Landscape module from Engine split out.

diff --git a/Engine/Source/Runtime/Landscape/Classes/LandscapeMeshCollisionComponent.h b/Engine/Source/Runtime/Landscape/Classes/LandscapeMeshCollisionComponent.h
new file mode 100644
--- /dev/null
+++ b/Engine/Source/Runtime/Landscape/Classes/LandscapeMeshCollisionComponent.h
@@ -0,0 +1,98 @@
+// Copyright 1998-2014 Epic Games, Inc. All Rights Reserved.
+
+
+#pragma once
+
+#include "LandscapeHeightfieldCollisionComponent.h"
+
+#include "LandscapeMeshCollisionComponent.generated.h"
+
+#if WITH_PHYSX
+namespace physx
+{
+	class PxMaterial;
+	class PxTriangleMesh;
+}
+#endif // WITH_PHYSX
+
+UCLASS()
+class ULandscapeMeshCollisionComponent : public ULandscapeHeightfieldCollisionComponent
+{
+	GENERATED_UCLASS_BODY()
+
+	// Keep the possibility to share projected height field PhysX object with editor mesh collision objects...
+
+	/** Guid used to share PhysX heightfield objects in the editor */
+	UPROPERTY()
+	FGuid MeshGuid;
+
+	struct FPhysXMeshRef : public FRefCountedObject
+	{
+		FGuid Guid;
+
+#if WITH_PHYSX
+		/** List of PxMaterials used on this landscape */
+		TArray<physx::PxMaterial*>	UsedPhysicalMaterialArray;
+		physx::PxTriangleMesh*		RBTriangleMesh;
+#if WITH_EDITOR
+		physx::PxTriangleMesh*		RBTriangleMeshEd; // Used only by landscape editor, does not have holes in it
+#endif	//WITH_EDITOR
+#endif	//WITH_PHYSX
+
+		/** tors **/
+		FPhysXMeshRef() 
+#if WITH_PHYSX
+			:	RBTriangleMesh(NULL)
+#if WITH_EDITOR
+			,	RBTriangleMeshEd(NULL)
+#endif	//WITH_EDITOR
+#endif	//WITH_PHYSX
+		{}
+		FPhysXMeshRef(FGuid& InGuid)
+			:	Guid(InGuid)
+#if WITH_PHYSX
+			,	RBTriangleMesh(NULL)
+#if WITH_EDITOR
+			,	RBTriangleMeshEd(NULL)
+#endif	//WITH_EDITOR
+#endif	//WITH_PHYSX
+		{}
+		virtual ~FPhysXMeshRef();
+	};
+
+	/** The collision mesh values. */
+	FWordBulkData								CollisionXYOffsetData; //  X, Y Offset in raw format...
+
+	/** Physics engine version of heightfield data. */
+	TRefCountPtr<struct FPhysXMeshRef>			MeshRef;
+
+	// Begin UActorComponent interface.
+	virtual void CreatePhysicsState() override;
+	virtual void ApplyWorldOffset(const FVector& InOffset, bool bWorldShift) override;
+	// End UActorComponent interface.
+
+	// Begin USceneComponent interface.
+	virtual void DestroyComponent() override;
+	// End USceneComponent interface.
+
+	// Begin UPrimitiveComponent interface
+	virtual bool DoCustomNavigableGeometryExport(struct FNavigableGeometryExport* GeomExport) const override;
+	//End UPrimitiveComponent interface
+
+	// Begin UObject Interface.
+	virtual void Serialize(FArchive& Ar) override;
+	virtual void BeginDestroy() override;
+#if WITH_EDITOR
+	virtual void ExportCustomProperties(FOutputDevice& Out, uint32 Indent) override;
+	virtual void ImportCustomProperties(const TCHAR* SourceText, FFeedbackContext* Warn) override;
+
+	virtual bool CookCollsionData(const FName& Format, bool bUseOnlyDefMaterial, TArray<uint8>& OutCookedData, TArray<UPhysicalMaterial*>& OutMaterails) const override;
+#endif
+	// End UObject Interface.
+
+	// Begin ULandscapeHeightfieldCollisionComponent Interface
+	virtual void CreateCollisionObject() override;
+	virtual void RecreateCollision(bool bUpdateAddCollision = true) override;
+	// End ULandscapeHeightfieldCollisionComponent Interface
+};
+
diff --git a/Engine/Source/Runtime/Landscape/Classes/LandscapeProxy.h b/Engine/Source/Runtime/Landscape/Classes/LandscapeProxy.h
new file mode 100644
--- /dev/null
+++ b/Engine/Source/Runtime/Landscape/Classes/LandscapeProxy.h
@@ -0,0 +1,440 @@
+// Copyright 1998-2014 Epic Games, Inc. All Rights Reserved.
+
+#pragma once
+
+#include "LandscapeInfo.h"
+
+#include "Engine/EngineTypes.h"
+#include "Engine/Texture.h"
+#include "Engine/TextureDefines.h"
+
+#include "GameFramework/Actor.h"
+#include "PhysicsEngine/BodyInstance.h"
+#include "Materials/MaterialInstanceConstant.h"
+
+#include "LandscapeProxy.generated.h"
+
+class ULandscapeMaterialInstanceConstant;
+class ULandscapeLayerInfoObject;
+class ULandscapeSplinesComponent;
+class ULandscapeHeightfieldCollisionComponent;
+class UMaterialInterface;
+class UTexture2D;
+class ALandscape;
+class ALandscapeProxy;
+class ULandscapeComponent;
+
+/** Structure storing channel usage for weightmap textures */
+USTRUCT()
+struct FLandscapeWeightmapUsage
+{
+	GENERATED_USTRUCT_BODY()
+
+	UPROPERTY()
+	ULandscapeComponent* ChannelUsage[4];
+
+	FLandscapeWeightmapUsage()
+	{
+		ChannelUsage[0] = NULL;
+		ChannelUsage[1] = NULL;
+		ChannelUsage[2] = NULL;
+		ChannelUsage[3] = NULL;
+	}
+	friend FArchive& operator<<( FArchive& Ar, FLandscapeWeightmapUsage& U );
+	int32 FreeChannelCount() const
+	{
+		return	((ChannelUsage[0] == NULL) ? 1 : 0) + 
+				((ChannelUsage[1] == NULL) ? 1 : 0) + 
+				((ChannelUsage[2] == NULL) ? 1 : 0) + 
+				((ChannelUsage[3] == NULL) ? 1 : 0);
+	}
+};
+
+USTRUCT()
+struct FLandscapeEditorLayerSettings
+{
+	GENERATED_USTRUCT_BODY()
+
+#if WITH_EDITORONLY_DATA
+	UPROPERTY()
+	ULandscapeLayerInfoObject* LayerInfoObj;
+
+	UPROPERTY()
+	FString ReimportLayerFilePath;
+
+	FLandscapeEditorLayerSettings()
+		: LayerInfoObj(NULL)
+		, ReimportLayerFilePath()
+	{
+	}
+
+	FLandscapeEditorLayerSettings(ULandscapeLayerInfoObject* InLayerInfo, const FString& InFilePath = FString())
+		: LayerInfoObj(InLayerInfo)
+		, ReimportLayerFilePath(InFilePath)
+	{
+	}
+
+	bool operator==(const FLandscapeEditorLayerSettings& rhs) const
+	{
+		return LayerInfoObj == rhs.LayerInfoObj;
+	}
+#endif // WITH_EDITORONLY_DATA
+};
+
+USTRUCT()
+struct FLandscapeLayerStruct
+{
+	GENERATED_USTRUCT_BODY()
+
+	UPROPERTY()
+	ULandscapeLayerInfoObject* LayerInfoObj;
+
+#if WITH_EDITORONLY_DATA
+	UPROPERTY(transient)
+	UMaterialInstanceConstant* ThumbnailMIC;
+
+	UPROPERTY()
+	ALandscapeProxy* Owner;
+
+	UPROPERTY(transient)
+	int32 DebugColorChannel;
+
+	UPROPERTY(transient)
+	uint32 bSelected:1;
+
+	UPROPERTY()
+	FString SourceFilePath;
+#endif // WITH_EDITORONLY_DATA
+
+	FLandscapeLayerStruct()
+		: LayerInfoObj(NULL)
+#if WITH_EDITORONLY_DATA
+		, ThumbnailMIC(NULL)
+		, Owner(NULL)
+		, DebugColorChannel(0)
+		, bSelected(false)
+		, SourceFilePath()
+#endif // WITH_EDITORONLY_DATA
+	{
+	}
+};
+
+/** Structure storing Layer Data for import */
+USTRUCT()
+struct FLandscapeImportLayerInfo
+{
+	GENERATED_USTRUCT_BODY()
+
+#if WITH_EDITORONLY_DATA
+	UPROPERTY(Category="Import", VisibleAnywhere)
+	FName LayerName;
+
+	UPROPERTY(Category="Import", EditAnywhere)
+	ULandscapeLayerInfoObject* LayerInfo;
+
+	UPROPERTY(Category="Import", VisibleAnywhere)
+	UMaterialInstanceConstant* ThumbnailMIC; // Optional
+
+	UPROPERTY(Category="Import", EditAnywhere)
+	FString SourceFilePath; // Optional
+	
+	// Raw weightmap data
+	TArray<uint8> LayerData;		
+#endif
+
+#if WITH_EDITOR
+	FLandscapeImportLayerInfo(FName InLayerName = NAME_None)
+	:	LayerName(InLayerName)
+	,	LayerInfo(NULL)
+	,	ThumbnailMIC(NULL)
+	,	SourceFilePath("")
+	{
+	}
+
+	LANDSCAPE_API FLandscapeImportLayerInfo(const struct FLandscapeInfoLayerSettings& InLayerSettings);
+#endif
+};
+
+// this is only here because putting it in LandscapeEditorObject.h (where it belongs)
+// results in Engine being dependent on LandscapeEditor, as the actual landscape editing
+// code (e.g. LandscapeEdit.h) is in /Engine/ for some reason...
+UENUM()
+namespace ELandscapeLayerPaintingRestriction
+{
+	enum Type
+	{
+		// No restriction, can paint anywhere (default)
+		None         UMETA(DisplayName="None"),
+
+		// Uses the MaxPaintedLayersPerComponent setting from the LandscapeProxy
+		UseMaxLayers UMETA(DisplayName="Limit Layer Count"),
+
+		// Restricts painting to only components that already have this layer
+		ExistingOnly UMETA(DisplayName="Existing Layers Only"),
+	};
+}
+
+UENUM()
+namespace ELandscapeLODFalloff
+{
+	enum Type
+	{
+		// Default mode
+		Linear			UMETA(DisplayName = "Linear"),
+		// Square Root give more natural transition, and also keep the same LOD 
+		SquareRoot		UMETA(DisplayName = "Square Root"),
+	};
+}
+
+UCLASS(NotPlaceable, hidecategories=(Display, Attachment, Physics, Debug, Lighting, LOD), showcategories=(Rendering, "Utilities|Transformation"), MinimalAPI)
+class ALandscapeProxy : public AActor
+{
+	GENERATED_UCLASS_BODY()
+
+	UPROPERTY()
+	ULandscapeSplinesComponent* SplineComponent;
+
+protected:
+	/** Guid for LandscapeEditorInfo **/
+	UPROPERTY()
+	FGuid LandscapeGuid;
+
+public:
+	/** Offset in quads from landscape actor origin **/
+	UPROPERTY()
+	FIntPoint LandscapeSectionOffset;
+
+	/** Max LOD level to use when rendering */
+	UPROPERTY(EditAnywhere, Category=LOD)
+	int32 MaxLODLevel;
+
+#if WITH_EDITORONLY_DATA
+	/** LOD level to use when exporting the landscape to obj or FBX */
+	UPROPERTY(EditAnywhere, Category=LOD)
+	int32 ExportLOD;
+#endif
+
+	/** LOD level to use when running lightmass (increase to 1 or 2 for large landscapes to stop lightmass crashing) */
+	UPROPERTY(EditAnywhere, Category=Lighting)
+	int32 StaticLightingLOD;
+
+	/** Default physical material, used when no per-layer values physical materials */
+	UPROPERTY(EditAnywhere, Category=Landscape)
+	UPhysicalMaterial* DefaultPhysMaterial;
+
+	/**
+	 * Allows artists to adjust the distance where textures using UV 0 are streamed in/out.
+	 * 1.0 is the default, whereas a higher value increases the streamed-in resolution.
+	 */
+	UPROPERTY(EditAnywhere, Category=Landscape, meta=(ClampMin=0))
+	float StreamingDistanceMultiplier;
+
+	/** Combined material used to render the landscape */
+	UPROPERTY(EditAnywhere, Category=Landscape)
+	UMaterialInterface* LandscapeMaterial;
+
+	/** Material used to render landscape components with holes. Should be a BLEND_Masked version of LandscapeMaterial. */
+	UPROPERTY(EditAnywhere, Category=Landscape)
+	UMaterialInterface* LandscapeHoleMaterial;
+
+	UPROPERTY(EditAnywhere, Category=LOD)
+	float LODDistanceFactor;
+
+	/** The array of LandscapeComponent that are used by the landscape */
+	UPROPERTY()
+	TArray<ULandscapeComponent*> LandscapeComponents;
+
+	/** Array of LandscapeHeightfieldCollisionComponent */
+	UPROPERTY()
+	TArray<ULandscapeHeightfieldCollisionComponent*> CollisionComponents;
+
+	/**
+	 *	The resolution to cache lighting at, in texels/quad in one axis
+	 *  Total resolution would be changed by StaticLightingResolution*StaticLightingResolution
+	 *	Automatically calculate proper value for removing seams
+	 */
+	UPROPERTY(EditAnywhere, Category=Lighting)
+	float StaticLightingResolution;
+
+	UPROPERTY(EditAnywhere, Category=LandscapeProxy)
+	TLazyObjectPtr<ALandscape> LandscapeActor;
+
+	UPROPERTY(EditAnywhere, Category=Lighting)
+	uint32 bCastStaticShadow:1;
+
+	/** Whether this primitive should cast dynamic shadows as if it were a two sided material. */
+	UPROPERTY(EditAnywhere, AdvancedDisplay, Category=Lighting)
+	uint32 bCastShadowAsTwoSided:1;
+
+	UPROPERTY()
+	uint32 bIsProxy:1;
+
+#if WITH_EDITORONLY_DATA
+	UPROPERTY(transient)
+	uint32 bIsMovingToLevel:1;    // Check for the Move to Current Level case
+#endif // WITH_EDITORONLY_DATA
+
+	/** The Lightmass settings for this object. */
+	UPROPERTY(EditAnywhere, Category=Lightmass)
+	struct FLightmassPrimitiveSettings LightmassSettings;
+
+	// Landscape LOD to use for collision tests. Higher numbers use less memory and process faster, but are much less accurate
+	UPROPERTY(EditAnywhere, Category=Landscape)
+	int32 CollisionMipLevel;
+
+	/** Thickness of the collision surface, in unreal units */
+	UPROPERTY(EditAnywhere, Category=Landscape)
+	float CollisionThickness;
+
+	/** Collision profile settings for this landscape */
+	UPROPERTY(EditAnywhere, BlueprintReadOnly, Category=Collision, meta=(ShowOnlyInnerProperties))
+	FBodyInstance BodyInstance;
+
+	UPROPERTY()
+	TArray<struct FLandscapeLayerStruct> LayerInfoObjs_DEPRECATED;
+
+#if WITH_EDITORONLY_DATA
+	UPROPERTY()
+	TArray<ULandscapeLayerInfoObject*> EditorCachedLayerInfos_DEPRECATED;
+
+	UPROPERTY()
+	FString ReimportHeightmapFilePath;
+
+	UPROPERTY()
+	TArray<struct FLandscapeEditorLayerSettings> EditorLayerSettings;
+#endif
+
+	/** Data set at creation time */
+	UPROPERTY()
+	int32 ComponentSizeQuads;    // Total number of quads in each component
+
+	UPROPERTY()
+	int32 SubsectionSizeQuads;    // Number of quads for a subsection of a component. SubsectionSizeQuads+1 must be a power of two.
+
+	UPROPERTY()
+	int32 NumSubsections;    // Number of subsections in X and Y axis
+
+	/** Change the Level of Detail distance factor */
+	UFUNCTION(BlueprintCallable, Category="Rendering")
+	virtual void ChangeLODDistanceFactor(float InLODDistanceFactor);
+
+	/** Hints navigation system whether this landscape will ever be navigated on. true by default, but make sure to set it to false for faraway, background landscapes */
+	UPROPERTY(EditAnywhere, Category=Landscape)
+	uint32 bUsedForNavigation:1;
+
+	UPROPERTY(EditAnywhere, Category=LOD)
+	TEnumAsByte<enum ELandscapeLODFalloff::Type> LODFalloff;
+
+#if WITH_EDITORONLY_DATA
+	UPROPERTY(EditAnywhere, Category=Landscape)
+	int32 MaxPaintedLayersPerComponent; // 0 = disabled
+#endif
+
+public:
+	LANDSCAPE_API static ULandscapeLayerInfoObject* VisibilityLayer;
+
+	/** Map of material instance constants used to for the components. Key is generated with ULandscapeComponent::GetLayerAllocationKey() */
+	TMap<FString, UMaterialInstanceConstant*> MaterialInstanceConstantMap;
+
+	/** Map of weightmap usage */
+	TMap<UTexture2D*, struct FLandscapeWeightmapUsage> WeightmapUsageMap;
+
+	// Begin AActor Interface
+	virtual void UnregisterAllComponents() override;
+	virtual void RegisterAllComponents() override;
+	virtual void RerunConstructionScripts() override {}
+	virtual bool IsLevelBoundsRelevant() const override { return true; }
+#if WITH_EDITOR
+	virtual void Destroyed() override;
+	virtual void EditorApplyScale(const FVector& DeltaScale, const FVector* PivotLocation, bool bAltDown, bool bShiftDown, bool bCtrlDown) override;
+	virtual void PostEditMove(bool bFinished) override;
+	virtual bool ShouldImport(FString* ActorPropString, bool IsMovingLevel) override;
+	virtual bool ShouldExport() override;
+	virtual bool GetSelectedComponents(TArray<UObject*>& SelectedObjects) override;
+	// End AActor Interface
+#endif	//WITH_EDITOR
+
+	FGuid GetLandscapeGuid() const { return LandscapeGuid; }
+	virtual ALandscape* GetLandscapeActor();
+
+	// Begin UObject interface.
+	virtual void Serialize(FArchive& Ar) override;
+	static void AddReferencedObjects(UObject* InThis, FReferenceCollector& Collector);
+	virtual void PostLoad() override;
+
+#if WITH_EDITOR
+	virtual void PostEditChangeProperty(FPropertyChangedEvent& PropertyChangedEvent) override;
+	virtual void PostEditChangeChainProperty(FPropertyChangedChainEvent& PropertyChangedEvent) override;
+	virtual void PreEditUndo() override;
+	virtual void PostEditUndo() override;
+	virtual void PostEditImport() override;
+	// End UObject Interface
+
+	FLandscapeLayerStruct* GetLayerInfo_Deprecated(FName LayerName);
+	LANDSCAPE_API static TArray<FName> GetLayersFromMaterial(UMaterialInterface* Material);
+	LANDSCAPE_API TArray<FName> GetLayersFromMaterial() const;
+	LANDSCAPE_API static ULandscapeLayerInfoObject* CreateLayerInfo(const TCHAR* LayerName, ULevel* Level);
+	LANDSCAPE_API ULandscapeLayerInfoObject* CreateLayerInfo(const TCHAR* LayerName);
+
+	LANDSCAPE_API ULandscapeInfo* GetLandscapeInfo(bool bSpawnNewActor = true);
+
+	// Get Landscape Material assigned to this Landscape
+	virtual UMaterialInterface* GetLandscapeMaterial() const;
+
+	// Get Hole Landscape Material assigned to this Landscape
+	virtual UMaterialInterface* GetLandscapeHoleMaterial() const;
+
+	// Remove Invalid weightmaps
+	LANDSCAPE_API void RemoveInvalidWeightmaps();
+
+	// Changed Physical Material
+	LANDSCAPE_API void ChangedPhysMaterial();
+
+	// Check input Landscape actor is match for this LandscapeProxy (by GUID)
+	LANDSCAPE_API bool IsValidLandscapeActor(ALandscape* Landscape);
+
+	// Copy properties from parent Landscape actor
+	LANDSCAPE_API void GetSharedProperties(ALandscapeProxy* Landscape);
+
+	/** Get the LandcapeActor-to-world transform with respect to landscape section offset*/
+	LANDSCAPE_API FTransform LandscapeActorToWorld() const;
+	
+	/** Set landscape absolute location in section space */
+	LANDSCAPE_API void SetAbsoluteSectionBase(FIntPoint SectionOffset);
+	
+	/** Get landscape position in section space */
+	LANDSCAPE_API FIntPoint GetSectionBaseOffset() const;
+	
+	/** Recreate all components rendering and collision states */
+	LANDSCAPE_API void RecreateComponentsState();
+
+	/** Recreate all collision components based on render component */
+	LANDSCAPE_API void RecreateCollisionComponents();
+
+	LANDSCAPE_API static ULandscapeMaterialInstanceConstant* GetLayerThumbnailMIC(UMaterialInterface* LandscapeMaterial, FName LayerName, UTexture2D* ThumbnailWeightmap, UTexture2D* ThumbnailHeightmap, ALandscapeProxy* Proxy);
+
+	LANDSCAPE_API void Import(FGuid Guid, int32 VertsX, int32 VertsY, 
+							int32 ComponentSizeQuads, int32 NumSubsections, int32 SubsectionSizeQuads, 
+							const uint16* HeightData, const TCHAR* HeightmapFileName, 
+							const TArray<FLandscapeImportLayerInfo>& ImportLayerInfos);
+
+	/**
+	 * Exports landscape into raw mesh
+	 * 
+	 * @param InExportLOD Landscape LOD level to use while exporting, INDEX_NONE will use ALanscapeProxy::ExportLOD settings
+	 * @param OutRawMesh - Resulting raw mesh
+	 * @return true if successful
+	 */
+	LANDSCAPE_API bool ExportToRawMesh(int32 InExportLOD, struct FRawMesh& OutRawMesh) const;
+
+
+	/** @return Current size of bounding rectangle in quads space */
+	LANDSCAPE_API FIntRect GetBoundingRect() const;
+
+	/** Creates a Texture2D for use by this landscape proxy or one of it's components. If OptionalOverrideOuter is not specified, the level is used. */
+	LANDSCAPE_API UTexture2D* CreateLandscapeTexture(int32 InSizeX, int32 InSizeY, TextureGroup InLODGroup, ETextureSourceFormat InFormat, UObject* OptionalOverrideOuter = nullptr) const;
+#endif
+};
+
+
+
diff --git a/Engine/Source/Runtime/Landscape/Classes/LandscapeSplineControlPoint.h b/Engine/Source/Runtime/Landscape/Classes/LandscapeSplineControlPoint.h
new file mode 100644
--- /dev/null
+++ b/Engine/Source/Runtime/Landscape/Classes/LandscapeSplineControlPoint.h
@@ -0,0 +1,181 @@
+// Copyright 1998-2014 Epic Games, Inc. All Rights Reserved.
+
+#pragma once
+
+#include "LandscapeSplineSegment.h"
+
+#include "LandscapeSplineControlPoint.generated.h"
+
+struct FLandscapeSplineInterpPoint;
+
+USTRUCT()
+struct FLandscapeSplineConnection
+{
+	GENERATED_USTRUCT_BODY()
+
+	FLandscapeSplineConnection() {}
+
+	FLandscapeSplineConnection(class ULandscapeSplineSegment* InSegment, int32 InEnd)
+		: Segment(InSegment)
+		, End(InEnd)
+	{
+	}
+
+	bool operator == (const FLandscapeSplineConnection& rhs) const
+	{
+		return Segment == rhs.Segment && End == rhs.End;
+	}
+
+	// Segment connected to this control point
+	UPROPERTY()
+	class ULandscapeSplineSegment* Segment;
+
+	// Which end of the segment is connected to this control point
+	UPROPERTY()
+	uint32 End:1;
+
+	LANDSCAPE_API struct FLandscapeSplineSegmentConnection& GetNearConnection() const;
+	LANDSCAPE_API struct FLandscapeSplineSegmentConnection& GetFarConnection() const;
+};
+
+UCLASS(Within=LandscapeSplinesComponent,autoExpandCategories=LandscapeSplineControlPoint,MinimalAPI)
+class ULandscapeSplineControlPoint : public UObject
+{
+	GENERATED_UCLASS_BODY()
+
+// Directly editable data:
+	/** Location in Landscape-space */
+	UPROPERTY(EditAnywhere, Category=LandscapeSplineControlPoint)
+	FVector Location;
+
+	/** Rotation of tangent vector at this point (in landscape-space) */
+	UPROPERTY(EditAnywhere, Category=LandscapeSplineControlPoint)
+	FRotator Rotation;
+
+	/** Width of the spline at this point. */
+	UPROPERTY(EditAnywhere, Category=LandscapeSplineControlPoint)
+	float Width;
+
+	/** Falloff at the sides of the spline at this point. */
+	UPROPERTY(EditAnywhere, Category=LandscapeSplineControlPoint)
+	float SideFalloff;
+
+	/** Falloff at the start/end of the spline (if this point is a start or end point, otherwise ignored). */
+	UPROPERTY(EditAnywhere, Category=LandscapeSplineControlPoint)
+	float EndFalloff;
+
+#if WITH_EDITORONLY_DATA
+	/** Mesh to use on the control point */
+	UPROPERTY(EditAnywhere, Category=Mesh)
+	class UStaticMesh* Mesh;
+
+	/** Overrides mesh's materials */
+	UPROPERTY(EditAnywhere, Category=Mesh)
+	TArray<class UMaterialInterface*> MaterialOverrides;
+
+	/** Scale of the control point mesh */
+	UPROPERTY(EditAnywhere, Category=Mesh)
+	FVector MeshScale;
+
+	/**  Max draw distance for the mesh used on this control point */
+	UPROPERTY(EditAnywhere, AdvancedDisplay, Category=Mesh, meta = (DisplayName = "Max Draw Distance"))
+	float LDMaxDrawDistance;
+
+	/**
+	 * Name of blend layer to paint when applying spline to landscape
+	 * If "none", no layer is painted
+	 */
+	UPROPERTY(EditAnywhere, Category=LandscapeSplineControlPoint)
+	FName LayerName;
+
+	/** If the spline is above the terrain, whether to raise the terrain up to the level of the spline when applying it to the landscape. */
+	UPROPERTY(EditAnywhere, Category=LandscapeSplineControlPoint)
+	uint32 bRaiseTerrain:1;
+
+	/** If the spline is below the terrain, whether to lower the terrain down to the level of the spline when applying it to the landscape. */
+	UPROPERTY(EditAnywhere, Category=LandscapeSplineControlPoint)
+	uint32 bLowerTerrain:1;
+
+	/** Whether to enable collision for the Control Point Mesh. */
+	UPROPERTY(EditAnywhere, Category=Mesh)
+	uint32 bEnableCollision:1;
+
+	/** Whether the Control Point Mesh should cast a shadow. */
+	UPROPERTY(EditAnywhere, Category=Mesh)
+	uint32 bCastShadow:1;
+
+protected:
+	UPROPERTY(Transient)
+	uint32 bSelected : 1;
+
+	UPROPERTY(Transient)
+	uint32 bNavDirty : 1;
+#endif
+
+// Procedural data:
+public:
+	UPROPERTY(TextExportTransient)
+	TArray<FLandscapeSplineConnection> ConnectedSegments;
+
+protected:
+	/** Control point mesh */
+	UPROPERTY(TextExportTransient)
+	class UControlPointMeshComponent* MeshComponent;
+
+	/** Spline points */
+	UPROPERTY()
+	TArray<struct FLandscapeSplineInterpPoint> Points;
+
+	/** Bounds of points */
+	UPROPERTY()
+	FBox Bounds;
+
+public:
+	const FBox& GetBounds() const { return Bounds; }
+	const TArray<struct FLandscapeSplineInterpPoint>& GetPoints() const { return Points; }
+
+	// Get the name of the best connection point (socket) to use for a particular destination
+	virtual FName GetBestConnectionTo(FVector Destination) const;
+
+	// Get the location and rotation of a connection point (socket) in spline space
+	virtual void GetConnectionLocalLocationAndRotation(FName SocketName, OUT FVector& OutLocation, OUT FRotator& OutRotation) const;
+
+	// Get the location and rotation of a connection point (socket) in spline space
+	virtual void GetConnectionLocationAndRotation(FName SocketName, OUT FVector& OutLocation, OUT FRotator& OutRotation) const;
+
+#if WITH_EDITOR
+	bool IsSplineSelected() const { return bSelected; }
+	virtual void SetSplineSelected(bool bInSelected);
+
+	/** Calculates rotation from connected segments */
+	virtual void AutoCalcRotation();
+
+	/**  */
+	virtual void AutoFlipTangents();
+
+	virtual void AutoSetConnections(bool bIncludingValid);
+
+	/** Update spline points */
+	virtual void UpdateSplinePoints(bool bUpdateCollision = true, bool bUpdateAttachedSegments = true);
+
+	/** Delete spline points */
+	virtual void DeleteSplinePoints();
+#endif // WITH_EDITOR
+
+	// Begin UObject Interface
+	virtual void PostLoad() override;
+#if WITH_EDITOR
+	virtual void PostEditUndo() override;
+	virtual void PostDuplicate(bool bDuplicateForPIE) override;
+	virtual void PostEditImport() override;
+	virtual void PostEditChangeProperty(FPropertyChangedEvent& PropertyChangedEvent) override;
+	// End UObject Interface
+#endif // WITH_EDITOR
+
+	void RegisterComponents();
+	void UnregisterComponents();
+
+	virtual bool OwnsComponent(const class UControlPointMeshComponent* ControlPointMeshComponent) const;
+
+	friend class FEdModeLandscape;
+};
diff --git a/Engine/Source/Runtime/Landscape/Classes/Materials/MaterialExpressionLandscapeLayerBlend.h b/Engine/Source/Runtime/Landscape/Classes/Materials/MaterialExpressionLandscapeLayerBlend.h
new file mode 100644
--- /dev/null
+++ b/Engine/Source/Runtime/Landscape/Classes/Materials/MaterialExpressionLandscapeLayerBlend.h
@@ -0,0 +1,93 @@
+// Copyright 1998-2014 Epic Games, Inc. All Rights Reserved.
+
+
+#pragma once
+#include "Materials/MaterialExpression.h"
+#include "MaterialExpressionLandscapeLayerBlend.generated.h"
+
+UENUM()
+enum ELandscapeLayerBlendType
+{
+	LB_WeightBlend,
+	LB_AlphaBlend,
+	LB_HeightBlend,
+	LB_MAX,
+};
+
+USTRUCT()
+struct FLayerBlendInput
+{
+	GENERATED_USTRUCT_BODY()
+
+	UPROPERTY(EditAnywhere, Category=LayerBlendInput)
+	FName LayerName;
+
+	UPROPERTY(EditAnywhere, Category=LayerBlendInput)
+	TEnumAsByte<ELandscapeLayerBlendType> BlendType;
+
+	UPROPERTY(meta = (RequiredInput = "false", ToolTip = "Defaults to 'ConstLayerInput' if not specified"))
+	FExpressionInput LayerInput;
+
+	UPROPERTY(meta = (RequiredInput = "false", ToolTip = "Defaults to 'ConstHeightInput' if not specified"))
+	FExpressionInput HeightInput;
+
+	UPROPERTY(EditAnywhere, Category=LayerBlendInput)
+	float PreviewWeight;
+
+	/** only used if LayerInput is not hooked up */
+	UPROPERTY(EditAnywhere, Category = LayerBlendInput)
+	FVector ConstLayerInput;
+
+	/** only used if HeightInput is not hooked up */
+	UPROPERTY(EditAnywhere, Category=LayerBlendInput)
+	float ConstHeightInput;
+
+	FLayerBlendInput()
+		: BlendType(0)
+		, PreviewWeight(0.0f)
+		, ConstLayerInput(0.0f, 0.0f, 0.0f)
+		, ConstHeightInput(0.0f)
+	{
+	}
+};
+
+UCLASS(collapsecategories, hidecategories=Object)
+class UMaterialExpressionLandscapeLayerBlend : public UMaterialExpression
+{
+	GENERATED_UCLASS_BODY()
+
+	UPROPERTY(EditAnywhere, Category=MaterialExpressionLandscapeLayerBlend)
+	TArray<FLayerBlendInput> Layers;
+
+	/** GUID that should be unique within the material, this is used for parameter renaming. */
+	UPROPERTY()
+	FGuid ExpressionGUID;
+
+
+	// Begin UObject Interface
+	void Serialize(FArchive& Ar) override;
+#if WITH_EDITOR
+	virtual void PostEditChangeProperty(FPropertyChangedEvent& PropertyChangedEvent) override;
+	// End UObject Interface
+
+	// Begin UMaterialExpression Interface
+#endif
+	virtual bool IsResultMaterialAttributes(int32 OutputIndex) override;
+	virtual int32 Compile(class FMaterialCompiler* Compiler, int32 OutputIndex, int32 MultiplexIndex) override;
+	virtual void GetCaption(TArray<FString>& OutCaptions) const override;
+	virtual const TArray<FExpressionInput*> GetInputs() override;
+	virtual FExpressionInput* GetInput(int32 InputIndex) override;
+	virtual FString GetInputName(int32 InputIndex) const override;
+	virtual UTexture* GetReferencedTexture() override;
+	// End UMaterialExpression Interface
+
+	LANDSCAPE_API virtual FGuid& GetParameterExpressionId() override;
+
+	/**
+	 * Get list of parameter names for static parameter sets
+	 */
+	void GetAllParameterNames(TArray<FName> &OutParameterNames, TArray<FGuid> &OutParameterIds);
+};
+
+
+
diff --git a/Engine/Source/Runtime/Landscape/Classes/Materials/MaterialExpressionLandscapeLayerCoords.h b/Engine/Source/Runtime/Landscape/Classes/Materials/MaterialExpressionLandscapeLayerCoords.h
new file mode 100644
--- /dev/null
+++ b/Engine/Source/Runtime/Landscape/Classes/Materials/MaterialExpressionLandscapeLayerCoords.h
@@ -0,0 +1,65 @@
+// Copyright 1998-2014 Epic Games, Inc. All Rights Reserved.
+
+
+#pragma once
+#include "Materials/MaterialExpression.h"
+#include "MaterialExpressionLandscapeLayerCoords.generated.h"
+
+UENUM()
+enum ETerrainCoordMappingType
+{
+	TCMT_Auto,
+	TCMT_XY,
+	TCMT_XZ,
+	TCMT_YZ,
+	TCMT_MAX,
+};
+
+UENUM()
+enum ELandscapeCustomizedCoordType
+{
+	LCCT_None, // Don't use customized UV, just use original UV
+	LCCT_CustomUV0,
+	LCCT_CustomUV1,
+	LCCT_CustomUV2,
+	LCCT_WeightMapUV, // Use original WeightMapUV, which could not be customized
+	LCCT_MAX,
+};
+
+UCLASS(collapsecategories, hidecategories=Object)
+class UMaterialExpressionLandscapeLayerCoords : public UMaterialExpression
+{
+	GENERATED_UCLASS_BODY()
+
+	/** Determines the mapping place to use on the terrain. */
+	UPROPERTY(EditAnywhere, Category=MaterialExpressionLandscapeLayerCoords)
+	TEnumAsByte<enum ETerrainCoordMappingType> MappingType;
+
+	/** Determines the mapping place to use on the terrain. */
+	UPROPERTY(EditAnywhere, Category=MaterialExpressionLandscapeLayerCoords)
+	TEnumAsByte<enum ELandscapeCustomizedCoordType> CustomUVType;
+
+	/** Uniform scale to apply to the mapping. */
+	UPROPERTY(EditAnywhere, Category=MaterialExpressionLandscapeLayerCoords)
+	float MappingScale;
+
+	/** Rotation to apply to the mapping. */
+	UPROPERTY(EditAnywhere, Category=MaterialExpressionLandscapeLayerCoords)
+	float MappingRotation;
+
+	/** Offset to apply to the mapping along U. */
+	UPROPERTY(EditAnywhere, Category=MaterialExpressionLandscapeLayerCoords)
+	float MappingPanU;
+
+	/** Offset to apply to the mapping along V. */
+	UPROPERTY(EditAnywhere, Category=MaterialExpressionLandscapeLayerCoords)
+	float MappingPanV;
+
+	// Begin UMaterialExpression Interface
+	virtual int32 Compile(class FMaterialCompiler* Compiler, int32 OutputIndex, int32 MultiplexIndex) override;
+	virtual void GetCaption(TArray<FString>& OutCaptions) const override;
+	// End UMaterialExpression Interface
+};
+
+
+
diff --git a/Engine/Source/Runtime/Landscape/Classes/Materials/MaterialExpressionLandscapeLayerSwitch.h b/Engine/Source/Runtime/Landscape/Classes/Materials/MaterialExpressionLandscapeLayerSwitch.h
new file mode 100644
--- /dev/null
+++ b/Engine/Source/Runtime/Landscape/Classes/Materials/MaterialExpressionLandscapeLayerSwitch.h
@@ -0,0 +1,59 @@
+// Copyright 1998-2014 Epic Games, Inc. All Rights Reserved.
+
+
+#pragma once
+#include "Materials/MaterialExpression.h"
+#include "MaterialExpressionLandscapeLayerSwitch.generated.h"
+
+UCLASS(collapsecategories, hidecategories=Object)
+class UMaterialExpressionLandscapeLayerSwitch : public UMaterialExpression
+{
+	GENERATED_UCLASS_BODY()
+
+	UPROPERTY()
+	FExpressionInput LayerUsed;
+
+	UPROPERTY()
+	FExpressionInput LayerNotUsed;
+
+	UPROPERTY(EditAnywhere, Category=MaterialExpressionLandscapeLayerSwitch)
+	FName ParameterName;
+
+	UPROPERTY(EditAnywhere, Category=MaterialExpressionLandscapeLayerSwitch)
+	uint32 PreviewUsed:1;
+
+	/** GUID that should be unique within the material, this is used for parameter renaming. */
+	UPROPERTY()
+	FGuid ExpressionGUID;
+
+
+public:
+
+	// Begin UObject Interface
+	virtual void Serialize(FArchive& Ar) override;
+	// End UObject Interface
+
+	// Begin UMaterialExpression Interface
+	virtual bool IsResultMaterialAttributes(int32 OutputIndex) override;
+	virtual int32 Compile(class FMaterialCompiler* Compiler, int32 OutputIndex, int32 MultiplexIndex) override;
+	virtual void GetCaption(TArray<FString>& OutCaptions) const override;
+	virtual UTexture* GetReferencedTexture() override;
+#if WITH_EDITOR
+	virtual uint32 GetInputType(int32 InputIndex) override {return MCT_Unknown;}
+	virtual uint32 GetOutputType(int32 InputIndex) override {return MCT_Unknown;}
+#endif
+	// End UMaterialExpression Interface
+
+	LANDSCAPE_API virtual FGuid& GetParameterExpressionId() override;
+
+	void GetAllParameterNames(TArray<FName> &OutParameterNames, TArray<FGuid> &OutParameterIds);
+
+	/**
+	 * Do any object-specific cleanup required immediately after loading an object,
+	 * and immediately after any undo/redo.
+	 */
+	void PostLoad() override;
+};
+
+
+
diff --git a/Engine/Source/Runtime/Landscape/Classes/Materials/MaterialExpressionLandscapeLayerWeight.h b/Engine/Source/Runtime/Landscape/Classes/Materials/MaterialExpressionLandscapeLayerWeight.h
new file mode 100644
--- /dev/null
+++ b/Engine/Source/Runtime/Landscape/Classes/Materials/MaterialExpressionLandscapeLayerWeight.h
@@ -0,0 +1,60 @@
+// Copyright 1998-2014 Epic Games, Inc. All Rights Reserved.
+
+
+#pragma once
+#include "Materials/MaterialExpression.h"
+#include "MaterialExpressionLandscapeLayerWeight.generated.h"
+
+UCLASS(collapsecategories, hidecategories=Object)
+class UMaterialExpressionLandscapeLayerWeight : public UMaterialExpression
+{
+	GENERATED_UCLASS_BODY()
+
+	UPROPERTY(meta = (RequiredInput = "false", ToolTip = "Defaults to 'ConstBase' if not specified"))
+	FExpressionInput Base;
+
+	UPROPERTY(meta = (RequiredInput = "false", ToolTip = "Ignored if not specified"))
+	FExpressionInput Layer;
+
+	UPROPERTY(EditAnywhere, Category=MaterialExpressionLandscapeLayerWeight)
+	FName ParameterName;
+
+	UPROPERTY(EditAnywhere, Category=MaterialExpressionLandscapeLayerWeight)
+	float PreviewWeight;
+
+	/** only used if Base is not hooked up */
+	UPROPERTY(EditAnywhere, Category = MaterialExpressionLandscapeLayerWeight)
+	FVector ConstBase;
+
+	/** GUID that should be unique within the material, this is used for parameter renaming. */
+	UPROPERTY()
+	FGuid ExpressionGUID;
+
+public:
+
+	// Begin UMaterialExpression Interface
+	virtual bool IsResultMaterialAttributes(int32 OutputIndex) override;
+	virtual int32 Compile(class FMaterialCompiler* Compiler, int32 OutputIndex, int32 MultiplexIndex) override;
+	virtual void GetCaption(TArray<FString>& OutCaptions) const override;
+	virtual UTexture* GetReferencedTexture() override;
+#if WITH_EDITOR
+	virtual uint32 GetInputType(int32 InputIndex) override {return MCT_Float | MCT_MaterialAttributes;}
+#endif //WITH_EDITOR
+	// End UMaterialExpression Interface
+
+	LANDSCAPE_API virtual FGuid& GetParameterExpressionId() override;
+
+	/**
+	 * Called to get list of parameter names for static parameter sets
+	 */
+	void GetAllParameterNames(TArray<FName> &OutParameterNames, TArray<FGuid> &OutParameterIds);
+
+	/**
+	 * Do any object-specific cleanup required immediately after loading an object,
+	 * and immediately after any undo/redo.
+	 */
+	void PostLoad() override;
+};
+
+
+
diff --git a/Engine/Source/Runtime/Landscape/Classes/Materials/MaterialExpressionLandscapeVisibilityMask.h b/Engine/Source/Runtime/Landscape/Classes/Materials/MaterialExpressionLandscapeVisibilityMask.h
new file mode 100644
--- /dev/null
+++ b/Engine/Source/Runtime/Landscape/Classes/Materials/MaterialExpressionLandscapeVisibilityMask.h
@@ -0,0 +1,36 @@
+// Copyright 1998-2014 Epic Games, Inc. All Rights Reserved.
+
+
+#pragma once
+#include "Materials/MaterialExpression.h"
+#include "MaterialExpressionLandscapeVisibilityMask.generated.h"
+
+UCLASS(collapsecategories, hidecategories=Object)
+class UMaterialExpressionLandscapeVisibilityMask : public UMaterialExpression
+{
+	GENERATED_UCLASS_BODY()
+
+	/** GUID that should be unique within the material, this is used for parameter renaming. */
+	UPROPERTY()
+	FGuid ExpressionGUID;
+
+public:
+
+	LANDSCAPE_API static FName ParameterName;
+
+	// Begin UMaterialExpression Interface
+	virtual int32 Compile(class FMaterialCompiler* Compiler, int32 OutputIndex, int32 MultiplexIndex) override;
+	virtual void GetCaption(TArray<FString>& OutCaptions) const override;
+	virtual UTexture* GetReferencedTexture() override;
+	// End UMaterialExpression Interface
+
+	LANDSCAPE_API virtual FGuid& GetParameterExpressionId() override;
+
+	/**
+	 * Called to get list of parameter names for static parameter sets
+	 */
+	void GetAllParameterNames(TArray<FName> &OutParameterNames, TArray<FGuid> &OutParameterIds);
+};
+
+
+
diff --git a/Engine/Source/Runtime/Landscape/Landscape.Build.cs b/Engine/Source/Runtime/Landscape/Landscape.Build.cs
new file mode 100644
--- /dev/null
+++ b/Engine/Source/Runtime/Landscape/Landscape.Build.cs
@@ -0,0 +1,68 @@
+// Copyright 1998-2014 Epic Games, Inc. All Rights Reserved.
+
+using UnrealBuildTool;
+using System.IO;
+
+public class Landscape : ModuleRules
+{
+	public Landscape(TargetInfo Target)
+	{
+		bIsAutoStartupModule = true;
+
+		PrivateIncludePaths.AddRange(
+			new string[] {
+                "Runtime/Engine/Private",
+			}
+		);
+
+		PrivateIncludePathModuleNames.AddRange(
+			new string[] {
+				"TargetPlatform",
+				"DerivedDataCache",
+				"ImageWrapper",
+			}
+		);
+
+		PrivateDependencyModuleNames.AddRange(
+			new string[] {
+				"Core",
+				"CoreUObject",
+				"Engine",
+				"RenderCore", 
+				"RHI",
+				"ShaderCore",
+			}
+		);
+
+		SetupModulePhysXAPEXSupport(Target);
+		if (UEBuildConfiguration.bCompilePhysX && UEBuildConfiguration.bRuntimePhysicsCooking)
+		{
+			DynamicallyLoadedModuleNames.Add("PhysXFormats");
+			PrivateIncludePathModuleNames.Add("PhysXFormats");
+		}
+
+		if (UEBuildConfiguration.bBuildDeveloperTools && Target.Type != TargetRules.TargetType.Server)
+		{
+			PrivateDependencyModuleNames.AddRange(
+				new string[] {
+					"RawMesh"
+				}
+			);
+		}
+
+		if (UEBuildConfiguration.bBuildEditor == true)
+		{
+			PrivateDependencyModuleNames.AddRange(
+				new string[] {
+                    "UnrealEd", 
+                }
+			);
+
+			CircularlyReferencedDependentModules.AddRange(
+				new string[] {
+                    "UnrealEd",
+                }
+			);
+		}
+	}
+}
diff --git a/Engine/Source/Runtime/Landscape/Private/Landscape.cpp b/Engine/Source/Runtime/Landscape/Private/Landscape.cpp
new file mode 100644
--- /dev/null
+++ b/Engine/Source/Runtime/Landscape/Private/Landscape.cpp
@@ -0,0 +1,2040 @@
+// Copyright 1998-2014 Epic Games, Inc. All Rights Reserved.
+
+/*=============================================================================
+Landscape.cpp: Terrain rendering
+=============================================================================*/
+
+#include "Landscape.h"
+#include "EditorSupportDelegates.h"
+#include "LandscapeDataAccess.h"
+#include "LandscapeRender.h"
+#include "LandscapeRenderMobile.h"
+#include "MessageLog.h"
+#include "UObjectToken.h"
+#include "MapErrors.h"
+#include "DerivedDataCacheInterface.h"
+#include "TargetPlatform.h"
+#include "LandscapeMeshCollisionComponent.h"
+#include "LandscapeMaterialInstanceConstant.h"
+#include "LandscapeSplinesComponent.h"
+#include "LandscapeInfo.h"
+#include "LandscapeInfoMap.h"
+#include "LandscapeLayerInfoObject.h"
+#include "LightMap.h"
+#include "ShadowMap.h"
+#include "Engine/CollisionProfile.h"
+#include "Materials/Material.h"
+#include "EngineGlobals.h"
+#include "Engine/Engine.h"
+#include "AI/Navigation/NavigationSystem.h"
+#include "EngineUtils.h"
+
+#include "Materials/MaterialExpressionLandscapeLayerWeight.h"
+#include "Materials/MaterialExpressionLandscapeLayerSwitch.h"
+#include "Materials/MaterialExpressionLandscapeLayerSample.h"
+#include "Materials/MaterialExpressionLandscapeLayerBlend.h"
+#include "Materials/MaterialExpressionLandscapeVisibilityMask.h"
+
+/** Landscape stats */
+
+DEFINE_STAT(STAT_LandscapeDynamicDrawTime);
+DEFINE_STAT(STAT_LandscapeStaticDrawLODTime);
+DEFINE_STAT(STAT_LandscapeVFDrawTime);
+DEFINE_STAT(STAT_LandscapeComponents);
+DEFINE_STAT(STAT_LandscapeDrawCalls);
+DEFINE_STAT(STAT_LandscapeTriangles);
+DEFINE_STAT(STAT_LandscapeVertexMem);
+DEFINE_STAT(STAT_LandscapeComponentMem);
+
+// Set this to 0 to disable landscape cooking and thus disable it on device.
+#define ENABLE_LANDSCAPE_COOKING 1
+
+#define LOCTEXT_NAMESPACE "Landscape"
+
+static void PrintNumLandscapeShadows()
+{
+	int32 NumComponents = 0;
+	int32 NumShadowCasters = 0;
+	for (TObjectIterator<ULandscapeComponent> It; It; ++It)
+	{
+		ULandscapeComponent* LC = *It;
+		NumComponents++;
+		if (LC->CastShadow && LC->bCastDynamicShadow)
+		{
+			NumShadowCasters++;
+		}
+	}
+	UE_LOG(LogConsoleResponse, Display, TEXT("%d/%d landscape components cast shadows"), NumShadowCasters, NumComponents);
+}
+
+FAutoConsoleCommand CmdPrintNumLandscapeShadows(
+	TEXT("ls.PrintNumLandscapeShadows"),
+	TEXT("Prints the number of landscape components that cast shadows."),
+	FConsoleCommandDelegate::CreateStatic(PrintNumLandscapeShadows)
+	);
+
+ULandscapeComponent::ULandscapeComponent(const FObjectInitializer& ObjectInitializer)
+	: Super(ObjectInitializer)
+{
+	BodyInstance.bEnableCollision_DEPRECATED = true;
+	SetCollisionProfileName(UCollisionProfile::BlockAll_ProfileName);
+	bGenerateOverlapEvents = false;
+	CastShadow = true;
+	bUseAsOccluder = true;
+	bAllowCullDistanceVolume = false;
+	CollisionMipLevel = 0;
+	StaticLightingResolution = 0.f; // Default value 0 means no overriding
+
+	HeightmapScaleBias = FVector4(0.0f, 0.0f, 0.0f, 1.0f);
+	WeightmapScaleBias = FVector4(0.0f, 0.0f, 0.0f, 1.0f);
+
+	bBoundsChangeTriggersStreamingDataRebuild = true;
+	ForcedLOD = -1;
+	// Neighbor LOD and LODBias are saved in BYTE, so need to convert to range [-128:127]
+	NeighborLOD[0] = 255;
+	NeighborLOD[1] = 255;
+	NeighborLOD[2] = 255;
+	NeighborLOD[3] = 255;
+	NeighborLOD[4] = 255;
+	NeighborLOD[5] = 255;
+	NeighborLOD[6] = 255;
+	NeighborLOD[7] = 255;
+	LODBias = 0;
+	NeighborLODBias[0] = 128;
+	NeighborLODBias[1] = 128;
+	NeighborLODBias[2] = 128;
+	NeighborLODBias[3] = 128;
+	NeighborLODBias[4] = 128;
+	NeighborLODBias[5] = 128;
+	NeighborLODBias[6] = 128;
+	NeighborLODBias[7] = 128;
+#if WITH_EDITORONLY_DATA
+	LightingLODBias = -1; // -1 Means automatic LOD calculation based on ForcedLOD + LODBias
+#endif
+
+	Mobility = EComponentMobility::Static;
+
+#if WITH_EDITORONLY_DATA
+	EditToolRenderData = NULL;
+#endif
+
+	LpvBiasMultiplier = 0.0f; // Bias is 0 for landscape, since it's single sided
+}
+
+ULandscapeComponent::~ULandscapeComponent()
+{
+
+}
+
+
+void ULandscapeComponent::AddReferencedObjects(UObject* InThis, FReferenceCollector& Collector)
+{
+	ULandscapeComponent* This = CastChecked<ULandscapeComponent>(InThis);
+	if (This->LightMap != NULL)
+	{
+		This->LightMap->AddReferencedObjects(Collector);
+	}
+	if (This->ShadowMap != NULL)
+	{
+		This->ShadowMap->AddReferencedObjects(Collector);
+	}
+	Super::AddReferencedObjects(This, Collector);
+}
+
+void ULandscapeComponent::Serialize(FArchive& Ar)
+{
+#if ENABLE_LANDSCAPE_COOKING && WITH_EDITOR
+	// Saving for cooking path
+	if (Ar.IsCooking() && Ar.IsSaving() && !HasAnyFlags(RF_ClassDefaultObject))
+	{
+		if (!Ar.CookingTarget()->SupportsFeature(ETargetPlatformFeatures::VertexShaderTextureSampling))
+		{
+			if (!PlatformData.HasValidPlatformData())
+			{
+				GeneratePlatformVertexData();
+				MaterialInstance = Cast<UMaterialInstanceConstant>(GeneratePlatformPixelData(WeightmapTextures, true));
+				HeightmapTexture = NULL;
+			}
+		}
+	}
+#endif
+
+	Super::Serialize(Ar);
+
+	Ar << LightMap;
+	if (Ar.UE4Ver() >= VER_UE4_PRECOMPUTED_SHADOW_MAPS_BSP)
+	{
+		Ar << ShadowMap;
+	}
+
+#if WITH_EDITOR
+	if (Ar.IsTransacting())
+	{
+		if (EditToolRenderData)
+		{
+			Ar << EditToolRenderData->SelectedType;
+		}
+		else
+		{
+			int32 TempV = 0;
+			Ar << TempV;
+		}
+	}
+#endif
+
+	if (Ar.UE4Ver() < VER_UE4_CHANGED_IRRELEVANT_LIGHT_GUIDS)
+	{
+		// Irrelevant lights were incorrect before VER_UE4_CHANGED_IRRELEVANT_LIGHT_GUIDS
+		IrrelevantLights.Empty();
+	}
+
+	bool bCooked = false;
+
+	if (Ar.UE4Ver() >= VER_UE4_LANDSCAPE_PLATFORMDATA_COOKING && !HasAnyFlags(RF_ClassDefaultObject))
+	{
+		bCooked = Ar.IsCooking();
+		// Save a bool indicating whether this is cooked data
+		// This is needed when loading cooked data, to know to serialize differently
+		Ar << bCooked;
+	}
+
+	if (FPlatformProperties::RequiresCookedData() && !bCooked && Ar.IsLoading())
+	{
+		UE_LOG(LogLandscape, Fatal, TEXT("This platform requires cooked packages, and this landscape does not contain cooked data %s."), *GetName());
+	}
+
+#if ENABLE_LANDSCAPE_COOKING
+	if (bCooked)
+	{
+		// Saving for cooking path
+		if (Ar.IsCooking() && !Ar.CookingTarget()->SupportsFeature(ETargetPlatformFeatures::VertexShaderTextureSampling))
+		{
+			bool bValid = PlatformData.HasValidPlatformData();
+			Ar << bValid;
+			if (bValid)
+			{
+				Ar << PlatformData;
+			}
+		}
+		else if (!FPlatformProperties::SupportsVertexShaderTextureSampling())
+		{
+			// Loading cooked data path
+			bool bValid = false;
+			Ar << bValid;
+
+			if (bValid)
+			{
+				Ar << PlatformData;
+			}
+		}
+	}
+#endif
+}
+
+#if WITH_EDITOR
+UMaterialInterface* ULandscapeComponent::GetLandscapeMaterial() const
+{
+	if (OverrideMaterial)
+	{
+		return OverrideMaterial;
+	}
+	ALandscapeProxy* Proxy = GetLandscapeProxy();
+	if (Proxy)
+	{
+		return Proxy->GetLandscapeMaterial();
+	}
+	return UMaterial::GetDefaultMaterial(MD_Surface);
+}
+
+UMaterialInterface* ULandscapeComponent::GetLandscapeHoleMaterial() const
+{
+	if (OverrideHoleMaterial)
+	{
+		return OverrideHoleMaterial;
+	}
+	ALandscapeProxy* Proxy = GetLandscapeProxy();
+	if (Proxy)
+	{
+		return Proxy->GetLandscapeHoleMaterial();
+	}
+	return UMaterial::GetDefaultMaterial(MD_Surface);
+}
+
+bool ULandscapeComponent::ComponentHasVisibilityPainted() const
+{
+	for (const FWeightmapLayerAllocationInfo& Allocation : WeightmapLayerAllocations)
+	{
+		if (Allocation.LayerInfo == ALandscapeProxy::VisibilityLayer)
+		{
+			return true;
+		}
+	}
+
+	return false;
+}
+
+FString ULandscapeComponent::GetLayerAllocationKey(bool bMobile /* = false */) const
+{
+	UMaterialInterface* LandscapeMaterial = GetLandscapeMaterial();
+	if (!LandscapeMaterial)
+	{
+		return FString();
+	}
+
+	FString Result = LandscapeMaterial->GetPathName();
+
+	// Sort the allocations
+	TArray<FString> LayerStrings;
+	for (int32 LayerIdx = 0; LayerIdx < WeightmapLayerAllocations.Num(); LayerIdx++)
+	{
+		new(LayerStrings)FString(*FString::Printf(TEXT("_%s_%d"), *WeightmapLayerAllocations[LayerIdx].GetLayerName().ToString(), bMobile ? 0 : WeightmapLayerAllocations[LayerIdx].WeightmapTextureIndex));
+	}
+	/**
+	 * Generate a key for this component's layer allocations to use with MaterialInstanceConstantMap.
+	 */
+	LayerStrings.Sort(TGreater<FString>());
+
+	for (int32 LayerIdx = 0; LayerIdx < LayerStrings.Num(); LayerIdx++)
+	{
+		Result += LayerStrings[LayerIdx];
+	}
+	return Result;
+}
+
+void ULandscapeComponent::GetLayerDebugColorKey(int32& R, int32& G, int32& B) const
+{
+	ULandscapeInfo* Info = GetLandscapeInfo(false);
+	if (Info)
+	{
+		R = INDEX_NONE, G = INDEX_NONE, B = INDEX_NONE;
+
+		for (auto It = Info->Layers.CreateConstIterator(); It; It++)
+		{
+			const FLandscapeInfoLayerSettings& LayerStruct = *It;
+			if (LayerStruct.DebugColorChannel > 0
+				&& LayerStruct.LayerInfoObj)
+			{
+				for (int32 LayerIdx = 0; LayerIdx < WeightmapLayerAllocations.Num(); LayerIdx++)
+				{
+					if (WeightmapLayerAllocations[LayerIdx].LayerInfo == LayerStruct.LayerInfoObj)
+					{
+						if (LayerStruct.DebugColorChannel & 1) // R
+						{
+							R = (WeightmapLayerAllocations[LayerIdx].WeightmapTextureIndex * 4 + WeightmapLayerAllocations[LayerIdx].WeightmapTextureChannel);
+						}
+						if (LayerStruct.DebugColorChannel & 2) // G
+						{
+							G = (WeightmapLayerAllocations[LayerIdx].WeightmapTextureIndex * 4 + WeightmapLayerAllocations[LayerIdx].WeightmapTextureChannel);
+						}
+						if (LayerStruct.DebugColorChannel & 4) // B
+						{
+							B = (WeightmapLayerAllocations[LayerIdx].WeightmapTextureIndex * 4 + WeightmapLayerAllocations[LayerIdx].WeightmapTextureChannel);
+						}
+						break;
+					}
+				}
+			}
+		}
+	}
+}
+#endif	//WITH_EDITOR
+
+ULandscapeMeshCollisionComponent::ULandscapeMeshCollisionComponent(const FObjectInitializer& ObjectInitializer)
+	: Super(ObjectInitializer)
+{
+	// make landscape always create? 
+	bAlwaysCreatePhysicsState = true;
+}
+
+ULandscapeInfo::ULandscapeInfo(const FObjectInitializer& ObjectInitializer)
+	: Super(ObjectInitializer)
+{
+#if WITH_EDITOR
+	bIsValid = false;
+#endif
+}
+
+#if WITH_EDITOR
+void ULandscapeInfo::UpdateDebugColorMaterial()
+{
+	FlushRenderingCommands();
+	//GWarn->BeginSlowTask( *FString::Printf(TEXT("Compiling layer color combinations for %s"), *GetName()), true);
+
+	for (auto It = XYtoComponentMap.CreateIterator(); It; ++It)
+	{
+		ULandscapeComponent* Comp = It.Value();
+		if (Comp && Comp->EditToolRenderData)
+		{
+			Comp->EditToolRenderData->UpdateDebugColorMaterial();
+		}
+	}
+	FlushRenderingCommands();
+	//GWarn->EndSlowTask();
+}
+
+void ULandscapeComponent::PostLoad()
+{
+	Super::PostLoad();
+
+	// Ensure that the component's lighting settings matches the actor's.
+	ALandscapeProxy* LandscapeProxy = GetLandscapeProxy();
+	if (ensure(LandscapeProxy))
+	{
+		bCastStaticShadow = LandscapeProxy->bCastStaticShadow;
+		bCastShadowAsTwoSided = LandscapeProxy->bCastShadowAsTwoSided;
+
+		// convert from deprecated layer names to direct LayerInfo references
+		static const FName DataWeightmapName = FName("__DataLayer__");
+		for (int32 i = 0; i < WeightmapLayerAllocations.Num(); i++)
+		{
+			if (WeightmapLayerAllocations[i].LayerName_DEPRECATED != NAME_None)
+			{
+				if (WeightmapLayerAllocations[i].LayerName_DEPRECATED == DataWeightmapName)
+				{
+					WeightmapLayerAllocations[i].LayerInfo = ALandscapeProxy::VisibilityLayer;
+				}
+				else
+				{
+					FLandscapeLayerStruct* Layer = LandscapeProxy->GetLayerInfo_Deprecated(WeightmapLayerAllocations[i].LayerName_DEPRECATED);
+					WeightmapLayerAllocations[i].LayerInfo = Layer ? Layer->LayerInfoObj : NULL;
+				}
+				WeightmapLayerAllocations[i].LayerName_DEPRECATED = NAME_None;
+			}
+		}
+	}
+
+#if WITH_EDITOR
+	if (GIsEditor && !HasAnyFlags(RF_ClassDefaultObject))
+	{
+		// Remove standalone flags from data textures to ensure data is unloaded in the editor when reverting an unsaved level.
+		// Previous version of landscape set these flags on creation.
+		if (HeightmapTexture && HeightmapTexture->HasAnyFlags(RF_Standalone))
+		{
+			HeightmapTexture->ClearFlags(RF_Standalone);
+		}
+		for (int32 Idx = 0; Idx < WeightmapTextures.Num(); Idx++)
+		{
+			if (WeightmapTextures[Idx] && WeightmapTextures[Idx]->HasAnyFlags(RF_Standalone))
+			{
+				WeightmapTextures[Idx]->ClearFlags(RF_Standalone);
+			}
+		}
+	}
+#endif
+
+	if (!HasAnyFlags(RF_ClassDefaultObject))
+	{
+		if (!CachedLocalBox.IsValid && CachedBoxSphereBounds_DEPRECATED.SphereRadius > 0)
+		{
+			USceneComponent* LandscapeRoot = LandscapeProxy->GetRootComponent();
+			check(LandscapeRoot == AttachParent);
+
+			// Component isn't attached so we can't use its LocalToWorld
+			FTransform ComponentLtWTransform = FTransform(RelativeRotation, RelativeLocation, RelativeScale3D) * FTransform(LandscapeRoot->RelativeRotation, LandscapeRoot->RelativeLocation, LandscapeRoot->RelativeScale3D);
+
+			// This is good enough. The exact box will be calculated during painting.
+			CachedLocalBox = CachedBoxSphereBounds_DEPRECATED.GetBox().InverseTransformBy(ComponentLtWTransform);
+		}
+
+		// If we're loading on a platform that doesn't require cooked data, attempt to load missing data from the DDC
+		if (!FPlatformProperties::RequiresCookedData() && (GIsEditor || (GMaxRHIFeatureLevel == ERHIFeatureLevel::ES2 || GMaxRHIFeatureLevel == ERHIFeatureLevel::ES3_1)))
+		{
+			// Only check the DDC if we don't already have it loaded
+			if (!PlatformData.HasValidPlatformData())
+			{
+				// Attempt to load the ES2 landscape data from the DDC
+				if (!PlatformData.LoadFromDDC(StateId))
+				{
+					// Height Data is available after loading height map, so need to defer it
+					UE_LOG(LogLandscape, Warning, TEXT("Attempt to access the DDC when there is none available on component '%s'."), *GetFullName());
+				}
+			}
+		}
+	}
+
+#if WITH_EDITOR
+	if (GIsEditor && !HasAnyFlags(RF_ClassDefaultObject))
+	{
+		// Move the MICs and Textures back to the Package if they're currently in the level
+		// Moving them into the level caused them to be duplicated when running PIE, which is *very very slow*, so we've reverted that change
+		// Also clear the public flag to avoid various issues, e.g. generating and saving thumbnails that can never be seen
+		{
+			ULevel* Level = GetLevel();
+			if (ensure(Level))
+			{
+				TArray<UObject*> ObjectsToMoveFromLevelToPackage;
+				GetGeneratedTexturesAndMaterialInstances(ObjectsToMoveFromLevelToPackage);
+
+				UPackage* MyPackage = GetOutermost();
+				for (auto* Obj : ObjectsToMoveFromLevelToPackage)
+				{
+					Obj->ClearFlags(RF_Public);
+					if (Obj->GetOuter() == Level)
+					{
+						Obj->Rename(NULL, MyPackage, REN_DoNotDirty | REN_DontCreateRedirectors | REN_ForceNoResetLoaders | REN_NonTransactional);
+					}
+				}
+			}
+		}
+	}
+#endif
+}
+
+#endif // WITH_EDITOR
+
+ALandscape::ALandscape(const FObjectInitializer& ObjectInitializer)
+	: Super(ObjectInitializer)
+{
+	bIsProxy = false;
+
+#if WITH_EDITORONLY_DATA
+	bLockLocation = false;
+#endif // WITH_EDITORONLY_DATA
+}
+
+ALandscape* ALandscape::GetLandscapeActor()
+{
+	return this;
+}
+
+ALandscapeProxy::ALandscapeProxy(const FObjectInitializer& ObjectInitializer)
+	: Super(ObjectInitializer)
+{
+	PrimaryActorTick.bCanEverTick = false;
+	bReplicates = false;
+	NetUpdateFrequency = 10.0f;
+	bHidden = false;
+	bReplicateMovement = false;
+	bCanBeDamaged = false;
+
+	TSubobjectPtr<USceneComponent> SceneComponent = ObjectInitializer.CreateDefaultSubobject<USceneComponent>(this, TEXT("RootComponent0"));
+	RootComponent = SceneComponent;
+	RootComponent->RelativeScale3D = FVector(128.0f, 128.0f, 256.0f); // Old default scale, preserved for compatibility. See ULandscapeEditorObject::NewLandscape_Scale
+	RootComponent->Mobility = EComponentMobility::Static;
+	LandscapeSectionOffset = FIntPoint::ZeroValue;
+
+	StaticLightingResolution = 1.0f;
+	StreamingDistanceMultiplier = 1.0f;
+	bIsProxy = true;
+	MaxLODLevel = -1;
+#if WITH_EDITORONLY_DATA
+	bLockLocation = true;
+	bIsMovingToLevel = false;
+#endif // WITH_EDITORONLY_DATA
+	LODDistanceFactor = 1.0f;
+	LODFalloff = ELandscapeLODFalloff::Linear;
+	bCastStaticShadow = true;
+	bCastShadowAsTwoSided = false;
+	bUsedForNavigation = true;
+	CollisionThickness = 16;
+	BodyInstance.SetCollisionProfileName(UCollisionProfile::BlockAll_ProfileName);
+#if WITH_EDITORONLY_DATA
+	MaxPaintedLayersPerComponent = 0;
+#endif
+
+	if (VisibilityLayer == NULL)
+	{
+		// Structure to hold one-time initialization
+		struct FConstructorStatics
+		{
+			ConstructorHelpers::FObjectFinderOptional<ULandscapeLayerInfoObject> DataLayer;
+			FConstructorStatics()
+				: DataLayer(TEXT("LandscapeLayerInfoObject'/Engine/EditorLandscapeResources/DataLayer.DataLayer'"))
+			{
+			}
+		};
+		static FConstructorStatics ConstructorStatics;
+
+		VisibilityLayer = ConstructorStatics.DataLayer.Get();
+		check(VisibilityLayer);
+#if WITH_EDITORONLY_DATA
+		// This layer should be no weight blending
+		VisibilityLayer->bNoWeightBlend = true;
+#endif
+		VisibilityLayer->AddToRoot();
+	}
+}
+
+ALandscape* ALandscapeProxy::GetLandscapeActor()
+{
+#if WITH_EDITORONLY_DATA
+	return LandscapeActor.Get();
+#else
+	return 0;
+#endif // WITH_EDITORONLY_DATA
+}
+
+#if WITH_EDITOR
+ULandscapeInfo* ALandscapeProxy::GetLandscapeInfo(bool bSpawnNewActor /*= true*/)
+{
+	// LandscapeInfo generate
+	if (GIsEditor)
+	{
+		UWorld* OwningWorld = GetWorld();
+		if (OwningWorld)
+		{
+			auto &LandscapeInfoMap = GetLandscapeInfoMap(OwningWorld);
+
+			ULandscapeInfo* LandscapeInfo = LandscapeInfoMap.Map.FindRef(LandscapeGuid);
+			if (!LandscapeInfo && bSpawnNewActor && !HasAnyFlags(RF_BeginDestroyed))
+			{
+				LandscapeInfo = ConstructObject<ULandscapeInfo>(ULandscapeInfo::StaticClass(), OwningWorld, NAME_None, RF_Transactional | RF_Transient);
+				OwningWorld->Modify(false);
+				LandscapeInfoMap.Map.Add(LandscapeGuid, LandscapeInfo);
+			}
+
+			return LandscapeInfo;
+		}
+	}
+
+	return NULL;
+}
+#endif
+
+ALandscape* ULandscapeComponent::GetLandscapeActor() const
+{
+	ALandscapeProxy* Landscape = GetLandscapeProxy();
+	if (Landscape)
+	{
+		return Landscape->GetLandscapeActor();
+	}
+	return NULL;
+}
+
+ULevel* ULandscapeComponent::GetLevel() const
+{
+	AActor* MyOwner = GetOwner();
+	return MyOwner ? MyOwner->GetLevel() : NULL;
+}
+
+void ULandscapeComponent::GetGeneratedTexturesAndMaterialInstances(TArray<UObject*>& OutTexturesAndMaterials) const
+{
+	for (UMaterialInstance* CurrentMIC = MaterialInstance; CurrentMIC && CurrentMIC->IsA<ULandscapeMaterialInstanceConstant>(); CurrentMIC = Cast<UMaterialInstance>(CurrentMIC->Parent))
+	{
+		OutTexturesAndMaterials.Add(CurrentMIC);
+	}
+
+	if (HeightmapTexture)
+	{
+		OutTexturesAndMaterials.Add(HeightmapTexture);
+	}
+
+	for (auto* Tex : WeightmapTextures)
+	{
+		OutTexturesAndMaterials.Add(Tex);
+	}
+
+	if (XYOffsetmapTexture)
+	{
+		OutTexturesAndMaterials.Add(XYOffsetmapTexture);
+	}
+}
+
+ALandscapeProxy* ULandscapeComponent::GetLandscapeProxy() const
+{
+	return CastChecked<ALandscapeProxy>(GetOuter());
+}
+
+FIntPoint ULandscapeComponent::GetSectionBase() const
+{
+	return FIntPoint(SectionBaseX, SectionBaseY);
+}
+
+void ULandscapeComponent::SetSectionBase(FIntPoint InSectionBase)
+{
+	SectionBaseX = InSectionBase.X;
+	SectionBaseY = InSectionBase.Y;
+}
+
+#if WITH_EDITOR
+ULandscapeInfo* ULandscapeComponent::GetLandscapeInfo(bool bSpawnNewActor /*= true*/) const
+{
+	if (GetLandscapeProxy())
+	{
+		return GetLandscapeProxy()->GetLandscapeInfo(bSpawnNewActor);
+	}
+	return NULL;
+}
+#endif
+
+void ULandscapeComponent::BeginDestroy()
+{
+	Super::BeginDestroy();
+
+#if WITH_EDITOR
+	if (EditToolRenderData != NULL)
+	{
+		// Ask render thread to destroy EditToolRenderData
+		EditToolRenderData->Cleanup();
+		EditToolRenderData = NULL;
+	}
+
+	if (GIsEditor && !HasAnyFlags(RF_ClassDefaultObject))
+	{
+		ALandscapeProxy* Proxy = GetLandscapeProxy();
+
+		// Remove any weightmap allocations from the Landscape Actor's map
+		for (int32 LayerIdx = 0; LayerIdx < WeightmapLayerAllocations.Num(); LayerIdx++)
+		{
+			int32 WeightmapIndex = WeightmapLayerAllocations[LayerIdx].WeightmapTextureIndex;
+			if (WeightmapTextures.IsValidIndex(WeightmapIndex))
+			{
+				UTexture2D* WeightmapTexture = WeightmapTextures[WeightmapIndex];
+				FLandscapeWeightmapUsage* Usage = Proxy->WeightmapUsageMap.Find(WeightmapTexture);
+				if (Usage != NULL)
+				{
+					Usage->ChannelUsage[WeightmapLayerAllocations[LayerIdx].WeightmapTextureChannel] = NULL;
+
+					if (Usage->FreeChannelCount() == 4)
+					{
+						Proxy->WeightmapUsageMap.Remove(WeightmapTexture);
+					}
+				}
+			}
+		}
+	}
+#endif
+}
+
+FPrimitiveSceneProxy* ULandscapeComponent::CreateSceneProxy()
+{
+	const auto FeatureLevel = GetWorld()->FeatureLevel;
+	FPrimitiveSceneProxy* Proxy = NULL;
+	if (FeatureLevel >= ERHIFeatureLevel::SM4)
+	{
+#if WITH_EDITOR
+		if (EditToolRenderData == NULL)
+		{
+			EditToolRenderData = new FLandscapeEditToolRenderData(this);
+		}
+		Proxy = new FLandscapeComponentSceneProxy(this, EditToolRenderData);
+#else
+		Proxy = new FLandscapeComponentSceneProxy(this, NULL);
+#endif
+	}
+	else
+	{
+#if WITH_EDITOR
+		if (!PlatformData.HasValidPlatformData()) // Deferred generation
+		{
+			// Try to reload the ES2 landscape data from the DDC
+			if (!PlatformData.LoadFromDDC(StateId))
+			{
+				// Generate and save to the derived data cache
+				GeneratePlatformVertexData();
+
+				if (PlatformData.HasValidPlatformData())
+				{
+					PlatformData.SaveToDDC(StateId);
+				}
+			}
+		}
+
+		if (PlatformData.HasValidPlatformData())
+		{
+			if (EditToolRenderData == NULL)
+			{
+				EditToolRenderData = new FLandscapeEditToolRenderData(this);
+			}
+
+			Proxy = new FLandscapeComponentSceneProxyMobile(this, EditToolRenderData);
+		}
+#else
+		if (PlatformData.HasValidPlatformData())
+		{
+			Proxy = new FLandscapeComponentSceneProxyMobile(this, NULL);
+		}
+#endif
+	}
+	return Proxy;
+}
+
+void ULandscapeComponent::DestroyComponent()
+{
+	ALandscapeProxy* Proxy = GetLandscapeProxy();
+	if (Proxy)
+	{
+		Proxy->LandscapeComponents.Remove(this);
+	}
+
+	Super::DestroyComponent();
+}
+
+FBoxSphereBounds ULandscapeComponent::CalcBounds(const FTransform& LocalToWorld) const
+{
+	return FBoxSphereBounds(CachedLocalBox.TransformBy(LocalToWorld));
+}
+
+void ULandscapeComponent::OnRegister()
+{
+	Super::OnRegister();
+
+#if WITH_EDITOR
+	ULandscapeInfo* Info = GetLandscapeInfo(false);
+	if (Info)
+	{
+		Info->RegisterActorComponent(this);
+	}
+#endif
+}
+
+void ULandscapeComponent::OnUnregister()
+{
+	Super::OnUnregister();
+
+#if WITH_EDITOR
+	ULandscapeInfo* Info = GetLandscapeInfo(false);
+	if (Info)
+	{
+		Info->UnregisterActorComponent(this);
+	}
+#endif
+}
+
+void ULandscapeComponent::GetUsedMaterials(TArray<UMaterialInterface*>& OutMaterials) const
+{
+	if (MaterialInstance != NULL)
+	{
+		OutMaterials.Add(MaterialInstance);
+	}
+}
+
+void ALandscapeProxy::RegisterAllComponents()
+{
+	Super::RegisterAllComponents();
+
+	// Landscape was added to world 
+	// We might need to update shared landscape data
+	if (GIsEditor && GetWorld() && !GetWorld()->IsPlayInEditor())
+	{
+		GEngine->DeferredCommands.AddUnique(TEXT("UpdateLandscapeEditorData"));
+	}
+}
+
+void ALandscapeProxy::UnregisterAllComponents()
+{
+	Super::UnregisterAllComponents();
+
+	// Landscape was removed from world 
+	// We might need to update shared landscape data
+	if (GEngine && GIsEditor && GetWorld() && !GetWorld()->IsPlayInEditor())
+	{
+		GEngine->DeferredCommands.AddUnique(TEXT("UpdateLandscapeEditorData"));
+	}
+}
+
+// FLandscapeWeightmapUsage serializer
+FArchive& operator<<(FArchive& Ar, FLandscapeWeightmapUsage& U)
+{
+	return Ar << U.ChannelUsage[0] << U.ChannelUsage[1] << U.ChannelUsage[2] << U.ChannelUsage[3];
+}
+
+FArchive& operator<<(FArchive& Ar, FLandscapeAddCollision& U)
+{
+#if WITH_EDITORONLY_DATA
+	return Ar << U.Corners[0] << U.Corners[1] << U.Corners[2] << U.Corners[3];
+#else
+	return Ar;
+#endif // WITH_EDITORONLY_DATA
+}
+
+FArchive& operator<<(FArchive& Ar, FLandscapeLayerStruct*& L)
+{
+	if (L)
+	{
+		Ar << L->LayerInfoObj;
+#if WITH_EDITORONLY_DATA
+		return Ar << L->ThumbnailMIC;
+#else
+		return Ar;
+#endif // WITH_EDITORONLY_DATA
+	}
+	return Ar;
+}
+
+void ULandscapeInfo::Serialize(FArchive& Ar)
+{
+	Super::Serialize(Ar);
+
+	// We do not serialize XYtoComponentMap as we don't want these references to hold a component in memory.
+	// The references are automatically cleaned up by the components' BeginDestroy method.
+	if (Ar.IsTransacting())
+	{
+		Ar << XYtoComponentMap;
+		Ar << XYtoAddCollisionMap;
+		Ar << SelectedComponents;
+		Ar << SelectedRegion;
+		Ar << SelectedRegionComponents;
+	}
+}
+
+void ALandscape::PostLoad()
+{
+	if (!LandscapeGuid.IsValid())
+	{
+		LandscapeGuid = FGuid::NewGuid();
+	}
+
+	Super::PostLoad();
+}
+
+void ULandscapeInfo::BeginDestroy()
+{
+	Super::BeginDestroy();
+
+#if WITH_EDITOR
+	if (DataInterface)
+	{
+		delete DataInterface;
+		DataInterface = NULL;
+	}
+#endif
+}
+
+#if WITH_EDITOR
+
+void ALandscape::CheckForErrors()
+{
+}
+
+#endif
+
+void ALandscapeProxy::Serialize(FArchive& Ar)
+{
+	Super::Serialize(Ar);
+
+#if WITH_EDITOR
+	if (Ar.IsTransacting())
+	{
+		Ar << WeightmapUsageMap;
+	}
+#endif
+}
+
+void ALandscapeProxy::AddReferencedObjects(UObject* InThis, FReferenceCollector& Collector)
+{
+	ALandscapeProxy* This = CastChecked<ALandscapeProxy>(InThis);
+
+	Super::AddReferencedObjects(InThis, Collector);
+
+	Collector.AddReferencedObjects(This->MaterialInstanceConstantMap, This);
+
+	for (auto It = This->WeightmapUsageMap.CreateIterator(); It; ++It)
+	{
+		Collector.AddReferencedObject(It.Key(), This);
+		Collector.AddReferencedObject(It.Value().ChannelUsage[0], This);
+		Collector.AddReferencedObject(It.Value().ChannelUsage[1], This);
+		Collector.AddReferencedObject(It.Value().ChannelUsage[2], This);
+		Collector.AddReferencedObject(It.Value().ChannelUsage[3], This);
+	}
+}
+
+#if WITH_EDITOR
+void ALandscapeProxy::PreEditUndo()
+{
+	Super::PreEditUndo();
+
+	// PostEditUndo doesn't get called when undoing a create action
+	// and PreEditUndo doesn't get called when undoing a delete action
+	// so this code needs to be in both!
+	if (GIsEditor && GetWorld() && !GetWorld()->IsPlayInEditor())
+	{
+		GEngine->DeferredCommands.AddUnique(TEXT("UpdateLandscapeEditorData"));
+
+		UNavigationSystem::ClearNavOctreeAll(this);
+	}
+}
+
+void ALandscapeProxy::PostEditUndo()
+{
+	Super::PostEditUndo();
+
+	if (GIsEditor && GetWorld() && !GetWorld()->IsPlayInEditor())
+	{
+		GEngine->DeferredCommands.AddUnique(TEXT("UpdateLandscapeEditorData"));
+
+		UNavigationSystem::UpdateNavOctreeAll(this);
+	}
+}
+
+void ULandscapeInfo::PostEditUndo()
+{
+	Super::PostEditUndo();
+
+	if (GIsEditor && GetWorld() && !GetWorld()->IsPlayInEditor())
+	{
+		GEngine->DeferredCommands.AddUnique(TEXT("UpdateLandscapeEditorData"));
+	}
+}
+
+FName FLandscapeInfoLayerSettings::GetLayerName() const
+{
+	checkSlow(LayerInfoObj == NULL || LayerInfoObj->LayerName == LayerName);
+
+	return LayerName;
+}
+
+FLandscapeEditorLayerSettings& FLandscapeInfoLayerSettings::GetEditorSettings() const
+{
+	check(LayerInfoObj);
+
+	ULandscapeInfo* LandscapeInfo = Owner->GetLandscapeInfo();
+	return LandscapeInfo->GetLayerEditorSettings(LayerInfoObj);
+}
+
+FLandscapeEditorLayerSettings& ULandscapeInfo::GetLayerEditorSettings(ULandscapeLayerInfoObject* LayerInfo) const
+{
+	auto& LayerSettingsList = GetLandscapeProxy()->EditorLayerSettings;
+	FLandscapeEditorLayerSettings* EditorLayerSettings = LayerSettingsList.FindByKey(LayerInfo);
+	if (EditorLayerSettings)
+	{
+		return *EditorLayerSettings;
+	}
+	else
+	{
+		int32 Index = LayerSettingsList.Add(LayerInfo);
+		return LayerSettingsList[Index];
+	}
+}
+
+void ULandscapeInfo::CreateLayerEditorSettingsFor(ULandscapeLayerInfoObject* LayerInfo)
+{
+	ALandscape* Landscape = LandscapeActor.Get();
+	if (Landscape != NULL)
+	{
+		FLandscapeEditorLayerSettings* EditorLayerSettings = Landscape->EditorLayerSettings.FindByKey(LayerInfo);
+		if (EditorLayerSettings == NULL)
+		{
+			Landscape->Modify();
+			Landscape->EditorLayerSettings.Add(LayerInfo);
+		}
+	}
+
+	for (auto It = Proxies.CreateConstIterator(); It; ++It)
+	{
+		ALandscapeProxy* Proxy = *It;
+		FLandscapeEditorLayerSettings* EditorLayerSettings = Proxy->EditorLayerSettings.FindByKey(LayerInfo);
+		if (EditorLayerSettings == NULL)
+		{
+			Proxy->Modify();
+			Proxy->EditorLayerSettings.Add(LayerInfo);
+		}
+	}
+}
+
+ULandscapeLayerInfoObject* ULandscapeInfo::GetLayerInfoByName(FName LayerName, ALandscapeProxy* Owner /*= NULL*/) const
+{
+	ULandscapeLayerInfoObject* LayerInfo = NULL;
+	for (int32 j = 0; j < Layers.Num(); j++)
+	{
+		if (Layers[j].LayerInfoObj && Layers[j].LayerInfoObj->LayerName == LayerName
+			&& (Owner == NULL || Layers[j].Owner == Owner))
+		{
+			LayerInfo = Layers[j].LayerInfoObj;
+		}
+	}
+	return LayerInfo;
+}
+
+int32 ULandscapeInfo::GetLayerInfoIndex(ULandscapeLayerInfoObject* LayerInfo, ALandscapeProxy* Owner /*= NULL*/) const
+{
+	for (int32 j = 0; j < Layers.Num(); j++)
+	{
+		if (Layers[j].LayerInfoObj && Layers[j].LayerInfoObj == LayerInfo
+			&& (Owner == NULL || Layers[j].Owner == Owner))
+		{
+			return j;
+		}
+	}
+
+	return INDEX_NONE;
+}
+
+int32 ULandscapeInfo::GetLayerInfoIndex(FName LayerName, ALandscapeProxy* Owner /*= NULL*/) const
+{
+	for (int32 j = 0; j < Layers.Num(); j++)
+	{
+		if (Layers[j].GetLayerName() == LayerName
+			&& (Owner == NULL || Layers[j].Owner == Owner))
+		{
+			return j;
+		}
+	}
+
+	return INDEX_NONE;
+}
+
+
+bool ULandscapeInfo::UpdateLayerInfoMap(ALandscapeProxy* Proxy /*= NULL*/, bool bInvalidate /*= false*/)
+{
+	bool bHasCollision = false;
+	if (GIsEditor)
+	{
+		bIsValid = !bInvalidate;
+
+		if (Proxy)
+		{
+			if (bInvalidate)
+			{
+				// this is a horribly dangerous combination of parameters...
+
+				for (int32 i = 0; i < Layers.Num(); i++)
+				{
+					if (Layers[i].Owner == Proxy)
+					{
+						Layers.RemoveAt(i--);
+					}
+				}
+			}
+			else // Proxy && !bInvalidate
+			{
+				TArray<FName> LayerNames = Proxy->GetLayersFromMaterial();
+
+				// Validate any existing layer infos owned by this proxy
+				for (int32 i = 0; i < Layers.Num(); i++)
+				{
+					if (Layers[i].Owner == Proxy)
+					{
+						Layers[i].bValid = LayerNames.Contains(Layers[i].GetLayerName());
+					}
+				}
+
+				// Add placeholders for any unused material layers
+				for (int32 i = 0; i < LayerNames.Num(); i++)
+				{
+					int32 LayerInfoIndex = GetLayerInfoIndex(LayerNames[i]);
+					if (LayerInfoIndex == INDEX_NONE)
+					{
+						FLandscapeInfoLayerSettings LayerSettings(LayerNames[i], Proxy);
+						LayerSettings.bValid = true;
+						Layers.Add(LayerSettings);
+					}
+				}
+
+				// Populate from layers used in components
+				for (int32 ComponentIndex = 0; ComponentIndex < Proxy->LandscapeComponents.Num(); ComponentIndex++)
+				{
+					ULandscapeComponent* Component = Proxy->LandscapeComponents[ComponentIndex];
+
+					// Add layers from per-component override materials
+					if (Component->OverrideMaterial != NULL)
+					{
+						TArray<FName> ComponentLayerNames = Proxy->GetLayersFromMaterial(Component->OverrideMaterial);
+						for (int32 i = 0; i < ComponentLayerNames.Num(); i++)
+						{
+							int32 LayerInfoIndex = GetLayerInfoIndex(ComponentLayerNames[i]);
+							if (LayerInfoIndex == INDEX_NONE)
+							{
+								FLandscapeInfoLayerSettings LayerSettings(ComponentLayerNames[i], Proxy);
+								LayerSettings.bValid = true;
+								Layers.Add(LayerSettings);
+							}
+						}
+					}
+
+					for (int32 AllocationIndex = 0; AllocationIndex < Component->WeightmapLayerAllocations.Num(); AllocationIndex++)
+					{
+						ULandscapeLayerInfoObject* LayerInfo = Component->WeightmapLayerAllocations[AllocationIndex].LayerInfo;
+						if (LayerInfo)
+						{
+							int32 LayerInfoIndex = GetLayerInfoIndex(LayerInfo);
+							bool bValid = LayerNames.Contains(LayerInfo->LayerName);
+
+							if (LayerInfoIndex != INDEX_NONE)
+							{
+								FLandscapeInfoLayerSettings& LayerSettings = Layers[LayerInfoIndex];
+
+								// Valid layer infos take precedence over invalid ones
+								// Landscape Actors take precedence over Proxies
+								if ((bValid && !LayerSettings.bValid)
+									|| (bValid == LayerSettings.bValid && !Proxy->bIsProxy))
+								{
+									LayerSettings.Owner = Proxy;
+									LayerSettings.bValid = bValid;
+									LayerSettings.ThumbnailMIC = NULL;
+								}
+							}
+							else
+							{
+								// handle existing placeholder layers
+								LayerInfoIndex = GetLayerInfoIndex(LayerInfo->LayerName);
+								if (LayerInfoIndex != INDEX_NONE)
+								{
+									FLandscapeInfoLayerSettings& LayerSettings = Layers[LayerInfoIndex];
+
+									//if (LayerSettings.Owner == Proxy)
+									{
+										LayerSettings.Owner = Proxy;
+										LayerSettings.LayerInfoObj = LayerInfo;
+										LayerSettings.bValid = bValid;
+										LayerSettings.ThumbnailMIC = NULL;
+									}
+								}
+								else
+								{
+									FLandscapeInfoLayerSettings LayerSettings(LayerInfo, Proxy);
+									LayerSettings.bValid = bValid;
+									Layers.Add(LayerSettings);
+								}
+							}
+						}
+					}
+				}
+
+				// Add any layer infos cached in the actor
+				Proxy->EditorLayerSettings.Remove(NULL);
+				for (int32 i = 0; i < Proxy->EditorLayerSettings.Num(); i++)
+				{
+					FLandscapeEditorLayerSettings& LayerInfo = Proxy->EditorLayerSettings[i];
+					if (LayerNames.Contains(LayerInfo.LayerInfoObj->LayerName))
+					{
+						// intentionally using the layer name here so we don't add layer infos from
+						// the cache that have the same name as an actual assignment from a component above
+						int32 LayerInfoIndex = GetLayerInfoIndex(LayerInfo.LayerInfoObj->LayerName);
+						if (LayerInfoIndex != INDEX_NONE)
+						{
+							FLandscapeInfoLayerSettings& LayerSettings = Layers[LayerInfoIndex];
+							if (LayerSettings.LayerInfoObj == NULL)
+							{
+								LayerSettings.Owner = Proxy;
+								LayerSettings.LayerInfoObj = LayerInfo.LayerInfoObj;
+								LayerSettings.bValid = true;
+							}
+						}
+					}
+					else
+					{
+						Proxy->Modify();
+						Proxy->EditorLayerSettings.RemoveAt(i--);
+					}
+				}
+			}
+		}
+		else // !Proxy
+		{
+			Layers.Empty();
+
+			if (!bInvalidate)
+			{
+				if (LandscapeActor.IsValid() && !LandscapeActor->IsPendingKillPending())
+				{
+					checkSlow(LandscapeActor->GetLandscapeInfo() == this);
+					UpdateLayerInfoMap(LandscapeActor.Get(), false);
+				}
+				for (auto It = Proxies.CreateIterator(); It; ++It)
+				{
+					ALandscapeProxy* LandscapeProxy = *It;
+					checkSlow(LandscapeProxy && LandscapeProxy->GetLandscapeInfo() == this);
+					UpdateLayerInfoMap(LandscapeProxy, false);
+				}
+			}
+		}
+
+		//if (GCallbackEvent)
+		//{
+		//	GCallbackEvent->Send( CALLBACK_EditorPostModal );
+		//}
+	}
+	return bHasCollision;
+}
+#endif
+
+void ALandscapeProxy::PostLoad()
+{
+	Super::PostLoad();
+
+	// Temporary
+	if (ComponentSizeQuads == 0 && LandscapeComponents.Num() > 0)
+	{
+		ULandscapeComponent* Comp = LandscapeComponents[0];
+		if (Comp)
+		{
+			ComponentSizeQuads = Comp->ComponentSizeQuads;
+			SubsectionSizeQuads = Comp->SubsectionSizeQuads;
+			NumSubsections = Comp->NumSubsections;
+		}
+	}
+
+	if (IsTemplate() == false)
+	{
+		BodyInstance.FixupData(this);
+	}
+
+#if WITH_EDITOR
+	if (GetLinker() && (GetLinker()->UE4Ver() < VER_UE4_LANDSCAPE_COMPONENT_LAZY_REFERENCES) || LandscapeComponents.Num() != CollisionComponents.Num())
+	{
+		// Need to clean up invalid collision components
+		RecreateCollisionComponents();
+	}
+
+	EditorLayerSettings.Remove(NULL);
+
+	if (EditorCachedLayerInfos_DEPRECATED.Num() > 0)
+	{
+		for (int32 i = 0; i < EditorCachedLayerInfos_DEPRECATED.Num(); i++)
+		{
+			EditorLayerSettings.Add(EditorCachedLayerInfos_DEPRECATED[i]);
+		}
+		EditorCachedLayerInfos_DEPRECATED.Empty();
+	}
+
+	if (GIsEditor && !GetWorld()->IsPlayInEditor())
+	{
+		// defer LandscapeInfo setup
+		GEngine->DeferredCommands.AddUnique(TEXT("UpdateLandscapeEditorData -warnings"));
+	}
+#endif
+}
+
+#if WITH_EDITOR
+void ALandscapeProxy::Destroyed()
+{
+	Super::Destroyed();
+
+	if (GIsEditor)
+	{
+		ULandscapeInfo::RecreateLandscapeInfo(GetWorld(), false);
+	}
+
+	if (SplineComponent)
+	{
+		SplineComponent->ModifySplines();
+	}
+}
+#endif
+
+#if WITH_EDITOR
+void ALandscapeProxy::GetSharedProperties(ALandscapeProxy* Landscape)
+{
+	if (GIsEditor && Landscape)
+	{
+		Modify();
+
+		LandscapeGuid = Landscape->LandscapeGuid;
+
+		//@todo UE4 merge, landscape, this needs work
+		RootComponent->SetRelativeScale3D(Landscape->GetRootComponent()->GetComponentToWorld().GetScale3D());
+
+		//PrePivot = Landscape->PrePivot;
+		StaticLightingResolution = Landscape->StaticLightingResolution;
+		bCastStaticShadow = Landscape->bCastStaticShadow;
+		bCastShadowAsTwoSided = Landscape->bCastShadowAsTwoSided;
+		ComponentSizeQuads = Landscape->ComponentSizeQuads;
+		NumSubsections = Landscape->NumSubsections;
+		SubsectionSizeQuads = Landscape->SubsectionSizeQuads;
+		MaxLODLevel = Landscape->MaxLODLevel;
+		LODDistanceFactor = Landscape->LODDistanceFactor;
+		LODFalloff = Landscape->LODFalloff;
+		CollisionMipLevel = Landscape->CollisionMipLevel;
+		if (!LandscapeMaterial)
+		{
+			LandscapeMaterial = Landscape->LandscapeMaterial;
+		}
+		if (!LandscapeHoleMaterial)
+		{
+			LandscapeHoleMaterial = Landscape->LandscapeHoleMaterial;
+		}
+		if (LandscapeMaterial == Landscape->LandscapeMaterial)
+		{
+			EditorLayerSettings = Landscape->EditorLayerSettings;
+		}
+		if (!DefaultPhysMaterial)
+		{
+			DefaultPhysMaterial = Landscape->DefaultPhysMaterial;
+		}
+		LightmassSettings = Landscape->LightmassSettings;
+	}
+}
+
+FTransform ALandscapeProxy::LandscapeActorToWorld() const
+{
+	FTransform TM = ActorToWorld();
+	// Add this proxy landscape section offset to obtain landscape actor transform
+	TM.AddToTranslation(TM.TransformVector(-FVector(LandscapeSectionOffset)));
+	return TM;
+}
+
+void ALandscapeProxy::SetAbsoluteSectionBase(FIntPoint InSectionBase)
+{
+	FIntPoint Difference = InSectionBase - LandscapeSectionOffset;
+	LandscapeSectionOffset = InSectionBase;
+
+	for (int32 CompIdx = 0; CompIdx < LandscapeComponents.Num(); CompIdx++)
+	{
+		ULandscapeComponent* Comp = LandscapeComponents[CompIdx];
+		if (Comp)
+		{
+			FIntPoint AbsoluteSectionBase = Comp->GetSectionBase() + Difference;
+			Comp->SetSectionBase(AbsoluteSectionBase);
+			Comp->RecreateRenderState_Concurrent();
+		}
+	}
+
+	for (int32 CompIdx = 0; CompIdx < CollisionComponents.Num(); CompIdx++)
+	{
+		ULandscapeHeightfieldCollisionComponent* Comp = CollisionComponents[CompIdx];
+		if (Comp)
+		{
+			FIntPoint AbsoluteSectionBase = Comp->GetSectionBase() + Difference;
+			Comp->SetSectionBase(AbsoluteSectionBase);
+		}
+	}
+}
+
+FIntPoint ALandscapeProxy::GetSectionBaseOffset() const
+{
+	return LandscapeSectionOffset;
+}
+
+void ALandscapeProxy::RecreateComponentsState()
+{
+	for (int32 ComponentIndex = 0; ComponentIndex < LandscapeComponents.Num(); ComponentIndex++)
+	{
+		ULandscapeComponent* Comp = LandscapeComponents[ComponentIndex];
+		if (Comp)
+		{
+			Comp->UpdateComponentToWorld();
+			Comp->UpdateCachedBounds();
+			Comp->UpdateBounds();
+			Comp->RecreateRenderState_Concurrent(); // @todo UE4 jackp just render state needs update?
+		}
+	}
+
+	for (int32 ComponentIndex = 0; ComponentIndex < CollisionComponents.Num(); ComponentIndex++)
+	{
+		ULandscapeHeightfieldCollisionComponent* Comp = CollisionComponents[ComponentIndex];
+		if (Comp)
+		{
+			Comp->UpdateComponentToWorld();
+			Comp->RecreateCollision();
+		}
+	}
+}
+
+UMaterialInterface* ALandscapeProxy::GetLandscapeMaterial() const
+{
+	if (LandscapeMaterial)
+	{
+		return LandscapeMaterial;
+	}
+	else if (LandscapeActor)
+	{
+		return LandscapeActor->GetLandscapeMaterial();
+	}
+	return UMaterial::GetDefaultMaterial(MD_Surface);
+}
+
+UMaterialInterface* ALandscapeProxy::GetLandscapeHoleMaterial() const
+{
+	if (LandscapeHoleMaterial)
+	{
+		return LandscapeHoleMaterial;
+	}
+	else if (LandscapeActor)
+	{
+		return LandscapeActor->GetLandscapeHoleMaterial();
+	}
+	return UMaterial::GetDefaultMaterial(MD_Surface);
+}
+
+UMaterialInterface* ALandscape::GetLandscapeMaterial() const
+{
+	if (LandscapeMaterial)
+	{
+		return LandscapeMaterial;
+	}
+	return UMaterial::GetDefaultMaterial(MD_Surface);
+}
+
+UMaterialInterface* ALandscape::GetLandscapeHoleMaterial() const
+{
+	if (LandscapeHoleMaterial)
+	{
+		return LandscapeHoleMaterial;
+	}
+	return UMaterial::GetDefaultMaterial(MD_Surface);
+}
+
+void ALandscape::PreSave()
+{
+	Super::PreSave();
+	//ULandscapeInfo* Info = GetLandscapeInfo();
+	//if (GIsEditor && Info && !IsRunningCommandlet())
+	//{
+	//	for (TSet<ALandscapeProxy*>::TIterator It(Info->Proxies); It; ++It)
+	//	{
+	//		ALandscapeProxy* Proxy = *It;
+	//		if (!ensure(Proxy->LandscapeActor == this))
+	//		{
+	//			Proxy->LandscapeActor = this;
+	//			Proxy->GetSharedProperties(this);
+	//		}
+	//	}
+	//}
+}
+
+ALandscapeProxy* ULandscapeInfo::GetCurrentLevelLandscapeProxy(bool bRegistered) const
+{
+	ALandscape* Landscape = LandscapeActor.Get();
+
+	if (Landscape && (!bRegistered || Landscape->GetRootComponent()->IsRegistered()))
+	{
+		UWorld* LandscapeWorld = Landscape->GetWorld();
+		if (LandscapeWorld &&
+			LandscapeWorld->GetCurrentLevel() == Landscape->GetOuter())
+		{
+			return Landscape;
+		}
+	}
+
+	for (auto It = Proxies.CreateConstIterator(); It; ++It)
+	{
+		ALandscapeProxy* LandscapeProxy = (*It);
+
+		if (LandscapeProxy && (!bRegistered || LandscapeProxy->GetRootComponent()->IsRegistered()))
+		{
+			UWorld* ProxyWorld = LandscapeProxy->GetWorld();
+			if (ProxyWorld &&
+				ProxyWorld->GetCurrentLevel() == LandscapeProxy->GetOuter())
+			{
+				return LandscapeProxy;
+			}
+		}
+	}
+
+	return NULL;
+}
+
+ALandscapeProxy* ULandscapeInfo::GetLandscapeProxy() const
+{
+	// Mostly this Proxy used to calculate transformations
+	// in Editor all proxies of same landscape actor have root components in same locations
+	// so it doesn't really matter which proxy we return here
+
+	// prefer LandscapeActor in case it is loaded
+	ALandscape* Landscape = LandscapeActor.Get();
+	if (Landscape != NULL &&
+		Landscape->GetRootComponent()->IsRegistered())
+	{
+		return Landscape;
+	}
+
+	// prefer current level proxy 
+	ALandscapeProxy* Proxy = GetCurrentLevelLandscapeProxy(true);
+	if (Proxy != NULL)
+	{
+		return Proxy;
+	}
+
+	// any proxy in the world
+	for (auto It = Proxies.CreateConstIterator(); It; ++It)
+	{
+		Proxy = (*It);
+		if (Proxy != NULL &&
+			Proxy->GetRootComponent()->IsRegistered())
+		{
+			return Proxy;
+		}
+	}
+
+	return NULL;
+}
+
+void ULandscapeInfo::RegisterActor(ALandscapeProxy* Proxy, bool bMapCheck)
+{
+	// do not pass here invalid actors
+	check(Proxy->GetLandscapeGuid().IsValid());
+	UWorld* OwningWorld = CastChecked<UWorld>(GetOuter());
+
+	// in case this Info object is not initialized yet
+	// initialized it with properties from passed actor
+	if (LandscapeGuid.IsValid() == false ||
+		(GetLandscapeProxy() == NULL && ensure(LandscapeGuid == Proxy->GetLandscapeGuid())))
+	{
+		LandscapeGuid = Proxy->GetLandscapeGuid();
+		ComponentSizeQuads = Proxy->ComponentSizeQuads;
+		ComponentNumSubsections = Proxy->NumSubsections;
+		SubsectionSizeQuads = Proxy->SubsectionSizeQuads;
+		DrawScale = Proxy->GetRootComponent()->RelativeScale3D;
+	}
+
+	// check that passed actor matches all shared parameters
+	check(LandscapeGuid == Proxy->GetLandscapeGuid());
+	check(ComponentSizeQuads == Proxy->ComponentSizeQuads);
+	check(ComponentNumSubsections == Proxy->NumSubsections);
+	check(SubsectionSizeQuads == Proxy->SubsectionSizeQuads);
+	check(DrawScale == Proxy->GetRootComponent()->RelativeScale3D);
+
+	// register
+	ALandscape* Landscape = Cast<ALandscape>(Proxy);
+	if (Landscape)
+	{
+		LandscapeActor = Landscape;
+		// In world composition user is not allowed to move landscape in editor, only through WorldBrowser 
+		LandscapeActor->bLockLocation = (OwningWorld->WorldComposition != NULL);
+
+		// update proxies reference actor
+		for (auto It = Proxies.CreateConstIterator(); It; ++It)
+		{
+			(*It)->LandscapeActor = LandscapeActor;
+		}
+	}
+	else
+	{
+		Proxies.Add(Proxy);
+		Proxy->LandscapeActor = LandscapeActor;
+	}
+
+	//
+	// add proxy components to the XY map
+	//
+	for (int32 CompIdx = 0; CompIdx < Proxy->LandscapeComponents.Num(); ++CompIdx)
+	{
+		RegisterActorComponent(Proxy->LandscapeComponents[CompIdx], bMapCheck);
+	}
+}
+
+void ULandscapeInfo::UnregisterActor(ALandscapeProxy* Proxy)
+{
+	ALandscape* Landscape = Cast<ALandscape>(Proxy);
+	if (Landscape)
+	{
+		check(LandscapeActor.Get() == Landscape);
+		LandscapeActor = 0;
+
+		// update proxies reference to landscape actor
+		for (auto It = Proxies.CreateConstIterator(); It; ++It)
+		{
+			(*It)->LandscapeActor = 0;
+		}
+	}
+	else
+	{
+		Proxies.Remove(Proxy);
+		Proxy->LandscapeActor = 0;
+	}
+
+	// remove proxy components from the XY map
+	for (int32 CompIdx = 0; CompIdx < Proxy->LandscapeComponents.Num(); ++CompIdx)
+	{
+		UnregisterActorComponent(Proxy->LandscapeComponents[CompIdx]);
+	}
+
+	XYtoComponentMap.Compact();
+}
+
+void ULandscapeInfo::RegisterActorComponent(ULandscapeComponent* Component, bool bMapCheck)
+{
+	// Do not register components which are not part of the world
+	if (Component == NULL ||
+		Component->IsRegistered() == false)
+	{
+		return;
+	}
+
+	check(Component != NULL);
+	FIntPoint ComponentKey = Component->GetSectionBase() / Component->ComponentSizeQuads;
+	auto RegisteredComponent = XYtoComponentMap.FindRef(ComponentKey);
+
+	if (RegisteredComponent == NULL)
+	{
+		XYtoComponentMap.Add(ComponentKey, Component);
+	}
+	else if (bMapCheck)
+	{
+		ALandscapeProxy* Proxy = Component->GetLandscapeProxy();
+		FFormatNamedArguments Arguments;
+		Arguments.Add(TEXT("ProxyName"), FText::FromString(Proxy->GetName()));
+		Arguments.Add(TEXT("XLocation"), Component->GetSectionBase().X);
+		Arguments.Add(TEXT("YLocation"), Component->GetSectionBase().Y);
+		FMessageLog("MapCheck").Warning()
+			->AddToken(FUObjectToken::Create(Proxy))
+			->AddToken(FTextToken::Create(FText::Format(LOCTEXT("MapCheck_Message_LandscapeComponentPostLoad_Warning", "Landscape ({ProxyName}) has overlapping render components at location ({XLocation}, {YLocation})."), Arguments)))
+			->AddToken(FMapErrorToken::Create(FMapErrors::LandscapeComponentPostLoad_Warning));
+
+		// Show MapCheck window
+		FMessageLog("MapCheck").Open(EMessageSeverity::Warning);
+	}
+
+	// Update Selected Components/Regions
+#if WITH_EDITOR
+	if (Component->EditToolRenderData != NULL && Component->EditToolRenderData->SelectedType)
+	{
+		if (Component->EditToolRenderData->SelectedType & FLandscapeEditToolRenderData::ST_COMPONENT)
+		{
+			SelectedComponents.Add(Component);
+		}
+		else if (Component->EditToolRenderData->SelectedType & FLandscapeEditToolRenderData::ST_REGION)
+		{
+			SelectedRegionComponents.Add(Component);
+		}
+	}
+#endif
+}
+
+void ULandscapeInfo::UnregisterActorComponent(ULandscapeComponent* Component)
+{
+	check(Component != NULL);
+
+	FIntPoint ComponentKey = Component->GetSectionBase() / Component->ComponentSizeQuads;
+	auto RegisteredComponent = XYtoComponentMap.FindRef(ComponentKey);
+
+	if (RegisteredComponent == Component)
+	{
+		XYtoComponentMap.Remove(ComponentKey);
+	}
+
+	SelectedComponents.Remove(Component);
+	SelectedRegionComponents.Remove(Component);
+}
+
+void ULandscapeInfo::Reset()
+{
+	LandscapeActor.Reset();
+
+	Proxies.Empty();
+	XYtoComponentMap.Empty();
+	XYtoAddCollisionMap.Empty();
+
+	//SelectedComponents.Empty();
+	//SelectedRegionComponents.Empty();
+	//SelectedRegion.Empty();
+}
+
+void ULandscapeInfo::FixupProxiesTransform()
+{
+	UWorld* OwningWorld = CastChecked<UWorld>(GetOuter());
+	if (OwningWorld->WorldComposition != NULL)
+	{
+		return;
+	}
+
+	ALandscape* Landscape = LandscapeActor.Get();
+
+	if (Landscape == NULL ||
+		Landscape->GetRootComponent()->IsRegistered() == false)
+	{
+		return;
+	}
+
+	FTransform LandscapeTM = Landscape->LandscapeActorToWorld();
+	// Update transformations of all linked landscape proxies
+	for (auto It = Proxies.CreateConstIterator(); It; ++It)
+	{
+		ALandscapeProxy* Proxy = *It;
+		FTransform ProxyRelativeTM(FVector(Proxy->LandscapeSectionOffset));
+		FTransform ProxyTransform = ProxyRelativeTM*LandscapeTM;
+
+		if (!Proxy->GetTransform().Equals(ProxyTransform))
+		{
+			Proxy->SetActorTransform(ProxyTransform);
+			Proxy->RecreateComponentsState();
+
+			// Let other systems know that an actor was moved
+			GEngine->BroadcastOnActorMoved(Proxy);
+		}
+	}
+}
+
+void ULandscapeInfo::FixupProxiesWeightmaps()
+{
+	if (LandscapeActor.IsValid())
+	{
+		LandscapeActor->WeightmapUsageMap.Empty();
+		for (int32 CompIdx = 0; CompIdx < LandscapeActor->LandscapeComponents.Num(); ++CompIdx)
+		{
+			ULandscapeComponent* Comp = LandscapeActor->LandscapeComponents[CompIdx];
+			if (Comp)
+			{
+				Comp->FixupWeightmaps();
+			}
+		}
+	}
+
+	for (auto It = Proxies.CreateConstIterator(); It; ++It)
+	{
+		ALandscapeProxy* Proxy = (*It);
+		Proxy->WeightmapUsageMap.Empty();
+		for (int32 CompIdx = 0; CompIdx < Proxy->LandscapeComponents.Num(); ++CompIdx)
+		{
+			ULandscapeComponent* Comp = Proxy->LandscapeComponents[CompIdx];
+			if (Comp)
+			{
+				Comp->FixupWeightmaps();
+			}
+		}
+	}
+}
+
+void ULandscapeInfo::RecreateLandscapeInfo(UWorld* InWorld, bool bMapCheck)
+{
+	check(InWorld);
+
+	auto& LandscapeInfoMap = GetLandscapeInfoMap(InWorld);
+
+	// reset all LandscapeInfo objects
+	for (auto It = LandscapeInfoMap.Map.CreateIterator(); It; ++It)
+	{
+		It.Value()->Modify();
+		It.Value()->Reset();
+	}
+
+	TMap<FGuid, TArray<ALandscapeProxy*>> ValidLandscapesMap;
+	// Gather all valid landscapes in the world
+	for (TActorIterator<ALandscapeProxy> It(InWorld); It; ++It)
+	{
+		ALandscapeProxy* Proxy = *It;
+		if (Proxy
+			&& Proxy->HasAnyFlags(RF_BeginDestroyed) == false
+			&& Proxy->IsPendingKill() == false
+			&& Proxy->IsPendingKillPending() == false)
+		{
+			ValidLandscapesMap.FindOrAdd(Proxy->GetLandscapeGuid()).Add(Proxy);
+		}
+	}
+
+	// Register gathered landscapes in shared LandscapeInfo map
+	for (auto It = ValidLandscapesMap.CreateIterator(); It; ++It)
+	{
+		auto& LandscapeList = It.Value();
+		ALandscapeProxy* FirstVisibleLandscape = NULL;
+
+		for (int32 LandscapeIdx = 0; LandscapeIdx < LandscapeList.Num(); LandscapeIdx++)
+		{
+			ALandscapeProxy* Proxy = LandscapeList[LandscapeIdx];
+			if (InWorld->WorldComposition &&
+				Proxy->GetRootComponent()->IsRegistered())
+			{
+				if (FirstVisibleLandscape == NULL)
+				{
+					FirstVisibleLandscape = Proxy;
+				}
+
+				// Setup first landscape in the list as origin for section based grid
+				// And based on that calculate section coordinates for all landscapes in the list
+				FVector Offset = Proxy->GetActorLocation() - FirstVisibleLandscape->GetActorLocation();
+				FVector	DrawScale = FirstVisibleLandscape->GetRootComponent()->RelativeScale3D;
+
+				FIntPoint QuadsSpaceOffset;
+				QuadsSpaceOffset.X = FMath::RoundToInt(Offset.X / DrawScale.X);
+				QuadsSpaceOffset.Y = FMath::RoundToInt(Offset.Y / DrawScale.Y);
+				Proxy->SetAbsoluteSectionBase(QuadsSpaceOffset);
+			}
+
+			// Register gathered landscape actors inside global LandscapeInfo object
+			Proxy->GetLandscapeInfo(true)->RegisterActor(Proxy, bMapCheck);
+		}
+	}
+
+	// Remove empty entries from global LandscapeInfo map
+	for (auto It = LandscapeInfoMap.Map.CreateIterator(); It; ++It)
+	{
+		if (It.Value()->GetLandscapeProxy() == NULL)
+		{
+			It.Value()->MarkPendingKill();
+			It.RemoveCurrent();
+		}
+	}
+
+	// Update layer info maps
+	for (auto It = LandscapeInfoMap.Map.CreateIterator(); It; ++It)
+	{
+		ULandscapeInfo* LandscapeInfo = It.Value();
+		if (LandscapeInfo)
+		{
+			LandscapeInfo->UpdateLayerInfoMap();
+		}
+	}
+
+	// Update add collision
+	for (auto It = LandscapeInfoMap.Map.CreateIterator(); It; ++It)
+	{
+		ULandscapeInfo* LandscapeInfo = It.Value();
+		if (LandscapeInfo)
+		{
+			LandscapeInfo->UpdateAllAddCollisions();
+		}
+	}
+
+	// We need to inform Landscape editor tools about LandscapeInfo updates
+	FEditorSupportDelegates::WorldChange.Broadcast();
+}
+
+
+#endif
+
+void ULandscapeComponent::PostInitProperties()
+{
+	Super::PostInitProperties();
+
+	// Create a new guid in case this is a newly created component
+	// If not, this guid will be overwritten when serialized
+	FPlatformMisc::CreateGuid(StateId);
+}
+
+void ULandscapeComponent::PostDuplicate(bool bDuplicateForPIE)
+{
+	if (!bDuplicateForPIE)
+	{
+		// Reset the StateId on duplication since it needs to be unique for each capture.
+		// PostDuplicate covers direct calls to StaticDuplicateObject, but not actor duplication (see PostEditImport)
+		FPlatformMisc::CreateGuid(StateId);
+	}
+}
+
+// Generate a new guid to force a recache of all landscape derived data
+#define LANDSCAPE_FULL_DERIVEDDATA_VER			TEXT("28f0da4d085e48a680482c9835907e2f")
+
+FString FLandscapeComponentDerivedData::GetDDCKeyString(const FGuid& StateId)
+{
+	return FDerivedDataCacheInterface::BuildCacheKey(TEXT("LS_FULL"), LANDSCAPE_FULL_DERIVEDDATA_VER, *StateId.ToString());
+}
+
+void FLandscapeComponentDerivedData::InitializeFromUncompressedData(const TArray<uint8>& UncompressedData)
+{
+	int32 UncompressedSize = UncompressedData.Num() * UncompressedData.GetTypeSize();
+
+	TArray<uint8> TempCompressedMemory;
+	// Compressed can be slightly larger than uncompressed
+	TempCompressedMemory.Empty(UncompressedSize * 4 / 3);
+	TempCompressedMemory.AddUninitialized(UncompressedSize * 4 / 3);
+	int32 CompressedSize = TempCompressedMemory.Num() * TempCompressedMemory.GetTypeSize();
+
+	verify(FCompression::CompressMemory(
+		(ECompressionFlags)(COMPRESS_ZLIB | COMPRESS_BiasMemory),
+		TempCompressedMemory.GetData(),
+		CompressedSize,
+		UncompressedData.GetData(),
+		UncompressedSize));
+
+	// Note: change LANDSCAPE_FULL_DERIVEDDATA_VER when modifying the serialization layout
+	FMemoryWriter FinalArchive(CompressedLandscapeData, true);
+	FinalArchive << UncompressedSize;
+	FinalArchive << CompressedSize;
+	FinalArchive.Serialize(TempCompressedMemory.GetData(), CompressedSize);
+}
+
+void FLandscapeComponentDerivedData::GetUncompressedData(TArray<uint8>& OutUncompressedData)
+{
+	check(CompressedLandscapeData.Num() > 0);
+
+	FMemoryReader Ar(CompressedLandscapeData);
+
+	// Note: change LANDSCAPE_FULL_DERIVEDDATA_VER when modifying the serialization layout
+	int32 UncompressedSize;
+	Ar << UncompressedSize;
+
+	int32 CompressedSize;
+	Ar << CompressedSize;
+
+	TArray<uint8> CompressedData;
+	CompressedData.Empty(CompressedSize);
+	CompressedData.AddUninitialized(CompressedSize);
+	Ar.Serialize(CompressedData.GetData(), CompressedSize);
+
+	OutUncompressedData.Empty(UncompressedSize);
+	OutUncompressedData.AddUninitialized(UncompressedSize);
+
+	verify(FCompression::UncompressMemory((ECompressionFlags)COMPRESS_ZLIB, OutUncompressedData.GetData(), UncompressedSize, CompressedData.GetData(), CompressedSize));
+
+	// free the compressed data now that we have the uncompressed version
+	CompressedLandscapeData.Empty();
+}
+
+FArchive& operator<<(FArchive& Ar, FLandscapeComponentDerivedData& Data)
+{
+	check(!Ar.IsSaving() || Data.CompressedLandscapeData.Num() > 0);
+	return Ar << Data.CompressedLandscapeData;
+}
+
+bool FLandscapeComponentDerivedData::LoadFromDDC(const FGuid& StateId)
+{
+	return GetDerivedDataCacheRef().GetSynchronous(*GetDDCKeyString(StateId), CompressedLandscapeData);
+}
+
+void FLandscapeComponentDerivedData::SaveToDDC(const FGuid& StateId)
+{
+	check(CompressedLandscapeData.Num() > 0);
+	GetDerivedDataCacheRef().Put(*GetDDCKeyString(StateId), CompressedLandscapeData);
+}
+
+void WorldDestroyEventFunction(UWorld* World)
+{
+	World->PerModuleDataObjects.RemoveAll(
+		[](UObject* Object)
+		{
+			return Object->IsA(ULandscapeInfoMap::StaticClass());
+		}
+	);
+}
+
+#if WITH_EDITORONLY_DATA
+
+ULandscapeInfoMap& GetLandscapeInfoMap(UWorld* World)
+{
+	ULandscapeInfoMap *FoundObject = nullptr;
+
+	World->PerModuleDataObjects.FindItemByClass(&FoundObject);
+
+	checkf(FoundObject, TEXT("ULandscapInfoMap object was not created for this UWorld."));
+
+	return *FoundObject;
+}
+
+#endif // WITH_EDITORONLY_DATA
+
+void LandscapeMaterialsParameterValuesGetter(FStaticParameterSet &OutStaticParameterSet, UMaterialInstance* Material)
+{
+	if (Material->Parent)
+	{
+		UMaterial *ParentMaterial = Material->Parent->GetMaterial();
+
+		TArray<FName> ParameterNames;
+		TArray<FGuid> Guids;
+		ParentMaterial->GetAllParameterNames<UMaterialExpressionLandscapeLayerWeight>(ParameterNames, Guids);
+		ParentMaterial->GetAllParameterNames<UMaterialExpressionLandscapeLayerSwitch>(ParameterNames, Guids);
+		ParentMaterial->GetAllParameterNames<UMaterialExpressionLandscapeLayerSample>(ParameterNames, Guids);
+		ParentMaterial->GetAllParameterNames<UMaterialExpressionLandscapeLayerBlend>(ParameterNames, Guids);
+		ParentMaterial->GetAllParameterNames<UMaterialExpressionLandscapeVisibilityMask>(ParameterNames, Guids);
+
+		OutStaticParameterSet.TerrainLayerWeightParameters.AddZeroed(ParameterNames.Num());
+		for (int32 ParameterIdx = 0; ParameterIdx < ParameterNames.Num(); ParameterIdx++)
+		{
+			FStaticTerrainLayerWeightParameter& ParentParameter = OutStaticParameterSet.TerrainLayerWeightParameters[ParameterIdx];
+			FName ParameterName = ParameterNames[ParameterIdx];
+			FGuid ExpressionId = Guids[ParameterIdx];
+			int32 WeightmapIndex = INDEX_NONE;
+
+			ParentParameter.bOverride = false;
+			ParentParameter.ParameterName = ParameterName;
+			//get the settings from the parent in the MIC chain
+			if (Material->Parent->GetTerrainLayerWeightParameterValue(ParameterName, WeightmapIndex, ExpressionId))
+			{
+				ParentParameter.WeightmapIndex = WeightmapIndex;
+			}
+			ParentParameter.ExpressionGUID = ExpressionId;
+
+			// if the SourceInstance is overriding this parameter, use its settings
+			for (int32 WeightParamIdx = 0; WeightParamIdx < Material->GetStaticParameters().TerrainLayerWeightParameters.Num(); WeightParamIdx++)
+			{
+				const FStaticTerrainLayerWeightParameter &TerrainLayerWeightParam = Material->GetStaticParameters().TerrainLayerWeightParameters[WeightParamIdx];
+
+				if (ParameterName == TerrainLayerWeightParam.ParameterName)
+				{
+					ParentParameter.bOverride = TerrainLayerWeightParam.bOverride;
+					if (TerrainLayerWeightParam.bOverride)
+					{
+						ParentParameter.WeightmapIndex = TerrainLayerWeightParam.WeightmapIndex;
+					}
+				}
+			}
+		}
+	}
+}
+
+bool LandscapeMaterialsParameterSetUpdater(FStaticParameterSet &StaticParameterSet, UMaterial* ParentMaterial)
+{
+	return UpdateParameterSet<FStaticTerrainLayerWeightParameter, UMaterialExpressionLandscapeLayerWeight>(StaticParameterSet.TerrainLayerWeightParameters, ParentMaterial);
+}
+
+#undef LOCTEXT_NAMESPACE
\ No newline at end of file
