Landscape module from Engine split out.

diff --git a/Engine/Source/Runtime/Landscape/Private/LandscapeCollision.cpp b/Engine/Source/Runtime/Landscape/Private/LandscapeCollision.cpp
new file mode 100644
--- /dev/null
+++ b/Engine/Source/Runtime/Landscape/Private/LandscapeCollision.cpp
@@ -0,0 +1,1708 @@
+// Copyright 1998-2014 Epic Games, Inc. All Rights Reserved.
+
+/*=============================================================================
+LandscapeCollision.cpp: Landscape collision
+=============================================================================*/
+
+#include "Landscape.h"
+#include "PhysicsPublic.h"
+#include "LandscapeDataAccess.h"
+#include "LandscapeRender.h"
+#include "Collision/PhysXCollision.h"
+#include "DerivedDataPluginInterface.h"
+#include "DerivedDataCacheInterface.h"
+#include "PhysicsEngine/PhysXSupport.h"
+#include "PhysicsEngine/PhysDerivedData.h"
+#include "PhysicalMaterials/PhysicalMaterial.h"
+#include "LandscapeHeightfieldCollisionComponent.h"
+#include "LandscapeMeshCollisionComponent.h"
+#include "LandscapeLayerInfoObject.h"
+#include "LandscapeInfo.h"
+#include "InstancedFoliage.h"
+#include "Foliage/FoliageType.h"
+#include "Engine/StaticMesh.h"
+#include "Interfaces/Interface_CollisionDataProvider.h"
+#include "AI/NavigationSystemHelpers.h"
+
+TMap<FGuid, ULandscapeHeightfieldCollisionComponent::FPhysXHeightfieldRef* > GSharedHeightfieldRefs;
+
+ULandscapeHeightfieldCollisionComponent::FPhysXHeightfieldRef::~FPhysXHeightfieldRef()
+{
+#if WITH_PHYSX
+	// Free the existing heightfield data.
+	if (RBHeightfield)
+	{
+		GPhysXPendingKillHeightfield.Add(RBHeightfield);
+		RBHeightfield = NULL;
+	}
+#if WITH_EDITOR
+	if (RBHeightfieldEd)
+	{
+		GPhysXPendingKillHeightfield.Add(RBHeightfieldEd);
+		RBHeightfieldEd = NULL;
+	}
+#endif// WITH_EDITOR
+#endif// WITH_PHYSX
+
+	// Remove ourselves from the shared map.
+	GSharedHeightfieldRefs.Remove(Guid);
+}
+
+TMap<FGuid, ULandscapeMeshCollisionComponent::FPhysXMeshRef* > GSharedMeshRefs;
+
+ULandscapeMeshCollisionComponent::FPhysXMeshRef::~FPhysXMeshRef()
+{
+#if WITH_PHYSX
+	// Free the existing heightfield data.
+	if (RBTriangleMesh)
+	{
+		GPhysXPendingKillTriMesh.Add(RBTriangleMesh);
+		RBTriangleMesh = NULL;
+	}
+
+#if WITH_EDITOR
+	if (RBTriangleMeshEd)
+	{
+		GPhysXPendingKillTriMesh.Add(RBTriangleMeshEd);
+		RBTriangleMeshEd = NULL;
+	}
+#endif// WITH_EDITOR
+#endif// WITH_PHYSX
+
+	// Remove ourselves from the shared map.
+	GSharedMeshRefs.Remove(Guid);
+}
+
+ECollisionEnabled::Type ULandscapeHeightfieldCollisionComponent::GetCollisionEnabled() const
+{
+	ALandscapeProxy* Proxy = GetLandscapeProxy();
+
+	return Proxy->BodyInstance.GetCollisionEnabled();
+}
+
+ECollisionResponse ULandscapeHeightfieldCollisionComponent::GetCollisionResponseToChannel(ECollisionChannel Channel) const
+{
+	ALandscapeProxy* Proxy = GetLandscapeProxy();
+
+	return Proxy->BodyInstance.GetResponseToChannel(Channel);
+}
+
+ECollisionChannel ULandscapeHeightfieldCollisionComponent::GetCollisionObjectType() const
+{
+	ALandscapeProxy* Proxy = GetLandscapeProxy();
+
+	return Proxy->BodyInstance.GetObjectType();
+}
+
+const FCollisionResponseContainer& ULandscapeHeightfieldCollisionComponent::GetCollisionResponseToChannels() const
+{
+	ALandscapeProxy* Proxy = GetLandscapeProxy();
+
+	return Proxy->BodyInstance.GetResponseToChannels();
+}
+
+void ULandscapeHeightfieldCollisionComponent::CreatePhysicsState()
+{
+	USceneComponent::CreatePhysicsState(); // route CreatePhysicsState, skip PrimitiveComponent implementation
+
+	if (!BodyInstance.IsValidBodyInstance())
+	{
+#if WITH_PHYSX
+		// This will do nothing, because we create heightfield at component PostLoad event, unless we destroyed it explicitly
+		CreateCollisionObject();
+
+		if (IsValidRef(HeightfieldRef))
+		{
+			// Make transform for this landscape component PxActor
+			FTransform LandscapeComponentTransform = GetComponentToWorld();
+			FMatrix LandscapeComponentMatrix = LandscapeComponentTransform.ToMatrixWithScale();
+			bool bIsMirrored = LandscapeComponentMatrix.Determinant() < 0.f;
+			if (!bIsMirrored)
+			{
+				// Unreal and PhysX have opposite handedness, so we need to translate the origin and rearrange the data
+				LandscapeComponentMatrix = FTranslationMatrix(FVector(CollisionSizeQuads*CollisionScale, 0, 0)) * LandscapeComponentMatrix;
+			}
+
+			// Get the scale to give to PhysX
+			FVector LandscapeScale = LandscapeComponentMatrix.ExtractScaling();
+
+			// Reorder the axes
+			FVector TerrainX = LandscapeComponentMatrix.GetScaledAxis(EAxis::X);
+			FVector TerrainY = LandscapeComponentMatrix.GetScaledAxis(EAxis::Y);
+			FVector TerrainZ = LandscapeComponentMatrix.GetScaledAxis(EAxis::Z);
+			LandscapeComponentMatrix.SetAxis(0, TerrainX);
+			LandscapeComponentMatrix.SetAxis(2, TerrainY);
+			LandscapeComponentMatrix.SetAxis(1, TerrainZ);
+
+			PxTransform PhysXLandscapeComponentTransform = U2PTransform(FTransform(LandscapeComponentMatrix));
+
+			// Create the geometry
+			PxHeightFieldGeometry LandscapeComponentGeom(HeightfieldRef->RBHeightfield, PxMeshGeometryFlags(), LandscapeScale.Z * LANDSCAPE_ZSCALE, LandscapeScale.Y * CollisionScale, LandscapeScale.X * CollisionScale);
+
+			if (LandscapeComponentGeom.isValid())
+			{
+				// Creating both a sync and async actor, since this object is static
+
+				// Create the sync scene actor
+				PxRigidStatic* HeightFieldActorSync = GPhysXSDK->createRigidStatic(PhysXLandscapeComponentTransform);
+				PxShape* HeightFieldShapeSync = HeightFieldActorSync->createShape(LandscapeComponentGeom, HeightfieldRef->UsedPhysicalMaterialArray.GetData(), HeightfieldRef->UsedPhysicalMaterialArray.Num());
+				check(HeightFieldShapeSync);
+
+				// Setup filtering
+				PxFilterData PQueryFilterData, PSimFilterData;
+				CreateShapeFilterData(GetCollisionObjectType(), GetOwner()->GetUniqueID(), GetCollisionResponseToChannels(), 0, 0, PQueryFilterData, PSimFilterData, true, false, true);
+
+				// Heightfield is used for simple and complex collision
+				PQueryFilterData.word3 |= (EPDF_SimpleCollision | EPDF_ComplexCollision);
+				PSimFilterData.word3 |= (EPDF_SimpleCollision | EPDF_ComplexCollision);
+				HeightFieldShapeSync->setQueryFilterData(PQueryFilterData);
+				HeightFieldShapeSync->setSimulationFilterData(PSimFilterData);
+				HeightFieldShapeSync->setFlag(PxShapeFlag::eSCENE_QUERY_SHAPE, true);
+				HeightFieldShapeSync->setFlag(PxShapeFlag::eSIMULATION_SHAPE, true);
+				HeightFieldShapeSync->setFlag(PxShapeFlag::eVISUALIZATION, true);
+
+#if WITH_EDITOR
+				// Create a shape for a heightfield which is used only by the landscape editor
+				if (!GetWorld()->IsGameWorld())
+				{
+					PxHeightFieldGeometry LandscapeComponentGeomEd(HeightfieldRef->RBHeightfieldEd, PxMeshGeometryFlags(), LandscapeScale.Z * LANDSCAPE_ZSCALE, LandscapeScale.Y * CollisionScale, LandscapeScale.X * CollisionScale);
+					if (LandscapeComponentGeomEd.isValid())
+					{
+						PxMaterial* PDefaultMat = GEngine->DefaultPhysMaterial->GetPhysXMaterial();
+						PxShape* HeightFieldEdShapeSync = HeightFieldActorSync->createShape(LandscapeComponentGeomEd, &PDefaultMat, 1);
+						check(HeightFieldEdShapeSync);
+
+						FCollisionResponseContainer CollisionResponse;
+						CollisionResponse.SetAllChannels(ECollisionResponse::ECR_Ignore);
+						CollisionResponse.SetResponse(ECollisionChannel::ECC_Visibility, ECR_Block);
+						PxFilterData PQueryFilterDataEd, PSimFilterDataEd;
+						CreateShapeFilterData(ECollisionChannel::ECC_Visibility, GetOwner()->GetUniqueID(), CollisionResponse, 0, 0, PQueryFilterDataEd, PSimFilterDataEd, true, false, true);
+
+						PQueryFilterDataEd.word3 |= (EPDF_SimpleCollision | EPDF_ComplexCollision);
+						HeightFieldEdShapeSync->setQueryFilterData(PQueryFilterDataEd);
+						HeightFieldEdShapeSync->setFlag(PxShapeFlag::eSCENE_QUERY_SHAPE, true);
+					}
+				}
+#endif// WITH_EDITOR
+
+				FPhysScene* PhysScene = GetWorld()->GetPhysicsScene();
+
+				PxRigidStatic* HeightFieldActorAsync = NULL;
+				if (PhysScene->HasAsyncScene())
+				{
+					// Create the async scene actor
+					HeightFieldActorAsync = GPhysXSDK->createRigidStatic(PhysXLandscapeComponentTransform);
+					PxShape* HeightFieldShapeAsync = HeightFieldActorAsync->createShape(LandscapeComponentGeom, HeightfieldRef->UsedPhysicalMaterialArray.GetData(), HeightfieldRef->UsedPhysicalMaterialArray.Num());
+
+					HeightFieldShapeAsync->setQueryFilterData(PQueryFilterData);
+					HeightFieldShapeAsync->setSimulationFilterData(PSimFilterData);
+					// Only perform scene queries in the synchronous scene for static shapes
+					HeightFieldShapeAsync->setFlag(PxShapeFlag::eSCENE_QUERY_SHAPE, false);
+					HeightFieldShapeAsync->setFlag(PxShapeFlag::eSIMULATION_SHAPE, true);
+					HeightFieldShapeAsync->setFlag(PxShapeFlag::eVISUALIZATION, true);
+				}
+
+				// Set body instance data
+				BodyInstance.PhysxUserData = FPhysxUserData(&BodyInstance);
+				BodyInstance.OwnerComponent = this;
+				BodyInstance.SceneIndexSync = PhysScene->PhysXSceneIndex[PST_Sync];
+				BodyInstance.SceneIndexAsync = PhysScene->HasAsyncScene() ? PhysScene->PhysXSceneIndex[PST_Async] : 0;
+				BodyInstance.RigidActorSync = HeightFieldActorSync;
+				BodyInstance.RigidActorAsync = HeightFieldActorAsync;
+				HeightFieldActorSync->userData = &BodyInstance.PhysxUserData;
+				if (PhysScene->HasAsyncScene())
+				{
+					HeightFieldActorAsync->userData = &BodyInstance.PhysxUserData;
+				}
+
+				// Add to scenes
+				PhysScene->GetPhysXScene(PST_Sync)->addActor(*HeightFieldActorSync);
+
+				if (PhysScene->HasAsyncScene())
+				{
+					PxScene* AsyncScene = PhysScene->GetPhysXScene(PST_Async);
+
+					SCOPED_SCENE_WRITE_LOCK(AsyncScene);
+					AsyncScene->addActor(*HeightFieldActorAsync);
+				}
+			}
+
+		}
+#endif// WITH_PHYSX
+	}
+}
+
+void ULandscapeHeightfieldCollisionComponent::ApplyWorldOffset(const FVector& InOffset, bool bWorldShift)
+{
+	Super::ApplyWorldOffset(InOffset, bWorldShift);
+
+	if (!bWorldShift || !FPhysScene::SupportsOriginShifting())
+	{
+		RecreatePhysicsState();
+	}
+}
+
+void ULandscapeHeightfieldCollisionComponent::CreateCollisionObject()
+{
+	if (IsTemplate())
+	{
+		return;
+	}
+
+#if WITH_PHYSX	
+	// If we have not created a heightfield yet - do it now.
+	if (!IsValidRef(HeightfieldRef))
+	{
+		if (!HeightfieldGuid.IsValid())
+		{
+			HeightfieldGuid = FGuid::NewGuid();
+		}
+
+		// Look for a heightfield object with the current Guid (this occurs with PIE)
+		FPhysXHeightfieldRef* ExistingHeightfieldRef = GSharedHeightfieldRefs.FindRef(HeightfieldGuid);
+		if (ExistingHeightfieldRef)
+		{
+			HeightfieldRef = ExistingHeightfieldRef;
+		}
+		else
+		{
+#if WITH_EDITOR
+			// Prepare heightfield data
+			static FName PhysicsFormatName(FPlatformProperties::GetPhysicsFormat());
+			CookCollsionData(PhysicsFormatName, false, CookedCollisionData, CookedPhysicalMaterials);
+#endif //WITH_EDITOR
+
+			if (CookedCollisionData.Num())
+			{
+				HeightfieldRef = GSharedHeightfieldRefs.Add(HeightfieldGuid, new FPhysXHeightfieldRef(HeightfieldGuid));
+
+				// Create heightfield shape
+				{
+					FPhysXInputStream HeightFieldStream(CookedCollisionData.GetData(), CookedCollisionData.Num());
+					HeightfieldRef->RBHeightfield = GPhysXSDK->createHeightField(HeightFieldStream);
+				}
+
+				for (UPhysicalMaterial* PhysicalMaterial : CookedPhysicalMaterials)
+				{
+					HeightfieldRef->UsedPhysicalMaterialArray.Add(PhysicalMaterial->GetPhysXMaterial());
+				}
+
+				// Release cooked data
+				// In cooked builds created PhysX heightfield object will never be deleted until component is alive, so we don't need this data anymore
+				// In the editor this data will be regenerated just before creating new heightfield object
+				CookedCollisionData.Empty();
+
+#if WITH_EDITOR
+				// Create heightfield for the landscape editor (no holes in it)
+				TArray<uint8> CookedHeightmapEd;
+				TArray<UPhysicalMaterial*> CookedMaterialsEd;
+				if (CookCollsionData(PhysicsFormatName, true, CookedHeightmapEd, CookedMaterialsEd))
+				{
+					FPhysXInputStream HeightFieldStream(CookedHeightmapEd.GetData(), CookedHeightmapEd.Num());
+					HeightfieldRef->RBHeightfieldEd = GPhysXSDK->createHeightField(HeightFieldStream);
+				}
+#endif //WITH_EDITOR
+			}
+		}
+	}
+#endif //WITH_PHYSX	
+}
+
+#if WITH_EDITOR
+bool ULandscapeHeightfieldCollisionComponent::CookCollsionData(const FName& Format, bool bUseDefMaterial, TArray<uint8>& OutCookedData, TArray<UPhysicalMaterial*>& OutMaterails) const
+{
+#if WITH_PHYSX
+
+	ALandscapeProxy* Proxy = GetLandscapeProxy();
+	if (!Proxy || !Proxy->GetRootComponent())
+	{
+		return false;
+	}
+
+	UPhysicalMaterial* DefMaterial = Proxy->DefaultPhysMaterial ? Proxy->DefaultPhysMaterial : GEngine->DefaultPhysMaterial;
+
+	// ComponentToWorld might not be initialized at this point, so use landscape transform
+	FVector LandscapeScale = Proxy->GetRootComponent()->RelativeScale3D;
+	bool bIsMirrored = (LandscapeScale.X*LandscapeScale.Y*LandscapeScale.Z) < 0.f;
+
+	int32 CollisionSizeVerts = CollisionSizeQuads + 1;
+
+	const uint16* Heights = (const uint16*)CollisionHeightData.LockReadOnly();
+	check(CollisionHeightData.GetElementCount() == FMath::Square(CollisionSizeVerts));
+
+	//
+	const uint8* DominantLayers = nullptr;
+	if (DominantLayerData.GetElementCount())
+	{
+		DominantLayers = (const uint8*)DominantLayerData.LockReadOnly();
+	}
+
+	// List of materials which is actually used by heightfield
+	OutMaterails.Empty();
+
+	TArray<PxHeightFieldSample> Samples;
+	Samples.AddZeroed(FMath::Square(CollisionSizeVerts));
+
+	for (int32 RowIndex = 0; RowIndex < CollisionSizeVerts; RowIndex++)
+	{
+		for (int32 ColIndex = 0; ColIndex < CollisionSizeVerts; ColIndex++)
+		{
+			int32 SrcSampleIndex = (ColIndex * CollisionSizeVerts) + (bIsMirrored ? RowIndex : (CollisionSizeVerts - RowIndex - 1));
+			int32 DstSampleIndex = (RowIndex * CollisionSizeVerts) + ColIndex;
+
+			PxHeightFieldSample& Sample = Samples[DstSampleIndex];
+			Sample.height = FMath::Clamp<int32>(((int32)Heights[SrcSampleIndex] - 32768), -32768, 32767);
+
+			// Materials are not relevant on the last row/column because they are per-triangle and the last row/column don't own any
+			if (RowIndex < CollisionSizeVerts - 1 &&
+				ColIndex < CollisionSizeVerts - 1)
+			{
+				int32 MaterialIndex = 0; // Default physical material.
+				if (!bUseDefMaterial && DominantLayers)
+				{
+					uint8 DominantLayerIdx = DominantLayers[SrcSampleIndex];
+					if (ComponentLayerInfos.IsValidIndex(DominantLayerIdx))
+					{
+						ULandscapeLayerInfoObject* Layer = ComponentLayerInfos[DominantLayerIdx];
+						if (Layer == ALandscapeProxy::VisibilityLayer)
+						{
+							// If it's a hole, override with the hole flag.
+							MaterialIndex = PxHeightFieldMaterial::eHOLE;
+						}
+						else
+						{
+							UPhysicalMaterial* DominantMaterial = Layer && Layer->PhysMaterial ? Layer->PhysMaterial : DefMaterial;
+							MaterialIndex = OutMaterails.AddUnique(DominantMaterial);
+						}
+					}
+				}
+
+				Sample.materialIndex0 = MaterialIndex;
+				Sample.materialIndex1 = MaterialIndex;
+			}
+
+			// TODO: edge turning
+		}
+	}
+
+	CollisionHeightData.Unlock();
+	if (DominantLayers)
+	{
+		DominantLayerData.Unlock();
+	}
+
+	// Add the default physical material to be used used when we have no dominant data.
+	if (OutMaterails.Num() == 0)
+	{
+		OutMaterails.Add(DefMaterial);
+	}
+
+	//
+	FIntPoint HFSize = FIntPoint(CollisionSizeVerts, CollisionSizeVerts);
+	float HFThickness = -Proxy->CollisionThickness / (LandscapeScale.Z * LANDSCAPE_ZSCALE);
+	TArray<uint8> OutData;
+
+	ITargetPlatformManagerModule* TPM = GetTargetPlatformManager();
+	const IPhysXFormat* Cooker = TPM->FindPhysXFormat(Format);
+	bool Result = Cooker->CookHeightField(Format, HFSize, HFThickness, Samples.GetData(), Samples.GetTypeSize(), OutData);
+
+	if (Result)
+	{
+		OutCookedData.SetNumUninitialized(OutData.Num());
+		FMemory::Memcpy(OutCookedData.GetData(), OutData.GetData(), OutData.Num());
+	}
+	else
+	{
+		OutCookedData.Empty();
+		OutMaterails.Empty();
+	}
+
+	return Result;
+#endif	// WITH_PHYSX
+
+	return false;
+}
+
+bool ULandscapeMeshCollisionComponent::CookCollsionData(const FName& Format, bool bUseDefMaterial, TArray<uint8>& OutCookedData, TArray<UPhysicalMaterial*>& OutMaterails) const
+{
+	if (IsTemplate())
+	{
+		return false;
+	}
+#if WITH_PHYSX
+	ALandscapeProxy* Proxy = GetLandscapeProxy();
+	UPhysicalMaterial* DefMaterial = (Proxy && Proxy->DefaultPhysMaterial != nullptr) ? Proxy->DefaultPhysMaterial : GEngine->DefaultPhysMaterial;
+
+	// List of materials which is actually used by trimesh
+	OutMaterails.Empty();
+
+	TArray<FVector>			Vertices;
+	TArray<FTriIndices>		Indices;
+	TArray<uint16>			MaterialIndices;
+
+	const int32 CollisionSizeVerts = CollisionSizeQuads + 1;
+	const int32 NumVerts = FMath::Square(CollisionSizeVerts);
+
+	const uint16* Heights = (const uint16*)CollisionHeightData.LockReadOnly();
+	const uint16* XYOffsets = (const uint16*)CollisionXYOffsetData.LockReadOnly();
+	check(CollisionHeightData.GetElementCount() == NumVerts);
+	check(CollisionXYOffsetData.GetElementCount() == NumVerts * 2);
+
+	const uint8* DominantLayers = nullptr;
+	if (DominantLayerData.GetElementCount() > 0)
+	{
+		DominantLayers = (const uint8*)DominantLayerData.LockReadOnly();
+	}
+
+	// Scale all verts into temporary vertex buffer.
+	Vertices.AddUninitialized(NumVerts);
+	for (int32 i = 0; i < NumVerts; i++)
+	{
+		int32 X = i % CollisionSizeVerts;
+		int32 Y = i / CollisionSizeVerts;
+		Vertices[i].Set(X + ((float)XYOffsets[i * 2] - 32768.f) * LANDSCAPE_XYOFFSET_SCALE, Y + ((float)XYOffsets[i * 2 + 1] - 32768.f) * LANDSCAPE_XYOFFSET_SCALE, ((float)Heights[i] - 32768.f) * LANDSCAPE_ZSCALE);
+	}
+
+	const int32 NumTris = FMath::Square(CollisionSizeQuads) * 2;
+	Indices.AddUninitialized(NumTris);
+	if (DominantLayers)
+	{
+		MaterialIndices.AddUninitialized(NumTris);
+	}
+
+	int32 TriangleIdx = 0;
+	for (int32 y = 0; y < CollisionSizeQuads; y++)
+	{
+		for (int32 x = 0; x < CollisionSizeQuads; x++)
+		{
+			int32 DataIdx = x + y * CollisionSizeVerts;
+			bool bHole = false;
+
+			int32 MaterialIndex = 0; // Default physical material.
+			if (!bUseDefMaterial && DominantLayers)
+			{
+				uint8 DominantLayerIdx = DominantLayers[DataIdx];
+				if (ComponentLayerInfos.IsValidIndex(DominantLayerIdx))
+				{
+					ULandscapeLayerInfoObject* Layer = ComponentLayerInfos[DominantLayerIdx];
+					if (Layer == ALandscapeProxy::VisibilityLayer)
+					{
+						// If it's a hole, override with the hole flag.
+						bHole = true;
+					}
+					else
+					{
+						UPhysicalMaterial* DominantMaterial = Layer && Layer->PhysMaterial ? Layer->PhysMaterial : DefMaterial;
+						MaterialIndex = OutMaterails.AddUnique(DominantMaterial);
+					}
+				}
+			}
+
+			FTriIndices& TriIndex1 = Indices[TriangleIdx];
+			if (bHole)
+			{
+				TriIndex1.v0 = (x + 0) + (y + 0) * CollisionSizeVerts;
+				TriIndex1.v1 = TriIndex1.v0;
+				TriIndex1.v2 = TriIndex1.v0;
+			}
+			else
+			{
+				TriIndex1.v0 = (x + 0) + (y + 0) * CollisionSizeVerts;
+				TriIndex1.v1 = (x + 1) + (y + 1) * CollisionSizeVerts;
+				TriIndex1.v2 = (x + 1) + (y + 0) * CollisionSizeVerts;
+			}
+
+			if (DominantLayers)
+			{
+				MaterialIndices[TriangleIdx] = MaterialIndex;
+			}
+			TriangleIdx++;
+
+			FTriIndices& TriIndex2 = Indices[TriangleIdx];
+			if (bHole)
+			{
+				TriIndex2.v0 = (x + 0) + (y + 0) * CollisionSizeVerts;
+				TriIndex2.v1 = TriIndex2.v0;
+				TriIndex2.v2 = TriIndex2.v0;
+			}
+			else
+			{
+				TriIndex2.v0 = (x + 0) + (y + 0) * CollisionSizeVerts;
+				TriIndex2.v1 = (x + 0) + (y + 1) * CollisionSizeVerts;
+				TriIndex2.v2 = (x + 1) + (y + 1) * CollisionSizeVerts;
+			}
+
+			if (DominantLayers)
+			{
+				MaterialIndices[TriangleIdx] = MaterialIndex;
+			}
+			TriangleIdx++;
+		}
+	}
+
+	CollisionHeightData.Unlock();
+	CollisionXYOffsetData.Unlock();
+	if (DominantLayers)
+	{
+		DominantLayerData.Unlock();
+	}
+
+	// Add the default physical material to be used used when we have no dominant data.
+	if (OutMaterails.Num() == 0)
+	{
+		OutMaterails.Add(DefMaterial);
+	}
+
+	bool bFlipNormals = true;
+	TArray<uint8> OutData;
+	ITargetPlatformManagerModule* TPM = GetTargetPlatformManager();
+	const IPhysXFormat* Cooker = TPM->FindPhysXFormat(Format);
+	bool Result = Cooker->CookTriMesh(Format, Vertices, Indices, MaterialIndices, bFlipNormals, OutData);
+
+	if (Result)
+	{
+		OutCookedData.SetNumUninitialized(OutData.Num());
+		FMemory::Memcpy(OutCookedData.GetData(), OutData.GetData(), OutData.Num());
+	}
+	else
+	{
+		OutCookedData.Empty();
+		OutMaterails.Empty();
+	}
+
+	return Result;
+
+#endif // WITH_PHYSX
+	return false;
+}
+#endif //WITH_EDITOR
+
+void ULandscapeMeshCollisionComponent::CreateCollisionObject()
+{
+	if (IsTemplate())
+	{
+		return;
+	}
+
+#if WITH_PHYSX	
+	// If we have not created a heightfield yet - do it now.
+	if (!IsValidRef(MeshRef))
+	{
+		if (!MeshGuid.IsValid())
+		{
+			MeshGuid = FGuid::NewGuid();
+		}
+
+		// Look for a heightfield object with the current Guid (this occurs with PIE)
+		FPhysXMeshRef* ExistingMeshRef = GSharedMeshRefs.FindRef(MeshGuid);
+		if (ExistingMeshRef)
+		{
+			MeshRef = ExistingMeshRef;
+		}
+		else
+		{
+#if WITH_EDITOR
+			// Create cooked physics data
+			static FName PhysicsFormatName(FPlatformProperties::GetPhysicsFormat());
+			CookCollsionData(PhysicsFormatName, false, CookedCollisionData, CookedPhysicalMaterials);
+#endif //WITH_EDITOR
+
+			if (CookedCollisionData.Num())
+			{
+				MeshRef = GSharedMeshRefs.Add(MeshGuid, new FPhysXMeshRef(MeshGuid));
+
+				// Create physics objects
+				FPhysXInputStream Buffer(CookedCollisionData.GetData(), CookedCollisionData.Num());
+				MeshRef->RBTriangleMesh = GPhysXSDK->createTriangleMesh(Buffer);
+
+				for (UPhysicalMaterial* PhysicalMaterial : CookedPhysicalMaterials)
+				{
+					MeshRef->UsedPhysicalMaterialArray.Add(PhysicalMaterial->GetPhysXMaterial());
+				}
+
+				// Release cooked data
+				// In cooked builds created collision object will never be deleted until component is alive, so we don't need this data anymore
+				// In the editor this data will be regenerated just before creating new collision object
+				CookedCollisionData.Empty();
+
+#if WITH_EDITOR
+				// Create collision mesh for the landscape editor (no holes in it)
+				TArray<uint8> CookedMeshEd;
+				TArray<UPhysicalMaterial*> CookedMaterialsEd;
+				if (CookCollsionData(PhysicsFormatName, true, CookedMeshEd, CookedMaterialsEd))
+				{
+					FPhysXInputStream MeshStream(CookedMeshEd.GetData(), CookedMeshEd.Num());
+					MeshRef->RBTriangleMeshEd = GPhysXSDK->createTriangleMesh(MeshStream);
+				}
+#endif //WITH_EDITOR
+			}
+		}
+	}
+#endif //WITH_PHYSX
+}
+
+void ULandscapeMeshCollisionComponent::CreatePhysicsState()
+{
+	USceneComponent::CreatePhysicsState(); // route CreatePhysicsState, skip PrimitiveComponent implementation
+
+	if (!BodyInstance.IsValidBodyInstance())
+	{
+#if WITH_PHYSX
+		// This will do nothing, because we create trimesh at component PostLoad event, unless we destroyed it explicitly
+		CreateCollisionObject();
+
+		if (IsValidRef(MeshRef))
+		{
+			// Make transform for this landscape component PxActor
+			FTransform LandscapeComponentTransform = GetComponentToWorld();
+			FMatrix LandscapeComponentMatrix = LandscapeComponentTransform.ToMatrixWithScale();
+			bool bIsMirrored = LandscapeComponentMatrix.Determinant() < 0.f;
+			if (bIsMirrored)
+			{
+				// Unreal and PhysX have opposite handedness, so we need to translate the origin and rearrange the data
+				LandscapeComponentMatrix = FTranslationMatrix(FVector(CollisionSizeQuads, 0, 0)) * LandscapeComponentMatrix;
+			}
+
+			// Get the scale to give to PhysX
+			FVector LandscapeScale = LandscapeComponentMatrix.ExtractScaling();
+			PxTransform PhysXLandscapeComponentTransform = U2PTransform(FTransform(LandscapeComponentMatrix));
+
+			// Create tri-mesh shape
+			PxTriangleMeshGeometry PTriMeshGeom;
+			PTriMeshGeom.triangleMesh = MeshRef->RBTriangleMesh;
+			PTriMeshGeom.scale.scale.x = LandscapeScale.X * CollisionScale;
+			PTriMeshGeom.scale.scale.y = LandscapeScale.Y * CollisionScale;
+			PTriMeshGeom.scale.scale.z = LandscapeScale.Z;
+
+			if (PTriMeshGeom.isValid())
+			{
+				// Creating both a sync and async actor, since this object is static
+
+				// Create the sync scene actor
+				PxRigidStatic* MeshActorSync = GPhysXSDK->createRigidStatic(PhysXLandscapeComponentTransform);
+				PxShape* MeshShapeSync = MeshActorSync->createShape(PTriMeshGeom, MeshRef->UsedPhysicalMaterialArray.GetData(), MeshRef->UsedPhysicalMaterialArray.Num());
+				check(MeshShapeSync);
+
+				// Setup filtering
+				PxFilterData PQueryFilterData, PSimFilterData;
+				CreateShapeFilterData(GetCollisionObjectType(), GetOwner()->GetUniqueID(), GetCollisionResponseToChannels(), 0, 0, PQueryFilterData, PSimFilterData, false, false, true);
+
+				// Heightfield is used for simple and complex collision
+				PQueryFilterData.word3 |= (EPDF_SimpleCollision | EPDF_ComplexCollision);
+				PSimFilterData.word3 |= (EPDF_SimpleCollision | EPDF_ComplexCollision);
+				MeshShapeSync->setQueryFilterData(PQueryFilterData);
+				MeshShapeSync->setSimulationFilterData(PSimFilterData);
+				MeshShapeSync->setFlag(PxShapeFlag::eSCENE_QUERY_SHAPE, true);
+				MeshShapeSync->setFlag(PxShapeFlag::eSIMULATION_SHAPE, true);
+				MeshShapeSync->setFlag(PxShapeFlag::eVISUALIZATION, true);
+
+				FPhysScene* PhysScene = GetWorld()->GetPhysicsScene();
+
+				PxRigidStatic* MeshActorAsync = NULL;
+				if (PhysScene->HasAsyncScene())
+				{
+					// Create the async scene actor
+					MeshActorAsync = GPhysXSDK->createRigidStatic(PhysXLandscapeComponentTransform);
+					PxShape* MeshShapeAsync = MeshActorAsync->createShape(PTriMeshGeom, MeshRef->UsedPhysicalMaterialArray.GetData(), MeshRef->UsedPhysicalMaterialArray.Num());
+					check(MeshShapeAsync);
+
+					MeshShapeAsync->setQueryFilterData(PQueryFilterData);
+					MeshShapeAsync->setSimulationFilterData(PSimFilterData);
+					// Only perform scene queries in the synchronous scene for static shapes
+					MeshShapeAsync->setFlag(PxShapeFlag::eSCENE_QUERY_SHAPE, false);
+					MeshShapeAsync->setFlag(PxShapeFlag::eSIMULATION_SHAPE, true);
+					MeshShapeAsync->setFlag(PxShapeFlag::eVISUALIZATION, true);	// Setting visualization flag, in case we visualize only the async scene
+				}
+
+#if WITH_EDITOR
+				// Create a shape for a mesh which is used only by the landscape editor
+				if (!GetWorld()->IsGameWorld())
+				{
+					PxTriangleMeshGeometry PTriMeshGeomEd;
+					PTriMeshGeomEd.triangleMesh = MeshRef->RBTriangleMeshEd;
+					PTriMeshGeomEd.scale.scale.x = LandscapeScale.X * CollisionScale;
+					PTriMeshGeomEd.scale.scale.y = LandscapeScale.Y * CollisionScale;
+					PTriMeshGeomEd.scale.scale.z = LandscapeScale.Z;
+					if (PTriMeshGeomEd.isValid())
+					{
+						PxMaterial* PDefaultMat = GEngine->DefaultPhysMaterial->GetPhysXMaterial();
+						PxShape* MeshShapeEdSync = MeshActorSync->createShape(PTriMeshGeomEd, &PDefaultMat, 1);
+						check(MeshShapeEdSync);
+
+						FCollisionResponseContainer CollisionResponse;
+						CollisionResponse.SetAllChannels(ECollisionResponse::ECR_Ignore);
+						CollisionResponse.SetResponse(ECollisionChannel::ECC_Visibility, ECR_Block);
+						PxFilterData PQueryFilterDataEd, PSimFilterDataEd;
+						CreateShapeFilterData(ECollisionChannel::ECC_Visibility, GetOwner()->GetUniqueID(), CollisionResponse, 0, 0, PQueryFilterDataEd, PSimFilterDataEd, true, false, true);
+
+						PQueryFilterDataEd.word3 |= (EPDF_SimpleCollision | EPDF_ComplexCollision);
+						MeshShapeEdSync->setQueryFilterData(PQueryFilterDataEd);
+						MeshShapeEdSync->setFlag(PxShapeFlag::eSCENE_QUERY_SHAPE, true);
+					}
+				}
+#endif// WITH_EDITOR
+
+				// Set body instance data
+				BodyInstance.PhysxUserData = FPhysxUserData(&BodyInstance);
+				BodyInstance.OwnerComponent = this;
+				BodyInstance.SceneIndexSync = PhysScene->PhysXSceneIndex[PST_Sync];
+				BodyInstance.SceneIndexAsync = PhysScene->HasAsyncScene() ? PhysScene->PhysXSceneIndex[PST_Async] : 0;
+				BodyInstance.RigidActorSync = MeshActorSync;
+				BodyInstance.RigidActorAsync = MeshActorAsync;
+				MeshActorSync->userData = &BodyInstance.PhysxUserData;
+				if (PhysScene->HasAsyncScene())
+				{
+					MeshActorAsync->userData = &BodyInstance.PhysxUserData;
+				}
+
+				// Add to scenes
+				PhysScene->GetPhysXScene(PST_Sync)->addActor(*MeshActorSync);
+
+				if (PhysScene->HasAsyncScene())
+				{
+					PxScene* AsyncScene = PhysScene->GetPhysXScene(PST_Async);
+
+					SCOPED_SCENE_WRITE_LOCK(AsyncScene);
+					AsyncScene->addActor(*MeshActorAsync);
+				}
+			}
+			else
+			{
+				UE_LOG(LogLandscape, Log, TEXT("ULandscapeMeshCollisionComponent::CreatePhysicsState(): TriMesh invalid"));
+			}
+		}
+#endif // WITH_PHYSX
+	}
+}
+
+void ULandscapeMeshCollisionComponent::ApplyWorldOffset(const FVector& InOffset, bool bWorldShift)
+{
+	Super::ApplyWorldOffset(InOffset, bWorldShift);
+
+	if (!bWorldShift || !FPhysScene::SupportsOriginShifting())
+	{
+		RecreatePhysicsState();
+	}
+}
+
+void ULandscapeMeshCollisionComponent::DestroyComponent()
+{
+	ALandscapeProxy* Proxy = GetLandscapeProxy();
+	if (Proxy)
+	{
+		Proxy->CollisionComponents.Remove(this);
+	}
+
+	Super::DestroyComponent();
+}
+
+#if WITH_EDITOR
+void ULandscapeHeightfieldCollisionComponent::UpdateHeightfieldRegion(int32 ComponentX1, int32 ComponentY1, int32 ComponentX2, int32 ComponentY2)
+{
+#if WITH_PHYSX
+	if (IsValidRef(HeightfieldRef))
+	{
+		// If we're currently sharing this data with a PIE session, we need to make a new heightfield.
+		if (HeightfieldRef->GetRefCount() > 1)
+		{
+			RecreateCollision(false);
+			return;
+		}
+
+		if (BodyInstance.RigidActorSync == NULL)
+		{
+			return;
+		}
+
+		int32 CollisionSizeVerts = CollisionSizeQuads + 1;
+
+		bool bIsMirrored = GetComponentToWorld().GetDeterminant() < 0.f;
+
+		uint16* Heights = (uint16*)CollisionHeightData.Lock(LOCK_READ_ONLY);
+		check(CollisionHeightData.GetElementCount() == FMath::Square(CollisionSizeVerts));
+
+		// PhysX heightfield has the X and Y axis swapped, and the X component is also inverted
+		int32 HeightfieldX1 = ComponentY1;
+		int32 HeightfieldY1 = (bIsMirrored ? ComponentX1 : (CollisionSizeVerts - ComponentX2 - 1));
+		int32 DstVertsX = ComponentY2 - ComponentY1 + 1;
+		int32 DstVertsY = ComponentX2 - ComponentX1 + 1;
+
+		TArray<PxHeightFieldSample> Samples;
+		Samples.AddZeroed(DstVertsX*DstVertsY);
+
+		// Traverse the area in destination heigthfield coordinates
+		for (int32 RowIndex = 0; RowIndex < DstVertsY; RowIndex++)
+		{
+			for (int32 ColIndex = 0; ColIndex < DstVertsX; ColIndex++)
+			{
+				int32 SrcX = bIsMirrored ? (RowIndex + ComponentX1) : (ComponentX2 - RowIndex);
+				int32 SrcY = ColIndex + ComponentY1;
+				int32 SrcSampleIndex = (SrcY * CollisionSizeVerts) + SrcX;
+				check(SrcSampleIndex < FMath::Square(CollisionSizeVerts));
+				int32 DstSampleIndex = (RowIndex * DstVertsX) + ColIndex;
+
+				PxHeightFieldSample& Sample = Samples[DstSampleIndex];
+				Sample.height = FMath::Clamp<int32>(((int32)Heights[SrcSampleIndex] - 32768), -32768, 32767);
+
+				Sample.materialIndex0 = 0;
+				Sample.materialIndex1 = 0;
+			}
+		}
+
+		CollisionHeightData.Unlock();
+
+		PxHeightFieldDesc SubDesc;
+		SubDesc.format = PxHeightFieldFormat::eS16_TM;
+		SubDesc.nbColumns = DstVertsX;
+		SubDesc.nbRows = DstVertsY;
+		SubDesc.samples.data = Samples.GetData();
+		SubDesc.samples.stride = sizeof(PxU32);
+		SubDesc.flags = PxHeightFieldFlag::eNO_BOUNDARY_EDGES;
+
+
+		HeightfieldRef->RBHeightfieldEd->modifySamples(HeightfieldX1, HeightfieldY1, SubDesc, true);
+
+		//
+		// Reset geometry of heightfield shape. Required by the modifySamples
+		//
+		FVector LandscapeScale = GetComponentToWorld().GetScale3D();
+		// Create the geometry
+		PxHeightFieldGeometry LandscapeComponentGeom(HeightfieldRef->RBHeightfieldEd, PxMeshGeometryFlags(), LandscapeScale.Z * LANDSCAPE_ZSCALE, LandscapeScale.Y * CollisionScale, LandscapeScale.X * CollisionScale);
+
+		if (BodyInstance.RigidActorSync)
+		{
+			TArray<PxShape*, TInlineAllocator<8>> PShapes;
+			PShapes.AddZeroed(BodyInstance.RigidActorSync->getNbShapes());
+			int32 NumShapes = BodyInstance.RigidActorSync->getShapes(PShapes.GetData(), PShapes.Num());
+			if (NumShapes > 1)
+			{
+				PShapes[1]->setGeometry(LandscapeComponentGeom);
+			}
+		}
+	}
+
+#endif// WITH_PHYSX
+}
+#endif// WITH_EDITOR
+
+void ULandscapeHeightfieldCollisionComponent::DestroyComponent()
+{
+	ALandscapeProxy* Proxy = GetLandscapeProxy();
+	if (Proxy)
+	{
+		Proxy->CollisionComponents.Remove(this);
+	}
+
+	Super::DestroyComponent();
+}
+
+FBoxSphereBounds ULandscapeHeightfieldCollisionComponent::CalcBounds(const FTransform& LocalToWorld) const
+{
+	return CachedLocalBox.TransformBy(LocalToWorld);
+}
+
+void ULandscapeHeightfieldCollisionComponent::BeginDestroy()
+{
+	HeightfieldRef = NULL;
+	HeightfieldGuid = FGuid();
+	Super::BeginDestroy();
+}
+
+void ULandscapeMeshCollisionComponent::BeginDestroy()
+{
+	if (!HasAnyFlags(RF_ClassDefaultObject))
+	{
+		MeshRef = NULL;
+		MeshGuid = FGuid();
+	}
+
+	Super::BeginDestroy();
+}
+
+void ULandscapeHeightfieldCollisionComponent::RecreateCollision(bool bUpdateAddCollision/*= true*/)
+{
+	if (!HasAnyFlags(RF_ClassDefaultObject))
+	{
+		HeightfieldRef = NULL;
+		HeightfieldGuid = FGuid();
+#if WITH_EDITOR
+		if (bUpdateAddCollision)
+		{
+			UpdateAddCollisions();
+		}
+#endif
+
+		RecreatePhysicsState();
+	}
+}
+
+#if WITH_EDITORONLY_DATA
+void ULandscapeHeightfieldCollisionComponent::SnapFoliageInstances(AInstancedFoliageActor &IFA, const FBox& InInstanceBox)
+{
+	for (auto& MeshPair : IFA.FoliageMeshes)
+	{
+		// Find the per-mesh info matching the mesh.
+		UFoliageType* Settings = MeshPair.Key;
+		FFoliageMeshInfo& MeshInfo = *MeshPair.Value;
+
+		FFoliageComponentHashInfo* ComponentHashInfo = MeshInfo.ComponentHash.Find(this);
+		if (ComponentHashInfo)
+		{
+			float TraceExtentSize = Bounds.SphereRadius * 2.f + 10.f; // extend a little
+			FVector TraceVector = GetOwner()->GetRootComponent()->ComponentToWorld.GetUnitAxis(EAxis::Z) * TraceExtentSize;
+
+			bool bFirst = true;
+			TArray<int32> InstancesToRemove;
+			for (int32 InstanceIndex : ComponentHashInfo->Instances)
+			{
+				FFoliageInstance& Instance = MeshInfo.Instances[InstanceIndex];
+
+				// Test location should remove any Z offset
+				FVector TestLocation = FMath::Abs(Instance.ZOffset) > KINDA_SMALL_NUMBER
+					? (FVector)Instance.GetInstanceWorldTransform().TransformPosition(FVector(0, 0, -Instance.ZOffset))
+					: Instance.Location;
+
+				if (InInstanceBox.IsInside(TestLocation))
+				{
+					if (bFirst)
+					{
+						bFirst = false;
+						Modify();
+					}
+
+					FVector Start = TestLocation + TraceVector;
+					FVector End = TestLocation - TraceVector;
+
+					static FName TraceTag = FName(TEXT("FoliageSnapToLandscape"));
+					TArray<FHitResult> Results;
+					UWorld* World = GetWorld();
+					check(World);
+					// Editor specific landscape heightfield uses ECC_Visibility collision channel
+					World->LineTraceMulti(Results, Start, End, FCollisionQueryParams(TraceTag, true), FCollisionObjectQueryParams(ECollisionChannel::ECC_Visibility));
+
+					bool bFoundHit = false;
+					for (const FHitResult& Hit : Results)
+					{
+						if (Hit.Component == this)
+						{
+							bFoundHit = true;
+							if ((TestLocation - Hit.Location).SizeSquared() > KINDA_SMALL_NUMBER)
+							{
+								// Remove instance location from the hash. Do not need to update ComponentHash as we re-add below.
+								MeshInfo.InstanceHash->RemoveInstance(Instance.Location, InstanceIndex);
+
+								// Update the instance editor data
+								Instance.Location = Hit.Location;
+
+								if (Instance.Flags & FOLIAGE_AlignToNormal)
+								{
+									// Remove previous alignment and align to new normal.
+									Instance.Rotation = Instance.PreAlignRotation;
+									Instance.AlignToNormal(Hit.Normal, Settings->AlignMaxAngle);
+								}
+
+								// Reapply the Z offset in local space
+								if (FMath::Abs(Instance.ZOffset) > KINDA_SMALL_NUMBER)
+								{
+									Instance.Location = Instance.GetInstanceWorldTransform().TransformPosition(FVector(0, 0, Instance.ZOffset));
+								}
+
+								// Todo: add do validation with other parameters such as max/min height etc.
+
+								// Update this instances' transform in the UInstancedStaticMeshComponent
+								if (MeshInfo.InstanceClusters.IsValidIndex(Instance.ClusterIndex))
+								{
+									FFoliageInstanceCluster& Cluster = MeshInfo.InstanceClusters[Instance.ClusterIndex];
+
+									Cluster.ClusterComponent->Modify();
+
+									int32 ClusterInstanceDataIndex = Cluster.InstanceIndices.Find(InstanceIndex);
+
+									check(ClusterInstanceDataIndex != INDEX_NONE);
+									check(Cluster.ClusterComponent->PerInstanceSMData.IsValidIndex(ClusterInstanceDataIndex));
+									Cluster.ClusterComponent->MarkRenderStateDirty();
+
+									// Update bounds
+									FTransform InstanceToWorld = Instance.GetInstanceWorldTransform();
+									Cluster.Bounds = Cluster.Bounds + Cluster.ClusterComponent->StaticMesh->GetBounds().TransformBy(InstanceToWorld);
+
+									// Update transform in InstancedStaticMeshComponent
+									FTransform InstanceToLocal = InstanceToWorld.GetRelativeTransform(Cluster.ClusterComponent->ComponentToWorld);
+									Cluster.ClusterComponent->PerInstanceSMData[ClusterInstanceDataIndex].Transform = InstanceToLocal.ToMatrixWithScale();
+									Cluster.ClusterComponent->InvalidateLightingCache();
+								}
+
+								// Re-add the new instance location to the hash
+								MeshInfo.InstanceHash->InsertInstance(Instance.Location, InstanceIndex);
+							}
+							break;
+						}
+					}
+
+					if (!bFoundHit)
+					{
+						// Couldn't find new spot - remove instance
+						InstancesToRemove.Add(InstanceIndex);
+					}
+				}
+			}
+
+			// Remove any unused instances
+			MeshInfo.RemoveInstances(&IFA, InstancesToRemove);
+		}
+	}
+}
+#endif // WITH_EDITORONLY_DATA
+
+void ULandscapeMeshCollisionComponent::RecreateCollision(bool bUpdateAddCollision/*= true*/)
+{
+	if (!HasAnyFlags(RF_ClassDefaultObject))
+	{
+		MeshRef = NULL;
+		MeshGuid = FGuid();
+	}
+
+	Super::RecreateCollision(bUpdateAddCollision);
+}
+
+void ULandscapeHeightfieldCollisionComponent::Serialize(FArchive& Ar)
+{
+	Super::Serialize(Ar);
+
+	if (Ar.UE4Ver() < VER_UE4_LANDSCAPE_COLLISION_DATA_COOKING)
+	{
+		CollisionHeightData.Serialize(Ar, this);
+		DominantLayerData.Serialize(Ar, this);
+	}
+	else
+	{
+		bool bCooked = Ar.IsCooking();
+		Ar << bCooked;
+
+		if (FPlatformProperties::RequiresCookedData() && !bCooked && Ar.IsLoading())
+		{
+			UE_LOG(LogPhysics, Fatal, TEXT("This platform requires cooked packages, and physX data was not cooked into %s."), *GetFullName());
+		}
+
+		if (bCooked)
+		{
+#if WITH_EDITOR
+			if (Ar.IsCooking() && !HasAnyFlags(RF_ClassDefaultObject))
+			{
+				CookCollsionData(Ar.CookingTarget()->GetPhysicsFormat(NULL), false, CookedCollisionData, CookedPhysicalMaterials);
+			}
+#endif// WITH_EDITOR
+
+			Ar << CookedCollisionData;
+			Ar << CookedPhysicalMaterials;
+		}
+		else
+		{
+			CollisionHeightData.Serialize(Ar, this);
+			DominantLayerData.Serialize(Ar, this);
+		}
+	}
+}
+
+void ULandscapeMeshCollisionComponent::Serialize(FArchive& Ar)
+{
+	Super::Serialize(Ar);
+
+	if (Ar.UE4Ver() < VER_UE4_LANDSCAPE_COLLISION_DATA_COOKING)
+	{
+		// conditional serialization in later versions
+		CollisionXYOffsetData.Serialize(Ar, this);
+	}
+
+	// PhysX cooking mesh data
+	bool bCooked = false;
+	if (Ar.UE4Ver() >= VER_UE4_ADD_COOKED_TO_LANDSCAPE)
+	{
+		bCooked = Ar.IsCooking();
+		Ar << bCooked;
+	}
+
+	if (FPlatformProperties::RequiresCookedData() && !bCooked && Ar.IsLoading())
+	{
+		UE_LOG(LogPhysics, Fatal, TEXT("This platform requires cooked packages, and physX data was not cooked into %s."), *GetFullName());
+	}
+
+	if (bCooked)
+	{
+		// triangle mesh cooked data should be serialized in ULandscapeHeightfieldCollisionComponent
+	}
+	else if (Ar.UE4Ver() >= VER_UE4_LANDSCAPE_COLLISION_DATA_COOKING)
+	{
+		// we serialize raw collision data only with non-cooked content
+		CollisionXYOffsetData.Serialize(Ar, this);
+	}
+}
+
+#if WITH_EDITOR
+void ULandscapeHeightfieldCollisionComponent::PostEditImport()
+{
+	Super::PostEditImport();
+	// Reinitialize physics after paste
+	if (CollisionSizeQuads > 0)
+	{
+		RecreateCollision(false);
+	}
+}
+
+void ULandscapeHeightfieldCollisionComponent::PostEditUndo()
+{
+	Super::PostEditUndo();
+
+	// Reinitialize physics after undo
+	if (CollisionSizeQuads > 0)
+	{
+		RecreateCollision(false);
+	}
+
+	UNavigationSystem::UpdateNavOctree(this);
+}
+#endif
+
+bool ULandscapeHeightfieldCollisionComponent::DoCustomNavigableGeometryExport(FNavigableGeometryExport* GeomExport) const
+{
+	check(IsInGameThread());
+#if WITH_PHYSX
+	if (IsValidRef(HeightfieldRef) && HeightfieldRef->RBHeightfield != nullptr)
+	{
+		FTransform HFToW = ComponentToWorld;
+		HFToW.MultiplyScale3D(FVector(CollisionScale, CollisionScale, LANDSCAPE_ZSCALE));
+
+		GeomExport->ExportPxHeightField(HeightfieldRef->RBHeightfield, HFToW);
+	}
+#endif// WITH_PHYSX
+	return false;
+}
+
+bool ULandscapeMeshCollisionComponent::DoCustomNavigableGeometryExport(struct FNavigableGeometryExport* GeomExport) const
+{
+	check(IsInGameThread());
+#if WITH_PHYSX
+	if (IsValidRef(MeshRef) && MeshRef->RBTriangleMesh != nullptr)
+	{
+		FTransform MeshToW = ComponentToWorld;
+		MeshToW.MultiplyScale3D(FVector(CollisionScale, CollisionScale, 1.f));
+
+		if (MeshRef->RBTriangleMesh->getTriangleMeshFlags() & PxTriangleMeshFlag::eHAS_16BIT_TRIANGLE_INDICES)
+		{
+			GeomExport->ExportPxTriMesh16Bit(MeshRef->RBTriangleMesh, MeshToW);
+		}
+		else
+		{
+			GeomExport->ExportPxTriMesh32Bit(MeshRef->RBTriangleMesh, MeshToW);
+		}
+	}
+#endif// WITH_PHYSX
+	return false;
+}
+
+void ULandscapeHeightfieldCollisionComponent::PostLoad()
+{
+	Super::PostLoad();
+
+	// Create collision object right after component was loaded
+	// so we can free cooked collision data
+	CreateCollisionObject();
+
+#if WITH_EDITOR
+	if (!HasAnyFlags(RF_ClassDefaultObject))
+	{
+		if (!CachedLocalBox.IsValid && CachedBoxSphereBounds_DEPRECATED.SphereRadius > 0)
+		{
+			ALandscapeProxy* LandscapeProxy = GetLandscapeProxy();
+			USceneComponent* LandscapeRoot = LandscapeProxy->GetRootComponent();
+			check(LandscapeRoot == AttachParent);
+
+			// Component isn't attached so we can't use its LocalToWorld
+			FTransform ComponentLtWTransform = FTransform(RelativeRotation, RelativeLocation, RelativeScale3D) * FTransform(LandscapeRoot->RelativeRotation, LandscapeRoot->RelativeLocation, LandscapeRoot->RelativeScale3D);
+
+			// This is good enough. The exact box will be calculated during painting.
+			CachedLocalBox = CachedBoxSphereBounds_DEPRECATED.GetBox().InverseTransformBy(ComponentLtWTransform);
+		}
+	}
+
+	// convert from deprecated layer names to direct LayerInfo references
+	ALandscapeProxy* LandscapeProxy = GetLandscapeProxy();
+	if (ensure(LandscapeProxy) && ComponentLayerInfos.Num() == 0 && ComponentLayers_DEPRECATED.Num() > 0)
+	{
+		ComponentLayerInfos.AddUninitialized(ComponentLayers_DEPRECATED.Num());
+
+		static const FName DataWeightmapName = FName("__DataLayer__");
+		for (int32 i = 0; i < ComponentLayers_DEPRECATED.Num(); i++)
+		{
+			if (ComponentLayers_DEPRECATED[i] != NAME_None)
+			{
+				if (ComponentLayers_DEPRECATED[i] == DataWeightmapName)
+				{
+					ComponentLayerInfos[i] = ALandscapeProxy::VisibilityLayer;
+				}
+				else
+				{
+					FLandscapeLayerStruct* Layer = LandscapeProxy->GetLayerInfo_Deprecated(ComponentLayers_DEPRECATED[i]);
+					ComponentLayerInfos[i] = Layer ? Layer->LayerInfoObj : NULL;
+				}
+			}
+		}
+
+		ComponentLayers_DEPRECATED.Empty();
+	}
+#endif//WITH_EDITOR
+}
+
+#if WITH_EDITOR
+void ULandscapeInfo::UpdateAllAddCollisions()
+{
+	for (auto It = XYtoComponentMap.CreateIterator(); It; ++It)
+	{
+		ULandscapeComponent* Comp = It.Value();
+		if (Comp)
+		{
+			ULandscapeHeightfieldCollisionComponent* CollisionComp = Comp->CollisionComponent.Get();
+			if (CollisionComp)
+			{
+				CollisionComp->UpdateAddCollisions();
+			}
+		}
+	}
+}
+
+void ULandscapeHeightfieldCollisionComponent::UpdateAddCollisions()
+{
+	ULandscapeInfo* Info = GetLandscapeInfo();
+	if (Info)
+	{
+		ALandscapeProxy* Proxy = GetLandscapeProxy();
+		FIntPoint ComponentBase = GetSectionBase() / Proxy->ComponentSizeQuads;
+
+		FIntPoint NeighborsKeys[8] =
+		{
+			ComponentBase + FIntPoint(-1, -1),
+			ComponentBase + FIntPoint(+0, -1),
+			ComponentBase + FIntPoint(+1, -1),
+			ComponentBase + FIntPoint(-1, +0),
+			ComponentBase + FIntPoint(+1, +0),
+			ComponentBase + FIntPoint(-1, +1),
+			ComponentBase + FIntPoint(+0, +1),
+			ComponentBase + FIntPoint(+1, +1)
+		};
+
+		// Search for Neighbors...
+		for (int32 i = 0; i < 8; ++i)
+		{
+			ULandscapeComponent* Comp = Info->XYtoComponentMap.FindRef(NeighborsKeys[i]);
+			if (!Comp || !Comp->CollisionComponent.IsValid())
+			{
+				Info->UpdateAddCollision(NeighborsKeys[i]);
+			}
+			else
+			{
+				Info->XYtoAddCollisionMap.Remove(NeighborsKeys[i]);
+			}
+		}
+	}
+}
+
+void ULandscapeInfo::UpdateAddCollision(FIntPoint LandscapeKey)
+{
+	FLandscapeAddCollision* AddCollision = XYtoAddCollisionMap.Find(LandscapeKey);
+	if (!AddCollision)
+	{
+		AddCollision = &XYtoAddCollisionMap.Add(LandscapeKey, FLandscapeAddCollision());
+	}
+
+	check(AddCollision);
+
+	// 8 Neighbors... 
+	// 0 1 2 
+	// 3   4
+	// 5 6 7
+	FIntPoint NeighborsKeys[8] =
+	{
+		LandscapeKey + FIntPoint(-1, -1),
+		LandscapeKey + FIntPoint(+0, -1),
+		LandscapeKey + FIntPoint(+1, -1),
+		LandscapeKey + FIntPoint(-1, +0),
+		LandscapeKey + FIntPoint(+1, +0),
+		LandscapeKey + FIntPoint(-1, +1),
+		LandscapeKey + FIntPoint(+0, +1),
+		LandscapeKey + FIntPoint(+1, +1)
+	};
+
+	ULandscapeHeightfieldCollisionComponent* NeighborCollisions[8];
+	// Search for Neighbors...
+	for (int32 i = 0; i < 8; ++i)
+	{
+		ULandscapeComponent* Comp = XYtoComponentMap.FindRef(NeighborsKeys[i]);
+		if (Comp)
+		{
+			NeighborCollisions[i] = Comp->CollisionComponent.Get();
+		}
+		else
+		{
+			NeighborCollisions[i] = NULL;
+		}
+	}
+
+	uint8 CornerSet = 0;
+	uint16 HeightCorner[4];
+
+	// Corner Cases...
+	if (NeighborCollisions[0])
+	{
+		uint16* Heights = (uint16*)NeighborCollisions[0]->CollisionHeightData.Lock(LOCK_READ_ONLY);
+		int32 CollisionSizeVerts = NeighborCollisions[0]->CollisionSizeQuads + 1;
+		HeightCorner[0] = Heights[CollisionSizeVerts - 1 + (CollisionSizeVerts - 1)*CollisionSizeVerts];
+		CornerSet |= 1;
+		NeighborCollisions[0]->CollisionHeightData.Unlock();
+	}
+	if (NeighborCollisions[2])
+	{
+		uint16* Heights = (uint16*)NeighborCollisions[2]->CollisionHeightData.Lock(LOCK_READ_ONLY);
+		int32 CollisionSizeVerts = NeighborCollisions[2]->CollisionSizeQuads + 1;
+		HeightCorner[1] = Heights[(CollisionSizeVerts - 1)*CollisionSizeVerts];
+		CornerSet |= 1 << 1;
+		NeighborCollisions[2]->CollisionHeightData.Unlock();
+	}
+	if (NeighborCollisions[5])
+	{
+		uint16* Heights = (uint16*)NeighborCollisions[5]->CollisionHeightData.Lock(LOCK_READ_ONLY);
+		int32 CollisionSizeVerts = NeighborCollisions[5]->CollisionSizeQuads + 1;
+		HeightCorner[2] = Heights[(CollisionSizeVerts - 1)];
+		CornerSet |= 1 << 2;
+		NeighborCollisions[5]->CollisionHeightData.Unlock();
+	}
+	if (NeighborCollisions[7])
+	{
+		uint16* Heights = (uint16*)NeighborCollisions[7]->CollisionHeightData.Lock(LOCK_READ_ONLY);
+		int32 CollisionSizeVerts = NeighborCollisions[7]->CollisionSizeQuads + 1;
+		HeightCorner[3] = Heights[0];
+		CornerSet |= 1 << 3;
+		NeighborCollisions[7]->CollisionHeightData.Unlock();
+	}
+
+	// Other cases...
+	if (NeighborCollisions[1])
+	{
+		uint16* Heights = (uint16*)NeighborCollisions[1]->CollisionHeightData.Lock(LOCK_READ_ONLY);
+		int32 CollisionSizeVerts = NeighborCollisions[1]->CollisionSizeQuads + 1;
+		HeightCorner[0] = Heights[(CollisionSizeVerts - 1)*CollisionSizeVerts];
+		CornerSet |= 1;
+		HeightCorner[1] = Heights[CollisionSizeVerts - 1 + (CollisionSizeVerts - 1)*CollisionSizeVerts];
+		CornerSet |= 1 << 1;
+		NeighborCollisions[1]->CollisionHeightData.Unlock();
+	}
+	if (NeighborCollisions[3])
+	{
+		uint16* Heights = (uint16*)NeighborCollisions[3]->CollisionHeightData.Lock(LOCK_READ_ONLY);
+		int32 CollisionSizeVerts = NeighborCollisions[3]->CollisionSizeQuads + 1;
+		HeightCorner[0] = Heights[(CollisionSizeVerts - 1)];
+		CornerSet |= 1;
+		HeightCorner[2] = Heights[CollisionSizeVerts - 1 + (CollisionSizeVerts - 1)*CollisionSizeVerts];
+		CornerSet |= 1 << 2;
+		NeighborCollisions[3]->CollisionHeightData.Unlock();
+	}
+	if (NeighborCollisions[4])
+	{
+		uint16* Heights = (uint16*)NeighborCollisions[4]->CollisionHeightData.Lock(LOCK_READ_ONLY);
+		int32 CollisionSizeVerts = NeighborCollisions[4]->CollisionSizeQuads + 1;
+		HeightCorner[1] = Heights[0];
+		CornerSet |= 1 << 1;
+		HeightCorner[3] = Heights[(CollisionSizeVerts - 1)*CollisionSizeVerts];
+		CornerSet |= 1 << 3;
+		NeighborCollisions[4]->CollisionHeightData.Unlock();
+	}
+	if (NeighborCollisions[6])
+	{
+		uint16* Heights = (uint16*)NeighborCollisions[6]->CollisionHeightData.Lock(LOCK_READ_ONLY);
+		int32 CollisionSizeVerts = NeighborCollisions[6]->CollisionSizeQuads + 1;
+		HeightCorner[2] = Heights[0];
+		CornerSet |= 1 << 2;
+		HeightCorner[3] = Heights[(CollisionSizeVerts - 1)];
+		CornerSet |= 1 << 3;
+		NeighborCollisions[6]->CollisionHeightData.Unlock();
+	}
+
+	// Fill unset values
+	// First iteration only for valid values distance 1 propagation
+	// Second iteration fills left ones...
+	FillCornerValues(CornerSet, HeightCorner);
+	//check(CornerSet == 15);
+
+	FIntPoint SectionBase = LandscapeKey*ComponentSizeQuads;
+
+	// Transform Height to Vectors...
+	FMatrix LtoW = GetLandscapeProxy()->LandscapeActorToWorld().ToMatrixWithScale();
+	AddCollision->Corners[0] = LtoW.TransformPosition(FVector(SectionBase.X, SectionBase.Y, LandscapeDataAccess::GetLocalHeight(HeightCorner[0])));
+	AddCollision->Corners[1] = LtoW.TransformPosition(FVector(SectionBase.X + ComponentSizeQuads, SectionBase.Y, LandscapeDataAccess::GetLocalHeight(HeightCorner[1])));
+	AddCollision->Corners[2] = LtoW.TransformPosition(FVector(SectionBase.X, SectionBase.Y + ComponentSizeQuads, LandscapeDataAccess::GetLocalHeight(HeightCorner[2])));
+	AddCollision->Corners[3] = LtoW.TransformPosition(FVector(SectionBase.X + ComponentSizeQuads, SectionBase.Y + ComponentSizeQuads, LandscapeDataAccess::GetLocalHeight(HeightCorner[3])));
+}
+
+void ULandscapeHeightfieldCollisionComponent::ExportCustomProperties(FOutputDevice& Out, uint32 Indent)
+{
+	if (HasAnyFlags(RF_ClassDefaultObject))
+	{
+		return;
+	}
+
+	uint16* Heights = (uint16*)CollisionHeightData.Lock(LOCK_READ_ONLY);
+	int32 NumHeights = FMath::Square(CollisionSizeQuads + 1);
+	check(CollisionHeightData.GetElementCount() == NumHeights);
+
+	Out.Logf(TEXT("%sCustomProperties CollisionHeightData "), FCString::Spc(Indent));
+	for (int32 i = 0; i < NumHeights; i++)
+	{
+		Out.Logf(TEXT("%d "), Heights[i]);
+	}
+
+	CollisionHeightData.Unlock();
+	Out.Logf(TEXT("\r\n"));
+
+	int32 NumDominantLayerSamples = DominantLayerData.GetElementCount();
+	check(NumDominantLayerSamples == 0 || NumDominantLayerSamples == NumHeights);
+
+	if (NumDominantLayerSamples > 0)
+	{
+		uint8* DominantLayerSamples = (uint8*)DominantLayerData.Lock(LOCK_READ_ONLY);
+
+		Out.Logf(TEXT("%sCustomProperties DominantLayerData "), FCString::Spc(Indent));
+		for (int32 i = 0; i < NumDominantLayerSamples; i++)
+		{
+			Out.Logf(TEXT("%02x"), DominantLayerSamples[i]);
+		}
+
+		DominantLayerData.Unlock();
+		Out.Logf(TEXT("\r\n"));
+	}
+}
+
+void ULandscapeHeightfieldCollisionComponent::ImportCustomProperties(const TCHAR* SourceText, FFeedbackContext* Warn)
+{
+	if (FParse::Command(&SourceText, TEXT("CollisionHeightData")))
+	{
+		int32 NumHeights = FMath::Square(CollisionSizeQuads + 1);
+
+		CollisionHeightData.Lock(LOCK_READ_WRITE);
+		uint16* Heights = (uint16*)CollisionHeightData.Realloc(NumHeights);
+		FMemory::Memzero(Heights, sizeof(uint16)*NumHeights);
+
+		FParse::Next(&SourceText);
+		int32 i = 0;
+		while (FChar::IsDigit(*SourceText))
+		{
+			if (i < NumHeights)
+			{
+				Heights[i++] = FCString::Atoi(SourceText);
+				while (FChar::IsDigit(*SourceText))
+				{
+					SourceText++;
+				}
+			}
+
+			FParse::Next(&SourceText);
+		}
+
+		CollisionHeightData.Unlock();
+
+		if (i != NumHeights)
+		{
+			Warn->Logf(*NSLOCTEXT("Core", "SyntaxError", "Syntax Error").ToString());
+		}
+	}
+	else if (FParse::Command(&SourceText, TEXT("DominantLayerData")))
+	{
+		int32 NumDominantLayerSamples = FMath::Square(CollisionSizeQuads + 1);
+
+		DominantLayerData.Lock(LOCK_READ_WRITE);
+		uint8* DominantLayerSamples = (uint8*)DominantLayerData.Realloc(NumDominantLayerSamples);
+		FMemory::Memzero(DominantLayerSamples, NumDominantLayerSamples);
+
+		FParse::Next(&SourceText);
+		int32 i = 0;
+		while (SourceText[0] && SourceText[1])
+		{
+			if (i < NumDominantLayerSamples)
+			{
+				DominantLayerSamples[i++] = FParse::HexDigit(SourceText[0]) * 16 + FParse::HexDigit(SourceText[1]);
+			}
+			SourceText += 2;
+		}
+
+		DominantLayerData.Unlock();
+
+		if (i != NumDominantLayerSamples)
+		{
+			Warn->Logf(*NSLOCTEXT("Core", "SyntaxError", "Syntax Error").ToString());
+		}
+	}
+}
+
+void ULandscapeMeshCollisionComponent::ExportCustomProperties(FOutputDevice& Out, uint32 Indent)
+{
+	if (HasAnyFlags(RF_ClassDefaultObject))
+	{
+		return;
+	}
+
+	Super::ExportCustomProperties(Out, Indent);
+
+	uint16* XYOffsets = (uint16*)CollisionXYOffsetData.Lock(LOCK_READ_ONLY);
+	int32 NumOffsets = FMath::Square(CollisionSizeQuads + 1) * 2;
+	check(CollisionXYOffsetData.GetElementCount() == NumOffsets);
+
+	Out.Logf(TEXT("%sCustomProperties CollisionXYOffsetData "), FCString::Spc(Indent));
+	for (int32 i = 0; i < NumOffsets; i++)
+	{
+		Out.Logf(TEXT("%d "), XYOffsets[i]);
+	}
+
+	CollisionXYOffsetData.Unlock();
+	Out.Logf(TEXT("\r\n"));
+}
+
+void ULandscapeMeshCollisionComponent::ImportCustomProperties(const TCHAR* SourceText, FFeedbackContext* Warn)
+{
+	if (FParse::Command(&SourceText, TEXT("CollisionHeightData")))
+	{
+		int32 NumHeights = FMath::Square(CollisionSizeQuads + 1);
+
+		CollisionHeightData.Lock(LOCK_READ_WRITE);
+		uint16* Heights = (uint16*)CollisionHeightData.Realloc(NumHeights);
+		FMemory::Memzero(Heights, sizeof(uint16)*NumHeights);
+
+		FParse::Next(&SourceText);
+		int32 i = 0;
+		while (FChar::IsDigit(*SourceText))
+		{
+			if (i < NumHeights)
+			{
+				Heights[i++] = FCString::Atoi(SourceText);
+				while (FChar::IsDigit(*SourceText))
+				{
+					SourceText++;
+				}
+			}
+
+			FParse::Next(&SourceText);
+		}
+
+		CollisionHeightData.Unlock();
+
+		if (i != NumHeights)
+		{
+			Warn->Logf(*NSLOCTEXT("Core", "SyntaxError", "Syntax Error").ToString());
+		}
+	}
+	else if (FParse::Command(&SourceText, TEXT("DominantLayerData")))
+	{
+		int32 NumDominantLayerSamples = FMath::Square(CollisionSizeQuads + 1);
+
+		DominantLayerData.Lock(LOCK_READ_WRITE);
+		uint8* DominantLayerSamples = (uint8*)DominantLayerData.Realloc(NumDominantLayerSamples);
+		FMemory::Memzero(DominantLayerSamples, NumDominantLayerSamples);
+
+		FParse::Next(&SourceText);
+		int32 i = 0;
+		while (SourceText[0] && SourceText[1])
+		{
+			if (i < NumDominantLayerSamples)
+			{
+				DominantLayerSamples[i++] = FParse::HexDigit(SourceText[0]) * 16 + FParse::HexDigit(SourceText[1]);
+			}
+			SourceText += 2;
+		}
+
+		DominantLayerData.Unlock();
+
+		if (i != NumDominantLayerSamples)
+		{
+			Warn->Logf(*NSLOCTEXT("Core", "SyntaxError", "Syntax Error").ToString());
+		}
+	}
+	else if (FParse::Command(&SourceText, TEXT("CollisionXYOffsetData")))
+	{
+		int32 NumOffsets = FMath::Square(CollisionSizeQuads + 1) * 2;
+
+		CollisionXYOffsetData.Lock(LOCK_READ_WRITE);
+		uint16* Offsets = (uint16*)CollisionXYOffsetData.Realloc(NumOffsets);
+		FMemory::Memzero(Offsets, sizeof(uint16)*NumOffsets);
+
+		FParse::Next(&SourceText);
+		int32 i = 0;
+		while (FChar::IsDigit(*SourceText))
+		{
+			if (i < NumOffsets)
+			{
+				Offsets[i++] = FCString::Atoi(SourceText);
+				while (FChar::IsDigit(*SourceText))
+				{
+					SourceText++;
+				}
+			}
+
+			FParse::Next(&SourceText);
+		}
+
+		CollisionXYOffsetData.Unlock();
+
+		if (i != NumOffsets)
+		{
+			Warn->Logf(*NSLOCTEXT("Core", "SyntaxError", "Syntax Error").ToString());
+		}
+	}
+}
+
+ULandscapeInfo* ULandscapeHeightfieldCollisionComponent::GetLandscapeInfo(bool bSpawnNewActor /*= true*/) const
+{
+	if (GetLandscapeProxy())
+	{
+		return GetLandscapeProxy()->GetLandscapeInfo(bSpawnNewActor);
+	}
+	return NULL;
+}
+
+#endif // WITH_EDITOR
+
+ALandscape* ULandscapeHeightfieldCollisionComponent::GetLandscapeActor() const
+{
+	ALandscapeProxy* Landscape = GetLandscapeProxy();
+	if (Landscape)
+	{
+		return Landscape->GetLandscapeActor();
+	}
+	return NULL;
+}
+
+ALandscapeProxy* ULandscapeHeightfieldCollisionComponent::GetLandscapeProxy() const
+{
+	return CastChecked<ALandscapeProxy>(GetOuter());
+}
+
+FIntPoint ULandscapeHeightfieldCollisionComponent::GetSectionBase() const
+{
+	return FIntPoint(SectionBaseX, SectionBaseY);
+}
+
+void ULandscapeHeightfieldCollisionComponent::SetSectionBase(FIntPoint InSectionBase)
+{
+	SectionBaseX = InSectionBase.X;
+	SectionBaseY = InSectionBase.Y;
+}
+
+ULandscapeHeightfieldCollisionComponent::ULandscapeHeightfieldCollisionComponent(const FObjectInitializer& ObjectInitializer)
+	: Super(ObjectInitializer)
+{
+	BodyInstance.bEnableCollision_DEPRECATED = true;
+	SetCollisionProfileName(UCollisionProfile::BlockAll_ProfileName);
+	bGenerateOverlapEvents = false;
+	CastShadow = false;
+	bUseAsOccluder = true;
+	bAllowCullDistanceVolume = false;
+	Mobility = EComponentMobility::Static;
+	bCanEverAffectNavigation = true;
+	bHasCustomNavigableGeometry = EHasCustomNavigableGeometry::Yes;
+}
diff --git a/Engine/Source/Runtime/Landscape/Private/LandscapeComponent.cpp b/Engine/Source/Runtime/Landscape/Private/LandscapeComponent.cpp
new file mode 100644
--- /dev/null
+++ b/Engine/Source/Runtime/Landscape/Private/LandscapeComponent.cpp
@@ -0,0 +1,14 @@
+// Copyright 1998-2014 Epic Games, Inc. All Rights Reserved.
+
+#include "Landscape.h"
+#include "LandscapeComponent.h"
+#include "LandscapeLayerInfoObject.h"
+
+FName FWeightmapLayerAllocationInfo::GetLayerName() const
+{
+	if (LayerInfo)
+	{
+		return LayerInfo->LayerName;
+	}
+	return NAME_None;
+}
\ No newline at end of file
