Landscape module from Engine split out.

diff --git a/Engine/Source/Runtime/Landscape/Private/LandscapeEdit.cpp b/Engine/Source/Runtime/Landscape/Private/LandscapeEdit.cpp
new file mode 100644
--- /dev/null
+++ b/Engine/Source/Runtime/Landscape/Private/LandscapeEdit.cpp
@@ -0,0 +1,5113 @@
+// Copyright 1998-2014 Epic Games, Inc. All Rights Reserved.
+
+/*=============================================================================
+LandscapeEdit.cpp: Landscape editing
+=============================================================================*/
+
+#include "Landscape.h"
+#include "Materials/MaterialInstanceDynamic.h"
+#include "Materials/MaterialExpressionLandscapeVisibilityMask.h"
+#include "Materials/MaterialExpressionLandscapeLayerWeight.h"
+#include "Materials/MaterialExpressionLandscapeLayerSample.h"
+#include "Materials/MaterialExpressionLandscapeLayerBlend.h"
+#include "Materials/MaterialExpressionLandscapeLayerSwitch.h"
+#include "LandscapeDataAccess.h"
+#include "LandscapeEdit.h"
+#include "LandscapeRender.h"
+#include "LandscapeRenderMobile.h"
+#include "Landscape.h"
+#include "LandscapeInfo.h"
+#include "LandscapeLayerInfoObject.h"
+#include "LandscapeMaterialInstanceConstant.h"
+#include "LandscapeHeightfieldCollisionComponent.h"
+#include "LandscapeMeshCollisionComponent.h"
+#include "LandscapeSplinesComponent.h"
+#include "LandscapeSplineControlPoint.h"
+#include "LandscapeSplineSegment.h"
+#include "LandscapeGizmoActiveActor.h"
+#include "Foliage/InstancedFoliageActor.h"
+#include "LevelUtils.h"
+#include "MessageLog.h"
+#include "MapErrors.h"
+#if WITH_EDITOR
+#include "RawMesh.h"
+#include "ScopedTransaction.h"
+#include "ImageWrapper.h"
+#endif
+#include "ComponentReregisterContext.h"
+
+DEFINE_LOG_CATEGORY(LogLandscape);
+
+#define LOCTEXT_NAMESPACE "Landscape"
+
+ULandscapeLayerInfoObject* ALandscapeProxy::VisibilityLayer = NULL;
+
+#if WITH_EDITOR
+
+void ULandscapeComponent::Init(int32 InBaseX, int32 InBaseY, int32 InComponentSizeQuads, int32 InNumSubsections, int32 InSubsectionSizeQuads)
+{
+	SetSectionBase(FIntPoint(InBaseX, InBaseY));
+	FVector RelativeLocation = FVector(GetSectionBase() - GetLandscapeProxy()->LandscapeSectionOffset);
+	SetRelativeLocation(RelativeLocation);
+	ComponentSizeQuads = InComponentSizeQuads;
+	NumSubsections = InNumSubsections;
+	SubsectionSizeQuads = InSubsectionSizeQuads;
+	check(NumSubsections * SubsectionSizeQuads == ComponentSizeQuads);
+	ULandscapeInfo* Info = GetLandscapeInfo();
+}
+
+void ULandscapeComponent::UpdateCachedBounds()
+{
+	FLandscapeComponentDataInterface CDI(this);
+
+	// Update local-space bounding box
+	CachedLocalBox.Init();
+	for (int32 y = 0; y < ComponentSizeQuads + 1; y++)
+	{
+		for (int32 x = 0; x < ComponentSizeQuads + 1; x++)
+		{
+			CachedLocalBox += CDI.GetLocalVertex(x, y);
+		}
+	}
+
+	// Update collision component bounds
+	ULandscapeHeightfieldCollisionComponent* HFCollisionComponent = CollisionComponent.Get();
+	if (HFCollisionComponent)
+	{
+		HFCollisionComponent->Modify();
+		HFCollisionComponent->CachedLocalBox = CachedLocalBox;
+		HFCollisionComponent->UpdateComponentToWorld();
+	}
+}
+
+void ULandscapeComponent::UpdateNavigationRelevance()
+{
+	ALandscapeProxy* Proxy = GetLandscapeProxy();
+	if (CollisionComponent && Proxy)
+	{
+		CollisionComponent->bCanEverAffectNavigation = Proxy->bUsedForNavigation;
+
+		UNavigationSystem::UpdateNavOctree(CollisionComponent.Get());
+	}
+}
+
+ULandscapeMaterialInstanceConstant* ALandscapeProxy::GetLayerThumbnailMIC(UMaterialInterface* LandscapeMaterial, FName LayerName, UTexture2D* ThumbnailWeightmap, UTexture2D* ThumbnailHeightmap, ALandscapeProxy* Proxy)
+{
+	if (!LandscapeMaterial)
+	{
+		LandscapeMaterial = Proxy ? Proxy->GetLandscapeMaterial() : UMaterial::GetDefaultMaterial(MD_Surface);
+	}
+
+	UMaterialInstanceConstant* CombinationMaterialInstance = NULL;
+	FString LayerKey = LandscapeMaterial->GetPathName() + FString::Printf(TEXT("_%s_0"), *LayerName.ToString());
+	if (Proxy)
+	{
+		CombinationMaterialInstance = Proxy->MaterialInstanceConstantMap.FindRef(*LayerKey);
+	}
+
+	if (CombinationMaterialInstance == NULL || CombinationMaterialInstance->Parent != LandscapeMaterial || (Proxy && Proxy->GetOutermost() != CombinationMaterialInstance->GetOutermost()))
+	{
+		FlushRenderingCommands();
+		UObject* MICOuter = GetTransientPackage();
+		if (Proxy)
+		{
+			MICOuter = Proxy->GetOutermost();
+		}
+		CombinationMaterialInstance = ConstructObject<ULandscapeMaterialInstanceConstant>(ULandscapeMaterialInstanceConstant::StaticClass(), MICOuter);
+		if (Proxy)
+		{
+			UE_LOG(LogLandscape, Log, TEXT("Looking for key %s, making new combination %s"), *LayerKey, *CombinationMaterialInstance->GetName());
+			Proxy->MaterialInstanceConstantMap.Add(*LayerKey, CombinationMaterialInstance);
+		}
+		CombinationMaterialInstance->SetParentEditorOnly(LandscapeMaterial);
+
+		FStaticParameterSet StaticParameters;
+		CombinationMaterialInstance->GetStaticParameterValues(StaticParameters);
+
+		for (int32 LayerParameterIdx = 0; LayerParameterIdx < StaticParameters.TerrainLayerWeightParameters.Num(); LayerParameterIdx++)
+		{
+			FStaticTerrainLayerWeightParameter& LayerParameter = StaticParameters.TerrainLayerWeightParameters[LayerParameterIdx];
+			if (LayerParameter.ParameterName == LayerName)
+			{
+				LayerParameter.WeightmapIndex = 0;
+				LayerParameter.bOverride = true;
+			}
+			else
+			{
+				LayerParameter.WeightmapIndex = INDEX_NONE;
+			}
+		}
+
+		CombinationMaterialInstance->UpdateStaticPermutation(StaticParameters);
+
+		CombinationMaterialInstance->PostEditChange();
+	}
+
+	// Create the instance for this component, that will use the layer combination instance.
+	ULandscapeMaterialInstanceConstant* MaterialInstance = ConstructObject<ULandscapeMaterialInstanceConstant>(ULandscapeMaterialInstanceConstant::StaticClass(), GetTransientPackage());
+	MaterialInstance->SetParentEditorOnly(CombinationMaterialInstance);
+	MaterialInstance->bIsLayerThumbnail = true;
+
+	FLinearColor Mask(1.f, 0.f, 0.f, 0.f);
+	MaterialInstance->SetVectorParameterValueEditorOnly(FName(*FString::Printf(TEXT("LayerMask_%s"), *LayerName.ToString())), Mask);
+	MaterialInstance->SetTextureParameterValueEditorOnly(FName(TEXT("Weightmap0")), ThumbnailWeightmap);
+	MaterialInstance->SetTextureParameterValueEditorOnly(FName(TEXT("Heightmap")), ThumbnailHeightmap);
+	MaterialInstance->PostEditChange();
+
+	return MaterialInstance;
+}
+
+UMaterialInstanceConstant* ULandscapeComponent::GetCombinationMaterial(bool bMobile /*= false*/)
+{
+	check(GIsEditor);
+
+	ALandscapeProxy* Proxy = GetLandscapeProxy();
+
+	UMaterialInterface* LandscapeMaterial = ComponentHasVisibilityPainted() ? GetLandscapeHoleMaterial() : GetLandscapeMaterial();
+
+	if (LandscapeMaterial != NULL)
+	{
+		// Ensure top level UMaterial has appropriate usage flags set.
+		bool bNeedsRecompile;
+		UMaterial* ParentUMaterial = LandscapeMaterial->GetMaterial();
+		if (ParentUMaterial && ParentUMaterial != UMaterial::GetDefaultMaterial(MD_Surface))
+		{
+			ParentUMaterial->SetMaterialUsage(bNeedsRecompile, MATUSAGE_Landscape);
+			ParentUMaterial->SetMaterialUsage(bNeedsRecompile, MATUSAGE_StaticLighting);
+		}
+
+		FString LayerKey = GetLayerAllocationKey(bMobile);
+		//UE_LOG(LogLandscape, Log, TEXT("Looking for key %s"), *LayerKey);
+
+		// Find or set a matching MIC in the Landscape's map.
+		UMaterialInstanceConstant* CombinationMaterialInstance = Proxy->MaterialInstanceConstantMap.FindRef(*LayerKey);
+		if (CombinationMaterialInstance == NULL || CombinationMaterialInstance->Parent != LandscapeMaterial || GetOutermost() != CombinationMaterialInstance->GetOutermost())
+		{
+			FlushRenderingCommands();
+
+			CombinationMaterialInstance = ConstructObject<ULandscapeMaterialInstanceConstant>(ULandscapeMaterialInstanceConstant::StaticClass(), GetOutermost());
+			UE_LOG(LogLandscape, Log, TEXT("Looking for key %s, making new combination %s"), *LayerKey, *CombinationMaterialInstance->GetName());
+			Proxy->MaterialInstanceConstantMap.Add(*LayerKey, CombinationMaterialInstance);
+			CombinationMaterialInstance->SetParentEditorOnly(LandscapeMaterial);
+
+			FStaticParameterSet StaticParameters;
+			CombinationMaterialInstance->GetStaticParameterValues(StaticParameters);
+
+			// Find weightmap mapping for each layer parameter, or disable if the layer is not used in this component.
+			for (int32 LayerParameterIdx = 0; LayerParameterIdx < StaticParameters.TerrainLayerWeightParameters.Num(); LayerParameterIdx++)
+			{
+				FStaticTerrainLayerWeightParameter& LayerParameter = StaticParameters.TerrainLayerWeightParameters[LayerParameterIdx];
+				LayerParameter.WeightmapIndex = INDEX_NONE;
+
+				// Look through our allocations to see if we need this layer.
+				// If not found, this component doesn't use the layer, and WeightmapIndex remains as INDEX_NONE.
+				for (int32 AllocIdx = 0; AllocIdx < WeightmapLayerAllocations.Num(); AllocIdx++)
+				{
+					FWeightmapLayerAllocationInfo& Allocation = WeightmapLayerAllocations[AllocIdx];
+					if (Allocation.LayerInfo != NULL)
+					{
+						FName ThisLayerName = (Allocation.LayerInfo == ALandscapeProxy::VisibilityLayer) ? UMaterialExpressionLandscapeVisibilityMask::ParameterName : Allocation.LayerInfo->LayerName;
+						if (ThisLayerName == LayerParameter.ParameterName)
+						{
+							LayerParameter.WeightmapIndex = Allocation.WeightmapTextureIndex;
+							LayerParameter.bOverride = true;
+							// UE_LOG(LogLandscape, Log, TEXT(" Layer %s channel %d"), *LayerParameter.ParameterName.ToString(), LayerParameter.WeightmapIndex);
+							break;
+						}
+					}
+				}
+			}
+
+			CombinationMaterialInstance->UpdateStaticPermutation(StaticParameters);
+
+			CombinationMaterialInstance->PostEditChange();
+		}
+
+		return CombinationMaterialInstance;
+	}
+	return NULL;
+}
+
+void ULandscapeComponent::UpdateMaterialInstances()
+{
+	check(GIsEditor);
+
+	// Find or set a matching MIC in the Landscape's map.
+	UMaterialInstanceConstant* CombinationMaterialInstance = GetCombinationMaterial(false);
+
+	if (CombinationMaterialInstance != NULL)
+	{
+		// Create the instance for this component, that will use the layer combination instance.
+		if (MaterialInstance == NULL || GetOutermost() != MaterialInstance->GetOutermost())
+		{
+			MaterialInstance = ConstructObject<ULandscapeMaterialInstanceConstant>(ULandscapeMaterialInstanceConstant::StaticClass(), GetOutermost());
+		}
+
+		// For undo
+		MaterialInstance->SetFlags(RF_Transactional);
+		MaterialInstance->Modify();
+
+		MaterialInstance->SetParentEditorOnly(CombinationMaterialInstance);
+
+		FLinearColor Masks[4];
+		Masks[0] = FLinearColor(1.f, 0.f, 0.f, 0.f);
+		Masks[1] = FLinearColor(0.f, 1.f, 0.f, 0.f);
+		Masks[2] = FLinearColor(0.f, 0.f, 1.f, 0.f);
+		Masks[3] = FLinearColor(0.f, 0.f, 0.f, 1.f);
+
+		// Set the layer mask
+		for (int32 AllocIdx = 0; AllocIdx < WeightmapLayerAllocations.Num(); AllocIdx++)
+		{
+			FWeightmapLayerAllocationInfo& Allocation = WeightmapLayerAllocations[AllocIdx];
+
+			FName LayerName = Allocation.LayerInfo == ALandscapeProxy::VisibilityLayer ? UMaterialExpressionLandscapeVisibilityMask::ParameterName : Allocation.LayerInfo ? Allocation.LayerInfo->LayerName : NAME_None;
+			MaterialInstance->SetVectorParameterValueEditorOnly(FName(*FString::Printf(TEXT("LayerMask_%s"), *LayerName.ToString())), Masks[Allocation.WeightmapTextureChannel]);
+		}
+
+		// Set the weightmaps
+		for (int32 i = 0; i < WeightmapTextures.Num(); i++)
+		{
+			// UE_LOG(LogLandscape, Log, TEXT("Setting Weightmap%d = %s"), i, *WeightmapTextures(i)->GetName());
+			MaterialInstance->SetTextureParameterValueEditorOnly(FName(*FString::Printf(TEXT("Weightmap%d"), i)), WeightmapTextures[i]);
+		}
+		// Set the heightmap, if needed.
+
+		if (HeightmapTexture)
+		{
+			MaterialInstance->SetTextureParameterValueEditorOnly(FName(TEXT("Heightmap")), HeightmapTexture);
+		}
+		MaterialInstance->PostEditChange();
+
+		// Recreate the render state, needed to update the static drawlist which has cached the MaterialRenderProxy.
+		RecreateRenderState_Concurrent();
+	}
+}
+
+int32 ULandscapeComponent::GetNumMaterials() const
+{
+	return 1;
+}
+
+class UMaterialInterface* ULandscapeComponent::GetMaterial(int32 ElementIndex) const
+{
+	if (ensure(ElementIndex == 0))
+	{
+		return GetLandscapeMaterial();
+	}
+	else
+	{
+		return NULL;
+	}
+}
+
+void ULandscapeComponent::SetMaterial(int32 ElementIndex, class UMaterialInterface* Material)
+{
+	if (ensure(ElementIndex == 0))
+	{
+		GetLandscapeProxy()->LandscapeMaterial = Material;
+	}
+}
+
+void ULandscapeComponent::PostEditUndo()
+{
+	Super::PostEditUndo();
+
+	UpdateMaterialInstances();
+
+	if (EditToolRenderData)
+	{
+		EditToolRenderData->UpdateDebugColorMaterial();
+		EditToolRenderData->UpdateSelectionMaterial(EditToolRenderData->SelectedType);
+	}
+}
+
+void ULandscapeComponent::FixupWeightmaps()
+{
+	if (GIsEditor && !HasAnyFlags(RF_ClassDefaultObject))
+	{
+		ULandscapeInfo* Info = GetLandscapeInfo();
+		ALandscapeProxy* Proxy = GetLandscapeProxy();
+
+		if (Info)
+		{
+			TArray<ULandscapeLayerInfoObject*> DeletedLayers;
+			bool bFixedLayerDeletion = false;
+
+			if (Info->Layers.Num() && Cast<ALandscape>(Proxy))
+			{
+				// LayerName Validation check...
+				for (int32 LayerIdx = 0; LayerIdx < WeightmapLayerAllocations.Num(); LayerIdx++)
+				{
+					if (!WeightmapLayerAllocations[LayerIdx].LayerInfo
+						|| (WeightmapLayerAllocations[LayerIdx].LayerInfo != ALandscapeProxy::VisibilityLayer && Info->GetLayerInfoIndex(WeightmapLayerAllocations[LayerIdx].LayerInfo) == INDEX_NONE))
+					{
+						if (!bFixedLayerDeletion)
+						{
+							FFormatNamedArguments Arguments;
+							Arguments.Add(TEXT("LandscapeName"), FText::FromString(GetName()));
+							FMessageLog("MapCheck").Warning()
+								->AddToken(FTextToken::Create(FText::Format(LOCTEXT("MapCheck_Message_FixedUpDeletedLayerWeightmap", "{LandscapeName} : Fixed up deleted layer weightmap"), Arguments)))
+								->AddToken(FMapErrorToken::Create(FMapErrors::FixedUpDeletedLayerWeightmap));
+						}
+
+						bFixedLayerDeletion = true;
+						DeletedLayers.Add(WeightmapLayerAllocations[LayerIdx].LayerInfo);
+					}
+				}
+			}
+
+			if (bFixedLayerDeletion)
+			{
+				FLandscapeEditDataInterface LandscapeEdit(Info);
+				for (int32 Idx = 0; Idx < DeletedLayers.Num(); ++Idx)
+				{
+					DeleteLayer(DeletedLayers[Idx], &LandscapeEdit);
+				}
+			}
+
+			bool bFixedWeightmapTextureIndex = false;
+
+			// Store the weightmap allocations in WeightmapUsageMap
+			for (int32 LayerIdx = 0; LayerIdx < WeightmapLayerAllocations.Num(); LayerIdx++)
+			{
+				FWeightmapLayerAllocationInfo& Allocation = WeightmapLayerAllocations[LayerIdx];
+
+				// Fix up any problems caused by the layer deletion bug.
+				if (Allocation.WeightmapTextureIndex >= WeightmapTextures.Num())
+				{
+					Allocation.WeightmapTextureIndex = WeightmapTextures.Num() - 1;
+					if (!bFixedWeightmapTextureIndex)
+					{
+						FFormatNamedArguments Arguments;
+						Arguments.Add(TEXT("LandscapeName"), FText::FromString(GetName()));
+						FMessageLog("MapCheck").Warning()
+							->AddToken(FTextToken::Create(FText::Format(LOCTEXT("MapCheck_Message_FixedUpIncorrectLayerWeightmap", "{LandscapeName} : Fixed up incorrect layer weightmap texture index"), Arguments)))
+							->AddToken(FMapErrorToken::Create(FMapErrors::FixedUpIncorrectLayerWeightmap));
+					}
+					bFixedWeightmapTextureIndex = true;
+				}
+
+				UTexture2D* WeightmapTexture = WeightmapTextures[Allocation.WeightmapTextureIndex];
+				FLandscapeWeightmapUsage& Usage = Proxy->WeightmapUsageMap.FindOrAdd(WeightmapTexture);
+
+				// Detect a shared layer allocation, caused by a previous undo or layer deletion bugs
+				if (Usage.ChannelUsage[Allocation.WeightmapTextureChannel] != NULL &&
+					Usage.ChannelUsage[Allocation.WeightmapTextureChannel] != this)
+				{
+					FFormatNamedArguments Arguments;
+					Arguments.Add(TEXT("LayerName"), FText::FromString(Allocation.GetLayerName().ToString()));
+					Arguments.Add(TEXT("LandscapeName"), FText::FromString(GetName()));
+					Arguments.Add(TEXT("ChannelName"), FText::FromString(Usage.ChannelUsage[Allocation.WeightmapTextureChannel]->GetName()));
+					FMessageLog("MapCheck").Warning()
+						->AddToken(FTextToken::Create(FText::Format(LOCTEXT("MapCheck_Message_FixedUpSharedLayerWeightmap", "Fixed up shared weightmap texture for layer {LayerName} in component '{LandscapeName}' (shares with '{ChannelName}')"), Arguments)))
+						->AddToken(FMapErrorToken::Create(FMapErrors::FixedUpSharedLayerWeightmap));
+					WeightmapLayerAllocations.RemoveAt(LayerIdx);
+					LayerIdx--;
+					continue;
+				}
+				else
+				{
+					Usage.ChannelUsage[Allocation.WeightmapTextureChannel] = this;
+				}
+			}
+
+			RemoveInvalidWeightmaps();
+
+			// Store the layer combination in the MaterialInstanceConstantMap
+			if (MaterialInstance != NULL)
+			{
+				UMaterialInstanceConstant* CombinationMaterialInstance = Cast<UMaterialInstanceConstant>(MaterialInstance->Parent);
+				if (CombinationMaterialInstance)
+				{
+					Proxy->MaterialInstanceConstantMap.Add(*GetLayerAllocationKey(), CombinationMaterialInstance);
+				}
+			}
+		}
+	}
+}
+
+//
+// LandscapeComponentAlphaInfo
+//
+struct FLandscapeComponentAlphaInfo
+{
+	int32 LayerIndex;
+	TArray<uint8> AlphaValues;
+
+	// tor
+	FLandscapeComponentAlphaInfo(ULandscapeComponent* InOwner, int32 InLayerIndex)
+		: LayerIndex(InLayerIndex)
+	{
+		AlphaValues.Empty(FMath::Square(InOwner->ComponentSizeQuads + 1));
+		AlphaValues.AddZeroed(FMath::Square(InOwner->ComponentSizeQuads + 1));
+	}
+
+	bool IsLayerAllZero() const
+	{
+		for (int32 Index = 0; Index < AlphaValues.Num(); Index++)
+		{
+			if (AlphaValues[Index] != 0)
+			{
+				return false;
+			}
+		}
+		return true;
+	}
+};
+
+
+void ULandscapeComponent::UpdateCollisionHeightData(const FColor* HeightmapTextureMipData, int32 ComponentX1, int32 ComponentY1, int32 ComponentX2, int32 ComponentY2, bool bUpdateBounds, const FColor* XYOffsetmapTextureData, bool bRebuild)
+{
+	ULandscapeInfo* Info = GetLandscapeInfo();
+	ALandscapeProxy* Proxy = GetLandscapeProxy();
+	FIntPoint ComponentKey = GetSectionBase() / ComponentSizeQuads;
+	ULandscapeHeightfieldCollisionComponent* CollisionComp = CollisionComponent.Get();
+	ULandscapeMeshCollisionComponent* MeshCollisionComponent = Cast<ULandscapeMeshCollisionComponent>(CollisionComp);
+
+	ULandscapeHeightfieldCollisionComponent* OldCollisionComponent = CollisionComp;
+
+	ALandscapeProxy* CollisionProxy = NULL;
+	if (CollisionComp && bRebuild)
+	{
+		// Remove existing component
+		CollisionProxy = CollisionComp->GetLandscapeProxy();
+		if (CollisionProxy)
+		{
+			CollisionComp->DestroyComponent();
+			CollisionComp = NULL;
+		}
+	}
+
+	int32 CollisionSubsectionSizeVerts = ((SubsectionSizeQuads + 1) >> CollisionMipLevel);
+	int32 CollisionSubsectionSizeQuads = CollisionSubsectionSizeVerts - 1;
+	int32 CollisionSizeVerts = NumSubsections*CollisionSubsectionSizeQuads + 1;
+
+	uint16* CollisionHeightData = NULL;
+	uint16* CollisionXYOffsetData = NULL;
+	bool CreatedNew = false;
+	bool ChangeType = false;
+	TArray<uint8> DominantLayerData;
+	TArray<ULandscapeLayerInfoObject*> LayerInfos;
+
+	if (CollisionComp)
+	{
+		CollisionComp->Modify();
+	}
+
+	// Existing collision component is same type with collision
+	if (CollisionComp && ((XYOffsetmapTexture == NULL) == (MeshCollisionComponent == NULL)))
+	{
+		if (bUpdateBounds)
+		{
+			CollisionComp->CachedLocalBox = CachedLocalBox;
+			CollisionComp->UpdateComponentToWorld();
+		}
+
+		CollisionHeightData = (uint16*)CollisionComp->CollisionHeightData.Lock(LOCK_READ_WRITE);
+
+		if (XYOffsetmapTexture && MeshCollisionComponent)
+		{
+			CollisionXYOffsetData = (uint16*)MeshCollisionComponent->CollisionXYOffsetData.Lock(LOCK_READ_WRITE);
+		}
+	}
+	else
+	{
+		ComponentX1 = 0;
+		ComponentY1 = 0;
+		ComponentX2 = MAX_int32;
+		ComponentY2 = MAX_int32;
+
+		if (CollisionComp) // remove old component before changing to other type collision...
+		{
+			ChangeType = true;
+
+			if (CollisionComp->DominantLayerData.GetElementCount())
+			{
+				DominantLayerData.AddUninitialized(FMath::Square(CollisionSizeVerts));
+
+				const uint8* SrcDominantLayerData = (uint8*)CollisionComp->DominantLayerData.Lock(LOCK_READ_ONLY);
+				FMemory::Memcpy(DominantLayerData.GetData(), SrcDominantLayerData, sizeof(uint8)*FMath::Square(CollisionSizeVerts));
+				CollisionComp->DominantLayerData.Unlock();
+			}
+
+			if (CollisionComp->ComponentLayerInfos.Num())
+			{
+				LayerInfos = CollisionComp->ComponentLayerInfos;
+			}
+
+			if (Info)
+			{
+				Info->Modify();
+			}
+			Proxy->Modify();
+			CollisionComp->DestroyComponent();
+			CollisionComp = NULL;
+		}
+
+		MeshCollisionComponent = XYOffsetmapTexture ? ConstructObject<ULandscapeMeshCollisionComponent>(ULandscapeMeshCollisionComponent::StaticClass(), Proxy, NAME_None, RF_Transactional) : NULL;
+		CollisionComp = XYOffsetmapTexture ? Cast<ULandscapeHeightfieldCollisionComponent>(MeshCollisionComponent) :
+			ConstructObject<ULandscapeHeightfieldCollisionComponent>(ULandscapeHeightfieldCollisionComponent::StaticClass(), Proxy, NAME_None, RF_Transactional);
+
+		CollisionComp->SetRelativeLocation(RelativeLocation);
+		CollisionComp->AttachTo(Proxy->GetRootComponent(), NAME_None);
+		Proxy->CollisionComponents.Add(CollisionComp);
+
+		CollisionComp->RenderComponent = this;
+		CollisionComp->SetSectionBase(GetSectionBase());
+		CollisionComp->CollisionSizeQuads = CollisionSubsectionSizeQuads * NumSubsections;
+		CollisionComp->CollisionScale = (float)(ComponentSizeQuads) / (float)(CollisionComp->CollisionSizeQuads);
+		CollisionComp->CachedLocalBox = CachedLocalBox;
+		CreatedNew = true;
+
+		// Reallocate raw collision data
+		CollisionComp->CollisionHeightData.Lock(LOCK_READ_WRITE);
+		CollisionHeightData = (uint16*)CollisionComp->CollisionHeightData.Realloc(FMath::Square(CollisionSizeVerts));
+		FMemory::Memzero(CollisionHeightData, sizeof(uint16)*FMath::Square(CollisionSizeVerts));
+
+		if (XYOffsetmapTexture && MeshCollisionComponent)
+		{
+			// Need XYOffsetData for Collision Component
+			MeshCollisionComponent->CollisionXYOffsetData.Lock(LOCK_READ_WRITE);
+			CollisionXYOffsetData = (uint16*)MeshCollisionComponent->CollisionXYOffsetData.Realloc(FMath::Square(CollisionSizeVerts) * 2);
+			FMemory::Memzero(CollisionXYOffsetData, sizeof(uint16)*FMath::Square(CollisionSizeVerts) * 2);
+
+			if (DominantLayerData.Num())
+			{
+				MeshCollisionComponent->DominantLayerData.Lock(LOCK_READ_WRITE);
+				uint8* DestDominantLayerData = (uint8*)MeshCollisionComponent->DominantLayerData.Realloc(FMath::Square(CollisionSizeVerts));
+				FMemory::Memcpy(DestDominantLayerData, DominantLayerData.GetData(), sizeof(uint8)*FMath::Square(CollisionSizeVerts));
+				MeshCollisionComponent->DominantLayerData.Unlock();
+			}
+
+			if (LayerInfos.Num())
+			{
+				MeshCollisionComponent->ComponentLayerInfos = LayerInfos;
+			}
+		}
+	}
+
+	int32 HeightmapSizeU = HeightmapTexture->Source.GetSizeX();
+	int32 HeightmapSizeV = HeightmapTexture->Source.GetSizeY();
+	int32 MipSizeU = HeightmapSizeU >> CollisionMipLevel;
+	int32 MipSizeV = HeightmapSizeV >> CollisionMipLevel;
+
+	int32 XYMipSizeU = XYOffsetmapTexture ? XYOffsetmapTexture->Source.GetSizeX() >> CollisionMipLevel : 0;
+	int32 XYMipSizeV = XYOffsetmapTexture ? XYOffsetmapTexture->Source.GetSizeY() >> CollisionMipLevel : 0;
+
+	// Ratio to convert update region coordinate to collision mip coordinates
+	float CollisionQuadRatio = (float)CollisionSubsectionSizeQuads / (float)SubsectionSizeQuads;
+
+	// XY offset into heightmap mip data
+	int32 HeightmapOffsetX = FMath::RoundToInt(HeightmapScaleBias.Z * (float)HeightmapSizeU) >> CollisionMipLevel;
+	int32 HeightmapOffsetY = FMath::RoundToInt(HeightmapScaleBias.W * (float)HeightmapSizeV) >> CollisionMipLevel;
+
+	//int32 WeightmapOffsetX = FMath::RoundToInt(WeightmapScaleBias.Z * (float)XYMipSizeU) >> CollisionMipLevel;
+	//int32 WeightmapOffsetY = FMath::RoundToInt(WeightmapScaleBias.W * (float)XYMipSizeV) >> CollisionMipLevel;
+
+	for (int32 SubsectionY = 0; SubsectionY < NumSubsections; SubsectionY++)
+	{
+		// Check if subsection is fully above or below the area we are interested in
+		if ((ComponentY2 < SubsectionSizeQuads*SubsectionY) ||		// above
+			(ComponentY1 > SubsectionSizeQuads*(SubsectionY + 1)))	// below
+		{
+			continue;
+		}
+
+		for (int32 SubsectionX = 0; SubsectionX < NumSubsections; SubsectionX++)
+		{
+			// Check if subsection is fully to the left or right of the area we are interested in
+			if ((ComponentX2 < SubsectionSizeQuads*SubsectionX) ||		// left
+				(ComponentX1 > SubsectionSizeQuads*(SubsectionX + 1)))	// right
+			{
+				continue;
+			}
+
+			// Area to update in subsection coordinates
+			int32 SubX1 = ComponentX1 - SubsectionSizeQuads*SubsectionX;
+			int32 SubY1 = ComponentY1 - SubsectionSizeQuads*SubsectionY;
+			int32 SubX2 = ComponentX2 - SubsectionSizeQuads*SubsectionX;
+			int32 SubY2 = ComponentY2 - SubsectionSizeQuads*SubsectionY;
+
+			// Area to update in collision mip level coords
+			int32 CollisionSubX1 = FMath::FloorToInt((float)SubX1 * CollisionQuadRatio);
+			int32 CollisionSubY1 = FMath::FloorToInt((float)SubY1 * CollisionQuadRatio);
+			int32 CollisionSubX2 = FMath::CeilToInt((float)SubX2 * CollisionQuadRatio);
+			int32 CollisionSubY2 = FMath::CeilToInt((float)SubY2 * CollisionQuadRatio);
+
+			// Clamp area to update
+			int32 VertX1 = FMath::Clamp<int32>(CollisionSubX1, 0, CollisionSubsectionSizeQuads);
+			int32 VertY1 = FMath::Clamp<int32>(CollisionSubY1, 0, CollisionSubsectionSizeQuads);
+			int32 VertX2 = FMath::Clamp<int32>(CollisionSubX2, 0, CollisionSubsectionSizeQuads);
+			int32 VertY2 = FMath::Clamp<int32>(CollisionSubY2, 0, CollisionSubsectionSizeQuads);
+
+			for (int32 VertY = VertY1; VertY <= VertY2; VertY++)
+			{
+				for (int32 VertX = VertX1; VertX <= VertX2; VertX++)
+				{
+					{
+						// X/Y of the vertex we're looking indexed into the texture data
+						int32 TexX = HeightmapOffsetX + CollisionSubsectionSizeVerts * SubsectionX + VertX;
+						int32 TexY = HeightmapOffsetY + CollisionSubsectionSizeVerts * SubsectionY + VertY;
+						const FColor& TexData = HeightmapTextureMipData[TexX + TexY * MipSizeU];
+
+						// this uses Quads as we don't want the duplicated vertices
+						int32 CompVertX = CollisionSubsectionSizeQuads * SubsectionX + VertX;
+						int32 CompVertY = CollisionSubsectionSizeQuads * SubsectionY + VertY;
+
+						// Copy collision data
+						uint16& CollisionHeight = CollisionHeightData[CompVertX + CompVertY*CollisionSizeVerts];
+						uint16 NewHeight = TexData.R << 8 | TexData.G;
+
+						CollisionHeight = NewHeight;
+					}
+
+					if (XYOffsetmapTexture && XYOffsetmapTextureData && CollisionXYOffsetData)
+					{
+						// X/Y of the vertex we're looking indexed into the texture data
+						int32 TexX = CollisionSubsectionSizeVerts * SubsectionX + VertX;
+						int32 TexY = CollisionSubsectionSizeVerts * SubsectionY + VertY;
+						const FColor& TexData = XYOffsetmapTextureData[TexX + TexY * XYMipSizeU];
+
+						// this uses Quads as we don't want the duplicated vertices
+						int32 CompVertX = CollisionSubsectionSizeQuads * SubsectionX + VertX;
+						int32 CompVertY = CollisionSubsectionSizeQuads * SubsectionY + VertY;
+
+						// Copy collision data
+						uint16 NewXOffset = TexData.R << 8 | TexData.G;
+						uint16 NewYOffset = TexData.B << 8 | TexData.A;
+
+						int32 XYIndex = CompVertX + CompVertY*CollisionSizeVerts;
+						CollisionXYOffsetData[XYIndex * 2] = NewXOffset;
+						CollisionXYOffsetData[XYIndex * 2 + 1] = NewYOffset;
+					}
+				}
+			}
+		}
+	}
+
+	CollisionComp->CollisionHeightData.Unlock();
+
+	if (XYOffsetmapTexture && MeshCollisionComponent)
+	{
+		MeshCollisionComponent->CollisionXYOffsetData.Unlock();
+	}
+
+	// If we updated an existing component, we need to update the PhysX copy of the data
+	if (!CreatedNew)
+	{
+		if (MeshCollisionComponent)
+		{
+			// Will be done once for XY Offset data update in FXYOffsetmapAccessor() destructor with UpdateCachedBounds()
+			//MeshCollisionComponent->RecreateCollision(false);
+		}
+		else if (CollisionMipLevel == 0)
+		{
+			CollisionComp->UpdateHeightfieldRegion(ComponentX1, ComponentY1, ComponentX2, ComponentY2);
+		}
+		else
+		{
+			int32 CollisionCompX1 = FMath::FloorToInt((float)ComponentX1 * CollisionQuadRatio);
+			int32 CollisionCompY1 = FMath::FloorToInt((float)ComponentY1 * CollisionQuadRatio);
+			int32 CollisionCompX2 = FMath::CeilToInt((float)ComponentX2 * CollisionQuadRatio);
+			int32 CollisionCompY2 = FMath::CeilToInt((float)ComponentY2 * CollisionQuadRatio);
+			CollisionComp->UpdateHeightfieldRegion(CollisionCompX1, CollisionCompY1, CollisionCompX2, CollisionCompY2);
+		}
+	}
+
+	{
+		// set relevancy for navigation system
+		ALandscapeProxy* LandscapeProxy = CollisionComp->GetLandscapeProxy();
+		CollisionComp->bCanEverAffectNavigation = LandscapeProxy ? LandscapeProxy->bUsedForNavigation : false;
+	}
+
+	if (bRebuild && CollisionProxy)
+	{
+		CollisionProxy->RegisterAllComponents();
+
+		// Move any foliage instances if we created a new collision component.
+		if (OldCollisionComponent && OldCollisionComponent != CollisionComp)
+		{
+			AInstancedFoliageActor* IFA = AInstancedFoliageActor::GetInstancedFoliageActorForCurrentLevel(OldCollisionComponent->GetWorld());
+			IFA->MoveInstancesToNewComponent(OldCollisionComponent, CollisionComp);
+		}
+	}
+
+	// Set new collision component to pointer
+	CollisionComponent = CollisionComp;
+
+	if (ChangeType || CreatedNew)
+	{
+		Proxy->RegisterAllComponents();
+	}
+}
+
+
+void ULandscapeComponent::UpdateCollisionLayerData(TArray<FColor*>& WeightmapTextureMipData, int32 ComponentX1, int32 ComponentY1, int32 ComponentX2, int32 ComponentY2)
+{
+	ULandscapeInfo* Info = GetLandscapeInfo();
+	ALandscapeProxy* Proxy = GetLandscapeProxy();
+	FIntPoint ComponentKey = GetSectionBase() / ComponentSizeQuads;
+
+	if (CollisionComponent)
+	{
+		CollisionComponent->Modify();
+
+		TArray<ULandscapeLayerInfoObject*> CandidateLayers;
+		TArray<uint8*> CandidateDataPtrs;
+
+		// Channel remapping
+		int32 ChannelOffsets[4] = { (int32)STRUCT_OFFSET(FColor, R), (int32)STRUCT_OFFSET(FColor, G), (int32)STRUCT_OFFSET(FColor, B), (int32)STRUCT_OFFSET(FColor, A) };
+
+		bool bExistingLayerMismatch = false;
+		int32 DataLayerIdx = INDEX_NONE;
+
+		// Find the layers we're interested in
+		for (int32 AllocIdx = 0; AllocIdx < WeightmapLayerAllocations.Num(); AllocIdx++)
+		{
+			FWeightmapLayerAllocationInfo& AllocInfo = WeightmapLayerAllocations[AllocIdx];
+			ULandscapeLayerInfoObject* LayerInfo = AllocInfo.LayerInfo;
+			if (AllocInfo.LayerInfo == ALandscapeProxy::VisibilityLayer || (LayerInfo && LayerInfo->PhysMaterial))
+			{
+				int32 Idx = CandidateLayers.Add(AllocInfo.LayerInfo);
+				CandidateDataPtrs.Add(((uint8*)WeightmapTextureMipData[AllocInfo.WeightmapTextureIndex]) + ChannelOffsets[AllocInfo.WeightmapTextureChannel]);
+
+				// Check if we still match the collision component.
+				if (!CollisionComponent->ComponentLayerInfos.IsValidIndex(Idx) || CollisionComponent->ComponentLayerInfos[Idx] != AllocInfo.LayerInfo)
+				{
+					bExistingLayerMismatch = true;
+				}
+
+				if (AllocInfo.LayerInfo == ALandscapeProxy::VisibilityLayer)
+				{
+					DataLayerIdx = Idx;
+					bExistingLayerMismatch = true; // always rebuild whole component for hole
+				}
+			}
+		}
+
+		if (CandidateLayers.Num() == 0)
+		{
+			// No layers, so don't update any weights
+			CollisionComponent->DominantLayerData.RemoveBulkData();
+			CollisionComponent->ComponentLayerInfos.Empty();
+		}
+		else
+		{
+			int32 CollisionSubsectionSizeVerts = ((SubsectionSizeQuads + 1) >> CollisionMipLevel);
+			int32 CollisionSubsectionSizeQuads = CollisionSubsectionSizeVerts - 1;
+			int32 CollisionSizeVerts = NumSubsections*CollisionSubsectionSizeQuads + 1;
+			uint8* DominantLayerData = NULL;
+
+			// If there's no existing data, or the layer allocations have changed, we need to update the data for the whole component.
+			if (bExistingLayerMismatch || CollisionComponent->DominantLayerData.GetElementCount() == 0)
+			{
+				ComponentX1 = 0;
+				ComponentY1 = 0;
+				ComponentX2 = MAX_int32;
+				ComponentY2 = MAX_int32;
+
+				CollisionComponent->DominantLayerData.Lock(LOCK_READ_WRITE);
+				DominantLayerData = (uint8*)CollisionComponent->DominantLayerData.Realloc(FMath::Square(CollisionSizeVerts));
+				FMemory::Memzero(DominantLayerData, FMath::Square(CollisionSizeVerts));
+
+				CollisionComponent->ComponentLayerInfos = CandidateLayers;
+			}
+			else
+			{
+				DominantLayerData = (uint8*)CollisionComponent->DominantLayerData.Lock(LOCK_READ_WRITE);
+			}
+
+			int32 MipSizeU = (WeightmapTextures[0]->Source.GetSizeX()) >> CollisionMipLevel;
+
+			// Ratio to convert update region coordinate to collision mip coordinates
+			float CollisionQuadRatio = (float)CollisionSubsectionSizeQuads / (float)SubsectionSizeQuads;
+
+			for (int32 SubsectionY = 0; SubsectionY < NumSubsections; SubsectionY++)
+			{
+				// Check if subsection is fully above or below the area we are interested in
+				if ((ComponentY2 < SubsectionSizeQuads*SubsectionY) ||		// above
+					(ComponentY1 > SubsectionSizeQuads*(SubsectionY + 1)))	// below
+				{
+					continue;
+				}
+
+				for (int32 SubsectionX = 0; SubsectionX < NumSubsections; SubsectionX++)
+				{
+					// Check if subsection is fully to the left or right of the area we are interested in
+					if ((ComponentX2 < SubsectionSizeQuads*SubsectionX) ||		// left
+						(ComponentX1 > SubsectionSizeQuads*(SubsectionX + 1)))	// right
+					{
+						continue;
+					}
+
+					// Area to update in subsection coordinates
+					int32 SubX1 = ComponentX1 - SubsectionSizeQuads*SubsectionX;
+					int32 SubY1 = ComponentY1 - SubsectionSizeQuads*SubsectionY;
+					int32 SubX2 = ComponentX2 - SubsectionSizeQuads*SubsectionX;
+					int32 SubY2 = ComponentY2 - SubsectionSizeQuads*SubsectionY;
+
+					// Area to update in collision mip level coords
+					int32 CollisionSubX1 = FMath::FloorToInt((float)SubX1 * CollisionQuadRatio);
+					int32 CollisionSubY1 = FMath::FloorToInt((float)SubY1 * CollisionQuadRatio);
+					int32 CollisionSubX2 = FMath::CeilToInt((float)SubX2 * CollisionQuadRatio);
+					int32 CollisionSubY2 = FMath::CeilToInt((float)SubY2 * CollisionQuadRatio);
+
+					// Clamp area to update
+					int32 VertX1 = FMath::Clamp<int32>(CollisionSubX1, 0, CollisionSubsectionSizeQuads);
+					int32 VertY1 = FMath::Clamp<int32>(CollisionSubY1, 0, CollisionSubsectionSizeQuads);
+					int32 VertX2 = FMath::Clamp<int32>(CollisionSubX2, 0, CollisionSubsectionSizeQuads);
+					int32 VertY2 = FMath::Clamp<int32>(CollisionSubY2, 0, CollisionSubsectionSizeQuads);
+
+					for (int32 VertY = VertY1; VertY <= VertY2; VertY++)
+					{
+						for (int32 VertX = VertX1; VertX <= VertX2; VertX++)
+						{
+							// X/Y of the vertex we're looking indexed into the texture data
+							int32 TexX = CollisionSubsectionSizeVerts * SubsectionX + VertX;
+							int32 TexY = CollisionSubsectionSizeVerts * SubsectionY + VertY;
+							int32 DataOffset = (TexX + TexY * MipSizeU) * sizeof(FColor);
+
+							uint8 DominantLayer = 255; // 255 as invalid value
+							int32 DominantWeight = 0;
+							for (int32 LayerIdx = 0; LayerIdx < CandidateDataPtrs.Num(); LayerIdx++)
+							{
+								const uint8 LayerWeight = CandidateDataPtrs[LayerIdx][DataOffset];
+
+								if (LayerIdx == DataLayerIdx) // Override value for hole
+								{
+									if (LayerWeight > 170) // 255 * 0.66...
+									{
+										DominantLayer = LayerIdx;
+										DominantWeight = INT_MAX;
+									}
+								}
+								else if (LayerWeight > DominantWeight)
+								{
+									DominantLayer = LayerIdx;
+									DominantWeight = LayerWeight;
+								}
+							}
+
+							// this uses Quads as we don't want the duplicated vertices
+							int32 CompVertX = CollisionSubsectionSizeQuads * SubsectionX + VertX;
+							int32 CompVertY = CollisionSubsectionSizeQuads * SubsectionY + VertY;
+
+							// Set collision data
+							DominantLayerData[CompVertX + CompVertY*CollisionSizeVerts] = DominantLayer;
+						}
+					}
+				}
+			}
+			CollisionComponent->DominantLayerData.Unlock();
+		}
+
+		// We do not force an update of the physics data here. We don't need the layer information in the editor and it
+		// causes problems if we update it multiple times in a single frame.
+	}
+}
+
+
+void ULandscapeComponent::UpdateCollisionLayerData()
+{
+	// Generate the dominant layer data
+	TArray<FColor*> WeightmapTextureMipData;
+	TArray<TArray<uint8> > CachedWeightmapTextureMipData;
+
+	WeightmapTextureMipData.Empty(WeightmapTextures.Num());
+	CachedWeightmapTextureMipData.Empty(WeightmapTextures.Num());
+	for (int32 Idx = 0; Idx < WeightmapTextures.Num(); Idx++)
+	{
+		TArray<uint8>* MipData = new(CachedWeightmapTextureMipData)TArray<uint8>();
+		WeightmapTextures[Idx]->Source.GetMipData(*MipData, CollisionMipLevel);
+		WeightmapTextureMipData.Add((FColor*)MipData->GetData());
+	}
+
+	UpdateCollisionLayerData(WeightmapTextureMipData);
+}
+
+
+
+
+void ULandscapeComponent::GenerateHeightmapMips(TArray<FColor*>& HeightmapTextureMipData, int32 ComponentX1/*=0*/, int32 ComponentY1/*=0*/, int32 ComponentX2/*=MAX_int32*/, int32 ComponentY2/*=MAX_int32*/, struct FLandscapeTextureDataInfo* TextureDataInfo/*=NULL*/)
+{
+	bool EndX = false;
+	bool EndY = false;
+
+	if (ComponentX1 == MAX_int32)
+	{
+		EndX = true;
+		ComponentX1 = 0;
+	}
+
+	if (ComponentY1 == MAX_int32)
+	{
+		EndY = true;
+		ComponentY1 = 0;
+	}
+
+	if (ComponentX2 == MAX_int32)
+	{
+		ComponentX2 = ComponentSizeQuads;
+	}
+	if (ComponentY2 == MAX_int32)
+	{
+		ComponentY2 = ComponentSizeQuads;
+	}
+
+	int32 HeightmapSizeU = HeightmapTexture->Source.GetSizeX();
+	int32 HeightmapSizeV = HeightmapTexture->Source.GetSizeY();
+
+	int32 HeightmapOffsetX = FMath::RoundToInt(HeightmapScaleBias.Z * (float)HeightmapSizeU);
+	int32 HeightmapOffsetY = FMath::RoundToInt(HeightmapScaleBias.W * (float)HeightmapSizeV);
+
+	for (int32 SubsectionY = 0; SubsectionY < NumSubsections; SubsectionY++)
+	{
+		// Check if subsection is fully above or below the area we are interested in
+		if ((ComponentY2 < SubsectionSizeQuads*SubsectionY) ||		// above
+			(ComponentY1 > SubsectionSizeQuads*(SubsectionY + 1)))	// below
+		{
+			continue;
+		}
+
+		for (int32 SubsectionX = 0; SubsectionX < NumSubsections; SubsectionX++)
+		{
+			// Check if subsection is fully to the left or right of the area we are interested in
+			if ((ComponentX2 < SubsectionSizeQuads*SubsectionX) ||		// left
+				(ComponentX1 > SubsectionSizeQuads*(SubsectionX + 1)))	// right
+			{
+				continue;
+			}
+
+			// Area to update in previous mip level coords
+			int32 PrevMipSubX1 = ComponentX1 - SubsectionSizeQuads*SubsectionX;
+			int32 PrevMipSubY1 = ComponentY1 - SubsectionSizeQuads*SubsectionY;
+			int32 PrevMipSubX2 = ComponentX2 - SubsectionSizeQuads*SubsectionX;
+			int32 PrevMipSubY2 = ComponentY2 - SubsectionSizeQuads*SubsectionY;
+
+			int32 PrevMipSubsectionSizeQuads = SubsectionSizeQuads;
+			float InvPrevMipSubsectionSizeQuads = 1.f / (float)SubsectionSizeQuads;
+
+			int32 PrevMipSizeU = HeightmapSizeU;
+			int32 PrevMipSizeV = HeightmapSizeV;
+
+			int32 PrevMipHeightmapOffsetX = HeightmapOffsetX;
+			int32 PrevMipHeightmapOffsetY = HeightmapOffsetY;
+
+			for (int32 Mip = 1; Mip < HeightmapTextureMipData.Num(); Mip++)
+			{
+				int32 MipSizeU = HeightmapSizeU >> Mip;
+				int32 MipSizeV = HeightmapSizeV >> Mip;
+
+				int32 MipSubsectionSizeQuads = ((SubsectionSizeQuads + 1) >> Mip) - 1;
+				float InvMipSubsectionSizeQuads = 1.f / (float)MipSubsectionSizeQuads;
+
+				int32 MipHeightmapOffsetX = HeightmapOffsetX >> Mip;
+				int32 MipHeightmapOffsetY = HeightmapOffsetY >> Mip;
+
+				// Area to update in current mip level coords
+				int32 MipSubX1 = FMath::FloorToInt((float)MipSubsectionSizeQuads * (float)PrevMipSubX1 * InvPrevMipSubsectionSizeQuads);
+				int32 MipSubY1 = FMath::FloorToInt((float)MipSubsectionSizeQuads * (float)PrevMipSubY1 * InvPrevMipSubsectionSizeQuads);
+				int32 MipSubX2 = FMath::CeilToInt((float)MipSubsectionSizeQuads * (float)PrevMipSubX2 * InvPrevMipSubsectionSizeQuads);
+				int32 MipSubY2 = FMath::CeilToInt((float)MipSubsectionSizeQuads * (float)PrevMipSubY2 * InvPrevMipSubsectionSizeQuads);
+
+				// Clamp area to update
+				int32 VertX1 = FMath::Clamp<int32>(MipSubX1, 0, MipSubsectionSizeQuads);
+				int32 VertY1 = FMath::Clamp<int32>(MipSubY1, 0, MipSubsectionSizeQuads);
+				int32 VertX2 = FMath::Clamp<int32>(MipSubX2, 0, MipSubsectionSizeQuads);
+				int32 VertY2 = FMath::Clamp<int32>(MipSubY2, 0, MipSubsectionSizeQuads);
+
+				for (int32 VertY = VertY1; VertY <= VertY2; VertY++)
+				{
+					for (int32 VertX = VertX1; VertX <= VertX2; VertX++)
+					{
+						// Convert VertX/Y into previous mip's coords
+						float PrevMipVertX = (float)PrevMipSubsectionSizeQuads * (float)VertX * InvMipSubsectionSizeQuads;
+						float PrevMipVertY = (float)PrevMipSubsectionSizeQuads * (float)VertY * InvMipSubsectionSizeQuads;
+
+#if 0
+						// Validate that the vertex we skip wouldn't use the updated data in the parent mip.
+						// Note this validation is doesn't do anything unless you change the VertY/VertX loops 
+						// above to process all verts from 0 .. MipSubsectionSizeQuads.
+						if (VertX < VertX1 || VertX > VertX2)
+						{
+							check(FMath::CeilToInt(PrevMipVertX) < PrevMipSubX1 || FMath::FloorToInt(PrevMipVertX) > PrevMipSubX2);
+							continue;
+						}
+
+						if (VertY < VertY1 || VertY > VertY2)
+						{
+							check(FMath::CeilToInt(PrevMipVertY) < PrevMipSubY1 || FMath::FloorToInt(PrevMipVertY) > PrevMipSubY2);
+							continue;
+						}
+#endif
+
+						// X/Y of the vertex we're looking indexed into the texture data
+						int32 TexX = (MipHeightmapOffsetX)+(MipSubsectionSizeQuads + 1) * SubsectionX + VertX;
+						int32 TexY = (MipHeightmapOffsetY)+(MipSubsectionSizeQuads + 1) * SubsectionY + VertY;
+
+						float fPrevMipTexX = (float)(PrevMipHeightmapOffsetX)+(float)((PrevMipSubsectionSizeQuads + 1) * SubsectionX) + PrevMipVertX;
+						float fPrevMipTexY = (float)(PrevMipHeightmapOffsetY)+(float)((PrevMipSubsectionSizeQuads + 1) * SubsectionY) + PrevMipVertY;
+
+						int32 PrevMipTexX = FMath::FloorToInt(fPrevMipTexX);
+						float fPrevMipTexFracX = FMath::Fractional(fPrevMipTexX);
+						int32 PrevMipTexY = FMath::FloorToInt(fPrevMipTexY);
+						float fPrevMipTexFracY = FMath::Fractional(fPrevMipTexY);
+
+						checkSlow(TexX >= 0 && TexX < MipSizeU);
+						checkSlow(TexY >= 0 && TexY < MipSizeV);
+						checkSlow(PrevMipTexX >= 0 && PrevMipTexX < PrevMipSizeU);
+						checkSlow(PrevMipTexY >= 0 && PrevMipTexY < PrevMipSizeV);
+
+						int32 PrevMipTexX1 = FMath::Min<int32>(PrevMipTexX + 1, PrevMipSizeU - 1);
+						int32 PrevMipTexY1 = FMath::Min<int32>(PrevMipTexY + 1, PrevMipSizeV - 1);
+
+						// Padding for missing data for MIP 0
+						if (Mip == 1)
+						{
+							if (EndX && SubsectionX == NumSubsections - 1 && VertX == VertX2)
+							{
+								for (int32 PaddingIdx = PrevMipTexX + PrevMipTexY * PrevMipSizeU; PaddingIdx + 1 < PrevMipTexY1 * PrevMipSizeU; ++PaddingIdx)
+								{
+									HeightmapTextureMipData[Mip - 1][PaddingIdx + 1] = HeightmapTextureMipData[Mip - 1][PaddingIdx];
+								}
+							}
+
+							if (EndY && SubsectionX == NumSubsections - 1 && SubsectionY == NumSubsections - 1 && VertY == VertY2 && VertX == VertX2)
+							{
+								for (int32 PaddingYIdx = PrevMipTexY; PaddingYIdx + 1 < PrevMipSizeV; ++PaddingYIdx)
+								{
+									for (int32 PaddingXIdx = 0; PaddingXIdx < PrevMipSizeU; ++PaddingXIdx)
+									{
+										HeightmapTextureMipData[Mip - 1][PaddingXIdx + (PaddingYIdx + 1) * PrevMipSizeU] = HeightmapTextureMipData[Mip - 1][PaddingXIdx + PaddingYIdx * PrevMipSizeU];
+									}
+								}
+							}
+						}
+
+						FColor* TexData = &(HeightmapTextureMipData[Mip])[TexX + TexY * MipSizeU];
+						FColor *PreMipTexData00 = &(HeightmapTextureMipData[Mip - 1])[PrevMipTexX + PrevMipTexY * PrevMipSizeU];
+						FColor *PreMipTexData01 = &(HeightmapTextureMipData[Mip - 1])[PrevMipTexX + PrevMipTexY1 * PrevMipSizeU];
+						FColor *PreMipTexData10 = &(HeightmapTextureMipData[Mip - 1])[PrevMipTexX1 + PrevMipTexY * PrevMipSizeU];
+						FColor *PreMipTexData11 = &(HeightmapTextureMipData[Mip - 1])[PrevMipTexX1 + PrevMipTexY1 * PrevMipSizeU];
+
+						// Lerp height values
+						uint16 PrevMipHeightValue00 = PreMipTexData00->R << 8 | PreMipTexData00->G;
+						uint16 PrevMipHeightValue01 = PreMipTexData01->R << 8 | PreMipTexData01->G;
+						uint16 PrevMipHeightValue10 = PreMipTexData10->R << 8 | PreMipTexData10->G;
+						uint16 PrevMipHeightValue11 = PreMipTexData11->R << 8 | PreMipTexData11->G;
+						uint16 HeightValue = FMath::RoundToInt(
+							FMath::Lerp(
+							FMath::Lerp((float)PrevMipHeightValue00, (float)PrevMipHeightValue10, fPrevMipTexFracX),
+							FMath::Lerp((float)PrevMipHeightValue01, (float)PrevMipHeightValue11, fPrevMipTexFracX),
+							fPrevMipTexFracY));
+
+						TexData->R = HeightValue >> 8;
+						TexData->G = HeightValue & 255;
+
+						// Lerp tangents
+						TexData->B = FMath::RoundToInt(
+							FMath::Lerp(
+							FMath::Lerp((float)PreMipTexData00->B, (float)PreMipTexData10->B, fPrevMipTexFracX),
+							FMath::Lerp((float)PreMipTexData01->B, (float)PreMipTexData11->B, fPrevMipTexFracX),
+							fPrevMipTexFracY));
+
+						TexData->A = FMath::RoundToInt(
+							FMath::Lerp(
+							FMath::Lerp((float)PreMipTexData00->A, (float)PreMipTexData10->A, fPrevMipTexFracX),
+							FMath::Lerp((float)PreMipTexData01->A, (float)PreMipTexData11->A, fPrevMipTexFracX),
+							fPrevMipTexFracY));
+
+						// Padding for missing data
+						if (EndX && SubsectionX == NumSubsections - 1 && VertX == VertX2)
+						{
+							for (int32 PaddingIdx = TexX + TexY * MipSizeU; PaddingIdx + 1 < (TexY + 1) * MipSizeU; ++PaddingIdx)
+							{
+								HeightmapTextureMipData[Mip][PaddingIdx + 1] = HeightmapTextureMipData[Mip][PaddingIdx];
+							}
+						}
+
+						if (EndY && SubsectionX == NumSubsections - 1 && SubsectionY == NumSubsections - 1 && VertY == VertY2 && VertX == VertX2)
+						{
+							for (int32 PaddingYIdx = TexY; PaddingYIdx + 1 < MipSizeV; ++PaddingYIdx)
+							{
+								for (int32 PaddingXIdx = 0; PaddingXIdx < MipSizeU; ++PaddingXIdx)
+								{
+									HeightmapTextureMipData[Mip][PaddingXIdx + (PaddingYIdx + 1) * MipSizeU] = HeightmapTextureMipData[Mip][PaddingXIdx + PaddingYIdx * MipSizeU];
+								}
+							}
+						}
+
+					}
+				}
+
+				// Record the areas we updated
+				if (TextureDataInfo)
+				{
+					int32 TexX1 = (MipHeightmapOffsetX)+(MipSubsectionSizeQuads + 1) * SubsectionX + VertX1;
+					int32 TexY1 = (MipHeightmapOffsetY)+(MipSubsectionSizeQuads + 1) * SubsectionY + VertY1;
+					int32 TexX2 = (MipHeightmapOffsetX)+(MipSubsectionSizeQuads + 1) * SubsectionX + VertX2;
+					int32 TexY2 = (MipHeightmapOffsetY)+(MipSubsectionSizeQuads + 1) * SubsectionY + VertY2;
+					TextureDataInfo->AddMipUpdateRegion(Mip, TexX1, TexY1, TexX2, TexY2);
+				}
+
+				// Copy current mip values to prev as we move to the next mip.
+				PrevMipSubsectionSizeQuads = MipSubsectionSizeQuads;
+				InvPrevMipSubsectionSizeQuads = InvMipSubsectionSizeQuads;
+
+				PrevMipSizeU = MipSizeU;
+				PrevMipSizeV = MipSizeV;
+
+				PrevMipHeightmapOffsetX = MipHeightmapOffsetX;
+				PrevMipHeightmapOffsetY = MipHeightmapOffsetY;
+
+				// Use this mip's area as we move to the next mip
+				PrevMipSubX1 = MipSubX1;
+				PrevMipSubY1 = MipSubY1;
+				PrevMipSubX2 = MipSubX2;
+				PrevMipSubY2 = MipSubY2;
+			}
+		}
+	}
+}
+
+void ULandscapeComponent::CreateEmptyTextureMips(UTexture2D* Texture, bool bClear /*= false*/)
+{
+	ETextureSourceFormat WeightmapFormat = Texture->Source.GetFormat();
+	int32 WeightmapSizeU = Texture->Source.GetSizeX();
+	int32 WeightmapSizeV = Texture->Source.GetSizeY();
+
+	if (bClear)
+	{
+		Texture->Source.Init2DWithMipChain(WeightmapSizeU, WeightmapSizeV, WeightmapFormat);
+		int32 NumMips = Texture->Source.GetNumMips();
+		for (int32 MipIndex = 0; MipIndex < NumMips; ++MipIndex)
+		{
+			uint8* MipData = Texture->Source.LockMip(MipIndex);
+			FMemory::Memzero(MipData, Texture->Source.CalcMipSize(MipIndex));
+			Texture->Source.UnlockMip(MipIndex);
+		}
+	}
+	else
+	{
+		TArray<uint8> TopMipData;
+		Texture->Source.GetMipData(TopMipData, 0);
+		Texture->Source.Init2DWithMipChain(WeightmapSizeU, WeightmapSizeV, WeightmapFormat);
+		int32 NumMips = Texture->Source.GetNumMips();
+		uint8* MipData = Texture->Source.LockMip(0);
+		FMemory::Memcpy(MipData, TopMipData.GetData(), TopMipData.Num());
+		Texture->Source.UnlockMip(0);
+	}
+}
+
+template<typename DataType>
+void ULandscapeComponent::GenerateMipsTempl(int32 InNumSubsections, int32 InSubsectionSizeQuads, UTexture2D* Texture, DataType* BaseMipData)
+{
+	// Stores pointers to the locked mip data
+	TArray<DataType*> MipData;
+	MipData.Add(BaseMipData);
+	for (int32 MipIndex = 1; MipIndex < Texture->Source.GetNumMips(); ++MipIndex)
+	{
+		MipData.Add((DataType*)Texture->Source.LockMip(MipIndex));
+	}
+
+	// Update the newly created mips
+	UpdateMipsTempl<DataType>(InNumSubsections, InSubsectionSizeQuads, Texture, MipData);
+
+	// Unlock all the new mips, but not the base mip's data
+	for (int32 i = 1; i < MipData.Num(); i++)
+	{
+		Texture->Source.UnlockMip(i);
+	}
+}
+
+void ULandscapeComponent::GenerateWeightmapMips(int32 InNumSubsections, int32 InSubsectionSizeQuads, UTexture2D* WeightmapTexture, FColor* BaseMipData)
+{
+	GenerateMipsTempl<FColor>(InNumSubsections, InSubsectionSizeQuads, WeightmapTexture, BaseMipData);
+}
+
+namespace
+{
+	template<typename DataType>
+	void BiLerpTextureData(DataType* Output, const DataType* Data00, const DataType* Data10, const DataType* Data01, const DataType* Data11, float FracX, float FracY)
+	{
+		*Output = FMath::RoundToInt(
+			FMath::Lerp(
+			FMath::Lerp((float)*Data00, (float)*Data10, FracX),
+			FMath::Lerp((float)*Data01, (float)*Data11, FracX),
+			FracY));
+	}
+
+	template<>
+	void BiLerpTextureData(FColor* Output, const FColor* Data00, const FColor* Data10, const FColor* Data01, const FColor* Data11, float FracX, float FracY)
+	{
+		Output->R = FMath::RoundToInt(
+			FMath::Lerp(
+			FMath::Lerp((float)Data00->R, (float)Data10->R, FracX),
+			FMath::Lerp((float)Data01->R, (float)Data11->R, FracX),
+			FracY));
+		Output->G = FMath::RoundToInt(
+			FMath::Lerp(
+			FMath::Lerp((float)Data00->G, (float)Data10->G, FracX),
+			FMath::Lerp((float)Data01->G, (float)Data11->G, FracX),
+			FracY));
+		Output->B = FMath::RoundToInt(
+			FMath::Lerp(
+			FMath::Lerp((float)Data00->B, (float)Data10->B, FracX),
+			FMath::Lerp((float)Data01->B, (float)Data11->B, FracX),
+			FracY));
+		Output->A = FMath::RoundToInt(
+			FMath::Lerp(
+			FMath::Lerp((float)Data00->A, (float)Data10->A, FracX),
+			FMath::Lerp((float)Data01->A, (float)Data11->A, FracX),
+			FracY));
+	}
+
+	template<typename DataType>
+	void AverageTexData(DataType* Output, const DataType* Data00, const DataType* Data10, const DataType* Data01, const DataType* Data11)
+	{
+		*Output = (((int32)(*Data00) + (int32)(*Data10) + (int32)(*Data01) + (int32)(*Data11)) >> 2);
+	}
+
+	template<>
+	void AverageTexData(FColor* Output, const FColor* Data00, const FColor* Data10, const FColor* Data01, const FColor* Data11)
+	{
+		Output->R = (((int32)Data00->R + (int32)Data10->R + (int32)Data01->R + (int32)Data11->R) >> 2);
+		Output->G = (((int32)Data00->G + (int32)Data10->G + (int32)Data01->G + (int32)Data11->G) >> 2);
+		Output->B = (((int32)Data00->B + (int32)Data10->B + (int32)Data01->B + (int32)Data11->B) >> 2);
+		Output->A = (((int32)Data00->A + (int32)Data10->A + (int32)Data01->A + (int32)Data11->A) >> 2);
+	}
+
+};
+
+template<typename DataType>
+void ULandscapeComponent::UpdateMipsTempl(int32 InNumSubsections, int32 InSubsectionSizeQuads, UTexture2D* Texture, TArray<DataType*>& TextureMipData, int32 ComponentX1/*=0*/, int32 ComponentY1/*=0*/, int32 ComponentX2/*=MAX_int32*/, int32 ComponentY2/*=MAX_int32*/, struct FLandscapeTextureDataInfo* TextureDataInfo/*=NULL*/)
+{
+	int32 WeightmapSizeU = Texture->Source.GetSizeX();
+	int32 WeightmapSizeV = Texture->Source.GetSizeY();
+
+	// Find the maximum mip where each texel's data comes from just one subsection.
+	int32 MaxWholeSubsectionMip = FMath::FloorLog2(InSubsectionSizeQuads + 1) - 1;
+
+	// Update the mip where each texel's data comes from just one subsection.
+	for (int32 SubsectionY = 0; SubsectionY < InNumSubsections; SubsectionY++)
+	{
+		// Check if subsection is fully above or below the area we are interested in
+		if ((ComponentY2 < InSubsectionSizeQuads*SubsectionY) ||	// above
+			(ComponentY1 > InSubsectionSizeQuads*(SubsectionY + 1)))	// below
+		{
+			continue;
+		}
+
+		for (int32 SubsectionX = 0; SubsectionX < InNumSubsections; SubsectionX++)
+		{
+			// Check if subsection is fully to the left or right of the area we are interested in
+			if ((ComponentX2 < InSubsectionSizeQuads*SubsectionX) ||	// left
+				(ComponentX1 > InSubsectionSizeQuads*(SubsectionX + 1)))	// right
+			{
+				continue;
+			}
+
+			// Area to update in previous mip level coords
+			int32 PrevMipSubX1 = ComponentX1 - InSubsectionSizeQuads*SubsectionX;
+			int32 PrevMipSubY1 = ComponentY1 - InSubsectionSizeQuads*SubsectionY;
+			int32 PrevMipSubX2 = ComponentX2 - InSubsectionSizeQuads*SubsectionX;
+			int32 PrevMipSubY2 = ComponentY2 - InSubsectionSizeQuads*SubsectionY;
+
+			int32 PrevMipSubsectionSizeQuads = InSubsectionSizeQuads;
+			float InvPrevMipSubsectionSizeQuads = 1.f / (float)InSubsectionSizeQuads;
+
+			int32 PrevMipSizeU = WeightmapSizeU;
+			int32 PrevMipSizeV = WeightmapSizeV;
+
+			for (int32 Mip = 1; Mip <= MaxWholeSubsectionMip; Mip++)
+			{
+				int32 MipSizeU = WeightmapSizeU >> Mip;
+				int32 MipSizeV = WeightmapSizeV >> Mip;
+
+				int32 MipSubsectionSizeQuads = ((InSubsectionSizeQuads + 1) >> Mip) - 1;
+				float InvMipSubsectionSizeQuads = 1.f / (float)MipSubsectionSizeQuads;
+
+				// Area to update in current mip level coords
+				int32 MipSubX1 = FMath::FloorToInt((float)MipSubsectionSizeQuads * (float)PrevMipSubX1 * InvPrevMipSubsectionSizeQuads);
+				int32 MipSubY1 = FMath::FloorToInt((float)MipSubsectionSizeQuads * (float)PrevMipSubY1 * InvPrevMipSubsectionSizeQuads);
+				int32 MipSubX2 = FMath::CeilToInt((float)MipSubsectionSizeQuads * (float)PrevMipSubX2 * InvPrevMipSubsectionSizeQuads);
+				int32 MipSubY2 = FMath::CeilToInt((float)MipSubsectionSizeQuads * (float)PrevMipSubY2 * InvPrevMipSubsectionSizeQuads);
+
+				// Clamp area to update
+				int32 VertX1 = FMath::Clamp<int32>(MipSubX1, 0, MipSubsectionSizeQuads);
+				int32 VertY1 = FMath::Clamp<int32>(MipSubY1, 0, MipSubsectionSizeQuads);
+				int32 VertX2 = FMath::Clamp<int32>(MipSubX2, 0, MipSubsectionSizeQuads);
+				int32 VertY2 = FMath::Clamp<int32>(MipSubY2, 0, MipSubsectionSizeQuads);
+
+				for (int32 VertY = VertY1; VertY <= VertY2; VertY++)
+				{
+					for (int32 VertX = VertX1; VertX <= VertX2; VertX++)
+					{
+						// Convert VertX/Y into previous mip's coords
+						float PrevMipVertX = (float)PrevMipSubsectionSizeQuads * (float)VertX * InvMipSubsectionSizeQuads;
+						float PrevMipVertY = (float)PrevMipSubsectionSizeQuads * (float)VertY * InvMipSubsectionSizeQuads;
+
+						// X/Y of the vertex we're looking indexed into the texture data
+						int32 TexX = (MipSubsectionSizeQuads + 1) * SubsectionX + VertX;
+						int32 TexY = (MipSubsectionSizeQuads + 1) * SubsectionY + VertY;
+
+						float fPrevMipTexX = (float)((PrevMipSubsectionSizeQuads + 1) * SubsectionX) + PrevMipVertX;
+						float fPrevMipTexY = (float)((PrevMipSubsectionSizeQuads + 1) * SubsectionY) + PrevMipVertY;
+
+						int32 PrevMipTexX = FMath::FloorToInt(fPrevMipTexX);
+						float fPrevMipTexFracX = FMath::Fractional(fPrevMipTexX);
+						int32 PrevMipTexY = FMath::FloorToInt(fPrevMipTexY);
+						float fPrevMipTexFracY = FMath::Fractional(fPrevMipTexY);
+
+						check(TexX >= 0 && TexX < MipSizeU);
+						check(TexY >= 0 && TexY < MipSizeV);
+						check(PrevMipTexX >= 0 && PrevMipTexX < PrevMipSizeU);
+						check(PrevMipTexY >= 0 && PrevMipTexY < PrevMipSizeV);
+
+						int32 PrevMipTexX1 = FMath::Min<int32>(PrevMipTexX + 1, PrevMipSizeU - 1);
+						int32 PrevMipTexY1 = FMath::Min<int32>(PrevMipTexY + 1, PrevMipSizeV - 1);
+
+						DataType* TexData = &(TextureMipData[Mip])[TexX + TexY * MipSizeU];
+						DataType *PreMipTexData00 = &(TextureMipData[Mip - 1])[PrevMipTexX + PrevMipTexY * PrevMipSizeU];
+						DataType *PreMipTexData01 = &(TextureMipData[Mip - 1])[PrevMipTexX + PrevMipTexY1 * PrevMipSizeU];
+						DataType *PreMipTexData10 = &(TextureMipData[Mip - 1])[PrevMipTexX1 + PrevMipTexY * PrevMipSizeU];
+						DataType *PreMipTexData11 = &(TextureMipData[Mip - 1])[PrevMipTexX1 + PrevMipTexY1 * PrevMipSizeU];
+
+						// Lerp weightmap data
+						BiLerpTextureData<DataType>(TexData, PreMipTexData00, PreMipTexData10, PreMipTexData01, PreMipTexData11, fPrevMipTexFracX, fPrevMipTexFracY);
+					}
+				}
+
+				// Record the areas we updated
+				if (TextureDataInfo)
+				{
+					int32 TexX1 = (MipSubsectionSizeQuads + 1) * SubsectionX + VertX1;
+					int32 TexY1 = (MipSubsectionSizeQuads + 1) * SubsectionY + VertY1;
+					int32 TexX2 = (MipSubsectionSizeQuads + 1) * SubsectionX + VertX2;
+					int32 TexY2 = (MipSubsectionSizeQuads + 1) * SubsectionY + VertY2;
+					TextureDataInfo->AddMipUpdateRegion(Mip, TexX1, TexY1, TexX2, TexY2);
+				}
+
+				// Copy current mip values to prev as we move to the next mip.
+				PrevMipSubsectionSizeQuads = MipSubsectionSizeQuads;
+				InvPrevMipSubsectionSizeQuads = InvMipSubsectionSizeQuads;
+
+				PrevMipSizeU = MipSizeU;
+				PrevMipSizeV = MipSizeV;
+
+				// Use this mip's area as we move to the next mip
+				PrevMipSubX1 = MipSubX1;
+				PrevMipSubY1 = MipSubY1;
+				PrevMipSubX2 = MipSubX2;
+				PrevMipSubY2 = MipSubY2;
+			}
+		}
+	}
+
+	// Handle mips that have texels from multiple subsections
+	// not valid weight data, so just average the texels of the previous mip.
+	for (int32 Mip = MaxWholeSubsectionMip + 1;; ++Mip)
+	{
+		int32 MipSubsectionSizeQuads = ((InSubsectionSizeQuads + 1) >> Mip) - 1;
+		checkSlow(MipSubsectionSizeQuads <= 0);
+
+		int32 MipSizeU = FMath::Max<int32>(WeightmapSizeU >> Mip, 1);
+		int32 MipSizeV = FMath::Max<int32>(WeightmapSizeV >> Mip, 1);
+
+		int32 PrevMipSizeU = FMath::Max<int32>(WeightmapSizeU >> (Mip - 1), 1);
+		int32 PrevMipSizeV = FMath::Max<int32>(WeightmapSizeV >> (Mip - 1), 1);
+
+		for (int32 Y = 0; Y < MipSizeV; Y++)
+		{
+			for (int32 X = 0; X < MipSizeU; X++)
+			{
+				DataType* TexData = &(TextureMipData[Mip])[X + Y * MipSizeU];
+
+				DataType *PreMipTexData00 = &(TextureMipData[Mip - 1])[(X * 2 + 0) + (Y * 2 + 0)  * PrevMipSizeU];
+				DataType *PreMipTexData01 = &(TextureMipData[Mip - 1])[(X * 2 + 0) + (Y * 2 + 1)  * PrevMipSizeU];
+				DataType *PreMipTexData10 = &(TextureMipData[Mip - 1])[(X * 2 + 1) + (Y * 2 + 0)  * PrevMipSizeU];
+				DataType *PreMipTexData11 = &(TextureMipData[Mip - 1])[(X * 2 + 1) + (Y * 2 + 1)  * PrevMipSizeU];
+
+				AverageTexData<DataType>(TexData, PreMipTexData00, PreMipTexData10, PreMipTexData01, PreMipTexData11);
+			}
+		}
+
+		if (TextureDataInfo)
+		{
+			// These mip sizes are small enough that we may as well just update the whole mip.
+			TextureDataInfo->AddMipUpdateRegion(Mip, 0, 0, MipSizeU - 1, MipSizeV - 1);
+		}
+
+		if (MipSizeU == 1 && MipSizeV == 1)
+		{
+			break;
+		}
+	}
+}
+
+void ULandscapeComponent::UpdateWeightmapMips(int32 InNumSubsections, int32 InSubsectionSizeQuads, UTexture2D* WeightmapTexture, TArray<FColor*>& WeightmapTextureMipData, int32 ComponentX1/*=0*/, int32 ComponentY1/*=0*/, int32 ComponentX2/*=MAX_int32*/, int32 ComponentY2/*=MAX_int32*/, struct FLandscapeTextureDataInfo* TextureDataInfo/*=NULL*/)
+{
+	UpdateMipsTempl<FColor>(InNumSubsections, InSubsectionSizeQuads, WeightmapTexture, WeightmapTextureMipData, ComponentX1, ComponentY1, ComponentX2, ComponentY2, TextureDataInfo);
+}
+
+void ULandscapeComponent::UpdateDataMips(int32 InNumSubsections, int32 InSubsectionSizeQuads, UTexture2D* Texture, TArray<uint8*>& TextureMipData, int32 ComponentX1/*=0*/, int32 ComponentY1/*=0*/, int32 ComponentX2/*=MAX_int32*/, int32 ComponentY2/*=MAX_int32*/, struct FLandscapeTextureDataInfo* TextureDataInfo/*=NULL*/)
+{
+	UpdateMipsTempl<uint8>(InNumSubsections, InSubsectionSizeQuads, Texture, TextureMipData, ComponentX1, ComponentY1, ComponentX2, ComponentY2, TextureDataInfo);
+}
+
+float ULandscapeComponent::GetLayerWeightAtLocation(const FVector& InLocation, ULandscapeLayerInfoObject* LayerInfo, TArray<uint8>* LayerCache)
+{
+	// Allocate and discard locally if no external cache is passed in.
+	TArray<uint8> LocalCache;
+	if (LayerCache == NULL)
+	{
+		LayerCache = &LocalCache;
+	}
+
+	// Fill the cache if necessary
+	if (LayerCache->Num() == 0)
+	{
+		FLandscapeComponentDataInterface CDI(this);
+		if (!CDI.GetWeightmapTextureData(LayerInfo, *LayerCache))
+		{
+			// no data for this layer for this component.
+			return 0.f;
+		}
+	}
+
+	// Find location
+	// TODO: Root landscape isn't always loaded, would Proxy suffice?
+	ALandscape* Landscape = GetLandscapeActor();
+	const FVector DrawScale = Landscape->GetRootComponent()->RelativeScale3D;
+	float TestX = (InLocation.X - Landscape->GetActorLocation().X) / DrawScale.X - (float)GetSectionBase().X;
+	float TestY = (InLocation.Y - Landscape->GetActorLocation().Y) / DrawScale.Y - (float)GetSectionBase().Y;
+
+	// Find data
+	int32 X1 = FMath::FloorToInt(TestX);
+	int32 Y1 = FMath::FloorToInt(TestY);
+	int32 X2 = FMath::CeilToInt(TestX);
+	int32 Y2 = FMath::CeilToInt(TestY);
+
+	int32 Stride = (SubsectionSizeQuads + 1) * NumSubsections;
+
+	// Min is to prevent the sampling of the final column from overflowing
+	int32 IdxX1 = FMath::Min<int32>(((X1 / SubsectionSizeQuads) * (SubsectionSizeQuads + 1)) + (X1 % SubsectionSizeQuads), Stride - 1);
+	int32 IdxY1 = FMath::Min<int32>(((Y1 / SubsectionSizeQuads) * (SubsectionSizeQuads + 1)) + (Y1 % SubsectionSizeQuads), Stride - 1);
+	int32 IdxX2 = FMath::Min<int32>(((X2 / SubsectionSizeQuads) * (SubsectionSizeQuads + 1)) + (X2 % SubsectionSizeQuads), Stride - 1);
+	int32 IdxY2 = FMath::Min<int32>(((Y2 / SubsectionSizeQuads) * (SubsectionSizeQuads + 1)) + (Y2 % SubsectionSizeQuads), Stride - 1);
+
+	// sample
+	float Sample11 = (float)((*LayerCache)[IdxX1 + Stride*IdxY1]) / 255.f;
+	float Sample21 = (float)((*LayerCache)[IdxX2 + Stride*IdxY1]) / 255.f;
+	float Sample12 = (float)((*LayerCache)[IdxX1 + Stride*IdxY2]) / 255.f;
+	float Sample22 = (float)((*LayerCache)[IdxX2 + Stride*IdxY2]) / 255.f;
+
+	int32 LerpX = FMath::Fractional(TestX);
+	int32 LerpY = FMath::Fractional(TestY);
+
+	// Bilinear interpolate
+	return FMath::Lerp(
+		FMath::Lerp(Sample11, Sample21, LerpX),
+		FMath::Lerp(Sample12, Sample22, LerpX),
+		LerpY);
+}
+
+void ULandscapeComponent::GetComponentExtent(int32& MinX, int32& MinY, int32& MaxX, int32& MaxY) const
+{
+	MinX = FMath::Min(SectionBaseX, MinX);
+	MinY = FMath::Min(SectionBaseY, MinY);
+	MaxX = FMath::Max(SectionBaseX + ComponentSizeQuads, MaxX);
+	MaxY = FMath::Max(SectionBaseY + ComponentSizeQuads, MaxY);
+}
+
+//
+// ALandscape
+//
+
+#define MAX_LANDSCAPE_SUBSECTIONS 2
+
+void ULandscapeInfo::GetComponentsInRegion(int32 X1, int32 Y1, int32 X2, int32 Y2, TSet<ULandscapeComponent*>& OutComponents)
+{
+	if (ComponentSizeQuads <= 0)
+	{
+		return;
+	}
+
+	// Find component range for this block of data
+	int32 ComponentIndexX1, ComponentIndexY1, ComponentIndexX2, ComponentIndexY2;
+	ALandscape::CalcComponentIndicesOverlap(X1, Y1, X2, Y2, ComponentSizeQuads, ComponentIndexX1, ComponentIndexY1, ComponentIndexX2, ComponentIndexY2);
+
+	for (int32 ComponentIndexY = ComponentIndexY1; ComponentIndexY <= ComponentIndexY2; ComponentIndexY++)
+	{
+		for (int32 ComponentIndexX = ComponentIndexX1; ComponentIndexX <= ComponentIndexX2; ComponentIndexX++)
+		{
+			ULandscapeComponent* Component = XYtoComponentMap.FindRef(FIntPoint(ComponentIndexX, ComponentIndexY));
+			if (Component && !FLevelUtils::IsLevelLocked(Component->GetLandscapeProxy()->GetLevel()) && FLevelUtils::IsLevelVisible(Component->GetLandscapeProxy()->GetLevel()))
+			{
+				OutComponents.Add(Component);
+			}
+		}
+	}
+}
+
+// A struct to remember where we have spare texture channels.
+struct FWeightmapTextureAllocation
+{
+	int32 X;
+	int32 Y;
+	int32 ChannelsInUse;
+	UTexture2D* Texture;
+	FColor* TextureData;
+
+	FWeightmapTextureAllocation(int32 InX, int32 InY, int32 InChannels, UTexture2D* InTexture, FColor* InTextureData)
+		: X(InX)
+		, Y(InY)
+		, ChannelsInUse(InChannels)
+		, Texture(InTexture)
+		, TextureData(InTextureData)
+	{}
+};
+
+// A struct to hold the info about each texture chunk of the total heightmap
+struct FHeightmapInfo
+{
+	int32 HeightmapSizeU;
+	int32 HeightmapSizeV;
+	UTexture2D* HeightmapTexture;
+	TArray<FColor*> HeightmapTextureMipData;
+};
+
+FLandscapeLayerStruct* ALandscapeProxy::GetLayerInfo_Deprecated(FName LayerName)
+{
+	for (auto It = LayerInfoObjs_DEPRECATED.CreateIterator(); It; It++)
+	{
+		ULandscapeLayerInfoObject* LayerInfo = It->LayerInfoObj;
+		if (LayerInfo && LayerInfo->LayerName == LayerName)
+		{
+			return &*It;
+		}
+	}
+
+	return NULL;
+}
+
+TArray<FName> ALandscapeProxy::GetLayersFromMaterial(UMaterialInterface* Material)
+{
+	TArray<FName> Result;
+
+	if (Material)
+	{
+		const TArray<UMaterialExpression*>& Expressions = Material->GetMaterial()->Expressions;
+
+		// TODO: *Unconnected* layer expressions?
+		for (UMaterialExpression* Expression : Expressions)
+		{
+			if (auto LayerWeightExpression = Cast<UMaterialExpressionLandscapeLayerWeight>(Expression))
+			{
+				Result.AddUnique(LayerWeightExpression->ParameterName);
+			}
+			else if (auto LayerSampleExpression = Cast<UMaterialExpressionLandscapeLayerSample>(Expression))
+			{
+				Result.AddUnique(LayerSampleExpression->ParameterName);
+			}
+			else if (auto LayerSwitchExpression = Cast<UMaterialExpressionLandscapeLayerSwitch>(Expression))
+			{
+				Result.AddUnique(LayerSwitchExpression->ParameterName);
+			}
+			else if (auto LayerBlendExpression = Cast<UMaterialExpressionLandscapeLayerBlend>(Expression))
+			{
+				for (const auto& ExpressionLayer : LayerBlendExpression->Layers)
+				{
+					Result.AddUnique(ExpressionLayer.LayerName);
+				}
+			}
+		}
+	}
+
+	return Result;
+}
+
+TArray<FName> ALandscapeProxy::GetLayersFromMaterial() const
+{
+	return GetLayersFromMaterial(LandscapeMaterial);
+}
+
+ULandscapeLayerInfoObject* ALandscapeProxy::CreateLayerInfo(const TCHAR* LayerName, ULevel* Level)
+{
+	FName LayerObjectName = FName(*FString::Printf(TEXT("LayerInfoObject_%s"), LayerName));
+	FString Path = Level->GetOutermost()->GetName() + TEXT("_sharedassets/");
+	if (Path.StartsWith("/Temp/"))
+	{
+		Path = FString("/Game/") + Path.RightChop(FString("/Temp/").Len());
+	}
+	FString PackageName = Path + LayerObjectName.ToString();
+	FString PackageFilename;
+	int32 Suffix = 1;
+	while (FPackageName::DoesPackageExist(PackageName, NULL, &PackageFilename))
+	{
+		LayerObjectName = FName(*FString::Printf(TEXT("LayerInfoObject_%s_%d"), LayerName, Suffix));
+		PackageName = Path + LayerObjectName.ToString();
+		Suffix++;
+	}
+	UPackage* Package = CreatePackage(NULL, *PackageName);
+	ULandscapeLayerInfoObject* LayerInfo = ConstructObject<ULandscapeLayerInfoObject>(ULandscapeLayerInfoObject::StaticClass(), Package, LayerObjectName, RF_Public | RF_Standalone | RF_Transactional);
+	LayerInfo->LayerName = LayerName;
+
+	return LayerInfo;
+}
+
+ULandscapeLayerInfoObject* ALandscapeProxy::CreateLayerInfo(const TCHAR* LayerName)
+{
+	ULandscapeLayerInfoObject* LayerInfo = ALandscapeProxy::CreateLayerInfo(LayerName, GetLevel());
+
+	check(LayerInfo);
+
+	ULandscapeInfo* LandscapeInfo = GetLandscapeInfo();
+	if (LandscapeInfo)
+	{
+		int32 Index = LandscapeInfo->GetLayerInfoIndex(LayerName, this);
+		if (Index == INDEX_NONE)
+		{
+			LandscapeInfo->Layers.Add(FLandscapeInfoLayerSettings(LayerInfo, this));
+		}
+		else
+		{
+			LandscapeInfo->Layers[Index].LayerInfoObj = LayerInfo;
+		}
+	}
+
+	return LayerInfo;
+}
+
+#define HEIGHTDATA(X,Y) (HeightData[ FMath::Clamp<int32>(Y,0,VertsY) * VertsX + FMath::Clamp<int32>(X,0,VertsX) ])
+void ALandscapeProxy::Import(FGuid Guid, int32 VertsX, int32 VertsY,
+	int32 InComponentSizeQuads, int32 InNumSubsections, int32 InSubsectionSizeQuads,
+	const uint16* HeightData, const TCHAR* HeightmapFileName,
+	const TArray<FLandscapeImportLayerInfo>& ImportLayerInfos)
+{
+	GWarn->BeginSlowTask(LOCTEXT("BeingImportingLandscapeTask", "Importing Landscape"), true);
+
+	ComponentSizeQuads = InComponentSizeQuads;
+	NumSubsections = InNumSubsections;
+	SubsectionSizeQuads = InSubsectionSizeQuads;
+	LandscapeGuid = Guid;
+
+	MarkPackageDirty();
+
+	// Create and initialize landscape info object
+	GetLandscapeInfo(true)->RegisterActor(this);
+
+	int32 NumPatchesX = (VertsX - 1);
+	int32 NumPatchesY = (VertsY - 1);
+
+	int32 NumSectionsX = NumPatchesX / ComponentSizeQuads;
+	int32 NumSectionsY = NumPatchesY / ComponentSizeQuads;
+
+	LandscapeComponents.Empty(NumSectionsX * NumSectionsY);
+
+	for (int32 Y = 0; Y < NumSectionsY; Y++)
+	{
+		for (int32 X = 0; X < NumSectionsX; X++)
+		{
+			// The number of quads
+			int32 NumQuadsX = NumPatchesX;
+			int32 NumQuadsY = NumPatchesY;
+
+			int32 BaseX = X * ComponentSizeQuads;
+			int32 BaseY = Y * ComponentSizeQuads;
+
+			ULandscapeComponent* LandscapeComponent = ConstructObject<ULandscapeComponent>(ULandscapeComponent::StaticClass(), this, NAME_None, RF_Transactional);
+			LandscapeComponent->SetRelativeLocation(FVector(BaseX, BaseY, 0));
+			LandscapeComponent->AttachTo(GetRootComponent(), NAME_None);
+			LandscapeComponents.Add(LandscapeComponent);
+			LandscapeComponent->Init(
+				BaseX, BaseY,
+				ComponentSizeQuads,
+				NumSubsections,
+				SubsectionSizeQuads
+				);
+
+			// Assign shared properties
+			LandscapeComponent->bCastStaticShadow = bCastStaticShadow;
+			LandscapeComponent->bCastShadowAsTwoSided = bCastShadowAsTwoSided;
+		}
+	}
+
+#define MAX_HEIGHTMAP_TEXTURE_SIZE 512
+
+	int32 ComponentSizeVerts = NumSubsections * (SubsectionSizeQuads + 1);
+	int32 ComponentsPerHeightmap = MAX_HEIGHTMAP_TEXTURE_SIZE / ComponentSizeVerts;
+
+	// Ensure that we don't pack so many heightmaps into a texture that their lowest LOD isn't guaranteed to be resident
+	ComponentsPerHeightmap = FMath::Min(ComponentsPerHeightmap, 1 << (UTexture2D::GetMinTextureResidentMipCount() - 2));
+
+	// Count how many heightmaps we need and the X dimension of the final heightmap
+	int32 NumHeightmapsX = 1;
+	int32 FinalComponentsX = NumSectionsX;
+	while (FinalComponentsX > ComponentsPerHeightmap)
+	{
+		FinalComponentsX -= ComponentsPerHeightmap;
+		NumHeightmapsX++;
+	}
+	// Count how many heightmaps we need and the Y dimension of the final heightmap
+	int32 NumHeightmapsY = 1;
+	int32 FinalComponentsY = NumSectionsY;
+	while (FinalComponentsY > ComponentsPerHeightmap)
+	{
+		FinalComponentsY -= ComponentsPerHeightmap;
+		NumHeightmapsY++;
+	}
+
+	TArray<FHeightmapInfo> HeightmapInfos;
+
+	for (int32 HmY = 0; HmY < NumHeightmapsY; HmY++)
+	{
+		for (int32 HmX = 0; HmX < NumHeightmapsX; HmX++)
+		{
+			FHeightmapInfo& HeightmapInfo = HeightmapInfos[HeightmapInfos.AddZeroed()];
+
+			// make sure the heightmap UVs are powers of two.
+			HeightmapInfo.HeightmapSizeU = (1 << FMath::CeilLogTwo(((HmX == NumHeightmapsX - 1) ? FinalComponentsX : ComponentsPerHeightmap)*ComponentSizeVerts));
+			HeightmapInfo.HeightmapSizeV = (1 << FMath::CeilLogTwo(((HmY == NumHeightmapsY - 1) ? FinalComponentsY : ComponentsPerHeightmap)*ComponentSizeVerts));
+
+			// Construct the heightmap textures
+			HeightmapInfo.HeightmapTexture = CreateLandscapeTexture(HeightmapInfo.HeightmapSizeU, HeightmapInfo.HeightmapSizeV, TEXTUREGROUP_Terrain_Heightmap, TSF_BGRA8);
+
+			int32 MipSubsectionSizeQuads = SubsectionSizeQuads;
+			int32 MipSizeU = HeightmapInfo.HeightmapSizeU;
+			int32 MipSizeV = HeightmapInfo.HeightmapSizeV;
+			while (MipSizeU > 1 && MipSizeV > 1 && MipSubsectionSizeQuads >= 1)
+			{
+				int32 MipIndex = HeightmapInfo.HeightmapTextureMipData.Num();
+				FColor* HeightmapTextureData = (FColor*)HeightmapInfo.HeightmapTexture->Source.LockMip(MipIndex);
+				FMemory::Memzero(HeightmapTextureData, MipSizeU*MipSizeV*sizeof(FColor));
+				HeightmapInfo.HeightmapTextureMipData.Add(HeightmapTextureData);
+
+				MipSizeU >>= 1;
+				MipSizeV >>= 1;
+
+				MipSubsectionSizeQuads = ((MipSubsectionSizeQuads + 1) >> 1) - 1;
+			}
+		}
+	}
+
+	const FVector DrawScale3D = GetRootComponent()->RelativeScale3D;
+
+	// Calculate the normals for each of the two triangles per quad.
+	FVector* VertexNormals = new FVector[(NumPatchesX + 1)*(NumPatchesY + 1)];
+	FMemory::Memzero(VertexNormals, (NumPatchesX + 1)*(NumPatchesY + 1)*sizeof(FVector));
+	for (int32 QuadY = 0; QuadY < NumPatchesY; QuadY++)
+	{
+		for (int32 QuadX = 0; QuadX < NumPatchesX; QuadX++)
+		{
+			FVector Vert00 = FVector(0.f, 0.f, ((float)HEIGHTDATA(QuadX + 0, QuadY + 0) - 32768.f)*LANDSCAPE_ZSCALE) * DrawScale3D;
+			FVector Vert01 = FVector(0.f, 1.f, ((float)HEIGHTDATA(QuadX + 0, QuadY + 1) - 32768.f)*LANDSCAPE_ZSCALE) * DrawScale3D;
+			FVector Vert10 = FVector(1.f, 0.f, ((float)HEIGHTDATA(QuadX + 1, QuadY + 0) - 32768.f)*LANDSCAPE_ZSCALE) * DrawScale3D;
+			FVector Vert11 = FVector(1.f, 1.f, ((float)HEIGHTDATA(QuadX + 1, QuadY + 1) - 32768.f)*LANDSCAPE_ZSCALE) * DrawScale3D;
+
+			FVector FaceNormal1 = ((Vert00 - Vert10) ^ (Vert10 - Vert11)).SafeNormal();
+			FVector FaceNormal2 = ((Vert11 - Vert01) ^ (Vert01 - Vert00)).SafeNormal();
+
+			// contribute to the vertex normals.
+			VertexNormals[(QuadX + 1 + (NumPatchesX + 1)*(QuadY + 0))] += FaceNormal1;
+			VertexNormals[(QuadX + 0 + (NumPatchesX + 1)*(QuadY + 1))] += FaceNormal2;
+			VertexNormals[(QuadX + 0 + (NumPatchesX + 1)*(QuadY + 0))] += FaceNormal1 + FaceNormal2;
+			VertexNormals[(QuadX + 1 + (NumPatchesX + 1)*(QuadY + 1))] += FaceNormal1 + FaceNormal2;
+		}
+	}
+
+	// Weight values for each layer for each component.
+	TArray<TArray<TArray<uint8> > > ComponentWeightValues;
+	ComponentWeightValues.AddZeroed(NumSectionsX*NumSectionsY);
+
+	for (int32 ComponentY = 0; ComponentY < NumSectionsY; ComponentY++)
+	{
+		for (int32 ComponentX = 0; ComponentX < NumSectionsX; ComponentX++)
+		{
+			ULandscapeComponent* LandscapeComponent = LandscapeComponents[ComponentX + ComponentY*NumSectionsX];
+			TArray<TArray<uint8> >& WeightValues = ComponentWeightValues[ComponentX + ComponentY*NumSectionsX];
+
+			// Import alphamap data into local array and check for unused layers for this component.
+			TArray<FLandscapeComponentAlphaInfo> EditingAlphaLayerData;
+			for (int32 LayerIndex = 0; LayerIndex < ImportLayerInfos.Num(); LayerIndex++)
+			{
+				FLandscapeComponentAlphaInfo* NewAlphaInfo = new(EditingAlphaLayerData) FLandscapeComponentAlphaInfo(LandscapeComponent, LayerIndex);
+
+				if (ImportLayerInfos[LayerIndex].LayerData.Num())
+				{
+					for (int32 AlphaY = 0; AlphaY <= LandscapeComponent->ComponentSizeQuads; AlphaY++)
+					{
+						const uint8* OldAlphaRowStart = &ImportLayerInfos[LayerIndex].LayerData[(AlphaY + LandscapeComponent->GetSectionBase().Y) * VertsX + (LandscapeComponent->GetSectionBase().X)];
+						uint8* NewAlphaRowStart = &NewAlphaInfo->AlphaValues[AlphaY * (LandscapeComponent->ComponentSizeQuads + 1)];
+						FMemory::Memcpy(NewAlphaRowStart, OldAlphaRowStart, LandscapeComponent->ComponentSizeQuads + 1);
+					}
+				}
+			}
+
+			for (int32 AlphaMapIndex = 0; AlphaMapIndex < EditingAlphaLayerData.Num(); AlphaMapIndex++)
+			{
+				if (EditingAlphaLayerData[AlphaMapIndex].IsLayerAllZero())
+				{
+					EditingAlphaLayerData.RemoveAt(AlphaMapIndex);
+					AlphaMapIndex--;
+				}
+			}
+
+
+			UE_LOG(LogLandscape, Log, TEXT("%s needs %d alphamaps"), *LandscapeComponent->GetName(), EditingAlphaLayerData.Num());
+
+			// Calculate weightmap weights for this component
+			WeightValues.Empty(EditingAlphaLayerData.Num());
+			WeightValues.AddZeroed(EditingAlphaLayerData.Num());
+			LandscapeComponent->WeightmapLayerAllocations.Empty(EditingAlphaLayerData.Num());
+
+			TArray<bool> IsNoBlendArray;
+			IsNoBlendArray.Empty(EditingAlphaLayerData.Num());
+			IsNoBlendArray.AddZeroed(EditingAlphaLayerData.Num());
+
+			for (int32 WeightLayerIndex = 0; WeightLayerIndex < WeightValues.Num(); WeightLayerIndex++)
+			{
+				// Lookup the original layer name
+				WeightValues[WeightLayerIndex] = EditingAlphaLayerData[WeightLayerIndex].AlphaValues;
+				new(LandscapeComponent->WeightmapLayerAllocations) FWeightmapLayerAllocationInfo(ImportLayerInfos[EditingAlphaLayerData[WeightLayerIndex].LayerIndex].LayerInfo);
+				IsNoBlendArray[WeightLayerIndex] = ImportLayerInfos[EditingAlphaLayerData[WeightLayerIndex].LayerIndex].LayerInfo->bNoWeightBlend;
+			}
+
+			// Discard the temporary alpha data
+			EditingAlphaLayerData.Empty();
+
+			// For each layer...
+			for (int32 WeightLayerIndex = WeightValues.Num() - 1; WeightLayerIndex >= 0; WeightLayerIndex--)
+			{
+				// ... multiply all lower layers'...
+				for (int32 BelowWeightLayerIndex = WeightLayerIndex - 1; BelowWeightLayerIndex >= 0; BelowWeightLayerIndex--)
+				{
+					int32 TotalWeight = 0;
+
+					if (IsNoBlendArray[BelowWeightLayerIndex])
+					{
+						continue; // skip no blend
+					}
+
+					// ... values by...
+					for (int32 Idx = 0; Idx < WeightValues[WeightLayerIndex].Num(); Idx++)
+					{
+						// ... one-minus the current layer's values
+						int32 NewValue = (int32)WeightValues[BelowWeightLayerIndex][Idx] * (int32)(255 - WeightValues[WeightLayerIndex][Idx]) / 255;
+						WeightValues[BelowWeightLayerIndex][Idx] = (uint8)NewValue;
+						TotalWeight += NewValue;
+					}
+
+					if (TotalWeight == 0)
+					{
+						// Remove the layer as it has no contribution
+						WeightValues.RemoveAt(BelowWeightLayerIndex);
+						LandscapeComponent->WeightmapLayerAllocations.RemoveAt(BelowWeightLayerIndex);
+						IsNoBlendArray.RemoveAt(BelowWeightLayerIndex);
+
+						// The current layer has been re-numbered
+						WeightLayerIndex--;
+					}
+				}
+			}
+
+			// Weight normalization for total should be 255...
+			if (WeightValues.Num())
+			{
+				for (int32 Idx = 0; Idx < WeightValues[0].Num(); Idx++)
+				{
+					int32 TotalWeight = 0;
+					int32 MaxLayerIdx = -1;
+					int32 MaxWeight = INT_MIN;
+
+					for (int32 WeightLayerIndex = 0; WeightLayerIndex < WeightValues.Num(); WeightLayerIndex++)
+					{
+						if (!IsNoBlendArray[WeightLayerIndex])
+						{
+							int32 Weight = WeightValues[WeightLayerIndex][Idx];
+							TotalWeight += Weight;
+							if (MaxWeight < Weight)
+							{
+								MaxWeight = Weight;
+								MaxLayerIdx = WeightLayerIndex;
+							}
+						}
+					}
+
+					if (TotalWeight == 0)
+					{
+						if (MaxLayerIdx >= 0)
+						{
+							WeightValues[MaxLayerIdx][Idx] = 255;
+						}
+					}
+					else if (TotalWeight != 255)
+					{
+						// normalization...
+						float Factor = 255.f / TotalWeight;
+						TotalWeight = 0;
+						for (int32 WeightLayerIndex = 0; WeightLayerIndex < WeightValues.Num(); WeightLayerIndex++)
+						{
+							if (!IsNoBlendArray[WeightLayerIndex])
+							{
+								WeightValues[WeightLayerIndex][Idx] = (uint8)(Factor * WeightValues[WeightLayerIndex][Idx]);
+								TotalWeight += WeightValues[WeightLayerIndex][Idx];
+							}
+						}
+
+						if (255 - TotalWeight && MaxLayerIdx >= 0)
+						{
+							WeightValues[MaxLayerIdx][Idx] += 255 - TotalWeight;
+						}
+					}
+				}
+			}
+		}
+	}
+
+	// Remember where we have spare texture channels.
+	TArray<FWeightmapTextureAllocation> TextureAllocations;
+
+	for (int32 ComponentY = 0; ComponentY < NumSectionsY; ComponentY++)
+	{
+		int32 HmY = ComponentY / ComponentsPerHeightmap;
+		int32 HeightmapOffsetY = (ComponentY - ComponentsPerHeightmap*HmY) * NumSubsections * (SubsectionSizeQuads + 1);
+
+		for (int32 ComponentX = 0; ComponentX < NumSectionsX; ComponentX++)
+		{
+			int32 HmX = ComponentX / ComponentsPerHeightmap;
+			FHeightmapInfo& HeightmapInfo = HeightmapInfos[HmX + HmY * NumHeightmapsX];
+
+			ULandscapeComponent* LandscapeComponent = LandscapeComponents[ComponentX + ComponentY*NumSectionsX];
+
+			// Lookup array of weight values for this component.
+			TArray<TArray<uint8> >& WeightValues = ComponentWeightValues[ComponentX + ComponentY*NumSectionsX];
+
+			// Heightmap offsets
+			int32 HeightmapOffsetX = (ComponentX - ComponentsPerHeightmap*HmX) * NumSubsections * (SubsectionSizeQuads + 1);
+
+			LandscapeComponent->HeightmapScaleBias = FVector4(1.f / (float)HeightmapInfo.HeightmapSizeU, 1.f / (float)HeightmapInfo.HeightmapSizeV, (float)((HeightmapOffsetX)) / (float)HeightmapInfo.HeightmapSizeU, ((float)(HeightmapOffsetY)) / (float)HeightmapInfo.HeightmapSizeV);
+			LandscapeComponent->HeightmapTexture = HeightmapInfo.HeightmapTexture;
+
+			// Weightmap is sized the same as the component
+			int32 WeightmapSize = (SubsectionSizeQuads + 1) * NumSubsections;
+			// Should be power of two
+			check(((WeightmapSize - 1) & WeightmapSize) == 0);
+
+			LandscapeComponent->WeightmapScaleBias = FVector4(1.f / (float)WeightmapSize, 1.f / (float)WeightmapSize, 0.5f / (float)WeightmapSize, 0.5f / (float)WeightmapSize);
+			LandscapeComponent->WeightmapSubsectionOffset = (float)(SubsectionSizeQuads + 1) / (float)WeightmapSize;
+
+			// Pointers to the texture data where we'll store each layer. Stride is 4 (FColor)
+			TArray<uint8*> WeightmapTextureDataPointers;
+
+			UE_LOG(LogLandscape, Log, TEXT("%s needs %d weightmap channels"), *LandscapeComponent->GetName(), WeightValues.Num());
+
+			// Find texture channels to store each layer.
+			int32 LayerIndex = 0;
+			while (LayerIndex < WeightValues.Num())
+			{
+				int32 RemainingLayers = WeightValues.Num() - LayerIndex;
+
+				int32 BestAllocationIndex = -1;
+
+				// if we need less than 4 channels, try to find them somewhere to put all of them
+				if (RemainingLayers < 4)
+				{
+					int32 BestDistSquared = MAX_int32;
+					for (int32 TryAllocIdx = 0; TryAllocIdx < TextureAllocations.Num(); TryAllocIdx++)
+					{
+						if (TextureAllocations[TryAllocIdx].ChannelsInUse + RemainingLayers <= 4)
+						{
+							FWeightmapTextureAllocation& TryAllocation = TextureAllocations[TryAllocIdx];
+							int32 TryDistSquared = FMath::Square(TryAllocation.X - ComponentX) + FMath::Square(TryAllocation.Y - ComponentY);
+							if (TryDistSquared < BestDistSquared)
+							{
+								BestDistSquared = TryDistSquared;
+								BestAllocationIndex = TryAllocIdx;
+							}
+						}
+					}
+				}
+
+				if (BestAllocationIndex != -1)
+				{
+					FWeightmapTextureAllocation& Allocation = TextureAllocations[BestAllocationIndex];
+					FLandscapeWeightmapUsage& WeightmapUsage = WeightmapUsageMap.FindChecked(Allocation.Texture);
+
+					UE_LOG(LogLandscape, Log, TEXT("  ==> Storing %d channels starting at %s[%d]"), RemainingLayers, *Allocation.Texture->GetName(), Allocation.ChannelsInUse);
+
+					for (int32 i = 0; i < RemainingLayers; i++)
+					{
+						LandscapeComponent->WeightmapLayerAllocations[LayerIndex + i].WeightmapTextureIndex = LandscapeComponent->WeightmapTextures.Num();
+						LandscapeComponent->WeightmapLayerAllocations[LayerIndex + i].WeightmapTextureChannel = Allocation.ChannelsInUse;
+						WeightmapUsage.ChannelUsage[Allocation.ChannelsInUse] = LandscapeComponent;
+						switch (Allocation.ChannelsInUse)
+						{
+						case 1:
+							WeightmapTextureDataPointers.Add((uint8*)&Allocation.TextureData->G);
+							break;
+						case 2:
+							WeightmapTextureDataPointers.Add((uint8*)&Allocation.TextureData->B);
+							break;
+						case 3:
+							WeightmapTextureDataPointers.Add((uint8*)&Allocation.TextureData->A);
+							break;
+						default:
+							// this should not occur.
+							check(0);
+
+						}
+						Allocation.ChannelsInUse++;
+					}
+
+					LayerIndex += RemainingLayers;
+					LandscapeComponent->WeightmapTextures.Add(Allocation.Texture);
+				}
+				else
+				{
+					// We couldn't find a suitable place for these layers, so lets make a new one.
+					UTexture2D* WeightmapTexture = CreateLandscapeTexture(WeightmapSize, WeightmapSize, TEXTUREGROUP_Terrain_Weightmap, TSF_BGRA8);
+					FColor* MipData = (FColor*)WeightmapTexture->Source.LockMip(0);
+
+					int32 ThisAllocationLayers = FMath::Min<int32>(RemainingLayers, 4);
+					new(TextureAllocations)FWeightmapTextureAllocation(ComponentX, ComponentY, ThisAllocationLayers, WeightmapTexture, MipData);
+					FLandscapeWeightmapUsage& WeightmapUsage = WeightmapUsageMap.Add(WeightmapTexture, FLandscapeWeightmapUsage());
+
+					UE_LOG(LogLandscape, Log, TEXT("  ==> Storing %d channels in new texture %s"), ThisAllocationLayers, *WeightmapTexture->GetName());
+
+					WeightmapTextureDataPointers.Add((uint8*)&MipData->R);
+					LandscapeComponent->WeightmapLayerAllocations[LayerIndex + 0].WeightmapTextureIndex = LandscapeComponent->WeightmapTextures.Num();
+					LandscapeComponent->WeightmapLayerAllocations[LayerIndex + 0].WeightmapTextureChannel = 0;
+					WeightmapUsage.ChannelUsage[0] = LandscapeComponent;
+
+					if (ThisAllocationLayers > 1)
+					{
+						WeightmapTextureDataPointers.Add((uint8*)&MipData->G);
+						LandscapeComponent->WeightmapLayerAllocations[LayerIndex + 1].WeightmapTextureIndex = LandscapeComponent->WeightmapTextures.Num();
+						LandscapeComponent->WeightmapLayerAllocations[LayerIndex + 1].WeightmapTextureChannel = 1;
+						WeightmapUsage.ChannelUsage[1] = LandscapeComponent;
+
+						if (ThisAllocationLayers > 2)
+						{
+							WeightmapTextureDataPointers.Add((uint8*)&MipData->B);
+							LandscapeComponent->WeightmapLayerAllocations[LayerIndex + 2].WeightmapTextureIndex = LandscapeComponent->WeightmapTextures.Num();
+							LandscapeComponent->WeightmapLayerAllocations[LayerIndex + 2].WeightmapTextureChannel = 2;
+							WeightmapUsage.ChannelUsage[2] = LandscapeComponent;
+
+							if (ThisAllocationLayers > 3)
+							{
+								WeightmapTextureDataPointers.Add((uint8*)&MipData->A);
+								LandscapeComponent->WeightmapLayerAllocations[LayerIndex + 3].WeightmapTextureIndex = LandscapeComponent->WeightmapTextures.Num();
+								LandscapeComponent->WeightmapLayerAllocations[LayerIndex + 3].WeightmapTextureChannel = 3;
+								WeightmapUsage.ChannelUsage[3] = LandscapeComponent;
+							}
+						}
+					}
+					LandscapeComponent->WeightmapTextures.Add(WeightmapTexture);
+
+					LayerIndex += ThisAllocationLayers;
+				}
+			}
+			check(WeightmapTextureDataPointers.Num() == WeightValues.Num());
+
+			FVector* LocalVerts = new FVector[FMath::Square(ComponentSizeQuads + 1)];
+
+			for (int32 SubsectionY = 0; SubsectionY < NumSubsections; SubsectionY++)
+			{
+				for (int32 SubsectionX = 0; SubsectionX < NumSubsections; SubsectionX++)
+				{
+					for (int32 SubY = 0; SubY <= SubsectionSizeQuads; SubY++)
+					{
+						for (int32 SubX = 0; SubX <= SubsectionSizeQuads; SubX++)
+						{
+							// X/Y of the vertex we're looking at in component's coordinates.
+							int32 CompX = SubsectionSizeQuads * SubsectionX + SubX;
+							int32 CompY = SubsectionSizeQuads * SubsectionY + SubY;
+
+							// X/Y of the vertex we're looking indexed into the texture data
+							int32 TexX = (SubsectionSizeQuads + 1) * SubsectionX + SubX;
+							int32 TexY = (SubsectionSizeQuads + 1) * SubsectionY + SubY;
+
+							int32 WeightSrcDataIdx = CompY * (ComponentSizeQuads + 1) + CompX;
+							int32 HeightTexDataIdx = (HeightmapOffsetX + TexX) + (HeightmapOffsetY + TexY) * (HeightmapInfo.HeightmapSizeU);
+
+							int32 WeightTexDataIdx = (TexX)+(TexY)* (WeightmapSize);
+
+							// copy height and normal data
+							uint16 HeightValue = HEIGHTDATA(CompX + LandscapeComponent->GetSectionBase().X, CompY + LandscapeComponent->GetSectionBase().Y);
+							FVector Normal = VertexNormals[CompX + LandscapeComponent->GetSectionBase().X + (NumPatchesX + 1)*(CompY + LandscapeComponent->GetSectionBase().Y)].SafeNormal();
+
+							HeightmapInfo.HeightmapTextureMipData[0][HeightTexDataIdx].R = HeightValue >> 8;
+							HeightmapInfo.HeightmapTextureMipData[0][HeightTexDataIdx].G = HeightValue & 255;
+							HeightmapInfo.HeightmapTextureMipData[0][HeightTexDataIdx].B = FMath::RoundToInt(127.5f * (Normal.X + 1.f));
+							HeightmapInfo.HeightmapTextureMipData[0][HeightTexDataIdx].A = FMath::RoundToInt(127.5f * (Normal.Y + 1.f));
+
+							for (int32 WeightmapIndex = 0; WeightmapIndex < WeightValues.Num(); WeightmapIndex++)
+							{
+								WeightmapTextureDataPointers[WeightmapIndex][WeightTexDataIdx * 4] = WeightValues[WeightmapIndex][WeightSrcDataIdx];
+							}
+
+							// Get local space verts
+							FVector LocalVertex(CompX, CompY, LandscapeDataAccess::GetLocalHeight(HeightValue));
+							LocalVerts[(LandscapeComponent->ComponentSizeQuads + 1) * CompY + CompX] = LocalVertex;
+						}
+					}
+				}
+			}
+
+			LandscapeComponent->CachedLocalBox = FBox(LocalVerts, FMath::Square(ComponentSizeQuads + 1));
+			delete[] LocalVerts;
+
+			// Update MaterialInstance
+			LandscapeComponent->UpdateMaterialInstances();
+		}
+	}
+
+	// Unlock the weightmaps' base mips
+	for (int32 AllocationIndex = 0; AllocationIndex < TextureAllocations.Num(); AllocationIndex++)
+	{
+		UTexture2D* WeightmapTexture = TextureAllocations[AllocationIndex].Texture;
+		FColor* BaseMipData = TextureAllocations[AllocationIndex].TextureData;
+
+		// Generate mips for weightmaps
+		ULandscapeComponent::GenerateWeightmapMips(NumSubsections, SubsectionSizeQuads, WeightmapTexture, BaseMipData);
+
+		WeightmapTexture->Source.UnlockMip(0);
+		WeightmapTexture->PostEditChange();
+	}
+
+	delete[] VertexNormals;
+
+	// Generate mipmaps for the components, and create the collision components
+	for (int32 ComponentY = 0; ComponentY < NumSectionsY; ComponentY++)
+	{
+		for (int32 ComponentX = 0; ComponentX < NumSectionsX; ComponentX++)
+		{
+			int32 HmX = ComponentX / ComponentsPerHeightmap;
+			int32 HmY = ComponentY / ComponentsPerHeightmap;
+			FHeightmapInfo& HeightmapInfo = HeightmapInfos[HmX + HmY * NumHeightmapsX];
+
+			ULandscapeComponent* LandscapeComponent = LandscapeComponents[ComponentX + ComponentY*NumSectionsX];
+			LandscapeComponent->GenerateHeightmapMips(HeightmapInfo.HeightmapTextureMipData, ComponentX == NumSectionsX - 1 ? MAX_int32 : 0, ComponentY == NumSectionsY - 1 ? MAX_int32 : 0);
+			LandscapeComponent->UpdateCollisionHeightData(HeightmapInfo.HeightmapTextureMipData[LandscapeComponent->CollisionMipLevel]);
+			LandscapeComponent->UpdateCollisionLayerData();
+		}
+	}
+
+	for (int32 HmIdx = 0; HmIdx < HeightmapInfos.Num(); HmIdx++)
+	{
+		FHeightmapInfo& HeightmapInfo = HeightmapInfos[HmIdx];
+
+		// Add remaining mips down to 1x1 to heightmap texture. These do not represent quads and are just a simple averages of the previous mipmaps. 
+		// These mips are not used for sampling in the vertex shader but could be sampled in the pixel shader.
+		int32 Mip = HeightmapInfo.HeightmapTextureMipData.Num();
+		int32 MipSizeU = (HeightmapInfo.HeightmapTexture->Source.GetSizeX()) >> Mip;
+		int32 MipSizeV = (HeightmapInfo.HeightmapTexture->Source.GetSizeY()) >> Mip;
+		while (MipSizeU > 1 && MipSizeV > 1)
+		{
+			HeightmapInfo.HeightmapTextureMipData.Add((FColor*)HeightmapInfo.HeightmapTexture->Source.LockMip(Mip));
+			int32 PrevMipSizeU = (HeightmapInfo.HeightmapTexture->Source.GetSizeX()) >> (Mip - 1);
+			int32 PrevMipSizeV = (HeightmapInfo.HeightmapTexture->Source.GetSizeY()) >> (Mip - 1);
+
+			for (int32 Y = 0; Y < MipSizeV; Y++)
+			{
+				for (int32 X = 0; X < MipSizeU; X++)
+				{
+					FColor* TexData = &(HeightmapInfo.HeightmapTextureMipData[Mip])[X + Y * MipSizeU];
+
+					FColor *PreMipTexData00 = &(HeightmapInfo.HeightmapTextureMipData[Mip - 1])[(X * 2 + 0) + (Y * 2 + 0)  * PrevMipSizeU];
+					FColor *PreMipTexData01 = &(HeightmapInfo.HeightmapTextureMipData[Mip - 1])[(X * 2 + 0) + (Y * 2 + 1)  * PrevMipSizeU];
+					FColor *PreMipTexData10 = &(HeightmapInfo.HeightmapTextureMipData[Mip - 1])[(X * 2 + 1) + (Y * 2 + 0)  * PrevMipSizeU];
+					FColor *PreMipTexData11 = &(HeightmapInfo.HeightmapTextureMipData[Mip - 1])[(X * 2 + 1) + (Y * 2 + 1)  * PrevMipSizeU];
+
+					TexData->R = (((int32)PreMipTexData00->R + (int32)PreMipTexData01->R + (int32)PreMipTexData10->R + (int32)PreMipTexData11->R) >> 2);
+					TexData->G = (((int32)PreMipTexData00->G + (int32)PreMipTexData01->G + (int32)PreMipTexData10->G + (int32)PreMipTexData11->G) >> 2);
+					TexData->B = (((int32)PreMipTexData00->B + (int32)PreMipTexData01->B + (int32)PreMipTexData10->B + (int32)PreMipTexData11->B) >> 2);
+					TexData->A = (((int32)PreMipTexData00->A + (int32)PreMipTexData01->A + (int32)PreMipTexData10->A + (int32)PreMipTexData11->A) >> 2);
+				}
+			}
+			Mip++;
+			MipSizeU >>= 1;
+			MipSizeV >>= 1;
+		}
+
+		for (int32 i = 0; i < HeightmapInfo.HeightmapTextureMipData.Num(); i++)
+		{
+			HeightmapInfo.HeightmapTexture->Source.UnlockMip(i);
+		}
+		HeightmapInfo.HeightmapTexture->PostEditChange();
+	}
+
+	if (GetLevel()->bIsVisible)
+	{
+		// Update our new components
+		ReregisterAllComponents();
+	}
+
+	ReimportHeightmapFilePath = HeightmapFileName;
+
+	ULandscapeInfo::RecreateLandscapeInfo(GetWorld(), false);
+
+	GWarn->EndSlowTask();
+}
+
+bool ALandscapeProxy::ExportToRawMesh(int32 InExportLOD, FRawMesh& OutRawMesh) const
+{
+	TArray<ULandscapeComponent*> RegisteredLandscapeComponents;
+	GetComponents<ULandscapeComponent>(RegisteredLandscapeComponents);
+
+	const FIntRect LandscapeSectionRect = GetBoundingRect();
+	const FVector2D LandscapeUVScale = FVector2D(1.f, 1.f) / FVector2D(LandscapeSectionRect.Size());
+
+	// User specified LOD to export
+	int32 LandscapeLODToExport = ExportLOD;
+	if (InExportLOD != INDEX_NONE)
+	{
+		LandscapeLODToExport = FMath::Clamp<int32>(InExportLOD, 0, FMath::CeilLogTwo(SubsectionSizeQuads + 1) - 1);
+	}
+
+	// Export data for each component
+	for (auto It = RegisteredLandscapeComponents.CreateConstIterator(); It; ++It)
+	{
+		ULandscapeComponent* Component = (*It);
+		FLandscapeComponentDataInterface CDI(Component, LandscapeLODToExport);
+		const int32 ComponentSizeQuadsLOD = ((Component->ComponentSizeQuads + 1) >> LandscapeLODToExport) - 1;
+		const int32 SubsectionSizeQuadsLOD = ((Component->SubsectionSizeQuads + 1) >> LandscapeLODToExport) - 1;
+		const FIntPoint ComponentOffsetQuads = Component->GetSectionBase() - LandscapeSectionOffset - LandscapeSectionRect.Min;
+		const FVector2D ComponentUVOffsetLOD = FVector2D(ComponentOffsetQuads)*((float)ComponentSizeQuadsLOD / ComponentSizeQuads);
+		const FVector2D ComponentUVScaleLOD = LandscapeUVScale*((float)ComponentSizeQuads / ComponentSizeQuadsLOD);
+
+		const int32 NumFaces = FMath::Square(ComponentSizeQuadsLOD) * 2;
+		const int32 NumVertices = NumFaces * 3;
+		const int32 VerticesOffset = OutRawMesh.VertexPositions.Num();
+		const int32 IndicesOffset = OutRawMesh.WedgeIndices.Num();
+
+		//
+		OutRawMesh.FaceMaterialIndices.AddZeroed(NumFaces);
+		OutRawMesh.FaceSmoothingMasks.AddZeroed(NumFaces);
+
+		OutRawMesh.VertexPositions.AddZeroed(NumVertices);
+		OutRawMesh.WedgeIndices.AddZeroed(NumVertices);
+		OutRawMesh.WedgeTangentX.AddZeroed(NumVertices);
+		OutRawMesh.WedgeTangentY.AddZeroed(NumVertices);
+		OutRawMesh.WedgeTangentZ.AddZeroed(NumVertices);
+		OutRawMesh.WedgeTexCoords[0].AddZeroed(NumVertices);
+
+
+		// Check if there are any holes
+		TArray<uint8> VisDataMap;
+
+		for (int32 AllocIdx = 0; AllocIdx < Component->WeightmapLayerAllocations.Num(); AllocIdx++)
+		{
+			FWeightmapLayerAllocationInfo& AllocInfo = Component->WeightmapLayerAllocations[AllocIdx];
+			if (AllocInfo.LayerInfo == ALandscapeProxy::VisibilityLayer)
+			{
+				CDI.GetWeightmapTextureData(AllocInfo.LayerInfo, VisDataMap);
+			}
+		}
+
+		const FIntPoint QuadPattern[6] =
+		{
+			//face 1
+			FIntPoint(0, 0),
+			FIntPoint(0, 1),
+			FIntPoint(1, 1),
+			//face 2
+			FIntPoint(0, 0),
+			FIntPoint(1, 1),
+			FIntPoint(1, 0),
+		};
+
+		const int32 VisThreshold = 170;
+		const int32 WeightMapSize = (SubsectionSizeQuadsLOD + 1) * Component->NumSubsections;
+		uint32* Faces = OutRawMesh.WedgeIndices.GetData() + IndicesOffset;
+
+		// Export verts
+		int32 VertexIdx = VerticesOffset;
+		for (int32 y = 0; y < ComponentSizeQuadsLOD; y++)
+		{
+			for (int32 x = 0; x < ComponentSizeQuadsLOD; x++)
+			{
+				// Fill indices
+				{
+
+					// Whether this vertex is in hole
+					bool bInvisible = false;
+					if (VisDataMap.Num())
+					{
+						int32 TexelX, TexelY;
+						CDI.VertexXYToTexelXY(x, y, TexelX, TexelY);
+						bInvisible = (VisDataMap[CDI.TexelXYToIndex(TexelX, TexelY)] >= VisThreshold);
+					}
+
+					// triangulation matches FLandscapeIndexBuffer constructor
+					Faces[0] = VertexIdx;
+					Faces[1] = bInvisible ? Faces[0] : VertexIdx + 1;
+					Faces[2] = bInvisible ? Faces[0] : VertexIdx + 2;
+					Faces += 3;
+
+					Faces[0] = VertexIdx + 3;
+					Faces[1] = bInvisible ? Faces[0] : VertexIdx + 4;
+					Faces[2] = bInvisible ? Faces[0] : VertexIdx + 5;
+					Faces += 3;
+				}
+
+				// Fill vertices
+				for (int32 i = 0; i < ARRAY_COUNT(QuadPattern); i++)
+				{
+					int32 VertexX = x + QuadPattern[i].X;
+					int32 VertexY = y + QuadPattern[i].Y;
+					FVector LocalVertexPos = CDI.GetWorldVertex(VertexX, VertexY);
+
+					FVector LocalTangentX, LocalTangentY, LocalTangentZ;
+					CDI.GetLocalTangentVectors(VertexX, VertexY, LocalTangentX, LocalTangentY, LocalTangentZ);
+
+					OutRawMesh.VertexPositions[VertexIdx] = LocalVertexPos;
+					OutRawMesh.WedgeTangentX[VertexIdx] = LocalTangentX;
+					OutRawMesh.WedgeTangentY[VertexIdx] = LocalTangentY;
+					OutRawMesh.WedgeTangentZ[VertexIdx] = LocalTangentZ;
+
+					OutRawMesh.WedgeTexCoords[0][VertexIdx] = (ComponentUVOffsetLOD + FVector2D(VertexX, VertexY))*ComponentUVScaleLOD;
+
+					VertexIdx++;
+				}
+			}
+		}
+	}
+
+	// Add lightmap UVs
+	OutRawMesh.WedgeTexCoords[1].Append(OutRawMesh.WedgeTexCoords[0]);
+
+	return true;
+}
+
+FIntRect ALandscapeProxy::GetBoundingRect() const
+{
+	FIntRect Rect(MAX_int32, MAX_int32, MIN_int32, MIN_int32);
+
+	for (int32 CompIdx = 0; CompIdx < LandscapeComponents.Num(); CompIdx++)
+	{
+		Rect.Include(LandscapeComponents[CompIdx]->GetSectionBase());
+	}
+
+	if (LandscapeComponents.Num() > 0)
+	{
+		Rect.Max += FIntPoint(ComponentSizeQuads, ComponentSizeQuads);
+		Rect -= LandscapeSectionOffset;
+	}
+	else
+	{
+		Rect = FIntRect();
+	}
+
+	return Rect;
+}
+
+bool ALandscape::HasAllComponent()
+{
+	ULandscapeInfo* Info = GetLandscapeInfo();
+	if (Info && Info->XYtoComponentMap.Num() == LandscapeComponents.Num())
+	{
+		// all components are owned by this Landscape actor (no Landscape Proxies)
+		return true;
+	}
+	return false;
+}
+
+bool ULandscapeInfo::GetLandscapeExtent(int32& MinX, int32& MinY, int32& MaxX, int32& MaxY)
+{
+	MinX = MAX_int32;
+	MinY = MAX_int32;
+	MaxX = MIN_int32;
+	MaxY = MIN_int32;
+
+	// Find range of entire landscape
+	for (auto It = XYtoComponentMap.CreateIterator(); It; ++It)
+	{
+		ULandscapeComponent* Comp = It.Value();
+		Comp->GetComponentExtent(MinX, MinY, MaxX, MaxY);
+	}
+	return (MinX != MAX_int32);
+}
+
+bool ULandscapeInfo::GetSelectedExtent(int32& MinX, int32& MinY, int32& MaxX, int32& MaxY)
+{
+	MinX = MinY = MAX_int32;
+	MaxX = MaxY = MIN_int32;
+	for (auto It = SelectedRegion.CreateIterator(); It; ++It)
+	{
+		int32 X, Y;
+		ALandscape::UnpackKey(It.Key(), X, Y);
+		if (MinX > X) MinX = X;
+		if (MaxX < X) MaxX = X;
+		if (MinY > Y) MinY = Y;
+		if (MaxY < Y) MaxY = Y;
+	}
+	if (MinX != MAX_int32)
+	{
+		return true;
+	}
+	// if SelectedRegion is empty, try SelectedComponents
+	for (TSet<ULandscapeComponent*>::TIterator It(SelectedComponents); It; ++It)
+	{
+		ULandscapeComponent* Comp = *It;
+		Comp->GetComponentExtent(MinX, MinY, MaxX, MaxY);
+	}
+	return MinX != MAX_int32;
+}
+
+FVector ULandscapeInfo::GetLandscapeCenterPos(float& LengthZ, int32 MinX /*= MAX_INT*/, int32 MinY /*= MAX_INT*/, int32 MaxX /*= MIN_INT*/, int32 MaxY /*= MIN_INT*/)
+{
+	// MinZ, MaxZ is Local coordinate
+	float MaxZ = -HALF_WORLD_MAX, MinZ = HALF_WORLD_MAX;
+	const float ScaleZ = DrawScale.Z;
+
+	if (MinX == MAX_int32)
+	{
+		// Find range of entire landscape
+		for (auto It = XYtoComponentMap.CreateIterator(); It; ++It)
+		{
+			ULandscapeComponent* Comp = It.Value();
+			Comp->GetComponentExtent(MinX, MinY, MaxX, MaxY);
+		}
+
+		const int32 Dist = (ComponentSizeQuads + 1) >> 1; // Should be same in ALandscapeGizmoActiveActor::SetTargetLandscape
+		FVector2D MidPoint(((float)(MinX + MaxX)) / 2.f, ((float)(MinY + MaxY)) / 2.f);
+		MinX = FMath::FloorToInt(MidPoint.X) - Dist;
+		MaxX = FMath::CeilToInt(MidPoint.X) + Dist;
+		MinY = FMath::FloorToInt(MidPoint.Y) - Dist;
+		MaxY = FMath::CeilToInt(MidPoint.Y) + Dist;
+		check(MidPoint.X == ((float)(MinX + MaxX)) / 2.f && MidPoint.Y == ((float)(MinY + MaxY)) / 2.f);
+	}
+
+	check(MinX != MAX_int32);
+	//if (MinX != MAX_int32)
+	{
+		int32 CompX1, CompX2, CompY1, CompY2;
+		ALandscape::CalcComponentIndicesOverlap(MinX, MinY, MaxX, MaxY, ComponentSizeQuads, CompX1, CompY1, CompX2, CompY2);
+		for (int32 IndexY = CompY1; IndexY <= CompY2; ++IndexY)
+		{
+			for (int32 IndexX = CompX1; IndexX <= CompX2; ++IndexX)
+			{
+				ULandscapeComponent* Comp = XYtoComponentMap.FindRef(FIntPoint(IndexX, IndexY));
+				if (Comp)
+				{
+					ULandscapeHeightfieldCollisionComponent* CollisionComp = Comp->CollisionComponent.Get();
+					if (CollisionComp)
+					{
+						uint16* Heights = (uint16*)CollisionComp->CollisionHeightData.Lock(LOCK_READ_ONLY);
+						int32 CollisionSizeVerts = CollisionComp->CollisionSizeQuads + 1;
+
+						int32 StartX = FMath::Max(0, MinX - CollisionComp->GetSectionBase().X);
+						int32 StartY = FMath::Max(0, MinY - CollisionComp->GetSectionBase().Y);
+						int32 EndX = FMath::Min(CollisionSizeVerts, MaxX - CollisionComp->GetSectionBase().X + 1);
+						int32 EndY = FMath::Min(CollisionSizeVerts, MaxY - CollisionComp->GetSectionBase().Y + 1);
+
+						for (int32 Y = StartY; Y < EndY; ++Y)
+						{
+							for (int32 X = StartX; X < EndX; ++X)
+							{
+								float Height = LandscapeDataAccess::GetLocalHeight(Heights[X + Y*CollisionSizeVerts]);
+								MaxZ = FMath::Max(Height, MaxZ);
+								MinZ = FMath::Min(Height, MinZ);
+							}
+						}
+						CollisionComp->CollisionHeightData.Unlock();
+					}
+				}
+			}
+		}
+	}
+
+	const float MarginZ = 3;
+	if (MaxZ < MinZ)
+	{
+		MaxZ = +MarginZ;
+		MinZ = -MarginZ;
+	}
+	LengthZ = (MaxZ - MinZ + 2 * MarginZ) * ScaleZ;
+
+	const FVector LocalPosition(((float)(MinX + MaxX)) / 2.f, ((float)(MinY + MaxY)) / 2.f, MinZ - MarginZ);
+	//return GetLandscapeProxy()->TransformLandscapeLocationToWorld(LocalPosition);
+	return GetLandscapeProxy()->LandscapeActorToWorld().TransformPosition(LocalPosition);
+}
+
+bool ULandscapeInfo::IsValidPosition(int32 X, int32 Y)
+{
+	int32 CompX1, CompX2, CompY1, CompY2;
+	ALandscape::CalcComponentIndicesOverlap(X, Y, X, Y, ComponentSizeQuads, CompX1, CompY1, CompX2, CompY2);
+	if (XYtoComponentMap.FindRef(FIntPoint(CompX1, CompY1)))
+	{
+		return true;
+	}
+	if (XYtoComponentMap.FindRef(FIntPoint(CompX2, CompY2)))
+	{
+		return true;
+	}
+	return false;
+}
+
+void ULandscapeInfo::Export(const TArray<ULandscapeLayerInfoObject*>& LayerInfos, const TArray<FString>& Filenames)
+{
+	check(Filenames.Num() > 0);
+
+	int32 MinX = MAX_int32;
+	int32 MinY = MAX_int32;
+	int32 MaxX = -MAX_int32;
+	int32 MaxY = -MAX_int32;
+
+	if (!GetLandscapeExtent(MinX, MinY, MaxX, MaxY))
+	{
+		return;
+	}
+
+	GWarn->BeginSlowTask(LOCTEXT("BeginExportingLandscapeTask", "Exporting Landscape"), true);
+
+	FLandscapeEditDataInterface LandscapeEdit(this);
+
+	TArray<uint8> HeightData;
+	HeightData.AddZeroed((1 + MaxX - MinX)*(1 + MaxY - MinY)*sizeof(uint16));
+	LandscapeEdit.GetHeightDataFast(MinX, MinY, MaxX, MaxY, (uint16*)HeightData.GetData(), 0);
+	FFileHelper::SaveArrayToFile(HeightData, *Filenames[0]);
+
+	for (int32 i = 1; i < Filenames.Num(); i++)
+	{
+		if (i <= LayerInfos.Num())
+		{
+			TArray<uint8> WeightData;
+			WeightData.AddZeroed((1 + MaxX - MinX)*(1 + MaxY - MinY));
+			ULandscapeLayerInfoObject* LayerInfo = LayerInfos[i - 1];
+			if (LayerInfo)
+			{
+				LandscapeEdit.GetWeightDataFast(LayerInfo, MinX, MinY, MaxX, MaxY, WeightData.GetData(), 0);
+			}
+			FFileHelper::SaveArrayToFile(WeightData, *Filenames[i]);
+		}
+	}
+
+	GWarn->EndSlowTask();
+}
+
+void ULandscapeInfo::ExportHeightmap(const FString& Filename)
+{
+	int32 MinX = MAX_int32;
+	int32 MinY = MAX_int32;
+	int32 MaxX = -MAX_int32;
+	int32 MaxY = -MAX_int32;
+
+	if (!GetLandscapeExtent(MinX, MinY, MaxX, MaxY))
+	{
+		return;
+	}
+
+	GWarn->BeginSlowTask(LOCTEXT("BeginExportingLandscapeHeightmapTask", "Exporting Landscape Heightmap"), true);
+
+	FLandscapeEditDataInterface LandscapeEdit(this);
+
+	TArray<uint8> HeightData;
+	HeightData.AddZeroed((MaxX - MinX + 1) * (MaxY - MinY + 1) * sizeof(uint16));
+	LandscapeEdit.GetHeightDataFast(MinX, MinY, MaxX, MaxY, (uint16*)HeightData.GetData(), 0);
+
+	if (Filename.EndsWith(".png"))
+	{
+		IImageWrapperModule& ImageWrapperModule = FModuleManager::LoadModuleChecked<IImageWrapperModule>("ImageWrapper");
+		IImageWrapperPtr ImageWrapper = ImageWrapperModule.CreateImageWrapper(EImageFormat::PNG);
+
+		const TArray<uint8>* RawData = NULL;
+		if (ImageWrapper->SetRaw(HeightData.GetData(), HeightData.Num(), (MaxX - MinX + 1), (MaxY - MinY + 1), ERGBFormat::Gray, 16))
+		{
+			HeightData = ImageWrapper->GetCompressed();
+		}
+	}
+
+	FFileHelper::SaveArrayToFile(HeightData, *Filename);
+
+	GWarn->EndSlowTask();
+}
+
+void ULandscapeInfo::ExportLayer(ULandscapeLayerInfoObject* LayerInfo, const FString& Filename)
+{
+	int32 MinX = MAX_int32;
+	int32 MinY = MAX_int32;
+	int32 MaxX = -MAX_int32;
+	int32 MaxY = -MAX_int32;
+
+	if (!GetLandscapeExtent(MinX, MinY, MaxX, MaxY))
+	{
+		return;
+	}
+
+	GWarn->BeginSlowTask(LOCTEXT("BeginExportingLandscapeWeightmapTask", "Exporting Landscape Layer Weightmap"), true);
+
+	TArray<uint8> WeightData;
+	WeightData.AddZeroed((MaxX - MinX + 1) * (MaxY - MinY + 1));
+	if (LayerInfo)
+	{
+		FLandscapeEditDataInterface LandscapeEdit(this);
+		LandscapeEdit.GetWeightDataFast(LayerInfo, MinX, MinY, MaxX, MaxY, WeightData.GetData(), 0);
+	}
+
+	if (Filename.EndsWith(".png"))
+	{
+		IImageWrapperModule& ImageWrapperModule = FModuleManager::LoadModuleChecked<IImageWrapperModule>("ImageWrapper");
+		IImageWrapperPtr ImageWrapper = ImageWrapperModule.CreateImageWrapper(EImageFormat::PNG);
+
+		const TArray<uint8>* RawData = NULL;
+		if (ImageWrapper->SetRaw(WeightData.GetData(), WeightData.Num(), (MaxX - MinX + 1), (MaxY - MinY + 1), ERGBFormat::Gray, 8))
+		{
+			WeightData = ImageWrapper->GetCompressed();
+		}
+	}
+
+	FFileHelper::SaveArrayToFile(WeightData, *Filename);
+
+	GWarn->EndSlowTask();
+}
+
+//////////////////////////////////////////////////////////////////////////
+// Apply splines
+//////////////////////////////////////////////////////////////////////////
+#include "Raster.h"
+
+class FLandscapeSplineHeightsRasterPolicy
+{
+public:
+	// X = Side Alpha, Y = End Alpha, Z = Height
+	typedef FVector InterpolantType;
+
+	/** Initialization constructor. */
+	FLandscapeSplineHeightsRasterPolicy(TArray<uint16>& InData, int32 InMinX, int32 InMinY, int32 InMaxX, int32 InMaxY, bool InbRaiseTerrain, bool InbLowerTerrain) :
+		Data(InData),
+		MinX(InMinX),
+		MinY(InMinY),
+		MaxX(InMaxX),
+		MaxY(InMaxY),
+		bRaiseTerrain(InbRaiseTerrain),
+		bLowerTerrain(InbLowerTerrain)
+	{
+	}
+
+protected:
+
+	// FTriangleRasterizer policy interface.
+
+	int32 GetMinX() const { return MinX; }
+	int32 GetMaxX() const { return MaxX; }
+	int32 GetMinY() const { return MinY; }
+	int32 GetMaxY() const { return MaxY; }
+
+	void ProcessPixel(int32 X, int32 Y, const InterpolantType& Interpolant, bool BackFacing)
+	{
+		const float CosInterpX = (Interpolant.X >= 1 ? 1 : 0.5f - 0.5f * FMath::Cos(Interpolant.X * PI));
+		const float CosInterpY = (Interpolant.Y >= 1 ? 1 : 0.5f - 0.5f * FMath::Cos(Interpolant.Y * PI));
+		const float Alpha = CosInterpX * CosInterpY;
+		uint16& Dest = Data[(Y - MinY)*(1 + MaxX - MinX) + X - MinX];
+		float Value = FMath::Lerp((float)Dest, Interpolant.Z, Alpha);
+		uint16 DValue = (uint32)FMath::Clamp<float>(Value, 0, LandscapeDataAccess::MaxValue);
+		if ((bRaiseTerrain && DValue > Dest) ||
+			(bLowerTerrain && DValue < Dest))
+		{
+			Dest = DValue;
+		}
+	}
+
+private:
+	TArray<uint16>& Data;
+	int32 MinX, MinY, MaxX, MaxY;
+	uint32 bRaiseTerrain : 1, bLowerTerrain : 1;
+};
+
+class FLandscapeSplineBlendmaskRasterPolicy
+{
+public:
+	// X = Side Alpha, Y = End Alpha, Z = Blend Value
+	typedef FVector InterpolantType;
+
+	/** Initialization constructor. */
+	FLandscapeSplineBlendmaskRasterPolicy(TArray<uint8>& InData, int32 InMinX, int32 InMinY, int32 InMaxX, int32 InMaxY) :
+		Data(InData),
+		MinX(InMinX),
+		MinY(InMinY),
+		MaxX(InMaxX),
+		MaxY(InMaxY)
+	{
+	}
+
+protected:
+
+	// FTriangleRasterizer policy interface.
+
+	int32 GetMinX() const { return MinX; }
+	int32 GetMaxX() const { return MaxX; }
+	int32 GetMinY() const { return MinY; }
+	int32 GetMaxY() const { return MaxY; }
+
+	void ProcessPixel(int32 X, int32 Y, const InterpolantType& Interpolant, bool BackFacing)
+	{
+		const float CosInterpX = (Interpolant.X >= 1 ? 1 : 0.5f - 0.5f * FMath::Cos(Interpolant.X * PI));
+		const float CosInterpY = (Interpolant.Y >= 1 ? 1 : 0.5f - 0.5f * FMath::Cos(Interpolant.Y * PI));
+		const float Alpha = CosInterpX * CosInterpY;
+		uint8& Dest = Data[(Y - MinY)*(1 + MaxX - MinX) + X - MinX];
+		float Value = FMath::Lerp((float)Dest, Interpolant.Z, Alpha);
+		Dest = (uint32)FMath::Clamp<float>(Value, 0, LandscapeDataAccess::MaxValue);
+	}
+
+private:
+	TArray<uint8>& Data;
+	int32 MinX, MinY, MaxX, MaxY;
+};
+
+template<typename T>
+void ShrinkData(TArray<T>& Data, int32 OldMinX, int32 OldMinY, int32 OldMaxX, int32 OldMaxY, int32 NewMinX, int32 NewMinY, int32 NewMaxX, int32 NewMaxY)
+{
+	checkSlow(OldMinX <= OldMaxX && OldMinY <= OldMaxY);
+	checkSlow(NewMinX >= OldMinX && NewMaxX <= OldMaxX);
+	checkSlow(NewMinY >= OldMinY && NewMaxY <= OldMaxY);
+
+	if (NewMinX != OldMinX || NewMinY != OldMinY ||
+		NewMaxX != OldMaxX || NewMaxY != OldMaxY)
+	{
+		// if only the MaxY changes we don't need to do the moving, only the truncate
+		if (NewMinX != OldMinX || NewMinY != OldMinY || NewMaxX != OldMaxX)
+		{
+			for (int32 DestY = 0, SrcY = NewMinY - OldMinY; DestY <= NewMaxY - NewMinY; DestY++, SrcY++)
+			{
+				UE_LOG(LogLandscape, Warning, TEXT("Dest: %d, %d = %d Src: %d, %d = %d Width = %d"), 0, DestY, DestY * (1 + NewMaxX - NewMinX), NewMinX - OldMinX, SrcY, SrcY * (1 + OldMaxX - OldMinX) + NewMinX - OldMinX, (1 + NewMaxX - NewMinX));
+				T* DestData = &Data[DestY * (1 + NewMaxX - NewMinX)];
+				const T* SrcData = &Data[SrcY * (1 + OldMaxX - OldMinX) + NewMinX - OldMinX];
+				FMemory::Memmove(DestData, SrcData, (1 + NewMaxX - NewMinX) * sizeof(T));
+			}
+		}
+
+		const int32 NewSize = (1 + NewMaxY - NewMinY) * (1 + NewMaxX - NewMinX);
+		Data.RemoveAt(NewSize, Data.Num() - NewSize);
+	}
+}
+
+bool ULandscapeInfo::ApplySplines(bool bOnlySelected)
+{
+	bool bResult = false;
+
+	ALandscape* Landscape = LandscapeActor.Get();
+
+	bResult |= ApplySplinesInternal(bOnlySelected, Landscape);
+
+	for (auto It = Proxies.CreateIterator(); It; ++It)
+	{
+		ALandscapeProxy* LandscapeProxy = (*It);
+		bResult |= ApplySplinesInternal(bOnlySelected, LandscapeProxy);
+	}
+
+	return bResult;
+}
+
+bool ULandscapeInfo::ApplySplinesInternal(bool bOnlySelected, ALandscapeProxy* Landscape)
+{
+	if (!Landscape || !Landscape->SplineComponent || Landscape->SplineComponent->ControlPoints.Num() == 0 || Landscape->SplineComponent->Segments.Num() == 0)
+	{
+		return false;
+	}
+
+	FScopedTransaction Transaction(LOCTEXT("LandscapeSpline_ApplySplines", "Apply Splines to Landscape"));
+
+	const FTransform SplineToLandscape = Landscape->SplineComponent->ComponentToWorld.GetRelativeTransform(Landscape->LandscapeActorToWorld());
+
+	FLandscapeEditDataInterface LandscapeEdit(this);
+	TSet<ULandscapeComponent*> ModifiedComponents;
+
+	for (auto It = Landscape->SplineComponent->ControlPoints.CreateConstIterator(); It; ++It)
+	{
+		const ULandscapeSplineControlPoint* ControlPoint = *It;
+
+		if (bOnlySelected && !ControlPoint->IsSplineSelected())
+		{
+			continue;
+		}
+
+		if (ControlPoint->GetPoints().Num() < 2)
+		{
+			continue;
+		}
+
+		FBox LandscapeBounds = ControlPoint->GetBounds();
+		LandscapeBounds = LandscapeBounds.TransformBy(SplineToLandscape.ToMatrixWithScale());
+
+		int32 MinX = FMath::CeilToInt(LandscapeBounds.Min.X);
+		int32 MinY = FMath::CeilToInt(LandscapeBounds.Min.Y);
+		int32 MaxX = FMath::FloorToInt(LandscapeBounds.Max.X);
+		int32 MaxY = FMath::FloorToInt(LandscapeBounds.Max.Y);
+
+		TArray<FLandscapeSplineInterpPoint> Points = ControlPoint->GetPoints();
+		for (int32 j = 0; j < Points.Num(); j++)
+		{
+			Points[j].Center = SplineToLandscape.TransformPosition(Points[j].Center);
+			Points[j].Left = SplineToLandscape.TransformPosition(Points[j].Left);
+			Points[j].Right = SplineToLandscape.TransformPosition(Points[j].Right);
+			Points[j].FalloffLeft = SplineToLandscape.TransformPosition(Points[j].FalloffLeft);
+			Points[j].FalloffRight = SplineToLandscape.TransformPosition(Points[j].FalloffRight);
+
+			// local-heights to texture value heights
+			Points[j].Left.Z = Points[j].Left.Z * LANDSCAPE_INV_ZSCALE + LandscapeDataAccess::MidValue;
+			Points[j].Right.Z = Points[j].Right.Z * LANDSCAPE_INV_ZSCALE + LandscapeDataAccess::MidValue;
+			Points[j].FalloffLeft.Z = Points[j].FalloffLeft.Z * LANDSCAPE_INV_ZSCALE + LandscapeDataAccess::MidValue;
+			Points[j].FalloffRight.Z = Points[j].FalloffRight.Z * LANDSCAPE_INV_ZSCALE + LandscapeDataAccess::MidValue;
+		}
+
+		// Heights raster
+		if (ControlPoint->bRaiseTerrain || ControlPoint->bLowerTerrain)
+		{
+			TArray<uint16> Data;
+			Data.AddZeroed((1 + MaxY - MinY) * (1 + MaxX - MinX));
+
+			int32 ValidMinX = MinX;
+			int32 ValidMinY = MinY;
+			int32 ValidMaxX = MaxX;
+			int32 ValidMaxY = MaxY;
+			LandscapeEdit.GetHeightData(ValidMinX, ValidMinY, ValidMaxX, ValidMaxY, Data.GetData(), 0);
+
+			if (ValidMinX > ValidMaxX || ValidMinY > ValidMaxY)
+			{
+				// The control point's bounds don't intersect any data, so we skip it entirely
+				// it wouldn't intersect any weightmap data either so we don't even bother trying
+				continue;
+			}
+
+			ShrinkData(Data, MinX, MinY, MaxX, MaxY, ValidMinX, ValidMinY, ValidMaxX, ValidMaxY);
+
+			MinX = ValidMinX;
+			MinY = ValidMinY;
+			MaxX = ValidMaxX;
+			MaxY = ValidMaxY;
+
+			FTriangleRasterizer<FLandscapeSplineHeightsRasterPolicy> Rasterizer(
+				FLandscapeSplineHeightsRasterPolicy(Data, MinX, MinY, MaxX, MaxY, ControlPoint->bRaiseTerrain, ControlPoint->bLowerTerrain));
+
+			const FVector Center3D = SplineToLandscape.TransformPosition(ControlPoint->Location);
+			const FVector2D CenterPos = FVector2D(Center3D);
+			const FVector Center = FVector(1.0f, Points[0].StartEndFalloff, Center3D.Z * LANDSCAPE_INV_ZSCALE + LandscapeDataAccess::MidValue);
+
+			for (int32 i = Points.Num() - 1, j = 0; j < Points.Num(); i = j++)
+			{
+				// Solid center
+				const FVector2D Right0Pos = FVector2D(Points[i].Right);
+				const FVector2D Left1Pos = FVector2D(Points[j].Left);
+				const FVector2D Right1Pos = FVector2D(Points[j].Right);
+				const FVector Right0 = FVector(1.0f, Points[i].StartEndFalloff, Points[i].Right.Z);
+				const FVector Left1 = FVector(1.0f, Points[j].StartEndFalloff, Points[j].Left.Z);
+				const FVector Right1 = FVector(1.0f, Points[j].StartEndFalloff, Points[j].Right.Z);
+
+				Rasterizer.DrawTriangle(Center, Right0, Left1, CenterPos, Right0Pos, Left1Pos, false);
+				Rasterizer.DrawTriangle(Center, Left1, Right1, CenterPos, Left1Pos, Right1Pos, false);
+
+				// Falloff
+				FVector2D FalloffRight0Pos = FVector2D(Points[i].FalloffRight);
+				FVector2D FalloffLeft1Pos = FVector2D(Points[j].FalloffLeft);
+				FVector FalloffRight0 = FVector(0.0f, Points[i].StartEndFalloff, Points[i].FalloffRight.Z);
+				FVector FalloffLeft1 = FVector(0.0f, Points[j].StartEndFalloff, Points[j].FalloffLeft.Z);
+				Rasterizer.DrawTriangle(Right0, FalloffRight0, Left1, Right0Pos, FalloffRight0Pos, Left1Pos, false);
+				Rasterizer.DrawTriangle(FalloffRight0, Left1, FalloffLeft1, FalloffRight0Pos, Left1Pos, FalloffLeft1Pos, false);
+			}
+
+			LandscapeEdit.SetHeightData(MinX, MinY, MaxX, MaxY, Data.GetData(), 0, true);
+			LandscapeEdit.GetComponentsInRegion(MinX, MinY, MaxX, MaxY, &ModifiedComponents);
+		}
+
+		// Blend layer raster
+		ULandscapeLayerInfoObject* LayerInfo = GetLayerInfoByName(ControlPoint->LayerName);
+		if (ControlPoint->LayerName != NAME_None && LayerInfo != NULL)
+		{
+			TArray<uint8> Data;
+			Data.AddZeroed((1 + MaxY - MinY) * (1 + MaxX - MinX));
+
+			int32 ValidMinX = MinX;
+			int32 ValidMinY = MinY;
+			int32 ValidMaxX = MaxX;
+			int32 ValidMaxY = MaxY;
+			LandscapeEdit.GetWeightData(LayerInfo, ValidMinX, ValidMinY, ValidMaxX, ValidMaxY, Data.GetData(), 0);
+
+			if (ValidMinX > ValidMaxX || ValidMinY > ValidMaxY)
+			{
+				// The control point's bounds don't intersect any data, so skip it
+				continue;
+			}
+
+			ShrinkData(Data, MinX, MinY, MaxX, MaxY, ValidMinX, ValidMinY, ValidMaxX, ValidMaxY);
+
+			MinX = ValidMinX;
+			MinY = ValidMinY;
+			MaxX = ValidMaxX;
+			MaxY = ValidMaxY;
+
+			FTriangleRasterizer<FLandscapeSplineBlendmaskRasterPolicy> Rasterizer(
+				FLandscapeSplineBlendmaskRasterPolicy(Data, MinX, MinY, MaxX, MaxY));
+
+			const float BlendValue = 255;
+
+			const FVector Center3D = SplineToLandscape.TransformPosition(ControlPoint->Location);
+			const FVector2D CenterPos = FVector2D(Center3D);
+			const FVector Center = FVector(1.0f, Points[0].StartEndFalloff, BlendValue);
+
+			for (int32 i = Points.Num() - 1, j = 0; j < Points.Num(); i = j++)
+			{
+				// Solid center
+				const FVector2D Right0Pos = FVector2D(Points[i].Right);
+				const FVector2D Left1Pos = FVector2D(Points[j].Left);
+				const FVector2D Right1Pos = FVector2D(Points[j].Right);
+				const FVector Right0 = FVector(1.0f, Points[i].StartEndFalloff, BlendValue);
+				const FVector Left1 = FVector(1.0f, Points[j].StartEndFalloff, BlendValue);
+				const FVector Right1 = FVector(1.0f, Points[j].StartEndFalloff, BlendValue);
+
+				Rasterizer.DrawTriangle(Center, Right0, Left1, CenterPos, Right0Pos, Left1Pos, false);
+				Rasterizer.DrawTriangle(Center, Left1, Right1, CenterPos, Left1Pos, Right1Pos, false);
+
+				// Falloff
+				FVector2D FalloffRight0Pos = FVector2D(Points[i].FalloffRight);
+				FVector2D FalloffLeft1Pos = FVector2D(Points[j].FalloffLeft);
+				FVector FalloffRight0 = FVector(0.0f, Points[i].StartEndFalloff, BlendValue);
+				FVector FalloffLeft1 = FVector(0.0f, Points[j].StartEndFalloff, BlendValue);
+				Rasterizer.DrawTriangle(Right0, FalloffRight0, Left1, Right0Pos, FalloffRight0Pos, Left1Pos, false);
+				Rasterizer.DrawTriangle(FalloffRight0, Left1, FalloffLeft1, FalloffRight0Pos, Left1Pos, FalloffLeft1Pos, false);
+			}
+
+			LandscapeEdit.SetAlphaData(LayerInfo, MinX, MinY, MaxX, MaxY, Data.GetData(), 0, ELandscapeLayerPaintingRestriction::None, true, false);
+			LandscapeEdit.GetComponentsInRegion(MinX, MinY, MaxX, MaxY, &ModifiedComponents);
+		}
+	}
+
+	for (auto It = Landscape->SplineComponent->Segments.CreateConstIterator(); It; ++It)
+	{
+		const ULandscapeSplineSegment* Segment = *It;
+
+		if (bOnlySelected && !Segment->IsSplineSelected())
+		{
+			continue;
+		}
+
+		FBox LandscapeBounds = Segment->GetBounds();
+		LandscapeBounds = LandscapeBounds.TransformBy(SplineToLandscape.ToMatrixWithScale());
+
+		int32 MinX = FMath::CeilToInt(LandscapeBounds.Min.X);
+		int32 MinY = FMath::CeilToInt(LandscapeBounds.Min.Y);
+		int32 MaxX = FMath::FloorToInt(LandscapeBounds.Max.X);
+		int32 MaxY = FMath::FloorToInt(LandscapeBounds.Max.Y);
+
+		TArray<FLandscapeSplineInterpPoint> Points = Segment->GetPoints();
+		for (int32 j = 0; j < Points.Num(); j++)
+		{
+			Points[j].Center = SplineToLandscape.TransformPosition(Points[j].Center);
+			Points[j].Left = SplineToLandscape.TransformPosition(Points[j].Left);
+			Points[j].Right = SplineToLandscape.TransformPosition(Points[j].Right);
+			Points[j].FalloffLeft = SplineToLandscape.TransformPosition(Points[j].FalloffLeft);
+			Points[j].FalloffRight = SplineToLandscape.TransformPosition(Points[j].FalloffRight);
+
+			// local-heights to texture value heights
+			Points[j].Left.Z = Points[j].Left.Z * LANDSCAPE_INV_ZSCALE + LandscapeDataAccess::MidValue;
+			Points[j].Right.Z = Points[j].Right.Z * LANDSCAPE_INV_ZSCALE + LandscapeDataAccess::MidValue;
+			Points[j].FalloffLeft.Z = Points[j].FalloffLeft.Z * LANDSCAPE_INV_ZSCALE + LandscapeDataAccess::MidValue;
+			Points[j].FalloffRight.Z = Points[j].FalloffRight.Z * LANDSCAPE_INV_ZSCALE + LandscapeDataAccess::MidValue;
+		}
+
+		// Heights raster
+		if (Segment->bRaiseTerrain || Segment->bLowerTerrain)
+		{
+			TArray<uint16> Data;
+			Data.AddZeroed((1 + MaxY - MinY) * (1 + MaxX - MinX));
+
+			int32 ValidMinX = MinX;
+			int32 ValidMinY = MinY;
+			int32 ValidMaxX = MaxX;
+			int32 ValidMaxY = MaxY;
+			LandscapeEdit.GetHeightData(ValidMinX, ValidMinY, ValidMaxX, ValidMaxY, Data.GetData(), 0);
+
+			if (ValidMinX > ValidMaxX || ValidMinY > ValidMaxY)
+			{
+				// The segment's bounds don't intersect any data, so we skip it entirely
+				// it wouldn't intersect any weightmap data either so we don't even bother trying
+				continue;
+			}
+
+			ShrinkData(Data, MinX, MinY, MaxX, MaxY, ValidMinX, ValidMinY, ValidMaxX, ValidMaxY);
+
+			MinX = ValidMinX;
+			MinY = ValidMinY;
+			MaxX = ValidMaxX;
+			MaxY = ValidMaxY;
+
+			FTriangleRasterizer<FLandscapeSplineHeightsRasterPolicy> Rasterizer(
+				FLandscapeSplineHeightsRasterPolicy(Data, MinX, MinY, MaxX, MaxY, Segment->bRaiseTerrain, Segment->bLowerTerrain));
+
+			for (int32 j = 1; j < Points.Num(); j++)
+			{
+				// Middle
+				FVector2D Left0Pos = FVector2D(Points[j - 1].Left);
+				FVector2D Right0Pos = FVector2D(Points[j - 1].Right);
+				FVector2D Left1Pos = FVector2D(Points[j].Left);
+				FVector2D Right1Pos = FVector2D(Points[j].Right);
+				FVector Left0 = FVector(1.0f, Points[j - 1].StartEndFalloff, Points[j - 1].Left.Z);
+				FVector Right0 = FVector(1.0f, Points[j - 1].StartEndFalloff, Points[j - 1].Right.Z);
+				FVector Left1 = FVector(1.0f, Points[j].StartEndFalloff, Points[j].Left.Z);
+				FVector Right1 = FVector(1.0f, Points[j].StartEndFalloff, Points[j].Right.Z);
+				Rasterizer.DrawTriangle(Left0, Right0, Left1, Left0Pos, Right0Pos, Left1Pos, false);
+				Rasterizer.DrawTriangle(Right0, Left1, Right1, Right0Pos, Left1Pos, Right1Pos, false);
+
+				// Left Falloff
+				FVector2D FalloffLeft0Pos = FVector2D(Points[j - 1].FalloffLeft);
+				FVector2D FalloffLeft1Pos = FVector2D(Points[j].FalloffLeft);
+				FVector FalloffLeft0 = FVector(0.0f, Points[j - 1].StartEndFalloff, Points[j - 1].FalloffLeft.Z);
+				FVector FalloffLeft1 = FVector(0.0f, Points[j].StartEndFalloff, Points[j].FalloffLeft.Z);
+				Rasterizer.DrawTriangle(FalloffLeft0, Left0, FalloffLeft1, FalloffLeft0Pos, Left0Pos, FalloffLeft1Pos, false);
+				Rasterizer.DrawTriangle(Left0, FalloffLeft1, Left1, Left0Pos, FalloffLeft1Pos, Left1Pos, false);
+
+				// Right Falloff
+				FVector2D FalloffRight0Pos = FVector2D(Points[j - 1].FalloffRight);
+				FVector2D FalloffRight1Pos = FVector2D(Points[j].FalloffRight);
+				FVector FalloffRight0 = FVector(0.0f, Points[j - 1].StartEndFalloff, Points[j - 1].FalloffRight.Z);
+				FVector FalloffRight1 = FVector(0.0f, Points[j].StartEndFalloff, Points[j].FalloffRight.Z);
+				Rasterizer.DrawTriangle(Right0, FalloffRight0, Right1, Right0Pos, FalloffRight0Pos, Right1Pos, false);
+				Rasterizer.DrawTriangle(FalloffRight0, Right1, FalloffRight1, FalloffRight0Pos, Right1Pos, FalloffRight1Pos, false);
+			}
+
+			LandscapeEdit.SetHeightData(MinX, MinY, MaxX, MaxY, Data.GetData(), 0, true);
+			LandscapeEdit.GetComponentsInRegion(MinX, MinY, MaxX, MaxY, &ModifiedComponents);
+		}
+
+		// Blend layer raster
+		ULandscapeLayerInfoObject* LayerInfo = GetLayerInfoByName(Segment->LayerName);
+		if (Segment->LayerName != NAME_None && LayerInfo != NULL)
+		{
+			TArray<uint8> Data;
+			Data.AddZeroed((1 + MaxY - MinY) * (1 + MaxX - MinX));
+
+			int32 ValidMinX = MinX;
+			int32 ValidMinY = MinY;
+			int32 ValidMaxX = MaxX;
+			int32 ValidMaxY = MaxY;
+			LandscapeEdit.GetWeightData(LayerInfo, ValidMinX, ValidMinY, ValidMaxX, ValidMaxY, Data.GetData(), 0);
+
+			if (ValidMinX > ValidMaxX || ValidMinY > ValidMaxY)
+			{
+				// The segment's bounds don't intersect any data, so skip it
+				continue;
+			}
+
+			ShrinkData(Data, MinX, MinY, MaxX, MaxY, ValidMinX, ValidMinY, ValidMaxX, ValidMaxY);
+
+			MinX = ValidMinX;
+			MinY = ValidMinY;
+			MaxX = ValidMaxX;
+			MaxY = ValidMaxY;
+
+			FTriangleRasterizer<FLandscapeSplineBlendmaskRasterPolicy> Rasterizer(
+				FLandscapeSplineBlendmaskRasterPolicy(Data, MinX, MinY, MaxX, MaxY));
+
+			const float BlendValue = 255;
+
+			for (int32 j = 1; j < Points.Num(); j++)
+			{
+				// Middle
+				FVector2D Left0Pos = FVector2D(Points[j - 1].Left);
+				FVector2D Right0Pos = FVector2D(Points[j - 1].Right);
+				FVector2D Left1Pos = FVector2D(Points[j].Left);
+				FVector2D Right1Pos = FVector2D(Points[j].Right);
+				FVector Left0 = FVector(1.0f, Points[j - 1].StartEndFalloff, BlendValue);
+				FVector Right0 = FVector(1.0f, Points[j - 1].StartEndFalloff, BlendValue);
+				FVector Left1 = FVector(1.0f, Points[j].StartEndFalloff, BlendValue);
+				FVector Right1 = FVector(1.0f, Points[j].StartEndFalloff, BlendValue);
+				Rasterizer.DrawTriangle(Left0, Right0, Left1, Left0Pos, Right0Pos, Left1Pos, false);
+				Rasterizer.DrawTriangle(Right0, Left1, Right1, Right0Pos, Left1Pos, Right1Pos, false);
+
+				// Left Falloff
+				FVector2D FalloffLeft0Pos = FVector2D(Points[j - 1].FalloffLeft);
+				FVector2D FalloffLeft1Pos = FVector2D(Points[j].FalloffLeft);
+				FVector FalloffLeft0 = FVector(0.0f, Points[j - 1].StartEndFalloff, BlendValue);
+				FVector FalloffLeft1 = FVector(0.0f, Points[j].StartEndFalloff, BlendValue);
+				Rasterizer.DrawTriangle(FalloffLeft0, Left0, FalloffLeft1, FalloffLeft0Pos, Left0Pos, FalloffLeft1Pos, false);
+				Rasterizer.DrawTriangle(Left0, FalloffLeft1, Left1, Left0Pos, FalloffLeft1Pos, Left1Pos, false);
+
+				// Right Falloff
+				FVector2D FalloffRight0Pos = FVector2D(Points[j - 1].FalloffRight);
+				FVector2D FalloffRight1Pos = FVector2D(Points[j].FalloffRight);
+				FVector FalloffRight0 = FVector(0.0f, Points[j - 1].StartEndFalloff, BlendValue);
+				FVector FalloffRight1 = FVector(0.0f, Points[j].StartEndFalloff, BlendValue);
+				Rasterizer.DrawTriangle(Right0, FalloffRight0, Right1, Right0Pos, FalloffRight0Pos, Right1Pos, false);
+				Rasterizer.DrawTriangle(FalloffRight0, Right1, FalloffRight1, FalloffRight0Pos, Right1Pos, FalloffRight1Pos, false);
+			}
+
+			LandscapeEdit.SetAlphaData(LayerInfo, MinX, MinY, MaxX, MaxY, Data.GetData(), 0, ELandscapeLayerPaintingRestriction::None, true, false);
+			LandscapeEdit.GetComponentsInRegion(MinX, MinY, MaxX, MaxY, &ModifiedComponents);
+		}
+	}
+
+	LandscapeEdit.Flush();
+
+	for (ULandscapeComponent* Component : ModifiedComponents)
+	{
+		// Recreate collision for modified components and update the navmesh
+		ULandscapeHeightfieldCollisionComponent* CollisionComponent = Component->CollisionComponent.Get();
+		if (CollisionComponent)
+		{
+			CollisionComponent->RecreateCollision(false);
+			UNavigationSystem* NavSys = UNavigationSystem::GetCurrent(Component);
+			if (NavSys)
+			{
+				NavSys->UpdateNavOctree(CollisionComponent);
+			}
+		}
+	}
+
+	return true;
+}
+
+
+void ULandscapeInfo::DeleteLayer(ULandscapeLayerInfoObject* LayerInfo)
+{
+	GWarn->BeginSlowTask(LOCTEXT("BeginDeletingLayerTask", "Deleting Layer"), true);
+
+	// Remove data from all components
+	FLandscapeEditDataInterface LandscapeEdit(this);
+	LandscapeEdit.DeleteLayer(LayerInfo);
+
+	// Remove from array
+	for (int32 j = 0; j < Layers.Num(); j++)
+	{
+		if (Layers[j].LayerInfoObj && Layers[j].LayerInfoObj == LayerInfo)
+		{
+			Layers.RemoveAt(j);
+			break;
+		}
+	}
+
+	ALandscape* Landscape = LandscapeActor.Get();
+	if (Landscape != NULL)
+	{
+		Landscape->Modify();
+		Landscape->EditorLayerSettings.Remove(LayerInfo);
+	}
+
+	for (auto It = Proxies.CreateConstIterator(); It; ++It)
+	{
+		ALandscapeProxy* Proxy = *It;
+		Proxy->Modify();
+		Proxy->EditorLayerSettings.Remove(LayerInfo);
+	}
+
+	//UpdateLayerInfoMap();
+
+	GWarn->EndSlowTask();
+}
+
+void ULandscapeInfo::ReplaceLayer(ULandscapeLayerInfoObject* FromLayerInfo, ULandscapeLayerInfoObject* ToLayerInfo)
+{
+	if (ensure(FromLayerInfo != ToLayerInfo))
+	{
+		GWarn->BeginSlowTask(LOCTEXT("BeginReplacingLayerTask", "Replacing Layer"), true);
+
+		// Remove data from all components
+		FLandscapeEditDataInterface LandscapeEdit(this);
+		LandscapeEdit.ReplaceLayer(FromLayerInfo, ToLayerInfo);
+
+		// Convert array
+		for (int32 j = 0; j < Layers.Num(); j++)
+		{
+			if (Layers[j].LayerInfoObj && Layers[j].LayerInfoObj == FromLayerInfo)
+			{
+				Layers[j].LayerInfoObj = ToLayerInfo;
+			}
+		}
+
+		ALandscape* Landscape = LandscapeActor.Get();
+		if (Landscape != NULL)
+		{
+			Landscape->Modify();
+			FLandscapeEditorLayerSettings* ToEditorLayerSettings = Landscape->EditorLayerSettings.FindByKey(ToLayerInfo);
+			if (ToEditorLayerSettings != NULL)
+			{
+				// If the new layer already exists, simple remove the old layer
+				Landscape->EditorLayerSettings.Remove(FromLayerInfo);
+			}
+			else
+			{
+				FLandscapeEditorLayerSettings* FromEditorLayerSettings = Landscape->EditorLayerSettings.FindByKey(FromLayerInfo);
+				if (FromEditorLayerSettings != NULL)
+				{
+					// If only the old layer exists (most common case), change it to point to the new layer info
+					FromEditorLayerSettings->LayerInfoObj = ToLayerInfo;
+				}
+				else
+				{
+					// If neither exists in the EditorLayerSettings cache, add it
+					Landscape->EditorLayerSettings.Add(ToLayerInfo);
+				}
+			}
+		}
+
+		for (auto It = Proxies.CreateConstIterator(); It; ++It)
+		{
+			ALandscapeProxy* Proxy = *It;
+			Proxy->Modify();
+			FLandscapeEditorLayerSettings* ToEditorLayerSettings = Proxy->EditorLayerSettings.FindByKey(ToLayerInfo);
+			if (ToEditorLayerSettings != NULL)
+			{
+				// If the new layer already exists, simple remove the old layer
+				Proxy->EditorLayerSettings.Remove(FromLayerInfo);
+			}
+			else
+			{
+				FLandscapeEditorLayerSettings* FromEditorLayerSettings = Proxy->EditorLayerSettings.FindByKey(FromLayerInfo);
+				if (FromEditorLayerSettings != NULL)
+				{
+					// If only the old layer exists (most common case), change it to point to the new layer info
+					FromEditorLayerSettings->LayerInfoObj = ToLayerInfo;
+				}
+				else
+				{
+					// If neither exists in the EditorLayerSettings cache, add it
+					Proxy->EditorLayerSettings.Add(ToLayerInfo);
+				}
+			}
+		}
+
+		//UpdateLayerInfoMap();
+
+		GWarn->EndSlowTask();
+	}
+}
+
+void ALandscapeProxy::EditorApplyScale(const FVector& DeltaScale, const FVector* PivotLocation, bool bAltDown, bool bShiftDown, bool bCtrlDown)
+{
+	FVector ModifiedScale = DeltaScale;
+
+	// Lock X and Y scaling to the same value
+	ModifiedScale.X = ModifiedScale.Y = (FMath::Abs(DeltaScale.X) > FMath::Abs(DeltaScale.Y)) ? DeltaScale.X : DeltaScale.Y;
+
+	// Correct for attempts to scale to 0 on any axis
+	FVector CurrentScale = GetRootComponent()->RelativeScale3D;
+	if (AActor::bUsePercentageBasedScaling)
+	{
+		if (ModifiedScale.X == -1)
+		{
+			ModifiedScale.X = ModifiedScale.Y = -(CurrentScale.X - 1) / CurrentScale.X;
+		}
+		if (ModifiedScale.Z == -1)
+		{
+			ModifiedScale.Z = -(CurrentScale.Z - 1) / CurrentScale.Z;
+		}
+	}
+	else
+	{
+		if (ModifiedScale.X == -CurrentScale.X)
+		{
+			CurrentScale.X += 1;
+			CurrentScale.Y += 1;
+		}
+		if (ModifiedScale.Z == -CurrentScale.Z)
+		{
+			CurrentScale.Z += 1;
+		}
+	}
+
+	Super::EditorApplyScale(ModifiedScale, PivotLocation, bAltDown, bShiftDown, bCtrlDown);
+}
+
+void ALandscapeProxy::PostEditMove(bool bFinished)
+{
+	// This point is only reached when Copy and Pasted
+	Super::PostEditMove(bFinished);
+
+	if (bFinished)
+	{
+		ULandscapeInfo::RecreateLandscapeInfo(GetWorld(), true);
+		RecreateComponentsState();
+	}
+
+	if (bFinished)
+	{
+		if (SplineComponent)
+		{
+			SplineComponent->CheckSplinesValid();
+		}
+	}
+}
+
+void ALandscapeProxy::PostEditImport()
+{
+	Super::PostEditImport();
+	if (!bIsProxy && GetWorld()) // For Landscape
+	{
+		for (TActorIterator<ALandscape> It(GetWorld()); It; ++It)
+		{
+			ALandscape* Landscape = *It;
+			if (Landscape != this && !Landscape->HasAnyFlags(RF_BeginDestroyed) && Landscape->LandscapeGuid == LandscapeGuid)
+			{
+				// Copy/Paste case, need to generate new GUID
+				LandscapeGuid = FGuid::NewGuid();
+			}
+		}
+	}
+
+	for (int32 ComponentIndex = 0; ComponentIndex < LandscapeComponents.Num(); ++ComponentIndex)
+	{
+		ULandscapeComponent* Comp = LandscapeComponents[ComponentIndex];
+		if (Comp)
+		{
+			// Update the MIC
+			Comp->UpdateMaterialInstances();
+		}
+	}
+
+	ULandscapeInfo::RecreateLandscapeInfo(GetWorld(), true);
+}
+
+void ALandscape::PostEditMove(bool bFinished)
+{
+	if (bFinished)
+	{
+		// align all proxies to landscape actor
+		GetLandscapeInfo()->FixupProxiesTransform();
+	}
+
+	Super::PostEditMove(bFinished);
+}
+#endif	//WITH_EDITOR
+
+ULandscapeLayerInfoObject::ULandscapeLayerInfoObject(const FObjectInitializer& ObjectInitializer)
+	: Super(ObjectInitializer)
+{
+	Hardness = 0.5f;
+#if WITH_EDITORONLY_DATA
+	bNoWeightBlend = false;
+#endif // WITH_EDITORONLY_DATA
+}
+
+#if WITH_EDITOR
+void ULandscapeLayerInfoObject::PostEditChangeProperty(FPropertyChangedEvent& PropertyChangedEvent)
+{
+	static const FName NAME_Hardness = FName(TEXT("Hardness"));
+	static const FName NAME_PhysMaterial = FName(TEXT("PhysMaterial"));
+
+	Super::PostEditChangeProperty(PropertyChangedEvent);
+
+	const FName PropertyName = PropertyChangedEvent.Property ? PropertyChangedEvent.Property->GetFName() : NAME_None;
+
+	if (GIsEditor)
+	{
+		if (PropertyName == NAME_Hardness)
+		{
+			Hardness = FMath::Clamp<float>(Hardness, 0.f, 1.f);
+		}
+		else if (PropertyName == NAME_PhysMaterial)
+		{
+			// Only care current world object
+			for (TActorIterator<ALandscapeProxy> It(GWorld); It; ++It)
+			{
+				ALandscapeProxy* Proxy = *It;
+				ULandscapeInfo* Info = Proxy->GetLandscapeInfo(false);
+				if (Info)
+				{
+					for (int32 i = 0; i < Info->Layers.Num(); ++i)
+					{
+						if (Info->Layers[i].LayerInfoObj == this)
+						{
+							Proxy->ChangedPhysMaterial();
+							break;
+						}
+					}
+				}
+			}
+		}
+	}
+}
+
+void ULandscapeLayerInfoObject::PostLoad()
+{
+	Super::PostLoad();
+	if (GIsEditor)
+	{
+		if (!HasAnyFlags(RF_Standalone))
+		{
+			SetFlags(RF_Standalone);
+		}
+		Hardness = FMath::Clamp<float>(Hardness, 0.f, 1.f);
+	}
+}
+
+void ALandscapeProxy::RecreateCollisionComponents()
+{
+	// Clear old CollisionComponent containers
+	CollisionComponents.Empty();
+
+	// Clear any Owned Collision Components
+	TArray<ULandscapeHeightfieldCollisionComponent*> CollisionComps;
+	GetComponents(CollisionComps);
+	for (int32 i = 0; i < CollisionComps.Num(); ++i)
+	{
+		RemoveOwnedComponent(CollisionComps[i]);
+	}
+
+	CollisionMipLevel = FMath::Clamp<int32>(CollisionMipLevel, 0, FMath::CeilLogTwo(SubsectionSizeQuads + 1) - 1);
+	for (int32 i = 0; i < LandscapeComponents.Num(); ++i)
+	{
+		ULandscapeComponent* Comp = LandscapeComponents[i];
+		if (Comp)
+		{
+			Comp->CollisionMipLevel = CollisionMipLevel;
+			TArray<uint8> CollisionMipData;
+			Comp->HeightmapTexture->Source.GetMipData(CollisionMipData, CollisionMipLevel);
+			Comp->UpdateCollisionHeightData((FColor*)CollisionMipData.GetData(), 0, 0, MAX_int32, MAX_int32, true, NULL, true); // Rebuild for new CollisionMipLevel
+		}
+	}
+}
+
+void ULandscapeInfo::RecreateCollisionComponents()
+{
+	if (LandscapeActor.IsValid())
+	{
+		LandscapeActor->RecreateCollisionComponents();
+	}
+
+	for (auto It = Proxies.CreateConstIterator(); It; ++It)
+	{
+		ALandscapeProxy* Proxy = (*It);
+		Proxy->RecreateCollisionComponents();
+	}
+}
+
+namespace
+{
+	inline float AdjustStaticLightingResolution(float StaticLightingResolution, int32 NumSubsections, int32 SubsectionSizeQuads, int32 ComponentSizeQuads)
+	{
+		// Change Lighting resolution to proper one...
+		if (StaticLightingResolution > 1.f)
+		{
+			StaticLightingResolution = (int32)StaticLightingResolution;
+		}
+		else if (StaticLightingResolution < 1.f)
+		{
+			// Restrict to 1/16
+			if (StaticLightingResolution < 0.0625)
+			{
+				StaticLightingResolution = 0.0625;
+			}
+
+			// Adjust to 1/2^n
+			int32 i = 2;
+			int32 LightmapSize = (NumSubsections * (SubsectionSizeQuads + 1)) >> 1;
+			while (StaticLightingResolution < (1.f / i) && LightmapSize > 4)
+			{
+				i <<= 1;
+				LightmapSize >>= 1;
+			}
+			StaticLightingResolution = 1.f / i;
+
+			int32 PixelPaddingX = GPixelFormats[PF_DXT1].BlockSizeX;
+
+			int32 DestSize = (int32)((2 * PixelPaddingX + ComponentSizeQuads + 1) * StaticLightingResolution);
+			StaticLightingResolution = (float)DestSize / (2 * PixelPaddingX + ComponentSizeQuads + 1);
+		}
+
+		return StaticLightingResolution;
+	}
+};
+
+void ALandscapeProxy::PostEditChangeProperty(FPropertyChangedEvent& PropertyChangedEvent)
+{
+	Super::PostEditChangeProperty(PropertyChangedEvent);
+	const FName PropertyName = PropertyChangedEvent.MemberProperty ? PropertyChangedEvent.MemberProperty->GetFName() : NAME_None;
+	if (bIsProxy)
+	{
+		if (PropertyName == FName(TEXT("LandscapeActor")))
+		{
+			if (LandscapeActor && IsValidLandscapeActor(LandscapeActor.Get()))
+			{
+				// defer LandscapeInfo setup
+				if (GIsEditor && GetWorld() && !GetWorld()->IsPlayInEditor())
+				{
+					GEngine->DeferredCommands.AddUnique(TEXT("UpdateLandscapeEditorData"));
+				}
+			}
+			else
+			{
+				LandscapeActor = 0;
+			}
+		}
+		else if (PropertyName == FName(TEXT("LandscapeMaterial")) || PropertyName == FName(TEXT("LandscapeHoleMaterial")))
+		{
+			FMaterialUpdateContext MaterialUpdateContext;
+			GetLandscapeInfo()->UpdateLayerInfoMap(/*this*/);
+
+			// Clear the parents out of combination material instances
+			for (TMap<FString, UMaterialInstanceConstant*>::TIterator It(MaterialInstanceConstantMap); It; ++It)
+			{
+				It.Value()->SetParentEditorOnly(NULL);
+				MaterialUpdateContext.AddMaterial(It.Value()->GetMaterial());
+			}
+
+			// Remove our references to any material instances
+			MaterialInstanceConstantMap.Empty();
+
+			for (int32 ComponentIndex = 0; ComponentIndex < LandscapeComponents.Num(); ComponentIndex++)
+			{
+				ULandscapeComponent* Comp = LandscapeComponents[ComponentIndex];
+				if (Comp)
+				{
+					// Update the MIC
+					Comp->UpdateMaterialInstances();
+				}
+			}
+		}
+	}
+
+	if (GIsEditor && PropertyName == FName(TEXT("StreamingDistanceMultiplier")))
+	{
+		// Recalculate in a few seconds.
+		GetWorld()->TriggerStreamingDataRebuild();
+	}
+	else if (GIsEditor && PropertyName == FName(TEXT("DefaultPhysMaterial")))
+	{
+		ChangedPhysMaterial();
+	}
+	else if (GIsEditor && (PropertyName == FName(TEXT("CollisionMipLevel"))))
+	{
+		RecreateCollisionComponents();
+	}
+	else if (PropertyName == FName(TEXT("bCastStaticShadow")) || PropertyName == FName(TEXT("bCastShadowAsTwoSided")))
+	{
+		// Replicate shared properties to all components.
+		for (int32 ComponentIndex = 0; ComponentIndex < LandscapeComponents.Num(); ComponentIndex++)
+		{
+			ULandscapeComponent* Comp = LandscapeComponents[ComponentIndex];
+			if (Comp)
+			{
+				Comp->bCastStaticShadow = bCastStaticShadow;
+				Comp->bCastShadowAsTwoSided = bCastShadowAsTwoSided;
+			}
+		}
+	}
+}
+
+void ALandscapeProxy::PostEditChangeChainProperty(FPropertyChangedChainEvent& PropertyChangedEvent)
+{
+	const FName MemberPropertyName = PropertyChangedEvent.PropertyChain.GetActiveMemberNode()->GetValue()->GetFName();
+	const FName PropertyName = PropertyChangedEvent.PropertyChain.GetActiveNode()->GetValue()->GetFName();
+
+	if (MemberPropertyName == FName(TEXT("RelativeScale3D")))
+	{
+		// RelativeScale3D isn't even a property of ALandscapeProxy, it's a property of the root component
+		if (RootComponent)
+		{
+			FVector ModifiedScale = RootComponent->RelativeScale3D;
+
+			// Lock X and Y scaling to the same value
+			if (PropertyName == FName("Y"))
+			{
+				ModifiedScale.X = FMath::Abs(RootComponent->RelativeScale3D.Y)*FMath::Sign(ModifiedScale.X);
+			}
+			else
+			{
+				// There's no "if name == X" here so that if we can't tell which has changed out of X and Y, we just use X
+				ModifiedScale.Y = FMath::Abs(RootComponent->RelativeScale3D.X)*FMath::Sign(ModifiedScale.Y);
+			}
+
+			ULandscapeInfo* Info = GetLandscapeInfo(false);
+
+			// Correct for attempts to scale to 0 on any axis
+			if (ModifiedScale.X == 0)
+			{
+				if (Info && Info->DrawScale.X < 0)
+				{
+					ModifiedScale.Y = ModifiedScale.X = -1;
+				}
+				else
+				{
+					ModifiedScale.Y = ModifiedScale.X = 1;
+				}
+			}
+			if (ModifiedScale.Z == 0)
+			{
+				if (Info && Info->DrawScale.Z < 0)
+				{
+					ModifiedScale.Z = -1;
+				}
+				else
+				{
+					ModifiedScale.Z = 1;
+				}
+			}
+
+			RootComponent->SetRelativeScale3D(ModifiedScale);
+
+			// Update ULandscapeInfo cached DrawScale
+			if (Info)
+			{
+				Info->DrawScale = ModifiedScale;
+			}
+
+			// We need to regenerate collision objects, they depend on scale value 
+			for (int32 ComponentIndex = 0; ComponentIndex < CollisionComponents.Num(); ComponentIndex++)
+			{
+				ULandscapeHeightfieldCollisionComponent* Comp = CollisionComponents[ComponentIndex];
+				if (Comp)
+				{
+					Comp->RecreateCollision(false);
+				}
+			}
+		}
+	}
+
+	// Must do this *after* correcting the scale or reattaching the landscape components will crash!
+	Super::PostEditChangeChainProperty(PropertyChangedEvent);
+}
+
+void ALandscape::PostEditChangeProperty(FPropertyChangedEvent& PropertyChangedEvent)
+{
+	Super::PostEditChangeProperty(PropertyChangedEvent);
+
+	const FName PropertyName = PropertyChangedEvent.Property ? PropertyChangedEvent.Property->GetFName() : NAME_None;
+
+	bool ChangedMaterial = false;
+	bool bNeedsRecalcBoundingBox = false;
+	bool bChangedLighting = false;
+	bool bChangedNavRelevance = false;
+	bool bPropagateToProxies = false;
+
+	ULandscapeInfo* Info = GetLandscapeInfo();
+
+	if (PropertyName == FName(TEXT("LandscapeMaterial")) || PropertyName == FName(TEXT("LandscapeHoleMaterial")))
+	{
+		FMaterialUpdateContext MaterialUpdateContext;
+		GetLandscapeInfo()->UpdateLayerInfoMap(/*this*/);
+
+		ChangedMaterial = true;
+
+		// Clear the parents out of combination material instances
+		for (TMap<FString, UMaterialInstanceConstant*>::TIterator It(MaterialInstanceConstantMap); It; ++It)
+		{
+			It.Value()->SetParentEditorOnly(NULL);
+			MaterialUpdateContext.AddMaterial(It.Value()->GetMaterial());
+		}
+
+		// Remove our references to any material instances
+		MaterialInstanceConstantMap.Empty();
+	}
+	else if (PropertyName == FName(TEXT("RelativeScale3D")) ||
+		PropertyName == FName(TEXT("RelativeLocation")) ||
+		PropertyName == FName(TEXT("RelativeRotation")))
+	{
+		// update transformations for all linked proxies 
+		Info->FixupProxiesTransform();
+		bNeedsRecalcBoundingBox = true;
+	}
+	else if (GIsEditor && PropertyName == FName(TEXT("MaxLODLevel")))
+	{
+		MaxLODLevel = FMath::Min<int32>(MaxLODLevel, FMath::CeilLogTwo(SubsectionSizeQuads + 1) - 1);
+		bPropagateToProxies = true;
+	}
+	else if (PropertyName == FName(TEXT("LODDistanceFactor")))
+	{
+		LODDistanceFactor = FMath::Clamp<float>(LODDistanceFactor, 0.1f, MAX_LANDSCAPE_LOD_DISTANCE_FACTOR); // limit because LOD transition became too popping...
+		bPropagateToProxies = true;
+	}
+	else if (PropertyName == FName(TEXT("CollisionMipLevel")))
+	{
+		CollisionMipLevel = FMath::Clamp<int32>(CollisionMipLevel, 0, FMath::CeilLogTwo(SubsectionSizeQuads + 1) - 1);
+		bPropagateToProxies = true;
+	}
+	else if (PropertyName == FName(TEXT("LODFalloff")))
+	{
+		bPropagateToProxies = true;
+	}
+	else if (GIsEditor && PropertyName == FName(TEXT("StaticLightingResolution")))
+	{
+		StaticLightingResolution = ::AdjustStaticLightingResolution(StaticLightingResolution, NumSubsections, SubsectionSizeQuads, ComponentSizeQuads);
+		bChangedLighting = true;
+	}
+	else if (GIsEditor && PropertyName == GET_MEMBER_NAME_CHECKED(ALandscapeProxy, StaticLightingLOD))
+	{
+		StaticLightingLOD = FMath::Clamp<int32>(StaticLightingLOD, 0, FMath::CeilLogTwo(SubsectionSizeQuads + 1) - 1);
+		bChangedLighting = true;
+	}
+	else if (GIsEditor && PropertyName == GET_MEMBER_NAME_CHECKED(ALandscapeProxy, ExportLOD))
+	{
+		ExportLOD = FMath::Clamp<int32>(ExportLOD, 0, FMath::CeilLogTwo(SubsectionSizeQuads + 1) - 1);
+	}
+	else if (GIsEditor && PropertyName == GET_MEMBER_NAME_CHECKED(ALandscapeProxy, bUsedForNavigation))
+	{
+		bChangedNavRelevance = true;
+	}
+
+	bPropagateToProxies = bPropagateToProxies || bNeedsRecalcBoundingBox || bChangedLighting;
+
+	if (Info)
+	{
+		if (bPropagateToProxies)
+		{
+			// Propagate Event to Proxies...
+			for (TSet<ALandscapeProxy*>::TIterator It(Info->Proxies); It; ++It)
+			{
+				(*It)->GetSharedProperties(this);
+				(*It)->PostEditChangeProperty(PropertyChangedEvent);
+			}
+		}
+
+		// Update normals if DrawScale3D is changed
+		if (PropertyName == FName(TEXT("RelativeScale3D")))
+		{
+			FLandscapeEditDataInterface LandscapeEdit(Info);
+			LandscapeEdit.RecalculateNormals();
+		}
+
+		TArray<ULandscapeComponent*> AllComponents;
+		Info->XYtoComponentMap.GenerateValueArray(AllComponents);
+
+		// We cannot iterate the XYtoComponentMap directly because reregistering components modifies the array.
+		for (auto It = AllComponents.CreateIterator(); It; ++It)
+		{
+			ULandscapeComponent* Comp = *It;
+			if (Comp)
+			{
+				if (bNeedsRecalcBoundingBox)
+				{
+					Comp->UpdateCachedBounds();
+					Comp->UpdateBounds();
+				}
+
+				if (ChangedMaterial)
+				{
+					// Update the MIC
+					Comp->UpdateMaterialInstances();
+				}
+
+				if (bChangedLighting)
+				{
+					Comp->InvalidateLightingCache();
+				}
+
+				if (bChangedNavRelevance)
+				{
+					Comp->UpdateNavigationRelevance();
+				}
+
+				// Reattach all components
+				FComponentReregisterContext ReregisterContext(Comp);
+			}
+		}
+
+		// Need to update Gizmo scene proxy
+		if (bNeedsRecalcBoundingBox && GetWorld())
+		{
+			for (TActorIterator<ALandscapeGizmoActiveActor> It(GetWorld()); It; ++It)
+			{
+				It->ReregisterAllComponents();
+			}
+		}
+
+		if (ChangedMaterial)
+		{
+			if (GIsEditor && GetWorld() && !GetWorld()->IsPlayInEditor())
+			{
+				GEngine->DeferredCommands.AddUnique(TEXT("UpdateLandscapeMIC"));
+			}
+
+			// Update all the proxies...
+			for (TSet<ALandscapeProxy*>::TIterator It(Info->Proxies); It; ++It)
+			{
+				(*It)->MarkComponentsRenderStateDirty();
+			}
+		}
+	}
+}
+
+void ALandscapeProxy::ChangedPhysMaterial()
+{
+	ULandscapeInfo* LandscapeInfo = GetLandscapeInfo();
+	if (!LandscapeInfo) return;
+	for (auto It = LandscapeInfo->XYtoComponentMap.CreateIterator(); It; ++It)
+	{
+		ULandscapeComponent* Comp = It.Value();
+		if (Comp)
+		{
+			ULandscapeHeightfieldCollisionComponent* CollisionComponent = Comp->CollisionComponent.Get();
+			if (CollisionComponent)
+			{
+				Comp->UpdateCollisionLayerData();
+			}
+		}
+	}
+}
+
+void ULandscapeComponent::SetLOD(bool bForcedLODChanged, int32 InLODValue)
+{
+	if (bForcedLODChanged)
+	{
+		ForcedLOD = InLODValue;
+		if (ForcedLOD >= 0)
+		{
+			ForcedLOD = FMath::Clamp<int32>(ForcedLOD, 0, FMath::CeilLogTwo(SubsectionSizeQuads + 1) - 1);
+		}
+		else
+		{
+			ForcedLOD = -1;
+		}
+	}
+	else
+	{
+		int32 MaxLOD = FMath::CeilLogTwo(SubsectionSizeQuads + 1) - 1;
+		LODBias = FMath::Clamp<int32>(InLODValue, -MaxLOD, MaxLOD);
+	}
+
+	InvalidateLightingCache();
+
+	// Update neighbor components
+	ULandscapeInfo* Info = GetLandscapeInfo(false);
+	if (Info)
+	{
+		FIntPoint ComponentBase = GetSectionBase() / ComponentSizeQuads;
+		FIntPoint LandscapeKey[8] =
+		{
+			ComponentBase + FIntPoint(-1, -1),
+			ComponentBase + FIntPoint(+0, -1),
+			ComponentBase + FIntPoint(+1, -1),
+			ComponentBase + FIntPoint(-1, +0),
+			ComponentBase + FIntPoint(+1, +0),
+			ComponentBase + FIntPoint(-1, +1),
+			ComponentBase + FIntPoint(+0, +1),
+			ComponentBase + FIntPoint(+1, +1)
+		};
+
+		for (int32 Idx = 0; Idx < 8; ++Idx)
+		{
+			ULandscapeComponent* Comp = Info->XYtoComponentMap.FindRef(LandscapeKey[Idx]);
+			if (Comp)
+			{
+				Comp->Modify();
+				if (bForcedLODChanged)
+				{
+					Comp->NeighborLOD[7 - Idx] = ForcedLOD >= 0 ? ForcedLOD : 255; // Use 255 as unspecified value
+				}
+				else
+				{
+					// Neighbor LODBias are saved in BYTE, so need to convert to range [-128:127]
+					Comp->NeighborLODBias[7 - Idx] = LODBias + 128;
+				}
+
+				Comp->InvalidateLightingCache();
+				FComponentReregisterContext ReregisterContext(Comp);
+			}
+		}
+	}
+	FComponentReregisterContext ReregisterContext(this);
+}
+
+void ULandscapeComponent::PreEditChange(UProperty* PropertyThatWillChange)
+{
+	Super::PreEditChange(PropertyThatWillChange);
+	if (GIsEditor && PropertyThatWillChange && (PropertyThatWillChange->GetName() == TEXT("ForcedLOD") || PropertyThatWillChange->GetName() == TEXT("LODBias")))
+	{
+		// PreEdit unregister component and re-register after PostEdit so we will lose XYtoComponentMap for this component
+		ULandscapeInfo* Info = GetLandscapeInfo(false);
+		if (Info)
+		{
+			FIntPoint ComponentKey = GetSectionBase() / ComponentSizeQuads;
+			auto RegisteredComponent = Info->XYtoComponentMap.FindRef(ComponentKey);
+
+			if (RegisteredComponent == NULL)
+			{
+				Info->XYtoComponentMap.Add(ComponentKey, this);
+			}
+		}
+	}
+}
+
+void ULandscapeComponent::PostEditChangeProperty(FPropertyChangedEvent& PropertyChangedEvent)
+{
+	Super::PostEditChangeProperty(PropertyChangedEvent);
+	const FName PropertyName = PropertyChangedEvent.Property ? PropertyChangedEvent.Property->GetFName() : NAME_None;
+	if (PropertyName == FName(TEXT("OverrideMaterial")))
+	{
+		UpdateMaterialInstances();
+		// Reregister all components
+		FComponentReregisterContext ReregisterContext(this);
+	}
+	else if (GIsEditor && (PropertyName == FName(TEXT("ForcedLOD")) || PropertyName == FName(TEXT("LODBias"))))
+	{
+		bool bForcedLODChanged = PropertyName == FName(TEXT("ForcedLOD"));
+		SetLOD(bForcedLODChanged, bForcedLODChanged ? ForcedLOD : LODBias);
+	}
+	else if (GIsEditor && PropertyName == FName(TEXT("StaticLightingResolution")))
+	{
+		if (StaticLightingResolution > 0.f)
+		{
+			StaticLightingResolution = ::AdjustStaticLightingResolution(StaticLightingResolution, NumSubsections, SubsectionSizeQuads, ComponentSizeQuads);
+		}
+		else
+		{
+			StaticLightingResolution = 0;
+		}
+		InvalidateLightingCache();
+	}
+	else if (GIsEditor && PropertyName == FName(TEXT("LightingLODBias")))
+	{
+		int32 MaxLOD = FMath::CeilLogTwo(SubsectionSizeQuads + 1) - 1;
+		LightingLODBias = FMath::Clamp<int32>(LightingLODBias, -1, MaxLOD);
+		InvalidateLightingCache();
+	}
+	else if (GIsEditor && (PropertyName == FName(TEXT("CollisionMipLevel"))))
+	{
+		CollisionMipLevel = FMath::Clamp<int32>(CollisionMipLevel, 0, FMath::CeilLogTwo(SubsectionSizeQuads + 1) - 1);
+		TArray<uint8> CollisionMipData;
+		HeightmapTexture->Source.GetMipData(CollisionMipData, CollisionMipLevel);
+		UpdateCollisionHeightData((FColor*)CollisionMipData.GetData(), 0, 0, MAX_int32, MAX_int32, true, NULL, true); // Rebuild for new CollisionMipLevel
+	}
+}
+
+TSet<class ULandscapeComponent*> ULandscapeInfo::GetSelectedComponents() const
+{
+	return SelectedComponents;
+}
+
+TSet<class ULandscapeComponent*> ULandscapeInfo::GetSelectedRegionComponents() const
+{
+	return SelectedRegionComponents;
+}
+
+void ULandscapeInfo::UpdateSelectedComponents(TSet<ULandscapeComponent*>& NewComponents, bool bIsComponentwise /*=true*/)
+{
+	int32 InSelectType = bIsComponentwise ? FLandscapeEditToolRenderData::ST_COMPONENT : FLandscapeEditToolRenderData::ST_REGION;
+
+	if (bIsComponentwise)
+	{
+		for (TSet<ULandscapeComponent*>::TIterator It(NewComponents); It; ++It)
+		{
+			ULandscapeComponent* Comp = *It;
+			if (Comp->EditToolRenderData != NULL && (Comp->EditToolRenderData->SelectedType & InSelectType) == 0)
+			{
+				Comp->Modify();
+				int32 SelectedType = Comp->EditToolRenderData->SelectedType;
+				SelectedType |= InSelectType;
+				Comp->EditToolRenderData->UpdateSelectionMaterial(SelectedType);
+			}
+		}
+
+		// Remove the material from any old components that are no longer in the region
+		TSet<ULandscapeComponent*> RemovedComponents = SelectedComponents.Difference(NewComponents);
+		for (TSet<ULandscapeComponent*>::TIterator It(RemovedComponents); It; ++It)
+		{
+			ULandscapeComponent* Comp = *It;
+			if (Comp->EditToolRenderData != NULL)
+			{
+				Comp->Modify();
+				int32 SelectedType = Comp->EditToolRenderData->SelectedType;
+				SelectedType &= ~InSelectType;
+				Comp->EditToolRenderData->UpdateSelectionMaterial(SelectedType);
+			}
+		}
+		SelectedComponents = NewComponents;
+	}
+	else
+	{
+		// Only add components...
+		if (NewComponents.Num())
+		{
+			for (TSet<ULandscapeComponent*>::TIterator It(NewComponents); It; ++It)
+			{
+				ULandscapeComponent* Comp = *It;
+				if (Comp->EditToolRenderData != NULL && (Comp->EditToolRenderData->SelectedType & InSelectType) == 0)
+				{
+					Comp->Modify();
+					int32 SelectedType = Comp->EditToolRenderData->SelectedType;
+					SelectedType |= InSelectType;
+					Comp->EditToolRenderData->UpdateSelectionMaterial(SelectedType);
+				}
+
+				SelectedRegionComponents.Add(*It);
+			}
+		}
+		else
+		{
+			// Remove the material from any old components that are no longer in the region
+			for (TSet<ULandscapeComponent*>::TIterator It(SelectedRegionComponents); It; ++It)
+			{
+				ULandscapeComponent* Comp = *It;
+				if (Comp->EditToolRenderData != NULL)
+				{
+					Comp->Modify();
+					int32 SelectedType = Comp->EditToolRenderData->SelectedType;
+					SelectedType &= ~InSelectType;
+					Comp->EditToolRenderData->UpdateSelectionMaterial(SelectedType);
+				}
+			}
+			SelectedRegionComponents = NewComponents;
+		}
+	}
+}
+
+void ULandscapeInfo::SortSelectedComponents()
+{
+	struct FCompareULandscapeComponentBySectionBase
+	{
+		FORCEINLINE bool operator()(const ULandscapeComponent& A, const ULandscapeComponent& B) const
+		{
+			return (A.GetSectionBase().X == B.GetSectionBase().X) ? (A.GetSectionBase().Y < B.GetSectionBase().Y) : (A.GetSectionBase().X < B.GetSectionBase().X);
+		}
+	};
+	SelectedComponents.Sort(FCompareULandscapeComponentBySectionBase());
+}
+
+void ULandscapeInfo::ClearSelectedRegion(bool bIsComponentwise /*= true*/)
+{
+	TSet<ULandscapeComponent*> NewComponents;
+	UpdateSelectedComponents(NewComponents, bIsComponentwise);
+	if (!bIsComponentwise)
+	{
+		SelectedRegion.Empty();
+	}
+}
+
+struct FLandscapeDataInterface* ULandscapeInfo::GetDataInterface()
+{
+	if (DataInterface == NULL)
+	{
+		DataInterface = new FLandscapeDataInterface();
+	}
+
+	return DataInterface;
+}
+
+void ULandscapeComponent::ReallocateWeightmaps(FLandscapeEditDataInterface* DataInterface)
+{
+	ALandscapeProxy* Proxy = GetLandscapeProxy();
+
+	int32 NeededNewChannels = 0;
+	for (int32 LayerIdx = 0; LayerIdx < WeightmapLayerAllocations.Num(); LayerIdx++)
+	{
+		if (WeightmapLayerAllocations[LayerIdx].WeightmapTextureIndex == 255)
+		{
+			NeededNewChannels++;
+		}
+	}
+
+	// All channels allocated!
+	if (NeededNewChannels == 0)
+	{
+		return;
+	}
+
+	Modify();
+	//Landscape->Modify();
+	Proxy->Modify();
+
+	// UE_LOG(LogLandscape, Log, TEXT("----------------------"));
+	// UE_LOG(LogLandscape, Log, TEXT("Component %s needs %d layers (%d new)"), *GetName(), WeightmapLayerAllocations.Num(), NeededNewChannels);
+
+	// See if our existing textures have sufficient space
+	int32 ExistingTexAvailableChannels = 0;
+	for (int32 TexIdx = 0; TexIdx < WeightmapTextures.Num(); TexIdx++)
+	{
+		FLandscapeWeightmapUsage* Usage = Proxy->WeightmapUsageMap.Find(WeightmapTextures[TexIdx]);
+		check(Usage);
+
+		ExistingTexAvailableChannels += Usage->FreeChannelCount();
+
+		if (ExistingTexAvailableChannels >= NeededNewChannels)
+		{
+			break;
+		}
+	}
+
+	if (ExistingTexAvailableChannels >= NeededNewChannels)
+	{
+		// UE_LOG(LogLandscape, Log, TEXT("Existing texture has available channels"));
+
+		// Allocate using our existing textures' spare channels.
+		for (int32 TexIdx = 0; TexIdx < WeightmapTextures.Num(); TexIdx++)
+		{
+			FLandscapeWeightmapUsage* Usage = Proxy->WeightmapUsageMap.Find(WeightmapTextures[TexIdx]);
+
+			for (int32 ChanIdx = 0; ChanIdx < 4; ChanIdx++)
+			{
+				if (Usage->ChannelUsage[ChanIdx] == NULL)
+				{
+					for (int32 LayerIdx = 0; LayerIdx < WeightmapLayerAllocations.Num(); LayerIdx++)
+					{
+						FWeightmapLayerAllocationInfo& AllocInfo = WeightmapLayerAllocations[LayerIdx];
+						if (AllocInfo.WeightmapTextureIndex == 255)
+						{
+							// Zero out the data for this texture channel
+							if (DataInterface)
+							{
+								DataInterface->ZeroTextureChannel(WeightmapTextures[TexIdx], ChanIdx);
+							}
+
+							AllocInfo.WeightmapTextureIndex = TexIdx;
+							AllocInfo.WeightmapTextureChannel = ChanIdx;
+							Usage->ChannelUsage[ChanIdx] = this;
+							NeededNewChannels--;
+
+							if (NeededNewChannels == 0)
+							{
+								return;
+							}
+						}
+					}
+				}
+			}
+		}
+		// we should never get here.
+		check(false);
+	}
+
+	// UE_LOG(LogLandscape, Log, TEXT("Reallocating."));
+
+	// We are totally reallocating the weightmap
+	int32 TotalNeededChannels = WeightmapLayerAllocations.Num();
+	int32 CurrentLayer = 0;
+	TArray<UTexture2D*> NewWeightmapTextures;
+	while (TotalNeededChannels > 0)
+	{
+		// UE_LOG(LogLandscape, Log, TEXT("Still need %d channels"), TotalNeededChannels);
+
+		UTexture2D* CurrentWeightmapTexture = NULL;
+		FLandscapeWeightmapUsage* CurrentWeightmapUsage = NULL;
+
+		if (TotalNeededChannels < 4)
+		{
+			// UE_LOG(LogLandscape, Log, TEXT("Looking for nearest"));
+
+			// see if we can find a suitable existing weightmap texture with sufficient channels
+			int32 BestDistanceSquared = MAX_int32;
+			for (TMap<UTexture2D*, struct FLandscapeWeightmapUsage>::TIterator It(Proxy->WeightmapUsageMap); It; ++It)
+			{
+				FLandscapeWeightmapUsage* TryWeightmapUsage = &It.Value();
+				if (TryWeightmapUsage->FreeChannelCount() >= TotalNeededChannels)
+				{
+					// See if this candidate is closer than any others we've found
+					for (int32 ChanIdx = 0; ChanIdx < 4; ChanIdx++)
+					{
+						if (TryWeightmapUsage->ChannelUsage[ChanIdx] != NULL)
+						{
+							int32 TryDistanceSquared = (TryWeightmapUsage->ChannelUsage[ChanIdx]->GetSectionBase() - GetSectionBase()).SizeSquared();
+							if (TryDistanceSquared < BestDistanceSquared)
+							{
+								CurrentWeightmapTexture = It.Key();
+								CurrentWeightmapUsage = TryWeightmapUsage;
+								BestDistanceSquared = TryDistanceSquared;
+							}
+						}
+					}
+				}
+			}
+		}
+
+		bool NeedsUpdateResource = false;
+		// No suitable weightmap texture
+		if (CurrentWeightmapTexture == NULL)
+		{
+			MarkPackageDirty();
+
+			// Weightmap is sized the same as the component
+			int32 WeightmapSize = (SubsectionSizeQuads + 1) * NumSubsections;
+
+			// We need a new weightmap texture
+			CurrentWeightmapTexture = GetLandscapeProxy()->CreateLandscapeTexture(WeightmapSize, WeightmapSize, TEXTUREGROUP_Terrain_Weightmap, TSF_BGRA8);
+			// Alloc dummy mips
+			CreateEmptyTextureMips(CurrentWeightmapTexture);
+			CurrentWeightmapTexture->PostEditChange();
+
+			// Store it in the usage map
+			CurrentWeightmapUsage = &Proxy->WeightmapUsageMap.Add(CurrentWeightmapTexture, FLandscapeWeightmapUsage());
+
+			// UE_LOG(LogLandscape, Log, TEXT("Making a new texture %s"), *CurrentWeightmapTexture->GetName());
+		}
+
+		NewWeightmapTextures.Add(CurrentWeightmapTexture);
+
+		for (int32 ChanIdx = 0; ChanIdx < 4 && TotalNeededChannels > 0; ChanIdx++)
+		{
+			// UE_LOG(LogLandscape, Log, TEXT("Finding allocation for layer %d"), CurrentLayer);
+
+			if (CurrentWeightmapUsage->ChannelUsage[ChanIdx] == NULL)
+			{
+				// Use this allocation
+				FWeightmapLayerAllocationInfo& AllocInfo = WeightmapLayerAllocations[CurrentLayer];
+
+				if (AllocInfo.WeightmapTextureIndex == 255)
+				{
+					// New layer - zero out the data for this texture channel
+					if (DataInterface)
+					{
+						DataInterface->ZeroTextureChannel(CurrentWeightmapTexture, ChanIdx);
+						// UE_LOG(LogLandscape, Log, TEXT("Zeroing out channel %s.%d"), *CurrentWeightmapTexture->GetName(), ChanIdx);
+					}
+				}
+				else
+				{
+					UTexture2D* OldWeightmapTexture = WeightmapTextures[AllocInfo.WeightmapTextureIndex];
+
+					// Copy the data
+					if (ensure(DataInterface != NULL)) // it's not safe to skip the copy
+					{
+						DataInterface->CopyTextureChannel(CurrentWeightmapTexture, ChanIdx, OldWeightmapTexture, AllocInfo.WeightmapTextureChannel);
+						DataInterface->ZeroTextureChannel(OldWeightmapTexture, AllocInfo.WeightmapTextureChannel);
+						// UE_LOG(LogLandscape, Log, TEXT("Copying old channel (%s).%d to new channel (%s).%d"), *OldWeightmapTexture->GetName(), AllocInfo.WeightmapTextureChannel, *CurrentWeightmapTexture->GetName(), ChanIdx);
+					}
+
+					// Remove the old allocation
+					FLandscapeWeightmapUsage* OldWeightmapUsage = Proxy->WeightmapUsageMap.Find(OldWeightmapTexture);
+					OldWeightmapUsage->ChannelUsage[AllocInfo.WeightmapTextureChannel] = NULL;
+				}
+
+				// Assign the new allocation
+				CurrentWeightmapUsage->ChannelUsage[ChanIdx] = this;
+				AllocInfo.WeightmapTextureIndex = NewWeightmapTextures.Num() - 1;
+				AllocInfo.WeightmapTextureChannel = ChanIdx;
+				CurrentLayer++;
+				TotalNeededChannels--;
+			}
+		}
+	}
+
+	// Replace the weightmap textures
+	WeightmapTextures = MoveTemp(NewWeightmapTextures);
+
+	if (DataInterface)
+	{
+		// Update the mipmaps for the textures we edited
+		for (int32 Idx = 0; Idx < WeightmapTextures.Num(); Idx++)
+		{
+			UTexture2D* WeightmapTexture = WeightmapTextures[Idx];
+			FLandscapeTextureDataInfo* WeightmapDataInfo = DataInterface->GetTextureDataInfo(WeightmapTexture);
+
+			int32 NumMips = WeightmapTexture->Source.GetNumMips();
+			TArray<FColor*> WeightmapTextureMipData;
+			WeightmapTextureMipData.AddUninitialized(NumMips);
+			for (int32 MipIdx = 0; MipIdx < NumMips; MipIdx++)
+			{
+				WeightmapTextureMipData[MipIdx] = (FColor*)WeightmapDataInfo->GetMipData(MipIdx);
+			}
+
+			ULandscapeComponent::UpdateWeightmapMips(NumSubsections, SubsectionSizeQuads, WeightmapTexture, WeightmapTextureMipData, 0, 0, MAX_int32, MAX_int32, WeightmapDataInfo);
+		}
+	}
+}
+
+void ALandscapeProxy::RemoveInvalidWeightmaps()
+{
+	if (GIsEditor)
+	{
+		for (TMap< UTexture2D*, struct FLandscapeWeightmapUsage >::TIterator It(WeightmapUsageMap); It; ++It)
+		{
+			UTexture2D* Tex = It.Key();
+			FLandscapeWeightmapUsage& Usage = It.Value();
+			if (Usage.FreeChannelCount() == 4) // Invalid Weight-map
+			{
+				if (Tex)
+				{
+					Tex->SetFlags(RF_Transactional);
+					Tex->Modify();
+					Tex->MarkPackageDirty();
+					Tex->ClearFlags(RF_Standalone);
+				}
+				WeightmapUsageMap.Remove(Tex);
+			}
+		}
+
+		// Remove Unused Weightmaps...
+		for (int32 Idx = 0; Idx < LandscapeComponents.Num(); ++Idx)
+		{
+			ULandscapeComponent* Component = LandscapeComponents[Idx];
+			Component->RemoveInvalidWeightmaps();
+		}
+	}
+}
+
+void ULandscapeComponent::RemoveInvalidWeightmaps()
+{
+	// Adjust WeightmapTextureIndex index for other layers
+	TSet<int32> UsedTextureIndices;
+	TSet<int32> AllTextureIndices;
+	for (int32 LayerIdx = 0; LayerIdx < WeightmapLayerAllocations.Num(); LayerIdx++)
+	{
+		UsedTextureIndices.Add(WeightmapLayerAllocations[LayerIdx].WeightmapTextureIndex);
+	}
+
+	for (int32 WeightIdx = 0; WeightIdx < WeightmapTextures.Num(); ++WeightIdx)
+	{
+		AllTextureIndices.Add(WeightIdx);
+	}
+
+	TSet<int32> UnUsedTextureIndices = AllTextureIndices.Difference(UsedTextureIndices);
+
+	int32 DeletedLayers = 0;
+	for (TSet<int32>::TIterator It(UnUsedTextureIndices); It; ++It)
+	{
+		int32 DeleteLayerWeightmapTextureIndex = *It - DeletedLayers;
+		WeightmapTextures[DeleteLayerWeightmapTextureIndex]->SetFlags(RF_Transactional);
+		WeightmapTextures[DeleteLayerWeightmapTextureIndex]->Modify();
+		WeightmapTextures[DeleteLayerWeightmapTextureIndex]->MarkPackageDirty();
+		WeightmapTextures[DeleteLayerWeightmapTextureIndex]->ClearFlags(RF_Standalone);
+		WeightmapTextures.RemoveAt(DeleteLayerWeightmapTextureIndex);
+
+		// Adjust WeightmapTextureIndex index for other layers
+		for (int32 LayerIdx = 0; LayerIdx < WeightmapLayerAllocations.Num(); LayerIdx++)
+		{
+			FWeightmapLayerAllocationInfo& Allocation = WeightmapLayerAllocations[LayerIdx];
+
+			if (Allocation.WeightmapTextureIndex > DeleteLayerWeightmapTextureIndex)
+			{
+				Allocation.WeightmapTextureIndex--;
+			}
+
+			check(Allocation.WeightmapTextureIndex < WeightmapTextures.Num());
+		}
+		DeletedLayers++;
+	}
+}
+
+void ULandscapeComponent::InitHeightmapData(TArray<FColor>& Heights, bool bUpdateCollision)
+{
+	int32 ComponentSizeVerts = NumSubsections * (SubsectionSizeQuads + 1);
+
+	if (Heights.Num() != FMath::Square(ComponentSizeVerts))
+	{
+		return;
+	}
+
+	// Handling old Height map....
+	if (HeightmapTexture && HeightmapTexture->GetOutermost() != GetTransientPackage()
+		&& HeightmapTexture->GetOutermost() == GetOutermost()
+		&& HeightmapTexture->Source.GetSizeX() >= ComponentSizeVerts) // if Height map is not valid...
+	{
+		HeightmapTexture->SetFlags(RF_Transactional);
+		HeightmapTexture->Modify();
+		HeightmapTexture->MarkPackageDirty();
+		HeightmapTexture->ClearFlags(RF_Standalone); // Delete if no reference...
+	}
+
+	// New Height map
+	TArray<FColor*> HeightmapTextureMipData;
+	// make sure the heightmap UVs are powers of two.
+	int32 HeightmapSizeU = (1 << FMath::CeilLogTwo(ComponentSizeVerts));
+	int32 HeightmapSizeV = (1 << FMath::CeilLogTwo(ComponentSizeVerts));
+
+	// Height map construction
+	HeightmapTexture = GetLandscapeProxy()->CreateLandscapeTexture(HeightmapSizeU, HeightmapSizeV, TEXTUREGROUP_Terrain_Heightmap, TSF_BGRA8);
+
+	int32 MipSubsectionSizeQuads = SubsectionSizeQuads;
+	int32 MipSizeU = HeightmapSizeU;
+	int32 MipSizeV = HeightmapSizeV;
+
+	HeightmapScaleBias = FVector4(1.f / (float)HeightmapSizeU, 1.f / (float)HeightmapSizeV, 0.f, 0.f);
+
+	int32 Mip = 0;
+	while (MipSizeU > 1 && MipSizeV > 1 && MipSubsectionSizeQuads >= 1)
+	{
+		FColor* HeightmapTextureData = (FColor*)HeightmapTexture->Source.LockMip(Mip);
+		if (Mip == 0)
+		{
+			FMemory::Memcpy(HeightmapTextureData, Heights.GetData(), MipSizeU*MipSizeV*sizeof(FColor));
+		}
+		else
+		{
+			FMemory::Memzero(HeightmapTextureData, MipSizeU*MipSizeV*sizeof(FColor));
+		}
+		HeightmapTextureMipData.Add(HeightmapTextureData);
+
+		MipSizeU >>= 1;
+		MipSizeV >>= 1;
+		Mip++;
+
+		MipSubsectionSizeQuads = ((MipSubsectionSizeQuads + 1) >> 1) - 1;
+	}
+	ULandscapeComponent::GenerateHeightmapMips(HeightmapTextureMipData);
+
+	if (bUpdateCollision)
+	{
+		ULandscapeComponent::UpdateCollisionHeightData(HeightmapTextureMipData[CollisionMipLevel]);
+	}
+
+	for (int32 i = 0; i < HeightmapTextureMipData.Num(); i++)
+	{
+		HeightmapTexture->Source.UnlockMip(i);
+	}
+	HeightmapTexture->PostEditChange();
+}
+
+void ULandscapeComponent::InitWeightmapData(TArray<ULandscapeLayerInfoObject*>& LayerInfos, TArray<TArray<uint8> >& WeightmapData)
+{
+	if (LayerInfos.Num() != WeightmapData.Num() || LayerInfos.Num() <= 0)
+	{
+		return;
+	}
+
+	int32 ComponentSizeVerts = NumSubsections * (SubsectionSizeQuads + 1);
+
+	// Validation..
+	for (int32 Idx = 0; Idx < WeightmapData.Num(); ++Idx)
+	{
+		if (WeightmapData[Idx].Num() != FMath::Square(ComponentSizeVerts))
+		{
+			return;
+		}
+	}
+
+	for (int32 Idx = 0; Idx < WeightmapTextures.Num(); ++Idx)
+	{
+		if (WeightmapTextures[Idx] && WeightmapTextures[Idx]->GetOutermost() != GetTransientPackage()
+			&& WeightmapTextures[Idx]->GetOutermost() == GetOutermost()
+			&& WeightmapTextures[Idx]->Source.GetSizeX() == ComponentSizeVerts)
+		{
+			WeightmapTextures[Idx]->SetFlags(RF_Transactional);
+			WeightmapTextures[Idx]->Modify();
+			WeightmapTextures[Idx]->MarkPackageDirty();
+			WeightmapTextures[Idx]->ClearFlags(RF_Standalone); // Delete if no reference...
+		}
+	}
+	WeightmapTextures.Empty();
+
+	WeightmapLayerAllocations.Empty(LayerInfos.Num());
+	for (int32 Idx = 0; Idx < LayerInfos.Num(); ++Idx)
+	{
+		new (WeightmapLayerAllocations)FWeightmapLayerAllocationInfo(LayerInfos[Idx]);
+	}
+
+	ReallocateWeightmaps(NULL);
+
+	check(WeightmapLayerAllocations.Num() > 0 && WeightmapTextures.Num() > 0);
+
+	int32 WeightmapSize = ComponentSizeVerts;
+	WeightmapScaleBias = FVector4(1.f / (float)WeightmapSize, 1.f / (float)WeightmapSize, 0.5f / (float)WeightmapSize, 0.5f / (float)WeightmapSize);
+	WeightmapSubsectionOffset = (float)(SubsectionSizeQuads + 1) / (float)WeightmapSize;
+
+	// Channel remapping
+	int32 ChannelOffsets[4] = { (int32)STRUCT_OFFSET(FColor, R), (int32)STRUCT_OFFSET(FColor, G), (int32)STRUCT_OFFSET(FColor, B), (int32)STRUCT_OFFSET(FColor, A) };
+
+	TArray<void*> WeightmapDataPtrs;
+	WeightmapDataPtrs.AddUninitialized(WeightmapTextures.Num());
+	for (int32 WeightmapIdx = 0; WeightmapIdx < WeightmapTextures.Num(); ++WeightmapIdx)
+	{
+		WeightmapDataPtrs[WeightmapIdx] = WeightmapTextures[WeightmapIdx]->Source.LockMip(0);
+	}
+
+	for (int32 LayerIdx = 0; LayerIdx < WeightmapLayerAllocations.Num(); ++LayerIdx)
+	{
+		void* DestDataPtr = WeightmapDataPtrs[WeightmapLayerAllocations[LayerIdx].WeightmapTextureIndex];
+		uint8* DestTextureData = (uint8*)DestDataPtr + ChannelOffsets[WeightmapLayerAllocations[LayerIdx].WeightmapTextureChannel];
+		uint8* SrcTextureData = (uint8*)&WeightmapData[LayerIdx][0];
+
+		for (int32 i = 0; i < WeightmapData[LayerIdx].Num(); i++)
+		{
+			DestTextureData[i * 4] = SrcTextureData[i];
+		}
+	}
+
+	for (int32 Idx = 0; Idx < WeightmapTextures.Num(); Idx++)
+	{
+		UTexture2D* WeightmapTexture = WeightmapTextures[Idx];
+		WeightmapTexture->Source.UnlockMip(0);
+	}
+
+	for (int32 Idx = 0; Idx < WeightmapTextures.Num(); Idx++)
+	{
+		UTexture2D* WeightmapTexture = WeightmapTextures[Idx];
+		{
+			FLandscapeTextureDataInfo WeightmapDataInfo(WeightmapTexture);
+
+			int32 NumMips = WeightmapTexture->Source.GetNumMips();
+			TArray<FColor*> WeightmapTextureMipData;
+			WeightmapTextureMipData.AddUninitialized(NumMips);
+			for (int32 MipIdx = 0; MipIdx < NumMips; MipIdx++)
+			{
+				WeightmapTextureMipData[MipIdx] = (FColor*)WeightmapDataInfo.GetMipData(MipIdx);
+			}
+
+			ULandscapeComponent::UpdateWeightmapMips(NumSubsections, SubsectionSizeQuads, WeightmapTexture, WeightmapTextureMipData, 0, 0, MAX_int32, MAX_int32, &WeightmapDataInfo);
+		}
+
+		WeightmapTexture->PostEditChange();
+	}
+
+	FlushRenderingCommands();
+
+	MaterialInstance = NULL;
+
+}
+
+#define MAX_LANDSCAPE_EXPORT_COMPONENTS_NUM		16
+#define MAX_LANDSCAPE_PROP_TEXT_LENGTH			1024*1024*16
+
+
+bool ALandscapeProxy::ShouldExport()
+{
+	if (!bIsMovingToLevel && LandscapeComponents.Num() > MAX_LANDSCAPE_EXPORT_COMPONENTS_NUM)
+	{
+		// Prompt to save startup packages
+		if (EAppReturnType::Yes == FMessageDialog::Open(EAppMsgType::YesNo, FText::Format(
+			NSLOCTEXT("UnrealEd", "LandscapeExport_Warning", "Landscape has large number({0}) of components, so it will use large amount memory to copy it to the clipboard. Do you want to proceed?"), FText::AsNumber(LandscapeComponents.Num()))))
+		{
+			return true;
+		}
+		else
+		{
+			return false;
+		}
+	}
+	return true;
+}
+
+bool ALandscapeProxy::ShouldImport(FString* ActorPropString, bool IsMovingToLevel)
+{
+	bIsMovingToLevel = IsMovingToLevel;
+	if (!bIsMovingToLevel && ActorPropString && ActorPropString->Len() > MAX_LANDSCAPE_PROP_TEXT_LENGTH)
+	{
+		// Prompt to save startup packages
+		if (EAppReturnType::Yes == FMessageDialog::Open(EAppMsgType::YesNo, FText::Format(
+			NSLOCTEXT("UnrealEd", "LandscapeImport_Warning", "Landscape is about to import large amount memory ({0}MB) from the clipboard, which will take some time. Do you want to proceed?"), FText::AsNumber(ActorPropString->Len() >> 20))))
+		{
+			return true;
+		}
+		else
+		{
+			return false;
+		}
+	}
+	return true;
+}
+
+void ULandscapeComponent::ExportCustomProperties(FOutputDevice& Out, uint32 Indent)
+{
+	if (HasAnyFlags(RF_ClassDefaultObject))
+	{
+		return;
+	}
+	// Height map
+	int32 NumVertices = FMath::Square(NumSubsections*(SubsectionSizeQuads + 1));
+	FLandscapeComponentDataInterface DataInterface(this);
+	TArray<FColor> Heightmap;
+	DataInterface.GetHeightmapTextureData(Heightmap);
+	check(Heightmap.Num() == NumVertices);
+
+	Out.Logf(TEXT("%sCustomProperties LandscapeHeightData "), FCString::Spc(Indent));
+	for (int32 i = 0; i < NumVertices; i++)
+	{
+		Out.Logf(TEXT("%x "), Heightmap[i].DWColor());
+	}
+
+	TArray<uint8> Weightmap;
+	// Weight map
+	Out.Logf(TEXT("LayerNum=%d "), WeightmapLayerAllocations.Num());
+	for (int32 i = 0; i < WeightmapLayerAllocations.Num(); i++)
+	{
+		if (DataInterface.GetWeightmapTextureData(WeightmapLayerAllocations[i].LayerInfo, Weightmap))
+		{
+			Out.Logf(TEXT("LayerInfo=%s "), *WeightmapLayerAllocations[i].LayerInfo->GetPathName());
+			for (int32 VertexIndex = 0; VertexIndex < NumVertices; VertexIndex++)
+			{
+				Out.Logf(TEXT("%x "), Weightmap[VertexIndex]);
+			}
+		}
+	}
+
+	Out.Logf(TEXT("\r\n"));
+}
+
+
+void ULandscapeComponent::ImportCustomProperties(const TCHAR* SourceText, FFeedbackContext* Warn)
+{
+	if (FParse::Command(&SourceText, TEXT("LandscapeHeightData")))
+	{
+		int32 NumVertices = FMath::Square(NumSubsections*(SubsectionSizeQuads + 1));
+
+		TArray<FColor> Heights;
+		Heights.Empty(NumVertices);
+		Heights.AddZeroed(NumVertices);
+
+		FParse::Next(&SourceText);
+		int32 i = 0;
+		TCHAR* StopStr;
+		while (FChar::IsHexDigit(*SourceText))
+		{
+			if (i < NumVertices)
+			{
+				Heights[i++].DWColor() = FCString::Strtoi(SourceText, &StopStr, 16);
+				while (FChar::IsHexDigit(*SourceText))
+				{
+					SourceText++;
+				}
+			}
+
+			FParse::Next(&SourceText);
+		}
+
+		if (i != NumVertices)
+		{
+			Warn->Logf(*NSLOCTEXT("Core", "SyntaxError", "Syntax Error").ToString());
+		}
+
+		int32 ComponentSizeVerts = NumSubsections * (SubsectionSizeQuads + 1);
+
+		InitHeightmapData(Heights, false);
+
+		// Weight maps
+		int32 LayerNum = 0;
+		if (FParse::Value(SourceText, TEXT("LayerNum="), LayerNum))
+		{
+			while (*SourceText && (!FChar::IsWhitespace(*SourceText)))
+			{
+				++SourceText;
+			}
+			FParse::Next(&SourceText);
+		}
+
+		if (LayerNum <= 0)
+		{
+			return;
+		}
+
+		// Init memory
+		TArray<ULandscapeLayerInfoObject*> LayerInfos;
+		LayerInfos.Empty(LayerNum);
+		TArray<TArray<uint8>> WeightmapData;
+		for (int32 LayerIndex = 0; LayerIndex < LayerNum; ++LayerIndex)
+		{
+			TArray<uint8> Weights;
+			Weights.Empty(NumVertices);
+			Weights.AddUninitialized(NumVertices);
+			WeightmapData.Add(Weights);
+		}
+
+		int32 LayerIdx = 0;
+		FString LayerInfoPath;
+		while (*SourceText)
+		{
+			if (FParse::Value(SourceText, TEXT("LayerInfo="), LayerInfoPath))
+			{
+				LayerInfos.Add(LoadObject<ULandscapeLayerInfoObject>(NULL, *LayerInfoPath));
+
+				while (*SourceText && (!FChar::IsWhitespace(*SourceText)))
+				{
+					++SourceText;
+				}
+				FParse::Next(&SourceText);
+				check(*SourceText);
+
+				i = 0;
+				while (FChar::IsHexDigit(*SourceText))
+				{
+					if (i < NumVertices)
+					{
+						(WeightmapData[LayerIdx])[i++] = (uint8)FCString::Strtoi(SourceText, &StopStr, 16);
+						while (FChar::IsHexDigit(*SourceText))
+						{
+							SourceText++;
+						}
+					}
+					FParse::Next(&SourceText);
+				}
+
+				if (i != NumVertices)
+				{
+					Warn->Logf(*NSLOCTEXT("Core", "SyntaxError", "Syntax Error").ToString());
+				}
+				LayerIdx++;
+			}
+			else
+			{
+				break;
+			}
+		}
+
+		InitWeightmapData(LayerInfos, WeightmapData);
+	}
+}
+
+bool ALandscapeProxy::IsValidLandscapeActor(ALandscape* Landscape)
+{
+	if (bIsProxy && Landscape)
+	{
+		if (!Landscape->HasAnyFlags(RF_BeginDestroyed))
+		{
+			if (LandscapeActor.IsNull() && !LandscapeGuid.IsValid())
+			{
+				return true; // always valid for newly created Proxy
+			}
+			if (((LandscapeActor && LandscapeActor == Landscape)
+				|| (LandscapeActor.IsNull() && LandscapeGuid.IsValid() && LandscapeGuid == Landscape->LandscapeGuid))
+				&& ComponentSizeQuads == Landscape->ComponentSizeQuads
+				&& NumSubsections == Landscape->NumSubsections
+				&& SubsectionSizeQuads == Landscape->SubsectionSizeQuads)
+			{
+				return true;
+			}
+		}
+	}
+	return false;
+}
+
+bool ALandscapeProxy::GetSelectedComponents(TArray<UObject*>& SelectedObjects)
+{
+	ULandscapeInfo* Info = GetLandscapeInfo(false);
+	if (Info && Info->bCurrentlyEditing)
+	{
+		auto SelectedComponents = Info->GetSelectedComponents();
+		for (auto It = SelectedComponents.CreateIterator(); It; ++It)
+		{
+			SelectedObjects.Add(*It);
+		}
+		return true;
+	}
+	return false;
+}
+
+struct FMobileLayerAllocation
+{
+	FWeightmapLayerAllocationInfo Allocation;
+
+	friend bool operator<(const FMobileLayerAllocation& lhs, const FMobileLayerAllocation& rhs)
+	{
+		if (!lhs.Allocation.LayerInfo && !rhs.Allocation.LayerInfo) return false; // equally broken :P
+		if (!lhs.Allocation.LayerInfo && rhs.Allocation.LayerInfo) return false; // broken layers sort to the end
+		if (!rhs.Allocation.LayerInfo && lhs.Allocation.LayerInfo) return true;
+
+		if (lhs.Allocation.LayerInfo == ALandscapeProxy::VisibilityLayer) return true; // visibility layer to the front
+		if (rhs.Allocation.LayerInfo == ALandscapeProxy::VisibilityLayer) return false;
+
+		if (lhs.Allocation.LayerInfo->bNoWeightBlend && !rhs.Allocation.LayerInfo->bNoWeightBlend) return false; // non-blended layers sort to the end
+		if (rhs.Allocation.LayerInfo->bNoWeightBlend && !lhs.Allocation.LayerInfo->bNoWeightBlend) return true;
+
+		// TODO: If we want to support cleanly decaying a pc landscape for mobile
+		// we should probably add other sort criteria, e.g. coverage
+		// or e.g. add an "importance" to layerinfos and sort on that
+
+		return false; // equal, preserve order
+	}
+};
+
+UMaterialInstance* ULandscapeComponent::GeneratePlatformPixelData(TArray<UTexture2D*>& InWeightmapTextures, bool bIsCooking)
+{
+	if (IsTemplate() || HeightmapTexture == NULL)
+	{
+		return MaterialInstance;
+	}
+
+	TArray<FMobileLayerAllocation> MobileLayerAllocations;
+	MobileLayerAllocations.Reserve(WeightmapLayerAllocations.Num());
+	for (const auto& Allocation : WeightmapLayerAllocations)
+	{
+		MobileLayerAllocations.Add(FMobileLayerAllocation{ Allocation });
+	}
+	MobileLayerAllocations.StableSort();
+
+	// in the current mobile shader only 3 layers are supported (the 3rd only as a blended layer)
+	// so make sure we have a blended layer for layer 3 if possible
+	if (MobileLayerAllocations.Num() >= 3 &&
+		MobileLayerAllocations[2].Allocation.LayerInfo && MobileLayerAllocations[2].Allocation.LayerInfo->bNoWeightBlend)
+	{
+		int32 BlendedLayerToMove = INDEX_NONE;
+
+		// First try to swap layer 3 with an earlier blended layer
+		// this will allow both to work
+		for (int32 i = 1; i >= 0; --i)
+		{
+			if (MobileLayerAllocations[i].Allocation.LayerInfo && !MobileLayerAllocations[i].Allocation.LayerInfo->bNoWeightBlend)
+			{
+				BlendedLayerToMove = i;
+				break;
+			}
+		}
+
+		// otherwise swap layer 3 with the first weight-blended layer found
+		// as non-blended layers aren't supported for layer 3 it wasn't going to work anyway, might as well swap it out for one that will work
+		if (BlendedLayerToMove == INDEX_NONE)
+		{
+			// I wish I could specify a start index here, but it doesn't affect the result
+			BlendedLayerToMove = MobileLayerAllocations.IndexOfByPredicate([](const FMobileLayerAllocation& MobileAllocation){ return MobileAllocation.Allocation.LayerInfo && !MobileAllocation.Allocation.LayerInfo->bNoWeightBlend; });
+		}
+
+		if (BlendedLayerToMove != INDEX_NONE)
+		{
+			// Preserve order of all but the blended layer we're moving into slot 3
+			FMobileLayerAllocation TempAllocation = MoveTemp(MobileLayerAllocations[BlendedLayerToMove]);
+			MobileLayerAllocations.RemoveAt(BlendedLayerToMove, 1, false);
+			MobileLayerAllocations.Insert(MoveTemp(TempAllocation), 2);
+		}
+	}
+
+	int32 WeightmapSize = (SubsectionSizeQuads + 1) * NumSubsections;
+	UTexture2D* WeightmapTexture = GetLandscapeProxy()->CreateLandscapeTexture(WeightmapSize, WeightmapSize, TEXTUREGROUP_Terrain_Weightmap, TSF_BGRA8);
+	CreateEmptyTextureMips(WeightmapTexture);
+
+	{
+		FLandscapeEditDataInterface LandscapeEdit(GetLandscapeInfo(false));
+
+		if (InWeightmapTextures.Num() > 0)
+		{
+			int32 CurrentIdx = 0;
+			for (const auto& MobileAllocation : MobileLayerAllocations)
+			{
+				// Only for valid Layers
+				if (MobileAllocation.Allocation.LayerInfo)
+				{
+					uint8 TextureIndex = MobileAllocation.Allocation.WeightmapTextureIndex;
+					uint8 TextureChannel = MobileAllocation.Allocation.WeightmapTextureChannel;
+					LandscapeEdit.CopyTextureChannel(WeightmapTexture, CurrentIdx, InWeightmapTextures[TextureIndex], TextureChannel);
+					CurrentIdx++;
+					if (CurrentIdx >= 2) // Only support 2 layers in texture
+					{
+						break;
+					}
+				}
+			}
+		}
+
+		LandscapeEdit.CopyTextureFromHeightmap(WeightmapTexture, 2, this, 2);
+		LandscapeEdit.CopyTextureFromHeightmap(WeightmapTexture, 3, this, 3);
+	}
+
+	WeightmapTexture->PostEditChange();
+	InWeightmapTextures.Empty();
+	InWeightmapTextures.Add(WeightmapTexture);
+
+	FLinearColor Masks[5];
+	Masks[0] = FLinearColor(1, 0, 0, 0);
+	Masks[1] = FLinearColor(0, 1, 0, 0);
+	Masks[2] = FLinearColor(0, 0, 1, 0);
+	Masks[3] = FLinearColor(0, 0, 0, 1);
+	Masks[4] = FLinearColor(0, 0, 0, 0); // mask out layers 4+ altogether
+
+	if (!bIsCooking)
+	{
+		UMaterialInstanceDynamic* MobileMaterialInstance = UMaterialInstanceDynamic::Create(MaterialInstance, GetOutermost());
+
+		MobileBlendableLayerMask = 0;
+
+		// Set the layer mask
+		int32 CurrentIdx = 0;
+		for (const auto& MobileAllocation : MobileLayerAllocations)
+		{
+			const FWeightmapLayerAllocationInfo& Allocation = MobileAllocation.Allocation;
+			if (Allocation.LayerInfo)
+			{
+				FName LayerName = Allocation.LayerInfo == ALandscapeProxy::VisibilityLayer ? UMaterialExpressionLandscapeVisibilityMask::ParameterName : Allocation.LayerInfo->LayerName;
+				MobileMaterialInstance->SetVectorParameterValue(FName(*FString::Printf(TEXT("LayerMask_%s"), *LayerName.ToString())), Masks[FMath::Min(4, CurrentIdx)]);
+				MobileBlendableLayerMask |= (!Allocation.LayerInfo->bNoWeightBlend ? (1 << CurrentIdx) : 0);
+				CurrentIdx++;
+			}
+		}
+		return MobileMaterialInstance;
+	}
+	else // for cooking
+	{
+		UMaterialInstanceConstant* CombinationMaterialInstance = GetCombinationMaterial(true);
+		UMaterialInstanceConstant* MobileMaterialInstance = ConstructObject<ULandscapeMaterialInstanceConstant>(ULandscapeMaterialInstanceConstant::StaticClass(), GetOutermost());
+
+		MobileMaterialInstance->SetParentEditorOnly(CombinationMaterialInstance);
+
+		MobileBlendableLayerMask = 0;
+
+		// Set the layer mask
+		int32 CurrentIdx = 0;
+		for (const auto& MobileAllocation : MobileLayerAllocations)
+		{
+			const FWeightmapLayerAllocationInfo& Allocation = MobileAllocation.Allocation;
+			if (Allocation.LayerInfo)
+			{
+				FName LayerName = Allocation.LayerInfo == ALandscapeProxy::VisibilityLayer ? UMaterialExpressionLandscapeVisibilityMask::ParameterName : Allocation.LayerInfo->LayerName;
+				MobileMaterialInstance->SetVectorParameterValueEditorOnly(FName(*FString::Printf(TEXT("LayerMask_%s"), *LayerName.ToString())), Masks[FMath::Min(4, CurrentIdx)]);
+				MobileBlendableLayerMask |= (!Allocation.LayerInfo->bNoWeightBlend ? (1 << CurrentIdx) : 0);
+				CurrentIdx++;
+			}
+		}
+
+		MobileMaterialInstance->PostEditChange();
+
+		return MobileMaterialInstance;
+	}
+}
+
+//
+// Generates vertex buffer data from the component's heightmap texture, for use on platforms without vertex texture fetch
+//
+void ULandscapeComponent::GeneratePlatformVertexData()
+{
+	if (IsTemplate())
+	{
+		return;
+	}
+	check(HeightmapTexture);
+	check(HeightmapTexture->Source.GetFormat() == TSF_BGRA8);
+
+	TArray<uint8> NewPlatformData;
+	int32 NewPlatformDataSize = 0;
+
+	int32 SubsectionSizeVerts = SubsectionSizeQuads + 1;
+	int32 MaxLOD = FMath::CeilLogTwo(SubsectionSizeVerts) - 1;
+
+	float HeightmapSubsectionOffsetU = (float)(SubsectionSizeVerts) / (float)HeightmapTexture->Source.GetSizeX();
+	float HeightmapSubsectionOffsetV = (float)(SubsectionSizeVerts) / (float)HeightmapTexture->Source.GetSizeY();
+
+	NewPlatformDataSize += sizeof(FLandscapeMobileVertex) * FMath::Square(SubsectionSizeVerts * NumSubsections);
+	NewPlatformData.AddZeroed(NewPlatformDataSize);
+
+	// Get the required mip data
+	TArray<FColor*> HeightmapMipData;
+	for (int32 MipIdx = 0; MipIdx < FMath::Min(LANDSCAPE_MAX_ES_LOD, HeightmapTexture->Source.GetNumMips()); MipIdx++)
+	{
+		int32 MipSubsectionSizeVerts = (SubsectionSizeVerts) >> MipIdx;
+		if (MipSubsectionSizeVerts > 1)
+		{
+			HeightmapMipData.Add((FColor*)HeightmapTexture->Source.LockMip(MipIdx));
+		}
+	}
+
+	TMap<uint64, int32> VertexMap;
+	TArray<FLandscapeVertexRef> VertexOrder;
+	VertexOrder.Empty(FMath::Square(SubsectionSizeVerts * NumSubsections));
+
+	// Layout index buffer to determine best vertex order
+	for (int32 Mip = MaxLOD; Mip >= 0; Mip--)
+	{
+		int32 LodSubsectionSizeQuads = (SubsectionSizeVerts >> Mip) - 1;
+		float MipRatio = (float)SubsectionSizeQuads / (float)LodSubsectionSizeQuads; // Morph current MIP to base MIP
+
+		for (int32 SubY = 0; SubY < NumSubsections; SubY++)
+		{
+			for (int32 SubX = 0; SubX < NumSubsections; SubX++)
+			{
+				for (int32 y = 0; y < LodSubsectionSizeQuads; y++)
+				{
+					for (int32 x = 0; x < LodSubsectionSizeQuads; x++)
+					{
+						int32 x0 = FMath::RoundToInt((float)x * MipRatio);
+						int32 y0 = FMath::RoundToInt((float)y * MipRatio);
+						int32 x1 = FMath::RoundToInt((float)(x + 1) * MipRatio);
+						int32 y1 = FMath::RoundToInt((float)(y + 1) * MipRatio);
+
+						FLandscapeVertexRef V1(x0, y0, SubX, SubY);
+						FLandscapeVertexRef V2(x1, y0, SubX, SubY);
+						FLandscapeVertexRef V3(x1, y1, SubX, SubY);
+						FLandscapeVertexRef V4(x0, y1, SubX, SubY);
+
+						uint64 Key1 = V1.MakeKey();
+						if (VertexMap.Find(Key1) == NULL)
+						{
+							VertexMap.Add(Key1, VertexOrder.Num());
+							VertexOrder.Add(V1);
+						}
+						uint64 Key2 = V2.MakeKey();
+						if (VertexMap.Find(Key2) == NULL)
+						{
+							VertexMap.Add(Key2, VertexOrder.Num());
+							VertexOrder.Add(V2);
+						}
+						uint64 Key3 = V3.MakeKey();
+						if (VertexMap.Find(Key3) == NULL)
+						{
+							VertexMap.Add(Key3, VertexOrder.Num());
+							VertexOrder.Add(V3);
+						}
+						uint64 Key4 = V4.MakeKey();
+						if (VertexMap.Find(Key4) == NULL)
+						{
+							VertexMap.Add(Key4, VertexOrder.Num());
+							VertexOrder.Add(V4);
+						}
+					}
+				}
+			}
+		}
+	}
+	check(VertexOrder.Num() == FMath::Square(SubsectionSizeVerts) * FMath::Square(NumSubsections));
+
+	// Fill in the vertices in the specified order
+	FLandscapeMobileVertex* DstVert = (FLandscapeMobileVertex*)NewPlatformData.GetData();
+	for (int32 Idx = 0; Idx < VertexOrder.Num(); Idx++)
+	{
+		int32 X = VertexOrder[Idx].X;
+		int32 Y = VertexOrder[Idx].Y;
+		int32 SubX = VertexOrder[Idx].SubX;
+		int32 SubY = VertexOrder[Idx].SubY;
+
+		float HeightmapScaleBiasZ = HeightmapScaleBias.Z + HeightmapSubsectionOffsetU * (float)SubX;
+		float HeightmapScaleBiasW = HeightmapScaleBias.W + HeightmapSubsectionOffsetV * (float)SubY;
+		int32 BaseMipOfsX = FMath::RoundToInt(HeightmapScaleBiasZ * (float)HeightmapTexture->Source.GetSizeX());
+		int32 BaseMipOfsY = FMath::RoundToInt(HeightmapScaleBiasW * (float)HeightmapTexture->Source.GetSizeY());
+
+		DstVert->Position[0] = X;
+		DstVert->Position[1] = Y;
+		DstVert->Position[2] = SubX;
+		DstVert->Position[3] = SubY;
+
+		TArray<int32> MipHeights;
+		MipHeights.AddZeroed(HeightmapMipData.Num());
+		int32 LastIndex = 0;
+		uint16 MaxHeight = 0, MinHeight = 65535;
+
+		for (int32 Mip = 0; Mip < HeightmapMipData.Num(); ++Mip)
+		{
+			int32 MipSizeX = HeightmapTexture->Source.GetSizeX() >> Mip;
+
+			int32 MipSubsectionSizeVerts = (SubsectionSizeVerts >> Mip);
+			int32 MipSubsectionSizeQuads = MipSubsectionSizeVerts - 1;
+
+			int32 CurrentMipOfsX = BaseMipOfsX >> Mip;
+			int32 CurrentMipOfsY = BaseMipOfsY >> Mip;
+
+			float MipRatio = (float)MipSubsectionSizeQuads / (float)SubsectionSizeQuads; // Morph Base to current MIP
+			int32 MipX = FMath::RoundToInt((float)X * MipRatio);
+			int32 MipY = FMath::RoundToInt((float)Y * MipRatio);
+
+			FColor* CurrentMipSrcRow = HeightmapMipData[Mip] + (CurrentMipOfsY + MipY) * MipSizeX + CurrentMipOfsX;
+			uint16 Height = CurrentMipSrcRow[MipX].R << 8 | CurrentMipSrcRow[MipX].G;
+
+			MipHeights[Mip] = Height;
+			MaxHeight = FMath::Max(MaxHeight, Height);
+			MinHeight = FMath::Min(MinHeight, Height);
+		}
+
+		DstVert->LODHeights[0] = MinHeight >> 8;
+		DstVert->LODHeights[1] = (MinHeight & 255);
+		DstVert->LODHeights[2] = MaxHeight >> 8;
+		DstVert->LODHeights[3] = (MaxHeight & 255);
+
+		for (int32 Mip = 0; Mip < HeightmapMipData.Num(); ++Mip)
+		{
+			if (Mip < 4)
+			{
+				DstVert->LODHeights[4 + Mip] = FMath::RoundToInt(float(MipHeights[Mip] - MinHeight) / (MaxHeight - MinHeight) * 255);
+			}
+			else // Mip 4 5 packed into SubX, SubY
+			{
+				DstVert->Position[Mip - 2] += (FMath::RoundToInt(float(MipHeights[Mip] - MinHeight) / (MaxHeight - MinHeight) * 255)) & (0xfffe);
+			}
+		}
+
+		DstVert++;
+	}
+
+	for (int32 MipIdx = 0; MipIdx < HeightmapTexture->Source.GetNumMips(); MipIdx++)
+	{
+		HeightmapTexture->Source.UnlockMip(MipIdx);
+	}
+
+	// Copy to PlatformData as Compressed
+	PlatformData.InitializeFromUncompressedData(NewPlatformData);
+}
+
+UTexture2D* ALandscapeProxy::CreateLandscapeTexture(int32 InSizeX, int32 InSizeY, TextureGroup InLODGroup, ETextureSourceFormat InFormat, UObject* OptionalOverrideOuter) const
+{
+	UObject* TexOuter = OptionalOverrideOuter ? OptionalOverrideOuter : GetOutermost();
+	UTexture2D* NewTexture = ConstructObject<UTexture2D>(UTexture2D::StaticClass(), TexOuter);
+	NewTexture->Source.Init2DWithMipChain(InSizeX, InSizeY, InFormat);
+	NewTexture->SRGB = false;
+	NewTexture->CompressionNone = true;
+	NewTexture->MipGenSettings = TMGS_LeaveExistingMips;
+	NewTexture->AddressX = TA_Clamp;
+	NewTexture->AddressY = TA_Clamp;
+	NewTexture->LODGroup = InLODGroup;
+
+	return NewTexture;
+}
+
+#endif //WITH_EDITOR
+
+#undef LOCTEXT_NAMESPACE
