Landscape module from Engine split out.

diff --git a/Engine/Source/Runtime/Landscape/Private/LandscapeEditInterface.cpp b/Engine/Source/Runtime/Landscape/Private/LandscapeEditInterface.cpp
new file mode 100644
--- /dev/null
+++ b/Engine/Source/Runtime/Landscape/Private/LandscapeEditInterface.cpp
@@ -0,0 +1,4761 @@
+// Copyright 1998-2014 Epic Games, Inc. All Rights Reserved.
+
+/*=============================================================================
+LandscapeEditInterface.cpp: Landscape editing interface
+=============================================================================*/
+
+#include "Landscape.h"
+#include "LandscapeDataAccess.h"
+#include "LandscapeEdit.h"
+#include "LandscapeRender.h"
+#include "Landscape.h"
+#include "LandscapeInfo.h"
+#include "LandscapeLayerInfoObject.h"
+#include "ComponentReregisterContext.h"
+
+#if WITH_EDITOR
+
+//
+// FLandscapeEditDataInterface
+//
+FLandscapeEditDataInterface::FLandscapeEditDataInterface(ULandscapeInfo* InLandscapeInfo)
+{
+	if (InLandscapeInfo)
+	{
+		LandscapeInfo			= InLandscapeInfo;
+		ComponentSizeQuads		= InLandscapeInfo->ComponentSizeQuads;
+		SubsectionSizeQuads		= InLandscapeInfo->SubsectionSizeQuads;
+		ComponentNumSubsections = InLandscapeInfo->ComponentNumSubsections;
+		DrawScale				= InLandscapeInfo->DrawScale;
+	}
+}
+
+FLandscapeEditDataInterface::~FLandscapeEditDataInterface()
+{
+	Flush();
+}
+
+void FLandscapeEditDataInterface::Flush()
+{
+	bool bNeedToWaitForUpdate = false;
+
+	// Update all textures
+	for( TMap<UTexture2D*, FLandscapeTextureDataInfo*>::TIterator It(TextureDataMap); It;  ++It )
+	{
+		if( It.Value()->UpdateTextureData() )
+		{
+			bNeedToWaitForUpdate = true;
+		}
+	}
+
+	if( bNeedToWaitForUpdate )
+	{
+		FlushRenderingCommands();
+	}
+
+	// delete all the FLandscapeTextureDataInfo allocations
+	for( TMap<UTexture2D*, FLandscapeTextureDataInfo*>::TIterator It(TextureDataMap); It;  ++It )
+	{
+		delete It.Value();	// FLandscapeTextureDataInfo destructors will unlock any texture data
+	}
+
+	TextureDataMap.Empty();
+}
+
+#include "LevelUtils.h"
+
+void ALandscape::CalcComponentIndicesOverlap(const int32 X1, const int32 Y1, const int32 X2, const int32 Y2, const int32 ComponentSizeQuads, 
+											 int32& ComponentIndexX1, int32& ComponentIndexY1, int32& ComponentIndexX2, int32& ComponentIndexY2)
+{
+	// Find component range for this block of data
+	ComponentIndexX1 = (X1-1 >= 0) ? (X1-1) / ComponentSizeQuads : (X1) / ComponentSizeQuads - 1;	// -1 because we need to pick up vertices shared between components
+	ComponentIndexY1 = (Y1-1 >= 0) ? (Y1-1) / ComponentSizeQuads : (Y1) / ComponentSizeQuads - 1;
+	ComponentIndexX2 = (X2 >= 0) ? X2 / ComponentSizeQuads : (X2+1) / ComponentSizeQuads - 1;
+	ComponentIndexY2 = (Y2 >= 0) ? Y2 / ComponentSizeQuads : (Y2+1) / ComponentSizeQuads - 1;
+}
+
+void ALandscape::CalcComponentIndices(const int32 X1, const int32 Y1, const int32 X2, const int32 Y2, const int32 ComponentSizeQuads, 
+									  int32& ComponentIndexX1, int32& ComponentIndexY1, int32& ComponentIndexX2, int32& ComponentIndexY2)
+{
+	// Find component range for this block of data
+	ComponentIndexX1 = (X1 >= 0) ? X1 / ComponentSizeQuads : (X1+1) / ComponentSizeQuads - 1;	// -1 because we need to pick up vertices shared between components
+	ComponentIndexY1 = (Y1 >= 0) ? Y1 / ComponentSizeQuads : (Y1+1) / ComponentSizeQuads - 1;
+	ComponentIndexX2 = (X2-1 >= 0) ? (X2-1) / ComponentSizeQuads : (X2) / ComponentSizeQuads - 1;
+	ComponentIndexY2 = (Y2-1 >= 0) ? (Y2-1) / ComponentSizeQuads : (Y2) / ComponentSizeQuads - 1;
+	// Shrink indices for shared values
+	if ( ComponentIndexX2 < ComponentIndexX1)
+	{
+		ComponentIndexX2 = ComponentIndexX1;
+	}
+	if ( ComponentIndexY2 < ComponentIndexY1)
+	{
+		ComponentIndexY2 = ComponentIndexY1;
+	}
+}
+
+namespace
+{
+	// Ugly helper function, all arrays should be only size 4
+	template<typename T, typename F>
+	FORCEINLINE void CalcInterpValue(const int32* Dist, const bool* Exist, const T* Value, F& ValueX, F& ValueY)
+	{
+		if (Exist[0] && Exist[1])
+		{
+			ValueX = (F)(Dist[1] * Value[0] + Dist[0] * Value[1]) / (Dist[0] + Dist[1]);
+		}
+		else
+		{
+			if (Exist[0])
+			{
+				ValueX = Value[0];
+			}
+			else if (Exist[1])
+			{
+				ValueX = Value[1];
+			}
+		}
+
+		if (Exist[2] && Exist[3])
+		{
+			ValueY = (F)(Dist[3] * Value[2] + Dist[2] * Value[3]) / (Dist[2] + Dist[3]);
+		}
+		else
+		{
+			if (Exist[2])
+			{
+				ValueY = Value[2];
+			}
+			else if (Exist[3])
+			{
+				ValueY = Value[3];
+			}
+		}
+	}
+
+	template<typename T>
+	FORCEINLINE T CalcValueFromValueXY( const int32* Dist, const T& ValueX, const T& ValueY, const uint8& CornerSet, const T* CornerValues )
+	{
+		T FinalValue;
+		int32 DistX = FMath::Min(Dist[0], Dist[1]);
+		int32 DistY = FMath::Min(Dist[2], Dist[3]);
+		if (DistX+DistY > 0)
+		{
+			FinalValue = ((ValueX * DistY) + (ValueY * DistX)) / (float)(DistX + DistY);
+		}
+		else
+		{
+			if ((CornerSet & 1) && Dist[0] == 0 && Dist[2] == 0)
+			{
+				FinalValue = CornerValues[0];
+			}
+			else if ((CornerSet & 1 << 1) && Dist[1] == 0 && Dist[2] == 0)
+			{
+				FinalValue = CornerValues[1];
+			}
+			else if ((CornerSet & 1 << 2) && Dist[0] == 0 && Dist[3] == 0)
+			{
+				FinalValue = CornerValues[2];
+			}
+			else if ((CornerSet & 1 << 3) && Dist[1] == 0 && Dist[3] == 0)
+			{
+				FinalValue = CornerValues[3];
+			}
+			else
+			{
+				FinalValue = ValueX;
+			}
+		}
+		return FinalValue;
+	}
+
+};
+
+bool FLandscapeEditDataInterface::GetComponentsInRegion(int32 X1, int32 Y1, int32 X2, int32 Y2, TSet<ULandscapeComponent*>* OutComponents /*= NULL*/)
+{
+	if (ComponentSizeQuads <= 0 || !LandscapeInfo)
+	{
+		return false;
+	}
+	// Find component range for this block of data
+	int32 ComponentIndexX1, ComponentIndexY1, ComponentIndexX2, ComponentIndexY2;
+	ALandscape::CalcComponentIndicesOverlap(X1, Y1, X2, Y2, ComponentSizeQuads, ComponentIndexX1, ComponentIndexY1, ComponentIndexX2, ComponentIndexY2);
+
+	bool bNotLocked = true;
+	for( int32 ComponentIndexY=ComponentIndexY1;ComponentIndexY<=ComponentIndexY2;ComponentIndexY++ )
+	{
+		for( int32 ComponentIndexX=ComponentIndexX1;ComponentIndexX<=ComponentIndexX2;ComponentIndexX++ )
+		{
+			ULandscapeComponent* Component = LandscapeInfo->XYtoComponentMap.FindRef(FIntPoint(ComponentIndexX,ComponentIndexY));
+			if( Component )
+			{
+				bNotLocked = bNotLocked && ( !FLevelUtils::IsLevelLocked(Component->GetLandscapeProxy()->GetLevel()) ) && FLevelUtils::IsLevelVisible(Component->GetLandscapeProxy()->GetLevel());
+				if (OutComponents)
+				{
+					OutComponents->Add(Component);
+				}
+			}
+		}
+	}
+	return bNotLocked;
+}
+
+void FLandscapeEditDataInterface::SetHeightData(int32 X1, int32 Y1, int32 X2, int32 Y2, const uint16* Data, int32 Stride, bool CalcNormals, const uint16* NormalData /*= NULL*/, bool CreateComponents /*= false*/)
+{
+	if( Stride==0 )
+	{
+		Stride = (1+X2-X1);
+	}
+
+	check(ComponentSizeQuads > 0);
+	// Find component range for this block of data
+	int32 ComponentIndexX1, ComponentIndexY1, ComponentIndexX2, ComponentIndexY2;
+	ALandscape::CalcComponentIndicesOverlap(X1, Y1, X2, Y2, ComponentSizeQuads, ComponentIndexX1, ComponentIndexY1, ComponentIndexX2, ComponentIndexY2);
+
+	FVector* VertexNormals = NULL;
+	if( CalcNormals )
+	{
+		// Calculate the normals for each of the two triangles per quad.
+		// Note that the normals at the edges are not correct because they include normals
+		// from triangles outside the current area. They are not updated
+		int32 NumVertsX = 1+X2-X1;
+		int32 NumVertsY = 1+Y2-Y1;
+		VertexNormals = new FVector[NumVertsX*NumVertsY];
+		FMemory::Memzero(VertexNormals, NumVertsX*NumVertsY*sizeof(FVector));
+
+		// Need to consider XYOffset for XY displacemented map
+		FVector2D* XYOffsets = new FVector2D[NumVertsX*NumVertsY];
+		FMemory::Memzero(XYOffsets, NumVertsX*NumVertsY*sizeof(FVector2D));
+		GetXYOffsetDataFast(X1, Y1, X2, Y2, XYOffsets, 0);
+
+		for( int32 Y=0;Y<NumVertsY-1;Y++ )
+		{
+			for( int32 X=0;X<NumVertsX-1;X++ )
+			{
+				FVector Vert00 = FVector(XYOffsets[(X+0) + NumVertsX*(Y+0)].X,		XYOffsets[(X+0) + NumVertsX*(Y+0)].Y, ((float)Data[(X+0) + Stride*(Y+0)] - 32768.f)*LANDSCAPE_ZSCALE) * DrawScale;
+				FVector Vert01 = FVector(XYOffsets[(X+0) + NumVertsX*(Y+0)].X,		XYOffsets[(X+0) + NumVertsX*(Y+0)].Y+1.f, ((float)Data[(X+0) + Stride*(Y+1)] - 32768.f)*LANDSCAPE_ZSCALE) * DrawScale;
+				FVector Vert10 = FVector(XYOffsets[(X+0) + NumVertsX*(Y+0)].X+1.f,	XYOffsets[(X+0) + NumVertsX*(Y+0)].Y, ((float)Data[(X+1) + Stride*(Y+0)] - 32768.f)*LANDSCAPE_ZSCALE) * DrawScale;
+				FVector Vert11 = FVector(XYOffsets[(X+0) + NumVertsX*(Y+0)].X+1.f,	XYOffsets[(X+0) + NumVertsX*(Y+0)].Y+1.f, ((float)Data[(X+1) + Stride*(Y+1)] - 32768.f)*LANDSCAPE_ZSCALE) * DrawScale;
+
+				FVector FaceNormal1 = ((Vert00-Vert10) ^ (Vert10-Vert11)).SafeNormal();
+				FVector FaceNormal2 = ((Vert11-Vert01) ^ (Vert01-Vert00)).SafeNormal(); 
+
+				// contribute to the vertex normals.
+				VertexNormals[(X+1 + NumVertsX*(Y+0))] += FaceNormal1;
+				VertexNormals[(X+0 + NumVertsX*(Y+1))] += FaceNormal2;
+				VertexNormals[(X+0 + NumVertsX*(Y+0))] += FaceNormal1 + FaceNormal2;
+				VertexNormals[(X+1 + NumVertsX*(Y+1))] += FaceNormal1 + FaceNormal2;
+			}
+		}
+	}
+
+	for( int32 ComponentIndexY=ComponentIndexY1;ComponentIndexY<=ComponentIndexY2;ComponentIndexY++ )
+	{
+		for( int32 ComponentIndexX=ComponentIndexX1;ComponentIndexX<=ComponentIndexX2;ComponentIndexX++ )
+		{
+			FIntPoint ComponentKey(ComponentIndexX,ComponentIndexY);
+			ULandscapeComponent* Component = LandscapeInfo->XYtoComponentMap.FindRef(ComponentKey);
+
+			// if NULL, it was painted away
+			if( Component==NULL )
+			{
+				if( CreateComponents )
+				{
+					// not yet implemented
+					continue;
+				}
+				else
+				{
+					continue;
+				}
+			}
+
+			Component->Modify();
+
+			FLandscapeTextureDataInfo* TexDataInfo = GetTextureDataInfo(Component->HeightmapTexture);
+			FColor* HeightmapTextureData = (FColor*)TexDataInfo->GetMipData(0);
+
+			FColor* XYOffsetMipData = NULL;
+			if (Component->XYOffsetmapTexture)
+			{
+				FLandscapeTextureDataInfo* XYTexDataInfo = GetTextureDataInfo(Component->XYOffsetmapTexture);
+				XYOffsetMipData = (FColor*)XYTexDataInfo->GetMipData(Component->CollisionMipLevel); 
+			}
+
+			// Find the texture data corresponding to this vertex
+			int32 SizeU = Component->HeightmapTexture->Source.GetSizeX();
+			int32 SizeV = Component->HeightmapTexture->Source.GetSizeY();
+			int32 HeightmapOffsetX = Component->HeightmapScaleBias.Z * (float)SizeU;
+			int32 HeightmapOffsetY = Component->HeightmapScaleBias.W * (float)SizeV;
+
+			// Find coordinates of box that lies inside component
+			int32 ComponentX1 = FMath::Clamp<int32>(X1-ComponentIndexX*ComponentSizeQuads, 0, ComponentSizeQuads);
+			int32 ComponentY1 = FMath::Clamp<int32>(Y1-ComponentIndexY*ComponentSizeQuads, 0, ComponentSizeQuads);
+			int32 ComponentX2 = FMath::Clamp<int32>(X2-ComponentIndexX*ComponentSizeQuads, 0, ComponentSizeQuads);
+			int32 ComponentY2 = FMath::Clamp<int32>(Y2-ComponentIndexY*ComponentSizeQuads, 0, ComponentSizeQuads);
+
+			// Find subsection range for this box
+			int32 SubIndexX1 = FMath::Clamp<int32>((ComponentX1-1) / SubsectionSizeQuads,0,ComponentNumSubsections-1);	// -1 because we need to pick up vertices shared between subsections
+			int32 SubIndexY1 = FMath::Clamp<int32>((ComponentY1-1) / SubsectionSizeQuads,0,ComponentNumSubsections-1);
+			int32 SubIndexX2 = FMath::Clamp<int32>(ComponentX2 / SubsectionSizeQuads,0,ComponentNumSubsections-1);
+			int32 SubIndexY2 = FMath::Clamp<int32>(ComponentY2 / SubsectionSizeQuads,0,ComponentNumSubsections-1);
+
+			// To adjust bounding box
+			uint16 MinHeight = MAX_uint16;
+			uint16 MaxHeight = 0;
+
+			for( int32 SubIndexY=SubIndexY1;SubIndexY<=SubIndexY2;SubIndexY++ )
+			{
+				for( int32 SubIndexX=SubIndexX1;SubIndexX<=SubIndexX2;SubIndexX++ )
+				{
+					// Find coordinates of box that lies inside subsection
+					int32 SubX1 = FMath::Clamp<int32>(ComponentX1-SubsectionSizeQuads*SubIndexX, 0, SubsectionSizeQuads);
+					int32 SubY1 = FMath::Clamp<int32>(ComponentY1-SubsectionSizeQuads*SubIndexY, 0, SubsectionSizeQuads);
+					int32 SubX2 = FMath::Clamp<int32>(ComponentX2-SubsectionSizeQuads*SubIndexX, 0, SubsectionSizeQuads);
+					int32 SubY2 = FMath::Clamp<int32>(ComponentY2-SubsectionSizeQuads*SubIndexY, 0, SubsectionSizeQuads);
+
+					// Update texture data for the box that lies inside subsection
+					for( int32 SubY=SubY1;SubY<=SubY2;SubY++ )
+					{
+						for( int32 SubX=SubX1;SubX<=SubX2;SubX++ )
+						{
+							int32 LandscapeX = SubIndexX*SubsectionSizeQuads + ComponentIndexX*ComponentSizeQuads + SubX;
+							int32 LandscapeY = SubIndexY*SubsectionSizeQuads + ComponentIndexY*ComponentSizeQuads + SubY;
+							checkSlow( LandscapeX >= X1 && LandscapeX <= X2 );
+							checkSlow( LandscapeY >= Y1 && LandscapeY <= Y2 );
+
+							// Find the input data corresponding to this vertex
+							int32 DataIndex = (LandscapeX-X1) + Stride * (LandscapeY-Y1);
+							const uint16& Height = Data[DataIndex];
+
+							// for bounding box
+							if( Height < MinHeight )
+							{
+								MinHeight = Height;
+							}
+							if( Height > MaxHeight )
+							{
+								MaxHeight = Height;
+							}
+
+							int32 TexX = HeightmapOffsetX + (SubsectionSizeQuads+1) * SubIndexX + SubX;
+							int32 TexY = HeightmapOffsetY + (SubsectionSizeQuads+1) * SubIndexY + SubY;
+							FColor& TexData = HeightmapTextureData[ TexX + TexY * SizeU ];
+
+							// Update the texture
+							TexData.R = Height >> 8;
+							TexData.G = Height & 255;
+
+							// Update normals if we're not on an edge vertex
+							if( VertexNormals && LandscapeX > X1 && LandscapeX < X2 && LandscapeY > Y1 && LandscapeY < Y2 )
+							{
+								FVector Normal = VertexNormals[DataIndex].SafeNormal();
+								TexData.B = FMath::RoundToInt( 127.5f * (Normal.X + 1.f) );
+								TexData.A = FMath::RoundToInt( 127.5f * (Normal.Y + 1.f) );
+							}
+							else if (NormalData)
+							{
+								// Need data validation?
+								const uint16& Normal = NormalData[DataIndex];
+								TexData.B = Normal >> 8;
+								TexData.A = Normal & 255;
+							}
+						}
+					}
+
+					// Record the areas of the texture we need to re-upload
+					int32 TexX1 = HeightmapOffsetX + (SubsectionSizeQuads+1) * SubIndexX + SubX1;
+					int32 TexY1 = HeightmapOffsetY + (SubsectionSizeQuads+1) * SubIndexY + SubY1;
+					int32 TexX2 = HeightmapOffsetX + (SubsectionSizeQuads+1) * SubIndexX + SubX2;
+					int32 TexY2 = HeightmapOffsetY + (SubsectionSizeQuads+1) * SubIndexY + SubY2;
+					TexDataInfo->AddMipUpdateRegion(0,TexX1,TexY1,TexX2,TexY2);
+				}
+			}
+
+			// See if we need to adjust the bounds. Note we never shrink the bounding box at this point
+			float MinLocalZ = LandscapeDataAccess::GetLocalHeight(MinHeight);
+			float MaxLocalZ = LandscapeDataAccess::GetLocalHeight(MaxHeight);
+
+			bool bUpdateBoxSphereBounds = false;
+			if( MinLocalZ < Component->CachedLocalBox.Min.Z )
+			{
+				Component->CachedLocalBox.Min.Z = MinLocalZ;
+				bUpdateBoxSphereBounds = true;
+			}
+			if( MaxLocalZ > Component->CachedLocalBox.Max.Z )
+			{
+				Component->CachedLocalBox.Max.Z = MaxLocalZ;
+				bUpdateBoxSphereBounds = true;
+			}
+
+			if( bUpdateBoxSphereBounds )
+			{
+				Component->UpdateComponentToWorld();
+			}
+
+			// Update mipmaps
+
+			// Work out how many mips should be calculated directly from one component's data.
+			// The remaining mips are calculated on a per texture basis.
+			// eg if subsection is 7x7 quads, we need one 3 mips total: (8x8, 4x4, 2x2 verts)
+			int32 BaseNumMips = FMath::CeilLogTwo(SubsectionSizeQuads+1);
+			TArray<FColor*> MipData;
+			MipData.AddUninitialized(BaseNumMips);
+			MipData[0] = HeightmapTextureData;
+			for( int32 MipIdx=1;MipIdx<BaseNumMips;MipIdx++ )
+			{
+				MipData[MipIdx] = (FColor*)TexDataInfo->GetMipData(MipIdx);
+			}
+			Component->GenerateHeightmapMips( MipData, ComponentX1, ComponentY1, ComponentX2, ComponentY2, TexDataInfo );
+
+			// Update collision
+			Component->UpdateCollisionHeightData(MipData[Component->CollisionMipLevel], ComponentX1, ComponentY1, ComponentX2, ComponentY2, bUpdateBoxSphereBounds, XYOffsetMipData );
+
+			// Update GUID for Platform Data
+			FPlatformMisc::CreateGuid(Component->StateId);
+		}
+	}
+
+	if( VertexNormals )
+	{
+		delete[] VertexNormals;
+	}
+}
+
+//
+// RecalculateNormals - Regenerate normals for the entire landscape. Called after modifying DrawScale3D.
+//
+void FLandscapeEditDataInterface::RecalculateNormals()
+{
+	if (!LandscapeInfo) return;
+	// Recalculate normals for each component in turn
+	for( auto It = LandscapeInfo->XYtoComponentMap.CreateIterator(); It; ++It )
+	{
+		ULandscapeComponent* Component = It.Value();
+
+		// one extra row of vertex either side of the component
+		int32 X1 = Component->GetSectionBase().X-1;
+		int32 Y1 = Component->GetSectionBase().Y-1;
+		int32 X2 = Component->GetSectionBase().X+ComponentSizeQuads+1;
+		int32 Y2 = Component->GetSectionBase().Y+ComponentSizeQuads+1;
+		int32 Stride = ComponentSizeQuads+3; 
+
+		uint16* HeightData = new uint16[FMath::Square(Stride)];
+		FVector* VertexNormals = new FVector[FMath::Square(Stride)];
+		FMemory::Memzero(VertexNormals, FMath::Square(Stride)*sizeof(FVector));
+		FVector2D* XYOffsets = new FVector2D[FMath::Square(Stride)];
+		FMemory::Memzero(XYOffsets, FMath::Square(Stride)*sizeof(FVector2D));
+
+		// Get XY offset
+		GetXYOffsetDataFast(X1,Y1,X2,Y2,XYOffsets,0);
+		// Get the vertex positions for entire quad
+		GetHeightData(X1,Y1,X2,Y2,HeightData,0);
+
+		// Contribute face normals for all triangles contributing to this components' normals
+		for( int32 Y=0;Y<Stride-1;Y++ )
+		{
+			for( int32 X=0;X<Stride-1;X++ )
+			{
+				FVector Vert00 = FVector(XYOffsets[(X+0) + Stride*(Y+0)].X,		XYOffsets[(X+0) + Stride*(Y+0)].Y, ((float)HeightData[(X+0) + Stride*(Y+0)] - 32768.f)*LANDSCAPE_ZSCALE) * DrawScale;
+				FVector Vert01 = FVector(XYOffsets[(X+0) + Stride*(Y+0)].X,		XYOffsets[(X+0) + Stride*(Y+0)].Y+1.f, ((float)HeightData[(X+0) + Stride*(Y+1)] - 32768.f)*LANDSCAPE_ZSCALE) * DrawScale;
+				FVector Vert10 = FVector(XYOffsets[(X+0) + Stride*(Y+0)].X+1.f,	XYOffsets[(X+0) + Stride*(Y+0)].Y, ((float)HeightData[(X+1) + Stride*(Y+0)] - 32768.f)*LANDSCAPE_ZSCALE) * DrawScale;
+				FVector Vert11 = FVector(XYOffsets[(X+0) + Stride*(Y+0)].X+1.f,	XYOffsets[(X+0) + Stride*(Y+0)].Y+1.f,((float)HeightData[(X+1) + Stride*(Y+1)] - 32768.f)*LANDSCAPE_ZSCALE) * DrawScale;
+
+				FVector FaceNormal1 = ((Vert00-Vert10) ^ (Vert10-Vert11)).SafeNormal();
+				FVector FaceNormal2 = ((Vert11-Vert01) ^ (Vert01-Vert00)).SafeNormal(); 
+
+				// contribute to the vertex normals.
+				VertexNormals[(X+1 + Stride*(Y+0))] += FaceNormal1;
+				VertexNormals[(X+0 + Stride*(Y+1))] += FaceNormal2;
+				VertexNormals[(X+0 + Stride*(Y+0))] += FaceNormal1 + FaceNormal2;
+				VertexNormals[(X+1 + Stride*(Y+1))] += FaceNormal1 + FaceNormal2;
+			}
+		}
+
+		// Find the texture data corresponding to this vertex
+		int32 SizeU = Component->HeightmapTexture->Source.GetSizeX();
+		int32 SizeV = Component->HeightmapTexture->Source.GetSizeY();
+		int32 HeightmapOffsetX = Component->HeightmapScaleBias.Z * (float)SizeU;
+		int32 HeightmapOffsetY = Component->HeightmapScaleBias.W * (float)SizeV;
+
+		FLandscapeTextureDataInfo* TexDataInfo = GetTextureDataInfo(Component->HeightmapTexture);
+		FColor* HeightmapTextureData = (FColor*)TexDataInfo->GetMipData(0);
+
+		// Apply vertex normals to the component
+		for( int32 SubIndexY=0;SubIndexY<Component->NumSubsections;SubIndexY++ )
+		{
+			for( int32 SubIndexX=0;SubIndexX<Component->NumSubsections;SubIndexX++ )
+			{
+				for( int32 SubY=0;SubY<=SubsectionSizeQuads;SubY++ )
+				{
+					for( int32 SubX=0;SubX<=SubsectionSizeQuads;SubX++ )
+					{
+						int32 X = (SubsectionSizeQuads+1) * SubIndexX + SubX;
+						int32 Y = (SubsectionSizeQuads+1) * SubIndexY + SubY;
+						int32 DataIndex = (X+1) + (Y+1) * Stride;
+
+						int32 TexX = HeightmapOffsetX + X;
+						int32 TexY = HeightmapOffsetY + Y;
+						FColor& TexData = HeightmapTextureData[ TexX + TexY * SizeU ];
+
+						// Update the texture
+						FVector Normal = VertexNormals[DataIndex].SafeNormal();
+						TexData.B = FMath::RoundToInt( 127.5f * (Normal.X + 1.f) );
+						TexData.A = FMath::RoundToInt( 127.5f * (Normal.Y + 1.f) );
+					}
+				}
+			}
+		}
+
+		delete[] HeightData;
+		delete[] VertexNormals;
+
+		// Record the areas of the texture we need to re-upload
+		int32 TexX1 = HeightmapOffsetX;
+		int32 TexY1 = HeightmapOffsetY;
+		int32 TexX2 = HeightmapOffsetX + (SubsectionSizeQuads+1) * Component->NumSubsections - 1;
+		int32 TexY2 = HeightmapOffsetY + (SubsectionSizeQuads+1) * Component->NumSubsections - 1;
+		TexDataInfo->AddMipUpdateRegion(0,TexX1,TexY1,TexX2,TexY2);
+
+		// Work out how many mips should be calculated directly from one component's data.
+		// The remaining mips are calculated on a per texture basis.
+		// eg if subsection is 7x7 quads, we need one 3 mips total: (8x8, 4x4, 2x2 verts)
+		int32 BaseNumMips = FMath::CeilLogTwo(SubsectionSizeQuads+1);
+		TArray<FColor*> MipData;
+		MipData.AddUninitialized(BaseNumMips);
+		MipData[0] = HeightmapTextureData;
+		for( int32 MipIdx=1;MipIdx<BaseNumMips;MipIdx++ )
+		{
+			MipData[MipIdx] = (FColor*)TexDataInfo->GetMipData(MipIdx);
+		}
+		Component->GenerateHeightmapMips( MipData, 0, 0, ComponentSizeQuads, ComponentSizeQuads, TexDataInfo );
+	}
+}
+
+template<typename TStoreData>
+void FLandscapeEditDataInterface::GetHeightDataTemplFast(const int32 X1, const int32 Y1, const int32 X2, const int32 Y2, TStoreData& StoreData, TStoreData* NormalData /*= NULL*/)
+{
+	if (!LandscapeInfo) return;
+	int32 ComponentIndexX1, ComponentIndexY1, ComponentIndexX2, ComponentIndexY2;
+	ALandscape::CalcComponentIndices(X1, Y1, X2, Y2, ComponentSizeQuads, ComponentIndexX1, ComponentIndexY1, ComponentIndexX2, ComponentIndexY2);
+
+	for( int32 ComponentIndexY=ComponentIndexY1;ComponentIndexY<=ComponentIndexY2;ComponentIndexY++ )
+	{
+		for( int32 ComponentIndexX=ComponentIndexX1;ComponentIndexX<=ComponentIndexX2;ComponentIndexX++ )
+		{		
+			ULandscapeComponent* Component = LandscapeInfo->XYtoComponentMap.FindRef(FIntPoint(ComponentIndexX,ComponentIndexY));
+
+			FLandscapeTextureDataInfo* TexDataInfo = NULL;
+			FColor* HeightmapTextureData = NULL;
+			if( Component )
+			{
+				TexDataInfo = GetTextureDataInfo(Component->HeightmapTexture);
+				HeightmapTextureData = (FColor*)TexDataInfo->GetMipData(0);
+			}
+			else
+			{
+				continue;
+			}
+
+			// Find coordinates of box that lies inside component
+			int32 ComponentX1 = FMath::Clamp<int32>(X1-ComponentIndexX*ComponentSizeQuads, 0, ComponentSizeQuads);
+			int32 ComponentY1 = FMath::Clamp<int32>(Y1-ComponentIndexY*ComponentSizeQuads, 0, ComponentSizeQuads);
+			int32 ComponentX2 = FMath::Clamp<int32>(X2-ComponentIndexX*ComponentSizeQuads, 0, ComponentSizeQuads);
+			int32 ComponentY2 = FMath::Clamp<int32>(Y2-ComponentIndexY*ComponentSizeQuads, 0, ComponentSizeQuads);
+
+			// Find subsection range for this box
+			int32 SubIndexX1 = FMath::Clamp<int32>((ComponentX1-1) / SubsectionSizeQuads,0,ComponentNumSubsections-1);	// -1 because we need to pick up vertices shared between subsections
+			int32 SubIndexY1 = FMath::Clamp<int32>((ComponentY1-1) / SubsectionSizeQuads,0,ComponentNumSubsections-1);
+			int32 SubIndexX2 = FMath::Clamp<int32>(ComponentX2 / SubsectionSizeQuads,0,ComponentNumSubsections-1);
+			int32 SubIndexY2 = FMath::Clamp<int32>(ComponentY2 / SubsectionSizeQuads,0,ComponentNumSubsections-1);
+
+			for( int32 SubIndexY=SubIndexY1;SubIndexY<=SubIndexY2;SubIndexY++ )
+			{
+				for( int32 SubIndexX=SubIndexX1;SubIndexX<=SubIndexX2;SubIndexX++ )
+				{
+					// Find coordinates of box that lies inside subsection
+					int32 SubX1 = FMath::Clamp<int32>(ComponentX1-SubsectionSizeQuads*SubIndexX, 0, SubsectionSizeQuads);
+					int32 SubY1 = FMath::Clamp<int32>(ComponentY1-SubsectionSizeQuads*SubIndexY, 0, SubsectionSizeQuads);
+					int32 SubX2 = FMath::Clamp<int32>(ComponentX2-SubsectionSizeQuads*SubIndexX, 0, SubsectionSizeQuads);
+					int32 SubY2 = FMath::Clamp<int32>(ComponentY2-SubsectionSizeQuads*SubIndexY, 0, SubsectionSizeQuads);
+
+					// Update texture data for the box that lies inside subsection
+					for( int32 SubY=SubY1;SubY<=SubY2;SubY++ )
+					{
+						for( int32 SubX=SubX1;SubX<=SubX2;SubX++ )
+						{
+							int32 LandscapeX = SubIndexX*SubsectionSizeQuads + ComponentIndexX*ComponentSizeQuads + SubX;
+							int32 LandscapeY = SubIndexY*SubsectionSizeQuads + ComponentIndexY*ComponentSizeQuads + SubY;
+
+							// Find the texture data corresponding to this vertex
+							int32 SizeU = Component->HeightmapTexture->Source.GetSizeX();
+							int32 SizeV = Component->HeightmapTexture->Source.GetSizeY();
+							int32 HeightmapOffsetX = Component->HeightmapScaleBias.Z * (float)SizeU;
+							int32 HeightmapOffsetY = Component->HeightmapScaleBias.W * (float)SizeV;
+
+							int32 TexX = HeightmapOffsetX + (SubsectionSizeQuads+1) * SubIndexX + SubX;
+							int32 TexY = HeightmapOffsetY + (SubsectionSizeQuads+1) * SubIndexY + SubY;
+							FColor& TexData = HeightmapTextureData[ TexX + TexY * SizeU ];
+
+							uint16 Height = (((uint16)TexData.R) << 8) | TexData.G;
+							StoreData.Store(LandscapeX, LandscapeY, Height);
+							if (NormalData)
+							{
+								uint16 Normals = (((uint16)TexData.B) << 8) | TexData.A;
+								NormalData->Store(LandscapeX, LandscapeY, Normals);
+							}
+						}
+					}
+				}
+			}
+		}
+	}
+}
+
+template<typename TData, typename TStoreData, typename FType>
+void FLandscapeEditDataInterface::CalcMissingValues(const int32& X1, const int32& X2, const int32& Y1, const int32& Y2, 
+								   const int32& ComponentIndexX1, const int32& ComponentIndexX2, const int32& ComponentIndexY1, const int32& ComponentIndexY2, 
+								   const int32& ComponentSizeX, const int32& ComponentSizeY, TData* CornerValues, 
+								   TArray<bool>& NoBorderY1, TArray<bool>& NoBorderY2, TArray<bool>& ComponentDataExist, TStoreData& StoreData)
+{
+	bool NoBorderX1 = false, NoBorderX2 = false;
+	// Init data...
+	FMemory::Memzero(NoBorderY1.GetData(), ComponentSizeX*sizeof(bool));
+	FMemory::Memzero(NoBorderY2.GetData(), ComponentSizeX*sizeof(bool));
+	int32 BorderX1 = INT_MAX, BorderX2 = INT_MIN;
+	TArray<int32> BorderY1, BorderY2;
+	BorderY1.Empty(ComponentSizeX);
+	BorderY2.Empty(ComponentSizeX);
+	for( int32 ComponentIndexX=ComponentIndexX1;ComponentIndexX<=ComponentIndexX2;ComponentIndexX++ )
+	{
+		new (BorderY1) int32(INT_MAX);
+		new (BorderY2) int32(INT_MIN);
+	}
+
+	// fill up missing values...
+	for( int32 ComponentIndexY=ComponentIndexY1;ComponentIndexY<=ComponentIndexY2;ComponentIndexY++ )
+	{
+		NoBorderX1 = false;
+		NoBorderX2 = false;
+		BorderX1 = INT_MAX;
+		BorderX2 = INT_MIN;
+		for( int32 ComponentIndexX=ComponentIndexX1;ComponentIndexX<=ComponentIndexX2;ComponentIndexX++ )
+		{
+			int32 ComponentIndexXY = ComponentSizeX*(ComponentIndexY-ComponentIndexY1) + ComponentIndexX-ComponentIndexX1;
+			if (!ComponentDataExist[ComponentIndexXY])
+			{
+				int32 ComponentIndexXX = ComponentIndexX - ComponentIndexX1;
+				int32 ComponentIndexYY = ComponentIndexY - ComponentIndexY1;
+
+				uint8 CornerSet = 0;
+				bool ExistLeft = ComponentIndexXX > 0 && ComponentDataExist[ ComponentIndexXX-1 + ComponentIndexYY * ComponentSizeX ];
+				bool ExistUp = ComponentIndexYY > 0 && ComponentDataExist[ ComponentIndexXX + (ComponentIndexYY-1) * ComponentSizeX ];
+
+				// Search for neighbor component for interpolation
+				bool bShouldSearchX = (BorderX2 <= ComponentIndexX);
+				bool bShouldSearchY = (BorderY2[ComponentIndexXX] <= ComponentIndexY);
+				// Search for left-closest component
+				if ( bShouldSearchX || (!NoBorderX1 && BorderX1 == INT_MAX) )
+				{
+					NoBorderX1 = true;
+					BorderX1 = INT_MAX;
+					for (int32 X = ComponentIndexX-1; X >= ComponentIndexX1; X--)
+					{
+						if (ComponentDataExist[ComponentSizeX*(ComponentIndexY-ComponentIndexY1) + X-ComponentIndexX1])
+						{
+							NoBorderX1 = false;
+							BorderX1 = X;
+							break;
+						}
+					}
+				}
+				// Search for right-closest component
+				if ( bShouldSearchX || (!NoBorderX2 && BorderX2 == INT_MIN) )
+				{
+					NoBorderX2 = true;
+					BorderX2 = INT_MIN;
+					for (int32 X = ComponentIndexX+1; X <= ComponentIndexX2; X++)
+					{
+						if (ComponentDataExist[ComponentSizeX*(ComponentIndexY-ComponentIndexY1) + X-ComponentIndexX1])
+						{
+							NoBorderX2 = false;
+							BorderX2 = X;
+							break;
+						}
+					}
+				}
+				// Search for up-closest component
+				if ( bShouldSearchY || (!NoBorderY1[ComponentIndexXX] && BorderY1[ComponentIndexXX] == INT_MAX))
+				{
+					NoBorderY1[ComponentIndexXX] = true;
+					BorderY1[ComponentIndexXX] = INT_MAX;
+					for (int32 Y = ComponentIndexY-1; Y >= ComponentIndexY1; Y--)
+					{
+						if (ComponentDataExist[ComponentSizeX*(Y-ComponentIndexY1) + ComponentIndexX-ComponentIndexX1])
+						{
+							NoBorderY1[ComponentIndexXX] = false;
+							BorderY1[ComponentIndexXX] = Y;
+							break;
+						}
+					}
+				}
+				// Search for bottom-closest component
+				if ( bShouldSearchY || (!NoBorderY2[ComponentIndexXX] && BorderY2[ComponentIndexXX] == INT_MIN))
+				{
+					NoBorderY2[ComponentIndexXX] = true;
+					BorderY2[ComponentIndexXX] = INT_MIN;
+					for (int32 Y = ComponentIndexY+1; Y <= ComponentIndexY2; Y++)
+					{
+						if (ComponentDataExist[ComponentSizeX*(Y-ComponentIndexY1) + ComponentIndexX-ComponentIndexX1])
+						{
+							NoBorderY2[ComponentIndexXX] = false;
+							BorderY2[ComponentIndexXX] = Y;
+							break;
+						}
+					}
+				}
+
+				if (((ComponentIndexX == ComponentIndexX1) || (ComponentIndexY == ComponentIndexY1)) ? false : ComponentDataExist[ComponentSizeX*(ComponentIndexY-1-ComponentIndexY1) + ComponentIndexX-1-ComponentIndexX1])
+				{
+					CornerSet |= 1;
+					CornerValues[0] = TData(StoreData.Load(ComponentIndexX*ComponentSizeQuads, ComponentIndexY*ComponentSizeQuads));
+				}
+				if (((ComponentIndexX == ComponentIndexX2) || (ComponentIndexY == ComponentIndexY1)) ? false : ComponentDataExist[ComponentSizeX*(ComponentIndexY-1-ComponentIndexY1) + ComponentIndexX+1-ComponentIndexX1])
+				{
+					CornerSet |= 1 << 1;
+					CornerValues[1] = TData(StoreData.Load((ComponentIndexX + 1)*ComponentSizeQuads, ComponentIndexY*ComponentSizeQuads));
+				}
+				if (((ComponentIndexX == ComponentIndexX1) || (ComponentIndexY == ComponentIndexY2)) ? false : ComponentDataExist[ComponentSizeX*(ComponentIndexY+1-ComponentIndexY1) + ComponentIndexX-1-ComponentIndexX1])
+				{
+					CornerSet |= 1 << 2;
+					CornerValues[2] = TData(StoreData.Load(ComponentIndexX*ComponentSizeQuads, (ComponentIndexY + 1)*ComponentSizeQuads));
+				}
+				if (((ComponentIndexX == ComponentIndexX2) || (ComponentIndexY == ComponentIndexY2)) ? false : ComponentDataExist[ComponentSizeX*(ComponentIndexY+1-ComponentIndexY1) + ComponentIndexX+1-ComponentIndexX1])
+				{
+					CornerSet |= 1 << 3;
+					CornerValues[3] = TData(StoreData.Load((ComponentIndexX + 1)*ComponentSizeQuads, (ComponentIndexY + 1)*ComponentSizeQuads));
+				}
+
+				FillCornerValues(CornerSet, CornerValues);
+
+				// Find coordinates of box that lies inside component
+				int32 ComponentX1 = FMath::Clamp<int32>(X1-ComponentIndexX*ComponentSizeQuads, 0, ComponentSizeQuads);
+				int32 ComponentY1 = FMath::Clamp<int32>(Y1-ComponentIndexY*ComponentSizeQuads, 0, ComponentSizeQuads);
+				int32 ComponentX2 = FMath::Clamp<int32>(X2-ComponentIndexX*ComponentSizeQuads, 0, ComponentSizeQuads);
+				int32 ComponentY2 = FMath::Clamp<int32>(Y2-ComponentIndexY*ComponentSizeQuads, 0, ComponentSizeQuads);
+
+				// Find subsection range for this box
+				int32 SubIndexX1 = FMath::Clamp<int32>((ComponentX1-1) / SubsectionSizeQuads,0,ComponentNumSubsections-1);	// -1 because we need to pick up vertices shared between subsections
+				int32 SubIndexY1 = FMath::Clamp<int32>((ComponentY1-1) / SubsectionSizeQuads,0,ComponentNumSubsections-1);
+				int32 SubIndexX2 = FMath::Clamp<int32>(ComponentX2 / SubsectionSizeQuads,0,ComponentNumSubsections-1);
+				int32 SubIndexY2 = FMath::Clamp<int32>(ComponentY2 / SubsectionSizeQuads,0,ComponentNumSubsections-1);
+
+				for( int32 SubIndexY=SubIndexY1;SubIndexY<=SubIndexY2;SubIndexY++ )
+				{
+					for( int32 SubIndexX=SubIndexX1;SubIndexX<=SubIndexX2;SubIndexX++ )
+					{
+						// Find coordinates of box that lies inside subsection
+						int32 SubX1 = FMath::Clamp<int32>(ComponentX1-SubsectionSizeQuads*SubIndexX, 0, SubsectionSizeQuads);
+						int32 SubY1 = FMath::Clamp<int32>(ComponentY1-SubsectionSizeQuads*SubIndexY, 0, SubsectionSizeQuads);
+						int32 SubX2 = FMath::Clamp<int32>(ComponentX2-SubsectionSizeQuads*SubIndexX, 0, SubsectionSizeQuads);
+						int32 SubY2 = FMath::Clamp<int32>(ComponentY2-SubsectionSizeQuads*SubIndexY, 0, SubsectionSizeQuads);
+
+						// Update texture data for the box that lies inside subsection
+						for( int32 SubY=SubY1;SubY<=SubY2;SubY++ )
+						{
+							for( int32 SubX=SubX1;SubX<=SubX2;SubX++ )
+							{
+								int32 LandscapeX = SubIndexX*SubsectionSizeQuads + ComponentIndexX*ComponentSizeQuads + SubX;
+								int32 LandscapeY = SubIndexY*SubsectionSizeQuads + ComponentIndexY*ComponentSizeQuads + SubY;
+
+								// Find the texture data corresponding to this vertex
+								TData Value[4];
+								FMemory::Memzero(Value, sizeof(TData)* 4);
+								int32 Dist[4] = {INT_MAX, INT_MAX, INT_MAX, INT_MAX};
+								FType ValueX, ValueY;
+								FMemory::Memzero(&ValueX, sizeof(FType));
+								FMemory::Memzero(&ValueY, sizeof(FType));
+								bool Exist[4] = {false, false, false, false};
+
+								if (ExistLeft)
+								{
+									Value[0] = TData(StoreData.Load(ComponentIndexX*ComponentSizeQuads, LandscapeY));
+									Dist[0] = LandscapeX - (ComponentIndexX*ComponentSizeQuads);
+									Exist[0] = true;
+								}
+								else if ( BorderX1 != INT_MAX )
+								{
+									int32 BorderIdxX = (BorderX1+1)*ComponentSizeQuads;
+									Value[0] = TData(StoreData.Load(BorderIdxX, LandscapeY));
+									Dist[0] = LandscapeX - (BorderIdxX-1);
+									Exist[0] = true;
+								}
+								else 
+								{
+									if ((CornerSet & 1) && (CornerSet & 1 << 2))
+									{
+										int32 Dist1 = LandscapeY - (ComponentIndexY*ComponentSizeQuads);
+										int32 Dist2 = ((ComponentIndexY+1)*ComponentSizeQuads) - LandscapeY;
+										Value[0] = (FType)(Dist2 * CornerValues[0] + Dist1 * CornerValues[2]) / (Dist1 + Dist2);
+										Dist[0] = LandscapeX - (ComponentIndexX*ComponentSizeQuads);
+										Exist[0] = true;
+									}
+								}
+
+								if ( BorderX2 != INT_MIN )
+								{
+									int32 BorderIdxX = BorderX2*ComponentSizeQuads;
+									Value[1] = TData(StoreData.Load(BorderIdxX, LandscapeY));
+									Dist[1] = (BorderIdxX+1) - LandscapeX;
+									Exist[1] = true;
+								}
+								else 
+								{
+									if ((CornerSet & 1 << 1) && (CornerSet & 1 << 3))
+									{
+										int32 Dist1 = LandscapeY - (ComponentIndexY*ComponentSizeQuads);
+										int32 Dist2 = ((ComponentIndexY+1)*ComponentSizeQuads) - LandscapeY;
+										Value[1] = (FType)(Dist2 * CornerValues[1] + Dist1 * CornerValues[3]) / (Dist1 + Dist2);
+										Dist[1] = (ComponentIndexX+1)*ComponentSizeQuads - LandscapeX;
+										Exist[1] = true;
+									}
+								}
+
+								if (ExistUp)
+								{
+									Value[2] = TData(StoreData.Load(LandscapeX, ComponentIndexY*ComponentSizeQuads));
+									Dist[2] = LandscapeY - (ComponentIndexY*ComponentSizeQuads);
+									Exist[2] = true;
+								}
+								else if ( BorderY1[ComponentIndexXX] != INT_MAX )
+								{
+									int32 BorderIdxY = (BorderY1[ComponentIndexXX]+1)*ComponentSizeQuads;
+									Value[2] = TData(StoreData.Load(LandscapeX, BorderIdxY));
+									Dist[2] = LandscapeY - BorderIdxY;
+									Exist[2] = true;
+								}
+								else 
+								{
+									if ((CornerSet & 1) && (CornerSet & 1 << 1))
+									{
+										int32 Dist1 = LandscapeX - (ComponentIndexX*ComponentSizeQuads);
+										int32 Dist2 = (ComponentIndexX+1)*ComponentSizeQuads - LandscapeX;
+										Value[2] = (FType)(Dist2 * CornerValues[0] + Dist1 * CornerValues[1]) / (Dist1 + Dist2);
+										Dist[2] = LandscapeY - (ComponentIndexY*ComponentSizeQuads);
+										Exist[2] = true;
+									}
+								}
+
+								if ( BorderY2[ComponentIndexXX] != INT_MIN )
+								{
+									int32 BorderIdxY = BorderY2[ComponentIndexXX]*ComponentSizeQuads;
+									Value[3] = TData(StoreData.Load(LandscapeX, BorderIdxY));
+									Dist[3] = BorderIdxY - LandscapeY;
+									Exist[3] = true;
+								}
+								else 
+								{
+									if ((CornerSet & 1 << 2) && (CornerSet & 1 << 3))
+									{
+										int32 Dist1 = LandscapeX - (ComponentIndexX*ComponentSizeQuads);
+										int32 Dist2 = (ComponentIndexX+1)*ComponentSizeQuads - LandscapeX;
+										Value[3] = (FType)(Dist2 * CornerValues[2] + Dist1 * CornerValues[3]) / (Dist1 + Dist2);
+										Dist[3] = (ComponentIndexY+1)*ComponentSizeQuads - LandscapeY;
+										Exist[3] = true;
+									}
+								}
+
+								CalcInterpValue<TData, FType>(Dist, Exist, Value, ValueX, ValueY);
+
+								TData FinalValue; // Default Value
+								FMemory::Memzero(&FinalValue, sizeof(TData));
+								if ( (Exist[0] || Exist[1]) && (Exist[2] || Exist[3]) )
+								{
+									FinalValue = CalcValueFromValueXY<TData>(Dist, ValueX, ValueY, CornerSet, CornerValues);
+								}
+								else if ( (Exist[0] || Exist[1]) )
+								{
+									FinalValue = ValueX;
+								}
+								else if ( (Exist[2] || Exist[3]) )
+								{
+									FinalValue = ValueY;
+								}
+
+								StoreData.Store(LandscapeX, LandscapeY, FinalValue);
+							}
+						}
+					}
+				}
+			}
+		}
+	}
+}
+
+uint16 FLandscapeEditDataInterface::GetHeightMapData(const ULandscapeComponent* Component, int32 TexU, int32 TexV, FColor* TextureData /*= NULL*/)
+{
+	check(Component);
+	if (!TextureData)
+	{
+		FLandscapeTextureDataInfo* TexDataInfo = GetTextureDataInfo(Component->HeightmapTexture);
+		TextureData = (FColor*)TexDataInfo->GetMipData(0);	
+	}
+
+	int32 SizeU = Component->HeightmapTexture->Source.GetSizeX();
+	int32 SizeV = Component->HeightmapTexture->Source.GetSizeY();
+	int32 HeightmapOffsetX = Component->HeightmapScaleBias.Z * (float)SizeU;
+	int32 HeightmapOffsetY = Component->HeightmapScaleBias.W * (float)SizeV;
+
+	int32 TexX = HeightmapOffsetX + TexU;
+	int32 TexY = HeightmapOffsetY + TexV;
+	FColor& TexData = TextureData[ TexX + TexY * SizeU ];
+
+	return ((((uint16)TexData.R) << 8) | TexData.G);
+}
+
+template<typename T>
+void FillCornerValues(uint8& CornerSet, T* CornerValues)
+{
+	uint8 OriginalSet = CornerSet;
+
+	if (CornerSet)
+	{
+		// Fill unset values
+		while (CornerSet != 15)
+		{
+			if (CornerSet != 15 && (OriginalSet & 1))
+			{
+				if (!(CornerSet & 1 << 1))
+				{
+					CornerValues[1] = CornerValues[0];
+					CornerSet |= 1 << 1;
+				}
+				if (!(CornerSet & 1 << 2))
+				{
+					CornerValues[2] = CornerValues[0];
+					CornerSet |= 1 << 2;
+				}
+			}
+			if (CornerSet != 15 && (OriginalSet & 1 << 1))
+			{
+				if (!(CornerSet & 1))
+				{
+					CornerValues[0] = CornerValues[1];
+					CornerSet |= 1;
+				}
+				if (!(CornerSet & 1 << 3))
+				{
+					CornerValues[3] = CornerValues[1];
+					CornerSet |= 1 << 3;
+				}
+			}
+			if (CornerSet != 15 && (OriginalSet & 1 << 2))
+			{
+				if (!(CornerSet & 1))
+				{
+					CornerValues[0] = CornerValues[2];
+					CornerSet |= 1;
+				}
+				if (!(CornerSet & 1 << 3))
+				{
+					CornerValues[3] = CornerValues[2];
+					CornerSet |= 1 << 3;
+				}
+			}
+			if (CornerSet != 15 && (OriginalSet & 1 << 3))
+			{
+				if (!(CornerSet & 1 << 1))
+				{
+					CornerValues[1] = CornerValues[3];
+					CornerSet |= 1 << 1;
+				}
+				if (!(CornerSet & 1 << 2))
+				{
+					CornerValues[2] = CornerValues[3];
+					CornerSet |= 1 << 2;
+				}
+			}
+
+			OriginalSet = CornerSet;
+		}
+	}
+}
+
+template<typename TStoreData>
+void FLandscapeEditDataInterface::GetHeightDataTempl(int32& ValidX1, int32& ValidY1, int32& ValidX2, int32& ValidY2, TStoreData& StoreData)
+{
+	// Copy variables
+	int32 X1 = ValidX1, X2 = ValidX2, Y1 = ValidY1, Y2 = ValidY2;
+	ValidX1 = INT_MAX; ValidX2 = INT_MIN; ValidY1 = INT_MAX; ValidY2 = INT_MIN;
+
+	int32 ComponentIndexX1, ComponentIndexY1, ComponentIndexX2, ComponentIndexY2;
+	ALandscape::CalcComponentIndicesOverlap(X1, Y1, X2, Y2, ComponentSizeQuads, ComponentIndexX1, ComponentIndexY1, ComponentIndexX2, ComponentIndexY2);
+
+	int32 ComponentSizeX = ComponentIndexX2-ComponentIndexX1+1;
+	int32 ComponentSizeY = ComponentIndexY2-ComponentIndexY1+1;
+
+	// Neighbor Components
+	ULandscapeComponent* BorderComponent[4] = {0, 0, 0, 0};
+	ULandscapeComponent* CornerComponent[4] = {0, 0, 0, 0};
+	bool NoBorderX1 = false, NoBorderX2 = false;
+	TArray<bool> NoBorderY1, NoBorderY2, ComponentDataExist;
+	TArray<ULandscapeComponent*> BorderComponentY1, BorderComponentY2;
+	ComponentDataExist.Empty(ComponentSizeX*ComponentSizeY);
+	ComponentDataExist.AddZeroed(ComponentSizeX*ComponentSizeY);
+	bool bHasMissingValue = false;
+
+	FLandscapeTextureDataInfo* NeighborTexDataInfo[4] = {0, 0, 0, 0};
+	FColor* NeighborHeightmapTextureData[4] = {0, 0, 0, 0};
+	uint16 CornerValues[4] = {0, 0, 0, 0};
+
+	int32 EdgeCoord = (SubsectionSizeQuads+1) * ComponentNumSubsections - 1; //ComponentSizeQuads;
+
+	// initial loop....
+	for( int32 ComponentIndexY=ComponentIndexY1;ComponentIndexY<=ComponentIndexY2;ComponentIndexY++ )
+	{
+		NoBorderX1 = false;
+		NoBorderX2 = false;
+		BorderComponent[0] = BorderComponent[1] = NULL;
+		for( int32 ComponentIndexX=ComponentIndexX1;ComponentIndexX<=ComponentIndexX2;ComponentIndexX++ )
+		{
+			BorderComponent[2] = BorderComponent[3] = NULL;
+			int32 ComponentIndexXY = ComponentSizeX*(ComponentIndexY-ComponentIndexY1) + ComponentIndexX-ComponentIndexX1;
+			int32 ComponentIndexXX = ComponentIndexX - ComponentIndexX1;
+			int32 ComponentIndexYY = ComponentIndexY - ComponentIndexY1;
+			ComponentDataExist[ComponentIndexXY] = false;
+			ULandscapeComponent* Component = LandscapeInfo->XYtoComponentMap.FindRef(FIntPoint(ComponentIndexX,ComponentIndexY));
+
+			FLandscapeTextureDataInfo* TexDataInfo = NULL;
+			FColor* HeightmapTextureData = NULL;
+			uint8 CornerSet = 0;
+			bool ExistLeft = ComponentIndexXX > 0 && ComponentDataExist[ ComponentIndexXX-1 + ComponentIndexYY * ComponentSizeX ];
+			bool ExistUp = ComponentIndexYY > 0 && ComponentDataExist[ ComponentIndexXX + (ComponentIndexYY-1) * ComponentSizeX ];
+
+			if( Component )
+			{
+				TexDataInfo = GetTextureDataInfo(Component->HeightmapTexture);
+				HeightmapTextureData = (FColor*)TexDataInfo->GetMipData(0);
+				ComponentDataExist[ComponentIndexXY] = true;
+				// Update valid region
+				ValidX1 = FMath::Min<int32>(Component->GetSectionBase().X, ValidX1);
+				ValidX2 = FMath::Max<int32>(Component->GetSectionBase().X+ComponentSizeQuads, ValidX2);
+				ValidY1 = FMath::Min<int32>(Component->GetSectionBase().Y, ValidY1);
+				ValidY2 = FMath::Max<int32>(Component->GetSectionBase().Y+ComponentSizeQuads, ValidY2);
+			}
+			else
+			{
+				if (!bHasMissingValue)
+				{
+					NoBorderY1.Empty(ComponentSizeX);
+					NoBorderY2.Empty(ComponentSizeX);
+					NoBorderY1.AddZeroed(ComponentSizeX);
+					NoBorderY2.AddZeroed(ComponentSizeX);
+					BorderComponentY1.Empty(ComponentSizeX);
+					BorderComponentY2.Empty(ComponentSizeX);
+					BorderComponentY1.AddZeroed(ComponentSizeX);
+					BorderComponentY2.AddZeroed(ComponentSizeX);
+					bHasMissingValue = true;
+				}
+
+				// Search for neighbor component for interpolation
+				bool bShouldSearchX = (BorderComponent[1] && BorderComponent[1]->GetSectionBase().X / ComponentSizeQuads <= ComponentIndexX);
+				bool bShouldSearchY = (BorderComponentY2[ComponentIndexXX] && BorderComponentY2[ComponentIndexXX]->GetSectionBase().Y / ComponentSizeQuads <= ComponentIndexY);
+				// Search for left-closest component
+				if ( bShouldSearchX || (!NoBorderX1 && !BorderComponent[0]))
+				{
+					NoBorderX1 = true;
+					for (int32 X = ComponentIndexX-1; X >= ComponentIndexX1; X--)
+					{
+						BorderComponent[0] = LandscapeInfo->XYtoComponentMap.FindRef(FIntPoint(X,ComponentIndexY));
+						if (BorderComponent[0])
+						{
+							NoBorderX1 = false;
+							NeighborTexDataInfo[0] = GetTextureDataInfo(BorderComponent[0]->HeightmapTexture);
+							NeighborHeightmapTextureData[0] = (FColor*)NeighborTexDataInfo[0]->GetMipData(0);
+							break;
+						}
+					}
+				}
+				// Search for right-closest component
+				if ( bShouldSearchX || (!NoBorderX2 && !BorderComponent[1]))
+				{
+					NoBorderX2 = true;
+					for (int32 X = ComponentIndexX+1; X <= ComponentIndexX2; X++)
+					{
+						BorderComponent[1] = LandscapeInfo->XYtoComponentMap.FindRef(FIntPoint(X,ComponentIndexY));
+						if (BorderComponent[1])
+						{
+							NoBorderX2 = false;
+							NeighborTexDataInfo[1] = GetTextureDataInfo(BorderComponent[1]->HeightmapTexture);
+							NeighborHeightmapTextureData[1] = (FColor*)NeighborTexDataInfo[1]->GetMipData(0);
+							break;
+						}
+					}
+				}
+				// Search for up-closest component
+				if ( bShouldSearchY || (!NoBorderY1[ComponentIndexXX] && !BorderComponentY1[ComponentIndexXX]))
+				{
+					NoBorderY1[ComponentIndexXX] = true;
+					for (int32 Y = ComponentIndexY-1; Y >= ComponentIndexY1; Y--)
+					{
+						BorderComponentY1[ComponentIndexXX] = BorderComponent[2] = LandscapeInfo->XYtoComponentMap.FindRef(FIntPoint(ComponentIndexX,Y));
+						if (BorderComponent[2])
+						{
+							NoBorderY1[ComponentIndexXX] = false;
+							NeighborTexDataInfo[2] = GetTextureDataInfo(BorderComponent[2]->HeightmapTexture);
+							NeighborHeightmapTextureData[2] = (FColor*)NeighborTexDataInfo[2]->GetMipData(0);
+							break;
+						}
+					}
+				}
+				else
+				{
+					BorderComponent[2] = BorderComponentY1[ComponentIndexXX];
+					if (BorderComponent[2])
+					{
+						NeighborTexDataInfo[2] = GetTextureDataInfo(BorderComponent[2]->HeightmapTexture);
+						NeighborHeightmapTextureData[2] = (FColor*)NeighborTexDataInfo[2]->GetMipData(0);
+					}
+				}
+				// Search for bottom-closest component
+				if ( bShouldSearchY || (!NoBorderY2[ComponentIndexXX] && !BorderComponentY2[ComponentIndexXX]))
+				{
+					NoBorderY2[ComponentIndexXX] = true;
+					for (int32 Y = ComponentIndexY+1; Y <= ComponentIndexY2; Y++)
+					{
+						BorderComponentY2[ComponentIndexXX] = BorderComponent[3] = LandscapeInfo->XYtoComponentMap.FindRef(FIntPoint(ComponentIndexX,Y));
+						if (BorderComponent[3])
+						{
+							NoBorderY2[ComponentIndexXX] = false;
+							NeighborTexDataInfo[3] = GetTextureDataInfo(BorderComponent[3]->HeightmapTexture);
+							NeighborHeightmapTextureData[3] = (FColor*)NeighborTexDataInfo[3]->GetMipData(0);
+							break;
+						}
+					}
+				}
+				else
+				{
+					BorderComponent[3] = BorderComponentY2[ComponentIndexXX];
+					if (BorderComponent[3])
+					{
+						NeighborTexDataInfo[3] = GetTextureDataInfo(BorderComponent[3]->HeightmapTexture);
+						NeighborHeightmapTextureData[3] = (FColor*)NeighborTexDataInfo[3]->GetMipData(0);
+					}
+				}
+
+				CornerComponent[0] = ComponentIndexX >= ComponentIndexX1 && ComponentIndexY >= ComponentIndexY1 ?
+					LandscapeInfo->XYtoComponentMap.FindRef(FIntPoint((ComponentIndexX-1),(ComponentIndexY-1))) : NULL;
+				CornerComponent[1] = ComponentIndexX <= ComponentIndexX2 && ComponentIndexY >= ComponentIndexY1 ?
+					LandscapeInfo->XYtoComponentMap.FindRef(FIntPoint((ComponentIndexX+1),(ComponentIndexY-1))) : NULL;
+				CornerComponent[2] = ComponentIndexX >= ComponentIndexX1 && ComponentIndexY <= ComponentIndexY2 ?
+					LandscapeInfo->XYtoComponentMap.FindRef(FIntPoint((ComponentIndexX-1),(ComponentIndexY+1))) : NULL;
+				CornerComponent[3] = ComponentIndexX <= ComponentIndexX2 && ComponentIndexY <= ComponentIndexY2 ?
+					LandscapeInfo->XYtoComponentMap.FindRef(FIntPoint((ComponentIndexX+1),(ComponentIndexY+1))) : NULL;
+
+				if (CornerComponent[0])
+				{
+					CornerSet |= 1;
+					CornerValues[0] = GetHeightMapData(CornerComponent[0], EdgeCoord, EdgeCoord);
+				}
+				else if ((ExistLeft || ExistUp) && X1 <= ComponentIndexX*ComponentSizeQuads && Y1 <= ComponentIndexY*ComponentSizeQuads  )
+				{
+					CornerSet |= 1;
+					CornerValues[0] = StoreData.Load( ComponentIndexX*ComponentSizeQuads, ComponentIndexY*ComponentSizeQuads);
+				}
+				else if (BorderComponent[0])
+				{
+					CornerSet |= 1;
+					CornerValues[0] = GetHeightMapData(BorderComponent[0], EdgeCoord, 0, NeighborHeightmapTextureData[0]);
+				}
+				else if (BorderComponent[2])
+				{
+					CornerSet |= 1;
+					CornerValues[0] = GetHeightMapData(BorderComponent[2], 0, EdgeCoord, NeighborHeightmapTextureData[2]);
+				}
+
+				if (CornerComponent[1])
+				{
+					CornerSet |= 1 << 1;
+					CornerValues[1] = GetHeightMapData(CornerComponent[1], 0, EdgeCoord);
+				}
+				else if (ExistUp && X2 >= (ComponentIndexX+1)*ComponentSizeQuads)
+				{
+					CornerSet |= 1 << 1;
+					CornerValues[1] = StoreData.Load( (ComponentIndexX+1)*ComponentSizeQuads, ComponentIndexY*ComponentSizeQuads);
+				}
+				else if (BorderComponent[1])
+				{
+					CornerSet |= 1 << 1;
+					CornerValues[1] = GetHeightMapData(BorderComponent[1], 0, 0, NeighborHeightmapTextureData[1]);
+				}
+				else if (BorderComponent[2])
+				{
+					CornerSet |= 1 << 1;
+					CornerValues[1] = GetHeightMapData(BorderComponent[2], EdgeCoord, EdgeCoord, NeighborHeightmapTextureData[2]);
+				}
+
+				if (CornerComponent[2])
+				{
+					CornerSet |= 1 << 2;
+					CornerValues[2] = GetHeightMapData(CornerComponent[2], EdgeCoord, 0);
+				}
+				else if (ExistLeft && Y2 >= (ComponentIndexY+1)*ComponentSizeQuads) // Use data already stored for 0, 2
+				{
+					CornerSet |= 1 << 2;
+					CornerValues[2] = StoreData.Load( ComponentIndexX*ComponentSizeQuads, (ComponentIndexY+1)*ComponentSizeQuads);
+				}
+				else if (BorderComponent[0])
+				{
+					CornerSet |= 1 << 2;
+					CornerValues[2] = GetHeightMapData(BorderComponent[0], EdgeCoord, EdgeCoord, NeighborHeightmapTextureData[0]);
+				}
+				else if (BorderComponent[3])
+				{
+					CornerSet |= 1 << 2;
+					CornerValues[2] = GetHeightMapData(BorderComponent[3], 0, 0, NeighborHeightmapTextureData[3]);
+				}
+
+				if (CornerComponent[3])
+				{
+					CornerSet |= 1 << 3;
+					CornerValues[3] = GetHeightMapData(CornerComponent[3], 0, 0);
+				}
+				else if (BorderComponent[1])
+				{
+					CornerSet |= 1 << 3;
+					CornerValues[3] = GetHeightMapData(BorderComponent[1], 0, EdgeCoord, NeighborHeightmapTextureData[1]);
+				}
+				else if (BorderComponent[3])
+				{
+					CornerSet |= 1 << 3;
+					CornerValues[3] = GetHeightMapData(BorderComponent[3], EdgeCoord, 0, NeighborHeightmapTextureData[3]);
+				}
+
+				FillCornerValues(CornerSet, CornerValues);
+				ComponentDataExist[ComponentIndexXY] = ExistLeft || ExistUp || (BorderComponent[0] || BorderComponent[1] || BorderComponent[2] || BorderComponent[3]) || CornerSet;
+			}
+
+			if (!ComponentDataExist[ComponentIndexXY])
+			{
+				continue;
+			}
+
+			// Find coordinates of box that lies inside component
+			int32 ComponentX1 = FMath::Clamp<int32>(X1-ComponentIndexX*ComponentSizeQuads, 0, ComponentSizeQuads);
+			int32 ComponentY1 = FMath::Clamp<int32>(Y1-ComponentIndexY*ComponentSizeQuads, 0, ComponentSizeQuads);
+			int32 ComponentX2 = FMath::Clamp<int32>(X2-ComponentIndexX*ComponentSizeQuads, 0, ComponentSizeQuads);
+			int32 ComponentY2 = FMath::Clamp<int32>(Y2-ComponentIndexY*ComponentSizeQuads, 0, ComponentSizeQuads);
+
+			// Find subsection range for this box
+			int32 SubIndexX1 = FMath::Clamp<int32>((ComponentX1-1) / SubsectionSizeQuads,0,ComponentNumSubsections-1);	// -1 because we need to pick up vertices shared between subsections
+			int32 SubIndexY1 = FMath::Clamp<int32>((ComponentY1-1) / SubsectionSizeQuads,0,ComponentNumSubsections-1);
+			int32 SubIndexX2 = FMath::Clamp<int32>(ComponentX2 / SubsectionSizeQuads,0,ComponentNumSubsections-1);
+			int32 SubIndexY2 = FMath::Clamp<int32>(ComponentY2 / SubsectionSizeQuads,0,ComponentNumSubsections-1);
+
+			for( int32 SubIndexY=SubIndexY1;SubIndexY<=SubIndexY2;SubIndexY++ )
+			{
+				for( int32 SubIndexX=SubIndexX1;SubIndexX<=SubIndexX2;SubIndexX++ )
+				{
+					// Find coordinates of box that lies inside subsection
+					int32 SubX1 = FMath::Clamp<int32>(ComponentX1-SubsectionSizeQuads*SubIndexX, 0, SubsectionSizeQuads);
+					int32 SubY1 = FMath::Clamp<int32>(ComponentY1-SubsectionSizeQuads*SubIndexY, 0, SubsectionSizeQuads);
+					int32 SubX2 = FMath::Clamp<int32>(ComponentX2-SubsectionSizeQuads*SubIndexX, 0, SubsectionSizeQuads);
+					int32 SubY2 = FMath::Clamp<int32>(ComponentY2-SubsectionSizeQuads*SubIndexY, 0, SubsectionSizeQuads);
+
+					// Update texture data for the box that lies inside subsection
+					for( int32 SubY=SubY1;SubY<=SubY2;SubY++ )
+					{
+						for( int32 SubX=SubX1;SubX<=SubX2;SubX++ )
+						{
+							int32 LandscapeX = SubIndexX*SubsectionSizeQuads + ComponentIndexX*ComponentSizeQuads + SubX;
+							int32 LandscapeY = SubIndexY*SubsectionSizeQuads + ComponentIndexY*ComponentSizeQuads + SubY;
+
+							// Find the input data corresponding to this vertex
+							if( Component )
+							{
+								// Find the texture data corresponding to this vertex
+								uint16 Height = GetHeightMapData(Component, (SubsectionSizeQuads+1) * SubIndexX + SubX, (SubsectionSizeQuads+1) * SubIndexY + SubY, HeightmapTextureData);
+								StoreData.Store(LandscapeX, LandscapeY, Height);
+							}
+							else
+							{
+								// Find the texture data corresponding to this vertex
+								uint16 Value[4] = {0, 0, 0, 0};
+								int32 Dist[4] = {INT_MAX, INT_MAX, INT_MAX, INT_MAX};
+								float ValueX = 0.f, ValueY = 0.f;
+								bool Exist[4] = {false, false, false, false};
+
+								// Use data already stored for 0, 2
+								if (ExistLeft)
+								{
+									Value[0] = StoreData.Load( ComponentIndexX*ComponentSizeQuads, LandscapeY);
+									Dist[0] = LandscapeX - (ComponentIndexX*ComponentSizeQuads);
+									Exist[0] = true;
+								}
+								else if (BorderComponent[0])
+								{
+									Value[0] = GetHeightMapData(BorderComponent[0], EdgeCoord, (SubsectionSizeQuads+1) * SubIndexY + SubY, NeighborHeightmapTextureData[0]);
+									Dist[0] = LandscapeX - (BorderComponent[0]->GetSectionBase().X + ComponentSizeQuads);
+									Exist[0] = true;
+								}
+								else 
+								{
+									if ((CornerSet & 1) && (CornerSet & 1 << 2))
+									{
+										int32 Dist1 = LandscapeY - (ComponentIndexY*ComponentSizeQuads);
+										int32 Dist2 = ((ComponentIndexY+1)*ComponentSizeQuads) - LandscapeY;
+										Value[0] = (float)(Dist2 * CornerValues[0] + Dist1 * CornerValues[2]) / (Dist1 + Dist2);
+										Dist[0] = LandscapeX - (ComponentIndexX*ComponentSizeQuads);
+										Exist[0] = true;
+									}
+								}
+
+								if (BorderComponent[1])
+								{
+									Value[1] = GetHeightMapData(BorderComponent[1], 0, (SubsectionSizeQuads+1) * SubIndexY + SubY, NeighborHeightmapTextureData[1]);
+									Dist[1] = (BorderComponent[1]->GetSectionBase().X) - LandscapeX;
+									Exist[1] = true;
+								}
+								else
+								{
+									if ((CornerSet & 1 << 1) && (CornerSet & 1 << 3))
+									{
+										int32 Dist1 = LandscapeY - (ComponentIndexY*ComponentSizeQuads);
+										int32 Dist2 = ((ComponentIndexY+1)*ComponentSizeQuads) - LandscapeY;
+										Value[1] = (float)(Dist2 * CornerValues[1] + Dist1 * CornerValues[3]) / (Dist1 + Dist2);
+										Dist[1] = (ComponentIndexX+1)*ComponentSizeQuads - LandscapeX;
+										Exist[1] = true;
+									}
+								}
+
+								if (ExistUp)
+								{
+									Value[2] = StoreData.Load( LandscapeX, ComponentIndexY*ComponentSizeQuads);
+									Dist[2] = LandscapeY - (ComponentIndexY*ComponentSizeQuads);
+									Exist[2] = true;
+								}
+								else if (BorderComponent[2])
+								{
+									Value[2] = GetHeightMapData(BorderComponent[2], (SubsectionSizeQuads+1) * SubIndexX + SubX, EdgeCoord, NeighborHeightmapTextureData[2]);
+									Dist[2] = LandscapeY - (BorderComponent[2]->GetSectionBase().Y + ComponentSizeQuads);
+									Exist[2] = true;
+								}
+								else
+								{
+									if ((CornerSet & 1) && (CornerSet & 1 << 1))
+									{
+										int32 Dist1 = LandscapeX - (ComponentIndexX*ComponentSizeQuads);
+										int32 Dist2 = (ComponentIndexX+1)*ComponentSizeQuads - LandscapeX;
+										Value[2] = (float)(Dist2 * CornerValues[0] + Dist1 * CornerValues[1]) / (Dist1 + Dist2);
+										Dist[2] = LandscapeY - (ComponentIndexY*ComponentSizeQuads);
+										Exist[2] = true;
+									}
+								}
+
+								if (BorderComponent[3])
+								{
+									Value[3] = GetHeightMapData(BorderComponent[3], (SubsectionSizeQuads+1) * SubIndexX + SubX, 0, NeighborHeightmapTextureData[3]);
+									Dist[3] = (BorderComponent[3]->GetSectionBase().Y) - LandscapeY;
+									Exist[3] = true;
+								}
+								else
+								{
+									if ((CornerSet & 1 << 2) && (CornerSet & 1 << 3))
+									{
+										int32 Dist1 = LandscapeX - (ComponentIndexX*ComponentSizeQuads);
+										int32 Dist2 = (ComponentIndexX+1)*ComponentSizeQuads - LandscapeX;
+										Value[3] = (float)(Dist2 * CornerValues[2] + Dist1 * CornerValues[3]) / (Dist1 + Dist2);
+										Dist[3] = (ComponentIndexY+1)*ComponentSizeQuads - LandscapeY;
+										Exist[3] = true;
+									}
+								}
+
+								CalcInterpValue<uint16>(Dist, Exist, Value, ValueX, ValueY);
+
+								uint16 FinalValue = 0; // Default Value
+								if ( (Exist[0] || Exist[1]) && (Exist[2] || Exist[3]) )
+								{
+									FinalValue = CalcValueFromValueXY<uint16>(Dist, ValueX, ValueY, CornerSet, CornerValues);
+								}
+								else if ( (BorderComponent[0] || BorderComponent[1]) )
+								{
+									FinalValue = ValueX;
+								}
+								else if ( (BorderComponent[2] || BorderComponent[3]) )
+								{
+									FinalValue = ValueY;
+								}
+								else if ( (Exist[0] || Exist[1]) )
+								{
+									FinalValue = ValueX;
+								}
+								else if ( (Exist[2] || Exist[3]) )
+								{
+									FinalValue = ValueY;
+								}
+
+								StoreData.Store(LandscapeX, LandscapeY, FinalValue);
+								//StoreData.StoreDefault(LandscapeX, LandscapeY);
+							}
+						}
+					}
+				}
+			}
+		}
+	}
+
+	if (bHasMissingValue)
+	{
+		CalcMissingValues<uint16, TStoreData, float>( X1, X2, Y1, Y2,
+			ComponentIndexX1, ComponentIndexX2, ComponentIndexY1, ComponentIndexY2, 
+			ComponentSizeX, ComponentSizeY, CornerValues,
+			NoBorderY1, NoBorderY2, ComponentDataExist, StoreData );
+		// Update valid region
+		ValidX1 = FMath::Max<int32>(X1, ValidX1);
+		ValidX2 = FMath::Min<int32>(X2, ValidX2);
+		ValidY1 = FMath::Max<int32>(Y1, ValidY1);
+		ValidY2 = FMath::Min<int32>(Y2, ValidY2);
+	}
+	else
+	{
+		ValidX1 = X1;
+		ValidX2 = X2;
+		ValidY1 = Y1;
+		ValidY2 = Y2;
+	}
+}
+
+namespace
+{
+	struct FArrayStoreData
+	{
+		int32 X1;
+		int32 Y1;
+		uint16* Data;
+		int32 Stride;
+
+		FArrayStoreData(int32 InX1, int32 InY1, uint16* InData, int32 InStride)
+			:	X1(InX1)
+			,	Y1(InY1)
+			,	Data(InData)
+			,	Stride(InStride)
+		{}
+
+		inline void Store(int32 LandscapeX, int32 LandscapeY, uint16 Height)
+		{
+			Data[ (LandscapeY-Y1) * Stride + (LandscapeX-X1) ] = Height;
+		}
+
+		// for interpolation
+		inline uint16 Load(int32 LandscapeX, int32 LandscapeY)
+		{
+			return Data[ (LandscapeY-Y1) * Stride + (LandscapeX-X1) ];
+		}
+
+		inline void StoreDefault(int32 LandscapeX, int32 LandscapeY)
+		{
+			Data[ (LandscapeY-Y1) * Stride + (LandscapeX-X1) ] = 0;
+		}
+	};
+
+	struct FSparseStoreData
+	{
+		TMap<FIntPoint, uint16>& SparseData;
+
+		FSparseStoreData(TMap<FIntPoint, uint16>& InSparseData)
+			:	SparseData(InSparseData)
+		{}
+
+		inline void Store(int32 LandscapeX, int32 LandscapeY, uint16 Height)
+		{
+			SparseData.Add(ALandscape::MakeKey(LandscapeX,LandscapeY), Height);
+		}
+
+		inline uint16 Load(int32 LandscapeX, int32 LandscapeY)
+		{
+			return SparseData.FindRef(ALandscape::MakeKey(LandscapeX,LandscapeY));
+		}
+
+		inline void StoreDefault(int32 LandscapeX, int32 LandscapeY)
+		{
+		}
+	};
+
+};
+
+void FLandscapeEditDataInterface::GetHeightData(int32& X1, int32& Y1, int32& X2, int32& Y2, uint16* Data, int32 Stride)
+{
+	if( Stride==0 )
+	{
+		Stride = (1+X2-X1);
+	}
+
+	FArrayStoreData ArrayStoreData(X1, Y1, Data, Stride);
+	GetHeightDataTempl(X1, Y1, X2, Y2, ArrayStoreData);
+}
+
+void FLandscapeEditDataInterface::GetHeightDataFast(const int32 X1, const int32 Y1, const int32 X2, const int32 Y2, uint16* Data, int32 Stride, uint16* NormalData /*= NULL*/)
+{
+	if( Stride==0 )
+	{
+		Stride = (1+X2-X1);
+	}
+
+	FArrayStoreData ArrayStoreData(X1, Y1, Data, Stride);
+	if (NormalData)
+	{
+		FArrayStoreData ArrayNormalData(X1, Y1, NormalData, Stride);
+		GetHeightDataTemplFast(X1, Y1, X2, Y2, ArrayStoreData, &ArrayNormalData);
+	}
+	else
+	{
+		GetHeightDataTemplFast(X1, Y1, X2, Y2, ArrayStoreData);
+	}
+}
+
+void FLandscapeEditDataInterface::GetHeightData(int32& X1, int32& Y1, int32& X2, int32& Y2, TMap<FIntPoint, uint16>& SparseData)
+{
+	FSparseStoreData SparseStoreData(SparseData);
+	GetHeightDataTempl(X1, Y1, X2, Y2, SparseStoreData);
+}
+
+void FLandscapeEditDataInterface::GetHeightDataFast(const int32 X1, const int32 Y1, const int32 X2, const int32 Y2, TMap<FIntPoint, uint16>& SparseData, TMap<FIntPoint, uint16>* NormalData /*= NULL*/)
+{
+	FSparseStoreData SparseStoreData(SparseData);
+	if (NormalData)
+	{
+		FSparseStoreData SparseNormalData(*NormalData);
+		GetHeightDataTemplFast(X1, Y1, X2, Y2, SparseStoreData, &SparseNormalData);
+	}
+	else
+	{
+		GetHeightDataTemplFast(X1, Y1, X2, Y2, SparseStoreData);
+	}
+}
+
+void ULandscapeComponent::DeleteLayer( ULandscapeLayerInfoObject* LayerInfo, struct FLandscapeEditDataInterface* LandscapeEdit )
+{
+	ULandscapeComponent* Component = this; //Landscape->LandscapeComponents(ComponentIdx);
+
+	// Find the index for this layer in this component.
+	int32 DeleteLayerIdx = INDEX_NONE;
+
+	for( int32 LayerIdx=0;LayerIdx<Component->WeightmapLayerAllocations.Num();LayerIdx++ )
+	{
+		FWeightmapLayerAllocationInfo& Allocation = Component->WeightmapLayerAllocations[LayerIdx];
+		if( Allocation.LayerInfo == LayerInfo )
+		{
+			DeleteLayerIdx = LayerIdx;
+		}
+	}
+	if( DeleteLayerIdx == INDEX_NONE )
+	{
+		// Layer not used for this component.
+		//continue;
+		return;
+	}
+
+	ULandscapeInfo* Info = GetLandscapeInfo();
+	if (!Info)
+	{
+		return;
+	}
+
+	FWeightmapLayerAllocationInfo& DeleteLayerAllocation = Component->WeightmapLayerAllocations[DeleteLayerIdx];
+	int32 DeleteLayerWeightmapTextureIndex = DeleteLayerAllocation.WeightmapTextureIndex;
+
+	// See if we'll be able to remove the texture completely.
+	bool bCanRemoveLayerTexture = true;
+	for( int32 LayerIdx=0;LayerIdx<Component->WeightmapLayerAllocations.Num();LayerIdx++ )
+	{
+		FWeightmapLayerAllocationInfo& Allocation = Component->WeightmapLayerAllocations[LayerIdx];
+
+		// check if we will be able to remove the texture also
+		if( LayerIdx!=DeleteLayerIdx && Allocation.WeightmapTextureIndex == DeleteLayerWeightmapTextureIndex )
+		{
+			bCanRemoveLayerTexture = false;
+		}
+	}
+
+	// See if the deleted layer is a NoWeightBlend layer - if not, we don't have to worry about normalization
+	bool bDeleteLayerIsNoWeightBlend = (LayerInfo && LayerInfo->bNoWeightBlend);
+
+	if( !bDeleteLayerIsNoWeightBlend )
+	{
+		// Lock data for all the weightmaps
+		TArray<FLandscapeTextureDataInfo*> TexDataInfos;
+		for( int32 WeightmapIdx=0;WeightmapIdx < Component->WeightmapTextures.Num();WeightmapIdx++ )
+		{
+			if (LandscapeEdit)
+			{
+				TexDataInfos.Add(LandscapeEdit->GetTextureDataInfo(Component->WeightmapTextures[WeightmapIdx]));
+			}
+		}
+
+		// Channel remapping
+		int32 ChannelOffsets[4] = {(int32)STRUCT_OFFSET(FColor,R),(int32)STRUCT_OFFSET(FColor,G),(int32)STRUCT_OFFSET(FColor,B),(int32)STRUCT_OFFSET(FColor,A)};
+
+		TArray<bool> LayerNoWeightBlends;	// Array of NoWeightBlend flags
+		TArray<uint8*> LayerDataPtrs;		// Pointers to all layers' data 
+
+		// Get the data for each layer
+		for( int32 LayerIdx=0;LayerIdx<Component->WeightmapLayerAllocations.Num();LayerIdx++ )
+		{
+			FWeightmapLayerAllocationInfo& Allocation = Component->WeightmapLayerAllocations[LayerIdx];
+			LayerDataPtrs.Add( (uint8*)TexDataInfos[Allocation.WeightmapTextureIndex]->GetMipData(0) + ChannelOffsets[Allocation.WeightmapTextureChannel] );
+
+			// Find the layer info and record if it is a bNoWeightBlend layer.
+			LayerNoWeightBlends.Add( Allocation.LayerInfo && Allocation.LayerInfo->bNoWeightBlend );
+		}
+
+		// Find the texture data corresponding to this vertex
+		int32 SizeU = (SubsectionSizeQuads+1) * NumSubsections; //Component->WeightmapTextures(0)->SizeX;		// not exactly correct.  
+		int32 SizeV = (SubsectionSizeQuads+1) * NumSubsections; //Component->WeightmapTextures(0)->SizeY;
+		int32 WeightmapOffsetX = Component->WeightmapScaleBias.Z * (float)SizeU;
+		int32 WeightmapOffsetY = Component->WeightmapScaleBias.W * (float)SizeV;
+
+		for( int32 SubIndexY=0;SubIndexY<NumSubsections;SubIndexY++ )
+		{
+			for( int32 SubIndexX=0;SubIndexX<NumSubsections;SubIndexX++ )
+			{
+				for( int32 SubY=0;SubY<=SubsectionSizeQuads;SubY++ )
+				{
+					for( int32 SubX=0;SubX<=SubsectionSizeQuads;SubX++ )
+					{
+						int32 TexX = WeightmapOffsetX + (SubsectionSizeQuads+1) * SubIndexX + SubX;
+						int32 TexY = WeightmapOffsetY + (SubsectionSizeQuads+1) * SubIndexY + SubY;
+						int32 TexDataIndex = 4 * (TexX + TexY * SizeU);
+
+						// Calculate the sum of other layer weights
+						int32 OtherLayerWeightSum = 0;
+						for( int32 LayerIdx=0;LayerIdx<Component->WeightmapLayerAllocations.Num();LayerIdx++ )
+						{
+							if( LayerIdx != DeleteLayerIdx && LayerNoWeightBlends[LayerIdx]==false )
+							{
+								OtherLayerWeightSum += LayerDataPtrs[LayerIdx][TexDataIndex];
+							}
+						}
+
+						if( OtherLayerWeightSum == 0 )
+						{
+							// Set the first other weight-blend layer we can find to 255 to avoid a black hole
+							// This isn't ideal but it's the best option
+							// There's nothing we can easily do if this was the only weight-blend layer on this component
+							for( int32 LayerIdx=0;LayerIdx<Component->WeightmapLayerAllocations.Num();LayerIdx++ )
+							{
+								if( LayerIdx != DeleteLayerIdx && LayerNoWeightBlends[LayerIdx]==false )
+								{
+									uint8& Weight = LayerDataPtrs[LayerIdx][TexDataIndex];
+									Weight = 255;
+									break;
+								}
+							}
+						}
+						else
+						{
+							// Adjust other layer weights
+							for( int32 LayerIdx=0;LayerIdx<Component->WeightmapLayerAllocations.Num();LayerIdx++ )
+							{
+								if( LayerIdx != DeleteLayerIdx && LayerNoWeightBlends[LayerIdx]==false )
+								{
+									uint8& Weight = LayerDataPtrs[LayerIdx][TexDataIndex];
+									Weight = FMath::Clamp<int32>( FMath::RoundToInt(255.f * (float)Weight/(float)OtherLayerWeightSum), 0, 255 );
+								}
+							}
+						}
+					}
+				}
+			}
+		}
+
+		// Update all the textures and mips
+		for( int32 Idx=0;Idx<Component->WeightmapTextures.Num();Idx++)
+		{
+			if( bCanRemoveLayerTexture && Idx==DeleteLayerWeightmapTextureIndex )
+			{
+				// We're going to remove this texture anyway, so don't bother updating
+				continue;
+			}
+
+			UTexture2D* WeightmapTexture = Component->WeightmapTextures[Idx];
+			FLandscapeTextureDataInfo* WeightmapDataInfo = TexDataInfos[Idx];
+
+			int32 NumMips = WeightmapTexture->Source.GetNumMips();
+			TArray<FColor*> WeightmapTextureMipData;
+			WeightmapTextureMipData.AddUninitialized(NumMips);
+			for( int32 MipIdx=0;MipIdx<NumMips;MipIdx++ )
+			{
+				WeightmapTextureMipData[MipIdx] = (FColor*)WeightmapDataInfo->GetMipData(MipIdx);
+			}
+
+			ULandscapeComponent::UpdateWeightmapMips(Component->NumSubsections, Component->SubsectionSizeQuads, WeightmapTexture, WeightmapTextureMipData, 0, 0, MAX_int32, MAX_int32, WeightmapDataInfo);
+
+			WeightmapDataInfo->AddMipUpdateRegion(0,0,0,WeightmapTexture->Source.GetSizeX()-1,WeightmapTexture->Source.GetSizeY()-1);
+		}
+	}
+
+	ALandscapeProxy* Proxy = GetLandscapeProxy(); //CastChecked<ALandscapeProxy>(Component->GetOuter());
+	// Mark the channel as unallocated, so we can reuse it later
+	FLandscapeWeightmapUsage* Usage = Proxy->WeightmapUsageMap.Find(Component->WeightmapTextures[DeleteLayerAllocation.WeightmapTextureIndex]);
+	//check(Usage);
+	if (Usage)
+	{
+		Usage->ChannelUsage[DeleteLayerAllocation.WeightmapTextureChannel] = NULL;
+	}
+
+	// Remove the layer
+	Component->WeightmapLayerAllocations.RemoveAt(DeleteLayerIdx);
+
+	// If this layer was the last usage for this channel in this layer, we can remove it.
+	if( bCanRemoveLayerTexture )
+	{
+		Component->WeightmapTextures[DeleteLayerWeightmapTextureIndex]->SetFlags(RF_Transactional);
+		Component->WeightmapTextures[DeleteLayerWeightmapTextureIndex]->Modify();
+		Component->WeightmapTextures[DeleteLayerWeightmapTextureIndex]->MarkPackageDirty();
+		Component->WeightmapTextures[DeleteLayerWeightmapTextureIndex]->ClearFlags(RF_Standalone);
+
+		Component->WeightmapTextures.RemoveAt(DeleteLayerWeightmapTextureIndex);
+
+		// Adjust WeightmapTextureIndex index for other layers
+		for( int32 LayerIdx=0;LayerIdx<Component->WeightmapLayerAllocations.Num();LayerIdx++ )
+		{
+			FWeightmapLayerAllocationInfo& Allocation = Component->WeightmapLayerAllocations[LayerIdx];
+
+			if( Allocation.WeightmapTextureIndex > DeleteLayerWeightmapTextureIndex )
+			{
+				Allocation.WeightmapTextureIndex--;
+			}
+
+			check( Allocation.WeightmapTextureIndex < Component->WeightmapTextures.Num() );
+		}
+	}
+
+	// Update the shaders for this component
+	Component->UpdateMaterialInstances();
+}
+
+void FLandscapeEditDataInterface::DeleteLayer( ULandscapeLayerInfoObject* LayerInfo )
+{
+	if (!LandscapeInfo) return;
+	for( auto It = LandscapeInfo->XYtoComponentMap.CreateIterator(); It; ++It )
+	{
+		ULandscapeComponent* Component = It.Value();
+		Component->DeleteLayer(LayerInfo, this);
+
+		// Update dominant layer info stored in collision component
+		TArray<FColor*> CollisionWeightmapMipData;
+		for( int32 WeightmapIdx=0;WeightmapIdx < Component->WeightmapTextures.Num();WeightmapIdx++ )
+		{
+			CollisionWeightmapMipData.Add( (FColor*)GetTextureDataInfo(Component->WeightmapTextures[WeightmapIdx])->GetMipData(Component->CollisionMipLevel) );
+		}
+		Component->UpdateCollisionLayerData(CollisionWeightmapMipData);
+	}
+}
+
+void ULandscapeComponent::ReplaceLayer(ULandscapeLayerInfoObject* FromLayerInfo, ULandscapeLayerInfoObject* ToLayerInfo, struct FLandscapeEditDataInterface* LandscapeEdit)
+{
+	check(FromLayerInfo && ToLayerInfo && LandscapeEdit);
+	if (FromLayerInfo == ToLayerInfo)
+	{
+		return;
+	}
+
+	// Find the index for this layer in this component.
+	int32 FromLayerIdx = INDEX_NONE;
+	for (int32 LayerIdx = 0; LayerIdx < WeightmapLayerAllocations.Num(); LayerIdx++)
+	{
+		FWeightmapLayerAllocationInfo& Allocation = WeightmapLayerAllocations[LayerIdx];
+		if (Allocation.LayerInfo == FromLayerInfo)
+		{
+			FromLayerIdx = LayerIdx;
+		}
+	}
+	if (FromLayerIdx == INDEX_NONE)
+	{
+		// Layer not used for this component, nothing to do.
+		return;
+	}
+
+	bool bMerging = true;
+
+	// Find the index for this layer in this component.
+	int32 ToLayerIdx = INDEX_NONE;
+	for (int32 LayerIdx = 0; LayerIdx < WeightmapLayerAllocations.Num(); LayerIdx++)
+	{
+		FWeightmapLayerAllocationInfo& Allocation = WeightmapLayerAllocations[LayerIdx];
+		if (Allocation.LayerInfo == ToLayerInfo)
+		{
+			ToLayerIdx = LayerIdx;
+		}
+	}
+	if (ToLayerIdx == INDEX_NONE)
+	{
+		// Layer not used for this component, so do trivial replace.
+		WeightmapLayerAllocations[FromLayerIdx].LayerInfo = ToLayerInfo;
+		bMerging = false;
+	}
+
+	FWeightmapLayerAllocationInfo& FromLayerAllocation = WeightmapLayerAllocations[FromLayerIdx];
+
+	// See if we'll be able to remove the texture completely.
+	bool bCanRemoveLayerTexture = false;
+	if (bMerging)
+	{
+		bCanRemoveLayerTexture = true;
+		for (int32 LayerIdx = 0; LayerIdx < WeightmapLayerAllocations.Num(); LayerIdx++)
+		{
+			FWeightmapLayerAllocationInfo& Allocation = WeightmapLayerAllocations[LayerIdx];
+
+			// check if we will be able to remove the texture also
+			if (LayerIdx != FromLayerIdx && Allocation.WeightmapTextureIndex == FromLayerAllocation.WeightmapTextureIndex)
+			{
+				bCanRemoveLayerTexture = false;
+				break;
+			}
+		}
+	}
+
+	// See if the deleted layer is a NoWeightBlend layer - if not, we don't have to worry about normalization
+	const bool bFromLayerIsNoWeightBlend = (FromLayerInfo && FromLayerInfo->bNoWeightBlend);
+	const bool bToLayerIsNoWeightBlend = (ToLayerInfo && ToLayerInfo->bNoWeightBlend);
+
+	const bool bRequireNormalization = (bFromLayerIsNoWeightBlend != bToLayerIsNoWeightBlend);
+	checkf(!bRequireNormalization, TEXT("It is not yet supported to replace a layer with another whose bNoWeightBlend setting does not match"))
+
+	if (bMerging)
+	{
+		FWeightmapLayerAllocationInfo& ToLayerAllocation = WeightmapLayerAllocations[ToLayerIdx];
+
+		// Lock data for all the weightmaps
+		FLandscapeTextureDataInfo* FromTexDataInfo = LandscapeEdit->GetTextureDataInfo(WeightmapTextures[FromLayerAllocation.WeightmapTextureIndex]);
+		FLandscapeTextureDataInfo* ToTexDataInfo = LandscapeEdit->GetTextureDataInfo(WeightmapTextures[ToLayerAllocation.WeightmapTextureIndex]);
+
+		check(FromTexDataInfo->GetMipSizeX(0) == FromTexDataInfo->GetMipSizeY(0));
+		check(ToTexDataInfo->GetMipSizeX(0) == ToTexDataInfo->GetMipSizeY(0));
+		check(FromTexDataInfo->GetMipSizeX(0) == ToTexDataInfo->GetMipSizeX(0));
+		const int32 MipSize = FromTexDataInfo->GetMipSizeX(0);
+
+		// Channel remapping
+		const int32 ChannelOffsets[4] = {(int32)STRUCT_OFFSET(FColor,R), (int32)STRUCT_OFFSET(FColor,G), (int32)STRUCT_OFFSET(FColor,B), (int32)STRUCT_OFFSET(FColor,A)};
+
+		uint8* const SrcTextureData = (uint8*)FromTexDataInfo->GetMipData(0) + ChannelOffsets[FromLayerAllocation.WeightmapTextureChannel];
+		uint8* const DestTextureData = (uint8*)ToTexDataInfo->GetMipData(0) + ChannelOffsets[ToLayerAllocation.WeightmapTextureChannel];
+
+		for (int32 i = 0; i < FMath::Square(MipSize); i++)
+		{
+			DestTextureData[i*4] = FMath::Min(255, (uint16)DestTextureData[i*4] + (uint16)SrcTextureData[i*4]);
+		}
+
+		// Update all mips
+		if (!bCanRemoveLayerTexture)
+		{
+			UTexture2D* WeightmapTexture = WeightmapTextures[FromLayerAllocation.WeightmapTextureIndex];
+			FLandscapeTextureDataInfo* WeightmapDataInfo = FromTexDataInfo;
+
+			const int32 NumMips = WeightmapTexture->Source.GetNumMips();
+			TArray<FColor*> WeightmapTextureMipData;
+			WeightmapTextureMipData.AddUninitialized(NumMips);
+			for (int32 MipIdx = 0; MipIdx < NumMips; MipIdx++)
+			{
+				WeightmapTextureMipData[MipIdx] = (FColor*)WeightmapDataInfo->GetMipData(MipIdx);
+			}
+
+			ULandscapeComponent::UpdateWeightmapMips(NumSubsections, SubsectionSizeQuads, WeightmapTexture, WeightmapTextureMipData, 0, 0, MAX_int32, MAX_int32, WeightmapDataInfo);
+
+			WeightmapDataInfo->AddMipUpdateRegion(0, 0, 0, WeightmapTexture->Source.GetSizeX() - 1, WeightmapTexture->Source.GetSizeY() - 1);
+		}
+
+		if (FromTexDataInfo != ToTexDataInfo)
+		{
+			UTexture2D* WeightmapTexture = WeightmapTextures[ToLayerAllocation.WeightmapTextureIndex];
+			FLandscapeTextureDataInfo* WeightmapDataInfo = ToTexDataInfo;
+
+			const int32 NumMips = WeightmapTexture->Source.GetNumMips();
+			TArray<FColor*> WeightmapTextureMipData;
+			WeightmapTextureMipData.AddUninitialized(NumMips);
+			for (int32 MipIdx = 0; MipIdx < NumMips; MipIdx++)
+			{
+				WeightmapTextureMipData[MipIdx] = (FColor*)WeightmapDataInfo->GetMipData(MipIdx);
+			}
+
+			ULandscapeComponent::UpdateWeightmapMips(NumSubsections, SubsectionSizeQuads, WeightmapTexture, WeightmapTextureMipData, 0, 0, MAX_int32, MAX_int32, WeightmapDataInfo);
+
+			WeightmapDataInfo->AddMipUpdateRegion(0, 0, 0, WeightmapTexture->Source.GetSizeX() - 1, WeightmapTexture->Source.GetSizeY() - 1);
+		}
+	}
+
+	if (bRequireNormalization)
+	{
+		// TODO
+	}
+
+	// if merging into an existing layer, remove the layer and potentially the texture
+	if (bMerging)
+	{
+		ALandscapeProxy* Proxy = GetLandscapeProxy();
+		// Mark the channel as unallocated, so we can reuse it later
+		FLandscapeWeightmapUsage* Usage = Proxy->WeightmapUsageMap.Find(WeightmapTextures[FromLayerAllocation.WeightmapTextureIndex]);
+		//check(Usage);
+		if (Usage)
+		{
+			Usage->ChannelUsage[FromLayerAllocation.WeightmapTextureChannel] = NULL;
+		}
+
+		// If this layer was the last usage for this texture, we can remove it.
+		if (bCanRemoveLayerTexture)
+		{
+			WeightmapTextures[FromLayerAllocation.WeightmapTextureIndex]->SetFlags(RF_Transactional);
+			WeightmapTextures[FromLayerAllocation.WeightmapTextureIndex]->Modify();
+			WeightmapTextures[FromLayerAllocation.WeightmapTextureIndex]->MarkPackageDirty();
+			WeightmapTextures[FromLayerAllocation.WeightmapTextureIndex]->ClearFlags(RF_Standalone);
+
+			WeightmapTextures.RemoveAt(FromLayerAllocation.WeightmapTextureIndex);
+
+			// Adjust WeightmapTextureIndex index for other layers
+			for (int32 LayerIdx = 0; LayerIdx < WeightmapLayerAllocations.Num(); LayerIdx++)
+			{
+				if (LayerIdx == FromLayerIdx)
+				{
+					continue;
+				}
+
+				FWeightmapLayerAllocationInfo& Allocation = WeightmapLayerAllocations[LayerIdx];
+
+				if (Allocation.WeightmapTextureIndex > FromLayerAllocation.WeightmapTextureIndex)
+				{
+					Allocation.WeightmapTextureIndex--;
+				}
+
+				check(Allocation.WeightmapTextureIndex < WeightmapTextures.Num());
+			}
+		}
+
+		// Remove the layer
+		WeightmapLayerAllocations.RemoveAt(FromLayerIdx);
+
+		// Update the shaders for this component
+		UpdateMaterialInstances();
+	}
+}
+
+void FLandscapeEditDataInterface::ReplaceLayer(ULandscapeLayerInfoObject* FromLayerInfo, ULandscapeLayerInfoObject* ToLayerInfo)
+{
+	if (!LandscapeInfo) return;
+	for( auto It = LandscapeInfo->XYtoComponentMap.CreateIterator(); It; ++It )
+	{
+		ULandscapeComponent* Component = It.Value();
+		Component->ReplaceLayer(FromLayerInfo, ToLayerInfo, this);
+
+		// Update dominant layer info stored in collision component
+		TArray<FColor*> CollisionWeightmapMipData;
+		for( int32 WeightmapIdx=0;WeightmapIdx < Component->WeightmapTextures.Num();WeightmapIdx++ )
+		{
+			CollisionWeightmapMipData.Add( (FColor*)GetTextureDataInfo(Component->WeightmapTextures[WeightmapIdx])->GetMipData(Component->CollisionMipLevel) );
+		}
+		Component->UpdateCollisionLayerData(CollisionWeightmapMipData);
+	}
+}
+
+// simple classes for the template....
+namespace
+{
+	template<typename TDataType>
+	struct TArrayStoreData
+	{
+		int32 X1;
+		int32 Y1;
+		TDataType* Data;
+		int32 Stride;
+		int32 ArraySize;
+
+		TArrayStoreData(int32 InX1, int32 InY1, TDataType* InData, int32 InStride)
+			:	X1(InX1)
+			,	Y1(InY1)
+			,	Data(InData)
+			,	Stride(InStride)
+			,	ArraySize(1)
+		{}
+
+		inline void Store(int32 LandscapeX, int32 LandscapeY, uint8 Weight) {}
+		inline void Store(int32 LandscapeX, int32 LandscapeY, uint8 Weight, int32 LayerIdx) {}
+		inline void Store(int32 LandscapeX, int32 LandscapeY, FVector2D Offset) {}
+		inline TDataType Load(int32 LandscapeX, int32 LandscapeY) { return 0; }
+		inline void PreInit(int32 InArraySize) { ArraySize = InArraySize; }
+	};
+
+	template<> void TArrayStoreData<uint8>::Store(int32 LandscapeX, int32 LandscapeY, uint8 Weight)
+	{
+		Data[ (LandscapeY-Y1) * Stride + (LandscapeX-X1) ] = Weight;
+	}
+
+	template<> uint8 TArrayStoreData<uint8>::Load(int32 LandscapeX, int32 LandscapeY)
+	{
+		return Data[ (LandscapeY-Y1) * Stride + (LandscapeX-X1) ];
+	}
+
+	template<> FVector2D TArrayStoreData<FVector2D>::Load(int32 LandscapeX, int32 LandscapeY)
+	{
+		return Data[(LandscapeY - Y1) * Stride + (LandscapeX - X1)];
+	}
+
+	template<> FVector TArrayStoreData<FVector>::Load(int32 LandscapeX, int32 LandscapeY)
+	{
+		return Data[(LandscapeY - Y1) * Stride + (LandscapeX - X1)];
+	}
+
+	template<> void TArrayStoreData<FVector2D>::Store(int32 LandscapeX, int32 LandscapeY, FVector2D Offset)
+	{
+		Data[ (LandscapeY-Y1) * Stride + (LandscapeX-X1) ] = Offset;
+	}
+
+	template<> void TArrayStoreData<FVector>::Store(int32 LandscapeX, int32 LandscapeY, FVector2D Offset)
+	{
+		Data[ (LandscapeY-Y1) * Stride + (LandscapeX-X1) ] = FVector(Offset.X, Offset.Y, 0.f);
+	}
+
+	// Data items should be initialized with ArraySize
+	template<> void TArrayStoreData<TArray<uint8>>::Store(int32 LandscapeX, int32 LandscapeY, uint8 Weight, int32 LayerIdx)
+	{
+		TArray<uint8>& Value = Data[ ((LandscapeY-Y1) * Stride + (LandscapeX-X1)) ];
+		if (Value.Num() != ArraySize)
+		{
+			Value.Empty(ArraySize);
+			Value.AddZeroed(ArraySize);
+		}
+		Value[LayerIdx] = Weight;
+	}
+
+	template<typename TDataType>
+	struct TSparseStoreData
+	{
+		TMap<FIntPoint, TDataType>& SparseData;
+		int32 ArraySize;
+
+		TSparseStoreData(TMap<FIntPoint, TDataType>& InSparseData)
+			:	SparseData(InSparseData)
+			,	ArraySize(1)
+		{}
+
+		inline void Store(int32 LandscapeX, int32 LandscapeY, uint8 Weight) {}
+		inline void Store(int32 LandscapeX, int32 LandscapeY, uint8 Weight, int32 LayerIdx) {}
+		inline void Store(int32 LandscapeX, int32 LandscapeY, FVector2D Offset) {}
+		inline TDataType Load(int32 LandscapeX, int32 LandscapeY) { return 0; }
+		inline void PreInit(int32 InArraySize) { ArraySize = InArraySize; }
+	};
+
+	template<> void TSparseStoreData<uint8>::Store(int32 LandscapeX, int32 LandscapeY, uint8 Weight)
+	{
+		SparseData.Add(ALandscape::MakeKey(LandscapeX,LandscapeY), Weight);
+	}
+
+	template<> uint8 TSparseStoreData<uint8>::Load(int32 LandscapeX, int32 LandscapeY)
+	{
+		return SparseData.FindRef(ALandscape::MakeKey(LandscapeX,LandscapeY));
+	}
+
+	template<> FVector2D TSparseStoreData<FVector2D>::Load(int32 LandscapeX, int32 LandscapeY)
+	{
+		return SparseData.FindRef(ALandscape::MakeKey(LandscapeX, LandscapeY));
+	}
+
+	template<> FVector TSparseStoreData<FVector>::Load(int32 LandscapeX, int32 LandscapeY)
+	{
+		return SparseData.FindRef(ALandscape::MakeKey(LandscapeX, LandscapeY));
+	}
+
+	template<> void TSparseStoreData<TArray<uint8>>::Store(int32 LandscapeX, int32 LandscapeY, uint8 Weight, int32 LayerIdx)
+	{
+		TArray<uint8>* Value = SparseData.Find(ALandscape::MakeKey(LandscapeX,LandscapeY));
+		if (Value)
+		{
+			(*Value)[LayerIdx] = Weight;
+		}
+		else
+		{
+			TArray<uint8> Values;
+			Values.Empty(ArraySize);
+			Values.AddZeroed(ArraySize);
+			Values[LayerIdx] = Weight;
+			SparseData.Add(ALandscape::MakeKey(LandscapeX, LandscapeY), Values);
+		}
+	}
+
+	template<> void TSparseStoreData<FVector2D>::Store(int32 LandscapeX, int32 LandscapeY, FVector2D Offset)
+	{
+		SparseData.Add(ALandscape::MakeKey(LandscapeX,LandscapeY), Offset);
+	}
+
+	template<> void TSparseStoreData<FVector>::Store(int32 LandscapeX, int32 LandscapeY, FVector2D Offset)
+	{
+		// Preserve old Z value
+		FVector* PrevValue = SparseData.Find(ALandscape::MakeKey(LandscapeX,LandscapeY));
+		if (PrevValue != NULL)
+		{
+			PrevValue->X = Offset.X;
+			PrevValue->Y = Offset.Y;
+		}
+		else
+		{
+			SparseData.Add(ALandscape::MakeKey(LandscapeX,LandscapeY), FVector(Offset.X, Offset.Y, 0.f));
+		}
+	}
+};
+
+bool DeleteLayerIfAllZero(ULandscapeComponent* const Component, const uint8* const TexDataPtr, int32 TexSize, int32 LayerIdx)
+{
+	// Check the data for the entire component and to see if it's all zero
+	for (int32 TexY = 0; TexY < TexSize; TexY++)
+	{
+		for (int32 TexX = 0; TexX < TexSize; TexX++)
+		{
+			const int32 TexDataIndex = 4 * (TexX + TexY * TexSize);
+
+			// Stop the first time we see any non-zero data
+			const uint8& Weight = TexDataPtr[TexDataIndex];
+			if (Weight != 0)
+			{
+				return false;
+			}
+		}
+	}
+
+	Component->Modify();
+
+	// Mark the channel as unallocated, so we can reuse it later
+	const int32 DeleteLayerWeightmapTextureIndex = Component->WeightmapLayerAllocations[LayerIdx].WeightmapTextureIndex;
+	FLandscapeWeightmapUsage& Usage = Component->GetLandscapeProxy()->WeightmapUsageMap.FindChecked(Component->WeightmapTextures[DeleteLayerWeightmapTextureIndex]);
+	Usage.ChannelUsage[Component->WeightmapLayerAllocations[LayerIdx].WeightmapTextureChannel] = NULL;
+
+	// Remove the layer as it's totally painted away.
+	Component->WeightmapLayerAllocations.RemoveAt(LayerIdx);
+
+	// Check if the weightmap texture used by the layer we just removed is used by any other layer, and if so, remove the texture too
+	bool bCanRemoveLayerTexture = !Component->WeightmapLayerAllocations.ContainsByPredicate([DeleteLayerWeightmapTextureIndex](const FWeightmapLayerAllocationInfo& Allocation){ return Allocation.WeightmapTextureIndex == DeleteLayerWeightmapTextureIndex; });
+	if (bCanRemoveLayerTexture)
+	{
+		Component->WeightmapTextures[DeleteLayerWeightmapTextureIndex]->MarkPackageDirty();
+		Component->WeightmapTextures[DeleteLayerWeightmapTextureIndex]->ClearFlags(RF_Standalone);
+		Component->WeightmapTextures.RemoveAt(DeleteLayerWeightmapTextureIndex);
+
+		// Adjust WeightmapTextureChannel index for other layers
+		for (auto It = Component->WeightmapLayerAllocations.CreateIterator(); It; ++It)
+		{
+			FWeightmapLayerAllocationInfo& Allocation = *It;
+			if (Allocation.WeightmapTextureIndex > DeleteLayerWeightmapTextureIndex)
+			{
+				Allocation.WeightmapTextureIndex--;
+			}
+		}
+	}
+
+	return true;
+}
+
+void FLandscapeEditDataInterface::SetAlphaData(ULandscapeLayerInfoObject* const LayerInfo, const int32 X1, const int32 Y1, const int32 X2, const int32 Y2, const uint8* Data, int32 Stride, ELandscapeLayerPaintingRestriction::Type PaintingRestriction /*= None*/, bool bWeightAdjust /*= true*/, bool bTotalWeightAdjust /*= false*/)
+{
+	check(LayerInfo != NULL);
+
+	if (Stride == 0)
+	{
+		Stride = (1+X2-X1);
+	}
+
+	// Channel remapping
+	size_t ChannelOffsets[4] = {STRUCT_OFFSET(FColor,R), STRUCT_OFFSET(FColor,G), STRUCT_OFFSET(FColor,B), STRUCT_OFFSET(FColor,A)};
+
+	check(ComponentSizeQuads > 0);
+	// Find component range for this block of data
+	int32 ComponentIndexX1 = (X1-1 >= 0) ? (X1-1) / ComponentSizeQuads : (X1) / ComponentSizeQuads - 1;	// -1 because we need to pick up vertices shared between components
+	int32 ComponentIndexY1 = (Y1-1 >= 0) ? (Y1-1) / ComponentSizeQuads : (Y1) / ComponentSizeQuads - 1;
+	int32 ComponentIndexX2 = (X2 >= 0) ? X2 / ComponentSizeQuads : (X2+1) / ComponentSizeQuads - 1;
+	int32 ComponentIndexY2 = (Y2 >= 0) ? Y2 / ComponentSizeQuads : (Y2+1) / ComponentSizeQuads - 1;
+
+	TArray<FLandscapeTextureDataInfo*, TInlineAllocator<2>> TexDataInfos;
+	TArray<uint8*, TInlineAllocator<8>> LayerDataPtrs;		// Pointers to all layers' data 
+	TArray<bool, TInlineAllocator<8>> LayerNoWeightBlends;	// NoWeightBlend flags
+	TArray<bool, TInlineAllocator<8>> LayerEditDataAllZero; // Whether the data we are editing for this layer is all zero 
+	TArray<FColor*> CollisionWeightmapMipData;
+	TArray<FColor*> WeightmapTextureMipData;
+
+	for (int32 ComponentIndexY = ComponentIndexY1; ComponentIndexY <= ComponentIndexY2; ComponentIndexY++)
+	{
+		for (int32 ComponentIndexX = ComponentIndexX1; ComponentIndexX <= ComponentIndexX2; ComponentIndexX++)
+		{
+			FIntPoint ComponentKey(ComponentIndexX,ComponentIndexY);
+			ULandscapeComponent* Component = LandscapeInfo->XYtoComponentMap.FindRef(ComponentKey);
+
+			// if NULL, there is no component at this location
+			if (Component == NULL)
+			{
+				continue;
+			}
+
+			int32 UpdateLayerIdx = Component->WeightmapLayerAllocations.IndexOfByPredicate([LayerInfo](const FWeightmapLayerAllocationInfo& Allocation){ return Allocation.LayerInfo == LayerInfo; });
+
+			// Need allocation for weightmap
+			if (UpdateLayerIdx == INDEX_NONE)
+			{
+				const int32 LayerLimit = Component->GetLandscapeProxy()->MaxPaintedLayersPerComponent;
+
+				// if we can't allocate a layer, then there is nothing to paint
+				if (PaintingRestriction == ELandscapeLayerPaintingRestriction::ExistingOnly ||
+					(PaintingRestriction == ELandscapeLayerPaintingRestriction::UseMaxLayers &&
+					 LayerLimit > 0 && Component->WeightmapLayerAllocations.Num() >= LayerLimit))
+				{
+					continue;
+				}
+
+				Component->Modify();
+
+				UpdateLayerIdx = Component->WeightmapLayerAllocations.Num();
+				new (Component->WeightmapLayerAllocations) FWeightmapLayerAllocationInfo(LayerInfo);
+				Component->ReallocateWeightmaps(this);
+
+				Component->UpdateMaterialInstances();
+				if( Component->EditToolRenderData )
+				{
+					Component->EditToolRenderData->UpdateDebugColorMaterial();
+				}
+			}
+
+			// Lock data for all the weightmaps
+			TexDataInfos.Reset();
+			TexDataInfos.AddUninitialized(Component->WeightmapTextures.Num());
+
+			for (int32 WeightmapIdx = 0; WeightmapIdx < Component->WeightmapTextures.Num(); ++WeightmapIdx)
+			{
+				TexDataInfos[WeightmapIdx] = GetTextureDataInfo(Component->WeightmapTextures[WeightmapIdx]);
+			};
+
+			LayerDataPtrs.Reset();        // Pointers to all layers' data
+			LayerDataPtrs.AddUninitialized(Component->WeightmapLayerAllocations.Num());
+			LayerNoWeightBlends.Reset();  // NoWeightBlend flags
+			LayerNoWeightBlends.AddUninitialized(Component->WeightmapLayerAllocations.Num());
+			LayerEditDataAllZero.Reset(); // Whether the data we are editing for this layer is all zero 
+			LayerEditDataAllZero.AddUninitialized(Component->WeightmapLayerAllocations.Num());
+
+			for (int32 LayerIdx = 0; LayerIdx < Component->WeightmapLayerAllocations.Num(); LayerIdx++)
+			{
+				FWeightmapLayerAllocationInfo& Allocation = Component->WeightmapLayerAllocations[LayerIdx];
+				LayerDataPtrs[LayerIdx] = (uint8*)TexDataInfos[Allocation.WeightmapTextureIndex]->GetMipData(0) + ChannelOffsets[Allocation.WeightmapTextureChannel];
+				LayerNoWeightBlends[LayerIdx] = Allocation.LayerInfo->bNoWeightBlend;
+				LayerEditDataAllZero[LayerIdx] = true;
+			}
+
+			// Find the texture data corresponding to this vertex
+			const int32 TexSize = (Component->SubsectionSizeQuads+1) * Component->NumSubsections; 
+
+			// Find coordinates of box that lies inside component
+			const int32 ComponentX1 = FMath::Clamp<int32>(X1-ComponentIndexX*ComponentSizeQuads, 0, ComponentSizeQuads);
+			const int32 ComponentY1 = FMath::Clamp<int32>(Y1-ComponentIndexY*ComponentSizeQuads, 0, ComponentSizeQuads);
+			const int32 ComponentX2 = FMath::Clamp<int32>(X2-ComponentIndexX*ComponentSizeQuads, 0, ComponentSizeQuads);
+			const int32 ComponentY2 = FMath::Clamp<int32>(Y2-ComponentIndexY*ComponentSizeQuads, 0, ComponentSizeQuads);
+
+			// Find subsection range for this box
+			const int32 SubIndexX1 = FMath::Clamp<int32>((ComponentX1-1) / SubsectionSizeQuads,0,ComponentNumSubsections-1);	// -1 because we need to pick up vertices shared between subsections
+			const int32 SubIndexY1 = FMath::Clamp<int32>((ComponentY1-1) / SubsectionSizeQuads,0,ComponentNumSubsections-1);
+			const int32 SubIndexX2 = FMath::Clamp<int32>(ComponentX2 / SubsectionSizeQuads,0,ComponentNumSubsections-1);
+			const int32 SubIndexY2 = FMath::Clamp<int32>(ComponentY2 / SubsectionSizeQuads,0,ComponentNumSubsections-1);
+
+			for (int32 SubIndexY = SubIndexY1; SubIndexY <= SubIndexY2; SubIndexY++)
+			{
+				for (int32 SubIndexX = SubIndexX1; SubIndexX <= SubIndexX2; SubIndexX++)
+				{
+					// Find coordinates of box that lies inside subsection
+					const int32 SubX1 = FMath::Clamp<int32>(ComponentX1-SubsectionSizeQuads*SubIndexX, 0, SubsectionSizeQuads);
+					const int32 SubY1 = FMath::Clamp<int32>(ComponentY1-SubsectionSizeQuads*SubIndexY, 0, SubsectionSizeQuads);
+					const int32 SubX2 = FMath::Clamp<int32>(ComponentX2-SubsectionSizeQuads*SubIndexX, 0, SubsectionSizeQuads);
+					const int32 SubY2 = FMath::Clamp<int32>(ComponentY2-SubsectionSizeQuads*SubIndexY, 0, SubsectionSizeQuads);
+
+					// Update texture data for the box that lies inside subsection
+					for (int32 SubY = SubY1; SubY <= SubY2; SubY++)
+					{
+						for (int32 SubX = SubX1; SubX <= SubX2; SubX++)
+						{
+							const int32 LandscapeX = SubIndexX*SubsectionSizeQuads + ComponentIndexX*ComponentSizeQuads + SubX;
+							const int32 LandscapeY = SubIndexY*SubsectionSizeQuads + ComponentIndexY*ComponentSizeQuads + SubY;
+							checkSlow( LandscapeX >= X1 && LandscapeX <= X2 );
+							checkSlow( LandscapeY >= Y1 && LandscapeY <= Y2 );
+
+							// Find the input data corresponding to this vertex
+							const int32 DataIndex = (LandscapeX-X1) + Stride * (LandscapeY-Y1);
+							uint8 NewWeight = Data[DataIndex];
+
+							// Adjust all layer weights
+							const int32 TexX = (SubsectionSizeQuads+1) * SubIndexX + SubX;
+							const int32 TexY = (SubsectionSizeQuads+1) * SubIndexY + SubY;
+
+							const int32 TexDataIndex = 4 * (TexX + TexY * TexSize);
+
+							int32 OtherLayerWeightSum = 0;
+
+							if (bWeightAdjust)
+							{
+								if (bTotalWeightAdjust)
+								{
+									int32 MaxLayerIdx = -1;
+									int32 MaxWeight = INT_MIN;
+
+									// Adjust other layers' weights accordingly
+									for (int32 LayerIdx = 0; LayerIdx < Component->WeightmapLayerAllocations.Num(); LayerIdx++)
+									{
+										uint8& ExistingWeight = LayerDataPtrs[LayerIdx][TexDataIndex];
+
+										if (LayerIdx == UpdateLayerIdx)
+										{
+											ExistingWeight = NewWeight;
+										}
+										// Exclude bNoWeightBlend layers
+										if (LayerNoWeightBlends[LayerIdx] == false)
+										{
+											OtherLayerWeightSum += ExistingWeight;
+											if (MaxWeight < ExistingWeight)
+											{
+												MaxWeight = ExistingWeight;
+												MaxLayerIdx = LayerIdx;
+											}
+										}
+									}
+
+									if (OtherLayerWeightSum != 255)
+									{
+										const float Factor = 255.f / OtherLayerWeightSum;
+										OtherLayerWeightSum = 0;
+
+										// Normalize
+										for (int32 LayerIdx = 0; LayerIdx<Component->WeightmapLayerAllocations.Num(); LayerIdx++)
+										{
+											uint8& ExistingWeight = LayerDataPtrs[LayerIdx][TexDataIndex];
+
+											if (LayerNoWeightBlends[LayerIdx] == false)
+											{
+												// normalization...
+												ExistingWeight = (uint8)(Factor * ExistingWeight);
+												OtherLayerWeightSum += ExistingWeight;
+
+												if (ExistingWeight != 0)
+												{
+													LayerEditDataAllZero[LayerIdx] = false;
+												}
+											}
+										}
+
+										if ((255 - OtherLayerWeightSum) && MaxLayerIdx >= 0)
+										{
+											LayerDataPtrs[MaxLayerIdx][TexDataIndex] += 255 - OtherLayerWeightSum;
+										}
+									}
+								}
+								else
+								{
+									// Adjust other layers' weights accordingly
+									for (int32 LayerIdx = 0; LayerIdx < Component->WeightmapLayerAllocations.Num(); LayerIdx++)
+									{
+										const uint8 ExistingWeight = LayerDataPtrs[LayerIdx][TexDataIndex];
+										// Exclude bNoWeightBlend layers
+										if (LayerIdx != UpdateLayerIdx && LayerNoWeightBlends[LayerIdx] == false)
+										{
+											OtherLayerWeightSum += ExistingWeight;
+										}
+									}
+
+									if (OtherLayerWeightSum == 0)
+									{
+										NewWeight = 255;
+										OtherLayerWeightSum = 1;
+									}
+
+									for (int32 LayerIdx = 0; LayerIdx < Component->WeightmapLayerAllocations.Num(); LayerIdx++)
+									{
+										uint8& Weight = LayerDataPtrs[LayerIdx][TexDataIndex];
+
+										if (LayerIdx == UpdateLayerIdx)
+										{
+											Weight = NewWeight;
+										}
+										else
+										{
+											// Exclude bNoWeightBlend layers
+											if (LayerNoWeightBlends[LayerIdx] == false)
+											{
+												Weight = FMath::Clamp<uint8>( FMath::RoundToInt((float)(255 - NewWeight) * (float)Weight/(float)OtherLayerWeightSum), 0, 255 );
+											}
+										}
+
+										if (Weight != 0)
+										{
+											LayerEditDataAllZero[LayerIdx] = false;
+										}
+									}
+								}
+							}
+							else
+							{
+								// Weight value set without adjusting other layers' weights
+								uint8& Weight = LayerDataPtrs[UpdateLayerIdx][TexDataIndex];
+								Weight = NewWeight;
+								if (Weight != 0)
+								{
+									LayerEditDataAllZero[UpdateLayerIdx] = false;
+								}
+							}
+						}
+					}
+
+					// Record the areas of the texture we need to re-upload
+					const int32 TexX1 = (SubsectionSizeQuads+1) * SubIndexX + SubX1;
+					const int32 TexY1 = (SubsectionSizeQuads+1) * SubIndexY + SubY1;
+					const int32 TexX2 = (SubsectionSizeQuads+1) * SubIndexX + SubX2;
+					const int32 TexY2 = (SubsectionSizeQuads+1) * SubIndexY + SubY2;
+					for (int32 WeightmapIdx = 0; WeightmapIdx < Component->WeightmapTextures.Num(); WeightmapIdx++)
+					{
+						if (TexDataInfos[WeightmapIdx] != NULL)
+						{
+							TexDataInfos[WeightmapIdx]->AddMipUpdateRegion(0,TexX1,TexY1,TexX2,TexY2);
+						}
+					}
+				}
+			}
+
+			// Update mipmaps
+			CollisionWeightmapMipData.Reset();
+			CollisionWeightmapMipData.AddUninitialized(Component->WeightmapTextures.Num());
+			for (int32 WeightmapIdx = 0; WeightmapIdx < Component->WeightmapTextures.Num(); WeightmapIdx++)
+			{
+				UTexture2D* const WeightmapTexture = Component->WeightmapTextures[WeightmapIdx];
+
+				const int32 NumMips = WeightmapTexture->Source.GetNumMips();
+				WeightmapTextureMipData.Reset();
+				WeightmapTextureMipData.AddUninitialized(NumMips);
+				for (int32 MipIdx = 0; MipIdx < NumMips; MipIdx++)
+				{
+					FColor* const MipData = (FColor*)TexDataInfos[WeightmapIdx]->GetMipData(MipIdx);
+					WeightmapTextureMipData[MipIdx] = MipData;
+				}
+				CollisionWeightmapMipData[WeightmapIdx] = WeightmapTextureMipData[Component->CollisionMipLevel];
+
+				ULandscapeComponent::UpdateWeightmapMips(ComponentNumSubsections, SubsectionSizeQuads, WeightmapTexture, WeightmapTextureMipData, ComponentX1, ComponentY1, ComponentX2, ComponentY2, TexDataInfos[WeightmapIdx]);
+				WeightmapTextureMipData.Reset();
+			}
+
+			// Update dominant layer info stored in collision component
+			Component->UpdateCollisionLayerData(CollisionWeightmapMipData, ComponentX1, ComponentY1, ComponentX2, ComponentY2);
+			CollisionWeightmapMipData.Reset();
+
+			// Check if we need to remove weightmap allocations for layers that were completely painted away
+			bool bRemovedLayer = false;
+			for (int32 LayerIdx = 0; LayerIdx < Component->WeightmapLayerAllocations.Num(); LayerIdx++)
+			{
+				if (LayerEditDataAllZero[LayerIdx])
+				{
+					bool bLayerDeleted = DeleteLayerIfAllZero(Component, LayerDataPtrs[LayerIdx], TexSize, LayerIdx);
+
+					if (bLayerDeleted)
+					{
+						LayerEditDataAllZero.RemoveAt(LayerIdx);
+						LayerDataPtrs.RemoveAt(LayerIdx);
+						LayerIdx--;
+
+						bRemovedLayer = true;
+					}
+				}
+			}
+
+			if (bRemovedLayer)
+			{
+				Component->UpdateMaterialInstances();
+
+				if (Component->EditToolRenderData)
+				{
+					Component->EditToolRenderData->UpdateDebugColorMaterial();
+				}
+			}
+		}
+	}
+}
+
+void FLandscapeEditDataInterface::SetAlphaData(const TSet<ULandscapeLayerInfoObject*>& DirtyLayerInfos, const int32 X1, const int32 Y1, const int32 X2, const int32 Y2, const uint8* Data, int32 Stride, ELandscapeLayerPaintingRestriction::Type PaintingRestriction /*= None*/)
+{
+	if (DirtyLayerInfos.Num() == 0)
+	{
+		return;
+	}
+
+	if (Stride == 0)
+	{
+		Stride = (1+X2-X1) * LandscapeInfo->Layers.Num();
+	}
+
+	// Channel remapping
+	size_t ChannelOffsets[4] = {STRUCT_OFFSET(FColor,R), STRUCT_OFFSET(FColor,G), STRUCT_OFFSET(FColor,B), STRUCT_OFFSET(FColor,A)};
+
+	check(ComponentSizeQuads > 0);
+	// Find component range for this block of data
+	int32 ComponentIndexX1 = (X1-1 >= 0) ? (X1-1) / ComponentSizeQuads : (X1) / ComponentSizeQuads - 1;	// -1 because we need to pick up vertices shared between components
+	int32 ComponentIndexY1 = (Y1-1 >= 0) ? (Y1-1) / ComponentSizeQuads : (Y1) / ComponentSizeQuads - 1;
+	int32 ComponentIndexX2 = (X2 >= 0) ? X2 / ComponentSizeQuads : (X2+1) / ComponentSizeQuads - 1;
+	int32 ComponentIndexY2 = (Y2 >= 0) ? Y2 / ComponentSizeQuads : (Y2+1) / ComponentSizeQuads - 1;
+
+	TArray<ULandscapeLayerInfoObject*, TInlineAllocator<8>> NeedAllocationInfos;
+	TArray<FLandscapeTextureDataInfo*, TInlineAllocator<2>> TexDataInfos;
+
+	struct FLayerDataInfo
+	{
+		const uint8* InDataPtr;
+		uint8* TexDataPtr;
+	};
+
+	TArray<FLayerDataInfo, TInlineAllocator<8>> LayerDataInfos;		// Pointers to all layers' data 
+	TArray<bool, TInlineAllocator<8>> LayerEditDataAllZero; // Whether the data we are editing for this layer is all zero 
+	TArray<FColor*> CollisionWeightmapMipData;
+	TArray<FColor*> WeightmapTextureMipData;
+
+	for (int32 ComponentIndexY = ComponentIndexY1; ComponentIndexY <= ComponentIndexY2; ComponentIndexY++)
+	{
+		for (int32 ComponentIndexX = ComponentIndexX1; ComponentIndexX <= ComponentIndexX2; ComponentIndexX++)
+		{
+			FIntPoint ComponentKey(ComponentIndexX,ComponentIndexY);
+			ULandscapeComponent* Component = LandscapeInfo->XYtoComponentMap.FindRef(ComponentKey);
+
+			// if NULL, there is no component at this location
+			if (Component == NULL)
+			{
+				continue;
+			}
+
+			const int32 LayerLimit = Component->GetLandscapeProxy()->MaxPaintedLayersPerComponent;
+
+			NeedAllocationInfos.Reset();
+
+			for (auto It = DirtyLayerInfos.CreateConstIterator(); It; ++It)
+			{
+				ULandscapeLayerInfoObject* const LayerInfo = *It;
+				const bool bFound = Component->WeightmapLayerAllocations.ContainsByPredicate([LayerInfo](const FWeightmapLayerAllocationInfo& Allocation){ return Allocation.LayerInfo == LayerInfo; });
+				if (bFound)
+				{
+					NeedAllocationInfos.Add(LayerInfo);
+				}
+			}
+
+			// Need allocation for weightmaps
+			if (NeedAllocationInfos.Num() > 0)
+			{
+				if (NeedAllocationInfos.Num() == DirtyLayerInfos.Num())
+				{
+					if (PaintingRestriction == ELandscapeLayerPaintingRestriction::ExistingOnly ||
+						(PaintingRestriction == ELandscapeLayerPaintingRestriction::UseMaxLayers &&
+						 Component->WeightmapLayerAllocations.Num() >= LayerLimit))
+					{
+						// nothing to paint to this component due to layer limit
+						continue;
+					}
+				}
+				if (PaintingRestriction != ELandscapeLayerPaintingRestriction::ExistingOnly)
+				{
+					Component->Modify();
+					for (int32 i = 0; i < NeedAllocationInfos.Num(); ++i)
+					{
+						if (PaintingRestriction == ELandscapeLayerPaintingRestriction::UseMaxLayers &&
+							LayerLimit > 0 && Component->WeightmapLayerAllocations.Num() >= LayerLimit)
+						{
+							break;
+						}
+						new (Component->WeightmapLayerAllocations) FWeightmapLayerAllocationInfo(NeedAllocationInfos[i]);
+					}
+					Component->ReallocateWeightmaps(this);
+					Component->UpdateMaterialInstances();
+					if (Component->EditToolRenderData)
+					{
+						Component->EditToolRenderData->UpdateDebugColorMaterial();
+					}
+				}
+			}
+
+			// Lock data for all the weightmaps
+			TexDataInfos.Reset();
+			TexDataInfos.AddUninitialized(Component->WeightmapTextures.Num());
+
+			for (int32 WeightmapIdx = 0; WeightmapIdx < Component->WeightmapTextures.Num(); ++WeightmapIdx)
+			{
+				TexDataInfos[WeightmapIdx] = GetTextureDataInfo(Component->WeightmapTextures[WeightmapIdx]);
+			};
+
+			LayerDataInfos.Reset();        // Pointers to all layers' data
+			LayerDataInfos.AddUninitialized(Component->WeightmapLayerAllocations.Num());
+			LayerEditDataAllZero.Reset(); // Whether the data we are editing for this layer is all zero 
+			LayerEditDataAllZero.AddUninitialized(Component->WeightmapLayerAllocations.Num());
+
+			for (int32 LayerIdx = 0; LayerIdx < Component->WeightmapLayerAllocations.Num(); LayerIdx++)
+			{
+				FWeightmapLayerAllocationInfo& Allocation = Component->WeightmapLayerAllocations[LayerIdx];
+				const int32 LayerDataIdx = LandscapeInfo->GetLayerInfoIndex(Component->WeightmapLayerAllocations[LayerIdx].LayerInfo);
+				check(LayerDataIdx != INDEX_NONE);
+				LayerDataInfos[LayerIdx].InDataPtr = Data + LayerDataIdx;
+				LayerDataInfos[LayerIdx].TexDataPtr = (uint8*)TexDataInfos[Allocation.WeightmapTextureIndex]->GetMipData(0) + ChannelOffsets[Allocation.WeightmapTextureChannel];
+				LayerEditDataAllZero[LayerIdx] = true;
+			}
+
+			// Find the texture data corresponding to this vertex
+			const int32 TexSize = (Component->SubsectionSizeQuads+1) * Component->NumSubsections; 
+
+			// Find coordinates of box that lies inside component
+			const int32 ComponentX1 = FMath::Clamp<int32>(X1-ComponentIndexX*ComponentSizeQuads, 0, ComponentSizeQuads);
+			const int32 ComponentY1 = FMath::Clamp<int32>(Y1-ComponentIndexY*ComponentSizeQuads, 0, ComponentSizeQuads);
+			const int32 ComponentX2 = FMath::Clamp<int32>(X2-ComponentIndexX*ComponentSizeQuads, 0, ComponentSizeQuads);
+			const int32 ComponentY2 = FMath::Clamp<int32>(Y2-ComponentIndexY*ComponentSizeQuads, 0, ComponentSizeQuads);
+
+			// Find subsection range for this box
+			const int32 SubIndexX1 = FMath::Clamp<int32>((ComponentX1-1) / SubsectionSizeQuads,0,ComponentNumSubsections-1);	// -1 because we need to pick up vertices shared between subsections
+			const int32 SubIndexY1 = FMath::Clamp<int32>((ComponentY1-1) / SubsectionSizeQuads,0,ComponentNumSubsections-1);
+			const int32 SubIndexX2 = FMath::Clamp<int32>(ComponentX2 / SubsectionSizeQuads,0,ComponentNumSubsections-1);
+			const int32 SubIndexY2 = FMath::Clamp<int32>(ComponentY2 / SubsectionSizeQuads,0,ComponentNumSubsections-1);
+
+			for (int32 SubIndexY = SubIndexY1; SubIndexY <= SubIndexY2; SubIndexY++)
+			{
+				for (int32 SubIndexX = SubIndexX1; SubIndexX <= SubIndexX2; SubIndexX++)
+				{
+					// Find coordinates of box that lies inside subsection
+					const int32 SubX1 = FMath::Clamp<int32>(ComponentX1-SubsectionSizeQuads*SubIndexX, 0, SubsectionSizeQuads);
+					const int32 SubY1 = FMath::Clamp<int32>(ComponentY1-SubsectionSizeQuads*SubIndexY, 0, SubsectionSizeQuads);
+					const int32 SubX2 = FMath::Clamp<int32>(ComponentX2-SubsectionSizeQuads*SubIndexX, 0, SubsectionSizeQuads);
+					const int32 SubY2 = FMath::Clamp<int32>(ComponentY2-SubsectionSizeQuads*SubIndexY, 0, SubsectionSizeQuads);
+
+					// Update texture data for the box that lies inside subsection
+					for (int32 SubY = SubY1; SubY <= SubY2; SubY++)
+					{
+						for (int32 SubX = SubX1; SubX <= SubX2; SubX++)
+						{
+							const int32 LandscapeX = SubIndexX*SubsectionSizeQuads + ComponentIndexX*ComponentSizeQuads + SubX;
+							const int32 LandscapeY = SubIndexY*SubsectionSizeQuads + ComponentIndexY*ComponentSizeQuads + SubY;
+							checkSlow( LandscapeX >= X1 && LandscapeX <= X2 );
+							checkSlow( LandscapeY >= Y1 && LandscapeY <= Y2 );
+
+							// Find the input data corresponding to this vertex
+							const int32 DataIndex = (LandscapeY-Y1) * Stride + (LandscapeX-X1) * LandscapeInfo->Layers.Num();
+
+							// Adjust all layer weights
+							const int32 TexX = (SubsectionSizeQuads+1) * SubIndexX + SubX;
+							const int32 TexY = (SubsectionSizeQuads+1) * SubIndexY + SubY;
+
+							const int32 TexDataIndex = 4 * (TexX + TexY * TexSize);
+
+							int32 OtherLayerWeightSum = 0;
+
+							// Apply weights to all layers simultaneously.
+							for (int32 LayerIdx = 0; LayerIdx < Component->WeightmapLayerAllocations.Num(); LayerIdx++)
+							{
+								// this is equivalent to saying if (DirtyLayerInfos.Contains(Allocation.LayerInfo))
+								// which is what we really mean here, but this is quicker
+								// and I've lost count of the depth we've nested for loops at this point
+								if (LayerDataInfos[LayerIdx].TexDataPtr != NULL)
+								{
+									uint8& Weight = LayerDataInfos[LayerIdx].TexDataPtr[TexDataIndex];
+
+									Weight = Data[DataIndex]; // Only for whole weight
+									if (Weight != 0)
+									{
+										LayerEditDataAllZero[LayerIdx] = false;
+									}
+								}
+							}
+						}
+					}
+
+					// Record the areas of the texture we need to re-upload
+					const int32 TexX1 = (SubsectionSizeQuads+1) * SubIndexX + SubX1;
+					const int32 TexY1 = (SubsectionSizeQuads+1) * SubIndexY + SubY1;
+					const int32 TexX2 = (SubsectionSizeQuads+1) * SubIndexX + SubX2;
+					const int32 TexY2 = (SubsectionSizeQuads+1) * SubIndexY + SubY2;
+					for (int32 WeightmapIdx = 0; WeightmapIdx < Component->WeightmapTextures.Num(); WeightmapIdx++)
+					{
+						if (TexDataInfos[WeightmapIdx] != NULL)
+						{
+							TexDataInfos[WeightmapIdx]->AddMipUpdateRegion(0,TexX1,TexY1,TexX2,TexY2);
+						}
+					}
+				}
+			}
+
+			// Update mipmaps
+			CollisionWeightmapMipData.Reset();
+			CollisionWeightmapMipData.AddUninitialized(Component->WeightmapTextures.Num());
+			for (int32 WeightmapIdx = 0; WeightmapIdx < Component->WeightmapTextures.Num(); WeightmapIdx++)
+			{
+				UTexture2D* const WeightmapTexture = Component->WeightmapTextures[WeightmapIdx];
+
+				const int32 NumMips = WeightmapTexture->Source.GetNumMips();
+				WeightmapTextureMipData.Reset();
+				WeightmapTextureMipData.AddUninitialized(NumMips);
+				for (int32 MipIdx = 0; MipIdx < NumMips; MipIdx++)
+				{
+					FColor* const MipData = (FColor*)TexDataInfos[WeightmapIdx]->GetMipData(MipIdx);
+					WeightmapTextureMipData[MipIdx] = MipData;
+				}
+				CollisionWeightmapMipData[WeightmapIdx] = WeightmapTextureMipData[Component->CollisionMipLevel];
+
+				ULandscapeComponent::UpdateWeightmapMips(ComponentNumSubsections, SubsectionSizeQuads, WeightmapTexture, WeightmapTextureMipData, ComponentX1, ComponentY1, ComponentX2, ComponentY2, TexDataInfos[WeightmapIdx]);
+				WeightmapTextureMipData.Reset();
+			}
+
+			// Update dominant layer info stored in collision component
+			Component->UpdateCollisionLayerData(CollisionWeightmapMipData, ComponentX1, ComponentY1, ComponentX2, ComponentY2);
+			CollisionWeightmapMipData.Reset();
+
+			// Check if we need to remove weightmap allocations for layers that were completely painted away
+			bool bRemovedLayer = false;
+			for (int32 LayerIdx = 0; LayerIdx < Component->WeightmapLayerAllocations.Num(); LayerIdx++)
+			{
+				if (LayerEditDataAllZero[LayerIdx])
+				{
+					bool bLayerDeleted = DeleteLayerIfAllZero(Component, LayerDataInfos[LayerIdx].TexDataPtr, TexSize, LayerIdx);
+
+					if (bLayerDeleted)
+					{
+						LayerEditDataAllZero.RemoveAt(LayerIdx);
+						LayerDataInfos.RemoveAt(LayerIdx);
+						LayerIdx--;
+
+						bRemovedLayer = true;
+					}
+				}
+			}
+
+			if (bRemovedLayer)
+			{
+				Component->UpdateMaterialInstances();
+
+				if (Component->EditToolRenderData)
+				{
+					Component->EditToolRenderData->UpdateDebugColorMaterial();
+				}
+			}
+		}
+	}
+}
+
+template<typename TStoreData>
+void FLandscapeEditDataInterface::GetWeightDataTemplFast(ULandscapeLayerInfoObject* LayerInfo, const int32 X1, const int32 Y1, const int32 X2, const int32 Y2, TStoreData& StoreData)
+{
+	int32 ComponentIndexX1, ComponentIndexY1, ComponentIndexX2, ComponentIndexY2;
+	ALandscape::CalcComponentIndices(X1, Y1, X2, Y2, ComponentSizeQuads, ComponentIndexX1, ComponentIndexY1, ComponentIndexX2, ComponentIndexY2);
+
+	// Channel remapping
+	int32 ChannelOffsets[4] = {(int32)STRUCT_OFFSET(FColor,R),(int32)STRUCT_OFFSET(FColor,G),(int32)STRUCT_OFFSET(FColor,B),(int32)STRUCT_OFFSET(FColor,A)};
+
+	for( int32 ComponentIndexY=ComponentIndexY1;ComponentIndexY<=ComponentIndexY2;ComponentIndexY++ )
+	{
+		for( int32 ComponentIndexX=ComponentIndexX1;ComponentIndexX<=ComponentIndexX2;ComponentIndexX++ )
+		{		
+			ULandscapeComponent* Component = LandscapeInfo->XYtoComponentMap.FindRef(FIntPoint(ComponentIndexX,ComponentIndexY));
+
+			if( !Component )
+			{
+				continue;
+			}
+
+			UTexture2D* WeightmapTexture = NULL;
+			FLandscapeTextureDataInfo* TexDataInfo = NULL;
+			uint8* WeightmapTextureData = NULL;
+			uint8 WeightmapChannelOffset = 0;
+			TArray<FLandscapeTextureDataInfo*> TexDataInfos; // added for whole weight case...
+
+			if (LayerInfo != NULL)
+			{
+				for( int32 LayerIdx=0;LayerIdx<Component->WeightmapLayerAllocations.Num();LayerIdx++ )
+				{
+					if( Component->WeightmapLayerAllocations[LayerIdx].LayerInfo == LayerInfo )
+					{
+						WeightmapTexture = Component->WeightmapTextures[Component->WeightmapLayerAllocations[LayerIdx].WeightmapTextureIndex];
+						TexDataInfo = GetTextureDataInfo(WeightmapTexture);
+						WeightmapTextureData = (uint8*)TexDataInfo->GetMipData(0);
+						WeightmapChannelOffset = ChannelOffsets[Component->WeightmapLayerAllocations[LayerIdx].WeightmapTextureChannel];
+						break;
+					}
+				}
+			}
+			else
+			{
+				// Lock data for all the weightmaps
+				for( int32 WeightmapIdx=0;WeightmapIdx < Component->WeightmapTextures.Num();WeightmapIdx++ )
+				{
+					TexDataInfos.Add(GetTextureDataInfo(Component->WeightmapTextures[WeightmapIdx]));
+				}
+			}
+
+			// Find coordinates of box that lies inside component
+			int32 ComponentX1 = FMath::Clamp<int32>(X1-ComponentIndexX*ComponentSizeQuads, 0, ComponentSizeQuads);
+			int32 ComponentY1 = FMath::Clamp<int32>(Y1-ComponentIndexY*ComponentSizeQuads, 0, ComponentSizeQuads);
+			int32 ComponentX2 = FMath::Clamp<int32>(X2-ComponentIndexX*ComponentSizeQuads, 0, ComponentSizeQuads);
+			int32 ComponentY2 = FMath::Clamp<int32>(Y2-ComponentIndexY*ComponentSizeQuads, 0, ComponentSizeQuads);
+
+			// Find subsection range for this box
+			int32 SubIndexX1 = FMath::Clamp<int32>((ComponentX1-1) / SubsectionSizeQuads,0,ComponentNumSubsections-1);	// -1 because we need to pick up vertices shared between subsections
+			int32 SubIndexY1 = FMath::Clamp<int32>((ComponentY1-1) / SubsectionSizeQuads,0,ComponentNumSubsections-1);
+			int32 SubIndexX2 = FMath::Clamp<int32>(ComponentX2 / SubsectionSizeQuads,0,ComponentNumSubsections-1);
+			int32 SubIndexY2 = FMath::Clamp<int32>(ComponentY2 / SubsectionSizeQuads,0,ComponentNumSubsections-1);
+
+			for( int32 SubIndexY=SubIndexY1;SubIndexY<=SubIndexY2;SubIndexY++ )
+			{
+				for( int32 SubIndexX=SubIndexX1;SubIndexX<=SubIndexX2;SubIndexX++ )
+				{
+					// Find coordinates of box that lies inside subsection
+					int32 SubX1 = FMath::Clamp<int32>(ComponentX1-SubsectionSizeQuads*SubIndexX, 0, SubsectionSizeQuads);
+					int32 SubY1 = FMath::Clamp<int32>(ComponentY1-SubsectionSizeQuads*SubIndexY, 0, SubsectionSizeQuads);
+					int32 SubX2 = FMath::Clamp<int32>(ComponentX2-SubsectionSizeQuads*SubIndexX, 0, SubsectionSizeQuads);
+					int32 SubY2 = FMath::Clamp<int32>(ComponentY2-SubsectionSizeQuads*SubIndexY, 0, SubsectionSizeQuads);
+
+					// Update texture data for the box that lies inside subsection
+					for( int32 SubY=SubY1;SubY<=SubY2;SubY++ )
+					{
+						for( int32 SubX=SubX1;SubX<=SubX2;SubX++ )
+						{
+							int32 LandscapeX = SubIndexX*SubsectionSizeQuads + ComponentIndexX*ComponentSizeQuads + SubX;
+							int32 LandscapeY = SubIndexY*SubsectionSizeQuads + ComponentIndexY*ComponentSizeQuads + SubY;
+
+							if (LayerInfo != NULL)
+							{
+								// Find the input data corresponding to this vertex
+								uint8 Weight;
+								if( WeightmapTexture )
+								{
+									// Find the texture data corresponding to this vertex
+									int32 SizeU = WeightmapTexture->Source.GetSizeX();
+									int32 SizeV = WeightmapTexture->Source.GetSizeY();
+									int32 WeightmapOffsetX = Component->WeightmapScaleBias.Z * (float)SizeU;
+									int32 WeightmapOffsetY = Component->WeightmapScaleBias.W * (float)SizeV;
+
+									int32 TexX = WeightmapOffsetX + (SubsectionSizeQuads+1) * SubIndexX + SubX;
+									int32 TexY = WeightmapOffsetY + (SubsectionSizeQuads+1) * SubIndexY + SubY;
+									Weight = WeightmapTextureData[ 4 * (TexX + TexY * SizeU) + WeightmapChannelOffset ];
+								}
+								else
+								{
+									Weight = 0;
+								}
+
+								StoreData.Store(LandscapeX, LandscapeY, Weight);
+							}
+							else // Whole weight map case...
+							{
+								StoreData.PreInit(LandscapeInfo->Layers.Num());
+								for( int32 LayerIdx=0;LayerIdx<Component->WeightmapLayerAllocations.Num();LayerIdx++ )
+								{
+									int32 Idx = Component->WeightmapLayerAllocations[LayerIdx].WeightmapTextureIndex;
+									UTexture2D* WeightmapTexture = Component->WeightmapTextures[Idx];
+									uint8* WeightmapTextureData = (uint8*)TexDataInfos[Idx]->GetMipData(0);
+									uint8 WeightmapChannelOffset = ChannelOffsets[Component->WeightmapLayerAllocations[LayerIdx].WeightmapTextureChannel];
+
+									// Find the texture data corresponding to this vertex
+									int32 SizeU = WeightmapTexture->Source.GetSizeX();
+									int32 SizeV = WeightmapTexture->Source.GetSizeY();
+									int32 WeightmapOffsetX = Component->WeightmapScaleBias.Z * (float)SizeU;
+									int32 WeightmapOffsetY = Component->WeightmapScaleBias.W * (float)SizeV;
+
+									int32 TexX = WeightmapOffsetX + (SubsectionSizeQuads+1) * SubIndexX + SubX;
+									int32 TexY = WeightmapOffsetY + (SubsectionSizeQuads+1) * SubIndexY + SubY;
+
+									uint8 Weight = WeightmapTextureData[ 4 * (TexX + TexY * SizeU) + WeightmapChannelOffset ];
+
+									// Find index in LayerInfos
+									{
+										int32 LayerInfoIdx = LandscapeInfo->GetLayerInfoIndex(Component->WeightmapLayerAllocations[LayerIdx].LayerInfo);
+										if (LayerInfoIdx != INDEX_NONE)
+										{
+											StoreData.Store(LandscapeX, LandscapeY, Weight, LayerInfoIdx);
+										}
+									}
+								}
+							}
+						}
+					}
+				}
+			}
+		}
+	}
+}
+
+uint8 FLandscapeEditDataInterface::GetWeightMapData(const ULandscapeComponent* Component, ULandscapeLayerInfoObject* LayerInfo, int32 TexU, int32 TexV, uint8 Offset /*= 0*/, UTexture2D* Texture /*= NULL*/, uint8* TextureData /*= NULL*/)
+{
+	check(Component);
+	int32 ChannelOffsets[4] = {(int32)STRUCT_OFFSET(FColor,R),(int32)STRUCT_OFFSET(FColor,G),(int32)STRUCT_OFFSET(FColor,B),(int32)STRUCT_OFFSET(FColor,A)};
+	if (!Texture || !TextureData)
+	{
+		if (LayerInfo != NULL)
+		{
+			for( int32 LayerIdx=0;LayerIdx<Component->WeightmapLayerAllocations.Num();LayerIdx++ )
+			{
+				if( Component->WeightmapLayerAllocations[LayerIdx].LayerInfo == LayerInfo )
+				{
+					Texture = Component->WeightmapTextures[Component->WeightmapLayerAllocations[LayerIdx].WeightmapTextureIndex];
+					FLandscapeTextureDataInfo* TexDataInfo = GetTextureDataInfo(Texture);
+					TextureData = (uint8*)TexDataInfo->GetMipData(0);
+					Offset = ChannelOffsets[Component->WeightmapLayerAllocations[LayerIdx].WeightmapTextureChannel];
+					break;
+				}
+			}
+		}
+	}
+
+	if (Texture && TextureData)
+	{
+		int32 SizeU = Texture->Source.GetSizeX();
+		int32 SizeV = Texture->Source.GetSizeY();
+		int32 WeightmapOffsetX = Component->WeightmapScaleBias.Z * (float)SizeU;
+		int32 WeightmapOffsetY = Component->WeightmapScaleBias.W * (float)SizeV;
+
+		int32 TexX = WeightmapOffsetX + TexU;
+		int32 TexY = WeightmapOffsetY + TexV;
+		return TextureData[ 4 * (TexX + TexY * SizeU) + Offset ];
+	}
+	return 0;
+}
+
+template<typename TStoreData>
+void FLandscapeEditDataInterface::GetWeightDataTempl(ULandscapeLayerInfoObject* LayerInfo, int32& ValidX1, int32& ValidY1, int32& ValidX2, int32& ValidY2, TStoreData& StoreData)
+{
+	// Copy variables
+	int32 X1 = ValidX1, X2 = ValidX2, Y1 = ValidY1, Y2 = ValidY2;
+	ValidX1 = INT_MAX; ValidX2 = INT_MIN; ValidY1 = INT_MAX; ValidY2 = INT_MIN;
+
+	int32 ComponentIndexX1, ComponentIndexY1, ComponentIndexX2, ComponentIndexY2;
+	ALandscape::CalcComponentIndices(X1, Y1, X2, Y2, ComponentSizeQuads, ComponentIndexX1, ComponentIndexY1, ComponentIndexX2, ComponentIndexY2);
+
+	int32 ComponentSizeX = ComponentIndexX2-ComponentIndexX1+1;
+	int32 ComponentSizeY = ComponentIndexY2-ComponentIndexY1+1;
+
+	// Neighbor Components
+	ULandscapeComponent* BorderComponent[4] = {0, 0, 0, 0};
+	ULandscapeComponent* CornerComponent[4] = {0, 0, 0, 0};
+	bool NoBorderX1 = false, NoBorderX2 = false;
+	TArray<bool> NoBorderY1, NoBorderY2, ComponentDataExist;
+	TArray<ULandscapeComponent*> BorderComponentY1, BorderComponentY2;
+	ComponentDataExist.Empty(ComponentSizeX*ComponentSizeY);
+	ComponentDataExist.AddZeroed(ComponentSizeX*ComponentSizeY);
+	bool bHasMissingValue = false;
+
+	UTexture2D* NeighborWeightmapTexture[4] = {0, 0, 0, 0};
+	FLandscapeTextureDataInfo* NeighborTexDataInfo[4] = {0, 0, 0, 0};
+	uint8* NeighborWeightmapTextureData[4] = {0, 0, 0, 0};
+	uint8 NeighborWeightmapChannelOffset[4] = {0, 0, 0, 0};
+	uint8 CornerValues[4] = {0, 0, 0, 0};
+	int32 EdgeCoord = (SubsectionSizeQuads+1) * ComponentNumSubsections - 1; //ComponentSizeQuads;
+
+	// Channel remapping
+	int32 ChannelOffsets[4] = {(int32)STRUCT_OFFSET(FColor,R),(int32)STRUCT_OFFSET(FColor,G),(int32)STRUCT_OFFSET(FColor,B),(int32)STRUCT_OFFSET(FColor,A)};
+
+	// initial loop....
+	for( int32 ComponentIndexY=ComponentIndexY1;ComponentIndexY<=ComponentIndexY2;ComponentIndexY++ )
+	{
+		NoBorderX1 = false;
+		NoBorderX2 = false;
+		BorderComponent[0] = BorderComponent[1] = NULL;
+		for( int32 ComponentIndexX=ComponentIndexX1;ComponentIndexX<=ComponentIndexX2;ComponentIndexX++ )
+		{	
+			BorderComponent[2] = BorderComponent[3] = NULL;
+			int32 ComponentIndexXY = ComponentSizeX*(ComponentIndexY-ComponentIndexY1) + ComponentIndexX-ComponentIndexX1;
+			int32 ComponentIndexXX = ComponentIndexX - ComponentIndexX1;
+			int32 ComponentIndexYY = ComponentIndexY - ComponentIndexY1;
+			ComponentDataExist[ComponentIndexXY] = false;
+			ULandscapeComponent* Component = LandscapeInfo->XYtoComponentMap.FindRef(FIntPoint(ComponentIndexX,ComponentIndexY));
+
+			UTexture2D* WeightmapTexture = NULL;
+			FLandscapeTextureDataInfo* TexDataInfo = NULL;
+			uint8* WeightmapTextureData = NULL;
+			uint8 WeightmapChannelOffset = 0;
+			TArray<FLandscapeTextureDataInfo*> TexDataInfos; // added for whole weight case...
+			uint8 CornerSet = 0;
+			bool ExistLeft = ComponentIndexXX > 0 && ComponentDataExist[ ComponentIndexXX-1 + ComponentIndexYY * ComponentSizeX ];
+			bool ExistUp = ComponentIndexYY > 0 && ComponentDataExist[ ComponentIndexXX + (ComponentIndexYY-1) * ComponentSizeX ];
+
+			if( Component )
+			{
+				if (LayerInfo != NULL)
+				{
+					for( int32 LayerIdx=0;LayerIdx<Component->WeightmapLayerAllocations.Num();LayerIdx++ )
+					{
+						if( Component->WeightmapLayerAllocations[LayerIdx].LayerInfo == LayerInfo )
+						{
+							WeightmapTexture = Component->WeightmapTextures[Component->WeightmapLayerAllocations[LayerIdx].WeightmapTextureIndex];
+							TexDataInfo = GetTextureDataInfo(WeightmapTexture);
+							WeightmapTextureData = (uint8*)TexDataInfo->GetMipData(0);
+							WeightmapChannelOffset = ChannelOffsets[Component->WeightmapLayerAllocations[LayerIdx].WeightmapTextureChannel];
+							break;
+						}
+					}
+				}
+				else
+				{
+					// Lock data for all the weightmaps
+					for( int32 WeightmapIdx=0;WeightmapIdx < Component->WeightmapTextures.Num();WeightmapIdx++ )
+					{
+						TexDataInfos.Add(GetTextureDataInfo(Component->WeightmapTextures[WeightmapIdx]));
+					}
+				}
+				ComponentDataExist[ComponentIndexXY] = true;
+				// Update valid region
+				ValidX1 = FMath::Min<int32>(Component->GetSectionBase().X, ValidX1);
+				ValidX2 = FMath::Max<int32>(Component->GetSectionBase().X+ComponentSizeQuads, ValidX2);
+				ValidY1 = FMath::Min<int32>(Component->GetSectionBase().Y, ValidY1);
+				ValidY2 = FMath::Max<int32>(Component->GetSectionBase().Y+ComponentSizeQuads, ValidY2);
+			}
+			else
+			{
+				if (!bHasMissingValue)
+				{
+					NoBorderY1.Empty(ComponentSizeX);
+					NoBorderY2.Empty(ComponentSizeX);
+					NoBorderY1.AddZeroed(ComponentSizeX);
+					NoBorderY2.AddZeroed(ComponentSizeX);
+					BorderComponentY1.Empty(ComponentSizeX);
+					BorderComponentY2.Empty(ComponentSizeX);
+					BorderComponentY1.AddZeroed(ComponentSizeX);
+					BorderComponentY2.AddZeroed(ComponentSizeX);
+					bHasMissingValue = true;
+				}
+
+				// Search for neighbor component for interpolation
+				bool bShouldSearchX = (BorderComponent[1] && BorderComponent[1]->GetSectionBase().X / ComponentSizeQuads <= ComponentIndexX);
+				bool bShouldSearchY = (BorderComponentY2[ComponentIndexXX] && BorderComponentY2[ComponentIndexXX]->GetSectionBase().Y / ComponentSizeQuads <= ComponentIndexY);
+				// Search for left-closest component
+				if ( bShouldSearchX || (!NoBorderX1 && !BorderComponent[0]))
+				{
+					NoBorderX1 = true;
+					for (int32 X = ComponentIndexX-1; X >= ComponentIndexX1; X--)
+					{
+						BorderComponent[0] = LandscapeInfo->XYtoComponentMap.FindRef(FIntPoint(X,ComponentIndexY));
+						if (BorderComponent[0])
+						{
+							NoBorderX1 = false;
+							if (LayerInfo != NULL)
+							{
+								for( int32 LayerIdx=0;LayerIdx<BorderComponent[0]->WeightmapLayerAllocations.Num();LayerIdx++ )
+								{
+									if( BorderComponent[0]->WeightmapLayerAllocations[LayerIdx].LayerInfo == LayerInfo )
+									{
+										NeighborWeightmapTexture[0] = BorderComponent[0]->WeightmapTextures[BorderComponent[0]->WeightmapLayerAllocations[LayerIdx].WeightmapTextureIndex];
+										NeighborTexDataInfo[0] = GetTextureDataInfo(NeighborWeightmapTexture[0]);
+										NeighborWeightmapTextureData[0] = (uint8*)NeighborTexDataInfo[0]->GetMipData(0);
+										NeighborWeightmapChannelOffset[0] = ChannelOffsets[BorderComponent[0]->WeightmapLayerAllocations[LayerIdx].WeightmapTextureChannel];
+										break;
+									}
+								}
+							}
+							break;
+						}
+					}
+				}
+				// Search for right-closest component
+				if ( bShouldSearchX || (!NoBorderX2 && !BorderComponent[1]))
+				{
+					NoBorderX2 = true;
+					for (int32 X = ComponentIndexX+1; X <= ComponentIndexX2; X++)
+					{
+						BorderComponent[1] = LandscapeInfo->XYtoComponentMap.FindRef(FIntPoint(X,ComponentIndexY));
+						if (BorderComponent[1])
+						{
+							NoBorderX2 = false;
+							if (LayerInfo != NULL)
+							{
+								for( int32 LayerIdx=0;LayerIdx<BorderComponent[1]->WeightmapLayerAllocations.Num();LayerIdx++ )
+								{
+									if( BorderComponent[1]->WeightmapLayerAllocations[LayerIdx].LayerInfo == LayerInfo )
+									{
+										NeighborWeightmapTexture[1] = BorderComponent[1]->WeightmapTextures[BorderComponent[1]->WeightmapLayerAllocations[LayerIdx].WeightmapTextureIndex];
+										NeighborTexDataInfo[1] = GetTextureDataInfo(NeighborWeightmapTexture[1]);
+										NeighborWeightmapTextureData[1] = (uint8*)NeighborTexDataInfo[1]->GetMipData(0);
+										NeighborWeightmapChannelOffset[1] = ChannelOffsets[BorderComponent[1]->WeightmapLayerAllocations[LayerIdx].WeightmapTextureChannel];
+										break;
+									}
+								}
+							}
+						}
+					}
+				}
+				// Search for up-closest component
+				if ( bShouldSearchY || (!NoBorderY1[ComponentIndexXX] && !BorderComponentY1[ComponentIndexXX]))
+				{
+					NoBorderY1[ComponentIndexXX] = true;
+					for (int32 Y = ComponentIndexY-1; Y >= ComponentIndexY1; Y--)
+					{
+						BorderComponentY1[ComponentIndexXX] = BorderComponent[2] = LandscapeInfo->XYtoComponentMap.FindRef(FIntPoint(ComponentIndexX,Y));
+						if (BorderComponent[2])
+						{
+							NoBorderY1[ComponentIndexXX] = false;
+							if (LayerInfo != NULL)
+							{
+								for( int32 LayerIdx=0;LayerIdx<BorderComponent[2]->WeightmapLayerAllocations.Num();LayerIdx++ )
+								{
+									if( BorderComponent[2]->WeightmapLayerAllocations[LayerIdx].LayerInfo == LayerInfo )
+									{
+										NeighborWeightmapTexture[2] = BorderComponent[2]->WeightmapTextures[BorderComponent[2]->WeightmapLayerAllocations[LayerIdx].WeightmapTextureIndex];
+										NeighborTexDataInfo[2] = GetTextureDataInfo(NeighborWeightmapTexture[2]);
+										NeighborWeightmapTextureData[2] = (uint8*)NeighborTexDataInfo[2]->GetMipData(0);
+										NeighborWeightmapChannelOffset[2] = ChannelOffsets[BorderComponent[2]->WeightmapLayerAllocations[LayerIdx].WeightmapTextureChannel];
+										break;
+									}
+								}
+							}
+						}
+					}
+				}
+				else
+				{
+					BorderComponent[2] = BorderComponentY1[ComponentIndexXX];
+					if (BorderComponent[2])
+					{
+						if (LayerInfo != NULL)
+						{
+							for( int32 LayerIdx=0;LayerIdx<BorderComponent[2]->WeightmapLayerAllocations.Num();LayerIdx++ )
+							{
+								if( BorderComponent[2]->WeightmapLayerAllocations[LayerIdx].LayerInfo == LayerInfo )
+								{
+									NeighborWeightmapTexture[2] = BorderComponent[2]->WeightmapTextures[BorderComponent[2]->WeightmapLayerAllocations[LayerIdx].WeightmapTextureIndex];
+									NeighborTexDataInfo[2] = GetTextureDataInfo(NeighborWeightmapTexture[2]);
+									NeighborWeightmapTextureData[2] = (uint8*)NeighborTexDataInfo[2]->GetMipData(0);
+									NeighborWeightmapChannelOffset[2] = ChannelOffsets[BorderComponent[2]->WeightmapLayerAllocations[LayerIdx].WeightmapTextureChannel];
+									break;
+								}
+							}
+						}
+					}
+				}
+				// Search for bottom-closest component
+				if ( bShouldSearchY || (!NoBorderY2[ComponentIndexXX] && !BorderComponentY2[ComponentIndexXX]))
+				{
+					NoBorderY2[ComponentIndexXX] = true;
+					for (int32 Y = ComponentIndexY+1; Y <= ComponentIndexY2; Y++)
+					{
+						BorderComponentY2[ComponentIndexXX] = BorderComponent[3] = LandscapeInfo->XYtoComponentMap.FindRef(FIntPoint(ComponentIndexX,Y));
+						if (BorderComponent[3])
+						{
+							NoBorderY2[ComponentIndexXX] = false;
+							if (LayerInfo != NULL)
+							{
+								for( int32 LayerIdx=0;LayerIdx<BorderComponent[3]->WeightmapLayerAllocations.Num();LayerIdx++ )
+								{
+									if( BorderComponent[3]->WeightmapLayerAllocations[LayerIdx].LayerInfo == LayerInfo )
+									{
+										NeighborWeightmapTexture[3] = BorderComponent[3]->WeightmapTextures[BorderComponent[3]->WeightmapLayerAllocations[LayerIdx].WeightmapTextureIndex];
+										NeighborTexDataInfo[3] = GetTextureDataInfo(NeighborWeightmapTexture[3]);
+										NeighborWeightmapTextureData[3] = (uint8*)NeighborTexDataInfo[3]->GetMipData(0);
+										NeighborWeightmapChannelOffset[3] = ChannelOffsets[BorderComponent[3]->WeightmapLayerAllocations[LayerIdx].WeightmapTextureChannel];
+										break;
+									}
+								}
+							}
+							break;
+						}
+					}
+				}
+				else
+				{
+					BorderComponent[3] = BorderComponentY2[ComponentIndexXX];
+					if (BorderComponent[3])
+					{
+						if (LayerInfo != NULL)
+						{
+							for( int32 LayerIdx=0;LayerIdx<BorderComponent[3]->WeightmapLayerAllocations.Num();LayerIdx++ )
+							{
+								if( BorderComponent[3]->WeightmapLayerAllocations[LayerIdx].LayerInfo == LayerInfo )
+								{
+									NeighborWeightmapTexture[3] = BorderComponent[3]->WeightmapTextures[BorderComponent[3]->WeightmapLayerAllocations[LayerIdx].WeightmapTextureIndex];
+									NeighborTexDataInfo[3] = GetTextureDataInfo(NeighborWeightmapTexture[3]);
+									NeighborWeightmapTextureData[3] = (uint8*)NeighborTexDataInfo[3]->GetMipData(0);
+									NeighborWeightmapChannelOffset[3] = ChannelOffsets[BorderComponent[3]->WeightmapLayerAllocations[LayerIdx].WeightmapTextureChannel];
+									break;
+								}
+							}
+						}
+					}
+				}
+
+				CornerComponent[0] = ComponentIndexX >= ComponentIndexX1 && ComponentIndexY >= ComponentIndexY1 ?
+					LandscapeInfo->XYtoComponentMap.FindRef(FIntPoint((ComponentIndexX-1),(ComponentIndexY-1))) : NULL;
+				CornerComponent[1] = ComponentIndexX <= ComponentIndexX2 && ComponentIndexY >= ComponentIndexY1 ?
+					LandscapeInfo->XYtoComponentMap.FindRef(FIntPoint((ComponentIndexX+1),(ComponentIndexY-1))) : NULL;
+				CornerComponent[2] = ComponentIndexX >= ComponentIndexX1 && ComponentIndexY <= ComponentIndexY2 ?
+					LandscapeInfo->XYtoComponentMap.FindRef(FIntPoint((ComponentIndexX-1),(ComponentIndexY+1))) : NULL;
+				CornerComponent[3] = ComponentIndexX <= ComponentIndexX2 && ComponentIndexY <= ComponentIndexY2 ?
+					LandscapeInfo->XYtoComponentMap.FindRef(FIntPoint((ComponentIndexX+1),(ComponentIndexY+1))) : NULL;
+
+				if (CornerComponent[0])
+				{
+					CornerSet |= 1;
+					CornerValues[0] = GetWeightMapData(CornerComponent[0], LayerInfo, EdgeCoord, EdgeCoord);
+				}
+				else if ((ExistLeft || ExistUp) && X1 <= ComponentIndexX*ComponentSizeQuads && Y1 <= ComponentIndexY*ComponentSizeQuads  )
+				{
+					CornerSet |= 1;
+					CornerValues[0] = StoreData.Load( ComponentIndexX*ComponentSizeQuads, ComponentIndexY*ComponentSizeQuads);
+				}
+				else if (BorderComponent[0])
+				{
+					CornerSet |= 1;
+					CornerValues[0] = GetWeightMapData(BorderComponent[0], LayerInfo, EdgeCoord, 0, NeighborWeightmapChannelOffset[0], NeighborWeightmapTexture[0], NeighborWeightmapTextureData[0]);
+				}
+				else if (BorderComponent[2])
+				{
+					CornerSet |= 1;
+					CornerValues[0] = GetWeightMapData(BorderComponent[2], LayerInfo, 0, EdgeCoord, NeighborWeightmapChannelOffset[2], NeighborWeightmapTexture[2], NeighborWeightmapTextureData[2]);
+				}
+
+				if (CornerComponent[1])
+				{
+					CornerSet |= 1 << 1;
+					CornerValues[1] = GetWeightMapData(CornerComponent[1], LayerInfo, 0, EdgeCoord);
+				}
+				else if (ExistUp && X2 >= (ComponentIndexX+1)*ComponentSizeQuads)
+				{
+					CornerSet |= 1 << 1;
+					CornerValues[1] = StoreData.Load( (ComponentIndexX+1)*ComponentSizeQuads, ComponentIndexY*ComponentSizeQuads);
+				}
+				else if (BorderComponent[1])
+				{
+					CornerSet |= 1 << 1;
+					CornerValues[1] = GetWeightMapData(BorderComponent[1], LayerInfo, 0, 0, NeighborWeightmapChannelOffset[1], NeighborWeightmapTexture[1], NeighborWeightmapTextureData[1]);
+				}
+				else if (BorderComponent[2])
+				{
+					CornerSet |= 1 << 1;
+					CornerValues[1] = GetWeightMapData(BorderComponent[2], LayerInfo, EdgeCoord, EdgeCoord, NeighborWeightmapChannelOffset[2], NeighborWeightmapTexture[2], NeighborWeightmapTextureData[2]);
+				}
+
+				if (CornerComponent[2])
+				{
+					CornerSet |= 1 << 2;
+					CornerValues[2] = GetWeightMapData(CornerComponent[2], LayerInfo, EdgeCoord, 0);
+				}
+				else if (ExistLeft && Y2 >= (ComponentIndexY+1)*ComponentSizeQuads) // Use data already stored for 0, 2
+				{
+					CornerSet |= 1 << 2;
+					CornerValues[2] = StoreData.Load( ComponentIndexX*ComponentSizeQuads, (ComponentIndexY+1)*ComponentSizeQuads);
+				}
+				else if (BorderComponent[0])
+				{
+					CornerSet |= 1 << 2;
+					CornerValues[2] = GetWeightMapData(BorderComponent[0], LayerInfo, EdgeCoord, EdgeCoord, NeighborWeightmapChannelOffset[0], NeighborWeightmapTexture[0], NeighborWeightmapTextureData[0]);
+				}
+				else if (BorderComponent[3])
+				{
+					CornerSet |= 1 << 2;
+					CornerValues[2] = GetWeightMapData(BorderComponent[3], LayerInfo, 0, 0, NeighborWeightmapChannelOffset[3], NeighborWeightmapTexture[3], NeighborWeightmapTextureData[3]);
+				}
+
+				if (CornerComponent[3])
+				{
+					CornerSet |= 1 << 3;
+					CornerValues[3] = GetWeightMapData(CornerComponent[3], LayerInfo, 0, 0);
+				}
+				else if (BorderComponent[1])
+				{
+					CornerSet |= 1 << 3;
+					CornerValues[3] = GetWeightMapData(BorderComponent[1], LayerInfo, 0, EdgeCoord, NeighborWeightmapChannelOffset[1], NeighborWeightmapTexture[1], NeighborWeightmapTextureData[1]);
+				}
+				else if (BorderComponent[3])
+				{
+					CornerSet |= 1 << 3;
+					CornerValues[3] = GetWeightMapData(BorderComponent[3], LayerInfo, EdgeCoord, 0, NeighborWeightmapChannelOffset[3], NeighborWeightmapTexture[3], NeighborWeightmapTextureData[3]);
+				}
+
+				FillCornerValues(CornerSet, CornerValues);
+				ComponentDataExist[ComponentIndexXY] = ExistLeft || ExistUp || (BorderComponent[0] || BorderComponent[1] || BorderComponent[2] || BorderComponent[3]) || CornerSet;
+			}
+
+			if (!ComponentDataExist[ComponentIndexXY])
+			{
+				continue;
+			}
+
+			// Find coordinates of box that lies inside component
+			int32 ComponentX1 = FMath::Clamp<int32>(X1-ComponentIndexX*ComponentSizeQuads, 0, ComponentSizeQuads);
+			int32 ComponentY1 = FMath::Clamp<int32>(Y1-ComponentIndexY*ComponentSizeQuads, 0, ComponentSizeQuads);
+			int32 ComponentX2 = FMath::Clamp<int32>(X2-ComponentIndexX*ComponentSizeQuads, 0, ComponentSizeQuads);
+			int32 ComponentY2 = FMath::Clamp<int32>(Y2-ComponentIndexY*ComponentSizeQuads, 0, ComponentSizeQuads);
+
+			// Find subsection range for this box
+			int32 SubIndexX1 = FMath::Clamp<int32>((ComponentX1-1) / SubsectionSizeQuads,0,ComponentNumSubsections-1);	// -1 because we need to pick up vertices shared between subsections
+			int32 SubIndexY1 = FMath::Clamp<int32>((ComponentY1-1) / SubsectionSizeQuads,0,ComponentNumSubsections-1);
+			int32 SubIndexX2 = FMath::Clamp<int32>(ComponentX2 / SubsectionSizeQuads,0,ComponentNumSubsections-1);
+			int32 SubIndexY2 = FMath::Clamp<int32>(ComponentY2 / SubsectionSizeQuads,0,ComponentNumSubsections-1);
+
+			for( int32 SubIndexY=SubIndexY1;SubIndexY<=SubIndexY2;SubIndexY++ )
+			{
+				for( int32 SubIndexX=SubIndexX1;SubIndexX<=SubIndexX2;SubIndexX++ )
+				{
+					// Find coordinates of box that lies inside subsection
+					int32 SubX1 = FMath::Clamp<int32>(ComponentX1-SubsectionSizeQuads*SubIndexX, 0, SubsectionSizeQuads);
+					int32 SubY1 = FMath::Clamp<int32>(ComponentY1-SubsectionSizeQuads*SubIndexY, 0, SubsectionSizeQuads);
+					int32 SubX2 = FMath::Clamp<int32>(ComponentX2-SubsectionSizeQuads*SubIndexX, 0, SubsectionSizeQuads);
+					int32 SubY2 = FMath::Clamp<int32>(ComponentY2-SubsectionSizeQuads*SubIndexY, 0, SubsectionSizeQuads);
+
+					// Update texture data for the box that lies inside subsection
+					for( int32 SubY=SubY1;SubY<=SubY2;SubY++ )
+					{
+						for( int32 SubX=SubX1;SubX<=SubX2;SubX++ )
+						{
+							int32 LandscapeX = SubIndexX*SubsectionSizeQuads + ComponentIndexX*ComponentSizeQuads + SubX;
+							int32 LandscapeY = SubIndexY*SubsectionSizeQuads + ComponentIndexY*ComponentSizeQuads + SubY;
+
+							if (LayerInfo != NULL)
+							{
+								// Find the input data corresponding to this vertex
+								uint8 Weight;
+								if( WeightmapTexture )
+								{
+									// Find the texture data corresponding to this vertex
+									Weight = GetWeightMapData(Component, LayerInfo, (SubsectionSizeQuads+1) * SubIndexX + SubX, (SubsectionSizeQuads+1) * SubIndexY + SubY, WeightmapChannelOffset, WeightmapTexture, WeightmapTextureData );
+									StoreData.Store(LandscapeX, LandscapeY, Weight);
+								}
+								else
+								{
+									// Find the texture data corresponding to this vertex
+									uint8 Value[4] = {0, 0, 0, 0};
+									int32 Dist[4] = {INT_MAX, INT_MAX, INT_MAX, INT_MAX};
+									float ValueX = 0.f, ValueY = 0.f;
+									bool Exist[4] = {false, false, false, false};
+
+									// Use data already stored for 0, 2
+									if (ExistLeft && SubX == 0)
+									{
+										Value[0] = StoreData.Load( ComponentIndexX*ComponentSizeQuads, LandscapeY);
+										Dist[0] = LandscapeX - (ComponentIndexX*ComponentSizeQuads);
+										Exist[0] = true;
+									}
+									else if (BorderComponent[0] && NeighborWeightmapTexture[0])
+									{
+										Value[0] = GetWeightMapData(BorderComponent[0], LayerInfo, EdgeCoord, (SubsectionSizeQuads+1) * SubIndexY + SubY, NeighborWeightmapChannelOffset[0], NeighborWeightmapTexture[0], NeighborWeightmapTextureData[0]);
+										Dist[0] = LandscapeX - (BorderComponent[0]->GetSectionBase().X + ComponentSizeQuads);
+										Exist[0] = true;
+									}
+									else 
+									{
+										if ((CornerSet & 1) && (CornerSet & 1 << 2))
+										{
+											int32 Dist1 = LandscapeY - (ComponentIndexY*ComponentSizeQuads);
+											int32 Dist2 = ((ComponentIndexY+1)*ComponentSizeQuads) - LandscapeY;
+											Value[0] = (float)(Dist2 * CornerValues[0] + Dist1 * CornerValues[2]) / (Dist1 + Dist2);
+											Dist[0] = LandscapeX - (ComponentIndexX*ComponentSizeQuads);
+											Exist[0] = true;
+										}
+									}
+
+									if (BorderComponent[1] && NeighborWeightmapTexture[1])
+									{
+										Value[1] = GetWeightMapData(BorderComponent[1], LayerInfo, 0, (SubsectionSizeQuads+1) * SubIndexY + SubY, NeighborWeightmapChannelOffset[1], NeighborWeightmapTexture[1], NeighborWeightmapTextureData[1]);
+										Dist[1] = (BorderComponent[1]->GetSectionBase().X) - LandscapeX;
+										Exist[1] = true;
+									}
+									else
+									{
+										if ((CornerSet & 1 << 1) && (CornerSet & 1 << 3))
+										{
+											int32 Dist1 = LandscapeY - (ComponentIndexY*ComponentSizeQuads);
+											int32 Dist2 = ((ComponentIndexY+1)*ComponentSizeQuads) - LandscapeY;
+											Value[1] = (float)(Dist2 * CornerValues[1] + Dist1 * CornerValues[3]) / (Dist1 + Dist2);
+											Dist[1] = (ComponentIndexX+1)*ComponentSizeQuads - LandscapeX;
+											Exist[1] = true;
+										}
+									}
+
+									if (ExistUp && SubY == 0)
+									{
+										Value[2] = StoreData.Load( LandscapeX, ComponentIndexY*ComponentSizeQuads);
+										Dist[2] = LandscapeY - (ComponentIndexY*ComponentSizeQuads);
+										Exist[2] = true;
+									}
+									else if (BorderComponent[2] && NeighborWeightmapTexture[2])
+									{
+										Value[2] = GetWeightMapData(BorderComponent[2], LayerInfo, (SubsectionSizeQuads+1) * SubIndexX + SubX, EdgeCoord, NeighborWeightmapChannelOffset[2], NeighborWeightmapTexture[2], NeighborWeightmapTextureData[2]);
+										Dist[2] = LandscapeY - (BorderComponent[2]->GetSectionBase().Y + ComponentSizeQuads);
+										Exist[2] = true;
+									}
+									else
+									{
+										if ((CornerSet & 1) && (CornerSet & 1 << 1))
+										{
+											int32 Dist1 = LandscapeX - (ComponentIndexX*ComponentSizeQuads);
+											int32 Dist2 = (ComponentIndexX+1)*ComponentSizeQuads - LandscapeX;
+											Value[2] = (float)(Dist2 * CornerValues[0] + Dist1 * CornerValues[1]) / (Dist1 + Dist2);
+											Dist[2] = LandscapeY - (ComponentIndexY*ComponentSizeQuads);
+											Exist[2] = true;
+										}
+									}
+
+									if (BorderComponent[3] && NeighborWeightmapTexture[3])
+									{
+										Value[3] = GetWeightMapData(BorderComponent[3], LayerInfo, (SubsectionSizeQuads+1) * SubIndexX + SubX, 0, NeighborWeightmapChannelOffset[3], NeighborWeightmapTexture[3], NeighborWeightmapTextureData[3]);
+										Dist[3] = (BorderComponent[3]->GetSectionBase().Y) - LandscapeY;
+										Exist[3] = true;
+									}
+									else
+									{
+										if ((CornerSet & 1 << 2) && (CornerSet & 1 << 3))
+										{
+											int32 Dist1 = LandscapeX - (ComponentIndexX*ComponentSizeQuads);
+											int32 Dist2 = (ComponentIndexX+1)*ComponentSizeQuads - LandscapeX;
+											Value[3] = (float)(Dist2 * CornerValues[2] + Dist1 * CornerValues[3]) / (Dist1 + Dist2);
+											Dist[3] = (ComponentIndexY+1)*ComponentSizeQuads - LandscapeY;
+											Exist[3] = true;
+										}
+									}
+
+									CalcInterpValue<uint8>(Dist, Exist, Value, ValueX, ValueY);
+
+									uint8 FinalValue = 0; // Default Value
+									if ( (Exist[0] || Exist[1]) && (Exist[2] || Exist[3]) )
+									{
+										FinalValue = CalcValueFromValueXY<uint8>(Dist, ValueX, ValueY, CornerSet, CornerValues);
+									}
+									else if ( (Exist[0] || Exist[1]) )
+									{
+										FinalValue = ValueX;
+									}
+									else if ( (Exist[2] || Exist[3]) )
+									{
+										FinalValue = ValueY;
+									}
+
+									Weight = FinalValue;
+								}
+
+								StoreData.Store(LandscapeX, LandscapeY, Weight);
+							}
+							else // Whole weight map case... no interpolation now...
+							{
+								StoreData.PreInit(LandscapeInfo->Layers.Num());
+								for( int32 LayerIdx=0;LayerIdx<Component->WeightmapLayerAllocations.Num();LayerIdx++ )
+								{
+									int32 Idx = Component->WeightmapLayerAllocations[LayerIdx].WeightmapTextureIndex;
+									UTexture2D* WeightmapTexture = Component->WeightmapTextures[Idx];
+									uint8* WeightmapTextureData = (uint8*)TexDataInfos[Idx]->GetMipData(0);
+									uint8 WeightmapChannelOffset = ChannelOffsets[Component->WeightmapLayerAllocations[LayerIdx].WeightmapTextureChannel];
+
+									// Find the texture data corresponding to this vertex
+									int32 SizeU = WeightmapTexture->Source.GetSizeX();
+									int32 SizeV = WeightmapTexture->Source.GetSizeY();
+									int32 WeightmapOffsetX = Component->WeightmapScaleBias.Z * (float)SizeU;
+									int32 WeightmapOffsetY = Component->WeightmapScaleBias.W * (float)SizeV;
+
+									int32 TexX = WeightmapOffsetX + (SubsectionSizeQuads+1) * SubIndexX + SubX;
+									int32 TexY = WeightmapOffsetY + (SubsectionSizeQuads+1) * SubIndexY + SubY;
+
+									uint8 Weight = WeightmapTextureData[ 4 * (TexX + TexY * SizeU) + WeightmapChannelOffset ];
+
+									// Find index in LayerInfos
+									{
+										int32 LayerInfoIdx = LandscapeInfo->GetLayerInfoIndex(Component->WeightmapLayerAllocations[LayerIdx].LayerInfo);
+										if (LayerInfoIdx != INDEX_NONE)
+										{
+											StoreData.Store(LandscapeX, LandscapeY, Weight, LayerInfoIdx);
+										}
+									}
+								}
+							}
+						}
+					}
+				}
+			}
+		}
+	}
+
+	if (bHasMissingValue)
+	{
+		CalcMissingValues<uint8, TStoreData, float>( X1, X2, Y1, Y2,
+			ComponentIndexX1, ComponentIndexX2, ComponentIndexY1, ComponentIndexY2, 
+			ComponentSizeX, ComponentSizeY, CornerValues,
+			NoBorderY1, NoBorderY2, ComponentDataExist, StoreData );
+		// Update valid region
+		ValidX1 = FMath::Max<int32>(X1, ValidX1);
+		ValidX2 = FMath::Min<int32>(X2, ValidX2);
+		ValidY1 = FMath::Max<int32>(Y1, ValidY1);
+		ValidY2 = FMath::Min<int32>(Y2, ValidY2);
+	}
+	else
+	{
+		ValidX1 = X1;
+		ValidX2 = X2;
+		ValidY1 = Y1;
+		ValidY2 = Y2;
+	}
+}
+
+void FLandscapeEditDataInterface::GetWeightData(ULandscapeLayerInfoObject* LayerInfo, int32& X1, int32& Y1, int32& X2, int32& Y2, uint8* Data, int32 Stride)
+{
+	if( Stride==0 )
+	{
+		Stride = (1+X2-X1);
+	}
+	TArrayStoreData<uint8> ArrayStoreData(X1, Y1, Data, Stride);
+	GetWeightDataTempl(LayerInfo, X1, Y1, X2, Y2, ArrayStoreData);
+}
+
+void FLandscapeEditDataInterface::GetWeightDataFast(ULandscapeLayerInfoObject* LayerInfo, const int32 X1, const int32 Y1, const int32 X2, const int32 Y2, uint8* Data, int32 Stride)
+{
+	if( Stride==0 )
+	{
+		Stride = (1+X2-X1);
+	}
+	TArrayStoreData<uint8> ArrayStoreData(X1, Y1, Data, Stride);
+	GetWeightDataTemplFast(LayerInfo, X1, Y1, X2, Y2, ArrayStoreData);
+}
+
+void FLandscapeEditDataInterface::GetWeightData(ULandscapeLayerInfoObject* LayerInfo, int32& X1, int32& Y1, int32& X2, int32& Y2, TMap<FIntPoint, uint8>& SparseData)
+{
+	TSparseStoreData<uint8> SparseStoreData(SparseData);
+	GetWeightDataTempl(LayerInfo, X1, Y1, X2, Y2, SparseStoreData);
+}
+
+void FLandscapeEditDataInterface::GetWeightDataFast(ULandscapeLayerInfoObject* LayerInfo, const int32 X1, const int32 Y1, const int32 X2, const int32 Y2, TMap<FIntPoint, uint8>& SparseData)
+{
+	TSparseStoreData<uint8> SparseStoreData(SparseData);
+	GetWeightDataTemplFast(LayerInfo, X1, Y1, X2, Y2, SparseStoreData);
+}
+
+void FLandscapeEditDataInterface::GetWeightDataFast(ULandscapeLayerInfoObject* LayerInfo, const int32 X1, const int32 Y1, const int32 X2, const int32 Y2, TArray<uint8>* Data, int32 Stride)
+{
+	if( Stride==0 )
+	{
+		Stride = (1+X2-X1);
+	}
+	TArrayStoreData<TArray<uint8>> ArrayStoreData(X1, Y1, Data, Stride);
+	GetWeightDataTemplFast(LayerInfo, X1, Y1, X2, Y2, ArrayStoreData);
+}
+
+void FLandscapeEditDataInterface::GetWeightDataFast(ULandscapeLayerInfoObject* LayerInfo, const int32 X1, const int32 Y1, const int32 X2, const int32 Y2, TMap<FIntPoint, TArray<uint8>>& SparseData)
+{
+	TSparseStoreData<TArray<uint8>> SparseStoreData(SparseData);
+	GetWeightDataTemplFast(LayerInfo, X1, Y1, X2, Y2, SparseStoreData);
+}
+
+FLandscapeTextureDataInfo* FLandscapeEditDataInterface::GetTextureDataInfo(UTexture2D* Texture)
+{
+	FLandscapeTextureDataInfo* Result = TextureDataMap.FindRef(Texture);
+	if( !Result )
+	{
+		Result = TextureDataMap.Add(Texture, new FLandscapeTextureDataInfo(Texture));
+	}
+	return Result;
+}
+
+void FLandscapeEditDataInterface::CopyTextureChannel(UTexture2D* Dest, int32 DestChannel, UTexture2D* Src, int32 SrcChannel)
+{
+	FLandscapeTextureDataInfo* DestDataInfo = GetTextureDataInfo(Dest);
+	FLandscapeTextureDataInfo* SrcDataInfo = GetTextureDataInfo(Src);
+	int32 MipSize = Dest->Source.GetSizeX();
+	check(Dest->Source.GetSizeX() == Dest->Source.GetSizeY() && Src->Source.GetSizeX() == Dest->Source.GetSizeX());
+
+	// Channel remapping
+	int32 ChannelOffsets[4] = {(int32)STRUCT_OFFSET(FColor,R),(int32)STRUCT_OFFSET(FColor,G),(int32)STRUCT_OFFSET(FColor,B),(int32)STRUCT_OFFSET(FColor,A)};
+
+	for( int32 MipIdx=0;MipIdx<DestDataInfo->NumMips();MipIdx++ )
+	{
+		uint8* DestTextureData = (uint8*)DestDataInfo->GetMipData(MipIdx) + ChannelOffsets[DestChannel];
+		uint8* SrcTextureData = (uint8*)SrcDataInfo->GetMipData(MipIdx) + ChannelOffsets[SrcChannel];
+
+		for( int32 i=0;i<FMath::Square(MipSize);i++ )
+		{
+			DestTextureData[i*4] = SrcTextureData[i*4];
+		}
+
+		DestDataInfo->AddMipUpdateRegion(MipIdx, 0, 0, MipSize-1, MipSize-1);
+		MipSize >>= 1;
+	}
+}
+
+void FLandscapeEditDataInterface::CopyTextureFromHeightmap(UTexture2D* Dest, int32 DestChannel, ULandscapeComponent* Comp, int32 SrcChannel)
+{
+	FLandscapeTextureDataInfo* DestDataInfo = GetTextureDataInfo(Dest);
+	int32 MipSize = Dest->Source.GetSizeX();
+	check(Dest->Source.GetSizeX() == Dest->Source.GetSizeY());
+
+	// Channel remapping
+	int32 ChannelOffsets[4] = {(int32)STRUCT_OFFSET(FColor,R),(int32)STRUCT_OFFSET(FColor,G),(int32)STRUCT_OFFSET(FColor,B),(int32)STRUCT_OFFSET(FColor,A)};
+
+	for( int32 MipIdx=0;MipIdx<DestDataInfo->NumMips();MipIdx++ )
+	{
+		FLandscapeComponentDataInterface DataInterface(Comp, MipIdx);
+		TArray<FColor> Heightmap;
+		DataInterface.GetHeightmapTextureData(Heightmap);
+
+		uint8* DestTextureData = (uint8*)DestDataInfo->GetMipData(MipIdx) + ChannelOffsets[DestChannel];
+		uint8* SrcTextureData = (uint8*)Heightmap.GetData() + ChannelOffsets[SrcChannel];
+
+		for( int32 i=0;i<FMath::Square(MipSize);i++ )
+		{
+			DestTextureData[i*4] = SrcTextureData[i*4];
+		}
+
+		DestDataInfo->AddMipUpdateRegion(MipIdx, 0, 0, MipSize-1, MipSize-1);
+		MipSize >>= 1;
+	}
+}
+
+void FLandscapeEditDataInterface::ZeroTextureChannel(UTexture2D* Dest, int32 DestChannel)
+{
+	FLandscapeTextureDataInfo* DestDataInfo = GetTextureDataInfo(Dest);
+	int32 MipSize = Dest->Source.GetSizeX();
+	check(Dest->Source.GetSizeX() == Dest->Source.GetSizeY());
+
+	// Channel remapping
+	int32 ChannelOffsets[4] = {(int32)STRUCT_OFFSET(FColor,R),(int32)STRUCT_OFFSET(FColor,G),(int32)STRUCT_OFFSET(FColor,B),(int32)STRUCT_OFFSET(FColor,A)};
+
+	for( int32 MipIdx=0;MipIdx<DestDataInfo->NumMips();MipIdx++ )
+	{
+		uint8* DestTextureData = (uint8*)DestDataInfo->GetMipData(MipIdx) + ChannelOffsets[DestChannel];
+
+		for( int32 i=0;i<FMath::Square(MipSize);i++ )
+		{
+			DestTextureData[i*4] = 0;
+		}
+
+		DestDataInfo->AddMipUpdateRegion(MipIdx, 0, 0, MipSize-1, MipSize-1);
+		MipSize >>= 1;
+	}
+}
+
+template<typename TData>
+void FLandscapeEditDataInterface::SetTextureValueTempl(UTexture2D* Dest, TData Value)
+{
+	FLandscapeTextureDataInfo* DestDataInfo = GetTextureDataInfo(Dest);
+	int32 MipSize = Dest->Source.GetSizeX();
+	check(Dest->Source.GetSizeX() == Dest->Source.GetSizeY());
+
+	for( int32 MipIdx=0;MipIdx<DestDataInfo->NumMips();MipIdx++ )
+	{
+		TData* DestTextureData = (TData*)DestDataInfo->GetMipData(MipIdx);
+
+		for( int32 i=0;i<FMath::Square(MipSize);i++ )
+		{
+			DestTextureData[i] = Value;
+		}
+
+		DestDataInfo->AddMipUpdateRegion(MipIdx, 0, 0, MipSize-1, MipSize-1);
+		MipSize >>= 1;
+	}
+}
+
+void FLandscapeEditDataInterface::ZeroTexture(UTexture2D* Dest)
+{
+	SetTextureValueTempl<uint8>(Dest, 0);
+}
+
+void FLandscapeEditDataInterface::SetTextureValue(UTexture2D* Dest, FColor Value)
+{
+	SetTextureValueTempl<FColor>(Dest, Value);
+}
+
+template<typename TData>
+bool FLandscapeEditDataInterface::EqualTextureValueTempl(UTexture2D* Src, TData Value)
+{
+	FLandscapeTextureDataInfo* DestDataInfo = GetTextureDataInfo(Src);
+	TData* DestTextureData = (TData*)DestDataInfo->GetMipData(0);
+	int32 Size = Src->Source.GetSizeX() * Src->Source.GetSizeY();
+
+	for( int32 i = 0 ;i < Size; ++i )
+	{
+		if (DestTextureData[i] != Value)
+		{
+			return false;
+		}
+	}
+
+	return true;
+}
+
+bool FLandscapeEditDataInterface::EqualTextureValue(UTexture2D* Src, FColor Value)
+{
+	return EqualTextureValueTempl<FColor>(Src, Value);
+}
+
+
+template<typename TStoreData>
+void FLandscapeEditDataInterface::GetSelectDataTempl(const int32 X1, const int32 Y1, const int32 X2, const int32 Y2, TStoreData& StoreData)
+{
+	int32 ComponentIndexX1, ComponentIndexY1, ComponentIndexX2, ComponentIndexY2;
+	ALandscape::CalcComponentIndices(X1, Y1, X2, Y2, ComponentSizeQuads, ComponentIndexX1, ComponentIndexY1, ComponentIndexX2, ComponentIndexY2);
+
+	for( int32 ComponentIndexY=ComponentIndexY1;ComponentIndexY<=ComponentIndexY2;ComponentIndexY++ )
+	{
+		for( int32 ComponentIndexX=ComponentIndexX1;ComponentIndexX<=ComponentIndexX2;ComponentIndexX++ )
+		{		
+			ULandscapeComponent* Component = LandscapeInfo->XYtoComponentMap.FindRef(FIntPoint(ComponentIndexX,ComponentIndexY));
+
+			FLandscapeTextureDataInfo* TexDataInfo = NULL;
+			uint8* SelectTextureData = NULL;
+			if( Component && Component->EditToolRenderData && Component->EditToolRenderData->DataTexture )
+			{
+				TexDataInfo = GetTextureDataInfo(Component->EditToolRenderData->DataTexture);
+				SelectTextureData = (uint8*)TexDataInfo->GetMipData(0);
+			}
+
+			// Find coordinates of box that lies inside component
+			int32 ComponentX1 = FMath::Clamp<int32>(X1-ComponentIndexX*ComponentSizeQuads, 0, ComponentSizeQuads);
+			int32 ComponentY1 = FMath::Clamp<int32>(Y1-ComponentIndexY*ComponentSizeQuads, 0, ComponentSizeQuads);
+			int32 ComponentX2 = FMath::Clamp<int32>(X2-ComponentIndexX*ComponentSizeQuads, 0, ComponentSizeQuads);
+			int32 ComponentY2 = FMath::Clamp<int32>(Y2-ComponentIndexY*ComponentSizeQuads, 0, ComponentSizeQuads);
+
+			// Find subsection range for this box
+			int32 SubIndexX1 = FMath::Clamp<int32>((ComponentX1-1) / SubsectionSizeQuads,0,ComponentNumSubsections-1);	// -1 because we need to pick up vertices shared between subsections
+			int32 SubIndexY1 = FMath::Clamp<int32>((ComponentY1-1) / SubsectionSizeQuads,0,ComponentNumSubsections-1);
+			int32 SubIndexX2 = FMath::Clamp<int32>(ComponentX2 / SubsectionSizeQuads,0,ComponentNumSubsections-1);
+			int32 SubIndexY2 = FMath::Clamp<int32>(ComponentY2 / SubsectionSizeQuads,0,ComponentNumSubsections-1);
+
+			for( int32 SubIndexY=SubIndexY1;SubIndexY<=SubIndexY2;SubIndexY++ )
+			{
+				for( int32 SubIndexX=SubIndexX1;SubIndexX<=SubIndexX2;SubIndexX++ )
+				{
+					// Find coordinates of box that lies inside subsection
+					int32 SubX1 = FMath::Clamp<int32>(ComponentX1-SubsectionSizeQuads*SubIndexX, 0, SubsectionSizeQuads);
+					int32 SubY1 = FMath::Clamp<int32>(ComponentY1-SubsectionSizeQuads*SubIndexY, 0, SubsectionSizeQuads);
+					int32 SubX2 = FMath::Clamp<int32>(ComponentX2-SubsectionSizeQuads*SubIndexX, 0, SubsectionSizeQuads);
+					int32 SubY2 = FMath::Clamp<int32>(ComponentY2-SubsectionSizeQuads*SubIndexY, 0, SubsectionSizeQuads);
+
+					// Update texture data for the box that lies inside subsection
+					for( int32 SubY=SubY1;SubY<=SubY2;SubY++ )
+					{
+						for( int32 SubX=SubX1;SubX<=SubX2;SubX++ )
+						{
+							int32 LandscapeX = SubIndexX*SubsectionSizeQuads + ComponentIndexX*ComponentSizeQuads + SubX;
+							int32 LandscapeY = SubIndexY*SubsectionSizeQuads + ComponentIndexY*ComponentSizeQuads + SubY;
+
+							// Find the input data corresponding to this vertex
+							if( Component && SelectTextureData )
+							{
+								// Find the texture data corresponding to this vertex
+								int32 SizeU = Component->EditToolRenderData->DataTexture->Source.GetSizeX();
+								int32 SizeV = Component->EditToolRenderData->DataTexture->Source.GetSizeY();
+								int32 WeightmapOffsetX = Component->WeightmapScaleBias.Z * (float)SizeU;
+								int32 WeightmapOffsetY = Component->WeightmapScaleBias.W * (float)SizeV;
+
+								int32 TexX = WeightmapOffsetX + (SubsectionSizeQuads+1) * SubIndexX + SubX;
+								int32 TexY = WeightmapOffsetY + (SubsectionSizeQuads+1) * SubIndexY + SubY;
+								uint8& TexData = SelectTextureData[ TexX + TexY * SizeU ];
+
+								StoreData.Store(LandscapeX, LandscapeY, TexData);
+							}
+							else
+							{
+								StoreData.Store(LandscapeX, LandscapeY, 0);
+							}
+
+						}
+					}
+				}
+			}
+		}
+	}
+}
+
+void FLandscapeEditDataInterface::GetSelectData(const int32 X1, const int32 Y1, const int32 X2, const int32 Y2, TMap<FIntPoint, uint8>& SparseData)
+{
+	TSparseStoreData<uint8> SparseStoreData(SparseData);
+	GetSelectDataTempl(X1, Y1, X2, Y2, SparseStoreData);
+}
+
+void FLandscapeEditDataInterface::GetSelectData(const int32 X1, const int32 Y1, const int32 X2, const int32 Y2, uint8* Data, int32 Stride)
+{
+	if( Stride==0 )
+	{
+		Stride = (1+X2-X1);
+	}
+	TArrayStoreData<uint8> ArrayStoreData(X1, Y1, Data, Stride);
+	GetSelectDataTempl(X1, Y1, X2, Y2, ArrayStoreData);
+}
+
+void FLandscapeEditDataInterface::SetSelectData(int32 X1, int32 Y1, int32 X2, int32 Y2, const uint8* Data, int32 Stride)
+{
+	if( Stride==0 )
+	{
+		Stride = (1+X2-X1);
+	}
+
+	check(ComponentSizeQuads > 0);
+	// Find component range for this block of data
+	int32 ComponentIndexX1, ComponentIndexY1, ComponentIndexX2, ComponentIndexY2;
+	ALandscape::CalcComponentIndices(X1, Y1, X2, Y2, ComponentSizeQuads, ComponentIndexX1, ComponentIndexY1, ComponentIndexX2, ComponentIndexY2);
+
+	for( int32 ComponentIndexY=ComponentIndexY1;ComponentIndexY<=ComponentIndexY2;ComponentIndexY++ )
+	{
+		for( int32 ComponentIndexX=ComponentIndexX1;ComponentIndexX<=ComponentIndexX2;ComponentIndexX++ )
+		{	
+			FIntPoint ComponentKey(ComponentIndexX,ComponentIndexY);
+			ULandscapeComponent* Component = LandscapeInfo->XYtoComponentMap.FindRef(ComponentKey);
+
+			UTexture2D* DataTexture = NULL;
+			// if NULL, it was painted away
+			if( Component==NULL || Component->EditToolRenderData==NULL)
+			{
+				continue;
+			}
+			else if (Component->EditToolRenderData->DataTexture==NULL)
+			{
+				//FlushRenderingCommands();
+				// Construct Texture...
+				int32 WeightmapSize = (Component->SubsectionSizeQuads+1) * Component->NumSubsections;
+				DataTexture = Component->GetLandscapeProxy()->CreateLandscapeTexture(WeightmapSize, WeightmapSize, TEXTUREGROUP_Terrain_Weightmap, TSF_G8);
+				// Alloc dummy mips
+				ULandscapeComponent::CreateEmptyTextureMips(DataTexture, true);
+				DataTexture->PostEditChange();
+
+				//FlushRenderingCommands();
+				ZeroTexture(DataTexture);
+				FLandscapeTextureDataInfo* TexDataInfo = GetTextureDataInfo(DataTexture);
+				int32 NumMips = DataTexture->Source.GetNumMips();
+				TArray<uint8*> TextureMipData;
+				TextureMipData.AddUninitialized(NumMips);
+				for( int32 MipIdx=0;MipIdx<NumMips;MipIdx++ )
+				{
+					TextureMipData[MipIdx] = (uint8*)TexDataInfo->GetMipData(MipIdx);
+				}
+				ULandscapeComponent::UpdateDataMips(ComponentNumSubsections, SubsectionSizeQuads, DataTexture, TextureMipData, 0, 0, MAX_int32, MAX_int32, TexDataInfo);
+				
+				Component->SelectDataTexture = DataTexture;
+				ENQUEUE_UNIQUE_RENDER_COMMAND_TWOPARAMETER(
+					UpdateEditToolRenderDataDataTexture,
+					FLandscapeEditToolRenderData*, LandscapeEditToolRenderData, Component->EditToolRenderData,
+					UTexture2D*, InDataTexture, DataTexture,
+				{
+					LandscapeEditToolRenderData->DataTexture  = InDataTexture;
+				});	
+			}
+			else
+			{
+				DataTexture = Component->EditToolRenderData->DataTexture;
+			}
+
+			FLandscapeTextureDataInfo* TexDataInfo = GetTextureDataInfo(DataTexture);
+			uint8* SelectTextureData = (uint8*)TexDataInfo->GetMipData(0);
+
+			// Find the texture data corresponding to this vertex
+			int32 SizeU = DataTexture->Source.GetSizeX();
+			int32 SizeV = DataTexture->Source.GetSizeY();
+			int32 WeightmapOffsetX = Component->WeightmapScaleBias.Z * (float)SizeU;
+			int32 WeightmapOffsetY = Component->WeightmapScaleBias.W * (float)SizeV;
+
+			// Find coordinates of box that lies inside component
+			int32 ComponentX1 = FMath::Clamp<int32>(X1-ComponentIndexX*ComponentSizeQuads, 0, ComponentSizeQuads);
+			int32 ComponentY1 = FMath::Clamp<int32>(Y1-ComponentIndexY*ComponentSizeQuads, 0, ComponentSizeQuads);
+			int32 ComponentX2 = FMath::Clamp<int32>(X2-ComponentIndexX*ComponentSizeQuads, 0, ComponentSizeQuads);
+			int32 ComponentY2 = FMath::Clamp<int32>(Y2-ComponentIndexY*ComponentSizeQuads, 0, ComponentSizeQuads);
+
+			// Find subsection range for this box
+			int32 SubIndexX1 = FMath::Clamp<int32>((ComponentX1-1) / SubsectionSizeQuads,0,ComponentNumSubsections-1);	// -1 because we need to pick up vertices shared between subsections
+			int32 SubIndexY1 = FMath::Clamp<int32>((ComponentY1-1) / SubsectionSizeQuads,0,ComponentNumSubsections-1);
+			int32 SubIndexX2 = FMath::Clamp<int32>(ComponentX2 / SubsectionSizeQuads,0,ComponentNumSubsections-1);
+			int32 SubIndexY2 = FMath::Clamp<int32>(ComponentY2 / SubsectionSizeQuads,0,ComponentNumSubsections-1);
+
+			for( int32 SubIndexY=SubIndexY1;SubIndexY<=SubIndexY2;SubIndexY++ )
+			{
+				for( int32 SubIndexX=SubIndexX1;SubIndexX<=SubIndexX2;SubIndexX++ )
+				{
+					// Find coordinates of box that lies inside subsection
+					int32 SubX1 = FMath::Clamp<int32>(ComponentX1-SubsectionSizeQuads*SubIndexX, 0, SubsectionSizeQuads);
+					int32 SubY1 = FMath::Clamp<int32>(ComponentY1-SubsectionSizeQuads*SubIndexY, 0, SubsectionSizeQuads);
+					int32 SubX2 = FMath::Clamp<int32>(ComponentX2-SubsectionSizeQuads*SubIndexX, 0, SubsectionSizeQuads);
+					int32 SubY2 = FMath::Clamp<int32>(ComponentY2-SubsectionSizeQuads*SubIndexY, 0, SubsectionSizeQuads);
+
+					// Update texture data for the box that lies inside subsection
+					for( int32 SubY=SubY1;SubY<=SubY2;SubY++ )
+					{
+						for( int32 SubX=SubX1;SubX<=SubX2;SubX++ )
+						{
+							int32 LandscapeX = SubIndexX*SubsectionSizeQuads + ComponentIndexX*ComponentSizeQuads + SubX;
+							int32 LandscapeY = SubIndexY*SubsectionSizeQuads + ComponentIndexY*ComponentSizeQuads + SubY;
+							checkSlow( LandscapeX >= X1 && LandscapeX <= X2 );
+							checkSlow( LandscapeY >= Y1 && LandscapeY <= Y2 );
+
+							// Find the input data corresponding to this vertex
+							int32 DataIndex = (LandscapeX-X1) + Stride * (LandscapeY-Y1);
+							const uint8& Value = Data[DataIndex];
+
+							int32 TexX = WeightmapOffsetX + (SubsectionSizeQuads+1) * SubIndexX + SubX;
+							int32 TexY = WeightmapOffsetY + (SubsectionSizeQuads+1) * SubIndexY + SubY;
+							uint8& TexData = SelectTextureData[ TexX + TexY * SizeU ];
+
+							TexData = Value;
+						}
+					}
+
+					// Record the areas of the texture we need to re-upload
+					int32 TexX1 = WeightmapOffsetX + (SubsectionSizeQuads+1) * SubIndexX + SubX1;
+					int32 TexY1 = WeightmapOffsetY + (SubsectionSizeQuads+1) * SubIndexY + SubY1;
+					int32 TexX2 = WeightmapOffsetX + (SubsectionSizeQuads+1) * SubIndexX + SubX2;
+					int32 TexY2 = WeightmapOffsetY + (SubsectionSizeQuads+1) * SubIndexY + SubY2;
+					TexDataInfo->AddMipUpdateRegion(0,TexX1,TexY1,TexX2,TexY2);
+				}
+			}
+			// Update mipmaps
+			int32 NumMips = DataTexture->Source.GetNumMips();
+			TArray<uint8*> TextureMipData;
+			TextureMipData.AddUninitialized(NumMips);
+			for( int32 MipIdx=0;MipIdx<NumMips;MipIdx++ )
+			{
+				TextureMipData[MipIdx] = (uint8*)TexDataInfo->GetMipData(MipIdx);
+			}
+			ULandscapeComponent::UpdateDataMips(ComponentNumSubsections, SubsectionSizeQuads, DataTexture, TextureMipData, ComponentX1, ComponentY1, ComponentX2, ComponentY2, TexDataInfo);
+		}
+	}
+}
+
+template<typename T>
+void FLandscapeEditDataInterface::SetXYOffsetDataTempl(int32 X1, int32 Y1, int32 X2, int32 Y2, const T* Data, int32 Stride)
+{
+	if( Stride==0 )
+	{
+		Stride = (1+X2-X1);
+	}
+
+	check(ComponentSizeQuads > 0);
+	// Find component range for this block of data
+	int32 ComponentIndexX1, ComponentIndexY1, ComponentIndexX2, ComponentIndexY2;
+	ALandscape::CalcComponentIndices(X1, Y1, X2, Y2, ComponentSizeQuads, ComponentIndexX1, ComponentIndexY1, ComponentIndexX2, ComponentIndexY2);
+
+	const FColor DefaultValue(128, 0, 128, 0);
+
+	for( int32 ComponentIndexY=ComponentIndexY1;ComponentIndexY<=ComponentIndexY2;ComponentIndexY++ )
+	{
+		for( int32 ComponentIndexX=ComponentIndexX1;ComponentIndexX<=ComponentIndexX2;ComponentIndexX++ )
+		{	
+			FIntPoint ComponentKey(ComponentIndexX,ComponentIndexY);
+			ULandscapeComponent* Component = LandscapeInfo->XYtoComponentMap.FindRef(ComponentKey);
+
+			UTexture2D* XYOffsetTexture = NULL;
+			if( Component==NULL)
+			{
+				continue;
+			}
+			else
+			{
+				if (Component->XYOffsetmapTexture == NULL)
+				{
+					Component->Modify();
+					//FlushRenderingCommands();
+					// Construct Texture...
+					int32 WeightmapSize = (Component->SubsectionSizeQuads+1) * Component->NumSubsections;
+					XYOffsetTexture = Component->GetLandscapeProxy()->CreateLandscapeTexture(WeightmapSize, WeightmapSize, TEXTUREGROUP_Terrain_Weightmap, TSF_BGRA8);
+					// Alloc dummy mips
+					ULandscapeComponent::CreateEmptyTextureMips(XYOffsetTexture, true);
+					XYOffsetTexture->PostEditChange();
+
+					//FlushRenderingCommands();
+					SetTextureValue(XYOffsetTexture, DefaultValue);
+					FLandscapeTextureDataInfo* TexDataInfo = GetTextureDataInfo(XYOffsetTexture);
+					int32 NumMips = XYOffsetTexture->Source.GetNumMips();
+					TArray<FColor*> TextureMipData;
+					TextureMipData.AddUninitialized(NumMips);
+					for( int32 MipIdx=0;MipIdx<NumMips;MipIdx++ )
+					{
+						TextureMipData[MipIdx] = (FColor*)TexDataInfo->GetMipData(MipIdx);
+					}
+					ULandscapeComponent::UpdateWeightmapMips(ComponentNumSubsections, SubsectionSizeQuads, XYOffsetTexture, TextureMipData, 0, 0, MAX_int32, MAX_int32, TexDataInfo);
+
+					Component->XYOffsetmapTexture = XYOffsetTexture;
+					FComponentReregisterContext ReregisterContext(Component);
+				}
+				else
+				{
+					XYOffsetTexture = Component->XYOffsetmapTexture;
+				}
+			}
+
+			FLandscapeTextureDataInfo* TexDataInfo = GetTextureDataInfo(XYOffsetTexture);
+			FColor* XYOffsetTextureData = (FColor*)TexDataInfo->GetMipData(0);
+
+			// Find the texture data corresponding to this vertex
+			int32 SizeU = XYOffsetTexture->Source.GetSizeX();
+			int32 SizeV = XYOffsetTexture->Source.GetSizeY();
+			int32 WeightmapOffsetX = Component->WeightmapScaleBias.Z * (float)SizeU;
+			int32 WeightmapOffsetY = Component->WeightmapScaleBias.W * (float)SizeV;
+
+			// Find coordinates of box that lies inside component
+			int32 ComponentX1 = FMath::Clamp<int32>(X1-ComponentIndexX*ComponentSizeQuads, 0, ComponentSizeQuads);
+			int32 ComponentY1 = FMath::Clamp<int32>(Y1-ComponentIndexY*ComponentSizeQuads, 0, ComponentSizeQuads);
+			int32 ComponentX2 = FMath::Clamp<int32>(X2-ComponentIndexX*ComponentSizeQuads, 0, ComponentSizeQuads);
+			int32 ComponentY2 = FMath::Clamp<int32>(Y2-ComponentIndexY*ComponentSizeQuads, 0, ComponentSizeQuads);
+
+			// Find subsection range for this box
+			int32 SubIndexX1 = FMath::Clamp<int32>((ComponentX1-1) / SubsectionSizeQuads,0,ComponentNumSubsections-1);	// -1 because we need to pick up vertices shared between subsections
+			int32 SubIndexY1 = FMath::Clamp<int32>((ComponentY1-1) / SubsectionSizeQuads,0,ComponentNumSubsections-1);
+			int32 SubIndexX2 = FMath::Clamp<int32>(ComponentX2 / SubsectionSizeQuads,0,ComponentNumSubsections-1);
+			int32 SubIndexY2 = FMath::Clamp<int32>(ComponentY2 / SubsectionSizeQuads,0,ComponentNumSubsections-1);
+
+			for( int32 SubIndexY=SubIndexY1;SubIndexY<=SubIndexY2;SubIndexY++ )
+			{
+				for( int32 SubIndexX=SubIndexX1;SubIndexX<=SubIndexX2;SubIndexX++ )
+				{
+					// Find coordinates of box that lies inside subsection
+					int32 SubX1 = FMath::Clamp<int32>(ComponentX1-SubsectionSizeQuads*SubIndexX, 0, SubsectionSizeQuads);
+					int32 SubY1 = FMath::Clamp<int32>(ComponentY1-SubsectionSizeQuads*SubIndexY, 0, SubsectionSizeQuads);
+					int32 SubX2 = FMath::Clamp<int32>(ComponentX2-SubsectionSizeQuads*SubIndexX, 0, SubsectionSizeQuads);
+					int32 SubY2 = FMath::Clamp<int32>(ComponentY2-SubsectionSizeQuads*SubIndexY, 0, SubsectionSizeQuads);
+
+					// Update texture data for the box that lies inside subsection
+					for( int32 SubY=SubY1;SubY<=SubY2;SubY++ )
+					{
+						for( int32 SubX=SubX1;SubX<=SubX2;SubX++ )
+						{
+							int32 LandscapeX = SubIndexX*SubsectionSizeQuads + ComponentIndexX*ComponentSizeQuads + SubX;
+							int32 LandscapeY = SubIndexY*SubsectionSizeQuads + ComponentIndexY*ComponentSizeQuads + SubY;
+							checkSlow( LandscapeX >= X1 && LandscapeX <= X2 );
+							checkSlow( LandscapeY >= Y1 && LandscapeY <= Y2 );
+
+							// Find the input data corresponding to this vertex
+							int32 DataIndex = (LandscapeX-X1) + Stride * (LandscapeY-Y1);
+							const T& Value = Data[DataIndex];
+
+							int32 TexX = WeightmapOffsetX + (SubsectionSizeQuads+1) * SubIndexX + SubX;
+							int32 TexY = WeightmapOffsetY + (SubsectionSizeQuads+1) * SubIndexY + SubY;
+							FColor& TexData = XYOffsetTextureData[ TexX + TexY * SizeU ];
+
+							uint16 XOffset = FMath::Clamp<uint16>(Value.X * LANDSCAPE_INV_XYOFFSET_SCALE + 32768.f, 0, 65535);
+							uint16 YOffset = FMath::Clamp<uint16>(Value.Y * LANDSCAPE_INV_XYOFFSET_SCALE + 32768.f, 0, 65535);
+
+							TexData.R = XOffset >> 8;
+							TexData.G = XOffset & 255;
+							TexData.B = YOffset >> 8;
+							TexData.A = YOffset & 255;
+						}
+					}
+
+					// Record the areas of the texture we need to re-upload
+					int32 TexX1 = WeightmapOffsetX + (SubsectionSizeQuads+1) * SubIndexX + SubX1;
+					int32 TexY1 = WeightmapOffsetY + (SubsectionSizeQuads+1) * SubIndexY + SubY1;
+					int32 TexX2 = WeightmapOffsetX + (SubsectionSizeQuads+1) * SubIndexX + SubX2;
+					int32 TexY2 = WeightmapOffsetY + (SubsectionSizeQuads+1) * SubIndexY + SubY2;
+					TexDataInfo->AddMipUpdateRegion(0,TexX1,TexY1,TexX2,TexY2);
+				}
+			}
+
+			// Update mipmaps
+			int32 NumMips = XYOffsetTexture->Source.GetNumMips();
+			TArray<FColor*> TextureMipData;
+			TextureMipData.AddUninitialized(NumMips);
+			for( int32 MipIdx=0;MipIdx<NumMips;MipIdx++ )
+			{
+				TextureMipData[MipIdx] = (FColor*)TexDataInfo->GetMipData(MipIdx);
+			}
+			ULandscapeComponent::UpdateWeightmapMips(ComponentNumSubsections, SubsectionSizeQuads, XYOffsetTexture, TextureMipData, ComponentX1, ComponentY1, ComponentX2, ComponentY2, TexDataInfo);
+		}
+	}
+}
+
+void FLandscapeEditDataInterface::SetXYOffsetData(int32 X1, int32 Y1, int32 X2, int32 Y2, const FVector2D* Data, int32 Stride)
+{
+	SetXYOffsetDataTempl<FVector2D>(X1, Y1, X2, Y2, Data, Stride);
+}
+
+void FLandscapeEditDataInterface::SetXYOffsetData(int32 X1, int32 Y1, int32 X2, int32 Y2, const FVector* Data, int32 Stride)
+{
+	SetXYOffsetDataTempl<FVector>(X1, Y1, X2, Y2, Data, Stride);
+}
+
+FVector2D FLandscapeEditDataInterface::GetXYOffsetmapData(const ULandscapeComponent* Component, int32 TexU, int32 TexV, FColor* TextureData/* = NULL*/)
+{
+	check(Component);
+	if (!TextureData && Component->XYOffsetmapTexture)
+	{
+		FLandscapeTextureDataInfo* TexDataInfo = GetTextureDataInfo(Component->XYOffsetmapTexture);
+		TextureData = (FColor*)TexDataInfo->GetMipData(0);	
+	}
+
+	if (TextureData)
+	{
+		int32 SizeU = Component->NumSubsections * (Component->SubsectionSizeQuads + 1);
+		int32 SizeV = Component->NumSubsections * (Component->SubsectionSizeQuads + 1);
+
+		int32 TexX = TexU;
+		int32 TexY = TexV;
+		FColor& TexData = TextureData[ TexX + TexY * SizeU ];
+		return FVector2D(((TexData.R * 256.0 + TexData.G) - 32768.0) * LANDSCAPE_XYOFFSET_SCALE, ((TexData.B * 256.0 + TexData.A) - 32768.0) * LANDSCAPE_XYOFFSET_SCALE );
+	}
+	return FVector2D::ZeroVector;
+}
+
+// XYOffset Interpolation version
+template<typename TStoreData>
+void FLandscapeEditDataInterface::GetXYOffsetDataTempl(int32& ValidX1, int32& ValidY1, int32& ValidX2, int32& ValidY2, TStoreData& StoreData)
+{
+	// Copy variables
+	int32 X1 = ValidX1, X2 = ValidX2, Y1 = ValidY1, Y2 = ValidY2;
+	ValidX1 = INT_MAX; ValidX2 = INT_MIN; ValidY1 = INT_MAX; ValidY2 = INT_MIN;
+
+	int32 ComponentIndexX1, ComponentIndexY1, ComponentIndexX2, ComponentIndexY2;
+	ALandscape::CalcComponentIndicesOverlap(X1, Y1, X2, Y2, ComponentSizeQuads, ComponentIndexX1, ComponentIndexY1, ComponentIndexX2, ComponentIndexY2);
+
+	int32 ComponentSizeX = ComponentIndexX2 - ComponentIndexX1 + 1;
+	int32 ComponentSizeY = ComponentIndexY2 - ComponentIndexY1 + 1;
+
+	// Neighbor Components
+	ULandscapeComponent* BorderComponent[4] = { 0, 0, 0, 0 };
+	ULandscapeComponent* CornerComponent[4] = { 0, 0, 0, 0 };
+	bool NoBorderX1 = false, NoBorderX2 = false;
+	TArray<bool> NoBorderY1, NoBorderY2, ComponentDataExist;
+	TArray<ULandscapeComponent*> BorderComponentY1, BorderComponentY2;
+	ComponentDataExist.Empty(ComponentSizeX*ComponentSizeY);
+	ComponentDataExist.AddZeroed(ComponentSizeX*ComponentSizeY);
+	bool bHasMissingValue = false;
+
+	FLandscapeTextureDataInfo* NeighborTexDataInfo[4] = { 0, 0, 0, 0 };
+	FColor* NeighborXYOffsetmapTextureData[4] = { 0, 0, 0, 0 };
+	FVector2D CornerValues[4] = { FVector2D::ZeroVector, FVector2D::ZeroVector, FVector2D::ZeroVector, FVector2D::ZeroVector };
+
+	int32 EdgeCoord = (SubsectionSizeQuads + 1) * ComponentNumSubsections - 1; //ComponentSizeQuads;
+
+	TArray<FColor> EmptyXYOffset;
+	int32 XYOffsetSize = (LandscapeInfo->SubsectionSizeQuads + 1) * LandscapeInfo->ComponentNumSubsections;
+	XYOffsetSize = XYOffsetSize * XYOffsetSize;
+	EmptyXYOffset.Empty(XYOffsetSize);
+	for (int32 i = 0; i < XYOffsetSize; ++i)
+	{
+		EmptyXYOffset.Add(FColor(128, 0, 128, 0));
+	}
+
+	// initial loop....
+	for (int32 ComponentIndexY = ComponentIndexY1; ComponentIndexY <= ComponentIndexY2; ComponentIndexY++)
+	{
+		NoBorderX1 = false;
+		NoBorderX2 = false;
+		BorderComponent[0] = BorderComponent[1] = NULL;
+		for (int32 ComponentIndexX = ComponentIndexX1; ComponentIndexX <= ComponentIndexX2; ComponentIndexX++)
+		{
+			BorderComponent[2] = BorderComponent[3] = NULL;
+			int32 ComponentIndexXY = ComponentSizeX*(ComponentIndexY - ComponentIndexY1) + ComponentIndexX - ComponentIndexX1;
+			int32 ComponentIndexXX = ComponentIndexX - ComponentIndexX1;
+			int32 ComponentIndexYY = ComponentIndexY - ComponentIndexY1;
+			ComponentDataExist[ComponentIndexXY] = false;
+			ULandscapeComponent* Component = LandscapeInfo->XYtoComponentMap.FindRef(FIntPoint(ComponentIndexX, ComponentIndexY));
+
+			FLandscapeTextureDataInfo* TexDataInfo = NULL;
+			FColor* XYOffsetmapTextureData = NULL;
+			uint8 CornerSet = 0;
+			bool ExistLeft = ComponentIndexXX > 0 && ComponentDataExist[ComponentIndexXX - 1 + ComponentIndexYY * ComponentSizeX];
+			bool ExistUp = ComponentIndexYY > 0 && ComponentDataExist[ComponentIndexXX + (ComponentIndexYY - 1) * ComponentSizeX];
+
+			if (Component)
+			{
+				if (Component->XYOffsetmapTexture)
+				{
+					TexDataInfo = GetTextureDataInfo(Component->XYOffsetmapTexture);
+					XYOffsetmapTextureData = (FColor*)TexDataInfo->GetMipData(0);
+				}
+				else
+				{
+					XYOffsetmapTextureData = EmptyXYOffset.GetData();
+				}
+				ComponentDataExist[ComponentIndexXY] = true;
+				// Update valid region
+				ValidX1 = FMath::Min<int32>(Component->GetSectionBase().X, ValidX1);
+				ValidX2 = FMath::Max<int32>(Component->GetSectionBase().X + ComponentSizeQuads, ValidX2);
+				ValidY1 = FMath::Min<int32>(Component->GetSectionBase().Y, ValidY1);
+				ValidY2 = FMath::Max<int32>(Component->GetSectionBase().Y + ComponentSizeQuads, ValidY2);
+			}
+			else
+			{
+				if (!bHasMissingValue)
+				{
+					NoBorderY1.Empty(ComponentSizeX);
+					NoBorderY2.Empty(ComponentSizeX);
+					NoBorderY1.AddZeroed(ComponentSizeX);
+					NoBorderY2.AddZeroed(ComponentSizeX);
+					BorderComponentY1.Empty(ComponentSizeX);
+					BorderComponentY2.Empty(ComponentSizeX);
+					BorderComponentY1.AddZeroed(ComponentSizeX);
+					BorderComponentY2.AddZeroed(ComponentSizeX);
+					bHasMissingValue = true;
+				}
+
+				// Search for neighbor component for interpolation
+				bool bShouldSearchX = (BorderComponent[1] && BorderComponent[1]->GetSectionBase().X / ComponentSizeQuads <= ComponentIndexX);
+				bool bShouldSearchY = (BorderComponentY2[ComponentIndexXX] && BorderComponentY2[ComponentIndexXX]->GetSectionBase().Y / ComponentSizeQuads <= ComponentIndexY);
+				// Search for left-closest component
+				if (bShouldSearchX || (!NoBorderX1 && !BorderComponent[0]))
+				{
+					NoBorderX1 = true;
+					for (int32 X = ComponentIndexX - 1; X >= ComponentIndexX1; X--)
+					{
+						BorderComponent[0] = LandscapeInfo->XYtoComponentMap.FindRef(FIntPoint(X, ComponentIndexY));
+						if (BorderComponent[0])
+						{
+							NoBorderX1 = false;
+							if (BorderComponent[0]->XYOffsetmapTexture)
+							{
+								NeighborTexDataInfo[0] = GetTextureDataInfo(BorderComponent[0]->XYOffsetmapTexture);
+								NeighborXYOffsetmapTextureData[0] = (FColor*)NeighborTexDataInfo[0]->GetMipData(0);
+							}
+							else
+							{
+								NeighborXYOffsetmapTextureData[0] = EmptyXYOffset.GetData();
+							}
+							break;
+						}
+					}
+				}
+				// Search for right-closest component
+				if (bShouldSearchX || (!NoBorderX2 && !BorderComponent[1]))
+				{
+					NoBorderX2 = true;
+					for (int32 X = ComponentIndexX + 1; X <= ComponentIndexX2; X++)
+					{
+						BorderComponent[1] = LandscapeInfo->XYtoComponentMap.FindRef(FIntPoint(X, ComponentIndexY));
+						if (BorderComponent[1])
+						{
+							NoBorderX2 = false;
+							if (BorderComponent[1]->XYOffsetmapTexture)
+							{
+								NeighborTexDataInfo[1] = GetTextureDataInfo(BorderComponent[1]->XYOffsetmapTexture);
+								NeighborXYOffsetmapTextureData[1] = (FColor*)NeighborTexDataInfo[1]->GetMipData(0);
+							}
+							else
+							{
+								NeighborXYOffsetmapTextureData[1] = EmptyXYOffset.GetData();
+							}
+							break;
+						}
+					}
+				}
+				// Search for up-closest component
+				if (bShouldSearchY || (!NoBorderY1[ComponentIndexXX] && !BorderComponentY1[ComponentIndexXX]))
+				{
+					NoBorderY1[ComponentIndexXX] = true;
+					for (int32 Y = ComponentIndexY - 1; Y >= ComponentIndexY1; Y--)
+					{
+						BorderComponentY1[ComponentIndexXX] = BorderComponent[2] = LandscapeInfo->XYtoComponentMap.FindRef(FIntPoint(ComponentIndexX, Y));
+						if (BorderComponent[2])
+						{
+							NoBorderY1[ComponentIndexXX] = false;
+							if (BorderComponent[2]->XYOffsetmapTexture)
+							{
+								NeighborTexDataInfo[2] = GetTextureDataInfo(BorderComponent[2]->XYOffsetmapTexture);
+								NeighborXYOffsetmapTextureData[2] = (FColor*)NeighborTexDataInfo[2]->GetMipData(0);
+							}
+							else
+							{
+								NeighborXYOffsetmapTextureData[2] = EmptyXYOffset.GetData();
+							}
+							break;
+						}
+					}
+				}
+				else
+				{
+					BorderComponent[2] = BorderComponentY1[ComponentIndexXX];
+					if (BorderComponent[2])
+					{
+						if (BorderComponent[2]->XYOffsetmapTexture)
+						{
+							NeighborTexDataInfo[2] = GetTextureDataInfo(BorderComponent[2]->XYOffsetmapTexture);
+							NeighborXYOffsetmapTextureData[2] = (FColor*)NeighborTexDataInfo[2]->GetMipData(0);
+						}
+						else
+						{
+							NeighborXYOffsetmapTextureData[2] = EmptyXYOffset.GetData();
+						}
+
+					}
+				}
+				// Search for bottom-closest component
+				if (bShouldSearchY || (!NoBorderY2[ComponentIndexXX] && !BorderComponentY2[ComponentIndexXX]))
+				{
+					NoBorderY2[ComponentIndexXX] = true;
+					for (int32 Y = ComponentIndexY + 1; Y <= ComponentIndexY2; Y++)
+					{
+						BorderComponentY2[ComponentIndexXX] = BorderComponent[3] = LandscapeInfo->XYtoComponentMap.FindRef(FIntPoint(ComponentIndexX, Y));
+						if (BorderComponent[3])
+						{
+							NoBorderY2[ComponentIndexXX] = false;
+							if (BorderComponent[3]->XYOffsetmapTexture)
+							{
+								NeighborTexDataInfo[3] = GetTextureDataInfo(BorderComponent[3]->XYOffsetmapTexture);
+								NeighborXYOffsetmapTextureData[3] = (FColor*)NeighborTexDataInfo[3]->GetMipData(0);
+							}
+							else
+							{
+								NeighborXYOffsetmapTextureData[3] = EmptyXYOffset.GetData();
+							}
+							break;
+						}
+					}
+				}
+				else
+				{
+					BorderComponent[3] = BorderComponentY2[ComponentIndexXX];
+					if (BorderComponent[3])
+					{
+						if (BorderComponent[3]->XYOffsetmapTexture)
+						{
+							NeighborTexDataInfo[3] = GetTextureDataInfo(BorderComponent[3]->XYOffsetmapTexture);
+							NeighborXYOffsetmapTextureData[3] = (FColor*)NeighborTexDataInfo[3]->GetMipData(0);
+						}
+						else
+						{
+							NeighborXYOffsetmapTextureData[3] = EmptyXYOffset.GetData();
+						}
+					}
+				}
+
+				CornerComponent[0] = ComponentIndexX >= ComponentIndexX1 && ComponentIndexY >= ComponentIndexY1 ?
+					LandscapeInfo->XYtoComponentMap.FindRef(FIntPoint((ComponentIndexX - 1), (ComponentIndexY - 1))) : NULL;
+				CornerComponent[1] = ComponentIndexX <= ComponentIndexX2 && ComponentIndexY >= ComponentIndexY1 ?
+					LandscapeInfo->XYtoComponentMap.FindRef(FIntPoint((ComponentIndexX + 1), (ComponentIndexY - 1))) : NULL;
+				CornerComponent[2] = ComponentIndexX >= ComponentIndexX1 && ComponentIndexY <= ComponentIndexY2 ?
+					LandscapeInfo->XYtoComponentMap.FindRef(FIntPoint((ComponentIndexX - 1), (ComponentIndexY + 1))) : NULL;
+				CornerComponent[3] = ComponentIndexX <= ComponentIndexX2 && ComponentIndexY <= ComponentIndexY2 ?
+					LandscapeInfo->XYtoComponentMap.FindRef(FIntPoint((ComponentIndexX + 1), (ComponentIndexY + 1))) : NULL;
+
+				if (CornerComponent[0])
+				{
+					CornerSet |= 1;
+					CornerValues[0] = GetXYOffsetmapData(CornerComponent[0], EdgeCoord, EdgeCoord);
+				}
+				else if ((ExistLeft || ExistUp) && X1 <= ComponentIndexX*ComponentSizeQuads && Y1 <= ComponentIndexY*ComponentSizeQuads)
+				{
+					CornerSet |= 1;
+					CornerValues[0] = FVector2D(StoreData.Load(ComponentIndexX*ComponentSizeQuads, ComponentIndexY*ComponentSizeQuads));
+				}
+				else if (BorderComponent[0])
+				{
+					CornerSet |= 1;
+					CornerValues[0] = GetXYOffsetmapData(BorderComponent[0], EdgeCoord, 0, NeighborXYOffsetmapTextureData[0]);
+				}
+				else if (BorderComponent[2])
+				{
+					CornerSet |= 1;
+					CornerValues[0] = GetXYOffsetmapData(BorderComponent[2], 0, EdgeCoord, NeighborXYOffsetmapTextureData[2]);
+				}
+
+				if (CornerComponent[1])
+				{
+					CornerSet |= 1 << 1;
+					CornerValues[1] = GetXYOffsetmapData(CornerComponent[1], 0, EdgeCoord);
+				}
+				else if (ExistUp && X2 >= (ComponentIndexX + 1)*ComponentSizeQuads)
+				{
+					CornerSet |= 1 << 1;
+					CornerValues[1] = FVector2D(StoreData.Load((ComponentIndexX + 1)*ComponentSizeQuads, ComponentIndexY*ComponentSizeQuads));
+				}
+				else if (BorderComponent[1])
+				{
+					CornerSet |= 1 << 1;
+					CornerValues[1] = GetXYOffsetmapData(BorderComponent[1], 0, 0, NeighborXYOffsetmapTextureData[1]);
+				}
+				else if (BorderComponent[2])
+				{
+					CornerSet |= 1 << 1;
+					CornerValues[1] = GetXYOffsetmapData(BorderComponent[2], EdgeCoord, EdgeCoord, NeighborXYOffsetmapTextureData[2]);
+				}
+
+				if (CornerComponent[2])
+				{
+					CornerSet |= 1 << 2;
+					CornerValues[2] = GetXYOffsetmapData(CornerComponent[2], EdgeCoord, 0);
+				}
+				else if (ExistLeft && Y2 >= (ComponentIndexY + 1)*ComponentSizeQuads) // Use data already stored for 0, 2
+				{
+					CornerSet |= 1 << 2;
+					CornerValues[2] = FVector2D(StoreData.Load(ComponentIndexX*ComponentSizeQuads, (ComponentIndexY + 1)*ComponentSizeQuads));
+				}
+				else if (BorderComponent[0])
+				{
+					CornerSet |= 1 << 2;
+					CornerValues[2] = GetXYOffsetmapData(BorderComponent[0], EdgeCoord, EdgeCoord, NeighborXYOffsetmapTextureData[0]);
+				}
+				else if (BorderComponent[3])
+				{
+					CornerSet |= 1 << 2;
+					CornerValues[2] = GetXYOffsetmapData(BorderComponent[3], 0, 0, NeighborXYOffsetmapTextureData[3]);
+				}
+
+				if (CornerComponent[3])
+				{
+					CornerSet |= 1 << 3;
+					CornerValues[3] = GetXYOffsetmapData(CornerComponent[3], 0, 0);
+				}
+				else if (BorderComponent[1])
+				{
+					CornerSet |= 1 << 3;
+					CornerValues[3] = GetXYOffsetmapData(BorderComponent[1], 0, EdgeCoord, NeighborXYOffsetmapTextureData[1]);
+				}
+				else if (BorderComponent[3])
+				{
+					CornerSet |= 1 << 3;
+					CornerValues[3] = GetXYOffsetmapData(BorderComponent[3], EdgeCoord, 0, NeighborXYOffsetmapTextureData[3]);
+				}
+
+				FillCornerValues(CornerSet, CornerValues);
+				ComponentDataExist[ComponentIndexXY] = ExistLeft || ExistUp || (BorderComponent[0] || BorderComponent[1] || BorderComponent[2] || BorderComponent[3]) || CornerSet;
+			}
+
+			if (!ComponentDataExist[ComponentIndexXY])
+			{
+				continue;
+			}
+
+			// Find coordinates of box that lies inside component
+			int32 ComponentX1 = FMath::Clamp<int32>(X1 - ComponentIndexX*ComponentSizeQuads, 0, ComponentSizeQuads);
+			int32 ComponentY1 = FMath::Clamp<int32>(Y1 - ComponentIndexY*ComponentSizeQuads, 0, ComponentSizeQuads);
+			int32 ComponentX2 = FMath::Clamp<int32>(X2 - ComponentIndexX*ComponentSizeQuads, 0, ComponentSizeQuads);
+			int32 ComponentY2 = FMath::Clamp<int32>(Y2 - ComponentIndexY*ComponentSizeQuads, 0, ComponentSizeQuads);
+
+			// Find subsection range for this box
+			int32 SubIndexX1 = FMath::Clamp<int32>((ComponentX1 - 1) / SubsectionSizeQuads, 0, ComponentNumSubsections - 1);	// -1 because we need to pick up vertices shared between subsections
+			int32 SubIndexY1 = FMath::Clamp<int32>((ComponentY1 - 1) / SubsectionSizeQuads, 0, ComponentNumSubsections - 1);
+			int32 SubIndexX2 = FMath::Clamp<int32>(ComponentX2 / SubsectionSizeQuads, 0, ComponentNumSubsections - 1);
+			int32 SubIndexY2 = FMath::Clamp<int32>(ComponentY2 / SubsectionSizeQuads, 0, ComponentNumSubsections - 1);
+
+			for (int32 SubIndexY = SubIndexY1; SubIndexY <= SubIndexY2; SubIndexY++)
+			{
+				for (int32 SubIndexX = SubIndexX1; SubIndexX <= SubIndexX2; SubIndexX++)
+				{
+					// Find coordinates of box that lies inside subsection
+					int32 SubX1 = FMath::Clamp<int32>(ComponentX1 - SubsectionSizeQuads*SubIndexX, 0, SubsectionSizeQuads);
+					int32 SubY1 = FMath::Clamp<int32>(ComponentY1 - SubsectionSizeQuads*SubIndexY, 0, SubsectionSizeQuads);
+					int32 SubX2 = FMath::Clamp<int32>(ComponentX2 - SubsectionSizeQuads*SubIndexX, 0, SubsectionSizeQuads);
+					int32 SubY2 = FMath::Clamp<int32>(ComponentY2 - SubsectionSizeQuads*SubIndexY, 0, SubsectionSizeQuads);
+
+					// Update texture data for the box that lies inside subsection
+					for (int32 SubY = SubY1; SubY <= SubY2; SubY++)
+					{
+						for (int32 SubX = SubX1; SubX <= SubX2; SubX++)
+						{
+							int32 LandscapeX = SubIndexX*SubsectionSizeQuads + ComponentIndexX*ComponentSizeQuads + SubX;
+							int32 LandscapeY = SubIndexY*SubsectionSizeQuads + ComponentIndexY*ComponentSizeQuads + SubY;
+
+							// Find the input data corresponding to this vertex
+							if (Component)
+							{
+								// Find the texture data corresponding to this vertex
+								FVector2D XYOffset = GetXYOffsetmapData(Component, (SubsectionSizeQuads + 1) * SubIndexX + SubX, (SubsectionSizeQuads + 1) * SubIndexY + SubY, XYOffsetmapTextureData);
+								StoreData.Store(LandscapeX, LandscapeY, XYOffset);
+							}
+							else
+							{
+								// Find the texture data corresponding to this vertex
+								FVector2D Value[4] = { FVector2D::ZeroVector, FVector2D::ZeroVector, FVector2D::ZeroVector, FVector2D::ZeroVector };
+								int32 Dist[4] = { INT_MAX, INT_MAX, INT_MAX, INT_MAX };
+								FVector2D ValueX = FVector2D::ZeroVector, ValueY = FVector2D::ZeroVector;
+								bool Exist[4] = { false, false, false, false };
+
+								// Use data already stored for 0, 2
+								if (ExistLeft)
+								{
+									Value[0] = FVector2D(StoreData.Load(ComponentIndexX*ComponentSizeQuads, LandscapeY));
+									Dist[0] = LandscapeX - (ComponentIndexX*ComponentSizeQuads);
+									Exist[0] = true;
+								}
+								else if (BorderComponent[0])
+								{
+									Value[0] = GetXYOffsetmapData(BorderComponent[0], EdgeCoord, (SubsectionSizeQuads + 1) * SubIndexY + SubY, NeighborXYOffsetmapTextureData[0]);
+									Dist[0] = LandscapeX - (BorderComponent[0]->GetSectionBase().X + ComponentSizeQuads);
+									Exist[0] = true;
+								}
+								else
+								{
+									if ((CornerSet & 1) && (CornerSet & 1 << 2))
+									{
+										int32 Dist1 = LandscapeY - (ComponentIndexY*ComponentSizeQuads);
+										int32 Dist2 = ((ComponentIndexY + 1)*ComponentSizeQuads) - LandscapeY;
+										Value[0] = (Dist2 * CornerValues[0] + Dist1 * CornerValues[2]) / (Dist1 + Dist2);
+										Dist[0] = LandscapeX - (ComponentIndexX*ComponentSizeQuads);
+										Exist[0] = true;
+									}
+								}
+
+								if (BorderComponent[1])
+								{
+									Value[1] = GetXYOffsetmapData(BorderComponent[1], 0, (SubsectionSizeQuads + 1) * SubIndexY + SubY, NeighborXYOffsetmapTextureData[1]);
+									Dist[1] = (BorderComponent[1]->GetSectionBase().X) - LandscapeX;
+									Exist[1] = true;
+								}
+								else
+								{
+									if ((CornerSet & 1 << 1) && (CornerSet & 1 << 3))
+									{
+										int32 Dist1 = LandscapeY - (ComponentIndexY*ComponentSizeQuads);
+										int32 Dist2 = ((ComponentIndexY + 1)*ComponentSizeQuads) - LandscapeY;
+										Value[1] = (Dist2 * CornerValues[1] + Dist1 * CornerValues[3]) / (Dist1 + Dist2);
+										Dist[1] = (ComponentIndexX + 1)*ComponentSizeQuads - LandscapeX;
+										Exist[1] = true;
+									}
+								}
+
+								if (ExistUp)
+								{
+									Value[2] = FVector2D(StoreData.Load(LandscapeX, ComponentIndexY*ComponentSizeQuads));
+									Dist[2] = LandscapeY - (ComponentIndexY*ComponentSizeQuads);
+									Exist[2] = true;
+								}
+								else if (BorderComponent[2])
+								{
+									Value[2] = GetXYOffsetmapData(BorderComponent[2], (SubsectionSizeQuads + 1) * SubIndexX + SubX, EdgeCoord, NeighborXYOffsetmapTextureData[2]);
+									Dist[2] = LandscapeY - (BorderComponent[2]->GetSectionBase().Y + ComponentSizeQuads);
+									Exist[2] = true;
+								}
+								else
+								{
+									if ((CornerSet & 1) && (CornerSet & 1 << 1))
+									{
+										int32 Dist1 = LandscapeX - (ComponentIndexX*ComponentSizeQuads);
+										int32 Dist2 = (ComponentIndexX + 1)*ComponentSizeQuads - LandscapeX;
+										Value[2] = (Dist2 * CornerValues[0] + Dist1 * CornerValues[1]) / (Dist1 + Dist2);
+										Dist[2] = LandscapeY - (ComponentIndexY*ComponentSizeQuads);
+										Exist[2] = true;
+									}
+								}
+
+								if (BorderComponent[3])
+								{
+									Value[3] = GetXYOffsetmapData(BorderComponent[3], (SubsectionSizeQuads + 1) * SubIndexX + SubX, 0, NeighborXYOffsetmapTextureData[3]);
+									Dist[3] = (BorderComponent[3]->GetSectionBase().Y) - LandscapeY;
+									Exist[3] = true;
+								}
+								else
+								{
+									if ((CornerSet & 1 << 2) && (CornerSet & 1 << 3))
+									{
+										int32 Dist1 = LandscapeX - (ComponentIndexX*ComponentSizeQuads);
+										int32 Dist2 = (ComponentIndexX + 1)*ComponentSizeQuads - LandscapeX;
+										Value[3] = (Dist2 * CornerValues[2] + Dist1 * CornerValues[3]) / (Dist1 + Dist2);
+										Dist[3] = (ComponentIndexY + 1)*ComponentSizeQuads - LandscapeY;
+										Exist[3] = true;
+									}
+								}
+
+								CalcInterpValue<FVector2D>(Dist, Exist, Value, ValueX, ValueY);
+
+								FVector2D FinalValue = FVector2D::ZeroVector; // Default Value
+								if ((Exist[0] || Exist[1]) && (Exist[2] || Exist[3]))
+								{
+									FinalValue = CalcValueFromValueXY<FVector2D>(Dist, ValueX, ValueY, CornerSet, CornerValues);
+								}
+								else if ((BorderComponent[0] || BorderComponent[1]))
+								{
+									FinalValue = ValueX;
+								}
+								else if ((BorderComponent[2] || BorderComponent[3]))
+								{
+									FinalValue = ValueY;
+								}
+								else if ((Exist[0] || Exist[1]))
+								{
+									FinalValue = ValueX;
+								}
+								else if ((Exist[2] || Exist[3]))
+								{
+									FinalValue = ValueY;
+								}
+
+								StoreData.Store(LandscapeX, LandscapeY, FinalValue);
+								//StoreData.StoreDefault(LandscapeX, LandscapeY);
+							}
+						}
+					}
+				}
+			}
+		}
+	}
+
+	if (bHasMissingValue)
+	{
+		CalcMissingValues<FVector2D, TStoreData, FVector2D>(X1, X2, Y1, Y2,
+			ComponentIndexX1, ComponentIndexX2, ComponentIndexY1, ComponentIndexY2,
+			ComponentSizeX, ComponentSizeY, CornerValues,
+			NoBorderY1, NoBorderY2, ComponentDataExist, StoreData);
+		// Update valid region
+		ValidX1 = FMath::Max<int32>(X1, ValidX1);
+		ValidX2 = FMath::Min<int32>(X2, ValidX2);
+		ValidY1 = FMath::Max<int32>(Y1, ValidY1);
+		ValidY2 = FMath::Min<int32>(Y2, ValidY2);
+	}
+	else
+	{
+		ValidX1 = X1;
+		ValidX2 = X2;
+		ValidY1 = Y1;
+		ValidY2 = Y2;
+	}
+}
+
+void FLandscapeEditDataInterface::GetXYOffsetData(int32& X1, int32& Y1, int32& X2, int32& Y2, FVector2D* Data, int32 Stride)
+{
+	if (Stride == 0)
+	{
+		Stride = (1 + X2 - X1);
+	}
+	TArrayStoreData<FVector2D> ArrayStoreData(X1, Y1, Data, Stride);
+	GetXYOffsetDataTempl(X1, Y1, X2, Y2, ArrayStoreData);
+}
+
+void FLandscapeEditDataInterface::GetXYOffsetData(int32& X1, int32& Y1, int32& X2, int32& Y2, TMap<FIntPoint, FVector2D>& SparseData)
+{
+	TSparseStoreData<FVector2D> SparseStoreData(SparseData);
+	GetXYOffsetDataTempl(X1, Y1, X2, Y2, SparseStoreData);
+}
+
+void FLandscapeEditDataInterface::GetXYOffsetData(int32& X1, int32& Y1, int32& X2, int32& Y2, FVector* Data, int32 Stride)
+{
+	if (Stride == 0)
+	{
+		Stride = (1 + X2 - X1);
+	}
+	TArrayStoreData<FVector> ArrayStoreData(X1, Y1, Data, Stride);
+	GetXYOffsetDataTempl(X1, Y1, X2, Y2, ArrayStoreData);
+}
+
+void FLandscapeEditDataInterface::GetXYOffsetData(int32& X1, int32& Y1, int32& X2, int32& Y2, TMap<FIntPoint, FVector>& SparseData)
+{
+	TSparseStoreData<FVector> SparseStoreData(SparseData);
+	GetXYOffsetDataTempl(X1, Y1, X2, Y2, SparseStoreData);
+}
+
+template<typename TStoreData>
+void FLandscapeEditDataInterface::GetXYOffsetDataTemplFast(const int32 X1, const int32 Y1, const int32 X2, const int32 Y2, TStoreData& StoreData)
+{
+	int32 ComponentIndexX1, ComponentIndexY1, ComponentIndexX2, ComponentIndexY2;
+	ALandscape::CalcComponentIndices(X1, Y1, X2, Y2, ComponentSizeQuads, ComponentIndexX1, ComponentIndexY1, ComponentIndexX2, ComponentIndexY2);
+
+	for( int32 ComponentIndexY=ComponentIndexY1;ComponentIndexY<=ComponentIndexY2;ComponentIndexY++ )
+	{
+		for( int32 ComponentIndexX=ComponentIndexX1;ComponentIndexX<=ComponentIndexX2;ComponentIndexX++ )
+		{		
+			ULandscapeComponent* Component = LandscapeInfo->XYtoComponentMap.FindRef(FIntPoint(ComponentIndexX,ComponentIndexY));
+
+			FLandscapeTextureDataInfo* TexDataInfo = NULL;
+			FColor* OffsetTextureData = NULL;
+			if( Component && Component->XYOffsetmapTexture )
+			{
+				TexDataInfo = GetTextureDataInfo(Component->XYOffsetmapTexture);
+				OffsetTextureData = (FColor*)TexDataInfo->GetMipData(0);
+			}
+
+			// Find coordinates of box that lies inside component
+			int32 ComponentX1 = FMath::Clamp<int32>(X1-ComponentIndexX*ComponentSizeQuads, 0, ComponentSizeQuads);
+			int32 ComponentY1 = FMath::Clamp<int32>(Y1-ComponentIndexY*ComponentSizeQuads, 0, ComponentSizeQuads);
+			int32 ComponentX2 = FMath::Clamp<int32>(X2-ComponentIndexX*ComponentSizeQuads, 0, ComponentSizeQuads);
+			int32 ComponentY2 = FMath::Clamp<int32>(Y2-ComponentIndexY*ComponentSizeQuads, 0, ComponentSizeQuads);
+
+			// Find subsection range for this box
+			int32 SubIndexX1 = FMath::Clamp<int32>((ComponentX1-1) / SubsectionSizeQuads,0,ComponentNumSubsections-1);	// -1 because we need to pick up vertices shared between subsections
+			int32 SubIndexY1 = FMath::Clamp<int32>((ComponentY1-1) / SubsectionSizeQuads,0,ComponentNumSubsections-1);
+			int32 SubIndexX2 = FMath::Clamp<int32>(ComponentX2 / SubsectionSizeQuads,0,ComponentNumSubsections-1);
+			int32 SubIndexY2 = FMath::Clamp<int32>(ComponentY2 / SubsectionSizeQuads,0,ComponentNumSubsections-1);
+
+			for( int32 SubIndexY=SubIndexY1;SubIndexY<=SubIndexY2;SubIndexY++ )
+			{
+				for( int32 SubIndexX=SubIndexX1;SubIndexX<=SubIndexX2;SubIndexX++ )
+				{
+					// Find coordinates of box that lies inside subsection
+					int32 SubX1 = FMath::Clamp<int32>(ComponentX1-SubsectionSizeQuads*SubIndexX, 0, SubsectionSizeQuads);
+					int32 SubY1 = FMath::Clamp<int32>(ComponentY1-SubsectionSizeQuads*SubIndexY, 0, SubsectionSizeQuads);
+					int32 SubX2 = FMath::Clamp<int32>(ComponentX2-SubsectionSizeQuads*SubIndexX, 0, SubsectionSizeQuads);
+					int32 SubY2 = FMath::Clamp<int32>(ComponentY2-SubsectionSizeQuads*SubIndexY, 0, SubsectionSizeQuads);
+
+					// Update texture data for the box that lies inside subsection
+					for( int32 SubY=SubY1;SubY<=SubY2;SubY++ )
+					{
+						for( int32 SubX=SubX1;SubX<=SubX2;SubX++ )
+						{
+							int32 LandscapeX = SubIndexX*SubsectionSizeQuads + ComponentIndexX*ComponentSizeQuads + SubX;
+							int32 LandscapeY = SubIndexY*SubsectionSizeQuads + ComponentIndexY*ComponentSizeQuads + SubY;
+
+							// Find the input data corresponding to this vertex
+							if( Component && OffsetTextureData )
+							{
+								FVector2D Value = GetXYOffsetmapData(Component, SubX, SubY, OffsetTextureData);
+								StoreData.Store(LandscapeX, LandscapeY, Value);
+							}
+							else
+							{
+								StoreData.Store(LandscapeX, LandscapeY, FVector2D(0.f, 0.f) );
+							}
+						}
+					}
+				}
+			}
+		}
+	}
+}
+
+void FLandscapeEditDataInterface::GetXYOffsetDataFast(const int32 X1, const int32 Y1, const int32 X2, const int32 Y2, FVector2D* Data, int32 Stride)
+{
+	if( Stride==0 )
+	{
+		Stride = (1+X2-X1);
+	}
+	TArrayStoreData<FVector2D> ArrayStoreData(X1, Y1, Data, Stride);
+	GetXYOffsetDataTemplFast(X1, Y1, X2, Y2, ArrayStoreData);
+}
+
+void FLandscapeEditDataInterface::GetXYOffsetDataFast(const int32 X1, const int32 Y1, const int32 X2, const int32 Y2, TMap<FIntPoint, FVector2D>& SparseData)
+{
+	TSparseStoreData<FVector2D> SparseStoreData(SparseData);
+	GetXYOffsetDataTemplFast(X1, Y1, X2, Y2, SparseStoreData);
+}
+
+void FLandscapeEditDataInterface::GetXYOffsetDataFast(const int32 X1, const int32 Y1, const int32 X2, const int32 Y2, FVector* Data, int32 Stride)
+{
+	if( Stride==0 )
+	{
+		Stride = (1+X2-X1);
+	}
+	TArrayStoreData<FVector> ArrayStoreData(X1, Y1, Data, Stride);
+	GetXYOffsetDataTemplFast(X1, Y1, X2, Y2, ArrayStoreData);
+}
+
+void FLandscapeEditDataInterface::GetXYOffsetDataFast(const int32 X1, const int32 Y1, const int32 X2, const int32 Y2, TMap<FIntPoint, FVector>& SparseData)
+{
+	TSparseStoreData<FVector> SparseStoreData(SparseData);
+	GetXYOffsetDataTemplFast(X1, Y1, X2, Y2, SparseStoreData);
+}
+
+//
+// FLandscapeTextureDataInfo
+//
+
+FLandscapeTextureDataInfo::FLandscapeTextureDataInfo(UTexture2D* InTexture)
+:	Texture(InTexture)
+{
+	MipInfo.AddZeroed(Texture->Source.GetNumMips());
+	Texture->SetFlags(RF_Transactional);
+	Texture->TemporarilyDisableStreaming();
+	Texture->Modify();
+}
+
+bool FLandscapeTextureDataInfo::UpdateTextureData()
+{
+	bool bNeedToWaitForUpdate = false;
+
+	int32 DataSize = sizeof(FColor);
+	if (Texture->GetPixelFormat() == PF_G8)
+	{
+		DataSize = sizeof(uint8);
+	}
+
+	for( int32 i=0;i<MipInfo.Num();i++ )
+	{
+		if( MipInfo[i].MipData && MipInfo[i].MipUpdateRegions.Num()>0 )
+		{
+			Texture->UpdateTextureRegions( i, MipInfo[i].MipUpdateRegions.Num(), &MipInfo[i].MipUpdateRegions[0], ((Texture->Source.GetSizeX())>>i)*DataSize, DataSize, (uint8*)MipInfo[i].MipData, false );
+			bNeedToWaitForUpdate = true;
+		}
+	}
+
+	return bNeedToWaitForUpdate;
+}
+
+FLandscapeTextureDataInfo::~FLandscapeTextureDataInfo()
+{
+	// Unlock any mips still locked.
+	for( int32 i=0;i<MipInfo.Num();i++ )
+	{
+		if( MipInfo[i].MipData )
+		{
+			if( MipInfo[i].MipData )
+			{
+				Texture->Source.UnlockMip(i);
+				MipInfo[i].MipData = NULL;
+			}
+		}
+	}
+	Texture->ClearFlags(RF_Transactional);
+}
+
+#endif // WITH_EDITOR
