Landscape module from Engine split out.

diff --git a/Engine/Source/Runtime/Landscape/Private/LandscapeGizmoActor.cpp b/Engine/Source/Runtime/Landscape/Private/LandscapeGizmoActor.cpp
new file mode 100644
--- /dev/null
+++ b/Engine/Source/Runtime/Landscape/Private/LandscapeGizmoActor.cpp
@@ -0,0 +1,1288 @@
+// Copyright 1998-2014 Epic Games, Inc. All Rights Reserved.
+
+
+#include "Landscape.h"
+#include "LandscapeDataAccess.h"
+#include "LandscapeRender.h"
+#include "LandscapeGizmoActiveActor.h"
+#include "LandscapeGizmoRenderComponent.h"
+#include "Landscape.h"
+#include "LandscapeInfo.h"
+#include "LandscapeLayerInfoObject.h"
+#include "DynamicMeshBuilder.h"
+#include "Engine/CollisionProfile.h"
+
+namespace
+{
+	enum PreviewType
+	{
+		Invalid = -1,
+		Both = 0,
+		Add = 1,
+		Sub = 2,
+	};
+}
+
+
+class FLandscapeGizmoMeshRenderProxy : public FMaterialRenderProxy
+{
+public:
+	const FMaterialRenderProxy* const Parent;
+	const float TopHeight;
+	const float BottomHeight;
+	const UTexture2D* AlphaTexture;
+	const FLinearColor ScaleBias;
+	const FMatrix WorldToLandscapeMatrix;
+
+	/** Initialization constructor. */
+	FLandscapeGizmoMeshRenderProxy(const FMaterialRenderProxy* InParent, const float InTop, const float InBottom, const UTexture2D* InAlphaTexture, const FLinearColor& InScaleBias, const FMatrix& InWorldToLandscapeMatrix)
+	:	Parent(InParent)
+	,	TopHeight(InTop)
+	,	BottomHeight(InBottom)
+	,	AlphaTexture(InAlphaTexture)
+	,	ScaleBias(InScaleBias)
+	,	WorldToLandscapeMatrix(InWorldToLandscapeMatrix)
+	{}
+
+	// FMaterialRenderProxy interface.
+	virtual const class FMaterial* GetMaterial(ERHIFeatureLevel::Type FeatureLevel) const
+	{
+		return Parent->GetMaterial(FeatureLevel);
+	}
+	virtual bool GetVectorValue(const FName ParameterName, FLinearColor* OutValue, const FMaterialRenderContext& Context) const
+	{
+		if (ParameterName == FName(TEXT("AlphaScaleBias")))
+		{
+			*OutValue = ScaleBias;
+			return true;
+		}
+		else
+		if (ParameterName == FName(TEXT("MatrixRow1")))
+		{
+			*OutValue = FLinearColor(WorldToLandscapeMatrix.M[0][0], WorldToLandscapeMatrix.M[0][1], WorldToLandscapeMatrix.M[0][2],WorldToLandscapeMatrix.M[0][3]);
+			return true;
+		}
+		else
+		if (ParameterName == FName(TEXT("MatrixRow2")))
+		{
+			*OutValue = FLinearColor(WorldToLandscapeMatrix.M[1][0], WorldToLandscapeMatrix.M[1][1], WorldToLandscapeMatrix.M[1][2],WorldToLandscapeMatrix.M[1][3]);
+			return true;
+		}
+		else
+		if (ParameterName == FName(TEXT("MatrixRow3")))
+		{
+			*OutValue = FLinearColor(WorldToLandscapeMatrix.M[2][0], WorldToLandscapeMatrix.M[2][1], WorldToLandscapeMatrix.M[2][2],WorldToLandscapeMatrix.M[2][3]);
+			return true;
+		}
+		else
+		if (ParameterName == FName(TEXT("MatrixRow4")))
+		{
+			*OutValue = FLinearColor(WorldToLandscapeMatrix.M[3][0], WorldToLandscapeMatrix.M[3][1], WorldToLandscapeMatrix.M[3][2],WorldToLandscapeMatrix.M[3][3]);
+			return true;
+		}
+
+		return Parent->GetVectorValue(ParameterName, OutValue, Context);
+	}
+	virtual bool GetScalarValue(const FName ParameterName, float* OutValue, const FMaterialRenderContext& Context) const
+	{
+		if (ParameterName == FName(TEXT("Top")))
+		{
+			*OutValue = TopHeight;
+			return true;
+		}
+		else if (ParameterName == FName(TEXT("Bottom")))
+		{
+			*OutValue = BottomHeight;
+			return true;
+		}
+		return Parent->GetScalarValue(ParameterName, OutValue, Context);
+	}
+	virtual bool GetTextureValue(const FName ParameterName,const UTexture** OutValue, const FMaterialRenderContext& Context) const
+	{
+		if (ParameterName == FName(TEXT("AlphaTexture")))
+		{
+			// FIXME: This needs to return a black texture if AlphaTexture is NULL.
+			// Returning NULL will cause the material to use GWhiteTexture.
+			*OutValue = AlphaTexture;
+			return true;
+		}
+		return Parent->GetTextureValue(ParameterName, OutValue, Context);
+	}
+};
+
+/** Represents a LandscapeGizmoRenderingComponent to the scene manager. */
+class FLandscapeGizmoRenderSceneProxy : public FPrimitiveSceneProxy
+{
+public:
+	FMatrix MeshRT;
+	FVector XAxis, YAxis, Origin;
+	FVector FrustumVerts[8];
+	float SampleSizeX, SampleSizeY;
+	TArray<FVector> SampledPositions;
+	TArray<FVector> SampledNormals;
+	bool bHeightmapRendering;
+	FLandscapeGizmoMeshRenderProxy* HeightmapRenderProxy;
+	FMaterialRenderProxy* GizmoRenderProxy;
+
+	FLandscapeGizmoRenderSceneProxy(const ULandscapeGizmoRenderComponent* InComponent):
+		FPrimitiveSceneProxy(InComponent),
+		bHeightmapRendering(false),
+		HeightmapRenderProxy(NULL),
+		GizmoRenderProxy(NULL)
+	{
+#if WITH_EDITOR	
+		ALandscapeGizmoActiveActor* Gizmo = Cast<ALandscapeGizmoActiveActor>(InComponent->GetOwner());
+		if (Gizmo && Gizmo->GizmoMeshMaterial && Gizmo->GizmoDataMaterial && Gizmo->GetRootComponent())
+		{
+			ULandscapeInfo* LandscapeInfo = Gizmo->TargetLandscapeInfo;
+			if (LandscapeInfo && LandscapeInfo->GetLandscapeProxy())
+			{
+				SampleSizeX = Gizmo->SampleSizeX;
+				SampleSizeY = Gizmo->SampleSizeY;
+				bHeightmapRendering = (Gizmo->DataType & LGT_Height);
+				FMatrix LToW = LandscapeInfo->GetLandscapeProxy()->LandscapeActorToWorld().ToMatrixWithScale();
+				FMatrix WToL = LToW.InverseFast();
+				FVector BaseLocation = WToL.TransformPosition(Gizmo->GetActorLocation());
+				float ScaleXY = LandscapeInfo->DrawScale.X;
+				float ScaleZ = LandscapeInfo->DrawScale.Z;
+				const float W = Gizmo->GetWidth() / (2 * ScaleXY);
+				const float H = Gizmo->GetHeight() / (2 * ScaleXY);
+				const float L = Gizmo->GetLength() / ScaleZ;
+				FMatrix GizmoRT = FRotationTranslationMatrix(FRotator(0, Gizmo->GetActorRotation().Yaw, 0), FVector(BaseLocation.X, BaseLocation.Y, 0)) * LToW;
+
+				FrustumVerts[0] = Gizmo->FrustumVerts[0] = GizmoRT.TransformPosition(FVector( - W, - H, BaseLocation.Z + L ));
+				FrustumVerts[1] = Gizmo->FrustumVerts[1] = GizmoRT.TransformPosition(FVector( + W, - H, BaseLocation.Z + L ));
+				FrustumVerts[2] = Gizmo->FrustumVerts[2] = GizmoRT.TransformPosition(FVector( + W, + H, BaseLocation.Z + L ));
+				FrustumVerts[3] = Gizmo->FrustumVerts[3] = GizmoRT.TransformPosition(FVector( - W, + H, BaseLocation.Z + L ));
+
+				FrustumVerts[4] = Gizmo->FrustumVerts[4] = GizmoRT.TransformPosition(FVector( - W, - H, BaseLocation.Z ));
+				FrustumVerts[5] = Gizmo->FrustumVerts[5] = GizmoRT.TransformPosition(FVector( + W, - H, BaseLocation.Z ));
+				FrustumVerts[6] = Gizmo->FrustumVerts[6] = GizmoRT.TransformPosition(FVector( + W, + H, BaseLocation.Z ));
+				FrustumVerts[7] = Gizmo->FrustumVerts[7] = GizmoRT.TransformPosition(FVector( - W, + H, BaseLocation.Z ));
+
+				XAxis = GizmoRT.TransformPosition(FVector( + W,	0,		BaseLocation.Z + L ));
+				YAxis = GizmoRT.TransformPosition(FVector( 0,	+ H,	BaseLocation.Z + L ));
+				Origin = GizmoRT.TransformPosition(FVector( 0,	0,		BaseLocation.Z + L ));
+
+				MeshRT = FTranslationMatrix(FVector(- W + 0.5, - H + 0.5, 0)) * FRotationTranslationMatrix(FRotator(0, Gizmo->GetActorRotation().Yaw, 0), FVector(BaseLocation.X, BaseLocation.Y, 0)) * LToW;
+				HeightmapRenderProxy = new FLandscapeGizmoMeshRenderProxy( Gizmo->GizmoMeshMaterial->GetRenderProxy(false), BaseLocation.Z + L, BaseLocation.Z, Gizmo->GizmoTexture, FLinearColor(Gizmo->TextureScale.X, Gizmo->TextureScale.Y, 0, 0), WToL );
+
+				GizmoRenderProxy = (Gizmo->DataType != LGT_None) ? Gizmo->GizmoDataMaterial->GetRenderProxy(false) : Gizmo->GizmoMaterial->GetRenderProxy(false);
+
+				// Cache sampled height
+				float ScaleX = Gizmo->GetWidth() / Gizmo->CachedWidth / ScaleXY * Gizmo->CachedScaleXY;
+				float ScaleY = Gizmo->GetHeight() / Gizmo->CachedHeight / ScaleXY * Gizmo->CachedScaleXY;
+				FScaleMatrix Mat(FVector(ScaleX, ScaleY, L));
+				FMatrix NormalM = Mat.InverseFast().GetTransposed();
+
+				int32 SamplingSize = Gizmo->SampleSizeX * Gizmo->SampleSizeY;
+				SampledPositions.Empty(SamplingSize);
+				SampledNormals.Empty(SamplingSize);
+
+				for (int32 Y = 0; Y < Gizmo->SampleSizeY; ++Y)
+				{
+					for (int32 X = 0; X < Gizmo->SampleSizeX; ++X)
+					{
+						FVector SampledPos = Gizmo->SampledHeight[X + Y * ALandscapeGizmoActiveActor::DataTexSize];
+						SampledPos.X *= ScaleX;
+						SampledPos.Y *= ScaleY;
+						SampledPos.Z = Gizmo->GetLandscapeHeight(SampledPos.Z);
+
+						FVector SampledNormal = NormalM.TransformVector(Gizmo->SampledNormal[X + Y * ALandscapeGizmoActiveActor::DataTexSize]);
+						SampledNormal = SampledNormal.SafeNormal();
+
+						SampledPositions.Add(SampledPos);
+						SampledNormals.Add(SampledNormal);
+						//MeshBuilder.AddVertex(SampledPos, FVector2D((float)X / (Gizmo->SampleSizeX), (float)Y / (Gizmo->SampleSizeY)), TangentX, SampledNormal^TangentX, SampledNormal, FColor(255, 255, 255) );
+					}
+				}
+			}
+		}
+#endif
+	}
+
+	~FLandscapeGizmoRenderSceneProxy()
+	{
+		delete HeightmapRenderProxy;
+		HeightmapRenderProxy = NULL;
+	}
+
+	virtual void GetDynamicMeshElements(const TArray<const FSceneView*>& Views, const FSceneViewFamily& ViewFamily, uint32 VisibilityMap, FMeshElementCollector& Collector) const override
+	{
+		//FMemMark Mark(FMemStack::Get());
+#if WITH_EDITOR
+		if( GizmoRenderProxy &&  HeightmapRenderProxy )
+		{
+			for (int32 ViewIndex = 0; ViewIndex < Views.Num(); ViewIndex++)
+			{
+				if (VisibilityMap & (1 << ViewIndex))
+				{
+					const FSceneView* View = Views[ViewIndex];
+					FPrimitiveDrawInterface* PDI = Collector.GetPDI(ViewIndex);
+
+					// Axis
+					PDI->DrawLine( Origin, XAxis, FLinearColor(1, 0, 0), SDPG_World );
+					PDI->DrawLine( Origin, YAxis, FLinearColor(0, 1, 0), SDPG_World );
+
+					{
+						FDynamicMeshBuilder MeshBuilder;
+
+						MeshBuilder.AddVertex(FrustumVerts[0], FVector2D(0, 0), FVector(1,0,0), FVector(0,1,0), FVector(0,0,1), FColor(255,255,255));
+						MeshBuilder.AddVertex(FrustumVerts[1], FVector2D(1, 0), FVector(1,0,0), FVector(0,1,0), FVector(0,0,1), FColor(255,255,255));
+						MeshBuilder.AddVertex(FrustumVerts[2], FVector2D(1, 1), FVector(1,0,0), FVector(0,1,0), FVector(0,0,1), FColor(255,255,255));
+						MeshBuilder.AddVertex(FrustumVerts[3], FVector2D(0, 1), FVector(1,0,0), FVector(0,1,0), FVector(0,0,1), FColor(255,255,255));
+
+						MeshBuilder.AddVertex(FrustumVerts[4], FVector2D(0, 0), FVector(1,0,0), FVector(0,1,0), FVector(0,0,1), FColor(255,255,255));
+						MeshBuilder.AddVertex(FrustumVerts[5], FVector2D(1, 0), FVector(1,0,0), FVector(0,1,0), FVector(0,0,1), FColor(255,255,255));
+						MeshBuilder.AddVertex(FrustumVerts[6], FVector2D(1, 1), FVector(1,0,0), FVector(0,1,0), FVector(0,0,1), FColor(255,255,255));
+						MeshBuilder.AddVertex(FrustumVerts[7], FVector2D(0, 1), FVector(1,0,0), FVector(0,1,0), FVector(0,0,1), FColor(255,255,255));
+
+						MeshBuilder.AddVertex(FrustumVerts[1], FVector2D(0, 0), FVector(1,0,0), FVector(0,1,0), FVector(0,0,1), FColor(255,255,255));
+						MeshBuilder.AddVertex(FrustumVerts[0], FVector2D(1, 0), FVector(1,0,0), FVector(0,1,0), FVector(0,0,1), FColor(255,255,255));
+						MeshBuilder.AddVertex(FrustumVerts[4], FVector2D(1, 1), FVector(1,0,0), FVector(0,1,0), FVector(0,0,1), FColor(255,255,255));
+						MeshBuilder.AddVertex(FrustumVerts[5], FVector2D(0, 1), FVector(1,0,0), FVector(0,1,0), FVector(0,0,1), FColor(255,255,255));
+
+						MeshBuilder.AddVertex(FrustumVerts[3], FVector2D(0, 0), FVector(1,0,0), FVector(0,1,0), FVector(0,0,1), FColor(255,255,255));
+						MeshBuilder.AddVertex(FrustumVerts[2], FVector2D(1, 0), FVector(1,0,0), FVector(0,1,0), FVector(0,0,1), FColor(255,255,255));
+						MeshBuilder.AddVertex(FrustumVerts[6], FVector2D(1, 1), FVector(1,0,0), FVector(0,1,0), FVector(0,0,1), FColor(255,255,255));
+						MeshBuilder.AddVertex(FrustumVerts[7], FVector2D(0, 1), FVector(1,0,0), FVector(0,1,0), FVector(0,0,1), FColor(255,255,255));
+
+						MeshBuilder.AddVertex(FrustumVerts[2], FVector2D(0, 0), FVector(1,0,0), FVector(0,1,0), FVector(0,0,1), FColor(255,255,255));
+						MeshBuilder.AddVertex(FrustumVerts[1], FVector2D(1, 0), FVector(1,0,0), FVector(0,1,0), FVector(0,0,1), FColor(255,255,255));
+						MeshBuilder.AddVertex(FrustumVerts[5], FVector2D(1, 1), FVector(1,0,0), FVector(0,1,0), FVector(0,0,1), FColor(255,255,255));
+						MeshBuilder.AddVertex(FrustumVerts[6], FVector2D(0, 1), FVector(1,0,0), FVector(0,1,0), FVector(0,0,1), FColor(255,255,255));
+
+						MeshBuilder.AddVertex(FrustumVerts[0], FVector2D(0, 0), FVector(1,0,0), FVector(0,1,0), FVector(0,0,1), FColor(255,255,255));
+						MeshBuilder.AddVertex(FrustumVerts[3], FVector2D(1, 0), FVector(1,0,0), FVector(0,1,0), FVector(0,0,1), FColor(255,255,255));
+						MeshBuilder.AddVertex(FrustumVerts[7], FVector2D(1, 1), FVector(1,0,0), FVector(0,1,0), FVector(0,0,1), FColor(255,255,255));
+						MeshBuilder.AddVertex(FrustumVerts[4], FVector2D(0, 1), FVector(1,0,0), FVector(0,1,0), FVector(0,0,1), FColor(255,255,255));
+
+						for (int32 i = 0; i < 6; ++i)
+						{
+							int32 Idx = i*4;
+							MeshBuilder.AddTriangle( Idx, Idx+2, Idx+1 );
+							MeshBuilder.AddTriangle( Idx, Idx+3, Idx+2 );
+						}
+
+						MeshBuilder.GetMesh(FMatrix::Identity, GizmoRenderProxy , SDPG_World, true, false, ViewIndex, Collector);
+					}
+
+					if (bHeightmapRendering)
+					{		  		
+						FDynamicMeshBuilder MeshBuilder;
+
+						for (int32 Y = 0; Y < SampleSizeY; ++Y)
+						{
+							for (int32 X = 0; X < SampleSizeX; ++X)
+							{
+								FVector SampledNormal = SampledNormals[X + Y * SampleSizeX];
+								FVector TangentX(SampledNormal.Z, 0, -SampledNormal.X);
+								TangentX = TangentX.SafeNormal();
+
+								MeshBuilder.AddVertex(SampledPositions[X + Y * SampleSizeX], FVector2D((float)X / (SampleSizeX), (float)Y / (SampleSizeY)), TangentX, SampledNormal^TangentX, SampledNormal, FColor(255, 255, 255) );
+							}
+						}
+
+						for (int32 Y = 0; Y < SampleSizeY; ++Y)
+						{
+							for (int32 X = 0; X < SampleSizeX; ++X)
+							{
+								if (X < SampleSizeX - 1 && Y < SampleSizeY - 1)
+								{
+									MeshBuilder.AddTriangle( (X+0) + (Y+0) * SampleSizeX, (X+1) + (Y+1) * SampleSizeX, (X+1) + (Y+0) * SampleSizeX );
+									MeshBuilder.AddTriangle( (X+0) + (Y+0) * SampleSizeX, (X+0) + (Y+1) * SampleSizeX, (X+1) + (Y+1) * SampleSizeX );
+								}
+							}
+						}
+
+						MeshBuilder.GetMesh(MeshRT, HeightmapRenderProxy , SDPG_World, false, false, ViewIndex, Collector);
+					}
+				}
+			}
+		}
+#endif
+	};
+
+	virtual void DrawDynamicElements(FPrimitiveDrawInterface* PDI,const FSceneView* View)
+	{
+		//FMemMark Mark(FMemStack::Get());
+#if WITH_EDITOR
+		if( GizmoRenderProxy &&  HeightmapRenderProxy )
+		{
+			// Axis
+			PDI->DrawLine( Origin, XAxis, FLinearColor(1, 0, 0), SDPG_World );
+			PDI->DrawLine( Origin, YAxis, FLinearColor(0, 1, 0), SDPG_World );
+
+			{
+				FDynamicMeshBuilder MeshBuilder;
+
+				MeshBuilder.AddVertex(FrustumVerts[0], FVector2D(0, 0), FVector(1,0,0), FVector(0,1,0), FVector(0,0,1), FColor(255,255,255));
+				MeshBuilder.AddVertex(FrustumVerts[1], FVector2D(1, 0), FVector(1,0,0), FVector(0,1,0), FVector(0,0,1), FColor(255,255,255));
+				MeshBuilder.AddVertex(FrustumVerts[2], FVector2D(1, 1), FVector(1,0,0), FVector(0,1,0), FVector(0,0,1), FColor(255,255,255));
+				MeshBuilder.AddVertex(FrustumVerts[3], FVector2D(0, 1), FVector(1,0,0), FVector(0,1,0), FVector(0,0,1), FColor(255,255,255));
+
+				MeshBuilder.AddVertex(FrustumVerts[4], FVector2D(0, 0), FVector(1,0,0), FVector(0,1,0), FVector(0,0,1), FColor(255,255,255));
+				MeshBuilder.AddVertex(FrustumVerts[5], FVector2D(1, 0), FVector(1,0,0), FVector(0,1,0), FVector(0,0,1), FColor(255,255,255));
+				MeshBuilder.AddVertex(FrustumVerts[6], FVector2D(1, 1), FVector(1,0,0), FVector(0,1,0), FVector(0,0,1), FColor(255,255,255));
+				MeshBuilder.AddVertex(FrustumVerts[7], FVector2D(0, 1), FVector(1,0,0), FVector(0,1,0), FVector(0,0,1), FColor(255,255,255));
+
+				MeshBuilder.AddVertex(FrustumVerts[1], FVector2D(0, 0), FVector(1,0,0), FVector(0,1,0), FVector(0,0,1), FColor(255,255,255));
+				MeshBuilder.AddVertex(FrustumVerts[0], FVector2D(1, 0), FVector(1,0,0), FVector(0,1,0), FVector(0,0,1), FColor(255,255,255));
+				MeshBuilder.AddVertex(FrustumVerts[4], FVector2D(1, 1), FVector(1,0,0), FVector(0,1,0), FVector(0,0,1), FColor(255,255,255));
+				MeshBuilder.AddVertex(FrustumVerts[5], FVector2D(0, 1), FVector(1,0,0), FVector(0,1,0), FVector(0,0,1), FColor(255,255,255));
+
+				MeshBuilder.AddVertex(FrustumVerts[3], FVector2D(0, 0), FVector(1,0,0), FVector(0,1,0), FVector(0,0,1), FColor(255,255,255));
+				MeshBuilder.AddVertex(FrustumVerts[2], FVector2D(1, 0), FVector(1,0,0), FVector(0,1,0), FVector(0,0,1), FColor(255,255,255));
+				MeshBuilder.AddVertex(FrustumVerts[6], FVector2D(1, 1), FVector(1,0,0), FVector(0,1,0), FVector(0,0,1), FColor(255,255,255));
+				MeshBuilder.AddVertex(FrustumVerts[7], FVector2D(0, 1), FVector(1,0,0), FVector(0,1,0), FVector(0,0,1), FColor(255,255,255));
+
+				MeshBuilder.AddVertex(FrustumVerts[2], FVector2D(0, 0), FVector(1,0,0), FVector(0,1,0), FVector(0,0,1), FColor(255,255,255));
+				MeshBuilder.AddVertex(FrustumVerts[1], FVector2D(1, 0), FVector(1,0,0), FVector(0,1,0), FVector(0,0,1), FColor(255,255,255));
+				MeshBuilder.AddVertex(FrustumVerts[5], FVector2D(1, 1), FVector(1,0,0), FVector(0,1,0), FVector(0,0,1), FColor(255,255,255));
+				MeshBuilder.AddVertex(FrustumVerts[6], FVector2D(0, 1), FVector(1,0,0), FVector(0,1,0), FVector(0,0,1), FColor(255,255,255));
+
+				MeshBuilder.AddVertex(FrustumVerts[0], FVector2D(0, 0), FVector(1,0,0), FVector(0,1,0), FVector(0,0,1), FColor(255,255,255));
+				MeshBuilder.AddVertex(FrustumVerts[3], FVector2D(1, 0), FVector(1,0,0), FVector(0,1,0), FVector(0,0,1), FColor(255,255,255));
+				MeshBuilder.AddVertex(FrustumVerts[7], FVector2D(1, 1), FVector(1,0,0), FVector(0,1,0), FVector(0,0,1), FColor(255,255,255));
+				MeshBuilder.AddVertex(FrustumVerts[4], FVector2D(0, 1), FVector(1,0,0), FVector(0,1,0), FVector(0,0,1), FColor(255,255,255));
+
+				for (int32 i = 0; i < 6; ++i)
+				{
+					int32 Idx = i*4;
+					MeshBuilder.AddTriangle( Idx, Idx+2, Idx+1 );
+					MeshBuilder.AddTriangle( Idx, Idx+3, Idx+2 );
+				}
+
+				MeshBuilder.Draw(PDI, FMatrix::Identity, GizmoRenderProxy , SDPG_World, true);
+			}
+
+			if (bHeightmapRendering)
+			{		  		
+				FDynamicMeshBuilder MeshBuilder;
+
+				for (int32 Y = 0; Y < SampleSizeY; ++Y)
+				{
+					for (int32 X = 0; X < SampleSizeX; ++X)
+					{
+						FVector SampledNormal = SampledNormals[X + Y * SampleSizeX];
+						FVector TangentX(SampledNormal.Z, 0, -SampledNormal.X);
+						TangentX = TangentX.SafeNormal();
+
+						MeshBuilder.AddVertex(SampledPositions[X + Y * SampleSizeX], FVector2D((float)X / (SampleSizeX), (float)Y / (SampleSizeY)), TangentX, SampledNormal^TangentX, SampledNormal, FColor(255, 255, 255) );
+					}
+				}
+
+				for (int32 Y = 0; Y < SampleSizeY; ++Y)
+				{
+					for (int32 X = 0; X < SampleSizeX; ++X)
+					{
+						if (X < SampleSizeX - 1 && Y < SampleSizeY - 1)
+						{
+							MeshBuilder.AddTriangle( (X+0) + (Y+0) * SampleSizeX, (X+1) + (Y+1) * SampleSizeX, (X+1) + (Y+0) * SampleSizeX );
+							MeshBuilder.AddTriangle( (X+0) + (Y+0) * SampleSizeX, (X+0) + (Y+1) * SampleSizeX, (X+1) + (Y+1) * SampleSizeX );
+						}
+					}
+				}
+
+				MeshBuilder.Draw(PDI, MeshRT, HeightmapRenderProxy, SDPG_World, false);
+			}
+		}
+#endif
+	};
+
+	virtual FPrimitiveViewRelevance GetViewRelevance(const FSceneView* View)
+	{
+		FPrimitiveViewRelevance Result;
+#if WITH_EDITOR
+		const bool bVisible = View->Family->EngineShowFlags.Landscape;
+		Result.bDrawRelevance = IsShown(View) && bVisible && !View->bIsGameView && GLandscapeEditRenderMode & ELandscapeEditRenderMode::Gizmo;
+		Result.bDynamicRelevance = true;
+		Result.bNormalTranslucencyRelevance = true;
+#endif
+		return Result;
+	}
+
+	virtual uint32 GetMemoryFootprint( void ) const { return( sizeof( *this ) + GetAllocatedSize() ); }
+	uint32 GetAllocatedSize( void ) const { return( FPrimitiveSceneProxy::GetAllocatedSize() ); }
+};
+
+ULandscapeGizmoRenderComponent::ULandscapeGizmoRenderComponent(const FObjectInitializer& ObjectInitializer)
+	: Super(ObjectInitializer)
+{
+	bHiddenInGame = true;
+	AlwaysLoadOnClient = false;
+	AlwaysLoadOnServer = false;
+	bSelectable = false;
+	BodyInstance.bEnableCollision_DEPRECATED = true;
+	SetCollisionProfileName(UCollisionProfile::BlockAll_ProfileName);
+}
+
+FPrimitiveSceneProxy* ULandscapeGizmoRenderComponent::CreateSceneProxy()
+{
+	return new FLandscapeGizmoRenderSceneProxy(this);
+}
+
+FBoxSphereBounds ULandscapeGizmoRenderComponent::CalcBounds(const FTransform& LocalToWorld) const
+{
+#if WITH_EDITOR
+	ALandscapeGizmoActiveActor* Gizmo = Cast<ALandscapeGizmoActiveActor>(GetOwner());
+	if( Gizmo )
+	{
+		return FBoxSphereBounds(Gizmo->FrustumVerts, 8);
+	}
+#endif
+	{
+		return Super::CalcBounds(LocalToWorld);
+	}
+}
+
+ALandscapeGizmoActor::ALandscapeGizmoActor(const FObjectInitializer& ObjectInitializer)
+	: Super(ObjectInitializer)
+{
+#if WITH_EDITORONLY_DATA
+	SpriteComponent = ObjectInitializer.CreateEditorOnlyDefaultSubobject<UBillboardComponent>(this, TEXT("Sprite"));
+	if (!IsRunningCommandlet() && (SpriteComponent != nullptr))
+	{
+		// Structure to hold one-time initialization
+		struct FConstructorStatics
+		{
+			ConstructorHelpers::FObjectFinderOptional<UTexture2D> DecalActorIconTexture;
+			FName ID_Misc;
+			FText NAME_Misc;
+			FConstructorStatics()
+				: DecalActorIconTexture(TEXT("Texture2D'/Engine/EditorResources/S_DecalActorIcon.S_DecalActorIcon'"))
+				, ID_Misc(TEXT("Misc"))
+				, NAME_Misc(NSLOCTEXT("SpriteCategory", "Misc", "Misc"))
+			{
+			}
+		};
+		static FConstructorStatics ConstructorStatics;
+
+		SpriteComponent->Sprite = ConstructorStatics.DecalActorIconTexture.Get();
+		SpriteComponent->RelativeScale3D = FVector(0.5f, 0.5f, 0.5f);
+		SpriteComponent->bHiddenInGame = true;
+		SpriteComponent->SpriteInfo.Category = ConstructorStatics.ID_Misc;
+		SpriteComponent->SpriteInfo.DisplayName = ConstructorStatics.NAME_Misc;
+		SpriteComponent->bIsScreenSizeScaled = true;
+	}
+#endif
+
+	TSubobjectPtr<USceneComponent> SceneComponent = ObjectInitializer.CreateDefaultSubobject<USceneComponent>(this, TEXT("RootComponent0"));
+	RootComponent = SceneComponent;
+	RootComponent->Mobility = EComponentMobility::Static;
+
+#if WITH_EDITORONLY_DATA
+	bEditable = false;
+	Width = 1280.0f;
+	Height = 1280.0f;
+	LengthZ = 1280.0f;
+	MarginZ = 512.0f;
+	MinRelativeZ = 0.0f;
+	RelativeScaleZ = 1.0f;
+#endif // WITH_EDITORONLY_DATA
+}
+
+#if WITH_EDITOR
+
+void ALandscapeGizmoActor::Duplicate(ALandscapeGizmoActor* Gizmo)
+{
+	Gizmo->Width = Width;
+	Gizmo->Height = Height;
+	Gizmo->LengthZ = LengthZ;
+	Gizmo->MarginZ = MarginZ;
+	//Gizmo->TargetLandscapeInfo = TargetLandscapeInfo;
+
+	Gizmo->SetActorLocation( GetActorLocation(), false );
+	Gizmo->SetActorRotation( GetActorRotation() );
+
+	if( Gizmo->GetRootComponent() != NULL && GetRootComponent() != NULL )
+	{
+		Gizmo->GetRootComponent()->SetRelativeScale3D( GetRootComponent()->RelativeScale3D );
+	}
+
+	Gizmo->MinRelativeZ = MinRelativeZ;
+	Gizmo->RelativeScaleZ = RelativeScaleZ;
+
+	Gizmo->ReregisterAllComponents();
+}
+#endif	//WITH_EDITOR
+
+ALandscapeGizmoActiveActor::ALandscapeGizmoActiveActor(const FObjectInitializer& ObjectInitializer)
+	: Super(ObjectInitializer
+		.DoNotCreateDefaultSubobject(TEXT("Sprite"))
+	)
+{
+#if WITH_EDITORONLY_DATA
+	if (!IsRunningCommandlet())
+	{
+		// Structure to hold one-time initialization
+		struct FConstructorStatics
+		{
+			ConstructorHelpers::FObjectFinder<UMaterial> LandscapeGizmo_Mat;
+			ConstructorHelpers::FObjectFinder<UMaterialInstanceConstant> LandscapeGizmo_Mat_Copied;
+			ConstructorHelpers::FObjectFinder<UMaterial> LandscapeGizmoHeight_Mat;
+			FConstructorStatics()
+				: LandscapeGizmo_Mat(TEXT("/Engine/EditorLandscapeResources/LandscapeGizmo_Mat"))
+				, LandscapeGizmo_Mat_Copied(TEXT("/Engine/EditorLandscapeResources/LandscapeGizmo_Mat_Copied"))
+				, LandscapeGizmoHeight_Mat(TEXT("/Engine/EditorLandscapeResources/LandscapeGizmoHeight_Mat"))
+			{
+			}
+		};
+		static FConstructorStatics ConstructorStatics;
+
+		GizmoMaterial = ConstructorStatics.LandscapeGizmo_Mat.Object;
+		GizmoDataMaterial = ConstructorStatics.LandscapeGizmo_Mat_Copied.Object;
+		GizmoMeshMaterial = ConstructorStatics.LandscapeGizmoHeight_Mat.Object;
+	}
+#endif // WITH_EDITORONLY_DATA
+
+	TSubobjectPtr<ULandscapeGizmoRenderComponent> LandscapeGizmoRenderComponent = ObjectInitializer.CreateDefaultSubobject<ULandscapeGizmoRenderComponent>(this, TEXT("GizmoRendererComponent0"));
+	LandscapeGizmoRenderComponent->SetCollisionProfileName(UCollisionProfile::BlockAllDynamic_ProfileName);
+
+	RootComponent = LandscapeGizmoRenderComponent;
+#if WITH_EDITORONLY_DATA
+	bEditable = true;
+	Width = 1280.0f;
+	Height = 1280.0f;
+	LengthZ = 1280.0f;
+	MarginZ = 512.0f;
+	DataType = LGT_None;
+	SampleSizeX = 0;
+	SampleSizeY = 0;
+	CachedWidth = 0.0f;
+	CachedHeight = 0.0f;
+	CachedScaleXY = 1.0f;
+#endif // WITH_EDITORONLY_DATA
+}
+
+#if WITH_EDITOR
+void ALandscapeGizmoActiveActor::PostEditChangeProperty(FPropertyChangedEvent& PropertyChangedEvent)
+{
+	const FName PropertyName = PropertyChangedEvent.Property ? PropertyChangedEvent.Property->GetFName() : NAME_None;
+	if( PropertyName == FName(TEXT("LengthZ")) )
+	{
+		if (LengthZ < 0)
+		{
+			LengthZ = MarginZ;
+		}
+	}
+	else if ( PropertyName == FName(TEXT("TargetLandscapeInfo")) )
+	{
+		SetTargetLandscape(TargetLandscapeInfo);
+	}
+	// AActor::PostEditChange will ForceUpdateComponents()
+	Super::PostEditChangeProperty(PropertyChangedEvent);
+}
+
+
+void ALandscapeGizmoActiveActor::PostEditMove(bool bFinished)
+{
+	Super::PostEditMove( bFinished );
+
+	if (bFinished)
+	{
+		UnsnappedRotation = FRotator::ZeroRotator;
+	}
+}
+
+FVector ALandscapeGizmoActiveActor::SnapToLandscapeGrid(const FVector& GizmoLocation) const
+{
+	check(TargetLandscapeInfo);
+	const FTransform LToW = TargetLandscapeInfo->GetLandscapeProxy()->LandscapeActorToWorld();
+	const FVector LandscapeSpaceLocation = LToW.InverseTransformPosition(GizmoLocation);
+	const FVector SnappedLandscapeSpaceLocation = LandscapeSpaceLocation.GridSnap(1);
+	const FVector ResultLocation = LToW.TransformPosition(SnappedLandscapeSpaceLocation);
+	return ResultLocation;
+}
+
+void ALandscapeGizmoActiveActor::EditorApplyTranslation(const FVector& DeltaTranslation, bool bAltDown, bool bShiftDown, bool bCtrlDown)
+{
+	if (bSnapToLandscapeGrid)
+	{
+		const FVector GizmoLocation = GetActorLocation() + DeltaTranslation;
+		const FVector ResultLocation = SnapToLandscapeGrid(GizmoLocation);
+
+		SetActorLocation(ResultLocation, false);
+	}
+	else
+	{
+		Super::EditorApplyTranslation(DeltaTranslation, bAltDown, bShiftDown, bCtrlDown);
+	}
+
+	ReregisterAllComponents();
+}
+
+FRotator ALandscapeGizmoActiveActor::SnapToLandscapeGrid(const FRotator& GizmoRotation) const
+{
+	// Snap to multiples of 90 Yaw in landscape coordinate system
+	//check(TargetLandscapeInfo && TargetLandscapeInfo->LandscapeProxy);
+	//const FTransform LToW = TargetLandscapeInfo->LandscapeProxy->ActorToWorld();
+	//const FRotator LandscapeSpaceRotation = (LToW.GetRotation().InverseFast() * GizmoRotation.Quaternion()).Rotator().GetNormalized();
+	//const FRotator SnappedLandscapeSpaceRotation = FRotator(0, FMath::GridSnap(LandscapeSpaceRotation.Yaw, 90), 0);
+	//const FRotator ResultRotation = (SnappedLandscapeSpaceRotation.Quaternion() * LToW.GetRotation()).Rotator().GetNormalized();
+
+	// Gizmo rotation is used as if it was relative to the landscape even though it isn't, so snap in world space
+	const FRotator ResultRotation = FRotator(0, FMath::GridSnap(GizmoRotation.Yaw, 90), 0);
+	return ResultRotation;
+}
+
+void ALandscapeGizmoActiveActor::EditorApplyRotation(const FRotator& DeltaRotation, bool bAltDown, bool bShiftDown, bool bCtrlDown)
+{
+	if (bSnapToLandscapeGrid)
+	{
+		// Based on AActor::EditorApplyRotation
+		FRotator GizmoRotation = GetActorRotation() + UnsnappedRotation;
+		FRotator Winding, Remainder;
+		GizmoRotation.GetWindingAndRemainder(Winding, Remainder);
+		const FQuat ActorQ = Remainder.Quaternion();
+		const FQuat DeltaQ = DeltaRotation.Quaternion();
+		const FQuat ResultQ = DeltaQ * ActorQ;
+		const FRotator NewActorRotRem = FRotator( ResultQ );
+		FRotator DeltaRot = NewActorRotRem - Remainder;
+		DeltaRot.Normalize();
+
+		GizmoRotation += DeltaRot;
+
+		const FRotator ResultRotation = SnapToLandscapeGrid(GizmoRotation);
+
+		UnsnappedRotation = GizmoRotation - ResultRotation;
+		UnsnappedRotation.Pitch = 0;
+		UnsnappedRotation.Roll = 0;
+		UnsnappedRotation.Normalize();
+
+		SetActorRotation(ResultRotation);
+	}
+	else
+	{
+		Super::EditorApplyRotation(DeltaRotation, bAltDown, bShiftDown, bCtrlDown);
+	}
+
+	ReregisterAllComponents();
+}
+
+ALandscapeGizmoActor* ALandscapeGizmoActiveActor::SpawnGizmoActor()
+{
+	// ALandscapeGizmoActor is history for ALandscapeGizmoActiveActor
+	ALandscapeGizmoActor* NewActor = GetWorld()->SpawnActor<ALandscapeGizmoActor>();
+	Duplicate(NewActor);
+	return NewActor;
+}
+
+void ALandscapeGizmoActiveActor::SetTargetLandscape(ULandscapeInfo* LandscapeInfo)
+{
+	ULandscapeInfo* PrevInfo = TargetLandscapeInfo;
+	if (!LandscapeInfo || LandscapeInfo->HasAnyFlags(RF_BeginDestroyed))
+	{
+		TargetLandscapeInfo = NULL;
+		if (GetWorld())
+		{
+			for (auto It = GetLandscapeInfoMap(GetWorld()).Map.CreateIterator(); It; ++It)
+			{
+				LandscapeInfo = It.Value();
+				if (LandscapeInfo && !LandscapeInfo->HasAnyFlags(RF_BeginDestroyed))
+				{
+					//if (LandscapeInfo->LandscapeProxy && !LandscapeInfo->LandscapeProxy->HasAnyFlags(RF_BeginDestroyed) )
+					{
+						TargetLandscapeInfo = LandscapeInfo;
+						break;
+					}
+				}
+			}
+		}
+	}
+	else
+	{
+		TargetLandscapeInfo = LandscapeInfo;
+	}
+
+	if (TargetLandscapeInfo && TargetLandscapeInfo != PrevInfo)
+	{
+		MarginZ = TargetLandscapeInfo->DrawScale.Z * 3;
+		Width = Height = TargetLandscapeInfo->DrawScale.X * (TargetLandscapeInfo->ComponentSizeQuads+1);
+
+		float LengthZ;
+		FVector NewLocation = TargetLandscapeInfo->GetLandscapeCenterPos(LengthZ);
+		SetLength(LengthZ);
+		SetActorLocation( NewLocation, false );
+		SetActorRotation(FRotator::ZeroRotator);
+		ReregisterAllComponents();
+	}
+}
+
+void ALandscapeGizmoActiveActor::ClearGizmoData()
+{
+	DataType = LGT_None;
+	SelectedData.Empty();
+	LayerInfos.Empty();
+	ReregisterAllComponents();
+}
+
+void ALandscapeGizmoActiveActor::FitToSelection()
+{
+	if (TargetLandscapeInfo)
+	{
+		// Find fit size
+		int32 MinX = MAX_int32, MinY = MAX_int32, MaxX = MIN_int32, MaxY = MIN_int32;
+		TargetLandscapeInfo->GetSelectedExtent(MinX, MinY, MaxX, MaxY);
+		if (MinX != MAX_int32)
+		{
+			float ScaleXY = TargetLandscapeInfo->DrawScale.X;
+			Width = ScaleXY * (MaxX - MinX + 1) / (GetRootComponent()->RelativeScale3D.X);
+			Height = ScaleXY * (MaxY - MinY + 1) / (GetRootComponent()->RelativeScale3D.Y);
+			float LengthZ;
+			FVector NewLocation = TargetLandscapeInfo->GetLandscapeCenterPos(LengthZ, MinX, MinY, MaxX, MaxY);
+			SetLength(LengthZ);
+			SetActorLocation(NewLocation, false);
+			SetActorRotation(FRotator::ZeroRotator);
+			// Reset Z render scale values...
+			MinRelativeZ = 0.f;
+			RelativeScaleZ = 1.f;
+			ReregisterAllComponents();
+		}
+	}
+}
+
+void ALandscapeGizmoActiveActor::FitMinMaxHeight()
+{
+	if (TargetLandscapeInfo)
+	{
+		float MinZ = HALF_WORLD_MAX, MaxZ = -HALF_WORLD_MAX;
+		// Change MinRelativeZ and RelativeZScale to fit Gizmo Box
+		for (auto It = SelectedData.CreateConstIterator(); It; ++It )
+		{
+			const FGizmoSelectData& Data = It.Value();
+			MinZ = FMath::Min(MinZ, Data.HeightData);
+			MaxZ = FMath::Max(MaxZ, Data.HeightData);
+		}
+
+		if (MinZ != HALF_WORLD_MAX && MaxZ > MinZ + KINDA_SMALL_NUMBER)
+		{
+			MinRelativeZ = MinZ;
+			RelativeScaleZ = 1.f / (MaxZ - MinZ);
+			ReregisterAllComponents();
+		}
+	}
+}
+
+float ALandscapeGizmoActiveActor::GetNormalizedHeight(uint16 LandscapeHeight) const
+{
+	if (TargetLandscapeInfo)
+	{
+		ALandscapeProxy* Proxy = TargetLandscapeInfo->GetCurrentLevelLandscapeProxy(true);
+		if (Proxy)
+		{
+			// Need to make it scale...?
+			float ZScale = GetLength();
+			if (ZScale > KINDA_SMALL_NUMBER)
+			{
+				FVector LocalGizmoPos = Proxy->LandscapeActorToWorld().InverseTransformPosition(GetActorLocation());
+				return FMath::Clamp<float>( (( LandscapeDataAccess::GetLocalHeight(LandscapeHeight) - LocalGizmoPos.Z) * TargetLandscapeInfo->DrawScale.Z) / ZScale, 0.f, 1.f );
+			}
+		}
+	}
+	return 0.f;
+}
+
+float ALandscapeGizmoActiveActor::GetWorldHeight(float NormalizedHeight) const
+{
+	if (TargetLandscapeInfo)
+	{
+		ALandscapeProxy* Proxy = TargetLandscapeInfo->GetCurrentLevelLandscapeProxy(true);
+		if (Proxy)
+		{
+			float ZScale = GetLength();
+			if (ZScale > KINDA_SMALL_NUMBER)
+			{
+				FVector LocalGizmoPos = Proxy->LandscapeActorToWorld().InverseTransformPosition(GetActorLocation());
+				return NormalizedHeight * ZScale + LocalGizmoPos.Z * TargetLandscapeInfo->DrawScale.Z;
+			}
+		}
+	}
+	return 0.f;
+}
+
+float ALandscapeGizmoActiveActor::GetLandscapeHeight(float NormalizedHeight) const
+{
+	if (TargetLandscapeInfo)
+	{
+		NormalizedHeight = (NormalizedHeight - MinRelativeZ) * RelativeScaleZ;
+		float ScaleZ = TargetLandscapeInfo->DrawScale.Z;
+		return (GetWorldHeight(NormalizedHeight) / ScaleZ);
+	}
+	return 0.f;
+}
+
+void ALandscapeGizmoActiveActor::CalcNormal()
+{
+	int32 SquaredDataTex = DataTexSize * DataTexSize;
+	if (SampledHeight.Num() == SquaredDataTex && SampleSizeX > 0 && SampleSizeY > 0 )
+	{
+		if (SampledNormal.Num() != SquaredDataTex)
+		{
+			SampledNormal.Empty(SquaredDataTex);
+			SampledNormal.AddZeroed(SquaredDataTex);
+		}
+		for (int32 Y = 0; Y < SampleSizeY-1; ++Y)
+		{
+			for (int32 X = 0; X < SampleSizeX-1; ++X)
+			{
+				FVector Vert00 = SampledHeight[X + Y*DataTexSize];
+				FVector Vert01 = SampledHeight[X + (Y+1)*DataTexSize];
+				FVector Vert10 = SampledHeight[X+1 + Y*DataTexSize];
+				FVector Vert11 = SampledHeight[X+1 + (Y+1)*DataTexSize];
+
+				FVector FaceNormal1 = ((Vert00-Vert10) ^ (Vert10-Vert11)).SafeNormal();
+				FVector FaceNormal2 = ((Vert11-Vert01) ^ (Vert01-Vert00)).SafeNormal(); 
+
+				// contribute to the vertex normals.
+				SampledNormal[X + Y*DataTexSize] += FaceNormal1;
+				SampledNormal[X + (Y+1)*DataTexSize] += FaceNormal2;
+				SampledNormal[X+1 + Y*DataTexSize] += FaceNormal1 + FaceNormal2;
+				SampledNormal[X+1 + (Y+1)*DataTexSize] += FaceNormal1 + FaceNormal2;
+			}
+		}
+		for (int32 Y = 0; Y < SampleSizeY; ++Y)
+		{
+			for (int32 X = 0; X < SampleSizeX; ++X)
+			{
+				SampledNormal[X + Y*DataTexSize] = SampledNormal[X + Y*DataTexSize].SafeNormal();
+			}
+		}
+	}
+}
+
+void ALandscapeGizmoActiveActor::SampleData(int32 SizeX, int32 SizeY)
+{
+	if (TargetLandscapeInfo && GizmoTexture)
+	{
+		// Rasterize rendering Texture...
+		int32 TexSizeX = FMath::Min(ALandscapeGizmoActiveActor::DataTexSize, SizeX);
+		int32 TexSizeY = FMath::Min(ALandscapeGizmoActiveActor::DataTexSize, SizeY);
+		SampleSizeX = TexSizeX;
+		SampleSizeY = TexSizeY;
+
+		// Update Data Texture...
+		//DataTexture->SetFlags(RF_Transactional);
+		//DataTexture->Modify();
+
+		TextureScale = FVector2D( (float)SizeX / FMath::Max(ALandscapeGizmoActiveActor::DataTexSize, SizeX), (float)SizeY / FMath::Max(ALandscapeGizmoActiveActor::DataTexSize, SizeY));
+		uint8* TexData = GizmoTexture->Source.LockMip(0);
+		int32 GizmoTexSizeX = GizmoTexture->Source.GetSizeX();
+		for (int32 Y = 0; Y < TexSizeY; ++Y)
+		{
+			for (int32 X = 0; X < TexSizeX; ++X)
+			{
+				float TexX = X * SizeX / TexSizeX;
+				float TexY = Y * SizeY / TexSizeY;
+				int32 LX = FMath::FloorToInt(TexX);
+				int32 LY = FMath::FloorToInt(TexY);
+
+				float FracX = TexX - LX;
+				float FracY = TexY - LY;
+
+				FGizmoSelectData* Data00 = SelectedData.Find(ALandscape::MakeKey(LX, LY));
+				FGizmoSelectData* Data10 = SelectedData.Find(ALandscape::MakeKey(LX+1, LY));
+				FGizmoSelectData* Data01 = SelectedData.Find(ALandscape::MakeKey(LX, LY+1));
+				FGizmoSelectData* Data11 = SelectedData.Find(ALandscape::MakeKey(LX+1, LY+1));
+
+				// Invert Tex Data to show selected region more visible
+				TexData[X + Y*GizmoTexSizeX] = 255 - FMath::Lerp(
+					FMath::Lerp(Data00 ? Data00->Ratio : 0, Data10 ? Data10->Ratio : 0, FracX),
+					FMath::Lerp(Data01 ? Data01->Ratio : 0, Data11 ? Data11->Ratio : 0, FracX),
+					FracY
+					) * 255;
+
+				if (DataType & LGT_Height)
+				{
+					float NormalizedHeight = FMath::Lerp(
+						FMath::Lerp(Data00 ? Data00->HeightData : 0, Data10 ? Data10->HeightData : 0, FracX),
+						FMath::Lerp(Data01 ? Data01->HeightData : 0, Data11 ? Data11->HeightData : 0, FracX),
+						FracY
+						);
+
+					SampledHeight[X + Y*GizmoTexSizeX] = FVector(LX, LY, NormalizedHeight);
+				}
+			}
+		}
+
+		if (DataType & LGT_Height)
+		{
+			CalcNormal();
+		}
+
+		GizmoTexture->TemporarilyDisableStreaming();
+		FUpdateTextureRegion2D Region(0, 0, 0, 0, TexSizeX, TexSizeY);
+		GizmoTexture->UpdateTextureRegions(0, 1, &Region, GizmoTexSizeX, sizeof(uint8), TexData, false);
+		FlushRenderingCommands();
+		GizmoTexture->Source.UnlockMip(0);
+
+		ReregisterAllComponents();
+	}
+}
+
+LANDSCAPE_API void ALandscapeGizmoActiveActor::Import( int32 VertsX, int32 VertsY, uint16* HeightData, TArray<ULandscapeLayerInfoObject*> ImportLayerInfos, uint8* LayerDataPointers[] )
+{
+	if (VertsX <= 0 || VertsY <= 0 || HeightData == NULL || TargetLandscapeInfo == NULL || GizmoTexture == NULL || (ImportLayerInfos.Num() && !LayerDataPointers) )
+	{
+		return;
+	}
+
+	GWarn->BeginSlowTask( NSLOCTEXT("Landscape", "BeginImportingGizmoDataTask", "Importing Gizmo Data"), true);
+
+	ClearGizmoData();
+
+	CachedScaleXY = TargetLandscapeInfo->DrawScale.X;
+	CachedWidth = CachedScaleXY * VertsX; // (DrawScale * DrawScale3D.X);
+	CachedHeight = CachedScaleXY * VertsY; // (DrawScale * DrawScale3D.Y);
+	
+	float CurrentWidth = GetWidth();
+	float CurrentHeight = GetHeight();
+	LengthZ = GetLength();
+
+	FVector Scale3D = FVector(CurrentWidth / CachedWidth, CurrentHeight / CachedHeight, 1.f);
+	GetRootComponent()->SetRelativeScale3D(Scale3D);
+
+	Width = CachedWidth;
+	Height = CachedHeight;
+
+	DataType = ELandscapeGizmoType(DataType | LGT_Height);
+	if (ImportLayerInfos.Num())
+	{
+		DataType = ELandscapeGizmoType(DataType | LGT_Weight);
+	}
+
+	for (int32 Y = 0; Y < VertsY; ++Y)
+	{
+		for (int32 X = 0; X < VertsX; ++X)
+		{
+			FGizmoSelectData Data;
+			Data.Ratio = 1.f;
+			Data.HeightData = (float)HeightData[X + Y*VertsX] / 65535.f; //GetNormalizedHeight(HeightData[X + Y*VertsX]);
+			for (int32 i = 0; i < ImportLayerInfos.Num(); ++i)
+			{
+				Data.WeightDataMap.Add( ImportLayerInfos[i], LayerDataPointers[i][X + Y*VertsX] );
+			}
+			SelectedData.Add(ALandscape::MakeKey(X, Y), Data);
+		}
+	}
+
+	SampleData(VertsX, VertsY);
+
+	for (auto It = ImportLayerInfos.CreateConstIterator(); It; ++It)
+	{
+		LayerInfos.Add(*It);
+	}
+
+	GWarn->EndSlowTask();
+
+	ReregisterAllComponents();
+}
+
+void ALandscapeGizmoActiveActor::Export(int32 Index, TArray<FString>& Filenames)
+{
+	//guard around case where landscape has no layer structs
+	if (Filenames.Num() == 0)
+	{
+		return;
+	}
+
+	bool bExportOneTarget = (Filenames.Num() == 1);
+
+	if (TargetLandscapeInfo)
+	{
+		int32 MinX = MAX_int32, MinY = MAX_int32, MaxX = MIN_int32, MaxY = MIN_int32;
+		for (auto It = SelectedData.CreateConstIterator(); It; ++It )
+		{
+			int32 X, Y;
+			ALandscape::UnpackKey(It.Key(), X, Y);
+			if (MinX > X) MinX = X;
+			if (MaxX < X) MaxX = X;
+			if (MinY > Y) MinY = Y;
+			if (MaxY < Y) MaxY = Y;
+		}
+
+		if (MinX != MAX_int32)
+		{
+			GWarn->BeginSlowTask( NSLOCTEXT("Landscape", "BeginExportingGizmoDataTask", "Exporting Gizmo Data"), true);
+
+			TArray<uint8> HeightData;
+			if (!bExportOneTarget || Index == -1)
+			{
+				HeightData.AddZeroed((1+MaxX-MinX)*(1+MaxY-MinY)*sizeof(uint16));
+			}
+			uint16* pHeightData = (uint16*)HeightData.GetData();
+
+			TArray<TArray<uint8> > WeightDatas;
+			for( int32 i=1;i<Filenames.Num();i++ )
+			{
+				TArray<uint8> WeightData;
+				if (!bExportOneTarget || Index == i-1)
+				{
+					WeightData.AddZeroed((1+MaxX-MinX)*(1+MaxY-MinY));
+				}
+				WeightDatas.Add(WeightData);
+			}
+
+			for (int32 Y = MinY; Y <= MaxY; ++Y)
+			{
+				for (int32 X = MinX; X <= MaxX; ++X)
+				{
+					const FGizmoSelectData* Data = SelectedData.Find(ALandscape::MakeKey(X, Y));
+					if (Data)
+					{
+						int32 Idx = (X-MinX) + Y *(1+MaxX-MinX);
+						if (!bExportOneTarget || Index == -1)
+						{
+							pHeightData[Idx] = FMath::Clamp<uint16>(Data->HeightData * 65535.f, 0, 65535);
+						}
+
+						for( int32 i=1;i<Filenames.Num();i++ )
+						{
+							if (!bExportOneTarget || Index == i-1)
+							{
+								TArray<uint8>& WeightData = WeightDatas[i-1];
+								WeightData[Idx] = FMath::Clamp<uint8>(Data->WeightDataMap.FindRef(LayerInfos[i-1]), 0, 255);
+							}
+						}
+					}
+				}
+			}
+
+			if (!bExportOneTarget || Index == -1)
+			{
+				FFileHelper::SaveArrayToFile(HeightData,*Filenames[0]);
+			}
+
+			for( int32 i=1;i<Filenames.Num();i++ )
+			{
+				if (!bExportOneTarget || Index == i-1)
+				{
+					FFileHelper::SaveArrayToFile(WeightDatas[i-1],*Filenames[bExportOneTarget ? 0 : i]);
+				}
+			}
+
+			GWarn->EndSlowTask();
+		}
+		else
+		{
+			FMessageDialog::Open( EAppMsgType::Ok, NSLOCTEXT("UnrealEd", "LandscapeGizmoExport_Warning", "Landscape Gizmo has no copyed data. You need to choose proper targets and copy it to Gizmo."));
+		}
+	}
+}
+
+void ALandscapeGizmoActiveActor::ExportToClipboard()
+{
+	if (TargetLandscapeInfo && DataType != LGT_None)
+	{
+		//GWarn->BeginSlowTask( TEXT("Exporting Gizmo Data From Clipboard"), true);
+
+		FString ClipboardString(TEXT("GizmoData="));
+
+		ClipboardString += FString::Printf(TEXT(" Type=%d,TextureScaleX=%g,TextureScaleY=%g,SampleSizeX=%d,SampleSizeY=%d,CachedWidth=%g,CachedHeight=%g,CachedScaleXY=%g "), 
+			(int32)DataType, TextureScale.X, TextureScale.Y, SampleSizeX, SampleSizeY, CachedWidth, CachedHeight, CachedScaleXY);
+
+		for (int32 Y = 0; Y < SampleSizeY; ++Y )
+		{
+			for (int32 X = 0; X < SampleSizeX; ++X)
+			{
+				FVector& V = SampledHeight[X + Y * DataTexSize];
+				ClipboardString += FString::Printf(TEXT("%d %d %d "), (int32)V.X, (int32)V.Y, *(int32*)(&V.Z) );
+			}
+		}
+
+		ClipboardString += FString::Printf(TEXT("LayerInfos= "));
+
+		for (auto It = LayerInfos.CreateConstIterator(); It; ++It)
+		{
+			ClipboardString += FString::Printf(TEXT("%s "), *(*It)->GetPathName() );
+		}
+
+		ClipboardString += FString::Printf(TEXT("Region= "));
+
+		for (auto It = SelectedData.CreateConstIterator(); It; ++It )
+		{
+			int32 X, Y;
+			ALandscape::UnpackKey(It.Key(), X, Y);
+			const FGizmoSelectData& Data = It.Value();
+			ClipboardString += FString::Printf(TEXT("%d %d %d %d %d "), X, Y, *(int32*)(&Data.Ratio), *(int32*)(&Data.HeightData), Data.WeightDataMap.Num());
+
+			for (auto It2 = Data.WeightDataMap.CreateConstIterator(); It2; ++It2)
+			{
+				ClipboardString += FString::Printf(TEXT("%d %d "), LayerInfos.Find(It2.Key()), *(int32*)(&It2.Value()));
+			}
+		}
+
+		FPlatformMisc::ClipboardCopy(*ClipboardString);
+
+		//GWarn->EndSlowTask();
+	}
+}
+
+#define MAX_GIZMO_PROP_TEXT_LENGTH			1024*1024*8
+
+void ALandscapeGizmoActiveActor::ImportFromClipboard()
+{
+	FString ClipboardString;
+	FPlatformMisc::ClipboardPaste(ClipboardString);
+	const TCHAR* Str = *ClipboardString;
+	
+	if(FParse::Command(&Str,TEXT("GizmoData=")))
+	{
+		int32 ClipBoardSize = ClipboardString.Len();
+		if (ClipBoardSize > MAX_GIZMO_PROP_TEXT_LENGTH)
+		{
+			if( EAppReturnType::Yes != FMessageDialog::Open( EAppMsgType::YesNo,
+				FText::Format(NSLOCTEXT("UnrealEd", "LandscapeGizmoImport_Warning", "Landscape Gizmo is about to import large amount data ({0}MB) from the clipboard, which will take some time. Do you want to proceed?"),
+				FText::AsNumber(ClipBoardSize >> 20) ) ) )
+			{
+				return;
+			}
+		}
+
+		GWarn->BeginSlowTask( NSLOCTEXT("Landscape", "BeginImportingGizmoDataFromClipboardTask", "Importing Gizmo Data From Clipboard"), true);
+
+		FParse::Next(&Str);
+
+
+		int32 ReadNum = 0;
+
+		uint8 Type = 0;
+		ReadNum += FParse::Value(Str, TEXT("Type="), Type) ? 1 : 0;
+		DataType = (ELandscapeGizmoType)Type;
+
+		ReadNum += FParse::Value(Str, TEXT("TextureScaleX="), TextureScale.X) ? 1 : 0;
+		ReadNum += FParse::Value(Str, TEXT("TextureScaleY="), TextureScale.Y) ? 1 : 0;
+		ReadNum += FParse::Value(Str, TEXT("SampleSizeX="), SampleSizeX) ? 1 : 0;
+		ReadNum += FParse::Value(Str, TEXT("SampleSizeY="), SampleSizeY) ? 1 : 0;
+		ReadNum += FParse::Value(Str, TEXT("CachedWidth="), CachedWidth) ? 1 : 0;
+		ReadNum += FParse::Value(Str, TEXT("CachedHeight="), CachedHeight) ? 1 : 0;
+		ReadNum += FParse::Value(Str, TEXT("CachedScaleXY="), CachedScaleXY) ? 1 : 0;
+
+		if (ReadNum > 0)
+		{
+			while (!FChar::IsWhitespace(*Str))
+			{
+				Str++;
+			}
+			FParse::Next(&Str);
+
+			int32 SquaredDataTex = DataTexSize * DataTexSize;
+			if (SampledHeight.Num() != SquaredDataTex)
+			{
+				SampledHeight.Empty(SquaredDataTex);
+				SampledHeight.AddZeroed(SquaredDataTex);
+			}
+
+			// For Sample Height...
+			TCHAR* StopStr;
+			for (int32 Y = 0; Y < SampleSizeY; ++Y )
+			{
+				for (int32 X = 0; X < SampleSizeX; ++X)
+				{
+					FVector& V = SampledHeight[X + Y * DataTexSize];
+					V.X = FCString::Strtoi(Str, &StopStr, 10);
+					while (!FChar::IsWhitespace(*Str))
+					{
+						Str++;
+					}
+					FParse::Next(&Str);
+					V.Y = FCString::Strtoi(Str, &StopStr, 10);
+					while (!FChar::IsWhitespace(*Str))
+					{
+						Str++;
+					}
+					FParse::Next(&Str);
+					//V.Z = FCString::Atof(Str);
+					*((int32*)(&V.Z)) = FCString::Strtoi(Str, &StopStr, 10);
+					while (!FChar::IsWhitespace(*Str))
+					{
+						Str++;
+					}
+					FParse::Next(&Str);
+				}
+			}
+
+			CalcNormal();
+
+			TCHAR StrBuf[1024];
+			if(FParse::Command(&Str,TEXT("LayerInfos=")))
+			{
+				while( !FParse::Command(&Str,TEXT("Region=")) )
+				{
+					FParse::Next(&Str);
+					int i = 0;
+					while (!FChar::IsWhitespace(*Str))
+					{
+						StrBuf[i++] = *Str;
+						Str++;
+					}
+					StrBuf[i] = 0;
+					LayerInfos.Add( LoadObject<ULandscapeLayerInfoObject>(NULL, StrBuf) );
+				}
+			}
+
+			//if(FParse::Command(&Str,TEXT("Region=")))
+			{
+				while (*Str)
+				{
+					FParse::Next(&Str);
+					int32 X, Y, LayerNum;
+					FGizmoSelectData Data;
+					X = FCString::Strtoi(Str, &StopStr, 10);
+					while (!FChar::IsWhitespace(*Str))
+					{
+						Str++;
+					}
+					FParse::Next(&Str);
+					Y = FCString::Strtoi(Str, &StopStr, 10);
+					while (!FChar::IsWhitespace(*Str))
+					{
+						Str++;
+					}
+					FParse::Next(&Str);
+					*((int32*)(&Data.Ratio)) = FCString::Strtoi(Str, &StopStr, 10);
+					while (!FChar::IsWhitespace(*Str))
+					{
+						Str++;
+					}
+					FParse::Next(&Str);
+					*((int32*)(&Data.HeightData)) = FCString::Strtoi(Str, &StopStr, 10);
+					while (!FChar::IsWhitespace(*Str))
+					{
+						Str++;
+					}
+					FParse::Next(&Str);
+					LayerNum = FCString::Strtoi(Str, &StopStr, 10);
+					while (!FChar::IsWhitespace(*Str))
+					{
+						Str++;
+					}
+					FParse::Next(&Str);
+					for (int32 i = 0; i < LayerNum; ++i)
+					{
+						int32 LayerIndex = FCString::Strtoi(Str, &StopStr, 10);
+						while (!FChar::IsWhitespace(*Str))
+						{
+							Str++;
+						}
+						FParse::Next(&Str);
+						float Weight;
+						*((int32*)(&Weight)) = FCString::Strtoi(Str, &StopStr, 10);
+						while (!FChar::IsWhitespace(*Str))
+						{
+							Str++;
+						}
+						FParse::Next(&Str);
+						Data.WeightDataMap.Add(LayerInfos[LayerIndex], Weight);
+					}
+					SelectedData.Add(ALandscape::MakeKey(X, Y), Data);
+				}
+			}
+		}
+
+		GWarn->EndSlowTask();
+
+		ReregisterAllComponents();
+	}
+}
+#endif	//WITH_EDITOR
diff --git a/Engine/Source/Runtime/Landscape/Private/LandscapeInfo.cpp b/Engine/Source/Runtime/Landscape/Private/LandscapeInfo.cpp
new file mode 100644
--- /dev/null
+++ b/Engine/Source/Runtime/Landscape/Private/LandscapeInfo.cpp
@@ -0,0 +1,14 @@
+#include "Landscape.h"
+#include "LandscapeInfo.h"
+#include "LandscapeLayerInfoObject.h"
+
+LANDSCAPE_API FLandscapeInfoLayerSettings::FLandscapeInfoLayerSettings(ULandscapeLayerInfoObject* InLayerInfo, class ALandscapeProxy* InProxy)
+: LayerInfoObj(InLayerInfo)
+, LayerName((InLayerInfo != NULL) ? InLayerInfo->LayerName : NAME_None)
+#if WITH_EDITORONLY_DATA
+, ThumbnailMIC(NULL)
+, Owner(InProxy)
+, DebugColorChannel(0)
+, bValid(false)
+#endif
+{ }
\ No newline at end of file
diff --git a/Engine/Source/Runtime/Landscape/Private/LandscapeInfoMap.cpp b/Engine/Source/Runtime/Landscape/Private/LandscapeInfoMap.cpp
new file mode 100644
--- /dev/null
+++ b/Engine/Source/Runtime/Landscape/Private/LandscapeInfoMap.cpp
@@ -0,0 +1,40 @@
+// Copyright 1998-2014 Epic Games, Inc. All Rights Reserved.
+
+#include "Landscape.h"
+#include "LandscapeInfoMap.h"
+
+ULandscapeInfoMap::ULandscapeInfoMap(const FObjectInitializer& ObjectInitializer)
+	: Super(ObjectInitializer)
+{
+}
+
+void ULandscapeInfoMap::PostDuplicate(bool bDuplicateForPIE)
+{
+	// After duplication the pointers are duplicated, but the objects should
+	// be duplicated.
+
+	Super::PostDuplicate(bDuplicateForPIE);
+	
+	for (auto Iterator = Map.CreateConstIterator(); Iterator; ++Iterator)
+	{
+		Map[Iterator->Key] = Cast<ULandscapeInfo>(StaticDuplicateObject(
+			Iterator->Value, Iterator->Value->GetOuter(), nullptr));
+	}
+}
+
+void ULandscapeInfoMap::Serialize(FArchive& Ar)
+{
+	Super::Serialize(Ar);
+
+	Ar << Map;
+}
+
+void ULandscapeInfoMap::AddReferencedObjects(UObject* InThis, FReferenceCollector& Collector)
+{
+	ULandscapeInfoMap* This = CastChecked<ULandscapeInfoMap>(InThis);
+
+	for (auto ReferencedMapInfo = This->Map.CreateIterator(); ReferencedMapInfo; ++ReferencedMapInfo)
+	{
+		Collector.AddReferencedObject(ReferencedMapInfo->Value, This);
+	}
+}
diff --git a/Engine/Source/Runtime/Landscape/Private/LandscapeLight.cpp b/Engine/Source/Runtime/Landscape/Private/LandscapeLight.cpp
new file mode 100644
--- /dev/null
+++ b/Engine/Source/Runtime/Landscape/Private/LandscapeLight.cpp
@@ -0,0 +1,688 @@
+// Copyright 1998-2014 Epic Games, Inc. All Rights Reserved.
+
+/*=============================================================================
+LandscapeLight.cpp: Static lighting for LandscapeComponents
+=============================================================================*/
+
+#include "Landscape.h"
+#include "LandscapeLight.h"
+#include "LandscapeInfo.h"
+#include "LandscapeRender.h"
+#include "LandscapeDataAccess.h"
+
+#include "ComponentReregisterContext.h"
+
+#include "UnrealEngine.h"
+
+#if WITH_EDITOR
+
+#define LANDSCAPE_LIGHTMAP_UV_INDEX 1
+
+TMap<FIntPoint, FColor> FLandscapeStaticLightingMesh::LandscapeUpscaleHeightDataCache;
+TMap<FIntPoint, FColor> FLandscapeStaticLightingMesh::LandscapeUpscaleXYOffsetDataCache;
+
+/** A texture mapping for landscapes */
+/** Initialization constructor. */
+FLandscapeStaticLightingTextureMapping::FLandscapeStaticLightingTextureMapping(
+	ULandscapeComponent* InComponent,FStaticLightingMesh* InMesh,int32 InLightMapWidth,int32 InLightMapHeight,bool bPerformFullQualityRebuild) :
+FStaticLightingTextureMapping(
+							  InMesh,
+							  InComponent,
+							  InLightMapWidth,
+							  InLightMapHeight,
+							  LANDSCAPE_LIGHTMAP_UV_INDEX
+							  ),
+							  LandscapeComponent(InComponent)
+{
+}
+
+void FLandscapeStaticLightingTextureMapping::Apply(FQuantizedLightmapData* QuantizedData, const TMap<ULightComponent*,FShadowMapData2D*>& ShadowMapData)
+{
+	//ELightMapPaddingType PaddingType = GAllowLightmapPadding ? LMPT_NormalPadding : LMPT_NoPadding;
+	ELightMapPaddingType PaddingType = LMPT_NoPadding;
+
+	const bool bHasNonZeroData = QuantizedData != NULL && QuantizedData->HasNonZeroData();
+
+	// We always create a light map if the surface either has any non-zero lighting data, or if the surface has a shadow map.  The runtime
+	// shaders are always expecting a light map in the case of a shadow map, even if the lighting is entirely zero.  This is simply to reduce
+	// the number of shader permutations to support in the very unlikely case of a unshadowed surfaces that has lighting values of zero.
+	const bool bNeedsLightMap = bHasNonZeroData || ShadowMapData.Num() > 0 || (QuantizedData != NULL && QuantizedData->bHasSkyShadowing);
+	if (bNeedsLightMap)
+	{
+		// Create a light-map for the primitive.
+		LandscapeComponent->LightMap = FLightMap2D::AllocateLightMap(
+			LandscapeComponent,
+			QuantizedData,
+			LandscapeComponent->Bounds,
+			PaddingType,
+			LMF_Streamed
+			);
+	}
+	else
+	{
+		LandscapeComponent->LightMap = NULL;
+	}
+
+	if (ShadowMapData.Num() > 0)
+	{
+		LandscapeComponent->ShadowMap = FShadowMap2D::AllocateShadowMap(
+			LandscapeComponent,
+			ShadowMapData,
+			LandscapeComponent->Bounds,
+			PaddingType,
+			SMF_Streamed
+			);
+	}
+	else
+	{
+		LandscapeComponent->ShadowMap = NULL;
+	}
+
+	// Build the list of statically irrelevant lights.
+	// TODO: This should be stored per LOD.
+	LandscapeComponent->IrrelevantLights.Empty();
+	for(int32 LightIndex = 0;LightIndex < Mesh->RelevantLights.Num();LightIndex++)
+	{
+		const ULightComponent* Light = Mesh->RelevantLights[LightIndex];
+
+		// Check if the light is stored in the light-map.
+		const bool bIsInLightMap = LandscapeComponent->LightMap && LandscapeComponent->LightMap->LightGuids.Contains(Light->LightGuid);
+
+		// Add the light to the statically irrelevant light list if it is in the potentially relevant light list, but didn't contribute to the light-map.
+		if(!bIsInLightMap)
+		{	
+			LandscapeComponent->IrrelevantLights.AddUnique(Light->LightGuid);
+		}
+	}
+
+	LandscapeComponent->bHasCachedStaticLighting = true;
+
+	// Mark the primitive's package as dirty.
+	LandscapeComponent->MarkPackageDirty();
+}
+
+namespace
+{
+	// Calculate Geometric LOD for lighting
+	int32 GetLightingLOD(const ULandscapeComponent* InComponent)
+	{
+		if (InComponent->LightingLODBias < 0)
+		{
+			return FMath::Clamp<int32>(InComponent->ForcedLOD >= 0 ? InComponent->ForcedLOD : InComponent->LODBias, 0, FMath::CeilLogTwo(InComponent->SubsectionSizeQuads + 1) - 1);
+		}
+		else
+		{
+			return InComponent->LightingLODBias;
+		}
+	}
+};
+
+/** Initialization constructor. */
+FLandscapeStaticLightingMesh::FLandscapeStaticLightingMesh(ULandscapeComponent* InComponent, const TArray<ULightComponent*>& InRelevantLights, int32 InExpandQuadsX, int32 InExpandQuadsY, float InLightMapRatio, int32 InLOD)
+	: FStaticLightingMesh(
+		FMath::Square(((InComponent->ComponentSizeQuads + 1) >> InLOD) - 1 + 2 * InExpandQuadsX) * 2,
+		FMath::Square(((InComponent->ComponentSizeQuads + 1) >> InLOD) - 1 + 2 * InExpandQuadsX) * 2,
+		FMath::Square(((InComponent->ComponentSizeQuads + 1) >> InLOD) + 2 * InExpandQuadsX),
+		FMath::Square(((InComponent->ComponentSizeQuads + 1) >> InLOD) + 2 * InExpandQuadsX),
+		0,
+		!!(InComponent->CastShadow | InComponent->bCastHiddenShadow),
+		false,
+		InRelevantLights,
+		InComponent,
+		InComponent->Bounds.GetBox(),
+		InComponent->GetLightingGuid()
+	)
+	, LandscapeComponent(InComponent)
+	, LightMapRatio(InLightMapRatio)
+	, ExpandQuadsX(InExpandQuadsX)
+	, ExpandQuadsY(InExpandQuadsY)
+{
+	const float LODScale = (float)InComponent->ComponentSizeQuads / (((InComponent->ComponentSizeQuads + 1) >> InLOD) - 1);
+	LocalToWorld = FTransform(FQuat::Identity, FVector::ZeroVector, FVector(LODScale, LODScale, 1)) * InComponent->ComponentToWorld;
+	ComponentSizeQuads = ((InComponent->ComponentSizeQuads + 1) >> InLOD) - 1;
+	NumVertices = ComponentSizeQuads + 2*InExpandQuadsX + 1;
+	NumQuads = NumVertices - 1;
+	UVFactor = LightMapRatio / NumVertices;
+	bReverseWinding = (LocalToWorld.GetDeterminant() < 0.0f);
+
+	int32 GeometricLOD = ::GetLightingLOD(InComponent);
+	GetHeightmapData(InLOD, FMath::Max(GeometricLOD, InLOD));
+}
+
+FLandscapeStaticLightingMesh::~FLandscapeStaticLightingMesh()
+{
+}
+
+namespace
+{
+	void GetLODData(ULandscapeComponent* LandscapeComponent, int32 X, int32 Y, int32 HeightmapOffsetX, int32 HeightmapOffsetY, int32 LODValue, int32 HeightmapStride, FColor& OutHeight, FColor& OutXYOffset)
+	{
+		int32 ComponentSize = ((LandscapeComponent->SubsectionSizeQuads + 1) * LandscapeComponent->NumSubsections) >> LODValue;
+		int32 LODHeightmapSize = LandscapeComponent->HeightmapTexture->Source.GetSizeX() >> LODValue;
+		float Ratio = (float)(LODHeightmapSize) / (HeightmapStride);
+
+		int32 CurrentHeightmapOffsetX = FMath::RoundToInt((float)(LODHeightmapSize)* LandscapeComponent->HeightmapScaleBias.Z);
+		int32 CurrentHeightmapOffsetY = FMath::RoundToInt((float)(LODHeightmapSize)* LandscapeComponent->HeightmapScaleBias.W);
+
+		float XX = FMath::Clamp<float>((X - HeightmapOffsetX) * Ratio, 0.f, ComponentSize - 1.f) + CurrentHeightmapOffsetX;
+		int32 XI = (int32)XX;
+		float XF = XX - XI;
+
+		float YY = FMath::Clamp<float>((Y - HeightmapOffsetY) * Ratio, 0.f, ComponentSize - 1.f) + CurrentHeightmapOffsetY;
+		int32 YI = (int32)YY;
+		float YF = YY - YI;
+
+		FLandscapeComponentDataInterface DataInterface(LandscapeComponent, LODValue);
+		FColor* HeightMipData = DataInterface.GetRawHeightData();
+		FColor* XYOffsetMipData = DataInterface.GetRawXYOffsetData();
+
+		FColor H1 = HeightMipData[XI + YI * LODHeightmapSize];
+		FColor H2 = HeightMipData[FMath::Min(XI + 1, LODHeightmapSize - 1) + YI * LODHeightmapSize];
+		FColor H3 = HeightMipData[XI + FMath::Min(YI + 1, LODHeightmapSize - 1) * LODHeightmapSize];
+		FColor H4 = HeightMipData[FMath::Min(XI + 1, LODHeightmapSize - 1) + FMath::Min(YI + 1, LODHeightmapSize - 1) * LODHeightmapSize];
+
+		uint16 Height = FMath::RoundToInt(FMath::Lerp(FMath::Lerp<float>(((H1.R << 8) + H1.G), ((H2.R << 8) + H2.G), XF),
+			FMath::Lerp<float>(((H3.R << 8) + H3.G), ((H4.R << 8) + H4.G), XF), YF));
+		uint8 B = FMath::RoundToInt(FMath::Lerp(FMath::Lerp<float>((H1.B), (H2.B), XF),
+			FMath::Lerp<float>((H3.B), (H4.B), XF), YF));
+		uint8 A = FMath::RoundToInt(FMath::Lerp<float>(FMath::Lerp((H1.A), (H2.A), XF),
+			FMath::Lerp<float>((H3.A), (H4.A), XF), YF));
+
+		OutHeight = FColor((Height >> 8), Height & 255, B, A);
+
+		if (LandscapeComponent->XYOffsetmapTexture)
+		{
+			FColor X1 = XYOffsetMipData[XI + YI * LODHeightmapSize];
+			FColor X2 = XYOffsetMipData[FMath::Min(XI + 1, LODHeightmapSize - 1) + YI * LODHeightmapSize];
+			FColor X3 = XYOffsetMipData[XI + FMath::Min(YI + 1, LODHeightmapSize - 1) * LODHeightmapSize];
+			FColor X4 = XYOffsetMipData[FMath::Min(XI + 1, LODHeightmapSize - 1) + FMath::Min(YI + 1, LODHeightmapSize - 1) * LODHeightmapSize];
+
+			uint16 XComp = FMath::RoundToInt(FMath::Lerp(FMath::Lerp<float>(((X1.R << 8) + X1.G), ((X2.R << 8) + X2.G), XF),
+				FMath::Lerp<float>(((X3.R << 8) + X3.G), ((X4.R << 8) + X4.G), XF), YF));
+			uint16 YComp = FMath::RoundToInt(FMath::Lerp(FMath::Lerp<float>(((X1.B << 8) + X1.A), ((X2.B << 8) + X2.A), XF),
+				FMath::Lerp<float>(((X3.B << 8) + X3.A), ((X4.B << 8) + X4.A), XF), YF));
+
+			OutXYOffset = FColor((XComp >> 8), XComp & 255, YComp >> 8, YComp & 255);
+		}
+	}
+
+	void InternalUpscaling(FLandscapeComponentDataInterface& DataInterface, ULandscapeComponent* LandscapeComponent, int32 InLOD, int32 GeometryLOD, TArray<FColor>& CompHeightData, TArray<FColor>& CompXYOffsetData)
+	{
+		// Upscaling using Landscape LOD system 
+		ULandscapeInfo* const Info = LandscapeComponent->GetLandscapeInfo();
+		check(Info);
+
+		FIntPoint ComponentBase = LandscapeComponent->GetSectionBase() / LandscapeComponent->ComponentSizeQuads;
+		int32 NeighborLODs[8] = { 0, 0, 0, 0, 0, 0, 0, 0 };
+
+		int32 MaxLOD = FMath::CeilLogTwo(LandscapeComponent->SubsectionSizeQuads + 1) - 1;
+		bool bNeedUpscaling = GeometryLOD > InLOD;
+		int32 NeighborIdx = 0;
+
+		for (int32 y = -1; y <= 1; y++)
+		{
+			for (int32 x = -1; x <= 1; x++)
+			{
+				if (x == 0 && y == 0)
+				{
+					continue;
+				}
+
+				ULandscapeComponent* Neighbor = Info->XYtoComponentMap.FindRef(ComponentBase + FIntPoint(x, y));
+				int32 NeighborLOD = -1;
+
+				if (Neighbor)
+				{
+					NeighborLOD = ::GetLightingLOD(Neighbor);
+				}
+				else
+				{
+					NeighborLOD = 0;
+					// Sample neighbor components to find maximum LOD
+					for (int32 yy = -1; yy <= 1; yy++)
+					{
+						for (int32 xx = -1; xx <= 1; xx++)
+						{
+							if (xx == 0 && yy == 0)
+							{
+								continue;
+							}
+
+							ULandscapeComponent* Neighbor = Info->XYtoComponentMap.FindRef(ComponentBase + FIntPoint(x+xx, y+yy));
+							if (Neighbor)
+							{
+								NeighborLOD = FMath::Max(::GetLightingLOD(Neighbor), NeighborLOD);
+							}
+						}
+					}
+				}
+
+				bNeedUpscaling |= (NeighborLOD > InLOD);
+				NeighborLODs[NeighborIdx++] = NeighborLOD;
+			}
+		}
+
+		if (bNeedUpscaling)
+		{
+			check(LandscapeComponent);
+			// Need Upscaling
+			int32 HeightmapStride = LandscapeComponent->HeightmapTexture->Source.GetSizeX() >> InLOD;
+
+			int32 HeightDataSize = HeightmapStride * HeightmapStride;
+			CompHeightData.Empty(HeightDataSize);
+			CompXYOffsetData.Empty(HeightDataSize);
+			CompHeightData.AddZeroed(HeightDataSize);
+			CompXYOffsetData.AddZeroed(HeightDataSize);
+
+			// Update for only component region for performance
+			int32 ComponentSize = ((LandscapeComponent->SubsectionSizeQuads + 1) * LandscapeComponent->NumSubsections) >> InLOD;
+
+			for (int32 Y = DataInterface.HeightmapComponentOffsetY; Y < DataInterface.HeightmapComponentOffsetY + ComponentSize; ++Y)
+			{
+				for (int32 X = DataInterface.HeightmapComponentOffsetX; X < DataInterface.HeightmapComponentOffsetX + ComponentSize; ++X)
+				{
+					FIntPoint IXY(X - DataInterface.HeightmapComponentOffsetX, Y - DataInterface.HeightmapComponentOffsetY);
+					IXY += ComponentBase * (ComponentSize - 1);
+
+					FColor* CachedHeight = FLandscapeStaticLightingMesh::LandscapeUpscaleHeightDataCache.Find(IXY);
+					FColor* CachedXYOffset = FLandscapeStaticLightingMesh::LandscapeUpscaleXYOffsetDataCache.Find(IXY);
+
+					if (CachedHeight)
+					{
+						CompHeightData[X + Y * HeightmapStride] = *CachedHeight;
+						if (CachedXYOffset)
+						{
+							CompXYOffsetData[X + Y * HeightmapStride] = *CachedXYOffset;
+						}
+					}
+					else
+					{
+						// LOD System similar to the shader
+						FVector2D XY(float(X - DataInterface.HeightmapComponentOffsetX) / (ComponentSize - 1), float(Y - DataInterface.HeightmapComponentOffsetY) / (ComponentSize - 1));
+						XY = XY - 0.5f;
+
+						float RealLOD = GeometryLOD;
+
+						if (XY.X < 0.f)
+						{
+							if (XY.Y < 0.f)
+							{
+								RealLOD = FMath::Lerp(
+									FMath::Lerp<float>(NeighborLODs[0], NeighborLODs[1], XY.X + 1.f),
+									FMath::Lerp<float>(NeighborLODs[3], GeometryLOD, XY.X + 1.f),
+									XY.Y + 1.f); // 0
+							}
+							else
+							{
+								RealLOD = FMath::Lerp(
+									FMath::Lerp<float>(NeighborLODs[3], GeometryLOD, XY.X + 1.f),
+									FMath::Lerp<float>(NeighborLODs[5], NeighborLODs[6], XY.X + 1.f),
+									XY.Y); // 2
+							}
+						}
+						else
+						{
+							if (XY.Y < 0.f)
+							{
+								RealLOD = FMath::Lerp(
+									FMath::Lerp<float>(NeighborLODs[1], NeighborLODs[2], XY.X),
+									FMath::Lerp<float>(GeometryLOD, NeighborLODs[4], XY.X),
+									XY.Y + 1.f); // 1
+							}
+							else
+							{
+								RealLOD = FMath::Lerp(
+									FMath::Lerp<float>(GeometryLOD, NeighborLODs[4], XY.X),
+									FMath::Lerp<float>(NeighborLODs[6], NeighborLODs[7], XY.X),
+									XY.Y); // 3
+							}
+						}
+
+						RealLOD = FMath::Min(RealLOD, (float)MaxLOD);
+
+						int32 LODValue = (int32)RealLOD;
+						float MorphAlpha = FMath::Fractional(RealLOD);
+
+						FColor Height[2];
+						FColor XYOffset[2];
+						::GetLODData(LandscapeComponent, X, Y, DataInterface.HeightmapComponentOffsetX, DataInterface.HeightmapComponentOffsetY,
+							FMath::Min(MaxLOD, LODValue), HeightmapStride, Height[0], XYOffset[0]);
+
+						// Interpolation between two LOD
+						if ((RealLOD > InLOD) && (LODValue + 1 <= MaxLOD) && MorphAlpha != 0.f)
+						{
+							::GetLODData(LandscapeComponent, X, Y, DataInterface.HeightmapComponentOffsetX, DataInterface.HeightmapComponentOffsetY,
+								FMath::Min(MaxLOD, LODValue + 1), HeightmapStride, Height[1], XYOffset[1]);
+
+							// Need interpolation
+							uint16 Height0 = (Height[0].R << 8) + Height[0].G;
+							uint16 Height1 = (Height[1].R << 8) + Height[1].G;
+							uint16 LerpHeight = FMath::RoundToInt(FMath::Lerp<float>(Height0, Height1, MorphAlpha));
+
+							CompHeightData[X + Y * HeightmapStride] =
+								FColor((LerpHeight >> 8), LerpHeight & 255,
+								FMath::RoundToInt(FMath::Lerp<float>(Height[0].B, Height[1].B, MorphAlpha)),
+								FMath::RoundToInt(FMath::Lerp<float>(Height[0].A, Height[1].A, MorphAlpha)));
+							if (LandscapeComponent->XYOffsetmapTexture)
+							{
+								uint16 XComp0 = (XYOffset[0].R << 8) + XYOffset[0].G;
+								uint16 XComp1 = (XYOffset[1].R << 8) + XYOffset[1].G;
+								uint16 LerpXComp = FMath::RoundToInt(FMath::Lerp<float>(XComp0, XComp1, MorphAlpha));
+
+								uint16 YComp0 = (XYOffset[0].B << 8) + XYOffset[0].A;
+								uint16 YComp1 = (XYOffset[1].B << 8) + XYOffset[1].A;
+								uint16 LerpYComp = FMath::RoundToInt(FMath::Lerp<float>(YComp0, YComp1, MorphAlpha));
+
+								CompXYOffsetData[X + Y * HeightmapStride] =
+									FColor(LerpXComp >> 8, LerpXComp & 255, LerpYComp >> 8, LerpYComp & 255);
+							}
+						}
+						else
+						{
+							CompHeightData[X + Y * HeightmapStride] = Height[0];
+							CompXYOffsetData[X + Y * HeightmapStride] = XYOffset[0];
+						}
+
+						// Caching current calculated value
+						FLandscapeStaticLightingMesh::LandscapeUpscaleHeightDataCache.Add(IXY, CompHeightData[X + Y * HeightmapStride]);
+						if (LandscapeComponent->XYOffsetmapTexture)
+						{
+							FLandscapeStaticLightingMesh::LandscapeUpscaleHeightDataCache.Add(IXY, CompXYOffsetData[X + Y * HeightmapStride]);
+						}
+					}
+				}
+			}
+
+			DataInterface.SetRawHeightData(&CompHeightData[0]);
+			if (LandscapeComponent->XYOffsetmapTexture)
+			{
+				DataInterface.SetRawXYOffsetData(&CompXYOffsetData[0]);
+			}
+		}
+	}
+};
+
+void FLandscapeStaticLightingMesh::GetHeightmapData(int32 InLOD, int32 GeometryLOD)
+{
+	ULandscapeInfo* const Info = LandscapeComponent->GetLandscapeInfo();
+	check(Info);
+
+	HeightData.Empty(FMath::Square(NumVertices));
+	HeightData.AddUninitialized(FMath::Square(NumVertices));
+
+	const int32 NumSubsections = LandscapeComponent->NumSubsections;
+	const int32 SubsectionSizeVerts = (LandscapeComponent->SubsectionSizeQuads + 1) >> InLOD;
+	const int32 SubsectionSizeQuads = SubsectionSizeVerts - 1;
+	FIntPoint ComponentBase = LandscapeComponent->GetSectionBase()/LandscapeComponent->ComponentSizeQuads;
+
+	// assume that ExpandQuad size <= SubsectionSizeQuads...
+	check(ExpandQuadsX <= SubsectionSizeQuads);
+	check(ExpandQuadsY <= SubsectionSizeQuads);
+
+	int32 MaxLOD = FMath::CeilLogTwo(LandscapeComponent->SubsectionSizeQuads + 1) - 1;
+
+	// copy heightmap data for this component...
+	{
+		// Data array for upscaling case
+		TArray<FColor> CompHeightData;
+		TArray<FColor> CompXYOffsetData;
+		FLandscapeComponentDataInterface DataInterface(LandscapeComponent, InLOD);
+		::InternalUpscaling(DataInterface, LandscapeComponent, InLOD, GeometryLOD, CompHeightData, CompXYOffsetData);
+
+		for (int32 Y = 0; Y < ComponentSizeQuads + 1; Y++)
+		{
+			const FColor* const Data = DataInterface.GetHeightData(0, Y);
+
+			for (int32 SubsectionX = 0; SubsectionX < NumSubsections; SubsectionX++)
+			{
+				const int32 X = SubsectionSizeQuads * SubsectionX;
+				const int32 CompX = X + FMath::Min(X/SubsectionSizeQuads, NumSubsections - 1);
+				const FColor* const SubsectionData = &Data[CompX];
+
+				// Copy the data
+				FMemory::Memcpy( &HeightData[X + ExpandQuadsX + (Y + ExpandQuadsY) * NumVertices], SubsectionData, SubsectionSizeVerts * sizeof(FColor));
+			}
+		}
+	}
+
+	// copy surrounding heightmaps...
+	for (int32 ComponentY = 0; ComponentY < 3; ComponentY++)
+	{
+		for (int32 ComponentX = 0; ComponentX < 3; ComponentX++)
+		{
+			if (ComponentX == 1 && ComponentY == 1)
+			{
+				// Ourself
+				continue;
+			}
+
+			const int32 XSource = (ComponentX == 0) ? (ComponentSizeQuads - ExpandQuadsX) : ((ComponentX == 1) ? 0 : 1);
+			const int32 YSource = (ComponentY == 0) ? (ComponentSizeQuads - ExpandQuadsY) : ((ComponentY == 1) ? 0 : 1);
+			const int32 XDest = (ComponentX == 0) ? 0 : ((ComponentX == 1) ? ExpandQuadsX : (ComponentSizeQuads + ExpandQuadsX + 1));
+			const int32 YDest = (ComponentY == 0) ? 0 : ((ComponentY == 1) ? ExpandQuadsY : (ComponentSizeQuads + ExpandQuadsY + 1));
+			const int32 XNum = (ComponentX == 1) ? (ComponentSizeQuads + 1) : ExpandQuadsX;
+			const int32 YNum = (ComponentY == 1) ? (ComponentSizeQuads + 1) : ExpandQuadsY;
+			const int32 XBackup = (ComponentX == 2) ? (ComponentSizeQuads + ExpandQuadsX) : ExpandQuadsX;
+			const int32 YBackup = (ComponentY == 2) ? (ComponentSizeQuads + ExpandQuadsY) : ExpandQuadsY;
+			const int32 XBackupNum = (ComponentX == 1) ? (ComponentSizeQuads + 1) : 1;
+			const int32 YBackupNum = (ComponentY == 1) ? (ComponentSizeQuads + 1) : 1;
+			
+			ULandscapeComponent* Neighbor = Info->XYtoComponentMap.FindRef(ComponentBase + FIntPoint((ComponentX - 1), (ComponentY - 1)));
+			if (Neighbor)
+			{
+				// Data array for upscaling case
+				TArray<FColor> CompHeightData;
+				TArray<FColor> CompXYOffsetData;
+				int32 NeighborGeometricLOD = ::GetLightingLOD(Neighbor);
+				FLandscapeComponentDataInterface DataInterface(Neighbor, InLOD);
+				::InternalUpscaling(DataInterface, Neighbor, InLOD, NeighborGeometricLOD, CompHeightData, CompXYOffsetData);
+				for (int32 Y = 0; Y < YNum; Y++)
+				{
+					const FColor* const Data = DataInterface.GetHeightData(0, YSource + Y);
+
+					int32 NextX;
+					for (int32 X = XSource; X < XSource + XNum; X = NextX)
+					{
+						NextX = (X / SubsectionSizeQuads + 1) * SubsectionSizeQuads + 1;
+
+						const int32 CompX = X + FMath::Min(X/SubsectionSizeQuads, NumSubsections - 1);
+						const FColor* const SubsectionData = &Data[CompX];
+
+						// Copy the data
+						FMemory::Memcpy( &HeightData[XDest + (X - XSource) + (YDest + Y) * NumVertices], SubsectionData, FMath::Min(NextX - X, XSource + XNum - X) * sizeof(FColor));
+					}
+				}
+			}
+			else
+			{
+				for (int32 Y = 0; Y < YNum; Y++)
+				{
+					for (int32 X = 0; X < XNum; X += XBackupNum)
+					{
+						const FColor* const BackupData = &HeightData[XBackup + (YBackup + (Y % YBackupNum)) * NumVertices];
+
+						// Copy the data
+						FMemory::Memcpy( &HeightData[XDest + X + (YDest + Y) * NumVertices], BackupData, XBackupNum * sizeof(FColor));
+					}
+				}
+			}
+		}
+	}
+}
+
+/** Fills in the static lighting vertex data for the Landscape vertex. */
+void FLandscapeStaticLightingMesh::GetStaticLightingVertex(int32 VertexIndex, FStaticLightingVertex& OutVertex) const
+{
+	const int32 X = VertexIndex % NumVertices;
+	const int32 Y = VertexIndex / NumVertices;
+
+	int32 LocalX = X-ExpandQuadsX;
+	int32 LocalY = Y-ExpandQuadsY;
+
+	const FColor* Data = &HeightData[X + Y * NumVertices];
+
+	OutVertex.WorldTangentZ.X = 2.f / 255.f * (float)Data->B - 1.f;
+	OutVertex.WorldTangentZ.Y = 2.f / 255.f * (float)Data->A - 1.f;
+	OutVertex.WorldTangentZ.Z = FMath::Sqrt(1.f - (FMath::Square(OutVertex.WorldTangentZ.X)+FMath::Square(OutVertex.WorldTangentZ.Y)));
+	OutVertex.WorldTangentX = FVector4(OutVertex.WorldTangentZ.Z, 0.f, -OutVertex.WorldTangentZ.X);
+	OutVertex.WorldTangentY = OutVertex.WorldTangentZ ^ OutVertex.WorldTangentX;
+
+	// Assume there is no rotation, so we don't need to do any LocalToWorld.
+	uint16 Height = (Data->R << 8) + Data->G;
+
+	OutVertex.WorldPosition = LocalToWorld.TransformPosition( FVector( LocalX, LocalY, LandscapeDataAccess::GetLocalHeight(Height) ) );	
+
+	OutVertex.TextureCoordinates[0] = FVector2D((float)X / NumVertices, (float)Y / NumVertices); 
+	OutVertex.TextureCoordinates[LANDSCAPE_LIGHTMAP_UV_INDEX].X = X * UVFactor;
+	OutVertex.TextureCoordinates[LANDSCAPE_LIGHTMAP_UV_INDEX].Y = Y * UVFactor;
+}
+
+void FLandscapeStaticLightingMesh::GetTriangle(int32 TriangleIndex,FStaticLightingVertex& OutV0,FStaticLightingVertex& OutV1,FStaticLightingVertex& OutV2) const
+{
+	int32 I0, I1, I2;
+	GetTriangleIndices(TriangleIndex,I0, I1, I2);
+	GetStaticLightingVertex(I0,OutV0);
+	GetStaticLightingVertex(I1,OutV1);
+	GetStaticLightingVertex(I2,OutV2);
+}
+
+void FLandscapeStaticLightingMesh::GetTriangleIndices(int32 TriangleIndex,int32& OutI0,int32& OutI1,int32& OutI2) const
+{
+	int32 QuadIndex = TriangleIndex >> 1;
+	int32 QuadTriIndex = TriangleIndex & 1;
+
+	int32 QuadX = QuadIndex % (NumVertices - 1);
+	int32 QuadY = QuadIndex / (NumVertices - 1);
+
+	switch(QuadTriIndex)
+	{
+	case 0:
+		OutI0 = (QuadX + 0) + (QuadY + 0) * NumVertices;
+		OutI1 = (QuadX + 1) + (QuadY + 1) * NumVertices;
+		OutI2 = (QuadX + 1) + (QuadY + 0) * NumVertices;
+		break;
+	case 1:
+		OutI0 = (QuadX + 0) + (QuadY + 0) * NumVertices;
+		OutI1 = (QuadX + 0) + (QuadY + 1) * NumVertices;
+		OutI2 = (QuadX + 1) + (QuadY + 1) * NumVertices;
+		break;
+	}
+
+	if (bReverseWinding)
+	{
+		Swap(OutI1, OutI2);
+	}
+}
+
+
+const static FName FLandscapeStaticLightingMesh_IntersectLightRayName(TEXT("FLandscapeStaticLightingMesh_IntersectLightRay"));
+
+FLightRayIntersection FLandscapeStaticLightingMesh::IntersectLightRay(const FVector& Start,const FVector& End,bool bFindNearestIntersection) const
+{
+	// Intersect the light ray with the terrain component.
+	FHitResult Result(1.0f);
+
+	FHitResult NewHitInfo;
+	FCollisionQueryParams NewTraceParams( FLandscapeStaticLightingMesh_IntersectLightRayName, true );
+	
+	const bool bIntersects = LandscapeComponent->LineTraceComponent( Result, Start, End, NewTraceParams );
+
+	// Setup a vertex to represent the intersection.
+	FStaticLightingVertex IntersectionVertex;
+	if(bIntersects)
+	{
+		IntersectionVertex.WorldPosition = Result.Location;
+		IntersectionVertex.WorldTangentZ = Result.Normal;
+	}
+	else
+	{
+		IntersectionVertex.WorldPosition.Set(0,0,0);
+		IntersectionVertex.WorldTangentZ.Set(0,0,1);
+	}
+	return FLightRayIntersection(bIntersects,IntersectionVertex);
+}
+
+
+void ULandscapeComponent::GetStaticLightingInfo(FStaticLightingPrimitiveInfo& OutPrimitiveInfo,const TArray<ULightComponent*>& InRelevantLights,const FLightingBuildOptions& Options)
+{
+	if( HasStaticLighting() )
+	{
+		float LightMapRes = StaticLightingResolution > 0.f ? StaticLightingResolution : GetLandscapeProxy()->StaticLightingResolution;
+		int32 PatchExpandCountX = 1;
+		int32 PatchExpandCountY = 1;
+		int32 DesiredSize = 1;
+		int32 LightingLOD = GetLandscapeProxy()->StaticLightingLOD;
+
+		float LightMapRatio = ::GetTerrainExpandPatchCount(LightMapRes, PatchExpandCountX, PatchExpandCountY, ComponentSizeQuads, (NumSubsections * (SubsectionSizeQuads+1)), DesiredSize, LightingLOD);
+
+		int32 SizeX = DesiredSize;
+		int32 SizeY = DesiredSize;
+
+		if (SizeX > 0 && SizeY > 0)
+		{
+			FLandscapeStaticLightingMesh* StaticLightingMesh = new FLandscapeStaticLightingMesh(this, InRelevantLights, PatchExpandCountX, PatchExpandCountY, LightMapRatio, LightingLOD);
+			OutPrimitiveInfo.Meshes.Add(StaticLightingMesh);
+			// Create a static lighting texture mapping
+			OutPrimitiveInfo.Mappings.Add(new FLandscapeStaticLightingTextureMapping(
+				this,StaticLightingMesh,SizeX,SizeY,true));
+		}
+	}
+}
+
+bool ULandscapeComponent::GetLightMapResolution( int32& Width, int32& Height ) const
+{
+	// Assuming DXT_1 compression at the moment...
+	float LightMapRes = StaticLightingResolution > 0.f ? StaticLightingResolution : GetLandscapeProxy()->StaticLightingResolution;
+	int32 PatchExpandCountX = 1;
+	int32 PatchExpandCountY = 1;
+	int32 DesiredSize = 1;
+	uint32 LightingLOD = GetLandscapeProxy()->StaticLightingLOD;
+
+	float LightMapRatio = ::GetTerrainExpandPatchCount(LightMapRes, PatchExpandCountX, PatchExpandCountY, ComponentSizeQuads, (NumSubsections * (SubsectionSizeQuads+1)), DesiredSize, LightingLOD);
+
+	Width = DesiredSize;
+	Height = DesiredSize;
+
+	return false;
+}
+
+void ULandscapeComponent::GetLightAndShadowMapMemoryUsage( int32& LightMapMemoryUsage, int32& ShadowMapMemoryUsage ) const
+{
+	int32 Width, Height;
+	GetLightMapResolution(Width, Height);
+
+	const auto FeatureLevel = GetWorld() ? GetWorld()->FeatureLevel : GMaxRHIFeatureLevel;
+
+	if(AllowHighQualityLightmaps(FeatureLevel))
+	{
+		LightMapMemoryUsage  = NUM_HQ_LIGHTMAP_COEF * (Width * Height * 4 / 3); // assuming DXT5
+	}
+	else
+	{
+		LightMapMemoryUsage  = NUM_LQ_LIGHTMAP_COEF * (Width * Height * 4 / 3) / 2; // assuming DXT1
+	}
+
+	ShadowMapMemoryUsage = (Width * Height * 4 / 3); // assuming G8
+	return;
+}
+#endif
+
+void ULandscapeComponent::InvalidateLightingCacheDetailed(bool bInvalidateBuildEnqueuedLighting, bool bTranslationOnly)
+{
+	if (bHasCachedStaticLighting)
+	{
+		Modify();
+
+		FComponentReregisterContext ReregisterContext(this);
+
+		// Block until the RT processes the unregister before modifying variables that it may need to access
+		FlushRenderingCommands();
+
+		Super::InvalidateLightingCacheDetailed(bInvalidateBuildEnqueuedLighting, bTranslationOnly);
+
+		// Discard all cached lighting.
+		IrrelevantLights.Empty();
+		LightMap = NULL;
+		ShadowMap = NULL;
+	}
+}
diff --git a/Engine/Source/Runtime/Landscape/Private/LandscapeModule.cpp b/Engine/Source/Runtime/Landscape/Private/LandscapeModule.cpp
new file mode 100644
--- /dev/null
+++ b/Engine/Source/Runtime/Landscape/Private/LandscapeModule.cpp
@@ -0,0 +1,117 @@
+// Copyright 1998-2014 Epic Games, Inc. All Rights Reserved.
+
+#include "Landscape.h"
+#include "LandscapeModule.h"
+
+class FLandscapeModule : public IModuleInterface
+{
+public:
+	/** IModuleInterface implementation */
+	void StartupModule() override;
+	void ShutdownModule() override;
+};
+
+/**
+ * Add landscape-specific per-world data.
+ *
+ * @param World A pointer to world that this data should be created for.
+ */
+void AddPerWorldLandscapeData(UWorld* World)
+{
+	if (!World->PerModuleDataObjects.FindItemByClass<ULandscapeInfoMap>())
+	{
+		World->PerModuleDataObjects.Add(NewObject<ULandscapeInfoMap>(GetTransientPackage()));
+	}
+}
+
+/**
+ * Gets landscape-specific material's static parameters values.
+ *
+ * @param OutStaticParameterSet A set that should be updated with found parameters values.
+ * @param Material Material instance to look for parameters.
+ */
+void LandscapeMaterialsParameterValuesGetter(FStaticParameterSet &OutStaticParameterSet, UMaterialInstance* Material);
+
+/**
+ * Updates landscape-specific material parameters.
+ *
+ * @param OutStaticParameterSet A set of parameters.
+ * @param Material A material to update.
+ */
+bool LandscapeMaterialsParameterSetUpdater(FStaticParameterSet &OutStaticParameterSet, UMaterial* Material);
+
+/**
+ * Function that will fire everytime a world is created.
+ *
+ * @param World A world that was created.
+ * @param IVS Initialization values.
+ */
+void WorldInitializationEventFunction(UWorld* World, const UWorld::InitializationValues IVS)
+{
+	AddPerWorldLandscapeData(World);
+}
+
+
+/**
+ * Function that will fire everytime a world is destroyed.
+ *
+ * @param World A world that's being destroyed.
+ */
+void WorldDestroyEventFunction(UWorld* World);
+
+/**
+ * A function that fires everytime a world is duplicated.
+ *
+ * @param World A world that was duplicated.
+ * @param bDuplicateForPIE If this duplication was done for PIE.
+ */
+void WorldDuplicateEventFunction(UWorld* World, bool bDuplicateForPIE)
+{
+	int32 Index;
+	ULandscapeInfoMap *InfoMap;
+	if (World->PerModuleDataObjects.FindItemByClass(&InfoMap, &Index))
+	{
+		World->PerModuleDataObjects[Index] = Cast<ULandscapeInfoMap>(
+			StaticDuplicateObject(InfoMap, InfoMap->GetOuter(), nullptr)
+			);
+	}
+	else
+	{
+		AddPerWorldLandscapeData(World);
+	}
+}
+
+void FLandscapeModule::StartupModule()
+{
+	// This code will execute after your module is loaded into memory (but after global variables are initialized, of course.)
+	UMaterialInstance::CustomStaticParametersGetters.AddStatic(
+		&LandscapeMaterialsParameterValuesGetter
+	);
+
+	UMaterialInstance::CustomParameterSetUpdaters.Add(
+		UMaterialInstance::FCustomParameterSetUpdaterDelegate::CreateStatic(
+			&LandscapeMaterialsParameterSetUpdater
+		)
+	);
+
+#if WITH_EDITORONLY_DATA
+	FWorldDelegates::OnPostWorldInitialization.AddStatic(
+		&WorldInitializationEventFunction
+	);
+	FWorldDelegates::OnPreWorldFinishDestroy.AddStatic(
+		&WorldDestroyEventFunction
+	);
+#endif // WITH_EDITORONLY_DATA
+
+	FWorldDelegates::OnPostDuplicate.AddStatic(
+		&WorldDuplicateEventFunction
+	);
+}
+
+void FLandscapeModule::ShutdownModule()
+{
+	// This function may be called during shutdown to clean up your module.  For modules that support dynamic reloading,
+	// we call this function before unloading the module.
+}
+
+IMPLEMENT_MODULE(FLandscapeModule, Landscape);
\ No newline at end of file
diff --git a/Engine/Source/Runtime/Landscape/Private/LandscapeProxy.cpp b/Engine/Source/Runtime/Landscape/Private/LandscapeProxy.cpp
new file mode 100644
--- /dev/null
+++ b/Engine/Source/Runtime/Landscape/Private/LandscapeProxy.cpp
@@ -0,0 +1,16 @@
+// Copyright 1998-2014 Epic Games, Inc. All Rights Reserved.
+
+#include "Landscape.h"
+#include "LandscapeProxy.h"
+#include "LandscapeInfo.h"
+
+#if WITH_EDITOR
+
+LANDSCAPE_API FLandscapeImportLayerInfo::FLandscapeImportLayerInfo(const struct FLandscapeInfoLayerSettings& InLayerSettings)
+:	LayerName(InLayerSettings.GetLayerName())
+,	LayerInfo(InLayerSettings.LayerInfoObj)
+,	ThumbnailMIC(NULL)
+,	SourceFilePath(InLayerSettings.GetEditorSettings().ReimportLayerFilePath)
+{ }
+
+#endif // WITH_EDITOR
\ No newline at end of file
