Landscape module from Engine split out.

diff --git a/Engine/Source/Runtime/Landscape/Private/LandscapeRender.cpp b/Engine/Source/Runtime/Landscape/Private/LandscapeRender.cpp
new file mode 100644
--- /dev/null
+++ b/Engine/Source/Runtime/Landscape/Private/LandscapeRender.cpp
@@ -0,0 +1,2876 @@
+// Copyright 1998-2014 Epic Games, Inc. All Rights Reserved.
+
+/*=============================================================================
+LandscapeRender.cpp: New terrain rendering
+=============================================================================*/
+
+#include "Landscape.h"
+
+#include "EngineGlobals.h"
+#include "Materials/MaterialExpressionTextureCoordinate.h"
+#include "Materials/MaterialExpressionLandscapeLayerCoords.h"
+#include "ShaderParameters.h"
+#include "ShaderParameterUtils.h"
+#include "RHIStaticStates.h"
+#include "LandscapeRender.h"
+#include "LandscapeEdit.h"
+#include "LevelUtils.h"
+#include "MaterialCompiler.h"
+#include "LandscapeMaterialInstanceConstant.h"
+#include "RawIndexBuffer.h"
+#include "Engine/LevelStreaming.h"
+#include "Engine/ShadowMapTexture2D.h"
+
+IMPLEMENT_UNIFORM_BUFFER_STRUCT(FLandscapeUniformShaderParameters,TEXT("LandscapeParameters"));
+
+#define LANDSCAPE_LOD_DISTANCE_FACTOR 2.f
+#define LANDSCAPE_MAX_COMPONENT_SIZE 255
+#define LANDSCAPE_LOD_SQUARE_ROOT_FACTOR 1.5f
+
+/*------------------------------------------------------------------------------
+	Forsyth algorithm for cache optimizing index buffers.
+------------------------------------------------------------------------------*/
+
+// Forsyth algorithm to optimize post-transformed vertex cache
+namespace
+{
+	// code for computing vertex score was taken, as much as possible
+	// directly from the original publication.
+	float ComputeVertexCacheScore(int32 CachePosition, uint32 VertexCacheSize)
+	{
+		const float FindVertexScoreCacheDecayPower = 1.5f;
+		const float FindVertexScoreLastTriScore = 0.75f;
+
+		float Score = 0.0f;
+		if ( CachePosition < 0 )
+		{
+			// Vertex is not in FIFO cache - no score.
+		}
+		else
+		{
+			if ( CachePosition < 3 )
+			{
+				// This vertex was used in the last triangle,
+				// so it has a fixed score, whichever of the three
+				// it's in. Otherwise, you can get very different
+				// answers depending on whether you add
+				// the triangle 1,2,3 or 3,1,2 - which is silly.
+				Score = FindVertexScoreLastTriScore;
+			}
+			else
+			{
+				check( CachePosition < (int32)VertexCacheSize );
+				// Points for being high in the cache.
+				const float Scaler = 1.0f / ( VertexCacheSize - 3 );
+				Score = 1.0f - ( CachePosition - 3 ) * Scaler;
+				Score = FMath::Pow( Score, FindVertexScoreCacheDecayPower );
+			}
+		}
+
+		return Score;
+	}
+
+	float ComputeVertexValenceScore(uint32 numActiveFaces)
+	{
+		const float FindVertexScoreValenceBoostScale = 2.0f;
+		const float FindVertexScoreValenceBoostPower = 0.5f;
+
+		float Score = 0.f;
+
+		// Bonus points for having a low number of tris still to
+		// use the vert, so we get rid of lone verts quickly.
+		float ValenceBoost = FMath::Pow(float(numActiveFaces), -FindVertexScoreValenceBoostPower );
+		Score += FindVertexScoreValenceBoostScale * ValenceBoost;
+
+		return Score;
+	}
+
+	const uint32 MaxVertexCacheSize = 64;
+	const uint32 MaxPrecomputedVertexValenceScores = 64;
+	float VertexCacheScores[MaxVertexCacheSize+1][MaxVertexCacheSize];
+	float VertexValenceScores[MaxPrecomputedVertexValenceScores];
+	bool bVertexScoresComputed = false; //ComputeVertexScores();
+
+	bool ComputeVertexScores()
+	{
+		for (uint32 CacheSize=0; CacheSize<=MaxVertexCacheSize; ++CacheSize)
+		{
+			for (uint32 CachePos=0; CachePos<CacheSize; ++CachePos)
+			{
+				VertexCacheScores[CacheSize][CachePos] = ComputeVertexCacheScore(CachePos, CacheSize);
+			}
+		}
+
+		for (uint32 Valence=0; Valence<MaxPrecomputedVertexValenceScores; ++Valence)
+		{
+			VertexValenceScores[Valence] = ComputeVertexValenceScore(Valence);
+		}
+
+		return true;
+	}
+
+	inline float FindVertexCacheScore(uint32 CachePosition, uint32 MaxSizeVertexCache)
+	{
+		return VertexCacheScores[MaxSizeVertexCache][CachePosition];
+	}
+
+	inline float FindVertexValenceScore(uint32 NumActiveTris)
+	{
+		return VertexValenceScores[NumActiveTris];
+	}
+
+	float FindVertexScore(uint32 NumActiveFaces, uint32 CachePosition, uint32 VertexCacheSize)
+	{
+		check(bVertexScoresComputed);
+
+		if ( NumActiveFaces == 0 )
+		{
+			// No tri needs this vertex!
+			return -1.0f;
+		}
+
+		float Score = 0.f;
+		if (CachePosition < VertexCacheSize)
+		{
+			Score += VertexCacheScores[VertexCacheSize][CachePosition];
+		}
+
+		if (NumActiveFaces < MaxPrecomputedVertexValenceScores)
+		{
+			Score += VertexValenceScores[NumActiveFaces];
+		}
+		else
+		{
+			Score += ComputeVertexValenceScore(NumActiveFaces);
+		}
+
+		return Score;
+	}
+
+	struct OptimizeVertexData
+	{
+		float  Score;
+		uint32  ActiveFaceListStart;
+		uint32  ActiveFaceListSize;
+		uint32  CachePos0;
+		uint32  CachePos1;
+		OptimizeVertexData() : Score(0.f), ActiveFaceListStart(0), ActiveFaceListSize(0), CachePos0(0), CachePos1(0) { }
+	};
+
+	//-----------------------------------------------------------------------------
+	//  OptimizeFaces
+	//-----------------------------------------------------------------------------
+	//  Parameters:
+	//      InIndexList
+	//          input index list
+	//      OutIndexList
+	//          a pointer to a preallocated buffer the same size as indexList to
+	//          hold the optimized index list
+	//      LRUCacheSize
+	//          the size of the simulated post-transform cache (max:64)
+	//-----------------------------------------------------------------------------
+
+	template <typename INDEX_TYPE>
+	void OptimizeFaces(const TArray<INDEX_TYPE>& InIndexList, TArray<INDEX_TYPE>& OutIndexList, uint16 LRUCacheSize)
+	{
+		uint32 VertexCount = 0;
+		const uint32 IndexCount = InIndexList.Num();
+
+		// compute face count per vertex
+		for (uint32 i = 0; i < IndexCount; ++i)
+		{
+			uint32 Index = InIndexList[i];
+			VertexCount = FMath::Max(Index, VertexCount);
+		}
+		VertexCount++;
+
+		TArray<OptimizeVertexData> VertexDataList;
+		VertexDataList.Empty(VertexCount);
+		for (uint32 i = 0; i < VertexCount; i++)
+		{
+			VertexDataList.Add(OptimizeVertexData());
+		}
+
+		OutIndexList.Empty(IndexCount);
+		OutIndexList.AddZeroed(IndexCount);
+
+		// compute face count per vertex
+		for (uint32 i = 0; i < IndexCount; ++i)
+		{
+			uint32 Index = InIndexList[i];
+			OptimizeVertexData& VertexData = VertexDataList[Index];
+			VertexData.ActiveFaceListSize++;
+		}
+
+		TArray<uint32> ActiveFaceList;
+
+		const uint32 EvictedCacheIndex = TNumericLimits<uint32>::Max();
+
+		{
+			// allocate face list per vertex
+			uint32 CurActiveFaceListPos = 0;
+			for (uint32 i = 0; i < VertexCount; ++i)
+			{
+				OptimizeVertexData& VertexData = VertexDataList[i];
+				VertexData.CachePos0 = EvictedCacheIndex;
+				VertexData.CachePos1 = EvictedCacheIndex;
+				VertexData.ActiveFaceListStart = CurActiveFaceListPos;
+				CurActiveFaceListPos += VertexData.ActiveFaceListSize;
+				VertexData.Score = FindVertexScore(VertexData.ActiveFaceListSize, VertexData.CachePos0, LRUCacheSize);
+				VertexData.ActiveFaceListSize = 0;
+			}
+			ActiveFaceList.Empty(CurActiveFaceListPos);
+			ActiveFaceList.AddZeroed(CurActiveFaceListPos);
+		}
+
+		// fill out face list per vertex
+		for (uint32 i = 0; i < IndexCount; i+=3)
+		{
+			for (uint32 j = 0; j < 3; ++j)
+			{
+				uint32 Index = InIndexList[i+j];
+				OptimizeVertexData& VertexData = VertexDataList[Index];
+				ActiveFaceList[VertexData.ActiveFaceListStart + VertexData.ActiveFaceListSize] = i;
+				VertexData.ActiveFaceListSize++;
+			}
+		}
+
+		TArray<uint8> ProcessedFaceList;
+		ProcessedFaceList.Empty(IndexCount);
+		ProcessedFaceList.AddZeroed(IndexCount);
+
+		uint32 VertexCacheBuffer[(MaxVertexCacheSize+3)*2];
+		uint32* Cache0 = VertexCacheBuffer;
+		uint32* Cache1 = VertexCacheBuffer + (MaxVertexCacheSize+3);
+		uint32 EntriesInCache0 = 0;
+
+		uint32 BestFace = 0;
+		float BestScore = -1.f;
+
+		const float MaxValenceScore = FindVertexScore(1, EvictedCacheIndex, LRUCacheSize) * 3.f;
+
+		for (uint32 i = 0; i < IndexCount; i += 3)
+		{
+			if (BestScore < 0.f)
+			{
+				// no verts in the cache are used by any unprocessed faces so
+				// search all unprocessed faces for a new starting point
+				for (uint32 j = 0; j < IndexCount; j += 3)
+				{
+					if (ProcessedFaceList[j] == 0)
+					{
+						uint32 Face = j;
+						float FaceScore = 0.f;
+						for (uint32 k = 0; k < 3; ++k)
+						{
+							uint32 Index = InIndexList[Face+k];
+							OptimizeVertexData& VertexData = VertexDataList[Index];
+							check(VertexData.ActiveFaceListSize > 0);
+							check(VertexData.CachePos0 >= LRUCacheSize);
+							FaceScore += VertexData.Score;
+						}
+
+						if (FaceScore > BestScore)
+						{
+							BestScore = FaceScore;
+							BestFace = Face;
+
+							check(BestScore <= MaxValenceScore);
+							if (BestScore >= MaxValenceScore)
+							{
+								break;
+							}
+						}
+					}
+				}
+				check(BestScore >= 0.f);
+			}
+
+			ProcessedFaceList[BestFace] = 1;
+			uint32 EntriesInCache1 = 0;
+
+			// add bestFace to LRU cache and to newIndexList
+			for (uint32 V = 0; V < 3; ++V)
+			{
+				INDEX_TYPE Index = InIndexList[BestFace+V];
+				OutIndexList[i+V] = Index;
+
+				OptimizeVertexData& VertexData = VertexDataList[Index];
+
+				if (VertexData.CachePos1 >= EntriesInCache1)
+				{
+					VertexData.CachePos1 = EntriesInCache1;
+					Cache1[EntriesInCache1++] = Index;
+
+					if (VertexData.ActiveFaceListSize == 1)
+					{
+						--VertexData.ActiveFaceListSize;
+						continue;
+					}
+				}
+
+				check(VertexData.ActiveFaceListSize > 0);
+				uint32 FindIndex;
+				for (FindIndex = VertexData.ActiveFaceListStart; FindIndex < VertexData.ActiveFaceListStart + VertexData.ActiveFaceListSize; FindIndex++)
+				{
+					if (ActiveFaceList[FindIndex] == BestFace)
+					{
+						break;
+					}
+				}
+				check(FindIndex != VertexData.ActiveFaceListStart + VertexData.ActiveFaceListSize);
+
+				if (FindIndex != VertexData.ActiveFaceListStart + VertexData.ActiveFaceListSize - 1)
+				{
+					uint32 SwapTemp = ActiveFaceList[FindIndex];
+					ActiveFaceList[FindIndex] = ActiveFaceList[VertexData.ActiveFaceListStart + VertexData.ActiveFaceListSize - 1];
+					ActiveFaceList[VertexData.ActiveFaceListStart + VertexData.ActiveFaceListSize - 1] = SwapTemp;
+				}
+
+				--VertexData.ActiveFaceListSize;
+				VertexData.Score = FindVertexScore(VertexData.ActiveFaceListSize, VertexData.CachePos1, LRUCacheSize);
+
+			}
+
+			// move the rest of the old verts in the cache down and compute their new scores
+			for (uint32 C0 = 0; C0 < EntriesInCache0; ++C0)
+			{
+				uint32 Index = Cache0[C0];
+				OptimizeVertexData& VertexData = VertexDataList[Index];
+
+				if (VertexData.CachePos1 >= EntriesInCache1)
+				{
+					VertexData.CachePos1 = EntriesInCache1;
+					Cache1[EntriesInCache1++] = Index;
+					VertexData.Score = FindVertexScore(VertexData.ActiveFaceListSize, VertexData.CachePos1, LRUCacheSize);
+				}
+			}
+
+			// find the best scoring triangle in the current cache (including up to 3 that were just evicted)
+			BestScore = -1.f;
+			for (uint32 C1 = 0; C1 < EntriesInCache1; ++C1)
+			{
+				uint32 Index = Cache1[C1];
+				OptimizeVertexData& VertexData = VertexDataList[Index];
+				VertexData.CachePos0 = VertexData.CachePos1;
+				VertexData.CachePos1 = EvictedCacheIndex;
+				for (uint32 j = 0; j < VertexData.ActiveFaceListSize; ++j)
+				{
+					uint32 Face = ActiveFaceList[VertexData.ActiveFaceListStart + j];
+					float FaceScore = 0.f;
+					for (uint32 V=0; V < 3; V++)
+					{
+						uint32 FaceIndex = InIndexList[Face+V];
+						OptimizeVertexData& FaceVertexData = VertexDataList[FaceIndex];
+						FaceScore += FaceVertexData.Score;
+					}
+					if (FaceScore > BestScore)
+					{
+						BestScore = FaceScore;
+						BestFace = Face;
+					}
+				}
+			}
+
+			uint32* SwapTemp = Cache0;
+			Cache0 = Cache1;
+			Cache1 = SwapTemp;
+
+			EntriesInCache0 = FMath::Min(EntriesInCache1, (uint32)LRUCacheSize);
+		}
+	}
+
+} // namespace 
+
+struct FLandscapeDebugOptions
+{
+	FLandscapeDebugOptions()
+	:	bShowPatches(false)
+	,	bDisableStatic(false)
+	,	bDisableCombine(false)
+	,	PatchesConsoleCommand(
+		TEXT( "Landscape.Patches" ),
+		TEXT( "Show/hide Landscape patches" ),
+		FConsoleCommandDelegate::CreateRaw( this, &FLandscapeDebugOptions::Patches ) )
+	,	StaticConsoleCommand(
+		TEXT( "Landscape.Static" ),
+		TEXT( "Enable/disable Landscape static drawlists" ),
+		FConsoleCommandDelegate::CreateRaw( this, &FLandscapeDebugOptions::Static ) )
+	,	CombineConsoleCommand(
+		TEXT( "Landscape.Combine" ),
+		TEXT( "Enable/disable Landscape component combining" ),
+		FConsoleCommandDelegate::CreateRaw( this, &FLandscapeDebugOptions::Combine) )
+	{
+	}
+
+	bool bShowPatches;
+	bool bDisableStatic;
+	bool bDisableCombine;
+
+private:
+	FAutoConsoleCommand PatchesConsoleCommand;
+	FAutoConsoleCommand StaticConsoleCommand;
+	FAutoConsoleCommand CombineConsoleCommand;
+
+	void Patches()
+	{
+		bShowPatches = !bShowPatches;
+		UE_LOG(LogLandscape, Display, TEXT("Landscape.Patches: %s"), bShowPatches ? TEXT("Show") : TEXT("Hide"));
+	}
+
+	void Static()
+	{
+		bDisableStatic = !bDisableStatic;
+		UE_LOG(LogLandscape, Display, TEXT("Landscape.Static: %s"), bDisableStatic ? TEXT("Disabled") : TEXT("Enabled"));
+	}
+
+	void Combine()
+	{
+		bDisableCombine = !bDisableCombine;
+		UE_LOG(LogLandscape, Display, TEXT("Landscape.Combine: %s"), bDisableCombine ? TEXT("Disabled") : TEXT("Enabled"));
+	}
+};
+
+FLandscapeDebugOptions GLandscapeDebugOptions;
+
+
+#if WITH_EDITOR
+LANDSCAPE_API bool GLandscapeEditModeActive = false;
+LANDSCAPE_API ELandscapeViewMode::Type GLandscapeViewMode = ELandscapeViewMode::Normal;
+LANDSCAPE_API int32 GLandscapeEditRenderMode = ELandscapeEditRenderMode::None;
+LANDSCAPE_API int32 GLandscapePreviewMeshRenderMode = 0;
+UMaterial* GLayerDebugColorMaterial = NULL;
+UMaterialInstanceConstant* GSelectionColorMaterial = NULL;
+UMaterialInstanceConstant* GSelectionRegionMaterial = NULL;
+UMaterialInstanceConstant* GMaskRegionMaterial = NULL;
+UTexture2D* GLandscapeBlackTexture = NULL;
+
+// Game thread update
+void FLandscapeEditToolRenderData::Update( UMaterialInterface* InNewToolMaterial )
+{
+	ENQUEUE_UNIQUE_RENDER_COMMAND_TWOPARAMETER(
+		UpdateEditToolRenderData,
+		FLandscapeEditToolRenderData*, LandscapeEditToolRenderData, this,
+		UMaterialInterface*, NewToolMaterial, InNewToolMaterial,
+	{
+		LandscapeEditToolRenderData->ToolMaterial = NewToolMaterial;
+	});
+}
+
+void FLandscapeEditToolRenderData::UpdateGizmo( UMaterialInterface* InNewGizmoMaterial )
+{
+	ENQUEUE_UNIQUE_RENDER_COMMAND_TWOPARAMETER(
+		UpdateEditToolRenderData,
+		FLandscapeEditToolRenderData*, LandscapeEditToolRenderData, this,
+		UMaterialInterface*, NewGizmoMaterial, InNewGizmoMaterial,
+	{
+		LandscapeEditToolRenderData->GizmoMaterial = NewGizmoMaterial;
+	});
+}
+
+// Allows game thread to queue the deletion by the render thread
+void FLandscapeEditToolRenderData::Cleanup()
+{
+	ENQUEUE_UNIQUE_RENDER_COMMAND_ONEPARAMETER(
+		CleanupEditToolRenderData,
+		FLandscapeEditToolRenderData*, LandscapeEditToolRenderData, this,
+	{
+		delete LandscapeEditToolRenderData;
+	});
+}
+
+
+void FLandscapeEditToolRenderData::UpdateDebugColorMaterial()
+{
+	if (!LandscapeComponent)
+	{
+		return;
+	}
+
+	// Debug Color Rendering Material....
+	DebugChannelR = INDEX_NONE, DebugChannelG = INDEX_NONE, DebugChannelB = INDEX_NONE;
+	LandscapeComponent->GetLayerDebugColorKey(DebugChannelR, DebugChannelG, DebugChannelB);
+}
+
+void FLandscapeEditToolRenderData::UpdateSelectionMaterial(int32 InSelectedType)
+{
+	if (!LandscapeComponent)
+	{
+		return;
+	}
+
+	// Check selection
+	if (SelectedType != InSelectedType && (SelectedType & ST_REGION) && !(InSelectedType & ST_REGION) )
+	{
+		// Clear Select textures...
+		if (DataTexture)
+		{
+			FLandscapeEditDataInterface LandscapeEdit(LandscapeComponent->GetLandscapeInfo());
+			LandscapeEdit.ZeroTexture(DataTexture);
+		}
+	}
+
+	ENQUEUE_UNIQUE_RENDER_COMMAND_TWOPARAMETER(
+		UpdateSelectionMaterial,
+		FLandscapeEditToolRenderData*, LandscapeEditToolRenderData, this,
+		int32, InSelectedType, InSelectedType,
+		{
+			LandscapeEditToolRenderData->SelectedType = InSelectedType;
+		});
+}
+#endif
+
+
+//
+// FLandscapeComponentSceneProxy
+//
+TMap<uint32, FLandscapeSharedBuffers*>FLandscapeComponentSceneProxy::SharedBuffersMap;
+TMap<uint32, FLandscapeSharedAdjacencyIndexBuffer*>FLandscapeComponentSceneProxy::SharedAdjacencyIndexBufferMap;
+
+FLandscapeComponentSceneProxy::FLandscapeComponentSceneProxy(ULandscapeComponent* InComponent, FLandscapeEditToolRenderData* InEditToolRenderData)
+:	FPrimitiveSceneProxy(InComponent)
+,	MaxLOD(FMath::CeilLogTwo(InComponent->SubsectionSizeQuads+1)-1)
+,	ComponentSizeQuads(InComponent->ComponentSizeQuads)
+,	ComponentSizeVerts(InComponent->ComponentSizeQuads+1)
+,	NumSubsections(InComponent->NumSubsections)
+,	SubsectionSizeQuads(InComponent->SubsectionSizeQuads)
+,	SubsectionSizeVerts(InComponent->SubsectionSizeQuads+1)
+,	SectionBase(InComponent->GetSectionBase())
+,	StaticLightingLOD(InComponent->GetLandscapeProxy()->StaticLightingLOD)
+,	WeightmapScaleBias(InComponent->WeightmapScaleBias)
+,	WeightmapSubsectionOffset(InComponent->WeightmapSubsectionOffset)
+,	WeightmapTextures(InComponent->WeightmapTextures)
+,	NumWeightmapLayerAllocations(InComponent->WeightmapLayerAllocations.Num())
+,	NormalmapTexture(InComponent->HeightmapTexture)
+,	HeightmapTexture(InComponent->HeightmapTexture)
+,	HeightmapScaleBias(InComponent->HeightmapScaleBias)
+,	XYOffsetmapTexture(InComponent->XYOffsetmapTexture)
+,	SharedBuffersKey(0)
+,	SharedBuffers(NULL)
+,	VertexFactory(NULL)
+,	MaterialInterface(InComponent->MaterialInstance)
+,	EditToolRenderData(InEditToolRenderData)
+,	ComponentLightInfo(NULL)
+,	LandscapeComponent(InComponent)
+,	ForcedLOD(InComponent->ForcedLOD)
+,	LODBias(InComponent->LODBias)
+,	LODFalloff(InComponent->GetLandscapeProxy()->LODFalloff)
+{
+	LevelColor = FLinearColor(1.f, 1.f, 1.f);
+
+	const auto FeatureLevel = GetScene().GetFeatureLevel();
+	if (FeatureLevel <= ERHIFeatureLevel::ES3_1)
+	{
+		HeightmapTexture = NULL;
+		HeightmapSubsectionOffsetU = 0;
+		HeightmapSubsectionOffsetV = 0;
+	}
+	else
+	{
+		HeightmapSubsectionOffsetU = ((float)(InComponent->SubsectionSizeQuads+1) / (float)HeightmapTexture->GetSizeX());
+		HeightmapSubsectionOffsetV = ((float)(InComponent->SubsectionSizeQuads+1) / (float)HeightmapTexture->GetSizeY());
+	}
+
+	//    - - 0 - -
+	//    |       |
+	//    1   P   2
+	//    |       |
+	//    - - 3 - -
+
+	NeighborPosition[0].Set(0.5f * (float)SubsectionSizeQuads, -0.5f * (float)SubsectionSizeQuads);
+	NeighborPosition[1].Set(-0.5f * (float)SubsectionSizeQuads, 0.5f * (float)SubsectionSizeQuads);
+	NeighborPosition[2].Set(1.5f * (float)SubsectionSizeQuads, 0.5f * (float)SubsectionSizeQuads);
+	NeighborPosition[3].Set(0.5f * (float)SubsectionSizeQuads, 1.5f * (float)SubsectionSizeQuads);
+
+	ForcedNeighborLOD[0] = InComponent->NeighborLOD[1];
+	ForcedNeighborLOD[1] = InComponent->NeighborLOD[3];
+	ForcedNeighborLOD[2] = InComponent->NeighborLOD[4];
+	ForcedNeighborLOD[3] = InComponent->NeighborLOD[6];
+
+	NeighborLODBias[0] = InComponent->NeighborLODBias[1];
+	NeighborLODBias[1] = InComponent->NeighborLODBias[3];
+	NeighborLODBias[2] = InComponent->NeighborLODBias[4];
+	NeighborLODBias[3] = InComponent->NeighborLODBias[6];
+
+	LODBias = FMath::Clamp(LODBias, -MaxLOD, MaxLOD);
+
+	if( InComponent->GetLandscapeProxy()->MaxLODLevel >= 0 )
+	{
+		MaxLOD = FMath::Min<int32>(MaxLOD, InComponent->GetLandscapeProxy()->MaxLODLevel);
+	}
+
+	float LODDistanceFactor;
+	switch (LODFalloff)
+	{
+		case ELandscapeLODFalloff::SquareRoot:
+			LODDistanceFactor = FMath::Square(FMath::Min(LANDSCAPE_LOD_SQUARE_ROOT_FACTOR * InComponent->GetLandscapeProxy()->LODDistanceFactor, MAX_LANDSCAPE_LOD_DISTANCE_FACTOR));
+			break;
+		case ELandscapeLODFalloff::Linear:
+		default:
+			LODDistanceFactor = InComponent->GetLandscapeProxy()->LODDistanceFactor;
+			break;
+	}
+
+	LODDistance = FMath::Sqrt(2.f * FMath::Square((float)SubsectionSizeQuads)) * LANDSCAPE_LOD_DISTANCE_FACTOR / LODDistanceFactor; // vary in 0...1
+	DistDiff = -FMath::Sqrt(2.f * FMath::Square(0.5f*(float)SubsectionSizeQuads));
+	// This value is used for shader. I made some modification to LOD transition looks better than plain transition
+	// Used for LODBias.xy values in shader
+	PrecomputedLODFactor = FMath::Min(LODDistanceFactor, 3.f) * 0.33f; 
+
+	if (InComponent->StaticLightingResolution > 0.f)
+	{
+		StaticLightingResolution = InComponent->StaticLightingResolution;
+	}
+	else
+	{
+		StaticLightingResolution = InComponent->GetLandscapeProxy()->StaticLightingResolution;
+	}
+
+	ComponentLightInfo = new FLandscapeLCI(InComponent);
+	check(ComponentLightInfo);
+
+	const bool bHasStaticLighting = InComponent->LightMap != NULL || InComponent->ShadowMap != NULL;
+
+	// Check material usage
+	if( MaterialInterface == NULL ||
+		!MaterialInterface->CheckMaterialUsage(MATUSAGE_Landscape) ||
+		(bHasStaticLighting && !MaterialInterface->CheckMaterialUsage(MATUSAGE_StaticLighting)) )
+	{
+		MaterialInterface = UMaterial::GetDefaultMaterial(MD_Surface);
+	}
+
+	MaterialRelevance = MaterialInterface->GetRelevance(FeatureLevel);
+
+#if !(UE_BUILD_SHIPPING || UE_BUILD_TEST) || (UE_BUILD_SHIPPING && WITH_EDITOR)
+	if( GIsEditor )
+	{
+		ALandscapeProxy* Proxy = InComponent->GetLandscapeProxy();
+		// Try to find a color for level coloration.
+		if ( Proxy )
+		{
+			ULevel* Level = Proxy->GetLevel();
+			ULevelStreaming* LevelStreaming = FLevelUtils::FindStreamingLevel( Level );
+			if ( LevelStreaming )
+			{
+				LevelColor = LevelStreaming->LevelColor;
+			}
+		}
+	}
+#endif
+
+	bRequiresAdjacencyInformation = RequiresAdjacencyInformation(MaterialInterface, XYOffsetmapTexture == NULL ? &FLandscapeVertexFactory::StaticType : &FLandscapeXYOffsetVertexFactory::StaticType, InComponent->GetWorld()->FeatureLevel);
+	SharedBuffersKey = SubsectionSizeQuads | (NumSubsections << 16) | ( XYOffsetmapTexture == NULL ? 0 : 0x80000000);
+
+	DynamicMesh.LCI = ComponentLightInfo; 
+	DynamicMesh.CastShadow = true;
+	DynamicMesh.Type = PT_TriangleList;
+	DynamicMeshBatchParamArray.Empty(NumSubsections*NumSubsections);
+
+#if WITH_EDITOR
+	DynamicMeshTools = DynamicMesh;
+	DynamicMeshTools.Type = PT_TriangleList;
+	DynamicMeshTools.CastShadow = false;
+#endif
+
+	for( int32 SubY=0;SubY<NumSubsections;SubY++ )
+	{
+		for( int32 SubX=0;SubX<NumSubsections;SubX++ )
+		{
+			FMeshBatchElement* BatchElement = (SubX==0 && SubY==0) ? DynamicMesh.Elements.GetData() : new(DynamicMesh.Elements) FMeshBatchElement;
+			BatchElement->PrimitiveUniformBufferResource = &GetUniformBuffer();
+			FLandscapeBatchElementParams* BatchElementParams = new(DynamicMeshBatchParamArray) FLandscapeBatchElementParams;
+			BatchElementParams->LocalToWorldNoScalingPtr = &LocalToWorldNoScaling;
+			BatchElement->UserData = BatchElementParams;
+
+			int32 SubSectionIdx = SubX + SubY*NumSubsections;
+
+			BatchElementParams->SceneProxy = this;
+			BatchElementParams->SubX = SubX;
+			BatchElementParams->SubY = SubY;
+
+#if WITH_EDITOR
+			FMeshBatchElement* ToolBatchElement = (SubX==0 && SubY==0) ? DynamicMeshTools.Elements.GetData() : new(DynamicMeshTools.Elements) FMeshBatchElement;
+			*ToolBatchElement = *BatchElement;
+#endif
+		}
+	}
+}
+
+void FLandscapeComponentSceneProxy::CreateRenderThreadResources()
+{
+	check(HeightmapTexture != NULL);
+
+	auto FeatureLevel = GetScene().GetFeatureLevel();
+
+	SharedBuffers = FLandscapeComponentSceneProxy::SharedBuffersMap.FindRef(SharedBuffersKey);
+	if( SharedBuffers == NULL )
+	{
+		SharedBuffers = new FLandscapeSharedBuffers(SharedBuffersKey, SubsectionSizeQuads, NumSubsections, FeatureLevel, bRequiresAdjacencyInformation);
+		FLandscapeComponentSceneProxy::SharedBuffersMap.Add(SharedBuffersKey, SharedBuffers);
+
+		if (!XYOffsetmapTexture)
+		{
+			FLandscapeVertexFactory* LandscapeVertexFactory = new FLandscapeVertexFactory();
+			LandscapeVertexFactory->Data.PositionComponent = FVertexStreamComponent(SharedBuffers->VertexBuffer, 0, sizeof(FLandscapeVertex), VET_Float4);
+			LandscapeVertexFactory->InitResource();
+			SharedBuffers->VertexFactory = LandscapeVertexFactory;
+		}
+		else
+		{
+			FLandscapeXYOffsetVertexFactory* LandscapeXYOffsetVertexFactory = new FLandscapeXYOffsetVertexFactory();
+			LandscapeXYOffsetVertexFactory->Data.PositionComponent = FVertexStreamComponent(SharedBuffers->VertexBuffer, 0, sizeof(FLandscapeVertex), VET_Float4);
+			LandscapeXYOffsetVertexFactory->InitResource();
+			SharedBuffers->VertexFactory = LandscapeXYOffsetVertexFactory;
+		}
+	}
+
+	SharedBuffers->AddRef();
+
+	if (bRequiresAdjacencyInformation)
+	{
+		if (SharedBuffers->AdjacencyIndexBuffers == NULL)
+		{
+			ensure(SharedBuffers->NumIndexBuffers > 0);
+			if (SharedBuffers->IndexBuffers[0])
+			{
+				// Recreate Index Buffers, this case happens only there are Landscape Components using different material (one uses tessellation, other don't use it) 
+				if (SharedBuffers->bUse32BitIndices && !((FRawStaticIndexBuffer16or32<uint32>*)SharedBuffers->IndexBuffers[0])->Num())
+				{
+					SharedBuffers->CreateIndexBuffers<uint32>(FeatureLevel, bRequiresAdjacencyInformation);
+				}
+				else if (!((FRawStaticIndexBuffer16or32<uint16>*)SharedBuffers->IndexBuffers[0])->Num())
+				{
+					SharedBuffers->CreateIndexBuffers<uint16>(FeatureLevel, bRequiresAdjacencyInformation);
+				}
+			}
+
+			SharedBuffers->AdjacencyIndexBuffers = new FLandscapeSharedAdjacencyIndexBuffer(SharedBuffers);
+			FLandscapeComponentSceneProxy::SharedAdjacencyIndexBufferMap.Add(SharedBuffersKey, SharedBuffers->AdjacencyIndexBuffers);
+		}
+		SharedBuffers->AdjacencyIndexBuffers->AddRef();
+
+		// Delayed Initialize for IndexBuffers
+		for (int i = 0; i < SharedBuffers->NumIndexBuffers; i++)
+		{
+			SharedBuffers->IndexBuffers[i]->InitResource();
+		}
+	}
+
+	// Assign vertex factory
+	VertexFactory = SharedBuffers->VertexFactory;
+	DynamicMesh.VertexFactory = VertexFactory;
+#if WITH_EDITOR
+	DynamicMeshTools.VertexFactory = VertexFactory;
+#endif
+	
+	// Assign LandscapeUniformShaderParameters
+	LandscapeUniformShaderParameters.InitResource();
+
+	for( int32 ElementIdx=0;ElementIdx<DynamicMesh.Elements.Num();ElementIdx++ )
+	{
+		FMeshBatchElement& BatchElement = DynamicMesh.Elements[ElementIdx];
+		FLandscapeBatchElementParams* BatchElementParams = (FLandscapeBatchElementParams*)BatchElement.UserData;
+		BatchElementParams->LandscapeUniformShaderParametersResource = &LandscapeUniformShaderParameters;
+	}
+}
+
+FLandscapeComponentSceneProxy::~FLandscapeComponentSceneProxy()
+{
+	// Free the subsection uniform buffer
+	LandscapeUniformShaderParameters.ReleaseResource();
+
+	if( SharedBuffers )
+	{
+		check( SharedBuffers == FLandscapeComponentSceneProxy::SharedBuffersMap.FindRef(SharedBuffersKey) );
+		if( SharedBuffers->Release() == 0 )
+		{
+			FLandscapeComponentSceneProxy::SharedBuffersMap.Remove(SharedBuffersKey);
+		}
+		SharedBuffers = NULL;
+	}
+
+	delete ComponentLightInfo;
+	ComponentLightInfo = NULL;
+}
+
+int32 GAllowLandscapeShadows=1;
+static FAutoConsoleVariableRef CVarAllowLandscapeShadows(
+	TEXT("r.AllowLandscapeShadows"),
+	GAllowLandscapeShadows,
+	TEXT("Allow Landscape Shadows")
+	);
+
+bool FLandscapeComponentSceneProxy::CanBeOccluded() const
+{
+	return !MaterialRelevance.bDisableDepthTest;
+}
+
+FPrimitiveViewRelevance FLandscapeComponentSceneProxy::GetViewRelevance(const FSceneView* View)
+{
+	FPrimitiveViewRelevance Result;
+	Result.bDrawRelevance = IsShown(View) && View->Family->EngineShowFlags.Landscape;
+
+	auto FeatureLevel = View->GetFeatureLevel();
+
+#if WITH_EDITOR
+	if( !GLandscapeEditModeActive )
+	{
+		// No tools to render, just use the cached material relevance.
+#endif
+		MaterialRelevance.SetPrimitiveViewRelevance(Result);
+#if WITH_EDITOR
+	}
+	else
+	{
+		// Also add the tool material(s)'s relevance to the MaterialRelevance
+		FMaterialRelevance ToolRelevance = MaterialRelevance;
+
+		// Tool brushes and Gizmo
+		if( EditToolRenderData )
+		{
+			if( EditToolRenderData->ToolMaterial )
+			{
+				Result.bDynamicRelevance = true;
+				ToolRelevance |= EditToolRenderData->ToolMaterial->GetRelevance_Concurrent(FeatureLevel);
+			}
+
+			if( EditToolRenderData->GizmoMaterial )
+			{
+				Result.bDynamicRelevance = true;
+				ToolRelevance |= EditToolRenderData->GizmoMaterial->GetRelevance_Concurrent(FeatureLevel);
+			}
+		}
+
+		// Region selection
+		if ( EditToolRenderData && EditToolRenderData->SelectedType )
+		{
+			if ((GLandscapeEditRenderMode & ELandscapeEditRenderMode::SelectRegion) && (EditToolRenderData->SelectedType & FLandscapeEditToolRenderData::ST_REGION)
+				&& !(GLandscapeEditRenderMode & ELandscapeEditRenderMode::Mask) && GSelectionRegionMaterial)
+			{
+				Result.bDynamicRelevance = true;
+				ToolRelevance |= GSelectionRegionMaterial->GetRelevance_Concurrent(FeatureLevel);
+			}
+			if ((GLandscapeEditRenderMode & ELandscapeEditRenderMode::SelectComponent) && (EditToolRenderData->SelectedType & FLandscapeEditToolRenderData::ST_COMPONENT) && GSelectionColorMaterial)
+			{
+				Result.bDynamicRelevance = true;
+				ToolRelevance |= GSelectionColorMaterial->GetRelevance_Concurrent(FeatureLevel);
+			}
+		}
+
+		// Mask
+		if ( (GLandscapeEditRenderMode & ELandscapeEditRenderMode::Mask) && GMaskRegionMaterial != NULL && 
+			((EditToolRenderData && (EditToolRenderData->SelectedType & FLandscapeEditToolRenderData::ST_REGION)) || (!(GLandscapeEditRenderMode & ELandscapeEditRenderMode::InvertedMask))) )
+		{
+			Result.bDynamicRelevance = true;
+			ToolRelevance |= GMaskRegionMaterial->GetRelevance_Concurrent(FeatureLevel);
+		}
+
+		ToolRelevance.SetPrimitiveViewRelevance(Result);
+	}
+
+	// Various visualizations need to render using dynamic relevance
+	if( (View->Family->EngineShowFlags.Bounds && IsSelected()) ||
+		GLandscapeDebugOptions.bShowPatches )
+	{
+		Result.bDynamicRelevance = true;
+	}
+#endif
+
+	// Use the dynamic path for rendering landscape components pass only for Rich Views or if the static path is disabled for debug.
+	if(	IsRichView(*View->Family) || 
+		GLandscapeDebugOptions.bDisableStatic || 
+		View->Family->EngineShowFlags.Wireframe || 
+#if WITH_EDITOR
+		(IsSelected() && !GLandscapeEditModeActive) || 
+		GLandscapeViewMode != ELandscapeViewMode::Normal
+#else
+		IsSelected()
+#endif
+		)
+	{
+		Result.bDynamicRelevance = true;
+	}
+	else
+	{
+		Result.bStaticRelevance = true;
+	}
+
+	Result.bShadowRelevance = (GAllowLandscapeShadows > 0) && IsShadowCast(View);
+	return Result;
+}
+
+/**
+*	Determines the relevance of this primitive's elements to the given light.
+*	@param	LightSceneProxy			The light to determine relevance for
+*	@param	bDynamic (output)		The light is dynamic for this primitive
+*	@param	bRelevant (output)		The light is relevant for this primitive
+*	@param	bLightMapped (output)	The light is light mapped for this primitive
+*/
+void FLandscapeComponentSceneProxy::GetLightRelevance(const FLightSceneProxy* LightSceneProxy, bool& bDynamic, bool& bRelevant, bool& bLightMapped, bool& bShadowMapped) const
+{
+	// Attach the light to the primitive's static meshes.
+	bDynamic = true;
+	bRelevant = false;
+	bLightMapped = true;
+	bShadowMapped = true;
+
+	if (ComponentLightInfo)
+	{
+		ELightInteractionType InteractionType = ComponentLightInfo->GetInteraction(LightSceneProxy).GetType();
+
+		if(InteractionType != LIT_CachedIrrelevant)
+		{
+			bRelevant = true;
+		}
+
+		if(InteractionType != LIT_CachedLightMap && InteractionType != LIT_CachedIrrelevant)
+		{
+			bLightMapped = false;
+		}
+
+		if (InteractionType != LIT_Dynamic && InteractionType != LIT_CachedSignedDistanceFieldShadowMap2D)
+		{
+			bDynamic = false;
+		}
+
+		if (InteractionType != LIT_CachedSignedDistanceFieldShadowMap2D)
+		{
+			bShadowMapped = false;
+		}
+	}
+	else
+	{
+		bRelevant = true;
+		bLightMapped = false;
+	}
+}
+
+FLightInteraction FLandscapeComponentSceneProxy::FLandscapeLCI::GetInteraction(const class FLightSceneProxy* LightSceneProxy) const
+{
+	// Check if the light has static lighting or shadowing.
+	if(LightSceneProxy->HasStaticShadowing())
+	{
+		const FGuid LightGuid = LightSceneProxy->GetLightGuid();
+
+		if(LightMap && LightMap->ContainsLight(LightGuid))
+		{
+			return FLightInteraction::LightMap();
+		}
+
+		if(ShadowMap && ShadowMap->ContainsLight(LightGuid))
+		{
+			return FLightInteraction::ShadowMap2D();
+		}
+
+		if( IrrelevantLights.Contains(LightGuid) )
+		{
+			return FLightInteraction::Irrelevant();
+		}
+	}
+
+	// Use dynamic lighting if the light doesn't have static lighting.
+	return FLightInteraction::Dynamic();
+}
+
+#if WITH_EDITOR
+namespace DebugColorMask
+{
+	const FLinearColor Masks[5] = 
+	{
+		FLinearColor(1.f,0.f,0.f,0.f),
+		FLinearColor(0.f,1.f,0.f,0.f),
+		FLinearColor(0.f,0.f,1.f,0.f),
+		FLinearColor(0.f,0.f,0.f,1.f),
+		FLinearColor(0.f,0.f,0.f,0.f)
+	};
+};
+#endif
+
+void FLandscapeComponentSceneProxy::OnTransformChanged()
+{
+	// Set Lightmap ScaleBias
+	int32 PatchExpandCountX = 1;
+	int32 PatchExpandCountY = 1;
+	int32 DesiredSize = 1;
+	const float LightMapRatio = ::GetTerrainExpandPatchCount(StaticLightingResolution, PatchExpandCountX, PatchExpandCountY, ComponentSizeQuads, (NumSubsections * (SubsectionSizeQuads+1)), DesiredSize, StaticLightingLOD);
+	const float LightmapLODScaleX = LightMapRatio / ((ComponentSizeVerts>>StaticLightingLOD) + 2 * PatchExpandCountX);
+	const float LightmapLODScaleY = LightMapRatio / ((ComponentSizeVerts>>StaticLightingLOD) + 2 * PatchExpandCountY);
+	const float LightmapBiasX = PatchExpandCountX * LightmapLODScaleX;
+	const float LightmapBiasY = PatchExpandCountY * LightmapLODScaleY;
+	const float LightmapScaleX = LightmapLODScaleX * (float)((ComponentSizeVerts>>StaticLightingLOD)-1) / ComponentSizeQuads;
+	const float LightmapScaleY = LightmapLODScaleY * (float)((ComponentSizeVerts>>StaticLightingLOD)-1) / ComponentSizeQuads;
+	const float LightmapExtendFactorX = (float)SubsectionSizeQuads * LightmapScaleX;
+	const float LightmapExtendFactorY = (float)SubsectionSizeQuads * LightmapScaleY;
+	
+	// cache component's WorldToLocal
+	FMatrix LtoW = GetLocalToWorld();
+	WorldToLocal = LtoW.InverseFast();
+
+	// cache component's LocalToWorldNoScaling
+	LocalToWorldNoScaling = LtoW;
+	LocalToWorldNoScaling.RemoveScaling();
+	
+	// Set FLandscapeUniformVSParameters for this subsection
+	FLandscapeUniformShaderParameters LandscapeParams;
+	LandscapeParams.HeightmapUVScaleBias = HeightmapScaleBias;
+	LandscapeParams.WeightmapUVScaleBias = WeightmapScaleBias;
+	LandscapeParams.LocalToWorldNoScaling = LocalToWorldNoScaling;
+
+	LandscapeParams.LandscapeLightmapScaleBias = FVector4(
+		LightmapScaleX,
+		LightmapScaleY,
+		LightmapBiasY,
+		LightmapBiasX);
+	LandscapeParams.SubsectionSizeVertsLayerUVPan = FVector4(
+		SubsectionSizeVerts,
+		1.f / (float)SubsectionSizeQuads,
+		SectionBase.X,
+		SectionBase.Y
+		);
+	LandscapeParams.SubsectionOffsetParams = FVector4(
+		HeightmapSubsectionOffsetU,
+		HeightmapSubsectionOffsetV,
+		WeightmapSubsectionOffset,
+		SubsectionSizeQuads
+		);
+	LandscapeParams.LightmapSubsectionOffsetParams = FVector4(
+		LightmapExtendFactorX,
+		LightmapExtendFactorY,
+		0,
+		0
+		);
+
+	LandscapeUniformShaderParameters.SetContents(LandscapeParams);
+
+	DynamicMesh.ReverseCulling = IsLocalToWorldDeterminantNegative();
+}
+
+namespace
+{
+	inline bool RequiresAdjacencyInformation(FMaterialRenderProxy* MaterialRenderProxy, ERHIFeatureLevel::Type InFeatureLevel) // Assumes VertexFactory supports tessellation, and rendering thread with this function
+	{
+		if (RHISupportsTessellation(GShaderPlatformForFeatureLevel[InFeatureLevel]) && MaterialRenderProxy)
+		{
+			check ( IsInRenderingThread() );
+			const FMaterial* MaterialResource = MaterialRenderProxy->GetMaterial(InFeatureLevel);
+			check( MaterialResource );
+			EMaterialTessellationMode TessellationMode = MaterialResource->GetTessellationMode();
+			bool bEnableCrackFreeDisplacement = MaterialResource->IsCrackFreeDisplacementEnabled();
+
+			return TessellationMode == MTM_PNTriangles || ( TessellationMode == MTM_FlatTessellation && bEnableCrackFreeDisplacement );
+		}
+		else
+		{
+			return false;
+		}
+	}
+};
+
+/** 
+* Draw the scene proxy as a dynamic element
+*
+* @param	PDI - draw interface to render to
+* @param	View - current view
+*/
+
+void FLandscapeComponentSceneProxy::DrawStaticElements(FStaticPrimitiveDrawInterface* PDI)
+{
+	int32 FirstLOD = (ForcedLOD >= 0) ? FMath::Min<int32>(ForcedLOD, MaxLOD) : FMath::Max<int32>(LODBias, 0);
+	int32 LastLOD  = (ForcedLOD >= 0) ? FirstLOD : FMath::Min<int32>(MaxLOD, MaxLOD+LODBias);
+
+	StaticBatchParamArray.Empty((1+LastLOD-FirstLOD) * (FMath::Square(NumSubsections) + 1));
+
+	FMeshBatch MeshBatch;
+	MeshBatch.Elements.Empty((1+LastLOD-FirstLOD) * (FMath::Square(NumSubsections) + 1));
+
+	FMaterialRenderProxy* RenderProxy = MaterialInterface->GetRenderProxy(false);
+
+	// Could be different from bRequiresAdjacencyInformation during shader compilation
+	bool bCurrentRequiresAdjacencyInformation = RequiresAdjacencyInformation( RenderProxy, GetScene().GetFeatureLevel() );
+
+	MeshBatch.VertexFactory = VertexFactory;
+	MeshBatch.MaterialRenderProxy = RenderProxy;
+	MeshBatch.LCI = ComponentLightInfo; 
+	MeshBatch.ReverseCulling = IsLocalToWorldDeterminantNegative();
+	MeshBatch.CastShadow = true;
+	MeshBatch.Type = bCurrentRequiresAdjacencyInformation ? PT_12_ControlPointPatchList : PT_TriangleList;
+	MeshBatch.DepthPriorityGroup = SDPG_World;
+
+	for( int32 LOD = FirstLOD; LOD <= LastLOD; LOD++ )
+	{
+		int32 LodSubsectionSizeVerts = SubsectionSizeVerts >> LOD;
+
+		if( ForcedLOD < 0 && NumSubsections > 1 )
+		{
+			// Per-subsection batch elements
+			for( int32 SubY=0;SubY<NumSubsections;SubY++ )
+			{
+				for( int32 SubX=0;SubX<NumSubsections;SubX++ )
+				{
+					FMeshBatchElement* BatchElement = new(MeshBatch.Elements) FMeshBatchElement;
+					FLandscapeBatchElementParams* BatchElementParams = new(StaticBatchParamArray) FLandscapeBatchElementParams;
+					BatchElement->UserData = BatchElementParams;
+
+					BatchElement->PrimitiveUniformBufferResource = &GetUniformBuffer();
+
+					BatchElementParams->LandscapeUniformShaderParametersResource = &LandscapeUniformShaderParameters;
+					BatchElementParams->LocalToWorldNoScalingPtr = &LocalToWorldNoScaling;
+					BatchElementParams->SceneProxy = this;
+					BatchElementParams->SubX = SubX;
+					BatchElementParams->SubY = SubY;
+					BatchElementParams->CurrentLOD = LOD;
+
+					BatchElement->IndexBuffer = SharedBuffers->IndexBuffers[LOD];
+					BatchElement->NumPrimitives = FMath::Square((LodSubsectionSizeVerts-1)) * 2;
+					BatchElement->FirstIndex = (SubX + SubY * NumSubsections) * BatchElement->NumPrimitives * 3;
+					BatchElement->MinVertexIndex = SharedBuffers->IndexRanges[LOD].MinIndex[SubX][SubY];
+					BatchElement->MaxVertexIndex = SharedBuffers->IndexRanges[LOD].MaxIndex[SubX][SubY];
+				}
+			}
+		}
+
+		// Combined batch element
+		FMeshBatchElement* BatchElement = new(MeshBatch.Elements) FMeshBatchElement;
+		FLandscapeBatchElementParams* BatchElementParams = new(StaticBatchParamArray) FLandscapeBatchElementParams;
+		BatchElementParams->LocalToWorldNoScalingPtr = &LocalToWorldNoScaling;
+		BatchElement->UserData = BatchElementParams;
+		BatchElement->PrimitiveUniformBufferResource = &GetUniformBuffer();
+		BatchElementParams->LandscapeUniformShaderParametersResource = &LandscapeUniformShaderParameters;
+		BatchElementParams->SceneProxy = this;
+		BatchElementParams->SubX = -1;
+		BatchElementParams->SubY = -1;
+		BatchElementParams->CurrentLOD = LOD;
+
+		if (bCurrentRequiresAdjacencyInformation)
+		{
+			check(SharedBuffers->AdjacencyIndexBuffers);
+			BatchElement->IndexBuffer = SharedBuffers->AdjacencyIndexBuffers->IndexBuffers[LOD];
+		}
+		else
+		{
+			BatchElement->IndexBuffer = SharedBuffers->IndexBuffers[LOD];
+		}
+
+		BatchElement->NumPrimitives = FMath::Square((LodSubsectionSizeVerts-1)) * FMath::Square(NumSubsections) * 2;
+		BatchElement->FirstIndex = 0;
+		BatchElement->MinVertexIndex = SharedBuffers->IndexRanges[LOD].MinIndexFull;
+		BatchElement->MaxVertexIndex = SharedBuffers->IndexRanges[LOD].MaxIndexFull;
+	}
+
+	PDI->DrawMesh(MeshBatch, FLT_MAX);
+}
+
+uint64 FLandscapeVertexFactory::GetStaticBatchElementVisibility( const class FSceneView& View, const struct FMeshBatch* Batch ) const
+{
+	const FLandscapeComponentSceneProxy* SceneProxy = ((FLandscapeBatchElementParams*)Batch->Elements[0].UserData)->SceneProxy;
+	return SceneProxy->GetStaticBatchElementVisibility( View, Batch );
+}
+
+uint64 FLandscapeComponentSceneProxy::GetStaticBatchElementVisibility( const class FSceneView& View, const struct FMeshBatch* Batch ) const
+{
+	uint64 BatchesToRenderMask = 0;
+
+	SCOPE_CYCLE_COUNTER(STAT_LandscapeStaticDrawLODTime);
+	if (ForcedLOD >= 0)
+	{
+		for( int32 BatchElementIndex=0;BatchElementIndex < Batch->Elements.Num(); BatchElementIndex++ )
+		{
+			BatchesToRenderMask |= (((uint64)1)<<BatchElementIndex);
+			INC_DWORD_STAT(STAT_LandscapeDrawCalls);
+			INC_DWORD_STAT_BY(STAT_LandscapeTriangles,Batch->Elements[BatchElementIndex].NumPrimitives);
+		}
+	}
+	else
+	{
+		// camera position in local heightmap space
+		FVector CameraLocalPos3D = WorldToLocal.TransformPosition(View.ViewMatrices.ViewOrigin); 
+		FVector2D CameraLocalPos(CameraLocalPos3D.X, CameraLocalPos3D.Y);
+
+		int32 BatchesPerLOD = NumSubsections > 1 ? FMath::Square(NumSubsections) + 1 : 1;
+		int32 CalculatedLods[LANDSCAPE_MAX_SUBSECTION_NUM][LANDSCAPE_MAX_SUBSECTION_NUM];
+		int32 CombinedLOD = -1;
+		int32 bAllSameLOD = true;
+
+		for( int32 SubY=0;SubY<NumSubsections;SubY++ )
+		{
+			for( int32 SubX=0;SubX<NumSubsections;SubX++ )
+			{
+				int32 TempLOD = CalcLODForSubsectionNoForced(View, SubX, SubY, CameraLocalPos);
+
+				if(LODBias > 0) 
+				{ 
+					TempLOD = FMath::Max<int32>(TempLOD - LODBias, 0); 
+				} 
+
+				// check if all LODs are the same.
+				if( TempLOD != CombinedLOD && CombinedLOD != -1 )
+				{
+					bAllSameLOD = false;
+				}
+				CombinedLOD = TempLOD;
+				CalculatedLods[SubX][SubY] = TempLOD;
+			}
+		}
+
+		if( bAllSameLOD && NumSubsections > 1 && !GLandscapeDebugOptions.bDisableCombine )
+		{
+			// choose the combined batch element
+			int32 BatchElementIndex = (CombinedLOD+1)*BatchesPerLOD - 1;
+			BatchesToRenderMask |= (((uint64)1)<<BatchElementIndex);
+			INC_DWORD_STAT(STAT_LandscapeDrawCalls);
+			INC_DWORD_STAT_BY(STAT_LandscapeTriangles,Batch->Elements[BatchElementIndex].NumPrimitives);
+		}
+		else
+		{
+			for( int32 SubY=0;SubY<NumSubsections;SubY++ )
+			{;
+				for( int32 SubX=0;SubX<NumSubsections;SubX++ )
+				{
+					int32 BatchElementIndex = CalculatedLods[SubX][SubY]*BatchesPerLOD + SubY*NumSubsections + SubX;
+					BatchesToRenderMask |= (((uint64)1)<<BatchElementIndex);
+					INC_DWORD_STAT(STAT_LandscapeDrawCalls);
+					INC_DWORD_STAT_BY(STAT_LandscapeTriangles,Batch->Elements[BatchElementIndex].NumPrimitives);
+				}
+			}
+		}
+
+	}
+
+	INC_DWORD_STAT(STAT_LandscapeComponents);
+
+	return BatchesToRenderMask;
+}
+
+int32 FLandscapeComponentSceneProxy::CalcLODForSubsectionNoForced(const class FSceneView& View, int32 SubX, int32 SubY, const FVector2D& CameraLocalPos) const
+{
+	const int32 MinLOD = HeightmapTexture ? FMath::Min<int32>(HeightmapTexture->GetNumMips() - HeightmapTexture->ResidentMips, MaxLOD) : 0;
+#if WITH_EDITOR
+	if (View.Family->LandscapeLODOverride >= 0)
+	{
+		float fLOD = FMath::Clamp<float>(View.Family->LandscapeLODOverride, FMath::Max<int32>(LODBias, MinLOD), FMath::Min<int32>(MaxLOD, MaxLOD + LODBias));
+		return FMath::FloorToInt(fLOD);
+	}
+#endif
+
+	if (View.IsPerspectiveProjection())
+	{
+		FVector2D ComponentPosition(0.5f * (float)SubsectionSizeQuads, 0.5f * (float)SubsectionSizeQuads);
+		FVector2D CurrentCameraLocalPos = CameraLocalPos - FVector2D(SubX * SubsectionSizeQuads,SubY * SubsectionSizeQuads);
+		float ComponentDistance = FVector2D(CurrentCameraLocalPos-ComponentPosition).Size() + DistDiff;
+		// Clamp calculated distance based LOD with LODBiased values
+		float fLOD;
+		switch (LODFalloff)
+		{
+			case ELandscapeLODFalloff::SquareRoot:
+				fLOD = FMath::Sqrt(ComponentDistance / LODDistance);
+				break;
+			default:
+			case ELandscapeLODFalloff::Linear:
+				fLOD = ComponentDistance / LODDistance;
+				break;
+		}
+		// Clamping based on LODBias
+		fLOD = FMath::Clamp<float>(fLOD, FMath::Max<int32>(LODBias, MinLOD), FMath::Min<int32>(MaxLOD, MaxLOD + LODBias));
+
+		return FMath::FloorToInt( fLOD );
+	}
+	else
+	{
+		float Scale = 1.0f / (View.ViewRect.Width() * View.ViewMatrices.ProjMatrix.M[0][0]);
+
+		// The "/ 5.0f" is totally arbitrary
+		float fLOD;
+		switch (LODFalloff)
+		{
+			case ELandscapeLODFalloff::SquareRoot:
+				fLOD = FMath::Sqrt(FMath::Sqrt(Scale / 5.0f));
+				break;
+			default:
+			case ELandscapeLODFalloff::Linear:
+				fLOD = Scale / 5.0f;
+				break;
+		}
+
+		// Clamping based on LODBias
+		fLOD = FMath::Clamp<float>(fLOD, FMath::Max<int32>(LODBias, MinLOD), FMath::Min<int32>(MaxLOD, MaxLOD + LODBias));
+
+		return FMath::FloorToInt(fLOD);
+	}
+}
+
+int32 FLandscapeComponentSceneProxy::CalcLODForSubsection(const class FSceneView& View, int32 SubX, int32 SubY, const FVector2D& CameraLocalPos) const
+{
+	if( ForcedLOD >= 0 )
+	{
+		return ForcedLOD;
+	}
+	else
+	{
+		return CalcLODForSubsectionNoForced(View, SubX, SubY, CameraLocalPos);
+	}
+}
+
+void FLandscapeComponentSceneProxy::CalcLODParamsForSubsection(const class FSceneView& View, const FVector2D& CameraLocalPos, int32 SubX, int32 SubY, float& OutfLOD, FVector4& OutNeighborLODs, float& OutDistLOD) const
+{
+	FVector2D ComponentPosition(0.5f * (float)SubsectionSizeQuads, 0.5f * (float)SubsectionSizeQuads);
+	FVector2D CurrentCameraLocalPos = CameraLocalPos - FVector2D(SubX * SubsectionSizeQuads,SubY * SubsectionSizeQuads);
+
+	float ComponentDistance = FVector2D(CurrentCameraLocalPos - ComponentPosition).Size() + DistDiff;
+
+	OutDistLOD = ComponentDistance / LODDistance;
+
+	int32 FirstLOD = FMath::Max<int32>(LODBias, 0);
+	int32 LastLOD = FMath::Min<int32>(MaxLOD, MaxLOD+LODBias);
+	if (ForcedLOD >= 0)
+	{
+		OutfLOD = ForcedLOD;
+	}
+	else
+	{
+		switch (LODFalloff)
+		{
+		case ELandscapeLODFalloff::SquareRoot:
+			OutfLOD = FMath::Sqrt(OutDistLOD);
+			break;
+		default:
+		case ELandscapeLODFalloff::Linear:
+			OutfLOD = OutDistLOD;
+			break;
+		}
+
+		OutfLOD = FMath::Clamp<float>(OutfLOD, FirstLOD, LastLOD);
+	}
+
+	for (int32 Idx = 0; Idx < LANDSCAPE_NEIGHBOR_NUM; ++Idx)
+	{
+		float ComponentDistance = FVector2D(CurrentCameraLocalPos-NeighborPosition[Idx]).Size() + DistDiff;
+		float fLOD;
+		switch (LODFalloff)
+		{
+			case ELandscapeLODFalloff::SquareRoot:
+				fLOD = FMath::Sqrt(ComponentDistance / LODDistance);
+				break;
+			default:
+			case ELandscapeLODFalloff::Linear:
+				fLOD = ComponentDistance / LODDistance;
+				break;
+		}
+
+		if (NumSubsections > 1 
+			&& ((SubX == 0 && Idx == 2) 
+			|| (SubX == NumSubsections-1 && Idx == 1) 
+			|| (SubY == 0 && Idx == 3) 
+			|| (SubY == NumSubsections-1 && Idx == 0)) )
+		{
+			// Clamp calculated distance based LOD with LODBiased values
+			OutNeighborLODs[Idx] = ForcedLOD >= 0 ? ForcedLOD : FMath::Clamp<float>( fLOD, FirstLOD, LastLOD);
+		}
+		else
+		{
+			// Neighbor LODBias are saved in BYTE, so need to convert to range [-128:127]
+			OutNeighborLODs[Idx] = ForcedNeighborLOD[Idx] != 255 ? ForcedNeighborLOD[Idx] : FMath::Clamp<float>(fLOD, FMath::Max<float>(NeighborLODBias[Idx] - 128, 0.f), FMath::Min<float>(MaxLOD, MaxLOD + NeighborLODBias[Idx] - 128));
+		}
+
+		OutNeighborLODs[Idx] = FMath::Max<float>(OutfLOD, OutNeighborLODs[Idx]);
+	}
+}
+
+void FLandscapeComponentSceneProxy::GetDynamicMeshElements(const TArray<const FSceneView*>& Views, const FSceneViewFamily& ViewFamily, uint32 VisibilityMap, FMeshElementCollector& Collector) const
+{
+	QUICK_SCOPE_CYCLE_COUNTER(STAT_FLandscapeComponentSceneProxy_GetMeshElements);
+
+	int32 NumPasses=0;
+	int32 NumTriangles=0;
+	int32 NumDrawCalls=0;
+	const bool bIsWireframe = ViewFamily.EngineShowFlags.Wireframe;
+
+	for (int32 ViewIndex = 0; ViewIndex < Views.Num(); ViewIndex++)
+	{
+		if (VisibilityMap & (1 << ViewIndex))
+		{
+			const FSceneView* View = Views[ViewIndex];
+			FVector CameraLocalPos3D = WorldToLocal.TransformPosition(View->ViewMatrices.ViewOrigin); 
+			FVector2D CameraLocalPos(CameraLocalPos3D.X, CameraLocalPos3D.Y);
+
+			FLandscapeElementParamArray& ParameterArray = Collector.AllocateOneFrameResource<FLandscapeElementParamArray>();
+			ParameterArray.ElementParams.Empty(NumSubsections * NumSubsections);
+			ParameterArray.ElementParams.AddUninitialized(NumSubsections * NumSubsections);
+
+			FMeshBatch& Mesh = Collector.AllocateMesh();
+
+			// Could be different from bRequiresAdjacencyInformation during shader compilation
+			FMaterialRenderProxy* RenderProxy = MaterialInterface->GetRenderProxy(false);
+			bool bCurrentRequiresAdjacencyInformation = RequiresAdjacencyInformation( RenderProxy, View->GetFeatureLevel() );
+			Mesh.Type = bCurrentRequiresAdjacencyInformation ? PT_12_ControlPointPatchList : PT_TriangleList;
+			Mesh.LCI = ComponentLightInfo; 
+			Mesh.CastShadow = true;
+			Mesh.VertexFactory = VertexFactory;
+			Mesh.ReverseCulling = IsLocalToWorldDeterminantNegative();
+
+#if WITH_EDITOR
+			FMeshBatch& MeshTools = Collector.AllocateMesh();
+			MeshTools.LCI = ComponentLightInfo; 
+			MeshTools.Type = PT_TriangleList;
+			MeshTools.CastShadow = false;
+			MeshTools.VertexFactory = VertexFactory;
+			MeshTools.ReverseCulling = IsLocalToWorldDeterminantNegative();
+#endif
+
+			// Setup the LOD parameters
+			for( int32 SubY=0;SubY<NumSubsections;SubY++ )
+			{
+				for( int32 SubX=0;SubX<NumSubsections;SubX++ )
+				{
+					int32 SubSectionIdx = SubX + SubY*NumSubsections;
+					int32 CurrentLOD = CalcLODForSubsection(*View, SubX, SubY, CameraLocalPos);
+
+					FMeshBatchElement& BatchElement = (SubX==0 && SubY==0) ? *Mesh.Elements.GetData() : *(new(Mesh.Elements) FMeshBatchElement);
+					BatchElement.PrimitiveUniformBufferResource = &GetUniformBuffer();
+					FLandscapeBatchElementParams& BatchElementParams = ParameterArray.ElementParams[SubSectionIdx];
+					BatchElementParams.LocalToWorldNoScalingPtr = &LocalToWorldNoScaling;
+					BatchElement.UserData = &BatchElementParams;
+
+					BatchElementParams.LandscapeUniformShaderParametersResource = &LandscapeUniformShaderParameters;
+					BatchElementParams.SceneProxy = this;
+					BatchElementParams.SubX = SubX;
+					BatchElementParams.SubY = SubY;
+					BatchElementParams.CurrentLOD = CurrentLOD;
+
+					if (bCurrentRequiresAdjacencyInformation)
+					{
+						check(SharedBuffers->AdjacencyIndexBuffers);
+						BatchElement.IndexBuffer = SharedBuffers->AdjacencyIndexBuffers->IndexBuffers[CurrentLOD];
+					}
+					else
+					{
+						BatchElement.IndexBuffer = SharedBuffers->IndexBuffers[CurrentLOD];
+					}
+
+					int32 LodSubsectionSizeVerts = (SubsectionSizeVerts >> CurrentLOD);
+
+					int32 NumPrimitives = FMath::Square((LodSubsectionSizeVerts - 1)) * 2;
+					BatchElement.NumPrimitives = NumPrimitives;
+					BatchElement.FirstIndex = (SubX + SubY * NumSubsections) * NumPrimitives * 3;
+					BatchElement.MinVertexIndex = SharedBuffers->IndexRanges[CurrentLOD].MinIndex[SubX][SubY];
+					BatchElement.MaxVertexIndex = SharedBuffers->IndexRanges[CurrentLOD].MaxIndex[SubX][SubY];
+					
+#if WITH_EDITOR
+					FMeshBatchElement& BatchElementTools = (SubX==0 && SubY==0) ? *MeshTools.Elements.GetData() : *(new(MeshTools.Elements) FMeshBatchElement);
+					BatchElementTools.PrimitiveUniformBufferResource = &GetUniformBuffer();
+					BatchElementTools.UserData = &BatchElementParams;
+
+					// Tools never use tessellation
+					BatchElementTools.IndexBuffer = SharedBuffers->IndexBuffers[CurrentLOD];
+					BatchElementTools.NumPrimitives = NumPrimitives;
+					BatchElementTools.FirstIndex = BatchElement.FirstIndex;
+					BatchElementTools.MinVertexIndex = SharedBuffers->IndexRanges[CurrentLOD].MinIndex[SubX][SubY];
+					BatchElementTools.MaxVertexIndex = SharedBuffers->IndexRanges[CurrentLOD].MaxIndex[SubX][SubY];
+#endif
+				}
+			}
+			// Render the landscape component
+#if WITH_EDITOR
+			
+			switch(GLandscapeViewMode)
+			{
+			case ELandscapeViewMode::DebugLayer:
+				if( GLayerDebugColorMaterial && EditToolRenderData )
+				{
+					auto DebugColorMaterialInstance = new FLandscapeDebugMaterialRenderProxy(GLayerDebugColorMaterial->GetRenderProxy(false), 
+						(EditToolRenderData->DebugChannelR >= 0 ? WeightmapTextures[EditToolRenderData->DebugChannelR/4] : NULL),
+						(EditToolRenderData->DebugChannelG >= 0 ? WeightmapTextures[EditToolRenderData->DebugChannelG/4] : NULL),
+						(EditToolRenderData->DebugChannelB >= 0 ? WeightmapTextures[EditToolRenderData->DebugChannelB/4] : NULL),	
+						(EditToolRenderData->DebugChannelR >= 0 ? DebugColorMask::Masks[EditToolRenderData->DebugChannelR%4] : DebugColorMask::Masks[4]),
+						(EditToolRenderData->DebugChannelG >= 0 ? DebugColorMask::Masks[EditToolRenderData->DebugChannelG%4] : DebugColorMask::Masks[4]),
+						(EditToolRenderData->DebugChannelB >= 0 ? DebugColorMask::Masks[EditToolRenderData->DebugChannelB%4] : DebugColorMask::Masks[4])
+						);
+
+					MeshTools.MaterialRenderProxy = DebugColorMaterialInstance;
+					Collector.RegisterOneFrameMaterialProxy(DebugColorMaterialInstance);
+
+					MeshTools.bCanApplyViewModeOverrides = true;
+					MeshTools.bUseWireframeSelectionColoring = IsSelected();
+
+					Collector.AddMesh(ViewIndex, MeshTools);
+
+					NumPasses++;
+					NumTriangles += MeshTools.GetNumPrimitives();
+					NumDrawCalls += MeshTools.Elements.Num();
+				}
+				break;
+
+			case ELandscapeViewMode::LayerDensity:
+				if( EditToolRenderData )
+				{
+					int32 ColorIndex = FMath::Min<int32>(NumWeightmapLayerAllocations, GEngine->ShaderComplexityColors.Num());
+					auto LayerDensityMaterialInstance = new FColoredMaterialRenderProxy(GEngine->LevelColorationUnlitMaterial->GetRenderProxy(false), ColorIndex ? GEngine->ShaderComplexityColors[ColorIndex-1] : FLinearColor::Black);
+
+					MeshTools.MaterialRenderProxy = LayerDensityMaterialInstance;
+					Collector.RegisterOneFrameMaterialProxy(LayerDensityMaterialInstance);
+
+					MeshTools.bCanApplyViewModeOverrides = true;
+					MeshTools.bUseWireframeSelectionColoring = IsSelected();
+
+					Collector.AddMesh(ViewIndex, MeshTools);
+
+					NumPasses++;
+					NumTriangles += MeshTools.GetNumPrimitives();
+					NumDrawCalls += MeshTools.Elements.Num();
+				}
+				break;
+
+			case ELandscapeViewMode::LOD:
+				{
+					FLinearColor WireColors[7];
+					WireColors[0] = FLinearColor(1,1,1,1);
+					WireColors[1] = FLinearColor(1,0,0,1);
+					WireColors[2] = FLinearColor(0,1,0,1);
+					WireColors[3] = FLinearColor(0,0,1,1);
+					WireColors[4] = FLinearColor(1,1,0,1);
+					WireColors[5] = FLinearColor(1,0,1,1);
+					WireColors[6] = FLinearColor(0,1,1,1);
+
+					for( int32 i=0;i<MeshTools.Elements.Num();i++ )
+					{
+						FMeshBatch& LODMesh = Collector.AllocateMesh();
+						LODMesh = MeshTools;
+						LODMesh.Elements.Empty(1);
+						LODMesh.Elements.Add(MeshTools.Elements[i]);
+						int32 ColorIndex = ((FLandscapeBatchElementParams*)MeshTools.Elements[i].UserData)->CurrentLOD;
+						FLinearColor Color = ForcedLOD >= 0 ? WireColors[ColorIndex] : WireColors[ColorIndex]*0.2f;
+						auto LODMaterialInstance = new FColoredMaterialRenderProxy(GEngine->LevelColorationUnlitMaterial->GetRenderProxy(false), Color);
+						LODMesh.MaterialRenderProxy = LODMaterialInstance;
+						Collector.RegisterOneFrameMaterialProxy(LODMaterialInstance);
+
+						LODMesh.bCanApplyViewModeOverrides = true;
+						LODMesh.bUseWireframeSelectionColoring = IsSelected();
+
+						Collector.AddMesh(ViewIndex, LODMesh);
+
+						NumPasses++;
+						NumTriangles += MeshTools.Elements[i].NumPrimitives;
+						NumDrawCalls++;
+					}
+				}
+				break;
+
+			default:
+			
+#else
+			{
+#endif // WITH_EDITOR
+				// Regular Landscape rendering. Only use the dynamic path if we're rendering a rich view or we've disabled the static path for debugging.
+				if( IsRichView(ViewFamily) || 
+					GLandscapeDebugOptions.bDisableStatic || 
+					bIsWireframe ||
+#if WITH_EDITOR
+					(IsSelected() && !GLandscapeEditModeActive)
+#else
+					IsSelected()
+#endif
+					)
+				{
+					Mesh.MaterialRenderProxy = MaterialInterface->GetRenderProxy(false);
+
+					Mesh.bCanApplyViewModeOverrides = true;
+					Mesh.bUseWireframeSelectionColoring = IsSelected();
+
+					Collector.AddMesh(ViewIndex, Mesh);
+
+					NumPasses++;
+					NumTriangles += Mesh.GetNumPrimitives();
+					NumDrawCalls += Mesh.Elements.Num();
+				}
+			}
+#if WITH_EDITOR
+			
+			// Extra render passes for landscape tools
+			if( GLandscapeEditModeActive )
+			{
+				// Region selection
+				if ( EditToolRenderData && EditToolRenderData->SelectedType )
+				{
+					if ((GLandscapeEditRenderMode & ELandscapeEditRenderMode::SelectRegion) && (EditToolRenderData->SelectedType & FLandscapeEditToolRenderData::ST_REGION)
+						&& !(GLandscapeEditRenderMode & ELandscapeEditRenderMode::Mask))
+					{
+						FMeshBatch& SelectMesh = Collector.AllocateMesh();
+						SelectMesh = MeshTools;
+						auto SelectMaterialInstance = new FLandscapeSelectMaterialRenderProxy(GSelectionRegionMaterial->GetRenderProxy(false), EditToolRenderData->DataTexture ? EditToolRenderData->DataTexture : GLandscapeBlackTexture);
+						SelectMesh.MaterialRenderProxy = SelectMaterialInstance;
+						Collector.RegisterOneFrameMaterialProxy(SelectMaterialInstance);
+						Collector.AddMesh(ViewIndex, SelectMesh);
+						NumPasses++;
+						NumTriangles += SelectMesh.GetNumPrimitives();
+						NumDrawCalls += SelectMesh.Elements.Num();
+					}
+
+					if ((GLandscapeEditRenderMode & ELandscapeEditRenderMode::SelectComponent) && (EditToolRenderData->SelectedType & FLandscapeEditToolRenderData::ST_COMPONENT))
+					{
+						FMeshBatch& SelectMesh = Collector.AllocateMesh();
+						SelectMesh = MeshTools;
+						SelectMesh.MaterialRenderProxy = GSelectionColorMaterial->GetRenderProxy(0);
+						Collector.AddMesh(ViewIndex, SelectMesh);
+						NumPasses++;
+						NumTriangles += SelectMesh.GetNumPrimitives();
+						NumDrawCalls += SelectMesh.Elements.Num();
+					}
+				}
+
+				// Mask
+				if ((GLandscapeEditRenderMode & ELandscapeEditRenderMode::SelectRegion) && (GLandscapeEditRenderMode & ELandscapeEditRenderMode::Mask))
+				{
+					if (EditToolRenderData && (EditToolRenderData->SelectedType & FLandscapeEditToolRenderData::ST_REGION) )
+					{
+						FMeshBatch& MaskMesh = Collector.AllocateMesh();
+						MaskMesh = MeshTools;
+						auto MaskMaterialInstance = new FLandscapeMaskMaterialRenderProxy(GMaskRegionMaterial->GetRenderProxy(false), EditToolRenderData->DataTexture ? EditToolRenderData->DataTexture : GLandscapeBlackTexture, !!(GLandscapeEditRenderMode & ELandscapeEditRenderMode::InvertedMask) );
+						MaskMesh.MaterialRenderProxy = MaskMaterialInstance;
+						Collector.RegisterOneFrameMaterialProxy(MaskMaterialInstance);
+						Collector.AddMesh(ViewIndex, MaskMesh);
+						NumPasses++;
+						NumTriangles += MaskMesh.GetNumPrimitives();
+						NumDrawCalls += MaskMesh.Elements.Num();
+					}
+					else if (!(GLandscapeEditRenderMode & ELandscapeEditRenderMode::InvertedMask))
+					{
+						FMeshBatch& MaskMesh = Collector.AllocateMesh();
+						MaskMesh = MeshTools;
+						auto MaskMaterialInstance = new FLandscapeMaskMaterialRenderProxy(GMaskRegionMaterial->GetRenderProxy(false), GLandscapeBlackTexture, false );
+						MaskMesh.MaterialRenderProxy = MaskMaterialInstance;
+						Collector.RegisterOneFrameMaterialProxy(MaskMaterialInstance);
+						Collector.AddMesh(ViewIndex, MaskMesh);
+						NumPasses++;
+						NumTriangles += MaskMesh.GetNumPrimitives();
+						NumDrawCalls += MaskMesh.Elements.Num();
+					}
+				}
+
+				// Edit mode tools
+				if( EditToolRenderData )
+				{
+					if (EditToolRenderData->ToolMaterial)
+					{
+						FMeshBatch& EditMesh = Collector.AllocateMesh();
+						EditMesh = MeshTools;
+						EditMesh.MaterialRenderProxy = EditToolRenderData->ToolMaterial->GetRenderProxy(0);		
+						Collector.AddMesh(ViewIndex, EditMesh);
+						NumPasses++;
+						NumTriangles += EditMesh.GetNumPrimitives();
+						NumDrawCalls += EditMesh.Elements.Num();
+					}
+
+					if (EditToolRenderData->GizmoMaterial && GLandscapeEditRenderMode & ELandscapeEditRenderMode::Gizmo)
+					{
+						FMeshBatch& EditMesh = Collector.AllocateMesh();
+						EditMesh = MeshTools;
+						EditMesh.MaterialRenderProxy = EditToolRenderData->GizmoMaterial->GetRenderProxy(0);
+						Collector.AddMesh(ViewIndex, EditMesh);
+						NumPasses++;
+						NumTriangles += EditMesh.GetNumPrimitives();
+						NumDrawCalls += EditMesh.Elements.Num();
+					}
+				}
+			}
+#endif // WITH_EDITOR
+
+			if ( GLandscapeDebugOptions.bShowPatches )
+			{
+				DrawWireBox(Collector.GetPDI(ViewIndex), GetBounds().GetBox(), FColor(255, 255, 0), SDPG_World);
+			}
+
+			RenderBounds(Collector.GetPDI(ViewIndex), ViewFamily.EngineShowFlags, GetBounds(), IsSelected());
+		}
+	}
+
+	INC_DWORD_STAT_BY(STAT_LandscapeComponents, NumPasses);
+	INC_DWORD_STAT_BY(STAT_LandscapeDrawCalls, NumDrawCalls);
+	INC_DWORD_STAT_BY(STAT_LandscapeTriangles, NumTriangles * NumPasses);
+}
+
+void FLandscapeComponentSceneProxy::DrawDynamicElements(FPrimitiveDrawInterface* PDI,const FSceneView* View)
+{
+	SCOPE_CYCLE_COUNTER(STAT_LandscapeDynamicDrawTime);
+
+	int32 NumPasses=0;
+	int32 NumTriangles=0;
+	int32 NumDrawCalls=0;
+	const bool bIsWireframe = View->Family->EngineShowFlags.Wireframe;
+
+	FVector CameraLocalPos3D = WorldToLocal.TransformPosition(View->ViewMatrices.ViewOrigin); 
+	FVector2D CameraLocalPos(CameraLocalPos3D.X, CameraLocalPos3D.Y);
+
+	// Could be different from bRequiresAdjacencyInformation during shader compilation
+	FMaterialRenderProxy* RenderProxy = MaterialInterface->GetRenderProxy(false);
+	bool bCurrentRequiresAdjacencyInformation = RequiresAdjacencyInformation( RenderProxy, View->GetFeatureLevel() );
+	DynamicMesh.Type = bCurrentRequiresAdjacencyInformation ? PT_12_ControlPointPatchList : PT_TriangleList;
+
+	// Setup the LOD parameters in DynamicMesh.
+	for( int32 SubY=0;SubY<NumSubsections;SubY++ )
+	{
+		for( int32 SubX=0;SubX<NumSubsections;SubX++ )
+		{
+			int32 SubSectionIdx = SubX + SubY*NumSubsections;
+			int32 CurrentLOD = CalcLODForSubsection(*View, SubX, SubY, CameraLocalPos);
+
+			FMeshBatchElement& BatchElement = DynamicMesh.Elements[SubSectionIdx];
+			((FLandscapeBatchElementParams*)BatchElement.UserData)->CurrentLOD = CurrentLOD;
+
+			if (bCurrentRequiresAdjacencyInformation)
+			{
+				check(SharedBuffers->AdjacencyIndexBuffers);
+				BatchElement.IndexBuffer = SharedBuffers->AdjacencyIndexBuffers->IndexBuffers[CurrentLOD];
+			}
+			else
+			{
+				BatchElement.IndexBuffer = SharedBuffers->IndexBuffers[CurrentLOD];
+			}
+
+			int32 LodSubsectionSizeVerts = (SubsectionSizeVerts >> CurrentLOD);
+
+			int32 NumPrimitives = FMath::Square((LodSubsectionSizeVerts - 1)) * 2;
+			BatchElement.NumPrimitives = NumPrimitives;
+			BatchElement.FirstIndex = (SubX + SubY * NumSubsections) * NumPrimitives * 3;
+			BatchElement.MinVertexIndex = SharedBuffers->IndexRanges[CurrentLOD].MinIndex[SubX][SubY];
+			BatchElement.MaxVertexIndex = SharedBuffers->IndexRanges[CurrentLOD].MaxIndex[SubX][SubY];
+
+#if WITH_EDITOR
+			// Tools never use tessellation
+			FMeshBatchElement& BatchElementTools = DynamicMeshTools.Elements[SubSectionIdx];
+			((FLandscapeBatchElementParams*)BatchElementTools.UserData)->CurrentLOD = CurrentLOD;
+			BatchElementTools.IndexBuffer = SharedBuffers->IndexBuffers[CurrentLOD];
+			BatchElementTools.NumPrimitives = NumPrimitives;
+			BatchElementTools.FirstIndex = BatchElement.FirstIndex;
+			BatchElementTools.MinVertexIndex = SharedBuffers->IndexRanges[CurrentLOD].MinIndex[SubX][SubY];
+			BatchElementTools.MaxVertexIndex = SharedBuffers->IndexRanges[CurrentLOD].MaxIndex[SubX][SubY];
+#endif
+		}
+	}
+
+	// Render the landscape component
+#if WITH_EDITOR
+	switch(GLandscapeViewMode)
+	{
+	case ELandscapeViewMode::DebugLayer:
+		if( GLayerDebugColorMaterial && EditToolRenderData )
+		{
+			const FLandscapeDebugMaterialRenderProxy DebugColorMaterialInstance(GLayerDebugColorMaterial->GetRenderProxy(false), 
+				(EditToolRenderData->DebugChannelR >= 0 ? WeightmapTextures[EditToolRenderData->DebugChannelR/4] : NULL),
+				(EditToolRenderData->DebugChannelG >= 0 ? WeightmapTextures[EditToolRenderData->DebugChannelG/4] : NULL),
+				(EditToolRenderData->DebugChannelB >= 0 ? WeightmapTextures[EditToolRenderData->DebugChannelB/4] : NULL),	
+				(EditToolRenderData->DebugChannelR >= 0 ? DebugColorMask::Masks[EditToolRenderData->DebugChannelR%4] : DebugColorMask::Masks[4]),
+				(EditToolRenderData->DebugChannelG >= 0 ? DebugColorMask::Masks[EditToolRenderData->DebugChannelG%4] : DebugColorMask::Masks[4]),
+				(EditToolRenderData->DebugChannelB >= 0 ? DebugColorMask::Masks[EditToolRenderData->DebugChannelB%4] : DebugColorMask::Masks[4])
+				);
+
+			DynamicMeshTools.MaterialRenderProxy = &DebugColorMaterialInstance;
+			NumPasses += DrawRichMesh(PDI, DynamicMeshTools, FLinearColor(1,1,1,1), LevelColor, FLinearColor(1.0f,1.0f,1.0f), this, IsSelected(), bIsWireframe);
+			NumTriangles += DynamicMeshTools.GetNumPrimitives();
+			NumDrawCalls += DynamicMeshTools.Elements.Num();
+		}
+		break;
+
+	case ELandscapeViewMode::LayerDensity:
+		if( EditToolRenderData )
+		{
+			int32 ColorIndex = FMath::Min<int32>(NumWeightmapLayerAllocations, GEngine->ShaderComplexityColors.Num());
+			const FColoredMaterialRenderProxy LayerDensityMaterialInstance(GEngine->LevelColorationUnlitMaterial->GetRenderProxy(false), ColorIndex ? GEngine->ShaderComplexityColors[ColorIndex-1] : FLinearColor::Black);
+
+			DynamicMeshTools.MaterialRenderProxy = &LayerDensityMaterialInstance;
+
+			NumPasses += DrawRichMesh(PDI, DynamicMeshTools, FLinearColor(1,1,1,1), LevelColor, FLinearColor(1.0f,1.0f,1.0f), this, IsSelected(), bIsWireframe);
+			NumTriangles += DynamicMeshTools.GetNumPrimitives();
+			NumDrawCalls += DynamicMeshTools.Elements.Num();
+		}
+		break;
+
+	case ELandscapeViewMode::LOD:
+		{
+			FLinearColor WireColors[7];
+			WireColors[0] = FLinearColor(1,1,1,1);
+			WireColors[1] = FLinearColor(1,0,0,1);
+			WireColors[2] = FLinearColor(0,1,0,1);
+			WireColors[3] = FLinearColor(0,0,1,1);
+			WireColors[4] = FLinearColor(1,1,0,1);
+			WireColors[5] = FLinearColor(1,0,1,1);
+			WireColors[6] = FLinearColor(0,1,1,1);
+
+			for( int32 i=0;i<DynamicMeshTools.Elements.Num();i++ )
+			{
+				FMeshBatch TempMesh = DynamicMeshTools;
+				TempMesh.Elements.Empty(1);
+				TempMesh.Elements.Add(DynamicMeshTools.Elements[i]);
+				int32 ColorIndex = ((FLandscapeBatchElementParams*)DynamicMeshTools.Elements[i].UserData)->CurrentLOD;
+				FLinearColor Color = ForcedLOD >= 0 ? WireColors[ColorIndex] : WireColors[ColorIndex]*0.2f;
+				const FColoredMaterialRenderProxy LODMaterialInstance(GEngine->LevelColorationUnlitMaterial->GetRenderProxy(false), Color);
+				TempMesh.MaterialRenderProxy = &LODMaterialInstance;
+				NumPasses += DrawRichMesh(PDI, TempMesh, Color, LevelColor, FLinearColor(1.0f,1.0f,1.0f), this, IsSelected());
+				NumTriangles += DynamicMeshTools.Elements[i].NumPrimitives;
+				NumDrawCalls++;
+			}
+		}
+		break;
+
+	default:
+#else
+	{
+#endif // WITH_EDITOR
+		// Regular Landscape rendering. Only use the dynamic path if we're rendering a rich view or we've disabled the static path for debugging.
+		if( IsRichView(*View->Family) || 
+			GLandscapeDebugOptions.bDisableStatic || 
+			bIsWireframe ||
+#if WITH_EDITOR
+			(IsSelected() && !GLandscapeEditModeActive)
+#else
+			IsSelected()
+#endif
+			)
+		{
+			DynamicMesh.MaterialRenderProxy = MaterialInterface->GetRenderProxy(false);
+			NumPasses += DrawRichMesh(PDI, DynamicMesh, FLinearColor(1,1,1,1), LevelColor, FLinearColor(1.0f,1.0f,1.0f), this, IsSelected(), bIsWireframe);
+			NumTriangles += DynamicMesh.GetNumPrimitives();
+			NumDrawCalls += DynamicMesh.Elements.Num();
+		}
+	}
+
+#if WITH_EDITOR
+	// Extra render passes for landscape tools
+	if( GLandscapeEditModeActive )
+	{
+		// Region selection
+		if ( EditToolRenderData && EditToolRenderData->SelectedType )
+		{
+			if ((GLandscapeEditRenderMode & ELandscapeEditRenderMode::SelectRegion) && (EditToolRenderData->SelectedType & FLandscapeEditToolRenderData::ST_REGION)
+				&& !(GLandscapeEditRenderMode & ELandscapeEditRenderMode::Mask))
+			{
+				const FLandscapeSelectMaterialRenderProxy SelectMaterialInstance(GSelectionRegionMaterial->GetRenderProxy(false), EditToolRenderData->DataTexture ? EditToolRenderData->DataTexture : GLandscapeBlackTexture);
+				DynamicMeshTools.MaterialRenderProxy = &SelectMaterialInstance;
+				NumPasses += PDI->DrawMesh(DynamicMeshTools);
+				NumTriangles += DynamicMeshTools.GetNumPrimitives();
+				NumDrawCalls += DynamicMeshTools.Elements.Num();
+			}
+			if ((GLandscapeEditRenderMode & ELandscapeEditRenderMode::SelectComponent) && (EditToolRenderData->SelectedType & FLandscapeEditToolRenderData::ST_COMPONENT))
+			{
+				DynamicMeshTools.MaterialRenderProxy = GSelectionColorMaterial->GetRenderProxy(0);
+				NumPasses += PDI->DrawMesh(DynamicMeshTools);
+				NumTriangles += DynamicMeshTools.GetNumPrimitives();
+				NumDrawCalls += DynamicMeshTools.Elements.Num();
+			}
+		}
+
+		// Mask
+		if ((GLandscapeEditRenderMode & ELandscapeEditRenderMode::SelectRegion) && (GLandscapeEditRenderMode & ELandscapeEditRenderMode::Mask))
+		{
+			if (EditToolRenderData && (EditToolRenderData->SelectedType & FLandscapeEditToolRenderData::ST_REGION) )
+			{
+				const FLandscapeMaskMaterialRenderProxy MaskMaterialInstance(GMaskRegionMaterial->GetRenderProxy(false), EditToolRenderData->DataTexture ? EditToolRenderData->DataTexture : GLandscapeBlackTexture, !!(GLandscapeEditRenderMode & ELandscapeEditRenderMode::InvertedMask) );
+				DynamicMeshTools.MaterialRenderProxy = &MaskMaterialInstance;
+				NumPasses += PDI->DrawMesh(DynamicMeshTools);
+				NumTriangles += DynamicMeshTools.GetNumPrimitives();
+				NumDrawCalls += DynamicMeshTools.Elements.Num();
+			}
+			else if (!(GLandscapeEditRenderMode & ELandscapeEditRenderMode::InvertedMask))
+			{
+				const FLandscapeMaskMaterialRenderProxy MaskMaterialInstance(GMaskRegionMaterial->GetRenderProxy(false), GLandscapeBlackTexture, false );
+				DynamicMeshTools.MaterialRenderProxy = &MaskMaterialInstance;
+				NumPasses += PDI->DrawMesh(DynamicMeshTools);
+				NumTriangles += DynamicMeshTools.GetNumPrimitives();
+				NumDrawCalls += DynamicMeshTools.Elements.Num();
+			}
+		}
+
+		// Edit mode tools
+		if( EditToolRenderData )
+		{
+			if (EditToolRenderData->ToolMaterial)
+			{
+				DynamicMeshTools.MaterialRenderProxy = EditToolRenderData->ToolMaterial->GetRenderProxy(0);		
+				NumPasses += PDI->DrawMesh(DynamicMeshTools);
+				NumTriangles += DynamicMeshTools.GetNumPrimitives();
+				NumDrawCalls += DynamicMeshTools.Elements.Num();
+			}
+
+			if (EditToolRenderData->GizmoMaterial && GLandscapeEditRenderMode & ELandscapeEditRenderMode::Gizmo)
+			{
+				DynamicMeshTools.MaterialRenderProxy = EditToolRenderData->GizmoMaterial->GetRenderProxy(0);
+				NumPasses += PDI->DrawMesh(DynamicMeshTools);
+				NumTriangles += DynamicMeshTools.GetNumPrimitives();
+				NumDrawCalls += DynamicMeshTools.Elements.Num();
+			}
+		}
+	}
+#endif // WITH_EDITOR
+
+	if ( GLandscapeDebugOptions.bShowPatches )
+	{
+		DrawWireBox(PDI, GetBounds().GetBox(), FColor(255, 255, 0), SDPG_World);
+	}
+
+	RenderBounds(PDI, View->Family->EngineShowFlags, GetBounds(), IsSelected());
+
+	INC_DWORD_STAT_BY(STAT_LandscapeComponents, NumPasses);
+	INC_DWORD_STAT_BY(STAT_LandscapeDrawCalls, NumDrawCalls);
+	INC_DWORD_STAT_BY(STAT_LandscapeTriangles, NumTriangles * NumPasses);
+}
+
+//
+// FLandscapeVertexBuffer
+//
+
+/** 
+* Initialize the RHI for this rendering resource 
+*/
+void FLandscapeVertexBuffer::InitRHI()
+{
+	// create a static vertex buffer
+	FRHIResourceCreateInfo CreateInfo;
+	VertexBufferRHI = RHICreateVertexBuffer(FMath::Square(SubsectionSizeVerts) * FMath::Square(NumSubsections) * sizeof(FLandscapeVertex), BUF_Static, CreateInfo);
+	FLandscapeVertex* Vertex = (FLandscapeVertex*)RHILockVertexBuffer(VertexBufferRHI, 0, FMath::Square(SubsectionSizeVerts) * FMath::Square(NumSubsections) * sizeof(FLandscapeVertex),RLM_WriteOnly);
+
+	for( int32 SubY=0;SubY<NumSubsections;SubY++ )
+	{
+		for( int32 SubX=0;SubX<NumSubsections;SubX++ )
+		{
+			for( int32 y=0;y<SubsectionSizeVerts;y++ )
+			{
+				for( int32 x=0;x<SubsectionSizeVerts;x++ )
+				{
+					Vertex->VertexX = x;
+					Vertex->VertexY = y;
+					Vertex->SubX = SubX;
+					Vertex->SubY = SubY;
+					Vertex++;
+				}
+			}
+		}
+	}
+
+	RHIUnlockVertexBuffer(VertexBufferRHI);
+}
+
+//
+// FLandscapeSharedBuffers
+//
+
+template <typename INDEX_TYPE>
+void FLandscapeSharedBuffers::CreateIndexBuffers(ERHIFeatureLevel::Type InFeatureLevel, bool bRequiresAdjacencyInformation)
+{
+	if (InFeatureLevel <= ERHIFeatureLevel::ES3_1)
+	{
+		if (!bVertexScoresComputed)
+		{
+			bVertexScoresComputed = ComputeVertexScores();
+		}
+	}
+
+	TMap<uint64, INDEX_TYPE> VertexMap;
+	INDEX_TYPE VertexCount = 0;
+	int32 SubsectionSizeQuads = SubsectionSizeVerts - 1;
+
+	// Layout index buffer to determine best vertex order
+	int32 MaxLOD = NumIndexBuffers - 1;
+	for (int32 Mip = MaxLOD; Mip >= 0; Mip--)
+	{
+		int32 LodSubsectionSizeQuads = (SubsectionSizeVerts >> Mip) - 1;
+
+		TArray<INDEX_TYPE> NewIndices;
+		int32 ExpectedNumIndices = FMath::Square(NumSubsections) * FMath::Square(LodSubsectionSizeQuads) * 6;
+		NewIndices.Empty(ExpectedNumIndices);
+
+		int32& MaxIndexFull = IndexRanges[Mip].MaxIndexFull;
+		int32& MinIndexFull = IndexRanges[Mip].MinIndexFull;
+		MaxIndexFull = 0;
+		MinIndexFull = MAX_int32;
+
+		if (InFeatureLevel <= ERHIFeatureLevel::ES3_1)
+		{
+			// ES2 version
+			float MipRatio = (float)SubsectionSizeQuads / (float)LodSubsectionSizeQuads; // Morph current MIP to base MIP
+
+			for( int32 SubY=0;SubY<NumSubsections;SubY++ )
+			{
+				for( int32 SubX=0;SubX<NumSubsections;SubX++ )
+				{
+					TArray<INDEX_TYPE> SubIndices;
+					SubIndices.Empty(FMath::Square(LodSubsectionSizeQuads) * 6);
+
+					int32& MaxIndex = IndexRanges[Mip].MaxIndex[SubX][SubY];
+					int32& MinIndex = IndexRanges[Mip].MinIndex[SubX][SubY];
+					MaxIndex = 0;
+					MinIndex = MAX_int32;
+
+					for (int32 y = 0; y < LodSubsectionSizeQuads; y++)
+					{
+						for (int32 x = 0; x < LodSubsectionSizeQuads; x++)
+						{
+							int32 x0 = FMath::RoundToInt((float)x * MipRatio);
+							int32 y0 = FMath::RoundToInt((float)y * MipRatio);
+							int32 x1 = FMath::RoundToInt((float)(x + 1) * MipRatio);
+							int32 y1 = FMath::RoundToInt((float)(y + 1) * MipRatio);
+
+							FLandscapeVertexRef V00(x0, y0, SubX, SubY);
+							FLandscapeVertexRef V10(x1, y0, SubX, SubY);
+							FLandscapeVertexRef V11(x1, y1, SubX, SubY);
+							FLandscapeVertexRef V01(x0, y1, SubX, SubY);
+
+							uint64 Key00 = V00.MakeKey();
+							uint64 Key10 = V10.MakeKey();
+							uint64 Key11 = V11.MakeKey();
+							uint64 Key01 = V01.MakeKey();
+
+							INDEX_TYPE i00;
+							INDEX_TYPE i10;
+							INDEX_TYPE i11;
+							INDEX_TYPE i01;
+
+							INDEX_TYPE* KeyPtr = VertexMap.Find(Key00);
+							if (KeyPtr == NULL)
+							{
+								i00 = VertexCount++;
+								VertexMap.Add(Key00, i00);
+							}
+							else
+							{
+								i00 = *KeyPtr;
+							}
+
+							KeyPtr = VertexMap.Find(Key10);
+							if (KeyPtr == NULL)
+							{
+								i10 = VertexCount++;
+								VertexMap.Add(Key10, i10);
+							}
+							else
+							{
+								i10 = *KeyPtr;
+							}
+
+							KeyPtr = VertexMap.Find(Key11);
+							if (KeyPtr == NULL)
+							{
+								i11 = VertexCount++;
+								VertexMap.Add(Key11, i11);
+							}
+							else
+							{
+								i11 = *KeyPtr;
+							}
+
+							KeyPtr = VertexMap.Find(Key01);
+							if (KeyPtr == NULL)
+							{
+								i01 = VertexCount++;
+								VertexMap.Add(Key01, i01);
+							}
+							else
+							{
+								i01 = *KeyPtr;
+							}
+
+							// Update the min/max index ranges
+							MaxIndex = FMath::Max<int32>(MaxIndex, i00);
+							MinIndex = FMath::Min<int32>(MinIndex, i00);
+							MaxIndex = FMath::Max<int32>(MaxIndex, i10);
+							MinIndex = FMath::Min<int32>(MinIndex, i10);
+							MaxIndex = FMath::Max<int32>(MaxIndex, i11);
+							MinIndex = FMath::Min<int32>(MinIndex, i11);
+							MaxIndex = FMath::Max<int32>(MaxIndex, i01);
+							MinIndex = FMath::Min<int32>(MinIndex, i01);
+
+							SubIndices.Add(i00);
+							SubIndices.Add(i11);
+							SubIndices.Add(i10);
+
+							SubIndices.Add(i00);
+							SubIndices.Add(i01);
+							SubIndices.Add(i11);
+						}
+					}
+
+					// update min/max for full subsection
+					MaxIndexFull = FMath::Max<int32>(MaxIndexFull, MaxIndex);
+					MinIndexFull = FMath::Min<int32>(MinIndexFull, MinIndex);
+
+					TArray<INDEX_TYPE> NewSubIndices;
+					::OptimizeFaces<INDEX_TYPE>(SubIndices, NewSubIndices, 32);
+					NewIndices.Append(NewSubIndices);
+				}
+			}
+		}
+		else
+		{
+			// non-ES2 version
+			for (int32 SubY = 0; SubY < NumSubsections; SubY++)
+			{
+				for (int32 SubX = 0; SubX < NumSubsections; SubX++)
+				{
+					int32 SubOffset = (SubX + SubY * NumSubsections) * FMath::Square(SubsectionSizeVerts);
+
+					int32& MaxIndex = IndexRanges[Mip].MaxIndex[SubX][SubY];
+					int32& MinIndex = IndexRanges[Mip].MinIndex[SubX][SubY];
+					MaxIndex = 0;
+					MinIndex = MAX_int32;
+
+					for (int32 y = 0; y < LodSubsectionSizeQuads; y++)
+					{
+						for (int32 x = 0; x < LodSubsectionSizeQuads; x++)
+						{
+							INDEX_TYPE i00 = (x+0) + (y+0) * SubsectionSizeVerts + SubOffset;
+							INDEX_TYPE i10 = (x+1) + (y+0) * SubsectionSizeVerts + SubOffset;
+							INDEX_TYPE i11 = (x+1) + (y+1) * SubsectionSizeVerts + SubOffset;
+							INDEX_TYPE i01 = (x+0) + (y+1) * SubsectionSizeVerts + SubOffset;
+
+							NewIndices.Add(i00);
+							NewIndices.Add(i11);
+							NewIndices.Add(i10);
+
+							NewIndices.Add(i00);
+							NewIndices.Add(i01);
+							NewIndices.Add(i11);
+
+							// Update the min/max index ranges
+							MaxIndex = FMath::Max<int32>(MaxIndex, i00);
+							MinIndex = FMath::Min<int32>(MinIndex, i00);
+							MaxIndex = FMath::Max<int32>(MaxIndex, i10);
+							MinIndex = FMath::Min<int32>(MinIndex, i10);
+							MaxIndex = FMath::Max<int32>(MaxIndex, i11);
+							MinIndex = FMath::Min<int32>(MinIndex, i11);
+							MaxIndex = FMath::Max<int32>(MaxIndex, i01);
+							MinIndex = FMath::Min<int32>(MinIndex, i01);
+						}
+					}
+
+					// update min/max for full subsection
+					MaxIndexFull = FMath::Max<int32>(MaxIndexFull, MaxIndex);
+					MinIndexFull = FMath::Min<int32>(MinIndexFull, MinIndex);
+				}
+			}
+
+			check(MinIndexFull <= (uint32)((INDEX_TYPE)(~(INDEX_TYPE)0)));
+			check(NewIndices.Num() == ExpectedNumIndices);
+		}
+
+		// Create and init new index buffer with index data
+		FRawStaticIndexBuffer16or32<INDEX_TYPE>* IndexBuffer = (FRawStaticIndexBuffer16or32<INDEX_TYPE>*)IndexBuffers[Mip];
+		if (!IndexBuffer)
+		{
+			IndexBuffer = new FRawStaticIndexBuffer16or32<INDEX_TYPE>(false);
+		}
+		IndexBuffer->AssignNewBuffer(NewIndices);
+
+		// Delay init resource to keep CPU data until create AdjacencyIndexbuffers
+		if (!bRequiresAdjacencyInformation)
+		{
+			IndexBuffer->InitResource();
+		}
+
+		IndexBuffers[Mip] = IndexBuffer;
+	}
+}
+
+FLandscapeSharedBuffers::FLandscapeSharedBuffers(int32 InSharedBuffersKey, int32 InSubsectionSizeQuads, int32 InNumSubsections, ERHIFeatureLevel::Type InFeatureLevel, bool bRequiresAdjacencyInformation)
+:	SharedBuffersKey(InSharedBuffersKey)
+,	NumIndexBuffers(FMath::CeilLogTwo(InSubsectionSizeQuads+1))
+,	SubsectionSizeVerts(InSubsectionSizeQuads+1)
+,	NumSubsections(InNumSubsections)
+,	VertexFactory(NULL)
+,	VertexBuffer(NULL)
+,	AdjacencyIndexBuffers(NULL)
+,	bUse32BitIndices(false)
+{
+	if (InFeatureLevel > ERHIFeatureLevel::ES3_1)
+	{
+		// Vertex Buffer cannot be shared
+		VertexBuffer = new FLandscapeVertexBuffer(SubsectionSizeVerts, NumSubsections);
+	}
+	IndexBuffers = new FIndexBuffer*[NumIndexBuffers];
+	FMemory::Memzero(IndexBuffers, sizeof(FIndexBuffer*)* NumIndexBuffers);
+	IndexRanges = new FLandscapeIndexRanges[NumIndexBuffers]();
+
+	// See if we need to use 16 or 32-bit index buffers
+	if( FMath::Square(SubsectionSizeVerts) * FMath::Square(NumSubsections) > 65535 )
+	{
+		bUse32BitIndices = true;
+		CreateIndexBuffers<uint32>(InFeatureLevel, bRequiresAdjacencyInformation);
+	}
+	else
+	{
+		CreateIndexBuffers<uint16>(InFeatureLevel, bRequiresAdjacencyInformation);
+	}
+}
+
+FLandscapeSharedBuffers::~FLandscapeSharedBuffers()
+{
+	delete VertexBuffer;
+	
+	for( int32 i=0;i<NumIndexBuffers;i++ )
+	{
+		IndexBuffers[i]->ReleaseResource();
+		delete IndexBuffers[i];
+	}
+	delete[] IndexBuffers;
+	delete[] IndexRanges;
+
+	if (AdjacencyIndexBuffers)
+	{
+		if (AdjacencyIndexBuffers->Release() == 0)
+		{
+			FLandscapeComponentSceneProxy::SharedAdjacencyIndexBufferMap.Remove(SharedBuffersKey);
+		}
+		AdjacencyIndexBuffers = NULL;
+	}
+
+	delete VertexFactory;
+}
+
+template<typename IndexType>
+static void BuildLandscapeAdjacencyIndexBuffer(int32 LODSubsectionSizeQuads, const FRawStaticIndexBuffer16or32<IndexType>* Indices, TArray<IndexType>& OutPnAenIndices)
+{
+	if (Indices && Indices->Num())
+	{
+		// Landscape use regular grid, so only expand Index buffer works
+		// PN AEN Dominant Corner
+		uint32 TriCount = LODSubsectionSizeQuads*LODSubsectionSizeQuads * 2;
+		uint32 ExpandedCount = 12 * TriCount;
+		OutPnAenIndices.Empty( ExpandedCount );
+		OutPnAenIndices.AddUninitialized( ExpandedCount );
+
+		for ( uint32 TriIdx = 0; TriIdx < TriCount; ++TriIdx )
+		{
+			uint32 OutStartIdx = TriIdx*12;
+			uint32 InStartIdx = TriIdx*3;
+			OutPnAenIndices[ OutStartIdx +  0] =	Indices->Get(InStartIdx + 0);
+			OutPnAenIndices[ OutStartIdx +  1] =	Indices->Get(InStartIdx + 1);
+			OutPnAenIndices[ OutStartIdx +  2] =	Indices->Get(InStartIdx + 2);
+
+			OutPnAenIndices[ OutStartIdx +  3] =	Indices->Get(InStartIdx + 0);
+			OutPnAenIndices[ OutStartIdx +  4] =	Indices->Get(InStartIdx + 1);
+			OutPnAenIndices[ OutStartIdx +  5] =	Indices->Get(InStartIdx + 1);
+			OutPnAenIndices[ OutStartIdx +  6] =	Indices->Get(InStartIdx + 2);
+			OutPnAenIndices[ OutStartIdx +  7] =	Indices->Get(InStartIdx + 2);
+			OutPnAenIndices[ OutStartIdx +  8] =	Indices->Get(InStartIdx + 0);
+
+			OutPnAenIndices[ OutStartIdx +  9] =	Indices->Get(InStartIdx + 0);
+			OutPnAenIndices[ OutStartIdx +  10] =	Indices->Get(InStartIdx + 1);
+			OutPnAenIndices[ OutStartIdx +  11] =	Indices->Get(InStartIdx + 2);
+		}
+	}
+	else
+	{
+		OutPnAenIndices.Empty();
+	}
+}
+
+FLandscapeSharedAdjacencyIndexBuffer::FLandscapeSharedAdjacencyIndexBuffer(FLandscapeSharedBuffers* Buffers)
+{
+	ensure(Buffers && Buffers->IndexBuffers);
+
+	// Currently only support PN-AEN-Dominant Corner, which is the only mode for UE4 for now
+	IndexBuffers.Empty(Buffers->NumIndexBuffers);
+
+	bool b32BitIndex = FMath::Square(Buffers->SubsectionSizeVerts) * FMath::Square(Buffers->NumSubsections) > 65535;
+	for (int32 i = 0; i < Buffers->NumIndexBuffers; ++i)
+	{
+		if( b32BitIndex )
+		{
+			TArray<uint32> OutPnAenIndices;
+			BuildLandscapeAdjacencyIndexBuffer<uint32>( (Buffers->SubsectionSizeVerts>>i)-1, (FRawStaticIndexBuffer16or32<uint32>*)Buffers->IndexBuffers[i], OutPnAenIndices );
+
+			FRawStaticIndexBuffer16or32<uint32>* IndexBuffer = new FRawStaticIndexBuffer16or32<uint32>();
+			IndexBuffer->AssignNewBuffer(OutPnAenIndices);
+			IndexBuffers.Add(IndexBuffer);
+		}
+		else
+		{
+			TArray<uint16> OutPnAenIndices;
+			BuildLandscapeAdjacencyIndexBuffer<uint16>( (Buffers->SubsectionSizeVerts>>i)-1, (FRawStaticIndexBuffer16or32<uint16>*)Buffers->IndexBuffers[i], OutPnAenIndices );
+
+			FRawStaticIndexBuffer16or32<uint16>* IndexBuffer = new FRawStaticIndexBuffer16or32<uint16>();
+			IndexBuffer->AssignNewBuffer(OutPnAenIndices);
+			IndexBuffers.Add(IndexBuffer);
+		}
+
+		IndexBuffers[i]->InitResource();
+	}
+}
+
+FLandscapeSharedAdjacencyIndexBuffer::~FLandscapeSharedAdjacencyIndexBuffer()
+{
+	for (int i = 0; i < IndexBuffers.Num(); ++i)
+	{
+		IndexBuffers[i]->ReleaseResource();
+		delete IndexBuffers[i];
+	}
+}
+
+//
+// FLandscapeVertexFactoryVertexShaderParameters
+//
+
+/** Shader parameters for use with FLandscapeVertexFactory */
+class FLandscapeVertexFactoryVertexShaderParameters : public FVertexFactoryShaderParameters
+{
+public:
+	/**
+	* Bind shader constants by name
+	* @param	ParameterMap - mapping of named shader constants to indices
+	*/
+	virtual void Bind(const FShaderParameterMap& ParameterMap) override
+	{
+		HeightmapTextureParameter.Bind(ParameterMap,TEXT("HeightmapTexture"));
+		HeightmapTextureParameterSampler.Bind(ParameterMap,TEXT("HeightmapTextureSampler"));
+		LodValuesParameter.Bind(ParameterMap,TEXT("LodValues"));
+		NeighborSectionLodParameter.Bind(ParameterMap,TEXT("NeighborSectionLod"));
+		LodBiasParameter.Bind(ParameterMap,TEXT("LodBias"));
+		SectionLodsParameter.Bind(ParameterMap,TEXT("SectionLods"));
+		XYOffsetTextureParameter.Bind(ParameterMap,TEXT("XYOffsetmapTexture"));
+		XYOffsetTextureParameterSampler.Bind(ParameterMap,TEXT("XYOffsetmapTextureSampler"));
+	}
+
+	/**
+	* Serialize shader params to an archive
+	* @param	Ar - archive to serialize to
+	*/
+	virtual void Serialize(FArchive& Ar) override
+	{
+		Ar << HeightmapTextureParameter;
+		Ar << HeightmapTextureParameterSampler;
+		Ar << LodValuesParameter;
+		Ar << NeighborSectionLodParameter;
+		Ar << LodBiasParameter;
+		Ar << SectionLodsParameter;
+		Ar << XYOffsetTextureParameter;
+		Ar << XYOffsetTextureParameterSampler;
+	}
+
+	/**
+	* Set any shader data specific to this vertex factory
+	*/
+	virtual void SetMesh(FRHICommandList& RHICmdList, FShader* VertexShader,const class FVertexFactory* VertexFactory,const class FSceneView& View,const struct FMeshBatchElement& BatchElement,uint32 DataFlags) const override
+	{
+		SCOPE_CYCLE_COUNTER(STAT_LandscapeVFDrawTime);
+
+		const FLandscapeBatchElementParams* BatchElementParams = (const FLandscapeBatchElementParams*)BatchElement.UserData;
+		check(BatchElementParams);
+
+		const FLandscapeComponentSceneProxy* SceneProxy = BatchElementParams->SceneProxy;
+		SetUniformBufferParameter(RHICmdList, VertexShader->GetVertexShader(),VertexShader->GetUniformBufferParameter<FLandscapeUniformShaderParameters>(),*BatchElementParams->LandscapeUniformShaderParametersResource);
+
+		if( HeightmapTextureParameter.IsBound() )
+		{
+			SetTextureParameter(
+				RHICmdList, 
+				VertexShader->GetVertexShader(),
+				HeightmapTextureParameter,
+				HeightmapTextureParameterSampler,
+				TStaticSamplerState<SF_Point>::GetRHI(),
+				SceneProxy->HeightmapTexture->Resource->TextureRHI
+				);
+		}
+
+		if( LodBiasParameter.IsBound() )
+		{
+			FVector4 LodBias(
+				SceneProxy->PrecomputedLODFactor,
+				1.f / (1.f - SceneProxy->PrecomputedLODFactor),
+				SceneProxy->HeightmapTexture->GetNumMips() - FMath::Min(SceneProxy->HeightmapTexture->ResidentMips, SceneProxy->HeightmapTexture->RequestedMips),
+				0.f // Reserved
+				);
+			SetShaderValue(RHICmdList, VertexShader->GetVertexShader(), LodBiasParameter, LodBias);
+		}
+
+		// Calculate LOD params
+		FVector CameraLocalPos3D = SceneProxy->WorldToLocal.TransformPosition(View.ViewMatrices.ViewOrigin); 
+		FVector2D CameraLocalPos = FVector2D(CameraLocalPos3D.X, CameraLocalPos3D.Y);
+
+		FVector4 fCurrentLODs;
+		FVector4 CurrentNeighborLODs[4];
+		float DistLOD = 0.f;
+
+		if( BatchElementParams->SubX == -1 )
+		{
+			for( int32 SubY = 0; SubY < SceneProxy->NumSubsections; SubY++ )
+			{
+				for( int32 SubX = 0; SubX < SceneProxy->NumSubsections; SubX++ )
+				{
+					int32 SubIndex = SubX + 2 * SubY;
+					SceneProxy->CalcLODParamsForSubsection(View, CameraLocalPos, SubX, SubY, fCurrentLODs[SubIndex], CurrentNeighborLODs[SubIndex], DistLOD);
+				}
+			}
+		}
+		else
+		{
+			int32 SubIndex = BatchElementParams->SubX + 2 * BatchElementParams->SubY;
+			SceneProxy->CalcLODParamsForSubsection(View, CameraLocalPos, BatchElementParams->SubX, BatchElementParams->SubY, fCurrentLODs[SubIndex], CurrentNeighborLODs[SubIndex], DistLOD);
+		}
+
+		if( SectionLodsParameter.IsBound() )
+		{
+			SetShaderValue(RHICmdList, VertexShader->GetVertexShader(), SectionLodsParameter, fCurrentLODs);
+		}
+
+		if( NeighborSectionLodParameter.IsBound() )
+		{
+			SetShaderValue(RHICmdList, VertexShader->GetVertexShader(), NeighborSectionLodParameter, CurrentNeighborLODs);
+		}
+
+		if( LodValuesParameter.IsBound() )
+		{
+			FVector4 LodValues(
+				DistLOD,
+				// convert current LOD coordinates into highest LOD coordinates
+				(float)SceneProxy->SubsectionSizeQuads / (float)(((SceneProxy->SubsectionSizeVerts) >> BatchElementParams->CurrentLOD)-1),
+				(float)((SceneProxy->SubsectionSizeVerts >> BatchElementParams->CurrentLOD) - 1),
+				1.f/(float)((SceneProxy->SubsectionSizeVerts >> BatchElementParams->CurrentLOD) - 1) );
+
+			SetShaderValue(RHICmdList, VertexShader->GetVertexShader(),LodValuesParameter,LodValues);
+		}
+	}
+
+	virtual uint32 GetSize() const override
+	{
+		return sizeof(*this);
+	}
+
+protected:
+	FShaderParameter LodValuesParameter;
+	FShaderParameter NeighborSectionLodParameter;
+	FShaderParameter LodBiasParameter;
+	FShaderParameter SectionLodsParameter;
+	FShaderResourceParameter HeightmapTextureParameter;
+	FShaderResourceParameter HeightmapTextureParameterSampler;
+	FShaderResourceParameter XYOffsetTextureParameter;
+	FShaderResourceParameter XYOffsetTextureParameterSampler;
+	TShaderUniformBufferParameter<FLandscapeUniformShaderParameters> LandscapeShaderParameters;
+};
+
+/** Shader parameters for use with FLandscapeVertexFactory */
+class FLandscapeXYOffsetVertexFactoryVertexShaderParameters : public FLandscapeVertexFactoryVertexShaderParameters
+{
+public:
+	/**
+	* Set any shader data specific to this vertex factory
+	*/
+	virtual void SetMesh(FRHICommandList& RHICmdList, FShader* VertexShader,const class FVertexFactory* VertexFactory,const class FSceneView& View,const struct FMeshBatchElement& BatchElement,uint32 DataFlags) const override
+	{
+		SCOPE_CYCLE_COUNTER(STAT_LandscapeVFDrawTime);
+
+		const FLandscapeBatchElementParams* BatchElementParams = (const FLandscapeBatchElementParams*)BatchElement.UserData;
+		check(BatchElementParams);
+
+		const FLandscapeComponentSceneProxy* SceneProxy = BatchElementParams->SceneProxy;
+		SetUniformBufferParameter(RHICmdList, VertexShader->GetVertexShader(),VertexShader->GetUniformBufferParameter<FLandscapeUniformShaderParameters>(),*BatchElementParams->LandscapeUniformShaderParametersResource);
+
+		if( HeightmapTextureParameter.IsBound() )
+		{
+			SetTextureParameter(
+				RHICmdList, 
+				VertexShader->GetVertexShader(),
+				HeightmapTextureParameter,
+				HeightmapTextureParameterSampler,
+				TStaticSamplerState<SF_Point>::GetRHI(),
+				SceneProxy->HeightmapTexture->Resource->TextureRHI
+				);
+		}
+
+		if( LodBiasParameter.IsBound() )
+		{
+			FVector4 LodBias(
+				SceneProxy->PrecomputedLODFactor,
+				1.f / (1.f - SceneProxy->PrecomputedLODFactor),
+				SceneProxy->HeightmapTexture->GetNumMips() - FMath::Min(SceneProxy->HeightmapTexture->ResidentMips, SceneProxy->HeightmapTexture->RequestedMips),
+				SceneProxy->XYOffsetmapTexture->GetNumMips() - FMath::Min(SceneProxy->XYOffsetmapTexture->ResidentMips, SceneProxy->XYOffsetmapTexture->RequestedMips)
+				);
+			SetShaderValue(RHICmdList, VertexShader->GetVertexShader(), LodBiasParameter, LodBias);
+		}
+
+		// Calculate LOD params
+		FVector CameraLocalPos3D = SceneProxy->WorldToLocal.TransformPosition(View.ViewMatrices.ViewOrigin); 
+		FVector2D CameraLocalPos = FVector2D(CameraLocalPos3D.X, CameraLocalPos3D.Y);
+
+		FVector4 fCurrentLODs;
+		FVector4 CurrentNeighborLODs[4];
+		float DistLOD = 0.f;
+
+		if( BatchElementParams->SubX == -1 )
+		{
+			for( int32 SubY = 0; SubY < SceneProxy->NumSubsections; SubY++ )
+			{
+				for( int32 SubX = 0; SubX < SceneProxy->NumSubsections; SubX++ )
+				{
+					int32 SubIndex = SubX + 2 * SubY;
+					SceneProxy->CalcLODParamsForSubsection(View, CameraLocalPos, SubX, SubY, fCurrentLODs[SubIndex], CurrentNeighborLODs[SubIndex], DistLOD);
+				}
+			}
+		}
+		else
+		{
+			int32 SubIndex = BatchElementParams->SubX + 2 * BatchElementParams->SubY;
+			SceneProxy->CalcLODParamsForSubsection(View, CameraLocalPos, BatchElementParams->SubX, BatchElementParams->SubY, fCurrentLODs[SubIndex], CurrentNeighborLODs[SubIndex], DistLOD);
+		}
+
+		if( SectionLodsParameter.IsBound() )
+		{
+			SetShaderValue(RHICmdList, VertexShader->GetVertexShader(), SectionLodsParameter, fCurrentLODs);
+		}
+
+		if( NeighborSectionLodParameter.IsBound() )
+		{
+			SetShaderValue(RHICmdList, VertexShader->GetVertexShader(), NeighborSectionLodParameter, CurrentNeighborLODs);
+		}
+
+		if( LodValuesParameter.IsBound() )
+		{
+			FVector4 LodValues(
+				DistLOD,
+				// convert current LOD coordinates into highest LOD coordinates
+				(float)SceneProxy->SubsectionSizeQuads / (float)(((SceneProxy->SubsectionSizeVerts) >> BatchElementParams->CurrentLOD)-1),
+				(float)((SceneProxy->SubsectionSizeVerts >> BatchElementParams->CurrentLOD) - 1),
+				1.f/(float)((SceneProxy->SubsectionSizeVerts >> BatchElementParams->CurrentLOD) - 1) );
+
+			SetShaderValue(RHICmdList, VertexShader->GetVertexShader(),LodValuesParameter,LodValues);
+		}
+
+		if( XYOffsetTextureParameter.IsBound() && SceneProxy->XYOffsetmapTexture )
+		{
+			SetTextureParameter(
+				RHICmdList, 
+				VertexShader->GetVertexShader(),
+				XYOffsetTextureParameter,
+				XYOffsetTextureParameterSampler,
+				TStaticSamplerState<SF_Point>::GetRHI(),
+				SceneProxy->XYOffsetmapTexture->Resource->TextureRHI
+				);
+		}
+	}
+};
+
+//
+// FLandscapeVertexFactoryPixelShaderParameters
+//
+/**
+* Bind shader constants by name
+* @param	ParameterMap - mapping of named shader constants to indices
+*/
+void FLandscapeVertexFactoryPixelShaderParameters::Bind(const FShaderParameterMap& ParameterMap)
+{
+	NormalmapTextureParameter.Bind(ParameterMap,TEXT("NormalmapTexture"));
+	NormalmapTextureParameterSampler.Bind(ParameterMap,TEXT("NormalmapTextureSampler"));
+	LocalToWorldNoScalingParameter.Bind(ParameterMap,TEXT("LocalToWorldNoScaling"));
+}
+
+/**
+* Serialize shader params to an archive
+* @param	Ar - archive to serialize to
+*/
+void FLandscapeVertexFactoryPixelShaderParameters::Serialize(FArchive& Ar)
+{
+	Ar	<< NormalmapTextureParameter
+		<< NormalmapTextureParameterSampler
+		<< LocalToWorldNoScalingParameter;
+}
+
+/**
+* Set any shader data specific to this vertex factory
+*/
+void FLandscapeVertexFactoryPixelShaderParameters::SetMesh(FRHICommandList& RHICmdList, FShader* PixelShader,const class FVertexFactory* VertexFactory,const class FSceneView& View,const struct FMeshBatchElement& BatchElement,uint32 DataFlags) const
+{
+	SCOPE_CYCLE_COUNTER(STAT_LandscapeVFDrawTime);
+
+	const FLandscapeBatchElementParams* BatchElementParams = (const FLandscapeBatchElementParams*)BatchElement.UserData;
+
+	if( LocalToWorldNoScalingParameter.IsBound() )
+	{
+		SetShaderValue(RHICmdList, PixelShader->GetPixelShader(), LocalToWorldNoScalingParameter, *BatchElementParams->LocalToWorldNoScalingPtr);
+	}
+
+	if( NormalmapTextureParameter.IsBound() )
+	{
+		SetTextureParameter(
+			RHICmdList, 
+			PixelShader->GetPixelShader(),
+			NormalmapTextureParameter,
+			NormalmapTextureParameterSampler,
+			BatchElementParams->SceneProxy->NormalmapTexture->Resource);
+	}
+}
+
+
+//
+// FLandscapeVertexFactory
+//
+
+void FLandscapeVertexFactory::InitRHI()
+{
+	// list of declaration items
+	FVertexDeclarationElementList Elements;
+
+	// position decls
+	Elements.Add(AccessStreamComponent(Data.PositionComponent,0));
+
+	// create the actual device decls
+	InitDeclaration(Elements,FVertexFactory::DataType());
+}
+
+FVertexFactoryShaderParameters* FLandscapeVertexFactory::ConstructShaderParameters(EShaderFrequency ShaderFrequency)
+{
+	switch( ShaderFrequency )
+	{
+	case SF_Vertex:
+		return new FLandscapeVertexFactoryVertexShaderParameters();
+		break;
+	case SF_Pixel:
+		return new FLandscapeVertexFactoryPixelShaderParameters();
+		break;
+	default:
+		return NULL;
+	}
+}
+
+void FLandscapeVertexFactory::ModifyCompilationEnvironment( EShaderPlatform Platform, const FMaterial* Material, FShaderCompilerEnvironment& OutEnvironment )
+{
+	FVertexFactory::ModifyCompilationEnvironment(Platform, Material, OutEnvironment);
+}
+
+
+IMPLEMENT_VERTEX_FACTORY_TYPE(FLandscapeVertexFactory, "LandscapeVertexFactory", true, true, true, false, false);
+
+/**
+* Copy the data from another vertex factory
+* @param Other - factory to copy from
+*/
+void FLandscapeVertexFactory::Copy(const FLandscapeVertexFactory& Other)
+{
+	//SetSceneProxy(Other.Proxy());
+	ENQUEUE_UNIQUE_RENDER_COMMAND_TWOPARAMETER(
+		FLandscapeVertexFactoryCopyData,
+		FLandscapeVertexFactory*,VertexFactory,this,
+		const DataType*,DataCopy,&Other.Data,
+	{
+		VertexFactory->Data = *DataCopy;
+	});	
+	BeginUpdateResourceRHI(this);
+}
+
+//
+// FLandscapeVertexFactory
+//
+
+FVertexFactoryShaderParameters* FLandscapeXYOffsetVertexFactory::ConstructShaderParameters(EShaderFrequency ShaderFrequency)
+{
+	switch( ShaderFrequency )
+	{
+	case SF_Vertex:
+		return new FLandscapeXYOffsetVertexFactoryVertexShaderParameters();
+		break;
+	case SF_Pixel:
+		return new FLandscapeVertexFactoryPixelShaderParameters();
+		break;
+	default:
+		return NULL;
+	}
+}
+
+void FLandscapeXYOffsetVertexFactory::ModifyCompilationEnvironment( EShaderPlatform Platform, const FMaterial* Material, FShaderCompilerEnvironment& OutEnvironment )
+{
+	FLandscapeVertexFactory::ModifyCompilationEnvironment(Platform, Material, OutEnvironment);
+	OutEnvironment.SetDefine(TEXT("LANDSCAPE_XYOFFSET"), TEXT("1"));
+}
+
+IMPLEMENT_VERTEX_FACTORY_TYPE(FLandscapeXYOffsetVertexFactory, "LandscapeVertexFactory", true, true, true, false, false);
+
+/** ULandscapeMaterialInstanceConstant */
+ULandscapeMaterialInstanceConstant::ULandscapeMaterialInstanceConstant(const FObjectInitializer& ObjectInitializer)
+	: Super(ObjectInitializer)
+{
+	bIsLayerThumbnail = false;
+	DataWeightmapIndex = -1;
+	DataWeightmapSize = 0;
+}
+
+void ULandscapeComponent::GetStreamingTextureInfo(TArray<FStreamingTexturePrimitiveInfo>& OutStreamingTextures) const
+{
+	ALandscapeProxy* Proxy = Cast<ALandscapeProxy>(GetOuter());
+	FSphere BoundingSphere = Bounds.GetSphere();
+	float StreamingDistanceMultiplier = 1.f;
+	if (Proxy)
+	{
+		StreamingDistanceMultiplier = Proxy->StreamingDistanceMultiplier;
+	}
+	const float TexelFactor = 0.75f * StreamingDistanceMultiplier * ComponentSizeQuads * (Proxy->GetRootComponent()->RelativeScale3D.X);
+
+	// Normal usage...
+	// Enumerate the textures used by the material.
+	if (MaterialInstance)
+	{
+		TArray<UTexture*> Textures;
+		MaterialInstance->GetUsedTextures(Textures, EMaterialQualityLevel::Num, false, GetWorld()->FeatureLevel, false);
+		// Add each texture to the output with the appropriate parameters.
+		// TODO: Take into account which UVIndex is being used.
+		for (int32 TextureIndex = 0;TextureIndex < Textures.Num();TextureIndex++)
+		{
+			FStreamingTexturePrimitiveInfo& StreamingTexture = *new(OutStreamingTextures) FStreamingTexturePrimitiveInfo;
+			StreamingTexture.Bounds = BoundingSphere;
+			StreamingTexture.TexelFactor = TexelFactor;
+			StreamingTexture.Texture = Textures[TextureIndex];
+		}
+
+		UMaterial* Material = MaterialInstance->GetMaterial();
+		if (Material)
+		{
+			int32 NumExpressions = Material->Expressions.Num();
+			for(int32 ExpressionIndex = 0;ExpressionIndex < NumExpressions; ExpressionIndex++)
+			{
+				UMaterialExpression* Expression = Material->Expressions[ExpressionIndex];
+				UMaterialExpressionTextureSample* TextureSample = Cast<UMaterialExpressionTextureSample>(Expression);
+
+				// TODO: This is only works for direct Coordinate Texture Sample cases
+				if(TextureSample && TextureSample->Coordinates.Expression)
+				{
+					UMaterialExpressionTextureCoordinate* TextureCoordinate =
+						Cast<UMaterialExpressionTextureCoordinate>( TextureSample->Coordinates.Expression );
+
+					UMaterialExpressionLandscapeLayerCoords* TerrainTextureCoordinate =
+						Cast<UMaterialExpressionLandscapeLayerCoords>( TextureSample->Coordinates.Expression );
+
+					if (TextureCoordinate || TerrainTextureCoordinate)
+					{
+						for (int32 i = 0; i < OutStreamingTextures.Num(); ++i)
+						{
+							FStreamingTexturePrimitiveInfo& StreamingTexture = OutStreamingTextures[i];
+							if (StreamingTexture.Texture == TextureSample->Texture)
+							{
+								if ( TextureCoordinate )
+								{
+									StreamingTexture.TexelFactor = TexelFactor * FPlatformMath::Max(TextureCoordinate->UTiling, TextureCoordinate->VTiling);
+								}
+								else //if ( TerrainTextureCoordinate )
+								{
+									StreamingTexture.TexelFactor = TexelFactor * TerrainTextureCoordinate->MappingScale;
+								}
+								break;
+							}
+						}
+					}
+				}
+			}
+		}
+
+		// Lightmap
+		const auto FeatureLevel = GetWorld() ? GetWorld()->FeatureLevel : GMaxRHIFeatureLevel;
+
+		FLightMap2D* Lightmap = LightMap ? LightMap->GetLightMap2D() : nullptr;
+		uint32 LightmapIndex = AllowHighQualityLightmaps(FeatureLevel) ? 0 : 1;
+		if (Lightmap && Lightmap->IsValid(LightmapIndex))
+		{
+			const FVector2D& Scale = Lightmap->GetCoordinateScale();
+			if (Scale.X > SMALL_NUMBER && Scale.Y > SMALL_NUMBER)
+			{
+				float LightmapFactorX = TexelFactor / Scale.X;
+				float LightmapFactorY = TexelFactor / Scale.Y;
+				FStreamingTexturePrimitiveInfo& StreamingTexture = *new(OutStreamingTextures) FStreamingTexturePrimitiveInfo;
+				StreamingTexture.Bounds		 = BoundingSphere;
+				StreamingTexture.TexelFactor = FMath::Max(LightmapFactorX, LightmapFactorY);
+				StreamingTexture.Texture	 = Lightmap->GetTexture(LightmapIndex);
+			}
+		}
+		
+		// Shadowmap
+		FShadowMap2D* Shadowmap = ShadowMap ? ShadowMap->GetShadowMap2D() : nullptr;
+		if (Shadowmap && Shadowmap->IsValid())
+		{
+			const FVector2D& Scale = Shadowmap->GetCoordinateScale();
+			if (Scale.X > SMALL_NUMBER && Scale.Y > SMALL_NUMBER)
+			{
+				float ShadowmapFactorX		 = TexelFactor / Scale.X;
+				float ShadowmapFactorY		 = TexelFactor / Scale.Y;
+				FStreamingTexturePrimitiveInfo& StreamingTexture = *new(OutStreamingTextures) FStreamingTexturePrimitiveInfo;
+				StreamingTexture.Bounds		 = BoundingSphere;
+				StreamingTexture.TexelFactor = FMath::Max(ShadowmapFactorX, ShadowmapFactorY);
+				StreamingTexture.Texture	 = Shadowmap->GetTexture();
+			}
+		}
+	}
+
+	// Weightmap
+	for(int32 TextureIndex = 0;TextureIndex < WeightmapTextures.Num();TextureIndex++)
+	{
+		FStreamingTexturePrimitiveInfo& StreamingWeightmap = *new(OutStreamingTextures) FStreamingTexturePrimitiveInfo;
+		StreamingWeightmap.Bounds = BoundingSphere;
+		StreamingWeightmap.TexelFactor = TexelFactor;
+		StreamingWeightmap.Texture = WeightmapTextures[TextureIndex];
+	}
+
+	// Heightmap
+	{
+		float HeightmapTexelFactor = TexelFactor * ((HeightmapTexture == NULL) ? 1.0f : (HeightmapTexture->GetSizeY()/(ComponentSizeQuads+1)));
+		FStreamingTexturePrimitiveInfo& StreamingHeightmap = *new(OutStreamingTextures) FStreamingTexturePrimitiveInfo;
+		StreamingHeightmap.Bounds = BoundingSphere;
+		StreamingHeightmap.TexelFactor = ForcedLOD >= 0 ? -13+ForcedLOD : HeightmapTexelFactor; // Minus Value indicate ForcedLOD, 13 for 8k texture
+		StreamingHeightmap.Texture = HeightmapTexture;
+	}
+
+	// XYOffset
+	if (XYOffsetmapTexture)
+	{
+		FStreamingTexturePrimitiveInfo& StreamingXYOffset = *new(OutStreamingTextures) FStreamingTexturePrimitiveInfo;
+		StreamingXYOffset.Bounds = BoundingSphere;
+		StreamingXYOffset.TexelFactor = TexelFactor;
+		StreamingXYOffset.Texture = XYOffsetmapTexture;
+	}
+
+#if WITH_EDITOR
+	if (GIsEditor && EditToolRenderData && EditToolRenderData->DataTexture)
+	{
+		FStreamingTexturePrimitiveInfo& StreamingDatamap = *new(OutStreamingTextures) FStreamingTexturePrimitiveInfo;
+		StreamingDatamap.Bounds = BoundingSphere;
+		StreamingDatamap.TexelFactor = TexelFactor;
+		StreamingDatamap.Texture = EditToolRenderData->DataTexture;
+	}
+#endif
+}
+
+void ALandscapeProxy::ChangeLODDistanceFactor(float InLODDistanceFactor)
+{
+	LODDistanceFactor = FMath::Clamp<float>(InLODDistanceFactor, 0.1f, MAX_LANDSCAPE_LOD_DISTANCE_FACTOR);
+	float LODFactor;
+	switch (LODFalloff)
+	{
+		case ELandscapeLODFalloff::SquareRoot:
+			LODFactor = FMath::Square(FMath::Min(LANDSCAPE_LOD_SQUARE_ROOT_FACTOR * LODDistanceFactor, MAX_LANDSCAPE_LOD_DISTANCE_FACTOR));
+			break;
+		default:
+		case ELandscapeLODFalloff::Linear:
+			LODFactor = LODDistanceFactor;
+			break;
+	}
+	
+	if (LandscapeComponents.Num())
+	{
+		int32 CompNum = LandscapeComponents.Num();
+		FLandscapeComponentSceneProxy** Proxies = new FLandscapeComponentSceneProxy*[CompNum];
+		for (int32 Idx = 0; Idx < CompNum; ++Idx)
+		{
+			Proxies[Idx] = (FLandscapeComponentSceneProxy*)(LandscapeComponents[Idx]->SceneProxy);
+		}
+		
+		ENQUEUE_UNIQUE_RENDER_COMMAND_THREEPARAMETER(
+			LandscapeChangeLODDistanceFactorCommand,
+			FLandscapeComponentSceneProxy**, Proxies, Proxies,
+			int32, CompNum, CompNum,
+			FVector2D, InLODDistanceFactors, FVector2D(FMath::Sqrt(2.f * FMath::Square((float)SubsectionSizeQuads)) * LANDSCAPE_LOD_DISTANCE_FACTOR / LODFactor, FMath::Min(LODFactor, 3.f) * 0.33f),
+		{
+			for (int32 Idx = 0; Idx < CompNum; ++Idx)
+			{
+				Proxies[Idx]->ChangeLODDistanceFactor_RenderThread(InLODDistanceFactors);
+			}
+			delete[] Proxies;
+		}
+		);
+	}
+};
+
+void FLandscapeComponentSceneProxy::ChangeLODDistanceFactor_RenderThread(FVector2D InLODDistanceFactors)
+{
+	LODDistance = InLODDistanceFactors.X;
+	PrecomputedLODFactor = InLODDistanceFactors.Y;
+}
