Landscape module from Engine split out.

diff --git a/Engine/Source/Runtime/Landscape/Private/LandscapeRenderMobile.cpp b/Engine/Source/Runtime/Landscape/Private/LandscapeRenderMobile.cpp
new file mode 100644
--- /dev/null
+++ b/Engine/Source/Runtime/Landscape/Private/LandscapeRenderMobile.cpp
@@ -0,0 +1,288 @@
+// Copyright 1998-2014 Epic Games, Inc. All Rights Reserved.
+
+/*=============================================================================
+LandscapeRenderMobile.cpp: Landscape Rendering without using vertex texture fetch
+=============================================================================*/
+
+#include "Landscape.h"
+#include "ShaderParameters.h"
+#include "ShaderParameterUtils.h"
+#include "LandscapeRender.h"
+#include "LandscapeRenderMobile.h"
+
+void FLandscapeVertexFactoryMobile::InitRHI()
+{
+	// list of declaration items
+	FVertexDeclarationElementList Elements;
+
+	// position decls
+	Elements.Add(AccessStreamComponent(MobileData.PositionComponent,0));
+
+	if (MobileData.LODHeightsComponent.Num())
+	{
+		const int32 BaseAttribute = 1;
+		for(int32 Index = 0;Index < MobileData.LODHeightsComponent.Num();Index++)
+		{
+			Elements.Add(AccessStreamComponent(MobileData.LODHeightsComponent[Index], BaseAttribute + Index));
+		}
+	}
+
+	// create the actual device decls
+	InitDeclaration(Elements,FVertexFactory::DataType());
+}
+
+/** Shader parameters for use with FLandscapeVertexFactory */
+class FLandscapeVertexFactoryMobileVertexShaderParameters : public FVertexFactoryShaderParameters
+{
+public:
+	/**
+	* Bind shader constants by name
+	* @param	ParameterMap - mapping of named shader constants to indices
+	*/
+	virtual void Bind(const FShaderParameterMap& ParameterMap) override
+	{
+		LodValuesParameter.Bind(ParameterMap,TEXT("LodValues"));
+		NeighborSectionLodParameter.Bind(ParameterMap,TEXT("NeighborSectionLod"));
+		LodBiasParameter.Bind(ParameterMap,TEXT("LodBias"));
+		SectionLodsParameter.Bind(ParameterMap,TEXT("SectionLods"));
+	}
+
+	/**
+	* Serialize shader params to an archive
+	* @param	Ar - archive to serialize to
+	*/
+	virtual void Serialize(FArchive& Ar) override
+	{
+		Ar << LodValuesParameter;
+		Ar << NeighborSectionLodParameter;
+		Ar << LodBiasParameter;
+		Ar << SectionLodsParameter;
+	}
+	/**
+	* Set any shader data specific to this vertex factory
+	*/
+	virtual void SetMesh(FRHICommandList& RHICmdList, FShader* VertexShader,const class FVertexFactory* VertexFactory,const class FSceneView& View,const struct FMeshBatchElement& BatchElement,uint32 DataFlags) const override
+	{
+		SCOPE_CYCLE_COUNTER(STAT_LandscapeVFDrawTime);
+
+		const FLandscapeBatchElementParams* BatchElementParams = (const FLandscapeBatchElementParams*)BatchElement.UserData;
+		check(BatchElementParams);
+
+		const FLandscapeComponentSceneProxyMobile* SceneProxy = (const FLandscapeComponentSceneProxyMobile*)BatchElementParams->SceneProxy;
+		SetUniformBufferParameter(RHICmdList, VertexShader->GetVertexShader(),VertexShader->GetUniformBufferParameter<FLandscapeUniformShaderParameters>(),*BatchElementParams->LandscapeUniformShaderParametersResource);
+
+		FVector CameraLocalPos3D = SceneProxy->WorldToLocal.TransformPosition(View.ViewMatrices.ViewOrigin); 
+		FVector2D CameraLocalPos = FVector2D(CameraLocalPos3D.X, CameraLocalPos3D.Y);
+
+		if( LodBiasParameter.IsBound() )
+		{
+			FVector4 LodBias(
+				SceneProxy->PrecomputedLODFactor,
+				1.f / (1.f - SceneProxy->PrecomputedLODFactor),
+				CameraLocalPos3D.X + SceneProxy->SectionBase.X,
+				CameraLocalPos3D.Y + SceneProxy->SectionBase.Y 
+				);
+			SetShaderValue(RHICmdList, VertexShader->GetVertexShader(), LodBiasParameter, LodBias);
+		}
+
+		// Calculate LOD params
+		FVector4 fCurrentLODs;
+		FVector4 CurrentNeighborLODs[4];
+		float DistLOD = 0.f;
+
+		if( BatchElementParams->SubX == -1 )
+		{
+			for( int32 SubY = 0; SubY < SceneProxy->NumSubsections; SubY++ )
+			{
+				for( int32 SubX = 0; SubX < SceneProxy->NumSubsections; SubX++ )
+				{
+					int32 SubIndex = SubX + 2 * SubY;
+					SceneProxy->CalcLODParamsForSubsection(View, CameraLocalPos, SubX, SubY, fCurrentLODs[SubIndex], CurrentNeighborLODs[SubIndex], DistLOD);
+				}
+			}
+		}
+		else
+		{
+			int32 SubIndex = BatchElementParams->SubX + 2 * BatchElementParams->SubY;
+			SceneProxy->CalcLODParamsForSubsection(View, CameraLocalPos, BatchElementParams->SubX, BatchElementParams->SubY, fCurrentLODs[SubIndex], CurrentNeighborLODs[SubIndex], DistLOD);
+		}
+
+		if( SectionLodsParameter.IsBound() )
+		{
+			SetShaderValue(RHICmdList, VertexShader->GetVertexShader(), SectionLodsParameter, fCurrentLODs);
+		}
+
+		if( NeighborSectionLodParameter.IsBound() )
+		{
+			SetShaderValue(RHICmdList, VertexShader->GetVertexShader(), NeighborSectionLodParameter, CurrentNeighborLODs);
+		}
+
+		if( LodValuesParameter.IsBound() )
+		{
+			FVector4 LodValues(
+				DistLOD,
+				// convert current LOD coordinates into highest LOD coordinates
+				(float)SceneProxy->SubsectionSizeQuads / (float)(((SceneProxy->SubsectionSizeVerts) >> BatchElementParams->CurrentLOD)-1),
+				(float)SceneProxy->SubsectionSizeQuads, 
+				1.f / (float)SceneProxy->SubsectionSizeQuads );
+
+			SetShaderValue(RHICmdList, VertexShader->GetVertexShader(),LodValuesParameter,LodValues);
+		}
+	}
+protected:
+	FShaderParameter LodValuesParameter;
+	FShaderParameter NeighborSectionLodParameter;
+	FShaderParameter LodBiasParameter;
+	FShaderParameter SectionLodsParameter;
+	TShaderUniformBufferParameter<FLandscapeUniformShaderParameters> LandscapeShaderParameters;
+};
+
+/** Shader parameters for use with FLandscapeVertexFactory */
+class FLandscapeVertexFactoryMobilePixelShaderParameters : public FLandscapeVertexFactoryPixelShaderParameters
+{
+public:
+	/**
+	* Bind shader constants by name
+	* @param	ParameterMap - mapping of named shader constants to indices
+	*/
+	virtual void Bind(const FShaderParameterMap& ParameterMap) override
+	{
+		FLandscapeVertexFactoryPixelShaderParameters::Bind(ParameterMap);
+		BlendableLayerMaskParameter.Bind(ParameterMap, TEXT("BlendableLayerMask"));
+	}
+
+	/**
+	* Serialize shader params to an archive
+	* @param	Ar - archive to serialize to
+	*/
+	virtual void Serialize(FArchive& Ar) override
+	{
+		FLandscapeVertexFactoryPixelShaderParameters::Serialize(Ar);
+		Ar << BlendableLayerMaskParameter;
+	}
+	/**
+	* Set any shader data specific to this vertex factory
+	*/
+	virtual void SetMesh(FRHICommandList& RHICmdList, FShader* PixelShader, const FVertexFactory* VertexFactory, const FSceneView& View, const FMeshBatchElement& BatchElement, uint32 DataFlags) const override
+	{
+		SCOPE_CYCLE_COUNTER(STAT_LandscapeVFDrawTime);
+
+		FLandscapeVertexFactoryPixelShaderParameters::SetMesh(RHICmdList, PixelShader, VertexFactory, View, BatchElement, DataFlags);
+
+		const FLandscapeBatchElementParams* BatchElementParams = (const FLandscapeBatchElementParams*)BatchElement.UserData;
+		check(BatchElementParams);
+		const FLandscapeComponentSceneProxyMobile* SceneProxy = (const FLandscapeComponentSceneProxyMobile*)BatchElementParams->SceneProxy;
+
+		if (BlendableLayerMaskParameter.IsBound())
+		{
+			FVector MaskVector;
+			MaskVector[0] = (SceneProxy->BlendableLayerMask & (1 << 0)) ? 1 : 0;
+			MaskVector[1] = (SceneProxy->BlendableLayerMask & (1 << 1)) ? 1 : 0;
+			MaskVector[2] = (SceneProxy->BlendableLayerMask & (1 << 2)) ? 1 : 0;
+			SetShaderValue(RHICmdList, PixelShader->GetPixelShader(), BlendableLayerMaskParameter, MaskVector);
+		}
+	}
+protected:
+	FShaderParameter BlendableLayerMaskParameter;
+};
+
+FVertexFactoryShaderParameters* FLandscapeVertexFactoryMobile::ConstructShaderParameters(EShaderFrequency ShaderFrequency)
+{
+	switch( ShaderFrequency )
+	{
+	case SF_Vertex:
+		return new FLandscapeVertexFactoryMobileVertexShaderParameters();
+	case SF_Pixel:
+		return new FLandscapeVertexFactoryMobilePixelShaderParameters();
+	default:
+		return NULL;
+	}
+}
+
+IMPLEMENT_VERTEX_FACTORY_TYPE(FLandscapeVertexFactoryMobile, "LandscapeVertexFactory", true, true, true, false, false);
+
+/** 
+* Initialize the RHI for this rendering resource 
+*/
+void FLandscapeVertexBufferMobile::InitRHI()
+{
+	// create a static vertex buffer
+	FRHIResourceCreateInfo CreateInfo;
+	VertexBufferRHI = RHICreateVertexBuffer(DataSize, BUF_Static, CreateInfo);
+	void* VertexData = RHILockVertexBuffer(VertexBufferRHI, 0, DataSize, RLM_WriteOnly);
+	// Copy stored platform data
+	FMemory::Memcpy(VertexData, (uint8*)Data, DataSize);
+	RHIUnlockVertexBuffer(VertexBufferRHI);
+}
+
+FLandscapeComponentSceneProxyMobile::FLandscapeComponentSceneProxyMobile(ULandscapeComponent* InComponent, FLandscapeEditToolRenderData* InEditToolRenderData)
+	:	FLandscapeComponentSceneProxy(InComponent, InEditToolRenderData)
+{
+		check(InComponent && InComponent->PlatformData.HasValidPlatformData());
+		InComponent->PlatformData.GetUncompressedData(PlatformData);
+#if WITH_EDITOR
+	if (InComponent->HeightmapTexture != NULL)
+	{
+		MaterialInterface = InComponent->GeneratePlatformPixelData(WeightmapTextures, false);
+		InComponent->MobileMaterialInterface = MaterialInterface;
+		InComponent->MobileNormalmapTexture = WeightmapTextures[0];
+	}
+#endif
+	NormalmapTexture = WeightmapTextures[0]; // Use Weightmap0 as Normal/Weightmap Texture
+
+	BlendableLayerMask = InComponent->MobileBlendableLayerMask;
+}
+
+FLandscapeComponentSceneProxyMobile::~FLandscapeComponentSceneProxyMobile()
+{
+	if (VertexFactory)
+	{
+		delete VertexFactory;
+		VertexFactory = NULL;
+	}
+}
+
+void FLandscapeComponentSceneProxyMobile::CreateRenderThreadResources()
+{
+	// Use only Index buffers
+	SharedBuffers = FLandscapeComponentSceneProxy::SharedBuffersMap.FindRef(SharedBuffersKey);
+	if( SharedBuffers == NULL )
+	{
+		SharedBuffers = new FLandscapeSharedBuffers(SharedBuffersKey, SubsectionSizeQuads, NumSubsections, GetScene().GetFeatureLevel(), false);
+		FLandscapeComponentSceneProxy::SharedBuffersMap.Add(SharedBuffersKey, SharedBuffers);
+	}
+
+	SharedBuffers->AddRef();
+
+	int32 VertexBufferSize = PlatformData.Num();
+	// Copy platform data into vertex buffer
+	VertexBuffer = new FLandscapeVertexBufferMobile(PlatformData.GetData(), VertexBufferSize);
+
+	FLandscapeVertexFactoryMobile* LandscapeVertexFactory = new FLandscapeVertexFactoryMobile();
+	LandscapeVertexFactory->MobileData.PositionComponent = FVertexStreamComponent(VertexBuffer, STRUCT_OFFSET(FLandscapeMobileVertex,Position), sizeof(FLandscapeMobileVertex), VET_UByte4N);
+	for( uint32 Index = 0; Index < LANDSCAPE_MAX_ES_LOD_COMP; ++Index )
+	{
+		LandscapeVertexFactory->MobileData.LODHeightsComponent.Add
+			(FVertexStreamComponent(VertexBuffer, STRUCT_OFFSET(FLandscapeMobileVertex,LODHeights) + sizeof(uint8) * 4 * Index, sizeof(FLandscapeMobileVertex), VET_UByte4N));
+	}
+
+	LandscapeVertexFactory->InitResource();
+	VertexFactory = LandscapeVertexFactory;
+	DynamicMesh.VertexFactory = VertexFactory;
+#if WITH_EDITOR
+	DynamicMeshTools.VertexFactory = VertexFactory;
+#endif
+
+	// Assign LandscapeUniformShaderParameters
+	LandscapeUniformShaderParameters.InitResource();
+
+	for( int32 ElementIdx=0;ElementIdx<DynamicMesh.Elements.Num();ElementIdx++ )
+	{
+		FMeshBatchElement& BatchElement = DynamicMesh.Elements[ElementIdx];
+		FLandscapeBatchElementParams* BatchElementParams = (FLandscapeBatchElementParams*)BatchElement.UserData;
+		BatchElementParams->LandscapeUniformShaderParametersResource = &LandscapeUniformShaderParameters;
+	}
+
+	PlatformData.Empty();
+}
+
diff --git a/Engine/Source/Runtime/Landscape/Private/LandscapeSplineProxies.cpp b/Engine/Source/Runtime/Landscape/Private/LandscapeSplineProxies.cpp
new file mode 100644
--- /dev/null
+++ b/Engine/Source/Runtime/Landscape/Private/LandscapeSplineProxies.cpp
@@ -0,0 +1,10 @@
+// Copyright 1998-2014 Epic Games, Inc. All Rights Reserved.
+
+#include "Landscape.h"
+#include "LandscapeSplineProxies.h"
+#include "LandscapeSplineSegment.h"
+
+LANDSCAPE_API void HLandscapeSplineProxy_Tangent::Serialize(FArchive& Ar)
+{
+	Ar << SplineSegment;
+}
\ No newline at end of file
diff --git a/Engine/Source/Runtime/Landscape/Private/LandscapeSplines.cpp b/Engine/Source/Runtime/Landscape/Private/LandscapeSplines.cpp
new file mode 100644
--- /dev/null
+++ b/Engine/Source/Runtime/Landscape/Private/LandscapeSplines.cpp
@@ -0,0 +1,1981 @@
+// Copyright 1998-2014 Epic Games, Inc. All Rights Reserved.
+
+/*=============================================================================
+	LandscapeSpline.cpp
+=============================================================================*/
+
+#include "Landscape.h"
+#include "Components/SplineMeshComponent.h"
+#include "ShaderParameterUtils.h"
+#include "StaticMeshResources.h"
+#include "LandscapeSplineProxies.h"
+#include "LandscapeSplinesComponent.h"
+#include "LandscapeSplineControlPoint.h"
+#include "LandscapeSplineSegment.h"
+#include "ControlPointMeshComponent.h"
+#include "Engine/StaticMesh.h"
+#include "Engine/StaticMeshSocket.h"
+
+IMPLEMENT_HIT_PROXY(HLandscapeSplineProxy, HHitProxy);
+IMPLEMENT_HIT_PROXY(HLandscapeSplineProxy_Segment, HLandscapeSplineProxy);
+IMPLEMENT_HIT_PROXY(HLandscapeSplineProxy_ControlPoint, HLandscapeSplineProxy);
+IMPLEMENT_HIT_PROXY(HLandscapeSplineProxy_Tangent, HLandscapeSplineProxy);
+
+//////////////////////////////////////////////////////////////////////////
+// LANDSCAPE SPLINES SCENE PROXY
+
+/** Represents a ULandscapeSplinesComponent to the scene manager. */
+#if WITH_EDITOR
+class FLandscapeSplinesSceneProxy : public FPrimitiveSceneProxy
+{
+private:
+	const FLinearColor	SplineColor;
+
+	const UTexture2D*	ControlPointSprite;
+	const bool			bDrawFalloff;
+
+	struct FSegmentProxy
+	{
+		ULandscapeSplineSegment* Owner;
+		TRefCountPtr<HHitProxy> HitProxy;
+		TArray<FLandscapeSplineInterpPoint> Points;
+		uint32 bSelected : 1;
+	};
+	TArray<FSegmentProxy> Segments;
+
+	struct FControlPointProxy
+	{
+		ULandscapeSplineControlPoint* Owner;
+		TRefCountPtr<HHitProxy> HitProxy;
+		FVector Location;
+		TArray<FLandscapeSplineInterpPoint> Points;
+		float SpriteScale;
+		uint32 bSelected : 1;
+	};
+	TArray<FControlPointProxy> ControlPoints;
+
+public:
+
+	~FLandscapeSplinesSceneProxy()
+	{
+	}
+
+	FLandscapeSplinesSceneProxy(ULandscapeSplinesComponent* Component):
+		FPrimitiveSceneProxy(Component),
+		SplineColor(Component->SplineColor),
+		ControlPointSprite(Component->ControlPointSprite),
+		bDrawFalloff(Component->bShowSplineEditorMesh)
+	{
+		Segments.Reserve(Component->Segments.Num());
+		for (ULandscapeSplineSegment* Segment : Component->Segments)
+		{
+			FSegmentProxy SegmentProxy;
+			SegmentProxy.Owner = Segment;
+			SegmentProxy.HitProxy = NULL;
+			SegmentProxy.Points = Segment->GetPoints();
+			SegmentProxy.bSelected = Segment->IsSplineSelected();
+			Segments.Add(SegmentProxy);
+		}
+
+		ControlPoints.Reserve(Component->ControlPoints.Num());
+		for (ULandscapeSplineControlPoint* ControlPoint : Component->ControlPoints)
+		{
+			FControlPointProxy ControlPointProxy;
+			ControlPointProxy.Owner = ControlPoint;
+			ControlPointProxy.HitProxy = NULL;
+			ControlPointProxy.Location = ControlPoint->Location;
+			ControlPointProxy.Points = ControlPoint->GetPoints();
+			ControlPointProxy.SpriteScale = FMath::Clamp<float>(ControlPoint->Width != 0 ? ControlPoint->Width / 2 : ControlPoint->SideFalloff / 4, 10, 1000);
+			ControlPointProxy.bSelected = ControlPoint->IsSplineSelected();
+			ControlPoints.Add(ControlPointProxy);
+		}
+	}
+
+	virtual HHitProxy* CreateHitProxies(UPrimitiveComponent* Component, TArray<TRefCountPtr<HHitProxy> >& OutHitProxies)
+	{
+		OutHitProxies.Reserve(OutHitProxies.Num() + Segments.Num() + ControlPoints.Num());
+		for (FSegmentProxy& Segment : Segments)
+		{
+			Segment.HitProxy = new HLandscapeSplineProxy_Segment(Segment.Owner);
+			OutHitProxies.Add(Segment.HitProxy);
+		}
+		for (FControlPointProxy& ControlPoint : ControlPoints)
+		{
+			ControlPoint.HitProxy = new HLandscapeSplineProxy_ControlPoint(ControlPoint.Owner);
+			OutHitProxies.Add(ControlPoint.HitProxy);
+		}
+		return NULL;
+	}
+
+	virtual void GetDynamicMeshElements(const TArray<const FSceneView*>& Views, const FSceneViewFamily& ViewFamily, uint32 VisibilityMap, FMeshElementCollector& Collector) const override
+	{
+		// Slight Depth Bias so that the splines show up when they exactly match the target surface
+		// e.g. someone playing with splines on a newly-created perfectly-flat landscape
+		static const float DepthBias = -0.0001;
+
+		const FMatrix& LocalToWorld = GetLocalToWorld();
+
+		const FLinearColor SelectedSplineColor = GEngine->GetSelectedMaterialColor();
+		const FLinearColor SelectedControlPointSpriteColor = FLinearColor::White + (GEngine->GetSelectedMaterialColor() * GEngine->SelectionHighlightIntensityBillboards * 10); // copied from FSpriteSceneProxy::DrawDynamicElements()
+
+		for (int32 ViewIndex = 0; ViewIndex < Views.Num(); ViewIndex++)
+		{
+			if (VisibilityMap & (1 << ViewIndex))
+			{
+				const FSceneView* View = Views[ViewIndex];
+				FPrimitiveDrawInterface* PDI = Collector.GetPDI(ViewIndex);
+
+				for (const FSegmentProxy& Segment : Segments)
+				{
+					const FLinearColor SegmentColor = Segment.bSelected ? SelectedSplineColor : SplineColor;
+
+					FLandscapeSplineInterpPoint OldPoint = Segment.Points[0];
+					OldPoint.Center       = LocalToWorld.TransformPosition(OldPoint.Center);
+					OldPoint.Left         = LocalToWorld.TransformPosition(OldPoint.Left);
+					OldPoint.Right        = LocalToWorld.TransformPosition(OldPoint.Right);
+					OldPoint.FalloffLeft  = LocalToWorld.TransformPosition(OldPoint.FalloffLeft);
+					OldPoint.FalloffRight = LocalToWorld.TransformPosition(OldPoint.FalloffRight);
+					for (int32 i = 1; i < Segment.Points.Num(); i++)
+					{
+						FLandscapeSplineInterpPoint NewPoint = Segment.Points[i];
+						NewPoint.Center       = LocalToWorld.TransformPosition(NewPoint.Center);
+						NewPoint.Left         = LocalToWorld.TransformPosition(NewPoint.Left);
+						NewPoint.Right        = LocalToWorld.TransformPosition(NewPoint.Right);
+						NewPoint.FalloffLeft  = LocalToWorld.TransformPosition(NewPoint.FalloffLeft);
+						NewPoint.FalloffRight = LocalToWorld.TransformPosition(NewPoint.FalloffRight);
+
+						// Draw lines from the last keypoint.
+						PDI->SetHitProxy(Segment.HitProxy);
+
+						// center line
+						PDI->DrawLine(OldPoint.Center, NewPoint.Center, SegmentColor, GetDepthPriorityGroup(View), 0.0f, DepthBias);
+
+						// draw sides
+						PDI->DrawLine(OldPoint.Left, NewPoint.Left, SegmentColor, GetDepthPriorityGroup(View), 0.0f, DepthBias);
+						PDI->DrawLine(OldPoint.Right, NewPoint.Right, SegmentColor, GetDepthPriorityGroup(View), 0.0f, DepthBias);
+
+						PDI->SetHitProxy(NULL);
+
+						// draw falloff sides
+						if (bDrawFalloff)
+						{
+							DrawDashedLine(PDI, OldPoint.FalloffLeft, NewPoint.FalloffLeft, SegmentColor, 100, GetDepthPriorityGroup(View), DepthBias);
+							DrawDashedLine(PDI, OldPoint.FalloffRight, NewPoint.FalloffRight, SegmentColor, 100, GetDepthPriorityGroup(View), DepthBias);
+						}
+
+						OldPoint = NewPoint;
+					}
+				}
+
+				for (const FControlPointProxy& ControlPoint : ControlPoints)
+				{
+					const float ControlPointSpriteScale = LocalToWorld.GetScaleVector().X * ControlPoint.SpriteScale;
+					const FVector ControlPointLocation = LocalToWorld.TransformPosition(ControlPoint.Location) + FVector(0, 0, ControlPointSpriteScale * 0.75f);
+
+					// Draw Sprite
+
+					const FLinearColor ControlPointSpriteColor = ControlPoint.bSelected ? SelectedControlPointSpriteColor : FLinearColor::White;
+
+					PDI->SetHitProxy(ControlPoint.HitProxy);
+
+					PDI->DrawSprite(
+						ControlPointLocation,
+						ControlPointSpriteScale,
+						ControlPointSpriteScale,
+						ControlPointSprite->Resource,
+						ControlPointSpriteColor,
+						GetDepthPriorityGroup(View),
+						0, ControlPointSprite->Resource->GetSizeX(),
+						0, ControlPointSprite->Resource->GetSizeY(),
+						SE_BLEND_Masked);
+
+
+					// Draw Lines
+					const FLinearColor ControlPointColor = ControlPoint.bSelected ? SelectedSplineColor : SplineColor;
+
+					if (ControlPoint.Points.Num() == 1)
+					{
+						FLandscapeSplineInterpPoint NewPoint = ControlPoint.Points[0];
+						NewPoint.Center = LocalToWorld.TransformPosition(NewPoint.Center);
+						NewPoint.Left   = LocalToWorld.TransformPosition(NewPoint.Left);
+						NewPoint.Right  = LocalToWorld.TransformPosition(NewPoint.Right);
+						NewPoint.FalloffLeft  = LocalToWorld.TransformPosition(NewPoint.FalloffLeft);
+						NewPoint.FalloffRight = LocalToWorld.TransformPosition(NewPoint.FalloffRight);
+
+						// draw end for spline connection
+						PDI->DrawPoint(NewPoint.Center, ControlPointColor, 6.0f, GetDepthPriorityGroup(View));
+						PDI->DrawLine(NewPoint.Left, NewPoint.Center, ControlPointColor, GetDepthPriorityGroup(View), 0.0f, DepthBias);
+						PDI->DrawLine(NewPoint.Right, NewPoint.Center, ControlPointColor, GetDepthPriorityGroup(View), 0.0f, DepthBias);
+						if (bDrawFalloff)
+						{
+							DrawDashedLine(PDI, NewPoint.FalloffLeft, NewPoint.Left, ControlPointColor, 100, GetDepthPriorityGroup(View), DepthBias);
+							DrawDashedLine(PDI, NewPoint.FalloffRight, NewPoint.Right, ControlPointColor, 100, GetDepthPriorityGroup(View), DepthBias);
+						}
+					}
+					else if (ControlPoint.Points.Num() >= 2)
+					{
+						FLandscapeSplineInterpPoint OldPoint = ControlPoint.Points.Last();
+						//OldPoint.Left   = LocalToWorld.TransformPosition(OldPoint.Left);
+						OldPoint.Right  = LocalToWorld.TransformPosition(OldPoint.Right);
+						//OldPoint.FalloffLeft  = LocalToWorld.TransformPosition(OldPoint.FalloffLeft);
+						OldPoint.FalloffRight = LocalToWorld.TransformPosition(OldPoint.FalloffRight);
+
+						for (const FLandscapeSplineInterpPoint& Point : ControlPoint.Points)
+						{
+							FLandscapeSplineInterpPoint NewPoint = Point;
+							NewPoint.Center = LocalToWorld.TransformPosition(NewPoint.Center);
+							NewPoint.Left   = LocalToWorld.TransformPosition(NewPoint.Left);
+							NewPoint.Right  = LocalToWorld.TransformPosition(NewPoint.Right);
+							NewPoint.FalloffLeft  = LocalToWorld.TransformPosition(NewPoint.FalloffLeft);
+							NewPoint.FalloffRight = LocalToWorld.TransformPosition(NewPoint.FalloffRight);
+
+							PDI->SetHitProxy(ControlPoint.HitProxy);
+
+							// center line
+							PDI->DrawLine(ControlPointLocation, NewPoint.Center, ControlPointColor, GetDepthPriorityGroup(View), 0.0f, DepthBias);
+
+							// draw sides
+							PDI->DrawLine(OldPoint.Right, NewPoint.Left, ControlPointColor, GetDepthPriorityGroup(View), 0.0f, DepthBias);
+
+							PDI->SetHitProxy(NULL);
+
+							// draw falloff sides
+							if (bDrawFalloff)
+							{
+								DrawDashedLine(PDI, OldPoint.FalloffRight, NewPoint.FalloffLeft, ControlPointColor, 100, GetDepthPriorityGroup(View), DepthBias);
+							}
+
+							// draw end for spline connection
+							PDI->DrawPoint(NewPoint.Center, ControlPointColor, 6.0f, GetDepthPriorityGroup(View));
+							PDI->DrawLine(NewPoint.Left, NewPoint.Center, ControlPointColor, GetDepthPriorityGroup(View), 0.0f, DepthBias);
+							PDI->DrawLine(NewPoint.Right, NewPoint.Center, ControlPointColor, GetDepthPriorityGroup(View), 0.0f, DepthBias);
+							if (bDrawFalloff)
+							{
+								DrawDashedLine(PDI, NewPoint.FalloffLeft, NewPoint.Left, ControlPointColor, 100, GetDepthPriorityGroup(View), DepthBias);
+								DrawDashedLine(PDI, NewPoint.FalloffRight, NewPoint.Right, ControlPointColor, 100, GetDepthPriorityGroup(View), DepthBias);
+							}
+
+							//OldPoint = NewPoint;
+							OldPoint.Right = NewPoint.Right;
+							OldPoint.FalloffRight = NewPoint.FalloffRight;
+						}
+					}
+				}
+
+				PDI->SetHitProxy(NULL);
+			}
+		}
+	}
+
+	virtual void DrawDynamicElements(FPrimitiveDrawInterface* PDI,const FSceneView* View)
+	{
+		// Slight Depth Bias so that the splines show up when they exactly match the target surface
+		// e.g. someone playing with splines on a newly-created perfectly-flat landscape
+		static const float DepthBias = -0.0001;
+
+		const FMatrix& LocalToWorld = GetLocalToWorld();
+
+		const FLinearColor SelectedSplineColor = GEngine->GetSelectedMaterialColor();
+		const FLinearColor SelectedControlPointSpriteColor = FLinearColor::White + (GEngine->GetSelectedMaterialColor() * GEngine->SelectionHighlightIntensityBillboards * 10); // copied from FSpriteSceneProxy::DrawDynamicElements()
+
+		for (const FSegmentProxy& Segment : Segments)
+		{
+			const FLinearColor SegmentColor = Segment.bSelected ? SelectedSplineColor : SplineColor;
+
+			FLandscapeSplineInterpPoint OldPoint = Segment.Points[0];
+			OldPoint.Center       = LocalToWorld.TransformPosition(OldPoint.Center);
+			OldPoint.Left         = LocalToWorld.TransformPosition(OldPoint.Left);
+			OldPoint.Right        = LocalToWorld.TransformPosition(OldPoint.Right);
+			OldPoint.FalloffLeft  = LocalToWorld.TransformPosition(OldPoint.FalloffLeft);
+			OldPoint.FalloffRight = LocalToWorld.TransformPosition(OldPoint.FalloffRight);
+			for (int32 i = 1; i < Segment.Points.Num(); i++)
+			{
+				FLandscapeSplineInterpPoint NewPoint = Segment.Points[i];
+				NewPoint.Center       = LocalToWorld.TransformPosition(NewPoint.Center);
+				NewPoint.Left         = LocalToWorld.TransformPosition(NewPoint.Left);
+				NewPoint.Right        = LocalToWorld.TransformPosition(NewPoint.Right);
+				NewPoint.FalloffLeft  = LocalToWorld.TransformPosition(NewPoint.FalloffLeft);
+				NewPoint.FalloffRight = LocalToWorld.TransformPosition(NewPoint.FalloffRight);
+
+				// Draw lines from the last keypoint.
+				if (PDI->IsHitTesting()) PDI->SetHitProxy(Segment.HitProxy);
+
+				// center line
+				PDI->DrawLine(OldPoint.Center, NewPoint.Center, SegmentColor, GetDepthPriorityGroup(View), 0.0f, DepthBias);
+
+				// draw sides
+				PDI->DrawLine(OldPoint.Left, NewPoint.Left, SegmentColor, GetDepthPriorityGroup(View), 0.0f, DepthBias);
+				PDI->DrawLine(OldPoint.Right, NewPoint.Right, SegmentColor, GetDepthPriorityGroup(View), 0.0f, DepthBias);
+
+				if (PDI->IsHitTesting()) PDI->SetHitProxy(NULL);
+
+				// draw falloff sides
+				if (bDrawFalloff)
+				{
+					DrawDashedLine(PDI, OldPoint.FalloffLeft, NewPoint.FalloffLeft, SegmentColor, 100, GetDepthPriorityGroup(View), DepthBias);
+					DrawDashedLine(PDI, OldPoint.FalloffRight, NewPoint.FalloffRight, SegmentColor, 100, GetDepthPriorityGroup(View), DepthBias);
+				}
+
+				OldPoint = NewPoint;
+			}
+		}
+
+		for (const FControlPointProxy& ControlPoint : ControlPoints)
+		{
+			const float ControlPointSpriteScale = LocalToWorld.GetScaleVector().X * ControlPoint.SpriteScale;
+			const FVector ControlPointLocation = LocalToWorld.TransformPosition(ControlPoint.Location) + FVector(0, 0, ControlPointSpriteScale * 0.75f);
+
+			// Draw Sprite
+
+			const FLinearColor ControlPointSpriteColor = ControlPoint.bSelected ? SelectedControlPointSpriteColor : FLinearColor::White;
+
+			if (PDI->IsHitTesting()) PDI->SetHitProxy(ControlPoint.HitProxy);
+
+			PDI->DrawSprite(
+				ControlPointLocation,
+				ControlPointSpriteScale,
+				ControlPointSpriteScale,
+				ControlPointSprite->Resource,
+				ControlPointSpriteColor,
+				GetDepthPriorityGroup(View),
+				0, ControlPointSprite->Resource->GetSizeX(),
+				0, ControlPointSprite->Resource->GetSizeY(),
+				SE_BLEND_Masked);
+
+
+			// Draw Lines
+			const FLinearColor ControlPointColor = ControlPoint.bSelected ? SelectedSplineColor : SplineColor;
+
+			if (ControlPoint.Points.Num() == 1)
+			{
+				FLandscapeSplineInterpPoint NewPoint = ControlPoint.Points[0];
+				NewPoint.Center = LocalToWorld.TransformPosition(NewPoint.Center);
+				NewPoint.Left   = LocalToWorld.TransformPosition(NewPoint.Left);
+				NewPoint.Right  = LocalToWorld.TransformPosition(NewPoint.Right);
+				NewPoint.FalloffLeft  = LocalToWorld.TransformPosition(NewPoint.FalloffLeft);
+				NewPoint.FalloffRight = LocalToWorld.TransformPosition(NewPoint.FalloffRight);
+
+				// draw end for spline connection
+				PDI->DrawPoint(NewPoint.Center, ControlPointColor, 6.0f, GetDepthPriorityGroup(View));
+				PDI->DrawLine(NewPoint.Left, NewPoint.Center, ControlPointColor, GetDepthPriorityGroup(View), 0.0f, DepthBias);
+				PDI->DrawLine(NewPoint.Right, NewPoint.Center, ControlPointColor, GetDepthPriorityGroup(View), 0.0f, DepthBias);
+				if (bDrawFalloff)
+				{
+					DrawDashedLine(PDI, NewPoint.FalloffLeft, NewPoint.Left, ControlPointColor, 100, GetDepthPriorityGroup(View), DepthBias);
+					DrawDashedLine(PDI, NewPoint.FalloffRight, NewPoint.Right, ControlPointColor, 100, GetDepthPriorityGroup(View), DepthBias);
+				}
+			}
+			else if (ControlPoint.Points.Num() >= 2)
+			{
+				FLandscapeSplineInterpPoint OldPoint = ControlPoint.Points.Last();
+				//OldPoint.Left   = LocalToWorld.TransformPosition(OldPoint.Left);
+				OldPoint.Right  = LocalToWorld.TransformPosition(OldPoint.Right);
+				//OldPoint.FalloffLeft  = LocalToWorld.TransformPosition(OldPoint.FalloffLeft);
+				OldPoint.FalloffRight = LocalToWorld.TransformPosition(OldPoint.FalloffRight);
+
+				for (const FLandscapeSplineInterpPoint& Point : ControlPoint.Points)
+				{
+					FLandscapeSplineInterpPoint NewPoint = Point;
+					NewPoint.Center = LocalToWorld.TransformPosition(NewPoint.Center);
+					NewPoint.Left   = LocalToWorld.TransformPosition(NewPoint.Left);
+					NewPoint.Right  = LocalToWorld.TransformPosition(NewPoint.Right);
+					NewPoint.FalloffLeft  = LocalToWorld.TransformPosition(NewPoint.FalloffLeft);
+					NewPoint.FalloffRight = LocalToWorld.TransformPosition(NewPoint.FalloffRight);
+
+					if (PDI->IsHitTesting()) PDI->SetHitProxy(ControlPoint.HitProxy);
+
+					// center line
+					PDI->DrawLine(ControlPointLocation, NewPoint.Center, ControlPointColor, GetDepthPriorityGroup(View), 0.0f, DepthBias);
+
+					// draw sides
+					PDI->DrawLine(OldPoint.Right, NewPoint.Left, ControlPointColor, GetDepthPriorityGroup(View), 0.0f, DepthBias);
+
+					if (PDI->IsHitTesting()) PDI->SetHitProxy(NULL);
+
+					// draw falloff sides
+					if (bDrawFalloff)
+					{
+						DrawDashedLine(PDI, OldPoint.FalloffRight, NewPoint.FalloffLeft, ControlPointColor, 100, GetDepthPriorityGroup(View), DepthBias);
+					}
+
+					// draw end for spline connection
+					PDI->DrawPoint(NewPoint.Center, ControlPointColor, 6.0f, GetDepthPriorityGroup(View));
+					PDI->DrawLine(NewPoint.Left, NewPoint.Center, ControlPointColor, GetDepthPriorityGroup(View), 0.0f, DepthBias);
+					PDI->DrawLine(NewPoint.Right, NewPoint.Center, ControlPointColor, GetDepthPriorityGroup(View), 0.0f, DepthBias);
+					if (bDrawFalloff)
+					{
+						DrawDashedLine(PDI, NewPoint.FalloffLeft, NewPoint.Left, ControlPointColor, 100, GetDepthPriorityGroup(View), DepthBias);
+						DrawDashedLine(PDI, NewPoint.FalloffRight, NewPoint.Right, ControlPointColor, 100, GetDepthPriorityGroup(View), DepthBias);
+					}
+
+					//OldPoint = NewPoint;
+					OldPoint.Right = NewPoint.Right;
+					OldPoint.FalloffRight = NewPoint.FalloffRight;
+				}
+			}
+		}
+
+		if (PDI->IsHitTesting()) PDI->SetHitProxy(NULL);
+	}
+
+	virtual FPrimitiveViewRelevance GetViewRelevance(const FSceneView* View)
+	{
+		FPrimitiveViewRelevance Result;
+		Result.bDrawRelevance = IsShown(View) && View->Family->EngineShowFlags.Splines;
+		Result.bDynamicRelevance = true;
+		return Result;
+	}
+
+	virtual uint32 GetMemoryFootprint() const
+	{
+		return sizeof(*this) + GetAllocatedSize();
+	}
+	uint32 GetAllocatedSize() const
+	{
+		uint32 AllocatedSize = FPrimitiveSceneProxy::GetAllocatedSize() + Segments.GetAllocatedSize() + ControlPoints.GetAllocatedSize();
+		for (const FSegmentProxy& Segment : Segments)
+		{
+			AllocatedSize += Segment.Points.GetAllocatedSize();
+		}
+		for (const FControlPointProxy& ControlPoint : ControlPoints)
+		{
+			AllocatedSize += ControlPoint.Points.GetAllocatedSize();
+		}
+		return AllocatedSize;
+	}
+};
+#endif
+
+//////////////////////////////////////////////////////////////////////////
+// SPLINE COMPONENT
+
+ULandscapeSplinesComponent::ULandscapeSplinesComponent(const FObjectInitializer& ObjectInitializer)
+	: Super(ObjectInitializer)
+{
+	Mobility = EComponentMobility::Static;
+
+#if WITH_EDITORONLY_DATA
+	SplineResolution = 512;
+	SplineColor = FColor(0, 192, 48);
+
+	if (!IsRunningCommandlet())
+	{
+		// Structure to hold one-time initialization
+		struct FConstructorStatics
+		{
+			ConstructorHelpers::FObjectFinder<UTexture2D> SpriteTexture;
+			ConstructorHelpers::FObjectFinder<UStaticMesh> SplineEditorMesh;
+			FConstructorStatics()
+				: SpriteTexture(TEXT("/Engine/EditorResources/S_Terrain.S_Terrain"))
+				, SplineEditorMesh(TEXT("/Engine/EditorLandscapeResources/SplineEditorMesh"))
+			{
+			}
+		};
+		static FConstructorStatics ConstructorStatics;
+
+		ControlPointSprite = ConstructorStatics.SpriteTexture.Object;
+		SplineEditorMesh = ConstructorStatics.SplineEditorMesh.Object;
+	}
+#endif
+	//RelativeScale3D = FVector(1/100.0f, 1/100.0f, 1/100.0f); // cancel out landscape scale. The scale is set up when component is created, but for a default landscape it's this
+}
+
+void ULandscapeSplinesComponent::CheckSplinesValid()
+{
+#if DO_CHECK
+	// This shouldn't happen, but it has somehow (TTP #334549) so we have to fix it
+	ensure(!ControlPoints.Contains(nullptr));
+	ensure(!Segments.Contains(nullptr));
+
+	// Remove all null control points/segments
+	ControlPoints.Remove(nullptr);
+	Segments.Remove(nullptr);
+
+	// Check for cross-spline connections, as this is a potential source of nulls
+	for (ULandscapeSplineControlPoint* ControlPoint : ControlPoints)
+	{
+		ensure(ControlPoint->GetOuterULandscapeSplinesComponent() == this);
+		for (const FLandscapeSplineConnection& Connection : ControlPoint->ConnectedSegments)
+		{
+			ensure(Connection.Segment->GetOuterULandscapeSplinesComponent() == this);
+		}
+	}
+	for (ULandscapeSplineSegment* Segment : Segments)
+	{
+		ensure(Segment->GetOuterULandscapeSplinesComponent() == this);
+		for (const FLandscapeSplineSegmentConnection& Connection : Segment->Connections)
+		{
+			ensure(Connection.ControlPoint->GetOuterULandscapeSplinesComponent() == this);
+		}
+	}
+#endif
+}
+
+void ULandscapeSplinesComponent::OnRegister()
+{
+	CheckSplinesValid();
+
+	Super::OnRegister();
+}
+
+#if WITH_EDITOR
+FPrimitiveSceneProxy* ULandscapeSplinesComponent::CreateSceneProxy()
+{
+	CheckSplinesValid();
+
+	return new FLandscapeSplinesSceneProxy(this);
+}
+#endif
+
+FBoxSphereBounds ULandscapeSplinesComponent::CalcBounds(const FTransform& LocalToWorld) const
+{
+	FBox NewBoundsCalc(0);
+
+	for (ULandscapeSplineControlPoint* ControlPoint : ControlPoints)
+	{
+		// TTP #334549: Somehow we're getting nulls in the ControlPoints array
+		if (ControlPoint)
+		{
+			NewBoundsCalc += ControlPoint->GetBounds();
+		}
+	}
+
+	for (ULandscapeSplineSegment* Segment : Segments)
+	{
+		if (Segment)
+		{
+			NewBoundsCalc += Segment->GetBounds();
+		}
+	}
+
+	NewBoundsCalc = NewBoundsCalc.TransformBy(LocalToWorld);
+	return FBoxSphereBounds(NewBoundsCalc);
+}
+
+bool ULandscapeSplinesComponent::ModifySplines(bool bAlwaysMarkDirty /*= true*/)
+{
+	bool bSavedToTransactionBuffer = Modify(bAlwaysMarkDirty);
+
+	for (ULandscapeSplineControlPoint* ControlPoint : ControlPoints)
+	{
+		bSavedToTransactionBuffer = ControlPoint->Modify(bAlwaysMarkDirty) || bSavedToTransactionBuffer;
+	}
+	for (ULandscapeSplineSegment* Segment : Segments)
+	{
+		bSavedToTransactionBuffer = Segment->Modify(bAlwaysMarkDirty) || bSavedToTransactionBuffer;
+	}
+
+	return bSavedToTransactionBuffer;
+}
+
+void ULandscapeSplinesComponent::PostLoad()
+{
+	Super::PostLoad();
+
+	CheckSplinesValid();
+}
+
+#if WITH_EDITOR
+static bool bHackIsUndoingSplines = false;
+void ULandscapeSplinesComponent::PostEditChangeProperty(FPropertyChangedEvent& PropertyChangedEvent)
+{
+	Super::PostEditChangeProperty(PropertyChangedEvent);
+
+	// Don't update splines when undoing, not only is it unnecessary and expensive,
+	// it also causes failed asserts in debug builds when trying to register components
+	// (because the actor hasn't reset its OwnedComponents array yet)
+	if (!bHackIsUndoingSplines)
+	{
+		const bool bUpdateCollision = PropertyChangedEvent.ChangeType != EPropertyChangeType::Interactive;
+
+		for (ULandscapeSplineControlPoint* ControlPoint : ControlPoints)
+		{
+			ControlPoint->UpdateSplinePoints(bUpdateCollision);
+		}
+		for (ULandscapeSplineSegment* Segment : Segments)
+		{
+			Segment->UpdateSplinePoints(bUpdateCollision);
+		}
+	}
+}
+void ULandscapeSplinesComponent::PostEditUndo()
+{
+	bHackIsUndoingSplines = true;
+	Super::PostEditUndo();
+	bHackIsUndoingSplines = false;
+
+	MarkRenderStateDirty();
+}
+
+void ULandscapeSplinesComponent::ShowSplineEditorMesh(bool bShow)
+{
+	bShowSplineEditorMesh = bShow;
+
+	for (ULandscapeSplineSegment* Segment : Segments)
+	{
+		Segment->UpdateSplineEditorMesh();
+	}
+
+	MarkRenderStateDirty();
+}
+#endif // WITH_EDITOR
+
+//////////////////////////////////////////////////////////////////////////
+// CONTROL POINT MESH COMPONENT
+
+UControlPointMeshComponent::UControlPointMeshComponent(const FObjectInitializer& ObjectInitializer)
+	: Super(ObjectInitializer)
+{
+	SetCollisionProfileName(UCollisionProfile::BlockAll_ProfileName);
+	Mobility = EComponentMobility::Static;
+
+#if WITH_EDITORONLY_DATA
+	bSelected = false;
+#endif
+}
+
+//////////////////////////////////////////////////////////////////////////
+// SPLINE CONTROL POINT
+
+ULandscapeSplineControlPoint::ULandscapeSplineControlPoint(const FObjectInitializer& ObjectInitializer)
+	: Super(ObjectInitializer)
+{
+	Width = 1000;
+	SideFalloff = 1000;
+	EndFalloff = 2000;
+
+#if WITH_EDITORONLY_DATA
+	Mesh = NULL;
+	MeshScale = FVector(1);
+
+	LDMaxDrawDistance = 0;
+
+	LayerName = NAME_None;
+	bRaiseTerrain = true;
+	bLowerTerrain = true;
+
+	MeshComponent = NULL;
+	bEnableCollision = true;
+	bCastShadow = true;
+
+	// transients
+	bSelected = false;
+#endif
+}
+
+void ULandscapeSplineControlPoint::PostLoad()
+{
+	Super::PostLoad();
+
+#if WITH_EDITOR
+	if (MeshComponent != NULL)
+	{
+		// Fix collision profile
+		const FName CollisionProfile = bEnableCollision ? UCollisionProfile::BlockAll_ProfileName : UCollisionProfile::NoCollision_ProfileName;
+		if (MeshComponent->GetCollisionProfileName() != CollisionProfile)
+		{
+			MeshComponent->SetCollisionProfileName(CollisionProfile);
+		}
+
+		MeshComponent->SetFlags(RF_TextExportTransient);
+	}
+#endif
+}
+
+FLandscapeSplineSegmentConnection& FLandscapeSplineConnection::GetNearConnection() const
+{
+	return Segment->Connections[End];
+}
+
+FLandscapeSplineSegmentConnection& FLandscapeSplineConnection::GetFarConnection() const
+{
+	return Segment->Connections[1 - End];
+}
+
+FName ULandscapeSplineControlPoint::GetBestConnectionTo(FVector Destination) const
+{
+	FName BestSocket = NAME_None;
+	float BestScore = -FLT_MAX;
+
+#if WITH_EDITORONLY_DATA
+	if (Mesh != NULL)
+	{
+#else
+	if (MeshComponent != NULL)
+	{
+		const UStaticMesh* Mesh = MeshComponent->StaticMesh;
+		const FVector MeshScale = MeshComponent->RelativeScale3D;
+#endif
+		for (const UStaticMeshSocket* Socket : Mesh->Sockets)
+		{
+			FTransform SocketTransform = FTransform(Socket->RelativeRotation, Socket->RelativeLocation) * FTransform(Rotation, Location, MeshScale);
+			FVector SocketLocation = SocketTransform.GetTranslation();
+			FRotator SocketRotation = SocketTransform.GetRotation().Rotator();
+
+			float Score = (Destination - Location).Size() - (Destination - SocketLocation).Size(); // Score closer sockets higher
+			Score *= FMath::Abs(FVector::DotProduct((Destination - SocketLocation), SocketRotation.Vector())); // score closer rotation higher
+
+			if (Score > BestScore)
+			{
+				BestSocket = Socket->SocketName;
+				BestScore = Score;
+			}
+		}
+	}
+
+	return BestSocket;
+}
+
+void ULandscapeSplineControlPoint::GetConnectionLocalLocationAndRotation(FName SocketName, OUT FVector& OutLocation, OUT FRotator& OutRotation) const
+{
+	OutLocation = FVector::ZeroVector;
+	OutRotation = FRotator::ZeroRotator;
+
+#if WITH_EDITORONLY_DATA
+	if (Mesh != NULL)
+	{
+#else
+	if (MeshComponent != NULL)
+	{
+		UStaticMesh* Mesh = MeshComponent->StaticMesh;
+#endif
+		const UStaticMeshSocket* Socket = Mesh->FindSocket(SocketName);
+		if (Socket != NULL)
+		{
+			OutLocation = Socket->RelativeLocation;
+			OutRotation = Socket->RelativeRotation;
+		}
+	}
+}
+
+void ULandscapeSplineControlPoint::GetConnectionLocationAndRotation(FName SocketName, OUT FVector& OutLocation, OUT FRotator& OutRotation) const
+{
+	OutLocation = Location;
+	OutRotation = Rotation;
+
+#if WITH_EDITORONLY_DATA
+	if (Mesh != NULL)
+	{
+#else
+	if (MeshComponent != NULL)
+	{
+		UStaticMesh* Mesh = MeshComponent->StaticMesh;
+		const FVector MeshScale = MeshComponent->RelativeScale3D;
+#endif
+		const UStaticMeshSocket* Socket = Mesh->FindSocket(SocketName);
+		if (Socket != NULL)
+		{
+			FTransform SocketTransform = FTransform(Socket->RelativeRotation, Socket->RelativeLocation) * FTransform(Rotation, Location, MeshScale);
+			OutLocation = SocketTransform.GetTranslation();
+			OutRotation = SocketTransform.GetRotation().Rotator().GetNormalized();
+		}
+	}
+}
+
+#if WITH_EDITOR
+void ULandscapeSplineControlPoint::SetSplineSelected(bool bInSelected)
+{
+	bSelected = bInSelected;
+	GetOuterULandscapeSplinesComponent()->MarkRenderStateDirty();
+
+	if (MeshComponent != NULL)
+	{
+		MeshComponent->bSelected = bInSelected;
+		MeshComponent->PushSelectionToProxy();
+	}
+}
+
+void ULandscapeSplineControlPoint::AutoCalcRotation()
+{
+	Modify();
+
+	FRotator Delta = FRotator::ZeroRotator;
+
+	for (const FLandscapeSplineConnection& Connection : ConnectedSegments)
+	{
+		// Get the start and end location/rotation of this connection
+		FVector StartLocation; FRotator StartRotation;
+		this->GetConnectionLocationAndRotation(Connection.GetNearConnection().SocketName, StartLocation, StartRotation);
+		FVector StartLocalLocation; FRotator StartLocalRotation;
+		this->GetConnectionLocalLocationAndRotation(Connection.GetNearConnection().SocketName, StartLocalLocation, StartLocalRotation);
+		FVector EndLocation; FRotator EndRotation;
+		Connection.GetFarConnection().ControlPoint->GetConnectionLocationAndRotation(Connection.GetFarConnection().SocketName, EndLocation, EndRotation);
+
+		// Find the delta between the direction of the tangent at the connection point and
+		// the direction to the other end's control point
+		FQuat SocketLocalRotation = StartLocalRotation.Quaternion();
+		if (FMath::Sign(Connection.GetNearConnection().TangentLen) < 0)
+		{
+			SocketLocalRotation = SocketLocalRotation * FRotator(0, 180, 0).Quaternion();
+		}
+		const FVector  DesiredDirection = (EndLocation - StartLocation);
+		const FQuat    DesiredSocketRotation = DesiredDirection.Rotation().Quaternion();
+		const FRotator DesiredRotation = (DesiredSocketRotation * SocketLocalRotation.Inverse()).Rotator().GetNormalized();
+		const FRotator DesiredRotationDelta = (DesiredRotation - Rotation).GetNormalized();
+
+		Delta += DesiredRotationDelta;
+	}
+
+	// Average delta of all connections
+	Delta *= 1.0f / ConnectedSegments.Num();
+
+	// Apply Delta and normalize
+	Rotation = (Rotation + Delta).GetNormalized();
+}
+
+void ULandscapeSplineControlPoint::AutoFlipTangents()
+{
+	for (const FLandscapeSplineConnection& Connection : ConnectedSegments)
+	{
+		Connection.Segment->AutoFlipTangents();
+	}
+}
+
+void ULandscapeSplineControlPoint::AutoSetConnections(bool bIncludingValid)
+{
+	for (const FLandscapeSplineConnection& Connection : ConnectedSegments)
+	{
+		FLandscapeSplineSegmentConnection& NearConnection = Connection.GetNearConnection();
+		if (bIncludingValid ||
+			(Mesh != NULL && Mesh->FindSocket(NearConnection.SocketName) == NULL) ||
+			(Mesh == NULL && NearConnection.SocketName != NAME_None))
+		{
+			FLandscapeSplineSegmentConnection& FarConnection = Connection.GetFarConnection();
+			FVector EndLocation; FRotator EndRotation;
+			FarConnection.ControlPoint->GetConnectionLocationAndRotation(FarConnection.SocketName, EndLocation, EndRotation);
+
+			NearConnection.SocketName = GetBestConnectionTo(EndLocation);
+			NearConnection.TangentLen = FMath::Abs(NearConnection.TangentLen);
+
+			// Allow flipping tangent on the null connection
+			if (NearConnection.SocketName == NAME_None)
+			{
+				FVector StartLocation; FRotator StartRotation;
+				NearConnection.ControlPoint->GetConnectionLocationAndRotation(NearConnection.SocketName, StartLocation, StartRotation);
+
+				if (FVector::DotProduct((EndLocation - StartLocation).SafeNormal(), StartRotation.Vector()) < 0)
+				{
+					NearConnection.TangentLen = -NearConnection.TangentLen;
+				}
+			}
+		}
+	}
+}
+#endif
+
+void ULandscapeSplineControlPoint::RegisterComponents()
+{
+	if (MeshComponent != NULL)
+	{
+		MeshComponent->RegisterComponent();
+	}
+}
+
+void ULandscapeSplineControlPoint::UnregisterComponents()
+{
+	if (MeshComponent != NULL)
+	{
+		MeshComponent->UnregisterComponent();
+	}
+}
+
+bool ULandscapeSplineControlPoint::OwnsComponent(const class UControlPointMeshComponent* StaticMeshComponent) const
+{
+	return MeshComponent == StaticMeshComponent;
+}
+
+#if WITH_EDITOR
+void ULandscapeSplineControlPoint::UpdateSplinePoints(bool bUpdateCollision, bool bUpdateAttachedSegments)
+{
+	ULandscapeSplinesComponent* OuterSplines = GetOuterULandscapeSplinesComponent();
+
+	if (MeshComponent == NULL && Mesh != NULL ||
+		MeshComponent != NULL &&
+			(MeshComponent->StaticMesh != Mesh ||
+			MeshComponent->RelativeLocation != Location ||
+			MeshComponent->RelativeRotation != Rotation ||
+			MeshComponent->RelativeScale3D  != MeshScale)
+		)
+	{
+		OuterSplines->Modify();
+
+		if (MeshComponent)
+		{
+			MeshComponent->Modify();
+			MeshComponent->UnregisterComponent();
+		}
+
+		if (Mesh != NULL)
+		{
+			if (MeshComponent == NULL)
+			{
+				Modify();
+				OuterSplines->GetOuter()->Modify();
+				MeshComponent = ConstructObject<UControlPointMeshComponent>(UControlPointMeshComponent::StaticClass(), OuterSplines->GetOuter(), NAME_None, RF_Transactional | RF_TextExportTransient);
+				MeshComponent->bSelected = bSelected;
+				MeshComponent->AttachTo(OuterSplines);
+			}
+
+			if (MeshComponent->StaticMesh != Mesh)
+			{
+				MeshComponent->SetStaticMesh(Mesh);
+
+				AutoSetConnections(false);
+			}
+
+			MeshComponent->RelativeLocation = Location;
+			MeshComponent->RelativeRotation = Rotation;
+			MeshComponent->RelativeScale3D  = MeshScale;
+			MeshComponent->InvalidateLightingCache();
+			MeshComponent->RegisterComponent();
+		}
+		else
+		{
+			Modify();
+			MeshComponent->GetOuter()->Modify();
+			MeshComponent->DestroyComponent();
+			MeshComponent = NULL;
+
+			AutoSetConnections(false);
+		}
+	}
+
+	if (MeshComponent != NULL)
+	{
+		if (MeshComponent->Materials != MaterialOverrides)
+		{
+			MeshComponent->Modify();
+			MeshComponent->Materials = MaterialOverrides;
+			MeshComponent->MarkRenderStateDirty();
+			if (MeshComponent->BodyInstance.IsValidBodyInstance())
+			{
+				MeshComponent->BodyInstance.UpdatePhysicalMaterials();
+			}
+		}
+
+		if (MeshComponent->LDMaxDrawDistance != LDMaxDrawDistance)
+		{
+			MeshComponent->Modify();
+			MeshComponent->LDMaxDrawDistance = LDMaxDrawDistance;
+			MeshComponent->CachedMaxDrawDistance = 0;
+			MeshComponent->MarkRenderStateDirty();
+		}
+
+		if (MeshComponent->CastShadow != bCastShadow)
+		{
+			MeshComponent->Modify();
+			MeshComponent->SetCastShadow(bCastShadow);
+		}
+
+		const FName CollisionProfile = bEnableCollision ? UCollisionProfile::BlockAll_ProfileName : UCollisionProfile::NoCollision_ProfileName;
+		if (MeshComponent->BodyInstance.GetCollisionProfileName() != CollisionProfile)
+		{
+			MeshComponent->Modify();
+			MeshComponent->BodyInstance.SetCollisionProfileName(CollisionProfile);
+		}
+	}
+
+	if (Mesh != NULL)
+	{
+		Points.Reset(ConnectedSegments.Num());
+
+		for (const FLandscapeSplineConnection& Connection : ConnectedSegments)
+		{
+			FVector StartLocation; FRotator StartRotation;
+			GetConnectionLocationAndRotation(Connection.GetNearConnection().SocketName, StartLocation, StartRotation);
+
+			const float Roll = FMath::DegreesToRadians(StartRotation.Roll);
+			const FVector Tangent = StartRotation.Vector();
+			const FVector BiNormal = FQuat(Tangent, -Roll).RotateVector((Tangent ^ FVector(0, 0, -1)).SafeNormal());
+			const FVector LeftPos = StartLocation - BiNormal * Width;
+			const FVector RightPos = StartLocation + BiNormal * Width;
+			const FVector FalloffLeftPos = StartLocation - BiNormal * (Width + SideFalloff);
+			const FVector FalloffRightPos = StartLocation + BiNormal * (Width + SideFalloff);
+
+			new(Points) FLandscapeSplineInterpPoint(StartLocation, LeftPos, RightPos, FalloffLeftPos, FalloffRightPos, 1.0f);
+		}
+
+		const FVector CPLocation = Location;
+		Points.Sort([&CPLocation](const FLandscapeSplineInterpPoint& x, const FLandscapeSplineInterpPoint& y){return (x.Center - CPLocation).Rotation().Yaw < (y.Center - CPLocation).Rotation().Yaw;});
+	}
+	else
+	{
+		Points.Reset(1);
+
+		FVector StartLocation; FRotator StartRotation;
+		GetConnectionLocationAndRotation(NAME_None, StartLocation, StartRotation);
+
+		const float Roll = FMath::DegreesToRadians(StartRotation.Roll);
+		const FVector Tangent = StartRotation.Vector();
+		const FVector BiNormal = FQuat(Tangent, -Roll).RotateVector((Tangent ^ FVector(0, 0, -1)).SafeNormal());
+		const FVector LeftPos = StartLocation - BiNormal * Width;
+		const FVector RightPos = StartLocation + BiNormal * Width;
+		const FVector FalloffLeftPos = StartLocation - BiNormal * (Width + SideFalloff);
+		const FVector FalloffRightPos = StartLocation + BiNormal * (Width + SideFalloff);
+
+		new(Points) FLandscapeSplineInterpPoint(StartLocation, LeftPos, RightPos, FalloffLeftPos, FalloffRightPos, 1.0f);
+	}
+
+	// Update Bounds
+	Bounds = FBox(0);
+	for (const FLandscapeSplineInterpPoint& Point : Points)
+	{
+		Bounds += Point.FalloffLeft;
+		Bounds += Point.FalloffRight;
+	}
+
+	OuterSplines->MarkRenderStateDirty();
+
+	if (bUpdateAttachedSegments)
+	{
+		for (const FLandscapeSplineConnection& Connection : ConnectedSegments)
+		{
+			Connection.Segment->UpdateSplinePoints(bUpdateCollision);
+		}
+	}
+}
+
+void ULandscapeSplineControlPoint::DeleteSplinePoints()
+{
+	Modify();
+
+	ULandscapeSplinesComponent* OuterSplines = CastChecked<ULandscapeSplinesComponent>(GetOuter());
+
+	Points.Reset();
+	Bounds = FBox(0);
+
+	OuterSplines->MarkRenderStateDirty();
+
+	if (MeshComponent != NULL)
+	{
+		MeshComponent->GetOuter()->Modify();
+		MeshComponent->DestroyComponent();
+		MeshComponent = NULL;
+	}
+}
+
+void ULandscapeSplineControlPoint::PostEditUndo()
+{
+	bHackIsUndoingSplines = true;
+	Super::PostEditUndo();
+	bHackIsUndoingSplines = false;
+
+	GetOuterULandscapeSplinesComponent()->MarkRenderStateDirty();
+}
+
+void ULandscapeSplineControlPoint::PostDuplicate(bool bDuplicateForPIE)
+{
+	if (!bDuplicateForPIE)
+	{
+		ULandscapeSplinesComponent* OuterSplines = CastChecked<ULandscapeSplinesComponent>(GetOuter());
+
+		if (MeshComponent != NULL)
+		{
+			if (MeshComponent->GetOuter() != OuterSplines->GetOuter())
+			{
+				MeshComponent = DuplicateObject<UControlPointMeshComponent>(MeshComponent, OuterSplines->GetOuter(), *MeshComponent->GetName());
+				MeshComponent->AttachTo(OuterSplines);
+			}
+		}
+	}
+
+	Super::PostDuplicate(bDuplicateForPIE);
+}
+
+void ULandscapeSplineControlPoint::PostEditImport()
+{
+	Super::PostEditImport();
+
+	GetOuterULandscapeSplinesComponent()->ControlPoints.AddUnique(this);
+}
+
+void ULandscapeSplineControlPoint::PostEditChangeProperty(FPropertyChangedEvent& PropertyChangedEvent)
+{
+	Super::PostEditChangeProperty(PropertyChangedEvent);
+
+	Width = FMath::Max(Width, 0.001f);
+	SideFalloff = FMath::Max(SideFalloff, 0.0f);
+	EndFalloff = FMath::Max(EndFalloff, 0.0f);
+
+	// Don't update splines when undoing, not only is it unnecessary and expensive,
+	// it also causes failed asserts in debug builds when trying to register components
+	// (because the actor hasn't reset its OwnedComponents array yet)
+	if (!bHackIsUndoingSplines)
+	{
+		const bool bUpdateCollision = PropertyChangedEvent.ChangeType != EPropertyChangeType::Interactive;
+		UpdateSplinePoints(bUpdateCollision);
+	}
+}
+#endif // WITH_EDITOR
+
+//////////////////////////////////////////////////////////////////////////
+// SPLINE SEGMENT
+
+ULandscapeSplineSegment::ULandscapeSplineSegment(const FObjectInitializer& ObjectInitializer)
+	: Super(ObjectInitializer)
+{
+
+	Connections[0].ControlPoint = NULL;
+	Connections[0].TangentLen = 0;
+	Connections[1].ControlPoint = NULL;
+	Connections[1].TangentLen = 0;
+
+#if WITH_EDITORONLY_DATA
+	LayerName = NAME_None;
+	bRaiseTerrain = true;
+	bLowerTerrain = true;
+
+	// SplineMesh properties
+	SplineMeshes.Empty();
+	LDMaxDrawDistance = 0;
+	bEnableCollision = true;
+	bCastShadow = true;
+
+	// transients
+	bSelected = false;
+#endif
+}
+
+void ULandscapeSplineSegment::PostInitProperties()
+{
+	Super::PostInitProperties();
+
+#if WITH_EDITORONLY_DATA
+	if (!HasAnyFlags(RF_ClassDefaultObject | RF_NeedLoad | RF_AsyncLoading))
+	{
+		// create a new random seed for all new objects
+		RandomSeed = FMath::Rand();
+	}
+#endif
+}
+
+void ULandscapeSplineSegment::Serialize(FArchive& Ar)
+{
+	Super::Serialize(Ar);
+
+#if WITH_EDITOR
+	if (Ar.UE4Ver() < VER_UE4_SPLINE_MESH_ORIENTATION)
+	{
+		for (FLandscapeSplineMeshEntry& MeshEntry : SplineMeshes)
+		{
+			switch (MeshEntry.Orientation_DEPRECATED)
+			{
+			case LSMO_XUp:
+				MeshEntry.ForwardAxis = ESplineMeshAxis::Z;
+				MeshEntry.UpAxis = ESplineMeshAxis::X;
+				break;
+			case LSMO_YUp:
+				MeshEntry.ForwardAxis = ESplineMeshAxis::Z;
+				MeshEntry.UpAxis = ESplineMeshAxis::Y;
+				break;
+			}
+		}
+	}
+#endif
+}
+
+void ULandscapeSplineSegment::PostLoad()
+{
+	Super::PostLoad();
+
+#if WITH_EDITOR
+	if (GIsEditor &&
+		GetLinkerUE4Version() < VER_UE4_ADDED_LANDSCAPE_SPLINE_EDITOR_MESH &&
+		MeshComponents.Num() == 0)
+	{
+		UpdateSplinePoints();
+	}
+
+	if (GIsEditor)
+	{
+		// Regenerate spline mesh components if any of the meshes used have gone missing
+		// Otherwise the spline will have no editor mesh and won't be selectable
+		ULandscapeSplinesComponent* OuterSplines = GetOuterULandscapeSplinesComponent();
+		for (auto* MeshComponent : MeshComponents)
+		{
+			MeshComponent->ConditionalPostLoad();
+			if (MeshComponent->StaticMesh == NULL)
+			{
+				UpdateSplinePoints();
+				break;
+			}
+		}
+	}
+
+	for (auto* MeshComponent : MeshComponents)
+	{
+		// Fix collision profile
+		const bool bUsingEditorMesh = MeshComponent->bHiddenInGame;
+		const FName CollisionProfile = (bEnableCollision && !bUsingEditorMesh) ? UCollisionProfile::BlockAll_ProfileName : UCollisionProfile::NoCollision_ProfileName;
+		if (MeshComponent->GetCollisionProfileName() != CollisionProfile)
+		{
+			MeshComponent->SetCollisionProfileName(CollisionProfile);
+		}
+
+		MeshComponent->SetFlags(RF_TextExportTransient);
+	}
+#endif
+}
+
+/**  */
+#if WITH_EDITOR
+void ULandscapeSplineSegment::SetSplineSelected(bool bInSelected)
+{
+	bSelected = bInSelected;
+	GetOuterULandscapeSplinesComponent()->MarkRenderStateDirty();
+
+	for (auto* MeshComponent : MeshComponents)
+	{
+		MeshComponent->bSelected = bInSelected;
+		MeshComponent->PushSelectionToProxy();
+	}
+}
+
+void ULandscapeSplineSegment::AutoFlipTangents()
+{
+	FVector StartLocation; FRotator StartRotation;
+	Connections[0].ControlPoint->GetConnectionLocationAndRotation(Connections[0].SocketName, StartLocation, StartRotation);
+	FVector EndLocation; FRotator EndRotation;
+	Connections[1].ControlPoint->GetConnectionLocationAndRotation(Connections[1].SocketName, EndLocation, EndRotation);
+
+	// Flipping the tangent is only allowed if not using a socket
+	if (Connections[0].SocketName == NAME_None && FVector::DotProduct((EndLocation - StartLocation).SafeNormal() * Connections[0].TangentLen, StartRotation.Vector()) < 0)
+	{
+		Connections[0].TangentLen = -Connections[0].TangentLen;
+	}
+	if (Connections[1].SocketName == NAME_None && FVector::DotProduct((StartLocation - EndLocation).SafeNormal() * Connections[1].TangentLen, EndRotation.Vector()) < 0)
+	{
+		Connections[1].TangentLen = -Connections[1].TangentLen;
+	}
+}
+#endif
+
+void ULandscapeSplineSegment::RegisterComponents()
+{
+	for (auto* MeshComponent : MeshComponents)
+	{
+		MeshComponent->RegisterComponent();
+	}
+}
+
+void ULandscapeSplineSegment::UnregisterComponents()
+{
+	for (auto* MeshComponent : MeshComponents)
+	{
+		MeshComponent->UnregisterComponent();
+	}
+}
+
+bool ULandscapeSplineSegment::OwnsComponent(const class USplineMeshComponent* SplineMeshComponent) const
+{
+	return MeshComponents.FindByKey(SplineMeshComponent) != NULL;
+}
+
+static float ApproxLength(const FInterpCurveVector& SplineInfo, const float Start = 0.0f, const float End = 1.0f, const int32 ApproxSections = 4)
+{
+	float SplineLength = 0;
+	FVector OldPos = SplineInfo.Eval(Start, FVector::ZeroVector);
+	for (int32 i = 1; i <= ApproxSections; i++)
+	{
+		FVector NewPos = SplineInfo.Eval(FMath::Lerp(Start, End, (float)i / (float)ApproxSections), FVector::ZeroVector);
+		SplineLength += (NewPos - OldPos).Size();
+		OldPos = NewPos;
+	}
+
+	return SplineLength;
+}
+
+/** Util that takes a 2D vector and rotates it by RotAngle (given in radians) */
+static FVector2D RotateVec2D(const FVector2D InVec, float RotAngle)
+{
+	FVector2D OutVec;
+	OutVec.X = (InVec.X * FMath::Cos(RotAngle)) - (InVec.Y * FMath::Sin(RotAngle));
+	OutVec.Y = (InVec.X * FMath::Sin(RotAngle)) + (InVec.Y * FMath::Cos(RotAngle));
+	return OutVec;
+}
+
+
+
+static ESplineMeshAxis::Type CrossAxis(ESplineMeshAxis::Type InForwardAxis, ESplineMeshAxis::Type InUpAxis)
+{
+	check(InForwardAxis != InUpAxis);
+	return (ESplineMeshAxis::Type)(3 ^ InForwardAxis ^ InUpAxis);
+}
+
+bool FLandscapeSplineMeshEntry::IsValid() const
+{
+	return Mesh != NULL && ForwardAxis != UpAxis && Scale.GetAbsMin() > KINDA_SMALL_NUMBER;
+}
+
+#if WITH_EDITOR
+void ULandscapeSplineSegment::UpdateSplinePoints(bool bUpdateCollision)
+{
+	Modify();
+
+	ULandscapeSplinesComponent* OuterSplines = CastChecked<ULandscapeSplinesComponent>(GetOuter());
+
+	SplineInfo.Points.Empty(2);
+	Points.Reset();
+
+	if (Connections[0].ControlPoint == NULL
+		|| Connections[1].ControlPoint == NULL)
+	{
+		return;
+	}
+
+	// Set up BSpline
+	FVector StartLocation; FRotator StartRotation;
+	Connections[0].ControlPoint->GetConnectionLocationAndRotation(Connections[0].SocketName, StartLocation, StartRotation);
+	new(SplineInfo.Points) FInterpCurvePoint<FVector>(0.0f, StartLocation, StartRotation.Vector() * Connections[0].TangentLen, StartRotation.Vector() * Connections[0].TangentLen, CIM_CurveUser);
+	FVector EndLocation; FRotator EndRotation;
+	Connections[1].ControlPoint->GetConnectionLocationAndRotation(Connections[1].SocketName, EndLocation, EndRotation);
+	new(SplineInfo.Points) FInterpCurvePoint<FVector>(1.0f, EndLocation, EndRotation.Vector() * -Connections[1].TangentLen, EndRotation.Vector() * -Connections[1].TangentLen, CIM_CurveUser);
+
+	// Pointify
+
+	// Calculate spline length
+	const float SplineLength = ApproxLength(SplineInfo, 0.0f, 1.0f, 4);
+
+	float StartFalloffFraction = ((Connections[0].ControlPoint->ConnectedSegments.Num() > 1) ? 0 : (Connections[0].ControlPoint->EndFalloff / SplineLength));
+	float EndFalloffFraction = ((Connections[1].ControlPoint->ConnectedSegments.Num() > 1) ? 0 : (Connections[1].ControlPoint->EndFalloff / SplineLength));
+
+	// Stop the start and end fall-off overlapping
+	const float TotalFalloff = StartFalloffFraction + EndFalloffFraction;
+	if (TotalFalloff > 1.0f)
+	{
+		StartFalloffFraction /= TotalFalloff;
+		EndFalloffFraction /= TotalFalloff;
+	}
+
+	const float StartWidth = Connections[0].ControlPoint->Width;
+	const float EndWidth = Connections[1].ControlPoint->Width;
+	const float StartSideFalloff = Connections[0].ControlPoint->SideFalloff;
+	const float EndSideFalloff = Connections[1].ControlPoint->SideFalloff;
+	const float StartRollDegrees = StartRotation.Roll * (Connections[0].TangentLen > 0 ? 1 : -1);
+	const float EndRollDegrees = EndRotation.Roll * (Connections[1].TangentLen > 0 ? -1 : 1);
+	const float StartRoll = FMath::DegreesToRadians(StartRollDegrees);
+	const float EndRoll = FMath::DegreesToRadians(EndRollDegrees);
+
+	int32 NumPoints = FMath::CeilToInt(SplineLength / OuterSplines->SplineResolution);
+	NumPoints = FMath::Clamp(NumPoints, 1, 1000);
+
+	float OldKeyTime = 0;
+	for (int32 i = 0; i < SplineInfo.Points.Num(); i++)
+	{
+		const float NewKeyTime = SplineInfo.Points[i].InVal;
+		const float NewKeyCosInterp = 0.5f - 0.5f * FMath::Cos(NewKeyTime * PI);
+		const float NewKeyWidth = FMath::Lerp(StartWidth, EndWidth, NewKeyCosInterp);
+		const float NewKeyFalloff = FMath::Lerp(StartSideFalloff, EndSideFalloff, NewKeyCosInterp);
+		const float NewKeyRoll = FMath::Lerp(StartRoll, EndRoll, NewKeyCosInterp);
+		const FVector NewKeyPos = SplineInfo.Eval(NewKeyTime, FVector::ZeroVector);
+		const FVector NewKeyTangent = SplineInfo.EvalDerivative(NewKeyTime, FVector::ZeroVector).SafeNormal();
+		const FVector NewKeyBiNormal = FQuat(NewKeyTangent, -NewKeyRoll).RotateVector((NewKeyTangent ^ FVector(0, 0, -1)).SafeNormal());
+		const FVector NewKeyLeftPos = NewKeyPos - NewKeyBiNormal * NewKeyWidth;
+		const FVector NewKeyRightPos = NewKeyPos + NewKeyBiNormal * NewKeyWidth;
+		const FVector NewKeyFalloffLeftPos = NewKeyPos - NewKeyBiNormal * (NewKeyWidth + NewKeyFalloff);
+		const FVector NewKeyFalloffRightPos = NewKeyPos + NewKeyBiNormal * (NewKeyWidth + NewKeyFalloff);
+		const float NewKeyStartEndFalloff = FMath::Min((StartFalloffFraction > 0 ? NewKeyTime / StartFalloffFraction : 1.0f), (EndFalloffFraction > 0 ? (1 - NewKeyTime) / EndFalloffFraction : 1.0f));
+
+		// If not the first keypoint, interp from the last keypoint.
+		if (i > 0)
+		{
+			int32 NumSteps = FMath::CeilToInt( (NewKeyTime - OldKeyTime) * NumPoints );
+			float DrawSubstep = (NewKeyTime - OldKeyTime) / NumSteps;
+
+			// Add a point for each substep, except the ends because that's the point added outside the interp'ing.
+			for (int32 j = 1; j < NumSteps; j++)
+			{
+				const float NewTime = OldKeyTime + j*DrawSubstep;
+				const float NewCosInterp = 0.5f - 0.5f * FMath::Cos(NewTime * PI);
+				const float NewWidth = FMath::Lerp(StartWidth, EndWidth, NewCosInterp);
+				const float NewFalloff = FMath::Lerp(StartSideFalloff, EndSideFalloff, NewCosInterp);
+				const float NewRoll = FMath::Lerp(StartRoll, EndRoll, NewCosInterp);
+				const FVector NewPos = SplineInfo.Eval(NewTime, FVector::ZeroVector);
+				const FVector NewTangent = SplineInfo.EvalDerivative(NewTime, FVector::ZeroVector).SafeNormal();
+				const FVector NewBiNormal = FQuat(NewTangent, -NewRoll).RotateVector((NewTangent ^ FVector(0, 0, -1)).SafeNormal());
+				const FVector NewLeftPos = NewPos - NewBiNormal * NewWidth;
+				const FVector NewRightPos = NewPos + NewBiNormal * NewWidth;
+				const FVector NewFalloffLeftPos = NewPos - NewBiNormal * (NewWidth + NewFalloff);
+				const FVector NewFalloffRightPos = NewPos + NewBiNormal * (NewWidth + NewFalloff);
+				const float NewStartEndFalloff = FMath::Min((StartFalloffFraction > 0 ? NewTime / StartFalloffFraction : 1.0f), (EndFalloffFraction > 0 ? (1 - NewTime) / EndFalloffFraction : 1.0f));
+
+				new(Points) FLandscapeSplineInterpPoint(NewPos, NewLeftPos, NewRightPos, NewFalloffLeftPos, NewFalloffRightPos, NewStartEndFalloff);
+			}
+		}
+
+		new(Points) FLandscapeSplineInterpPoint(NewKeyPos, NewKeyLeftPos, NewKeyRightPos, NewKeyFalloffLeftPos, NewKeyFalloffRightPos, NewKeyStartEndFalloff);
+
+		OldKeyTime = NewKeyTime;
+	}
+
+	// Handle self-intersection errors due to tight turns
+	FixSelfIntersection(&FLandscapeSplineInterpPoint::Left);
+	FixSelfIntersection(&FLandscapeSplineInterpPoint::Right);
+	FixSelfIntersection(&FLandscapeSplineInterpPoint::FalloffLeft);
+	FixSelfIntersection(&FLandscapeSplineInterpPoint::FalloffRight);
+
+	// Update Bounds
+	Bounds = FBox(0);
+	for (const FLandscapeSplineInterpPoint& Point : Points)
+	{
+		Bounds += Point.FalloffLeft;
+		Bounds += Point.FalloffRight;
+	}
+
+	OuterSplines->MarkRenderStateDirty();
+
+	// Spline mesh components
+	TArray<const FLandscapeSplineMeshEntry*> UsableMeshes;
+	UsableMeshes.Reserve(SplineMeshes.Num());
+	for (const FLandscapeSplineMeshEntry& MeshEntry : SplineMeshes)
+	{
+		if (MeshEntry.IsValid())
+		{
+			UsableMeshes.Add(&MeshEntry);
+		}
+	}
+
+	// Editor mesh
+	FLandscapeSplineMeshEntry SplineEditorMeshEntry;
+	SplineEditorMeshEntry.Mesh = OuterSplines->SplineEditorMesh;
+	SplineEditorMeshEntry.Scale.X = 3;
+	SplineEditorMeshEntry.Offset.Y = 0.5f;
+	bool bUsingEditorMesh = false;
+	if (UsableMeshes.Num() == 0 && OuterSplines->SplineEditorMesh != NULL)
+	{
+		bUsingEditorMesh = true;
+		UsableMeshes.Add(&SplineEditorMeshEntry);
+	}
+
+	OuterSplines->Modify();
+	for (auto* MeshComponent : MeshComponents)
+	{
+		MeshComponent->Modify();
+	}
+	UnregisterComponents();
+
+	if (SplineLength > 0 && (StartWidth > 0 || EndWidth > 0) && UsableMeshes.Num() > 0)
+	{
+		float T = 0;
+		int32 iMesh = 0;
+
+		struct FMeshSettings
+		{
+			const float T;
+			const FLandscapeSplineMeshEntry* const MeshEntry;
+
+			FMeshSettings(float InT, const FLandscapeSplineMeshEntry* const InMeshEntry) :
+				T(InT), MeshEntry(InMeshEntry) { }
+		};
+		TArray<FMeshSettings> MeshSettings;
+		MeshSettings.Reserve(21);
+
+		FRandomStream Random(RandomSeed);
+
+		// First pass:
+		// Choose meshes, create components, calculate lengths
+		while (T < 1.0f && iMesh < 20) // Max 20 meshes per spline segment
+		{
+			const float CosInterp = 0.5f - 0.5f * FMath::Cos(T * PI);
+			const float Width = FMath::Lerp(StartWidth, EndWidth, CosInterp);
+
+			const FLandscapeSplineMeshEntry* MeshEntry = UsableMeshes[Random.RandHelper(UsableMeshes.Num())];
+			UStaticMesh* Mesh = MeshEntry->Mesh;
+			const FBoxSphereBounds MeshBounds = Mesh->GetBounds();
+
+			FVector Scale = MeshEntry->Scale;
+			if (MeshEntry->bScaleToWidth)
+			{
+				Scale *= Width / USplineMeshComponent::GetAxisValue(MeshBounds.BoxExtent, CrossAxis(MeshEntry->ForwardAxis, MeshEntry->UpAxis));
+			}
+
+			const float MeshLength = FMath::Abs(USplineMeshComponent::GetAxisValue(MeshBounds.BoxExtent, MeshEntry->ForwardAxis) * 2 * USplineMeshComponent::GetAxisValue(Scale, MeshEntry->ForwardAxis));
+			float MeshT = (MeshLength / SplineLength);
+
+			// Improve our approximation if we're not going off the end of the spline
+			if (T + MeshT <= 1.0f)
+			{
+				MeshT *= (MeshLength / ApproxLength(SplineInfo, T, T + MeshT, 4));
+				MeshT *= (MeshLength / ApproxLength(SplineInfo, T, T + MeshT, 4));
+			}
+
+			// If it's smaller to round up than down, don't add another component
+			if (iMesh != 0 && (1.0f - T) < (T + MeshT - 1.0f))
+			{
+				break;
+			}
+
+			USplineMeshComponent* MeshComponent = NULL;
+			if (iMesh < MeshComponents.Num())
+			{
+				MeshComponent = MeshComponents[iMesh];
+			}
+			else
+			{
+				OuterSplines->GetOuter()->Modify();
+				MeshComponent = ConstructObject<USplineMeshComponent>(USplineMeshComponent::StaticClass(), OuterSplines->GetOuter(), NAME_None, RF_Transactional|RF_TextExportTransient);
+				MeshComponent->bSelected = bSelected;
+				MeshComponent->AttachTo(OuterSplines);
+				MeshComponents.Add(MeshComponent);
+			}
+
+			MeshComponent->SetStaticMesh(Mesh);
+
+			MeshComponent->Materials = MeshEntry->MaterialOverrides;
+			MeshComponent->MarkRenderStateDirty();
+			if (MeshComponent->BodyInstance.IsValidBodyInstance())
+			{
+				MeshComponent->BodyInstance.UpdatePhysicalMaterials();
+			}
+
+			MeshComponent->SetHiddenInGame(bUsingEditorMesh);
+			MeshComponent->SetVisibility(!bUsingEditorMesh || OuterSplines->bShowSplineEditorMesh);
+
+			MeshSettings.Add(FMeshSettings(T, MeshEntry));
+			iMesh++;
+			T += MeshT;
+		}
+		MeshSettings.Add(FMeshSettings(T, NULL));
+
+		if (iMesh < MeshComponents.Num())
+		{
+			for (int32 i = iMesh; i < MeshComponents.Num(); i++)
+			{
+				MeshComponents[i]->GetOuter()->Modify();
+				MeshComponents[i]->DestroyComponent();
+			}
+
+			MeshComponents.RemoveAt(iMesh, MeshComponents.Num() - iMesh);
+		}
+
+		// Second pass:
+		// Rescale components to fit a whole number to the spline, set up final parameters
+		const float Rescale = 1.0f / T;
+		for (int32 i = 0; i < MeshComponents.Num(); i++)
+		{
+			USplineMeshComponent* const MeshComponent = MeshComponents[i];
+			const UStaticMesh* const Mesh = MeshComponent->StaticMesh;
+			const FBoxSphereBounds MeshBounds = Mesh->GetBounds();
+
+			const float RescaledT = MeshSettings[i].T * Rescale;
+			const FLandscapeSplineMeshEntry* MeshEntry = MeshSettings[i].MeshEntry;
+			const ESplineMeshAxis::Type SideAxis = CrossAxis(MeshEntry->ForwardAxis, MeshEntry->UpAxis);
+
+			const float TEnd = MeshSettings[i + 1].T * Rescale;
+
+			const float CosInterp = 0.5f - 0.5f * FMath::Cos(RescaledT * PI);
+			const float Width = FMath::Lerp(StartWidth, EndWidth, CosInterp);
+			const bool bDoOrientationRoll = (MeshEntry->ForwardAxis == ESplineMeshAxis::X && MeshEntry->UpAxis == ESplineMeshAxis::Y) ||
+			                                (MeshEntry->ForwardAxis == ESplineMeshAxis::Y && MeshEntry->UpAxis == ESplineMeshAxis::Z) ||
+			                                (MeshEntry->ForwardAxis == ESplineMeshAxis::Z && MeshEntry->UpAxis == ESplineMeshAxis::X);
+			const float Roll = FMath::Lerp(StartRoll, EndRoll, CosInterp) + (bDoOrientationRoll ? -HALF_PI : 0);
+
+			FVector Scale = MeshEntry->Scale;
+			if (MeshEntry->bScaleToWidth)
+			{
+				Scale *= Width / USplineMeshComponent::GetAxisValue(MeshBounds.BoxExtent, SideAxis);
+			}
+
+			FVector2D Offset = MeshEntry->Offset;
+			if (MeshEntry->bCenterH)
+			{
+				if (bDoOrientationRoll)
+				{
+					Offset.Y -= USplineMeshComponent::GetAxisValue(MeshBounds.Origin, SideAxis);
+				}
+				else
+				{
+					Offset.X -= USplineMeshComponent::GetAxisValue(MeshBounds.Origin, SideAxis);
+				}
+			}
+
+			FVector2D Scale2D;
+			switch (MeshEntry->ForwardAxis)
+			{
+			case ESplineMeshAxis::X:
+				Scale2D = FVector2D(Scale.Y, Scale.Z);
+				break;
+			case ESplineMeshAxis::Y:
+				Scale2D = FVector2D(Scale.Z, Scale.X);
+				break;
+			case ESplineMeshAxis::Z:
+				Scale2D = FVector2D(Scale.X, Scale.Y);
+				break;
+			default:
+				check(0);
+				break;
+			}
+			Offset *= Scale2D;
+			Offset = RotateVec2D(Offset, -Roll);
+
+			MeshComponent->SplineParams.StartPos = SplineInfo.Eval(RescaledT, FVector::ZeroVector);
+			MeshComponent->SplineParams.StartTangent = SplineInfo.EvalDerivative(RescaledT, FVector::ZeroVector) * (TEnd - RescaledT);
+			MeshComponent->SplineParams.StartScale = Scale2D;
+			MeshComponent->SplineParams.StartRoll = Roll;
+			MeshComponent->SplineParams.StartOffset = Offset;
+
+			const float CosInterpEnd = 0.5f - 0.5f * FMath::Cos(TEnd * PI);
+			const float WidthEnd = FMath::Lerp(StartWidth, EndWidth, CosInterpEnd);
+			const float RollEnd = FMath::Lerp(StartRoll, EndRoll, CosInterpEnd) + (bDoOrientationRoll ? -HALF_PI : 0);
+
+			FVector ScaleEnd = MeshEntry->Scale;
+			if (MeshEntry->bScaleToWidth)
+			{
+				ScaleEnd *= WidthEnd / USplineMeshComponent::GetAxisValue(MeshBounds.BoxExtent, SideAxis);
+			}
+
+			FVector2D OffsetEnd = MeshEntry->Offset;
+			if (MeshEntry->bCenterH)
+			{
+				if (bDoOrientationRoll)
+				{
+					OffsetEnd.Y -= USplineMeshComponent::GetAxisValue(MeshBounds.Origin, SideAxis);
+				}
+				else
+				{
+					OffsetEnd.X -= USplineMeshComponent::GetAxisValue(MeshBounds.Origin, SideAxis);
+				}
+			}
+
+			FVector2D Scale2DEnd;
+			switch (MeshEntry->ForwardAxis)
+			{
+			case ESplineMeshAxis::X:
+				Scale2DEnd = FVector2D(ScaleEnd.Y, ScaleEnd.Z);
+				break;
+			case ESplineMeshAxis::Y:
+				Scale2DEnd = FVector2D(ScaleEnd.Z, ScaleEnd.X);
+				break;
+			case ESplineMeshAxis::Z:
+				Scale2DEnd = FVector2D(ScaleEnd.X, ScaleEnd.Y);
+				break;
+			default:
+				check(0);
+				break;
+			}
+			OffsetEnd *= Scale2DEnd;
+			OffsetEnd = RotateVec2D(OffsetEnd, -RollEnd);
+
+			MeshComponent->SplineParams.EndPos = SplineInfo.Eval(TEnd, FVector::ZeroVector);
+			MeshComponent->SplineParams.EndTangent = SplineInfo.EvalDerivative(TEnd, FVector::ZeroVector) * (TEnd - RescaledT);
+			MeshComponent->SplineParams.EndScale = Scale2DEnd;
+			MeshComponent->SplineParams.EndRoll = RollEnd;
+			MeshComponent->SplineParams.EndOffset = OffsetEnd;
+
+			MeshComponent->SplineUpDir = FVector(0,0,1); // Up, to be consistent between joined meshes. We rotate it to horizontal using roll if using Z Forward X Up or X Forward Y Up
+			MeshComponent->ForwardAxis = MeshEntry->ForwardAxis;
+
+			if (USplineMeshComponent::GetAxisValue(MeshEntry->Scale, MeshEntry->ForwardAxis) < 0)
+			{
+				Swap(MeshComponent->SplineParams.StartPos, MeshComponent->SplineParams.EndPos);
+				Swap(MeshComponent->SplineParams.StartTangent, MeshComponent->SplineParams.EndTangent);
+				Swap(MeshComponent->SplineParams.StartScale, MeshComponent->SplineParams.EndScale);
+				Swap(MeshComponent->SplineParams.StartRoll, MeshComponent->SplineParams.EndRoll);
+				Swap(MeshComponent->SplineParams.StartOffset, MeshComponent->SplineParams.EndOffset);
+
+				MeshComponent->SplineParams.StartTangent = -MeshComponent->SplineParams.StartTangent;
+				MeshComponent->SplineParams.EndTangent = -MeshComponent->SplineParams.EndTangent;
+				MeshComponent->SplineParams.StartScale.X = -MeshComponent->SplineParams.StartScale.X;
+				MeshComponent->SplineParams.EndScale.X = -MeshComponent->SplineParams.EndScale.X;
+				MeshComponent->SplineParams.StartRoll = -MeshComponent->SplineParams.StartRoll;
+				MeshComponent->SplineParams.EndRoll = -MeshComponent->SplineParams.EndRoll;
+				MeshComponent->SplineParams.StartOffset.X = -MeshComponent->SplineParams.StartOffset.X;
+				MeshComponent->SplineParams.EndOffset.X = -MeshComponent->SplineParams.EndOffset.X;
+			}
+
+			// Set Mesh component's location to half way between the start and end points. Improves the bounds and allows LDMaxDrawDistance to work
+			MeshComponent->RelativeLocation = (MeshComponent->SplineParams.StartPos + MeshComponent->SplineParams.EndPos) / 2;
+			MeshComponent->SplineParams.StartPos -= MeshComponent->RelativeLocation;
+			MeshComponent->SplineParams.EndPos -= MeshComponent->RelativeLocation;
+
+			if (MeshComponent->LDMaxDrawDistance != LDMaxDrawDistance)
+			{
+				MeshComponent->LDMaxDrawDistance = LDMaxDrawDistance;
+				MeshComponent->CachedMaxDrawDistance = 0;
+			}
+
+			MeshComponent->SetCastShadow(bCastShadow);
+			MeshComponent->InvalidateLightingCache();
+
+			MeshComponent->BodyInstance.SetCollisionProfileName((bEnableCollision && !bUsingEditorMesh) ? UCollisionProfile::BlockAll_ProfileName : UCollisionProfile::NoCollision_ProfileName);
+
+#if WITH_EDITOR
+			if (bUpdateCollision)
+			{
+				MeshComponent->RecreateCollision();
+			}
+			else
+			{
+				if (MeshComponent->BodySetup)
+				{
+					MeshComponent->BodySetup->InvalidatePhysicsData();
+					MeshComponent->BodySetup->AggGeom.EmptyElements();
+				}
+			}
+#endif
+		}
+
+		// Finally, register components
+		RegisterComponents();
+	}
+	else
+	{
+		if (MeshComponents.Num() > 0)
+		{
+			for (auto* MeshComponent : MeshComponents)
+			{
+				MeshComponent->GetOuter()->Modify();
+				MeshComponent->DestroyComponent();
+			}
+			MeshComponents.Empty();
+		}
+	}
+}
+
+void ULandscapeSplineSegment::UpdateSplineEditorMesh()
+{
+	ULandscapeSplinesComponent* OuterSplines = CastChecked<ULandscapeSplinesComponent>(GetOuter());
+
+	for (auto* MeshComponent : MeshComponents)
+	{
+		if (MeshComponent->bHiddenInGame)
+		{
+			MeshComponent->SetVisibility(OuterSplines->bShowSplineEditorMesh);
+		}
+	}
+}
+
+void ULandscapeSplineSegment::DeleteSplinePoints()
+{
+	Modify();
+
+	ULandscapeSplinesComponent* OuterSplines = CastChecked<ULandscapeSplinesComponent>(GetOuter());
+
+	SplineInfo.Reset();
+	Points.Reset();
+	Bounds = FBox(0);
+
+	OuterSplines->MarkRenderStateDirty();
+
+	for (auto* MeshComponent : MeshComponents)
+	{
+		MeshComponent->GetOuter()->Modify();
+		MeshComponent->DestroyComponent();
+	}
+	MeshComponents.Empty();
+}
+
+static bool LineIntersect( const FVector2D& L1Start, const FVector2D& L1End, const FVector2D& L2Start, const FVector2D& L2End, FVector2D& Intersect, float Tolerance = KINDA_SMALL_NUMBER)
+{
+	float tA = (L2End - L2Start) ^ (L2Start - L1Start);
+	float tB = (L1End - L1Start) ^ (L2Start - L1Start);
+	float Denom = (L2End - L2Start) ^ (L1End - L1Start);
+
+	if( FMath::IsNearlyZero( tA ) && FMath::IsNearlyZero( tB )  )
+	{
+		// Lines are the same
+		Intersect = (L2Start + L2End) / 2;
+		return true;
+	}
+
+	if( FMath::IsNearlyZero(Denom) )
+	{
+		// Lines are parallel
+		Intersect = (L2Start + L2End) / 2;
+		return false;
+	}
+
+	tA /= Denom;
+	tB /= Denom;
+
+	Intersect = L1Start + tA * (L1End - L1Start);
+
+	if( tA >= -Tolerance && tA <= (1.0f + Tolerance ) && tB >= -Tolerance && tB <= (1.0f + Tolerance) )
+	{
+		return true;
+	}
+
+	return false;
+}
+
+bool ULandscapeSplineSegment::FixSelfIntersection(FVector FLandscapeSplineInterpPoint::* Side)
+{
+	int32 StartSide = INDEX_NONE;
+	for(int32 i = 0; i < Points.Num(); i++)
+	{
+		bool bReversed = false;
+
+		if (i < Points.Num() - 1)
+		{
+			const FLandscapeSplineInterpPoint& CurrentPoint = Points[i];
+			const FLandscapeSplineInterpPoint& NextPoint = Points[i+1];
+			const FVector Direction = (NextPoint.Center - CurrentPoint.Center).SafeNormal();
+			const FVector SideDirection = (NextPoint.*Side - CurrentPoint.*Side).SafeNormal();
+			bReversed = (SideDirection | Direction) < 0;
+		}
+
+		if (bReversed)
+		{
+			if (StartSide == INDEX_NONE)
+			{
+				StartSide = i;
+			}
+		}
+		else
+		{
+			if (StartSide != INDEX_NONE)
+			{
+				int32 EndSide = i;
+
+				// step startSide back until before the endSide point
+				while (StartSide > 0)
+				{
+					const float Projection = (Points[StartSide].*Side - Points[StartSide-1].*Side) | (Points[EndSide].*Side - Points[StartSide-1].*Side);
+					if (Projection >= 0)
+					{
+						break;
+					}
+					StartSide--;
+				}
+				// step endSide forwards until after the startSide point
+				while (EndSide < Points.Num() - 1)
+				{
+					const float Projection = (Points[EndSide].*Side - Points[EndSide+1].*Side) | (Points[StartSide].*Side - Points[EndSide+1].*Side);
+					if (Projection >= 0)
+					{
+						break;
+					}
+					EndSide++;
+				}
+
+				// Can't do anything if the start and end intersect, as they're both unalterable
+				if (StartSide == 0 && EndSide == Points.Num() - 1)
+				{
+					return false;
+				}
+
+				FVector2D Collapse;
+				if (StartSide == 0)
+				{
+					Collapse = FVector2D(Points[StartSide].*Side);
+					StartSide++;
+				}
+				else if (EndSide == Points.Num() - 1)
+				{
+					Collapse = FVector2D(Points[EndSide].*Side);
+					EndSide--;
+				}
+				else
+				{
+					LineIntersect(FVector2D(Points[StartSide-1].*Side), FVector2D(Points[StartSide].*Side),
+						FVector2D(Points[EndSide+1].*Side), FVector2D(Points[EndSide].*Side), Collapse);
+				}
+
+				for (int32 j = StartSide; j <= EndSide; j++)
+				{
+					(Points[j].*Side).X = Collapse.X;
+					(Points[j].*Side).Y = Collapse.Y;
+				}
+
+				StartSide = INDEX_NONE;
+				i = EndSide;
+			}
+		}
+	}
+
+	return true;
+}
+#endif
+
+void ULandscapeSplineSegment::FindNearest( const FVector& InLocation, float& t, FVector& OutLocation, FVector& OutTangent )
+{
+	float TempOutDistanceSq;
+	t = SplineInfo.InaccurateFindNearest(InLocation, TempOutDistanceSq);
+	OutLocation = SplineInfo.Eval(t, FVector::ZeroVector);
+	OutTangent = SplineInfo.EvalDerivative(t, FVector::ZeroVector);
+}
+
+bool ULandscapeSplineSegment::Modify(bool bAlwaysMarkDirty /*= true*/)
+{
+	bool bSavedToTransactionBuffer = Super::Modify(bAlwaysMarkDirty);
+
+	for (auto* MeshComponent : MeshComponents)
+	{
+		if (MeshComponent)
+		{
+			bSavedToTransactionBuffer = MeshComponent->Modify(bAlwaysMarkDirty) || bSavedToTransactionBuffer;
+		}
+	}
+
+	return bSavedToTransactionBuffer;
+}
+
+#if WITH_EDITOR
+void ULandscapeSplineSegment::PostEditUndo()
+{
+	bHackIsUndoingSplines = true;
+	Super::PostEditUndo();
+	bHackIsUndoingSplines = false;
+
+	GetOuterULandscapeSplinesComponent()->MarkRenderStateDirty();
+}
+
+void ULandscapeSplineSegment::PostDuplicate(bool bDuplicateForPIE)
+{
+	if (!bDuplicateForPIE)
+	{
+		ULandscapeSplinesComponent* OuterSplines = CastChecked<ULandscapeSplinesComponent>(GetOuter());
+
+		for (auto*& MeshComponent : MeshComponents)
+		{
+			if (MeshComponent->GetOuter() != OuterSplines->GetOuter())
+			{
+				MeshComponent = DuplicateObject<USplineMeshComponent>(MeshComponent, OuterSplines->GetOuter(), *MeshComponent->GetName());
+				MeshComponent->AttachTo(OuterSplines);
+			}
+		}
+	}
+
+	Super::PostDuplicate(bDuplicateForPIE);
+}
+
+void ULandscapeSplineSegment::PostEditImport()
+{
+	Super::PostEditImport();
+
+	GetOuterULandscapeSplinesComponent()->Segments.AddUnique(this);
+
+	if (Connections[0].ControlPoint != NULL)
+	{
+		Connections[0].ControlPoint->ConnectedSegments.AddUnique(FLandscapeSplineConnection(this, 0));
+		Connections[1].ControlPoint->ConnectedSegments.AddUnique(FLandscapeSplineConnection(this, 1));
+	}
+}
+
+void ULandscapeSplineSegment::PostEditChangeProperty(FPropertyChangedEvent& PropertyChangedEvent)
+{
+	Super::PostEditChangeProperty(PropertyChangedEvent);
+
+	// Flipping the tangent is only allowed if not using a socket
+	if (Connections[0].SocketName != NAME_None)
+	{
+		Connections[0].TangentLen = FMath::Abs(Connections[0].TangentLen);
+	}
+	if (Connections[1].SocketName != NAME_None)
+	{
+		Connections[1].TangentLen = FMath::Abs(Connections[1].TangentLen);
+	}
+
+	// Don't update splines when undoing, not only is it unnecessary and expensive,
+	// it also causes failed asserts in debug builds when trying to register components
+	// (because the actor hasn't reset its OwnedComponents array yet)
+	if (!bHackIsUndoingSplines)
+	{
+		const bool bUpdateCollision = PropertyChangedEvent.ChangeType != EPropertyChangeType::Interactive;
+		UpdateSplinePoints(bUpdateCollision);
+	}
+}
+#endif // WITH_EDITOR
+
+
+
diff --git a/Engine/Source/Runtime/Landscape/Private/Materials/MaterialExpressionLandscapeLayerBlend.cpp b/Engine/Source/Runtime/Landscape/Private/Materials/MaterialExpressionLandscapeLayerBlend.cpp
new file mode 100644
--- /dev/null
+++ b/Engine/Source/Runtime/Landscape/Private/Materials/MaterialExpressionLandscapeLayerBlend.cpp
@@ -0,0 +1,305 @@
+// Copyright 1998-2014 Epic Games, Inc. All Rights Reserved.
+
+#include "Landscape.h"
+#include "Materials/MaterialExpressionLandscapeLayerBlend.h"
+
+#define LOCTEXT_NAMESPACE "Landscape"
+
+///////////////////////////////////////////////////////////////////////////////
+// UMaterialExpressionLandscapeLayerBlend
+///////////////////////////////////////////////////////////////////////////////
+UMaterialExpressionLandscapeLayerBlend::UMaterialExpressionLandscapeLayerBlend(const FObjectInitializer& ObjectInitializer)
+	: Super(ObjectInitializer)
+{
+	// Structure to hold one-time initialization
+	struct FConstructorStatics
+	{
+		FString NAME_Landscape;
+		FConstructorStatics()
+			: NAME_Landscape(LOCTEXT("Landscape", "Landscape").ToString())
+		{
+		}
+	};
+	static FConstructorStatics ConstructorStatics;
+
+	bIsParameterExpression = true;
+
+	MenuCategories.Add(ConstructorStatics.NAME_Landscape);
+}
+
+FGuid& UMaterialExpressionLandscapeLayerBlend::GetParameterExpressionId()
+{
+	return ExpressionGUID;
+}
+
+void UMaterialExpressionLandscapeLayerBlend::Serialize(FArchive& Ar)
+{
+	Super::Serialize(Ar);
+
+	if (Ar.IsLoading() && Ar.UE4Ver() < VER_UE4_ADD_LB_WEIGHTBLEND)
+	{
+		// convert any LB_AlphaBlend entries to LB_WeightBlend
+		for (FLayerBlendInput& LayerInput : Layers)
+		{
+			if (LayerInput.BlendType == LB_AlphaBlend)
+			{
+				LayerInput.BlendType = LB_WeightBlend;
+			}
+		}
+	}
+}
+
+const TArray<FExpressionInput*> UMaterialExpressionLandscapeLayerBlend::GetInputs()
+{
+	TArray<FExpressionInput*> Result;
+	for (int32 LayerIdx = 0; LayerIdx<Layers.Num(); LayerIdx++)
+	{
+		Result.Add(&Layers[LayerIdx].LayerInput);
+		if (Layers[LayerIdx].BlendType == LB_HeightBlend)
+		{
+			Result.Add(&Layers[LayerIdx].HeightInput);
+		}
+	}
+	return Result;
+}
+
+FExpressionInput* UMaterialExpressionLandscapeLayerBlend::GetInput(int32 InputIndex)
+{
+	int32 Idx = 0;
+	for (int32 LayerIdx = 0; LayerIdx<Layers.Num(); LayerIdx++)
+	{
+		if (InputIndex == Idx++)
+		{
+			return &Layers[LayerIdx].LayerInput;
+		}
+		if (Layers[LayerIdx].BlendType == LB_HeightBlend)
+		{
+			if (InputIndex == Idx++)
+			{
+				return &Layers[LayerIdx].HeightInput;
+			}
+		}
+	}
+	return NULL;
+}
+
+FString UMaterialExpressionLandscapeLayerBlend::GetInputName(int32 InputIndex) const
+{
+	int32 Idx = 0;
+	for (int32 LayerIdx = 0; LayerIdx<Layers.Num(); LayerIdx++)
+	{
+		if (InputIndex == Idx++)
+		{
+			return FString::Printf(TEXT("Layer %s"), *Layers[LayerIdx].LayerName.ToString());
+		}
+		if (Layers[LayerIdx].BlendType == LB_HeightBlend)
+		{
+			if (InputIndex == Idx++)
+			{
+				return FString::Printf(TEXT("Height %s"), *Layers[LayerIdx].LayerName.ToString());
+			}
+		}
+	}
+	return TEXT("");
+}
+
+bool UMaterialExpressionLandscapeLayerBlend::IsResultMaterialAttributes(int32 OutputIndex)
+{
+	if (ContainsInputLoop())
+	{
+		// If there is a loop anywhere in this expression's inputs then we can't risk checking them
+		return false;
+	}
+	for (int32 LayerIdx = 0; LayerIdx < Layers.Num(); LayerIdx++)
+	{
+		if (Layers[LayerIdx].LayerInput.Expression && Layers[LayerIdx].LayerInput.Expression->IsResultMaterialAttributes(Layers[LayerIdx].LayerInput.OutputIndex))
+		{
+			return true;
+		}
+	}
+	return false;
+}
+
+int32 UMaterialExpressionLandscapeLayerBlend::Compile(class FMaterialCompiler* Compiler, int32 OutputIndex, int32 MultiplexIndex)
+{
+	// For renormalization
+	bool bNeedsRenormalize = false;
+	int32 WeightSumCode = Compiler->Constant(0);
+
+	// Temporary store for each layer's weight
+	TArray<int32> WeightCodes;
+	WeightCodes.Empty(Layers.Num());
+
+	for (int32 LayerIdx = 0; LayerIdx<Layers.Num(); LayerIdx++)
+	{
+		WeightCodes.Add(INDEX_NONE);
+
+		FLayerBlendInput& Layer = Layers[LayerIdx];
+
+		// LB_AlphaBlend layers are blended last
+		if (Layer.BlendType != LB_AlphaBlend)
+		{
+			// Height input
+			const int32 HeightCode = Layer.HeightInput.Expression ? Layer.HeightInput.Compile(Compiler, MultiplexIndex) : Compiler->Constant(Layer.ConstHeightInput);
+
+			const int32 WeightCode = Compiler->StaticTerrainLayerWeight(Layer.LayerName, Layer.PreviewWeight > 0.0f ? Compiler->Constant(Layer.PreviewWeight) : INDEX_NONE);
+			if (WeightCode != INDEX_NONE)
+			{
+				switch (Layer.BlendType)
+				{
+				case LB_WeightBlend:
+				{
+					// Store the weight plus accumulate the sum of all weights so far
+					WeightCodes[LayerIdx] = WeightCode;
+					WeightSumCode = Compiler->Add(WeightSumCode, WeightCode);
+				}
+					break;
+				case LB_HeightBlend:
+				{
+					bNeedsRenormalize = true;
+
+					// Modify weight with height
+					int32 ModifiedWeightCode = Compiler->Clamp(
+						Compiler->Add(Compiler->Lerp(Compiler->Constant(-1.f), Compiler->Constant(1.f), WeightCode), HeightCode),
+						Compiler->Constant(0.f), Compiler->Constant(1.f));
+
+					// Store the final weight plus accumulate the sum of all weights so far
+					WeightCodes[LayerIdx] = ModifiedWeightCode;
+					WeightSumCode = Compiler->Add(WeightSumCode, ModifiedWeightCode);
+				}
+					break;
+				}
+			}
+		}
+	}
+
+	int32 InvWeightSumCode = Compiler->Div(Compiler->Constant(1.f), WeightSumCode);
+
+	int32 OutputCode = Compiler->Constant(0);
+
+	for (int32 LayerIdx = 0; LayerIdx<Layers.Num(); LayerIdx++)
+	{
+		FLayerBlendInput& Layer = Layers[LayerIdx];
+
+		if (WeightCodes[LayerIdx] != INDEX_NONE)
+		{
+			// Layer input
+			int32 LayerCode = Layer.LayerInput.Expression ? Layer.LayerInput.Compile(Compiler, MultiplexIndex) : Compiler->Constant3(Layer.ConstLayerInput.X, Layer.ConstLayerInput.Y, Layer.ConstLayerInput.Z);
+
+			if (bNeedsRenormalize)
+			{
+				// Renormalize the weights as our height modification has made them non-uniform
+				OutputCode = Compiler->Add(OutputCode, Compiler->Mul(LayerCode, Compiler->Mul(InvWeightSumCode, WeightCodes[LayerIdx])));
+			}
+			else
+			{
+				// No renormalization is necessary, so just add the weights
+				OutputCode = Compiler->Add(OutputCode, Compiler->Mul(LayerCode, WeightCodes[LayerIdx]));
+			}
+		}
+	}
+
+	// Blend in LB_AlphaBlend layers
+	for (FLayerBlendInput& Layer : Layers)
+	{
+		if (Layer.BlendType == LB_AlphaBlend)
+		{
+			const int32 WeightCode = Compiler->StaticTerrainLayerWeight(Layer.LayerName, Layer.PreviewWeight > 0.0f ? Compiler->Constant(Layer.PreviewWeight) : INDEX_NONE);
+			if (WeightCode != INDEX_NONE)
+			{
+				const int32 LayerCode = Layer.LayerInput.Expression ? Layer.LayerInput.Compile(Compiler, MultiplexIndex) : Compiler->Constant3(Layer.ConstLayerInput.X, Layer.ConstLayerInput.Y, Layer.ConstLayerInput.Z);
+				// Blend in the layer using the alpha value
+				OutputCode = Compiler->Lerp(OutputCode, LayerCode, WeightCode);
+			}
+		}
+	}
+
+	if (OutputCode != INDEX_NONE)
+	{
+		// We've definitely passed the reentrant check here so we're good to call IsResultMaterialAttributes().
+		bool bFoundExpression = false;
+		bool bIsResultMaterialAttributes = false;
+		for (int32 LayerIdx = 0; LayerIdx < Layers.Num(); LayerIdx++)
+		{
+			if (Layers[LayerIdx].HeightInput.Expression)
+			{
+				bool bHeightIsMaterialAttributes = Layers[LayerIdx].HeightInput.Expression->IsResultMaterialAttributes(Layers[LayerIdx].HeightInput.OutputIndex);
+				if (bHeightIsMaterialAttributes)
+				{
+					Compiler->Errorf(TEXT("Height input (%s) does not accept MaterialAttributes"), Layers[LayerIdx].LayerName);
+				}
+			}
+			if (Layers[LayerIdx].LayerInput.Expression)
+			{
+				bool bLayerIsMaterialAttributes = Layers[LayerIdx].LayerInput.Expression->IsResultMaterialAttributes(Layers[LayerIdx].LayerInput.OutputIndex);
+				if (!bFoundExpression)
+				{
+					bFoundExpression = true;
+					bIsResultMaterialAttributes = bLayerIsMaterialAttributes;
+				}
+				else if (bIsResultMaterialAttributes != bLayerIsMaterialAttributes)
+				{
+					Compiler->Error(TEXT("Cannot mix MaterialAttributes and non MaterialAttributes nodes"));
+					break;
+				}
+			}
+		}
+	}
+
+	return OutputCode;
+}
+
+UTexture* UMaterialExpressionLandscapeLayerBlend::GetReferencedTexture()
+{
+	return GEngine->WeightMapPlaceholderTexture;
+}
+
+void UMaterialExpressionLandscapeLayerBlend::GetCaption(TArray<FString>& OutCaptions) const
+{
+	OutCaptions.Add(FString(TEXT("Layer Blend")));
+}
+
+#if WITH_EDITOR
+
+void UMaterialExpressionLandscapeLayerBlend::PostEditChangeProperty(FPropertyChangedEvent& PropertyChangedEvent)
+{
+	Super::PostEditChangeProperty(PropertyChangedEvent);
+
+	// Clear out any height expressions for layers not using height blending
+	for (int32 LayerIdx = 0; LayerIdx<Layers.Num(); LayerIdx++)
+	{
+		if (Layers[LayerIdx].BlendType != LB_HeightBlend)
+		{
+			Layers[LayerIdx].HeightInput.Expression = NULL;
+		}
+	}
+
+	if (PropertyChangedEvent.MemberProperty)
+	{
+		const FName PropertyName = PropertyChangedEvent.MemberProperty->GetFName();
+		if (PropertyName == GET_MEMBER_NAME_CHECKED(UMaterialExpressionLandscapeLayerBlend, Layers))
+		{
+			if (GraphNode)
+			{
+				GraphNode->ReconstructNode();
+			}
+		}
+	}
+}
+#endif
+
+void UMaterialExpressionLandscapeLayerBlend::GetAllParameterNames(TArray<FName> &OutParameterNames, TArray<FGuid> &OutParameterIds)
+{
+	for (int32 LayerIdx = 0; LayerIdx<Layers.Num(); LayerIdx++)
+	{
+		FLayerBlendInput& Layer = Layers[LayerIdx];
+
+		int32 CurrentSize = OutParameterNames.Num();
+		OutParameterNames.AddUnique(Layer.LayerName);
+
+		if (CurrentSize != OutParameterNames.Num())
+		{
+			OutParameterIds.Add(ExpressionGUID);
+		}
+	}
+}
diff --git a/Engine/Source/Runtime/Landscape/Private/Materials/MaterialExpressionLandscapeLayerCoords.cpp b/Engine/Source/Runtime/Landscape/Private/Materials/MaterialExpressionLandscapeLayerCoords.cpp
new file mode 100644
--- /dev/null
+++ b/Engine/Source/Runtime/Landscape/Private/Materials/MaterialExpressionLandscapeLayerCoords.cpp
@@ -0,0 +1,94 @@
+// Copyright 1998-2014 Epic Games, Inc. All Rights Reserved.
+
+#include "Landscape.h"
+#include "Materials/MaterialExpressionLandscapeLayerCoords.h"
+
+///////////////////////////////////////////////////////////////////////////////
+// UMaterialExpressionLandscapeLayerCoords
+///////////////////////////////////////////////////////////////////////////////
+UMaterialExpressionLandscapeLayerCoords::UMaterialExpressionLandscapeLayerCoords(const FObjectInitializer& ObjectInitializer)
+	: Super(ObjectInitializer)
+{
+	// Structure to hold one-time initialization
+	struct FConstructorStatics
+	{
+		FString NAME_Landscape;
+		FConstructorStatics()
+			: NAME_Landscape(LOCTEXT("Landscape", "Landscape").ToString())
+		{
+		}
+	};
+	static FConstructorStatics ConstructorStatics;
+
+	MenuCategories.Add(ConstructorStatics.NAME_Landscape);
+	bCollapsed = false;
+}
+
+int32 UMaterialExpressionLandscapeLayerCoords::Compile(class FMaterialCompiler* Compiler, int32 OutputIndex, int32 MultiplexIndex)
+{
+	switch (CustomUVType)
+	{
+	case LCCT_CustomUV0:
+		return Compiler->TextureCoordinate(0, false, false);
+	case LCCT_CustomUV1:
+		return Compiler->TextureCoordinate(1, false, false);
+	case LCCT_CustomUV2:
+		return Compiler->TextureCoordinate(2, false, false);
+	case LCCT_WeightMapUV:
+		return Compiler->TextureCoordinate(3, false, false);
+	default:
+		break;
+	}
+
+	int32 BaseUV;
+
+	switch (MappingType)
+	{
+	case TCMT_Auto:
+	case TCMT_XY: BaseUV = Compiler->TextureCoordinate(0, false, false); break;
+	case TCMT_XZ: BaseUV = Compiler->TextureCoordinate(1, false, false); break;
+	case TCMT_YZ: BaseUV = Compiler->TextureCoordinate(2, false, false); break;
+	default: UE_LOG(LogLandscape, Fatal, TEXT("Invalid mapping type %u"), (uint8)MappingType); return INDEX_NONE;
+	};
+
+	float Scale = (MappingScale == 0.0f) ? 1.0f : 1.0f / MappingScale;
+	int32 RealScale = Compiler->Constant(Scale);
+
+	float Cos = FMath::Cos(MappingRotation * PI / 180.0);
+	float Sin = FMath::Sin(MappingRotation * PI / 180.0);
+
+	int32 ResultUV = INDEX_NONE;
+	int32 TransformedUV = Compiler->Add(
+		Compiler->Mul(RealScale,
+		Compiler->AppendVector(
+		Compiler->Dot(BaseUV, Compiler->Constant2(+Cos, +Sin)),
+		Compiler->Dot(BaseUV, Compiler->Constant2(-Sin, +Cos)))
+		),
+		Compiler->Constant2(MappingPanU, MappingPanV)
+		);
+
+	if (Compiler->GetFeatureLevel() != ERHIFeatureLevel::ES2) // No need to localize UV
+	{
+		ResultUV = TransformedUV;
+	}
+	else
+	{
+		int32 Offset = Compiler->TextureCoordinateOffset();
+		int32 TransformedOffset =
+			Compiler->Floor(
+			Compiler->Mul(RealScale,
+			Compiler->AppendVector(
+			Compiler->Dot(Offset, Compiler->Constant2(+Cos, +Sin)),
+			Compiler->Dot(Offset, Compiler->Constant2(-Sin, +Cos)))
+			));
+
+		ResultUV = Compiler->Sub(TransformedUV, TransformedOffset);
+	}
+
+	return ResultUV;
+}
+
+void UMaterialExpressionLandscapeLayerCoords::GetCaption(TArray<FString>& OutCaptions) const
+{
+	OutCaptions.Add(FString(TEXT("LandscapeCoords")));
+}
diff --git a/Engine/Source/Runtime/Landscape/Private/Materials/MaterialExpressionLandscapeLayerSample.cpp b/Engine/Source/Runtime/Landscape/Private/Materials/MaterialExpressionLandscapeLayerSample.cpp
new file mode 100644
--- /dev/null
+++ b/Engine/Source/Runtime/Landscape/Private/Materials/MaterialExpressionLandscapeLayerSample.cpp
@@ -0,0 +1,56 @@
+// Copyright 1998-2014 Epic Games, Inc. All Rights Reserved.
+
+#include "Landscape.h"
+#include "Materials/MaterialExpressionLandscapeLayerSample.h"
+
+///////////////////////////////////////////////////////////////////////////////
+// UMaterialExpressionLandscapeLayerSample
+///////////////////////////////////////////////////////////////////////////////
+UMaterialExpressionLandscapeLayerSample::UMaterialExpressionLandscapeLayerSample(const FObjectInitializer& ObjectInitializer)
+	: Super(ObjectInitializer)
+{
+	// Structure to hold one-time initialization
+	struct FConstructorStatics
+	{
+		FString NAME_Landscape;
+		FConstructorStatics()
+			: NAME_Landscape(LOCTEXT("Landscape", "Landscape").ToString())
+		{
+		}
+	};
+	static FConstructorStatics ConstructorStatics;
+
+	bIsParameterExpression = true;
+	MenuCategories.Add(ConstructorStatics.NAME_Landscape);
+}
+
+FGuid& UMaterialExpressionLandscapeLayerSample::GetParameterExpressionId()
+{
+	return ExpressionGUID;
+}
+
+int32 UMaterialExpressionLandscapeLayerSample::Compile(class FMaterialCompiler* Compiler, int32 OutputIndex, int32 MultiplexIndex)
+{
+	const int32 WeightCode = Compiler->StaticTerrainLayerWeight(ParameterName, Compiler->Constant(PreviewWeight));
+
+	return WeightCode;
+}
+
+UTexture* UMaterialExpressionLandscapeLayerSample::GetReferencedTexture()
+{
+	return GEngine->WeightMapPlaceholderTexture;
+}
+
+void UMaterialExpressionLandscapeLayerSample::GetCaption(TArray<FString>& OutCaptions) const
+{
+	OutCaptions.Add(FString::Printf(TEXT("Sample '%s'"), *ParameterName.ToString()));
+}
+
+void UMaterialExpressionLandscapeLayerSample::GetAllParameterNames(TArray<FName> &OutParameterNames, TArray<FGuid> &OutParameterIds)
+{
+	if (!OutParameterNames.Contains(ParameterName))
+	{
+		OutParameterNames.Add(ParameterName);
+		OutParameterIds.Add(ExpressionGUID);
+	}
+}
