Landscape module from Engine split out.

diff --git a/Engine/Source/Runtime/Landscape/Private/Materials/MaterialExpressionLandscapeLayerSwitch.cpp b/Engine/Source/Runtime/Landscape/Private/Materials/MaterialExpressionLandscapeLayerSwitch.cpp
new file mode 100644
--- /dev/null
+++ b/Engine/Source/Runtime/Landscape/Private/Materials/MaterialExpressionLandscapeLayerSwitch.cpp
@@ -0,0 +1,114 @@
+// Copyright 1998-2014 Epic Games, Inc. All Rights Reserved.
+
+#include "Landscape.h"
+#include "Materials/MaterialExpressionLandscapeLayerSwitch.h"
+
+///////////////////////////////////////////////////////////////////////////////
+// UMaterialExpressionLandscapeLayerSwitch
+///////////////////////////////////////////////////////////////////////////////
+UMaterialExpressionLandscapeLayerSwitch::UMaterialExpressionLandscapeLayerSwitch(const FObjectInitializer& ObjectInitializer)
+	: Super(ObjectInitializer)
+{
+	// Structure to hold one-time initialization
+	struct FConstructorStatics
+	{
+		FString NAME_Landscape;
+		FConstructorStatics()
+			: NAME_Landscape(LOCTEXT("Landscape", "Landscape").ToString())
+		{
+		}
+	};
+	static FConstructorStatics ConstructorStatics;
+
+	bIsParameterExpression = true;
+	MenuCategories.Add(ConstructorStatics.NAME_Landscape);
+	PreviewUsed = true;
+	bCollapsed = false;
+}
+
+bool UMaterialExpressionLandscapeLayerSwitch::IsResultMaterialAttributes(int32 OutputIndex)
+{
+	if (ContainsInputLoop())
+	{
+		// If there is a loop anywhere in this expression's inputs then we can't risk checking them
+		return false;
+	}
+	bool bLayerUsedIsMaterialAttributes = LayerUsed.Expression != NULL && LayerUsed.Expression->IsResultMaterialAttributes(LayerUsed.OutputIndex);
+	bool bLayerNotUsedIsMaterialAttributes = LayerNotUsed.Expression != NULL && LayerNotUsed.Expression->IsResultMaterialAttributes(LayerNotUsed.OutputIndex);
+	return bLayerUsedIsMaterialAttributes || bLayerNotUsedIsMaterialAttributes;
+}
+
+int32 UMaterialExpressionLandscapeLayerSwitch::Compile(class FMaterialCompiler* Compiler, int32 OutputIndex, int32 MultiplexIndex)
+{
+	const int32 WeightCode = Compiler->StaticTerrainLayerWeight(
+		ParameterName,
+		PreviewUsed ? Compiler->Constant(1.0f) : INDEX_NONE
+		);
+
+	int32 ReturnCode = INDEX_NONE;
+	if (WeightCode != INDEX_NONE)
+	{
+		ReturnCode = LayerUsed.Compile(Compiler, MultiplexIndex);
+	}
+	else
+	{
+		ReturnCode = LayerNotUsed.Compile(Compiler, MultiplexIndex);
+	}
+
+	if (ReturnCode != INDEX_NONE && //If we've already failed for some other reason don't bother with this check. It could have been the reentrant check causing this to loop infinitely!
+		LayerUsed.Expression != NULL && LayerNotUsed.Expression != NULL &&
+		LayerUsed.Expression->IsResultMaterialAttributes(LayerUsed.OutputIndex) != LayerNotUsed.Expression->IsResultMaterialAttributes(LayerNotUsed.OutputIndex))
+	{
+		Compiler->Error(TEXT("Cannot mix MaterialAttributes and non MaterialAttributes nodes"));
+	}
+
+	return ReturnCode;
+}
+
+UTexture* UMaterialExpressionLandscapeLayerSwitch::GetReferencedTexture()
+{
+	return GEngine->WeightMapPlaceholderTexture;
+}
+
+void UMaterialExpressionLandscapeLayerSwitch::GetCaption(TArray<FString>& OutCaptions) const
+{
+	OutCaptions.Add(TEXT("Layer Switch"));
+	OutCaptions.Add(FString::Printf(TEXT("'%s'"), *ParameterName.ToString()));
+}
+
+void UMaterialExpressionLandscapeLayerSwitch::Serialize(FArchive& Ar)
+{
+	Super::Serialize(Ar);
+
+	if (Ar.UE4Ver() < VER_UE4_FIX_TERRAIN_LAYER_SWITCH_ORDER)
+	{
+		Swap(LayerUsed, LayerNotUsed);
+	}
+}
+
+void UMaterialExpressionLandscapeLayerSwitch::PostLoad()
+{
+	Super::PostLoad();
+
+	if (GetLinkerUE4Version() < VER_UE4_FIXUP_TERRAIN_LAYER_NODES)
+	{
+		UpdateParameterGuid(true, true);
+	}
+}
+
+FGuid& UMaterialExpressionLandscapeLayerSwitch::GetParameterExpressionId()
+{
+	return ExpressionGUID;
+}
+
+
+void UMaterialExpressionLandscapeLayerSwitch::GetAllParameterNames(TArray<FName> &OutParameterNames, TArray<FGuid> &OutParameterIds)
+{
+	int32 CurrentSize = OutParameterNames.Num();
+	OutParameterNames.AddUnique(ParameterName);
+
+	if (CurrentSize != OutParameterNames.Num())
+	{
+		OutParameterIds.Add(ExpressionGUID);
+	}
+}
diff --git a/Engine/Source/Runtime/Landscape/Private/Materials/MaterialExpressionLandscapeLayerWeight.cpp b/Engine/Source/Runtime/Landscape/Private/Materials/MaterialExpressionLandscapeLayerWeight.cpp
new file mode 100644
--- /dev/null
+++ b/Engine/Source/Runtime/Landscape/Private/Materials/MaterialExpressionLandscapeLayerWeight.cpp
@@ -0,0 +1,101 @@
+// Copyright 1998-2014 Epic Games, Inc. All Rights Reserved.
+
+#include "Landscape.h"
+#include "Materials/MaterialExpressionLandscapeLayerWeight.h"
+
+///////////////////////////////////////////////////////////////////////////////
+// UMaterialExpressionLandscapeLayerWeight
+///////////////////////////////////////////////////////////////////////////////
+UMaterialExpressionLandscapeLayerWeight::UMaterialExpressionLandscapeLayerWeight(const FObjectInitializer& ObjectInitializer)
+	: Super(ObjectInitializer)
+{
+	// Structure to hold one-time initialization
+	struct FConstructorStatics
+	{
+		FString NAME_Landscape;
+		FConstructorStatics()
+			: NAME_Landscape(LOCTEXT("Landscape", "Landscape").ToString())
+		{
+		}
+	};
+	static FConstructorStatics ConstructorStatics;
+
+	bIsParameterExpression = true;
+	MenuCategories.Add(ConstructorStatics.NAME_Landscape);
+	PreviewWeight = 0.0f;
+	ConstBase = FVector(0.f, 0.f, 0.f);
+}
+
+FGuid& UMaterialExpressionLandscapeLayerWeight::GetParameterExpressionId()
+{
+	return ExpressionGUID;
+}
+
+void UMaterialExpressionLandscapeLayerWeight::PostLoad()
+{
+	Super::PostLoad();
+
+	if (GetLinkerUE4Version() < VER_UE4_FIXUP_TERRAIN_LAYER_NODES)
+	{
+		UpdateParameterGuid(true, true);
+	}
+}
+
+bool UMaterialExpressionLandscapeLayerWeight::IsResultMaterialAttributes(int32 OutputIndex)
+{
+	if (ContainsInputLoop())
+	{
+		// If there is a loop anywhere in this expression's inputs then we can't risk checking them
+		return false;
+	}
+	bool bLayerIsMaterialAttributes = Layer.Expression != NULL && Layer.Expression->IsResultMaterialAttributes(Layer.OutputIndex);
+	bool bBaseIsMaterialAttributes = Base.Expression != NULL && Base.Expression->IsResultMaterialAttributes(Base.OutputIndex);
+	return bLayerIsMaterialAttributes || bBaseIsMaterialAttributes;
+}
+
+int32 UMaterialExpressionLandscapeLayerWeight::Compile(class FMaterialCompiler* Compiler, int32 OutputIndex, int32 MultiplexIndex)
+{
+	const int32 BaseCode = Base.Expression ? Base.Compile(Compiler, MultiplexIndex) : Compiler->Constant3(ConstBase.X, ConstBase.Y, ConstBase.Z);
+	const int32 WeightCode = Compiler->StaticTerrainLayerWeight(ParameterName, Compiler->Constant(PreviewWeight));
+
+	int32 ReturnCode = INDEX_NONE;
+	if (WeightCode == INDEX_NONE)
+	{
+		ReturnCode = BaseCode;
+	}
+	else
+	{
+		const int32 LayerCode = Layer.Compile(Compiler, MultiplexIndex);
+		ReturnCode = Compiler->Add(BaseCode, Compiler->Mul(LayerCode, WeightCode));
+	}
+
+	if (ReturnCode != INDEX_NONE && //If we've already failed for some other reason don't bother with this check. It could have been the reentrant check causing this to loop infinitely!
+		Layer.Expression != NULL && Base.Expression != NULL &&
+		Layer.Expression->IsResultMaterialAttributes(Layer.OutputIndex) != Base.Expression->IsResultMaterialAttributes(Base.OutputIndex))
+	{
+		Compiler->Error(TEXT("Cannot mix MaterialAttributes and non MaterialAttributes nodes"));
+	}
+
+	return ReturnCode;
+}
+
+UTexture* UMaterialExpressionLandscapeLayerWeight::GetReferencedTexture()
+{
+	return GEngine->WeightMapPlaceholderTexture;
+}
+
+void UMaterialExpressionLandscapeLayerWeight::GetCaption(TArray<FString>& OutCaptions) const
+{
+	OutCaptions.Add(FString::Printf(TEXT("Layer '%s'"), *ParameterName.ToString()));
+}
+
+void UMaterialExpressionLandscapeLayerWeight::GetAllParameterNames(TArray<FName> &OutParameterNames, TArray<FGuid> &OutParameterIds)
+{
+	int32 CurrentSize = OutParameterNames.Num();
+	OutParameterNames.AddUnique(ParameterName);
+
+	if (CurrentSize != OutParameterNames.Num())
+	{
+		OutParameterIds.Add(ExpressionGUID);
+	}
+}
\ No newline at end of file
diff --git a/Engine/Source/Runtime/Landscape/Private/Materials/MaterialExpressionLandscapeVisibilityMask.cpp b/Engine/Source/Runtime/Landscape/Private/Materials/MaterialExpressionLandscapeVisibilityMask.cpp
new file mode 100644
--- /dev/null
+++ b/Engine/Source/Runtime/Landscape/Private/Materials/MaterialExpressionLandscapeVisibilityMask.cpp
@@ -0,0 +1,59 @@
+// Copyright 1998-2014 Epic Games, Inc. All Rights Reserved.
+
+#include "Landscape.h"
+#include "Materials/MaterialExpressionLandscapeVisibilityMask.h"
+
+///////////////////////////////////////////////////////////////////////////////
+// UMaterialExpressionLandscapeVisibilityMask
+///////////////////////////////////////////////////////////////////////////////
+
+FName UMaterialExpressionLandscapeVisibilityMask::ParameterName = FName("__LANDSCAPE_VISIBILITY__");
+
+UMaterialExpressionLandscapeVisibilityMask::UMaterialExpressionLandscapeVisibilityMask(const FObjectInitializer& ObjectInitializer)
+	: Super(ObjectInitializer)
+{
+	// Structure to hold one-time initialization
+	struct FConstructorStatics
+	{
+		FString NAME_Landscape;
+		FConstructorStatics()
+			: NAME_Landscape(LOCTEXT("Landscape", "Landscape").ToString())
+		{
+		}
+	};
+	static FConstructorStatics ConstructorStatics;
+
+	bIsParameterExpression = true;
+	MenuCategories.Add(ConstructorStatics.NAME_Landscape);
+}
+
+FGuid& UMaterialExpressionLandscapeVisibilityMask::GetParameterExpressionId()
+{
+	return ExpressionGUID;
+}
+
+int32 UMaterialExpressionLandscapeVisibilityMask::Compile(class FMaterialCompiler* Compiler, int32 OutputIndex, int32 MultiplexIndex)
+{
+	return Compiler->Sub(Compiler->Constant(1.f), Compiler->StaticTerrainLayerWeight(ParameterName, Compiler->Constant(0.f)));
+}
+
+UTexture* UMaterialExpressionLandscapeVisibilityMask::GetReferencedTexture()
+{
+	return GEngine->WeightMapPlaceholderTexture;
+}
+
+void UMaterialExpressionLandscapeVisibilityMask::GetAllParameterNames(TArray<FName> &OutParameterNames, TArray<FGuid> &OutParameterIds)
+{
+	int32 CurrentSize = OutParameterNames.Num();
+	OutParameterNames.AddUnique(ParameterName);
+
+	if (CurrentSize != OutParameterNames.Num())
+	{
+		OutParameterIds.Add(ExpressionGUID);
+	}
+}
+
+void UMaterialExpressionLandscapeVisibilityMask::GetCaption(TArray<FString>& OutCaptions) const
+{
+	OutCaptions.Add(FString(TEXT("Landscape Visibility Mask")));
+}
\ No newline at end of file
diff --git a/Engine/Source/Runtime/Landscape/Public/LandscapeEdit.h b/Engine/Source/Runtime/Landscape/Public/LandscapeEdit.h
new file mode 100644
--- /dev/null
+++ b/Engine/Source/Runtime/Landscape/Public/LandscapeEdit.h
@@ -0,0 +1,195 @@
+// Copyright 1998-2014 Epic Games, Inc. All Rights Reserved.
+
+/*=============================================================================
+LandscapeEdit.h: Classes for the editor to access to Landscape data
+=============================================================================*/
+
+#ifndef _LANDSCAPEEDIT_H
+#define _LANDSCAPEEDIT_H
+
+#define MAX_LANDSCAPE_LOD_DISTANCE_FACTOR 10.f
+
+#if WITH_EDITOR
+
+#include "LandscapeProxy.h" // for ELandscapeLayerPaintingRestriction
+
+struct FLandscapeTextureDataInfo
+{
+	struct FMipInfo
+	{
+		void* MipData;
+		TArray<FUpdateTextureRegion2D> MipUpdateRegions;
+	};
+
+	FLandscapeTextureDataInfo(UTexture2D* InTexture);
+	virtual ~FLandscapeTextureDataInfo();
+
+	// returns true if we need to block on the render thread before unlocking the mip data
+	bool UpdateTextureData();
+
+	int32 NumMips() { return MipInfo.Num(); }
+
+	void AddMipUpdateRegion(int32 MipNum, int32 InX1, int32 InY1, int32 InX2, int32 InY2)
+	{
+		check( MipNum < MipInfo.Num() );
+		new(MipInfo[MipNum].MipUpdateRegions) FUpdateTextureRegion2D(InX1, InY1, InX1, InY1, 1+InX2-InX1, 1+InY2-InY1);
+	}
+
+	void* GetMipData(int32 MipNum)
+	{
+		check( MipNum < MipInfo.Num() );
+		if( !MipInfo[MipNum].MipData )
+		{
+			MipInfo[MipNum].MipData = Texture->Source.LockMip(MipNum);
+		}
+		return MipInfo[MipNum].MipData;
+	}
+
+	int32 GetMipSizeX(int32 MipNum)
+	{
+		return FMath::Max(Texture->Source.GetSizeX() >> MipNum, 1);
+	}
+
+	int32 GetMipSizeY(int32 MipNum)
+	{
+		return FMath::Max(Texture->Source.GetSizeY() >> MipNum, 1);
+	}
+
+private:
+	UTexture2D* Texture;
+	TArray<FMipInfo> MipInfo;
+};
+
+struct LANDSCAPE_API FLandscapeEditDataInterface
+{
+	// tors
+	FLandscapeEditDataInterface(ULandscapeInfo* InLandscape);
+	virtual ~FLandscapeEditDataInterface();
+
+	// Misc
+	bool GetComponentsInRegion(int32 X1, int32 Y1, int32 X2, int32 Y2, TSet<ULandscapeComponent*>* OutComponents = NULL);
+
+	//
+	// Heightmap access
+	//
+	void SetHeightData(int32 X1, int32 Y1, int32 X2, int32 Y2, const uint16* Data, int32 Stride, bool CalcNormals, const uint16* NormalData = NULL, bool CreateComponents=false);
+
+	// Helper accessor
+	FORCEINLINE uint16 GetHeightMapData(const ULandscapeComponent* Component, int32 TexU, int32 TexV, FColor* TextureData = NULL);
+	// Generic
+	template<typename TStoreData>
+	void GetHeightDataTempl(int32& X1, int32& Y1, int32& X2, int32& Y2, TStoreData& StoreData);
+	// Without data interpolation, able to get normal data
+	template<typename TStoreData>
+	void GetHeightDataTemplFast(const int32 X1, const int32 Y1, const int32 X2, const int32 Y2, TStoreData& StoreData, TStoreData* NormalData = NULL);
+	// Implementation for fixed array
+	void GetHeightData(int32& X1, int32& Y1, int32& X2, int32& Y2, uint16* Data, int32 Stride);
+	void GetHeightDataFast(const int32 X1, const int32 Y1, const int32 X2, const int32 Y2, uint16* Data, int32 Stride, uint16* NormalData = NULL);
+	// Implementation for sparse array
+	void GetHeightData(int32& X1, int32& Y1, int32& X2, int32& Y2, TMap<FIntPoint, uint16>& SparseData);
+	void GetHeightDataFast(const int32 X1, const int32 Y1, const int32 X2, const int32 Y2, TMap<FIntPoint, uint16>& SparseData, TMap<FIntPoint, uint16>* NormalData = NULL);
+
+	// Recaclulate normals for the entire landscape.
+	void RecalculateNormals();
+
+	//
+	// Weightmap access
+	//
+	// Helper accessor
+	FORCEINLINE uint8 GetWeightMapData(const ULandscapeComponent* Component, ULandscapeLayerInfoObject* LayerInfo, int32 TexU, int32 TexV, uint8 Offset = 0, UTexture2D* Texture = NULL, uint8* TextureData = NULL);
+	template<typename TStoreData>
+	void GetWeightDataTempl(ULandscapeLayerInfoObject* LayerInfo, int32& X1, int32& Y1, int32& X2, int32& Y2, TStoreData& StoreData);
+	// Without data interpolation
+	template<typename TStoreData>
+	void GetWeightDataTemplFast(ULandscapeLayerInfoObject* LayerInfo, const int32 X1, const int32 Y1, const int32 X2, const int32 Y2, TStoreData& StoreData);
+	// Implementation for fixed array
+	void GetWeightData(ULandscapeLayerInfoObject* LayerInfo, int32& X1, int32& Y1, int32& X2, int32& Y2, uint8* Data, int32 Stride);
+	//void GetWeightData(FName LayerName, int32& X1, int32& Y1, int32& X2, int32& Y2, TArray<uint8>* Data, int32 Stride);
+	void GetWeightDataFast(ULandscapeLayerInfoObject* LayerInfo, const int32 X1, const int32 Y1, const int32 X2, const int32 Y2, uint8* Data, int32 Stride);
+	void GetWeightDataFast(ULandscapeLayerInfoObject* LayerInfo, const int32 X1, const int32 Y1, const int32 X2, const int32 Y2, TArray<uint8>* Data, int32 Stride);
+	// Implementation for sparse array
+	void GetWeightData(ULandscapeLayerInfoObject* LayerInfo, int32& X1, int32& Y1, int32& X2, int32& Y2, TMap<FIntPoint, uint8>& SparseData);
+	//void GetWeightData(FName LayerName, int32& X1, int32& Y1, int32& X2, int32& Y2, TMap<uint64, TArray<uint8>>& SparseData);
+	void GetWeightDataFast(ULandscapeLayerInfoObject* LayerInfo, const int32 X1, const int32 Y1, const int32 X2, const int32 Y2, TMap<FIntPoint, uint8>& SparseData);
+	void GetWeightDataFast(ULandscapeLayerInfoObject* LayerInfo, const int32 X1, const int32 Y1, const int32 X2, const int32 Y2, TMap<FIntPoint, TArray<uint8>>& SparseData);
+	// Updates weightmap for LayerInfo, optionally adjusting all other weightmaps.
+	void SetAlphaData(ULandscapeLayerInfoObject* LayerInfo, const int32 X1, const int32 Y1, const int32 X2, const int32 Y2, const uint8* Data, int32 Stride, ELandscapeLayerPaintingRestriction::Type PaintingRestriction = ELandscapeLayerPaintingRestriction::None, bool bWeightAdjust = true, bool bTotalWeightAdjust = false);
+	// Updates weightmaps for all layers. Data points to packed data for all layers in the landscape info
+	void SetAlphaData(const TSet<ULandscapeLayerInfoObject*>& DirtyLayerInfos, const int32 X1, const int32 Y1, const int32 X2, const int32 Y2, const uint8* Data, int32 Stride, ELandscapeLayerPaintingRestriction::Type PaintingRestriction = ELandscapeLayerPaintingRestriction::None);
+	// Delete a layer and re-normalize other layers
+	void DeleteLayer(ULandscapeLayerInfoObject* LayerInfo);
+	// Replace/merge a layer
+	void ReplaceLayer(ULandscapeLayerInfoObject* FromLayerInfo, ULandscapeLayerInfoObject* ToLayerInfo);
+
+	// Without data interpolation, Select Data 
+	template<typename TStoreData>
+	void GetSelectDataTempl(const int32 X1, const int32 Y1, const int32 X2, const int32 Y2, TStoreData& StoreData);
+	void GetSelectData(const int32 X1, const int32 Y1, const int32 X2, const int32 Y2, uint8* Data, int32 Stride);
+	void GetSelectData(const int32 X1, const int32 Y1, const int32 X2, const int32 Y2, TMap<FIntPoint, uint8>& SparseData);
+	void SetSelectData(const int32 X1, const int32 Y1, const int32 X2, const int32 Y2, const uint8* Data, int32 Stride);
+
+	//
+	// XYOffsetmap access
+	//
+	template<typename T>
+	void SetXYOffsetDataTempl(int32 X1, int32 Y1, int32 X2, int32 Y2, const T* Data, int32 Stride);
+	void SetXYOffsetData(int32 X1, int32 Y1, int32 X2, int32 Y2, const FVector2D* Data, int32 Stride);
+	void SetXYOffsetData(int32 X1, int32 Y1, int32 X2, int32 Y2, const FVector* Data, int32 Stride);
+	// Helper accessor
+	FORCEINLINE FVector2D GetXYOffsetmapData(const ULandscapeComponent* Component, int32 TexU, int32 TexV, FColor* TextureData = NULL);
+
+	template<typename TStoreData>
+	void GetXYOffsetDataTempl(int32& X1, int32& Y1, int32& X2, int32& Y2, TStoreData& StoreData);
+	void GetXYOffsetData(int32& X1, int32& Y1, int32& X2, int32& Y2, FVector2D* Data, int32 Stride);
+	void GetXYOffsetData(int32& X1, int32& Y1, int32& X2, int32& Y2, TMap<FIntPoint, FVector2D>& SparseData);
+	void GetXYOffsetData(int32& X1, int32& Y1, int32& X2, int32& Y2, FVector* Data, int32 Stride);
+	void GetXYOffsetData(int32& X1, int32& Y1, int32& X2, int32& Y2, TMap<FIntPoint, FVector>& SparseData);
+	// Without data interpolation
+	template<typename TStoreData>
+	void GetXYOffsetDataTemplFast(const int32 X1, const int32 Y1, const int32 X2, const int32 Y2, TStoreData& StoreData);
+	// Without data interpolation, able to get normal data
+	void GetXYOffsetDataFast(const int32 X1, const int32 Y1, const int32 X2, const int32 Y2, FVector2D* Data, int32 Stride);
+	void GetXYOffsetDataFast(const int32 X1, const int32 Y1, const int32 X2, const int32 Y2, TMap<FIntPoint, FVector2D>& SparseData);
+	void GetXYOffsetDataFast(const int32 X1, const int32 Y1, const int32 X2, const int32 Y2, FVector* Data, int32 Stride);
+	void GetXYOffsetDataFast(const int32 X1, const int32 Y1, const int32 X2, const int32 Y2, TMap<FIntPoint, FVector>& SparseData);
+
+	// Texture data access
+	FLandscapeTextureDataInfo* GetTextureDataInfo(UTexture2D* Texture);
+
+	// Flush texture updates
+	void Flush();
+
+	// Texture bulk operations for weightmap reallocation
+	void CopyTextureChannel(UTexture2D* Dest, int32 DestChannel, UTexture2D* Src, int32 SrcChannel);
+	void ZeroTextureChannel(UTexture2D* Dest, int32 DestChannel);
+	void CopyTextureFromHeightmap(UTexture2D* Dest, int32 DestChannel, ULandscapeComponent* Comp, int32 SrcChannel);
+
+	template<typename TData>
+	void SetTextureValueTempl(UTexture2D* Dest, TData Value);
+	void ZeroTexture(UTexture2D* Dest);
+	void SetTextureValue(UTexture2D* Dest, FColor Value);
+
+	template<typename TData>
+	bool EqualTextureValueTempl(UTexture2D* Src, TData Value);
+	bool EqualTextureValue(UTexture2D* Src, FColor Value);
+
+private:
+	int32 ComponentSizeQuads;
+	int32 SubsectionSizeQuads;
+	int32 ComponentNumSubsections;
+	FVector DrawScale;
+
+	TMap<UTexture2D*, FLandscapeTextureDataInfo*> TextureDataMap;
+	ULandscapeInfo* LandscapeInfo;
+
+	// Only for Missing Data interpolation... only internal usage
+	template<typename TData, typename TStoreData, typename FType>
+	FORCEINLINE void CalcMissingValues(const int32& X1, const int32& X2, const int32& Y1, const int32& Y2, 
+		const int32& ComponentIndexX1, const int32& ComponentIndexX2, const int32& ComponentIndexY1, const int32& ComponentIndexY2, 
+		const int32& ComponentSizeX, const int32& ComponentSizeY, TData* CornerValues, 
+		TArray<bool>& NoBorderY1, TArray<bool>& NoBorderY2, TArray<bool>& ComponentDataExist, TStoreData& StoreData);
+};
+
+#endif
+
+#endif // _LANDSCAPEEDIT_H
diff --git a/Engine/Source/Runtime/Landscape/Public/LandscapeModule.h b/Engine/Source/Runtime/Landscape/Public/LandscapeModule.h
new file mode 100644
--- /dev/null
+++ b/Engine/Source/Runtime/Landscape/Public/LandscapeModule.h
@@ -0,0 +1,14 @@
+// Copyright 1998-2014 Epic Games, Inc. All Rights Reserved.
+
+#pragma once
+
+#include "LandscapeInfo.h"
+
+/**
+* Landscape module interface
+*/
+class ILandscapeModule : public IModuleInterface
+{
+public:
+};
+
diff --git a/Engine/Source/Runtime/Landscape/Public/LandscapeRender.h b/Engine/Source/Runtime/Landscape/Public/LandscapeRender.h
new file mode 100644
--- /dev/null
+++ b/Engine/Source/Runtime/Landscape/Public/LandscapeRender.h
@@ -0,0 +1,738 @@
+// Copyright 1998-2014 Epic Games, Inc. All Rights Reserved.
+
+/*=============================================================================
+LandscapeRender.h: New terrain rendering
+=============================================================================*/
+
+#ifndef _LANDSCAPERENDER_H
+#define _LANDSCAPERENDER_H
+
+#include "UniformBuffer.h"
+
+#include "LandscapeComponent.h"
+#include "LandscapeProxy.h"
+
+#include "PrimitiveSceneProxy.h"
+
+#include "LightMap.h"
+#include "MeshBatch.h"
+#include "ShadowMap.h"
+
+// This defines the number of border blocks to surround terrain by when generating lightmaps
+#define TERRAIN_PATCH_EXPAND_SCALAR	1
+
+#define LANDSCAPE_NEIGHBOR_NUM	4
+
+#define LANDSCAPE_LOD_LEVELS 8
+#define LANDSCAPE_MAX_SUBSECTION_NUM 2
+
+#if WITH_EDITOR
+namespace ELandscapeViewMode
+{
+	enum Type
+	{
+		Invalid = -1,
+		/** Color only */
+		Normal = 0,
+		EditLayer,
+		/** Layer debug only */
+		DebugLayer,
+		LayerDensity,
+		LOD,
+	};
+}
+
+extern LANDSCAPE_API ELandscapeViewMode::Type GLandscapeViewMode;
+
+namespace ELandscapeEditRenderMode
+{
+	enum Type
+	{
+		None				= 0x0,
+		Gizmo				= 0x1,
+		SelectRegion		= 0x2,
+		SelectComponent		= 0x4,
+		Select = SelectRegion | SelectComponent,
+		Mask				= 0x8, 
+		InvertedMask		= 0x10, // Should not be overlapped with other bits 
+		BitMaskForMask = Mask | InvertedMask,
+
+	};
+}
+
+LANDSCAPE_API extern bool GLandscapeEditModeActive;
+LANDSCAPE_API extern int32 GLandscapeEditRenderMode;
+LANDSCAPE_API extern int32 GLandscapePreviewMeshRenderMode;
+LANDSCAPE_API extern UMaterial* GLayerDebugColorMaterial;
+LANDSCAPE_API extern UMaterialInstanceConstant* GSelectionColorMaterial;
+LANDSCAPE_API extern UMaterialInstanceConstant* GSelectionRegionMaterial;
+LANDSCAPE_API extern UMaterialInstanceConstant* GMaskRegionMaterial;
+LANDSCAPE_API extern UTexture2D* GLandscapeBlackTexture;
+#endif
+
+
+/** The uniform shader parameters for a landscape draw call. */
+BEGIN_UNIFORM_BUFFER_STRUCT(FLandscapeUniformShaderParameters,LANDSCAPE_API)
+	/** vertex shader parameters */
+	DECLARE_UNIFORM_BUFFER_STRUCT_MEMBER(FVector4,HeightmapUVScaleBias)
+	DECLARE_UNIFORM_BUFFER_STRUCT_MEMBER(FVector4,WeightmapUVScaleBias)
+	DECLARE_UNIFORM_BUFFER_STRUCT_MEMBER(FVector4,LandscapeLightmapScaleBias)
+	DECLARE_UNIFORM_BUFFER_STRUCT_MEMBER(FVector4,SubsectionSizeVertsLayerUVPan)
+	DECLARE_UNIFORM_BUFFER_STRUCT_MEMBER(FVector4,SubsectionOffsetParams)
+	DECLARE_UNIFORM_BUFFER_STRUCT_MEMBER(FVector4,LightmapSubsectionOffsetParams)
+	DECLARE_UNIFORM_BUFFER_STRUCT_MEMBER(FMatrix,LocalToWorldNoScaling)
+END_UNIFORM_BUFFER_STRUCT(FLandscapeUniformShaderParameters)
+
+/* Data needed for the landscape vertex factory to set the render state for an individual batch element */
+struct FLandscapeBatchElementParams
+{
+	const TUniformBuffer<FLandscapeUniformShaderParameters>* LandscapeUniformShaderParametersResource;
+	const FMatrix* LocalToWorldNoScalingPtr;
+
+	// LOD calculation-related params
+	const class FLandscapeComponentSceneProxy* SceneProxy;
+	int32 SubX;
+	int32 SubY;
+	int32	CurrentLOD;
+};
+
+class FLandscapeElementParamArray : public FOneFrameResource
+{
+public:
+	TArray<FLandscapeBatchElementParams, SceneRenderingAllocator> ElementParams;
+};
+
+/** Pixel shader parameters for use with FLandscapeVertexFactory */
+class FLandscapeVertexFactoryPixelShaderParameters : public FVertexFactoryShaderParameters
+{
+public:
+	/**
+	* Bind shader constants by name
+	* @param	ParameterMap - mapping of named shader constants to indices
+	*/
+	virtual void Bind(const FShaderParameterMap& ParameterMap);
+
+	/**
+	* Serialize shader params to an archive
+	* @param	Ar - archive to serialize to
+	*/
+	virtual void Serialize(FArchive& Ar);
+
+	/**
+	* Set any shader data specific to this vertex factory
+	*/
+	virtual void SetMesh(FRHICommandList& RHICmdList, FShader* PixelShader,const class FVertexFactory* VertexFactory,const class FSceneView& View,const struct FMeshBatchElement& BatchElement,uint32 DataFlags) const override;
+	
+	virtual uint32 GetSize() const override
+	{
+		return sizeof(*this);
+	}
+
+private:
+	FShaderResourceParameter NormalmapTextureParameter;
+	FShaderResourceParameter NormalmapTextureParameterSampler;
+	FShaderParameter LocalToWorldNoScalingParameter;
+};
+
+/** vertex factory for VTF-heightmap terrain  */
+class FLandscapeVertexFactory : public FVertexFactory
+{
+	DECLARE_VERTEX_FACTORY_TYPE(FLandscapeVertexFactory);
+
+public:
+
+	FLandscapeVertexFactory()
+	{
+	}
+
+	virtual ~FLandscapeVertexFactory()
+	{
+		ReleaseResource();
+	}
+
+	static FVertexFactoryShaderParameters* ConstructShaderParameters(EShaderFrequency ShaderFrequency);
+
+	struct DataType
+	{
+		/** The stream to read the vertex position from. */
+		FVertexStreamComponent PositionComponent;
+	};
+
+	/**
+	* Should we cache the material's shadertype on this platform with this vertex factory? 
+	*/
+	static bool ShouldCache(EShaderPlatform Platform, const class FMaterial* Material, const class FShaderType* ShaderType)
+	{
+		// only compile landscape materials for landscape vertex factory
+		// The special engine materials must be compiled for the landscape vertex factory because they are used with it for wireframe, etc.
+		return IsFeatureLevelSupported(Platform, ERHIFeatureLevel::SM4) &&
+			(Material->IsUsedWithLandscape() || Material->IsSpecialEngineMaterial());
+	}
+
+	/**
+	* Can be overridden by FVertexFactory subclasses to modify their compile environment just before compilation occurs.
+	*/
+	static void ModifyCompilationEnvironment(EShaderPlatform Platform, const FMaterial* Material, FShaderCompilerEnvironment& OutEnvironment);
+
+	/**
+	* Copy the data from another vertex factory
+	* @param Other - factory to copy from
+	*/
+	void Copy(const FLandscapeVertexFactory& Other);
+
+	// FRenderResource interface.
+	virtual void InitRHI() override;
+
+	static bool SupportsTessellationShaders() { return true; }
+
+	/**
+	 * An implementation of the interface used by TSynchronizedResource to update the resource with new data from the game thread.
+	 */
+	void SetData(const DataType& InData)
+	{
+		Data = InData;
+		UpdateRHI();
+	}
+
+	virtual uint64 GetStaticBatchElementVisibility( const class FSceneView& View, const struct FMeshBatch* Batch ) const override;
+
+	/** stream component data bound to this vertex factory */
+	DataType Data;  
+};
+
+
+/** vertex factory for VTF-heightmap terrain  */
+class FLandscapeXYOffsetVertexFactory : public FLandscapeVertexFactory
+{
+	DECLARE_VERTEX_FACTORY_TYPE(FLandscapeXYOffsetVertexFactory);
+
+public:
+
+	FLandscapeXYOffsetVertexFactory()
+	{
+	}
+
+	virtual ~FLandscapeXYOffsetVertexFactory()
+	{
+		ReleaseResource();
+	}
+
+	static FVertexFactoryShaderParameters* ConstructShaderParameters(EShaderFrequency ShaderFrequency);
+
+	/**
+	* Can be overridden by FVertexFactory subclasses to modify their compile environment just before compilation occurs.
+	*/
+	static void ModifyCompilationEnvironment(EShaderPlatform Platform, const FMaterial* Material, FShaderCompilerEnvironment& OutEnvironment);
+};
+
+struct FLandscapeVertex
+{
+	float VertexX;
+	float VertexY;
+	float SubX;
+	float SubY;
+};
+
+//
+// FLandscapeVertexBuffer
+//
+class FLandscapeVertexBuffer : public FVertexBuffer
+{
+	int32 SubsectionSizeVerts;
+	int32 NumSubsections;
+public:
+	
+	/** Constructor. */
+	FLandscapeVertexBuffer(int32 InSubsectionSizeVerts, int32 InNumSubsections)
+	: 	SubsectionSizeVerts(InSubsectionSizeVerts)
+	,	NumSubsections(InNumSubsections)
+	{
+		InitResource();
+	}
+
+	/** Destructor. */
+	virtual ~FLandscapeVertexBuffer()
+	{
+		ReleaseResource();
+	}
+
+	/** 
+	* Initialize the RHI for this rendering resource 
+	*/
+	virtual void InitRHI() override;
+};
+
+
+//
+// FLandscapeSharedAdjacencyIndexBuffer
+//
+class FLandscapeSharedAdjacencyIndexBuffer : public FRefCountedObject
+{
+public:
+	FLandscapeSharedAdjacencyIndexBuffer(class FLandscapeSharedBuffers* SharedBuffer);
+	virtual ~FLandscapeSharedAdjacencyIndexBuffer();
+
+	TArray<FIndexBuffer*> IndexBuffers; // For tessellation
+};
+
+//
+// FLandscapeSharedBuffers
+//
+class FLandscapeSharedBuffers : public FRefCountedObject
+{
+public:
+	struct FLandscapeIndexRanges
+	{
+		int32 MinIndex[LANDSCAPE_MAX_SUBSECTION_NUM][LANDSCAPE_MAX_SUBSECTION_NUM];
+		int32 MaxIndex[LANDSCAPE_MAX_SUBSECTION_NUM][LANDSCAPE_MAX_SUBSECTION_NUM];
+		int32 MinIndexFull;
+		int32 MaxIndexFull;
+	};
+
+	int32 SharedBuffersKey;
+	int32 NumIndexBuffers;
+	int32 SubsectionSizeVerts;
+	int32 NumSubsections;
+
+	FLandscapeVertexFactory* VertexFactory;
+	FLandscapeVertexBuffer* VertexBuffer;
+	FIndexBuffer** IndexBuffers;
+	FLandscapeIndexRanges* IndexRanges;
+	FLandscapeSharedAdjacencyIndexBuffer* AdjacencyIndexBuffers;
+	bool bUse32BitIndices;
+
+	FLandscapeSharedBuffers(int32 SharedBuffersKey, int32 SubsectionSizeQuads, int32 NumSubsections, ERHIFeatureLevel::Type InFeatureLevel, bool bRequiresAdjacencyInformation);
+
+	template <typename INDEX_TYPE>
+	void CreateIndexBuffers(ERHIFeatureLevel::Type InFeatureLevel, bool bRequiresAdjacencyInformation);
+
+	virtual ~FLandscapeSharedBuffers();
+};
+
+//
+// FLandscapeEditToolRenderData
+//
+struct FLandscapeEditToolRenderData
+{
+	enum
+	{
+		ST_NONE = 0,
+		ST_COMPONENT = 1,
+		ST_REGION = 2,
+		// = 4...
+	};
+
+	FLandscapeEditToolRenderData(ULandscapeComponent* InComponent)
+		:	ToolMaterial(NULL),
+			GizmoMaterial(NULL),
+			LandscapeComponent(InComponent),
+			SelectedType(ST_NONE),
+			DebugChannelR(INDEX_NONE),
+			DebugChannelG(INDEX_NONE),
+			DebugChannelB(INDEX_NONE),
+			DataTexture(NULL)
+	{}
+
+	// Material used to render the tool.
+	UMaterialInterface* ToolMaterial;
+	// Material used to render the gizmo selection region...
+	UMaterialInterface* GizmoMaterial;
+
+	ULandscapeComponent* LandscapeComponent;
+
+	// Component is selected
+	int32 SelectedType;
+	int32 DebugChannelR, DebugChannelG, DebugChannelB;
+	UTexture2D* DataTexture; // Data texture other than height/weight
+
+#if WITH_EDITOR
+	void UpdateDebugColorMaterial();
+	void UpdateSelectionMaterial(int32 InSelectedType);
+#endif
+
+	// Game thread update
+	LANDSCAPE_API void Update( UMaterialInterface* InNewToolMaterial );
+	LANDSCAPE_API void UpdateGizmo( UMaterialInterface* InNewGizmoMaterial );
+	// Allows game thread to queue the deletion by the render thread
+	void Cleanup();
+};
+
+//
+// FLandscapeComponentSceneProxy
+//
+class FLandscapeComponentSceneProxy : public FPrimitiveSceneProxy
+{
+	friend class FLandscapeSharedBuffers;
+
+	class FLandscapeLCI : public FLightCacheInterface
+	{
+	public:
+		/** Initialization constructor. */
+		FLandscapeLCI(const ULandscapeComponent* InComponent)
+		{
+			LightMap = InComponent->LightMap;
+			ShadowMap = InComponent->ShadowMap;
+			IrrelevantLights = InComponent->IrrelevantLights;
+		}
+
+		virtual ~FLandscapeLCI()
+		{
+		}
+
+		// FLightCacheInterface
+		virtual FLightInteraction GetInteraction(const class FLightSceneProxy* LightSceneProxy) const;
+
+		virtual FLightMapInteraction GetLightMapInteraction(ERHIFeatureLevel::Type InFeatureLevel) const
+		{
+			return LightMap ? LightMap->GetInteraction(InFeatureLevel) : FLightMapInteraction();
+		}
+
+		virtual FShadowMapInteraction GetShadowMapInteraction() const
+		{
+			return ShadowMap ? ShadowMap->GetInteraction() : FShadowMapInteraction();
+		}
+
+	private:
+		/** The light-map used by the element. */
+		const FLightMap* LightMap;
+
+		/** The shadowmap used by the element. */
+		const FShadowMap* ShadowMap;
+
+		TArray<FGuid> IrrelevantLights;
+	};
+
+protected:
+	int32						MaxLOD;
+	int32						ComponentSizeQuads;	// Size of component in quads
+	int32						ComponentSizeVerts;
+	int32						NumSubsections;
+	int32						SubsectionSizeQuads;
+	int32						SubsectionSizeVerts;
+	FIntPoint					SectionBase;
+	float						StaticLightingResolution;
+	uint32						StaticLightingLOD;
+	FMatrix						LocalToWorldNoScaling;
+
+	// Storage for static draw list batch params
+	TArray<FLandscapeBatchElementParams> StaticBatchParamArray;
+
+	// Precomputed values
+	float					LODDistance;
+	float					PrecomputedLODFactor;
+	float					DistDiff;
+
+	FVector4 WeightmapScaleBias;
+	float WeightmapSubsectionOffset;
+	TArray<class UTexture2D*> WeightmapTextures;
+	int32 NumWeightmapLayerAllocations;
+
+	UTexture2D* NormalmapTexture; // PC : Heightmap, Mobile : Weightmap
+	
+	UTexture2D* HeightmapTexture; // PC : Heightmap, Mobile : Weightmap
+	FVector4 HeightmapScaleBias;
+	float HeightmapSubsectionOffsetU;
+	float HeightmapSubsectionOffsetV;
+
+	UTexture2D* XYOffsetmapTexture;
+
+	bool						bRequiresAdjacencyInformation;
+	uint32						SharedBuffersKey;
+	FLandscapeSharedBuffers*	SharedBuffers;
+	FLandscapeVertexFactory*	VertexFactory;
+	
+	UMaterialInterface* MaterialInterface;
+	FMaterialRelevance MaterialRelevance;
+
+	// Reference counted vertex and index buffer shared among all landscape scene proxies of the same component size
+	// Key is the component size and number of subsections.
+	static TMap<uint32, FLandscapeSharedBuffers*> SharedBuffersMap;
+	static TMap<uint32, FLandscapeSharedAdjacencyIndexBuffer*> SharedAdjacencyIndexBufferMap;
+
+	FLandscapeEditToolRenderData* EditToolRenderData;
+
+	// FLightCacheInterface
+	FLandscapeLCI* ComponentLightInfo;
+
+	const ULandscapeComponent* LandscapeComponent;
+
+	FVector2D				NeighborPosition[LANDSCAPE_NEIGHBOR_NUM];
+	int32					ForcedLOD;
+	int32					LODBias;
+	uint8					ForcedNeighborLOD[LANDSCAPE_NEIGHBOR_NUM];
+	uint8					NeighborLODBias[LANDSCAPE_NEIGHBOR_NUM];
+
+	enum ELandscapeLODFalloff::Type LODFalloff;
+
+	TUniformBuffer<FLandscapeUniformShaderParameters> LandscapeUniformShaderParameters;
+
+	// Cached versions of these
+	FMatrix					WorldToLocal;
+
+	virtual ~FLandscapeComponentSceneProxy();
+
+	// Used for DrawDynamicElements
+	FMeshBatch DynamicMesh;			// Landscape Rendering using Dynamic path
+#if WITH_EDITOR
+	FMeshBatch DynamicMeshTools;	// Tool rendering, don't support tessellation for now
+#endif
+	TArray<FLandscapeBatchElementParams> DynamicMeshBatchParamArray;
+
+
+public:
+	// constructor
+	FLandscapeComponentSceneProxy(ULandscapeComponent* InComponent, FLandscapeEditToolRenderData* InEditToolRenderData);
+
+	// FPrimitiveSceneProxy interface.
+	virtual void DrawStaticElements(FStaticPrimitiveDrawInterface* PDI) override;
+	virtual void GetDynamicMeshElements(const TArray<const FSceneView*>& Views, const FSceneViewFamily& ViewFamily, uint32 VisibilityMap, FMeshElementCollector& Collector) const override;
+	virtual void DrawDynamicElements(FPrimitiveDrawInterface* PDI,const FSceneView* View) override;
+	virtual uint32 GetMemoryFootprint( void ) const override { return( sizeof( *this ) + GetAllocatedSize() ); }
+	virtual FPrimitiveViewRelevance GetViewRelevance(const FSceneView* View) override;
+	virtual bool CanBeOccluded() const override;
+	virtual void GetLightRelevance(const FLightSceneProxy* LightSceneProxy, bool& bDynamic, bool& bRelevant, bool& bLightMapped, bool& bShadowMapped) const override;
+	virtual void OnTransformChanged() override;
+	virtual void CreateRenderThreadResources() override;
+
+	friend class ULandscapeComponent;
+	friend class FLandscapeVertexFactoryVertexShaderParameters;
+	friend class FLandscapeXYOffsetVertexFactoryVertexShaderParameters;
+	friend class FLandscapeVertexFactoryPixelShaderParameters;
+	friend struct FLandscapeBatchElementParams;
+	friend class FLandscapeVertexFactoryMobileVertexShaderParameters;
+	friend class FLandscapeVertexFactoryMobilePixelShaderParameters;
+
+	// FLandscapeComponentSceneProxy interface.
+	int32 CalcLODForSubsection(const class FSceneView& View, int32 SubX, int32 SubY, const FVector2D& CameraLocalPos) const;
+	int32 CalcLODForSubsectionNoForced(const class FSceneView& View, int32 SubX, int32 SubY, const FVector2D& CameraLocalPos) const;
+	void CalcLODParamsForSubsection(const class FSceneView& View, const FVector2D& CameraLocalPos, int32 SubX, int32 SubY, float& OutfLOD, FVector4& OutNeighborLODs, float& OutDistLOD) const;
+	uint64 GetStaticBatchElementVisibility( const class FSceneView& View, const struct FMeshBatch* Batch ) const;
+
+	// FLandcapeSceneProxy
+	void ChangeLODDistanceFactor_RenderThread(FVector2D InLODDistanceFactors);
+};
+
+class FLandscapeDebugMaterialRenderProxy : public FMaterialRenderProxy
+{
+public:
+	const FMaterialRenderProxy* const Parent;
+	const UTexture2D* RedTexture;
+	const UTexture2D* GreenTexture;
+	const UTexture2D* BlueTexture;
+	const FLinearColor R;
+	const FLinearColor G;
+	const FLinearColor B;
+
+	/** Initialization constructor. */
+	FLandscapeDebugMaterialRenderProxy(const FMaterialRenderProxy* InParent, const UTexture2D* TexR, const UTexture2D* TexG, const UTexture2D* TexB, 
+		const FLinearColor& InR, const FLinearColor& InG, const FLinearColor& InB ):
+	Parent(InParent),
+		RedTexture(TexR),
+		GreenTexture(TexG),
+		BlueTexture(TexB),
+		R(InR),
+		G(InG),
+		B(InB)
+	{}
+
+	// FMaterialRenderProxy interface.
+	virtual const class FMaterial* GetMaterial(ERHIFeatureLevel::Type FeatureLevel) const
+	{
+		return Parent->GetMaterial(FeatureLevel);
+	}
+	virtual bool GetVectorValue(const FName ParameterName, FLinearColor* OutValue, const FMaterialRenderContext& Context) const
+	{
+		if (ParameterName == FName(TEXT("Landscape_RedMask") ) )
+		{
+			*OutValue = R;
+			return true;
+		}
+		else if (ParameterName == FName(TEXT("Landscape_GreenMask") ) )
+		{
+			*OutValue = G;
+			return true;
+		}
+		else if (ParameterName == FName(TEXT("Landscape_BlueMask") ) )
+		{
+			*OutValue = B;
+			return true;
+		}
+		else
+		{
+			return Parent->GetVectorValue(ParameterName, OutValue, Context);
+		}
+	}
+	virtual bool GetScalarValue(const FName ParameterName, float* OutValue, const FMaterialRenderContext& Context) const
+	{
+		return Parent->GetScalarValue(ParameterName, OutValue, Context);
+	}
+	virtual bool GetTextureValue(const FName ParameterName,const UTexture** OutValue, const FMaterialRenderContext& Context) const
+	{
+		// NOTE: These should be returning black textures when NULL. The material will
+		// use a white texture if they are.
+		if (ParameterName == FName(TEXT("Landscape_RedTexture")) )
+		{
+			*OutValue = RedTexture;
+			return true;
+		}
+		else if (ParameterName == FName(TEXT("Landscape_GreenTexture"))  )
+		{
+			*OutValue = GreenTexture;
+			return true;
+		}
+		else if (ParameterName == FName(TEXT("Landscape_BlueTexture"))  )
+		{
+			*OutValue = BlueTexture;
+			return true;
+		}
+		else
+		{
+			return Parent->GetTextureValue(ParameterName, OutValue, Context);
+		}
+	}
+};
+
+class FLandscapeSelectMaterialRenderProxy : public FMaterialRenderProxy
+{
+public:
+	const FMaterialRenderProxy* const Parent;
+	const UTexture2D* SelectTexture;
+
+	/** Initialization constructor. */
+	FLandscapeSelectMaterialRenderProxy(const FMaterialRenderProxy* InParent, const UTexture2D* InTexture):
+	Parent(InParent),
+		SelectTexture(InTexture)
+	{}
+
+	// FMaterialRenderProxy interface.
+	virtual const class FMaterial* GetMaterial(ERHIFeatureLevel::Type FeatureLevel) const
+	{
+		return Parent->GetMaterial(FeatureLevel);
+	}
+	virtual bool GetVectorValue(const FName ParameterName, FLinearColor* OutValue, const FMaterialRenderContext& Context) const
+	{
+		if (ParameterName == FName(TEXT("HighlightColor")) )
+		{
+			*OutValue = FLinearColor(1.f, 0.5f, 0.5f);
+			return true;
+		}
+		else
+		{
+			return Parent->GetVectorValue(ParameterName, OutValue, Context);
+		}
+	}
+	virtual bool GetScalarValue(const FName ParameterName, float* OutValue, const FMaterialRenderContext& Context) const
+	{
+		return Parent->GetScalarValue(ParameterName, OutValue, Context);
+	}
+	virtual bool GetTextureValue(const FName ParameterName,const UTexture** OutValue, const FMaterialRenderContext& Context) const
+	{
+		if (ParameterName == FName(TEXT("SelectedData")) )
+		{
+			*OutValue = SelectTexture;
+			return true;
+		}
+		else
+		{
+			return Parent->GetTextureValue(ParameterName, OutValue, Context);
+		}
+	}
+};
+
+class FLandscapeMaskMaterialRenderProxy : public FMaterialRenderProxy
+{
+public:
+	const FMaterialRenderProxy* const Parent;
+	const UTexture2D* SelectTexture;
+	const bool bInverted;
+
+	/** Initialization constructor. */
+	FLandscapeMaskMaterialRenderProxy(const FMaterialRenderProxy* InParent, const UTexture2D* InTexture, const bool InbInverted):
+	Parent(InParent),
+		SelectTexture(InTexture),
+		bInverted(InbInverted)
+	{}
+
+	// FMaterialRenderProxy interface.
+	virtual const class FMaterial* GetMaterial(ERHIFeatureLevel::Type FeatureLevel) const
+	{
+		return Parent->GetMaterial(FeatureLevel);
+	}
+	virtual bool GetVectorValue(const FName ParameterName, FLinearColor* OutValue, const FMaterialRenderContext& Context) const
+	{
+		return Parent->GetVectorValue(ParameterName, OutValue, Context);
+	}
+	virtual bool GetScalarValue(const FName ParameterName, float* OutValue, const FMaterialRenderContext& Context) const
+	{
+		if (ParameterName == FName(TEXT("bInverted")) )
+		{
+			*OutValue = bInverted;
+			return true;
+		}
+		return Parent->GetScalarValue(ParameterName, OutValue, Context);
+	}
+	virtual bool GetTextureValue(const FName ParameterName,const UTexture** OutValue, const FMaterialRenderContext& Context) const
+	{
+		if (ParameterName == FName(TEXT("SelectedData")) )
+		{
+			*OutValue = SelectTexture;
+			return true;
+		}
+		else
+		{
+			return Parent->GetTextureValue(ParameterName, OutValue, Context);
+		}
+	}
+};
+
+namespace 
+{
+	// LightmapRes: Multiplier of lightmap size relative to landscape size
+	// X: (Output) PatchExpandCountX (at Lighting LOD)
+	// Y: (Output) PatchExpandCountY (at Lighting LOD)
+	// ComponentSize: Component size in patches (at LOD 0)
+	// LigtmapSize: Size desired for lightmap (texels)
+	// DesiredSize: (Output) Recommended lightmap size (texels)
+	// return: LightMapRatio
+	static float GetTerrainExpandPatchCount(float LightMapRes, int32& X, int32& Y, int32 ComponentSize, int32 LightmapSize, int32& DesiredSize, uint32 LightingLOD)
+	{
+		if (LightMapRes <= 0) return 0.f;
+
+		// Assuming DXT_1 compression at the moment...
+		int32 PixelPaddingX = GPixelFormats[PF_DXT1].BlockSizeX;
+		int32 PixelPaddingY = GPixelFormats[PF_DXT1].BlockSizeY;
+		int32 PatchExpandCountX = (LightMapRes >= 1.f) ? (PixelPaddingX) / LightMapRes : (PixelPaddingX);
+		int32 PatchExpandCountY = (LightMapRes >= 1.f )? (PixelPaddingY) / LightMapRes : (PixelPaddingY);
+
+		X = FMath::Max<int32>(1, PatchExpandCountX >> LightingLOD);
+		Y = FMath::Max<int32>(1, PatchExpandCountY >> LightingLOD);
+
+		DesiredSize = (LightMapRes >= 1.f) ? FMath::Min<int32>((int32)((ComponentSize + 1) * LightMapRes), 4096) : FMath::Min<int32>((int32)((LightmapSize) * LightMapRes), 4096);
+		int32 CurrentSize = (LightMapRes >= 1.f) ? FMath::Min<int32>((int32)((2*(X<<LightingLOD) + ComponentSize + 1) * LightMapRes), 4096) : FMath::Min<int32>((int32)((2*(X<<LightingLOD) + LightmapSize) * LightMapRes), 4096);
+
+		// Find proper Lightmap Size
+		if (CurrentSize > DesiredSize)
+		{
+			// Find maximum bit
+			int32 PriorSize = DesiredSize;
+			while (DesiredSize > 0)
+			{
+				PriorSize = DesiredSize;
+				DesiredSize = DesiredSize & ~(DesiredSize & ~(DesiredSize-1));
+			}
+
+			DesiredSize = PriorSize << 1; // next bigger size
+			if ( CurrentSize * CurrentSize <= ((PriorSize * PriorSize) << 1)  )
+			{
+				DesiredSize = PriorSize;
+			}
+		}
+
+		int32 DestSize = (float)DesiredSize / CurrentSize * (ComponentSize*LightMapRes);
+		float LightMapRatio = (float)DestSize / (ComponentSize*LightMapRes) * CurrentSize / DesiredSize;
+		return LightMapRatio;
+	}
+}
+
+
+#endif // _LANDSCAPERENDER_H
diff --git a/Engine/Source/Runtime/Landscape/Public/LandscapeRenderMobile.h b/Engine/Source/Runtime/Landscape/Public/LandscapeRenderMobile.h
new file mode 100644
--- /dev/null
+++ b/Engine/Source/Runtime/Landscape/Public/LandscapeRenderMobile.h
@@ -0,0 +1,136 @@
+// Copyright 1998-2014 Epic Games, Inc. All Rights Reserved.
+
+/*=============================================================================
+LandscapeRenderMobile.h: Mobile landscape rendering
+=============================================================================*/
+
+#ifndef _LANDSCAPERENDERMOBILE_H
+#define _LANDSCAPERENDERMOBILE_H
+
+#include "Landscape.h"
+
+class FLandscapeComponentSceneProxyMobile;
+
+#define LANDSCAPE_MAX_ES_LOD_COMP	2
+#define LANDSCAPE_MAX_ES_LOD		6
+
+struct FLandscapeMobileVertex
+{
+	uint8 Position[4]; // Pos + LOD 0 Height
+	uint8 LODHeights[LANDSCAPE_MAX_ES_LOD_COMP*4];
+};
+
+/** vertex factory for VTF-heightmap terrain  */
+class FLandscapeVertexFactoryMobile : public FLandscapeVertexFactory
+{
+	DECLARE_VERTEX_FACTORY_TYPE(FLandscapeVertexFactoryMobile);
+
+	typedef FLandscapeVertexFactory Super;
+public:
+
+	struct DataType : FLandscapeVertexFactory::DataType
+	{
+		/** stream which has heights of each LOD levels */
+		TArray<FVertexStreamComponent,TFixedAllocator<LANDSCAPE_MAX_ES_LOD_COMP> > LODHeightsComponent;
+	};
+
+	FLandscapeVertexFactoryMobile()
+	{
+	}
+
+	virtual ~FLandscapeVertexFactoryMobile()
+	{
+		ReleaseResource();
+	}
+
+	static FVertexFactoryShaderParameters* ConstructShaderParameters(EShaderFrequency ShaderFrequency);
+
+	/**
+	* Should we cache the material's shadertype on this platform with this vertex factory? 
+	*/
+	static bool ShouldCache(EShaderPlatform Platform, const class FMaterial* Material, const class FShaderType* ShaderType)
+	{
+		auto FeatureLevel = GetMaxSupportedFeatureLevel(Platform);
+		return (FeatureLevel == ERHIFeatureLevel::ES2 || FeatureLevel == ERHIFeatureLevel::ES3_1) &&
+			(Material->IsUsedWithLandscape() || Material->IsSpecialEngineMaterial());
+	}
+
+	static void ModifyCompilationEnvironment( EShaderPlatform Platform, const FMaterial* Material, FShaderCompilerEnvironment& OutEnvironment )
+	{
+		FVertexFactory::ModifyCompilationEnvironment(Platform, Material, OutEnvironment);
+		OutEnvironment.SetDefine(TEXT("NUM_VF_PACKED_INTERPOLANTS"), TEXT("1"));
+	}
+
+	// FRenderResource interface.
+	virtual void InitRHI() override;
+
+	/**
+	 * An implementation of the interface used by TSynchronizedResource to update the resource with new data from the game thread.
+	 */
+	void SetData(const DataType& InData)
+	{
+		MobileData = InData;
+		UpdateRHI();
+	}
+
+private:
+	/** stream component data bound to this vertex factory */
+	DataType MobileData; 
+
+	friend class FLandscapeComponentSceneProxyMobile;
+};
+
+//
+// FLandscapeVertexBuffer
+//
+class FLandscapeVertexBufferMobile : public FVertexBuffer
+{
+	const void* Data;
+	int32 DataSize;
+public:
+
+	/** Constructor. */
+	FLandscapeVertexBufferMobile(const void* InData, int32 InDataSize)
+	: 	Data(InData)
+	,	DataSize(InDataSize)
+	{
+		INC_DWORD_STAT_BY(STAT_LandscapeVertexMem, DataSize);
+		InitResource();
+	}
+
+	/** Destructor. */
+	virtual ~FLandscapeVertexBufferMobile()
+	{
+		ReleaseResource();
+		DEC_DWORD_STAT_BY(STAT_LandscapeVertexMem, DataSize);
+	}
+
+	/** 
+	* Initialize the RHI for this rendering resource 
+	*/
+	virtual void InitRHI() override;
+};
+
+//
+// FLandscapeComponentSceneProxy
+//
+class FLandscapeComponentSceneProxyMobile : public FLandscapeComponentSceneProxy
+{
+	FLandscapeVertexBufferMobile* VertexBuffer;
+
+	// Cooked data
+	TArray<uint8> PlatformData;
+
+	virtual ~FLandscapeComponentSceneProxyMobile();
+
+public:
+	FLandscapeComponentSceneProxyMobile(ULandscapeComponent* InComponent, FLandscapeEditToolRenderData* InEditToolRenderData);
+
+	virtual void CreateRenderThreadResources() override;
+
+	uint8 BlendableLayerMask;
+
+	friend class FLandscapeVertexBufferMobile;
+};
+
+#endif // _LANDSCAPERENDERMOBILE_H
\ No newline at end of file
diff --git a/Engine/Source/Runtime/Landscape/Public/LandscapeSplineProxies.h b/Engine/Source/Runtime/Landscape/Public/LandscapeSplineProxies.h
new file mode 100644
--- /dev/null
+++ b/Engine/Source/Runtime/Landscape/Public/LandscapeSplineProxies.h
@@ -0,0 +1,75 @@
+// Copyright 1998-2014 Epic Games, Inc. All Rights Reserved.
+
+#pragma once
+
+//////////////////////////////////////////////////////////////////////////
+// LANDSCAPE SPLINES HIT PROXY
+
+struct HLandscapeSplineProxy : public HHitProxy
+{
+	DECLARE_HIT_PROXY( LANDSCAPE_API );
+
+	HLandscapeSplineProxy(EHitProxyPriority InPriority = HPP_Wireframe) :
+		HHitProxy(InPriority)
+	{
+	}
+	virtual EMouseCursor::Type GetMouseCursor()
+	{
+		return EMouseCursor::Crosshairs;
+	}
+};
+
+struct HLandscapeSplineProxy_Segment : public HLandscapeSplineProxy
+{
+	DECLARE_HIT_PROXY( LANDSCAPE_API );
+
+	class ULandscapeSplineSegment* SplineSegment;
+
+	HLandscapeSplineProxy_Segment(class ULandscapeSplineSegment* InSplineSegment) :
+		HLandscapeSplineProxy(),
+		SplineSegment(InSplineSegment)
+	{
+	}
+	virtual void AddReferencedObjects( FReferenceCollector& Collector ) override
+	{
+		Collector.AddReferencedObject( SplineSegment );
+	}
+};
+
+struct HLandscapeSplineProxy_ControlPoint : public HLandscapeSplineProxy
+{
+	DECLARE_HIT_PROXY( LANDSCAPE_API );
+
+	class ULandscapeSplineControlPoint* ControlPoint;
+
+	HLandscapeSplineProxy_ControlPoint(class ULandscapeSplineControlPoint* InControlPoint) :
+		HLandscapeSplineProxy(HPP_Foreground),
+		ControlPoint(InControlPoint)
+	{
+	}
+	virtual void AddReferencedObjects( FReferenceCollector& Collector ) override
+	{
+		Collector.AddReferencedObject( ControlPoint );
+	}
+};
+
+struct HLandscapeSplineProxy_Tangent : public HLandscapeSplineProxy
+{
+	DECLARE_HIT_PROXY( LANDSCAPE_API );
+
+	ULandscapeSplineSegment* SplineSegment;
+	uint32 End:1;
+
+	HLandscapeSplineProxy_Tangent(class ULandscapeSplineSegment* InSplineSegment, bool InEnd) :
+		HLandscapeSplineProxy(HPP_UI),
+		SplineSegment(InSplineSegment),
+		End(InEnd)
+	{
+	}
+	LANDSCAPE_API virtual void Serialize(FArchive& Ar);
+
+	virtual EMouseCursor::Type GetMouseCursor()
+	{
+		return EMouseCursor::CardinalCross;
+	}
+};
