Build fixes for OS X.

diff --git a/Engine/Source/Runtime/Engine/Public/UnrealClient.h b/Engine/Source/Runtime/Engine/Public/UnrealClient.h
--- a/Engine/Source/Runtime/Engine/Public/UnrealClient.h
+++ b/Engine/Source/Runtime/Engine/Public/UnrealClient.h
@@ -9,7 +9,7 @@
 #include "RenderResource.h"
 #include "HitProxies.h"
 #include "InputCoreTypes.h"
-#include "EngineBaseTypes.h"
+#include "Engine/EngineBaseTypes.h"
 
 class FCanvas;
 class FViewportClient;
diff --git a/Engine/Source/Runtime/Landscape/Classes/Landscape.h b/Engine/Source/Runtime/Landscape/Classes/Landscape.h
--- a/Engine/Source/Runtime/Landscape/Classes/Landscape.h
+++ b/Engine/Source/Runtime/Landscape/Classes/Landscape.h
@@ -2,6 +2,11 @@
 
 #pragma once
 
+#include "Core.h"
+#include "CoreUObject.h"
+#include "Casts.h"
+#include "LazyObjectPtr.h"
+
 #include "LandscapeProxy.h"
 #include "LandscapeInfoMap.h"
 #include "Engine/World.h"
diff --git a/Engine/Source/Runtime/Landscape/Classes/LandscapeProxy.h b/Engine/Source/Runtime/Landscape/Classes/LandscapeProxy.h
--- a/Engine/Source/Runtime/Landscape/Classes/LandscapeProxy.h
+++ b/Engine/Source/Runtime/Landscape/Classes/LandscapeProxy.h
@@ -10,6 +10,7 @@
 
 #include "GameFramework/Actor.h"
 #include "PhysicsEngine/BodyInstance.h"
+#include "Materials/Material.h"
 #include "Materials/MaterialInstanceConstant.h"
 
 #include "LandscapeProxy.generated.h"
diff --git a/Engine/Source/Runtime/Landscape/Private/Materials/MaterialExpressionLandscapeLayerBlend.cpp b/Engine/Source/Runtime/Landscape/Private/Materials/MaterialExpressionLandscapeLayerBlend.cpp
--- a/Engine/Source/Runtime/Landscape/Private/Materials/MaterialExpressionLandscapeLayerBlend.cpp
+++ b/Engine/Source/Runtime/Landscape/Private/Materials/MaterialExpressionLandscapeLayerBlend.cpp
@@ -2,6 +2,7 @@
 
 #include "Landscape.h"
 #include "Materials/MaterialExpressionLandscapeLayerBlend.h"
+#include "EdGraph/EdGraphNode.h"
 
 #define LOCTEXT_NAMESPACE "Landscape"
 
