Landscape: Fixed non-unity builds

Also NULL to nullptr

diff --git a/Engine/Source/Runtime/Landscape/Private/Materials/MaterialExpressionLandscapeLayerBlend.cpp b/Engine/Source/Runtime/Landscape/Private/Materials/MaterialExpressionLandscapeLayerBlend.cpp
--- a/Engine/Source/Runtime/Landscape/Private/Materials/MaterialExpressionLandscapeLayerBlend.cpp
+++ b/Engine/Source/Runtime/Landscape/Private/Materials/MaterialExpressionLandscapeLayerBlend.cpp
@@ -1,11 +1,13 @@
 // Copyright 1998-2014 Epic Games, Inc. All Rights Reserved.
 
 #include "Landscape.h"
+#include "MaterialCompiler.h"
 #include "Materials/MaterialExpressionLandscapeLayerBlend.h"
 #include "EdGraph/EdGraphNode.h"
 
 #define LOCTEXT_NAMESPACE "Landscape"
 
+
 ///////////////////////////////////////////////////////////////////////////////
 // UMaterialExpressionLandscapeLayerBlend
 ///////////////////////////////////////////////////////////////////////////////
@@ -28,11 +30,13 @@ UMaterialExpressionLandscapeLayerBlend::UMaterialExpressionLandscapeLayerBlend(c
 	MenuCategories.Add(ConstructorStatics.NAME_Landscape);
 }
 
+
 FGuid& UMaterialExpressionLandscapeLayerBlend::GetParameterExpressionId()
 {
 	return ExpressionGUID;
 }
 
+
 void UMaterialExpressionLandscapeLayerBlend::Serialize(FArchive& Ar)
 {
 	Super::Serialize(Ar);
@@ -50,6 +54,7 @@ void UMaterialExpressionLandscapeLayerBlend::Serialize(FArchive& Ar)
 	}
 }
 
+
 const TArray<FExpressionInput*> UMaterialExpressionLandscapeLayerBlend::GetInputs()
 {
 	TArray<FExpressionInput*> Result;
@@ -64,6 +69,7 @@ const TArray<FExpressionInput*> UMaterialExpressionLandscapeLayerBlend::GetInput
 	return Result;
 }
 
+
 FExpressionInput* UMaterialExpressionLandscapeLayerBlend::GetInput(int32 InputIndex)
 {
 	int32 Idx = 0;
@@ -81,9 +87,10 @@ FExpressionInput* UMaterialExpressionLandscapeLayerBlend::GetInput(int32 InputIn
 			}
 		}
 	}
-	return NULL;
+	return nullptr;
 }
 
+
 FString UMaterialExpressionLandscapeLayerBlend::GetInputName(int32 InputIndex) const
 {
 	int32 Idx = 0;
@@ -104,6 +111,7 @@ FString UMaterialExpressionLandscapeLayerBlend::GetInputName(int32 InputIndex) c
 	return TEXT("");
 }
 
+
 bool UMaterialExpressionLandscapeLayerBlend::IsResultMaterialAttributes(int32 OutputIndex)
 {
 	if (ContainsInputLoop())
@@ -121,6 +129,7 @@ bool UMaterialExpressionLandscapeLayerBlend::IsResultMaterialAttributes(int32 Ou
 	return false;
 }
 
+
 int32 UMaterialExpressionLandscapeLayerBlend::Compile(class FMaterialCompiler* Compiler, int32 OutputIndex, int32 MultiplexIndex)
 {
 	// For renormalization
@@ -250,16 +259,19 @@ int32 UMaterialExpressionLandscapeLayerBlend::Compile(class FMaterialCompiler* C
 	return OutputCode;
 }
 
+
 UTexture* UMaterialExpressionLandscapeLayerBlend::GetReferencedTexture()
 {
 	return GEngine->WeightMapPlaceholderTexture;
 }
 
+
 void UMaterialExpressionLandscapeLayerBlend::GetCaption(TArray<FString>& OutCaptions) const
 {
 	OutCaptions.Add(FString(TEXT("Layer Blend")));
 }
 
+
 #if WITH_EDITOR
 
 void UMaterialExpressionLandscapeLayerBlend::PostEditChangeProperty(FPropertyChangedEvent& PropertyChangedEvent)
@@ -271,7 +283,7 @@ void UMaterialExpressionLandscapeLayerBlend::PostEditChangeProperty(FPropertyCha
 	{
 		if (Layers[LayerIdx].BlendType != LB_HeightBlend)
 		{
-			Layers[LayerIdx].HeightInput.Expression = NULL;
+			Layers[LayerIdx].HeightInput.Expression = nullptr;
 		}
 	}
 
@@ -289,6 +301,7 @@ void UMaterialExpressionLandscapeLayerBlend::PostEditChangeProperty(FPropertyCha
 }
 #endif
 
+
 void UMaterialExpressionLandscapeLayerBlend::GetAllParameterNames(TArray<FName> &OutParameterNames, TArray<FGuid> &OutParameterIds)
 {
 	for (int32 LayerIdx = 0; LayerIdx<Layers.Num(); LayerIdx++)
diff --git a/Engine/Source/Runtime/Landscape/Private/Materials/MaterialExpressionLandscapeLayerCoords.cpp b/Engine/Source/Runtime/Landscape/Private/Materials/MaterialExpressionLandscapeLayerCoords.cpp
--- a/Engine/Source/Runtime/Landscape/Private/Materials/MaterialExpressionLandscapeLayerCoords.cpp
+++ b/Engine/Source/Runtime/Landscape/Private/Materials/MaterialExpressionLandscapeLayerCoords.cpp
@@ -1,11 +1,17 @@
 // Copyright 1998-2014 Epic Games, Inc. All Rights Reserved.
 
 #include "Landscape.h"
+#include "MaterialCompiler.h"
 #include "Materials/MaterialExpressionLandscapeLayerCoords.h"
 
+
+#define LOCTEXT_NAMESPACE "Landscape"
+
+
 ///////////////////////////////////////////////////////////////////////////////
 // UMaterialExpressionLandscapeLayerCoords
 ///////////////////////////////////////////////////////////////////////////////
+
 UMaterialExpressionLandscapeLayerCoords::UMaterialExpressionLandscapeLayerCoords(const FObjectInitializer& ObjectInitializer)
 	: Super(ObjectInitializer)
 {
@@ -24,6 +30,7 @@ UMaterialExpressionLandscapeLayerCoords::UMaterialExpressionLandscapeLayerCoords
 	bCollapsed = false;
 }
 
+
 int32 UMaterialExpressionLandscapeLayerCoords::Compile(class FMaterialCompiler* Compiler, int32 OutputIndex, int32 MultiplexIndex)
 {
 	switch (CustomUVType)
@@ -88,7 +95,11 @@ int32 UMaterialExpressionLandscapeLayerCoords::Compile(class FMaterialCompiler*
 	return ResultUV;
 }
 
+
 void UMaterialExpressionLandscapeLayerCoords::GetCaption(TArray<FString>& OutCaptions) const
 {
 	OutCaptions.Add(FString(TEXT("LandscapeCoords")));
 }
+
+
+#undef LOCTEXT_NAMESPACE
diff --git a/Engine/Source/Runtime/Landscape/Private/Materials/MaterialExpressionLandscapeLayerSample.cpp b/Engine/Source/Runtime/Landscape/Private/Materials/MaterialExpressionLandscapeLayerSample.cpp
--- a/Engine/Source/Runtime/Landscape/Private/Materials/MaterialExpressionLandscapeLayerSample.cpp
+++ b/Engine/Source/Runtime/Landscape/Private/Materials/MaterialExpressionLandscapeLayerSample.cpp
@@ -1,11 +1,17 @@
 // Copyright 1998-2014 Epic Games, Inc. All Rights Reserved.
 
 #include "Landscape.h"
+#include "MaterialCompiler.h"
 #include "Materials/MaterialExpressionLandscapeLayerSample.h"
 
+
+#define LOCTEXT_NAMESPACE "Landscape"
+
+
 ///////////////////////////////////////////////////////////////////////////////
 // UMaterialExpressionLandscapeLayerSample
 ///////////////////////////////////////////////////////////////////////////////
+
 UMaterialExpressionLandscapeLayerSample::UMaterialExpressionLandscapeLayerSample(const FObjectInitializer& ObjectInitializer)
 	: Super(ObjectInitializer)
 {
@@ -24,11 +30,13 @@ UMaterialExpressionLandscapeLayerSample::UMaterialExpressionLandscapeLayerSample
 	MenuCategories.Add(ConstructorStatics.NAME_Landscape);
 }
 
+
 FGuid& UMaterialExpressionLandscapeLayerSample::GetParameterExpressionId()
 {
 	return ExpressionGUID;
 }
 
+
 int32 UMaterialExpressionLandscapeLayerSample::Compile(class FMaterialCompiler* Compiler, int32 OutputIndex, int32 MultiplexIndex)
 {
 	const int32 WeightCode = Compiler->StaticTerrainLayerWeight(ParameterName, Compiler->Constant(PreviewWeight));
@@ -36,16 +44,19 @@ int32 UMaterialExpressionLandscapeLayerSample::Compile(class FMaterialCompiler*
 	return WeightCode;
 }
 
+
 UTexture* UMaterialExpressionLandscapeLayerSample::GetReferencedTexture()
 {
 	return GEngine->WeightMapPlaceholderTexture;
 }
 
+
 void UMaterialExpressionLandscapeLayerSample::GetCaption(TArray<FString>& OutCaptions) const
 {
 	OutCaptions.Add(FString::Printf(TEXT("Sample '%s'"), *ParameterName.ToString()));
 }
 
+
 void UMaterialExpressionLandscapeLayerSample::GetAllParameterNames(TArray<FName> &OutParameterNames, TArray<FGuid> &OutParameterIds)
 {
 	if (!OutParameterNames.Contains(ParameterName))
@@ -54,3 +65,6 @@ void UMaterialExpressionLandscapeLayerSample::GetAllParameterNames(TArray<FName>
 		OutParameterIds.Add(ExpressionGUID);
 	}
 }
+
+
+#undef LOCTEXT_NAMESPACE
diff --git a/Engine/Source/Runtime/Landscape/Private/Materials/MaterialExpressionLandscapeLayerSwitch.cpp b/Engine/Source/Runtime/Landscape/Private/Materials/MaterialExpressionLandscapeLayerSwitch.cpp
--- a/Engine/Source/Runtime/Landscape/Private/Materials/MaterialExpressionLandscapeLayerSwitch.cpp
+++ b/Engine/Source/Runtime/Landscape/Private/Materials/MaterialExpressionLandscapeLayerSwitch.cpp
@@ -1,11 +1,17 @@
 // Copyright 1998-2014 Epic Games, Inc. All Rights Reserved.
 
 #include "Landscape.h"
+#include "MaterialCompiler.h"
 #include "Materials/MaterialExpressionLandscapeLayerSwitch.h"
 
+
+#define LOCTEXT_NAMESPACE "Landscape"
+
+
 ///////////////////////////////////////////////////////////////////////////////
 // UMaterialExpressionLandscapeLayerSwitch
 ///////////////////////////////////////////////////////////////////////////////
+
 UMaterialExpressionLandscapeLayerSwitch::UMaterialExpressionLandscapeLayerSwitch(const FObjectInitializer& ObjectInitializer)
 	: Super(ObjectInitializer)
 {
@@ -26,6 +32,7 @@ UMaterialExpressionLandscapeLayerSwitch::UMaterialExpressionLandscapeLayerSwitch
 	bCollapsed = false;
 }
 
+
 bool UMaterialExpressionLandscapeLayerSwitch::IsResultMaterialAttributes(int32 OutputIndex)
 {
 	if (ContainsInputLoop())
@@ -33,11 +40,12 @@ bool UMaterialExpressionLandscapeLayerSwitch::IsResultMaterialAttributes(int32 O
 		// If there is a loop anywhere in this expression's inputs then we can't risk checking them
 		return false;
 	}
-	bool bLayerUsedIsMaterialAttributes = LayerUsed.Expression != NULL && LayerUsed.Expression->IsResultMaterialAttributes(LayerUsed.OutputIndex);
-	bool bLayerNotUsedIsMaterialAttributes = LayerNotUsed.Expression != NULL && LayerNotUsed.Expression->IsResultMaterialAttributes(LayerNotUsed.OutputIndex);
+	bool bLayerUsedIsMaterialAttributes = LayerUsed.Expression != nullptr && LayerUsed.Expression->IsResultMaterialAttributes(LayerUsed.OutputIndex);
+	bool bLayerNotUsedIsMaterialAttributes = LayerNotUsed.Expression != nullptr && LayerNotUsed.Expression->IsResultMaterialAttributes(LayerNotUsed.OutputIndex);
 	return bLayerUsedIsMaterialAttributes || bLayerNotUsedIsMaterialAttributes;
 }
 
+
 int32 UMaterialExpressionLandscapeLayerSwitch::Compile(class FMaterialCompiler* Compiler, int32 OutputIndex, int32 MultiplexIndex)
 {
 	const int32 WeightCode = Compiler->StaticTerrainLayerWeight(
@@ -56,7 +64,7 @@ int32 UMaterialExpressionLandscapeLayerSwitch::Compile(class FMaterialCompiler*
 	}
 
 	if (ReturnCode != INDEX_NONE && //If we've already failed for some other reason don't bother with this check. It could have been the reentrant check causing this to loop infinitely!
-		LayerUsed.Expression != NULL && LayerNotUsed.Expression != NULL &&
+		LayerUsed.Expression != nullptr && LayerNotUsed.Expression != nullptr &&
 		LayerUsed.Expression->IsResultMaterialAttributes(LayerUsed.OutputIndex) != LayerNotUsed.Expression->IsResultMaterialAttributes(LayerNotUsed.OutputIndex))
 	{
 		Compiler->Error(TEXT("Cannot mix MaterialAttributes and non MaterialAttributes nodes"));
@@ -65,17 +73,20 @@ int32 UMaterialExpressionLandscapeLayerSwitch::Compile(class FMaterialCompiler*
 	return ReturnCode;
 }
 
+
 UTexture* UMaterialExpressionLandscapeLayerSwitch::GetReferencedTexture()
 {
 	return GEngine->WeightMapPlaceholderTexture;
 }
 
+
 void UMaterialExpressionLandscapeLayerSwitch::GetCaption(TArray<FString>& OutCaptions) const
 {
 	OutCaptions.Add(TEXT("Layer Switch"));
 	OutCaptions.Add(FString::Printf(TEXT("'%s'"), *ParameterName.ToString()));
 }
 
+
 void UMaterialExpressionLandscapeLayerSwitch::Serialize(FArchive& Ar)
 {
 	Super::Serialize(Ar);
@@ -86,6 +97,7 @@ void UMaterialExpressionLandscapeLayerSwitch::Serialize(FArchive& Ar)
 	}
 }
 
+
 void UMaterialExpressionLandscapeLayerSwitch::PostLoad()
 {
 	Super::PostLoad();
@@ -96,6 +108,7 @@ void UMaterialExpressionLandscapeLayerSwitch::PostLoad()
 	}
 }
 
+
 FGuid& UMaterialExpressionLandscapeLayerSwitch::GetParameterExpressionId()
 {
 	return ExpressionGUID;
@@ -112,3 +125,6 @@ void UMaterialExpressionLandscapeLayerSwitch::GetAllParameterNames(TArray<FName>
 		OutParameterIds.Add(ExpressionGUID);
 	}
 }
+
+
+#undef LOCTEXT_NAMESPACE
diff --git a/Engine/Source/Runtime/Landscape/Private/Materials/MaterialExpressionLandscapeLayerWeight.cpp b/Engine/Source/Runtime/Landscape/Private/Materials/MaterialExpressionLandscapeLayerWeight.cpp
--- a/Engine/Source/Runtime/Landscape/Private/Materials/MaterialExpressionLandscapeLayerWeight.cpp
+++ b/Engine/Source/Runtime/Landscape/Private/Materials/MaterialExpressionLandscapeLayerWeight.cpp
@@ -1,11 +1,17 @@
 // Copyright 1998-2014 Epic Games, Inc. All Rights Reserved.
 
 #include "Landscape.h"
+#include "MaterialCompiler.h"
 #include "Materials/MaterialExpressionLandscapeLayerWeight.h"
 
+
+#define LOCTEXT_NAMESPACE "Landscape"
+
+
 ///////////////////////////////////////////////////////////////////////////////
 // UMaterialExpressionLandscapeLayerWeight
 ///////////////////////////////////////////////////////////////////////////////
+
 UMaterialExpressionLandscapeLayerWeight::UMaterialExpressionLandscapeLayerWeight(const FObjectInitializer& ObjectInitializer)
 	: Super(ObjectInitializer)
 {
@@ -26,11 +32,13 @@ UMaterialExpressionLandscapeLayerWeight::UMaterialExpressionLandscapeLayerWeight
 	ConstBase = FVector(0.f, 0.f, 0.f);
 }
 
+
 FGuid& UMaterialExpressionLandscapeLayerWeight::GetParameterExpressionId()
 {
 	return ExpressionGUID;
 }
 
+
 void UMaterialExpressionLandscapeLayerWeight::PostLoad()
 {
 	Super::PostLoad();
@@ -41,6 +49,7 @@ void UMaterialExpressionLandscapeLayerWeight::PostLoad()
 	}
 }
 
+
 bool UMaterialExpressionLandscapeLayerWeight::IsResultMaterialAttributes(int32 OutputIndex)
 {
 	if (ContainsInputLoop())
@@ -48,11 +57,12 @@ bool UMaterialExpressionLandscapeLayerWeight::IsResultMaterialAttributes(int32 O
 		// If there is a loop anywhere in this expression's inputs then we can't risk checking them
 		return false;
 	}
-	bool bLayerIsMaterialAttributes = Layer.Expression != NULL && Layer.Expression->IsResultMaterialAttributes(Layer.OutputIndex);
-	bool bBaseIsMaterialAttributes = Base.Expression != NULL && Base.Expression->IsResultMaterialAttributes(Base.OutputIndex);
+	bool bLayerIsMaterialAttributes = Layer.Expression != nullptr && Layer.Expression->IsResultMaterialAttributes(Layer.OutputIndex);
+	bool bBaseIsMaterialAttributes = Base.Expression != nullptr && Base.Expression->IsResultMaterialAttributes(Base.OutputIndex);
 	return bLayerIsMaterialAttributes || bBaseIsMaterialAttributes;
 }
 
+
 int32 UMaterialExpressionLandscapeLayerWeight::Compile(class FMaterialCompiler* Compiler, int32 OutputIndex, int32 MultiplexIndex)
 {
 	const int32 BaseCode = Base.Expression ? Base.Compile(Compiler, MultiplexIndex) : Compiler->Constant3(ConstBase.X, ConstBase.Y, ConstBase.Z);
@@ -70,7 +80,7 @@ int32 UMaterialExpressionLandscapeLayerWeight::Compile(class FMaterialCompiler*
 	}
 
 	if (ReturnCode != INDEX_NONE && //If we've already failed for some other reason don't bother with this check. It could have been the reentrant check causing this to loop infinitely!
-		Layer.Expression != NULL && Base.Expression != NULL &&
+		Layer.Expression != nullptr && Base.Expression != nullptr &&
 		Layer.Expression->IsResultMaterialAttributes(Layer.OutputIndex) != Base.Expression->IsResultMaterialAttributes(Base.OutputIndex))
 	{
 		Compiler->Error(TEXT("Cannot mix MaterialAttributes and non MaterialAttributes nodes"));
@@ -79,16 +89,19 @@ int32 UMaterialExpressionLandscapeLayerWeight::Compile(class FMaterialCompiler*
 	return ReturnCode;
 }
 
+
 UTexture* UMaterialExpressionLandscapeLayerWeight::GetReferencedTexture()
 {
 	return GEngine->WeightMapPlaceholderTexture;
 }
 
+
 void UMaterialExpressionLandscapeLayerWeight::GetCaption(TArray<FString>& OutCaptions) const
 {
 	OutCaptions.Add(FString::Printf(TEXT("Layer '%s'"), *ParameterName.ToString()));
 }
 
+
 void UMaterialExpressionLandscapeLayerWeight::GetAllParameterNames(TArray<FName> &OutParameterNames, TArray<FGuid> &OutParameterIds)
 {
 	int32 CurrentSize = OutParameterNames.Num();
@@ -98,4 +111,7 @@ void UMaterialExpressionLandscapeLayerWeight::GetAllParameterNames(TArray<FName>
 	{
 		OutParameterIds.Add(ExpressionGUID);
 	}
-}
\ No newline at end of file
+}
+
+
+#undef LOCTEXT_NAMESPACE
diff --git a/Engine/Source/Runtime/Landscape/Private/Materials/MaterialExpressionLandscapeVisibilityMask.cpp b/Engine/Source/Runtime/Landscape/Private/Materials/MaterialExpressionLandscapeVisibilityMask.cpp
--- a/Engine/Source/Runtime/Landscape/Private/Materials/MaterialExpressionLandscapeVisibilityMask.cpp
+++ b/Engine/Source/Runtime/Landscape/Private/Materials/MaterialExpressionLandscapeVisibilityMask.cpp
@@ -1,8 +1,13 @@
 // Copyright 1998-2014 Epic Games, Inc. All Rights Reserved.
 
 #include "Landscape.h"
+#include "MaterialCompiler.h"
 #include "Materials/MaterialExpressionLandscapeVisibilityMask.h"
 
+
+#define LOCTEXT_NAMESPACE "Landscape"
+
+
 ///////////////////////////////////////////////////////////////////////////////
 // UMaterialExpressionLandscapeVisibilityMask
 ///////////////////////////////////////////////////////////////////////////////
@@ -27,21 +32,25 @@ UMaterialExpressionLandscapeVisibilityMask::UMaterialExpressionLandscapeVisibili
 	MenuCategories.Add(ConstructorStatics.NAME_Landscape);
 }
 
+
 FGuid& UMaterialExpressionLandscapeVisibilityMask::GetParameterExpressionId()
 {
 	return ExpressionGUID;
 }
 
+
 int32 UMaterialExpressionLandscapeVisibilityMask::Compile(class FMaterialCompiler* Compiler, int32 OutputIndex, int32 MultiplexIndex)
 {
 	return Compiler->Sub(Compiler->Constant(1.f), Compiler->StaticTerrainLayerWeight(ParameterName, Compiler->Constant(0.f)));
 }
 
+
 UTexture* UMaterialExpressionLandscapeVisibilityMask::GetReferencedTexture()
 {
 	return GEngine->WeightMapPlaceholderTexture;
 }
 
+
 void UMaterialExpressionLandscapeVisibilityMask::GetAllParameterNames(TArray<FName> &OutParameterNames, TArray<FGuid> &OutParameterIds)
 {
 	int32 CurrentSize = OutParameterNames.Num();
@@ -53,7 +62,11 @@ void UMaterialExpressionLandscapeVisibilityMask::GetAllParameterNames(TArray<FNa
 	}
 }
 
+
 void UMaterialExpressionLandscapeVisibilityMask::GetCaption(TArray<FString>& OutCaptions) const
 {
 	OutCaptions.Add(FString(TEXT("Landscape Visibility Mask")));
-}
\ No newline at end of file
+}
+
+
+#undef LOCTEXT_NAMESPACE
