Fix Landscape Ramp tool corruption if the bounds of the ramp go outside the landscape bounds
Made some minor changes to the spline tool's apply to landscape function to prevent allocating unnecessary memory for spline segments outside of the landscape bounds

diff --git a/Engine/Source/Editor/LandscapeEditor/Private/LandscapeEdModeRampTool.cpp b/Engine/Source/Editor/LandscapeEditor/Private/LandscapeEdModeRampTool.cpp
--- a/Engine/Source/Editor/LandscapeEditor/Private/LandscapeEdModeRampTool.cpp
+++ b/Engine/Source/Editor/LandscapeEditor/Private/LandscapeEdModeRampTool.cpp
@@ -450,10 +450,11 @@ public:
 	{
 		FScopedTransaction Transaction(LOCTEXT("Ramp_Apply", "Landscape Editing: Add ramp"));
 
-		const ALandscapeProxy* LandscapeProxy = EdMode->CurrentToolTarget.LandscapeInfo->GetLandscapeProxy();
+		const ULandscapeInfo* LandscapeInfo = EdMode->CurrentToolTarget.LandscapeInfo.Get();
+		const ALandscapeProxy* LandscapeProxy = LandscapeInfo->GetLandscapeProxy();
 		const FTransform LandscapeToWorld = LandscapeProxy->LandscapeActorToWorld();
 
-		const FVector2D Side = FVector2D(FVector::CrossProduct(Points[1] - Points[0], FVector(0, 0, 1))).SafeNormal();
+		const FVector2D Side = FVector2D(FVector::CrossProduct(Points[1] - Points[0], FVector(0,0,1))).SafeNormal();
 		const FVector2D InnerSide = Side * (EdMode->UISettings->RampWidth * 0.5f * (1 - EdMode->UISettings->RampSideFalloff)) / LandscapeToWorld.GetScale3D().X;
 		const FVector2D OuterSide = Side * (EdMode->UISettings->RampWidth * 0.5f) / LandscapeToWorld.GetScale3D().X;
 
@@ -478,6 +479,24 @@ public:
 		int32 MaxX = FMath::FloorToInt(FMath::Max(FMath::Max(OuterVerts[0][0].X, OuterVerts[0][1].X), FMath::Max(OuterVerts[1][0].X, OuterVerts[1][1].X))) + 1;
 		int32 MaxY = FMath::FloorToInt(FMath::Max(FMath::Max(OuterVerts[0][0].Y, OuterVerts[0][1].Y), FMath::Max(OuterVerts[1][0].Y, OuterVerts[1][1].Y))) + 1;
 
+		// I'd dearly love to use FIntRect in this code, but Landscape works with "Inclusive Max" and FIntRect is "Exclusive Max"
+		int32 LandscapeMinX, LandscapeMinY, LandscapeMaxX, LandscapeMaxY;
+		if (!LandscapeInfo->GetLandscapeExtent(LandscapeMinX, LandscapeMinY, LandscapeMaxX, LandscapeMaxY))
+		{
+			return;
+		}
+
+		MinX = FMath::Max(MinX, LandscapeMinX);
+		MinY = FMath::Max(MinY, LandscapeMinY);
+		MaxX = FMath::Min(MaxX, LandscapeMaxX);
+		MaxY = FMath::Min(MaxY, LandscapeMaxY);
+
+		if (MinX > MaxX || MinY > MaxY)
+		{
+			// The bounds don't intersect any data, so we skip applying the ramp entirely
+			return;
+		}
+
 		FLandscapeEditDataInterface LandscapeEdit(EdMode->CurrentToolTarget.LandscapeInfo.Get());
 
 		// Heights raster
@@ -496,11 +515,11 @@ public:
 
 			if (ValidMinX > ValidMaxX || ValidMinY > ValidMaxY)
 			{
-				// The bounds don't intersect any data, so we skip it entirely
+				// The bounds don't intersect any data, so we skip applying the ramp entirely
 				return;
 			}
 
-			//ShrinkData(Data, MinX, MinY, MaxX, MaxY, ValidMinX, ValidMinY, ValidMaxX, ValidMaxY);
+			FLandscapeEditDataInterface::ShrinkData(Data, MinX, MinY, MaxX, MaxY, ValidMinX, ValidMinY, ValidMaxX, ValidMaxY);
 
 			MinX = ValidMinX;
 			MinY = ValidMinY;
diff --git a/Engine/Source/Runtime/Landscape/Classes/LandscapeInfo.h b/Engine/Source/Runtime/Landscape/Classes/LandscapeInfo.h
--- a/Engine/Source/Runtime/Landscape/Classes/LandscapeInfo.h
+++ b/Engine/Source/Runtime/Landscape/Classes/LandscapeInfo.h
@@ -3,6 +3,15 @@
 #pragma once
 #include "LandscapeInfo.generated.h"
 
+// Forward declarations
+class ALandscapeProxy;
+class ALandscape;
+class ULandscapeComponent;
+class ULandscapeLayerInfoObject;
+class UMaterialInstanceConstant;
+struct FLandscapeEditorLayerSettings;
+struct FLandscapeDataInterface;
+
 /** Structure storing Collision for LandscapeComponent Add */
 USTRUCT()
 struct FLandscapeAddCollision
@@ -31,17 +40,17 @@ struct FLandscapeInfoLayerSettings
 	GENERATED_USTRUCT_BODY()
 
 	UPROPERTY()
-	class ULandscapeLayerInfoObject* LayerInfoObj;
+	ULandscapeLayerInfoObject* LayerInfoObj;
 
 	UPROPERTY()
 	FName LayerName;
 
 #if WITH_EDITORONLY_DATA
 	UPROPERTY(transient)
-	class UMaterialInstanceConstant* ThumbnailMIC;
+	UMaterialInstanceConstant* ThumbnailMIC;
 
 	UPROPERTY()
-	class ALandscapeProxy* Owner;
+	ALandscapeProxy* Owner;
 
 	UPROPERTY(transient)
 	int32 DebugColorChannel;
@@ -62,9 +71,9 @@ struct FLandscapeInfoLayerSettings
 	{
 	}
 
-	LANDSCAPE_API FLandscapeInfoLayerSettings(ULandscapeLayerInfoObject* InLayerInfo, class ALandscapeProxy* InProxy);
+	LANDSCAPE_API FLandscapeInfoLayerSettings(ULandscapeLayerInfoObject* InLayerInfo, ALandscapeProxy* InProxy);
 
-	FLandscapeInfoLayerSettings(FName InPlaceholderLayerName, class ALandscapeProxy* InProxy)
+	FLandscapeInfoLayerSettings(FName InPlaceholderLayerName, ALandscapeProxy* InProxy)
 		: LayerInfoObj(NULL)
 		, LayerName(InPlaceholderLayerName)
 #if WITH_EDITORONLY_DATA
@@ -79,7 +88,7 @@ struct FLandscapeInfoLayerSettings
 	LANDSCAPE_API FName GetLayerName() const;
 
 #if WITH_EDITORONLY_DATA
-	LANDSCAPE_API struct FLandscapeEditorLayerSettings& GetEditorSettings() const;
+	LANDSCAPE_API FLandscapeEditorLayerSettings& GetEditorSettings() const;
 #endif
 };
 
@@ -89,7 +98,7 @@ class ULandscapeInfo : public UObject
 	GENERATED_UCLASS_BODY()
 
 	UPROPERTY()
-	TLazyObjectPtr<class ALandscape> LandscapeActor;
+	TLazyObjectPtr<ALandscape> LandscapeActor;
 
 	UPROPERTY()
 	FGuid LandscapeGuid;
@@ -111,23 +120,23 @@ class ULandscapeInfo : public UObject
 	uint32 bIsValid:1;
 
 	UPROPERTY()
-	TArray<struct FLandscapeInfoLayerSettings> Layers;
+	TArray<FLandscapeInfoLayerSettings> Layers;
 #endif // WITH_EDITORONLY_DATA
 
 public:
-	struct FLandscapeDataInterface* DataInterface;
+	FLandscapeDataInterface* DataInterface;
 
 	/** Map of the offsets (in component space) to the component. Valid in editor only. */
-	TMap<FIntPoint,class ULandscapeComponent*> XYtoComponentMap;
+	TMap<FIntPoint, ULandscapeComponent*> XYtoComponentMap;
 
 	/** Map of the offsets to the newly added collision components. Only available near valid LandscapeComponents. Valid in editor only. */
 	TMap<FIntPoint, FLandscapeAddCollision> XYtoAddCollisionMap;
 
-	TSet<class ALandscapeProxy*> Proxies;
+	TSet<ALandscapeProxy*> Proxies;
 
 private:
-	TSet<class ULandscapeComponent*> SelectedComponents;
-	TSet<class ULandscapeComponent*> SelectedRegionComponents;
+	TSet<ULandscapeComponent*> SelectedComponents;
+	TSet<ULandscapeComponent*> SelectedRegionComponents;
 
 public:
 
@@ -145,17 +154,17 @@ public:
 
 	// @todo document 
 	// all below.
-	LANDSCAPE_API struct FLandscapeDataInterface* GetDataInterface();
+	LANDSCAPE_API FLandscapeDataInterface* GetDataInterface();
 
 	LANDSCAPE_API void GetComponentsInRegion(int32 X1, int32 Y1, int32 X2, int32 Y2, TSet<ULandscapeComponent*>& OutComponents);
-	LANDSCAPE_API bool GetLandscapeExtent(int32& MinX, int32& MinY, int32& MaxX, int32& MaxY);
+	LANDSCAPE_API bool GetLandscapeExtent(int32& MinX, int32& MinY, int32& MaxX, int32& MaxY) const;
 	LANDSCAPE_API void Export(const TArray<ULandscapeLayerInfoObject*>& LayerInfos, const TArray<FString>& Filenames);
 	LANDSCAPE_API void ExportHeightmap(const FString& Filename);
 	LANDSCAPE_API void ExportLayer(ULandscapeLayerInfoObject* LayerInfo, const FString& Filename);
 	LANDSCAPE_API bool ApplySplines(bool bOnlySelected);
 	bool ApplySplinesInternal(bool bOnlySelected, ALandscapeProxy* Landscape);
 
-	LANDSCAPE_API bool GetSelectedExtent(int32& MinX, int32& MinY, int32& MaxX, int32& MaxY);
+	LANDSCAPE_API bool GetSelectedExtent(int32& MinX, int32& MinY, int32& MaxX, int32& MaxY) const;
 	FVector GetLandscapeCenterPos(float& LengthZ, int32 MinX = MAX_int32, int32 MinY = MAX_int32, int32 MaxX = MIN_int32, int32 MaxY = MIN_int32);
 	LANDSCAPE_API bool IsValidPosition(int32 X, int32 Y);
 	LANDSCAPE_API void DeleteLayer(ULandscapeLayerInfoObject* LayerInfo);
@@ -163,9 +172,9 @@ public:
 
 	LANDSCAPE_API void UpdateDebugColorMaterial();
 
-	LANDSCAPE_API TSet<class ULandscapeComponent*> GetSelectedComponents() const;
-	LANDSCAPE_API TSet<class ULandscapeComponent*> GetSelectedRegionComponents() const;
-	LANDSCAPE_API void UpdateSelectedComponents(TSet<class ULandscapeComponent*>& NewComponents, bool bIsComponentwise = true);
+	LANDSCAPE_API TSet<ULandscapeComponent*> GetSelectedComponents() const;
+	LANDSCAPE_API TSet<ULandscapeComponent*> GetSelectedRegionComponents() const;
+	LANDSCAPE_API void UpdateSelectedComponents(TSet<ULandscapeComponent*>& NewComponents, bool bIsComponentwise = true);
 	LANDSCAPE_API void SortSelectedComponents();
 	LANDSCAPE_API void ClearSelectedRegion(bool bIsComponentwise = true);
 
@@ -175,10 +184,10 @@ public:
 	LANDSCAPE_API FLandscapeEditorLayerSettings& GetLayerEditorSettings(ULandscapeLayerInfoObject* LayerInfo) const;
 	LANDSCAPE_API void CreateLayerEditorSettingsFor(ULandscapeLayerInfoObject* LayerInfo);
 
-	LANDSCAPE_API ULandscapeLayerInfoObject* GetLayerInfoByName(FName LayerName, class ALandscapeProxy* Owner = NULL) const;
-	LANDSCAPE_API int32 GetLayerInfoIndex(FName LayerName, class ALandscapeProxy* Owner = NULL) const;
-	LANDSCAPE_API int32 GetLayerInfoIndex(ULandscapeLayerInfoObject* LayerInfo, class ALandscapeProxy* Owner = NULL) const;
-	LANDSCAPE_API bool UpdateLayerInfoMap(class ALandscapeProxy* Proxy = NULL, bool bInvalidate = false);
+	LANDSCAPE_API ULandscapeLayerInfoObject* GetLayerInfoByName(FName LayerName, ALandscapeProxy* Owner = NULL) const;
+	LANDSCAPE_API int32 GetLayerInfoIndex(FName LayerName, ALandscapeProxy* Owner = NULL) const;
+	LANDSCAPE_API int32 GetLayerInfoIndex(ULandscapeLayerInfoObject* LayerInfo, ALandscapeProxy* Owner = NULL) const;
+	LANDSCAPE_API bool UpdateLayerInfoMap(ALandscapeProxy* Proxy = NULL, bool bInvalidate = false);
 
 	/** 
 	 *  Returns landscape which is spawned in the current level that was previously added to this landscape info object
diff --git a/Engine/Source/Runtime/Landscape/Private/LandscapeEdit.cpp b/Engine/Source/Runtime/Landscape/Private/LandscapeEdit.cpp
--- a/Engine/Source/Runtime/Landscape/Private/LandscapeEdit.cpp
+++ b/Engine/Source/Runtime/Landscape/Private/LandscapeEdit.cpp
@@ -2418,7 +2418,7 @@ bool ALandscape::HasAllComponent()
 	return false;
 }
 
-bool ULandscapeInfo::GetLandscapeExtent(int32& MinX, int32& MinY, int32& MaxX, int32& MaxY)
+bool ULandscapeInfo::GetLandscapeExtent(int32& MinX, int32& MinY, int32& MaxX, int32& MaxY) const
 {
 	MinX = MAX_int32;
 	MinY = MAX_int32;
@@ -2426,22 +2426,22 @@ bool ULandscapeInfo::GetLandscapeExtent(int32& MinX, int32& MinY, int32& MaxX, i
 	MaxY = MIN_int32;
 
 	// Find range of entire landscape
-	for (auto It = XYtoComponentMap.CreateIterator(); It; ++It)
+	for (auto& XYComponentPair : XYtoComponentMap)
 	{
-		ULandscapeComponent* Comp = It.Value();
+		const ULandscapeComponent* Comp = XYComponentPair.Value;
 		Comp->GetComponentExtent(MinX, MinY, MaxX, MaxY);
 	}
 	return (MinX != MAX_int32);
 }
 
-bool ULandscapeInfo::GetSelectedExtent(int32& MinX, int32& MinY, int32& MaxX, int32& MaxY)
+bool ULandscapeInfo::GetSelectedExtent(int32& MinX, int32& MinY, int32& MaxX, int32& MaxY) const
 {
 	MinX = MinY = MAX_int32;
 	MaxX = MaxY = MIN_int32;
-	for (auto It = SelectedRegion.CreateIterator(); It; ++It)
+	for (auto& SelectedPointPair : SelectedRegion)
 	{
 		int32 X, Y;
-		ALandscape::UnpackKey(It.Key(), X, Y);
+		ALandscape::UnpackKey(SelectedPointPair.Key, X, Y);
 		if (MinX > X) MinX = X;
 		if (MaxX < X) MaxX = X;
 		if (MinY > Y) MinY = Y;
@@ -2452,9 +2452,8 @@ bool ULandscapeInfo::GetSelectedExtent(int32& MinX, int32& MinY, int32& MaxX, in
 		return true;
 	}
 	// if SelectedRegion is empty, try SelectedComponents
-	for (TSet<ULandscapeComponent*>::TIterator It(SelectedComponents); It; ++It)
+	for (const ULandscapeComponent* Comp : SelectedComponents)
 	{
-		ULandscapeComponent* Comp = *It;
 		Comp->GetComponentExtent(MinX, MinY, MaxX, MaxY);
 	}
 	return MinX != MAX_int32;
@@ -2761,33 +2760,6 @@ private:
 	int32 MinX, MinY, MaxX, MaxY;
 };
 
-template<typename T>
-void ShrinkData(TArray<T>& Data, int32 OldMinX, int32 OldMinY, int32 OldMaxX, int32 OldMaxY, int32 NewMinX, int32 NewMinY, int32 NewMaxX, int32 NewMaxY)
-{
-	checkSlow(OldMinX <= OldMaxX && OldMinY <= OldMaxY);
-	checkSlow(NewMinX >= OldMinX && NewMaxX <= OldMaxX);
-	checkSlow(NewMinY >= OldMinY && NewMaxY <= OldMaxY);
-
-	if (NewMinX != OldMinX || NewMinY != OldMinY ||
-		NewMaxX != OldMaxX || NewMaxY != OldMaxY)
-	{
-		// if only the MaxY changes we don't need to do the moving, only the truncate
-		if (NewMinX != OldMinX || NewMinY != OldMinY || NewMaxX != OldMaxX)
-		{
-			for (int32 DestY = 0, SrcY = NewMinY - OldMinY; DestY <= NewMaxY - NewMinY; DestY++, SrcY++)
-			{
-				UE_LOG(LogLandscape, Warning, TEXT("Dest: %d, %d = %d Src: %d, %d = %d Width = %d"), 0, DestY, DestY * (1 + NewMaxX - NewMinX), NewMinX - OldMinX, SrcY, SrcY * (1 + OldMaxX - OldMinX) + NewMinX - OldMinX, (1 + NewMaxX - NewMinX));
-				T* DestData = &Data[DestY * (1 + NewMaxX - NewMinX)];
-				const T* SrcData = &Data[SrcY * (1 + OldMaxX - OldMinX) + NewMinX - OldMinX];
-				FMemory::Memmove(DestData, SrcData, (1 + NewMaxX - NewMinX) * sizeof(T));
-			}
-		}
-
-		const int32 NewSize = (1 + NewMaxY - NewMinY) * (1 + NewMaxX - NewMinX);
-		Data.RemoveAt(NewSize, Data.Num() - NewSize);
-	}
-}
-
 bool ULandscapeInfo::ApplySplines(bool bOnlySelected)
 {
 	bool bResult = false;
@@ -2819,6 +2791,13 @@ bool ULandscapeInfo::ApplySplinesInternal(bool bOnlySelected, ALandscapeProxy* L
 	FLandscapeEditDataInterface LandscapeEdit(this);
 	TSet<ULandscapeComponent*> ModifiedComponents;
 
+	// I'd dearly love to use FIntRect in this code, but Landscape works with "Inclusive Max" and FIntRect is "Exclusive Max"
+	int32 LandscapeMinX, LandscapeMinY, LandscapeMaxX, LandscapeMaxY;
+	if (!GetLandscapeExtent(LandscapeMinX, LandscapeMinY, LandscapeMaxX, LandscapeMaxY))
+	{
+		return false;
+	}
+
 	for (auto It = Landscape->SplineComponent->ControlPoints.CreateConstIterator(); It; ++It)
 	{
 		const ULandscapeSplineControlPoint* ControlPoint = *It;
@@ -2833,13 +2812,24 @@ bool ULandscapeInfo::ApplySplinesInternal(bool bOnlySelected, ALandscapeProxy* L
 			continue;
 		}
 
-		FBox LandscapeBounds = ControlPoint->GetBounds();
-		LandscapeBounds = LandscapeBounds.TransformBy(SplineToLandscape.ToMatrixWithScale());
+		FBox ControlPointBounds = ControlPoint->GetBounds();
+		ControlPointBounds = ControlPointBounds.TransformBy(SplineToLandscape.ToMatrixWithScale());
+
+		int32 MinX = FMath::CeilToInt(ControlPointBounds.Min.X);
+		int32 MinY = FMath::CeilToInt(ControlPointBounds.Min.Y);
+		int32 MaxX = FMath::FloorToInt(ControlPointBounds.Max.X);
+		int32 MaxY = FMath::FloorToInt(ControlPointBounds.Max.Y);
 
-		int32 MinX = FMath::CeilToInt(LandscapeBounds.Min.X);
-		int32 MinY = FMath::CeilToInt(LandscapeBounds.Min.Y);
-		int32 MaxX = FMath::FloorToInt(LandscapeBounds.Max.X);
-		int32 MaxY = FMath::FloorToInt(LandscapeBounds.Max.Y);
+		MinX = FMath::Max(MinX, LandscapeMinX);
+		MinY = FMath::Max(MinY, LandscapeMinY);
+		MaxX = FMath::Min(MaxX, LandscapeMaxX);
+		MaxY = FMath::Min(MaxY, LandscapeMaxY);
+
+		if (MinX > MaxX || MinY > MaxY)
+		{
+			// The control point's bounds don't intersect any data, so we skip it entirely
+			continue;
+		}
 
 		TArray<FLandscapeSplineInterpPoint> Points = ControlPoint->GetPoints();
 		for (int32 j = 0; j < Points.Num(); j++)
@@ -2876,7 +2866,7 @@ bool ULandscapeInfo::ApplySplinesInternal(bool bOnlySelected, ALandscapeProxy* L
 				continue;
 			}
 
-			ShrinkData(Data, MinX, MinY, MaxX, MaxY, ValidMinX, ValidMinY, ValidMaxX, ValidMaxY);
+			FLandscapeEditDataInterface::ShrinkData(Data, MinX, MinY, MaxX, MaxY, ValidMinX, ValidMinY, ValidMaxX, ValidMaxY);
 
 			MinX = ValidMinX;
 			MinY = ValidMinY;
@@ -2935,7 +2925,7 @@ bool ULandscapeInfo::ApplySplinesInternal(bool bOnlySelected, ALandscapeProxy* L
 				continue;
 			}
 
-			ShrinkData(Data, MinX, MinY, MaxX, MaxY, ValidMinX, ValidMinY, ValidMaxX, ValidMaxY);
+			FLandscapeEditDataInterface::ShrinkData(Data, MinX, MinY, MaxX, MaxY, ValidMinX, ValidMinY, ValidMaxX, ValidMaxY);
 
 			MinX = ValidMinX;
 			MinY = ValidMinY;
@@ -2987,13 +2977,24 @@ bool ULandscapeInfo::ApplySplinesInternal(bool bOnlySelected, ALandscapeProxy* L
 			continue;
 		}
 
-		FBox LandscapeBounds = Segment->GetBounds();
-		LandscapeBounds = LandscapeBounds.TransformBy(SplineToLandscape.ToMatrixWithScale());
+		FBox SegmentBounds = Segment->GetBounds();
+		SegmentBounds = SegmentBounds.TransformBy(SplineToLandscape.ToMatrixWithScale());
+
+		int32 MinX = FMath::CeilToInt(SegmentBounds.Min.X);
+		int32 MinY = FMath::CeilToInt(SegmentBounds.Min.Y);
+		int32 MaxX = FMath::FloorToInt(SegmentBounds.Max.X);
+		int32 MaxY = FMath::FloorToInt(SegmentBounds.Max.Y);
+
+		MinX = FMath::Max(MinX, LandscapeMinX);
+		MinY = FMath::Max(MinY, LandscapeMinY);
+		MaxX = FMath::Min(MaxX, LandscapeMaxX);
+		MaxY = FMath::Min(MaxY, LandscapeMaxY);
 
-		int32 MinX = FMath::CeilToInt(LandscapeBounds.Min.X);
-		int32 MinY = FMath::CeilToInt(LandscapeBounds.Min.Y);
-		int32 MaxX = FMath::FloorToInt(LandscapeBounds.Max.X);
-		int32 MaxY = FMath::FloorToInt(LandscapeBounds.Max.Y);
+		if (MinX > MaxX || MinY > MaxY)
+		{
+			// The segment's bounds don't intersect any data, so we skip it entirely
+			continue;
+		}
 
 		TArray<FLandscapeSplineInterpPoint> Points = Segment->GetPoints();
 		for (int32 j = 0; j < Points.Num(); j++)
@@ -3030,7 +3031,7 @@ bool ULandscapeInfo::ApplySplinesInternal(bool bOnlySelected, ALandscapeProxy* L
 				continue;
 			}
 
-			ShrinkData(Data, MinX, MinY, MaxX, MaxY, ValidMinX, ValidMinY, ValidMaxX, ValidMaxY);
+			FLandscapeEditDataInterface::ShrinkData(Data, MinX, MinY, MaxX, MaxY, ValidMinX, ValidMinY, ValidMaxX, ValidMaxY);
 
 			MinX = ValidMinX;
 			MinY = ValidMinY;
@@ -3094,7 +3095,7 @@ bool ULandscapeInfo::ApplySplinesInternal(bool bOnlySelected, ALandscapeProxy* L
 				continue;
 			}
 
-			ShrinkData(Data, MinX, MinY, MaxX, MaxY, ValidMinX, ValidMinY, ValidMaxX, ValidMaxY);
+			FLandscapeEditDataInterface::ShrinkData(Data, MinX, MinY, MaxX, MaxY, ValidMinX, ValidMinY, ValidMaxX, ValidMaxY);
 
 			MinX = ValidMinX;
 			MinY = ValidMinY;
diff --git a/Engine/Source/Runtime/Landscape/Public/LandscapeEdit.h b/Engine/Source/Runtime/Landscape/Public/LandscapeEdit.h
--- a/Engine/Source/Runtime/Landscape/Public/LandscapeEdit.h
+++ b/Engine/Source/Runtime/Landscape/Public/LandscapeEdit.h
@@ -173,6 +173,9 @@ struct LANDSCAPE_API FLandscapeEditDataInterface
 	bool EqualTextureValueTempl(UTexture2D* Src, TData Value);
 	bool EqualTextureValue(UTexture2D* Src, FColor Value);
 
+	template<typename T>
+	static void ShrinkData(TArray<T>& Data, int32 OldMinX, int32 OldMinY, int32 OldMaxX, int32 OldMaxY, int32 NewMinX, int32 NewMinY, int32 NewMaxX, int32 NewMaxY);
+
 private:
 	int32 ComponentSizeQuads;
 	int32 SubsectionSizeQuads;
@@ -190,6 +193,33 @@ private:
 		TArray<bool>& NoBorderY1, TArray<bool>& NoBorderY2, TArray<bool>& ComponentDataExist, TStoreData& StoreData);
 };
 
+template<typename T>
+void FLandscapeEditDataInterface::ShrinkData(TArray<T>& Data, int32 OldMinX, int32 OldMinY, int32 OldMaxX, int32 OldMaxY, int32 NewMinX, int32 NewMinY, int32 NewMaxX, int32 NewMaxY)
+{
+	checkSlow(OldMinX <= OldMaxX && OldMinY <= OldMaxY);
+	checkSlow(NewMinX >= OldMinX && NewMaxX <= OldMaxX);
+	checkSlow(NewMinY >= OldMinY && NewMaxY <= OldMaxY);
+
+	if (NewMinX != OldMinX || NewMinY != OldMinY ||
+		NewMaxX != OldMaxX || NewMaxY != OldMaxY)
+	{
+		// if only the MaxY changes we don't need to do the moving, only the truncate
+		if (NewMinX != OldMinX || NewMinY != OldMinY || NewMaxX != OldMaxX)
+		{
+			for (int32 DestY = 0, SrcY = NewMinY - OldMinY; DestY <= NewMaxY - NewMinY; DestY++, SrcY++)
+			{
+				UE_LOG(LogLandscape, Warning, TEXT("Dest: %d, %d = %d Src: %d, %d = %d Width = %d"), 0, DestY, DestY * (1 + NewMaxX - NewMinX), NewMinX - OldMinX, SrcY, SrcY * (1 + OldMaxX - OldMinX) + NewMinX - OldMinX, (1 + NewMaxX - NewMinX));
+				T* DestData = &Data[DestY * (1 + NewMaxX - NewMinX)];
+				const T* SrcData = &Data[SrcY * (1 + OldMaxX - OldMinX) + NewMinX - OldMinX];
+				FMemory::Memmove(DestData, SrcData, (1 + NewMaxX - NewMinX) * sizeof(T));
+			}
+		}
+
+		const int32 NewSize = (1 + NewMaxY - NewMinY) * (1 + NewMaxX - NewMinX);
+		Data.RemoveAt(NewSize, Data.Num() - NewSize);
+	}
+}
+
 #endif
 
 #endif // _LANDSCAPEEDIT_H
