Fix FMobileLayerAllocation's operator<()
For some reason StableSort is comparing elements against themselves, FMobileLayerAllocation's operator<() would return true for comparing the visibility layer against itself and this put StableSort into an infinite loop

diff --git a/Engine/Source/Runtime/Landscape/Private/LandscapeEdit.cpp b/Engine/Source/Runtime/Landscape/Private/LandscapeEdit.cpp
--- a/Engine/Source/Runtime/Landscape/Private/LandscapeEdit.cpp
+++ b/Engine/Source/Runtime/Landscape/Private/LandscapeEdit.cpp
@@ -4769,12 +4769,14 @@ struct FMobileLayerAllocation
 
 	friend bool operator<(const FMobileLayerAllocation& lhs, const FMobileLayerAllocation& rhs)
 	{
+		if (&lhs == &rhs) return false; // I don't know why StableSort is comparing elements against themself
+
 		if (!lhs.Allocation.LayerInfo && !rhs.Allocation.LayerInfo) return false; // equally broken :P
 		if (!lhs.Allocation.LayerInfo && rhs.Allocation.LayerInfo) return false; // broken layers sort to the end
 		if (!rhs.Allocation.LayerInfo && lhs.Allocation.LayerInfo) return true;
 
-		if (lhs.Allocation.LayerInfo == ALandscapeProxy::VisibilityLayer) return true; // visibility layer to the front
-		if (rhs.Allocation.LayerInfo == ALandscapeProxy::VisibilityLayer) return false;
+		if (lhs.Allocation.LayerInfo == ALandscapeProxy::VisibilityLayer && rhs.Allocation.LayerInfo != ALandscapeProxy::VisibilityLayer) return true; // visibility layer to the front
+		if (rhs.Allocation.LayerInfo == ALandscapeProxy::VisibilityLayer && lhs.Allocation.LayerInfo != ALandscapeProxy::VisibilityLayer) return false;
 
 		if (lhs.Allocation.LayerInfo->bNoWeightBlend && !rhs.Allocation.LayerInfo->bNoWeightBlend) return false; // non-blended layers sort to the end
 		if (rhs.Allocation.LayerInfo->bNoWeightBlend && !lhs.Allocation.LayerInfo->bNoWeightBlend) return true;
