LandscapeEditor: Fixed build

diff --git a/Engine/Source/Runtime/Landscape/Public/LandscapeEdit.h b/Engine/Source/Runtime/Landscape/Public/LandscapeEdit.h
--- a/Engine/Source/Runtime/Landscape/Public/LandscapeEdit.h
+++ b/Engine/Source/Runtime/Landscape/Public/LandscapeEdit.h
@@ -208,7 +208,7 @@ void FLandscapeEditDataInterface::ShrinkData(TArray<T>& Data, int32 OldMinX, int
 		{
 			for (int32 DestY = 0, SrcY = NewMinY - OldMinY; DestY <= NewMaxY - NewMinY; DestY++, SrcY++)
 			{
-				UE_LOG(LogLandscape, Warning, TEXT("Dest: %d, %d = %d Src: %d, %d = %d Width = %d"), 0, DestY, DestY * (1 + NewMaxX - NewMinX), NewMinX - OldMinX, SrcY, SrcY * (1 + OldMaxX - OldMinX) + NewMinX - OldMinX, (1 + NewMaxX - NewMinX));
+//				UE_LOG(LogLandscape, Warning, TEXT("Dest: %d, %d = %d Src: %d, %d = %d Width = %d"), 0, DestY, DestY * (1 + NewMaxX - NewMinX), NewMinX - OldMinX, SrcY, SrcY * (1 + OldMaxX - OldMinX) + NewMinX - OldMinX, (1 + NewMaxX - NewMinX));
 				T* DestData = &Data[DestY * (1 + NewMaxX - NewMinX)];
 				const T* SrcData = &Data[SrcY * (1 + OldMaxX - OldMinX) + NewMinX - OldMinX];
 				FMemory::Memmove(DestData, SrcData, (1 + NewMaxX - NewMinX) * sizeof(T));
