Main branch stability fix.

diff --git a/Engine/Source/Runtime/Landscape/Private/LandscapeSplines.cpp b/Engine/Source/Runtime/Landscape/Private/LandscapeSplines.cpp
--- a/Engine/Source/Runtime/Landscape/Private/LandscapeSplines.cpp
+++ b/Engine/Source/Runtime/Landscape/Private/LandscapeSplines.cpp
@@ -13,6 +13,8 @@
 #include "LandscapeSplineControlPoint.h"
 #include "LandscapeSplineSegment.h"
 #include "ControlPointMeshComponent.h"
+#include "Engine/CollisionProfile.h"
+#include "EngineGlobals.h"
 #include "Engine/StaticMesh.h"
 #include "Engine/StaticMeshSocket.h"
 
