Missing headers NonUnityTestCompile_OnMac build fix.

diff --git a/Engine/Source/Runtime/Landscape/Private/LandscapeSplines.cpp b/Engine/Source/Runtime/Landscape/Private/LandscapeSplines.cpp
--- a/Engine/Source/Runtime/Landscape/Private/LandscapeSplines.cpp
+++ b/Engine/Source/Runtime/Landscape/Private/LandscapeSplines.cpp
@@ -14,6 +14,7 @@
 #include "LandscapeSplineSegment.h"
 #include "ControlPointMeshComponent.h"
 #include "Engine/CollisionProfile.h"
+#include "Engine/Engine.h"
 #include "EngineGlobals.h"
 #include "Engine/StaticMesh.h"
 #include "Engine/StaticMeshSocket.h"
diff --git a/Engine/Source/Runtime/Landscape/Private/Materials/MaterialExpressionLandscapeLayerBlend.cpp b/Engine/Source/Runtime/Landscape/Private/Materials/MaterialExpressionLandscapeLayerBlend.cpp
--- a/Engine/Source/Runtime/Landscape/Private/Materials/MaterialExpressionLandscapeLayerBlend.cpp
+++ b/Engine/Source/Runtime/Landscape/Private/Materials/MaterialExpressionLandscapeLayerBlend.cpp
@@ -4,6 +4,8 @@
 #include "MaterialCompiler.h"
 #include "Materials/MaterialExpressionLandscapeLayerBlend.h"
 #include "EdGraph/EdGraphNode.h"
+#include "Engine/Engine.h"
+#include "EngineGlobals.h"
 
 #define LOCTEXT_NAMESPACE "Landscape"
 
diff --git a/Engine/Source/Runtime/Landscape/Private/Materials/MaterialExpressionLandscapeLayerSample.cpp b/Engine/Source/Runtime/Landscape/Private/Materials/MaterialExpressionLandscapeLayerSample.cpp
--- a/Engine/Source/Runtime/Landscape/Private/Materials/MaterialExpressionLandscapeLayerSample.cpp
+++ b/Engine/Source/Runtime/Landscape/Private/Materials/MaterialExpressionLandscapeLayerSample.cpp
@@ -3,7 +3,8 @@
 #include "Landscape.h"
 #include "MaterialCompiler.h"
 #include "Materials/MaterialExpressionLandscapeLayerSample.h"
-
+#include "Engine/Engine.h"
+#include "EngineGlobals.h"
 
 #define LOCTEXT_NAMESPACE "Landscape"
 
diff --git a/Engine/Source/Runtime/Landscape/Private/Materials/MaterialExpressionLandscapeLayerSwitch.cpp b/Engine/Source/Runtime/Landscape/Private/Materials/MaterialExpressionLandscapeLayerSwitch.cpp
--- a/Engine/Source/Runtime/Landscape/Private/Materials/MaterialExpressionLandscapeLayerSwitch.cpp
+++ b/Engine/Source/Runtime/Landscape/Private/Materials/MaterialExpressionLandscapeLayerSwitch.cpp
@@ -3,7 +3,8 @@
 #include "Landscape.h"
 #include "MaterialCompiler.h"
 #include "Materials/MaterialExpressionLandscapeLayerSwitch.h"
-
+#include "Engine/Engine.h"
+#include "EngineGlobals.h"
 
 #define LOCTEXT_NAMESPACE "Landscape"
 
diff --git a/Engine/Source/Runtime/Landscape/Private/Materials/MaterialExpressionLandscapeLayerWeight.cpp b/Engine/Source/Runtime/Landscape/Private/Materials/MaterialExpressionLandscapeLayerWeight.cpp
--- a/Engine/Source/Runtime/Landscape/Private/Materials/MaterialExpressionLandscapeLayerWeight.cpp
+++ b/Engine/Source/Runtime/Landscape/Private/Materials/MaterialExpressionLandscapeLayerWeight.cpp
@@ -3,7 +3,8 @@
 #include "Landscape.h"
 #include "MaterialCompiler.h"
 #include "Materials/MaterialExpressionLandscapeLayerWeight.h"
-
+#include "Engine/Engine.h"
+#include "EngineGlobals.h"
 
 #define LOCTEXT_NAMESPACE "Landscape"
 
diff --git a/Engine/Source/Runtime/Landscape/Private/Materials/MaterialExpressionLandscapeVisibilityMask.cpp b/Engine/Source/Runtime/Landscape/Private/Materials/MaterialExpressionLandscapeVisibilityMask.cpp
--- a/Engine/Source/Runtime/Landscape/Private/Materials/MaterialExpressionLandscapeVisibilityMask.cpp
+++ b/Engine/Source/Runtime/Landscape/Private/Materials/MaterialExpressionLandscapeVisibilityMask.cpp
@@ -3,7 +3,8 @@
 #include "Landscape.h"
 #include "MaterialCompiler.h"
 #include "Materials/MaterialExpressionLandscapeVisibilityMask.h"
-
+#include "Engine/Engine.h"
+#include "EngineGlobals.h"
 
 #define LOCTEXT_NAMESPACE "Landscape"
 
