Fix spline component bounds for the case of having no splines (was *world* 0,0,0 before)

diff --git a/Engine/Source/Runtime/Landscape/Private/LandscapeSplines.cpp b/Engine/Source/Runtime/Landscape/Private/LandscapeSplines.cpp
--- a/Engine/Source/Runtime/Landscape/Private/LandscapeSplines.cpp
+++ b/Engine/Source/Runtime/Landscape/Private/LandscapeSplines.cpp
@@ -551,8 +551,26 @@ FBoxSphereBounds ULandscapeSplinesComponent::CalcBounds(const FTransform& LocalT
 		}
 	}
 
-	NewBoundsCalc = NewBoundsCalc.TransformBy(LocalToWorld);
-	return FBoxSphereBounds(NewBoundsCalc);
+	FBoxSphereBounds NewBounds;
+	if (NewBoundsCalc.IsValid)
+	{
+		NewBoundsCalc = NewBoundsCalc.TransformBy(LocalToWorld);
+		NewBounds = FBoxSphereBounds(NewBoundsCalc);
+	}
+	else
+	{
+		// There's no such thing as an "invalid" FBoxSphereBounds (unlike FBox)
+		// try to return something that won't modify the parent bounds
+		if (AttachParent)
+		{
+			NewBounds = FBoxSphereBounds(AttachParent->Bounds.Origin, FVector::ZeroVector, 0.0f);
+		}
+		else
+		{
+			NewBounds = FBoxSphereBounds(LocalToWorld.GetTranslation(), FVector::ZeroVector, 0.0f);
+		}
+	}
+	return NewBounds;
 }
 
 bool ULandscapeSplinesComponent::ModifySplines(bool bAlwaysMarkDirty /*= true*/)
