Fix crashes when undoing the delete of a landscape and then entering the landscape editor
Make bPendingKillPending a UPROPERTY() again so it is correctly undone when undoing an actor delete

diff --git a/Engine/Source/Runtime/Engine/Classes/GameFramework/Actor.h b/Engine/Source/Runtime/Engine/Classes/GameFramework/Actor.h
--- a/Engine/Source/Runtime/Engine/Classes/GameFramework/Actor.h
+++ b/Engine/Source/Runtime/Engine/Classes/GameFramework/Actor.h
@@ -308,6 +308,7 @@ public:
 	 * Set when actor is about to be deleted.
 	 * @see IsPendingKillPending()
 	 */
+	UPROPERTY(Transient, DuplicateTransient)
 	uint32 bPendingKillPending:1;    
 
 	/** This actor collides with the world when placing in the editor or when spawned, even if RootComponent collision is disabled */
diff --git a/Engine/Source/Runtime/Landscape/Private/Landscape.cpp b/Engine/Source/Runtime/Landscape/Private/Landscape.cpp
--- a/Engine/Source/Runtime/Landscape/Private/Landscape.cpp
+++ b/Engine/Source/Runtime/Landscape/Private/Landscape.cpp
@@ -579,7 +579,7 @@ ULandscapeInfo* ALandscapeProxy::GetLandscapeInfo(bool bSpawnNewActor /*= true*/
 			if (!LandscapeInfo && bSpawnNewActor && !HasAnyFlags(RF_BeginDestroyed))
 			{
 				LandscapeInfo = ConstructObject<ULandscapeInfo>(ULandscapeInfo::StaticClass(), OwningWorld, NAME_None, RF_Transactional | RF_Transient);
-				OwningWorld->Modify(false);
+				LandscapeInfoMap.Modify(false);
 				LandscapeInfoMap.Map.Add(LandscapeGuid, LandscapeInfo);
 			}
 
diff --git a/Engine/Source/Runtime/Landscape/Private/LandscapeInfoMap.cpp b/Engine/Source/Runtime/Landscape/Private/LandscapeInfoMap.cpp
--- a/Engine/Source/Runtime/Landscape/Private/LandscapeInfoMap.cpp
+++ b/Engine/Source/Runtime/Landscape/Private/LandscapeInfoMap.cpp
@@ -10,31 +10,24 @@ ULandscapeInfoMap::ULandscapeInfoMap(const FObjectInitializer& ObjectInitializer
 
 void ULandscapeInfoMap::PostDuplicate(bool bDuplicateForPIE)
 {
-	// After duplication the pointers are duplicated, but the objects should
-	// be duplicated.
-
 	Super::PostDuplicate(bDuplicateForPIE);
-	
-	for (auto Iterator = Map.CreateConstIterator(); Iterator; ++Iterator)
-	{
-		Map[Iterator->Key] = Cast<ULandscapeInfo>(StaticDuplicateObject(
-			Iterator->Value, Iterator->Value->GetOuter(), nullptr));
-	}
+
+	check(Map.Num() == 0);
 }
 
 void ULandscapeInfoMap::Serialize(FArchive& Ar)
 {
 	Super::Serialize(Ar);
 
-	Ar << Map;
+	if (!Ar.IsLoading() && !Ar.IsSaving())
+	{
+		Ar << Map;
+	}
 }
 
 void ULandscapeInfoMap::AddReferencedObjects(UObject* InThis, FReferenceCollector& Collector)
 {
 	ULandscapeInfoMap* This = CastChecked<ULandscapeInfoMap>(InThis);
 
-	for (auto ReferencedMapInfo = This->Map.CreateIterator(); ReferencedMapInfo; ++ReferencedMapInfo)
-	{
-		Collector.AddReferencedObject(ReferencedMapInfo->Value, This);
-	}
+	Collector.AddReferencedObjects(This->Map, This);
 }
