Clean up any stray components in RecreateLandscapeCollision console command

diff --git a/Engine/Source/Runtime/Landscape/Private/LandscapeEdit.cpp b/Engine/Source/Runtime/Landscape/Private/LandscapeEdit.cpp
--- a/Engine/Source/Runtime/Landscape/Private/LandscapeEdit.cpp
+++ b/Engine/Source/Runtime/Landscape/Private/LandscapeEdit.cpp
@@ -3441,14 +3441,27 @@ void ALandscapeProxy::RecreateCollisionComponents()
 	// Clear old CollisionComponent containers
 	CollisionComponents.Empty();
 
-	// Clear any Owned Collision Components
+	// Destroy any owned collision components
 	TArray<ULandscapeHeightfieldCollisionComponent*> CollisionComps;
 	GetComponents(CollisionComps);
-	for (int32 i = 0; i < CollisionComps.Num(); ++i)
+	for (ULandscapeHeightfieldCollisionComponent* Component : CollisionComps)
 	{
-		RemoveOwnedComponent(CollisionComps[i]);
+		Component->DestroyComponent();
 	}
 
+	TArray<USceneComponent*> AttachedCollisionComponents = RootComponent->AttachChildren.FilterByPredicate(
+		[](USceneComponent* Component)
+	{
+		return Cast<ULandscapeHeightfieldCollisionComponent>(Component);
+	});
+
+	// Destroy any attached but un-owned collision components
+	for (USceneComponent* Component : AttachedCollisionComponents)
+	{
+		Component->DestroyComponent();
+	}
+
+	// Recreate collision
 	CollisionMipLevel = FMath::Clamp<int32>(CollisionMipLevel, 0, FMath::CeilLogTwo(SubsectionSizeQuads + 1) - 1);
 	for (int32 i = 0; i < LandscapeComponents.Num(); ++i)
 	{
