#UE4 - Crash fix for mobile landscape

diff --git a/Engine/Source/Runtime/Landscape/Private/LandscapeRenderMobile.cpp b/Engine/Source/Runtime/Landscape/Private/LandscapeRenderMobile.cpp
--- a/Engine/Source/Runtime/Landscape/Private/LandscapeRenderMobile.cpp
+++ b/Engine/Source/Runtime/Landscape/Private/LandscapeRenderMobile.cpp
@@ -228,7 +228,10 @@ FLandscapeComponentSceneProxyMobile::FLandscapeComponentSceneProxyMobile(ULandsc
 		InComponent->MobileNormalmapTexture = WeightmapTextures[0];
 	}
 #endif
-	NormalmapTexture = WeightmapTextures[0]; // Use Weightmap0 as Normal/Weightmap Texture
+	if (WeightmapTextures.Num() > 0)
+	{
+		NormalmapTexture = WeightmapTextures[0]; // Use Weightmap0 as Normal/Weightmap Texture
+	}
 
 	BlendableLayerMask = InComponent->MobileBlendableLayerMask;
 }
