Fix slow PIE with Landscape

diff --git a/Engine/Source/Runtime/Landscape/Private/Landscape.cpp b/Engine/Source/Runtime/Landscape/Private/Landscape.cpp
--- a/Engine/Source/Runtime/Landscape/Private/Landscape.cpp
+++ b/Engine/Source/Runtime/Landscape/Private/Landscape.cpp
@@ -436,8 +436,8 @@ void ULandscapeComponent::PostLoad()
 			CachedLocalBox = CachedBoxSphereBounds_DEPRECATED.GetBox().InverseTransformBy(ComponentLtWTransform);
 		}
 
-		// If we're loading on a platform that doesn't require cooked data, attempt to load missing data from the DDC
-		if (!FPlatformProperties::RequiresCookedData() && (GIsEditor || (GMaxRHIFeatureLevel == ERHIFeatureLevel::ES2 || GMaxRHIFeatureLevel == ERHIFeatureLevel::ES3_1)))
+		// If we're loading on a platform that doesn't require cooked data, but *only* supports OpenGL ES, preload data from the DDC
+		if (!FPlatformProperties::RequiresCookedData() && GMaxRHIFeatureLevel <= ERHIFeatureLevel::ES3_1)
 		{
 			// Only check the DDC if we don't already have it loaded
 			if (!PlatformData.HasValidPlatformData())
@@ -712,10 +712,12 @@ FPrimitiveSceneProxy* ULandscapeComponent::CreateSceneProxy()
 		Proxy = new FLandscapeComponentSceneProxy(this, NULL);
 #endif
 	}
-	else
+	else // i.e. (FeatureLevel <= ERHIFeatureLevel::ES3_1)
 	{
 #if WITH_EDITOR
-		if (!PlatformData.HasValidPlatformData()) // Deferred generation
+		// We need to cook platform data for ES2 preview in editor
+		// We can't always pre-cook it in PostLoad/Serialize etc because landscape can be edited
+		if (!PlatformData.HasValidPlatformData())
 		{
 			// Try to reload the ES2 landscape data from the DDC
 			if (!PlatformData.LoadFromDDC(StateId))
