Landscape uses DDC to cache physx heightfield data for current editor platform. This improves loading time for landscape maps in editor builds
UE-4946

diff --git a/Engine/Source/Runtime/Landscape/Classes/LandscapeHeightfieldCollisionComponent.h b/Engine/Source/Runtime/Landscape/Classes/LandscapeHeightfieldCollisionComponent.h
--- a/Engine/Source/Runtime/Landscape/Classes/LandscapeHeightfieldCollisionComponent.h
+++ b/Engine/Source/Runtime/Landscape/Classes/LandscapeHeightfieldCollisionComponent.h
@@ -99,20 +99,25 @@ class ULandscapeHeightfieldCollisionComponent : public UPrimitiveComponent
 		virtual ~FPhysXHeightfieldRef();
 	};
 	
+#if WITH_EDITORONLY_DATA
 	/** The collision height values. Stripped from cooked content */
 	FWordBulkData								CollisionHeightData;
 
 	/** Indices into the ComponentLayers array for the per-vertex dominant layer. Stripped from cooked content */
 	FByteBulkData								DominantLayerData;
 
+	/*  Cooked editor specific heightfield data, never serialized  */
+	TArray<uint8>								CookedCollisionDataEd;
+#endif //WITH_EDITORONLY_DATA
+
 	/** 
 	 *	Cooked HeightField data. Serialized only with cooked content 
 	 *	Stored as array instead of BulkData to take advantage of precaching during async loading
 	 */
 	TArray<uint8>								CookedCollisionData;
-
-	/** This is a list of physical materials that is actually used by a cooked HeightField. Serialized only with cooked content */
-	UPROPERTY(transient)
+	
+	/** This is a list of physical materials that is actually used by a cooked HeightField */
+	UPROPERTY()
 	TArray<UPhysicalMaterial*>					CookedPhysicalMaterials;
 	
 	/** Physics engine version of heightfield data. */
@@ -152,6 +157,7 @@ class ULandscapeHeightfieldCollisionComponent : public UPrimitiveComponent
 	virtual void Serialize(FArchive& Ar) override;
 	virtual void BeginDestroy() override;
 	virtual void PostLoad() override;
+	virtual void PreSave() override;
 #if WITH_EDITOR
 	virtual void ExportCustomProperties(FOutputDevice& Out, uint32 Indent) override;
 	virtual void ImportCustomProperties(const TCHAR* SourceText, FFeedbackContext* Warn) override;
diff --git a/Engine/Source/Runtime/Landscape/Classes/LandscapeMeshCollisionComponent.h b/Engine/Source/Runtime/Landscape/Classes/LandscapeMeshCollisionComponent.h
--- a/Engine/Source/Runtime/Landscape/Classes/LandscapeMeshCollisionComponent.h
+++ b/Engine/Source/Runtime/Landscape/Classes/LandscapeMeshCollisionComponent.h
@@ -60,8 +60,10 @@ class ULandscapeMeshCollisionComponent : public ULandscapeHeightfieldCollisionCo
 		virtual ~FPhysXMeshRef();
 	};
 
+#if WITH_EDITORONLY_DATA
 	/** The collision mesh values. */
 	FWordBulkData								CollisionXYOffsetData; //  X, Y Offset in raw format...
+#endif //WITH_EDITORONLY_DATA
 
 	/** Physics engine version of heightfield data. */
 	TRefCountPtr<struct FPhysXMeshRef>			MeshRef;
diff --git a/Engine/Source/Runtime/Landscape/Private/LandscapeCollision.cpp b/Engine/Source/Runtime/Landscape/Private/LandscapeCollision.cpp
--- a/Engine/Source/Runtime/Landscape/Private/LandscapeCollision.cpp
+++ b/Engine/Source/Runtime/Landscape/Private/LandscapeCollision.cpp
@@ -77,6 +77,15 @@ ULandscapeMeshCollisionComponent::FPhysXMeshRef::~FPhysXMeshRef()
 	GSharedMeshRefs.Remove(Guid);
 }
 
+// Generate a new guid to force a recache of landscape collisoon derived data
+#define LANDSCAPE_COLLISION_DERIVEDDATA_VER	TEXT("5DF9E1AAB7CC4DCCB2965BA1A78DFE8")
+
+static FString GetHFDDCKeyString(const FName& Format, bool bDefMaterial, const FGuid& StateId)
+{
+	const FString KeyPrefix = FString::Printf(TEXT("%s_%s"), *Format.ToString(), (bDefMaterial ? TEXT("VIS") : TEXT("FULL")));
+	return FDerivedDataCacheInterface::BuildCacheKey(*KeyPrefix, LANDSCAPE_COLLISION_DERIVEDDATA_VER, *StateId.ToString());
+}
+
 ECollisionEnabled::Type ULandscapeHeightfieldCollisionComponent::GetCollisionEnabled() const
 {
 	ALandscapeProxy* Proxy = GetLandscapeProxy();
@@ -112,7 +121,6 @@ void ULandscapeHeightfieldCollisionComponent::CreatePhysicsState()
 	if (!BodyInstance.IsValidBodyInstance())
 	{
 #if WITH_PHYSX
-		// This will do nothing, because we create heightfield at component PostLoad event, unless we destroyed it explicitly
 		CreateCollisionObject();
 
 		if (IsValidRef(HeightfieldRef))
@@ -248,11 +256,6 @@ void ULandscapeHeightfieldCollisionComponent::ApplyWorldOffset(const FVector& In
 
 void ULandscapeHeightfieldCollisionComponent::CreateCollisionObject()
 {
-	if (IsTemplate())
-	{
-		return;
-	}
-
 #if WITH_PHYSX	
 	// If we have not created a heightfield yet - do it now.
 	if (!IsValidRef(HeightfieldRef))
@@ -291,19 +294,23 @@ void ULandscapeHeightfieldCollisionComponent::CreateCollisionObject()
 					HeightfieldRef->UsedPhysicalMaterialArray.Add(PhysicalMaterial->GetPhysXMaterial());
 				}
 
-				// Release cooked data
-				// In cooked builds created PhysX heightfield object will never be deleted until component is alive, so we don't need this data anymore
-				// In the editor this data will be regenerated just before creating new heightfield object
-				CookedCollisionData.Empty();
+				// Release cooked collison data
+				// In cooked builds created collision object will never be deleted while component is alive, so we don't need this data anymore
+				if (FPlatformProperties::RequiresCookedData() || GetWorld()->IsGameWorld())
+				{
+					CookedCollisionData.Empty();
+				}
 
 #if WITH_EDITOR
 				// Create heightfield for the landscape editor (no holes in it)
-				TArray<uint8> CookedHeightmapEd;
-				TArray<UPhysicalMaterial*> CookedMaterialsEd;
-				if (CookCollsionData(PhysicsFormatName, true, CookedHeightmapEd, CookedMaterialsEd))
+				if (!GetWorld()->IsGameWorld())
 				{
-					FPhysXInputStream HeightFieldStream(CookedHeightmapEd.GetData(), CookedHeightmapEd.Num());
-					HeightfieldRef->RBHeightfieldEd = GPhysXSDK->createHeightField(HeightFieldStream);
+					TArray<UPhysicalMaterial*> CookedMaterialsEd;
+					if (CookCollsionData(PhysicsFormatName, true, CookedCollisionDataEd, CookedMaterialsEd))
+					{
+						FPhysXInputStream HeightFieldStream(CookedCollisionDataEd.GetData(), CookedCollisionDataEd.Num());
+						HeightfieldRef->RBHeightfieldEd = GPhysXSDK->createHeightField(HeightFieldStream);
+					}
 				}
 #endif //WITH_EDITOR
 			}
@@ -317,12 +324,17 @@ bool ULandscapeHeightfieldCollisionComponent::CookCollsionData(const FName& Form
 {
 #if WITH_PHYSX
 
+	if (GetDerivedDataCacheRef().GetSynchronous(*GetHFDDCKeyString(Format, bUseDefMaterial, HeightfieldGuid), OutCookedData))
+	{
+		return true;
+	}
+				
 	ALandscapeProxy* Proxy = GetLandscapeProxy();
 	if (!Proxy || !Proxy->GetRootComponent())
 	{
 		return false;
 	}
-
+	
 	UPhysicalMaterial* DefMaterial = Proxy->DefaultPhysMaterial ? Proxy->DefaultPhysMaterial : GEngine->DefaultPhysMaterial;
 
 	// ComponentToWorld might not be initialized at this point, so use landscape transform
@@ -345,7 +357,9 @@ bool ULandscapeHeightfieldCollisionComponent::CookCollsionData(const FName& Form
 	OutMaterails.Empty();
 
 	TArray<PxHeightFieldSample> Samples;
-	Samples.AddZeroed(FMath::Square(CollisionSizeVerts));
+	const int32 NumSamples = FMath::Square(CollisionSizeVerts);
+	Samples.Reserve(NumSamples);
+	Samples.AddZeroed(NumSamples);
 
 	for (int32 RowIndex = 0; RowIndex < CollisionSizeVerts; RowIndex++)
 	{
@@ -412,7 +426,7 @@ bool ULandscapeHeightfieldCollisionComponent::CookCollsionData(const FName& Form
 
 	if (Result)
 	{
-		OutCookedData.SetNumUninitialized(OutData.Num());
+		OutCookedData.Init(OutData.Num());
 		FMemory::Memcpy(OutCookedData.GetData(), OutData.GetData(), OutData.Num());
 	}
 	else
@@ -429,11 +443,13 @@ bool ULandscapeHeightfieldCollisionComponent::CookCollsionData(const FName& Form
 
 bool ULandscapeMeshCollisionComponent::CookCollsionData(const FName& Format, bool bUseDefMaterial, TArray<uint8>& OutCookedData, TArray<UPhysicalMaterial*>& OutMaterails) const
 {
-	if (IsTemplate())
+#if WITH_PHYSX
+	
+	if (GetDerivedDataCacheRef().GetSynchronous(*GetHFDDCKeyString(Format, bUseDefMaterial, HeightfieldGuid), OutCookedData))
 	{
-		return false;
+		return true;
 	}
-#if WITH_PHYSX
+	
 	ALandscapeProxy* Proxy = GetLandscapeProxy();
 	UPhysicalMaterial* DefMaterial = (Proxy && Proxy->DefaultPhysMaterial != nullptr) ? Proxy->DefaultPhysMaterial : GEngine->DefaultPhysMaterial;
 
@@ -459,7 +475,7 @@ bool ULandscapeMeshCollisionComponent::CookCollsionData(const FName& Format, boo
 	}
 
 	// Scale all verts into temporary vertex buffer.
-	Vertices.AddUninitialized(NumVerts);
+	Vertices.Init(NumVerts);
 	for (int32 i = 0; i < NumVerts; i++)
 	{
 		int32 X = i % CollisionSizeVerts;
@@ -468,10 +484,10 @@ bool ULandscapeMeshCollisionComponent::CookCollsionData(const FName& Format, boo
 	}
 
 	const int32 NumTris = FMath::Square(CollisionSizeQuads) * 2;
-	Indices.AddUninitialized(NumTris);
+	Indices.Init(NumTris);
 	if (DominantLayers)
 	{
-		MaterialIndices.AddUninitialized(NumTris);
+		MaterialIndices.Init(NumTris);
 	}
 
 	int32 TriangleIdx = 0;
@@ -565,7 +581,7 @@ bool ULandscapeMeshCollisionComponent::CookCollsionData(const FName& Format, boo
 
 	if (Result)
 	{
-		OutCookedData.SetNumUninitialized(OutData.Num());
+		OutCookedData.Init(OutData.Num());
 		FMemory::Memcpy(OutCookedData.GetData(), OutData.GetData(), OutData.Num());
 	}
 	else
@@ -583,11 +599,6 @@ bool ULandscapeMeshCollisionComponent::CookCollsionData(const FName& Format, boo
 
 void ULandscapeMeshCollisionComponent::CreateCollisionObject()
 {
-	if (IsTemplate())
-	{
-		return;
-	}
-
 #if WITH_PHYSX	
 	// If we have not created a heightfield yet - do it now.
 	if (!IsValidRef(MeshRef))
@@ -624,19 +635,23 @@ void ULandscapeMeshCollisionComponent::CreateCollisionObject()
 					MeshRef->UsedPhysicalMaterialArray.Add(PhysicalMaterial->GetPhysXMaterial());
 				}
 
-				// Release cooked data
-				// In cooked builds created collision object will never be deleted until component is alive, so we don't need this data anymore
-				// In the editor this data will be regenerated just before creating new collision object
-				CookedCollisionData.Empty();
+				// Release cooked collison data
+				// In cooked builds created collision object will never be deleted while component is alive, so we don't need this data anymore
+				if (FPlatformProperties::RequiresCookedData() || GetWorld()->IsGameWorld())
+				{
+					CookedCollisionData.Empty();
+				}
 
 #if WITH_EDITOR
 				// Create collision mesh for the landscape editor (no holes in it)
-				TArray<uint8> CookedMeshEd;
-				TArray<UPhysicalMaterial*> CookedMaterialsEd;
-				if (CookCollsionData(PhysicsFormatName, true, CookedMeshEd, CookedMaterialsEd))
+				if (!GetWorld()->IsGameWorld())
 				{
-					FPhysXInputStream MeshStream(CookedMeshEd.GetData(), CookedMeshEd.Num());
-					MeshRef->RBTriangleMeshEd = GPhysXSDK->createTriangleMesh(MeshStream);
+					TArray<UPhysicalMaterial*> CookedMaterialsEd;
+					if (CookCollsionData(PhysicsFormatName, true, CookedCollisionDataEd, CookedMaterialsEd))
+					{
+						FPhysXInputStream MeshStream(CookedCollisionDataEd.GetData(), CookedCollisionDataEd.Num());
+						MeshRef->RBTriangleMeshEd = GPhysXSDK->createTriangleMesh(MeshStream);
+					}
 				}
 #endif //WITH_EDITOR
 			}
@@ -1071,12 +1086,27 @@ void ULandscapeMeshCollisionComponent::RecreateCollision(bool bUpdateAddCollisio
 
 void ULandscapeHeightfieldCollisionComponent::Serialize(FArchive& Ar)
 {
+#if WITH_EDITOR
+	if (Ar.UE4Ver() >= VER_UE4_LANDSCAPE_COLLISION_DATA_COOKING)
+	{
+		// Cook data here so CookedPhysicalMaterials is always up to date
+		if (Ar.IsCooking() && !HasAnyFlags(RF_ClassDefaultObject))
+		{
+			FName Format = Ar.CookingTarget()->GetPhysicsFormat(nullptr);
+			CookCollsionData(Format, false, CookedCollisionData, CookedPhysicalMaterials);
+		}
+	}
+#endif// WITH_EDITOR
+
+	// this will also serialize CookedPhysicalMaterials
 	Super::Serialize(Ar);
 
 	if (Ar.UE4Ver() < VER_UE4_LANDSCAPE_COLLISION_DATA_COOKING)
 	{
+#if WITH_EDITORONLY_DATA
 		CollisionHeightData.Serialize(Ar, this);
 		DominantLayerData.Serialize(Ar, this);
+#endif//WITH_EDITORONLY_DATA
 	}
 	else
 	{
@@ -1090,20 +1120,14 @@ void ULandscapeHeightfieldCollisionComponent::Serialize(FArchive& Ar)
 
 		if (bCooked)
 		{
-#if WITH_EDITOR
-			if (Ar.IsCooking() && !HasAnyFlags(RF_ClassDefaultObject))
-			{
-				CookCollsionData(Ar.CookingTarget()->GetPhysicsFormat(NULL), false, CookedCollisionData, CookedPhysicalMaterials);
-			}
-#endif// WITH_EDITOR
-
 			Ar << CookedCollisionData;
-			Ar << CookedPhysicalMaterials;
 		}
 		else
 		{
+#if WITH_EDITORONLY_DATA			
 			CollisionHeightData.Serialize(Ar, this);
 			DominantLayerData.Serialize(Ar, this);
+#endif//WITH_EDITORONLY_DATA
 		}
 	}
 }
@@ -1114,8 +1138,10 @@ void ULandscapeMeshCollisionComponent::Serialize(FArchive& Ar)
 
 	if (Ar.UE4Ver() < VER_UE4_LANDSCAPE_COLLISION_DATA_COOKING)
 	{
+#if WITH_EDITORONLY_DATA
 		// conditional serialization in later versions
 		CollisionXYOffsetData.Serialize(Ar, this);
+#endif// WITH_EDITORONLY_DATA
 	}
 
 	// PhysX cooking mesh data
@@ -1137,8 +1163,10 @@ void ULandscapeMeshCollisionComponent::Serialize(FArchive& Ar)
 	}
 	else if (Ar.UE4Ver() >= VER_UE4_LANDSCAPE_COLLISION_DATA_COOKING)
 	{
+#if WITH_EDITORONLY_DATA		
 		// we serialize raw collision data only with non-cooked content
 		CollisionXYOffsetData.Serialize(Ar, this);
+#endif// WITH_EDITORONLY_DATA
 	}
 }
 
@@ -1228,10 +1256,6 @@ void ULandscapeHeightfieldCollisionComponent::PostLoad()
 {
 	Super::PostLoad();
 
-	// Create collision object right after component was loaded
-	// so we can free cooked collision data
-	CreateCollisionObject();
-
 #if WITH_EDITOR
 	if (!HasAnyFlags(RF_ClassDefaultObject))
 	{
@@ -1277,6 +1301,20 @@ void ULandscapeHeightfieldCollisionComponent::PostLoad()
 #endif//WITH_EDITOR
 }
 
+void ULandscapeHeightfieldCollisionComponent::PreSave()
+{
+	Super::PreSave();
+
+	if (!IsRunningCommandlet())
+	{
+#if WITH_EDITOR
+		static FName PhysicsFormatName(FPlatformProperties::GetPhysicsFormat());
+		GetDerivedDataCacheRef().Put(*GetHFDDCKeyString(PhysicsFormatName, false, HeightfieldGuid), CookedCollisionData);
+		GetDerivedDataCacheRef().Put(*GetHFDDCKeyString(PhysicsFormatName, true, HeightfieldGuid), CookedCollisionDataEd);
+#endif// WITH_EDITOR
+	}
+}
+
 #if WITH_EDITOR
 void ULandscapeInfo::UpdateAllAddCollisions()
 {
