Added DDC population when landscape content is opened in the editor
UE-4946

diff --git a/Engine/Source/Runtime/Landscape/Classes/LandscapeHeightfieldCollisionComponent.h b/Engine/Source/Runtime/Landscape/Classes/LandscapeHeightfieldCollisionComponent.h
--- a/Engine/Source/Runtime/Landscape/Classes/LandscapeHeightfieldCollisionComponent.h
+++ b/Engine/Source/Runtime/Landscape/Classes/LandscapeHeightfieldCollisionComponent.h
@@ -108,6 +108,12 @@ class ULandscapeHeightfieldCollisionComponent : public UPrimitiveComponent
 
 	/*  Cooked editor specific heightfield data, never serialized  */
 	TArray<uint8>								CookedCollisionDataEd;
+
+	/** 
+	 *	Flag to indicate that the next time we cook data, we should save it to the DDC.
+	 *	Used to ensure DDC is populated when loading content for the first time. 
+	 */
+	mutable bool								bShouldSaveCookedDataToDDC;
 #endif //WITH_EDITORONLY_DATA
 
 	/** 
@@ -174,7 +180,7 @@ class ULandscapeHeightfieldCollisionComponent : public UPrimitiveComponent
 	/** 
 	 * Cooks raw height data into collision object binary stream
 	 */
-	virtual bool CookCollsionData(const FName& Format, bool bUseOnlyDefMaterial, TArray<uint8>& OutCookedData, TArray<UPhysicalMaterial*>& OutMaterails) const;
+	virtual bool CookCollsionData(const FName& Format, bool bUseOnlyDefMaterial, TArray<uint8>& OutCookedData, TArray<UPhysicalMaterial*>& OutMaterials) const;
 
 	/** Modify a sub-region of the PhysX heightfield. Note that this does not update the physical material */
 	void UpdateHeightfieldRegion(int32 ComponentX1, int32 ComponentY1, int32 ComponentX2, int32 ComponentY2);
diff --git a/Engine/Source/Runtime/Landscape/Classes/LandscapeMeshCollisionComponent.h b/Engine/Source/Runtime/Landscape/Classes/LandscapeMeshCollisionComponent.h
--- a/Engine/Source/Runtime/Landscape/Classes/LandscapeMeshCollisionComponent.h
+++ b/Engine/Source/Runtime/Landscape/Classes/LandscapeMeshCollisionComponent.h
@@ -88,7 +88,7 @@ class ULandscapeMeshCollisionComponent : public ULandscapeHeightfieldCollisionCo
 	virtual void ExportCustomProperties(FOutputDevice& Out, uint32 Indent) override;
 	virtual void ImportCustomProperties(const TCHAR* SourceText, FFeedbackContext* Warn) override;
 
-	virtual bool CookCollsionData(const FName& Format, bool bUseOnlyDefMaterial, TArray<uint8>& OutCookedData, TArray<UPhysicalMaterial*>& OutMaterails) const override;
+	virtual bool CookCollsionData(const FName& Format, bool bUseOnlyDefMaterial, TArray<uint8>& OutCookedData, TArray<UPhysicalMaterial*>& OutMaterials) const override;
 #endif
 	// End UObject Interface.
 
diff --git a/Engine/Source/Runtime/Landscape/Private/LandscapeCollision.cpp b/Engine/Source/Runtime/Landscape/Private/LandscapeCollision.cpp
--- a/Engine/Source/Runtime/Landscape/Private/LandscapeCollision.cpp
+++ b/Engine/Source/Runtime/Landscape/Private/LandscapeCollision.cpp
@@ -320,12 +320,13 @@ void ULandscapeHeightfieldCollisionComponent::CreateCollisionObject()
 }
 
 #if WITH_EDITOR
-bool ULandscapeHeightfieldCollisionComponent::CookCollsionData(const FName& Format, bool bUseDefMaterial, TArray<uint8>& OutCookedData, TArray<UPhysicalMaterial*>& OutMaterails) const
+bool ULandscapeHeightfieldCollisionComponent::CookCollsionData(const FName& Format, bool bUseDefMaterial, TArray<uint8>& OutCookedData, TArray<UPhysicalMaterial*>& OutMaterials) const
 {
 #if WITH_PHYSX
 
 	if (GetDerivedDataCacheRef().GetSynchronous(*GetHFDDCKeyString(Format, bUseDefMaterial, HeightfieldGuid), OutCookedData))
 	{
+		bShouldSaveCookedDataToDDC = false;
 		return true;
 	}
 				
@@ -354,7 +355,7 @@ bool ULandscapeHeightfieldCollisionComponent::CookCollsionData(const FName& Form
 	}
 
 	// List of materials which is actually used by heightfield
-	OutMaterails.Empty();
+	OutMaterials.Empty();
 
 	TArray<PxHeightFieldSample> Samples;
 	const int32 NumSamples = FMath::Square(CollisionSizeVerts);
@@ -390,7 +391,7 @@ bool ULandscapeHeightfieldCollisionComponent::CookCollsionData(const FName& Form
 						else
 						{
 							UPhysicalMaterial* DominantMaterial = Layer && Layer->PhysMaterial ? Layer->PhysMaterial : DefMaterial;
-							MaterialIndex = OutMaterails.AddUnique(DominantMaterial);
+							MaterialIndex = OutMaterials.AddUnique(DominantMaterial);
 						}
 					}
 				}
@@ -410,9 +411,9 @@ bool ULandscapeHeightfieldCollisionComponent::CookCollsionData(const FName& Form
 	}
 
 	// Add the default physical material to be used used when we have no dominant data.
-	if (OutMaterails.Num() == 0)
+	if (OutMaterials.Num() == 0)
 	{
-		OutMaterails.Add(DefMaterial);
+		OutMaterials.Add(DefMaterial);
 	}
 
 	//
@@ -428,11 +429,17 @@ bool ULandscapeHeightfieldCollisionComponent::CookCollsionData(const FName& Form
 	{
 		OutCookedData.Init(OutData.Num());
 		FMemory::Memcpy(OutCookedData.GetData(), OutData.GetData(), OutData.Num());
+
+		if (bShouldSaveCookedDataToDDC)
+		{
+			GetDerivedDataCacheRef().Put(*GetHFDDCKeyString(Format, bUseDefMaterial, HeightfieldGuid), OutCookedData);
+			bShouldSaveCookedDataToDDC = false;
+		}
 	}
 	else
 	{
 		OutCookedData.Empty();
-		OutMaterails.Empty();
+		OutMaterials.Empty();
 	}
 
 	return Result;
@@ -441,12 +448,13 @@ bool ULandscapeHeightfieldCollisionComponent::CookCollsionData(const FName& Form
 	return false;
 }
 
-bool ULandscapeMeshCollisionComponent::CookCollsionData(const FName& Format, bool bUseDefMaterial, TArray<uint8>& OutCookedData, TArray<UPhysicalMaterial*>& OutMaterails) const
+bool ULandscapeMeshCollisionComponent::CookCollsionData(const FName& Format, bool bUseDefMaterial, TArray<uint8>& OutCookedData, TArray<UPhysicalMaterial*>& OutMaterials) const
 {
 #if WITH_PHYSX
 	
 	if (GetDerivedDataCacheRef().GetSynchronous(*GetHFDDCKeyString(Format, bUseDefMaterial, HeightfieldGuid), OutCookedData))
 	{
+		bShouldSaveCookedDataToDDC = false;
 		return true;
 	}
 	
@@ -454,7 +462,7 @@ bool ULandscapeMeshCollisionComponent::CookCollsionData(const FName& Format, boo
 	UPhysicalMaterial* DefMaterial = (Proxy && Proxy->DefaultPhysMaterial != nullptr) ? Proxy->DefaultPhysMaterial : GEngine->DefaultPhysMaterial;
 
 	// List of materials which is actually used by trimesh
-	OutMaterails.Empty();
+	OutMaterials.Empty();
 
 	TArray<FVector>			Vertices;
 	TArray<FTriIndices>		Indices;
@@ -513,7 +521,7 @@ bool ULandscapeMeshCollisionComponent::CookCollsionData(const FName& Format, boo
 					else
 					{
 						UPhysicalMaterial* DominantMaterial = Layer && Layer->PhysMaterial ? Layer->PhysMaterial : DefMaterial;
-						MaterialIndex = OutMaterails.AddUnique(DominantMaterial);
+						MaterialIndex = OutMaterials.AddUnique(DominantMaterial);
 					}
 				}
 			}
@@ -568,9 +576,9 @@ bool ULandscapeMeshCollisionComponent::CookCollsionData(const FName& Format, boo
 	}
 
 	// Add the default physical material to be used used when we have no dominant data.
-	if (OutMaterails.Num() == 0)
+	if (OutMaterials.Num() == 0)
 	{
-		OutMaterails.Add(DefMaterial);
+		OutMaterials.Add(DefMaterial);
 	}
 
 	bool bFlipNormals = true;
@@ -583,11 +591,17 @@ bool ULandscapeMeshCollisionComponent::CookCollsionData(const FName& Format, boo
 	{
 		OutCookedData.Init(OutData.Num());
 		FMemory::Memcpy(OutCookedData.GetData(), OutData.GetData(), OutData.Num());
+
+		if (bShouldSaveCookedDataToDDC)
+		{
+			GetDerivedDataCacheRef().Put(*GetHFDDCKeyString(Format, bUseDefMaterial, HeightfieldGuid), OutCookedData);
+			bShouldSaveCookedDataToDDC = false;
+		}
 	}
 	else
 	{
 		OutCookedData.Empty();
-		OutMaterails.Empty();
+		OutMaterials.Empty();
 	}
 
 	return Result;
@@ -1094,6 +1108,7 @@ void ULandscapeHeightfieldCollisionComponent::Serialize(FArchive& Ar)
 		{
 			FName Format = Ar.CookingTarget()->GetPhysicsFormat(nullptr);
 			CookCollsionData(Format, false, CookedCollisionData, CookedPhysicalMaterials);
+			GetDerivedDataCacheRef().Put(*GetHFDDCKeyString(Format, false, HeightfieldGuid), CookedCollisionData);
 		}
 	}
 #endif// WITH_EDITOR
@@ -1259,6 +1274,11 @@ void ULandscapeHeightfieldCollisionComponent::PostLoad()
 #if WITH_EDITOR
 	if (!HasAnyFlags(RF_ClassDefaultObject))
 	{
+		// We can save collision data to DDC 
+		// only if content was already saved with physical materials list, 
+		// specifically after CL# 2349609
+		bShouldSaveCookedDataToDDC = (CookedPhysicalMaterials.Num() > 0);
+
 		if (!CachedLocalBox.IsValid && CachedBoxSphereBounds_DEPRECATED.SphereRadius > 0)
 		{
 			ALandscapeProxy* LandscapeProxy = GetLandscapeProxy();
