Conditionaly populate DDC for landscape collision during cooking

diff --git a/Engine/Source/Runtime/Landscape/Private/LandscapeCollision.cpp b/Engine/Source/Runtime/Landscape/Private/LandscapeCollision.cpp
--- a/Engine/Source/Runtime/Landscape/Private/LandscapeCollision.cpp
+++ b/Engine/Source/Runtime/Landscape/Private/LandscapeCollision.cpp
@@ -1108,7 +1108,11 @@ void ULandscapeHeightfieldCollisionComponent::Serialize(FArchive& Ar)
 		{
 			FName Format = Ar.CookingTarget()->GetPhysicsFormat(nullptr);
 			CookCollsionData(Format, false, CookedCollisionData, CookedPhysicalMaterials);
-			GetDerivedDataCacheRef().Put(*GetHFDDCKeyString(Format, false, HeightfieldGuid), CookedCollisionData);
+			if (bShouldSaveCookedDataToDDC)
+			{
+				GetDerivedDataCacheRef().Put(*GetHFDDCKeyString(Format, false, HeightfieldGuid), CookedCollisionData);
+				// no need to reset bShouldSaveCookedDataToDDC here
+			}
 		}
 	}
 #endif// WITH_EDITOR
