More bulletproof condition whether to use DDC data for landscape collision

diff --git a/Engine/Source/Runtime/Core/Public/UObject/ObjectVersion.h b/Engine/Source/Runtime/Core/Public/UObject/ObjectVersion.h
--- a/Engine/Source/Runtime/Core/Public/UObject/ObjectVersion.h
+++ b/Engine/Source/Runtime/Core/Public/UObject/ObjectVersion.h
@@ -642,6 +642,8 @@ enum EUnrealEngineObjectUE4Version
 	VER_UE4_MOVEMENTCOMPONENT_AXIS_SETTINGS,
 	// Switch to new interactive comments, requires boundry conversion to preserve previous states
 	VER_UE4_GRAPH_INTERACTIVE_COMMENTBUBBLES,
+	// Landscape serializes physical materials for collision objects 
+	VER_UE4_LANDSCAPE_SERIALIZE_PHYSICS_MATERIALS,
 	// -----<new versions can be added before this line>-------------------------------------------------
 	// - this needs to be the last line (see note below)
 	VER_UE4_AUTOMATIC_VERSION_PLUS_ONE,
diff --git a/Engine/Source/Runtime/Landscape/Private/LandscapeCollision.cpp b/Engine/Source/Runtime/Landscape/Private/LandscapeCollision.cpp
--- a/Engine/Source/Runtime/Landscape/Private/LandscapeCollision.cpp
+++ b/Engine/Source/Runtime/Landscape/Private/LandscapeCollision.cpp
@@ -324,10 +324,14 @@ bool ULandscapeHeightfieldCollisionComponent::CookCollsionData(const FName& Form
 {
 #if WITH_PHYSX
 
-	if (GetDerivedDataCacheRef().GetSynchronous(*GetHFDDCKeyString(Format, bUseDefMaterial, HeightfieldGuid), OutCookedData))
+	// Ensure that content was saved with physical materials before using DDC data
+	if (GetLinkerUE4Version() >= VER_UE4_LANDSCAPE_SERIALIZE_PHYSICS_MATERIALS) 
 	{
-		bShouldSaveCookedDataToDDC = false;
-		return true;
+		if (GetDerivedDataCacheRef().GetSynchronous(*GetHFDDCKeyString(Format, bUseDefMaterial, HeightfieldGuid), OutCookedData))
+		{
+			bShouldSaveCookedDataToDDC = false;
+			return true;
+		}
 	}
 				
 	ALandscapeProxy* Proxy = GetLandscapeProxy();
@@ -452,10 +456,14 @@ bool ULandscapeMeshCollisionComponent::CookCollsionData(const FName& Format, boo
 {
 #if WITH_PHYSX
 	
-	if (GetDerivedDataCacheRef().GetSynchronous(*GetHFDDCKeyString(Format, bUseDefMaterial, HeightfieldGuid), OutCookedData))
+	// Ensure that content was saved with physical materials before using DDC data
+	if (GetLinkerUE4Version() >= VER_UE4_LANDSCAPE_SERIALIZE_PHYSICS_MATERIALS)
 	{
-		bShouldSaveCookedDataToDDC = false;
-		return true;
+		if (GetDerivedDataCacheRef().GetSynchronous(*GetHFDDCKeyString(Format, bUseDefMaterial, HeightfieldGuid), OutCookedData))
+		{
+			bShouldSaveCookedDataToDDC = false;
+			return true;
+		}
 	}
 	
 	ALandscapeProxy* Proxy = GetLandscapeProxy();
@@ -1108,11 +1116,7 @@ void ULandscapeHeightfieldCollisionComponent::Serialize(FArchive& Ar)
 		{
 			FName Format = Ar.CookingTarget()->GetPhysicsFormat(nullptr);
 			CookCollsionData(Format, false, CookedCollisionData, CookedPhysicalMaterials);
-			if (bShouldSaveCookedDataToDDC)
-			{
-				GetDerivedDataCacheRef().Put(*GetHFDDCKeyString(Format, false, HeightfieldGuid), CookedCollisionData);
-				// no need to reset bShouldSaveCookedDataToDDC here
-			}
+			GetDerivedDataCacheRef().Put(*GetHFDDCKeyString(Format, false, HeightfieldGuid), CookedCollisionData);
 		}
 	}
 #endif// WITH_EDITOR
@@ -1278,10 +1282,7 @@ void ULandscapeHeightfieldCollisionComponent::PostLoad()
 #if WITH_EDITOR
 	if (!HasAnyFlags(RF_ClassDefaultObject))
 	{
-		// We can save collision data to DDC 
-		// only if content was already saved with physical materials list, 
-		// specifically after CL# 2349609
-		bShouldSaveCookedDataToDDC = (CookedPhysicalMaterials.Num() > 0);
+		bShouldSaveCookedDataToDDC = true;
 
 		if (!CachedLocalBox.IsValid && CachedBoxSphereBounds_DEPRECATED.SphereRadius > 0)
 		{
