Fixed: potential problem with landscape SquareRoot falloff

diff --git a/Engine/Source/Runtime/Landscape/Private/LandscapeRender.cpp b/Engine/Source/Runtime/Landscape/Private/LandscapeRender.cpp
--- a/Engine/Source/Runtime/Landscape/Private/LandscapeRender.cpp
+++ b/Engine/Source/Runtime/Landscape/Private/LandscapeRender.cpp
@@ -1264,7 +1264,7 @@ int32 FLandscapeComponentSceneProxy::CalcLODForSubsectionNoForced(const class FS
 		switch (LODFalloff)
 		{
 			case ELandscapeLODFalloff::SquareRoot:
-				fLOD = FMath::Sqrt(ComponentDistance / LODDistance);
+				fLOD = FMath::Sqrt(FMath::Max(0.f, ComponentDistance / LODDistance));
 				break;
 			default:
 			case ELandscapeLODFalloff::Linear:
@@ -1332,7 +1332,7 @@ void FLandscapeComponentSceneProxy::CalcLODParamsForSubsection(const class FScen
 		switch (LODFalloff)
 		{
 		case ELandscapeLODFalloff::SquareRoot:
-			OutfLOD = FMath::Sqrt(OutDistLOD);
+			OutfLOD = FMath::Sqrt(FMath::Max(0.f, OutDistLOD));
 			break;
 		default:
 		case ELandscapeLODFalloff::Linear:
@@ -1350,7 +1350,7 @@ void FLandscapeComponentSceneProxy::CalcLODParamsForSubsection(const class FScen
 		switch (LODFalloff)
 		{
 			case ELandscapeLODFalloff::SquareRoot:
-				fLOD = FMath::Sqrt(ComponentDistance / LODDistance);
+				fLOD = FMath::Sqrt(FMath::Max(0.f, ComponentDistance / LODDistance));
 				break;
 			default:
 			case ELandscapeLODFalloff::Linear:
