Fixed: typo in landscape lod calculation

diff --git a/Engine/Source/Runtime/Landscape/Private/LandscapeRender.cpp b/Engine/Source/Runtime/Landscape/Private/LandscapeRender.cpp
--- a/Engine/Source/Runtime/Landscape/Private/LandscapeRender.cpp
+++ b/Engine/Source/Runtime/Landscape/Private/LandscapeRender.cpp
@@ -1285,7 +1285,7 @@ int32 FLandscapeComponentSceneProxy::CalcLODForSubsectionNoForced(const class FS
 		switch (LODFalloff)
 		{
 			case ELandscapeLODFalloff::SquareRoot:
-				fLOD = FMath::Sqrt(FMath::Sqrt(Scale / 5.0f));
+				fLOD = FMath::Sqrt(Scale / 5.0f);
 				break;
 			default:
 			case ELandscapeLODFalloff::Linear:
