Fix incorrect landscape component offsets (UDN support request)

diff --git a/Engine/Source/Runtime/Landscape/Private/Landscape.cpp b/Engine/Source/Runtime/Landscape/Private/Landscape.cpp
--- a/Engine/Source/Runtime/Landscape/Private/Landscape.cpp
+++ b/Engine/Source/Runtime/Landscape/Private/Landscape.cpp
@@ -376,10 +376,10 @@ void ULandscapeComponent::PostLoad()
 {
 	Super::PostLoad();
 
-	// Ensure that the component's lighting settings matches the actor's.
 	ALandscapeProxy* LandscapeProxy = GetLandscapeProxy();
 	if (ensure(LandscapeProxy))
 	{
+		// Ensure that the component's lighting settings matches the actor's.
 		bCastStaticShadow = LandscapeProxy->bCastStaticShadow;
 		bCastShadowAsTwoSided = LandscapeProxy->bCastShadowAsTwoSided;
 
@@ -401,6 +401,18 @@ void ULandscapeComponent::PostLoad()
 				WeightmapLayerAllocations[i].LayerName_DEPRECATED = NAME_None;
 			}
 		}
+
+		// check SectionBaseX/Y are correct
+		int32 CheckSectionBaseX = FMath::RoundToInt(RelativeLocation.X) + LandscapeProxy->LandscapeSectionOffset.X;
+		int32 CheckSectionBaseY = FMath::RoundToInt(RelativeLocation.Y) + LandscapeProxy->LandscapeSectionOffset.Y;
+		if (CheckSectionBaseX != SectionBaseX ||
+			CheckSectionBaseY != SectionBaseY)
+		{
+			UE_LOG(LogLandscape, Warning, TEXT("LandscapeComponent SectionBaseX disagrees with its location, attempted automated fix: '%s', %d,%d vs %d,%d."),
+				*GetFullName(), SectionBaseX, SectionBaseY, CheckSectionBaseX, CheckSectionBaseY);
+			SectionBaseX = CheckSectionBaseX;
+			SectionBaseY = CheckSectionBaseY;
+		}
 	}
 
 #if WITH_EDITOR
