Added "RemoveLandscapeXYOffsets" console command to completely remove any XY offset data and restore the landscape to pure heightfield.

diff --git a/Engine/Source/Editor/UnrealEd/Classes/Editor/UnrealEdEngine.h b/Engine/Source/Editor/UnrealEd/Classes/Editor/UnrealEdEngine.h
--- a/Engine/Source/Editor/UnrealEd/Classes/Editor/UnrealEdEngine.h
+++ b/Engine/Source/Editor/UnrealEd/Classes/Editor/UnrealEdEngine.h
@@ -725,6 +725,7 @@ class UNREALED_API UUnrealEdEngine : public UEditorEngine, public FNotifyHook
 	bool HandleUpdateLandscapeEditorDataCommand( const TCHAR* Str, FOutputDevice& Ar, UWorld* InWorld );
 	bool HandleUpdateLandscapeMICCommand( const TCHAR* Str, FOutputDevice& Ar, UWorld* InWorld );
 	bool HandleRecreateLandscapeCollisionCommand(const TCHAR* Str, FOutputDevice& Ar, UWorld* InWorld);
+	bool HandleRemoveLandscapeXYOffsetsCommand(const TCHAR* Str, FOutputDevice& Ar, UWorld* InWorld);
 	bool HandleConvertMatineesCommand( const TCHAR* Str, FOutputDevice& Ar, UWorld* InWorld );
 	bool HandleDisasmScriptCommand( const TCHAR* Str, FOutputDevice& Ar );	
 
diff --git a/Engine/Source/Editor/UnrealEd/Private/UnrealEdSrv.cpp b/Engine/Source/Editor/UnrealEd/Private/UnrealEdSrv.cpp
--- a/Engine/Source/Editor/UnrealEd/Private/UnrealEdSrv.cpp
+++ b/Engine/Source/Editor/UnrealEd/Private/UnrealEdSrv.cpp
@@ -639,6 +639,19 @@ bool UUnrealEdEngine::HandleRecreateLandscapeCollisionCommand(const TCHAR* Str,
 	return true;
 }
 
+bool UUnrealEdEngine::HandleRemoveLandscapeXYOffsetsCommand(const TCHAR* Str, FOutputDevice& Ar, UWorld* InWorld)
+{
+	if (!PlayWorld && InWorld && InWorld->GetWorldSettings())
+	{
+		for (auto It = GetLandscapeInfoMap(InWorld).Map.CreateIterator(); It; ++It)
+		{
+			ULandscapeInfo* Info = It.Value();
+			Info->RemoveXYOffsets();
+		}
+	}
+	return true;
+}
+
 bool UUnrealEdEngine::HandleConvertMatineesCommand( const TCHAR* Str, FOutputDevice& Ar, UWorld* InWorld )
 {
 	FVector StartLocation= FVector::ZeroVector;
@@ -843,6 +856,12 @@ bool UUnrealEdEngine::Exec( UWorld* InWorld, const TCHAR* Stream, FOutputDevice&
 		UWorld* World = GetEditorWorldContext().World();
 		return HandleRecreateLandscapeCollisionCommand(Str, Ar, World);
 	}
+	else if (FParse::Command(&Str, TEXT("RemoveLandscapeXYOffsets")))
+	{
+		// InWorld above is the PIE world if PIE is active, but this is specifically an editor command
+		UWorld* World = GetEditorWorldContext().World();
+		return HandleRemoveLandscapeXYOffsetsCommand(Str, Ar, World);
+	}
 #endif // WITH_EDITOR
 	else if( FParse::Command(&Str, TEXT("CONVERTMATINEES")) )
 	{
diff --git a/Engine/Source/Runtime/Landscape/Classes/LandscapeInfo.h b/Engine/Source/Runtime/Landscape/Classes/LandscapeInfo.h
--- a/Engine/Source/Runtime/Landscape/Classes/LandscapeInfo.h
+++ b/Engine/Source/Runtime/Landscape/Classes/LandscapeInfo.h
@@ -241,5 +241,7 @@ public:
 	LANDSCAPE_API void FixupProxiesWeightmaps();
 
 	LANDSCAPE_API void RecreateCollisionComponents();
+
+	LANDSCAPE_API void RemoveXYOffsets();
 #endif
 };
diff --git a/Engine/Source/Runtime/Landscape/Classes/LandscapeProxy.h b/Engine/Source/Runtime/Landscape/Classes/LandscapeProxy.h
--- a/Engine/Source/Runtime/Landscape/Classes/LandscapeProxy.h
+++ b/Engine/Source/Runtime/Landscape/Classes/LandscapeProxy.h
@@ -412,6 +412,10 @@ public:
 	/** Recreate all collision components based on render component */
 	LANDSCAPE_API void RecreateCollisionComponents();
 
+	/** Remove all XYOffset values */
+	LANDSCAPE_API void RemoveXYOffsets();
+
+
 	LANDSCAPE_API static ULandscapeMaterialInstanceConstant* GetLayerThumbnailMIC(UMaterialInterface* LandscapeMaterial, FName LayerName, UTexture2D* ThumbnailWeightmap, UTexture2D* ThumbnailHeightmap, ALandscapeProxy* Proxy);
 
 	LANDSCAPE_API void Import(FGuid Guid, int32 VertsX, int32 VertsY, 
diff --git a/Engine/Source/Runtime/Landscape/Private/LandscapeEdit.cpp b/Engine/Source/Runtime/Landscape/Private/LandscapeEdit.cpp
--- a/Engine/Source/Runtime/Landscape/Private/LandscapeEdit.cpp
+++ b/Engine/Source/Runtime/Landscape/Private/LandscapeEdit.cpp
@@ -3458,6 +3458,33 @@ void ULandscapeLayerInfoObject::PostLoad()
 	}
 }
 
+void ALandscapeProxy::RemoveXYOffsets()
+{
+	bool bFoundXYOffset = false;
+
+	for (int32 i = 0; i < LandscapeComponents.Num(); ++i)
+	{
+		ULandscapeComponent* Comp = LandscapeComponents[i];
+		if (Comp && Comp->XYOffsetmapTexture)
+		{
+			Comp->XYOffsetmapTexture->SetFlags(RF_Transactional);
+			Comp->XYOffsetmapTexture->Modify();
+			Comp->XYOffsetmapTexture->MarkPackageDirty();
+			Comp->XYOffsetmapTexture->ClearFlags(RF_Standalone);
+			Comp->Modify();
+			Comp->MarkPackageDirty();
+			Comp->XYOffsetmapTexture = NULL;
+			Comp->MarkRenderStateDirty();
+			bFoundXYOffset = true;
+		}
+	}
+
+	if (bFoundXYOffset)
+	{
+		RecreateCollisionComponents();
+	}
+}
+
 void ALandscapeProxy::RecreateCollisionComponents()
 {
 	// Clear old CollisionComponent containers
@@ -3512,6 +3539,20 @@ void ULandscapeInfo::RecreateCollisionComponents()
 	}
 }
 
+void ULandscapeInfo::RemoveXYOffsets()
+{
+	if (LandscapeActor.IsValid())
+	{
+		LandscapeActor->RemoveXYOffsets();
+	}
+
+	for (auto It = Proxies.CreateConstIterator(); It; ++It)
+	{
+		ALandscapeProxy* Proxy = (*It);
+		Proxy->RemoveXYOffsets();
+	}
+}
+
 namespace
 {
 	inline float AdjustStaticLightingResolution(float StaticLightingResolution, int32 NumSubsections, int32 SubsectionSizeQuads, int32 ComponentSizeQuads)
