Fix for crash loading levels that predated VER_UE4_LANDSCAPE_COLLISION_DATA_COOKING that were resaved using a commandlet.

diff --git a/Engine/Source/Runtime/Landscape/Private/LandscapeCollision.cpp b/Engine/Source/Runtime/Landscape/Private/LandscapeCollision.cpp
--- a/Engine/Source/Runtime/Landscape/Private/LandscapeCollision.cpp
+++ b/Engine/Source/Runtime/Landscape/Private/LandscapeCollision.cpp
@@ -274,6 +274,13 @@ void ULandscapeHeightfieldCollisionComponent::CreateCollisionObject()
 			ExistingHeightfieldRef = GSharedHeightfieldRefs.FindRef(HeightfieldGuid);
 		}
 
+		// This should only occur if a level prior to VER_UE4_LANDSCAPE_COLLISION_DATA_COOKING 
+		// was resaved using a commandlet and not saved in the editor.
+		if (CookedPhysicalMaterials.Num() == 0)
+		{
+			bCheckDDC = false;
+		}
+
 		if (ExistingHeightfieldRef)
 		{
 			HeightfieldRef = ExistingHeightfieldRef;
