Fix LandscapeLayerSample node causing compile errors when a landscape component doesn't use the layer
(from GDC branch CL 2374577)

diff --git a/Engine/Source/Runtime/Landscape/Private/Materials/MaterialExpressionLandscapeLayerSample.cpp b/Engine/Source/Runtime/Landscape/Private/Materials/MaterialExpressionLandscapeLayerSample.cpp
--- a/Engine/Source/Runtime/Landscape/Private/Materials/MaterialExpressionLandscapeLayerSample.cpp
+++ b/Engine/Source/Runtime/Landscape/Private/Materials/MaterialExpressionLandscapeLayerSample.cpp
@@ -41,8 +41,15 @@ FGuid& UMaterialExpressionLandscapeLayerSample::GetParameterExpressionId()
 int32 UMaterialExpressionLandscapeLayerSample::Compile(class FMaterialCompiler* Compiler, int32 OutputIndex, int32 MultiplexIndex)
 {
 	const int32 WeightCode = Compiler->StaticTerrainLayerWeight(ParameterName, Compiler->Constant(PreviewWeight));
-
-	return WeightCode;
+	if (WeightCode == INDEX_NONE)
+	{
+		// layer is not used in this component, sample value is 0.
+		return Compiler->Constant(0.f);
+	}
+	else
+	{
+		return WeightCode;
+	}
 }
 
 
