Make landscape LOD colorizaiton work in wireframe mode

diff --git a/Engine/Source/Runtime/Landscape/Private/LandscapeRender.cpp b/Engine/Source/Runtime/Landscape/Private/LandscapeRender.cpp
--- a/Engine/Source/Runtime/Landscape/Private/LandscapeRender.cpp
+++ b/Engine/Source/Runtime/Landscape/Private/LandscapeRender.cpp
@@ -1499,8 +1499,16 @@ void FLandscapeComponentSceneProxy::GetDynamicMeshElements(const TArray<const FS
 						LODMesh.MaterialRenderProxy = LODMaterialInstance;
 						Collector.RegisterOneFrameMaterialProxy(LODMaterialInstance);
 
-						LODMesh.bCanApplyViewModeOverrides = true;
-						LODMesh.bUseWireframeSelectionColoring = IsSelected();
+						if (ViewFamily.EngineShowFlags.Wireframe)
+						{
+							LODMesh.bCanApplyViewModeOverrides = false;
+							LODMesh.bWireframe = true;
+						}
+						else
+						{
+							LODMesh.bCanApplyViewModeOverrides = true;
+							LODMesh.bUseWireframeSelectionColoring = IsSelected();
+						}
 
 						Collector.AddMesh(ViewIndex, LODMesh);
 
