Fix Landscape cracks:
This was mostly caused by rounding to the nearest vertex on multiple LOD transitions, which wasn't stable (On one component in LOD 0, vertex A in LOD 0 would round to vertex B in LOD 1 and vertex C in LOD 2 (vertex C being the nearest to vertex A, not necessarily vertex B), but the next component would be in LOD 1 and so would start with vertex B, and round it to a different LOD 2 vertex)
Partly responsible was also the "DistLOD" hack, because it used different shader code for if LodValues.x (aka "DistLOD") was <1 or >1, so would cause a beautiful seam between components which hit each case

There will still be temporary seams between landscape components that are still streaming their heightmap textures, a solution to that is in the works.

Also:
Removed FLandscapeXYOffsetVertexFactoryVertexShaderParameters and merged it with FLandscapeVertexFactoryVertexShaderParameters to remove a large amount of code duplication, and did some cleanup of LandscapeRender.cpp/.h (e.g. autoformatting)

diff --git a/Engine/Shaders/LandscapeVertexFactory.usf b/Engine/Shaders/LandscapeVertexFactory.usf
--- a/Engine/Shaders/LandscapeVertexFactory.usf
+++ b/Engine/Shaders/LandscapeVertexFactory.usf
@@ -12,11 +12,22 @@
 float4 LightMapCoordinateScaleBias;
 float4 ShadowMapCoordinateScaleBias;
 
-// Vertex shader
-float4 LodBias; 
+/* SM4+:
+   x = unused
+   y = unused
+   z = Heightmap texture LOD difference from current LOD to highest LOD
+   w = XYOffset texture LOD difference from current LOD to highest LOD
+
+   ES2 / ES3.1:
+   x = unused
+   y = unused
+   z = TexCoordOffset.x
+   w = TexCoordOffset.y
+*/
+float4 LodBias;
 
-/* x = LOD
-   y = convert current LOD coordinates into highest LOD coordinates
+/* x = current LOD
+   y = unused
    z = SectionSizeQuads in current LOD
    w = 1/SectionSizeQuads in current LOD
 */
@@ -419,51 +430,36 @@ float2 GetShadowMapCoordinate(FVertexFactoryInterpolantsVSToPS Interpolants)
 }
 #endif
 
-FVertexFactoryIntermediates GetVertexFactoryIntermediates(FVertexFactoryInput Input)
+float CalcLOD(float2 xy, float2 Subsection)
 {
-	FVertexFactoryIntermediates Intermediates;
-
-#if FEATURE_LEVEL >= FEATURE_LEVEL_SM4
-	Intermediates.InputPosition = Input.Position;
-#else
-	Intermediates.InputPosition = Input.PackedPosition * 255.f;
-	float2 SubXY = frac((Intermediates.InputPosition.zw) * .5f) * 2.f;
-	float2 LOD45 = (Intermediates.InputPosition.zw - SubXY) / 255.f;
-	Intermediates.InputPosition.zw = SubXY;
-#endif
-
-
 	// LOD value based on Barycentric coordinate
-	float2 xy = Intermediates.InputPosition.xy * LodValues.w;
-	float4 L0 = float4( xy.y, xy.x, (1-xy.x), (1-xy.y) ) * 2;
 
+	float4 L0 = float4(xy.y, xy.x, (1 - xy.x), (1 - xy.y)) * 2;
 
-	// Which subsection we're rendering (x/y encoded in Intermediates.InputPosition.z/w)
-	float SectionLod;
 	float4 LODCalculated4;
-	if( Intermediates.InputPosition.w > 0.5f )
+	if (Subsection.y > 0.5f)
 	{
-		if( Intermediates.InputPosition.z > 0.5f )
+		if (Subsection.x > 0.5f)
 		{
-			SectionLod = SectionLods.w;
+			float SectionLod = SectionLods.w;
 			LODCalculated4 = L0 * SectionLod + (float4(1, 1, 1, 1) - L0) * NeighborSectionLod[3];
 		}
 		else
 		{
-			SectionLod = SectionLods.z;
+			float SectionLod = SectionLods.z;
 			LODCalculated4 = L0 * SectionLod + (float4(1, 1, 1, 1) - L0) * NeighborSectionLod[2];
 		}
 	}
 	else
 	{
-		if( Intermediates.InputPosition.z > 0.5f )
+		if (Subsection.x > 0.5f)
 		{
-			SectionLod = SectionLods.y;
-			LODCalculated4 = L0 * SectionLod + (float4(1,1,1,1)-L0) * NeighborSectionLod[1];
+			float SectionLod = SectionLods.y;
+			LODCalculated4 = L0 * SectionLod + (float4(1, 1, 1, 1) - L0) * NeighborSectionLod[1];
 		}
 		else
 		{
-			SectionLod = SectionLods.x;
+			float SectionLod = SectionLods.x;
 			LODCalculated4 = L0 * SectionLod + (float4(1, 1, 1, 1) - L0) * NeighborSectionLod[0];
 		}
 	}
@@ -475,9 +471,9 @@ FVertexFactoryIntermediates GetVertexFactoryIntermediates(FVertexFactoryInput In
 	//   | / 3 \ |                 | /      \ |    /   \     
 	// 2 --------- 3               2          2  1 ----- 2   
 	float LODCalculated;
-	if( LodValues.z < (Intermediates.InputPosition.x + Intermediates.InputPosition.y) )
+	if ((xy.x + xy.y) > 1)
 	{
-		if( Intermediates.InputPosition.x < Intermediates.InputPosition.y) 
+		if (xy.x < xy.y)
 		{
 			LODCalculated = LODCalculated4.w;
 		}
@@ -488,7 +484,7 @@ FVertexFactoryIntermediates GetVertexFactoryIntermediates(FVertexFactoryInput In
 	}
 	else
 	{
-		if( Intermediates.InputPosition.x < Intermediates.InputPosition.y) 
+		if (xy.x < xy.y)
 		{
 			LODCalculated = LODCalculated4.y;
 		}
@@ -498,36 +494,45 @@ FVertexFactoryIntermediates GetVertexFactoryIntermediates(FVertexFactoryInput In
 		}
 	}
 
-	float LodValue = floor(LODCalculated);
-	float MorphAlpha = frac(LODCalculated);
-
-	if (LodValues.x < 1.0f)
-	{
-		MorphAlpha = saturate((MorphAlpha - LodBias.x) * (LodValue+1) * 2.f);
-	}
-	else
-	{
-		MorphAlpha = saturate((MorphAlpha - LodBias.x) * LodBias.y);
-	}
+	return LODCalculated;
+}
 
-	// Base to Actual LOD, Base to Next LOD
-	float2 CoordTranslate = float2( LandscapeParameters.SubsectionSizeVertsLayerUVPan.x * pow(2, -LodValue) - 1, max(LandscapeParameters.SubsectionSizeVertsLayerUVPan.x * pow(2, -LodValue-1), 2) - 1 ) * LandscapeParameters.SubsectionSizeVertsLayerUVPan.y;
+FVertexFactoryIntermediates GetVertexFactoryIntermediates(FVertexFactoryInput Input)
+{
+	FVertexFactoryIntermediates Intermediates;
 
-	// Position for current LOD in base LOD units
 #if FEATURE_LEVEL >= FEATURE_LEVEL_SM4
-	float2 InputPositionLODBase = Intermediates.InputPosition.xy * LodValues.y;
+	Intermediates.InputPosition = Input.Position;
 #else
-	float2 InputPositionLODBase = Intermediates.InputPosition.xy;
+	Intermediates.InputPosition = Input.PackedPosition * 255.f;
+	float2 SubXY = frac((Intermediates.InputPosition.zw) * .5f) * 2.f;
+	float2 LOD45 = (Intermediates.InputPosition.zw - SubXY) / 255.f;
+	Intermediates.InputPosition.zw = SubXY;
 #endif
 
-	// Position for actual LOD in base LOD units
-	float2 ActualLODCoordsInt = floor( InputPositionLODBase * CoordTranslate.x + float2(0.5,0.5) );
+	float2 xy = Intermediates.InputPosition.xy * LodValues.w;
+
+	float LODCalculated = CalcLOD(xy, Intermediates.InputPosition.zw);
+
+	float LodValue = floor(LODCalculated);
+	float MorphAlpha = frac(LODCalculated);
+
+	// Base to Actual LOD, Base to Next LOD
+	float2 CoordTranslate = float2( LandscapeParameters.SubsectionSizeVertsLayerUVPan.x * pow(2, -LodValue) - 1, max(LandscapeParameters.SubsectionSizeVertsLayerUVPan.x * pow(2, -LodValue-1), 2) - 1 ) * LandscapeParameters.SubsectionSizeVertsLayerUVPan.y;
+
+	// InputPositionLODAdjusted : Position for actual LOD in base LOD units
+	float2 ActualLODCoordsInt = floor(Intermediates.InputPosition.xy * pow(2, -(LodValue - LodValues.x)));
 	float2 InputPositionLODAdjusted = ActualLODCoordsInt / CoordTranslate.x;
 
-	// Position for next LOD in base LOD units
-	float2 NextLODCoordsInt = floor( InputPositionLODBase * CoordTranslate.y + float2(0.5,0.5) );
+	// InputPositionNextLOD : Position for next LOD in base LOD units
+	float2 NextLODCoordsInt = floor(ActualLODCoordsInt * 0.5);
 	float2 InputPositionNextLOD = NextLODCoordsInt / CoordTranslate.y;
 
+	// Adjust MorphAlpha based on actual LOD - Improves multi-LOD transitions
+	//xy = ActualLODCoordsInt / (LandscapeParameters.SubsectionSizeVertsLayerUVPan.x * pow(2, -LodValue) - 1);
+	//LODCalculated = CalcLOD(xy, Intermediates.InputPosition.zw);
+	//MorphAlpha = saturate(LODCalculated - LodValue);
+
 	// Get the height and normal XY for current and next LOD out of the textures
 #if FEATURE_LEVEL >= FEATURE_LEVEL_SM4
 	float2 SampleCoords = InputPositionLODAdjusted * LandscapeParameters.HeightmapUVScaleBias.xy + LandscapeParameters.HeightmapUVScaleBias.zw + 0.5*LandscapeParameters.HeightmapUVScaleBias.xy + Intermediates.InputPosition.zw * LandscapeParameters.SubsectionOffsetParams.xy;
@@ -586,7 +591,7 @@ FVertexFactoryIntermediates GetVertexFactoryIntermediates(FVertexFactoryInput In
 #endif
 
 
-#if LANDSCAPE_XYOFFSET
+#if LANDSCAPE_XYOFFSET // FEATURE_LEVEL >= FEATURE_LEVEL_SM4 only
 	float2 SampleCoords2 = float2(InputPositionLODAdjusted * LandscapeParameters.WeightmapUVScaleBias.xy + LandscapeParameters.WeightmapUVScaleBias.zw + Intermediates.InputPosition.zw * LandscapeParameters.SubsectionOffsetParams.zz);
 	float4 OffsetValue = Texture2DSampleLevel( XYOffsetmapTexture, XYOffsetmapTextureSampler, SampleCoords2, LodValue-LodBias.w );
 	float2 SampleCoordsNextLOD2 = float2(InputPositionNextLOD * LandscapeParameters.WeightmapUVScaleBias.xy + LandscapeParameters.WeightmapUVScaleBias.zw + Intermediates.InputPosition.zw * LandscapeParameters.SubsectionOffsetParams.zz);
diff --git a/Engine/Source/Runtime/Landscape/Private/Landscape.cpp b/Engine/Source/Runtime/Landscape/Private/Landscape.cpp
--- a/Engine/Source/Runtime/Landscape/Private/Landscape.cpp
+++ b/Engine/Source/Runtime/Landscape/Private/Landscape.cpp
@@ -89,7 +89,7 @@ ULandscapeComponent::ULandscapeComponent(const FObjectInitializer& ObjectInitial
 
 	bBoundsChangeTriggersStreamingDataRebuild = true;
 	ForcedLOD = -1;
-	// Neighbor LOD and LODBias are saved in BYTE, so need to convert to range [-128:127]
+	// Neighbor LOD and LODBias are saved in a legacy uint8 form instead of int8. Todo: convert
 	NeighborLOD[0] = 255;
 	NeighborLOD[1] = 255;
 	NeighborLOD[2] = 255;
@@ -1859,7 +1859,7 @@ void ULandscapeComponent::PostDuplicate(bool bDuplicateForPIE)
 }
 
 // Generate a new guid to force a recache of all landscape derived data
-#define LANDSCAPE_FULL_DERIVEDDATA_VER			TEXT("28f0da4d085e48a680482c9835907e2f")
+#define LANDSCAPE_FULL_DERIVEDDATA_VER			TEXT("016D326F3A954BBA9CCDFA00CEFA31E9")
 
 FString FLandscapeComponentDerivedData::GetDDCKeyString(const FGuid& StateId)
 {
diff --git a/Engine/Source/Runtime/Landscape/Private/LandscapeEdit.cpp b/Engine/Source/Runtime/Landscape/Private/LandscapeEdit.cpp
--- a/Engine/Source/Runtime/Landscape/Private/LandscapeEdit.cpp
+++ b/Engine/Source/Runtime/Landscape/Private/LandscapeEdit.cpp
@@ -3766,7 +3766,7 @@ void ALandscape::PostEditChangeProperty(FPropertyChangedEvent& PropertyChangedEv
 	}
 	else if (GIsEditor && PropertyName == FName(TEXT("MaxLODLevel")))
 	{
-		MaxLODLevel = FMath::Min<int32>(MaxLODLevel, FMath::CeilLogTwo(SubsectionSizeQuads + 1) - 1);
+		MaxLODLevel = FMath::Clamp<int32>(MaxLODLevel, 0, FMath::CeilLogTwo(SubsectionSizeQuads + 1) - 1);
 		bPropagateToProxies = true;
 	}
 	else if (PropertyName == FName(TEXT("LODDistanceFactor")))
@@ -5113,15 +5113,11 @@ void ULandscapeComponent::GeneratePlatformVertexData()
 		{
 			int32 MipSizeX = HeightmapTexture->Source.GetSizeX() >> Mip;
 
-			int32 MipSubsectionSizeVerts = (SubsectionSizeVerts >> Mip);
-			int32 MipSubsectionSizeQuads = MipSubsectionSizeVerts - 1;
-
 			int32 CurrentMipOfsX = BaseMipOfsX >> Mip;
 			int32 CurrentMipOfsY = BaseMipOfsY >> Mip;
 
-			float MipRatio = (float)MipSubsectionSizeQuads / (float)SubsectionSizeQuads; // Morph Base to current MIP
-			int32 MipX = FMath::RoundToInt((float)X * MipRatio);
-			int32 MipY = FMath::RoundToInt((float)Y * MipRatio);
+			int32 MipX = X >> Mip;
+			int32 MipY = Y >> Mip;
 
 			FColor* CurrentMipSrcRow = HeightmapMipData[Mip] + (CurrentMipOfsY + MipY) * MipSizeX + CurrentMipOfsX;
 			uint16 Height = CurrentMipSrcRow[MipX].R << 8 | CurrentMipSrcRow[MipX].G;
diff --git a/Engine/Source/Runtime/Landscape/Private/LandscapeRender.cpp b/Engine/Source/Runtime/Landscape/Private/LandscapeRender.cpp
--- a/Engine/Source/Runtime/Landscape/Private/LandscapeRender.cpp
+++ b/Engine/Source/Runtime/Landscape/Private/LandscapeRender.cpp
@@ -23,7 +23,7 @@ LandscapeRender.cpp: New terrain rendering
 #include "EngineGlobals.h"
 #include "UnrealEngine.h"
 
-IMPLEMENT_UNIFORM_BUFFER_STRUCT(FLandscapeUniformShaderParameters,TEXT("LandscapeParameters"));
+IMPLEMENT_UNIFORM_BUFFER_STRUCT(FLandscapeUniformShaderParameters, TEXT("LandscapeParameters"));
 
 #define LANDSCAPE_LOD_DISTANCE_FACTOR 2.f
 #define LANDSCAPE_MAX_COMPONENT_SIZE 255
@@ -44,13 +44,13 @@ namespace
 		const float FindVertexScoreLastTriScore = 0.75f;
 
 		float Score = 0.0f;
-		if ( CachePosition < 0 )
+		if (CachePosition < 0)
 		{
 			// Vertex is not in FIFO cache - no score.
 		}
 		else
 		{
-			if ( CachePosition < 3 )
+			if (CachePosition < 3)
 			{
 				// This vertex was used in the last triangle,
 				// so it has a fixed score, whichever of the three
@@ -61,11 +61,11 @@ namespace
 			}
 			else
 			{
-				check( CachePosition < (int32)VertexCacheSize );
+				check(CachePosition < (int32)VertexCacheSize);
 				// Points for being high in the cache.
-				const float Scaler = 1.0f / ( VertexCacheSize - 3 );
-				Score = 1.0f - ( CachePosition - 3 ) * Scaler;
-				Score = FMath::Pow( Score, FindVertexScoreCacheDecayPower );
+				const float Scaler = 1.0f / (VertexCacheSize - 3);
+				Score = 1.0f - (CachePosition - 3) * Scaler;
+				Score = FMath::Pow(Score, FindVertexScoreCacheDecayPower);
 			}
 		}
 
@@ -81,7 +81,7 @@ namespace
 
 		// Bonus points for having a low number of tris still to
 		// use the vert, so we get rid of lone verts quickly.
-		float ValenceBoost = FMath::Pow(float(numActiveFaces), -FindVertexScoreValenceBoostPower );
+		float ValenceBoost = FMath::Pow(float(numActiveFaces), -FindVertexScoreValenceBoostPower);
 		Score += FindVertexScoreValenceBoostScale * ValenceBoost;
 
 		return Score;
@@ -89,21 +89,21 @@ namespace
 
 	const uint32 MaxVertexCacheSize = 64;
 	const uint32 MaxPrecomputedVertexValenceScores = 64;
-	float VertexCacheScores[MaxVertexCacheSize+1][MaxVertexCacheSize];
+	float VertexCacheScores[MaxVertexCacheSize + 1][MaxVertexCacheSize];
 	float VertexValenceScores[MaxPrecomputedVertexValenceScores];
 	bool bVertexScoresComputed = false; //ComputeVertexScores();
 
 	bool ComputeVertexScores()
 	{
-		for (uint32 CacheSize=0; CacheSize<=MaxVertexCacheSize; ++CacheSize)
+		for (uint32 CacheSize = 0; CacheSize <= MaxVertexCacheSize; ++CacheSize)
 		{
-			for (uint32 CachePos=0; CachePos<CacheSize; ++CachePos)
+			for (uint32 CachePos = 0; CachePos < CacheSize; ++CachePos)
 			{
 				VertexCacheScores[CacheSize][CachePos] = ComputeVertexCacheScore(CachePos, CacheSize);
 			}
 		}
 
-		for (uint32 Valence=0; Valence<MaxPrecomputedVertexValenceScores; ++Valence)
+		for (uint32 Valence = 0; Valence < MaxPrecomputedVertexValenceScores; ++Valence)
 		{
 			VertexValenceScores[Valence] = ComputeVertexValenceScore(Valence);
 		}
@@ -125,7 +125,7 @@ namespace
 	{
 		check(bVertexScoresComputed);
 
-		if ( NumActiveFaces == 0 )
+		if (NumActiveFaces == 0)
 		{
 			// No tri needs this vertex!
 			return -1.0f;
@@ -226,11 +226,11 @@ namespace
 		}
 
 		// fill out face list per vertex
-		for (uint32 i = 0; i < IndexCount; i+=3)
+		for (uint32 i = 0; i < IndexCount; i += 3)
 		{
 			for (uint32 j = 0; j < 3; ++j)
 			{
-				uint32 Index = InIndexList[i+j];
+				uint32 Index = InIndexList[i + j];
 				OptimizeVertexData& VertexData = VertexDataList[Index];
 				ActiveFaceList[VertexData.ActiveFaceListStart + VertexData.ActiveFaceListSize] = i;
 				VertexData.ActiveFaceListSize++;
@@ -241,9 +241,9 @@ namespace
 		ProcessedFaceList.Empty(IndexCount);
 		ProcessedFaceList.AddZeroed(IndexCount);
 
-		uint32 VertexCacheBuffer[(MaxVertexCacheSize+3)*2];
+		uint32 VertexCacheBuffer[(MaxVertexCacheSize + 3) * 2];
 		uint32* Cache0 = VertexCacheBuffer;
-		uint32* Cache1 = VertexCacheBuffer + (MaxVertexCacheSize+3);
+		uint32* Cache1 = VertexCacheBuffer + (MaxVertexCacheSize + 3);
 		uint32 EntriesInCache0 = 0;
 
 		uint32 BestFace = 0;
@@ -265,7 +265,7 @@ namespace
 						float FaceScore = 0.f;
 						for (uint32 k = 0; k < 3; ++k)
 						{
-							uint32 Index = InIndexList[Face+k];
+							uint32 Index = InIndexList[Face + k];
 							OptimizeVertexData& VertexData = VertexDataList[Index];
 							check(VertexData.ActiveFaceListSize > 0);
 							check(VertexData.CachePos0 >= LRUCacheSize);
@@ -294,8 +294,8 @@ namespace
 			// add bestFace to LRU cache and to newIndexList
 			for (uint32 V = 0; V < 3; ++V)
 			{
-				INDEX_TYPE Index = InIndexList[BestFace+V];
-				OutIndexList[i+V] = Index;
+				INDEX_TYPE Index = InIndexList[BestFace + V];
+				OutIndexList[i + V] = Index;
 
 				OptimizeVertexData& VertexData = VertexDataList[Index];
 
@@ -360,9 +360,9 @@ namespace
 				{
 					uint32 Face = ActiveFaceList[VertexData.ActiveFaceListStart + j];
 					float FaceScore = 0.f;
-					for (uint32 V=0; V < 3; V++)
+					for (uint32 V = 0; V < 3; V++)
 					{
-						uint32 FaceIndex = InIndexList[Face+V];
+						uint32 FaceIndex = InIndexList[Face + V];
 						OptimizeVertexData& FaceVertexData = VertexDataList[FaceIndex];
 						FaceScore += FaceVertexData.Score;
 					}
@@ -387,21 +387,21 @@ namespace
 struct FLandscapeDebugOptions
 {
 	FLandscapeDebugOptions()
-	:	bShowPatches(false)
-	,	bDisableStatic(false)
-	,	bDisableCombine(false)
-	,	PatchesConsoleCommand(
-		TEXT( "Landscape.Patches" ),
-		TEXT( "Show/hide Landscape patches" ),
-		FConsoleCommandDelegate::CreateRaw( this, &FLandscapeDebugOptions::Patches ) )
-	,	StaticConsoleCommand(
-		TEXT( "Landscape.Static" ),
-		TEXT( "Enable/disable Landscape static drawlists" ),
-		FConsoleCommandDelegate::CreateRaw( this, &FLandscapeDebugOptions::Static ) )
-	,	CombineConsoleCommand(
-		TEXT( "Landscape.Combine" ),
-		TEXT( "Enable/disable Landscape component combining" ),
-		FConsoleCommandDelegate::CreateRaw( this, &FLandscapeDebugOptions::Combine) )
+		: bShowPatches(false)
+		, bDisableStatic(false)
+		, bDisableCombine(false)
+		, PatchesConsoleCommand(
+		TEXT("Landscape.Patches"),
+		TEXT("Show/hide Landscape patches"),
+		FConsoleCommandDelegate::CreateRaw(this, &FLandscapeDebugOptions::Patches))
+		, StaticConsoleCommand(
+		TEXT("Landscape.Static"),
+		TEXT("Enable/disable Landscape static drawlists"),
+		FConsoleCommandDelegate::CreateRaw(this, &FLandscapeDebugOptions::Static))
+		, CombineConsoleCommand(
+		TEXT("Landscape.Combine"),
+		TEXT("Enable/disable Landscape component combining"),
+		FConsoleCommandDelegate::CreateRaw(this, &FLandscapeDebugOptions::Combine))
 	{
 	}
 
@@ -441,33 +441,33 @@ LANDSCAPE_API bool GLandscapeEditModeActive = false;
 LANDSCAPE_API ELandscapeViewMode::Type GLandscapeViewMode = ELandscapeViewMode::Normal;
 LANDSCAPE_API int32 GLandscapeEditRenderMode = ELandscapeEditRenderMode::None;
 LANDSCAPE_API int32 GLandscapePreviewMeshRenderMode = 0;
-UMaterial* GLayerDebugColorMaterial = NULL;
-UMaterialInstanceConstant* GSelectionColorMaterial = NULL;
-UMaterialInstanceConstant* GSelectionRegionMaterial = NULL;
-UMaterialInstanceConstant* GMaskRegionMaterial = NULL;
-UTexture2D* GLandscapeBlackTexture = NULL;
+UMaterial* GLayerDebugColorMaterial = nullptr;
+UMaterialInstanceConstant* GSelectionColorMaterial = nullptr;
+UMaterialInstanceConstant* GSelectionRegionMaterial = nullptr;
+UMaterialInstanceConstant* GMaskRegionMaterial = nullptr;
+UTexture2D* GLandscapeBlackTexture = nullptr;
 
 // Game thread update
-void FLandscapeEditToolRenderData::Update( UMaterialInterface* InNewToolMaterial )
+void FLandscapeEditToolRenderData::Update(UMaterialInterface* InNewToolMaterial)
 {
 	ENQUEUE_UNIQUE_RENDER_COMMAND_TWOPARAMETER(
 		UpdateEditToolRenderData,
 		FLandscapeEditToolRenderData*, LandscapeEditToolRenderData, this,
 		UMaterialInterface*, NewToolMaterial, InNewToolMaterial,
-	{
-		LandscapeEditToolRenderData->ToolMaterial = NewToolMaterial;
-	});
+		{
+			LandscapeEditToolRenderData->ToolMaterial = NewToolMaterial;
+		});
 }
 
-void FLandscapeEditToolRenderData::UpdateGizmo( UMaterialInterface* InNewGizmoMaterial )
+void FLandscapeEditToolRenderData::UpdateGizmo(UMaterialInterface* InNewGizmoMaterial)
 {
 	ENQUEUE_UNIQUE_RENDER_COMMAND_TWOPARAMETER(
 		UpdateEditToolRenderData,
 		FLandscapeEditToolRenderData*, LandscapeEditToolRenderData, this,
 		UMaterialInterface*, NewGizmoMaterial, InNewGizmoMaterial,
-	{
-		LandscapeEditToolRenderData->GizmoMaterial = NewGizmoMaterial;
-	});
+		{
+			LandscapeEditToolRenderData->GizmoMaterial = NewGizmoMaterial;
+		});
 }
 
 // Allows game thread to queue the deletion by the render thread
@@ -476,9 +476,10 @@ void FLandscapeEditToolRenderData::Cleanup()
 	ENQUEUE_UNIQUE_RENDER_COMMAND_ONEPARAMETER(
 		CleanupEditToolRenderData,
 		FLandscapeEditToolRenderData*, LandscapeEditToolRenderData, this,
-	{
-		delete LandscapeEditToolRenderData;
-	});
+		{
+			delete LandscapeEditToolRenderData;
+		}
+	);
 }
 
 
@@ -502,7 +503,7 @@ void FLandscapeEditToolRenderData::UpdateSelectionMaterial(int32 InSelectedType)
 	}
 
 	// Check selection
-	if (SelectedType != InSelectedType && (SelectedType & ST_REGION) && !(InSelectedType & ST_REGION) )
+	if (SelectedType != InSelectedType && (SelectedType & ST_REGION) && !(InSelectedType & ST_REGION))
 	{
 		// Clear Select textures...
 		if (DataTexture)
@@ -530,94 +531,90 @@ TMap<uint32, FLandscapeSharedBuffers*>FLandscapeComponentSceneProxy::SharedBuffe
 TMap<uint32, FLandscapeSharedAdjacencyIndexBuffer*>FLandscapeComponentSceneProxy::SharedAdjacencyIndexBufferMap;
 
 FLandscapeComponentSceneProxy::FLandscapeComponentSceneProxy(ULandscapeComponent* InComponent, FLandscapeEditToolRenderData* InEditToolRenderData)
-:	FPrimitiveSceneProxy(InComponent)
-,	MaxLOD(FMath::CeilLogTwo(InComponent->SubsectionSizeQuads+1)-1)
-,	ComponentSizeQuads(InComponent->ComponentSizeQuads)
-,	ComponentSizeVerts(InComponent->ComponentSizeQuads+1)
-,	NumSubsections(InComponent->NumSubsections)
-,	SubsectionSizeQuads(InComponent->SubsectionSizeQuads)
-,	SubsectionSizeVerts(InComponent->SubsectionSizeQuads+1)
-,	SectionBase(InComponent->GetSectionBase())
-,	StaticLightingLOD(InComponent->GetLandscapeProxy()->StaticLightingLOD)
-,	WeightmapScaleBias(InComponent->WeightmapScaleBias)
-,	WeightmapSubsectionOffset(InComponent->WeightmapSubsectionOffset)
-,	WeightmapTextures(InComponent->WeightmapTextures)
-,	NumWeightmapLayerAllocations(InComponent->WeightmapLayerAllocations.Num())
-,	NormalmapTexture(InComponent->HeightmapTexture)
-,	HeightmapTexture(InComponent->HeightmapTexture)
-,	HeightmapScaleBias(InComponent->HeightmapScaleBias)
-,	XYOffsetmapTexture(InComponent->XYOffsetmapTexture)
-,	SharedBuffersKey(0)
-,	SharedBuffers(NULL)
-,	VertexFactory(NULL)
-,	MaterialInterface(InComponent->MaterialInstance)
-,	EditToolRenderData(InEditToolRenderData)
-,	ComponentLightInfo(NULL)
-,	LandscapeComponent(InComponent)
-,	ForcedLOD(InComponent->ForcedLOD)
-,	LODBias(InComponent->LODBias)
-,	LODFalloff(InComponent->GetLandscapeProxy()->LODFalloff)
+	: FPrimitiveSceneProxy(InComponent)
+	, MaxLOD(FMath::CeilLogTwo(InComponent->SubsectionSizeQuads + 1) - 1)
+	, NumSubsections(InComponent->NumSubsections)
+	, SubsectionSizeQuads(InComponent->SubsectionSizeQuads)
+	, SubsectionSizeVerts(InComponent->SubsectionSizeQuads + 1)
+	, ComponentSizeQuads(InComponent->ComponentSizeQuads)
+	, ComponentSizeVerts(InComponent->ComponentSizeQuads + 1)
+	, StaticLightingLOD(InComponent->GetLandscapeProxy()->StaticLightingLOD)
+	, SectionBase(InComponent->GetSectionBase())
+	, WeightmapScaleBias(InComponent->WeightmapScaleBias)
+	, WeightmapSubsectionOffset(InComponent->WeightmapSubsectionOffset)
+	, WeightmapTextures(InComponent->WeightmapTextures)
+	, NumWeightmapLayerAllocations(InComponent->WeightmapLayerAllocations.Num())
+	, NormalmapTexture(InComponent->HeightmapTexture)
+	, HeightmapTexture(InComponent->HeightmapTexture)
+	, HeightmapScaleBias(InComponent->HeightmapScaleBias)
+	, XYOffsetmapTexture(InComponent->XYOffsetmapTexture)
+	, SharedBuffersKey(0)
+	, SharedBuffers(nullptr)
+	, VertexFactory(nullptr)
+	, MaterialInterface(InComponent->MaterialInstance)
+	, EditToolRenderData(InEditToolRenderData)
+	, ComponentLightInfo(nullptr)
+	, LandscapeComponent(InComponent)
+	, ForcedLOD(InComponent->ForcedLOD)
+	, LODBias(InComponent->LODBias)
+	, LODFalloff(InComponent->GetLandscapeProxy()->LODFalloff)
 {
 	LevelColor = FLinearColor(1.f, 1.f, 1.f);
 
 	const auto FeatureLevel = GetScene().GetFeatureLevel();
 	if (FeatureLevel <= ERHIFeatureLevel::ES3_1)
 	{
-		HeightmapTexture = NULL;
+		HeightmapTexture = nullptr;
 		HeightmapSubsectionOffsetU = 0;
 		HeightmapSubsectionOffsetV = 0;
 	}
 	else
 	{
-		HeightmapSubsectionOffsetU = ((float)(InComponent->SubsectionSizeQuads+1) / (float)HeightmapTexture->GetSizeX());
-		HeightmapSubsectionOffsetV = ((float)(InComponent->SubsectionSizeQuads+1) / (float)HeightmapTexture->GetSizeY());
+		HeightmapSubsectionOffsetU = ((float)(InComponent->SubsectionSizeQuads + 1) / (float)HeightmapTexture->GetSizeX());
+		HeightmapSubsectionOffsetV = ((float)(InComponent->SubsectionSizeQuads + 1) / (float)HeightmapTexture->GetSizeY());
 	}
 
-	//    - - 0 - -
-	//    |       |
-	//    1   P   2
-	//    |       |
-	//    - - 3 - -
-
-	NeighborPosition[0].Set(0.5f * (float)SubsectionSizeQuads, -0.5f * (float)SubsectionSizeQuads);
-	NeighborPosition[1].Set(-0.5f * (float)SubsectionSizeQuads, 0.5f * (float)SubsectionSizeQuads);
-	NeighborPosition[2].Set(1.5f * (float)SubsectionSizeQuads, 0.5f * (float)SubsectionSizeQuads);
-	NeighborPosition[3].Set(0.5f * (float)SubsectionSizeQuads, 1.5f * (float)SubsectionSizeQuads);
+	//       -Y       
+	//    - - 0 - -   
+	//    |       |   
+	// -X 1   P   2 +X
+	//    |       |   
+	//    - - 3 - -   
+	//       +Y       
 
-	ForcedNeighborLOD[0] = InComponent->NeighborLOD[1];
-	ForcedNeighborLOD[1] = InComponent->NeighborLOD[3];
-	ForcedNeighborLOD[2] = InComponent->NeighborLOD[4];
-	ForcedNeighborLOD[3] = InComponent->NeighborLOD[6];
+	// Component->NeighborLOD and Component->NeighborLODBias are stored in a legacy uint8 form,
+	// so convert them here to something that matches ForcedLOD and LODBias (for this component)
+	ForcedNeighborLOD[0] = InComponent->NeighborLOD[1] != 255 ? InComponent->NeighborLOD[1] : -1;
+	ForcedNeighborLOD[1] = InComponent->NeighborLOD[3] != 255 ? InComponent->NeighborLOD[3] : -1;
+	ForcedNeighborLOD[2] = InComponent->NeighborLOD[4] != 255 ? InComponent->NeighborLOD[4] : -1;
+	ForcedNeighborLOD[3] = InComponent->NeighborLOD[6] != 255 ? InComponent->NeighborLOD[6] : -1;
 
-	NeighborLODBias[0] = InComponent->NeighborLODBias[1];
-	NeighborLODBias[1] = InComponent->NeighborLODBias[3];
-	NeighborLODBias[2] = InComponent->NeighborLODBias[4];
-	NeighborLODBias[3] = InComponent->NeighborLODBias[6];
+	NeighborLODBias[0] = InComponent->NeighborLODBias[1] - 128;
+	NeighborLODBias[1] = InComponent->NeighborLODBias[3] - 128;
+	NeighborLODBias[2] = InComponent->NeighborLODBias[4] - 128;
+	NeighborLODBias[3] = InComponent->NeighborLODBias[6] - 128;
 
-	LODBias = FMath::Clamp(LODBias, -MaxLOD, MaxLOD);
+	LODBias = FMath::Clamp<int8>(LODBias, -MaxLOD, MaxLOD);
 
-	if( InComponent->GetLandscapeProxy()->MaxLODLevel >= 0 )
+	if (InComponent->GetLandscapeProxy()->MaxLODLevel >= 0)
 	{
-		MaxLOD = FMath::Min<int32>(MaxLOD, InComponent->GetLandscapeProxy()->MaxLODLevel);
+		MaxLOD = FMath::Min<int8>(MaxLOD, InComponent->GetLandscapeProxy()->MaxLODLevel);
 	}
 
 	float LODDistanceFactor;
 	switch (LODFalloff)
 	{
-		case ELandscapeLODFalloff::SquareRoot:
-			LODDistanceFactor = FMath::Square(FMath::Min(LANDSCAPE_LOD_SQUARE_ROOT_FACTOR * InComponent->GetLandscapeProxy()->LODDistanceFactor, MAX_LANDSCAPE_LOD_DISTANCE_FACTOR));
-			break;
-		case ELandscapeLODFalloff::Linear:
-		default:
-			LODDistanceFactor = InComponent->GetLandscapeProxy()->LODDistanceFactor;
-			break;
+	case ELandscapeLODFalloff::SquareRoot:
+		LODDistanceFactor = FMath::Square(FMath::Min(LANDSCAPE_LOD_SQUARE_ROOT_FACTOR * InComponent->GetLandscapeProxy()->LODDistanceFactor, MAX_LANDSCAPE_LOD_DISTANCE_FACTOR));
+		break;
+	case ELandscapeLODFalloff::Linear:
+	default:
+		LODDistanceFactor = InComponent->GetLandscapeProxy()->LODDistanceFactor;
+		break;
 	}
 
 	LODDistance = FMath::Sqrt(2.f * FMath::Square((float)SubsectionSizeQuads)) * LANDSCAPE_LOD_DISTANCE_FACTOR / LODDistanceFactor; // vary in 0...1
 	DistDiff = -FMath::Sqrt(2.f * FMath::Square(0.5f*(float)SubsectionSizeQuads));
-	// This value is used for shader. I made some modification to LOD transition looks better than plain transition
-	// Used for LODBias.xy values in shader
-	PrecomputedLODFactor = FMath::Min(LODDistanceFactor, 3.f) * 0.33f; 
 
 	if (InComponent->StaticLightingResolution > 0.f)
 	{
@@ -628,15 +625,15 @@ FLandscapeComponentSceneProxy::FLandscapeComponentSceneProxy(ULandscapeComponent
 		StaticLightingResolution = InComponent->GetLandscapeProxy()->StaticLightingResolution;
 	}
 
-	ComponentLightInfo = new FLandscapeLCI(InComponent);
+	ComponentLightInfo = MakeUnique<FLandscapeLCI>(InComponent);
 	check(ComponentLightInfo);
 
-	const bool bHasStaticLighting = InComponent->LightMap != NULL || InComponent->ShadowMap != NULL;
+	const bool bHasStaticLighting = InComponent->LightMap != nullptr || InComponent->ShadowMap != nullptr;
 
 	// Check material usage
-	if( MaterialInterface == NULL ||
+	if (MaterialInterface == nullptr ||
 		!MaterialInterface->CheckMaterialUsage(MATUSAGE_Landscape) ||
-		(bHasStaticLighting && !MaterialInterface->CheckMaterialUsage(MATUSAGE_StaticLighting)) )
+		(bHasStaticLighting && !MaterialInterface->CheckMaterialUsage(MATUSAGE_StaticLighting)))
 	{
 		MaterialInterface = UMaterial::GetDefaultMaterial(MD_Surface);
 	}
@@ -644,15 +641,15 @@ FLandscapeComponentSceneProxy::FLandscapeComponentSceneProxy(ULandscapeComponent
 	MaterialRelevance = MaterialInterface->GetRelevance(FeatureLevel);
 
 #if !(UE_BUILD_SHIPPING || UE_BUILD_TEST) || (UE_BUILD_SHIPPING && WITH_EDITOR)
-	if( GIsEditor )
+	if (GIsEditor)
 	{
 		ALandscapeProxy* Proxy = InComponent->GetLandscapeProxy();
 		// Try to find a color for level coloration.
-		if ( Proxy )
+		if (Proxy)
 		{
 			ULevel* Level = Proxy->GetLevel();
-			ULevelStreaming* LevelStreaming = FLevelUtils::FindStreamingLevel( Level );
-			if ( LevelStreaming )
+			ULevelStreaming* LevelStreaming = FLevelUtils::FindStreamingLevel(Level);
+			if (LevelStreaming)
 			{
 				LevelColor = LevelStreaming->LevelColor;
 			}
@@ -660,20 +657,20 @@ FLandscapeComponentSceneProxy::FLandscapeComponentSceneProxy(ULandscapeComponent
 	}
 #endif
 
-	bRequiresAdjacencyInformation = RequiresAdjacencyInformation(MaterialInterface, XYOffsetmapTexture == NULL ? &FLandscapeVertexFactory::StaticType : &FLandscapeXYOffsetVertexFactory::StaticType, InComponent->GetWorld()->FeatureLevel);
-	SharedBuffersKey = (SubsectionSizeQuads & 0xffff) | ((NumSubsections & 0xf) << 16) | (FeatureLevel <= ERHIFeatureLevel::ES3_1 ? 0 : 1<<20) | (XYOffsetmapTexture == NULL ? 0 : 1 << 31);
+	bRequiresAdjacencyInformation = RequiresAdjacencyInformation(MaterialInterface, XYOffsetmapTexture == nullptr ? &FLandscapeVertexFactory::StaticType : &FLandscapeXYOffsetVertexFactory::StaticType, InComponent->GetWorld()->FeatureLevel);
+	SharedBuffersKey = (SubsectionSizeQuads & 0xffff) | ((NumSubsections & 0xf) << 16) | (FeatureLevel <= ERHIFeatureLevel::ES3_1 ? 0 : 1 << 20) | (XYOffsetmapTexture == nullptr ? 0 : 1 << 31);
 
 	bSupportsHeightfieldRepresentation = true;
 }
 
 void FLandscapeComponentSceneProxy::CreateRenderThreadResources()
 {
-	check(HeightmapTexture != NULL);
+	check(HeightmapTexture != nullptr);
 
 	auto FeatureLevel = GetScene().GetFeatureLevel();
 
 	SharedBuffers = FLandscapeComponentSceneProxy::SharedBuffersMap.FindRef(SharedBuffersKey);
-	if( SharedBuffers == NULL )
+	if (SharedBuffers == nullptr)
 	{
 		SharedBuffers = new FLandscapeSharedBuffers(SharedBuffersKey, SubsectionSizeQuads, NumSubsections, FeatureLevel, bRequiresAdjacencyInformation);
 		FLandscapeComponentSceneProxy::SharedBuffersMap.Add(SharedBuffersKey, SharedBuffers);
@@ -698,7 +695,7 @@ void FLandscapeComponentSceneProxy::CreateRenderThreadResources()
 
 	if (bRequiresAdjacencyInformation)
 	{
-		if (SharedBuffers->AdjacencyIndexBuffers == NULL)
+		if (SharedBuffers->AdjacencyIndexBuffers == nullptr)
 		{
 			ensure(SharedBuffers->NumIndexBuffers > 0);
 			if (SharedBuffers->IndexBuffers[0])
@@ -728,7 +725,7 @@ void FLandscapeComponentSceneProxy::CreateRenderThreadResources()
 
 	// Assign vertex factory
 	VertexFactory = SharedBuffers->VertexFactory;
-	
+
 	// Assign LandscapeUniformShaderParameters
 	LandscapeUniformShaderParameters.InitResource();
 }
@@ -738,21 +735,18 @@ FLandscapeComponentSceneProxy::~FLandscapeComponentSceneProxy()
 	// Free the subsection uniform buffer
 	LandscapeUniformShaderParameters.ReleaseResource();
 
-	if( SharedBuffers )
+	if (SharedBuffers)
 	{
-		check( SharedBuffers == FLandscapeComponentSceneProxy::SharedBuffersMap.FindRef(SharedBuffersKey) );
-		if( SharedBuffers->Release() == 0 )
+		check(SharedBuffers == FLandscapeComponentSceneProxy::SharedBuffersMap.FindRef(SharedBuffersKey));
+		if (SharedBuffers->Release() == 0)
 		{
 			FLandscapeComponentSceneProxy::SharedBuffersMap.Remove(SharedBuffersKey);
 		}
-		SharedBuffers = NULL;
+		SharedBuffers = nullptr;
 	}
-
-	delete ComponentLightInfo;
-	ComponentLightInfo = NULL;
 }
 
-int32 GAllowLandscapeShadows=1;
+int32 GAllowLandscapeShadows = 1;
 static FAutoConsoleVariableRef CVarAllowLandscapeShadows(
 	TEXT("r.AllowLandscapeShadows"),
 	GAllowLandscapeShadows,
@@ -772,7 +766,7 @@ FPrimitiveViewRelevance FLandscapeComponentSceneProxy::GetViewRelevance(const FS
 	auto FeatureLevel = View->GetFeatureLevel();
 
 #if WITH_EDITOR
-	if( !GLandscapeEditModeActive )
+	if (!GLandscapeEditModeActive)
 	{
 		// No tools to render, just use the cached material relevance.
 #endif
@@ -785,15 +779,15 @@ FPrimitiveViewRelevance FLandscapeComponentSceneProxy::GetViewRelevance(const FS
 		FMaterialRelevance ToolRelevance = MaterialRelevance;
 
 		// Tool brushes and Gizmo
-		if( EditToolRenderData )
+		if (EditToolRenderData)
 		{
-			if( EditToolRenderData->ToolMaterial )
+			if (EditToolRenderData->ToolMaterial)
 			{
 				Result.bDynamicRelevance = true;
 				ToolRelevance |= EditToolRenderData->ToolMaterial->GetRelevance_Concurrent(FeatureLevel);
 			}
 
-			if( EditToolRenderData->GizmoMaterial )
+			if (EditToolRenderData->GizmoMaterial)
 			{
 				Result.bDynamicRelevance = true;
 				ToolRelevance |= EditToolRenderData->GizmoMaterial->GetRelevance_Concurrent(FeatureLevel);
@@ -801,7 +795,7 @@ FPrimitiveViewRelevance FLandscapeComponentSceneProxy::GetViewRelevance(const FS
 		}
 
 		// Region selection
-		if ( EditToolRenderData && EditToolRenderData->SelectedType )
+		if (EditToolRenderData && EditToolRenderData->SelectedType)
 		{
 			if ((GLandscapeEditRenderMode & ELandscapeEditRenderMode::SelectRegion) && (EditToolRenderData->SelectedType & FLandscapeEditToolRenderData::ST_REGION)
 				&& !(GLandscapeEditRenderMode & ELandscapeEditRenderMode::Mask) && GSelectionRegionMaterial)
@@ -817,8 +811,8 @@ FPrimitiveViewRelevance FLandscapeComponentSceneProxy::GetViewRelevance(const FS
 		}
 
 		// Mask
-		if ( (GLandscapeEditRenderMode & ELandscapeEditRenderMode::Mask) && GMaskRegionMaterial != NULL && 
-			((EditToolRenderData && (EditToolRenderData->SelectedType & FLandscapeEditToolRenderData::ST_REGION)) || (!(GLandscapeEditRenderMode & ELandscapeEditRenderMode::InvertedMask))) )
+		if ((GLandscapeEditRenderMode & ELandscapeEditRenderMode::Mask) && GMaskRegionMaterial != nullptr &&
+			((EditToolRenderData && (EditToolRenderData->SelectedType & FLandscapeEditToolRenderData::ST_REGION)) || (!(GLandscapeEditRenderMode & ELandscapeEditRenderMode::InvertedMask))))
 		{
 			Result.bDynamicRelevance = true;
 			ToolRelevance |= GMaskRegionMaterial->GetRelevance_Concurrent(FeatureLevel);
@@ -828,19 +822,19 @@ FPrimitiveViewRelevance FLandscapeComponentSceneProxy::GetViewRelevance(const FS
 	}
 
 	// Various visualizations need to render using dynamic relevance
-	if( (View->Family->EngineShowFlags.Bounds && IsSelected()) ||
-		GLandscapeDebugOptions.bShowPatches )
+	if ((View->Family->EngineShowFlags.Bounds && IsSelected()) ||
+		GLandscapeDebugOptions.bShowPatches)
 	{
 		Result.bDynamicRelevance = true;
 	}
 #endif
 
 	// Use the dynamic path for rendering landscape components pass only for Rich Views or if the static path is disabled for debug.
-	if(	IsRichView(*View->Family) || 
-		GLandscapeDebugOptions.bDisableStatic || 
-		View->Family->EngineShowFlags.Wireframe || 
+	if (IsRichView(*View->Family) ||
+		GLandscapeDebugOptions.bDisableStatic ||
+		View->Family->EngineShowFlags.Wireframe ||
 #if WITH_EDITOR
-		(IsSelected() && !GLandscapeEditModeActive) || 
+		(IsSelected() && !GLandscapeEditModeActive) ||
 		GLandscapeViewMode != ELandscapeViewMode::Normal
 #else
 		IsSelected()
@@ -848,7 +842,7 @@ FPrimitiveViewRelevance FLandscapeComponentSceneProxy::GetViewRelevance(const FS
 		)
 	{
 		Result.bDynamicRelevance = true;
-	}
+}
 	else
 	{
 		Result.bStaticRelevance = true;
@@ -877,12 +871,12 @@ void FLandscapeComponentSceneProxy::GetLightRelevance(const FLightSceneProxy* Li
 	{
 		ELightInteractionType InteractionType = ComponentLightInfo->GetInteraction(LightSceneProxy).GetType();
 
-		if(InteractionType != LIT_CachedIrrelevant)
+		if (InteractionType != LIT_CachedIrrelevant)
 		{
 			bRelevant = true;
 		}
 
-		if(InteractionType != LIT_CachedLightMap && InteractionType != LIT_CachedIrrelevant)
+		if (InteractionType != LIT_CachedLightMap && InteractionType != LIT_CachedIrrelevant)
 		{
 			bLightMapped = false;
 		}
@@ -907,21 +901,21 @@ void FLandscapeComponentSceneProxy::GetLightRelevance(const FLightSceneProxy* Li
 FLightInteraction FLandscapeComponentSceneProxy::FLandscapeLCI::GetInteraction(const class FLightSceneProxy* LightSceneProxy) const
 {
 	// Check if the light has static lighting or shadowing.
-	if(LightSceneProxy->HasStaticShadowing())
+	if (LightSceneProxy->HasStaticShadowing())
 	{
 		const FGuid LightGuid = LightSceneProxy->GetLightGuid();
 
-		if(LightMap && LightMap->ContainsLight(LightGuid))
+		if (LightMap && LightMap->ContainsLight(LightGuid))
 		{
 			return FLightInteraction::LightMap();
 		}
 
-		if(ShadowMap && ShadowMap->ContainsLight(LightGuid))
+		if (ShadowMap && ShadowMap->ContainsLight(LightGuid))
 		{
 			return FLightInteraction::ShadowMap2D();
 		}
 
-		if( IrrelevantLights.Contains(LightGuid) )
+		if (IrrelevantLights.Contains(LightGuid))
 		{
 			return FLightInteraction::Irrelevant();
 		}
@@ -934,13 +928,13 @@ FLightInteraction FLandscapeComponentSceneProxy::FLandscapeLCI::GetInteraction(c
 #if WITH_EDITOR
 namespace DebugColorMask
 {
-	const FLinearColor Masks[5] = 
+	const FLinearColor Masks[5] =
 	{
-		FLinearColor(1.f,0.f,0.f,0.f),
-		FLinearColor(0.f,1.f,0.f,0.f),
-		FLinearColor(0.f,0.f,1.f,0.f),
-		FLinearColor(0.f,0.f,0.f,1.f),
-		FLinearColor(0.f,0.f,0.f,0.f)
+		FLinearColor(1.f, 0.f, 0.f, 0.f),
+		FLinearColor(0.f, 1.f, 0.f, 0.f),
+		FLinearColor(0.f, 0.f, 1.f, 0.f),
+		FLinearColor(0.f, 0.f, 0.f, 1.f),
+		FLinearColor(0.f, 0.f, 0.f, 0.f)
 	};
 };
 #endif
@@ -951,16 +945,16 @@ void FLandscapeComponentSceneProxy::OnTransformChanged()
 	int32 PatchExpandCountX = 1;
 	int32 PatchExpandCountY = 1;
 	int32 DesiredSize = 1;
-	const float LightMapRatio = ::GetTerrainExpandPatchCount(StaticLightingResolution, PatchExpandCountX, PatchExpandCountY, ComponentSizeQuads, (NumSubsections * (SubsectionSizeQuads+1)), DesiredSize, StaticLightingLOD);
-	const float LightmapLODScaleX = LightMapRatio / ((ComponentSizeVerts>>StaticLightingLOD) + 2 * PatchExpandCountX);
-	const float LightmapLODScaleY = LightMapRatio / ((ComponentSizeVerts>>StaticLightingLOD) + 2 * PatchExpandCountY);
+	const float LightMapRatio = ::GetTerrainExpandPatchCount(StaticLightingResolution, PatchExpandCountX, PatchExpandCountY, ComponentSizeQuads, (NumSubsections * (SubsectionSizeQuads + 1)), DesiredSize, StaticLightingLOD);
+	const float LightmapLODScaleX = LightMapRatio / ((ComponentSizeVerts >> StaticLightingLOD) + 2 * PatchExpandCountX);
+	const float LightmapLODScaleY = LightMapRatio / ((ComponentSizeVerts >> StaticLightingLOD) + 2 * PatchExpandCountY);
 	const float LightmapBiasX = PatchExpandCountX * LightmapLODScaleX;
 	const float LightmapBiasY = PatchExpandCountY * LightmapLODScaleY;
-	const float LightmapScaleX = LightmapLODScaleX * (float)((ComponentSizeVerts>>StaticLightingLOD)-1) / ComponentSizeQuads;
-	const float LightmapScaleY = LightmapLODScaleY * (float)((ComponentSizeVerts>>StaticLightingLOD)-1) / ComponentSizeQuads;
+	const float LightmapScaleX = LightmapLODScaleX * (float)((ComponentSizeVerts >> StaticLightingLOD) - 1) / ComponentSizeQuads;
+	const float LightmapScaleY = LightmapLODScaleY * (float)((ComponentSizeVerts >> StaticLightingLOD) - 1) / ComponentSizeQuads;
 	const float LightmapExtendFactorX = (float)SubsectionSizeQuads * LightmapScaleX;
 	const float LightmapExtendFactorY = (float)SubsectionSizeQuads * LightmapScaleY;
-	
+
 	// cache component's WorldToLocal
 	FMatrix LtoW = GetLocalToWorld();
 	WorldToLocal = LtoW.InverseFast();
@@ -968,7 +962,7 @@ void FLandscapeComponentSceneProxy::OnTransformChanged()
 	// cache component's LocalToWorldNoScaling
 	LocalToWorldNoScaling = LtoW;
 	LocalToWorldNoScaling.RemoveScaling();
-	
+
 	// Set FLandscapeUniformVSParameters for this subsection
 	FLandscapeUniformShaderParameters LandscapeParams;
 	LandscapeParams.HeightmapUVScaleBias = HeightmapScaleBias;
@@ -1008,13 +1002,13 @@ namespace
 	{
 		if (RHISupportsTessellation(GShaderPlatformForFeatureLevel[InFeatureLevel]) && MaterialRenderProxy)
 		{
-			check ( IsInRenderingThread() );
+			check(IsInRenderingThread());
 			const FMaterial* MaterialResource = MaterialRenderProxy->GetMaterial(InFeatureLevel);
-			check( MaterialResource );
+			check(MaterialResource);
 			EMaterialTessellationMode TessellationMode = MaterialResource->GetTessellationMode();
 			bool bEnableCrackFreeDisplacement = MaterialResource->IsCrackFreeDisplacementEnabled();
 
-			return TessellationMode == MTM_PNTriangles || ( TessellationMode == MTM_FlatTessellation && bEnableCrackFreeDisplacement );
+			return TessellationMode == MTM_PNTriangles || (TessellationMode == MTM_FlatTessellation && bEnableCrackFreeDisplacement);
 		}
 		else
 		{
@@ -1023,7 +1017,7 @@ namespace
 	}
 };
 
-/** 
+/**
 * Draw the scene proxy as a dynamic element
 *
 * @param	PDI - draw interface to render to
@@ -1033,39 +1027,39 @@ namespace
 void FLandscapeComponentSceneProxy::DrawStaticElements(FStaticPrimitiveDrawInterface* PDI)
 {
 	int32 FirstLOD = (ForcedLOD >= 0) ? FMath::Min<int32>(ForcedLOD, MaxLOD) : FMath::Max<int32>(LODBias, 0);
-	int32 LastLOD  = (ForcedLOD >= 0) ? FirstLOD : FMath::Min<int32>(MaxLOD, MaxLOD+LODBias);
+	int32 LastLOD = (ForcedLOD >= 0) ? FirstLOD : FMath::Min<int32>(MaxLOD, MaxLOD + LODBias);
 
-	StaticBatchParamArray.Empty((1+LastLOD-FirstLOD) * (FMath::Square(NumSubsections) + 1));
+	StaticBatchParamArray.Empty((1 + LastLOD - FirstLOD) * (FMath::Square(NumSubsections) + 1));
 
 	FMeshBatch MeshBatch;
-	MeshBatch.Elements.Empty((1+LastLOD-FirstLOD) * (FMath::Square(NumSubsections) + 1));
+	MeshBatch.Elements.Empty((1 + LastLOD - FirstLOD) * (FMath::Square(NumSubsections) + 1));
 
 	FMaterialRenderProxy* RenderProxy = MaterialInterface->GetRenderProxy(false);
 
 	// Could be different from bRequiresAdjacencyInformation during shader compilation
-	bool bCurrentRequiresAdjacencyInformation = RequiresAdjacencyInformation( RenderProxy, GetScene().GetFeatureLevel() );
+	bool bCurrentRequiresAdjacencyInformation = RequiresAdjacencyInformation(RenderProxy, GetScene().GetFeatureLevel());
 
 	MeshBatch.VertexFactory = VertexFactory;
 	MeshBatch.MaterialRenderProxy = RenderProxy;
-	MeshBatch.LCI = ComponentLightInfo; 
+	MeshBatch.LCI = ComponentLightInfo.Get();
 	MeshBatch.ReverseCulling = IsLocalToWorldDeterminantNegative();
 	MeshBatch.CastShadow = true;
 	MeshBatch.Type = bCurrentRequiresAdjacencyInformation ? PT_12_ControlPointPatchList : PT_TriangleList;
 	MeshBatch.DepthPriorityGroup = SDPG_World;
 
-	for( int32 LOD = FirstLOD; LOD <= LastLOD; LOD++ )
+	for (int32 LOD = FirstLOD; LOD <= LastLOD; LOD++)
 	{
 		int32 LodSubsectionSizeVerts = SubsectionSizeVerts >> LOD;
 
-		if( ForcedLOD < 0 && NumSubsections > 1 )
+		if (ForcedLOD < 0 && NumSubsections > 1)
 		{
 			// Per-subsection batch elements
-			for( int32 SubY=0;SubY<NumSubsections;SubY++ )
+			for (int32 SubY = 0; SubY < NumSubsections; SubY++)
 			{
-				for( int32 SubX=0;SubX<NumSubsections;SubX++ )
+				for (int32 SubX = 0; SubX < NumSubsections; SubX++)
 				{
 					FMeshBatchElement* BatchElement = new(MeshBatch.Elements) FMeshBatchElement;
-					FLandscapeBatchElementParams* BatchElementParams = new(StaticBatchParamArray) FLandscapeBatchElementParams;
+					FLandscapeBatchElementParams* BatchElementParams = new(StaticBatchParamArray)FLandscapeBatchElementParams;
 					BatchElement->UserData = BatchElementParams;
 
 					BatchElement->PrimitiveUniformBufferResource = &GetUniformBuffer();
@@ -1078,7 +1072,7 @@ void FLandscapeComponentSceneProxy::DrawStaticElements(FStaticPrimitiveDrawInter
 					BatchElementParams->CurrentLOD = LOD;
 
 					BatchElement->IndexBuffer = SharedBuffers->IndexBuffers[LOD];
-					BatchElement->NumPrimitives = FMath::Square((LodSubsectionSizeVerts-1)) * 2;
+					BatchElement->NumPrimitives = FMath::Square((LodSubsectionSizeVerts - 1)) * 2;
 					BatchElement->FirstIndex = (SubX + SubY * NumSubsections) * BatchElement->NumPrimitives * 3;
 					BatchElement->MinVertexIndex = SharedBuffers->IndexRanges[LOD].MinIndex[SubX][SubY];
 					BatchElement->MaxVertexIndex = SharedBuffers->IndexRanges[LOD].MaxIndex[SubX][SubY];
@@ -1088,7 +1082,7 @@ void FLandscapeComponentSceneProxy::DrawStaticElements(FStaticPrimitiveDrawInter
 
 		// Combined batch element
 		FMeshBatchElement* BatchElement = new(MeshBatch.Elements) FMeshBatchElement;
-		FLandscapeBatchElementParams* BatchElementParams = new(StaticBatchParamArray) FLandscapeBatchElementParams;
+		FLandscapeBatchElementParams* BatchElementParams = new(StaticBatchParamArray)FLandscapeBatchElementParams;
 		BatchElementParams->LocalToWorldNoScalingPtr = &LocalToWorldNoScaling;
 		BatchElement->UserData = BatchElementParams;
 		BatchElement->PrimitiveUniformBufferResource = &GetUniformBuffer();
@@ -1108,7 +1102,7 @@ void FLandscapeComponentSceneProxy::DrawStaticElements(FStaticPrimitiveDrawInter
 			BatchElement->IndexBuffer = SharedBuffers->IndexBuffers[LOD];
 		}
 
-		BatchElement->NumPrimitives = FMath::Square((LodSubsectionSizeVerts-1)) * FMath::Square(NumSubsections) * 2;
+		BatchElement->NumPrimitives = FMath::Square((LodSubsectionSizeVerts - 1)) * FMath::Square(NumSubsections) * 2;
 		BatchElement->FirstIndex = 0;
 		BatchElement->MinVertexIndex = SharedBuffers->IndexRanges[LOD].MinIndexFull;
 		BatchElement->MaxVertexIndex = SharedBuffers->IndexRanges[LOD].MaxIndexFull;
@@ -1117,30 +1111,31 @@ void FLandscapeComponentSceneProxy::DrawStaticElements(FStaticPrimitiveDrawInter
 	PDI->DrawMesh(MeshBatch, FLT_MAX);
 }
 
-uint64 FLandscapeVertexFactory::GetStaticBatchElementVisibility( const class FSceneView& View, const struct FMeshBatch* Batch ) const
+uint64 FLandscapeVertexFactory::GetStaticBatchElementVisibility(const class FSceneView& View, const struct FMeshBatch* Batch) const
 {
 	const FLandscapeComponentSceneProxy* SceneProxy = ((FLandscapeBatchElementParams*)Batch->Elements[0].UserData)->SceneProxy;
-	return SceneProxy->GetStaticBatchElementVisibility( View, Batch );
+	return SceneProxy->GetStaticBatchElementVisibility(View, Batch);
 }
 
-uint64 FLandscapeComponentSceneProxy::GetStaticBatchElementVisibility( const class FSceneView& View, const struct FMeshBatch* Batch ) const
+uint64 FLandscapeComponentSceneProxy::GetStaticBatchElementVisibility(const class FSceneView& View, const struct FMeshBatch* Batch) const
 {
 	uint64 BatchesToRenderMask = 0;
 
 	SCOPE_CYCLE_COUNTER(STAT_LandscapeStaticDrawLODTime);
 	if (ForcedLOD >= 0)
 	{
-		for( int32 BatchElementIndex=0;BatchElementIndex < Batch->Elements.Num(); BatchElementIndex++ )
+		// When forcing LOD we only create one Batch Element, so not sure this loop is necessary?
+		for (int32 BatchElementIndex = 0; BatchElementIndex < Batch->Elements.Num(); BatchElementIndex++)
 		{
-			BatchesToRenderMask |= (((uint64)1)<<BatchElementIndex);
+			BatchesToRenderMask |= (((uint64)1) << BatchElementIndex);
 			INC_DWORD_STAT(STAT_LandscapeDrawCalls);
-			INC_DWORD_STAT_BY(STAT_LandscapeTriangles,Batch->Elements[BatchElementIndex].NumPrimitives);
+			INC_DWORD_STAT_BY(STAT_LandscapeTriangles, Batch->Elements[BatchElementIndex].NumPrimitives);
 		}
 	}
 	else
 	{
 		// camera position in local heightmap space
-		FVector CameraLocalPos3D = WorldToLocal.TransformPosition(View.ViewMatrices.ViewOrigin); 
+		FVector CameraLocalPos3D = WorldToLocal.TransformPosition(View.ViewMatrices.ViewOrigin);
 		FVector2D CameraLocalPos(CameraLocalPos3D.X, CameraLocalPos3D.Y);
 
 		int32 BatchesPerLOD = NumSubsections > 1 ? FMath::Square(NumSubsections) + 1 : 1;
@@ -1148,19 +1143,19 @@ uint64 FLandscapeComponentSceneProxy::GetStaticBatchElementVisibility( const cla
 		int32 CombinedLOD = -1;
 		int32 bAllSameLOD = true;
 
-		for( int32 SubY=0;SubY<NumSubsections;SubY++ )
+		for (int32 SubY = 0; SubY < NumSubsections; SubY++)
 		{
-			for( int32 SubX=0;SubX<NumSubsections;SubX++ )
+			for (int32 SubX = 0; SubX < NumSubsections; SubX++)
 			{
-				int32 TempLOD = CalcLODForSubsectionNoForced(View, SubX, SubY, CameraLocalPos);
+				int32 TempLOD = CalcLODForSubsection(View, SubX, SubY, CameraLocalPos);
 
-				if(LODBias > 0) 
-				{ 
-					TempLOD = FMath::Max<int32>(TempLOD - LODBias, 0); 
-				} 
+				if (LODBias > 0)
+				{
+					TempLOD = FMath::Max<int32>(TempLOD - LODBias, 0);
+				}
 
 				// check if all LODs are the same.
-				if( TempLOD != CombinedLOD && CombinedLOD != -1 )
+				if (TempLOD != CombinedLOD && CombinedLOD != -1)
 				{
 					bAllSameLOD = false;
 				}
@@ -1169,24 +1164,24 @@ uint64 FLandscapeComponentSceneProxy::GetStaticBatchElementVisibility( const cla
 			}
 		}
 
-		if( bAllSameLOD && NumSubsections > 1 && !GLandscapeDebugOptions.bDisableCombine )
+		if (bAllSameLOD && NumSubsections > 1 && !GLandscapeDebugOptions.bDisableCombine)
 		{
 			// choose the combined batch element
-			int32 BatchElementIndex = (CombinedLOD+1)*BatchesPerLOD - 1;
-			BatchesToRenderMask |= (((uint64)1)<<BatchElementIndex);
+			int32 BatchElementIndex = (CombinedLOD + 1)*BatchesPerLOD - 1;
+			BatchesToRenderMask |= (((uint64)1) << BatchElementIndex);
 			INC_DWORD_STAT(STAT_LandscapeDrawCalls);
-			INC_DWORD_STAT_BY(STAT_LandscapeTriangles,Batch->Elements[BatchElementIndex].NumPrimitives);
+			INC_DWORD_STAT_BY(STAT_LandscapeTriangles, Batch->Elements[BatchElementIndex].NumPrimitives);
 		}
 		else
 		{
-			for( int32 SubY=0;SubY<NumSubsections;SubY++ )
-			{;
-				for( int32 SubX=0;SubX<NumSubsections;SubX++ )
+			for (int32 SubY = 0; SubY < NumSubsections; SubY++)
+			{
+				for (int32 SubX = 0; SubX < NumSubsections; SubX++)
 				{
-					int32 BatchElementIndex = CalculatedLods[SubX][SubY]*BatchesPerLOD + SubY*NumSubsections + SubX;
-					BatchesToRenderMask |= (((uint64)1)<<BatchElementIndex);
+					int32 BatchElementIndex = CalculatedLods[SubX][SubY] * BatchesPerLOD + SubY*NumSubsections + SubX;
+					BatchesToRenderMask |= (((uint64)1) << BatchElementIndex);
 					INC_DWORD_STAT(STAT_LandscapeDrawCalls);
-					INC_DWORD_STAT_BY(STAT_LandscapeTriangles,Batch->Elements[BatchElementIndex].NumPrimitives);
+					INC_DWORD_STAT_BY(STAT_LandscapeTriangles, Batch->Elements[BatchElementIndex].NumPrimitives);
 				}
 			}
 		}
@@ -1198,26 +1193,61 @@ uint64 FLandscapeComponentSceneProxy::GetStaticBatchElementVisibility( const cla
 	return BatchesToRenderMask;
 }
 
-int32 FLandscapeComponentSceneProxy::CalcLODForSubsectionNoForced(const class FSceneView& View, int32 SubX, int32 SubY, const FVector2D& CameraLocalPos) const
+float FLandscapeComponentSceneProxy::CalcDesiredLOD(const class FSceneView& View, const FVector2D& CameraLocalPos, int32 SubX, int32 SubY) const
 {
-	const int32 MinLOD = HeightmapTexture ? FMath::Min<int32>(HeightmapTexture->GetNumMips() - HeightmapTexture->ResidentMips, MaxLOD) : 0;
+	// FLandscapeComponentSceneProxy::NumSubsections, SubsectionSizeQuads, MaxLOD, LODFalloff and LODDistance are the same for all components and so are safe to use in the neighbour LOD calculations
+	// HeightmapTexture, LODBias, ForcedLOD are component-specific with neighbor lookup
+	const bool bIsInThisComponent = (SubX >= 0 && SubX < NumSubsections && SubY >= 0 && SubY < NumSubsections);
+
+	int8 SubsectionForcedLOD = ForcedLOD;
+	int8 SubsectionLODBias = LODBias;
+
+	if (SubX < 0)
+	{
+		SubsectionForcedLOD = ForcedNeighborLOD[1];
+		SubsectionLODBias = NeighborLODBias[1];
+	}
+	else if (SubX >= NumSubsections)
+	{
+		SubsectionForcedLOD = ForcedNeighborLOD[2];
+		SubsectionLODBias = NeighborLODBias[2];
+	}
+	else if (SubY < 0)
+	{
+		SubsectionForcedLOD = ForcedNeighborLOD[0];
+		SubsectionLODBias = NeighborLODBias[0];
+	}
+	else if (SubY >= NumSubsections)
+	{
+		SubsectionForcedLOD = ForcedNeighborLOD[3];
+		SubsectionLODBias = NeighborLODBias[3];
+	}
+
+	// TODO - this clamping isn't done for neighbors, but should be
+	const int32 MinLOD = (bIsInThisComponent && HeightmapTexture) ? FMath::Min<int32>(HeightmapTexture->GetNumMips() - HeightmapTexture->ResidentMips, MaxLOD) : 0;
+
+	float fLOD = FLT_MAX;
+
+	if (SubsectionForcedLOD >= 0)
+	{
+		fLOD = SubsectionForcedLOD;
+	}
+	else
 #if WITH_EDITOR
 	if (View.Family->LandscapeLODOverride >= 0)
 	{
-		float fLOD = FMath::Clamp<float>(View.Family->LandscapeLODOverride, FMath::Max<int32>(LODBias, MinLOD), FMath::Min<int32>(MaxLOD, MaxLOD + LODBias));
-		return FMath::FloorToInt(fLOD);
+		fLOD = View.Family->LandscapeLODOverride;
 	}
+	else
 #endif
-
-	if (View.IsPerspectiveProjection())
 	{
-		FVector2D ComponentPosition(0.5f * (float)SubsectionSizeQuads, 0.5f * (float)SubsectionSizeQuads);
-		FVector2D CurrentCameraLocalPos = CameraLocalPos - FVector2D(SubX * SubsectionSizeQuads,SubY * SubsectionSizeQuads);
-		float ComponentDistance = FVector2D(CurrentCameraLocalPos-ComponentPosition).Size() + DistDiff;
-		// Clamp calculated distance based LOD with LODBiased values
-		float fLOD;
-		switch (LODFalloff)
+		if (View.IsPerspectiveProjection())
 		{
+			FVector2D ComponentPosition(0.5f * (float)SubsectionSizeQuads, 0.5f * (float)SubsectionSizeQuads);
+			FVector2D CurrentCameraLocalPos = CameraLocalPos - FVector2D(SubX * SubsectionSizeQuads, SubY * SubsectionSizeQuads);
+			float ComponentDistance = FVector2D(CurrentCameraLocalPos - ComponentPosition).Size() + DistDiff;
+			switch (LODFalloff)
+			{
 			case ELandscapeLODFalloff::SquareRoot:
 				fLOD = FMath::Sqrt(FMath::Max(0.f, ComponentDistance / LODDistance));
 				break;
@@ -1225,20 +1255,15 @@ int32 FLandscapeComponentSceneProxy::CalcLODForSubsectionNoForced(const class FS
 			case ELandscapeLODFalloff::Linear:
 				fLOD = ComponentDistance / LODDistance;
 				break;
+			}
 		}
-		// Clamping based on LODBias
-		fLOD = FMath::Clamp<float>(fLOD, FMath::Max<int32>(LODBias, MinLOD), FMath::Min<int32>(MaxLOD, MaxLOD + LODBias));
-
-		return FMath::FloorToInt( fLOD );
-	}
-	else
-	{
-		float Scale = 1.0f / (View.ViewRect.Width() * View.ViewMatrices.ProjMatrix.M[0][0]);
-
-		// The "/ 5.0f" is totally arbitrary
-		float fLOD;
-		switch (LODFalloff)
+		else
 		{
+			float Scale = 1.0f / (View.ViewRect.Width() * View.ViewMatrices.ProjMatrix.M[0][0]);
+
+			// The "/ 5.0f" is totally arbitrary
+			switch (LODFalloff)
+			{
 			case ELandscapeLODFalloff::SquareRoot:
 				fLOD = FMath::Sqrt(Scale / 5.0f);
 				break;
@@ -1246,99 +1271,37 @@ int32 FLandscapeComponentSceneProxy::CalcLODForSubsectionNoForced(const class FS
 			case ELandscapeLODFalloff::Linear:
 				fLOD = Scale / 5.0f;
 				break;
+			}
 		}
+	}
 
-		// Clamping based on LODBias
-		fLOD = FMath::Clamp<float>(fLOD, FMath::Max<int32>(LODBias, MinLOD), FMath::Min<int32>(MaxLOD, MaxLOD + LODBias));
+	fLOD = FMath::Clamp<float>(fLOD, FMath::Max<int32>(SubsectionLODBias, MinLOD), FMath::Min<int32>(MaxLOD, MaxLOD + SubsectionLODBias));
 
-		return FMath::FloorToInt(fLOD);
-	}
+	return fLOD;
 }
 
 int32 FLandscapeComponentSceneProxy::CalcLODForSubsection(const class FSceneView& View, int32 SubX, int32 SubY, const FVector2D& CameraLocalPos) const
 {
-	if( ForcedLOD >= 0 )
-	{
-		return ForcedLOD;
-	}
-	else
-	{
-		return CalcLODForSubsectionNoForced(View, SubX, SubY, CameraLocalPos);
-	}
+	return FMath::FloorToInt(CalcDesiredLOD(View, CameraLocalPos, SubX, SubY));
 }
 
-void FLandscapeComponentSceneProxy::CalcLODParamsForSubsection(const class FSceneView& View, const FVector2D& CameraLocalPos, int32 SubX, int32 SubY, float& OutfLOD, FVector4& OutNeighborLODs, float& OutDistLOD) const
+void FLandscapeComponentSceneProxy::CalcLODParamsForSubsection(const class FSceneView& View, const FVector2D& CameraLocalPos, int32 SubX, int32 SubY, int32 BatchLOD, float& OutfLOD, FVector4& OutNeighborLODs) const
 {
-	FVector2D ComponentPosition(0.5f * (float)SubsectionSizeQuads, 0.5f * (float)SubsectionSizeQuads);
-	FVector2D CurrentCameraLocalPos = CameraLocalPos - FVector2D(SubX * SubsectionSizeQuads,SubY * SubsectionSizeQuads);
-
-	float ComponentDistance = FVector2D(CurrentCameraLocalPos - ComponentPosition).Size() + DistDiff;
-
-	OutDistLOD = ComponentDistance / LODDistance;
-
-	int32 FirstLOD = FMath::Max<int32>(LODBias, 0);
-	int32 LastLOD = FMath::Min<int32>(MaxLOD, MaxLOD+LODBias);
-	if (ForcedLOD >= 0)
-	{
-		OutfLOD = ForcedLOD;
-	}
-	else
-	{
-		switch (LODFalloff)
-		{
-		case ELandscapeLODFalloff::SquareRoot:
-			OutfLOD = FMath::Sqrt(FMath::Max(0.f, OutDistLOD));
-			break;
-		default:
-		case ELandscapeLODFalloff::Linear:
-			OutfLOD = OutDistLOD;
-			break;
-		}
+	OutfLOD = FMath::Max<float>(BatchLOD, CalcDesiredLOD(View, CameraLocalPos, SubX, SubY));
 
-		OutfLOD = FMath::Clamp<float>(OutfLOD, FirstLOD, LastLOD);
-	}
-
-	for (int32 Idx = 0; Idx < LANDSCAPE_NEIGHBOR_NUM; ++Idx)
-	{
-		float ComponentDistance = FVector2D(CurrentCameraLocalPos-NeighborPosition[Idx]).Size() + DistDiff;
-		float fLOD;
-		switch (LODFalloff)
-		{
-			case ELandscapeLODFalloff::SquareRoot:
-				fLOD = FMath::Sqrt(FMath::Max(0.f, ComponentDistance / LODDistance));
-				break;
-			default:
-			case ELandscapeLODFalloff::Linear:
-				fLOD = ComponentDistance / LODDistance;
-				break;
-		}
-
-		if (NumSubsections > 1 
-			&& ((SubX == 0 && Idx == 2) 
-			|| (SubX == NumSubsections-1 && Idx == 1) 
-			|| (SubY == 0 && Idx == 3) 
-			|| (SubY == NumSubsections-1 && Idx == 0)) )
-		{
-			// Clamp calculated distance based LOD with LODBiased values
-			OutNeighborLODs[Idx] = ForcedLOD >= 0 ? ForcedLOD : FMath::Clamp<float>( fLOD, FirstLOD, LastLOD);
-		}
-		else
-		{
-			// Neighbor LODBias are saved in BYTE, so need to convert to range [-128:127]
-			OutNeighborLODs[Idx] = ForcedNeighborLOD[Idx] != 255 ? ForcedNeighborLOD[Idx] : FMath::Clamp<float>(fLOD, FMath::Max<float>(NeighborLODBias[Idx] - 128, 0.f), FMath::Min<float>(MaxLOD, MaxLOD + NeighborLODBias[Idx] - 128));
-		}
-
-		OutNeighborLODs[Idx] = FMath::Max<float>(OutfLOD, OutNeighborLODs[Idx]);
-	}
+	OutNeighborLODs[0] = FMath::Max<float>(OutfLOD, CalcDesiredLOD(View, CameraLocalPos, SubX,     SubY - 1));
+	OutNeighborLODs[1] = FMath::Max<float>(OutfLOD, CalcDesiredLOD(View, CameraLocalPos, SubX - 1, SubY    ));
+	OutNeighborLODs[2] = FMath::Max<float>(OutfLOD, CalcDesiredLOD(View, CameraLocalPos, SubX + 1, SubY    ));
+	OutNeighborLODs[3] = FMath::Max<float>(OutfLOD, CalcDesiredLOD(View, CameraLocalPos, SubX,     SubY + 1));
 }
 
 void FLandscapeComponentSceneProxy::GetDynamicMeshElements(const TArray<const FSceneView*>& Views, const FSceneViewFamily& ViewFamily, uint32 VisibilityMap, FMeshElementCollector& Collector) const
 {
 	QUICK_SCOPE_CYCLE_COUNTER(STAT_FLandscapeComponentSceneProxy_GetMeshElements);
 
-	int32 NumPasses=0;
-	int32 NumTriangles=0;
-	int32 NumDrawCalls=0;
+	int32 NumPasses = 0;
+	int32 NumTriangles = 0;
+	int32 NumDrawCalls = 0;
 	const bool bIsWireframe = ViewFamily.EngineShowFlags.Wireframe;
 
 	for (int32 ViewIndex = 0; ViewIndex < Views.Num(); ViewIndex++)
@@ -1346,7 +1309,7 @@ void FLandscapeComponentSceneProxy::GetDynamicMeshElements(const TArray<const FS
 		if (VisibilityMap & (1 << ViewIndex))
 		{
 			const FSceneView* View = Views[ViewIndex];
-			FVector CameraLocalPos3D = WorldToLocal.TransformPosition(View->ViewMatrices.ViewOrigin); 
+			FVector CameraLocalPos3D = WorldToLocal.TransformPosition(View->ViewMatrices.ViewOrigin);
 			FVector2D CameraLocalPos(CameraLocalPos3D.X, CameraLocalPos3D.Y);
 
 			FLandscapeElementParamArray& ParameterArray = Collector.AllocateOneFrameResource<FLandscapeElementParamArray>();
@@ -1357,16 +1320,16 @@ void FLandscapeComponentSceneProxy::GetDynamicMeshElements(const TArray<const FS
 
 			// Could be different from bRequiresAdjacencyInformation during shader compilation
 			FMaterialRenderProxy* RenderProxy = MaterialInterface->GetRenderProxy(false);
-			bool bCurrentRequiresAdjacencyInformation = RequiresAdjacencyInformation( RenderProxy, View->GetFeatureLevel() );
+			bool bCurrentRequiresAdjacencyInformation = RequiresAdjacencyInformation(RenderProxy, View->GetFeatureLevel());
 			Mesh.Type = bCurrentRequiresAdjacencyInformation ? PT_12_ControlPointPatchList : PT_TriangleList;
-			Mesh.LCI = ComponentLightInfo; 
+			Mesh.LCI = ComponentLightInfo.Get();
 			Mesh.CastShadow = true;
 			Mesh.VertexFactory = VertexFactory;
 			Mesh.ReverseCulling = IsLocalToWorldDeterminantNegative();
 
 #if WITH_EDITOR
 			FMeshBatch& MeshTools = Collector.AllocateMesh();
-			MeshTools.LCI = ComponentLightInfo; 
+			MeshTools.LCI = ComponentLightInfo.Get();
 			MeshTools.Type = PT_TriangleList;
 			MeshTools.CastShadow = false;
 			MeshTools.VertexFactory = VertexFactory;
@@ -1374,14 +1337,14 @@ void FLandscapeComponentSceneProxy::GetDynamicMeshElements(const TArray<const FS
 #endif
 
 			// Setup the LOD parameters
-			for( int32 SubY=0;SubY<NumSubsections;SubY++ )
+			for (int32 SubY = 0; SubY < NumSubsections; SubY++)
 			{
-				for( int32 SubX=0;SubX<NumSubsections;SubX++ )
+				for (int32 SubX = 0; SubX < NumSubsections; SubX++)
 				{
 					int32 SubSectionIdx = SubX + SubY*NumSubsections;
 					int32 CurrentLOD = CalcLODForSubsection(*View, SubX, SubY, CameraLocalPos);
 
-					FMeshBatchElement& BatchElement = (SubX==0 && SubY==0) ? *Mesh.Elements.GetData() : *(new(Mesh.Elements) FMeshBatchElement);
+					FMeshBatchElement& BatchElement = (SubX == 0 && SubY == 0) ? *Mesh.Elements.GetData() : *(new(Mesh.Elements) FMeshBatchElement);
 					BatchElement.PrimitiveUniformBufferResource = &GetUniformBuffer();
 					FLandscapeBatchElementParams& BatchElementParams = ParameterArray.ElementParams[SubSectionIdx];
 					BatchElementParams.LocalToWorldNoScalingPtr = &LocalToWorldNoScaling;
@@ -1410,9 +1373,9 @@ void FLandscapeComponentSceneProxy::GetDynamicMeshElements(const TArray<const FS
 					BatchElement.FirstIndex = (SubX + SubY * NumSubsections) * NumPrimitives * 3;
 					BatchElement.MinVertexIndex = SharedBuffers->IndexRanges[CurrentLOD].MinIndex[SubX][SubY];
 					BatchElement.MaxVertexIndex = SharedBuffers->IndexRanges[CurrentLOD].MaxIndex[SubX][SubY];
-					
+
 #if WITH_EDITOR
-					FMeshBatchElement& BatchElementTools = (SubX==0 && SubY==0) ? *MeshTools.Elements.GetData() : *(new(MeshTools.Elements) FMeshBatchElement);
+					FMeshBatchElement& BatchElementTools = (SubX == 0 && SubY == 0) ? *MeshTools.Elements.GetData() : *(new(MeshTools.Elements) FMeshBatchElement);
 					BatchElementTools.PrimitiveUniformBufferResource = &GetUniformBuffer();
 					BatchElementTools.UserData = &BatchElementParams;
 
@@ -1427,19 +1390,19 @@ void FLandscapeComponentSceneProxy::GetDynamicMeshElements(const TArray<const FS
 			}
 			// Render the landscape component
 #if WITH_EDITOR
-			
-			switch(GLandscapeViewMode)
+
+			switch (GLandscapeViewMode)
 			{
 			case ELandscapeViewMode::DebugLayer:
-				if( GLayerDebugColorMaterial && EditToolRenderData )
+				if (GLayerDebugColorMaterial && EditToolRenderData)
 				{
-					auto DebugColorMaterialInstance = new FLandscapeDebugMaterialRenderProxy(GLayerDebugColorMaterial->GetRenderProxy(false), 
-						(EditToolRenderData->DebugChannelR >= 0 ? WeightmapTextures[EditToolRenderData->DebugChannelR/4] : NULL),
-						(EditToolRenderData->DebugChannelG >= 0 ? WeightmapTextures[EditToolRenderData->DebugChannelG/4] : NULL),
-						(EditToolRenderData->DebugChannelB >= 0 ? WeightmapTextures[EditToolRenderData->DebugChannelB/4] : NULL),	
-						(EditToolRenderData->DebugChannelR >= 0 ? DebugColorMask::Masks[EditToolRenderData->DebugChannelR%4] : DebugColorMask::Masks[4]),
-						(EditToolRenderData->DebugChannelG >= 0 ? DebugColorMask::Masks[EditToolRenderData->DebugChannelG%4] : DebugColorMask::Masks[4]),
-						(EditToolRenderData->DebugChannelB >= 0 ? DebugColorMask::Masks[EditToolRenderData->DebugChannelB%4] : DebugColorMask::Masks[4])
+					auto DebugColorMaterialInstance = new FLandscapeDebugMaterialRenderProxy(GLayerDebugColorMaterial->GetRenderProxy(false),
+						(EditToolRenderData->DebugChannelR >= 0 ? WeightmapTextures[EditToolRenderData->DebugChannelR / 4] : nullptr),
+						(EditToolRenderData->DebugChannelG >= 0 ? WeightmapTextures[EditToolRenderData->DebugChannelG / 4] : nullptr),
+						(EditToolRenderData->DebugChannelB >= 0 ? WeightmapTextures[EditToolRenderData->DebugChannelB / 4] : nullptr),
+						(EditToolRenderData->DebugChannelR >= 0 ? DebugColorMask::Masks[EditToolRenderData->DebugChannelR % 4] : DebugColorMask::Masks[4]),
+						(EditToolRenderData->DebugChannelG >= 0 ? DebugColorMask::Masks[EditToolRenderData->DebugChannelG % 4] : DebugColorMask::Masks[4]),
+						(EditToolRenderData->DebugChannelB >= 0 ? DebugColorMask::Masks[EditToolRenderData->DebugChannelB % 4] : DebugColorMask::Masks[4])
 						);
 
 					MeshTools.MaterialRenderProxy = DebugColorMaterialInstance;
@@ -1457,10 +1420,10 @@ void FLandscapeComponentSceneProxy::GetDynamicMeshElements(const TArray<const FS
 				break;
 
 			case ELandscapeViewMode::LayerDensity:
-				if( EditToolRenderData )
+				if (EditToolRenderData)
 				{
 					int32 ColorIndex = FMath::Min<int32>(NumWeightmapLayerAllocations, GEngine->ShaderComplexityColors.Num());
-					auto LayerDensityMaterialInstance = new FColoredMaterialRenderProxy(GEngine->LevelColorationUnlitMaterial->GetRenderProxy(false), ColorIndex ? GEngine->ShaderComplexityColors[ColorIndex-1] : FLinearColor::Black);
+					auto LayerDensityMaterialInstance = new FColoredMaterialRenderProxy(GEngine->LevelColorationUnlitMaterial->GetRenderProxy(false), ColorIndex ? GEngine->ShaderComplexityColors[ColorIndex - 1] : FLinearColor::Black);
 
 					MeshTools.MaterialRenderProxy = LayerDensityMaterialInstance;
 					Collector.RegisterOneFrameMaterialProxy(LayerDensityMaterialInstance);
@@ -1477,56 +1440,57 @@ void FLandscapeComponentSceneProxy::GetDynamicMeshElements(const TArray<const FS
 				break;
 
 			case ELandscapeViewMode::LOD:
+			{
+				FLinearColor WireColors[LANDSCAPE_LOD_LEVELS];
+				WireColors[0] = FLinearColor(1, 1, 1);
+				WireColors[1] = FLinearColor(1, 0, 0);
+				WireColors[2] = FLinearColor(0, 1, 0);
+				WireColors[3] = FLinearColor(0, 0, 1);
+				WireColors[4] = FLinearColor(1, 1, 0);
+				WireColors[5] = FLinearColor(1, 0, 1);
+				WireColors[6] = FLinearColor(0, 1, 1);
+				WireColors[7] = FLinearColor(0.5f, 0, 0.5f);
+
+				for (int32 i = 0; i < MeshTools.Elements.Num(); i++)
 				{
-					FLinearColor WireColors[7];
-					WireColors[0] = FLinearColor(1,1,1,1);
-					WireColors[1] = FLinearColor(1,0,0,1);
-					WireColors[2] = FLinearColor(0,1,0,1);
-					WireColors[3] = FLinearColor(0,0,1,1);
-					WireColors[4] = FLinearColor(1,1,0,1);
-					WireColors[5] = FLinearColor(1,0,1,1);
-					WireColors[6] = FLinearColor(0,1,1,1);
-
-					for( int32 i=0;i<MeshTools.Elements.Num();i++ )
+					FMeshBatch& LODMesh = Collector.AllocateMesh();
+					LODMesh = MeshTools;
+					LODMesh.Elements.Empty(1);
+					LODMesh.Elements.Add(MeshTools.Elements[i]);
+					int32 ColorIndex = ((FLandscapeBatchElementParams*)MeshTools.Elements[i].UserData)->CurrentLOD;
+					FLinearColor Color = ForcedLOD >= 0 ? WireColors[ColorIndex] : WireColors[ColorIndex] * 0.2f;
+					auto LODMaterialInstance = new FColoredMaterialRenderProxy(GEngine->LevelColorationUnlitMaterial->GetRenderProxy(false), Color);
+					LODMesh.MaterialRenderProxy = LODMaterialInstance;
+					Collector.RegisterOneFrameMaterialProxy(LODMaterialInstance);
+
+					if (ViewFamily.EngineShowFlags.Wireframe)
 					{
-						FMeshBatch& LODMesh = Collector.AllocateMesh();
-						LODMesh = MeshTools;
-						LODMesh.Elements.Empty(1);
-						LODMesh.Elements.Add(MeshTools.Elements[i]);
-						int32 ColorIndex = ((FLandscapeBatchElementParams*)MeshTools.Elements[i].UserData)->CurrentLOD;
-						FLinearColor Color = ForcedLOD >= 0 ? WireColors[ColorIndex] : WireColors[ColorIndex]*0.2f;
-						auto LODMaterialInstance = new FColoredMaterialRenderProxy(GEngine->LevelColorationUnlitMaterial->GetRenderProxy(false), Color);
-						LODMesh.MaterialRenderProxy = LODMaterialInstance;
-						Collector.RegisterOneFrameMaterialProxy(LODMaterialInstance);
-
-						if (ViewFamily.EngineShowFlags.Wireframe)
-						{
-							LODMesh.bCanApplyViewModeOverrides = false;
-							LODMesh.bWireframe = true;
-						}
-						else
-						{
-							LODMesh.bCanApplyViewModeOverrides = true;
-							LODMesh.bUseWireframeSelectionColoring = IsSelected();
-						}
+						LODMesh.bCanApplyViewModeOverrides = false;
+						LODMesh.bWireframe = true;
+					}
+					else
+					{
+						LODMesh.bCanApplyViewModeOverrides = true;
+						LODMesh.bUseWireframeSelectionColoring = IsSelected();
+					}
 
-						Collector.AddMesh(ViewIndex, LODMesh);
+					Collector.AddMesh(ViewIndex, LODMesh);
 
-						NumPasses++;
-						NumTriangles += MeshTools.Elements[i].NumPrimitives;
-						NumDrawCalls++;
-					}
+					NumPasses++;
+					NumTriangles += MeshTools.Elements[i].NumPrimitives;
+					NumDrawCalls++;
 				}
-				break;
+			}
+			break;
 
 			default:
-			
+
 #else
 			{
 #endif // WITH_EDITOR
 				// Regular Landscape rendering. Only use the dynamic path if we're rendering a rich view or we've disabled the static path for debugging.
 				if( IsRichView(ViewFamily) || 
-					GLandscapeDebugOptions.bDisableStatic || 
+					GLandscapeDebugOptions.bDisableStatic ||
 					bIsWireframe ||
 #if WITH_EDITOR
 					(IsSelected() && !GLandscapeEditModeActive)
@@ -1545,15 +1509,15 @@ void FLandscapeComponentSceneProxy::GetDynamicMeshElements(const TArray<const FS
 					NumPasses++;
 					NumTriangles += Mesh.GetNumPrimitives();
 					NumDrawCalls += Mesh.Elements.Num();
-				}
+			}
 			}
 #if WITH_EDITOR
-			
+
 			// Extra render passes for landscape tools
-			if( GLandscapeEditModeActive )
+			if (GLandscapeEditModeActive)
 			{
 				// Region selection
-				if ( EditToolRenderData && EditToolRenderData->SelectedType )
+				if (EditToolRenderData && EditToolRenderData->SelectedType)
 				{
 					if ((GLandscapeEditRenderMode & ELandscapeEditRenderMode::SelectRegion) && (EditToolRenderData->SelectedType & FLandscapeEditToolRenderData::ST_REGION)
 						&& !(GLandscapeEditRenderMode & ELandscapeEditRenderMode::Mask))
@@ -1584,11 +1548,11 @@ void FLandscapeComponentSceneProxy::GetDynamicMeshElements(const TArray<const FS
 				// Mask
 				if ((GLandscapeEditRenderMode & ELandscapeEditRenderMode::SelectRegion) && (GLandscapeEditRenderMode & ELandscapeEditRenderMode::Mask))
 				{
-					if (EditToolRenderData && (EditToolRenderData->SelectedType & FLandscapeEditToolRenderData::ST_REGION) )
+					if (EditToolRenderData && (EditToolRenderData->SelectedType & FLandscapeEditToolRenderData::ST_REGION))
 					{
 						FMeshBatch& MaskMesh = Collector.AllocateMesh();
 						MaskMesh = MeshTools;
-						auto MaskMaterialInstance = new FLandscapeMaskMaterialRenderProxy(GMaskRegionMaterial->GetRenderProxy(false), EditToolRenderData->DataTexture ? EditToolRenderData->DataTexture : GLandscapeBlackTexture, !!(GLandscapeEditRenderMode & ELandscapeEditRenderMode::InvertedMask) );
+						auto MaskMaterialInstance = new FLandscapeMaskMaterialRenderProxy(GMaskRegionMaterial->GetRenderProxy(false), EditToolRenderData->DataTexture ? EditToolRenderData->DataTexture : GLandscapeBlackTexture, !!(GLandscapeEditRenderMode & ELandscapeEditRenderMode::InvertedMask));
 						MaskMesh.MaterialRenderProxy = MaskMaterialInstance;
 						Collector.RegisterOneFrameMaterialProxy(MaskMaterialInstance);
 						Collector.AddMesh(ViewIndex, MaskMesh);
@@ -1600,7 +1564,7 @@ void FLandscapeComponentSceneProxy::GetDynamicMeshElements(const TArray<const FS
 					{
 						FMeshBatch& MaskMesh = Collector.AllocateMesh();
 						MaskMesh = MeshTools;
-						auto MaskMaterialInstance = new FLandscapeMaskMaterialRenderProxy(GMaskRegionMaterial->GetRenderProxy(false), GLandscapeBlackTexture, false );
+						auto MaskMaterialInstance = new FLandscapeMaskMaterialRenderProxy(GMaskRegionMaterial->GetRenderProxy(false), GLandscapeBlackTexture, false);
 						MaskMesh.MaterialRenderProxy = MaskMaterialInstance;
 						Collector.RegisterOneFrameMaterialProxy(MaskMaterialInstance);
 						Collector.AddMesh(ViewIndex, MaskMesh);
@@ -1611,13 +1575,13 @@ void FLandscapeComponentSceneProxy::GetDynamicMeshElements(const TArray<const FS
 				}
 
 				// Edit mode tools
-				if( EditToolRenderData )
+				if (EditToolRenderData)
 				{
 					if (EditToolRenderData->ToolMaterial)
 					{
 						FMeshBatch& EditMesh = Collector.AllocateMesh();
 						EditMesh = MeshTools;
-						EditMesh.MaterialRenderProxy = EditToolRenderData->ToolMaterial->GetRenderProxy(0);		
+						EditMesh.MaterialRenderProxy = EditToolRenderData->ToolMaterial->GetRenderProxy(0);
 						Collector.AddMesh(ViewIndex, EditMesh);
 						NumPasses++;
 						NumTriangles += EditMesh.GetNumPrimitives();
@@ -1638,7 +1602,7 @@ void FLandscapeComponentSceneProxy::GetDynamicMeshElements(const TArray<const FS
 			}
 #endif // WITH_EDITOR
 
-			if ( GLandscapeDebugOptions.bShowPatches )
+			if (GLandscapeDebugOptions.bShowPatches)
 			{
 				DrawWireBox(Collector.GetPDI(ViewIndex), GetBounds().GetBox(), FColor(255, 255, 0), SDPG_World);
 			}
@@ -1656,23 +1620,23 @@ void FLandscapeComponentSceneProxy::GetDynamicMeshElements(const TArray<const FS
 // FLandscapeVertexBuffer
 //
 
-/** 
-* Initialize the RHI for this rendering resource 
+/**
+* Initialize the RHI for this rendering resource
 */
 void FLandscapeVertexBuffer::InitRHI()
 {
 	// create a static vertex buffer
 	FRHIResourceCreateInfo CreateInfo;
 	VertexBufferRHI = RHICreateVertexBuffer(FMath::Square(SubsectionSizeVerts) * FMath::Square(NumSubsections) * sizeof(FLandscapeVertex), BUF_Static, CreateInfo);
-	FLandscapeVertex* Vertex = (FLandscapeVertex*)RHILockVertexBuffer(VertexBufferRHI, 0, FMath::Square(SubsectionSizeVerts) * FMath::Square(NumSubsections) * sizeof(FLandscapeVertex),RLM_WriteOnly);
+	FLandscapeVertex* Vertex = (FLandscapeVertex*)RHILockVertexBuffer(VertexBufferRHI, 0, FMath::Square(SubsectionSizeVerts) * FMath::Square(NumSubsections) * sizeof(FLandscapeVertex), RLM_WriteOnly);
 
-	for( int32 SubY=0;SubY<NumSubsections;SubY++ )
+	for (int32 SubY = 0; SubY < NumSubsections; SubY++)
 	{
-		for( int32 SubX=0;SubX<NumSubsections;SubX++ )
+		for (int32 SubX = 0; SubX < NumSubsections; SubX++)
 		{
-			for( int32 y=0;y<SubsectionSizeVerts;y++ )
+			for (int32 y = 0; y < SubsectionSizeVerts; y++)
 			{
-				for( int32 x=0;x<SubsectionSizeVerts;x++ )
+				for (int32 x = 0; x < SubsectionSizeVerts; x++)
 				{
 					Vertex->VertexX = x;
 					Vertex->VertexY = y;
@@ -1726,9 +1690,9 @@ void FLandscapeSharedBuffers::CreateIndexBuffers(ERHIFeatureLevel::Type InFeatur
 			// ES2 version
 			float MipRatio = (float)SubsectionSizeQuads / (float)LodSubsectionSizeQuads; // Morph current MIP to base MIP
 
-			for( int32 SubY=0;SubY<NumSubsections;SubY++ )
+			for (int32 SubY = 0; SubY < NumSubsections; SubY++)
 			{
-				for( int32 SubX=0;SubX<NumSubsections;SubX++ )
+				for (int32 SubX = 0; SubX < NumSubsections; SubX++)
 				{
 					TArray<INDEX_TYPE> SubIndices;
 					SubIndices.Empty(FMath::Square(LodSubsectionSizeQuads) * 6);
@@ -1763,7 +1727,7 @@ void FLandscapeSharedBuffers::CreateIndexBuffers(ERHIFeatureLevel::Type InFeatur
 							INDEX_TYPE i01;
 
 							INDEX_TYPE* KeyPtr = VertexMap.Find(Key00);
-							if (KeyPtr == NULL)
+							if (KeyPtr == nullptr)
 							{
 								i00 = VertexCount++;
 								VertexMap.Add(Key00, i00);
@@ -1774,7 +1738,7 @@ void FLandscapeSharedBuffers::CreateIndexBuffers(ERHIFeatureLevel::Type InFeatur
 							}
 
 							KeyPtr = VertexMap.Find(Key10);
-							if (KeyPtr == NULL)
+							if (KeyPtr == nullptr)
 							{
 								i10 = VertexCount++;
 								VertexMap.Add(Key10, i10);
@@ -1785,7 +1749,7 @@ void FLandscapeSharedBuffers::CreateIndexBuffers(ERHIFeatureLevel::Type InFeatur
 							}
 
 							KeyPtr = VertexMap.Find(Key11);
-							if (KeyPtr == NULL)
+							if (KeyPtr == nullptr)
 							{
 								i11 = VertexCount++;
 								VertexMap.Add(Key11, i11);
@@ -1796,7 +1760,7 @@ void FLandscapeSharedBuffers::CreateIndexBuffers(ERHIFeatureLevel::Type InFeatur
 							}
 
 							KeyPtr = VertexMap.Find(Key01);
-							if (KeyPtr == NULL)
+							if (KeyPtr == nullptr)
 							{
 								i01 = VertexCount++;
 								VertexMap.Add(Key01, i01);
@@ -1854,10 +1818,10 @@ void FLandscapeSharedBuffers::CreateIndexBuffers(ERHIFeatureLevel::Type InFeatur
 					{
 						for (int32 x = 0; x < LodSubsectionSizeQuads; x++)
 						{
-							INDEX_TYPE i00 = (x+0) + (y+0) * SubsectionSizeVerts + SubOffset;
-							INDEX_TYPE i10 = (x+1) + (y+0) * SubsectionSizeVerts + SubOffset;
-							INDEX_TYPE i11 = (x+1) + (y+1) * SubsectionSizeVerts + SubOffset;
-							INDEX_TYPE i01 = (x+0) + (y+1) * SubsectionSizeVerts + SubOffset;
+							INDEX_TYPE i00 = (x + 0) + (y + 0) * SubsectionSizeVerts + SubOffset;
+							INDEX_TYPE i10 = (x + 1) + (y + 0) * SubsectionSizeVerts + SubOffset;
+							INDEX_TYPE i11 = (x + 1) + (y + 1) * SubsectionSizeVerts + SubOffset;
+							INDEX_TYPE i01 = (x + 0) + (y + 1) * SubsectionSizeVerts + SubOffset;
 
 							NewIndices.Add(i00);
 							NewIndices.Add(i11);
@@ -1908,14 +1872,14 @@ void FLandscapeSharedBuffers::CreateIndexBuffers(ERHIFeatureLevel::Type InFeatur
 }
 
 FLandscapeSharedBuffers::FLandscapeSharedBuffers(int32 InSharedBuffersKey, int32 InSubsectionSizeQuads, int32 InNumSubsections, ERHIFeatureLevel::Type InFeatureLevel, bool bRequiresAdjacencyInformation)
-:	SharedBuffersKey(InSharedBuffersKey)
-,	NumIndexBuffers(FMath::CeilLogTwo(InSubsectionSizeQuads+1))
-,	SubsectionSizeVerts(InSubsectionSizeQuads+1)
-,	NumSubsections(InNumSubsections)
-,	VertexFactory(NULL)
-,	VertexBuffer(NULL)
-,	AdjacencyIndexBuffers(NULL)
-,	bUse32BitIndices(false)
+	: SharedBuffersKey(InSharedBuffersKey)
+	, NumIndexBuffers(FMath::CeilLogTwo(InSubsectionSizeQuads + 1))
+	, SubsectionSizeVerts(InSubsectionSizeQuads + 1)
+	, NumSubsections(InNumSubsections)
+	, VertexFactory(nullptr)
+	, VertexBuffer(nullptr)
+	, AdjacencyIndexBuffers(nullptr)
+	, bUse32BitIndices(false)
 {
 	if (InFeatureLevel > ERHIFeatureLevel::ES3_1)
 	{
@@ -1927,7 +1891,7 @@ FLandscapeSharedBuffers::FLandscapeSharedBuffers(int32 InSharedBuffersKey, int32
 	IndexRanges = new FLandscapeIndexRanges[NumIndexBuffers]();
 
 	// See if we need to use 16 or 32-bit index buffers
-	if( FMath::Square(SubsectionSizeVerts) * FMath::Square(NumSubsections) > 65535 )
+	if (FMath::Square(SubsectionSizeVerts) * FMath::Square(NumSubsections) > 65535)
 	{
 		bUse32BitIndices = true;
 		CreateIndexBuffers<uint32>(InFeatureLevel, bRequiresAdjacencyInformation);
@@ -1941,8 +1905,8 @@ FLandscapeSharedBuffers::FLandscapeSharedBuffers(int32 InSharedBuffersKey, int32
 FLandscapeSharedBuffers::~FLandscapeSharedBuffers()
 {
 	delete VertexBuffer;
-	
-	for( int32 i=0;i<NumIndexBuffers;i++ )
+
+	for (int32 i = 0; i < NumIndexBuffers; i++)
 	{
 		IndexBuffers[i]->ReleaseResource();
 		delete IndexBuffers[i];
@@ -1956,7 +1920,7 @@ FLandscapeSharedBuffers::~FLandscapeSharedBuffers()
 		{
 			FLandscapeComponentSceneProxy::SharedAdjacencyIndexBufferMap.Remove(SharedBuffersKey);
 		}
-		AdjacencyIndexBuffers = NULL;
+		AdjacencyIndexBuffers = nullptr;
 	}
 
 	delete VertexFactory;
@@ -1971,27 +1935,27 @@ static void BuildLandscapeAdjacencyIndexBuffer(int32 LODSubsectionSizeQuads, con
 		// PN AEN Dominant Corner
 		uint32 TriCount = LODSubsectionSizeQuads*LODSubsectionSizeQuads * 2;
 		uint32 ExpandedCount = 12 * TriCount;
-		OutPnAenIndices.Empty( ExpandedCount );
-		OutPnAenIndices.AddUninitialized( ExpandedCount );
+		OutPnAenIndices.Empty(ExpandedCount);
+		OutPnAenIndices.AddUninitialized(ExpandedCount);
 
-		for ( uint32 TriIdx = 0; TriIdx < TriCount; ++TriIdx )
+		for (uint32 TriIdx = 0; TriIdx < TriCount; ++TriIdx)
 		{
-			uint32 OutStartIdx = TriIdx*12;
-			uint32 InStartIdx = TriIdx*3;
-			OutPnAenIndices[ OutStartIdx +  0] =	Indices->Get(InStartIdx + 0);
-			OutPnAenIndices[ OutStartIdx +  1] =	Indices->Get(InStartIdx + 1);
-			OutPnAenIndices[ OutStartIdx +  2] =	Indices->Get(InStartIdx + 2);
+			uint32 OutStartIdx = TriIdx * 12;
+			uint32 InStartIdx = TriIdx * 3;
+			OutPnAenIndices[OutStartIdx + 0] = Indices->Get(InStartIdx + 0);
+			OutPnAenIndices[OutStartIdx + 1] = Indices->Get(InStartIdx + 1);
+			OutPnAenIndices[OutStartIdx + 2] = Indices->Get(InStartIdx + 2);
 
-			OutPnAenIndices[ OutStartIdx +  3] =	Indices->Get(InStartIdx + 0);
-			OutPnAenIndices[ OutStartIdx +  4] =	Indices->Get(InStartIdx + 1);
-			OutPnAenIndices[ OutStartIdx +  5] =	Indices->Get(InStartIdx + 1);
-			OutPnAenIndices[ OutStartIdx +  6] =	Indices->Get(InStartIdx + 2);
-			OutPnAenIndices[ OutStartIdx +  7] =	Indices->Get(InStartIdx + 2);
-			OutPnAenIndices[ OutStartIdx +  8] =	Indices->Get(InStartIdx + 0);
+			OutPnAenIndices[OutStartIdx + 3] = Indices->Get(InStartIdx + 0);
+			OutPnAenIndices[OutStartIdx + 4] = Indices->Get(InStartIdx + 1);
+			OutPnAenIndices[OutStartIdx + 5] = Indices->Get(InStartIdx + 1);
+			OutPnAenIndices[OutStartIdx + 6] = Indices->Get(InStartIdx + 2);
+			OutPnAenIndices[OutStartIdx + 7] = Indices->Get(InStartIdx + 2);
+			OutPnAenIndices[OutStartIdx + 8] = Indices->Get(InStartIdx + 0);
 
-			OutPnAenIndices[ OutStartIdx +  9] =	Indices->Get(InStartIdx + 0);
-			OutPnAenIndices[ OutStartIdx +  10] =	Indices->Get(InStartIdx + 1);
-			OutPnAenIndices[ OutStartIdx +  11] =	Indices->Get(InStartIdx + 2);
+			OutPnAenIndices[OutStartIdx + 9] = Indices->Get(InStartIdx + 0);
+			OutPnAenIndices[OutStartIdx + 10] = Indices->Get(InStartIdx + 1);
+			OutPnAenIndices[OutStartIdx + 11] = Indices->Get(InStartIdx + 2);
 		}
 	}
 	else
@@ -2010,10 +1974,10 @@ FLandscapeSharedAdjacencyIndexBuffer::FLandscapeSharedAdjacencyIndexBuffer(FLand
 	bool b32BitIndex = FMath::Square(Buffers->SubsectionSizeVerts) * FMath::Square(Buffers->NumSubsections) > 65535;
 	for (int32 i = 0; i < Buffers->NumIndexBuffers; ++i)
 	{
-		if( b32BitIndex )
+		if (b32BitIndex)
 		{
 			TArray<uint32> OutPnAenIndices;
-			BuildLandscapeAdjacencyIndexBuffer<uint32>( (Buffers->SubsectionSizeVerts>>i)-1, (FRawStaticIndexBuffer16or32<uint32>*)Buffers->IndexBuffers[i], OutPnAenIndices );
+			BuildLandscapeAdjacencyIndexBuffer<uint32>((Buffers->SubsectionSizeVerts >> i) - 1, (FRawStaticIndexBuffer16or32<uint32>*)Buffers->IndexBuffers[i], OutPnAenIndices);
 
 			FRawStaticIndexBuffer16or32<uint32>* IndexBuffer = new FRawStaticIndexBuffer16or32<uint32>();
 			IndexBuffer->AssignNewBuffer(OutPnAenIndices);
@@ -2022,7 +1986,7 @@ FLandscapeSharedAdjacencyIndexBuffer::FLandscapeSharedAdjacencyIndexBuffer(FLand
 		else
 		{
 			TArray<uint16> OutPnAenIndices;
-			BuildLandscapeAdjacencyIndexBuffer<uint16>( (Buffers->SubsectionSizeVerts>>i)-1, (FRawStaticIndexBuffer16or32<uint16>*)Buffers->IndexBuffers[i], OutPnAenIndices );
+			BuildLandscapeAdjacencyIndexBuffer<uint16>((Buffers->SubsectionSizeVerts >> i) - 1, (FRawStaticIndexBuffer16or32<uint16>*)Buffers->IndexBuffers[i], OutPnAenIndices);
 
 			FRawStaticIndexBuffer16or32<uint16>* IndexBuffer = new FRawStaticIndexBuffer16or32<uint16>();
 			IndexBuffer->AssignNewBuffer(OutPnAenIndices);
@@ -2056,14 +2020,14 @@ public:
 	*/
 	virtual void Bind(const FShaderParameterMap& ParameterMap) override
 	{
-		HeightmapTextureParameter.Bind(ParameterMap,TEXT("HeightmapTexture"));
-		HeightmapTextureParameterSampler.Bind(ParameterMap,TEXT("HeightmapTextureSampler"));
-		LodValuesParameter.Bind(ParameterMap,TEXT("LodValues"));
-		NeighborSectionLodParameter.Bind(ParameterMap,TEXT("NeighborSectionLod"));
-		LodBiasParameter.Bind(ParameterMap,TEXT("LodBias"));
-		SectionLodsParameter.Bind(ParameterMap,TEXT("SectionLods"));
-		XYOffsetTextureParameter.Bind(ParameterMap,TEXT("XYOffsetmapTexture"));
-		XYOffsetTextureParameterSampler.Bind(ParameterMap,TEXT("XYOffsetmapTextureSampler"));
+		HeightmapTextureParameter.Bind(ParameterMap, TEXT("HeightmapTexture"));
+		HeightmapTextureParameterSampler.Bind(ParameterMap, TEXT("HeightmapTextureSampler"));
+		LodValuesParameter.Bind(ParameterMap, TEXT("LodValues"));
+		NeighborSectionLodParameter.Bind(ParameterMap, TEXT("NeighborSectionLod"));
+		LodBiasParameter.Bind(ParameterMap, TEXT("LodBias"));
+		SectionLodsParameter.Bind(ParameterMap, TEXT("SectionLods"));
+		XYOffsetTextureParameter.Bind(ParameterMap, TEXT("XYOffsetmapTexture"));
+		XYOffsetTextureParameterSampler.Bind(ParameterMap, TEXT("XYOffsetmapTextureSampler"));
 	}
 
 	/**
@@ -2085,7 +2049,7 @@ public:
 	/**
 	* Set any shader data specific to this vertex factory
 	*/
-	virtual void SetMesh(FRHICommandList& RHICmdList, FShader* VertexShader,const class FVertexFactory* VertexFactory,const class FSceneView& View,const struct FMeshBatchElement& BatchElement,uint32 DataFlags) const override
+	virtual void SetMesh(FRHICommandList& RHICmdList, FShader* VertexShader, const class FVertexFactory* VertexFactory, const class FSceneView& View, const struct FMeshBatchElement& BatchElement, uint32 DataFlags) const override
 	{
 		SCOPE_CYCLE_COUNTER(STAT_LandscapeVFDrawTime);
 
@@ -2093,12 +2057,12 @@ public:
 		check(BatchElementParams);
 
 		const FLandscapeComponentSceneProxy* SceneProxy = BatchElementParams->SceneProxy;
-		SetUniformBufferParameter(RHICmdList, VertexShader->GetVertexShader(),VertexShader->GetUniformBufferParameter<FLandscapeUniformShaderParameters>(),*BatchElementParams->LandscapeUniformShaderParametersResource);
+		SetUniformBufferParameter(RHICmdList, VertexShader->GetVertexShader(), VertexShader->GetUniformBufferParameter<FLandscapeUniformShaderParameters>(), *BatchElementParams->LandscapeUniformShaderParametersResource);
 
-		if( HeightmapTextureParameter.IsBound() )
+		if (HeightmapTextureParameter.IsBound())
 		{
 			SetTextureParameter(
-				RHICmdList, 
+				RHICmdList,
 				VertexShader->GetVertexShader(),
 				HeightmapTextureParameter,
 				HeightmapTextureParameterSampler,
@@ -2107,62 +2071,74 @@ public:
 				);
 		}
 
-		if( LodBiasParameter.IsBound() )
+		if (LodBiasParameter.IsBound())
 		{
 			FVector4 LodBias(
-				SceneProxy->PrecomputedLODFactor,
-				1.f / (1.f - SceneProxy->PrecomputedLODFactor),
+				0.0f, // unused
+				0.0f, // unused
 				SceneProxy->HeightmapTexture->GetNumMips() - FMath::Min(SceneProxy->HeightmapTexture->ResidentMips, SceneProxy->HeightmapTexture->RequestedMips),
-				0.f // Reserved
+				SceneProxy->XYOffsetmapTexture
+					? SceneProxy->XYOffsetmapTexture->GetNumMips() - FMath::Min(SceneProxy->XYOffsetmapTexture->ResidentMips, SceneProxy->XYOffsetmapTexture->RequestedMips)
+					: 0.0f
 				);
 			SetShaderValue(RHICmdList, VertexShader->GetVertexShader(), LodBiasParameter, LodBias);
 		}
 
 		// Calculate LOD params
-		FVector CameraLocalPos3D = SceneProxy->WorldToLocal.TransformPosition(View.ViewMatrices.ViewOrigin); 
+		FVector CameraLocalPos3D = SceneProxy->WorldToLocal.TransformPosition(View.ViewMatrices.ViewOrigin);
 		FVector2D CameraLocalPos = FVector2D(CameraLocalPos3D.X, CameraLocalPos3D.Y);
 
 		FVector4 fCurrentLODs;
 		FVector4 CurrentNeighborLODs[4];
-		float DistLOD = 0.f;
 
-		if( BatchElementParams->SubX == -1 )
+		if (BatchElementParams->SubX == -1)
 		{
-			for( int32 SubY = 0; SubY < SceneProxy->NumSubsections; SubY++ )
+			for (int32 SubY = 0; SubY < SceneProxy->NumSubsections; SubY++)
 			{
-				for( int32 SubX = 0; SubX < SceneProxy->NumSubsections; SubX++ )
+				for (int32 SubX = 0; SubX < SceneProxy->NumSubsections; SubX++)
 				{
 					int32 SubIndex = SubX + 2 * SubY;
-					SceneProxy->CalcLODParamsForSubsection(View, CameraLocalPos, SubX, SubY, fCurrentLODs[SubIndex], CurrentNeighborLODs[SubIndex], DistLOD);
+					SceneProxy->CalcLODParamsForSubsection(View, CameraLocalPos, SubX, SubY, BatchElementParams->CurrentLOD, fCurrentLODs[SubIndex], CurrentNeighborLODs[SubIndex]);
 				}
 			}
 		}
 		else
 		{
 			int32 SubIndex = BatchElementParams->SubX + 2 * BatchElementParams->SubY;
-			SceneProxy->CalcLODParamsForSubsection(View, CameraLocalPos, BatchElementParams->SubX, BatchElementParams->SubY, fCurrentLODs[SubIndex], CurrentNeighborLODs[SubIndex], DistLOD);
+			SceneProxy->CalcLODParamsForSubsection(View, CameraLocalPos, BatchElementParams->SubX, BatchElementParams->SubY, BatchElementParams->CurrentLOD, fCurrentLODs[SubIndex], CurrentNeighborLODs[SubIndex]);
 		}
 
-		if( SectionLodsParameter.IsBound() )
+		if (SectionLodsParameter.IsBound())
 		{
 			SetShaderValue(RHICmdList, VertexShader->GetVertexShader(), SectionLodsParameter, fCurrentLODs);
 		}
 
-		if( NeighborSectionLodParameter.IsBound() )
+		if (NeighborSectionLodParameter.IsBound())
 		{
 			SetShaderValue(RHICmdList, VertexShader->GetVertexShader(), NeighborSectionLodParameter, CurrentNeighborLODs);
 		}
 
-		if( LodValuesParameter.IsBound() )
+		if (LodValuesParameter.IsBound())
 		{
 			FVector4 LodValues(
-				DistLOD,
-				// convert current LOD coordinates into highest LOD coordinates
-				(float)SceneProxy->SubsectionSizeQuads / (float)(((SceneProxy->SubsectionSizeVerts) >> BatchElementParams->CurrentLOD)-1),
+				BatchElementParams->CurrentLOD,
+				0.0f, // unused
 				(float)((SceneProxy->SubsectionSizeVerts >> BatchElementParams->CurrentLOD) - 1),
-				1.f/(float)((SceneProxy->SubsectionSizeVerts >> BatchElementParams->CurrentLOD) - 1) );
+				1.f / (float)((SceneProxy->SubsectionSizeVerts >> BatchElementParams->CurrentLOD) - 1));
 
-			SetShaderValue(RHICmdList, VertexShader->GetVertexShader(),LodValuesParameter,LodValues);
+			SetShaderValue(RHICmdList, VertexShader->GetVertexShader(), LodValuesParameter, LodValues);
+		}
+
+		if (XYOffsetTextureParameter.IsBound() && SceneProxy->XYOffsetmapTexture)
+		{
+			SetTextureParameter(
+				RHICmdList,
+				VertexShader->GetVertexShader(),
+				XYOffsetTextureParameter,
+				XYOffsetTextureParameterSampler,
+				TStaticSamplerState<SF_Point>::GetRHI(),
+				SceneProxy->XYOffsetmapTexture->Resource->TextureRHI
+				);
 		}
 	}
 
@@ -2183,107 +2159,6 @@ protected:
 	TShaderUniformBufferParameter<FLandscapeUniformShaderParameters> LandscapeShaderParameters;
 };
 
-/** Shader parameters for use with FLandscapeVertexFactory */
-class FLandscapeXYOffsetVertexFactoryVertexShaderParameters : public FLandscapeVertexFactoryVertexShaderParameters
-{
-public:
-	/**
-	* Set any shader data specific to this vertex factory
-	*/
-	virtual void SetMesh(FRHICommandList& RHICmdList, FShader* VertexShader,const class FVertexFactory* VertexFactory,const class FSceneView& View,const struct FMeshBatchElement& BatchElement,uint32 DataFlags) const override
-	{
-		SCOPE_CYCLE_COUNTER(STAT_LandscapeVFDrawTime);
-
-		const FLandscapeBatchElementParams* BatchElementParams = (const FLandscapeBatchElementParams*)BatchElement.UserData;
-		check(BatchElementParams);
-
-		const FLandscapeComponentSceneProxy* SceneProxy = BatchElementParams->SceneProxy;
-		SetUniformBufferParameter(RHICmdList, VertexShader->GetVertexShader(),VertexShader->GetUniformBufferParameter<FLandscapeUniformShaderParameters>(),*BatchElementParams->LandscapeUniformShaderParametersResource);
-
-		if( HeightmapTextureParameter.IsBound() )
-		{
-			SetTextureParameter(
-				RHICmdList, 
-				VertexShader->GetVertexShader(),
-				HeightmapTextureParameter,
-				HeightmapTextureParameterSampler,
-				TStaticSamplerState<SF_Point>::GetRHI(),
-				SceneProxy->HeightmapTexture->Resource->TextureRHI
-				);
-		}
-
-		if( LodBiasParameter.IsBound() )
-		{
-			FVector4 LodBias(
-				SceneProxy->PrecomputedLODFactor,
-				1.f / (1.f - SceneProxy->PrecomputedLODFactor),
-				SceneProxy->HeightmapTexture->GetNumMips() - FMath::Min(SceneProxy->HeightmapTexture->ResidentMips, SceneProxy->HeightmapTexture->RequestedMips),
-				SceneProxy->XYOffsetmapTexture->GetNumMips() - FMath::Min(SceneProxy->XYOffsetmapTexture->ResidentMips, SceneProxy->XYOffsetmapTexture->RequestedMips)
-				);
-			SetShaderValue(RHICmdList, VertexShader->GetVertexShader(), LodBiasParameter, LodBias);
-		}
-
-		// Calculate LOD params
-		FVector CameraLocalPos3D = SceneProxy->WorldToLocal.TransformPosition(View.ViewMatrices.ViewOrigin); 
-		FVector2D CameraLocalPos = FVector2D(CameraLocalPos3D.X, CameraLocalPos3D.Y);
-
-		FVector4 fCurrentLODs;
-		FVector4 CurrentNeighborLODs[4];
-		float DistLOD = 0.f;
-
-		if( BatchElementParams->SubX == -1 )
-		{
-			for( int32 SubY = 0; SubY < SceneProxy->NumSubsections; SubY++ )
-			{
-				for( int32 SubX = 0; SubX < SceneProxy->NumSubsections; SubX++ )
-				{
-					int32 SubIndex = SubX + 2 * SubY;
-					SceneProxy->CalcLODParamsForSubsection(View, CameraLocalPos, SubX, SubY, fCurrentLODs[SubIndex], CurrentNeighborLODs[SubIndex], DistLOD);
-				}
-			}
-		}
-		else
-		{
-			int32 SubIndex = BatchElementParams->SubX + 2 * BatchElementParams->SubY;
-			SceneProxy->CalcLODParamsForSubsection(View, CameraLocalPos, BatchElementParams->SubX, BatchElementParams->SubY, fCurrentLODs[SubIndex], CurrentNeighborLODs[SubIndex], DistLOD);
-		}
-
-		if( SectionLodsParameter.IsBound() )
-		{
-			SetShaderValue(RHICmdList, VertexShader->GetVertexShader(), SectionLodsParameter, fCurrentLODs);
-		}
-
-		if( NeighborSectionLodParameter.IsBound() )
-		{
-			SetShaderValue(RHICmdList, VertexShader->GetVertexShader(), NeighborSectionLodParameter, CurrentNeighborLODs);
-		}
-
-		if( LodValuesParameter.IsBound() )
-		{
-			FVector4 LodValues(
-				DistLOD,
-				// convert current LOD coordinates into highest LOD coordinates
-				(float)SceneProxy->SubsectionSizeQuads / (float)(((SceneProxy->SubsectionSizeVerts) >> BatchElementParams->CurrentLOD)-1),
-				(float)((SceneProxy->SubsectionSizeVerts >> BatchElementParams->CurrentLOD) - 1),
-				1.f/(float)((SceneProxy->SubsectionSizeVerts >> BatchElementParams->CurrentLOD) - 1) );
-
-			SetShaderValue(RHICmdList, VertexShader->GetVertexShader(),LodValuesParameter,LodValues);
-		}
-
-		if( XYOffsetTextureParameter.IsBound() && SceneProxy->XYOffsetmapTexture )
-		{
-			SetTextureParameter(
-				RHICmdList, 
-				VertexShader->GetVertexShader(),
-				XYOffsetTextureParameter,
-				XYOffsetTextureParameterSampler,
-				TStaticSamplerState<SF_Point>::GetRHI(),
-				SceneProxy->XYOffsetmapTexture->Resource->TextureRHI
-				);
-		}
-	}
-};
-
 //
 // FLandscapeVertexFactoryPixelShaderParameters
 //
@@ -2293,9 +2168,9 @@ public:
 */
 void FLandscapeVertexFactoryPixelShaderParameters::Bind(const FShaderParameterMap& ParameterMap)
 {
-	NormalmapTextureParameter.Bind(ParameterMap,TEXT("NormalmapTexture"));
-	NormalmapTextureParameterSampler.Bind(ParameterMap,TEXT("NormalmapTextureSampler"));
-	LocalToWorldNoScalingParameter.Bind(ParameterMap,TEXT("LocalToWorldNoScaling"));
+	NormalmapTextureParameter.Bind(ParameterMap, TEXT("NormalmapTexture"));
+	NormalmapTextureParameterSampler.Bind(ParameterMap, TEXT("NormalmapTextureSampler"));
+	LocalToWorldNoScalingParameter.Bind(ParameterMap, TEXT("LocalToWorldNoScaling"));
 }
 
 /**
@@ -2304,7 +2179,7 @@ void FLandscapeVertexFactoryPixelShaderParameters::Bind(const FShaderParameterMa
 */
 void FLandscapeVertexFactoryPixelShaderParameters::Serialize(FArchive& Ar)
 {
-	Ar	<< NormalmapTextureParameter
+	Ar << NormalmapTextureParameter
 		<< NormalmapTextureParameterSampler
 		<< LocalToWorldNoScalingParameter;
 }
@@ -2312,21 +2187,21 @@ void FLandscapeVertexFactoryPixelShaderParameters::Serialize(FArchive& Ar)
 /**
 * Set any shader data specific to this vertex factory
 */
-void FLandscapeVertexFactoryPixelShaderParameters::SetMesh(FRHICommandList& RHICmdList, FShader* PixelShader,const class FVertexFactory* VertexFactory,const class FSceneView& View,const struct FMeshBatchElement& BatchElement,uint32 DataFlags) const
+void FLandscapeVertexFactoryPixelShaderParameters::SetMesh(FRHICommandList& RHICmdList, FShader* PixelShader, const class FVertexFactory* VertexFactory, const class FSceneView& View, const struct FMeshBatchElement& BatchElement, uint32 DataFlags) const
 {
 	SCOPE_CYCLE_COUNTER(STAT_LandscapeVFDrawTime);
 
 	const FLandscapeBatchElementParams* BatchElementParams = (const FLandscapeBatchElementParams*)BatchElement.UserData;
 
-	if( LocalToWorldNoScalingParameter.IsBound() )
+	if (LocalToWorldNoScalingParameter.IsBound())
 	{
 		SetShaderValue(RHICmdList, PixelShader->GetPixelShader(), LocalToWorldNoScalingParameter, *BatchElementParams->LocalToWorldNoScalingPtr);
 	}
 
-	if( NormalmapTextureParameter.IsBound() )
+	if (NormalmapTextureParameter.IsBound())
 	{
 		SetTextureParameter(
-			RHICmdList, 
+			RHICmdList,
 			PixelShader->GetPixelShader(),
 			NormalmapTextureParameter,
 			NormalmapTextureParameterSampler,
@@ -2345,15 +2220,15 @@ void FLandscapeVertexFactory::InitRHI()
 	FVertexDeclarationElementList Elements;
 
 	// position decls
-	Elements.Add(AccessStreamComponent(Data.PositionComponent,0));
+	Elements.Add(AccessStreamComponent(Data.PositionComponent, 0));
 
 	// create the actual device decls
-	InitDeclaration(Elements,FVertexFactory::DataType());
+	InitDeclaration(Elements, FVertexFactory::DataType());
 }
 
 FVertexFactoryShaderParameters* FLandscapeVertexFactory::ConstructShaderParameters(EShaderFrequency ShaderFrequency)
 {
-	switch( ShaderFrequency )
+	switch (ShaderFrequency)
 	{
 	case SF_Vertex:
 		return new FLandscapeVertexFactoryVertexShaderParameters();
@@ -2362,11 +2237,11 @@ FVertexFactoryShaderParameters* FLandscapeVertexFactory::ConstructShaderParamete
 		return new FLandscapeVertexFactoryPixelShaderParameters();
 		break;
 	default:
-		return NULL;
+		return nullptr;
 	}
 }
 
-void FLandscapeVertexFactory::ModifyCompilationEnvironment( EShaderPlatform Platform, const FMaterial* Material, FShaderCompilerEnvironment& OutEnvironment )
+void FLandscapeVertexFactory::ModifyCompilationEnvironment(EShaderPlatform Platform, const FMaterial* Material, FShaderCompilerEnvironment& OutEnvironment)
 {
 	FVertexFactory::ModifyCompilationEnvironment(Platform, Material, OutEnvironment);
 }
@@ -2383,34 +2258,19 @@ void FLandscapeVertexFactory::Copy(const FLandscapeVertexFactory& Other)
 	//SetSceneProxy(Other.Proxy());
 	ENQUEUE_UNIQUE_RENDER_COMMAND_TWOPARAMETER(
 		FLandscapeVertexFactoryCopyData,
-		FLandscapeVertexFactory*,VertexFactory,this,
-		const DataType*,DataCopy,&Other.Data,
-	{
-		VertexFactory->Data = *DataCopy;
-	});	
+		FLandscapeVertexFactory*, VertexFactory, this,
+		const DataType*, DataCopy, &Other.Data,
+		{
+			VertexFactory->Data = *DataCopy;
+		});
 	BeginUpdateResourceRHI(this);
 }
 
 //
-// FLandscapeVertexFactory
+// FLandscapeXYOffsetVertexFactory
 //
 
-FVertexFactoryShaderParameters* FLandscapeXYOffsetVertexFactory::ConstructShaderParameters(EShaderFrequency ShaderFrequency)
-{
-	switch( ShaderFrequency )
-	{
-	case SF_Vertex:
-		return new FLandscapeXYOffsetVertexFactoryVertexShaderParameters();
-		break;
-	case SF_Pixel:
-		return new FLandscapeVertexFactoryPixelShaderParameters();
-		break;
-	default:
-		return NULL;
-	}
-}
-
-void FLandscapeXYOffsetVertexFactory::ModifyCompilationEnvironment( EShaderPlatform Platform, const FMaterial* Material, FShaderCompilerEnvironment& OutEnvironment )
+void FLandscapeXYOffsetVertexFactory::ModifyCompilationEnvironment(EShaderPlatform Platform, const FMaterial* Material, FShaderCompilerEnvironment& OutEnvironment)
 {
 	FLandscapeVertexFactory::ModifyCompilationEnvironment(Platform, Material, OutEnvironment);
 	OutEnvironment.SetDefine(TEXT("LANDSCAPE_XYOFFSET"), TEXT("1"));
@@ -2446,9 +2306,9 @@ void ULandscapeComponent::GetStreamingTextureInfo(TArray<FStreamingTexturePrimit
 		MaterialInstance->GetUsedTextures(Textures, EMaterialQualityLevel::Num, false, GetWorld()->FeatureLevel, false);
 		// Add each texture to the output with the appropriate parameters.
 		// TODO: Take into account which UVIndex is being used.
-		for (int32 TextureIndex = 0;TextureIndex < Textures.Num();TextureIndex++)
+		for (int32 TextureIndex = 0; TextureIndex < Textures.Num(); TextureIndex++)
 		{
-			FStreamingTexturePrimitiveInfo& StreamingTexture = *new(OutStreamingTextures) FStreamingTexturePrimitiveInfo;
+			FStreamingTexturePrimitiveInfo& StreamingTexture = *new(OutStreamingTextures)FStreamingTexturePrimitiveInfo;
 			StreamingTexture.Bounds = BoundingSphere;
 			StreamingTexture.TexelFactor = TexelFactor;
 			StreamingTexture.Texture = Textures[TextureIndex];
@@ -2458,19 +2318,19 @@ void ULandscapeComponent::GetStreamingTextureInfo(TArray<FStreamingTexturePrimit
 		if (Material)
 		{
 			int32 NumExpressions = Material->Expressions.Num();
-			for(int32 ExpressionIndex = 0;ExpressionIndex < NumExpressions; ExpressionIndex++)
+			for (int32 ExpressionIndex = 0; ExpressionIndex < NumExpressions; ExpressionIndex++)
 			{
 				UMaterialExpression* Expression = Material->Expressions[ExpressionIndex];
 				UMaterialExpressionTextureSample* TextureSample = Cast<UMaterialExpressionTextureSample>(Expression);
 
 				// TODO: This is only works for direct Coordinate Texture Sample cases
-				if(TextureSample && TextureSample->Coordinates.Expression)
+				if (TextureSample && TextureSample->Coordinates.Expression)
 				{
 					UMaterialExpressionTextureCoordinate* TextureCoordinate =
-						Cast<UMaterialExpressionTextureCoordinate>( TextureSample->Coordinates.Expression );
+						Cast<UMaterialExpressionTextureCoordinate>(TextureSample->Coordinates.Expression);
 
 					UMaterialExpressionLandscapeLayerCoords* TerrainTextureCoordinate =
-						Cast<UMaterialExpressionLandscapeLayerCoords>( TextureSample->Coordinates.Expression );
+						Cast<UMaterialExpressionLandscapeLayerCoords>(TextureSample->Coordinates.Expression);
 
 					if (TextureCoordinate || TerrainTextureCoordinate)
 					{
@@ -2479,7 +2339,7 @@ void ULandscapeComponent::GetStreamingTextureInfo(TArray<FStreamingTexturePrimit
 							FStreamingTexturePrimitiveInfo& StreamingTexture = OutStreamingTextures[i];
 							if (StreamingTexture.Texture == TextureSample->Texture)
 							{
-								if ( TextureCoordinate )
+								if (TextureCoordinate)
 								{
 									StreamingTexture.TexelFactor = TexelFactor * FPlatformMath::Max(TextureCoordinate->UTiling, TextureCoordinate->VTiling);
 								}
@@ -2507,13 +2367,13 @@ void ULandscapeComponent::GetStreamingTextureInfo(TArray<FStreamingTexturePrimit
 			{
 				float LightmapFactorX = TexelFactor / Scale.X;
 				float LightmapFactorY = TexelFactor / Scale.Y;
-				FStreamingTexturePrimitiveInfo& StreamingTexture = *new(OutStreamingTextures) FStreamingTexturePrimitiveInfo;
-				StreamingTexture.Bounds		 = BoundingSphere;
+				FStreamingTexturePrimitiveInfo& StreamingTexture = *new(OutStreamingTextures)FStreamingTexturePrimitiveInfo;
+				StreamingTexture.Bounds = BoundingSphere;
 				StreamingTexture.TexelFactor = FMath::Max(LightmapFactorX, LightmapFactorY);
-				StreamingTexture.Texture	 = Lightmap->GetTexture(LightmapIndex);
+				StreamingTexture.Texture = Lightmap->GetTexture(LightmapIndex);
 			}
 		}
-		
+
 		// Shadowmap
 		FShadowMap2D* Shadowmap = ShadowMap ? ShadowMap->GetShadowMap2D() : nullptr;
 		if (Shadowmap && Shadowmap->IsValid())
@@ -2521,20 +2381,20 @@ void ULandscapeComponent::GetStreamingTextureInfo(TArray<FStreamingTexturePrimit
 			const FVector2D& Scale = Shadowmap->GetCoordinateScale();
 			if (Scale.X > SMALL_NUMBER && Scale.Y > SMALL_NUMBER)
 			{
-				float ShadowmapFactorX		 = TexelFactor / Scale.X;
-				float ShadowmapFactorY		 = TexelFactor / Scale.Y;
-				FStreamingTexturePrimitiveInfo& StreamingTexture = *new(OutStreamingTextures) FStreamingTexturePrimitiveInfo;
-				StreamingTexture.Bounds		 = BoundingSphere;
+				float ShadowmapFactorX = TexelFactor / Scale.X;
+				float ShadowmapFactorY = TexelFactor / Scale.Y;
+				FStreamingTexturePrimitiveInfo& StreamingTexture = *new(OutStreamingTextures)FStreamingTexturePrimitiveInfo;
+				StreamingTexture.Bounds = BoundingSphere;
 				StreamingTexture.TexelFactor = FMath::Max(ShadowmapFactorX, ShadowmapFactorY);
-				StreamingTexture.Texture	 = Shadowmap->GetTexture();
+				StreamingTexture.Texture = Shadowmap->GetTexture();
 			}
 		}
 	}
 
 	// Weightmap
-	for(int32 TextureIndex = 0;TextureIndex < WeightmapTextures.Num();TextureIndex++)
+	for (int32 TextureIndex = 0; TextureIndex < WeightmapTextures.Num(); TextureIndex++)
 	{
-		FStreamingTexturePrimitiveInfo& StreamingWeightmap = *new(OutStreamingTextures) FStreamingTexturePrimitiveInfo;
+		FStreamingTexturePrimitiveInfo& StreamingWeightmap = *new(OutStreamingTextures)FStreamingTexturePrimitiveInfo;
 		StreamingWeightmap.Bounds = BoundingSphere;
 		StreamingWeightmap.TexelFactor = TexelFactor;
 		StreamingWeightmap.Texture = WeightmapTextures[TextureIndex];
@@ -2542,17 +2402,17 @@ void ULandscapeComponent::GetStreamingTextureInfo(TArray<FStreamingTexturePrimit
 
 	// Heightmap
 	{
-		float HeightmapTexelFactor = TexelFactor * ((HeightmapTexture == NULL) ? 1.0f : (HeightmapTexture->GetSizeY()/(ComponentSizeQuads+1)));
-		FStreamingTexturePrimitiveInfo& StreamingHeightmap = *new(OutStreamingTextures) FStreamingTexturePrimitiveInfo;
+		float HeightmapTexelFactor = TexelFactor * ((HeightmapTexture == nullptr) ? 1.0f : (HeightmapTexture->GetSizeY() / (ComponentSizeQuads + 1)));
+		FStreamingTexturePrimitiveInfo& StreamingHeightmap = *new(OutStreamingTextures)FStreamingTexturePrimitiveInfo;
 		StreamingHeightmap.Bounds = BoundingSphere;
-		StreamingHeightmap.TexelFactor = ForcedLOD >= 0 ? -13+ForcedLOD : HeightmapTexelFactor; // Minus Value indicate ForcedLOD, 13 for 8k texture
+		StreamingHeightmap.TexelFactor = ForcedLOD >= 0 ? -13 + ForcedLOD : HeightmapTexelFactor; // Minus Value indicate ForcedLOD, 13 for 8k texture
 		StreamingHeightmap.Texture = HeightmapTexture;
 	}
 
 	// XYOffset
 	if (XYOffsetmapTexture)
 	{
-		FStreamingTexturePrimitiveInfo& StreamingXYOffset = *new(OutStreamingTextures) FStreamingTexturePrimitiveInfo;
+		FStreamingTexturePrimitiveInfo& StreamingXYOffset = *new(OutStreamingTextures)FStreamingTexturePrimitiveInfo;
 		StreamingXYOffset.Bounds = BoundingSphere;
 		StreamingXYOffset.TexelFactor = TexelFactor;
 		StreamingXYOffset.Texture = XYOffsetmapTexture;
@@ -2561,7 +2421,7 @@ void ULandscapeComponent::GetStreamingTextureInfo(TArray<FStreamingTexturePrimit
 #if WITH_EDITOR
 	if (GIsEditor && EditToolRenderData && EditToolRenderData->DataTexture)
 	{
-		FStreamingTexturePrimitiveInfo& StreamingDatamap = *new(OutStreamingTextures) FStreamingTexturePrimitiveInfo;
+		FStreamingTexturePrimitiveInfo& StreamingDatamap = *new(OutStreamingTextures)FStreamingTexturePrimitiveInfo;
 		StreamingDatamap.Bounds = BoundingSphere;
 		StreamingDatamap.TexelFactor = TexelFactor;
 		StreamingDatamap.Texture = EditToolRenderData->DataTexture;
@@ -2575,15 +2435,15 @@ void ALandscapeProxy::ChangeLODDistanceFactor(float InLODDistanceFactor)
 	float LODFactor;
 	switch (LODFalloff)
 	{
-		case ELandscapeLODFalloff::SquareRoot:
-			LODFactor = FMath::Square(FMath::Min(LANDSCAPE_LOD_SQUARE_ROOT_FACTOR * LODDistanceFactor, MAX_LANDSCAPE_LOD_DISTANCE_FACTOR));
-			break;
-		default:
-		case ELandscapeLODFalloff::Linear:
-			LODFactor = LODDistanceFactor;
-			break;
+	case ELandscapeLODFalloff::SquareRoot:
+		LODFactor = FMath::Square(FMath::Min(LANDSCAPE_LOD_SQUARE_ROOT_FACTOR * LODDistanceFactor, MAX_LANDSCAPE_LOD_DISTANCE_FACTOR));
+		break;
+	default:
+	case ELandscapeLODFalloff::Linear:
+		LODFactor = LODDistanceFactor;
+		break;
 	}
-	
+
 	if (LandscapeComponents.Num())
 	{
 		int32 CompNum = LandscapeComponents.Num();
@@ -2592,25 +2452,24 @@ void ALandscapeProxy::ChangeLODDistanceFactor(float InLODDistanceFactor)
 		{
 			Proxies[Idx] = (FLandscapeComponentSceneProxy*)(LandscapeComponents[Idx]->SceneProxy);
 		}
-		
+
 		ENQUEUE_UNIQUE_RENDER_COMMAND_THREEPARAMETER(
 			LandscapeChangeLODDistanceFactorCommand,
 			FLandscapeComponentSceneProxy**, Proxies, Proxies,
 			int32, CompNum, CompNum,
-			FVector2D, InLODDistanceFactors, FVector2D(FMath::Sqrt(2.f * FMath::Square((float)SubsectionSizeQuads)) * LANDSCAPE_LOD_DISTANCE_FACTOR / LODFactor, FMath::Min(LODFactor, 3.f) * 0.33f),
-		{
-			for (int32 Idx = 0; Idx < CompNum; ++Idx)
+			float, InLODDistanceFactor, FMath::Sqrt(2.f * FMath::Square((float)SubsectionSizeQuads)) * LANDSCAPE_LOD_DISTANCE_FACTOR / LODFactor,
 			{
-				Proxies[Idx]->ChangeLODDistanceFactor_RenderThread(InLODDistanceFactors);
+				for (int32 Idx = 0; Idx < CompNum; ++Idx)
+				{
+					Proxies[Idx]->ChangeLODDistanceFactor_RenderThread(InLODDistanceFactor);
+				}
+				delete[] Proxies;
 			}
-			delete[] Proxies;
-		}
 		);
 	}
 };
 
-void FLandscapeComponentSceneProxy::ChangeLODDistanceFactor_RenderThread(FVector2D InLODDistanceFactors)
+void FLandscapeComponentSceneProxy::ChangeLODDistanceFactor_RenderThread(float InLODDistanceFactor)
 {
-	LODDistance = InLODDistanceFactors.X;
-	PrecomputedLODFactor = InLODDistanceFactors.Y;
+	LODDistance = InLODDistanceFactor;
 }
diff --git a/Engine/Source/Runtime/Landscape/Private/LandscapeRenderMobile.cpp b/Engine/Source/Runtime/Landscape/Private/LandscapeRenderMobile.cpp
--- a/Engine/Source/Runtime/Landscape/Private/LandscapeRenderMobile.cpp
+++ b/Engine/Source/Runtime/Landscape/Private/LandscapeRenderMobile.cpp
@@ -77,8 +77,8 @@ public:
 		if( LodBiasParameter.IsBound() )
 		{
 			FVector4 LodBias(
-				SceneProxy->PrecomputedLODFactor,
-				1.f / (1.f - SceneProxy->PrecomputedLODFactor),
+				0.0f, // unused
+				0.0f, // unused
 				CameraLocalPos3D.X + SceneProxy->SectionBase.X,
 				CameraLocalPos3D.Y + SceneProxy->SectionBase.Y 
 				);
@@ -88,7 +88,6 @@ public:
 		// Calculate LOD params
 		FVector4 fCurrentLODs;
 		FVector4 CurrentNeighborLODs[4];
-		float DistLOD = 0.f;
 
 		if( BatchElementParams->SubX == -1 )
 		{
@@ -97,14 +96,14 @@ public:
 				for( int32 SubX = 0; SubX < SceneProxy->NumSubsections; SubX++ )
 				{
 					int32 SubIndex = SubX + 2 * SubY;
-					SceneProxy->CalcLODParamsForSubsection(View, CameraLocalPos, SubX, SubY, fCurrentLODs[SubIndex], CurrentNeighborLODs[SubIndex], DistLOD);
+					SceneProxy->CalcLODParamsForSubsection(View, CameraLocalPos, SubX, SubY, BatchElementParams->CurrentLOD, fCurrentLODs[SubIndex], CurrentNeighborLODs[SubIndex]);
 				}
 			}
 		}
 		else
 		{
 			int32 SubIndex = BatchElementParams->SubX + 2 * BatchElementParams->SubY;
-			SceneProxy->CalcLODParamsForSubsection(View, CameraLocalPos, BatchElementParams->SubX, BatchElementParams->SubY, fCurrentLODs[SubIndex], CurrentNeighborLODs[SubIndex], DistLOD);
+			SceneProxy->CalcLODParamsForSubsection(View, CameraLocalPos, BatchElementParams->SubX, BatchElementParams->SubY, BatchElementParams->CurrentLOD, fCurrentLODs[SubIndex], CurrentNeighborLODs[SubIndex]);
 		}
 
 		if( SectionLodsParameter.IsBound() )
@@ -120,10 +119,9 @@ public:
 		if( LodValuesParameter.IsBound() )
 		{
 			FVector4 LodValues(
-				DistLOD,
-				// convert current LOD coordinates into highest LOD coordinates
-				(float)SceneProxy->SubsectionSizeQuads / (float)(((SceneProxy->SubsectionSizeVerts) >> BatchElementParams->CurrentLOD)-1),
-				(float)SceneProxy->SubsectionSizeQuads, 
+				0.0f, // this is the mesh's LOD, ES2 always uses the LOD0 mesh
+				0.0f, // unused
+				(float)SceneProxy->SubsectionSizeQuads,
 				1.f / (float)SceneProxy->SubsectionSizeQuads );
 
 			SetShaderValue(RHICmdList, VertexShader->GetVertexShader(),LodValuesParameter,LodValues);
diff --git a/Engine/Source/Runtime/Landscape/Public/LandscapeRender.h b/Engine/Source/Runtime/Landscape/Public/LandscapeRender.h
--- a/Engine/Source/Runtime/Landscape/Public/LandscapeRender.h
+++ b/Engine/Source/Runtime/Landscape/Public/LandscapeRender.h
@@ -4,8 +4,7 @@
 LandscapeRender.h: New terrain rendering
 =============================================================================*/
 
-#ifndef _LANDSCAPERENDER_H
-#define _LANDSCAPERENDER_H
+#pragma once
 
 #include "UniformBuffer.h"
 
@@ -26,6 +25,9 @@ LandscapeRender.h: New terrain rendering
 #define LANDSCAPE_LOD_LEVELS 8
 #define LANDSCAPE_MAX_SUBSECTION_NUM 2
 
+// Forward declarations
+class FLandscapeComponentSceneProxy;
+
 #if WITH_EDITOR
 namespace ELandscapeViewMode
 {
@@ -48,13 +50,13 @@ namespace ELandscapeEditRenderMode
 {
 	enum Type
 	{
-		None				= 0x0,
-		Gizmo				= 0x1,
-		SelectRegion		= 0x2,
-		SelectComponent		= 0x4,
+		None = 0x0,
+		Gizmo = 0x1,
+		SelectRegion = 0x2,
+		SelectComponent = 0x4,
 		Select = SelectRegion | SelectComponent,
-		Mask				= 0x8, 
-		InvertedMask		= 0x10, // Should not be overlapped with other bits 
+		Mask = 0x8,
+		InvertedMask = 0x10, // Should not be overlapped with other bits 
 		BitMaskForMask = Mask | InvertedMask,
 
 	};
@@ -72,15 +74,15 @@ LANDSCAPE_API extern UTexture2D* GLandscapeBlackTexture;
 
 
 /** The uniform shader parameters for a landscape draw call. */
-BEGIN_UNIFORM_BUFFER_STRUCT(FLandscapeUniformShaderParameters,LANDSCAPE_API)
-	/** vertex shader parameters */
-	DECLARE_UNIFORM_BUFFER_STRUCT_MEMBER(FVector4,HeightmapUVScaleBias)
-	DECLARE_UNIFORM_BUFFER_STRUCT_MEMBER(FVector4,WeightmapUVScaleBias)
-	DECLARE_UNIFORM_BUFFER_STRUCT_MEMBER(FVector4,LandscapeLightmapScaleBias)
-	DECLARE_UNIFORM_BUFFER_STRUCT_MEMBER(FVector4,SubsectionSizeVertsLayerUVPan)
-	DECLARE_UNIFORM_BUFFER_STRUCT_MEMBER(FVector4,SubsectionOffsetParams)
-	DECLARE_UNIFORM_BUFFER_STRUCT_MEMBER(FVector4,LightmapSubsectionOffsetParams)
-	DECLARE_UNIFORM_BUFFER_STRUCT_MEMBER(FMatrix,LocalToWorldNoScaling)
+BEGIN_UNIFORM_BUFFER_STRUCT(FLandscapeUniformShaderParameters, LANDSCAPE_API)
+/** vertex shader parameters */
+DECLARE_UNIFORM_BUFFER_STRUCT_MEMBER(FVector4, HeightmapUVScaleBias)
+DECLARE_UNIFORM_BUFFER_STRUCT_MEMBER(FVector4, WeightmapUVScaleBias)
+DECLARE_UNIFORM_BUFFER_STRUCT_MEMBER(FVector4, LandscapeLightmapScaleBias)
+DECLARE_UNIFORM_BUFFER_STRUCT_MEMBER(FVector4, SubsectionSizeVertsLayerUVPan)
+DECLARE_UNIFORM_BUFFER_STRUCT_MEMBER(FVector4, SubsectionOffsetParams)
+DECLARE_UNIFORM_BUFFER_STRUCT_MEMBER(FVector4, LightmapSubsectionOffsetParams)
+DECLARE_UNIFORM_BUFFER_STRUCT_MEMBER(FMatrix, LocalToWorldNoScaling)
 END_UNIFORM_BUFFER_STRUCT(FLandscapeUniformShaderParameters)
 
 /* Data needed for the landscape vertex factory to set the render state for an individual batch element */
@@ -90,7 +92,7 @@ struct FLandscapeBatchElementParams
 	const FMatrix* LocalToWorldNoScalingPtr;
 
 	// LOD calculation-related params
-	const class FLandscapeComponentSceneProxy* SceneProxy;
+	const FLandscapeComponentSceneProxy* SceneProxy;
 	int32 SubX;
 	int32 SubY;
 	int32	CurrentLOD;
@@ -121,8 +123,8 @@ public:
 	/**
 	* Set any shader data specific to this vertex factory
 	*/
-	virtual void SetMesh(FRHICommandList& RHICmdList, FShader* PixelShader,const class FVertexFactory* VertexFactory,const class FSceneView& View,const struct FMeshBatchElement& BatchElement,uint32 DataFlags) const override;
-	
+	virtual void SetMesh(FRHICommandList& RHICmdList, FShader* PixelShader, const FVertexFactory* VertexFactory, const FSceneView& View, const FMeshBatchElement& BatchElement, uint32 DataFlags) const override;
+
 	virtual uint32 GetSize() const override
 	{
 		return sizeof(*this);
@@ -141,12 +143,11 @@ class FLandscapeVertexFactory : public FVertexFactory
 
 public:
 
-	FLandscapeVertexFactory()
-	{
-	}
+	FLandscapeVertexFactory() = default;
 
 	virtual ~FLandscapeVertexFactory()
 	{
+		// can only be destroyed from the render thread
 		ReleaseResource();
 	}
 
@@ -159,9 +160,9 @@ public:
 	};
 
 	/**
-	* Should we cache the material's shadertype on this platform with this vertex factory? 
+	* Should we cache the material's shadertype on this platform with this vertex factory?
 	*/
-	static bool ShouldCache(EShaderPlatform Platform, const class FMaterial* Material, const class FShaderType* ShaderType)
+	static bool ShouldCache(EShaderPlatform Platform, const FMaterial* Material, const FShaderType* ShaderType)
 	{
 		// only compile landscape materials for landscape vertex factory
 		// The special engine materials must be compiled for the landscape vertex factory because they are used with it for wireframe, etc.
@@ -194,10 +195,10 @@ public:
 		UpdateRHI();
 	}
 
-	virtual uint64 GetStaticBatchElementVisibility( const class FSceneView& View, const struct FMeshBatch* Batch ) const override;
+	virtual uint64 GetStaticBatchElementVisibility(const FSceneView& View, const FMeshBatch* Batch) const override;
 
 	/** stream component data bound to this vertex factory */
-	DataType Data;  
+	DataType Data;
 };
 
 
@@ -207,21 +208,9 @@ class FLandscapeXYOffsetVertexFactory : public FLandscapeVertexFactory
 	DECLARE_VERTEX_FACTORY_TYPE(FLandscapeXYOffsetVertexFactory);
 
 public:
+	FLandscapeXYOffsetVertexFactory() = default;
+	virtual ~FLandscapeXYOffsetVertexFactory() = default;
 
-	FLandscapeXYOffsetVertexFactory()
-	{
-	}
-
-	virtual ~FLandscapeXYOffsetVertexFactory()
-	{
-		ReleaseResource();
-	}
-
-	static FVertexFactoryShaderParameters* ConstructShaderParameters(EShaderFrequency ShaderFrequency);
-
-	/**
-	* Can be overridden by FVertexFactory subclasses to modify their compile environment just before compilation occurs.
-	*/
 	static void ModifyCompilationEnvironment(EShaderPlatform Platform, const FMaterial* Material, FShaderCompilerEnvironment& OutEnvironment);
 };
 
@@ -241,11 +230,11 @@ class FLandscapeVertexBuffer : public FVertexBuffer
 	int32 SubsectionSizeVerts;
 	int32 NumSubsections;
 public:
-	
+
 	/** Constructor. */
 	FLandscapeVertexBuffer(int32 InSubsectionSizeVerts, int32 InNumSubsections)
-	: 	SubsectionSizeVerts(InSubsectionSizeVerts)
-	,	NumSubsections(InNumSubsections)
+		: SubsectionSizeVerts(InSubsectionSizeVerts)
+		, NumSubsections(InNumSubsections)
 	{
 		InitResource();
 	}
@@ -256,8 +245,8 @@ public:
 		ReleaseResource();
 	}
 
-	/** 
-	* Initialize the RHI for this rendering resource 
+	/**
+	* Initialize the RHI for this rendering resource
 	*/
 	virtual void InitRHI() override;
 };
@@ -323,14 +312,14 @@ struct FLandscapeEditToolRenderData
 	};
 
 	FLandscapeEditToolRenderData(ULandscapeComponent* InComponent)
-		:	ToolMaterial(NULL),
-			GizmoMaterial(NULL),
-			LandscapeComponent(InComponent),
-			SelectedType(ST_NONE),
-			DebugChannelR(INDEX_NONE),
-			DebugChannelG(INDEX_NONE),
-			DebugChannelB(INDEX_NONE),
-			DataTexture(NULL)
+		: ToolMaterial(NULL),
+		GizmoMaterial(NULL),
+		LandscapeComponent(InComponent),
+		SelectedType(ST_NONE),
+		DebugChannelR(INDEX_NONE),
+		DebugChannelG(INDEX_NONE),
+		DebugChannelB(INDEX_NONE),
+		DataTexture(NULL)
 	{}
 
 	// Material used to render the tool.
@@ -351,8 +340,8 @@ struct FLandscapeEditToolRenderData
 #endif
 
 	// Game thread update
-	LANDSCAPE_API void Update( UMaterialInterface* InNewToolMaterial );
-	LANDSCAPE_API void UpdateGizmo( UMaterialInterface* InNewGizmoMaterial );
+	LANDSCAPE_API void Update(UMaterialInterface* InNewToolMaterial);
+	LANDSCAPE_API void UpdateGizmo(UMaterialInterface* InNewGizmoMaterial);
 	// Allows game thread to queue the deletion by the render thread
 	void Cleanup();
 };
@@ -364,7 +353,7 @@ class FLandscapeComponentSceneProxy : public FPrimitiveSceneProxy
 {
 	friend class FLandscapeSharedBuffers;
 
-	class FLandscapeLCI : public FLightCacheInterface
+	class FLandscapeLCI final : public FLightCacheInterface
 	{
 	public:
 		/** Initialization constructor. */
@@ -375,12 +364,8 @@ class FLandscapeComponentSceneProxy : public FPrimitiveSceneProxy
 			IrrelevantLights = InComponent->IrrelevantLights;
 		}
 
-		virtual ~FLandscapeLCI()
-		{
-		}
-
 		// FLightCacheInterface
-		virtual FLightInteraction GetInteraction(const class FLightSceneProxy* LightSceneProxy) const;
+		virtual FLightInteraction GetInteraction(const FLightSceneProxy* LightSceneProxy) const;
 
 		virtual FLightMapInteraction GetLightMapInteraction(ERHIFeatureLevel::Type InFeatureLevel) const
 		{
@@ -403,15 +388,15 @@ class FLandscapeComponentSceneProxy : public FPrimitiveSceneProxy
 	};
 
 protected:
-	int32						MaxLOD;
-	int32						ComponentSizeQuads;	// Size of component in quads
-	int32						ComponentSizeVerts;
-	int32						NumSubsections;
-	int32						SubsectionSizeQuads;
-	int32						SubsectionSizeVerts;
-	FIntPoint					SectionBase;
+	int8						MaxLOD;
+	int8						NumSubsections;
+	int16						SubsectionSizeQuads;
+	int16						SubsectionSizeVerts;
+	int16						ComponentSizeQuads;	// Size of component in quads
+	int16						ComponentSizeVerts;
+	uint8						StaticLightingLOD;
 	float						StaticLightingResolution;
-	uint32						StaticLightingLOD;
+	FIntPoint					SectionBase;
 	FMatrix						LocalToWorldNoScaling;
 
 	// Storage for static draw list batch params
@@ -419,16 +404,15 @@ protected:
 
 	// Precomputed values
 	float					LODDistance;
-	float					PrecomputedLODFactor;
 	float					DistDiff;
 
 	FVector4 WeightmapScaleBias;
 	float WeightmapSubsectionOffset;
-	TArray<class UTexture2D*> WeightmapTextures;
-	int32 NumWeightmapLayerAllocations;
+	TArray<UTexture2D*> WeightmapTextures;
+	int8 NumWeightmapLayerAllocations;
 
 	UTexture2D* NormalmapTexture; // PC : Heightmap, Mobile : Weightmap
-	
+
 	UTexture2D* HeightmapTexture; // PC : Heightmap, Mobile : Weightmap
 	FVector4 HeightmapScaleBias;
 	float HeightmapSubsectionOffsetU;
@@ -440,7 +424,7 @@ protected:
 	uint32						SharedBuffersKey;
 	FLandscapeSharedBuffers*	SharedBuffers;
 	FLandscapeVertexFactory*	VertexFactory;
-	
+
 	UMaterialInterface* MaterialInterface;
 	FMaterialRelevance MaterialRelevance;
 
@@ -452,17 +436,16 @@ protected:
 	FLandscapeEditToolRenderData* EditToolRenderData;
 
 	// FLightCacheInterface
-	FLandscapeLCI* ComponentLightInfo;
+	TUniquePtr<FLandscapeLCI> ComponentLightInfo;
 
 	const ULandscapeComponent* LandscapeComponent;
 
-	FVector2D				NeighborPosition[LANDSCAPE_NEIGHBOR_NUM];
-	int32					ForcedLOD;
-	int32					LODBias;
-	uint8					ForcedNeighborLOD[LANDSCAPE_NEIGHBOR_NUM];
-	uint8					NeighborLODBias[LANDSCAPE_NEIGHBOR_NUM];
+	int8					ForcedLOD;
+	int8					LODBias;
+	int8					ForcedNeighborLOD[LANDSCAPE_NEIGHBOR_NUM];
+	int8					NeighborLODBias[LANDSCAPE_NEIGHBOR_NUM];
 
-	enum ELandscapeLODFalloff::Type LODFalloff;
+	ELandscapeLODFalloff::Type LODFalloff;
 
 	TUniformBuffer<FLandscapeUniformShaderParameters> LandscapeUniformShaderParameters;
 
@@ -478,7 +461,7 @@ public:
 	// FPrimitiveSceneProxy interface.
 	virtual void DrawStaticElements(FStaticPrimitiveDrawInterface* PDI) override;
 	virtual void GetDynamicMeshElements(const TArray<const FSceneView*>& Views, const FSceneViewFamily& ViewFamily, uint32 VisibilityMap, FMeshElementCollector& Collector) const override;
-	virtual uint32 GetMemoryFootprint( void ) const override { return( sizeof( *this ) + GetAllocatedSize() ); }
+	virtual uint32 GetMemoryFootprint() const override { return(sizeof(*this) + GetAllocatedSize()); }
 	virtual FPrimitiveViewRelevance GetViewRelevance(const FSceneView* View) override;
 	virtual bool CanBeOccluded() const override;
 	virtual void GetLightRelevance(const FLightSceneProxy* LightSceneProxy, bool& bDynamic, bool& bRelevant, bool& bLightMapped, bool& bShadowMapped) const override;
@@ -494,13 +477,13 @@ public:
 	friend class FLandscapeVertexFactoryMobilePixelShaderParameters;
 
 	// FLandscapeComponentSceneProxy interface.
-	int32 CalcLODForSubsection(const class FSceneView& View, int32 SubX, int32 SubY, const FVector2D& CameraLocalPos) const;
-	int32 CalcLODForSubsectionNoForced(const class FSceneView& View, int32 SubX, int32 SubY, const FVector2D& CameraLocalPos) const;
-	void CalcLODParamsForSubsection(const class FSceneView& View, const FVector2D& CameraLocalPos, int32 SubX, int32 SubY, float& OutfLOD, FVector4& OutNeighborLODs, float& OutDistLOD) const;
-	uint64 GetStaticBatchElementVisibility( const class FSceneView& View, const struct FMeshBatch* Batch ) const;
+	float CalcDesiredLOD(const FSceneView& View, const FVector2D& CameraLocalPos, int32 SubX, int32 SubY) const;
+	int32 CalcLODForSubsection(const FSceneView& View, int32 SubX, int32 SubY, const FVector2D& CameraLocalPos) const;
+	void CalcLODParamsForSubsection(const FSceneView& View, const FVector2D& CameraLocalPos, int32 SubX, int32 SubY, int32 BatchLOD, float& OutfLOD, FVector4& OutNeighborLODs) const;
+	uint64 GetStaticBatchElementVisibility(const FSceneView& View, const FMeshBatch* Batch) const;
 
 	// FLandcapeSceneProxy
-	void ChangeLODDistanceFactor_RenderThread(FVector2D InLODDistanceFactors);
+	void ChangeLODDistanceFactor_RenderThread(float InLODDistanceFactor);
 
 	virtual void GetHeightfieldRepresentation(UTexture2D*& OutHeightmapTexture, FVector4& OutHeightfieldScaleBias, FVector4& OutMinMaxUV) override
 	{
@@ -525,9 +508,9 @@ public:
 	const FLinearColor B;
 
 	/** Initialization constructor. */
-	FLandscapeDebugMaterialRenderProxy(const FMaterialRenderProxy* InParent, const UTexture2D* TexR, const UTexture2D* TexG, const UTexture2D* TexB, 
-		const FLinearColor& InR, const FLinearColor& InG, const FLinearColor& InB ):
-	Parent(InParent),
+	FLandscapeDebugMaterialRenderProxy(const FMaterialRenderProxy* InParent, const UTexture2D* TexR, const UTexture2D* TexG, const UTexture2D* TexB,
+		const FLinearColor& InR, const FLinearColor& InG, const FLinearColor& InB) :
+		Parent(InParent),
 		RedTexture(TexR),
 		GreenTexture(TexG),
 		BlueTexture(TexB),
@@ -537,23 +520,23 @@ public:
 	{}
 
 	// FMaterialRenderProxy interface.
-	virtual const class FMaterial* GetMaterial(ERHIFeatureLevel::Type FeatureLevel) const
+	virtual const FMaterial* GetMaterial(ERHIFeatureLevel::Type FeatureLevel) const
 	{
 		return Parent->GetMaterial(FeatureLevel);
 	}
 	virtual bool GetVectorValue(const FName ParameterName, FLinearColor* OutValue, const FMaterialRenderContext& Context) const
 	{
-		if (ParameterName == FName(TEXT("Landscape_RedMask") ) )
+		if (ParameterName == FName(TEXT("Landscape_RedMask")))
 		{
 			*OutValue = R;
 			return true;
 		}
-		else if (ParameterName == FName(TEXT("Landscape_GreenMask") ) )
+		else if (ParameterName == FName(TEXT("Landscape_GreenMask")))
 		{
 			*OutValue = G;
 			return true;
 		}
-		else if (ParameterName == FName(TEXT("Landscape_BlueMask") ) )
+		else if (ParameterName == FName(TEXT("Landscape_BlueMask")))
 		{
 			*OutValue = B;
 			return true;
@@ -567,21 +550,21 @@ public:
 	{
 		return Parent->GetScalarValue(ParameterName, OutValue, Context);
 	}
-	virtual bool GetTextureValue(const FName ParameterName,const UTexture** OutValue, const FMaterialRenderContext& Context) const
+	virtual bool GetTextureValue(const FName ParameterName, const UTexture** OutValue, const FMaterialRenderContext& Context) const
 	{
 		// NOTE: These should be returning black textures when NULL. The material will
 		// use a white texture if they are.
-		if (ParameterName == FName(TEXT("Landscape_RedTexture")) )
+		if (ParameterName == FName(TEXT("Landscape_RedTexture")))
 		{
 			*OutValue = RedTexture;
 			return true;
 		}
-		else if (ParameterName == FName(TEXT("Landscape_GreenTexture"))  )
+		else if (ParameterName == FName(TEXT("Landscape_GreenTexture")))
 		{
 			*OutValue = GreenTexture;
 			return true;
 		}
-		else if (ParameterName == FName(TEXT("Landscape_BlueTexture"))  )
+		else if (ParameterName == FName(TEXT("Landscape_BlueTexture")))
 		{
 			*OutValue = BlueTexture;
 			return true;
@@ -600,19 +583,19 @@ public:
 	const UTexture2D* SelectTexture;
 
 	/** Initialization constructor. */
-	FLandscapeSelectMaterialRenderProxy(const FMaterialRenderProxy* InParent, const UTexture2D* InTexture):
-	Parent(InParent),
+	FLandscapeSelectMaterialRenderProxy(const FMaterialRenderProxy* InParent, const UTexture2D* InTexture) :
+		Parent(InParent),
 		SelectTexture(InTexture)
 	{}
 
 	// FMaterialRenderProxy interface.
-	virtual const class FMaterial* GetMaterial(ERHIFeatureLevel::Type FeatureLevel) const
+	virtual const FMaterial* GetMaterial(ERHIFeatureLevel::Type FeatureLevel) const
 	{
 		return Parent->GetMaterial(FeatureLevel);
 	}
 	virtual bool GetVectorValue(const FName ParameterName, FLinearColor* OutValue, const FMaterialRenderContext& Context) const
 	{
-		if (ParameterName == FName(TEXT("HighlightColor")) )
+		if (ParameterName == FName(TEXT("HighlightColor")))
 		{
 			*OutValue = FLinearColor(1.f, 0.5f, 0.5f);
 			return true;
@@ -626,9 +609,9 @@ public:
 	{
 		return Parent->GetScalarValue(ParameterName, OutValue, Context);
 	}
-	virtual bool GetTextureValue(const FName ParameterName,const UTexture** OutValue, const FMaterialRenderContext& Context) const
+	virtual bool GetTextureValue(const FName ParameterName, const UTexture** OutValue, const FMaterialRenderContext& Context) const
 	{
-		if (ParameterName == FName(TEXT("SelectedData")) )
+		if (ParameterName == FName(TEXT("SelectedData")))
 		{
 			*OutValue = SelectTexture;
 			return true;
@@ -648,14 +631,14 @@ public:
 	const bool bInverted;
 
 	/** Initialization constructor. */
-	FLandscapeMaskMaterialRenderProxy(const FMaterialRenderProxy* InParent, const UTexture2D* InTexture, const bool InbInverted):
-	Parent(InParent),
+	FLandscapeMaskMaterialRenderProxy(const FMaterialRenderProxy* InParent, const UTexture2D* InTexture, const bool InbInverted) :
+		Parent(InParent),
 		SelectTexture(InTexture),
 		bInverted(InbInverted)
 	{}
 
 	// FMaterialRenderProxy interface.
-	virtual const class FMaterial* GetMaterial(ERHIFeatureLevel::Type FeatureLevel) const
+	virtual const FMaterial* GetMaterial(ERHIFeatureLevel::Type FeatureLevel) const
 	{
 		return Parent->GetMaterial(FeatureLevel);
 	}
@@ -665,16 +648,16 @@ public:
 	}
 	virtual bool GetScalarValue(const FName ParameterName, float* OutValue, const FMaterialRenderContext& Context) const
 	{
-		if (ParameterName == FName(TEXT("bInverted")) )
+		if (ParameterName == FName(TEXT("bInverted")))
 		{
 			*OutValue = bInverted;
 			return true;
 		}
 		return Parent->GetScalarValue(ParameterName, OutValue, Context);
 	}
-	virtual bool GetTextureValue(const FName ParameterName,const UTexture** OutValue, const FMaterialRenderContext& Context) const
+	virtual bool GetTextureValue(const FName ParameterName, const UTexture** OutValue, const FMaterialRenderContext& Context) const
 	{
-		if (ParameterName == FName(TEXT("SelectedData")) )
+		if (ParameterName == FName(TEXT("SelectedData")))
 		{
 			*OutValue = SelectTexture;
 			return true;
@@ -686,7 +669,7 @@ public:
 	}
 };
 
-namespace 
+namespace
 {
 	// LightmapRes: Multiplier of lightmap size relative to landscape size
 	// X: (Output) PatchExpandCountX (at Lighting LOD)
@@ -703,13 +686,13 @@ namespace
 		int32 PixelPaddingX = GPixelFormats[PF_DXT1].BlockSizeX;
 		int32 PixelPaddingY = GPixelFormats[PF_DXT1].BlockSizeY;
 		int32 PatchExpandCountX = (LightMapRes >= 1.f) ? (PixelPaddingX) / LightMapRes : (PixelPaddingX);
-		int32 PatchExpandCountY = (LightMapRes >= 1.f )? (PixelPaddingY) / LightMapRes : (PixelPaddingY);
+		int32 PatchExpandCountY = (LightMapRes >= 1.f) ? (PixelPaddingY) / LightMapRes : (PixelPaddingY);
 
 		X = FMath::Max<int32>(1, PatchExpandCountX >> LightingLOD);
 		Y = FMath::Max<int32>(1, PatchExpandCountY >> LightingLOD);
 
-		DesiredSize = (LightMapRes >= 1.f) ? FMath::Min<int32>((int32)((ComponentSize + 1) * LightMapRes), 4096) : FMath::Min<int32>((int32)((LightmapSize) * LightMapRes), 4096);
-		int32 CurrentSize = (LightMapRes >= 1.f) ? FMath::Min<int32>((int32)((2*(X<<LightingLOD) + ComponentSize + 1) * LightMapRes), 4096) : FMath::Min<int32>((int32)((2*(X<<LightingLOD) + LightmapSize) * LightMapRes), 4096);
+		DesiredSize = (LightMapRes >= 1.f) ? FMath::Min<int32>((int32)((ComponentSize + 1) * LightMapRes), 4096) : FMath::Min<int32>((int32)((LightmapSize)* LightMapRes), 4096);
+		int32 CurrentSize = (LightMapRes >= 1.f) ? FMath::Min<int32>((int32)((2 * (X << LightingLOD) + ComponentSize + 1) * LightMapRes), 4096) : FMath::Min<int32>((int32)((2 * (X << LightingLOD) + LightmapSize) * LightMapRes), 4096);
 
 		// Find proper Lightmap Size
 		if (CurrentSize > DesiredSize)
@@ -719,11 +702,11 @@ namespace
 			while (DesiredSize > 0)
 			{
 				PriorSize = DesiredSize;
-				DesiredSize = DesiredSize & ~(DesiredSize & ~(DesiredSize-1));
+				DesiredSize = DesiredSize & ~(DesiredSize & ~(DesiredSize - 1));
 			}
 
 			DesiredSize = PriorSize << 1; // next bigger size
-			if ( CurrentSize * CurrentSize <= ((PriorSize * PriorSize) << 1)  )
+			if (CurrentSize * CurrentSize <= ((PriorSize * PriorSize) << 1))
 			{
 				DesiredSize = PriorSize;
 			}
@@ -734,6 +717,3 @@ namespace
 		return LightMapRatio;
 	}
 }
-
-
-#endif // _LANDSCAPERENDER_H
