Fix compile error - unsupported C++11 feature.

diff --git a/Engine/Source/Runtime/Landscape/Public/LandscapeRender.h b/Engine/Source/Runtime/Landscape/Public/LandscapeRender.h
--- a/Engine/Source/Runtime/Landscape/Public/LandscapeRender.h
+++ b/Engine/Source/Runtime/Landscape/Public/LandscapeRender.h
@@ -143,7 +143,7 @@ class FLandscapeVertexFactory : public FVertexFactory
 
 public:
 
-	FLandscapeVertexFactory() = default;
+	FLandscapeVertexFactory() {}
 
 	virtual ~FLandscapeVertexFactory()
 	{
@@ -208,8 +208,8 @@ class FLandscapeXYOffsetVertexFactory : public FLandscapeVertexFactory
 	DECLARE_VERTEX_FACTORY_TYPE(FLandscapeXYOffsetVertexFactory);
 
 public:
-	FLandscapeXYOffsetVertexFactory() = default;
-	virtual ~FLandscapeXYOffsetVertexFactory() = default;
+	FLandscapeXYOffsetVertexFactory() {}
+	virtual ~FLandscapeXYOffsetVertexFactory() {}
 
 	static void ModifyCompilationEnvironment(EShaderPlatform Platform, const FMaterial* Material, FShaderCompilerEnvironment& OutEnvironment);
 };
