Fixed: Texture streaming TexelFactor calculation for mirrored landscapes

diff --git a/Engine/Source/Runtime/Landscape/Private/LandscapeRender.cpp b/Engine/Source/Runtime/Landscape/Private/LandscapeRender.cpp
--- a/Engine/Source/Runtime/Landscape/Private/LandscapeRender.cpp
+++ b/Engine/Source/Runtime/Landscape/Private/LandscapeRender.cpp
@@ -2321,7 +2321,7 @@ void ULandscapeComponent::GetStreamingTextureInfo(TArray<FStreamingTexturePrimit
 	{
 		StreamingDistanceMultiplier = Proxy->StreamingDistanceMultiplier;
 	}
-	const float TexelFactor = 0.75f * StreamingDistanceMultiplier * ComponentSizeQuads * (Proxy->GetRootComponent()->RelativeScale3D.X);
+	const float TexelFactor = 0.75f * StreamingDistanceMultiplier * ComponentSizeQuads * FMath::Abs(Proxy->GetRootComponent()->RelativeScale3D.X);
 
 	// Normal usage...
 	// Enumerate the textures used by the material.
@@ -2426,8 +2426,9 @@ void ULandscapeComponent::GetStreamingTextureInfo(TArray<FStreamingTexturePrimit
 	}
 
 	// Heightmap
+	if (HeightmapTexture)
 	{
-		float HeightmapTexelFactor = TexelFactor * ((HeightmapTexture == nullptr) ? 1.0f : (HeightmapTexture->GetSizeY() / (ComponentSizeQuads + 1)));
+		float HeightmapTexelFactor = TexelFactor * (HeightmapTexture->GetSizeY() / (ComponentSizeQuads + 1));
 		FStreamingTexturePrimitiveInfo& StreamingHeightmap = *new(OutStreamingTextures)FStreamingTexturePrimitiveInfo;
 		StreamingHeightmap.Bounds = BoundingSphere;
 		StreamingHeightmap.TexelFactor = ForcedLOD >= 0 ? -13 + ForcedLOD : HeightmapTexelFactor; // Minus Value indicate ForcedLOD, 13 for 8k texture
