Fixed: Navmesh distorts after mirroring landscape
UE-6342

diff --git a/Engine/Source/Runtime/Landscape/Classes/LandscapeProxy.h b/Engine/Source/Runtime/Landscape/Classes/LandscapeProxy.h
--- a/Engine/Source/Runtime/Landscape/Classes/LandscapeProxy.h
+++ b/Engine/Source/Runtime/Landscape/Classes/LandscapeProxy.h
@@ -345,6 +345,7 @@ public:
 #if WITH_EDITOR
 	virtual void Destroyed() override;
 	virtual void EditorApplyScale(const FVector& DeltaScale, const FVector* PivotLocation, bool bAltDown, bool bShiftDown, bool bCtrlDown) override;
+	virtual void EditorApplyMirror(const FVector& MirrorScale, const FVector& PivotLocation) override;
 	virtual void PostEditMove(bool bFinished) override;
 	virtual bool ShouldImport(FString* ActorPropString, bool IsMovingLevel) override;
 	virtual bool ShouldExport() override;
diff --git a/Engine/Source/Runtime/Landscape/Private/LandscapeEdit.cpp b/Engine/Source/Runtime/Landscape/Private/LandscapeEdit.cpp
--- a/Engine/Source/Runtime/Landscape/Private/LandscapeEdit.cpp
+++ b/Engine/Source/Runtime/Landscape/Private/LandscapeEdit.cpp
@@ -3320,6 +3320,29 @@ void ALandscapeProxy::EditorApplyScale(const FVector& DeltaScale, const FVector*
 	}
 
 	Super::EditorApplyScale(ModifiedScale, PivotLocation, bAltDown, bShiftDown, bCtrlDown);
+
+	// We need to regenerate collision objects, they depend on scale value 
+	for (ULandscapeHeightfieldCollisionComponent* Comp : CollisionComponents)
+	{
+		if (Comp)
+		{
+			Comp->RecreateCollision(false);
+		}
+	}
+}
+
+void ALandscapeProxy::EditorApplyMirror(const FVector& MirrorScale, const FVector& PivotLocation)
+{
+	Super::EditorApplyMirror(MirrorScale, PivotLocation);
+
+	// We need to regenerate collision objects, they depend on scale value 
+	for (ULandscapeHeightfieldCollisionComponent* Comp : CollisionComponents)
+	{
+		if (Comp)
+		{
+			Comp->RecreateCollision(false);
+		}
+	}
 }
 
 void ALandscapeProxy::PostEditMove(bool bFinished)
