Fix crash when moving at least two landscapes to the another level

diff --git a/Engine/Source/Runtime/Landscape/Private/LandscapeEdit.cpp b/Engine/Source/Runtime/Landscape/Private/LandscapeEdit.cpp
--- a/Engine/Source/Runtime/Landscape/Private/LandscapeEdit.cpp
+++ b/Engine/Source/Runtime/Landscape/Private/LandscapeEdit.cpp
@@ -3391,7 +3391,7 @@ void ALandscapeProxy::PostEditImport()
 		}
 	}
 
-	ULandscapeInfo::RecreateLandscapeInfo(GetWorld(), true);
+	GEngine->DeferredCommands.AddUnique(TEXT("UpdateLandscapeEditorData"));
 }
 
 void ALandscape::PostEditMove(bool bFinished)
