CIS fix.

Adding another include.

diff --git a/Engine/Source/Runtime/Landscape/Private/LandscapeModule.cpp b/Engine/Source/Runtime/Landscape/Private/LandscapeModule.cpp
--- a/Engine/Source/Runtime/Landscape/Private/LandscapeModule.cpp
+++ b/Engine/Source/Runtime/Landscape/Private/LandscapeModule.cpp
@@ -5,6 +5,7 @@
 #include "LandscapeComponent.h"
 #include "UObjectHash.h"
 #include "Engine/World.h"
+#include "Engine/Level.h"
 
 class FLandscapeModule : public IModuleInterface
 {
