Fix Landscape compile error under Visual Studio 2012

diff --git a/Engine/Source/Runtime/Landscape/Private/LandscapeEdit.cpp b/Engine/Source/Runtime/Landscape/Private/LandscapeEdit.cpp
--- a/Engine/Source/Runtime/Landscape/Private/LandscapeEdit.cpp
+++ b/Engine/Source/Runtime/Landscape/Private/LandscapeEdit.cpp
@@ -4852,6 +4852,11 @@ struct FMobileLayerAllocation
 {
 	FWeightmapLayerAllocationInfo Allocation;
 
+	FMobileLayerAllocation(const FWeightmapLayerAllocationInfo& InAllocation)
+		: Allocation(InAllocation)
+	{
+	}
+
 	friend bool operator<(const FMobileLayerAllocation& lhs, const FMobileLayerAllocation& rhs)
 	{
 		if (!lhs.Allocation.LayerInfo && !rhs.Allocation.LayerInfo) return false; // equally broken :P
@@ -4883,7 +4888,7 @@ UMaterialInstance* ULandscapeComponent::GeneratePlatformPixelData(TArray<UTextur
 	MobileLayerAllocations.Reserve(WeightmapLayerAllocations.Num());
 	for (const auto& Allocation : WeightmapLayerAllocations)
 	{
-		MobileLayerAllocations.Add(FMobileLayerAllocation{ Allocation });
+		MobileLayerAllocations.Emplace(Allocation);
 	}
 	MobileLayerAllocations.StableSort();
 
