Fixed a memory leak with landscape editing

diff --git a/Engine/Source/Runtime/Landscape/Private/LandscapeEditInterface.cpp b/Engine/Source/Runtime/Landscape/Private/LandscapeEditInterface.cpp
--- a/Engine/Source/Runtime/Landscape/Private/LandscapeEditInterface.cpp
+++ b/Engine/Source/Runtime/Landscape/Private/LandscapeEditInterface.cpp
@@ -246,6 +246,8 @@ void FLandscapeEditDataInterface::SetHeightData(int32 X1, int32 Y1, int32 X2, in
 				VertexNormals[(X+1 + NumVertsX*(Y+1))] += FaceNormal1 + FaceNormal2;
 			}
 		}
+
+		delete[] XYOffsets;
 	}
 
 	for( int32 ComponentIndexY=ComponentIndexY1;ComponentIndexY<=ComponentIndexY2;ComponentIndexY++ )
@@ -506,6 +508,7 @@ void FLandscapeEditDataInterface::RecalculateNormals()
 			}
 		}
 
+		delete[] XYOffsets;
 		delete[] HeightData;
 		delete[] VertexNormals;
 
