UE4 - prototype grass system

diff --git a/Engine/Source/Editor/UnrealEd/Private/Editor.cpp b/Engine/Source/Editor/UnrealEd/Private/Editor.cpp
--- a/Engine/Source/Editor/UnrealEd/Private/Editor.cpp
+++ b/Engine/Source/Editor/UnrealEd/Private/Editor.cpp
@@ -1666,7 +1666,7 @@ bool UEditorEngine::UpdateSingleViewportClient(FEditorViewportClient* InViewport
 		// Add view information for perspective viewports.
 		if( InViewportClient->IsPerspective() )
 		{
-			GWorld->ViewLocationsRenderedLastFrame.Add(InViewportClient->GetViewLocation());
+			InViewportClient->GetWorld()->ViewLocationsRenderedLastFrame.Add(InViewportClient->GetViewLocation());
 	
 			// If we're currently simulating in editor, then we'll need to make sure that sub-levels are streamed in.
 			// When using PIE, this normally happens by UGameViewportClient::Draw().  But for SIE, we need to do
@@ -1674,7 +1674,7 @@ bool UEditorEngine::UpdateSingleViewportClient(FEditorViewportClient* InViewport
 			if( PlayWorld != NULL && bIsSimulatingInEditor && InViewportClient->IsSimulateInEditorViewport() )
 			{
 				// Update level streaming.
-				GWorld->UpdateLevelStreaming();
+				InViewportClient->GetWorld()->UpdateLevelStreaming();
 
 				// Also make sure hit proxies are refreshed for SIE viewports, as the user may be trying to grab an object or widget manipulator that's moving!
 				if( InViewportClient->IsRealtime() )
diff --git a/Engine/Source/Runtime/Engine/Classes/Components/HierarchicalInstancedStaticMeshComponent.h b/Engine/Source/Runtime/Engine/Classes/Components/HierarchicalInstancedStaticMeshComponent.h
--- a/Engine/Source/Runtime/Engine/Classes/Components/HierarchicalInstancedStaticMeshComponent.h
+++ b/Engine/Source/Runtime/Engine/Classes/Components/HierarchicalInstancedStaticMeshComponent.h
@@ -74,6 +74,9 @@ class ENGINE_API UHierarchicalInstancedStaticMeshComponent : public UInstancedSt
 	bool bIsAsyncBuilding;
 	bool bConcurrentRemoval;
 
+	UPROPERTY()
+	bool bDisableCollision;
+
 	// Apply the results of the async build
 	void ApplyBuildTreeAsync(ENamedThreads::Type CurrentThread, const FGraphEventRef& MyCompletionGraphEvent, TSharedRef<FClusterBuilder, ESPMode::ThreadSafe> Builder, double StartTime);
 
@@ -88,8 +91,23 @@ public:
 	virtual bool UpdateInstanceTransform(int32 InstanceIndex, const FTransform& NewInstanceTransform, bool bWorldSpace) override;
 	virtual void ClearInstances() override;
 
+	virtual bool ShouldCreatePhysicsState() const override;
+
 	void BuildTree();
 	void BuildTreeAsync();
+	static void BuildTreeAnyThread(
+		const TArray<FInstancedStaticMeshInstanceData>& PerInstanceSMData, 
+		const FBox& MeshBox,
+		TArray<FClusterNode>& OutClusterTree,
+		TArray<int32>& OutSortedInstances,
+		TArray<int32>& OutInstanceReorderTable
+		);
+	void AcceptPrebuiltTree(
+		TArray<FClusterNode>& InClusterTree,
+		TArray<int32>& InSortedInstances,
+		TArray<int32>& InInstanceReorderTable
+		);
+	void BuildFlatTree(const TArray<int32>& LeafInstanceCounts);
 	bool IsAsyncBuilding() const { return bIsAsyncBuilding; }
 	bool IsTreeFullyBuilt() const { return NumBuiltInstances == PerInstanceSMData.Num() && RemovedInstances.Num() == 0; }
 
diff --git a/Engine/Source/Runtime/Engine/Classes/Components/InstancedStaticMeshComponent.h b/Engine/Source/Runtime/Engine/Classes/Components/InstancedStaticMeshComponent.h
--- a/Engine/Source/Runtime/Engine/Classes/Components/InstancedStaticMeshComponent.h
+++ b/Engine/Source/Runtime/Engine/Classes/Components/InstancedStaticMeshComponent.h
@@ -120,7 +120,7 @@ class ENGINE_API UInstancedStaticMeshComponent : public UStaticMeshComponent
 	UFUNCTION(BlueprintCallable, Category = "Components|InstancedStaticMesh")
 	void SetCullDistances(int32 StartCullDistance, int32 EndCullDistance);
 
-	virtual bool ShouldCreatePhysicsState() const;
+	virtual bool ShouldCreatePhysicsState() const override;
 
 public:
 #if WITH_EDITOR
diff --git a/Engine/Source/Runtime/Engine/Private/HierarchicalInstancedStaticMesh.cpp b/Engine/Source/Runtime/Engine/Private/HierarchicalInstancedStaticMesh.cpp
--- a/Engine/Source/Runtime/Engine/Private/HierarchicalInstancedStaticMesh.cpp
+++ b/Engine/Source/Runtime/Engine/Private/HierarchicalInstancedStaticMesh.cpp
@@ -10,7 +10,6 @@
 #include "InstancedStaticMesh.h"
 #include "NavigationSystemHelpers.h"
 #include "AI/Navigation/NavCollision.h"
-#include "../../Renderer/Private/ScenePrivate.h"
 
 static TAutoConsoleVariable<int32> CVarFoliageSplitFactor(
 	TEXT("foliage.SplitFactor"),
@@ -207,7 +206,7 @@ public:
 		, InstBox(InInstBox)
 		, Result(nullptr)
 	{
-		BranchingFactor = CVarFoliageSplitFactor.GetValueOnGameThread();
+		BranchingFactor = CVarFoliageSplitFactor.GetValueOnAnyThread();
 
 		Exchange(Transforms, InTransforms);
 		SortIndex.AddUninitialized(Num);
@@ -261,6 +260,15 @@ public:
 		TArray<int32> NodesPerLevel;
 		NodesPerLevel.Add(NumRoots);
 		int32 LOD = 0;
+
+		TArray<int32> InverseSortIndex;
+		TArray<int32> RemapSortIndex;
+		TArray<int32> InverseInstanceIndex;
+		TArray<int32> OldInstanceIndex;
+		TArray<int32> LevelStarts;
+		TArray<int32> InverseChildIndex;
+		TArray<FClusterNode> OldNodes;
+
 		while (NumRoots > 1)
 		{
 			SortIndex.Reset();
@@ -275,7 +283,7 @@ public:
 			}
 			Split(NumRoots);
 
-			TArray<int32> InverseSortIndex;
+			InverseSortIndex.Reset();
 			InverseSortIndex.AddUninitialized(NumRoots);
 			for (int32 Index = 0; Index < NumRoots; Index++)
 			{
@@ -284,7 +292,7 @@ public:
 
 			{
 				// rearrange the instances to match the new order of the old roots
-				TArray<int32> RemapSortIndex;
+				RemapSortIndex.Reset();
 				RemapSortIndex.AddUninitialized(Num);
 				int32 OutIndex = 0;
 				for (int32 Index = 0; Index < NumRoots; Index++)
@@ -295,7 +303,7 @@ public:
 						RemapSortIndex[OutIndex++] = InstanceIndex;
 					}
 				}
-				TArray<int32> InverseInstanceIndex;
+				InverseInstanceIndex.Reset();
 				InverseInstanceIndex.AddUninitialized(Num);
 				for (int32 Index = 0; Index < Num; Index++)
 				{
@@ -307,7 +315,7 @@ public:
 					Node.FirstInstance = InverseInstanceIndex[Node.FirstInstance];
 					Node.LastInstance = InverseInstanceIndex[Node.LastInstance];
 				}
-				TArray<int32> OldInstanceIndex;
+				OldInstanceIndex.Reset();
 				Exchange(OldInstanceIndex, SortedInstances);
 				SortedInstances.AddUninitialized(Num);
 				for (int32 Index = 0; Index < Num; Index++)
@@ -317,11 +325,11 @@ public:
 			}
 			{
 				// rearrange the nodes to match the new order of the old roots
-				TArray<int32> RemapSortIndex;
+				RemapSortIndex.Reset();
 				int32 NewNum = Result->Nodes.Num() + Clusters.Num();
 				// RemapSortIndex[new index] == old index
 				RemapSortIndex.AddUninitialized(NewNum);
-				TArray<int32> LevelStarts;
+				LevelStarts.Reset();
 				LevelStarts.Add(Clusters.Num());
 				for (int32 Index = 0; Index < NodesPerLevel.Num() - 1; Index++)
 				{
@@ -360,7 +368,7 @@ public:
 					}
 				}
 				check(LevelStarts[LevelStarts.Num() - 1] == NewNum);
-				TArray<int32> InverseChildIndex;
+				InverseChildIndex.Reset();
 				// InverseChildIndex[old index] == new index
 				InverseChildIndex.AddUninitialized(NewNum);
 				for (int32 Index = Clusters.Num(); Index < NewNum; Index++)
@@ -377,7 +385,7 @@ public:
 					}
 				}
 				{
-					TArray<FClusterNode> OldNodes = Result->Nodes;
+					Exchange(OldNodes, Result->Nodes);
 					Result->Nodes.Empty(NewNum);
 					for (int32 Index = 0; Index < Clusters.Num(); Index++)
 					{
@@ -1272,18 +1280,21 @@ FBoxSphereBounds UHierarchicalInstancedStaticMeshComponent::CalcBounds(const FTr
 		// @todo: Level error or something to let LDs know this
 		1;//StaticMesh->LODModels(0).Elements.Num() == 1;
 
-	if (ClusterTree.Num() > 0 && bMeshIsValid)
+	if ((ClusterTree.Num() || UnbuiltInstanceBounds.IsValid) && bMeshIsValid)
 	{
-		FBoxSphereBounds Result = FBox(ClusterTree[0].BoundMin, ClusterTree[0].BoundMax);
-		for (int32 Index = 1; Index < ClusterTree.Num(); Index++)
+		FBoxSphereBounds Result;
+		if (ClusterTree.Num())
 		{
-			Result = Result + FBox(ClusterTree[Index].BoundMin, ClusterTree[Index].BoundMax);
+			Result = FBox(ClusterTree[0].BoundMin, ClusterTree[0].BoundMax);
+			// Include any unbuilt instances
+			if (UnbuiltInstanceBounds.IsValid)
+			{
+				Result = Result + UnbuiltInstanceBounds;
+			}
 		}
-
-		// Include any unbuilt instances
-		if (UnbuiltInstanceBounds.IsValid)
+		else
 		{
-			Result = Result + UnbuiltInstanceBounds;
+			Result = UnbuiltInstanceBounds;
 		}
 
 		return Result.TransformBy(BoundTransform);
@@ -1435,6 +1446,15 @@ void UHierarchicalInstancedStaticMeshComponent::ClearInstances()
 	Super::ClearInstances();
 }
 
+bool UHierarchicalInstancedStaticMeshComponent::ShouldCreatePhysicsState() const
+{
+	if (bDisableCollision)
+	{
+		return false;
+	}
+	return Super::ShouldCreatePhysicsState();
+}
+
 void UHierarchicalInstancedStaticMeshComponent::BuildTree()
 {
 	// Verify that the mesh is valid before using it.
@@ -1450,7 +1470,7 @@ void UHierarchicalInstancedStaticMeshComponent::BuildTree()
 
 	if(bMeshIsValid)
 	{
-		double StartTime = FPlatformTime::Seconds();
+		//double StartTime = FPlatformTime::Seconds();
 		// If we don't have a random seed for this instanced static mesh component yet, then go ahead and
 		// generate one now.  This will be saved with the static mesh component and used for future generation
 		// of random numbers for this component's instances. (Used by the PerInstanceRandom material expression)
@@ -1479,10 +1499,170 @@ void UHierarchicalInstancedStaticMeshComponent::BuildTree()
 		
 		FlushAccumulatedNavigationUpdates();
 				
-		UE_LOG(LogStaticMesh, Display, TEXT("Built a foliage hierarchy with %d elements in %.1fs."), PerInstanceSMData.Num(), float(FPlatformTime::Seconds() - StartTime));
+		//UE_LOG(LogStaticMesh, Display, TEXT("Built a foliage hierarchy with %d elements in %.1fs."), PerInstanceSMData.Num(), float(FPlatformTime::Seconds() - StartTime));
 	}
 }
 
+void UHierarchicalInstancedStaticMeshComponent::BuildTreeAnyThread(
+	const TArray<FInstancedStaticMeshInstanceData>& PerInstanceSMData, 
+	const FBox& MeshBox,
+	TArray<FClusterNode>& OutClusterTree,
+	TArray<int32>& OutSortedInstances,
+	TArray<int32>& OutInstanceReorderTable
+	)
+{
+	TArray<FMatrix> InstanceTransforms;
+	InstanceTransforms.AddUninitialized(PerInstanceSMData.Num());
+	for (int32 Index = 0; Index < PerInstanceSMData.Num(); Index++)
+	{
+		InstanceTransforms[Index] = PerInstanceSMData[Index].Transform;
+	}
+
+	TUniquePtr<FClusterBuilder> Builder(new FClusterBuilder(InstanceTransforms, MeshBox));
+	Builder->Build();
+
+	Exchange(OutClusterTree, Builder->Result->Nodes);
+	Exchange(OutInstanceReorderTable, Builder->Result->InstanceReorderTable);
+	Exchange(OutSortedInstances, Builder->Result->SortedInstances);
+
+}
+
+void UHierarchicalInstancedStaticMeshComponent::AcceptPrebuiltTree(
+	TArray<FClusterNode>& InClusterTree,
+	TArray<int32>& InSortedInstances,
+	TArray<int32>& InInstanceReorderTable
+	)
+{
+	NumBuiltInstances = PerInstanceSMData.Num();
+	UnbuiltInstanceBounds.Init();
+	RemovedInstances.Empty();
+	ClusterTree.Empty();
+	InstanceReorderTable.Empty();
+	SortedInstances.Empty();
+
+	// Verify that the mesh is valid before using it.
+	const bool bMeshIsValid = 
+		// make sure we have instances
+		PerInstanceSMData.Num() > 0 &&
+		// make sure we have an actual staticmesh
+		StaticMesh &&
+		StaticMesh->HasValidRenderData() &&
+		// You really can't use hardware instancing on the consoles with multiple elements because they share the same index buffer. 
+		// @todo: Level error or something to let LDs know this
+		1;//StaticMesh->LODModels(0).Elements.Num() == 1;
+
+	if(bMeshIsValid)
+	{
+		//double StartTime = FPlatformTime::Seconds();
+		// If we don't have a random seed for this instanced static mesh component yet, then go ahead and
+		// generate one now.  This will be saved with the static mesh component and used for future generation
+		// of random numbers for this component's instances. (Used by the PerInstanceRandom material expression)
+		while( InstancingRandomSeed == 0 )
+		{
+			InstancingRandomSeed = FMath::Rand();
+		}
+
+		Exchange(ClusterTree, InClusterTree);
+		Exchange(InstanceReorderTable, InInstanceReorderTable);
+		Exchange(SortedInstances, InSortedInstances);
+
+		FlushAccumulatedNavigationUpdates();
+
+	}
+	MarkRenderStateDirty();
+}
+
+void UHierarchicalInstancedStaticMeshComponent::BuildFlatTree(const TArray<int32>& LeafInstanceCounts)
+{
+	// Verify that the mesh is valid before using it.
+	const bool bMeshIsValid = 
+		// make sure we have instances
+		PerInstanceSMData.Num() > 0 &&
+		// make sure we have an actual staticmesh
+		StaticMesh &&
+		StaticMesh->HasValidRenderData() &&
+		// You really can't use hardware instancing on the consoles with multiple elements because they share the same index buffer. 
+		// @todo: Level error or something to let LDs know this
+		1;//StaticMesh->LODModels(0).Elements.Num() == 1;
+
+	if(bMeshIsValid)
+	{
+		//double StartTime = FPlatformTime::Seconds();
+		// If we don't have a random seed for this instanced static mesh component yet, then go ahead and
+		// generate one now.  This will be saved with the static mesh component and used for future generation
+		// of random numbers for this component's instances. (Used by the PerInstanceRandom material expression)
+		while( InstancingRandomSeed == 0 )
+		{
+			InstancingRandomSeed = FMath::Rand();
+		}
+
+		NumBuiltInstances = PerInstanceSMData.Num();
+		UnbuiltInstanceBounds.Init();
+		RemovedInstances.Empty();
+
+		const FBox InstBox = StaticMesh->GetBounds().GetBox();
+
+		ClusterTree.Empty(LeafInstanceCounts.Num() + 1);
+		FClusterNode Root;
+		Root.FirstChild = 1;
+		Root.LastChild = LeafInstanceCounts.Num();
+		Root.FirstInstance = 0;
+		Root.LastInstance = PerInstanceSMData.Num() - 1;
+		ClusterTree.Add(Root);
+		FBox RootBox(ForceInit);
+
+		int32 InstIndex = 0;
+		for (auto Count : LeafInstanceCounts)
+		{
+			FClusterNode Leaf;
+			Leaf.FirstChild = -1;
+			Leaf.LastChild = -1;
+			Leaf.FirstInstance = InstIndex;
+			Leaf.LastInstance = InstIndex + Count - 1;
+			check(Count);
+
+
+			FBox LeafBox(ForceInit);
+			for (int32 SubIndex = 0; SubIndex < Count; SubIndex++)
+			{
+				const FMatrix& ThisInstTrans = PerInstanceSMData[InstIndex + SubIndex].Transform;
+				const FBox ThisInstBox = InstBox.TransformBy(ThisInstTrans);
+				LeafBox += ThisInstBox;
+			}
+			Leaf.BoundMin = LeafBox.Min;
+			Leaf.BoundMax = LeafBox.Max;
+			ClusterTree.Add(Leaf);
+			RootBox += LeafBox;
+			InstIndex += Count;
+		}
+		check(InstIndex == PerInstanceSMData.Num()); // else you botched your counts
+		ClusterTree[0].BoundMin = RootBox.Min;
+		ClusterTree[0].BoundMax = RootBox.Max;
+
+		if (ClusterTree.Num() == 2)
+		{
+			// no point in having two levels if the second level only has one node
+			ClusterTree.RemoveAt(1);
+			ClusterTree[0].FirstChild = -1;
+			ClusterTree[0].LastChild = -1;
+		}
+
+		InstanceReorderTable.Empty(PerInstanceSMData.Num());
+		SortedInstances.Empty(PerInstanceSMData.Num());
+
+		for (int32 Index = 0; Index < PerInstanceSMData.Num(); Index++)
+		{
+			InstanceReorderTable.Add(Index);
+			SortedInstances.Add(Index);
+		}
+
+		FlushAccumulatedNavigationUpdates();
+
+//		UE_LOG(LogStaticMesh, Display, TEXT("Built a flat foliage hierarchy with %d elements in %.2fms."), PerInstanceSMData.Num(), float(FPlatformTime::Seconds() - StartTime) * 1000.0f);
+	}
+}
+
+
 void UHierarchicalInstancedStaticMeshComponent::ApplyBuildTreeAsync(ENamedThreads::Type CurrentThread, const FGraphEventRef& MyCompletionGraphEvent, TSharedRef<FClusterBuilder, ESPMode::ThreadSafe> Builder, double StartTime)
 {
 	bIsAsyncBuilding = false;
diff --git a/Engine/Source/Runtime/Landscape/Classes/LandscapeLayerInfoObject.h b/Engine/Source/Runtime/Landscape/Classes/LandscapeLayerInfoObject.h
--- a/Engine/Source/Runtime/Landscape/Classes/LandscapeLayerInfoObject.h
+++ b/Engine/Source/Runtime/Landscape/Classes/LandscapeLayerInfoObject.h
@@ -18,6 +18,33 @@ class ULandscapeLayerInfoObject : public UObject
 	UPROPERTY(EditAnywhere, Category=LandscapeLayerInfoObject)
 	float Hardness;
 
+	UPROPERTY(EditAnywhere, Category=Grass)
+	class UStaticMesh* GrassMesh;
+
+	/* Instances per 10 square meters. */
+	UPROPERTY(EditAnywhere, Category=Grass)
+	float GrassDensity;
+
+	UPROPERTY(EditAnywhere, Category=Grass)
+	float PlacementJitter;
+
+	/* The distance where instances will begin to fade out if using a PerInstanceFadeAmount material node. 0 disables. */
+	UPROPERTY(EditAnywhere, Category=Grass)
+	int32 StartCullDistance;
+
+	/**
+	 * The distance where instances will have completely faded out when using a PerInstanceFadeAmount material node. 0 disables. 
+	 * When the entire cluster is beyond this distance, the cluster is completely culled and not rendered at all.
+	 */
+	UPROPERTY(EditAnywhere, Category = Grass)
+	int32 EndCullDistance;
+
+	UPROPERTY(EditAnywhere, Category = Grass)
+	bool RandomRotation;
+
+	UPROPERTY(EditAnywhere, Category = Grass)
+	bool AlignToSurface;
+
 #if WITH_EDITORONLY_DATA
 	UPROPERTY(VisibleAnywhere, Category=LandscapeLayerInfoObject)
 	uint32 bNoWeightBlend:1;
diff --git a/Engine/Source/Runtime/Landscape/Classes/LandscapeProxy.h b/Engine/Source/Runtime/Landscape/Classes/LandscapeProxy.h
--- a/Engine/Source/Runtime/Landscape/Classes/LandscapeProxy.h
+++ b/Engine/Source/Runtime/Landscape/Classes/LandscapeProxy.h
@@ -12,6 +12,9 @@
 #include "PhysicsEngine/BodyInstance.h"
 #include "Materials/Material.h"
 #include "Materials/MaterialInstanceConstant.h"
+#include "Components/HierarchicalInstancedStaticMeshComponent.h"
+#include "LandscapeComponent.h"
+#include "LandscapeLayerInfoObject.h"
 
 #include "LandscapeProxy.generated.h"
 
@@ -187,11 +190,65 @@ namespace ELandscapeLODFalloff
 	};
 }
 
+struct FCachedLandscapeFoliage
+{
+	struct FPerLayer
+	{
+		TWeakObjectPtr<UHierarchicalInstancedStaticMeshComponent> Foliage;
+		TWeakObjectPtr<ULandscapeLayerInfoObject> Layer;
+
+		FPerLayer(UHierarchicalInstancedStaticMeshComponent* InFoliage, ULandscapeLayerInfoObject* InLayer)
+			: Foliage(InFoliage)
+			, Layer(InLayer)
+		{
+		}
+	};
+	struct FPerComponent
+	{
+		TWeakObjectPtr<ULandscapeComponent> BasedOn;
+		TArray<FPerLayer> Layers;
+		FPerComponent(ULandscapeComponent* InBasedOn)
+			: BasedOn(InBasedOn)
+		{
+		}
+	};
+	TArray<FPerComponent> PerComponent;
+
+	void ClearCache()
+	{
+		PerComponent.Empty();
+	}
+};
+
+class FAsyncGrassTask : public FNonAbandonableTask
+{
+public:
+	struct FAsyncGrassBuilder* Builder;
+	TWeakObjectPtr<ULandscapeComponent> BasedOn;
+	TWeakObjectPtr<UHierarchicalInstancedStaticMeshComponent> Foliage;
+	TWeakObjectPtr<ULandscapeLayerInfoObject> Layer;
+
+	FAsyncGrassTask(struct FAsyncGrassBuilder* InBuilder, ULandscapeComponent* InBasedOn, UHierarchicalInstancedStaticMeshComponent* InFoliage, ULandscapeLayerInfoObject* InLayer)
+		: Builder(InBuilder)
+		, BasedOn(InBasedOn)
+		, Foliage(InFoliage)
+		, Layer(InLayer)
+	{
+	}
+	void DoWork();
+	static const TCHAR *Name()
+	{
+		return TEXT("FAsyncGrassTask");
+	}
+};
+
 UCLASS(NotPlaceable, hidecategories=(Display, Attachment, Physics, Debug, Lighting, LOD), showcategories=(Rendering, "Utilities|Transformation"), MinimalAPI)
 class ALandscapeProxy : public AActor
 {
 	GENERATED_UCLASS_BODY()
 
+	virtual ~ALandscapeProxy();
+
 	UPROPERTY()
 	ULandscapeSplinesComponent* SplineComponent;
 
@@ -249,6 +306,14 @@ public:
 	UPROPERTY()
 	TArray<ULandscapeHeightfieldCollisionComponent*> CollisionComponents;
 
+	UPROPERTY(transient, duplicatetransient)
+	TArray<UHierarchicalInstancedStaticMeshComponent*> FoliageComponents;
+
+	/** A transient data structure for tracking the grass */
+	FCachedLandscapeFoliage FoliageCache;
+	/** A transient data structure for tracking the grass tasks*/
+	TArray<FAsyncTask<FAsyncGrassTask>* > AsyncFoliageTasks;
+
 	/**
 	 *	The resolution to cache lighting at, in texels/quad in one axis
 	 *  Total resolution would be changed by StaticLightingResolution*StaticLightingResolution
@@ -342,6 +407,13 @@ public:
 	virtual void RegisterAllComponents() override;
 	virtual void RerunConstructionScripts() override {}
 	virtual bool IsLevelBoundsRelevant() const override { return true; }
+	virtual void Tick(float DeltaSeconds) override;
+	virtual bool ShouldTickIfViewportsOnly() const override
+	{
+		return true;
+	}
+
+
 #if WITH_EDITOR
 	virtual void Destroyed() override;
 	virtual void EditorApplyScale(const FVector& DeltaScale, const FVector* PivotLocation, bool bAltDown, bool bShiftDown, bool bCtrlDown) override;
@@ -356,6 +428,9 @@ public:
 	FGuid GetLandscapeGuid() const { return LandscapeGuid; }
 	virtual ALandscape* GetLandscapeActor();
 
+	/** Flush the foliage cache */
+	LANDSCAPE_API void FlushFoliageComponents();
+
 	// Begin UObject interface.
 	virtual void Serialize(FArchive& Ar) override;
 	static void AddReferencedObjects(UObject* InThis, FReferenceCollector& Collector);
diff --git a/Engine/Source/Runtime/Landscape/Private/Landscape.cpp b/Engine/Source/Runtime/Landscape/Private/Landscape.cpp
--- a/Engine/Source/Runtime/Landscape/Private/Landscape.cpp
+++ b/Engine/Source/Runtime/Landscape/Private/Landscape.cpp
@@ -28,6 +28,7 @@ Landscape.cpp: Terrain rendering
 #include "Engine/Engine.h"
 #include "AI/Navigation/NavigationSystem.h"
 #include "EngineUtils.h"
+#include "Engine/StaticMesh.h"
 
 #include "Materials/MaterialExpressionLandscapeLayerWeight.h"
 #include "Materials/MaterialExpressionLandscapeLayerSwitch.h"
@@ -471,7 +472,7 @@ ALandscape* ALandscape::GetLandscapeActor()
 ALandscapeProxy::ALandscapeProxy(const FObjectInitializer& ObjectInitializer)
 	: Super(ObjectInitializer)
 {
-	PrimaryActorTick.bCanEverTick = false;
+	PrimaryActorTick.bCanEverTick = true;
 	bReplicates = false;
 	NetUpdateFrequency = 10.0f;
 	bHidden = false;
@@ -2034,4 +2035,594 @@ bool LandscapeMaterialsParameterSetUpdater(FStaticParameterSet &StaticParameterS
 	return UpdateParameterSet<FStaticTerrainLayerWeightParameter, UMaterialExpressionLandscapeLayerWeight>(StaticParameterSet.TerrainLayerWeightParameters, ParentMaterial);
 }
 
+void ALandscapeProxy::FlushFoliageComponents()
+{
+	// Clear old foliage component containers
+	FoliageComponents.Empty();
+
+	// Might as well clear the cache...
+	FoliageCache.ClearCache();
+
+	// Destroy any owned foliage components
+	TArray<UHierarchicalInstancedStaticMeshComponent*> FoliageComps;
+	GetComponents(FoliageComps);
+	for (UHierarchicalInstancedStaticMeshComponent* Component : FoliageComps)
+	{
+		Component->ClearInstances();
+		Component->DestroyComponent();
+	}
+
+	TArray<USceneComponent*> AttachedFoliageComponents = RootComponent->AttachChildren.FilterByPredicate(
+		[](USceneComponent* Component)
+	{
+		return Cast<UHierarchicalInstancedStaticMeshComponent>(Component);
+	});
+
+	// Destroy any attached but un-owned foliage components
+	for (USceneComponent* Component : AttachedFoliageComponents)
+	{
+		Component->DestroyComponent();
+	}
+}
+
+static TAutoConsoleVariable<float> CVarGuardBandMultiplier(
+	TEXT("grass.GuardBandMultiplier"),
+	1.5f,
+	TEXT("Used to control discarding in the grass system. Approximate range, 1-2. Multiplied by the cull distance to control when we discard grass components."));
+
+struct FAsyncGrassBuilder
+{
+	bool bHaveValidData;
+	bool RandomRotation;
+	bool AlignToSurface;
+	float GrassDensity;
+	float PlacementJitter;
+	int32 EndCullDistance;
+
+	FVector DrawScale;
+	FVector DrawLoc;
+	FMatrix LandscapeToWorld;
+	FMatrix XForm;
+
+	FIntPoint SectionBase;
+	FIntPoint LandscapeSectionOffset;
+	int32 ComponentSizeQuads;
+	FVector Origin;
+	FVector Extent;
+	int32 SqrtMaxInstances;
+
+	int32 SubsectionSizeQuads;
+	int32 NumSubsections;
+	int32 Stride;
+
+
+	TArray<FColor> LocalHeightCache;
+	TArray<uint8> LocalCache;
+
+	FRandomStream RandomStream;
+
+	FBox MeshBox;
+
+	double RasterTime;
+	double BuildTime;
+	double InstanceTime;
+	int32 TotalInstances;
+
+	// output
+	TArray<FInstancedStaticMeshInstanceData> PerInstanceSMData;
+	TArray<FClusterNode> ClusterTree;
+	TArray<int32> SortedInstances;
+	TArray<int32> InstanceReorderTable;
+
+	FAsyncGrassBuilder(ALandscapeProxy* Landscape, ULandscapeComponent* Component, ULandscapeLayerInfoObject* LayerInfo, UHierarchicalInstancedStaticMeshComponent* HierarchicalInstancedStaticMeshComponent)
+	{
+		bHaveValidData = true;
+		GrassDensity = LayerInfo->GrassDensity;
+		PlacementJitter = LayerInfo->PlacementJitter;
+		EndCullDistance = LayerInfo->EndCullDistance;
+		RandomRotation = LayerInfo->RandomRotation;
+		AlignToSurface = LayerInfo->AlignToSurface;
+
+		DrawScale = Landscape->GetRootComponent()->RelativeScale3D;
+		DrawLoc = Landscape->GetActorLocation();
+		LandscapeSectionOffset = Landscape->LandscapeSectionOffset;
+		const FRotator DrawRot = Landscape->GetActorRotation();
+		LandscapeToWorld = Landscape->GetRootComponent()->ComponentToWorld.ToMatrixNoScale();
+
+#if WITH_EDITOR
+		{
+			FLandscapeComponentDataInterface CDI(Component);
+			CDI.GetHeightmapTextureData(LocalHeightCache);
+			if (!CDI.GetWeightmapTextureData(LayerInfo, LocalCache) || !LocalCache.Num() || LocalCache.Num() != LocalHeightCache.Num())
+			{
+				// no data for this layer for this component.
+				bHaveValidData = false;
+			}
+		}
+#else
+		bHaveValidData = false;
+#endif
+		SectionBase = Component->GetSectionBase();
+		ComponentSizeQuads = Component->ComponentSizeQuads;
+		SubsectionSizeQuads = Component->SubsectionSizeQuads;
+		NumSubsections = Component->NumSubsections;
+		Stride = (SubsectionSizeQuads + 1) * NumSubsections;
+
+		RandomStream.Initialize(HierarchicalInstancedStaticMeshComponent->InstancingRandomSeed);
+
+		Origin = FVector(DrawScale.X * float(SectionBase.X), DrawScale.Y * float(SectionBase.Y), 0.0f);
+		Extent = FVector(DrawScale.X * float(SectionBase.X + ComponentSizeQuads), DrawScale.Y * float(SectionBase.Y + ComponentSizeQuads), 0.0f) - Origin;
+
+		SqrtMaxInstances = FMath::CeilToInt(FMath::Sqrt(FMath::Abs(Extent.X * Extent.Y * GrassDensity / 1000.0f / 1000.0f)));
+
+		if (!SqrtMaxInstances)
+		{
+			bHaveValidData = false;
+		}
+
+		XForm = LandscapeToWorld * HierarchicalInstancedStaticMeshComponent->ComponentToWorld.ToMatrixWithScale().Inverse();
+
+		MeshBox = LayerInfo->GrassMesh->GetBounds().GetBox();
+
+		RasterTime = 0.0;
+		BuildTime = 0.0;
+		InstanceTime = 0.0;
+		TotalInstances = 0;
+	}
+
+	void Build()
+	{
+		check(bHaveValidData);
+		RasterTime -= FPlatformTime::Seconds();
+
+		float Div = 1.0f / float(SqrtMaxInstances);
+		float MaxJitter1D = FMath::Clamp<float>(PlacementJitter, 0.0f, .99f) * Div * .5f;
+		FVector MaxJitter(MaxJitter1D, MaxJitter1D, 0.0f);
+		MaxJitter *= Extent;
+		Origin += Extent * (Div * 0.5f);
+		struct FInstanceLocal
+	{
+			FVector Pos;
+			bool bKeep;
+		};
+		TArray<FInstanceLocal> Instances;
+		Instances.AddUninitialized(SqrtMaxInstances * SqrtMaxInstances);
+		int32 NumKept = 0;
+		{
+			int32 InstanceIndex = 0;
+			for (int32 xStart = 0; xStart < SqrtMaxInstances; xStart++)
+			{
+				for (int32 yStart = 0; yStart < SqrtMaxInstances; yStart++)
+				{
+					FVector Location(Origin.X + float(xStart) * Div * Extent.X, Origin.Y + float(yStart) * Div * Extent.Y, 0.0f);
+					Location += FVector(RandomStream.GetFraction() * 2.0f - 1.0f, RandomStream.GetFraction() * 2.0f - 1.0f, 0.0f) * MaxJitter;
+
+					FInstanceLocal& Instance = Instances[InstanceIndex];
+					float Weight = GetLayerWeightAtLocationLocal(Location, &Instance.Pos);
+					Instance.bKeep = Weight >= RandomStream.GetFraction();
+					if (Instance.bKeep)
+					{
+						NumKept++;
+					}
+					InstanceIndex++;
+				}
+			}
+		}
+		RasterTime += FPlatformTime::Seconds();
+
+		if (NumKept)
+		{
+			TotalInstances += NumKept;
+
+
+			FInstancedStaticMeshInstanceData Inst;
+			Inst.LightmapUVBias = FVector2D( -1.0f, -1.0f );
+			Inst.ShadowmapUVBias = FVector2D( -1.0f, -1.0f );
+			{
+				InstanceTime -= FPlatformTime::Seconds();
+				PerInstanceSMData.Reserve(NumKept);
+				int32 InstanceIndex = 0;
+				for (int32 xStart = 0; xStart < SqrtMaxInstances; xStart++ )
+				{
+					for (int32 yStart = 0; yStart < SqrtMaxInstances; yStart++ )
+					{
+						const FInstanceLocal& Instance = Instances[InstanceIndex];
+						if (Instance.bKeep)
+						{
+							float Rot = RandomRotation ? RandomStream.GetFraction() * 360.0f : 0.0f;
+							if (AlignToSurface)
+							{
+								FVector PosX1 = xStart ? Instances[InstanceIndex - SqrtMaxInstances].Pos : Instance.Pos;
+								FVector PosX2 = (xStart + 1 < SqrtMaxInstances) ? Instances[InstanceIndex + SqrtMaxInstances].Pos : Instance.Pos;
+								FVector PosY1 = yStart ? Instances[InstanceIndex - 1].Pos : Instance.Pos;
+								FVector PosY2 = (yStart + 1 < SqrtMaxInstances) ? Instances[InstanceIndex + 1].Pos : Instance.Pos;
+
+								if (PosX1 != PosX2 && PosY1 != PosY2)
+								{
+									FVector NewZ = ((PosX1 - PosX2) ^ (PosY1 - PosY2)).GetSafeNormal();
+									NewZ *= FMath::Sign(NewZ.Z);
+
+									const FVector NewX = (FVector(0,-1,0) ^ NewZ).GetSafeNormal();
+									const FVector NewY = NewZ ^ NewX;
+
+									FMatrix Align = FMatrix(NewX, NewY, NewZ, FVector::ZeroVector);
+									Inst.Transform = FRotationMatrix(FRotator(0.0f, Rot, 0.0f)) * Align * FTranslationMatrix(Instance.Pos) * XForm;
+								}
+								else
+								{
+									Inst.Transform = FRotationMatrix(FRotator(0.0f, Rot, 0.0f)) * FTranslationMatrix(Instance.Pos) * XForm;
+								}
+							}
+							else
+							{
+								Inst.Transform = FRotationMatrix(FRotator(0.0f, Rot, 0.0f)) * FTranslationMatrix(Instance.Pos) * XForm;
+							}
+							PerInstanceSMData.Add(Inst);
+						}
+						InstanceIndex++;
+					}
+				}
+				InstanceTime += FPlatformTime::Seconds();
+			}
+			{
+				BuildTime -= FPlatformTime::Seconds();
+				UHierarchicalInstancedStaticMeshComponent::BuildTreeAnyThread(PerInstanceSMData, MeshBox, ClusterTree, SortedInstances, InstanceReorderTable);
+				BuildTime += FPlatformTime::Seconds();
+			}
+		}
+	}
+	float GetLayerWeightAtLocationLocal(const FVector& InLocation, FVector* OutLocation)
+	{
+		// Find location
+		float TestX = InLocation.X / DrawScale.X - (float)SectionBase.X;
+		float TestY = InLocation.Y / DrawScale.Y - (float)SectionBase.Y;
+
+		// Find data
+		int32 X1 = FMath::FloorToInt(TestX);
+		int32 Y1 = FMath::FloorToInt(TestY);
+		int32 X2 = FMath::CeilToInt(TestX);
+		int32 Y2 = FMath::CeilToInt(TestY);
+
+		// Min is to prevent the sampling of the final column from overflowing
+		int32 IdxX1 = FMath::Min<int32>(((X1 / SubsectionSizeQuads) * (SubsectionSizeQuads + 1)) + (X1 % SubsectionSizeQuads), Stride - 1);
+		int32 IdxY1 = FMath::Min<int32>(((Y1 / SubsectionSizeQuads) * (SubsectionSizeQuads + 1)) + (Y1 % SubsectionSizeQuads), Stride - 1);
+		int32 IdxX2 = FMath::Min<int32>(((X2 / SubsectionSizeQuads) * (SubsectionSizeQuads + 1)) + (X2 % SubsectionSizeQuads), Stride - 1);
+		int32 IdxY2 = FMath::Min<int32>(((Y2 / SubsectionSizeQuads) * (SubsectionSizeQuads + 1)) + (Y2 % SubsectionSizeQuads), Stride - 1);
+
+		float LerpX = FMath::Fractional(TestX);
+		float LerpY = FMath::Fractional(TestY);
+
+		float Result;
+		{
+			// sample
+			float Sample11 = (float)((LocalCache)[IdxX1 + Stride*IdxY1]) / 255.f;
+			float Sample21 = (float)((LocalCache)[IdxX2 + Stride*IdxY1]) / 255.f;
+			float Sample12 = (float)((LocalCache)[IdxX1 + Stride*IdxY2]) / 255.f;
+			float Sample22 = (float)((LocalCache)[IdxX2 + Stride*IdxY2]) / 255.f;
+
+			// Bilinear interpolate
+			Result = FMath::Lerp(
+				FMath::Lerp(Sample11, Sample21, LerpX),
+				FMath::Lerp(Sample12, Sample22, LerpX),
+				LerpY);
+		}
+
+		{
+			struct FLocal
+			{
+				static float Sample(TArray<FColor>* HeightCache, int32 IdxX, int32 IdxY, int32 Stride)
+				{
+					const FColor& SampleVal = (*HeightCache)[IdxX + Stride*IdxY];
+					return LandscapeDataAccess::GetLocalHeight((uint16(SampleVal.R) << 8) + uint16(SampleVal.G));
+				}
+			};
+			// sample
+			float Sample11 = FLocal::Sample(&LocalHeightCache, IdxX1, IdxY1, Stride);
+			float Sample21 = FLocal::Sample(&LocalHeightCache, IdxX2, IdxY1, Stride);
+			float Sample12 = FLocal::Sample(&LocalHeightCache, IdxX1, IdxY2, Stride);
+			float Sample22 = FLocal::Sample(&LocalHeightCache, IdxX2, IdxY2, Stride);
+
+			OutLocation->X = InLocation.X - DrawScale.X * float(LandscapeSectionOffset.X);
+			OutLocation->Y = InLocation.Y - DrawScale.Y * float(LandscapeSectionOffset.Y);
+			// Bilinear interpolate
+			OutLocation->Z = DrawScale.Z * FMath::Lerp(
+				FMath::Lerp(Sample11, Sample21, LerpX),
+				FMath::Lerp(Sample12, Sample22, LerpX),
+				LerpY);
+		}
+		return Result;
+	}
+
+};
+
+
+void ALandscapeProxy::Tick(float DeltaSeconds)
+{
+#if !WITH_EDITOR
+	return;
+#endif
+
+	Super::Tick(DeltaSeconds);
+
+	UWorld* World = GetWorld();
+	if (!World)
+	{
+		return;
+	}
+
+	// there is a bug here, which often leaves us with no cameras in the editor
+	static TArray<FVector> OldCameras;
+
+	if (!OldCameras.Num() && !World->ViewLocationsRenderedLastFrame.Num())
+	{
+		// no cameras, no grass update
+		return;
+	}
+
+	const TArray<FVector>& Cameras = World->ViewLocationsRenderedLastFrame.Num() ? World->ViewLocationsRenderedLastFrame : OldCameras;
+
+	if (&Cameras != &OldCameras)
+	{
+		check(IsInGameThread());
+		OldCameras = Cameras;
+	}
+
+	FCachedLandscapeFoliage DesiredFoliage;
+	TSet<UHierarchicalInstancedStaticMeshComponent *> StillUsed;
+
+	double ComponentTime = 0.0;
+	double BuildTime = 0.0;
+	double ComponentMiscTime = 0.0;
+	double StartTime = FPlatformTime::Seconds();
+	int32 TotalComponents = 0;
+
+	float GuardBand = CVarGuardBandMultiplier.GetValueOnAnyThread();
+
+	for (int32 ComponentIndex = 0; ComponentIndex < LandscapeComponents.Num(); ComponentIndex++)
+	{
+		ComponentTime -= FPlatformTime::Seconds();
+		ULandscapeComponent* Component = LandscapeComponents[ComponentIndex];
+
+		const FCachedLandscapeFoliage::FPerComponent* ExistingCache = nullptr;
+		for (const FCachedLandscapeFoliage::FPerComponent& CacheItem : FoliageCache.PerComponent)
+		{
+			if (CacheItem.BasedOn.Get() == Component)
+			{
+				ExistingCache = &CacheItem;
+				break;
+			}
+		}
+		FCachedLandscapeFoliage::FPerComponent NewPerComponent(Component);
+
+		FBoxSphereBounds WorldBounds = Component->CalcBounds(Component->ComponentToWorld);
+		float MinDistanceToComp = MAX_flt;
+
+		for (auto& Pos : Cameras)
+		{
+			MinDistanceToComp = FMath::Min<float>(MinDistanceToComp, FVector::Dist(WorldBounds.Origin, Pos) - WorldBounds.SphereRadius);
+		}
+
+		for (FWeightmapLayerAllocationInfo& LayerAlloc : Component->WeightmapLayerAllocations)
+		{
+			ULandscapeLayerInfoObject* LayerInfo = LayerAlloc.LayerInfo;
+
+			if (LayerInfo && LayerInfo->GrassMesh && LayerInfo->GrassDensity > 0.0f && LayerInfo->EndCullDistance > 0)
+			{
+				bool bRandomRotation = LayerInfo->RandomRotation;
+				bool bAlign = LayerInfo->AlignToSurface;
+				float MustHaveDistance = (float)LayerInfo->EndCullDistance;
+				float DiscardDistance = MustHaveDistance * GuardBand;
+
+				if (MinDistanceToComp > DiscardDistance)
+				{
+					continue;
+				}
+				if (ExistingCache)
+				{
+					bool bFound = false;
+					for (const FCachedLandscapeFoliage::FPerLayer& Layer : ExistingCache->Layers)
+					{
+						if (Layer.Layer.Get() == LayerInfo && Layer.Foliage.Get())
+						{
+							NewPerComponent.Layers.Add(Layer);
+							StillUsed.Add(Layer.Foliage.Get());
+							bFound = true;
+							break;
+						}
+					}
+					if (bFound || MinDistanceToComp > MustHaveDistance)
+					{
+						continue;
+					}
+				}
+
+				int32 FolSeed = FCrc::StrCrc32((LayerInfo->LayerName.ToString() + Component->GetName()).GetCharArray().GetData());
+				if (FolSeed == 0)
+				{
+					FolSeed++;
+				}
+
+				UHierarchicalInstancedStaticMeshComponent* HierarchicalInstancedStaticMeshComponent = ConstructObject<UHierarchicalInstancedStaticMeshComponent>(UHierarchicalInstancedStaticMeshComponent::StaticClass(), this); //, NAME_None, RF_Transactional);
+				TotalComponents++;
+				HierarchicalInstancedStaticMeshComponent->Mobility = EComponentMobility::Static;
+
+				HierarchicalInstancedStaticMeshComponent->StaticMesh = LayerInfo->GrassMesh;
+				HierarchicalInstancedStaticMeshComponent->bSelectable = true;
+				HierarchicalInstancedStaticMeshComponent->bHasPerInstanceHitProxies = true;
+				static FName NoCollision(TEXT("NoCollision"));
+				HierarchicalInstancedStaticMeshComponent->SetCollisionProfileName(NoCollision);
+				HierarchicalInstancedStaticMeshComponent->bDisableCollision = true;
+				HierarchicalInstancedStaticMeshComponent->SetCanEverAffectNavigation(false);
+				HierarchicalInstancedStaticMeshComponent->InstancingRandomSeed = FolSeed;
+				HierarchicalInstancedStaticMeshComponent->InstanceStartCullDistance = LayerInfo->StartCullDistance;
+				HierarchicalInstancedStaticMeshComponent->InstanceEndCullDistance = LayerInfo->EndCullDistance;
+
+
+#if 0
+				HierarchicalInstancedStaticMeshComponent->CastShadow                     = Settings->CastShadow;
+				HierarchicalInstancedStaticMeshComponent->bCastDynamicShadow             = Settings->bCastDynamicShadow;
+				HierarchicalInstancedStaticMeshComponent->bCastStaticShadow              = Settings->bCastStaticShadow;
+				HierarchicalInstancedStaticMeshComponent->bAffectDynamicIndirectLighting = Settings->bAffectDynamicIndirectLighting;
+				HierarchicalInstancedStaticMeshComponent->bCastShadowAsTwoSided          = Settings->bCastShadowAsTwoSided;
+				HierarchicalInstancedStaticMeshComponent->bReceivesDecals                = Settings->bReceivesDecals;
+				HierarchicalInstancedStaticMeshComponent->bOverrideLightMapRes           = Settings->bOverrideLightMapRes;
+				HierarchicalInstancedStaticMeshComponent->OverriddenLightMapRes          = Settings->OverriddenLightMapRes;
+				HierarchicalInstancedStaticMeshComponent->BodyInstance.CopyBodyInstancePropertiesFrom(&Settings->BodyInstance);
+#endif
+
+				HierarchicalInstancedStaticMeshComponent->AttachTo(GetRootComponent());
+				FTransform DesiredTransform = GetRootComponent()->ComponentToWorld;
+				DesiredTransform.RemoveScaling();
+				HierarchicalInstancedStaticMeshComponent->SetWorldTransform(DesiredTransform);
+
+				NewPerComponent.Layers.Add(FCachedLandscapeFoliage::FPerLayer(HierarchicalInstancedStaticMeshComponent, LayerInfo));
+				FoliageComponents.Add(HierarchicalInstancedStaticMeshComponent);
+				StillUsed.Add(HierarchicalInstancedStaticMeshComponent);
+
+
+				FAsyncGrassBuilder* Builder = new FAsyncGrassBuilder(this, Component, LayerInfo, HierarchicalInstancedStaticMeshComponent);
+
+				if (Builder->bHaveValidData)
+				{
+					BuildTime -= FPlatformTime::Seconds();
+					FAsyncTask<FAsyncGrassTask>* Task = new FAsyncTask<FAsyncGrassTask>(Builder, Component, HierarchicalInstancedStaticMeshComponent, LayerInfo);
+
+					Task->StartBackgroundTask();
+
+					AsyncFoliageTasks.Add(Task);
+
+					BuildTime += FPlatformTime::Seconds();
+				}
+				else
+				{
+					delete Builder;
+				}
+
+
+				ComponentMiscTime -= FPlatformTime::Seconds();
+				HierarchicalInstancedStaticMeshComponent->RegisterComponent();
+				ComponentMiscTime += FPlatformTime::Seconds();
+			}
+		}
+		if (NewPerComponent.Layers.Num())
+		{
+			DesiredFoliage.PerComponent.Add(NewPerComponent);
+		}
+
+		ComponentTime += FPlatformTime::Seconds();
+	}
+
+	if (TotalComponents)
+	{
+		double TotalTime = FPlatformTime::Seconds() - StartTime;
+		UE_LOG(LogTemp, Display, TEXT("Started grass with  %d components in %.1fms     Build Time %.1fms    Component Time  %.1fms    RegisterComponent Time  %.1fms"), TotalComponents, float(TotalTime) * 1000.0f, float(BuildTime) * 1000.0f, float(ComponentTime) * 1000.0f, float(ComponentMiscTime) * 1000.0f);
+	}
+
+	Exchange(FoliageCache, DesiredFoliage);
+
+	TSet<UHierarchicalInstancedStaticMeshComponent *> ToDestroy;
+
+	TArray<UHierarchicalInstancedStaticMeshComponent*> FoliageComps;
+	GetComponents(FoliageComps);
+	for (UHierarchicalInstancedStaticMeshComponent* Component : FoliageComps)
+	{
+		if (!StillUsed.Contains(Component))
+		{
+			ToDestroy.Add(Component);
+		}
+	}
+
+	TArray<USceneComponent*> AttachedFoliageComponents = RootComponent->AttachChildren.FilterByPredicate(
+		[](USceneComponent* Component)
+	{
+		return Cast<UHierarchicalInstancedStaticMeshComponent>(Component);
+	});
+
+	// Destroy any attached but un-owned foliage components
+	for (USceneComponent* Component : AttachedFoliageComponents)
+	{
+		UHierarchicalInstancedStaticMeshComponent* HComponent = Cast<UHierarchicalInstancedStaticMeshComponent>(Component);
+		if (HComponent && !StillUsed.Contains(HComponent))
+		{
+			ToDestroy.Add(HComponent);
+		}
+	}
+	for (UHierarchicalInstancedStaticMeshComponent* Component : ToDestroy)
+	{
+		Component->ClearInstances();
+		Component->DestroyComponent();
+		FoliageComponents.Remove(Component);
+	}
+
+	StartTime = FPlatformTime::Seconds();
+	int32 TotalInstances = 0;
+	TotalComponents = 0;
+	for (int32 Index = 0; Index < AsyncFoliageTasks.Num(); Index++)
+	{
+		FAsyncTask<FAsyncGrassTask>* Task = AsyncFoliageTasks[Index];
+		if (Task->IsDone())
+		{
+			FAsyncGrassTask& Inner = Task->GetTask();
+			AsyncFoliageTasks.RemoveAtSwap(Index--);
+			TotalComponents++;
+			TotalInstances += Inner.Builder->TotalInstances;
+			UHierarchicalInstancedStaticMeshComponent* HierarchicalInstancedStaticMeshComponent = Inner.Foliage.Get();
+			if (HierarchicalInstancedStaticMeshComponent && StillUsed.Contains(HierarchicalInstancedStaticMeshComponent))
+			{
+				Exchange(HierarchicalInstancedStaticMeshComponent->PerInstanceSMData, Inner.Builder->PerInstanceSMData);
+				HierarchicalInstancedStaticMeshComponent->AcceptPrebuiltTree(Inner.Builder->ClusterTree, Inner.Builder->SortedInstances, Inner.Builder->InstanceReorderTable);
+
+				FBoxSphereBounds WorldBounds = HierarchicalInstancedStaticMeshComponent->CalcBounds(HierarchicalInstancedStaticMeshComponent->ComponentToWorld);
+				UE_LOG(LogTemp, Display, TEXT("Comp %d instances at %.0f %.0f %.0f"), HierarchicalInstancedStaticMeshComponent->PerInstanceSMData.Num(), WorldBounds.Origin.X, WorldBounds.Origin.Y, WorldBounds.Origin.Z);
+			}
+			delete Inner.Builder;
+			delete Task;
+		}
+	}
+	if (TotalComponents)
+	{
+		double TotalTime = FPlatformTime::Seconds() - StartTime;
+		UE_LOG(LogTemp, Display, TEXT("Finished grass with  %d components and %d instances in %.1fms"), TotalComponents, TotalInstances, float(TotalTime) * 1000.0f);
+	}
+}
+
+ALandscapeProxy::~ALandscapeProxy()
+{
+	for (int32 Index = 0; Index < AsyncFoliageTasks.Num(); Index++)
+	{
+		FAsyncTask<FAsyncGrassTask>* Task = AsyncFoliageTasks[Index];
+		Task->EnsureCompletion(true);
+		FAsyncGrassTask& Inner = Task->GetTask();
+		delete Inner.Builder;
+		delete Task;
+	}
+	AsyncFoliageTasks.Empty();
+}
+
+void FAsyncGrassTask::DoWork()
+{
+	Builder->Build();
+}
+
+static void FlushFoliage(const TArray<FString>& Args)
+{
+	for( TObjectIterator<ALandscapeProxy> It; It; ++It )
+	{
+		ALandscapeProxy* Landscape = *It;
+		if( Landscape && !Landscape->IsTemplate())
+		{
+			Landscape->FlushFoliageComponents();
+		}
+	}
+}
+
+static FAutoConsoleCommand FlushFoliageCmd(
+	TEXT("grass.FlushCache"),
+	TEXT("Flush the grass cache, debugging."),
+	FConsoleCommandWithArgsDelegate::CreateStatic(&FlushFoliage)
+	);
+
+
+
+
 #undef LOCTEXT_NAMESPACE
\ No newline at end of file
diff --git a/Engine/Source/Runtime/Landscape/Private/LandscapeEdit.cpp b/Engine/Source/Runtime/Landscape/Private/LandscapeEdit.cpp
--- a/Engine/Source/Runtime/Landscape/Private/LandscapeEdit.cpp
+++ b/Engine/Source/Runtime/Landscape/Private/LandscapeEdit.cpp
@@ -1491,15 +1491,15 @@ float ULandscapeComponent::GetLayerWeightAtLocation(const FVector& InLocation, U
 	}
 
 	// Fill the cache if necessary
-	if (LayerCache->Num() == 0)
-	{
-		FLandscapeComponentDataInterface CDI(this);
-		if (!CDI.GetWeightmapTextureData(LayerInfo, *LayerCache))
+		if (LayerCache->Num() == 0)
 		{
-			// no data for this layer for this component.
-			return 0.f;
+		FLandscapeComponentDataInterface CDI(this);
+			if (!CDI.GetWeightmapTextureData(LayerInfo, *LayerCache))
+			{
+				// no data for this layer for this component.
+				return 0.f;
+			}
 		}
-	}
 
 	// Find location
 	// TODO: Root landscape isn't always loaded, would Proxy suffice?
@@ -1528,14 +1528,14 @@ float ULandscapeComponent::GetLayerWeightAtLocation(const FVector& InLocation, U
 	float Sample12 = (float)((*LayerCache)[IdxX1 + Stride*IdxY2]) / 255.f;
 	float Sample22 = (float)((*LayerCache)[IdxX2 + Stride*IdxY2]) / 255.f;
 
-	int32 LerpX = FMath::Fractional(TestX);
-	int32 LerpY = FMath::Fractional(TestY);
+	float LerpX = FMath::Fractional(TestX);
+	float LerpY = FMath::Fractional(TestY);
 
-	// Bilinear interpolate
+		// Bilinear interpolate
 	return FMath::Lerp(
-		FMath::Lerp(Sample11, Sample21, LerpX),
-		FMath::Lerp(Sample12, Sample22, LerpX),
-		LerpY);
+			FMath::Lerp(Sample11, Sample21, LerpX),
+			FMath::Lerp(Sample12, Sample22, LerpX),
+			LerpY);
 }
 
 void ULandscapeComponent::GetComponentExtent(int32& MinX, int32& MinY, int32& MaxX, int32& MaxY) const
@@ -3410,6 +3410,12 @@ ULandscapeLayerInfoObject::ULandscapeLayerInfoObject(const FObjectInitializer& O
 	: Super(ObjectInitializer)
 {
 	Hardness = 0.5f;
+	GrassDensity = 400;
+	StartCullDistance = 10000.0f;
+	EndCullDistance = 10000.0f;
+	PlacementJitter = 1.0f;
+	RandomRotation = true;
+	AlignToSurface = true;
 #if WITH_EDITORONLY_DATA
 	bNoWeightBlend = false;
 #endif // WITH_EDITORONLY_DATA
@@ -3420,6 +3426,12 @@ void ULandscapeLayerInfoObject::PostEditChangeProperty(FPropertyChangedEvent& Pr
 {
 	static const FName NAME_Hardness = FName(TEXT("Hardness"));
 	static const FName NAME_PhysMaterial = FName(TEXT("PhysMaterial"));
+	static const FName NAME_GrassMesh = FName(TEXT("GrassMesh"));
+	static const FName NAME_GrassDensity = FName(TEXT("GrassDensity"));
+	static const FName NAME_PlacementJitter = FName(TEXT("PlacementJitter"));
+	static const FName NAME_EndCullDistance = FName(TEXT("EndCullDistance"));
+	static const FName NAME_RandomRotation = FName(TEXT("RandomRotation"));
+	static const FName NAME_AlignToSurface = FName(TEXT("AlignToSurface"));
 
 	Super::PostEditChangeProperty(PropertyChangedEvent);
 
@@ -3451,6 +3463,33 @@ void ULandscapeLayerInfoObject::PostEditChangeProperty(FPropertyChangedEvent& Pr
 				}
 			}
 		}
+		else if (
+			PropertyName == NAME_GrassMesh ||
+			PropertyName == NAME_GrassDensity ||
+			PropertyName == NAME_PlacementJitter ||
+			PropertyName == NAME_EndCullDistance ||
+			PropertyName == NAME_AlignToSurface ||
+			PropertyName == NAME_RandomRotation
+			)
+		{
+			// Only care current world object
+			for (TActorIterator<ALandscapeProxy> It(GWorld); It; ++It)
+			{
+				ALandscapeProxy* Proxy = *It;
+				ULandscapeInfo* Info = Proxy->GetLandscapeInfo(false);
+				if (Info)
+				{
+					for (int32 i = 0; i < Info->Layers.Num(); ++i)
+					{
+						if (Info->Layers[i].LayerInfoObj == this)
+						{
+							Proxy->FlushFoliageComponents();
+							break;
+						}
+					}
+				}
+			}
+		}
 	}
 }
 
@@ -3496,6 +3535,9 @@ void ALandscapeProxy::RemoveXYOffsets()
 
 void ALandscapeProxy::RecreateCollisionComponents()
 {
+	// We can assume these are all junk; they recreate as needed
+	FlushFoliageComponents();
+
 	// Clear old CollisionComponent containers
 	CollisionComponents.Empty();
 
