UE4 - implement cooking for grass, plus a few other grass tweaks

diff --git a/Engine/Source/Runtime/Core/Public/UObject/ObjectVersion.h b/Engine/Source/Runtime/Core/Public/UObject/ObjectVersion.h
--- a/Engine/Source/Runtime/Core/Public/UObject/ObjectVersion.h
+++ b/Engine/Source/Runtime/Core/Public/UObject/ObjectVersion.h
@@ -446,6 +446,8 @@ enum EUnrealEngineObjectUE4Version
 	VER_UE4_PERFRAME_MATERIAL_UNIFORM_EXPRESSIONS,
 	// Make MikkTSpace the default tangent space calculation method for static meshes.
 	VER_UE4_MIKKTSPACE_IS_DEFAULT,
+	// Only applies to cooked files, grass cooking support.
+	VER_UE4_LANDSCAPE_GRASS_COOKING,
 
 	// -----<new versions can be added before this line>-------------------------------------------------
 	// - this needs to be the last line (see note below)
diff --git a/Engine/Source/Runtime/Engine/Classes/Components/HierarchicalInstancedStaticMeshComponent.h b/Engine/Source/Runtime/Engine/Classes/Components/HierarchicalInstancedStaticMeshComponent.h
--- a/Engine/Source/Runtime/Engine/Classes/Components/HierarchicalInstancedStaticMeshComponent.h
+++ b/Engine/Source/Runtime/Engine/Classes/Components/HierarchicalInstancedStaticMeshComponent.h
@@ -100,7 +100,8 @@ public:
 		const FBox& MeshBox,
 		TArray<FClusterNode>& OutClusterTree,
 		TArray<int32>& OutSortedInstances,
-		TArray<int32>& OutInstanceReorderTable
+		TArray<int32>& OutInstanceReorderTable,
+		int32 MaxInstancesPerLeaf
 		);
 	void AcceptPrebuiltTree(
 		TArray<FClusterNode>& InClusterTree,
diff --git a/Engine/Source/Runtime/Engine/Private/HierarchicalInstancedStaticMesh.cpp b/Engine/Source/Runtime/Engine/Private/HierarchicalInstancedStaticMesh.cpp
--- a/Engine/Source/Runtime/Engine/Private/HierarchicalInstancedStaticMesh.cpp
+++ b/Engine/Source/Runtime/Engine/Private/HierarchicalInstancedStaticMesh.cpp
@@ -84,6 +84,8 @@ class FClusterBuilder
 	int32 Num;
 	FBox InstBox;
 	int32 BranchingFactor;
+	int32 InternalNodeBranchingFactor;
+	int32 MaxInstancesPerLeaf;
 	int32 NumRoots;
 	TArray<int32> SortIndex;
 	TArray<FVector> SortPoints;
@@ -201,12 +203,20 @@ class FClusterBuilder
 public:
 	FClusterTree* Result;
 
-	FClusterBuilder(TArray<FMatrix>& InTransforms, const FBox& InInstBox)
+	FClusterBuilder(TArray<FMatrix>& InTransforms, const FBox& InInstBox, int32 InMaxInstancesPerLeaf = 0)
 		: Num(InTransforms.Num())
 		, InstBox(InInstBox)
 		, Result(nullptr)
 	{
-		BranchingFactor = CVarFoliageSplitFactor.GetValueOnAnyThread();
+		InternalNodeBranchingFactor = CVarFoliageSplitFactor.GetValueOnAnyThread();
+		if (InMaxInstancesPerLeaf <= 0)
+		{
+			MaxInstancesPerLeaf = InternalNodeBranchingFactor;
+		}
+		else
+		{
+			MaxInstancesPerLeaf = InMaxInstancesPerLeaf;
+		}
 
 		Exchange(Transforms, InTransforms);
 		SortIndex.AddUninitialized(Num);
@@ -226,7 +236,9 @@ public:
 
 	void Build()
 	{
+		BranchingFactor = MaxInstancesPerLeaf;
 		Split(Num);
+		BranchingFactor = InternalNodeBranchingFactor;
 
 		Result = new FClusterTree;
 		TArray<int32>& SortedInstances = Result->SortedInstances;
@@ -1508,7 +1520,8 @@ void UHierarchicalInstancedStaticMeshComponent::BuildTreeAnyThread(
 	const FBox& MeshBox,
 	TArray<FClusterNode>& OutClusterTree,
 	TArray<int32>& OutSortedInstances,
-	TArray<int32>& OutInstanceReorderTable
+	TArray<int32>& OutInstanceReorderTable,
+	int32 MaxInstancesPerLeaf
 	)
 {
 	TArray<FMatrix> InstanceTransforms;
@@ -1518,7 +1531,7 @@ void UHierarchicalInstancedStaticMeshComponent::BuildTreeAnyThread(
 		InstanceTransforms[Index] = PerInstanceSMData[Index].Transform;
 	}
 
-	TUniquePtr<FClusterBuilder> Builder(new FClusterBuilder(InstanceTransforms, MeshBox));
+	TUniquePtr<FClusterBuilder> Builder(new FClusterBuilder(InstanceTransforms, MeshBox, MaxInstancesPerLeaf));
 	Builder->Build();
 
 	Exchange(OutClusterTree, Builder->Result->Nodes);
diff --git a/Engine/Source/Runtime/Landscape/Classes/LandscapeComponent.h b/Engine/Source/Runtime/Landscape/Classes/LandscapeComponent.h
--- a/Engine/Source/Runtime/Landscape/Classes/LandscapeComponent.h
+++ b/Engine/Source/Runtime/Landscape/Classes/LandscapeComponent.h
@@ -81,10 +81,14 @@ struct FWeightmapLayerAllocationInfo
 	UPROPERTY()
 	uint8 WeightmapTextureChannel;
 
+	/** Only relevant in non-editor builds, this indicates which channel in the data array is this layer...must be > 1 to be valid, the first two are height **/
+	UPROPERTY()
+	uint8 GrassMapChannelIndex;
 
 	FWeightmapLayerAllocationInfo()
 		: WeightmapTextureIndex(0)
 		, WeightmapTextureChannel(0)
+		, GrassMapChannelIndex(0)
 	{
 	}
 
@@ -93,12 +97,25 @@ struct FWeightmapLayerAllocationInfo
 		:	LayerInfo(InLayerInfo)
 		,	WeightmapTextureIndex(255)	// Indicates an invalid allocation
 		,	WeightmapTextureChannel(255)
+		,	GrassMapChannelIndex(0) // Indicates an invalid allocation
 	{
 	}
 	
 	FName GetLayerName() const;
 };
 
+struct FGrassMap
+{
+	/** Total number of channels; at least 3 since we have at least the height (two bytes) and at least one layer channel **/
+	int32 ElementStride;
+	TArray<uint8> Data;
+	FGrassMap(int32 InElementStride)
+		: ElementStride(InElementStride)
+	{
+	}
+};
+
+
 UCLASS(hidecategories=(Display, Attachment, Physics, Debug, Collision, Movement, Rendering, PrimitiveComponent, Object, Transform), showcategories=("Rendering|Material"), MinimalAPI)
 class ULandscapeComponent : public UPrimitiveComponent
 {
@@ -239,6 +256,9 @@ public:
 	/** Platform Data where don't support texture sampling in vertex buffer */
 	FLandscapeComponentDerivedData PlatformData;
 
+	/** Grass data for cooked platforms, also used temporarily when cooking */
+	TSharedPtr<FGrassMap, ESPMode::ThreadSafe> GrassMap;
+
 	virtual ~ULandscapeComponent();
 
 	// Begin UObject interface.	
@@ -299,6 +319,11 @@ public:
 	
 	void GeneratePlatformVertexData();
 	UMaterialInstance* GeneratePlatformPixelData(TArray<class UTexture2D*>& WeightmapTextures, bool bIsCooking);
+
+	/** Creates and destroys grass data for cooked platforms, used temporarily when cooking */
+	void GenerateGrassMap();
+	void RemoveGrassMap();
+
 #endif
 
 	/** @todo document */
diff --git a/Engine/Source/Runtime/Landscape/Private/Landscape.cpp b/Engine/Source/Runtime/Landscape/Private/Landscape.cpp
--- a/Engine/Source/Runtime/Landscape/Private/Landscape.cpp
+++ b/Engine/Source/Runtime/Landscape/Private/Landscape.cpp
@@ -156,6 +156,10 @@ void ULandscapeComponent::Serialize(FArchive& Ar)
 				HeightmapTexture = NULL;
 			}
 		}
+		if (Ar.UE4Ver() >= VER_UE4_LANDSCAPE_GRASS_COOKING)
+		{
+			GenerateGrassMap();
+		}
 	}
 #endif
 
@@ -218,6 +222,34 @@ void ULandscapeComponent::Serialize(FArchive& Ar)
 				Ar << PlatformData;
 			}
 		}
+		if (Ar.UE4Ver() >= VER_UE4_LANDSCAPE_GRASS_COOKING)
+		{
+			// GrassMap
+			if (Ar.IsCooking() && Ar.IsSaving())
+			{
+				int32 NumChannels = GrassMap.IsValid() ? GrassMap->ElementStride : 0;
+				Ar << NumChannels;
+				if (NumChannels)
+				{
+					GrassMap->Data.BulkSerialize(Ar);
+				}
+#if WITH_EDITOR
+				RemoveGrassMap();
+#else
+				check(0); // why are we cooking without an editor
+#endif
+			}
+			else
+			{
+				int32 NumChannels = 0;
+				Ar << NumChannels;
+				if (NumChannels)
+				{
+					GrassMap = MakeShareable(new FGrassMap(NumChannels));
+					GrassMap->Data.BulkSerialize(Ar);
+				}
+			}
+		}
 	}
 #endif
 }
@@ -2070,8 +2102,116 @@ static TAutoConsoleVariable<float> CVarGuardBandMultiplier(
 	1.5f,
 	TEXT("Used to control discarding in the grass system. Approximate range, 1-2. Multiplied by the cull distance to control when we discard grass components."));
 
+static TAutoConsoleVariable<float> CVarMinInstancesPerLeaf(
+	TEXT("grass.MinInstancesPerLeaf"),
+	256,
+	TEXT("Used to control the depth of hierarchical culling tree for grass. Clamped to 16-32000; a reasonable range."));
+
+#if WITH_EDITOR
+class FEditorGrassLayerData
+{
+	int32 Stride;
+	TArray<FColor> LocalHeightCache;
+	TArray<uint8> LocalCache;
+public:
+	FEditorGrassLayerData(ALandscapeProxy* Landscape, ULandscapeComponent* Component, ULandscapeLayerInfoObject* LayerInfo, UHierarchicalInstancedStaticMeshComponent* HierarchicalInstancedStaticMeshComponent)
+	{
+		int32 SubsectionSizeQuads = Component->SubsectionSizeQuads;
+		int32 NumSubsections = Component->NumSubsections;
+		Stride = (SubsectionSizeQuads + 1) * NumSubsections;
+		if (Stride)
+		{
+			FLandscapeComponentDataInterface CDI(Component);
+			CDI.GetHeightmapTextureData(LocalHeightCache);
+			if (!CDI.GetWeightmapTextureData(LayerInfo, LocalCache) || !LocalCache.Num() || LocalCache.Num() != LocalHeightCache.Num())
+			{
+				// no data for this layer for this component.
+				LocalHeightCache.Empty();
+				LocalCache.Empty();
+			}
+		}
+	}
+	bool HasData() const
+	{
+		return Stride && LocalCache.Num() && LocalHeightCache.Num() && LocalCache.Num() == LocalHeightCache.Num();
+	}
+	FORCEINLINE float GetHeight(int32 IdxX, int32 IdxY)
+	{
+		const FColor& SampleVal = LocalHeightCache[IdxX + Stride*IdxY];
+		return LandscapeDataAccess::GetLocalHeight((uint16(SampleVal.R) << 8) + uint16(SampleVal.G));
+	}
+	FORCEINLINE float GetLayer(int32 IdxX, int32 IdxY)
+	{
+		return (float)(LocalCache[IdxX + Stride*IdxY]) / 255.f;
+	}
+};
+
+typedef FEditorGrassLayerData FGrassLayerData;
+#else
+class FGameGrassLayerData
+{
+	int32 Stride;
+	int32 LayerOffset;
+	int32 ElementStride;
+	TSharedPtr<FGrassMap, ESPMode::ThreadSafe> GrassMap;
+	const uint8* GrassMapPtr;
+public:
+	FGameGrassLayerData(ALandscapeProxy* Landscape, ULandscapeComponent* Component, ULandscapeLayerInfoObject* LayerInfo, UHierarchicalInstancedStaticMeshComponent* HierarchicalInstancedStaticMeshComponent)
+	{
+		int32 SubsectionSizeQuads = Component->SubsectionSizeQuads;
+		int32 NumSubsections = Component->NumSubsections;
+		Stride = (SubsectionSizeQuads + 1) * NumSubsections;
+		GrassMapPtr = nullptr;
+		LayerOffset = 0;
+		ElementStride = 0;
+		if (Stride)
+		{
+			GrassMap = Component->GrassMap;
+			if (GrassMap.IsValid())
+			{
+				check(GrassMap->Data.Num() == GrassMap->ElementStride * Stride * Stride);
+				int32 LayerIdx = INDEX_NONE;
+				for (int32 Idx = 0; Idx < Component->WeightmapLayerAllocations.Num(); Idx++)
+				{
+					if (Component->WeightmapLayerAllocations[Idx].LayerInfo == LayerInfo)
+					{
+						LayerIdx = Idx;
+						LayerOffset = int32(Component->WeightmapLayerAllocations[Idx].GrassMapChannelIndex);
+						check(LayerOffset > 1);
+						break;
+					}
+				}
+				if (LayerIdx >= 0 && LayerOffset > 1)
+				{
+					GrassMapPtr = &GrassMap->Data[0];
+					ElementStride = GrassMap->ElementStride;
+				}
+			}
+		}
+	}
+	bool HasData() const
+	{
+		return !!GrassMapPtr;
+	}
+	FORCEINLINE float GetHeight(int32 IdxX, int32 IdxY)
+	{
+		const uint8* Sample = &GrassMapPtr[(IdxX + Stride*IdxY) * ElementStride];
+		return LandscapeDataAccess::GetLocalHeight((uint16(Sample[0]) << 8) + uint16(Sample[1]));
+	}
+	FORCEINLINE float GetLayer(int32 IdxX, int32 IdxY)
+	{
+		return (float)(GrassMapPtr[(IdxX + Stride*IdxY) * ElementStride + LayerOffset]) / 255.f;
+	}
+};
+
+typedef FGameGrassLayerData FGrassLayerData;
+
+#endif
+
 struct FAsyncGrassBuilder
 {
+	FGrassLayerData LayerData;
+
 	bool bHaveValidData;
 	bool RandomRotation;
 	bool AlignToSurface;
@@ -2095,10 +2235,6 @@ struct FAsyncGrassBuilder
 	int32 NumSubsections;
 	int32 Stride;
 
-
-	TArray<FColor> LocalHeightCache;
-	TArray<uint8> LocalCache;
-
 	FRandomStream RandomStream;
 
 	FBox MeshBox;
@@ -2115,8 +2251,9 @@ struct FAsyncGrassBuilder
 	TArray<int32> InstanceReorderTable;
 
 	FAsyncGrassBuilder(ALandscapeProxy* Landscape, ULandscapeComponent* Component, ULandscapeLayerInfoObject* LayerInfo, UHierarchicalInstancedStaticMeshComponent* HierarchicalInstancedStaticMeshComponent)
+		: LayerData(Landscape, Component, LayerInfo, HierarchicalInstancedStaticMeshComponent)
 	{
-		bHaveValidData = true;
+		bHaveValidData = LayerData.HasData();
 		GrassDensity = LayerInfo->GrassDensity;
 		PlacementJitter = LayerInfo->PlacementJitter;
 		EndCullDistance = LayerInfo->EndCullDistance;
@@ -2129,19 +2266,6 @@ struct FAsyncGrassBuilder
 		const FRotator DrawRot = Landscape->GetActorRotation();
 		LandscapeToWorld = Landscape->GetRootComponent()->ComponentToWorld.ToMatrixNoScale();
 
-#if WITH_EDITOR
-		{
-			FLandscapeComponentDataInterface CDI(Component);
-			CDI.GetHeightmapTextureData(LocalHeightCache);
-			if (!CDI.GetWeightmapTextureData(LayerInfo, LocalCache) || !LocalCache.Num() || LocalCache.Num() != LocalHeightCache.Num())
-			{
-				// no data for this layer for this component.
-				bHaveValidData = false;
-			}
-		}
-#else
-		bHaveValidData = false;
-#endif
 		SectionBase = Component->GetSectionBase();
 		ComponentSizeQuads = Component->ComponentSizeQuads;
 		SubsectionSizeQuads = Component->SubsectionSizeQuads;
@@ -2266,12 +2390,12 @@ struct FAsyncGrassBuilder
 			}
 			{
 				BuildTime -= FPlatformTime::Seconds();
-				UHierarchicalInstancedStaticMeshComponent::BuildTreeAnyThread(PerInstanceSMData, MeshBox, ClusterTree, SortedInstances, InstanceReorderTable);
+				UHierarchicalInstancedStaticMeshComponent::BuildTreeAnyThread(PerInstanceSMData, MeshBox, ClusterTree, SortedInstances, InstanceReorderTable, FMath::Clamp<int32>(CVarMinInstancesPerLeaf.GetValueOnAnyThread(), 16, 256 * 128));
 				BuildTime += FPlatformTime::Seconds();
 			}
 		}
 	}
-	float GetLayerWeightAtLocationLocal(const FVector& InLocation, FVector* OutLocation)
+	FORCEINLINE_DEBUGGABLE float GetLayerWeightAtLocationLocal(const FVector& InLocation, FVector* OutLocation)
 	{
 		// Find location
 		float TestX = InLocation.X / DrawScale.X - (float)SectionBase.X;
@@ -2295,10 +2419,10 @@ struct FAsyncGrassBuilder
 		float Result;
 		{
 			// sample
-			float Sample11 = (float)((LocalCache)[IdxX1 + Stride*IdxY1]) / 255.f;
-			float Sample21 = (float)((LocalCache)[IdxX2 + Stride*IdxY1]) / 255.f;
-			float Sample12 = (float)((LocalCache)[IdxX1 + Stride*IdxY2]) / 255.f;
-			float Sample22 = (float)((LocalCache)[IdxX2 + Stride*IdxY2]) / 255.f;
+			float Sample11 = LayerData.GetLayer(IdxX1, IdxY1);
+			float Sample21 = LayerData.GetLayer(IdxX2, IdxY1);
+			float Sample12 = LayerData.GetLayer(IdxX1, IdxY2);
+			float Sample22 = LayerData.GetLayer(IdxX2, IdxY2);
 
 			// Bilinear interpolate
 			Result = FMath::Lerp(
@@ -2308,19 +2432,11 @@ struct FAsyncGrassBuilder
 		}
 
 		{
-			struct FLocal
-			{
-				static float Sample(TArray<FColor>* HeightCache, int32 IdxX, int32 IdxY, int32 Stride)
-				{
-					const FColor& SampleVal = (*HeightCache)[IdxX + Stride*IdxY];
-					return LandscapeDataAccess::GetLocalHeight((uint16(SampleVal.R) << 8) + uint16(SampleVal.G));
-				}
-			};
 			// sample
-			float Sample11 = FLocal::Sample(&LocalHeightCache, IdxX1, IdxY1, Stride);
-			float Sample21 = FLocal::Sample(&LocalHeightCache, IdxX2, IdxY1, Stride);
-			float Sample12 = FLocal::Sample(&LocalHeightCache, IdxX1, IdxY2, Stride);
-			float Sample22 = FLocal::Sample(&LocalHeightCache, IdxX2, IdxY2, Stride);
+			float Sample11 = LayerData.GetHeight(IdxX1, IdxY1);
+			float Sample21 = LayerData.GetHeight(IdxX2, IdxY1);
+			float Sample12 = LayerData.GetHeight(IdxX1, IdxY2);
+			float Sample22 = LayerData.GetHeight(IdxX2, IdxY2);
 
 			OutLocation->X = InLocation.X - DrawScale.X * float(LandscapeSectionOffset.X);
 			OutLocation->Y = InLocation.Y - DrawScale.Y * float(LandscapeSectionOffset.Y);
@@ -2338,10 +2454,6 @@ struct FAsyncGrassBuilder
 
 void ALandscapeProxy::Tick(float DeltaSeconds)
 {
-#if !WITH_EDITOR
-	return;
-#endif
-
 	Super::Tick(DeltaSeconds);
 
 	UWorld* World = GetWorld();
@@ -2557,6 +2669,9 @@ void ALandscapeProxy::Tick(float DeltaSeconds)
 	StartTime = FPlatformTime::Seconds();
 	int32 TotalInstances = 0;
 	TotalComponents = 0;
+	int64 TreeSize = 0;
+	int64 InstSize = 0;
+	int64 InstBufSize = 0;
 	for (int32 Index = 0; Index < AsyncFoliageTasks.Num(); Index++)
 	{
 		FAsyncTask<FAsyncGrassTask>* Task = AsyncFoliageTasks[Index];
@@ -2574,15 +2689,19 @@ void ALandscapeProxy::Tick(float DeltaSeconds)
 
 				FBoxSphereBounds WorldBounds = HierarchicalInstancedStaticMeshComponent->CalcBounds(HierarchicalInstancedStaticMeshComponent->ComponentToWorld);
 				UE_LOG(LogTemp, Display, TEXT("Comp %d instances at %.0f %.0f %.0f"), HierarchicalInstancedStaticMeshComponent->PerInstanceSMData.Num(), WorldBounds.Origin.X, WorldBounds.Origin.Y, WorldBounds.Origin.Z);
+				TreeSize += HierarchicalInstancedStaticMeshComponent->ClusterTree.GetAllocatedSize();
+				InstSize += HierarchicalInstancedStaticMeshComponent->PerInstanceSMData.GetAllocatedSize();
+				InstBufSize += HierarchicalInstancedStaticMeshComponent->PerInstanceSMData.Num() * 7 * 16; // there is no place where this constant exists.
 			}
 			delete Inner.Builder;
 			delete Task;
+			break; // one per frame is fine
 		}
 	}
 	if (TotalComponents)
 	{
 		double TotalTime = FPlatformTime::Seconds() - StartTime;
-		UE_LOG(LogTemp, Display, TEXT("Finished grass with  %d components and %d instances in %.1fms"), TotalComponents, TotalInstances, float(TotalTime) * 1000.0f);
+		UE_LOG(LogTemp, Display, TEXT("Finished grass with  %d components and %d instances in %.1fms      InstMem = %lldMB   TreeMem = %lldMB   InstBufMem = %lldMB"), TotalComponents, TotalInstances, float(TotalTime) * 1000.0f, 1 + InstSize / 1024 / 1024, 1 + TreeSize / 1024 / 1024, 1 + InstBufSize / 1024 / 1024);
 	}
 }
 
diff --git a/Engine/Source/Runtime/Landscape/Private/LandscapeEdit.cpp b/Engine/Source/Runtime/Landscape/Private/LandscapeEdit.cpp
--- a/Engine/Source/Runtime/Landscape/Private/LandscapeEdit.cpp
+++ b/Engine/Source/Runtime/Landscape/Private/LandscapeEdit.cpp
@@ -5222,6 +5222,80 @@ void ULandscapeComponent::GeneratePlatformVertexData()
 	PlatformData.InitializeFromUncompressedData(NewPlatformData);
 }
 
+void ULandscapeComponent::GenerateGrassMap()
+{
+	RemoveGrassMap();
+	int32 Stride = (SubsectionSizeQuads + 1) * NumSubsections;
+	if (Stride)
+	{
+		int32 NumChannels = 2; // these are for the height
+		for (FWeightmapLayerAllocationInfo& Layer : WeightmapLayerAllocations)
+		{
+			ULandscapeLayerInfoObject* LayerInfo = Layer.LayerInfo;
+
+			if (LayerInfo && LayerInfo->GrassMesh && LayerInfo->GrassDensity > 0.0f && LayerInfo->EndCullDistance > 0)
+			{
+				Layer.GrassMapChannelIndex = (uint8)NumChannels++;
+			}
+			else
+			{
+				check(!Layer.GrassMapChannelIndex);
+			}
+		}
+		if (NumChannels > 2)
+		{
+			GrassMap = MakeShareable(new FGrassMap(NumChannels));
+			TArray<FColor> LocalHeightCache;
+			FLandscapeComponentDataInterface CDI(this);
+			CDI.GetHeightmapTextureData(LocalHeightCache);
+			if (LocalHeightCache.Num())
+			{
+				check(LocalHeightCache.Num() == Stride * Stride);
+				TArray<uint8>& Data = GrassMap->Data;
+				Data.AddZeroed(Stride * Stride * NumChannels);
+				for (int32 Index = 0; Index < LocalHeightCache.Num(); Index++)
+				{
+					const FColor& SampleVal = LocalHeightCache[Index];
+					uint8* Dest = &Data[Index * NumChannels];
+					*Dest++ = SampleVal.R;
+					*Dest = SampleVal.G;
+				}
+				TArray<uint8> LocalCache;
+				for (FWeightmapLayerAllocationInfo& Layer : WeightmapLayerAllocations)
+				{
+					if (Layer.GrassMapChannelIndex)
+					{
+						ULandscapeLayerInfoObject* LayerInfo = Layer.LayerInfo;
+						check(LayerInfo && LayerInfo->GrassMesh && LayerInfo->GrassDensity > 0.0f && LayerInfo->EndCullDistance > 0);
+						LocalCache.Reset();
+						if (CDI.GetWeightmapTextureData(LayerInfo, LocalCache) && LocalCache.Num() && LocalCache.Num() == LocalHeightCache.Num())
+						{
+							for (int32 Index = 0; Index < LocalCache.Num(); Index++)
+							{
+								Data[Index * NumChannels + Layer.GrassMapChannelIndex] = LocalCache[Index];
+							}
+						}
+					}
+				}
+			}
+			else
+			{
+				RemoveGrassMap();
+			}
+		}
+	}
+}
+
+void ULandscapeComponent::RemoveGrassMap()
+{
+	GrassMap.Reset();
+	for (FWeightmapLayerAllocationInfo& Layer : WeightmapLayerAllocations)
+	{
+		Layer.GrassMapChannelIndex = 0;
+	}
+}
+
+
 UTexture2D* ALandscapeProxy::CreateLandscapeTexture(int32 InSizeX, int32 InSizeY, TextureGroup InLODGroup, ETextureSourceFormat InFormat, UObject* OptionalOverrideOuter) const
 {
 	UObject* TexOuter = OptionalOverrideOuter ? OptionalOverrideOuter : GetOutermost();
