UE4 - add greater granularity to grass

diff --git a/Engine/Source/Runtime/Landscape/Classes/LandscapeProxy.h b/Engine/Source/Runtime/Landscape/Classes/LandscapeProxy.h
--- a/Engine/Source/Runtime/Landscape/Classes/LandscapeProxy.h
+++ b/Engine/Source/Runtime/Landscape/Classes/LandscapeProxy.h
@@ -194,13 +194,17 @@ struct FCachedLandscapeFoliage
 {
 	struct FPerLayer
 	{
-		TWeakObjectPtr<UHierarchicalInstancedStaticMeshComponent> Foliage;
+		int32 SqrtSubsections;
+		int32 CachedMaxInstancesPerComponent;
+		TArray<TWeakObjectPtr<UHierarchicalInstancedStaticMeshComponent> > Foliage;
 		TWeakObjectPtr<ULandscapeLayerInfoObject> Layer;
 
-		FPerLayer(UHierarchicalInstancedStaticMeshComponent* InFoliage, ULandscapeLayerInfoObject* InLayer)
-			: Foliage(InFoliage)
+		FPerLayer(int32 InSqrtSubsections, int32 InCachedMaxInstancesPerComponent, ULandscapeLayerInfoObject* InLayer)
+			: SqrtSubsections(InSqrtSubsections)
+			, CachedMaxInstancesPerComponent(InCachedMaxInstancesPerComponent)
 			, Layer(InLayer)
 		{
+			Foliage.Reserve(SqrtSubsections * SqrtSubsections);
 		}
 	};
 	struct FPerComponent
diff --git a/Engine/Source/Runtime/Landscape/Private/Landscape.cpp b/Engine/Source/Runtime/Landscape/Private/Landscape.cpp
--- a/Engine/Source/Runtime/Landscape/Private/Landscape.cpp
+++ b/Engine/Source/Runtime/Landscape/Private/Landscape.cpp
@@ -2102,11 +2102,16 @@ static TAutoConsoleVariable<float> CVarGuardBandMultiplier(
 	1.5f,
 	TEXT("Used to control discarding in the grass system. Approximate range, 1-2. Multiplied by the cull distance to control when we discard grass components."));
 
-static TAutoConsoleVariable<float> CVarMinInstancesPerLeaf(
+static TAutoConsoleVariable<int32> CVarMinInstancesPerLeaf(
 	TEXT("grass.MinInstancesPerLeaf"),
 	256,
 	TEXT("Used to control the depth of hierarchical culling tree for grass. Clamped to 16-32000; a reasonable range."));
 
+static TAutoConsoleVariable<int32> CVarMaxInstancesPerComponent(
+	TEXT("grass.MaxInstancesPerComponent"),
+	65536,
+	TEXT("Used to control the number of hierarchical components created. More can be more efficient, but can be hitchy as new components come into range"));
+
 #if WITH_EDITOR
 class FEditorGrassLayerData
 {
@@ -2208,21 +2213,13 @@ typedef FGameGrassLayerData FGrassLayerData;
 
 #endif
 
-struct FAsyncGrassBuilder
+struct FGrassBuilderBase
 {
-	FGrassLayerData LayerData;
-
 	bool bHaveValidData;
-	bool RandomRotation;
-	bool AlignToSurface;
 	float GrassDensity;
-	float PlacementJitter;
-	int32 EndCullDistance;
-
 	FVector DrawScale;
 	FVector DrawLoc;
 	FMatrix LandscapeToWorld;
-	FMatrix XForm;
 
 	FIntPoint SectionBase;
 	FIntPoint LandscapeSectionOffset;
@@ -2235,36 +2232,14 @@ struct FAsyncGrassBuilder
 	int32 NumSubsections;
 	int32 Stride;
 
-	FRandomStream RandomStream;
-
-	FBox MeshBox;
-
-	double RasterTime;
-	double BuildTime;
-	double InstanceTime;
-	int32 TotalInstances;
-
-	// output
-	TArray<FInstancedStaticMeshInstanceData> PerInstanceSMData;
-	TArray<FClusterNode> ClusterTree;
-	TArray<int32> SortedInstances;
-	TArray<int32> InstanceReorderTable;
-
-	FAsyncGrassBuilder(ALandscapeProxy* Landscape, ULandscapeComponent* Component, ULandscapeLayerInfoObject* LayerInfo, UHierarchicalInstancedStaticMeshComponent* HierarchicalInstancedStaticMeshComponent)
-		: LayerData(Landscape, Component, LayerInfo, HierarchicalInstancedStaticMeshComponent)
+	FGrassBuilderBase(ALandscapeProxy* Landscape, ULandscapeComponent* Component, ULandscapeLayerInfoObject* LayerInfo, int32 SqrtSubsections = 1, int32 SubX = 0, int32 SubY = 0)
 	{
-		bHaveValidData = LayerData.HasData();
+		bHaveValidData = true;
 		GrassDensity = LayerInfo->GrassDensity;
-		PlacementJitter = LayerInfo->PlacementJitter;
-		EndCullDistance = LayerInfo->EndCullDistance;
-		RandomRotation = LayerInfo->RandomRotation;
-		AlignToSurface = LayerInfo->AlignToSurface;
 
 		DrawScale = Landscape->GetRootComponent()->RelativeScale3D;
 		DrawLoc = Landscape->GetActorLocation();
 		LandscapeSectionOffset = Landscape->LandscapeSectionOffset;
-		const FRotator DrawRot = Landscape->GetActorRotation();
-		LandscapeToWorld = Landscape->GetRootComponent()->ComponentToWorld.ToMatrixNoScale();
 
 		SectionBase = Component->GetSectionBase();
 		ComponentSizeQuads = Component->ComponentSizeQuads;
@@ -2272,8 +2247,6 @@ struct FAsyncGrassBuilder
 		NumSubsections = Component->NumSubsections;
 		Stride = (SubsectionSizeQuads + 1) * NumSubsections;
 
-		RandomStream.Initialize(HierarchicalInstancedStaticMeshComponent->InstancingRandomSeed);
-
 		Origin = FVector(DrawScale.X * float(SectionBase.X), DrawScale.Y * float(SectionBase.Y), 0.0f);
 		Extent = FVector(DrawScale.X * float(SectionBase.X + ComponentSizeQuads), DrawScale.Y * float(SectionBase.Y + ComponentSizeQuads), 0.0f) - Origin;
 
@@ -2283,8 +2256,54 @@ struct FAsyncGrassBuilder
 		{
 			bHaveValidData = false;
 		}
+		const FRotator DrawRot = Landscape->GetActorRotation();
+		LandscapeToWorld = Landscape->GetRootComponent()->ComponentToWorld.ToMatrixNoScale();
+
+		if (bHaveValidData && SqrtSubsections != 1)
+		{
+			check(SqrtMaxInstances > 2 * SqrtSubsections);
+			SqrtMaxInstances /= SqrtSubsections;
+			check(SqrtMaxInstances > 0);
+
+			Extent /= float(SqrtSubsections);
+			Origin += Extent * FVector(float(SubX), float(SubY), 0.0f);
+		}
+	}
+
+};
+
+struct FAsyncGrassBuilder : public FGrassBuilderBase
+{
+	FGrassLayerData LayerData;
+	bool RandomRotation;
+	bool AlignToSurface;
+	float PlacementJitter;
+	FRandomStream RandomStream;
+	FMatrix XForm;
+	FBox MeshBox;
+
+	double RasterTime;
+	double BuildTime;
+	double InstanceTime;
+	int32 TotalInstances;
+
+	// output
+	TArray<FInstancedStaticMeshInstanceData> PerInstanceSMData;
+	TArray<FClusterNode> ClusterTree;
+	TArray<int32> SortedInstances;
+	TArray<int32> InstanceReorderTable;
 
+	FAsyncGrassBuilder(ALandscapeProxy* Landscape, ULandscapeComponent* Component, ULandscapeLayerInfoObject* LayerInfo, UHierarchicalInstancedStaticMeshComponent* HierarchicalInstancedStaticMeshComponent, int32 SqrtSubsections, int32 SubX, int32 SubY)
+		: FGrassBuilderBase(Landscape, Component, LayerInfo, SqrtSubsections, SubX, SubY)
+		, LayerData(Landscape, Component, LayerInfo, HierarchicalInstancedStaticMeshComponent)
+	{
+		bHaveValidData = bHaveValidData && LayerData.HasData();
 		XForm = LandscapeToWorld * HierarchicalInstancedStaticMeshComponent->ComponentToWorld.ToMatrixWithScale().Inverse();
+		PlacementJitter = LayerInfo->PlacementJitter;
+		RandomRotation = LayerInfo->RandomRotation;
+		AlignToSurface = LayerInfo->AlignToSurface;
+
+		RandomStream.Initialize(HierarchicalInstancedStaticMeshComponent->InstancingRandomSeed);
 
 		MeshBox = LayerInfo->GrassMesh->GetBounds().GetBox();
 
@@ -2323,7 +2342,7 @@ struct FAsyncGrassBuilder
 
 					FInstanceLocal& Instance = Instances[InstanceIndex];
 					float Weight = GetLayerWeightAtLocationLocal(Location, &Instance.Pos);
-					Instance.bKeep = Weight >= RandomStream.GetFraction();
+					Instance.bKeep = Weight > 0.0f && Weight >= RandomStream.GetFraction();
 					if (Instance.bKeep)
 					{
 						NumKept++;
@@ -2489,6 +2508,7 @@ void ALandscapeProxy::Tick(float DeltaSeconds)
 	int32 TotalComponents = 0;
 
 	float GuardBand = CVarGuardBandMultiplier.GetValueOnAnyThread();
+	int32 MaxInstancesPerComponent = FMath::Max<int32>(1024, CVarMaxInstancesPerComponent.GetValueOnAnyThread());
 
 	for (int32 ComponentIndex = 0; ComponentIndex < LandscapeComponents.Num(); ComponentIndex++)
 	{
@@ -2534,12 +2554,32 @@ void ALandscapeProxy::Tick(float DeltaSeconds)
 					bool bFound = false;
 					for (const FCachedLandscapeFoliage::FPerLayer& Layer : ExistingCache->Layers)
 					{
-						if (Layer.Layer.Get() == LayerInfo && Layer.Foliage.Get())
+						if (Layer.Layer.Get() == LayerInfo && Layer.CachedMaxInstancesPerComponent == MaxInstancesPerComponent)
 						{
-							NewPerComponent.Layers.Add(Layer);
-							StillUsed.Add(Layer.Foliage.Get());
-							bFound = true;
-							break;
+							check(Layer.SqrtSubsections >= 1);
+							int32 NumValid = 0;
+							for (const TWeakObjectPtr<UHierarchicalInstancedStaticMeshComponent>& Comp : Layer.Foliage)
+							{
+								if (Comp.Get())
+								{
+									NumValid++;
+								}
+								else
+								{
+									NumValid = -1;
+									break;
+								}
+							}
+							if (NumValid == Layer.SqrtSubsections * Layer.SqrtSubsections)
+							{
+								NewPerComponent.Layers.Add(Layer);
+								for (const TWeakObjectPtr<UHierarchicalInstancedStaticMeshComponent>& Comp : Layer.Foliage)
+								{
+									StillUsed.Add(Comp.Get());
+								}
+								bFound = true;
+								break;
+							}
 						}
 					}
 					if (bFound || MinDistanceToComp > MustHaveDistance)
@@ -2548,72 +2588,90 @@ void ALandscapeProxy::Tick(float DeltaSeconds)
 					}
 				}
 
-				int32 FolSeed = FCrc::StrCrc32((LayerInfo->LayerName.ToString() + Component->GetName()).GetCharArray().GetData());
-				if (FolSeed == 0)
+				FGrassBuilderBase ForSubsectionMath(this, Component, LayerInfo);
+
+				int32 SqrtSubsections = 1;
+				
+				if (ForSubsectionMath.bHaveValidData && ForSubsectionMath.SqrtMaxInstances > 0)
 				{
-					FolSeed++;
+					SqrtSubsections = FMath::Clamp<int32>(FMath::CeilToInt(float(ForSubsectionMath.SqrtMaxInstances) / FMath::Sqrt((float)MaxInstancesPerComponent)), 1, 16);
 				}
 
-				UHierarchicalInstancedStaticMeshComponent* HierarchicalInstancedStaticMeshComponent = ConstructObject<UHierarchicalInstancedStaticMeshComponent>(UHierarchicalInstancedStaticMeshComponent::StaticClass(), this); //, NAME_None, RF_Transactional);
-				TotalComponents++;
-				HierarchicalInstancedStaticMeshComponent->Mobility = EComponentMobility::Static;
+				FCachedLandscapeFoliage::FPerLayer& NewLayer = NewPerComponent.Layers[NewPerComponent.Layers.Add(FCachedLandscapeFoliage::FPerLayer(SqrtSubsections, MaxInstancesPerComponent, LayerInfo))];
+
+				for (int32 SubX = 0; SubX < SqrtSubsections; SubX++)
+				{
+					for (int32 SubY = 0; SubY < SqrtSubsections; SubY++)
+					{
+						int32 FolSeed = FCrc::StrCrc32((LayerInfo->LayerName.ToString() + Component->GetName() + FString::Printf(TEXT("%d %d"), SubX, SubY)).GetCharArray().GetData());
+						if (FolSeed == 0)
+						{
+							FolSeed++;
+						}
+
+						UHierarchicalInstancedStaticMeshComponent* HierarchicalInstancedStaticMeshComponent = ConstructObject<UHierarchicalInstancedStaticMeshComponent>(UHierarchicalInstancedStaticMeshComponent::StaticClass(), this); //, NAME_None, RF_Transactional);
+						TotalComponents++;
+						HierarchicalInstancedStaticMeshComponent->Mobility = EComponentMobility::Static;
 
-				HierarchicalInstancedStaticMeshComponent->StaticMesh = LayerInfo->GrassMesh;
-				HierarchicalInstancedStaticMeshComponent->bSelectable = true;
-				HierarchicalInstancedStaticMeshComponent->bHasPerInstanceHitProxies = true;
-				static FName NoCollision(TEXT("NoCollision"));
-				HierarchicalInstancedStaticMeshComponent->SetCollisionProfileName(NoCollision);
-				HierarchicalInstancedStaticMeshComponent->bDisableCollision = true;
-				HierarchicalInstancedStaticMeshComponent->SetCanEverAffectNavigation(false);
-				HierarchicalInstancedStaticMeshComponent->InstancingRandomSeed = FolSeed;
-				HierarchicalInstancedStaticMeshComponent->InstanceStartCullDistance = LayerInfo->StartCullDistance;
-				HierarchicalInstancedStaticMeshComponent->InstanceEndCullDistance = LayerInfo->EndCullDistance;
+						HierarchicalInstancedStaticMeshComponent->StaticMesh = LayerInfo->GrassMesh;
+						HierarchicalInstancedStaticMeshComponent->bSelectable = false;
+						HierarchicalInstancedStaticMeshComponent->bHasPerInstanceHitProxies = true;
+						static FName NoCollision(TEXT("NoCollision"));
+						HierarchicalInstancedStaticMeshComponent->SetCollisionProfileName(NoCollision);
+						HierarchicalInstancedStaticMeshComponent->bDisableCollision = true;
+						HierarchicalInstancedStaticMeshComponent->SetCanEverAffectNavigation(false);
+						HierarchicalInstancedStaticMeshComponent->InstancingRandomSeed = FolSeed;
+						HierarchicalInstancedStaticMeshComponent->InstanceStartCullDistance = LayerInfo->StartCullDistance;
+						HierarchicalInstancedStaticMeshComponent->InstanceEndCullDistance = LayerInfo->EndCullDistance;
 
 
 #if 0
-				HierarchicalInstancedStaticMeshComponent->CastShadow                     = Settings->CastShadow;
-				HierarchicalInstancedStaticMeshComponent->bCastDynamicShadow             = Settings->bCastDynamicShadow;
-				HierarchicalInstancedStaticMeshComponent->bCastStaticShadow              = Settings->bCastStaticShadow;
-				HierarchicalInstancedStaticMeshComponent->bAffectDynamicIndirectLighting = Settings->bAffectDynamicIndirectLighting;
-				HierarchicalInstancedStaticMeshComponent->bCastShadowAsTwoSided          = Settings->bCastShadowAsTwoSided;
-				HierarchicalInstancedStaticMeshComponent->bReceivesDecals                = Settings->bReceivesDecals;
-				HierarchicalInstancedStaticMeshComponent->bOverrideLightMapRes           = Settings->bOverrideLightMapRes;
-				HierarchicalInstancedStaticMeshComponent->OverriddenLightMapRes          = Settings->OverriddenLightMapRes;
-				HierarchicalInstancedStaticMeshComponent->BodyInstance.CopyBodyInstancePropertiesFrom(&Settings->BodyInstance);
+						HierarchicalInstancedStaticMeshComponent->CastShadow                     = Settings->CastShadow;
+						HierarchicalInstancedStaticMeshComponent->bCastDynamicShadow             = Settings->bCastDynamicShadow;
+						HierarchicalInstancedStaticMeshComponent->bCastStaticShadow              = Settings->bCastStaticShadow;
+						HierarchicalInstancedStaticMeshComponent->bAffectDynamicIndirectLighting = Settings->bAffectDynamicIndirectLighting;
+						HierarchicalInstancedStaticMeshComponent->bCastShadowAsTwoSided          = Settings->bCastShadowAsTwoSided;
+						HierarchicalInstancedStaticMeshComponent->bReceivesDecals                = Settings->bReceivesDecals;
+						HierarchicalInstancedStaticMeshComponent->bOverrideLightMapRes           = Settings->bOverrideLightMapRes;
+						HierarchicalInstancedStaticMeshComponent->OverriddenLightMapRes          = Settings->OverriddenLightMapRes;
+						HierarchicalInstancedStaticMeshComponent->BodyInstance.CopyBodyInstancePropertiesFrom(&Settings->BodyInstance);
 #endif
 
-				HierarchicalInstancedStaticMeshComponent->AttachTo(GetRootComponent());
-				FTransform DesiredTransform = GetRootComponent()->ComponentToWorld;
-				DesiredTransform.RemoveScaling();
-				HierarchicalInstancedStaticMeshComponent->SetWorldTransform(DesiredTransform);
+						HierarchicalInstancedStaticMeshComponent->AttachTo(GetRootComponent());
+						FTransform DesiredTransform = GetRootComponent()->ComponentToWorld;
+						DesiredTransform.RemoveScaling();
+						HierarchicalInstancedStaticMeshComponent->SetWorldTransform(DesiredTransform);
 
-				NewPerComponent.Layers.Add(FCachedLandscapeFoliage::FPerLayer(HierarchicalInstancedStaticMeshComponent, LayerInfo));
-				FoliageComponents.Add(HierarchicalInstancedStaticMeshComponent);
-				StillUsed.Add(HierarchicalInstancedStaticMeshComponent);
+						FoliageComponents.Add(HierarchicalInstancedStaticMeshComponent);
+						StillUsed.Add(HierarchicalInstancedStaticMeshComponent);
+						NewLayer.Foliage.Add(HierarchicalInstancedStaticMeshComponent);
 
 
-				FAsyncGrassBuilder* Builder = new FAsyncGrassBuilder(this, Component, LayerInfo, HierarchicalInstancedStaticMeshComponent);
+						FAsyncGrassBuilder* Builder = new FAsyncGrassBuilder(this, Component, LayerInfo, HierarchicalInstancedStaticMeshComponent, SqrtSubsections, SubX, SubY);
 
-				if (Builder->bHaveValidData)
-				{
-					BuildTime -= FPlatformTime::Seconds();
-					FAsyncTask<FAsyncGrassTask>* Task = new FAsyncTask<FAsyncGrassTask>(Builder, Component, HierarchicalInstancedStaticMeshComponent, LayerInfo);
+						if (Builder->bHaveValidData)
+						{
+							BuildTime -= FPlatformTime::Seconds();
+							FAsyncTask<FAsyncGrassTask>* Task = new FAsyncTask<FAsyncGrassTask>(Builder, Component, HierarchicalInstancedStaticMeshComponent, LayerInfo);
 
-					Task->StartBackgroundTask();
+							Task->StartBackgroundTask();
 
-					AsyncFoliageTasks.Add(Task);
+							AsyncFoliageTasks.Add(Task);
+
+							BuildTime += FPlatformTime::Seconds();
+						}
+						else
+						{
+							delete Builder;
+						}
 
-					BuildTime += FPlatformTime::Seconds();
-				}
-				else
-				{
-					delete Builder;
-				}
 
+						ComponentMiscTime -= FPlatformTime::Seconds();
+						HierarchicalInstancedStaticMeshComponent->RegisterComponent();
+						ComponentMiscTime += FPlatformTime::Seconds();
 
-				ComponentMiscTime -= FPlatformTime::Seconds();
-				HierarchicalInstancedStaticMeshComponent->RegisterComponent();
-				ComponentMiscTime += FPlatformTime::Seconds();
+					}
+				}
 			}
 		}
 		if (NewPerComponent.Layers.Num())
