Fix various landscape LOD discontinuities between neighbors and while texture streaming

diff --git a/Engine/Source/Editor/LandscapeEditor/Private/LandscapeEdModeComponentTools.cpp b/Engine/Source/Editor/LandscapeEditor/Private/LandscapeEdModeComponentTools.cpp
--- a/Engine/Source/Editor/LandscapeEditor/Private/LandscapeEdModeComponentTools.cpp
+++ b/Engine/Source/Editor/LandscapeEditor/Private/LandscapeEdModeComponentTools.cpp
@@ -865,30 +865,6 @@ public:
 
 			for (int32 Idx = 0; Idx < NewComponents.Num(); Idx++)
 			{
-				// Update LODBias from neighbors
-				FIntPoint ComponentBase = NewComponents[Idx]->GetSectionBase() / NewComponents[Idx]->ComponentSizeQuads;
-				FIntPoint LandscapeKey[8] =
-				{
-					ComponentBase + FIntPoint(-1, -1),
-					ComponentBase + FIntPoint(+0, -1),
-					ComponentBase + FIntPoint(+1, -1),
-					ComponentBase + FIntPoint(-1, +0),
-					ComponentBase + FIntPoint(+1, +0),
-					ComponentBase + FIntPoint(-1, +1),
-					ComponentBase + FIntPoint(+0, +1),
-					ComponentBase + FIntPoint(+1, +1)
-				};
-
-				for (int32 NeighborIdx = 0; NeighborIdx < 8; ++NeighborIdx)
-				{
-					ULandscapeComponent* Component = LandscapeInfo->XYtoComponentMap.FindRef(LandscapeKey[NeighborIdx]);
-					if (Component)
-					{
-						NewComponents[Idx]->NeighborLOD[NeighborIdx] = Component->ForcedLOD >= 0 ? Component->ForcedLOD : 255;
-						NewComponents[Idx]->NeighborLODBias[NeighborIdx] = Component->LODBias + 128;
-					}
-				}
-
 				// Update Collision
 				NewComponents[Idx]->UpdateCachedBounds();
 				NewComponents[Idx]->UpdateBounds();
@@ -1039,9 +1015,6 @@ public:
 					if (NeighborComp)
 					{
 						NeighborComp->Modify();
-						NeighborComp->NeighborLOD[7 - Idx] = 255; // Use 255 as unspecified value
-						NeighborComp->NeighborLODBias[7 - Idx] = 128;
-
 						NeighborComp->InvalidateLightingCache();
 						FComponentReregisterContext ReregisterContext(NeighborComp);
 					}
diff --git a/Engine/Source/Runtime/Landscape/Classes/LandscapeComponent.h b/Engine/Source/Runtime/Landscape/Classes/LandscapeComponent.h
--- a/Engine/Source/Runtime/Landscape/Classes/LandscapeComponent.h
+++ b/Engine/Source/Runtime/Landscape/Classes/LandscapeComponent.h
@@ -213,18 +213,10 @@ public:
 	UPROPERTY(EditAnywhere, BlueprintReadOnly, Category=LandscapeComponent)
 	int32 ForcedLOD;
 
-	/** Neighbor LOD data to use when rendering, 255 is unspecified */
-	UPROPERTY()
-	uint8 NeighborLOD[8];
-
 	/** LOD level Bias to use when rendering */
 	UPROPERTY(EditAnywhere, BlueprintReadOnly, Category=LandscapeComponent)
 	int32 LODBias;
 
-	/** Neighbor LOD bias to use when rendering, 128 is 0 bias, 0 is -128 bias, 255 is 127 bias */
-	UPROPERTY()
-	uint8 NeighborLODBias[8];
-
 	UPROPERTY()
 	FGuid StateId;
 
diff --git a/Engine/Source/Runtime/Landscape/Private/Landscape.cpp b/Engine/Source/Runtime/Landscape/Private/Landscape.cpp
--- a/Engine/Source/Runtime/Landscape/Private/Landscape.cpp
+++ b/Engine/Source/Runtime/Landscape/Private/Landscape.cpp
@@ -90,24 +90,7 @@ ULandscapeComponent::ULandscapeComponent(const FObjectInitializer& ObjectInitial
 
 	bBoundsChangeTriggersStreamingDataRebuild = true;
 	ForcedLOD = -1;
-	// Neighbor LOD and LODBias are saved in a legacy uint8 form instead of int8. Todo: convert
-	NeighborLOD[0] = 255;
-	NeighborLOD[1] = 255;
-	NeighborLOD[2] = 255;
-	NeighborLOD[3] = 255;
-	NeighborLOD[4] = 255;
-	NeighborLOD[5] = 255;
-	NeighborLOD[6] = 255;
-	NeighborLOD[7] = 255;
 	LODBias = 0;
-	NeighborLODBias[0] = 128;
-	NeighborLODBias[1] = 128;
-	NeighborLODBias[2] = 128;
-	NeighborLODBias[3] = 128;
-	NeighborLODBias[4] = 128;
-	NeighborLODBias[5] = 128;
-	NeighborLODBias[6] = 128;
-	NeighborLODBias[7] = 128;
 #if WITH_EDITORONLY_DATA
 	LightingLODBias = -1; // -1 Means automatic LOD calculation based on ForcedLOD + LODBias
 #endif
diff --git a/Engine/Source/Runtime/Landscape/Private/LandscapeEdit.cpp b/Engine/Source/Runtime/Landscape/Private/LandscapeEdit.cpp
--- a/Engine/Source/Runtime/Landscape/Private/LandscapeEdit.cpp
+++ b/Engine/Source/Runtime/Landscape/Private/LandscapeEdit.cpp
@@ -4008,16 +4008,6 @@ void ULandscapeComponent::SetLOD(bool bForcedLODChanged, int32 InLODValue)
 			if (Comp)
 			{
 				Comp->Modify();
-				if (bForcedLODChanged)
-				{
-					Comp->NeighborLOD[7 - Idx] = ForcedLOD >= 0 ? ForcedLOD : 255; // Use 255 as unspecified value
-				}
-				else
-				{
-					// Neighbor LODBias are saved in BYTE, so need to convert to range [-128:127]
-					Comp->NeighborLODBias[7 - Idx] = LODBias + 128;
-				}
-
 				Comp->InvalidateLightingCache();
 				FComponentReregisterContext ReregisterContext(Comp);
 			}
diff --git a/Engine/Source/Runtime/Landscape/Private/LandscapeRender.cpp b/Engine/Source/Runtime/Landscape/Private/LandscapeRender.cpp
--- a/Engine/Source/Runtime/Landscape/Private/LandscapeRender.cpp
+++ b/Engine/Source/Runtime/Landscape/Private/LandscapeRender.cpp
@@ -529,9 +529,11 @@ void FLandscapeEditToolRenderData::UpdateSelectionMaterial(int32 InSelectedType)
 //
 TMap<uint32, FLandscapeSharedBuffers*>FLandscapeComponentSceneProxy::SharedBuffersMap;
 TMap<uint32, FLandscapeSharedAdjacencyIndexBuffer*>FLandscapeComponentSceneProxy::SharedAdjacencyIndexBufferMap;
+TMap<FLandscapeComponentSceneProxy::FLandscapeKey, TMap<FIntPoint, const FLandscapeComponentSceneProxy*> > FLandscapeComponentSceneProxy::SharedSceneProxyMap;
 
 FLandscapeComponentSceneProxy::FLandscapeComponentSceneProxy(ULandscapeComponent* InComponent, FLandscapeEditToolRenderData* InEditToolRenderData)
 	: FPrimitiveSceneProxy(InComponent)
+	, LandscapeKey(InComponent->GetWorld(), InComponent->GetLandscapeProxy()->GetLandscapeGuid())
 	, MaxLOD(FMath::CeilLogTwo(InComponent->SubsectionSizeQuads + 1) - 1)
 	, NumSubsections(InComponent->NumSubsections)
 	, SubsectionSizeQuads(InComponent->SubsectionSizeQuads)
@@ -540,6 +542,7 @@ FLandscapeComponentSceneProxy::FLandscapeComponentSceneProxy(ULandscapeComponent
 	, ComponentSizeVerts(InComponent->ComponentSizeQuads + 1)
 	, StaticLightingLOD(InComponent->GetLandscapeProxy()->StaticLightingLOD)
 	, SectionBase(InComponent->GetSectionBase())
+	, ComponentBase(InComponent->GetSectionBase() / InComponent->ComponentSizeQuads)
 	, WeightmapScaleBias(InComponent->WeightmapScaleBias)
 	, WeightmapSubsectionOffset(InComponent->WeightmapSubsectionOffset)
 	, WeightmapTextures(InComponent->WeightmapTextures)
@@ -582,17 +585,10 @@ FLandscapeComponentSceneProxy::FLandscapeComponentSceneProxy(ULandscapeComponent
 	//    - - 3 - -   
 	//       +Y       
 
-	// Component->NeighborLOD and Component->NeighborLODBias are stored in a legacy uint8 form,
-	// so convert them here to something that matches ForcedLOD and LODBias (for this component)
-	ForcedNeighborLOD[0] = InComponent->NeighborLOD[1] != 255 ? InComponent->NeighborLOD[1] : -1;
-	ForcedNeighborLOD[1] = InComponent->NeighborLOD[3] != 255 ? InComponent->NeighborLOD[3] : -1;
-	ForcedNeighborLOD[2] = InComponent->NeighborLOD[4] != 255 ? InComponent->NeighborLOD[4] : -1;
-	ForcedNeighborLOD[3] = InComponent->NeighborLOD[6] != 255 ? InComponent->NeighborLOD[6] : -1;
-
-	NeighborLODBias[0] = InComponent->NeighborLODBias[1] - 128;
-	NeighborLODBias[1] = InComponent->NeighborLODBias[3] - 128;
-	NeighborLODBias[2] = InComponent->NeighborLODBias[4] - 128;
-	NeighborLODBias[3] = InComponent->NeighborLODBias[6] - 128;
+	Neighbors[0] = nullptr;
+	Neighbors[1] = nullptr;
+	Neighbors[2] = nullptr;
+	Neighbors[3] = nullptr;
 
 	LODBias = FMath::Clamp<int8>(LODBias, -MaxLOD, MaxLOD);
 
@@ -667,6 +663,43 @@ void FLandscapeComponentSceneProxy::CreateRenderThreadResources()
 {
 	check(HeightmapTexture != nullptr);
 
+	// Register ourselves in the map.
+	TMap<FIntPoint, const FLandscapeComponentSceneProxy*>& SceneProxyMap = SharedSceneProxyMap.FindOrAdd(LandscapeKey);
+	
+	const FLandscapeComponentSceneProxy* Existing = SceneProxyMap.FindRef(ComponentBase);
+	if (/*ensure*/(Existing == nullptr))
+	{
+		SceneProxyMap.Add(ComponentBase, this);
+
+		// Find Neighbors
+		Neighbors[0] = SceneProxyMap.FindRef(ComponentBase + FIntPoint(0, -1));
+		Neighbors[1] = SceneProxyMap.FindRef(ComponentBase + FIntPoint(-1, 0));
+		Neighbors[2] = SceneProxyMap.FindRef(ComponentBase + FIntPoint(1, 0));
+		Neighbors[3] = SceneProxyMap.FindRef(ComponentBase + FIntPoint(0, 1));
+
+		// Add ourselves to our neighbors
+		if (Neighbors[0])
+		{
+			Neighbors[0]->Neighbors[3] = this;
+		}
+		if (Neighbors[1])
+		{
+			Neighbors[1]->Neighbors[2] = this;
+		}
+		if (Neighbors[2])
+		{
+			Neighbors[2]->Neighbors[1] = this;
+		}
+		if (Neighbors[3])
+		{
+			Neighbors[3]->Neighbors[0] = this;
+		}
+	}
+	else
+	{
+		// UE_LOG(LogLandscape, Warning, TEXT("Duplicate ComponentBase %d, %d"), ComponentBase.X, ComponentBase.Y);
+	}
+
 	auto FeatureLevel = GetScene().GetFeatureLevel();
 
 	SharedBuffers = FLandscapeComponentSceneProxy::SharedBuffersMap.FindRef(SharedBuffersKey);
@@ -732,6 +765,42 @@ void FLandscapeComponentSceneProxy::CreateRenderThreadResources()
 
 FLandscapeComponentSceneProxy::~FLandscapeComponentSceneProxy()
 {
+	// Remove ourselves from the map
+	TMap<FIntPoint, const FLandscapeComponentSceneProxy*>* SceneProxyMap = SharedSceneProxyMap.Find(LandscapeKey);
+	check(SceneProxyMap);
+
+	const FLandscapeComponentSceneProxy* MapEntry = SceneProxyMap->FindRef(ComponentBase);
+	if (/*ensure*/(MapEntry == this))
+	{
+		SceneProxyMap->Remove(ComponentBase);
+
+		if (SceneProxyMap->Num() == 0)
+		{
+			// remove the entire LandscapeKey entry as this is the last scene proxy
+			SharedSceneProxyMap.Remove(LandscapeKey);
+		}
+		else
+		{
+			// remove reference to us from our neighbors
+			if (Neighbors[0])
+			{
+				Neighbors[0]->Neighbors[3] = nullptr;
+			}
+			if (Neighbors[1])
+			{
+				Neighbors[1]->Neighbors[2] = nullptr;
+			}
+			if (Neighbors[2])
+			{
+				Neighbors[2]->Neighbors[1] = nullptr;
+			}
+			if (Neighbors[3])
+			{
+				Neighbors[3]->Neighbors[0] = nullptr;
+			}
+		}
+	}
+
 	// Free the subsection uniform buffer
 	LandscapeUniformShaderParameters.ReleaseResource();
 
@@ -1199,32 +1268,36 @@ float FLandscapeComponentSceneProxy::CalcDesiredLOD(const class FSceneView& View
 	// HeightmapTexture, LODBias, ForcedLOD are component-specific with neighbor lookup
 	const bool bIsInThisComponent = (SubX >= 0 && SubX < NumSubsections && SubY >= 0 && SubY < NumSubsections);
 
+	auto* SubsectionHeightmapTexture = HeightmapTexture;
 	int8 SubsectionForcedLOD = ForcedLOD;
 	int8 SubsectionLODBias = LODBias;
 
 	if (SubX < 0)
 	{
-		SubsectionForcedLOD = ForcedNeighborLOD[1];
-		SubsectionLODBias = NeighborLODBias[1];
+		SubsectionHeightmapTexture = Neighbors[1] ? Neighbors[1]->HeightmapTexture : nullptr;
+		SubsectionForcedLOD = Neighbors[1] ? Neighbors[1]->ForcedLOD : -1;
+		SubsectionLODBias   = Neighbors[1] ? Neighbors[1]->LODBias : 0;
 	}
 	else if (SubX >= NumSubsections)
 	{
-		SubsectionForcedLOD = ForcedNeighborLOD[2];
-		SubsectionLODBias = NeighborLODBias[2];
+		SubsectionHeightmapTexture = Neighbors[2] ? Neighbors[2]->HeightmapTexture : nullptr;
+		SubsectionForcedLOD = Neighbors[2] ? Neighbors[2]->ForcedLOD : -1;
+		SubsectionLODBias   = Neighbors[2] ? Neighbors[2]->LODBias : 0;
 	}
 	else if (SubY < 0)
 	{
-		SubsectionForcedLOD = ForcedNeighborLOD[0];
-		SubsectionLODBias = NeighborLODBias[0];
+		SubsectionHeightmapTexture = Neighbors[0] ? Neighbors[0]->HeightmapTexture : nullptr;
+		SubsectionForcedLOD = Neighbors[0] ? Neighbors[0]->ForcedLOD : -1;
+		SubsectionLODBias   = Neighbors[0] ? Neighbors[0]->LODBias : 0;
 	}
 	else if (SubY >= NumSubsections)
 	{
-		SubsectionForcedLOD = ForcedNeighborLOD[3];
-		SubsectionLODBias = NeighborLODBias[3];
+		SubsectionHeightmapTexture = Neighbors[3] ? Neighbors[3]->HeightmapTexture : nullptr;
+		SubsectionForcedLOD = Neighbors[3] ? Neighbors[3]->ForcedLOD : -1;
+		SubsectionLODBias   = Neighbors[3] ? Neighbors[3]->LODBias : 0;
 	}
 
-	// TODO - this clamping isn't done for neighbors, but should be
-	const int32 MinLOD = (bIsInThisComponent && HeightmapTexture) ? FMath::Min<int32>(HeightmapTexture->GetNumMips() - HeightmapTexture->ResidentMips, MaxLOD) : 0;
+	const int32 MinStreamedLOD = SubsectionHeightmapTexture ? ((FTexture2DResource*)SubsectionHeightmapTexture->Resource)->GetCurrentFirstMip() : 0;
 
 	float fLOD = FLT_MAX;
 
@@ -1275,7 +1348,10 @@ float FLandscapeComponentSceneProxy::CalcDesiredLOD(const class FSceneView& View
 		}
 	}
 
-	fLOD = FMath::Clamp<float>(fLOD, FMath::Max<int32>(SubsectionLODBias, MinLOD), FMath::Min<int32>(MaxLOD, MaxLOD + SubsectionLODBias));
+	fLOD = FMath::Clamp<float>(fLOD, SubsectionLODBias, FMath::Min<int32>(MaxLOD, MaxLOD + SubsectionLODBias));
+
+	// ultimately due to texture streaming we sometimes need to go past MaxLOD
+	fLOD = FMath::Max<float>(fLOD, MinStreamedLOD);
 
 	return fLOD;
 }
@@ -2101,10 +2177,8 @@ public:
 			FVector4 LodBias(
 				0.0f, // unused
 				0.0f, // unused
-				SceneProxy->HeightmapTexture->GetNumMips() - FMath::Min(SceneProxy->HeightmapTexture->ResidentMips, SceneProxy->HeightmapTexture->RequestedMips),
-				SceneProxy->XYOffsetmapTexture
-					? SceneProxy->XYOffsetmapTexture->GetNumMips() - FMath::Min(SceneProxy->XYOffsetmapTexture->ResidentMips, SceneProxy->XYOffsetmapTexture->RequestedMips)
-					: 0.0f
+				((FTexture2DResource*)SceneProxy->HeightmapTexture->Resource)->GetCurrentFirstMip(),
+				SceneProxy->XYOffsetmapTexture ? ((FTexture2DResource*)SceneProxy->XYOffsetmapTexture->Resource)->GetCurrentFirstMip() : 0.0f
 				);
 			SetShaderValue(RHICmdList, VertexShader->GetVertexShader(), LodBiasParameter, LodBias);
 		}
diff --git a/Engine/Source/Runtime/Landscape/Public/LandscapeRender.h b/Engine/Source/Runtime/Landscape/Public/LandscapeRender.h
--- a/Engine/Source/Runtime/Landscape/Public/LandscapeRender.h
+++ b/Engine/Source/Runtime/Landscape/Public/LandscapeRender.h
@@ -354,6 +354,28 @@ class FLandscapeComponentSceneProxy : public FPrimitiveSceneProxy
 {
 	friend class FLandscapeSharedBuffers;
 
+	// Key to uniquely identify the landscape to find the correct render proxy map
+	class FLandscapeKey
+	{
+		const UWorld* World;
+		const FGuid Guid;
+	public:
+		FLandscapeKey(const UWorld* InWorld, const FGuid& InGuid)
+		: World(InWorld)
+		, Guid(InGuid)
+		{}
+
+		friend inline uint32 GetTypeHash(const FLandscapeKey& LandscapeKey)
+		{
+			return HashCombine(GetTypeHash(LandscapeKey.World), GetTypeHash(LandscapeKey.Guid));
+		}
+
+		friend bool operator==(const FLandscapeKey& A, const FLandscapeKey& B)
+		{
+			return A.World == B.World && A.Guid == B.Guid;
+		}
+	};
+
 	class FLandscapeLCI final : public FLightCacheInterface
 	{
 	public:
@@ -389,6 +411,7 @@ class FLandscapeComponentSceneProxy : public FPrimitiveSceneProxy
 	};
 
 protected:
+	FLandscapeKey				LandscapeKey;
 	int8						MaxLOD;
 	int8						NumSubsections;
 	int16						SubsectionSizeQuads;
@@ -398,6 +421,7 @@ protected:
 	uint8						StaticLightingLOD;
 	float						StaticLightingResolution;
 	FIntPoint					SectionBase;
+	FIntPoint					ComponentBase;
 	FMatrix						LocalToWorldNoScaling;
 
 	// Storage for static draw list batch params
@@ -434,6 +458,9 @@ protected:
 	static TMap<uint32, FLandscapeSharedBuffers*> SharedBuffersMap;
 	static TMap<uint32, FLandscapeSharedAdjacencyIndexBuffer*> SharedAdjacencyIndexBufferMap;
 
+	// Map of currently registered landscape proxies, used to register with our neighbors
+	static TMap<FLandscapeKey, TMap<FIntPoint, const FLandscapeComponentSceneProxy*> > SharedSceneProxyMap;
+
 	FLandscapeEditToolRenderData* EditToolRenderData;
 
 	// FLightCacheInterface
@@ -443,8 +470,9 @@ protected:
 
 	int8					ForcedLOD;
 	int8					LODBias;
-	int8					ForcedNeighborLOD[LANDSCAPE_NEIGHBOR_NUM];
-	int8					NeighborLODBias[LANDSCAPE_NEIGHBOR_NUM];
+
+	// Pointer to our neighbor's scene proxies in NWES order (nullptr if there is currently no neighbor)
+	mutable const FLandscapeComponentSceneProxy* Neighbors[4];
 
 	ELandscapeLODFalloff::Type LODFalloff;
 
