Fix for Mac compilation

diff --git a/Engine/Source/Runtime/Landscape/Private/LandscapeRender.cpp b/Engine/Source/Runtime/Landscape/Private/LandscapeRender.cpp
--- a/Engine/Source/Runtime/Landscape/Private/LandscapeRender.cpp
+++ b/Engine/Source/Runtime/Landscape/Private/LandscapeRender.cpp
@@ -667,7 +667,7 @@ void FLandscapeComponentSceneProxy::CreateRenderThreadResources()
 	TMap<FIntPoint, const FLandscapeComponentSceneProxy*>& SceneProxyMap = SharedSceneProxyMap.FindOrAdd(LandscapeKey);
 	
 	const FLandscapeComponentSceneProxy* Existing = SceneProxyMap.FindRef(ComponentBase);
-	if (/*ensure*/(Existing == nullptr))
+	if (Existing == nullptr)//(ensure(Existing == nullptr))
 	{
 		SceneProxyMap.Add(ComponentBase, this);
 
@@ -770,7 +770,7 @@ FLandscapeComponentSceneProxy::~FLandscapeComponentSceneProxy()
 	check(SceneProxyMap);
 
 	const FLandscapeComponentSceneProxy* MapEntry = SceneProxyMap->FindRef(ComponentBase);
-	if (/*ensure*/(MapEntry == this))
+	if (MapEntry == this) //(/*ensure*/(MapEntry == this))
 	{
 		SceneProxyMap->Remove(ComponentBase);
 
