UE4 - added grass enable cvar

diff --git a/Engine/Source/Runtime/Landscape/Private/Landscape.cpp b/Engine/Source/Runtime/Landscape/Private/Landscape.cpp
--- a/Engine/Source/Runtime/Landscape/Private/Landscape.cpp
+++ b/Engine/Source/Runtime/Landscape/Private/Landscape.cpp
@@ -2095,6 +2095,11 @@ static TAutoConsoleVariable<int32> CVarMaxInstancesPerComponent(
 	65536,
 	TEXT("Used to control the number of hierarchical components created. More can be more efficient, but can be hitchy as new components come into range"));
 
+static TAutoConsoleVariable<int32> CVarGrassEnable(
+	TEXT("grass.Enable"),
+	1,
+	TEXT("1: Enable Grass; 0: Disable Grass"));
+
 #if WITH_EDITOR
 class FEditorGrassLayerData
 {
@@ -2484,193 +2489,196 @@ void ALandscapeProxy::Tick(float DeltaSeconds)
 	FCachedLandscapeFoliage DesiredFoliage;
 	TSet<UHierarchicalInstancedStaticMeshComponent *> StillUsed;
 
-	double ComponentTime = 0.0;
-	double BuildTime = 0.0;
-	double ComponentMiscTime = 0.0;
-	double StartTime = FPlatformTime::Seconds();
-	int32 TotalComponents = 0;
-
-	float GuardBand = CVarGuardBandMultiplier.GetValueOnAnyThread();
-	int32 MaxInstancesPerComponent = FMath::Max<int32>(1024, CVarMaxInstancesPerComponent.GetValueOnAnyThread());
-
-	for (int32 ComponentIndex = 0; ComponentIndex < LandscapeComponents.Num(); ComponentIndex++)
+	if (CVarGrassEnable.GetValueOnAnyThread() > 0)
 	{
-		ComponentTime -= FPlatformTime::Seconds();
-		ULandscapeComponent* Component = LandscapeComponents[ComponentIndex];
+		double ComponentTime = 0.0;
+		double BuildTime = 0.0;
+		double ComponentMiscTime = 0.0;
+		double StartTime = FPlatformTime::Seconds();
+		int32 TotalComponents = 0;
+
+		float GuardBand = CVarGuardBandMultiplier.GetValueOnAnyThread();
+		int32 MaxInstancesPerComponent = FMath::Max<int32>(1024, CVarMaxInstancesPerComponent.GetValueOnAnyThread());
 
-		const FCachedLandscapeFoliage::FPerComponent* ExistingCache = nullptr;
-		for (const FCachedLandscapeFoliage::FPerComponent& CacheItem : FoliageCache.PerComponent)
+		for (int32 ComponentIndex = 0; ComponentIndex < LandscapeComponents.Num(); ComponentIndex++)
 		{
-			if (CacheItem.BasedOn.Get() == Component)
+			ComponentTime -= FPlatformTime::Seconds();
+			ULandscapeComponent* Component = LandscapeComponents[ComponentIndex];
+
+			const FCachedLandscapeFoliage::FPerComponent* ExistingCache = nullptr;
+			for (const FCachedLandscapeFoliage::FPerComponent& CacheItem : FoliageCache.PerComponent)
 			{
-				ExistingCache = &CacheItem;
-				break;
+				if (CacheItem.BasedOn.Get() == Component)
+				{
+					ExistingCache = &CacheItem;
+					break;
+				}
 			}
-		}
-		FCachedLandscapeFoliage::FPerComponent NewPerComponent(Component);
+			FCachedLandscapeFoliage::FPerComponent NewPerComponent(Component);
 
-		FBoxSphereBounds WorldBounds = Component->CalcBounds(Component->ComponentToWorld);
-		float MinDistanceToComp = MAX_flt;
+			FBoxSphereBounds WorldBounds = Component->CalcBounds(Component->ComponentToWorld);
+			float MinDistanceToComp = MAX_flt;
 
-		for (auto& Pos : Cameras)
-		{
-			MinDistanceToComp = FMath::Min<float>(MinDistanceToComp, FVector::Dist(WorldBounds.Origin, Pos) - WorldBounds.SphereRadius);
-		}
-
-		for (FWeightmapLayerAllocationInfo& LayerAlloc : Component->WeightmapLayerAllocations)
-		{
-			ULandscapeLayerInfoObject* LayerInfo = LayerAlloc.LayerInfo;
+			for (auto& Pos : Cameras)
+			{
+				MinDistanceToComp = FMath::Min<float>(MinDistanceToComp, FVector::Dist(WorldBounds.Origin, Pos) - WorldBounds.SphereRadius);
+			}
 
-			if (LayerInfo && LayerInfo->GrassMesh && LayerInfo->GrassDensity > 0.0f && LayerInfo->EndCullDistance > 0)
+			for (FWeightmapLayerAllocationInfo& LayerAlloc : Component->WeightmapLayerAllocations)
 			{
-				bool bRandomRotation = LayerInfo->RandomRotation;
-				bool bAlign = LayerInfo->AlignToSurface;
-				float MustHaveDistance = (float)LayerInfo->EndCullDistance;
-				float DiscardDistance = MustHaveDistance * GuardBand;
+				ULandscapeLayerInfoObject* LayerInfo = LayerAlloc.LayerInfo;
 
-				if (MinDistanceToComp > DiscardDistance)
-				{
-					continue;
-				}
-				if (ExistingCache)
+				if (LayerInfo && LayerInfo->GrassMesh && LayerInfo->GrassDensity > 0.0f && LayerInfo->EndCullDistance > 0)
 				{
-					bool bFound = false;
-					for (const FCachedLandscapeFoliage::FPerLayer& Layer : ExistingCache->Layers)
+					bool bRandomRotation = LayerInfo->RandomRotation;
+					bool bAlign = LayerInfo->AlignToSurface;
+					float MustHaveDistance = (float)LayerInfo->EndCullDistance;
+					float DiscardDistance = MustHaveDistance * GuardBand;
+
+					if (MinDistanceToComp > DiscardDistance)
 					{
-						if (Layer.Layer.Get() == LayerInfo && Layer.CachedMaxInstancesPerComponent == MaxInstancesPerComponent)
+						continue;
+					}
+					if (ExistingCache)
+					{
+						bool bFound = false;
+						for (const FCachedLandscapeFoliage::FPerLayer& Layer : ExistingCache->Layers)
 						{
-							check(Layer.SqrtSubsections >= 1);
-							int32 NumValid = 0;
-							for (const TWeakObjectPtr<UHierarchicalInstancedStaticMeshComponent>& Comp : Layer.Foliage)
+							if (Layer.Layer.Get() == LayerInfo && Layer.CachedMaxInstancesPerComponent == MaxInstancesPerComponent)
 							{
-								if (Comp.Get())
+								check(Layer.SqrtSubsections >= 1);
+								int32 NumValid = 0;
+								for (const TWeakObjectPtr<UHierarchicalInstancedStaticMeshComponent>& Comp : Layer.Foliage)
 								{
-									NumValid++;
+									if (Comp.Get())
+									{
+										NumValid++;
+									}
+									else
+									{
+										NumValid = -1;
+										break;
+									}
 								}
-								else
+								if (NumValid == Layer.SqrtSubsections * Layer.SqrtSubsections)
 								{
-									NumValid = -1;
+									NewPerComponent.Layers.Add(Layer);
+									for (const TWeakObjectPtr<UHierarchicalInstancedStaticMeshComponent>& Comp : Layer.Foliage)
+									{
+										StillUsed.Add(Comp.Get());
+									}
+									bFound = true;
 									break;
 								}
 							}
-							if (NumValid == Layer.SqrtSubsections * Layer.SqrtSubsections)
-							{
-								NewPerComponent.Layers.Add(Layer);
-								for (const TWeakObjectPtr<UHierarchicalInstancedStaticMeshComponent>& Comp : Layer.Foliage)
-								{
-									StillUsed.Add(Comp.Get());
-								}
-								bFound = true;
-								break;
-							}
+						}
+						if (bFound || MinDistanceToComp > MustHaveDistance)
+						{
+							continue;
 						}
 					}
-					if (bFound || MinDistanceToComp > MustHaveDistance)
-					{
-						continue;
-					}
-				}
 
-				FGrassBuilderBase ForSubsectionMath(this, Component, LayerInfo);
+					FGrassBuilderBase ForSubsectionMath(this, Component, LayerInfo);
 
-				int32 SqrtSubsections = 1;
+					int32 SqrtSubsections = 1;
 				
-				if (ForSubsectionMath.bHaveValidData && ForSubsectionMath.SqrtMaxInstances > 0)
-				{
-					SqrtSubsections = FMath::Clamp<int32>(FMath::CeilToInt(float(ForSubsectionMath.SqrtMaxInstances) / FMath::Sqrt((float)MaxInstancesPerComponent)), 1, 16);
-				}
+					if (ForSubsectionMath.bHaveValidData && ForSubsectionMath.SqrtMaxInstances > 0)
+					{
+						SqrtSubsections = FMath::Clamp<int32>(FMath::CeilToInt(float(ForSubsectionMath.SqrtMaxInstances) / FMath::Sqrt((float)MaxInstancesPerComponent)), 1, 16);
+					}
 
-				FCachedLandscapeFoliage::FPerLayer& NewLayer = NewPerComponent.Layers[NewPerComponent.Layers.Add(FCachedLandscapeFoliage::FPerLayer(SqrtSubsections, MaxInstancesPerComponent, LayerInfo))];
+					FCachedLandscapeFoliage::FPerLayer& NewLayer = NewPerComponent.Layers[NewPerComponent.Layers.Add(FCachedLandscapeFoliage::FPerLayer(SqrtSubsections, MaxInstancesPerComponent, LayerInfo))];
 
-				for (int32 SubX = 0; SubX < SqrtSubsections; SubX++)
-				{
-					for (int32 SubY = 0; SubY < SqrtSubsections; SubY++)
+					for (int32 SubX = 0; SubX < SqrtSubsections; SubX++)
 					{
-						int32 FolSeed = FCrc::StrCrc32((LayerInfo->LayerName.ToString() + Component->GetName() + FString::Printf(TEXT("%d %d"), SubX, SubY)).GetCharArray().GetData());
-						if (FolSeed == 0)
+						for (int32 SubY = 0; SubY < SqrtSubsections; SubY++)
 						{
-							FolSeed++;
-						}
-
-						UHierarchicalInstancedStaticMeshComponent* HierarchicalInstancedStaticMeshComponent = ConstructObject<UHierarchicalInstancedStaticMeshComponent>(UHierarchicalInstancedStaticMeshComponent::StaticClass(), this); //, NAME_None, RF_Transactional);
-						TotalComponents++;
-						HierarchicalInstancedStaticMeshComponent->Mobility = EComponentMobility::Static;
-
-						HierarchicalInstancedStaticMeshComponent->StaticMesh = LayerInfo->GrassMesh;
-						HierarchicalInstancedStaticMeshComponent->bSelectable = false;
-						HierarchicalInstancedStaticMeshComponent->bHasPerInstanceHitProxies = true;
-						static FName NoCollision(TEXT("NoCollision"));
-						HierarchicalInstancedStaticMeshComponent->SetCollisionProfileName(NoCollision);
-						HierarchicalInstancedStaticMeshComponent->bDisableCollision = true;
-						HierarchicalInstancedStaticMeshComponent->SetCanEverAffectNavigation(false);
-						HierarchicalInstancedStaticMeshComponent->InstancingRandomSeed = FolSeed;
-						HierarchicalInstancedStaticMeshComponent->InstanceStartCullDistance = LayerInfo->StartCullDistance;
-						HierarchicalInstancedStaticMeshComponent->InstanceEndCullDistance = LayerInfo->EndCullDistance;
-
-
-#if 0
-						HierarchicalInstancedStaticMeshComponent->CastShadow                     = Settings->CastShadow;
-						HierarchicalInstancedStaticMeshComponent->bCastDynamicShadow             = Settings->bCastDynamicShadow;
-						HierarchicalInstancedStaticMeshComponent->bCastStaticShadow              = Settings->bCastStaticShadow;
-						HierarchicalInstancedStaticMeshComponent->bAffectDynamicIndirectLighting = Settings->bAffectDynamicIndirectLighting;
-						HierarchicalInstancedStaticMeshComponent->bCastShadowAsTwoSided          = Settings->bCastShadowAsTwoSided;
-						HierarchicalInstancedStaticMeshComponent->bReceivesDecals                = Settings->bReceivesDecals;
-						HierarchicalInstancedStaticMeshComponent->bOverrideLightMapRes           = Settings->bOverrideLightMapRes;
-						HierarchicalInstancedStaticMeshComponent->OverriddenLightMapRes          = Settings->OverriddenLightMapRes;
-						HierarchicalInstancedStaticMeshComponent->BodyInstance.CopyBodyInstancePropertiesFrom(&Settings->BodyInstance);
-#endif
-
-						HierarchicalInstancedStaticMeshComponent->AttachTo(GetRootComponent());
-						FTransform DesiredTransform = GetRootComponent()->ComponentToWorld;
-						DesiredTransform.RemoveScaling();
-						HierarchicalInstancedStaticMeshComponent->SetWorldTransform(DesiredTransform);
+							int32 FolSeed = FCrc::StrCrc32((LayerInfo->LayerName.ToString() + Component->GetName() + FString::Printf(TEXT("%d %d"), SubX, SubY)).GetCharArray().GetData());
+							if (FolSeed == 0)
+							{
+								FolSeed++;
+							}
 
-						FoliageComponents.Add(HierarchicalInstancedStaticMeshComponent);
-						StillUsed.Add(HierarchicalInstancedStaticMeshComponent);
-						NewLayer.Foliage.Add(HierarchicalInstancedStaticMeshComponent);
+							UHierarchicalInstancedStaticMeshComponent* HierarchicalInstancedStaticMeshComponent = ConstructObject<UHierarchicalInstancedStaticMeshComponent>(UHierarchicalInstancedStaticMeshComponent::StaticClass(), this); //, NAME_None, RF_Transactional);
+							TotalComponents++;
+							HierarchicalInstancedStaticMeshComponent->Mobility = EComponentMobility::Static;
+
+							HierarchicalInstancedStaticMeshComponent->StaticMesh = LayerInfo->GrassMesh;
+							HierarchicalInstancedStaticMeshComponent->bSelectable = false;
+							HierarchicalInstancedStaticMeshComponent->bHasPerInstanceHitProxies = true;
+							static FName NoCollision(TEXT("NoCollision"));
+							HierarchicalInstancedStaticMeshComponent->SetCollisionProfileName(NoCollision);
+							HierarchicalInstancedStaticMeshComponent->bDisableCollision = true;
+							HierarchicalInstancedStaticMeshComponent->SetCanEverAffectNavigation(false);
+							HierarchicalInstancedStaticMeshComponent->InstancingRandomSeed = FolSeed;
+							HierarchicalInstancedStaticMeshComponent->InstanceStartCullDistance = LayerInfo->StartCullDistance;
+							HierarchicalInstancedStaticMeshComponent->InstanceEndCullDistance = LayerInfo->EndCullDistance;
+
+
+	#if 0
+							HierarchicalInstancedStaticMeshComponent->CastShadow                     = Settings->CastShadow;
+							HierarchicalInstancedStaticMeshComponent->bCastDynamicShadow             = Settings->bCastDynamicShadow;
+							HierarchicalInstancedStaticMeshComponent->bCastStaticShadow              = Settings->bCastStaticShadow;
+							HierarchicalInstancedStaticMeshComponent->bAffectDynamicIndirectLighting = Settings->bAffectDynamicIndirectLighting;
+							HierarchicalInstancedStaticMeshComponent->bCastShadowAsTwoSided          = Settings->bCastShadowAsTwoSided;
+							HierarchicalInstancedStaticMeshComponent->bReceivesDecals                = Settings->bReceivesDecals;
+							HierarchicalInstancedStaticMeshComponent->bOverrideLightMapRes           = Settings->bOverrideLightMapRes;
+							HierarchicalInstancedStaticMeshComponent->OverriddenLightMapRes          = Settings->OverriddenLightMapRes;
+							HierarchicalInstancedStaticMeshComponent->BodyInstance.CopyBodyInstancePropertiesFrom(&Settings->BodyInstance);
+	#endif
+
+							HierarchicalInstancedStaticMeshComponent->AttachTo(GetRootComponent());
+							FTransform DesiredTransform = GetRootComponent()->ComponentToWorld;
+							DesiredTransform.RemoveScaling();
+							HierarchicalInstancedStaticMeshComponent->SetWorldTransform(DesiredTransform);
+
+							FoliageComponents.Add(HierarchicalInstancedStaticMeshComponent);
+							StillUsed.Add(HierarchicalInstancedStaticMeshComponent);
+							NewLayer.Foliage.Add(HierarchicalInstancedStaticMeshComponent);
+
+
+							FAsyncGrassBuilder* Builder = new FAsyncGrassBuilder(this, Component, LayerInfo, HierarchicalInstancedStaticMeshComponent, SqrtSubsections, SubX, SubY);
+
+							if (Builder->bHaveValidData)
+							{
+								BuildTime -= FPlatformTime::Seconds();
+								FAsyncTask<FAsyncGrassTask>* Task = new FAsyncTask<FAsyncGrassTask>(Builder, Component, HierarchicalInstancedStaticMeshComponent, LayerInfo);
 
+								Task->StartBackgroundTask();
 
-						FAsyncGrassBuilder* Builder = new FAsyncGrassBuilder(this, Component, LayerInfo, HierarchicalInstancedStaticMeshComponent, SqrtSubsections, SubX, SubY);
+								AsyncFoliageTasks.Add(Task);
 
-						if (Builder->bHaveValidData)
-						{
-							BuildTime -= FPlatformTime::Seconds();
-							FAsyncTask<FAsyncGrassTask>* Task = new FAsyncTask<FAsyncGrassTask>(Builder, Component, HierarchicalInstancedStaticMeshComponent, LayerInfo);
+								BuildTime += FPlatformTime::Seconds();
+							}
+							else
+							{
+								delete Builder;
+							}
 
-							Task->StartBackgroundTask();
 
-							AsyncFoliageTasks.Add(Task);
+							ComponentMiscTime -= FPlatformTime::Seconds();
+							HierarchicalInstancedStaticMeshComponent->RegisterComponent();
+							ComponentMiscTime += FPlatformTime::Seconds();
 
-							BuildTime += FPlatformTime::Seconds();
 						}
-						else
-						{
-							delete Builder;
-						}
-
-
-						ComponentMiscTime -= FPlatformTime::Seconds();
-						HierarchicalInstancedStaticMeshComponent->RegisterComponent();
-						ComponentMiscTime += FPlatformTime::Seconds();
-
 					}
 				}
 			}
+			if (NewPerComponent.Layers.Num())
+			{
+				DesiredFoliage.PerComponent.Add(NewPerComponent);
+			}
+
+			ComponentTime += FPlatformTime::Seconds();
 		}
-		if (NewPerComponent.Layers.Num())
+
+		if (TotalComponents)
 		{
-			DesiredFoliage.PerComponent.Add(NewPerComponent);
+			double TotalTime = FPlatformTime::Seconds() - StartTime;
+			UE_LOG(LogTemp, Display, TEXT("Started grass with  %d components in %.1fms     Build Time %.1fms    Component Time  %.1fms    RegisterComponent Time  %.1fms"), TotalComponents, float(TotalTime) * 1000.0f, float(BuildTime) * 1000.0f, float(ComponentTime) * 1000.0f, float(ComponentMiscTime) * 1000.0f);
 		}
 
-		ComponentTime += FPlatformTime::Seconds();
 	}
-
-	if (TotalComponents)
-	{
-		double TotalTime = FPlatformTime::Seconds() - StartTime;
-		UE_LOG(LogTemp, Display, TEXT("Started grass with  %d components in %.1fms     Build Time %.1fms    Component Time  %.1fms    RegisterComponent Time  %.1fms"), TotalComponents, float(TotalTime) * 1000.0f, float(BuildTime) * 1000.0f, float(ComponentTime) * 1000.0f, float(ComponentMiscTime) * 1000.0f);
-	}
-
 	Exchange(FoliageCache, DesiredFoliage);
 
 	TSet<UHierarchicalInstancedStaticMeshComponent *> ToDestroy;
@@ -2707,9 +2715,9 @@ void ALandscapeProxy::Tick(float DeltaSeconds)
 		FoliageComponents.Remove(Component);
 	}
 
-	StartTime = FPlatformTime::Seconds();
+	double StartTime = FPlatformTime::Seconds();
 	int32 TotalInstances = 0;
-	TotalComponents = 0;
+	int32 TotalComponents = 0;
 	int64 TreeSize = 0;
 	int64 InstSize = 0;
 	int64 InstBufSize = 0;
