Force disable distance field lighting on grass, prevents extreme hitching from GPU management of 1m instances constantly added and removed

diff --git a/Engine/Source/Runtime/Landscape/Private/Landscape.cpp b/Engine/Source/Runtime/Landscape/Private/Landscape.cpp
--- a/Engine/Source/Runtime/Landscape/Private/Landscape.cpp
+++ b/Engine/Source/Runtime/Landscape/Private/Landscape.cpp
@@ -2614,7 +2614,8 @@ void ALandscapeProxy::Tick(float DeltaSeconds)
 							HierarchicalInstancedStaticMeshComponent->InstanceStartCullDistance = LayerInfo->StartCullDistance;
 							HierarchicalInstancedStaticMeshComponent->InstanceEndCullDistance = LayerInfo->EndCullDistance;
 
-
+							//@todo - take the settings from a UFoliageType object.  For now, disable distance field lighting on grass so we don't hitch.
+							HierarchicalInstancedStaticMeshComponent->bAffectDistanceFieldLighting = false;
 	#if 0
 							HierarchicalInstancedStaticMeshComponent->CastShadow                     = Settings->CastShadow;
 							HierarchicalInstancedStaticMeshComponent->bCastDynamicShadow             = Settings->bCastDynamicShadow;
