Fix for potential crash when shutting down landscape

diff --git a/Engine/Source/Runtime/Landscape/Private/LandscapeRender.cpp b/Engine/Source/Runtime/Landscape/Private/LandscapeRender.cpp
--- a/Engine/Source/Runtime/Landscape/Private/LandscapeRender.cpp
+++ b/Engine/Source/Runtime/Landscape/Private/LandscapeRender.cpp
@@ -534,6 +534,7 @@ TMap<FLandscapeComponentSceneProxy::FLandscapeKey, TMap<FIntPoint, const FLandsc
 FLandscapeComponentSceneProxy::FLandscapeComponentSceneProxy(ULandscapeComponent* InComponent, FLandscapeEditToolRenderData* InEditToolRenderData)
 	: FPrimitiveSceneProxy(InComponent)
 	, LandscapeKey(InComponent->GetWorld(), InComponent->GetLandscapeProxy()->GetLandscapeGuid())
+	, bAddedToSceneProxyMap(false)
 	, MaxLOD(FMath::CeilLogTwo(InComponent->SubsectionSizeQuads + 1) - 1)
 	, NumSubsections(InComponent->NumSubsections)
 	, SubsectionSizeQuads(InComponent->SubsectionSizeQuads)
@@ -670,6 +671,7 @@ void FLandscapeComponentSceneProxy::CreateRenderThreadResources()
 	if (Existing == nullptr)//(ensure(Existing == nullptr))
 	{
 		SceneProxyMap.Add(ComponentBase, this);
+		bAddedToSceneProxyMap = true;
 
 		// Find Neighbors
 		Neighbors[0] = SceneProxyMap.FindRef(ComponentBase + FIntPoint(0, -1));
@@ -766,37 +768,40 @@ void FLandscapeComponentSceneProxy::CreateRenderThreadResources()
 FLandscapeComponentSceneProxy::~FLandscapeComponentSceneProxy()
 {
 	// Remove ourselves from the map
-	TMap<FIntPoint, const FLandscapeComponentSceneProxy*>* SceneProxyMap = SharedSceneProxyMap.Find(LandscapeKey);
-	check(SceneProxyMap);
-
-	const FLandscapeComponentSceneProxy* MapEntry = SceneProxyMap->FindRef(ComponentBase);
-	if (MapEntry == this) //(/*ensure*/(MapEntry == this))
+	if (bAddedToSceneProxyMap)
 	{
-		SceneProxyMap->Remove(ComponentBase);
+		TMap<FIntPoint, const FLandscapeComponentSceneProxy*>* SceneProxyMap = SharedSceneProxyMap.Find(LandscapeKey);
+		check(SceneProxyMap);
 
-		if (SceneProxyMap->Num() == 0)
+		const FLandscapeComponentSceneProxy* MapEntry = SceneProxyMap->FindRef(ComponentBase);
+		if (MapEntry == this) //(/*ensure*/(MapEntry == this))
 		{
-			// remove the entire LandscapeKey entry as this is the last scene proxy
-			SharedSceneProxyMap.Remove(LandscapeKey);
-		}
-		else
-		{
-			// remove reference to us from our neighbors
-			if (Neighbors[0])
-			{
-				Neighbors[0]->Neighbors[3] = nullptr;
-			}
-			if (Neighbors[1])
-			{
-				Neighbors[1]->Neighbors[2] = nullptr;
-			}
-			if (Neighbors[2])
+			SceneProxyMap->Remove(ComponentBase);
+
+			if (SceneProxyMap->Num() == 0)
 			{
-				Neighbors[2]->Neighbors[1] = nullptr;
+				// remove the entire LandscapeKey entry as this is the last scene proxy
+				SharedSceneProxyMap.Remove(LandscapeKey);
 			}
-			if (Neighbors[3])
+			else
 			{
-				Neighbors[3]->Neighbors[0] = nullptr;
+				// remove reference to us from our neighbors
+				if (Neighbors[0])
+				{
+					Neighbors[0]->Neighbors[3] = nullptr;
+				}
+				if (Neighbors[1])
+				{
+					Neighbors[1]->Neighbors[2] = nullptr;
+				}
+				if (Neighbors[2])
+				{
+					Neighbors[2]->Neighbors[1] = nullptr;
+				}
+				if (Neighbors[3])
+				{
+					Neighbors[3]->Neighbors[0] = nullptr;
+				}
 			}
 		}
 	}
diff --git a/Engine/Source/Runtime/Landscape/Public/LandscapeRender.h b/Engine/Source/Runtime/Landscape/Public/LandscapeRender.h
--- a/Engine/Source/Runtime/Landscape/Public/LandscapeRender.h
+++ b/Engine/Source/Runtime/Landscape/Public/LandscapeRender.h
@@ -412,6 +412,7 @@ class FLandscapeComponentSceneProxy : public FPrimitiveSceneProxy
 
 protected:
 	FLandscapeKey				LandscapeKey;
+	bool						bAddedToSceneProxyMap;
 	int8						MaxLOD;
 	int8						NumSubsections;
 	int16						SubsectionSizeQuads;
