Log warning (instead of crashing) when landscape proxy scale is not equal to main landscape actor scale

diff --git a/Engine/Source/Runtime/Landscape/Private/Landscape.cpp b/Engine/Source/Runtime/Landscape/Private/Landscape.cpp
--- a/Engine/Source/Runtime/Landscape/Private/Landscape.cpp
+++ b/Engine/Source/Runtime/Landscape/Private/Landscape.cpp
@@ -1576,8 +1576,13 @@ void ULandscapeInfo::RegisterActor(ALandscapeProxy* Proxy, bool bMapCheck)
 	check(ComponentSizeQuads == Proxy->ComponentSizeQuads);
 	check(ComponentNumSubsections == Proxy->NumSubsections);
 	check(SubsectionSizeQuads == Proxy->SubsectionSizeQuads);
-	check(DrawScale == Proxy->GetRootComponent()->RelativeScale3D);
 
+	if (!DrawScale.Equals(Proxy->GetRootComponent()->RelativeScale3D))
+	{
+		UE_LOG(LogLandscape, Warning, TEXT("Landscape proxy (%s) scale (%s) does not match to main actor scale (%s)."), 
+			*Proxy->GetName(), *Proxy->GetRootComponent()->RelativeScale3D.ToCompactString(), *DrawScale.ToCompactString());
+	}
+	
 	// register
 	ALandscape* Landscape = Cast<ALandscape>(Proxy);
 	if (Landscape)
