Allow landscape proxies to fixup their transforms when loaded under world composition

diff --git a/Engine/Source/Runtime/Landscape/Private/Landscape.cpp b/Engine/Source/Runtime/Landscape/Private/Landscape.cpp
--- a/Engine/Source/Runtime/Landscape/Private/Landscape.cpp
+++ b/Engine/Source/Runtime/Landscape/Private/Landscape.cpp
@@ -1721,12 +1721,6 @@ void ULandscapeInfo::Reset()
 
 void ULandscapeInfo::FixupProxiesTransform()
 {
-	UWorld* OwningWorld = CastChecked<UWorld>(GetOuter());
-	if (OwningWorld->WorldComposition != NULL)
-	{
-		return;
-	}
-
 	ALandscape* Landscape = LandscapeActor.Get();
 
 	if (Landscape == NULL ||
