UE4 - never do render thread updates for grass after it is created, probably a temp solution

diff --git a/Engine/Source/Runtime/Landscape/Private/Landscape.cpp b/Engine/Source/Runtime/Landscape/Private/Landscape.cpp
--- a/Engine/Source/Runtime/Landscape/Private/Landscape.cpp
+++ b/Engine/Source/Runtime/Landscape/Private/Landscape.cpp
@@ -2729,6 +2729,9 @@ void ALandscapeProxy::Tick(float DeltaSeconds)
 				Exchange(HierarchicalInstancedStaticMeshComponent->PerInstanceSMData, Inner.Builder->PerInstanceSMData);
 				HierarchicalInstancedStaticMeshComponent->AcceptPrebuiltTree(Inner.Builder->ClusterTree, Inner.Builder->SortedInstances, Inner.Builder->InstanceReorderTable);
 
+				//@todo - this is sketchy, what you really want to do is flush the grass cache instead of doing a render thread update
+				HierarchicalInstancedStaticMeshComponent->bNeverNeedsRenderUpdate = true;
+
 				FBoxSphereBounds WorldBounds = HierarchicalInstancedStaticMeshComponent->CalcBounds(HierarchicalInstancedStaticMeshComponent->ComponentToWorld);
 				UE_LOG(LogTemp, Display, TEXT("Comp %d instances at %.0f %.0f %.0f"), HierarchicalInstancedStaticMeshComponent->PerInstanceSMData.Num(), WorldBounds.Origin.X, WorldBounds.Origin.Y, WorldBounds.Origin.Z);
 				TreeSize += HierarchicalInstancedStaticMeshComponent->ClusterTreePtr->GetAllocatedSize();
