UE4 - Grass: Fixed poor editor experience (hitching) with high density grass. Allowed grass to be rendering into landscape LODs. Fixed grass not ticking in editor without realtime.

diff --git a/Engine/Source/Editor/WorldBrowser/Private/Tiles/WorldTileCollectionModel.cpp b/Engine/Source/Editor/WorldBrowser/Private/Tiles/WorldTileCollectionModel.cpp
--- a/Engine/Source/Editor/WorldBrowser/Private/Tiles/WorldTileCollectionModel.cpp
+++ b/Engine/Source/Editor/WorldBrowser/Private/Tiles/WorldTileCollectionModel.cpp
@@ -1924,7 +1924,7 @@ bool FWorldTileCollectionModel::GenerateLODLevels(FLevelModelList InLevelList, i
 				{
 					LandscapeActors.Add(LandscapeProxy);
 				}
-				else
+				else 
 				{
 					Actors.Add(Actor);
 				}
@@ -1984,7 +1984,7 @@ bool FWorldTileCollectionModel::GenerateLODLevels(FLevelModelList InLevelList, i
 				UObject* PrimitiveOuter = PrimitiveComp->GetOuter();
 				
 				const bool bTargetPrim = 
-					PrimitiveComp->GetOuter() == Landscape || 
+					(PrimitiveComp->GetOuter() == Landscape && !(!SimplificationDetails.bBakeGrassToLandscape && PrimitiveComp->IsA(UHierarchicalInstancedStaticMeshComponent::StaticClass()))) || 
 					(SimplificationDetails.bBakeFoliageToLandscape && PrimitiveOuter->IsA(AInstancedFoliageActor::StaticClass()));
 
 				if (!bTargetPrim && PrimitiveComp->IsRegistered() && PrimitiveComp->SceneProxy)
@@ -1992,6 +1992,13 @@ bool FWorldTileCollectionModel::GenerateLODLevels(FLevelModelList InLevelList, i
 					PrimitivesToHide.Add(PrimitiveComp->SceneProxy->GetPrimitiveComponentId());
 				}
 			}
+
+			if (SimplificationDetails.bBakeGrassToLandscape)
+			{
+				Landscape->FlushFoliageComponents();
+				TArray<FVector> Cameras;
+				Landscape->UpdateFoliage(Cameras, nullptr, true);
+			}
 								
 			// This is texture resolution for a landscape mesh, probably needs to be calculated using landscape size
 			const FIntPoint LandscapeTextureSize(1024, 1024);
@@ -2003,6 +2010,10 @@ bool FWorldTileCollectionModel::GenerateLODLevels(FLevelModelList InLevelList, i
 			
 			ExportMaterial(Landscape, PrimitivesToHide, LandscapeFlattenMaterial);
 		
+			if (SimplificationDetails.bBakeGrassToLandscape)
+			{
+				Landscape->FlushFoliageComponents(); // wipe this and let it fix itself later
+			}
 			FString LandscapeBaseAssetName = Landscape->GetName();
 			// Construct landscape material
 			UMaterial* StaticLandscapeMaterial = MaterialExportUtils::CreateMaterial(
diff --git a/Engine/Source/Runtime/Engine/Classes/Engine/Level.h b/Engine/Source/Runtime/Engine/Classes/Engine/Level.h
--- a/Engine/Source/Runtime/Engine/Classes/Engine/Level.h
+++ b/Engine/Source/Runtime/Engine/Classes/Engine/Level.h
@@ -266,6 +266,10 @@ struct ENGINE_API FLevelSimplificationDetails
 	UPROPERTY(Category=Landscape, EditAnywhere)
 	bool bBakeFoliageToLandscape;
 
+	// Whether to bake grass into landscape static mesh texture
+	UPROPERTY(Category=Landscape, EditAnywhere)
+	bool bBakeGrassToLandscape;
+
 	FLevelSimplificationDetails();
 
 	bool operator == (const FLevelSimplificationDetails& Other) const;
diff --git a/Engine/Source/Runtime/Engine/Private/InstancedStaticMesh.cpp b/Engine/Source/Runtime/Engine/Private/InstancedStaticMesh.cpp
--- a/Engine/Source/Runtime/Engine/Private/InstancedStaticMesh.cpp
+++ b/Engine/Source/Runtime/Engine/Private/InstancedStaticMesh.cpp
@@ -1290,6 +1290,7 @@ void UInstancedStaticMeshComponent::ClearInstances()
 	ClearAllInstanceBodies();
 
 	// Indicate we need to update render state to reflect changes
+	bNeverNeedsRenderUpdate = false;
 	MarkRenderStateDirty();
 
 	UNavigationSystem::UpdateNavOctree(this);
diff --git a/Engine/Source/Runtime/Engine/Private/Level.cpp b/Engine/Source/Runtime/Engine/Private/Level.cpp
--- a/Engine/Source/Runtime/Engine/Private/Level.cpp
+++ b/Engine/Source/Runtime/Engine/Private/Level.cpp
@@ -170,6 +170,7 @@ FLevelSimplificationDetails::FLevelSimplificationDetails()
  , bGenerateLandscapeRoughnessMap(false)
  , bGenerateLandscapeSpecularMap(false)
  , bBakeFoliageToLandscape(false)
+ , bBakeGrassToLandscape(false)
 {
 }
 
@@ -183,7 +184,8 @@ bool FLevelSimplificationDetails::operator == (const FLevelSimplificationDetails
 		bGenerateLandscapeMetallicMap == Other.bGenerateLandscapeMetallicMap &&
 		bGenerateLandscapeRoughnessMap == Other.bGenerateLandscapeRoughnessMap &&
 		bGenerateLandscapeSpecularMap == Other.bGenerateLandscapeSpecularMap &&
-		bBakeFoliageToLandscape == Other.bBakeFoliageToLandscape;
+		bBakeFoliageToLandscape == Other.bBakeFoliageToLandscape &&
+		bBakeGrassToLandscape == Other.bBakeGrassToLandscape;
 }
 
 TMap<FName, TWeakObjectPtr<UWorld> > ULevel::StreamedLevelsOwningWorld;
diff --git a/Engine/Source/Runtime/Landscape/Classes/LandscapeProxy.h b/Engine/Source/Runtime/Landscape/Classes/LandscapeProxy.h
--- a/Engine/Source/Runtime/Landscape/Classes/LandscapeProxy.h
+++ b/Engine/Source/Runtime/Landscape/Classes/LandscapeProxy.h
@@ -15,6 +15,7 @@
 #include "Components/HierarchicalInstancedStaticMeshComponent.h"
 #include "LandscapeComponent.h"
 #include "LandscapeLayerInfoObject.h"
+#include "Tickable.h"
 
 #include "LandscapeProxy.generated.h"
 
@@ -247,7 +248,7 @@ public:
 };
 
 UCLASS(NotPlaceable, hidecategories=(Display, Attachment, Physics, Debug, Lighting, LOD), showcategories=(Rendering, "Utilities|Transformation"), MinimalAPI)
-class ALandscapeProxy : public AActor
+class ALandscapeProxy : public AActor, public FTickableGameObject
 {
 	GENERATED_UCLASS_BODY()
 
@@ -411,12 +412,6 @@ public:
 	virtual void RegisterAllComponents() override;
 	virtual void RerunConstructionScripts() override {}
 	virtual bool IsLevelBoundsRelevant() const override { return true; }
-	virtual void Tick(float DeltaSeconds) override;
-	virtual bool ShouldTickIfViewportsOnly() const override
-	{
-		return true;
-	}
-
 
 #if WITH_EDITOR
 	virtual void Destroyed() override;
@@ -435,6 +430,33 @@ public:
 	/** Flush the foliage cache */
 	LANDSCAPE_API void FlushFoliageComponents();
 
+	/** 
+		Update Foliage 
+		* @param Cameras to use for culling, if empty, then NO culling
+		* @param OnlyComponent if non-null only do this component
+		* @param bForceSync if true, block and finish all work
+	*/
+	LANDSCAPE_API void UpdateFoliage(const TArray<FVector>& Cameras, ULandscapeComponent* OnlyComponent = nullptr, bool bForceSync = false);
+
+	// Begin FTickableGameObject interface.
+	virtual void Tick(float DeltaTime) override;
+	virtual bool IsTickable() const override 
+	{ 
+		return !HasAnyFlags(RF_ClassDefaultObject); 
+	}
+	virtual bool IsTickableWhenPaused() const override
+	{
+		return !HasAnyFlags(RF_ClassDefaultObject); 
+	}
+	virtual bool IsTickableInEditor() const override
+	{
+		return !HasAnyFlags(RF_ClassDefaultObject); 
+	}
+	virtual TStatId GetStatId() const override
+	{
+		return GetStatID();
+	}
+
 	// Begin UObject interface.
 	virtual void Serialize(FArchive& Ar) override;
 	static void AddReferencedObjects(UObject* InThis, FReferenceCollector& Collector);
diff --git a/Engine/Source/Runtime/Landscape/Private/Landscape.cpp b/Engine/Source/Runtime/Landscape/Private/Landscape.cpp
--- a/Engine/Source/Runtime/Landscape/Private/Landscape.cpp
+++ b/Engine/Source/Runtime/Landscape/Private/Landscape.cpp
@@ -487,7 +487,6 @@ ALandscape* ALandscape::GetLandscapeActor()
 ALandscapeProxy::ALandscapeProxy(const FObjectInitializer& ObjectInitializer)
 	: Super(ObjectInitializer)
 {
-	PrimaryActorTick.bCanEverTick = true;
 	bReplicates = false;
 	NetUpdateFrequency = 10.0f;
 	bHidden = false;
@@ -2457,11 +2456,11 @@ struct FAsyncGrassBuilder : public FGrassBuilderBase
 
 };
 
-
 void ALandscapeProxy::Tick(float DeltaSeconds)
 {
-	Super::Tick(DeltaSeconds);
-
+	// this is NOT an actor tick, it is a FTickableGameObject tick
+	// the super tick is for an actor tick...
+	//Super::Tick(DeltaSeconds);
 	UWorld* World = GetWorld();
 	if (!World)
 	{
@@ -2484,7 +2483,11 @@ void ALandscapeProxy::Tick(float DeltaSeconds)
 		check(IsInGameThread());
 		OldCameras = Cameras;
 	}
+	UpdateFoliage(Cameras);
+}
 
+void ALandscapeProxy::UpdateFoliage(const TArray<FVector>& Cameras, ULandscapeComponent* OnlyComponent, bool bForceSync)
+{
 	FCachedLandscapeFoliage DesiredFoliage;
 	TSet<UHierarchicalInstancedStaticMeshComponent *> StillUsed;
 
@@ -2501,8 +2504,12 @@ void ALandscapeProxy::Tick(float DeltaSeconds)
 
 		for (int32 ComponentIndex = 0; ComponentIndex < LandscapeComponents.Num(); ComponentIndex++)
 		{
-			ComponentTime -= FPlatformTime::Seconds();
 			ULandscapeComponent* Component = LandscapeComponents[ComponentIndex];
+			if (OnlyComponent && OnlyComponent != Component)
+			{
+				continue;
+			}
+			ComponentTime -= FPlatformTime::Seconds();
 
 			const FCachedLandscapeFoliage::FPerComponent* ExistingCache = nullptr;
 			for (const FCachedLandscapeFoliage::FPerComponent& CacheItem : FoliageCache.PerComponent)
@@ -2516,7 +2523,7 @@ void ALandscapeProxy::Tick(float DeltaSeconds)
 			FCachedLandscapeFoliage::FPerComponent NewPerComponent(Component);
 
 			FBoxSphereBounds WorldBounds = Component->CalcBounds(Component->ComponentToWorld);
-			float MinDistanceToComp = MAX_flt;
+			float MinDistanceToComp = Cameras.Num() ? MAX_flt : 0.0f;
 
 			for (auto& Pos : Cameras)
 			{
@@ -2610,8 +2617,18 @@ void ALandscapeProxy::Tick(float DeltaSeconds)
 							HierarchicalInstancedStaticMeshComponent->bDisableCollision = true;
 							HierarchicalInstancedStaticMeshComponent->SetCanEverAffectNavigation(false);
 							HierarchicalInstancedStaticMeshComponent->InstancingRandomSeed = FolSeed;
-							HierarchicalInstancedStaticMeshComponent->InstanceStartCullDistance = LayerInfo->StartCullDistance;
-							HierarchicalInstancedStaticMeshComponent->InstanceEndCullDistance = LayerInfo->EndCullDistance;
+
+							if (!Cameras.Num())
+							{
+								// if we don't have any cameras, then we are rendering landscape LOD materials or somesuch and we want to disable culling
+								HierarchicalInstancedStaticMeshComponent->InstanceStartCullDistance = 0;
+								HierarchicalInstancedStaticMeshComponent->InstanceEndCullDistance = 0;
+							}
+							else
+							{
+								HierarchicalInstancedStaticMeshComponent->InstanceStartCullDistance = LayerInfo->StartCullDistance;
+								HierarchicalInstancedStaticMeshComponent->InstanceEndCullDistance = LayerInfo->EndCullDistance;
+							}
 
 							//@todo - take the settings from a UFoliageType object.  For now, disable distance field lighting on grass so we don't hitch.
 							HierarchicalInstancedStaticMeshComponent->bAffectDistanceFieldLighting = false;
@@ -2717,6 +2734,10 @@ void ALandscapeProxy::Tick(float DeltaSeconds)
 	for (int32 Index = 0; Index < AsyncFoliageTasks.Num(); Index++)
 	{
 		FAsyncTask<FAsyncGrassTask>* Task = AsyncFoliageTasks[Index];
+		if (bForceSync)
+		{
+			Task->EnsureCompletion();
+		}
 		if (Task->IsDone())
 		{
 			FAsyncGrassTask& Inner = Task->GetTask();
@@ -2732,15 +2753,22 @@ void ALandscapeProxy::Tick(float DeltaSeconds)
 				//@todo - this is sketchy, what you really want to do is flush the grass cache instead of doing a render thread update
 				HierarchicalInstancedStaticMeshComponent->bNeverNeedsRenderUpdate = true;
 
-				FBoxSphereBounds WorldBounds = HierarchicalInstancedStaticMeshComponent->CalcBounds(HierarchicalInstancedStaticMeshComponent->ComponentToWorld);
-				UE_LOG(LogTemp, Display, TEXT("Comp %d instances at %.0f %.0f %.0f"), HierarchicalInstancedStaticMeshComponent->PerInstanceSMData.Num(), WorldBounds.Origin.X, WorldBounds.Origin.Y, WorldBounds.Origin.Z);
+				//FBoxSphereBounds WorldBounds = HierarchicalInstancedStaticMeshComponent->CalcBounds(HierarchicalInstancedStaticMeshComponent->ComponentToWorld);
+				//UE_LOG(LogTemp, Display, TEXT("Comp %d instances at %.0f %.0f %.0f"), HierarchicalInstancedStaticMeshComponent->PerInstanceSMData.Num(), WorldBounds.Origin.X, WorldBounds.Origin.Y, WorldBounds.Origin.Z);
 				TreeSize += HierarchicalInstancedStaticMeshComponent->ClusterTreePtr->GetAllocatedSize();
 				InstSize += HierarchicalInstancedStaticMeshComponent->PerInstanceSMData.GetAllocatedSize();
 				InstBufSize += HierarchicalInstancedStaticMeshComponent->PerInstanceSMData.Num() * 7 * 16; // there is no place where this constant exists.
+				if (bForceSync)
+				{
+					HierarchicalInstancedStaticMeshComponent->RecreateRenderState_Concurrent();
+				}
 			}
 			delete Inner.Builder;
 			delete Task;
-			break; // one per frame is fine
+			if (!bForceSync)
+			{
+				break; // one per frame is fine
+			}
 		}
 	}
 	if (TotalComponents)
