UE4 - temp fix for camera badness in PIE with grass

diff --git a/Engine/Source/Runtime/Landscape/Private/Landscape.cpp b/Engine/Source/Runtime/Landscape/Private/Landscape.cpp
--- a/Engine/Source/Runtime/Landscape/Private/Landscape.cpp
+++ b/Engine/Source/Runtime/Landscape/Private/Landscape.cpp
@@ -2467,23 +2467,62 @@ void ALandscapeProxy::Tick(float DeltaSeconds)
 		return;
 	}
 
+	struct FOldCamera
+	{
+		FVector Pos;
+		double LastTime;
+	};
+
 	// there is a bug here, which often leaves us with no cameras in the editor
-	static TArray<FVector> OldCameras;
+	// also sometimes the cameras are ping ponging and otherise really wrong in the editor
+	// the main thing is not to thrash
+	static TArray<FOldCamera> OldCameras;
 
-	if (!OldCameras.Num() && !World->ViewLocationsRenderedLastFrame.Num())
+	double Now = FPlatformTime::Seconds();
+	for (auto& NewCam : World->ViewLocationsRenderedLastFrame)
 	{
-		// no cameras, no grass update
-		return;
+		bool bFound = false;
+		for (auto& OldCam : OldCameras)
+		{
+			if (FVector::Dist(OldCam.Pos, NewCam) < 1000.0f)
+			{
+				OldCam.Pos = NewCam;
+				OldCam.LastTime = Now;
+				bFound = true;
+				break;
+			}
+		}
+		if (!bFound)
+		{
+			FOldCamera NewOldCam;
+			NewOldCam.Pos = NewCam;
+			NewOldCam.LastTime = Now;
+			OldCameras.Add(NewOldCam);
+		}
+	}
+	for (int32 Index = 0 ; Index < OldCameras.Num(); Index++)
+	{
+		if (OldCameras[Index].LastTime < Now - 2.0f)
+		{
+			OldCameras.RemoveAtSwap(Index);
+		}
 	}
 
-	const TArray<FVector>& Cameras = World->ViewLocationsRenderedLastFrame.Num() ? World->ViewLocationsRenderedLastFrame : OldCameras;
+	static TArray<FVector> UseCameras;
+	check(IsInGameThread());
+	UseCameras.Reset();
+	for (auto& OldCam : OldCameras)
+	{
+		UseCameras.Add(OldCam.Pos);
+	}
 
-	if (&Cameras != &OldCameras)
+	if (!UseCameras.Num())
 	{
-		check(IsInGameThread());
-		OldCameras = Cameras;
+		// no cameras, no grass update
+		return;
 	}
-	UpdateFoliage(Cameras);
+
+	UpdateFoliage(UseCameras);
 }
 
 void ALandscapeProxy::UpdateFoliage(const TArray<FVector>& Cameras, ULandscapeComponent* OnlyComponent, bool bForceSync)
