Hid Landscape Spline sprites when not in landscape spline edit mode

diff --git a/Engine/Source/Runtime/Landscape/Private/LandscapeSplines.cpp b/Engine/Source/Runtime/Landscape/Private/LandscapeSplines.cpp
--- a/Engine/Source/Runtime/Landscape/Private/LandscapeSplines.cpp
+++ b/Engine/Source/Runtime/Landscape/Private/LandscapeSplines.cpp
@@ -35,6 +35,7 @@ private:
 	const FLinearColor	SplineColor;
 
 	const UTexture2D*	ControlPointSprite;
+	const bool			bDrawControlPointSprite;
 	const bool			bDrawFalloff;
 
 	struct FSegmentProxy
@@ -67,6 +68,7 @@ public:
 		FPrimitiveSceneProxy(Component),
 		SplineColor(Component->SplineColor),
 		ControlPointSprite(Component->ControlPointSprite),
+		bDrawControlPointSprite(Component->bShowSplineEditorMesh),
 		bDrawFalloff(Component->bShowSplineEditorMesh)
 	{
 		Segments.Reserve(Component->Segments.Num());
@@ -172,25 +174,28 @@ public:
 
 				for (const FControlPointProxy& ControlPoint : ControlPoints)
 				{
-					const float ControlPointSpriteScale = LocalToWorld.GetScaleVector().X * ControlPoint.SpriteScale;
-					const FVector ControlPointLocation = LocalToWorld.TransformPosition(ControlPoint.Location) + FVector(0, 0, ControlPointSpriteScale * 0.75f);
+					const FVector ControlPointLocation = LocalToWorld.TransformPosition(ControlPoint.Location);
 
 					// Draw Sprite
-
-					const FLinearColor ControlPointSpriteColor = ControlPoint.bSelected ? SelectedControlPointSpriteColor : FLinearColor::White;
-
-					PDI->SetHitProxy(ControlPoint.HitProxy);
-
-					PDI->DrawSprite(
-						ControlPointLocation,
-						ControlPointSpriteScale,
-						ControlPointSpriteScale,
-						ControlPointSprite->Resource,
-						ControlPointSpriteColor,
-						GetDepthPriorityGroup(View),
-						0, ControlPointSprite->Resource->GetSizeX(),
-						0, ControlPointSprite->Resource->GetSizeY(),
-						SE_BLEND_Masked);
+					if (bDrawControlPointSprite)
+					{
+						const float ControlPointSpriteScale = LocalToWorld.GetScaleVector().X * ControlPoint.SpriteScale;
+						const FVector ControlPointSpriteLocation = ControlPointLocation + FVector(0, 0, ControlPointSpriteScale * 0.75f);
+						const FLinearColor ControlPointSpriteColor = ControlPoint.bSelected ? SelectedControlPointSpriteColor : FLinearColor::White;
+
+						PDI->SetHitProxy(ControlPoint.HitProxy);
+
+						PDI->DrawSprite(
+							ControlPointSpriteLocation,
+							ControlPointSpriteScale,
+							ControlPointSpriteScale,
+							ControlPointSprite->Resource,
+							ControlPointSpriteColor,
+							GetDepthPriorityGroup(View),
+							0, ControlPointSprite->Resource->GetSizeX(),
+							0, ControlPointSprite->Resource->GetSizeY(),
+							SE_BLEND_Masked);
+					}
 
 
 					// Draw Lines
