Fix seams between proxy streaming levels and real landscape

diff --git a/Engine/Source/Editor/WorldBrowser/Private/Tiles/WorldTileCollectionModel.cpp b/Engine/Source/Editor/WorldBrowser/Private/Tiles/WorldTileCollectionModel.cpp
--- a/Engine/Source/Editor/WorldBrowser/Private/Tiles/WorldTileCollectionModel.cpp
+++ b/Engine/Source/Editor/WorldBrowser/Private/Tiles/WorldTileCollectionModel.cpp
@@ -24,6 +24,8 @@
 #include "GameFramework/WorldSettings.h"
 #include "Engine/Light.h"
 #include "Foliage/InstancedFoliageActor.h"
+#include "LandscapeMeshProxyActor.h"
+#include "LandscapeMeshProxyComponent.h"
 
 #define LOCTEXT_NAMESPACE "WorldBrowser"
 
@@ -1903,8 +1905,20 @@ bool FWorldTileCollectionModel::GenerateLODLevels(FLevelModelList InLevelList, i
 			ActorsOffset = FVector::ZeroVector;
 		}
 	
-		TArray<UStaticMesh*>		AssetsToSpawn;
-		TArray<FTransform>			AssetsToSpawnTransform;
+		struct FAssetToSpawnInfo
+		{
+			FAssetToSpawnInfo(UStaticMesh* InStaticMesh, const FTransform& InTransform, ALandscapeProxy* InSourceLandscape = nullptr)
+			: StaticMesh(InStaticMesh)
+			, Transform(InTransform)
+			, SourceLandscape(InSourceLandscape)
+			{}
+
+			UStaticMesh* StaticMesh;
+			FTransform Transform;
+			ALandscapeProxy* SourceLandscape;
+		};
+
+		TArray<FAssetToSpawnInfo>	AssetsToSpawn;
 		TArray<AActor*>				Actors;
 		TArray<ALandscapeProxy*>	LandscapeActors;
 		TArray<UObject*>			GeneratedAssets;
@@ -1949,8 +1963,7 @@ bool FWorldTileCollectionModel::GenerateLODLevels(FLevelModelList InLevelList, i
 				UStaticMesh* ProxyMesh = nullptr;
 				if (OutAssets.FindItemByClass(&ProxyMesh))
 				{
-					AssetsToSpawn.Add(ProxyMesh);
-					AssetsToSpawnTransform.Add(FTransform(OutProxyLocation - ActorsOffset));
+					new(AssetsToSpawn) FAssetToSpawnInfo(ProxyMesh, FTransform(OutProxyLocation - ActorsOffset));
 				}
 
 				GeneratedAssets.Append(OutAssets);
@@ -2060,8 +2073,7 @@ bool FWorldTileCollectionModel::GenerateLODLevels(FLevelModelList InLevelList, i
 			}
 
 			GeneratedAssets.Add(StaticMesh);
-			AssetsToSpawn.Add(StaticMesh);
-			AssetsToSpawnTransform.Add(FTransform(LandscapeWorldLocation - ActorsOffset));
+			new(AssetsToSpawn) FAssetToSpawnInfo(StaticMesh, FTransform(LandscapeWorldLocation - ActorsOffset), Landscape);
 
 			LandscapeActorIndex++;
 		}
@@ -2108,13 +2120,24 @@ bool FWorldTileCollectionModel::GenerateLODLevels(FLevelModelList InLevelList, i
 			LODWorld = UWorld::CreateWorld(EWorldType::None, false, FPackageName::GetShortFName(LODPackage->GetFName()), LODPackage);
 			LODWorld->SetFlags(RF_Standalone);
 
-			for (int32 AssetIdx = 0; AssetIdx < AssetsToSpawn.Num(); ++AssetIdx)
+			for (FAssetToSpawnInfo& AssetInfo : AssetsToSpawn)
 			{
-				FVector Location = AssetsToSpawnTransform[AssetIdx].GetLocation();
+				FVector Location = AssetInfo.Transform.GetLocation();
 				FRotator Rotation(ForceInit);
-				AStaticMeshActor* MeshActor = LODWorld->SpawnActor<AStaticMeshActor>(Location, Rotation);
-				MeshActor->GetStaticMeshComponent()->StaticMesh = AssetsToSpawn[AssetIdx];
-				MeshActor->SetActorLabel(AssetsToSpawn[AssetIdx]->GetName());
+
+				if (AssetInfo.SourceLandscape != nullptr)
+				{
+					ALandscapeMeshProxyActor* MeshActor = LODWorld->SpawnActor<ALandscapeMeshProxyActor>(Location, Rotation);
+					MeshActor->GetLandscapeMeshProxyComponent()->StaticMesh = AssetInfo.StaticMesh;
+					MeshActor->GetLandscapeMeshProxyComponent()->InitializeForLandscape(AssetInfo.SourceLandscape, SimplificationDetails.LandscapeExportLOD);
+					MeshActor->SetActorLabel(AssetInfo.SourceLandscape->GetName());
+				}
+				else
+				{
+					AStaticMeshActor* MeshActor = LODWorld->SpawnActor<AStaticMeshActor>(Location, Rotation);
+					MeshActor->GetStaticMeshComponent()->StaticMesh = AssetInfo.StaticMesh;
+					MeshActor->SetActorLabel(AssetInfo.StaticMesh->GetName());
+				}
 			}
 		
 			// Save generated level
diff --git a/Engine/Source/Runtime/Engine/Private/Level.cpp b/Engine/Source/Runtime/Engine/Private/Level.cpp
--- a/Engine/Source/Runtime/Engine/Private/Level.cpp
+++ b/Engine/Source/Runtime/Engine/Private/Level.cpp
@@ -1538,11 +1538,6 @@ void ULevel::InitializeRenderingResources()
 		if( !PrecomputedLightVolume->IsAddedToScene() )
 		{
 			PrecomputedLightVolume->AddToScene(OwningWorld->Scene);
-
-			if (OwningWorld->Scene)
-			{
-				OwningWorld->Scene->OnLevelAddedToWorld(GetOutermost()->GetFName());
-			}
 		}
 	}
 }
diff --git a/Engine/Source/Runtime/Engine/Private/PrimitiveSceneProxy.cpp b/Engine/Source/Runtime/Engine/Private/PrimitiveSceneProxy.cpp
--- a/Engine/Source/Runtime/Engine/Private/PrimitiveSceneProxy.cpp
+++ b/Engine/Source/Runtime/Engine/Private/PrimitiveSceneProxy.cpp
@@ -51,6 +51,7 @@ FPrimitiveSceneProxy::FPrimitiveSceneProxy(const UPrimitiveComponent* InComponen
 ,	bUseEditorDepthTest(true)
 ,	bSupportsDistanceFieldRepresentation(false)
 ,	bSupportsHeightfieldRepresentation(false)
+,	bNeedsLevelAddedToWorldNotification(false)
 ,	bUseAsOccluder(InComponent->bUseAsOccluder)
 ,	bAllowApproximateOcclusion(InComponent->Mobility != EComponentMobility::Movable)
 ,	bSelectable(InComponent->bSelectable)
@@ -130,15 +131,8 @@ FPrimitiveSceneProxy::FPrimitiveSceneProxy(const UPrimitiveComponent* InComponen
 	// Flag components to render only after level will be fully added to the world
 	//
 	ULevel* ComponentLevel = InComponent->GetComponentLevel();
-	if (ComponentLevel && ComponentLevel->bRequireFullVisibilityToRender)
-	{
-		bRequiresVisibleLevelToRender = true;
-		bIsComponentLevelVisible = ComponentLevel->bIsVisible;
-	}
-	else
-	{
-		bRequiresVisibleLevelToRender = false;
-	}
+	bRequiresVisibleLevelToRender = (ComponentLevel && ComponentLevel->bRequireFullVisibilityToRender);
+	bIsComponentLevelVisible = (!ComponentLevel || ComponentLevel->bIsVisible);
 }
 
 FPrimitiveSceneProxy::~FPrimitiveSceneProxy()
diff --git a/Engine/Source/Runtime/Engine/Private/World.cpp b/Engine/Source/Runtime/Engine/Private/World.cpp
--- a/Engine/Source/Runtime/Engine/Private/World.cpp
+++ b/Engine/Source/Runtime/Engine/Private/World.cpp
@@ -2261,10 +2261,16 @@ void UWorld::UpdateLevelStreamingInner(ULevelStreaming* StreamingLevel)
 			}
 
 			// In case we have finished making level visible
-			// immediately discard previous level
 			if (Level->bIsVisible)
 			{
+				// immediately discard previous level
 				StreamingLevel->DiscardPendingUnloadLevel(this);
+
+				if (Scene)
+				{
+					// Notify the new level has been added after the old has been discarded
+					Scene->OnLevelAddedToWorld(Level->GetOutermost()->GetFName());
+				}
 			}
 		}
 		else
diff --git a/Engine/Source/Runtime/Engine/Public/PrimitiveSceneProxy.h b/Engine/Source/Runtime/Engine/Public/PrimitiveSceneProxy.h
--- a/Engine/Source/Runtime/Engine/Public/PrimitiveSceneProxy.h
+++ b/Engine/Source/Runtime/Engine/Public/PrimitiveSceneProxy.h
@@ -241,6 +241,13 @@ public:
 	 */
 	virtual void OnActorPositionChanged() {}
 
+	/**
+	 * Called to notify the proxy that the level has been fully added to
+	 * the world and the primitive will now be rendered.
+	 * Only called if bNeedsLevelAddedToWorldNotification is set to true.
+	 */
+	virtual void OnLevelAddedToWorld() {}
+
 	/**
 	* @return true if the proxy can be culled when occluded by other primitives
 	*/
@@ -355,6 +362,9 @@ public:
 	inline bool SupportsDistanceFieldRepresentation() const { return bSupportsDistanceFieldRepresentation; }
 	inline bool SupportsHeightfieldRepresentation() const { return bSupportsHeightfieldRepresentation; }
 	inline bool TreatAsBackgroundForOcclusion() const { return bTreatAsBackgroundForOcclusion; }
+	inline bool NeedsLevelAddedToWorldNotification() const { return bNeedsLevelAddedToWorldNotification; }
+	inline bool IsComponentLevelVisible() const { return bIsComponentLevelVisible; }
+
 #if WITH_EDITOR
 	inline int32 GetNumUncachedStaticLightingInteractions() { return NumUncachedStaticLightingInteractions; }
 #endif
@@ -543,6 +553,9 @@ protected:
 	/** Whether the primitive implements GetHeightfieldRepresentation() */
 	uint32 bSupportsHeightfieldRepresentation : 1;
 
+	/** Whether this primitive requires notification when its level is added to the world and made visible for the first time. */
+	uint32 bNeedsLevelAddedToWorldNotification : 1;
+
 private:
 
 	/** If this is True, this primitive will be used to occlusion cull other primitives. */
diff --git a/Engine/Source/Runtime/Engine/Public/StaticMeshResources.h b/Engine/Source/Runtime/Engine/Public/StaticMeshResources.h
--- a/Engine/Source/Runtime/Engine/Public/StaticMeshResources.h
+++ b/Engine/Source/Runtime/Engine/Public/StaticMeshResources.h
@@ -732,7 +732,7 @@ private:
 /**
  * A static mesh component scene proxy.
  */
-class FStaticMeshSceneProxy : public FPrimitiveSceneProxy
+class ENGINE_API FStaticMeshSceneProxy : public FPrimitiveSceneProxy
 {
 public:
 
diff --git a/Engine/Source/Runtime/Landscape/Classes/LandscapeMeshProxyActor.h b/Engine/Source/Runtime/Landscape/Classes/LandscapeMeshProxyActor.h
new file mode 100644
--- /dev/null
+++ b/Engine/Source/Runtime/Landscape/Classes/LandscapeMeshProxyActor.h
@@ -0,0 +1,24 @@
+// Copyright 1998-2015 Epic Games, Inc. All Rights Reserved.
+
+#pragma once
+
+#include "SceneTypes.h"
+#include "StaticLighting.h"
+#include "Components/PrimitiveComponent.h"
+
+#include "LandscapeMeshProxyActor.generated.h"
+
+UCLASS(MinimalAPI)
+class ALandscapeMeshProxyActor : public AActor
+{
+	GENERATED_UCLASS_BODY()
+
+private:
+	UPROPERTY(Category = LandscapeMeshProxyActor, VisibleAnywhere, BlueprintReadOnly, meta = (ExposeFunctionCategories = "Mesh,Rendering,Physics,Components|StaticMesh", AllowPrivateAccess = "true"))
+	class ULandscapeMeshProxyComponent* LandscapeMeshProxyComponent;
+
+public:
+	/** Returns StaticMeshComponent subobject **/
+	class ULandscapeMeshProxyComponent* GetLandscapeMeshProxyComponent() const { return LandscapeMeshProxyComponent; }
+};
+
diff --git a/Engine/Source/Runtime/Landscape/Classes/LandscapeMeshProxyComponent.h b/Engine/Source/Runtime/Landscape/Classes/LandscapeMeshProxyComponent.h
new file mode 100644
--- /dev/null
+++ b/Engine/Source/Runtime/Landscape/Classes/LandscapeMeshProxyComponent.h
@@ -0,0 +1,49 @@
+// Copyright 1998-2015 Epic Games, Inc. All Rights Reserved.
+
+#pragma once
+
+#include "SceneTypes.h"
+#include "StaticLighting.h"
+#include "Components/PrimitiveComponent.h"
+
+#include "LandscapeMeshProxyComponent.generated.h"
+
+UCLASS(MinimalAPI)
+class ULandscapeMeshProxyComponent : public UStaticMeshComponent
+{
+	GENERATED_UCLASS_BODY()
+
+private:
+	/* The landscape this proxy was generated for */
+	UPROPERTY()
+	FGuid LandscapeGuid;
+
+	/* The components this proxy was generated for */
+	UPROPERTY()
+	TArray<FIntPoint> ProxyComponentBases;
+
+	/* LOD level proxy was generated for */
+	UPROPERTY()
+	int8 ProxyLOD;
+public:
+	void InitializeForLandscape(ALandscapeProxy* Landscape, int8 InProxyLOD)
+	{
+		LandscapeGuid = Landscape->GetLandscapeGuid();
+
+		for (ULandscapeComponent* Component : Landscape->LandscapeComponents)
+		{
+			if (Component)
+			{
+				ProxyComponentBases.Add(Component->GetSectionBase() / Component->ComponentSizeQuads);
+			}
+		}
+
+		if (InProxyLOD != INDEX_NONE)
+		{
+			ProxyLOD = FMath::Clamp<int32>(InProxyLOD, 0, FMath::CeilLogTwo(Landscape->SubsectionSizeQuads + 1) - 1);
+		}
+	}
+	
+	virtual FPrimitiveSceneProxy* CreateSceneProxy() override;
+};
+
diff --git a/Engine/Source/Runtime/Landscape/Private/Landscape.cpp b/Engine/Source/Runtime/Landscape/Private/Landscape.cpp
--- a/Engine/Source/Runtime/Landscape/Private/Landscape.cpp
+++ b/Engine/Source/Runtime/Landscape/Private/Landscape.cpp
@@ -29,6 +29,8 @@ Landscape.cpp: Terrain rendering
 #include "AI/Navigation/NavigationSystem.h"
 #include "EngineUtils.h"
 #include "Engine/StaticMesh.h"
+#include "LandscapeMeshProxyComponent.h"
+#include "LandscapeMeshProxyActor.h"
 
 #include "Materials/MaterialExpressionLandscapeLayerWeight.h"
 #include "Materials/MaterialExpressionLandscapeLayerSwitch.h"
@@ -2853,7 +2855,28 @@ static FAutoConsoleCommand FlushFoliageCmd(
 	FConsoleCommandWithArgsDelegate::CreateStatic(&FlushFoliage)
 	);
 
+//
+// ALandscapeMeshProxyActor
+//
+ALandscapeMeshProxyActor::ALandscapeMeshProxyActor(const FObjectInitializer& ObjectInitializer)
+: Super(ObjectInitializer)
+{
+	bCanBeDamaged = false;
+
+	LandscapeMeshProxyComponent = CreateDefaultSubobject<ULandscapeMeshProxyComponent>(TEXT("LandscapeMeshProxyComponent0"));
+	LandscapeMeshProxyComponent->SetCollisionProfileName(UCollisionProfile::NoCollision_ProfileName);
+	LandscapeMeshProxyComponent->Mobility = EComponentMobility::Static;
+	LandscapeMeshProxyComponent->bGenerateOverlapEvents = false;
 
+	RootComponent = LandscapeMeshProxyComponent;
+}
 
+//
+// ULandscapeMeshProxyComponent
+//
+ULandscapeMeshProxyComponent::ULandscapeMeshProxyComponent(const FObjectInitializer& ObjectInitializer)
+: Super(ObjectInitializer)
+{
+}
 
 #undef LOCTEXT_NAMESPACE
\ No newline at end of file
diff --git a/Engine/Source/Runtime/Landscape/Private/LandscapeRender.cpp b/Engine/Source/Runtime/Landscape/Private/LandscapeRender.cpp
--- a/Engine/Source/Runtime/Landscape/Private/LandscapeRender.cpp
+++ b/Engine/Source/Runtime/Landscape/Private/LandscapeRender.cpp
@@ -529,12 +529,11 @@ void FLandscapeEditToolRenderData::UpdateSelectionMaterial(int32 InSelectedType)
 //
 TMap<uint32, FLandscapeSharedBuffers*>FLandscapeComponentSceneProxy::SharedBuffersMap;
 TMap<uint32, FLandscapeSharedAdjacencyIndexBuffer*>FLandscapeComponentSceneProxy::SharedAdjacencyIndexBufferMap;
-TMap<FLandscapeComponentSceneProxy::FLandscapeKey, TMap<FIntPoint, const FLandscapeComponentSceneProxy*> > FLandscapeComponentSceneProxy::SharedSceneProxyMap;
+TMap<FLandscapeNeighborInfo::FLandscapeKey, TMap<FIntPoint, const FLandscapeNeighborInfo*> > FLandscapeNeighborInfo::SharedSceneProxyMap;
 
 FLandscapeComponentSceneProxy::FLandscapeComponentSceneProxy(ULandscapeComponent* InComponent, FLandscapeEditToolRenderData* InEditToolRenderData)
 	: FPrimitiveSceneProxy(InComponent)
-	, LandscapeKey(InComponent->GetWorld(), InComponent->GetLandscapeProxy()->GetLandscapeGuid())
-	, bAddedToSceneProxyMap(false)
+	, FLandscapeNeighborInfo(InComponent->GetWorld(), InComponent->GetLandscapeProxy()->GetLandscapeGuid(), InComponent->GetSectionBase() / InComponent->ComponentSizeQuads, InComponent->HeightmapTexture, InComponent->ForcedLOD, InComponent->LODBias)
 	, MaxLOD(FMath::CeilLogTwo(InComponent->SubsectionSizeQuads + 1) - 1)
 	, NumSubsections(InComponent->NumSubsections)
 	, SubsectionSizeQuads(InComponent->SubsectionSizeQuads)
@@ -543,13 +542,11 @@ FLandscapeComponentSceneProxy::FLandscapeComponentSceneProxy(ULandscapeComponent
 	, ComponentSizeVerts(InComponent->ComponentSizeQuads + 1)
 	, StaticLightingLOD(InComponent->GetLandscapeProxy()->StaticLightingLOD)
 	, SectionBase(InComponent->GetSectionBase())
-	, ComponentBase(InComponent->GetSectionBase() / InComponent->ComponentSizeQuads)
 	, WeightmapScaleBias(InComponent->WeightmapScaleBias)
 	, WeightmapSubsectionOffset(InComponent->WeightmapSubsectionOffset)
 	, WeightmapTextures(InComponent->WeightmapTextures)
 	, NumWeightmapLayerAllocations(InComponent->WeightmapLayerAllocations.Num())
 	, NormalmapTexture(InComponent->HeightmapTexture)
-	, HeightmapTexture(InComponent->HeightmapTexture)
 	, HeightmapScaleBias(InComponent->HeightmapScaleBias)
 	, XYOffsetmapTexture(InComponent->XYOffsetmapTexture)
 	, SharedBuffersKey(0)
@@ -559,10 +556,13 @@ FLandscapeComponentSceneProxy::FLandscapeComponentSceneProxy(ULandscapeComponent
 	, EditToolRenderData(InEditToolRenderData)
 	, ComponentLightInfo(nullptr)
 	, LandscapeComponent(InComponent)
-	, ForcedLOD(InComponent->ForcedLOD)
-	, LODBias(InComponent->LODBias)
 	, LODFalloff(InComponent->GetLandscapeProxy()->LODFalloff)
 {
+	if (!IsComponentLevelVisible())
+	{
+		bNeedsLevelAddedToWorldNotification = true;
+	}
+
 	LevelColor = FLinearColor(1.f, 1.f, 1.f);
 
 	const auto FeatureLevel = GetScene().GetFeatureLevel();
@@ -578,19 +578,6 @@ FLandscapeComponentSceneProxy::FLandscapeComponentSceneProxy(ULandscapeComponent
 		HeightmapSubsectionOffsetV = ((float)(InComponent->SubsectionSizeQuads + 1) / (float)HeightmapTexture->GetSizeY());
 	}
 
-	//       -Y       
-	//    - - 0 - -   
-	//    |       |   
-	// -X 1   P   2 +X
-	//    |       |   
-	//    - - 3 - -   
-	//       +Y       
-
-	Neighbors[0] = nullptr;
-	Neighbors[1] = nullptr;
-	Neighbors[2] = nullptr;
-	Neighbors[3] = nullptr;
-
 	LODBias = FMath::Clamp<int8>(LODBias, -MaxLOD, MaxLOD);
 
 	if (InComponent->GetLandscapeProxy()->MaxLODLevel >= 0)
@@ -664,42 +651,9 @@ void FLandscapeComponentSceneProxy::CreateRenderThreadResources()
 {
 	check(HeightmapTexture != nullptr);
 
-	// Register ourselves in the map.
-	TMap<FIntPoint, const FLandscapeComponentSceneProxy*>& SceneProxyMap = SharedSceneProxyMap.FindOrAdd(LandscapeKey);
-	
-	const FLandscapeComponentSceneProxy* Existing = SceneProxyMap.FindRef(ComponentBase);
-	if (Existing == nullptr)//(ensure(Existing == nullptr))
-	{
-		SceneProxyMap.Add(ComponentBase, this);
-		bAddedToSceneProxyMap = true;
-
-		// Find Neighbors
-		Neighbors[0] = SceneProxyMap.FindRef(ComponentBase + FIntPoint(0, -1));
-		Neighbors[1] = SceneProxyMap.FindRef(ComponentBase + FIntPoint(-1, 0));
-		Neighbors[2] = SceneProxyMap.FindRef(ComponentBase + FIntPoint(1, 0));
-		Neighbors[3] = SceneProxyMap.FindRef(ComponentBase + FIntPoint(0, 1));
-
-		// Add ourselves to our neighbors
-		if (Neighbors[0])
-		{
-			Neighbors[0]->Neighbors[3] = this;
-		}
-		if (Neighbors[1])
+	if (IsComponentLevelVisible())
 		{
-			Neighbors[1]->Neighbors[2] = this;
-		}
-		if (Neighbors[2])
-		{
-			Neighbors[2]->Neighbors[1] = this;
-		}
-		if (Neighbors[3])
-		{
-			Neighbors[3]->Neighbors[0] = this;
-		}
-	}
-	else
-	{
-		// UE_LOG(LogLandscape, Warning, TEXT("Duplicate ComponentBase %d, %d"), ComponentBase.X, ComponentBase.Y);
+		RegisterNeighbors();
 	}
 
 	auto FeatureLevel = GetScene().GetFeatureLevel();
@@ -765,46 +719,15 @@ void FLandscapeComponentSceneProxy::CreateRenderThreadResources()
 	LandscapeUniformShaderParameters.InitResource();
 }
 
-FLandscapeComponentSceneProxy::~FLandscapeComponentSceneProxy()
+void FLandscapeComponentSceneProxy::OnLevelAddedToWorld()
 {
-	// Remove ourselves from the map
-	if (bAddedToSceneProxyMap)
-	{
-		TMap<FIntPoint, const FLandscapeComponentSceneProxy*>* SceneProxyMap = SharedSceneProxyMap.Find(LandscapeKey);
-		check(SceneProxyMap);
+	RegisterNeighbors();
+}
 
-		const FLandscapeComponentSceneProxy* MapEntry = SceneProxyMap->FindRef(ComponentBase);
-		if (MapEntry == this) //(/*ensure*/(MapEntry == this))
-		{
-			SceneProxyMap->Remove(ComponentBase);
 
-			if (SceneProxyMap->Num() == 0)
-			{
-				// remove the entire LandscapeKey entry as this is the last scene proxy
-				SharedSceneProxyMap.Remove(LandscapeKey);
-			}
-			else
+FLandscapeComponentSceneProxy::~FLandscapeComponentSceneProxy()
 			{
-				// remove reference to us from our neighbors
-				if (Neighbors[0])
-				{
-					Neighbors[0]->Neighbors[3] = nullptr;
-				}
-				if (Neighbors[1])
-				{
-					Neighbors[1]->Neighbors[2] = nullptr;
-				}
-				if (Neighbors[2])
-				{
-					Neighbors[2]->Neighbors[1] = nullptr;
-				}
-				if (Neighbors[3])
-				{
-					Neighbors[3]->Neighbors[0] = nullptr;
-				}
-			}
-		}
-	}
+	UnregisterNeighbors();
 
 	// Free the subsection uniform buffer
 	LandscapeUniformShaderParameters.ReleaseResource();
@@ -2590,4 +2513,153 @@ void FLandscapeComponentSceneProxy::GetHeightfieldRepresentation(UTexture2D*& Ou
 		//@todo - subsections duplicate a line of border heights in the heightmap, take this into account in the UV mapping
 		HeightmapScaleBias.Z + SubsectionSizeVerts * NumSubsections * HeightmapScaleBias.X, 
 		HeightmapScaleBias.W + SubsectionSizeVerts * NumSubsections * HeightmapScaleBias.Y);
-}
\ No newline at end of file
+}
+
+//
+// FLandscapeNeighborInfo
+//
+void FLandscapeNeighborInfo::RegisterNeighbors()
+{
+	if (!bRegistered)
+	{
+		// Register ourselves in the map.
+		TMap<FIntPoint, const FLandscapeNeighborInfo*>& SceneProxyMap = SharedSceneProxyMap.FindOrAdd(LandscapeKey);
+
+		const FLandscapeNeighborInfo* Existing = SceneProxyMap.FindRef(ComponentBase);
+		if (Existing == nullptr)//(ensure(Existing == nullptr))
+		{
+			SceneProxyMap.Add(ComponentBase, this);
+			bRegistered = true;
+
+			// Find Neighbors
+			Neighbors[0] = SceneProxyMap.FindRef(ComponentBase + FIntPoint(0, -1));
+			Neighbors[1] = SceneProxyMap.FindRef(ComponentBase + FIntPoint(-1, 0));
+			Neighbors[2] = SceneProxyMap.FindRef(ComponentBase + FIntPoint(1, 0));
+			Neighbors[3] = SceneProxyMap.FindRef(ComponentBase + FIntPoint(0, 1));
+
+			// Add ourselves to our neighbors
+			if (Neighbors[0])
+			{
+				Neighbors[0]->Neighbors[3] = this;
+			}
+			if (Neighbors[1])
+			{
+				Neighbors[1]->Neighbors[2] = this;
+			}
+			if (Neighbors[2])
+			{
+				Neighbors[2]->Neighbors[1] = this;
+			}
+			if (Neighbors[3])
+			{
+				Neighbors[3]->Neighbors[0] = this;
+			}
+		}
+		else
+		{
+			UE_LOG(LogLandscape, Warning, TEXT("Duplicate ComponentBase %d, %d"), ComponentBase.X, ComponentBase.Y);
+		}
+	}
+}
+
+void FLandscapeNeighborInfo::UnregisterNeighbors()
+{
+	if (bRegistered)
+	{
+		// Remove ourselves from the map
+		TMap<FIntPoint, const FLandscapeNeighborInfo*>* SceneProxyMap = SharedSceneProxyMap.Find(LandscapeKey);
+		check(SceneProxyMap);
+
+		const FLandscapeNeighborInfo* MapEntry = SceneProxyMap->FindRef(ComponentBase);
+		if (MapEntry == this) //(/*ensure*/(MapEntry == this))
+		{
+			SceneProxyMap->Remove(ComponentBase);
+
+			if (SceneProxyMap->Num() == 0)
+			{
+				// remove the entire LandscapeKey entry as this is the last scene proxy
+				SharedSceneProxyMap.Remove(LandscapeKey);
+			}
+			else
+			{
+				// remove reference to us from our neighbors
+				if (Neighbors[0])
+				{
+					Neighbors[0]->Neighbors[3] = nullptr;
+				}
+				if (Neighbors[1])
+				{
+					Neighbors[1]->Neighbors[2] = nullptr;
+				}
+				if (Neighbors[2])
+				{
+					Neighbors[2]->Neighbors[1] = nullptr;
+				}
+				if (Neighbors[3])
+				{
+					Neighbors[3]->Neighbors[0] = nullptr;
+				}
+			}
+		}
+	}
+}
+
+//
+// FLandscapeMeshProxySceneProxy
+//
+FLandscapeMeshProxySceneProxy::FLandscapeMeshProxySceneProxy(UStaticMeshComponent* InComponent, const FGuid& InGuid, const TArray<FIntPoint>& InProxyComponentBases, int8 InProxyLOD)
+: FStaticMeshSceneProxy(InComponent)
+{
+	if (!IsComponentLevelVisible())
+	{
+		bNeedsLevelAddedToWorldNotification = true;
+	}
+
+	ProxyNeighborInfos.Empty(InProxyComponentBases.Num());
+	for (FIntPoint ComponentBase : InProxyComponentBases)
+	{
+		new(ProxyNeighborInfos) FLandscapeNeighborInfo(InComponent->GetWorld(), InGuid, ComponentBase, nullptr, InProxyLOD, 0);
+	}
+}
+
+void FLandscapeMeshProxySceneProxy::CreateRenderThreadResources()
+{
+	FStaticMeshSceneProxy::CreateRenderThreadResources();
+
+	if (IsComponentLevelVisible())
+	{
+		for (FLandscapeNeighborInfo& Info : ProxyNeighborInfos)
+		{
+			Info.RegisterNeighbors();
+		}
+	}
+}
+
+void FLandscapeMeshProxySceneProxy::OnLevelAddedToWorld()
+{
+	for (FLandscapeNeighborInfo& Info : ProxyNeighborInfos)
+	{
+		Info.RegisterNeighbors();
+	}
+}
+
+FLandscapeMeshProxySceneProxy::~FLandscapeMeshProxySceneProxy()
+{
+	for (FLandscapeNeighborInfo& Info : ProxyNeighborInfos)
+	{
+		Info.UnregisterNeighbors();
+	}
+}
+
+FPrimitiveSceneProxy* ULandscapeMeshProxyComponent::CreateSceneProxy()
+{
+	if (StaticMesh == NULL
+		|| StaticMesh->RenderData == NULL
+		|| StaticMesh->RenderData->LODResources.Num() == 0
+		|| StaticMesh->RenderData->LODResources[0].VertexBuffer.GetNumVertices() == 0)
+	{
+		return NULL;
+	}
+
+	return new FLandscapeMeshProxySceneProxy(this, LandscapeGuid, ProxyComponentBases, ProxyLOD);
+}
diff --git a/Engine/Source/Runtime/Landscape/Public/LandscapeRender.h b/Engine/Source/Runtime/Landscape/Public/LandscapeRender.h
--- a/Engine/Source/Runtime/Landscape/Public/LandscapeRender.h
+++ b/Engine/Source/Runtime/Landscape/Public/LandscapeRender.h
@@ -10,8 +10,10 @@ LandscapeRender.h: New terrain rendering
 
 #include "LandscapeComponent.h"
 #include "LandscapeProxy.h"
+#include "LandscapeMeshProxyComponent.h"
 
 #include "PrimitiveSceneProxy.h"
+#include "StaticMeshResources.h"
 
 #include "LightMap.h"
 #include "MeshBatch.h"
@@ -20,8 +22,6 @@ LandscapeRender.h: New terrain rendering
 // This defines the number of border blocks to surround terrain by when generating lightmaps
 #define TERRAIN_PATCH_EXPAND_SCALAR	1
 
-#define LANDSCAPE_NEIGHBOR_NUM	4
-
 #define LANDSCAPE_LOD_LEVELS 8
 #define LANDSCAPE_MAX_SUBSECTION_NUM 2
 
@@ -348,12 +348,12 @@ struct FLandscapeEditToolRenderData
 };
 
 //
-// FLandscapeComponentSceneProxy
+// FLandscapeNeighborInfo
 //
-class FLandscapeComponentSceneProxy : public FPrimitiveSceneProxy
+class FLandscapeNeighborInfo
 {
-	friend class FLandscapeSharedBuffers;
-
+	bool bRegistered;
+protected:
 	// Key to uniquely identify the landscape to find the correct render proxy map
 	class FLandscapeKey
 	{
@@ -361,8 +361,8 @@ class FLandscapeComponentSceneProxy : public FPrimitiveSceneProxy
 		const FGuid Guid;
 	public:
 		FLandscapeKey(const UWorld* InWorld, const FGuid& InGuid)
-		: World(InWorld)
-		, Guid(InGuid)
+			: World(InWorld)
+			, Guid(InGuid)
 		{}
 
 		friend inline uint32 GetTypeHash(const FLandscapeKey& LandscapeKey)
@@ -376,6 +376,73 @@ class FLandscapeComponentSceneProxy : public FPrimitiveSceneProxy
 		}
 	};
 
+	// Map of currently registered landscape proxies, used to register with our neighbors
+	static TMap<FLandscapeKey, TMap<FIntPoint, const FLandscapeNeighborInfo*> > SharedSceneProxyMap;
+
+	// For neighbor lookup
+	FLandscapeKey			LandscapeKey;
+	FIntPoint				ComponentBase;
+
+	// Pointer to our neighbor's scene proxies in NWES order (nullptr if there is currently no neighbor)
+	mutable const FLandscapeNeighborInfo* Neighbors[4];
+
+	
+	// Data we need to be able to access about our neighbor
+	UTexture2D*				HeightmapTexture; // PC : Heightmap, Mobile : Weightmap
+	int8					ForcedLOD;
+	int8					LODBias;
+
+	friend class FLandscapeComponentSceneProxy;
+
+public:
+	FLandscapeNeighborInfo(const UWorld* InWorld, const FGuid& InGuid, const FIntPoint& InComponentBase, UTexture2D* InHeightmapTexture, int8 InForcedLOD, int8 InLODBias)
+	: bRegistered(false)
+	, LandscapeKey(InWorld, InGuid)
+	, ComponentBase(InComponentBase)
+	, HeightmapTexture(InHeightmapTexture)
+	, ForcedLOD(InForcedLOD)
+	, LODBias(InLODBias)
+	{
+		//       -Y       
+		//    - - 0 - -   
+		//    |       |   
+		// -X 1   P   2 +X
+		//    |       |   
+		//    - - 3 - -   
+		//       +Y       
+
+		Neighbors[0] = nullptr;
+		Neighbors[1] = nullptr;
+		Neighbors[2] = nullptr;
+		Neighbors[3] = nullptr;
+	}
+
+	void RegisterNeighbors();
+	void UnregisterNeighbors();
+};
+
+//
+// FLandscapeMeshProxySceneProxy
+//
+class FLandscapeMeshProxySceneProxy : public FStaticMeshSceneProxy
+{
+	TArray<FLandscapeNeighborInfo> ProxyNeighborInfos;
+public:
+	FLandscapeMeshProxySceneProxy(UStaticMeshComponent* InComponent, const FGuid& InGuid, const TArray<FIntPoint>& InProxyComponentBases, int8 InProxyLOD);
+	virtual ~FLandscapeMeshProxySceneProxy();
+	virtual void CreateRenderThreadResources() override;
+	virtual void OnLevelAddedToWorld() override;
+};
+
+
+//
+// FLandscapeComponentSceneProxy
+//
+class FLandscapeComponentSceneProxy : public FPrimitiveSceneProxy, public FLandscapeNeighborInfo
+{
+	friend class FLandscapeSharedBuffers;
+
+
 	class FLandscapeLCI final : public FLightCacheInterface
 	{
 	public:
@@ -411,8 +478,6 @@ class FLandscapeComponentSceneProxy : public FPrimitiveSceneProxy
 	};
 
 protected:
-	FLandscapeKey				LandscapeKey;
-	bool						bAddedToSceneProxyMap;
 	int8						MaxLOD;
 	int8						NumSubsections;
 	int16						SubsectionSizeQuads;
@@ -422,7 +487,6 @@ protected:
 	uint8						StaticLightingLOD;
 	float						StaticLightingResolution;
 	FIntPoint					SectionBase;
-	FIntPoint					ComponentBase;
 	FMatrix						LocalToWorldNoScaling;
 
 	// Storage for static draw list batch params
@@ -436,10 +500,7 @@ protected:
 	float WeightmapSubsectionOffset;
 	TArray<UTexture2D*> WeightmapTextures;
 	int8 NumWeightmapLayerAllocations;
-
 	UTexture2D* NormalmapTexture; // PC : Heightmap, Mobile : Weightmap
-
-	UTexture2D* HeightmapTexture; // PC : Heightmap, Mobile : Weightmap
 	FVector4 HeightmapScaleBias;
 	float HeightmapSubsectionOffsetU;
 	float HeightmapSubsectionOffsetV;
@@ -459,9 +520,6 @@ protected:
 	static TMap<uint32, FLandscapeSharedBuffers*> SharedBuffersMap;
 	static TMap<uint32, FLandscapeSharedAdjacencyIndexBuffer*> SharedAdjacencyIndexBufferMap;
 
-	// Map of currently registered landscape proxies, used to register with our neighbors
-	static TMap<FLandscapeKey, TMap<FIntPoint, const FLandscapeComponentSceneProxy*> > SharedSceneProxyMap;
-
 	FLandscapeEditToolRenderData* EditToolRenderData;
 
 	// FLightCacheInterface
@@ -469,12 +527,6 @@ protected:
 
 	const ULandscapeComponent* LandscapeComponent;
 
-	int8					ForcedLOD;
-	int8					LODBias;
-
-	// Pointer to our neighbor's scene proxies in NWES order (nullptr if there is currently no neighbor)
-	mutable const FLandscapeComponentSceneProxy* Neighbors[4];
-
 	ELandscapeLODFalloff::Type LODFalloff;
 
 	TUniformBuffer<FLandscapeUniformShaderParameters> LandscapeUniformShaderParameters;
@@ -497,7 +549,8 @@ public:
 	virtual void GetLightRelevance(const FLightSceneProxy* LightSceneProxy, bool& bDynamic, bool& bRelevant, bool& bLightMapped, bool& bShadowMapped) const override;
 	virtual void OnTransformChanged() override;
 	virtual void CreateRenderThreadResources() override;
-
+	virtual void OnLevelAddedToWorld() override;
+	
 	friend class ULandscapeComponent;
 	friend class FLandscapeVertexFactoryVertexShaderParameters;
 	friend class FLandscapeXYOffsetVertexFactoryVertexShaderParameters;
diff --git a/Engine/Source/Runtime/Renderer/Private/RendererScene.cpp b/Engine/Source/Runtime/Renderer/Private/RendererScene.cpp
--- a/Engine/Source/Runtime/Renderer/Private/RendererScene.cpp
+++ b/Engine/Source/Runtime/Renderer/Private/RendererScene.cpp
@@ -2174,6 +2174,10 @@ void FScene::OnLevelAddedToWorld_RenderThread(FName InLevelName)
 		if (Proxy->LevelName == InLevelName)
 		{
 			Proxy->bIsComponentLevelVisible = true;
+			if (Proxy->NeedsLevelAddedToWorldNotification())
+			{
+				Proxy->OnLevelAddedToWorld();
+			}
 		}
 	}
 }
