UE4 - grass fixes. Eliminated the possibility of grass cache thrash. Verified reasonable memory use plus minor memory fix. Support for flushing only certain components. Reuse weightmaps in editor to cut down on hitching. Only destroy one component per frame.

diff --git a/Engine/Source/Runtime/Engine/Private/HierarchicalInstancedStaticMesh.cpp b/Engine/Source/Runtime/Engine/Private/HierarchicalInstancedStaticMesh.cpp
--- a/Engine/Source/Runtime/Engine/Private/HierarchicalInstancedStaticMesh.cpp
+++ b/Engine/Source/Runtime/Engine/Private/HierarchicalInstancedStaticMesh.cpp
@@ -1462,6 +1462,7 @@ void UHierarchicalInstancedStaticMeshComponent::ClearInstances()
 
 	ClusterTreePtr = MakeShareable(new TArray<FClusterNode>);
 	NumBuiltInstances = 0;
+	SortedInstances.Empty();
 	UnbuiltInstanceBounds.Init();
 
 	Super::ClearInstances();
diff --git a/Engine/Source/Runtime/Landscape/Classes/LandscapeProxy.h b/Engine/Source/Runtime/Landscape/Classes/LandscapeProxy.h
--- a/Engine/Source/Runtime/Landscape/Classes/LandscapeProxy.h
+++ b/Engine/Source/Runtime/Landscape/Classes/LandscapeProxy.h
@@ -193,11 +193,30 @@ namespace ELandscapeLODFalloff
 
 struct FCachedLandscapeFoliage
 {
+	struct FPerGrassComp
+	{
+		TWeakObjectPtr<UHierarchicalInstancedStaticMeshComponent> Foliage;
+		uint32 LastUsedFrameNumber;
+		double LastUsedTime;
+		bool Pending;
+
+		FPerGrassComp(UHierarchicalInstancedStaticMeshComponent* InFoliage)
+			: Foliage(InFoliage)
+			, Pending(true)
+		{
+			Touch();
+		}
+		void Touch()
+		{
+			LastUsedFrameNumber = GFrameNumber;
+			LastUsedTime = FPlatformTime::Seconds();
+		}
+	};
 	struct FPerLayer
 	{
 		int32 SqrtSubsections;
 		int32 CachedMaxInstancesPerComponent;
-		TArray<TWeakObjectPtr<UHierarchicalInstancedStaticMeshComponent> > Foliage;
+		TArray<FPerGrassComp> Foliage;
 		TWeakObjectPtr<ULandscapeLayerInfoObject> Layer;
 
 		FPerLayer(int32 InSqrtSubsections, int32 InCachedMaxInstancesPerComponent, ULandscapeLayerInfoObject* InLayer)
@@ -428,7 +447,7 @@ public:
 	virtual ALandscape* GetLandscapeActor();
 
 	/** Flush the foliage cache */
-	LANDSCAPE_API void FlushFoliageComponents();
+	LANDSCAPE_API void FlushFoliageComponents(TSet<ULandscapeComponent*>* OnlyForComponents = nullptr);
 
 	/** 
 		Update Foliage 
diff --git a/Engine/Source/Runtime/Landscape/Private/Landscape.cpp b/Engine/Source/Runtime/Landscape/Private/Landscape.cpp
--- a/Engine/Source/Runtime/Landscape/Private/Landscape.cpp
+++ b/Engine/Source/Runtime/Landscape/Private/Landscape.cpp
@@ -2050,41 +2050,82 @@ bool LandscapeMaterialsParameterSetUpdater(FStaticParameterSet &StaticParameterS
 	return UpdateParameterSet<FStaticTerrainLayerWeightParameter, UMaterialExpressionLandscapeLayerWeight>(StaticParameterSet.TerrainLayerWeightParameters, ParentMaterial);
 }
 
-void ALandscapeProxy::FlushFoliageComponents()
+void ALandscapeProxy::FlushFoliageComponents(TSet<ULandscapeComponent*>* OnlyForComponents)
 {
-	// Clear old foliage component containers
-	FoliageComponents.Empty();
-
-	// Might as well clear the cache...
-	FoliageCache.ClearCache();
-
-	// Destroy any owned foliage components
-	TInlineComponentArray<UHierarchicalInstancedStaticMeshComponent*> FoliageComps;
-	GetComponents(FoliageComps);
-	for (UHierarchicalInstancedStaticMeshComponent* Component : FoliageComps)
+	if (OnlyForComponents)
 	{
-		Component->ClearInstances();
-		Component->DestroyComponent();
+		OnlyForComponents->Add(nullptr); // if the weak pointer in the cache is invalid, we should kill them anyway
+		for (int32 CacheIndex = 0; CacheIndex < FoliageCache.PerComponent.Num(); CacheIndex++)
+		{
+			FCachedLandscapeFoliage::FPerComponent& CacheItem = FoliageCache.PerComponent[CacheIndex];
+			if (OnlyForComponents->Contains(CacheItem.BasedOn.Get()))
+			{
+				for (int32 LayerIndex = 0; LayerIndex < CacheItem.Layers.Num(); LayerIndex++)
+				{
+					FCachedLandscapeFoliage::FPerLayer& LayerItem = CacheItem.Layers[LayerIndex];
+					for (int32 FoliageIndex = 0; FoliageIndex < LayerItem.Foliage.Num(); FoliageIndex++)
+					{
+						FCachedLandscapeFoliage::FPerGrassComp& GrassItem = LayerItem.Foliage[FoliageIndex];
+						UHierarchicalInstancedStaticMeshComponent *Used = GrassItem.Foliage.Get();
+						if (Used)
+						{
+							Used->ClearInstances();
+							Used->DestroyComponent();
+						}
+					}
+				}
+			}
+		}
 	}
-
-	TArray<USceneComponent*> AttachedFoliageComponents = RootComponent->AttachChildren.FilterByPredicate(
-		[](USceneComponent* Component)
+	else
 	{
-		return Cast<UHierarchicalInstancedStaticMeshComponent>(Component);
-	});
+		// Clear old foliage component containers
+		FoliageComponents.Empty();
 
-	// Destroy any attached but un-owned foliage components
-	for (USceneComponent* Component : AttachedFoliageComponents)
-	{
-		Component->DestroyComponent();
+		// Might as well clear the cache...
+		FoliageCache.ClearCache();
+		// Destroy any owned foliage components
+		TInlineComponentArray<UHierarchicalInstancedStaticMeshComponent*> FoliageComps;
+		GetComponents(FoliageComps);
+		for (UHierarchicalInstancedStaticMeshComponent* Component : FoliageComps)
+		{
+			Component->ClearInstances();
+			Component->DestroyComponent();
+		}
+
+		TArray<USceneComponent*> AttachedFoliageComponents = RootComponent->AttachChildren.FilterByPredicate(
+			[](USceneComponent* Component)
+		{
+			return Cast<UHierarchicalInstancedStaticMeshComponent>(Component);
+		});
+
+		// Destroy any attached but un-owned foliage components
+		for (USceneComponent* Component : AttachedFoliageComponents)
+		{
+			Component->DestroyComponent();
+		}
 	}
 }
 
+//// Grass
+
+DEFINE_LOG_CATEGORY_STATIC(LogGrass, Log, All);
+
 static TAutoConsoleVariable<float> CVarGuardBandMultiplier(
 	TEXT("grass.GuardBandMultiplier"),
 	1.5f,
 	TEXT("Used to control discarding in the grass system. Approximate range, 1-2. Multiplied by the cull distance to control when we discard grass components."));
 
+static TAutoConsoleVariable<int32> CVarMinFramesToKeepGrass(
+	TEXT("grass.MinFramesToKeepGrass"),
+	30,
+	TEXT("Minimum number of frames before cached grass can be discarded; used to prevent thrashing."));
+
+static TAutoConsoleVariable<float> CVarMinTimeToKeepGrass(
+	TEXT("grass.MinTimeToKeepGrass"),
+	5.0f,
+	TEXT("Minimum number of seconds before cached grass can be discarded; used to prevent thrashing."));
+
 static TAutoConsoleVariable<int32> CVarMinInstancesPerLeaf(
 	TEXT("grass.MinInstancesPerLeaf"),
 	256,
@@ -2107,20 +2148,29 @@ class FEditorGrassLayerData
 	TArray<FColor> LocalHeightCache;
 	TArray<uint8> LocalCache;
 public:
-	FEditorGrassLayerData(ALandscapeProxy* Landscape, ULandscapeComponent* Component, ULandscapeLayerInfoObject* LayerInfo, UHierarchicalInstancedStaticMeshComponent* HierarchicalInstancedStaticMeshComponent)
+	FEditorGrassLayerData(ALandscapeProxy* Landscape, ULandscapeComponent* Component, ULandscapeLayerInfoObject* LayerInfo, UHierarchicalInstancedStaticMeshComponent* HierarchicalInstancedStaticMeshComponent, FEditorGrassLayerData* CopyFrom)
 	{
 		int32 SubsectionSizeQuads = Component->SubsectionSizeQuads;
 		int32 NumSubsections = Component->NumSubsections;
 		Stride = (SubsectionSizeQuads + 1) * NumSubsections;
 		if (Stride)
 		{
-			FLandscapeComponentDataInterface CDI(Component);
-			CDI.GetHeightmapTextureData(LocalHeightCache);
-			if (!CDI.GetWeightmapTextureData(LayerInfo, LocalCache) || !LocalCache.Num() || LocalCache.Num() != LocalHeightCache.Num())
+			if (CopyFrom)
 			{
-				// no data for this layer for this component.
-				LocalHeightCache.Empty();
-				LocalCache.Empty();
+				check(Stride == CopyFrom->Stride);
+				LocalHeightCache = CopyFrom->LocalHeightCache;
+				LocalCache = CopyFrom->LocalCache;
+			}
+			else
+			{
+				FLandscapeComponentDataInterface CDI(Component);
+				CDI.GetHeightmapTextureData(LocalHeightCache);
+				if (!CDI.GetWeightmapTextureData(LayerInfo, LocalCache) || !LocalCache.Num() || LocalCache.Num() != LocalHeightCache.Num())
+				{
+					// no data for this layer for this component.
+					LocalHeightCache.Empty();
+					LocalCache.Empty();
+				}
 			}
 		}
 	}
@@ -2149,7 +2199,7 @@ class FGameGrassLayerData
 	TSharedPtr<FGrassMap, ESPMode::ThreadSafe> GrassMap;
 	const uint8* GrassMapPtr;
 public:
-	FGameGrassLayerData(ALandscapeProxy* Landscape, ULandscapeComponent* Component, ULandscapeLayerInfoObject* LayerInfo, UHierarchicalInstancedStaticMeshComponent* HierarchicalInstancedStaticMeshComponent)
+	FGameGrassLayerData(ALandscapeProxy* Landscape, ULandscapeComponent* Component, ULandscapeLayerInfoObject* LayerInfo, UHierarchicalInstancedStaticMeshComponent* HierarchicalInstancedStaticMeshComponent, FGameGrassLayerData*)
 	{
 		int32 SubsectionSizeQuads = Component->SubsectionSizeQuads;
 		int32 NumSubsections = Component->NumSubsections;
@@ -2281,9 +2331,9 @@ struct FAsyncGrassBuilder : public FGrassBuilderBase
 	TArray<int32> SortedInstances;
 	TArray<int32> InstanceReorderTable;
 
-	FAsyncGrassBuilder(ALandscapeProxy* Landscape, ULandscapeComponent* Component, ULandscapeLayerInfoObject* LayerInfo, UHierarchicalInstancedStaticMeshComponent* HierarchicalInstancedStaticMeshComponent, int32 SqrtSubsections, int32 SubX, int32 SubY)
+	FAsyncGrassBuilder(ALandscapeProxy* Landscape, ULandscapeComponent* Component, ULandscapeLayerInfoObject* LayerInfo, UHierarchicalInstancedStaticMeshComponent* HierarchicalInstancedStaticMeshComponent, int32 SqrtSubsections, int32 SubX, int32 SubY, FAsyncGrassBuilder* CopyFrom)
 		: FGrassBuilderBase(Landscape, Component, LayerInfo, SqrtSubsections, SubX, SubY)
-		, LayerData(Landscape, Component, LayerInfo, HierarchicalInstancedStaticMeshComponent)
+		, LayerData(Landscape, Component, LayerInfo, HierarchicalInstancedStaticMeshComponent, CopyFrom ? &CopyFrom->LayerData : nullptr)
 	{
 		bHaveValidData = bHaveValidData && LayerData.HasData();
 		XForm = LandscapeToWorld * HierarchicalInstancedStaticMeshComponent->ComponentToWorld.ToMatrixWithScale().Inverse();
@@ -2469,69 +2519,27 @@ void ALandscapeProxy::Tick(float DeltaSeconds)
 		return;
 	}
 
-	struct FOldCamera
-	{
-		FVector Pos;
-		double LastTime;
-	};
-
 	// there is a bug here, which often leaves us with no cameras in the editor
-	// also sometimes the cameras are ping ponging and otherise really wrong in the editor
-	// the main thing is not to thrash
-	static TArray<FOldCamera> OldCameras;
+	static TArray<FVector> OldCameras;
 
-	double Now = FPlatformTime::Seconds();
-	for (auto& NewCam : World->ViewLocationsRenderedLastFrame)
+	if (!OldCameras.Num() && !World->ViewLocationsRenderedLastFrame.Num())
 	{
-		bool bFound = false;
-		for (auto& OldCam : OldCameras)
-		{
-			if (FVector::Dist(OldCam.Pos, NewCam) < 1000.0f)
-			{
-				OldCam.Pos = NewCam;
-				OldCam.LastTime = Now;
-				bFound = true;
-				break;
-			}
-		}
-		if (!bFound)
-		{
-			FOldCamera NewOldCam;
-			NewOldCam.Pos = NewCam;
-			NewOldCam.LastTime = Now;
-			OldCameras.Add(NewOldCam);
-		}
-	}
-	for (int32 Index = 0 ; Index < OldCameras.Num(); Index++)
-	{
-		if (OldCameras[Index].LastTime < Now - 2.0f)
-		{
-			OldCameras.RemoveAtSwap(Index);
-		}
+		// no cameras, no grass update
+		return;
 	}
 
-	static TArray<FVector> UseCameras;
-	check(IsInGameThread());
-	UseCameras.Reset();
-	for (auto& OldCam : OldCameras)
-	{
-		UseCameras.Add(OldCam.Pos);
-	}
+	const TArray<FVector>& Cameras = World->ViewLocationsRenderedLastFrame.Num() ? World->ViewLocationsRenderedLastFrame : OldCameras;
 
-	if (!UseCameras.Num())
+	if (&Cameras != &OldCameras)
 	{
-		// no cameras, no grass update
-		return;
+		check(IsInGameThread());
+		OldCameras = Cameras;
 	}
-
-	UpdateFoliage(UseCameras);
+	UpdateFoliage(Cameras);
 }
 
 void ALandscapeProxy::UpdateFoliage(const TArray<FVector>& Cameras, ULandscapeComponent* OnlyComponent, bool bForceSync)
 {
-	FCachedLandscapeFoliage DesiredFoliage;
-	TSet<UHierarchicalInstancedStaticMeshComponent *> StillUsed;
-
 	if (CVarGrassEnable.GetValueOnAnyThread() > 0)
 	{
 		double ComponentTime = 0.0;
@@ -2552,8 +2560,8 @@ void ALandscapeProxy::UpdateFoliage(const TArray<FVector>& Cameras, ULandscapeCo
 			}
 			ComponentTime -= FPlatformTime::Seconds();
 
-			const FCachedLandscapeFoliage::FPerComponent* ExistingCache = nullptr;
-			for (const FCachedLandscapeFoliage::FPerComponent& CacheItem : FoliageCache.PerComponent)
+			FCachedLandscapeFoliage::FPerComponent* ExistingCache = nullptr;
+			for (FCachedLandscapeFoliage::FPerComponent& CacheItem : FoliageCache.PerComponent)
 			{
 				if (CacheItem.BasedOn.Get() == Component)
 				{
@@ -2561,7 +2569,12 @@ void ALandscapeProxy::UpdateFoliage(const TArray<FVector>& Cameras, ULandscapeCo
 					break;
 				}
 			}
-			FCachedLandscapeFoliage::FPerComponent NewPerComponent(Component);
+			if (!ExistingCache)
+			{
+				ExistingCache = &FoliageCache.PerComponent[FoliageCache.PerComponent.Add(FCachedLandscapeFoliage::FPerComponent(Component))];
+			}
+			check(ExistingCache);
+
 
 			FBoxSphereBounds WorldBounds = Component->CalcBounds(Component->ComponentToWorld);
 			float MinDistanceToComp = Cameras.Num() ? MAX_flt : 0.0f;
@@ -2586,44 +2599,49 @@ void ALandscapeProxy::UpdateFoliage(const TArray<FVector>& Cameras, ULandscapeCo
 					{
 						continue;
 					}
-					if (ExistingCache)
+					bool bFound = false;
+					FCachedLandscapeFoliage::FPerLayer* ExistingLayer = nullptr;
+					for (int32 Index = 0; Index < ExistingCache->Layers.Num(); Index++)
 					{
-						bool bFound = false;
-						for (const FCachedLandscapeFoliage::FPerLayer& Layer : ExistingCache->Layers)
+						FCachedLandscapeFoliage::FPerLayer& Layer = ExistingCache->Layers[Index];
+						if (Layer.Layer.Get() == LayerInfo && Layer.CachedMaxInstancesPerComponent == MaxInstancesPerComponent)
 						{
-							if (Layer.Layer.Get() == LayerInfo && Layer.CachedMaxInstancesPerComponent == MaxInstancesPerComponent)
+							ExistingLayer = &Layer;
+							check(Layer.SqrtSubsections >= 1);
+							int32 NumValid = 0;
+							for (const FCachedLandscapeFoliage::FPerGrassComp& GrassItem : Layer.Foliage)
 							{
-								check(Layer.SqrtSubsections >= 1);
-								int32 NumValid = 0;
-								for (const TWeakObjectPtr<UHierarchicalInstancedStaticMeshComponent>& Comp : Layer.Foliage)
+								const TWeakObjectPtr<UHierarchicalInstancedStaticMeshComponent>& Comp = GrassItem.Foliage;
+								if (Comp.Get())
 								{
-									if (Comp.Get())
-									{
-										NumValid++;
-									}
-									else
-									{
-										NumValid = -1;
-										break;
-									}
+									NumValid++;
 								}
-								if (NumValid == Layer.SqrtSubsections * Layer.SqrtSubsections)
+								else
 								{
-									NewPerComponent.Layers.Add(Layer);
-									for (const TWeakObjectPtr<UHierarchicalInstancedStaticMeshComponent>& Comp : Layer.Foliage)
-									{
-										StillUsed.Add(Comp.Get());
-									}
-									bFound = true;
+									NumValid = -1;
 									break;
 								}
 							}
-						}
-						if (bFound || MinDistanceToComp > MustHaveDistance)
-						{
-							continue;
+							if (NumValid == Layer.SqrtSubsections * Layer.SqrtSubsections)
+							{
+								for (FCachedLandscapeFoliage::FPerGrassComp& GrassItem : Layer.Foliage)
+								{
+									GrassItem.Touch();
+								}
+								bFound = true;
+								break;
+							}
+							else
+							{
+								// this must be stale, lets delete them
+								Layer.Foliage.Empty();
+							}
 						}
 					}
+					if (bFound || MinDistanceToComp > MustHaveDistance)
+					{
+						continue;
+					}
 
 					FGrassBuilderBase ForSubsectionMath(this, Component, LayerInfo);
 
@@ -2633,9 +2651,14 @@ void ALandscapeProxy::UpdateFoliage(const TArray<FVector>& Cameras, ULandscapeCo
 					{
 						SqrtSubsections = FMath::Clamp<int32>(FMath::CeilToInt(float(ForSubsectionMath.SqrtMaxInstances) / FMath::Sqrt((float)MaxInstancesPerComponent)), 1, 16);
 					}
+					if (!ExistingLayer)
+					{
+						ExistingLayer = &ExistingCache->Layers[ExistingCache->Layers.Add(FCachedLandscapeFoliage::FPerLayer(SqrtSubsections, MaxInstancesPerComponent, LayerInfo))];
+					}
+					check(ExistingLayer);
+					FCachedLandscapeFoliage::FPerLayer& NewLayer = *ExistingLayer;
 
-					FCachedLandscapeFoliage::FPerLayer& NewLayer = NewPerComponent.Layers[NewPerComponent.Layers.Add(FCachedLandscapeFoliage::FPerLayer(SqrtSubsections, MaxInstancesPerComponent, LayerInfo))];
-
+					FAsyncGrassBuilder* CopyFromBuilder = nullptr;
 					for (int32 SubX = 0; SubX < SqrtSubsections; SubX++)
 					{
 						for (int32 SubY = 0; SubY < SqrtSubsections; SubY++)
@@ -2691,11 +2714,11 @@ void ALandscapeProxy::UpdateFoliage(const TArray<FVector>& Cameras, ULandscapeCo
 							HierarchicalInstancedStaticMeshComponent->SetWorldTransform(DesiredTransform);
 
 							FoliageComponents.Add(HierarchicalInstancedStaticMeshComponent);
-							StillUsed.Add(HierarchicalInstancedStaticMeshComponent);
-							NewLayer.Foliage.Add(HierarchicalInstancedStaticMeshComponent);
+							NewLayer.Foliage.Add(FCachedLandscapeFoliage::FPerGrassComp(HierarchicalInstancedStaticMeshComponent));
 
 
-							FAsyncGrassBuilder* Builder = new FAsyncGrassBuilder(this, Component, LayerInfo, HierarchicalInstancedStaticMeshComponent, SqrtSubsections, SubX, SubY);
+							FAsyncGrassBuilder* Builder = new FAsyncGrassBuilder(this, Component, LayerInfo, HierarchicalInstancedStaticMeshComponent, SqrtSubsections, SubX, SubY, CopyFromBuilder);
+							CopyFromBuilder = Builder;
 
 							if (Builder->bHaveValidData)
 							{
@@ -2722,100 +2745,148 @@ void ALandscapeProxy::UpdateFoliage(const TArray<FVector>& Cameras, ULandscapeCo
 					}
 				}
 			}
-			if (NewPerComponent.Layers.Num())
-			{
-				DesiredFoliage.PerComponent.Add(NewPerComponent);
-			}
-
 			ComponentTime += FPlatformTime::Seconds();
 		}
 
 		if (TotalComponents)
 		{
 			double TotalTime = FPlatformTime::Seconds() - StartTime;
-			UE_LOG(LogTemp, Display, TEXT("Started grass with  %d components in %.1fms     Build Time %.1fms    Component Time  %.1fms    RegisterComponent Time  %.1fms"), TotalComponents, float(TotalTime) * 1000.0f, float(BuildTime) * 1000.0f, float(ComponentTime) * 1000.0f, float(ComponentMiscTime) * 1000.0f);
+			UE_LOG(LogGrass, Display, TEXT("Started grass with  %d components in %.1fms     Build Time %.1fms    Component Time  %.1fms    RegisterComponent Time  %.1fms"), TotalComponents, float(TotalTime) * 1000.0f, float(BuildTime) * 1000.0f, float(ComponentTime) * 1000.0f, float(ComponentMiscTime) * 1000.0f);
 		}
 
 	}
-	Exchange(FoliageCache, DesiredFoliage);
 
-	TSet<UHierarchicalInstancedStaticMeshComponent *> ToDestroy;
-
-	for (UActorComponent* ActorComponent : GetComponents())
+	TSet<UHierarchicalInstancedStaticMeshComponent *> StillUsed;
 	{
-		UHierarchicalInstancedStaticMeshComponent* HComponent = Cast<UHierarchicalInstancedStaticMeshComponent>(ActorComponent);
-		if (HComponent && !StillUsed.Contains(HComponent))
+		// trim cached items based on time, pending and emptiness
+		double OldestToKeepTime = FPlatformTime::Seconds() - CVarMinTimeToKeepGrass.GetValueOnGameThread();
+		uint32 OldestToKeepFrame = GFrameNumber - CVarMinFramesToKeepGrass.GetValueOnGameThread();
+		for (int32 CacheIndex = 0; CacheIndex < FoliageCache.PerComponent.Num(); CacheIndex++)
 		{
-			ToDestroy.Add(HComponent);
+			FCachedLandscapeFoliage::FPerComponent& CacheItem = FoliageCache.PerComponent[CacheIndex];
+			for (int32 LayerIndex = 0; LayerIndex < CacheItem.Layers.Num(); LayerIndex++)
+			{
+				FCachedLandscapeFoliage::FPerLayer& LayerItem = CacheItem.Layers[LayerIndex];
+				for (int32 FoliageIndex = 0; FoliageIndex < LayerItem.Foliage.Num(); FoliageIndex++)
+				{
+					FCachedLandscapeFoliage::FPerGrassComp& GrassItem = LayerItem.Foliage[FoliageIndex];
+					UHierarchicalInstancedStaticMeshComponent *Used = GrassItem.Foliage.Get();
+					if (!Used || (!GrassItem.Pending && GrassItem.LastUsedFrameNumber < OldestToKeepFrame && GrassItem.LastUsedTime < OldestToKeepTime))
+					{
+						LayerItem.Foliage.RemoveAtSwap(FoliageIndex--);
+					}
+					else
+					{
+						StillUsed.Add(Used);
+					}
+				}
+				if (!LayerItem.Foliage.Num())
+				{
+					CacheItem.Layers.RemoveAtSwap(LayerIndex--);
+				}
+			}
+			if (!CacheItem.Layers.Num())
+			{
+				FoliageCache.PerComponent.RemoveAtSwap(CacheIndex--);
+			}
 		}
 	}
-
-	// Destroy any attached but un-owned foliage components
-	for (USceneComponent* SceneComponent : RootComponent->AttachChildren)
 	{
-		UHierarchicalInstancedStaticMeshComponent* HComponent = Cast<UHierarchicalInstancedStaticMeshComponent>(SceneComponent);
-		if (HComponent && !StillUsed.Contains(HComponent))
+		// delete components that are no longer used
+		for (UActorComponent* ActorComponent : GetComponents())
 		{
-			ToDestroy.Add(HComponent);
+			UHierarchicalInstancedStaticMeshComponent* HComponent = Cast<UHierarchicalInstancedStaticMeshComponent>(ActorComponent);
+			if (HComponent && !StillUsed.Contains(HComponent))
+			{
+				double DestroyTime = FPlatformTime::Seconds();
+				{
+					HComponent->ClearInstances();
+					HComponent->DestroyComponent();
+					FoliageComponents.Remove(HComponent);
+				}
+				double TotalTime = FPlatformTime::Seconds() - DestroyTime;
+				UE_LOG(LogGrass, Display, TEXT("Destroyed grass component in %.1fms"), float(TotalTime) * 1000.0f);
+				if (!bForceSync)
+				{
+					break; // one per frame is fine
+				}
+			}
 		}
 	}
-	for (UHierarchicalInstancedStaticMeshComponent* Component : ToDestroy)
-	{
-		Component->ClearInstances();
-		Component->DestroyComponent();
-		FoliageComponents.Remove(Component);
-	}
-
-	double StartTime = FPlatformTime::Seconds();
-	int32 TotalInstances = 0;
-	int32 TotalComponents = 0;
-	int64 TreeSize = 0;
-	int64 InstSize = 0;
-	int64 InstBufSize = 0;
-	for (int32 Index = 0; Index < AsyncFoliageTasks.Num(); Index++)
 	{
-		FAsyncTask<FAsyncGrassTask>* Task = AsyncFoliageTasks[Index];
-		if (bForceSync)
-		{
-			Task->EnsureCompletion();
-		}
-		if (Task->IsDone())
+		// finsh async tasks
+		double StartTime = FPlatformTime::Seconds();
+		int32 TotalInstances = 0;
+		int32 TotalComponents = 0;
+		int64 TreeSize = 0;
+		int64 InstSize = 0;
+		int64 InstBufSize = 0;
+		for (int32 Index = 0; Index < AsyncFoliageTasks.Num(); Index++)
 		{
-			FAsyncGrassTask& Inner = Task->GetTask();
-			AsyncFoliageTasks.RemoveAtSwap(Index--);
-			TotalComponents++;
-			TotalInstances += Inner.Builder->TotalInstances;
-			UHierarchicalInstancedStaticMeshComponent* HierarchicalInstancedStaticMeshComponent = Inner.Foliage.Get();
-			if (HierarchicalInstancedStaticMeshComponent && StillUsed.Contains(HierarchicalInstancedStaticMeshComponent))
+			FAsyncTask<FAsyncGrassTask>* Task = AsyncFoliageTasks[Index];
+			if (bForceSync)
 			{
-				Exchange(HierarchicalInstancedStaticMeshComponent->PerInstanceSMData, Inner.Builder->PerInstanceSMData);
-				HierarchicalInstancedStaticMeshComponent->AcceptPrebuiltTree(Inner.Builder->ClusterTree, Inner.Builder->SortedInstances, Inner.Builder->InstanceReorderTable);
-
-				//@todo - this is sketchy, what you really want to do is flush the grass cache instead of doing a render thread update
-				HierarchicalInstancedStaticMeshComponent->bNeverNeedsRenderUpdate = true;
-
-				//FBoxSphereBounds WorldBounds = HierarchicalInstancedStaticMeshComponent->CalcBounds(HierarchicalInstancedStaticMeshComponent->ComponentToWorld);
-				//UE_LOG(LogTemp, Display, TEXT("Comp %d instances at %.0f %.0f %.0f"), HierarchicalInstancedStaticMeshComponent->PerInstanceSMData.Num(), WorldBounds.Origin.X, WorldBounds.Origin.Y, WorldBounds.Origin.Z);
-				TreeSize += HierarchicalInstancedStaticMeshComponent->ClusterTreePtr->GetAllocatedSize();
-				InstSize += HierarchicalInstancedStaticMeshComponent->PerInstanceSMData.GetAllocatedSize();
-				InstBufSize += HierarchicalInstancedStaticMeshComponent->PerInstanceSMData.Num() * 7 * 16; // there is no place where this constant exists.
-				if (bForceSync)
-				{
-					HierarchicalInstancedStaticMeshComponent->RecreateRenderState_Concurrent();
-				}
+				Task->EnsureCompletion();
 			}
-			delete Inner.Builder;
-			delete Task;
-			if (!bForceSync)
+			if (Task->IsDone())
 			{
-				break; // one per frame is fine
+				FAsyncGrassTask& Inner = Task->GetTask();
+				AsyncFoliageTasks.RemoveAtSwap(Index--);
+				TotalComponents++;
+				TotalInstances += Inner.Builder->TotalInstances;
+				UHierarchicalInstancedStaticMeshComponent* HierarchicalInstancedStaticMeshComponent = Inner.Foliage.Get();
+				if (HierarchicalInstancedStaticMeshComponent && StillUsed.Contains(HierarchicalInstancedStaticMeshComponent))
+				{
+					Exchange(HierarchicalInstancedStaticMeshComponent->PerInstanceSMData, Inner.Builder->PerInstanceSMData);
+					HierarchicalInstancedStaticMeshComponent->AcceptPrebuiltTree(Inner.Builder->ClusterTree, Inner.Builder->SortedInstances, Inner.Builder->InstanceReorderTable);
+
+					//@todo - this is sketchy, what you really want to do is flush the grass cache instead of doing a render thread update
+					HierarchicalInstancedStaticMeshComponent->bNeverNeedsRenderUpdate = true;
+
+					//FBoxSphereBounds WorldBounds = HierarchicalInstancedStaticMeshComponent->CalcBounds(HierarchicalInstancedStaticMeshComponent->ComponentToWorld);
+					//UE_LOG(LogTemp, Display, TEXT("Comp %d instances at %.0f %.0f %.0f"), HierarchicalInstancedStaticMeshComponent->PerInstanceSMData.Num(), WorldBounds.Origin.X, WorldBounds.Origin.Y, WorldBounds.Origin.Z);
+					TreeSize += HierarchicalInstancedStaticMeshComponent->ClusterTreePtr->GetAllocatedSize();
+					InstSize += HierarchicalInstancedStaticMeshComponent->PerInstanceSMData.GetAllocatedSize();
+					InstBufSize += HierarchicalInstancedStaticMeshComponent->PerInstanceSMData.Num() * 7 * 16; // there is no place where this constant exists.
+					if (bForceSync)
+					{
+						HierarchicalInstancedStaticMeshComponent->RecreateRenderState_Concurrent();
+					}
+					for (int32 CacheIndex = 0; CacheIndex < FoliageCache.PerComponent.Num(); CacheIndex++)
+					{
+						FCachedLandscapeFoliage::FPerComponent& CacheItem = FoliageCache.PerComponent[CacheIndex];
+						for (int32 LayerIndex = 0; LayerIndex < CacheItem.Layers.Num(); LayerIndex++)
+						{
+							FCachedLandscapeFoliage::FPerLayer& LayerItem = CacheItem.Layers[LayerIndex];
+							for (int32 FoliageIndex = 0; FoliageIndex < LayerItem.Foliage.Num(); FoliageIndex++)
+							{
+								FCachedLandscapeFoliage::FPerGrassComp& GrassItem = LayerItem.Foliage[FoliageIndex];
+								UHierarchicalInstancedStaticMeshComponent *Used = GrassItem.Foliage.Get();
+								if (Used == HierarchicalInstancedStaticMeshComponent)
+								{
+									GrassItem.Pending = false;
+									GrassItem.Touch();
+									CacheIndex = MAX_int32 - 1;
+									LayerIndex = MAX_int32 - 1;
+									break;
+								}
+							}
+						}
+					}
+				}
+				delete Inner.Builder;
+				delete Task;
+				if (!bForceSync)
+				{
+					break; // one per frame is fine
+				}
 			}
 		}
-	}
-	if (TotalComponents)
-	{
-		double TotalTime = FPlatformTime::Seconds() - StartTime;
-		UE_LOG(LogTemp, Display, TEXT("Finished grass with  %d components and %d instances in %.1fms      InstMem = %lldMB   TreeMem = %lldMB   InstBufMem = %lldMB"), TotalComponents, TotalInstances, float(TotalTime) * 1000.0f, 1 + InstSize / 1024 / 1024, 1 + TreeSize / 1024 / 1024, 1 + InstBufSize / 1024 / 1024);
+		if (TotalComponents)
+		{
+			double TotalTime = FPlatformTime::Seconds() - StartTime;
+			UE_LOG(LogGrass, Display, TEXT("Finished grass with  %d components and %d instances in %.1fms      InstMem = %lldMB   TreeMem = %lldMB   InstBufMem = %lldMB"), TotalComponents, TotalInstances, float(TotalTime) * 1000.0f, 1 + InstSize / 1024 / 1024, 1 + TreeSize / 1024 / 1024, 1 + InstBufSize / 1024 / 1024);
+		}
 	}
 }
 
