Fixed: World Composition should not be modyfing SectionBase on landscape (requires loading all tiles at once and saving)
UE-6910

diff --git a/Engine/Source/Editor/WorldBrowser/Private/Tiles/WorldTileModel.cpp b/Engine/Source/Editor/WorldBrowser/Private/Tiles/WorldTileModel.cpp
--- a/Engine/Source/Editor/WorldBrowser/Private/Tiles/WorldTileModel.cpp
+++ b/Engine/Source/Editor/WorldBrowser/Private/Tiles/WorldTileModel.cpp
@@ -471,7 +471,7 @@ void FWorldTileModel::SetLevelPosition(const FIntPoint& InPosition)
 
 	if (IsLandscapeBased())
 	{
-		FixLandscapeSectionsOffset();
+		UpdateLandscapeSectionsOffset(Offset);
 	}
 	
 	// Transform child levels
@@ -483,10 +483,18 @@ void FWorldTileModel::SetLevelPosition(const FIntPoint& InPosition)
 	}
 }
 
-void FWorldTileModel::FixLandscapeSectionsOffset()
+void FWorldTileModel::UpdateLandscapeSectionsOffset(FIntPoint LevelOffset)
 {
-	check(IsLandscapeBased());
-	Editor->DeferredCommands.AddUnique(TEXT("UpdateLandscapeEditorData"));
+	ALandscapeProxy* LandscapeProxy = GetLandscape();
+	if (LandscapeProxy)
+	{
+		// Calculate new section coordinates for landscape
+		FVector	DrawScale = LandscapeProxy->GetRootComponent()->RelativeScale3D;
+		FIntPoint QuadsSpaceOffset;
+		QuadsSpaceOffset.X = FMath::RoundToInt(LevelOffset.X / DrawScale.X);
+		QuadsSpaceOffset.Y = FMath::RoundToInt(LevelOffset.Y / DrawScale.Y);
+		LandscapeProxy->SetAbsoluteSectionBase(QuadsSpaceOffset + LandscapeProxy->LandscapeSectionOffset);
+	}
 }
 
 void FWorldTileModel::Update()
diff --git a/Engine/Source/Editor/WorldBrowser/Private/Tiles/WorldTileModel.h b/Engine/Source/Editor/WorldBrowser/Private/Tiles/WorldTileModel.h
--- a/Engine/Source/Editor/WorldBrowser/Private/Tiles/WorldTileModel.h
+++ b/Engine/Source/Editor/WorldBrowser/Private/Tiles/WorldTileModel.h
@@ -164,8 +164,8 @@ private:
 	/** Fixup invalid streaming objects inside level */
 	void FixupStreamingObjects();
 
-	/** Handle case when level which is based on landscape was moved or changed visibility  */
-	void FixLandscapeSectionsOffset();
+	/** When level with landscape is moved we need to update internal landscape coordinates to match landscape component grid  */
+	void UpdateLandscapeSectionsOffset(FIntPoint LevelOffset);
 	
 	/** Handler for LevelBoundsActorUpdated event */
 	void OnLevelBoundsActorUpdated();
diff --git a/Engine/Source/Runtime/Core/Public/UObject/ObjectVersion.h b/Engine/Source/Runtime/Core/Public/UObject/ObjectVersion.h
--- a/Engine/Source/Runtime/Core/Public/UObject/ObjectVersion.h
+++ b/Engine/Source/Runtime/Core/Public/UObject/ObjectVersion.h
@@ -450,6 +450,8 @@ enum EUnrealEngineObjectUE4Version
 	VER_UE4_LANDSCAPE_GRASS_COOKING,
 	// Fixed code for using the bOrientMeshEmitters property.
 	VER_UE4_FIX_SKEL_VERT_ORIENT_MESH_PARTICLES,
+	// Do not change landscape section offset on load under world composition
+	VER_UE4_LANDSCAPE_STATIC_SECTION_OFFSET,
 
 	// -----<new versions can be added before this line>-------------------------------------------------
 	// - this needs to be the last line (see note below)
diff --git a/Engine/Source/Runtime/Landscape/Classes/LandscapeComponent.h b/Engine/Source/Runtime/Landscape/Classes/LandscapeComponent.h
--- a/Engine/Source/Runtime/Landscape/Classes/LandscapeComponent.h
+++ b/Engine/Source/Runtime/Landscape/Classes/LandscapeComponent.h
@@ -120,10 +120,12 @@ UCLASS(hidecategories=(Display, Attachment, Physics, Debug, Collision, Movement,
 class ULandscapeComponent : public UPrimitiveComponent
 {
 	GENERATED_UCLASS_BODY()
-
+	
+	/** X offset from global components grid origin (in quads) */
 	UPROPERTY(VisibleAnywhere, BlueprintReadOnly, Category=LandscapeComponent)
 	int32 SectionBaseX;
 
+	/** Y offset from global components grid origin (in quads) */
 	UPROPERTY(VisibleAnywhere, BlueprintReadOnly, Category=LandscapeComponent)
 	int32 SectionBaseY;
 
diff --git a/Engine/Source/Runtime/Landscape/Classes/LandscapeProxy.h b/Engine/Source/Runtime/Landscape/Classes/LandscapeProxy.h
--- a/Engine/Source/Runtime/Landscape/Classes/LandscapeProxy.h
+++ b/Engine/Source/Runtime/Landscape/Classes/LandscapeProxy.h
@@ -282,10 +282,16 @@ protected:
 	FGuid LandscapeGuid;
 
 public:
-	/** Offset in quads from landscape actor origin **/
+	/** Offset in quads from global components grid origin (in quads) **/
 	UPROPERTY()
 	FIntPoint LandscapeSectionOffset;
 
+#if WITH_EDITORONLY_DATA
+	/** To support legacy landscape section offset modification under world composition mode */
+	UPROPERTY()
+	bool bStaticSectionOffset;
+#endif
+
 	/** Max LOD level to use when rendering */
 	UPROPERTY(EditAnywhere, Category=LOD)
 	int32 MaxLODLevel;
diff --git a/Engine/Source/Runtime/Landscape/Private/Landscape.cpp b/Engine/Source/Runtime/Landscape/Private/Landscape.cpp
--- a/Engine/Source/Runtime/Landscape/Private/Landscape.cpp
+++ b/Engine/Source/Runtime/Landscape/Private/Landscape.cpp
@@ -508,6 +508,7 @@ ALandscapeProxy::ALandscapeProxy(const FObjectInitializer& ObjectInitializer)
 #if WITH_EDITORONLY_DATA
 	bLockLocation = true;
 	bIsMovingToLevel = false;
+	bStaticSectionOffset = true;
 #endif // WITH_EDITORONLY_DATA
 	LODDistanceFactor = 1.0f;
 	LODFalloff = ELandscapeLODFalloff::Linear;
@@ -1277,7 +1278,12 @@ void ALandscapeProxy::PostLoad()
 		// Need to clean up invalid collision components
 		RecreateCollisionComponents();
 	}
-
+	
+	if (GetLinker() && GetLinker()->UE4Ver() < VER_UE4_LANDSCAPE_STATIC_SECTION_OFFSET)
+	{
+		bStaticSectionOffset = false;
+	}
+	
 	EditorLayerSettings.Remove(NULL);
 
 	if (EditorCachedLayerInfos_DEPRECATED.Num() > 0)
@@ -1779,6 +1785,65 @@ void ULandscapeInfo::FixupProxiesWeightmaps()
 	}
 }
 
+//
+// This handles legacy behavior of landscapes created under world composition
+// We adjust landscape section offsets and set a flag inside landscape that sections offset are static and should never be touched again
+//
+void AdjustLandscapeSectionOffsets(UWorld* InWorld, const TArray<ALandscapeProxy*> InLandscapeList)
+{
+	// We interested only in registered actors
+	TArray<ALandscapeProxy*> RegisteredLandscapeList = InLandscapeList.FilterByPredicate([](ALandscapeProxy* Proxy) { 
+		return Proxy->GetRootComponent()->IsRegistered(); 
+	});
+
+	// Main Landscape actor should act as origin of global components grid
+	int32 LandcapeIndex = RegisteredLandscapeList.IndexOfByPredicate([](ALandscapeProxy* Proxy) { 
+		return Proxy->IsA<ALandscape>(); 
+	});
+	ALandscapeProxy* StaticLandscape = RegisteredLandscapeList.IsValidIndex(LandcapeIndex) ? RegisteredLandscapeList[LandcapeIndex] : nullptr;
+
+	if (StaticLandscape && !StaticLandscape->bStaticSectionOffset)
+	{
+		StaticLandscape->SetAbsoluteSectionBase(FIntPoint::ZeroValue);
+	}
+	
+	// In case there is no main landscape actor loaded try use any landscape that already has static offsets
+	if (StaticLandscape == nullptr)
+	{
+		LandcapeIndex = RegisteredLandscapeList.IndexOfByPredicate([](ALandscapeProxy* Proxy) { 
+			return Proxy->bStaticSectionOffset; 
+		});
+		StaticLandscape = RegisteredLandscapeList.IsValidIndex(LandcapeIndex) ? RegisteredLandscapeList[LandcapeIndex] : nullptr;
+		// Otherwise offsets will stay variable
+	}
+		
+	ALandscapeProxy* OriginLandscape = StaticLandscape;
+		
+	for (ALandscapeProxy* Proxy : RegisteredLandscapeList)
+	{
+		if (OriginLandscape == nullptr)
+		{
+			OriginLandscape = Proxy;
+		}
+		else if (!Proxy->bStaticSectionOffset)
+		{
+			// Calculate section offset based on relative position from "origin" landscape
+			FVector Offset = Proxy->GetActorLocation() - OriginLandscape->GetActorLocation();
+			FVector	DrawScale = OriginLandscape->GetRootComponent()->RelativeScale3D;
+
+			FIntPoint QuadsSpaceOffset;
+			QuadsSpaceOffset.X = FMath::RoundToInt(Offset.X / DrawScale.X);
+			QuadsSpaceOffset.Y = FMath::RoundToInt(Offset.Y / DrawScale.Y);
+			Proxy->SetAbsoluteSectionBase(QuadsSpaceOffset + OriginLandscape->LandscapeSectionOffset);
+
+			if (StaticLandscape)
+			{
+				Proxy->bStaticSectionOffset = true;
+			}
+		}
+	}
+}
+
 void ULandscapeInfo::RecreateLandscapeInfo(UWorld* InWorld, bool bMapCheck)
 {
 	check(InWorld);
@@ -1806,35 +1871,21 @@ void ULandscapeInfo::RecreateLandscapeInfo(UWorld* InWorld, bool bMapCheck)
 		}
 	}
 
-	// Register gathered landscapes in shared LandscapeInfo map
+	// Handle legacy landscape data under world composition
+	if (InWorld->WorldComposition)
+	{
+		for (auto It = ValidLandscapesMap.CreateIterator(); It; ++It)
+		{
+			AdjustLandscapeSectionOffsets(InWorld, It.Value());
+		}
+	}
+
+	// Register landscapes in global landscape map
 	for (auto It = ValidLandscapesMap.CreateIterator(); It; ++It)
 	{
 		auto& LandscapeList = It.Value();
-		ALandscapeProxy* FirstVisibleLandscape = NULL;
-
-		for (int32 LandscapeIdx = 0; LandscapeIdx < LandscapeList.Num(); LandscapeIdx++)
+		for (ALandscapeProxy* Proxy : LandscapeList)
 		{
-			ALandscapeProxy* Proxy = LandscapeList[LandscapeIdx];
-			if (InWorld->WorldComposition &&
-				Proxy->GetRootComponent()->IsRegistered())
-			{
-				if (FirstVisibleLandscape == NULL)
-				{
-					FirstVisibleLandscape = Proxy;
-				}
-
-				// Setup first landscape in the list as origin for section based grid
-				// And based on that calculate section coordinates for all landscapes in the list
-				FVector Offset = Proxy->GetActorLocation() - FirstVisibleLandscape->GetActorLocation();
-				FVector	DrawScale = FirstVisibleLandscape->GetRootComponent()->RelativeScale3D;
-
-				FIntPoint QuadsSpaceOffset;
-				QuadsSpaceOffset.X = FMath::RoundToInt(Offset.X / DrawScale.X);
-				QuadsSpaceOffset.Y = FMath::RoundToInt(Offset.Y / DrawScale.Y);
-				Proxy->SetAbsoluteSectionBase(QuadsSpaceOffset);
-			}
-
-			// Register gathered landscape actors inside global LandscapeInfo object
 			Proxy->GetLandscapeInfo(true)->RegisterActor(Proxy, bMapCheck);
 		}
 	}
