UE4 - improve grass culling at the generation level

diff --git a/Engine/Source/Runtime/Landscape/Private/Landscape.cpp b/Engine/Source/Runtime/Landscape/Private/Landscape.cpp
--- a/Engine/Source/Runtime/Landscape/Private/Landscape.cpp
+++ b/Engine/Source/Runtime/Landscape/Private/Landscape.cpp
@@ -2164,7 +2164,7 @@ DEFINE_LOG_CATEGORY_STATIC(LogGrass, Log, All);
 
 static TAutoConsoleVariable<float> CVarGuardBandMultiplier(
 	TEXT("grass.GuardBandMultiplier"),
-	1.5f,
+	1.1f,
 	TEXT("Used to control discarding in the grass system. Approximate range, 1-2. Multiplied by the cull distance to control when we discard grass components."));
 
 static TAutoConsoleVariable<int32> CVarMinFramesToKeepGrass(
@@ -2632,9 +2632,11 @@ void ALandscapeProxy::UpdateFoliage(const TArray<FVector>& Cameras, ULandscapeCo
 
 			for (auto& Pos : Cameras)
 			{
-				MinDistanceToComp = FMath::Min<float>(MinDistanceToComp, FVector::Dist(WorldBounds.Origin, Pos) - WorldBounds.SphereRadius);
+				MinDistanceToComp = FMath::Min<float>(MinDistanceToComp, WorldBounds.ComputeSquaredDistanceFromBoxToPoint(Pos));
 			}
 
+			MinDistanceToComp = FMath::Sqrt(MinDistanceToComp);
+
 			for (FWeightmapLayerAllocationInfo& LayerAlloc : Component->WeightmapLayerAllocations)
 			{
 				ULandscapeLayerInfoObject* LayerInfo = LayerAlloc.LayerInfo;
