UE4 - grass, replace spew with stats

diff --git a/Engine/Source/Runtime/Engine/Classes/Components/InstancedStaticMeshComponent.h b/Engine/Source/Runtime/Engine/Classes/Components/InstancedStaticMeshComponent.h
--- a/Engine/Source/Runtime/Engine/Classes/Components/InstancedStaticMeshComponent.h
+++ b/Engine/Source/Runtime/Engine/Classes/Components/InstancedStaticMeshComponent.h
@@ -6,6 +6,8 @@
 
 #include "InstancedStaticMeshComponent.generated.h"
 
+DECLARE_STATS_GROUP(TEXT("Foliage"), STATGROUP_Foliage, STATCAT_Advanced);
+
 class FStaticLightingTextureMapping_InstancedStaticMesh;
 class FInstancedLightMap2D;
 class FInstancedShadowMap2D;
diff --git a/Engine/Source/Runtime/Engine/Private/HierarchicalInstancedStaticMesh.cpp b/Engine/Source/Runtime/Engine/Private/HierarchicalInstancedStaticMesh.cpp
--- a/Engine/Source/Runtime/Engine/Private/HierarchicalInstancedStaticMesh.cpp
+++ b/Engine/Source/Runtime/Engine/Private/HierarchicalInstancedStaticMesh.cpp
@@ -1805,6 +1805,8 @@ void UHierarchicalInstancedStaticMeshComponent::BuildTreeAsync()
 
 FPrimitiveSceneProxy* UHierarchicalInstancedStaticMeshComponent::CreateSceneProxy()
 {
+	QUICK_SCOPE_CYCLE_COUNTER(STAT_HierarchicalInstancedStaticMeshComponent_CreateSceneProxy);
+
 	// Verify that the mesh is valid before using it.
 	const bool bMeshIsValid = 
 		// make sure we have instances
diff --git a/Engine/Source/Runtime/Engine/Private/InstancedStaticMesh.h b/Engine/Source/Runtime/Engine/Private/InstancedStaticMesh.h
--- a/Engine/Source/Runtime/Engine/Private/InstancedStaticMesh.h
+++ b/Engine/Source/Runtime/Engine/Private/InstancedStaticMesh.h
@@ -54,8 +54,6 @@
 #include "AI/Navigation/NavCollision.h"
 #include "Components/InstancedStaticMeshComponent.h"
 
-DECLARE_STATS_GROUP(TEXT("Foliage"), STATGROUP_Foliage, STATCAT_Advanced);
-
 extern TAutoConsoleVariable<float> CVarFoliageMinimumScreenSize;
 extern TAutoConsoleVariable<float> CVarFoliageLODDistanceScale;
 
diff --git a/Engine/Source/Runtime/Landscape/Private/Landscape.cpp b/Engine/Source/Runtime/Landscape/Private/Landscape.cpp
--- a/Engine/Source/Runtime/Landscape/Private/Landscape.cpp
+++ b/Engine/Source/Runtime/Landscape/Private/Landscape.cpp
@@ -2101,6 +2101,13 @@ bool LandscapeMaterialsParameterSetUpdater(FStaticParameterSet &StaticParameterS
 	return UpdateParameterSet<FStaticTerrainLayerWeightParameter, UMaterialExpressionLandscapeLayerWeight>(StaticParameterSet.TerrainLayerWeightParameters, ParentMaterial);
 }
 
+
+DECLARE_CYCLE_STAT(TEXT("Grass Async Build Time"),STAT_FoliageGrassAsyncBuildTime,STATGROUP_Foliage);
+DECLARE_CYCLE_STAT(TEXT("Grass Start Comp"),STAT_FoliageGrassStartComp,STATGROUP_Foliage);
+DECLARE_CYCLE_STAT(TEXT("Grass End Comp"),STAT_FoliageGrassEndComp,STATGROUP_Foliage);
+DECLARE_CYCLE_STAT(TEXT("Grass Destroy Comps"),STAT_FoliageGrassDestoryComp,STATGROUP_Foliage);
+DECLARE_CYCLE_STAT(TEXT("Grass Update"),STAT_GrassUpdate,STATGROUP_Foliage);
+
 void ALandscapeProxy::FlushFoliageComponents(const TSet<ULandscapeComponent*>* OnlyForComponents)
 {
 	if (OnlyForComponents)
@@ -2118,6 +2125,7 @@ void ALandscapeProxy::FlushFoliageComponents(const TSet<ULandscapeComponent*>* O
 						UHierarchicalInstancedStaticMeshComponent *Used = GrassItem.Foliage.Get();
 						if (Used)
 						{
+							SCOPE_CYCLE_COUNTER(STAT_FoliageGrassDestoryComp);
 							Used->ClearInstances();
 							Used->DestroyComponent();
 						}
@@ -2138,6 +2146,7 @@ void ALandscapeProxy::FlushFoliageComponents(const TSet<ULandscapeComponent*>* O
 		GetComponents(FoliageComps);
 		for (UHierarchicalInstancedStaticMeshComponent* Component : FoliageComps)
 		{
+			SCOPE_CYCLE_COUNTER(STAT_FoliageGrassDestoryComp);
 			Component->ClearInstances();
 			Component->DestroyComponent();
 		}
@@ -2151,6 +2160,7 @@ void ALandscapeProxy::FlushFoliageComponents(const TSet<ULandscapeComponent*>* O
 		// Destroy any attached but un-owned foliage components
 		for (USceneComponent* Component : AttachedFoliageComponents)
 		{
+			SCOPE_CYCLE_COUNTER(STAT_FoliageGrassDestoryComp);
 			Component->DestroyComponent();
 		}
 	}
@@ -2191,6 +2201,7 @@ static TAutoConsoleVariable<int32> CVarGrassEnable(
 	TEXT("1: Enable Grass; 0: Disable Grass"));
 
 #if WITH_EDITOR
+DECLARE_CYCLE_STAT(TEXT("Grass Load Source Art"),STAT_GrassLoadSourceArt,STATGROUP_Foliage);
 class FEditorGrassLayerData
 {
 	int32 Stride;
@@ -2212,6 +2223,7 @@ public:
 			}
 			else
 			{
+				SCOPE_CYCLE_COUNTER(STAT_GrassLoadSourceArt);
 				FLandscapeComponentDataInterface CDI(Component);
 				CDI.GetHeightmapTextureData(LocalHeightCache);
 				if (!CDI.GetWeightmapTextureData(LayerInfo, LocalCache) || !LocalCache.Num() || LocalCache.Num() != LocalHeightCache.Num())
@@ -2402,6 +2414,7 @@ struct FAsyncGrassBuilder : public FGrassBuilderBase
 
 	void Build()
 	{
+		SCOPE_CYCLE_COUNTER(STAT_FoliageGrassAsyncBuildTime);
 		check(bHaveValidData);
 		RasterTime -= FPlatformTime::Seconds();
 
@@ -2411,7 +2424,7 @@ struct FAsyncGrassBuilder : public FGrassBuilderBase
 		MaxJitter *= Extent;
 		Origin += Extent * (Div * 0.5f);
 		struct FInstanceLocal
-	{
+		{
 			FVector Pos;
 			bool bKeep;
 		};
@@ -2589,14 +2602,10 @@ void ALandscapeProxy::Tick(float DeltaSeconds)
 
 void ALandscapeProxy::UpdateFoliage(const TArray<FVector>& Cameras, ULandscapeComponent* OnlyComponent, bool bForceSync)
 {
+	SCOPE_CYCLE_COUNTER(STAT_GrassUpdate);
+
 	if (CVarGrassEnable.GetValueOnAnyThread() > 0)
 	{
-		double ComponentTime = 0.0;
-		double BuildTime = 0.0;
-		double ComponentMiscTime = 0.0;
-		double StartTime = FPlatformTime::Seconds();
-		int32 TotalComponents = 0;
-
 		float GuardBand = CVarGuardBandMultiplier.GetValueOnAnyThread();
 		int32 MaxInstancesPerComponent = FMath::Max<int32>(1024, CVarMaxInstancesPerComponent.GetValueOnAnyThread());
 
@@ -2607,7 +2616,6 @@ void ALandscapeProxy::UpdateFoliage(const TArray<FVector>& Cameras, ULandscapeCo
 			{
 				continue;
 			}
-			ComponentTime -= FPlatformTime::Seconds();
 
 			FCachedLandscapeFoliage::FPerComponent* ExistingCache = nullptr;
 			for (FCachedLandscapeFoliage::FPerComponent& CacheItem : FoliageCache.PerComponent)
@@ -2714,14 +2722,19 @@ void ALandscapeProxy::UpdateFoliage(const TArray<FVector>& Cameras, ULandscapeCo
 					{
 						for (int32 SubY = 0; SubY < SqrtSubsections; SubY++)
 						{
+							SCOPE_CYCLE_COUNTER(STAT_FoliageGrassStartComp);
+
 							int32 FolSeed = FCrc::StrCrc32((LayerInfo->LayerName.ToString() + Component->GetName() + FString::Printf(TEXT("%d %d"), SubX, SubY)).GetCharArray().GetData());
 							if (FolSeed == 0)
 							{
 								FolSeed++;
 							}
 
-							UHierarchicalInstancedStaticMeshComponent* HierarchicalInstancedStaticMeshComponent = ConstructObject<UHierarchicalInstancedStaticMeshComponent>(UHierarchicalInstancedStaticMeshComponent::StaticClass(), this); //, NAME_None, RF_Transactional);
-							TotalComponents++;
+							UHierarchicalInstancedStaticMeshComponent* HierarchicalInstancedStaticMeshComponent;
+							{
+								QUICK_SCOPE_CYCLE_COUNTER(STAT_GrassCreateComp);
+								HierarchicalInstancedStaticMeshComponent = ConstructObject<UHierarchicalInstancedStaticMeshComponent>(UHierarchicalInstancedStaticMeshComponent::StaticClass(), this); //, NAME_None, RF_Transactional);
+							}
 							HierarchicalInstancedStaticMeshComponent->Mobility = EComponentMobility::Static;
 
 							HierarchicalInstancedStaticMeshComponent->StaticMesh = LayerInfo->GrassMesh;
@@ -2759,52 +2772,48 @@ void ALandscapeProxy::UpdateFoliage(const TArray<FVector>& Cameras, ULandscapeCo
 							HierarchicalInstancedStaticMeshComponent->BodyInstance.CopyBodyInstancePropertiesFrom(&Settings->BodyInstance);
 	#endif
 
-							HierarchicalInstancedStaticMeshComponent->AttachTo(GetRootComponent());
-							FTransform DesiredTransform = GetRootComponent()->ComponentToWorld;
-							DesiredTransform.RemoveScaling();
-							HierarchicalInstancedStaticMeshComponent->SetWorldTransform(DesiredTransform);
+							{
+								QUICK_SCOPE_CYCLE_COUNTER(STAT_GrassAttachComp);
 
-							FoliageComponents.Add(HierarchicalInstancedStaticMeshComponent);
-							NewLayer.Foliage.Add(FCachedLandscapeFoliage::FPerGrassComp(HierarchicalInstancedStaticMeshComponent));
+								HierarchicalInstancedStaticMeshComponent->AttachTo(GetRootComponent());
+								FTransform DesiredTransform = GetRootComponent()->ComponentToWorld;
+								DesiredTransform.RemoveScaling();
+								HierarchicalInstancedStaticMeshComponent->SetWorldTransform(DesiredTransform);
 
+								FoliageComponents.Add(HierarchicalInstancedStaticMeshComponent);
+								NewLayer.Foliage.Add(FCachedLandscapeFoliage::FPerGrassComp(HierarchicalInstancedStaticMeshComponent));
+							}
+
+							FAsyncGrassBuilder* Builder;
+
+							{
+								QUICK_SCOPE_CYCLE_COUNTER(STAT_GrassCreateBuilder);
+								Builder = new FAsyncGrassBuilder(this, Component, LayerInfo, HierarchicalInstancedStaticMeshComponent, SqrtSubsections, SubX, SubY, CopyFromBuilder);
+							}
 
-							FAsyncGrassBuilder* Builder = new FAsyncGrassBuilder(this, Component, LayerInfo, HierarchicalInstancedStaticMeshComponent, SqrtSubsections, SubX, SubY, CopyFromBuilder);
 							CopyFromBuilder = Builder;
 
 							if (Builder->bHaveValidData)
 							{
-								BuildTime -= FPlatformTime::Seconds();
 								FAsyncTask<FAsyncGrassTask>* Task = new FAsyncTask<FAsyncGrassTask>(Builder, Component, HierarchicalInstancedStaticMeshComponent, LayerInfo);
 
 								Task->StartBackgroundTask();
 
 								AsyncFoliageTasks.Add(Task);
-
-								BuildTime += FPlatformTime::Seconds();
 							}
 							else
 							{
 								delete Builder;
 							}
-
-
-							ComponentMiscTime -= FPlatformTime::Seconds();
-							HierarchicalInstancedStaticMeshComponent->RegisterComponent();
-							ComponentMiscTime += FPlatformTime::Seconds();
-
+							{
+								QUICK_SCOPE_CYCLE_COUNTER(STAT_GrassRegisterComp);
+								HierarchicalInstancedStaticMeshComponent->RegisterComponent();
+							}
 						}
 					}
 				}
 			}
-			ComponentTime += FPlatformTime::Seconds();
 		}
-
-		if (TotalComponents)
-		{
-			double TotalTime = FPlatformTime::Seconds() - StartTime;
-			UE_LOG(LogGrass, Display, TEXT("Started grass with  %d components in %.1fms     Build Time %.1fms    Component Time  %.1fms    RegisterComponent Time  %.1fms"), TotalComponents, float(TotalTime) * 1000.0f, float(BuildTime) * 1000.0f, float(ComponentTime) * 1000.0f, float(ComponentMiscTime) * 1000.0f);
-		}
-
 	}
 
 	TSet<UHierarchicalInstancedStaticMeshComponent *> StillUsed;
@@ -2843,20 +2852,19 @@ void ALandscapeProxy::UpdateFoliage(const TArray<FVector>& Cameras, ULandscapeCo
 		}
 	}
 	{
+
 		// delete components that are no longer used
 		for (UActorComponent* ActorComponent : GetComponents())
 		{
 			UHierarchicalInstancedStaticMeshComponent* HComponent = Cast<UHierarchicalInstancedStaticMeshComponent>(ActorComponent);
 			if (HComponent && !StillUsed.Contains(HComponent))
 			{
-				double DestroyTime = FPlatformTime::Seconds();
 				{
+					SCOPE_CYCLE_COUNTER(STAT_FoliageGrassDestoryComp);
 					HComponent->ClearInstances();
 					HComponent->DestroyComponent();
 					FoliageComponents.Remove(HComponent);
 				}
-				double TotalTime = FPlatformTime::Seconds() - DestroyTime;
-				UE_LOG(LogGrass, Display, TEXT("Destroyed grass component in %.1fms"), float(TotalTime) * 1000.0f);
 				if (!bForceSync)
 				{
 					break; // one per frame is fine
@@ -2866,12 +2874,6 @@ void ALandscapeProxy::UpdateFoliage(const TArray<FVector>& Cameras, ULandscapeCo
 	}
 	{
 		// finsh async tasks
-		double StartTime = FPlatformTime::Seconds();
-		int32 TotalInstances = 0;
-		int32 TotalComponents = 0;
-		int64 TreeSize = 0;
-		int64 InstSize = 0;
-		int64 InstBufSize = 0;
 		for (int32 Index = 0; Index < AsyncFoliageTasks.Num(); Index++)
 		{
 			FAsyncTask<FAsyncGrassTask>* Task = AsyncFoliageTasks[Index];
@@ -2881,10 +2883,9 @@ void ALandscapeProxy::UpdateFoliage(const TArray<FVector>& Cameras, ULandscapeCo
 			}
 			if (Task->IsDone())
 			{
+				SCOPE_CYCLE_COUNTER(STAT_FoliageGrassEndComp);
 				FAsyncGrassTask& Inner = Task->GetTask();
 				AsyncFoliageTasks.RemoveAtSwap(Index--);
-				TotalComponents++;
-				TotalInstances += Inner.Builder->TotalInstances;
 				UHierarchicalInstancedStaticMeshComponent* HierarchicalInstancedStaticMeshComponent = Inner.Foliage.Get();
 				if (HierarchicalInstancedStaticMeshComponent && StillUsed.Contains(HierarchicalInstancedStaticMeshComponent))
 				{
@@ -2894,11 +2895,6 @@ void ALandscapeProxy::UpdateFoliage(const TArray<FVector>& Cameras, ULandscapeCo
 					//@todo - this is sketchy, what you really want to do is flush the grass cache instead of doing a render thread update
 					HierarchicalInstancedStaticMeshComponent->bNeverNeedsRenderUpdate = true;
 
-					//FBoxSphereBounds WorldBounds = HierarchicalInstancedStaticMeshComponent->CalcBounds(HierarchicalInstancedStaticMeshComponent->ComponentToWorld);
-					//UE_LOG(LogTemp, Display, TEXT("Comp %d instances at %.0f %.0f %.0f"), HierarchicalInstancedStaticMeshComponent->PerInstanceSMData.Num(), WorldBounds.Origin.X, WorldBounds.Origin.Y, WorldBounds.Origin.Z);
-					TreeSize += HierarchicalInstancedStaticMeshComponent->ClusterTreePtr->GetAllocatedSize();
-					InstSize += HierarchicalInstancedStaticMeshComponent->PerInstanceSMData.GetAllocatedSize();
-					InstBufSize += HierarchicalInstancedStaticMeshComponent->PerInstanceSMData.Num() * 7 * 16; // there is no place where this constant exists.
 					if (bForceSync)
 					{
 						HierarchicalInstancedStaticMeshComponent->RecreateRenderState_Concurrent();
@@ -2933,11 +2929,6 @@ void ALandscapeProxy::UpdateFoliage(const TArray<FVector>& Cameras, ULandscapeCo
 				}
 			}
 		}
-		if (TotalComponents)
-		{
-			double TotalTime = FPlatformTime::Seconds() - StartTime;
-			UE_LOG(LogGrass, Display, TEXT("Finished grass with  %d components and %d instances in %.1fms      InstMem = %lldMB   TreeMem = %lldMB   InstBufMem = %lldMB"), TotalComponents, TotalInstances, float(TotalTime) * 1000.0f, 1 + InstSize / 1024 / 1024, 1 + TreeSize / 1024 / 1024, 1 + InstBufSize / 1024 / 1024);
-		}
 	}
 }
 
