UE4 - ability to use cooked path for grass in the editor, disabled ifdef

diff --git a/Engine/Source/Runtime/Landscape/Classes/LandscapeComponent.h b/Engine/Source/Runtime/Landscape/Classes/LandscapeComponent.h
--- a/Engine/Source/Runtime/Landscape/Classes/LandscapeComponent.h
+++ b/Engine/Source/Runtime/Landscape/Classes/LandscapeComponent.h
@@ -8,6 +8,8 @@
 
 #include "LandscapeComponent.generated.h"
 
+#define USE_PRECACHED_GRASSMAP_IN_EDITOR (0)
+
 class ULandscapeLayerInfoObject;
 class ULandscapeInfo;
 class ALandscapeProxy;
@@ -250,7 +252,7 @@ public:
 	/** Platform Data where don't support texture sampling in vertex buffer */
 	FLandscapeComponentDerivedData PlatformData;
 
-	/** Grass data for cooked platforms, also used temporarily when cooking */
+	/** Grass data for cooked platforms, also used temporarily when cooking, but see USE_PRECACHED_GRASSMAP_IN_EDITOR, this might be changing */
 	TSharedPtr<FGrassMap, ESPMode::ThreadSafe> GrassMap;
 
 	virtual ~ULandscapeComponent();
@@ -258,6 +260,7 @@ public:
 	// Begin UObject interface.	
 	virtual void PostInitProperties() override;	
 	virtual void Serialize(FArchive& Ar) override;
+	virtual SIZE_T GetResourceSize(EResourceSizeMode::Type Mode) override;
 	virtual void BeginDestroy() override;
 	static void AddReferencedObjects(UObject* InThis, FReferenceCollector& Collector);
 	virtual void PostDuplicate(bool bDuplicateForPIE) override;
diff --git a/Engine/Source/Runtime/Landscape/Private/Landscape.cpp b/Engine/Source/Runtime/Landscape/Private/Landscape.cpp
--- a/Engine/Source/Runtime/Landscape/Private/Landscape.cpp
+++ b/Engine/Source/Runtime/Landscape/Private/Landscape.cpp
@@ -239,6 +239,18 @@ void ULandscapeComponent::Serialize(FArchive& Ar)
 #endif
 }
 
+SIZE_T ULandscapeComponent::GetResourceSize(EResourceSizeMode::Type Mode)
+{
+	SIZE_T ResSize = Super::GetResourceSize(Mode);
+
+	if (GrassMap.IsValid())
+	{
+		ResSize += GrassMap->Data.GetAllocatedSize();
+	}
+
+	return ResSize;
+}
+
 #if WITH_EDITOR
 UMaterialInterface* ULandscapeComponent::GetLandscapeMaterial() const
 {
@@ -467,6 +479,10 @@ void ULandscapeComponent::PostLoad()
 		}
 	}
 #endif
+
+#if USE_PRECACHED_GRASSMAP_IN_EDITOR
+	GenerateGrassMap();
+#endif
 }
 
 #endif // WITH_EDITOR
@@ -2200,7 +2216,7 @@ static TAutoConsoleVariable<int32> CVarGrassEnable(
 	1,
 	TEXT("1: Enable Grass; 0: Disable Grass"));
 
-#if WITH_EDITOR
+#if WITH_EDITOR && !USE_PRECACHED_GRASSMAP_IN_EDITOR
 DECLARE_CYCLE_STAT(TEXT("Grass Load Source Art"),STAT_GrassLoadSourceArt,STATGROUP_Foliage);
 class FEditorGrassLayerData
 {
