UE4 - use streaming manager for grass

diff --git a/Engine/Source/Runtime/Engine/Public/ContentStreaming.h b/Engine/Source/Runtime/Engine/Public/ContentStreaming.h
--- a/Engine/Source/Runtime/Engine/Public/ContentStreaming.h
+++ b/Engine/Source/Runtime/Engine/Public/ContentStreaming.h
@@ -337,13 +337,13 @@ struct IStreamingManager
 	}
 
 	/** Returns the number of view infos. */
-	int32 GetNumViews() const
+	ENGINE_API int32 GetNumViews() const
 	{
 		return CurrentViewInfos.Num();
 	}
 
 	/** Returns the view info by the specified index. */
-	const FStreamingViewInfo& GetViewInformation( int32 ViewIndex ) const
+	ENGINE_API const FStreamingViewInfo& GetViewInformation( int32 ViewIndex ) const
 	{
 		return CurrentViewInfos[ ViewIndex ];
 	}
@@ -417,7 +417,7 @@ protected:
 	};
 
 	/** Current collection of views that need to be taken into account for streaming. Emptied every frame. */
-	static TArray<FStreamingViewInfo> CurrentViewInfos;
+	ENGINE_API static TArray<FStreamingViewInfo> CurrentViewInfos;
 
 	/** Pending views. Emptied every frame. */
 	static TArray<FStreamingViewInfo> PendingViewInfos;
diff --git a/Engine/Source/Runtime/Landscape/Private/Landscape.cpp b/Engine/Source/Runtime/Landscape/Private/Landscape.cpp
--- a/Engine/Source/Runtime/Landscape/Private/Landscape.cpp
+++ b/Engine/Source/Runtime/Landscape/Private/Landscape.cpp
@@ -31,6 +31,7 @@ Landscape.cpp: Terrain rendering
 #include "Engine/StaticMesh.h"
 #include "LandscapeMeshProxyComponent.h"
 #include "LandscapeMeshProxyActor.h"
+#include "ContentStreaming.h"
 
 #include "Materials/MaterialExpressionLandscapeLayerWeight.h"
 #include "Materials/MaterialExpressionLandscapeLayerSwitch.h"
@@ -2216,6 +2217,11 @@ static TAutoConsoleVariable<int32> CVarGrassEnable(
 	1,
 	TEXT("1: Enable Grass; 0: Disable Grass"));
 
+static TAutoConsoleVariable<int32> CVarUseStreamingManagerForCameras(
+	TEXT("grass.UseStreamingManagerForCameras"),
+	1,
+	TEXT("1: Use Streaming Manager; 0: Use ViewLocationsRenderedLastFrame"));
+
 #if WITH_EDITOR && !USE_PRECACHED_GRASSMAP_IN_EDITOR
 DECLARE_CYCLE_STAT(TEXT("Grass Load Source Art"),STAT_GrassLoadSourceArt,STATGROUP_Foliage);
 class FEditorGrassLayerData
@@ -2591,29 +2597,48 @@ void ALandscapeProxy::Tick(float DeltaSeconds)
 	// this is NOT an actor tick, it is a FTickableGameObject tick
 	// the super tick is for an actor tick...
 	//Super::Tick(DeltaSeconds);
-	UWorld* World = GetWorld();
-	if (!World)
-	{
-		return;
-	}
 
-	// there is a bug here, which often leaves us with no cameras in the editor
 	static TArray<FVector> OldCameras;
-
-	if (!OldCameras.Num() && !World->ViewLocationsRenderedLastFrame.Num())
+	if (CVarUseStreamingManagerForCameras.GetValueOnGameThread() == 0)
 	{
-		// no cameras, no grass update
-		return;
-	}
+		UWorld* World = GetWorld();
+		if (!World)
+		{
+			return;
+		}
 
-	const TArray<FVector>& Cameras = World->ViewLocationsRenderedLastFrame.Num() ? World->ViewLocationsRenderedLastFrame : OldCameras;
+		if (!OldCameras.Num() && !World->ViewLocationsRenderedLastFrame.Num())
+		{
+			// no cameras, no grass update
+			return;
+		}
+
+		// there is a bug here, which often leaves us with no cameras in the editor
+		const TArray<FVector>& Cameras = World->ViewLocationsRenderedLastFrame.Num() ? World->ViewLocationsRenderedLastFrame : OldCameras;
 
-	if (&Cameras != &OldCameras)
+		if (&Cameras != &OldCameras)
+		{
+			check(IsInGameThread());
+			OldCameras = Cameras;
+		}
+		UpdateFoliage(Cameras);
+	}
+	else
 	{
-		check(IsInGameThread());
-		OldCameras = Cameras;
+		int32 Num = IStreamingManager::Get().GetNumViews();
+		if (!Num)
+		{
+			// no cameras, no grass update
+			return;
+		}
+		OldCameras.Reset(Num);
+		for (int32 Index = 0; Index < Num; Index++)
+		{
+			auto& ViewInfo = IStreamingManager::Get().GetViewInformation(Index);
+			OldCameras.Add(ViewInfo.ViewOrigin);
+		}
+		UpdateFoliage(OldCameras);
 	}
-	UpdateFoliage(Cameras);
 }
 
 void ALandscapeProxy::UpdateFoliage(const TArray<FVector>& Cameras, ULandscapeComponent* OnlyComponent, bool bForceSync)
