Fix to CL 2406859
Removed unused, incorrectly calculated InstancedStaticMesh per-instance inverse matrix.

diff --git a/Engine/Source/Runtime/Landscape/Private/Landscape.cpp b/Engine/Source/Runtime/Landscape/Private/Landscape.cpp
--- a/Engine/Source/Runtime/Landscape/Private/Landscape.cpp
+++ b/Engine/Source/Runtime/Landscape/Private/Landscape.cpp
@@ -2536,9 +2536,7 @@ struct FAsyncGrassBuilder : public FGrassBuilderBase
 							RenderData[2] = FVector4(OutXForm.M[0][1], OutXForm.M[1][1], OutXForm.M[2][1], OutXForm.M[3][1]);
 							RenderData[3] = FVector4(OutXForm.M[0][2], OutXForm.M[1][2], OutXForm.M[2][2], OutXForm.M[3][2]);
 
-							RenderData[4] = FVector4(0.0f, 0.0f, 0.0f, -1.0f);
-							RenderData[5] = FVector4(0.0f, 0.0f, 0.0f, -1.0f);
-							RenderData[6] = FVector4(0.0f, 0.0f, 0.0f, RandomInstanceID);
+							RenderData[4] = FVector4(0.0f, 0.0f, RandomInstanceID, 0.0f);
 						}
 						InstanceIndex++;
 					}
