Fix for crash deleting landscape

diff --git a/Engine/Source/Runtime/Landscape/Private/Landscape.cpp b/Engine/Source/Runtime/Landscape/Private/Landscape.cpp
--- a/Engine/Source/Runtime/Landscape/Private/Landscape.cpp
+++ b/Engine/Source/Runtime/Landscape/Private/Landscape.cpp
@@ -2636,50 +2636,53 @@ struct FAsyncGrassBuilder : public FGrassBuilderBase
 
 void ALandscapeProxy::Tick(float DeltaSeconds)
 {
-	// this is NOT an actor tick, it is a FTickableGameObject tick
-	// the super tick is for an actor tick...
-	//Super::Tick(DeltaSeconds);
-
-	static TArray<FVector> OldCameras;
-	if (CVarUseStreamingManagerForCameras.GetValueOnGameThread() == 0)
+	if (!IsPendingKillPending())
 	{
-		UWorld* World = GetWorld();
-		if (!World)
-		{
-			return;
-		}
+		// this is NOT an actor tick, it is a FTickableGameObject tick
+		// the super tick is for an actor tick...
+		//Super::Tick(DeltaSeconds);
 
-		if (!OldCameras.Num() && !World->ViewLocationsRenderedLastFrame.Num())
+		static TArray<FVector> OldCameras;
+		if (CVarUseStreamingManagerForCameras.GetValueOnGameThread() == 0)
 		{
-			// no cameras, no grass update
-			return;
-		}
+			UWorld* World = GetWorld();
+			if (!World)
+			{
+				return;
+			}
 
-		// there is a bug here, which often leaves us with no cameras in the editor
-		const TArray<FVector>& Cameras = World->ViewLocationsRenderedLastFrame.Num() ? World->ViewLocationsRenderedLastFrame : OldCameras;
+			if (!OldCameras.Num() && !World->ViewLocationsRenderedLastFrame.Num())
+			{
+				// no cameras, no grass update
+				return;
+			}
 
-		if (&Cameras != &OldCameras)
-		{
-			check(IsInGameThread());
-			OldCameras = Cameras;
-		}
-		UpdateFoliage(Cameras);
-	}
-	else
-	{
-		int32 Num = IStreamingManager::Get().GetNumViews();
-		if (!Num)
-		{
-			// no cameras, no grass update
-			return;
+			// there is a bug here, which often leaves us with no cameras in the editor
+			const TArray<FVector>& Cameras = World->ViewLocationsRenderedLastFrame.Num() ? World->ViewLocationsRenderedLastFrame : OldCameras;
+
+			if (&Cameras != &OldCameras)
+			{
+				check(IsInGameThread());
+				OldCameras = Cameras;
+			}
+			UpdateFoliage(Cameras);
 		}
-		OldCameras.Reset(Num);
-		for (int32 Index = 0; Index < Num; Index++)
+		else
 		{
-			auto& ViewInfo = IStreamingManager::Get().GetViewInformation(Index);
-			OldCameras.Add(ViewInfo.ViewOrigin);
+			int32 Num = IStreamingManager::Get().GetNumViews();
+			if (!Num)
+			{
+				// no cameras, no grass update
+				return;
+			}
+			OldCameras.Reset(Num);
+			for (int32 Index = 0; Index < Num; Index++)
+			{
+				auto& ViewInfo = IStreamingManager::Get().GetViewInformation(Index);
+				OldCameras.Add(ViewInfo.ViewOrigin);
+			}
+			UpdateFoliage(OldCameras);
 		}
-		UpdateFoliage(OldCameras);
 	}
 }
 
