UE4 - switch to always use the cooked GrassMap path, even in the editor

diff --git a/Engine/Source/Runtime/Landscape/Classes/LandscapeComponent.h b/Engine/Source/Runtime/Landscape/Classes/LandscapeComponent.h
--- a/Engine/Source/Runtime/Landscape/Classes/LandscapeComponent.h
+++ b/Engine/Source/Runtime/Landscape/Classes/LandscapeComponent.h
@@ -8,8 +8,6 @@
 
 #include "LandscapeComponent.generated.h"
 
-#define USE_PRECACHED_GRASSMAP_IN_EDITOR (0)
-
 class ULandscapeLayerInfoObject;
 class ULandscapeInfo;
 class ALandscapeProxy;
@@ -252,7 +250,7 @@ public:
 	/** Platform Data where don't support texture sampling in vertex buffer */
 	FLandscapeComponentDerivedData PlatformData;
 
-	/** Grass data for cooked platforms, also used temporarily when cooking, but see USE_PRECACHED_GRASSMAP_IN_EDITOR, this might be changing */
+	/** Grass data for generation **/
 	TSharedPtr<FGrassMap, ESPMode::ThreadSafe> GrassMap;
 
 	virtual ~ULandscapeComponent();
diff --git a/Engine/Source/Runtime/Landscape/Private/Landscape.cpp b/Engine/Source/Runtime/Landscape/Private/Landscape.cpp
--- a/Engine/Source/Runtime/Landscape/Private/Landscape.cpp
+++ b/Engine/Source/Runtime/Landscape/Private/Landscape.cpp
@@ -481,9 +481,7 @@ void ULandscapeComponent::PostLoad()
 	}
 #endif
 
-#if USE_PRECACHED_GRASSMAP_IN_EDITOR
 	GenerateGrassMap();
-#endif
 }
 
 #endif // WITH_EDITOR
@@ -2129,6 +2127,13 @@ void ALandscapeProxy::FlushFoliageComponents(const TSet<ULandscapeComponent*>* O
 {
 	if (OnlyForComponents)
 	{
+		for (auto LandComp : *OnlyForComponents)
+		{
+			if (LandComp)
+			{
+				LandComp->RemoveGrassMap();
+			}
+		}
 		for (FCachedLandscapeFoliage::FPerComponent& CacheItem : FoliageCache.PerComponent)
 		{
 			ULandscapeComponent* Component = CacheItem.BasedOn.Get();
@@ -2180,6 +2185,14 @@ void ALandscapeProxy::FlushFoliageComponents(const TSet<ULandscapeComponent*>* O
 			SCOPE_CYCLE_COUNTER(STAT_FoliageGrassDestoryComp);
 			Component->DestroyComponent();
 		}
+
+		// Clear GrassMaps
+		TInlineComponentArray<ULandscapeComponent*> LandComps;
+		GetComponents(LandComps);
+		for (ULandscapeComponent* Component : LandComps)
+		{
+			Component->RemoveGrassMap();
+		}
 	}
 }
 
@@ -2222,7 +2235,9 @@ static TAutoConsoleVariable<int32> CVarUseStreamingManagerForCameras(
 	1,
 	TEXT("1: Use Streaming Manager; 0: Use ViewLocationsRenderedLastFrame"));
 
-#if WITH_EDITOR && !USE_PRECACHED_GRASSMAP_IN_EDITOR
+#if WITH_EDITOR && 0 // we will remove this shortly
+
+
 DECLARE_CYCLE_STAT(TEXT("Grass Load Source Art"),STAT_GrassLoadSourceArt,STATGROUP_Foliage);
 class FEditorGrassLayerData
 {
@@ -2292,6 +2307,11 @@ public:
 		ElementStride = 0;
 		if (Stride)
 		{
+			if (!Component->GrassMap.IsValid())
+			{
+				check(IsInGameThread());
+				Component->GenerateGrassMap();
+			}
 			GrassMap = Component->GrassMap;
 			if (GrassMap.IsValid())
 			{
