UE4 - fix non-editor builds

diff --git a/Engine/Source/Runtime/Landscape/Private/Landscape.cpp b/Engine/Source/Runtime/Landscape/Private/Landscape.cpp
--- a/Engine/Source/Runtime/Landscape/Private/Landscape.cpp
+++ b/Engine/Source/Runtime/Landscape/Private/Landscape.cpp
@@ -2127,6 +2127,7 @@ void ALandscapeProxy::FlushFoliageComponents(const TSet<ULandscapeComponent*>* O
 {
 	if (OnlyForComponents)
 	{
+#if WITH_EDITOR
 		for (auto LandComp : *OnlyForComponents)
 		{
 			if (LandComp)
@@ -2134,6 +2135,7 @@ void ALandscapeProxy::FlushFoliageComponents(const TSet<ULandscapeComponent*>* O
 				LandComp->RemoveGrassMap();
 			}
 		}
+#endif
 		for (FCachedLandscapeFoliage::FPerComponent& CacheItem : FoliageCache.PerComponent)
 		{
 			ULandscapeComponent* Component = CacheItem.BasedOn.Get();
@@ -2186,6 +2188,7 @@ void ALandscapeProxy::FlushFoliageComponents(const TSet<ULandscapeComponent*>* O
 			Component->DestroyComponent();
 		}
 
+#if WITH_EDITOR
 		// Clear GrassMaps
 		TInlineComponentArray<ULandscapeComponent*> LandComps;
 		GetComponents(LandComps);
@@ -2193,6 +2196,7 @@ void ALandscapeProxy::FlushFoliageComponents(const TSet<ULandscapeComponent*>* O
 		{
 			Component->RemoveGrassMap();
 		}
+#endif
 	}
 }
 
@@ -2307,11 +2311,13 @@ public:
 		ElementStride = 0;
 		if (Stride)
 		{
+#if WITH_EDITOR
 			if (!Component->GrassMap.IsValid())
 			{
 				check(IsInGameThread());
 				Component->GenerateGrassMap();
 			}
+#endif
 			GrassMap = Component->GrassMap;
 			if (GrassMap.IsValid())
 			{
