UE4 - finer granularity for grass building

diff --git a/Engine/Source/Runtime/Landscape/Classes/LandscapeProxy.h b/Engine/Source/Runtime/Landscape/Classes/LandscapeProxy.h
--- a/Engine/Source/Runtime/Landscape/Classes/LandscapeProxy.h
+++ b/Engine/Source/Runtime/Landscape/Classes/LandscapeProxy.h
@@ -193,16 +193,50 @@ namespace ELandscapeLODFalloff
 
 struct FCachedLandscapeFoliage
 {
-	struct FPerGrassComp
+	struct FGrassCompKey
 	{
+		TWeakObjectPtr<ULandscapeComponent> BasedOn;
+		TWeakObjectPtr<ULandscapeLayerInfoObject> Layer;
+		int32 SqrtSubsections;
+		int32 CachedMaxInstancesPerComponent;
+		int32 SubsectionX;
+		int32 SubsectionY;
+
+		FGrassCompKey()
+			: SqrtSubsections(0)
+			, CachedMaxInstancesPerComponent(0)
+			, SubsectionX(0)
+			, SubsectionY(0)
+		{
+		}
+		inline bool operator==(const FGrassCompKey& Other) const
+		{
+			return 
+				SqrtSubsections == Other.SqrtSubsections &&
+				CachedMaxInstancesPerComponent == Other.CachedMaxInstancesPerComponent &&
+				SubsectionX == Other.SubsectionX &&
+				SubsectionY == Other.SubsectionY &&
+				BasedOn == Other.BasedOn &&
+				Layer == Other.Layer;
+		}
+
+		friend uint32 GetTypeHash(const FGrassCompKey& Key)
+		{
+			return GetTypeHash(Key.BasedOn) ^ GetTypeHash(Key.Layer) ^ Key.SqrtSubsections ^ Key.CachedMaxInstancesPerComponent ^ (Key.SubsectionX >> 16) ^ (Key.SubsectionY >> 24);
+		}
+
+	};
+
+	struct FGrassComp
+	{
+		FGrassCompKey Key;
 		TWeakObjectPtr<UHierarchicalInstancedStaticMeshComponent> Foliage;
 		uint32 LastUsedFrameNumber;
 		double LastUsedTime;
 		bool Pending;
 
-		FPerGrassComp(UHierarchicalInstancedStaticMeshComponent* InFoliage)
-			: Foliage(InFoliage)
-			, Pending(true)
+		FGrassComp()
+			: Pending(true)
 		{
 			Touch();
 		}
@@ -212,35 +246,31 @@ struct FCachedLandscapeFoliage
 			LastUsedTime = FPlatformTime::Seconds();
 		}
 	};
-	struct FPerLayer
+
+	struct FGrassCompKeyFuncs : BaseKeyFuncs<FGrassComp,FGrassCompKey>
 	{
-		int32 SqrtSubsections;
-		int32 CachedMaxInstancesPerComponent;
-		TArray<FPerGrassComp> Foliage;
-		TWeakObjectPtr<ULandscapeLayerInfoObject> Layer;
+		static KeyInitType GetSetKey(const FGrassComp& Element)
+		{
+			return Element.Key;
+		}
 
-		FPerLayer(int32 InSqrtSubsections, int32 InCachedMaxInstancesPerComponent, ULandscapeLayerInfoObject* InLayer)
-			: SqrtSubsections(InSqrtSubsections)
-			, CachedMaxInstancesPerComponent(InCachedMaxInstancesPerComponent)
-			, Layer(InLayer)
+		static bool Matches(KeyInitType A, KeyInitType B)
 		{
-			Foliage.Reserve(SqrtSubsections * SqrtSubsections);
+			return A == B;
 		}
-	};
-	struct FPerComponent
-	{
-		TWeakObjectPtr<ULandscapeComponent> BasedOn;
-		TArray<FPerLayer> Layers;
-		FPerComponent(ULandscapeComponent* InBasedOn)
-			: BasedOn(InBasedOn)
+
+		static uint32 GetKeyHash(KeyInitType Key)
 		{
+			return GetTypeHash(Key);
 		}
 	};
-	TArray<FPerComponent> PerComponent;
+
+	typedef TSet<FGrassComp, FGrassCompKeyFuncs> TGrassSet;
+	TSet<FGrassComp, FGrassCompKeyFuncs> CachedGrassComps;
 
 	void ClearCache()
 	{
-		PerComponent.Empty();
+		CachedGrassComps.Empty();
 	}
 };
 
@@ -248,15 +278,13 @@ class FAsyncGrassTask : public FNonAbandonableTask
 {
 public:
 	struct FAsyncGrassBuilder* Builder;
-	TWeakObjectPtr<ULandscapeComponent> BasedOn;
+	FCachedLandscapeFoliage::FGrassCompKey Key;
 	TWeakObjectPtr<UHierarchicalInstancedStaticMeshComponent> Foliage;
-	TWeakObjectPtr<ULandscapeLayerInfoObject> Layer;
 
-	FAsyncGrassTask(struct FAsyncGrassBuilder* InBuilder, ULandscapeComponent* InBasedOn, UHierarchicalInstancedStaticMeshComponent* InFoliage, ULandscapeLayerInfoObject* InLayer)
+	FAsyncGrassTask(struct FAsyncGrassBuilder* InBuilder, const FCachedLandscapeFoliage::FGrassCompKey& InKey, UHierarchicalInstancedStaticMeshComponent* InFoliage)
 		: Builder(InBuilder)
-		, BasedOn(InBasedOn)
+		, Key(InKey)
 		, Foliage(InFoliage)
-		, Layer(InLayer)
 	{
 	}
 	void DoWork();
diff --git a/Engine/Source/Runtime/Landscape/Private/Landscape.cpp b/Engine/Source/Runtime/Landscape/Private/Landscape.cpp
--- a/Engine/Source/Runtime/Landscape/Private/Landscape.cpp
+++ b/Engine/Source/Runtime/Landscape/Private/Landscape.cpp
@@ -2136,17 +2136,13 @@ void ALandscapeProxy::FlushFoliageComponents(const TSet<ULandscapeComponent*>* O
 			}
 		}
 #endif
-		for (FCachedLandscapeFoliage::FPerComponent& CacheItem : FoliageCache.PerComponent)
+		for (auto& CacheItem : FoliageCache.CachedGrassComps)
 		{
-			ULandscapeComponent* Component = CacheItem.BasedOn.Get();
+			ULandscapeComponent* Component = CacheItem.Key.BasedOn.Get();
 			// if the weak pointer in the cache is invalid, we should kill them anyway
 			if (Component == nullptr || OnlyForComponents->Contains(Component))
 			{
-				for (FCachedLandscapeFoliage::FPerLayer& LayerItem : CacheItem.Layers)
-				{
-					for (FCachedLandscapeFoliage::FPerGrassComp& GrassItem : LayerItem.Foliage)
-					{
-						UHierarchicalInstancedStaticMeshComponent *Used = GrassItem.Foliage.Get();
+				UHierarchicalInstancedStaticMeshComponent *Used = CacheItem.Foliage.Get();
 						if (Used)
 						{
 							SCOPE_CYCLE_COUNTER(STAT_FoliageGrassDestoryComp);
@@ -2156,8 +2152,6 @@ void ALandscapeProxy::FlushFoliageComponents(const TSet<ULandscapeComponent*>* O
 					}
 				}
 			}
-		}
-	}
 	else
 	{
 		// Clear old foliage component containers
@@ -2206,8 +2200,13 @@ DEFINE_LOG_CATEGORY_STATIC(LogGrass, Log, All);
 
 static TAutoConsoleVariable<float> CVarGuardBandMultiplier(
 	TEXT("grass.GuardBandMultiplier"),
-	1.1f,
-	TEXT("Used to control discarding in the grass system. Approximate range, 1-2. Multiplied by the cull distance to control when we discard grass components."));
+	1.3f,
+	TEXT("Used to control discarding in the grass system. Approximate range, 1-4. Multiplied by the cull distance to control when we add grass components."));
+
+static TAutoConsoleVariable<float> CVarGuardBandDiscardMultiplier(
+	TEXT("grass.GuardBandDiscardMultiplier"),
+	1.4f,
+	TEXT("Used to control discarding in the grass system. Approximate range, 1-4. Multiplied by the cull distance to control when we discard grass components."));
 
 static TAutoConsoleVariable<int32> CVarMinFramesToKeepGrass(
 	TEXT("grass.MinFramesToKeepGrass"),
@@ -2239,6 +2238,16 @@ static TAutoConsoleVariable<int32> CVarUseStreamingManagerForCameras(
 	1,
 	TEXT("1: Use Streaming Manager; 0: Use ViewLocationsRenderedLastFrame"));
 
+static TAutoConsoleVariable<int32> CVarCullSubsections(
+	TEXT("grass.CullSubsections"),
+	1,
+	TEXT("1: Cull each foliage component; 0: Cull only based on the landscape component."));
+
+static TAutoConsoleVariable<int32> CVarDisableGPUCull(
+	TEXT("grass.DisableGPUCull"),
+	0,
+	TEXT("For debugging. Set this to zero to see where the grass is generated. Useful for tweaking the guard bands."));
+
 #if WITH_EDITOR && 0 // we will remove this shortly
 
 
@@ -2301,7 +2310,7 @@ class FGameGrassLayerData
 	TSharedPtr<FGrassMap, ESPMode::ThreadSafe> GrassMap;
 	const uint8* GrassMapPtr;
 public:
-	FGameGrassLayerData(ALandscapeProxy* Landscape, ULandscapeComponent* Component, ULandscapeLayerInfoObject* LayerInfo, UHierarchicalInstancedStaticMeshComponent* HierarchicalInstancedStaticMeshComponent, FGameGrassLayerData*)
+	FGameGrassLayerData(ALandscapeProxy* Landscape, ULandscapeComponent* Component, ULandscapeLayerInfoObject* LayerInfo, UHierarchicalInstancedStaticMeshComponent* HierarchicalInstancedStaticMeshComponent)
 	{
 		int32 SubsectionSizeQuads = Component->SubsectionSizeQuads;
 		int32 NumSubsections = Component->NumSubsections;
@@ -2416,7 +2425,6 @@ struct FGrassBuilderBase
 			Origin += Extent * FVector(float(SubX), float(SubY), 0.0f);
 		}
 	}
-
 };
 
 struct FAsyncGrassBuilder : public FGrassBuilderBase
@@ -2438,9 +2446,9 @@ struct FAsyncGrassBuilder : public FGrassBuilderBase
 	FStaticMeshInstanceData InstanceBuffer;
 	TArray<FClusterNode> ClusterTree;
 
-	FAsyncGrassBuilder(ALandscapeProxy* Landscape, ULandscapeComponent* Component, ULandscapeLayerInfoObject* LayerInfo, UHierarchicalInstancedStaticMeshComponent* HierarchicalInstancedStaticMeshComponent, int32 SqrtSubsections, int32 SubX, int32 SubY, FAsyncGrassBuilder* CopyFrom)
+	FAsyncGrassBuilder(ALandscapeProxy* Landscape, ULandscapeComponent* Component, ULandscapeLayerInfoObject* LayerInfo, UHierarchicalInstancedStaticMeshComponent* HierarchicalInstancedStaticMeshComponent, int32 SqrtSubsections, int32 SubX, int32 SubY)
 		: FGrassBuilderBase(Landscape, Component, LayerInfo, SqrtSubsections, SubX, SubY)
-		, LayerData(Landscape, Component, LayerInfo, HierarchicalInstancedStaticMeshComponent, CopyFrom ? &CopyFrom->LayerData : nullptr)
+		, LayerData(Landscape, Component, LayerInfo, HierarchicalInstancedStaticMeshComponent)
 	{
 		bHaveValidData = bHaveValidData && LayerData.HasData();
 		XForm = LandscapeToWorld * HierarchicalInstancedStaticMeshComponent->ComponentToWorld.ToMatrixWithScale().Inverse();
@@ -2719,6 +2727,9 @@ void ALandscapeProxy::UpdateFoliage(const TArray<FVector>& Cameras, ULandscapeCo
 	if (CVarGrassEnable.GetValueOnAnyThread() > 0)
 	{
 		float GuardBand = CVarGuardBandMultiplier.GetValueOnAnyThread();
+		float DiscardGuardBand = CVarGuardBandDiscardMultiplier.GetValueOnAnyThread();
+		bool bCullSubsections = CVarCullSubsections.GetValueOnAnyThread() > 0;
+		bool bDisableGPUCull = CVarDisableGPUCull.GetValueOnAnyThread() > 0;
 		int32 MaxInstancesPerComponent = FMath::Max<int32>(1024, CVarMaxInstancesPerComponent.GetValueOnAnyThread());
 
 		for (int32 ComponentIndex = 0; ComponentIndex < LandscapeComponents.Num(); ComponentIndex++)
@@ -2729,22 +2740,6 @@ void ALandscapeProxy::UpdateFoliage(const TArray<FVector>& Cameras, ULandscapeCo
 				continue;
 			}
 
-			FCachedLandscapeFoliage::FPerComponent* ExistingCache = nullptr;
-			for (FCachedLandscapeFoliage::FPerComponent& CacheItem : FoliageCache.PerComponent)
-			{
-				if (CacheItem.BasedOn.Get() == Component)
-				{
-					ExistingCache = &CacheItem;
-					break;
-				}
-			}
-			if (!ExistingCache)
-			{
-				ExistingCache = &FoliageCache.PerComponent[FoliageCache.PerComponent.Add(FCachedLandscapeFoliage::FPerComponent(Component))];
-			}
-			check(ExistingCache);
-
-
 			FBoxSphereBounds WorldBounds = Component->CalcBounds(Component->ComponentToWorld);
 			float MinDistanceToComp = Cameras.Num() ? MAX_flt : 0.0f;
 
@@ -2763,56 +2758,13 @@ void ALandscapeProxy::UpdateFoliage(const TArray<FVector>& Cameras, ULandscapeCo
 				{
 					bool bRandomRotation = LayerInfo->RandomRotation;
 					bool bAlign = LayerInfo->AlignToSurface;
-					float MustHaveDistance = (float)LayerInfo->EndCullDistance;
-					float DiscardDistance = MustHaveDistance * GuardBand;
+					float MustHaveDistance = GuardBand * (float)LayerInfo->EndCullDistance;
+					float DiscardDistance = DiscardGuardBand * (float)LayerInfo->EndCullDistance;
 
 					if (MinDistanceToComp > DiscardDistance)
 					{
 						continue;
 					}
-					bool bFound = false;
-					FCachedLandscapeFoliage::FPerLayer* ExistingLayer = nullptr;
-					for (int32 Index = 0; Index < ExistingCache->Layers.Num(); Index++)
-					{
-						FCachedLandscapeFoliage::FPerLayer& Layer = ExistingCache->Layers[Index];
-						if (Layer.Layer.Get() == LayerInfo && Layer.CachedMaxInstancesPerComponent == MaxInstancesPerComponent)
-						{
-							ExistingLayer = &Layer;
-							check(Layer.SqrtSubsections >= 1);
-							int32 NumValid = 0;
-							for (const FCachedLandscapeFoliage::FPerGrassComp& GrassItem : Layer.Foliage)
-							{
-								const TWeakObjectPtr<UHierarchicalInstancedStaticMeshComponent>& Comp = GrassItem.Foliage;
-								if (Comp.Get())
-								{
-									NumValid++;
-								}
-								else
-								{
-									NumValid = -1;
-									break;
-								}
-							}
-							if (NumValid == Layer.SqrtSubsections * Layer.SqrtSubsections)
-							{
-								for (FCachedLandscapeFoliage::FPerGrassComp& GrassItem : Layer.Foliage)
-								{
-									GrassItem.Touch();
-								}
-								bFound = true;
-								break;
-							}
-							else
-							{
-								// this must be stale, lets delete them
-								Layer.Foliage.Empty();
-							}
-						}
-					}
-					if (bFound || MinDistanceToComp > MustHaveDistance)
-					{
-						continue;
-					}
 
 					FGrassBuilderBase ForSubsectionMath(this, Component, LayerInfo);
 
@@ -2822,20 +2774,71 @@ void ALandscapeProxy::UpdateFoliage(const TArray<FVector>& Cameras, ULandscapeCo
 					{
 						SqrtSubsections = FMath::Clamp<int32>(FMath::CeilToInt(float(ForSubsectionMath.SqrtMaxInstances) / FMath::Sqrt((float)MaxInstancesPerComponent)), 1, 16);
 					}
-					if (!ExistingLayer)
-					{
-						ExistingLayer = &ExistingCache->Layers[ExistingCache->Layers.Add(FCachedLandscapeFoliage::FPerLayer(SqrtSubsections, MaxInstancesPerComponent, LayerInfo))];
-					}
-					check(ExistingLayer);
-					FCachedLandscapeFoliage::FPerLayer& NewLayer = *ExistingLayer;
 
-					FAsyncGrassBuilder* CopyFromBuilder = nullptr;
+					int32 NumCompsCreated = 0;
+					FBox LocalBox = Component->CachedLocalBox;
+					FVector LocalExtentDiv = (LocalBox.Max - LocalBox.Min) * FVector(1.0f / float(SqrtSubsections), 1.0f / float(SqrtSubsections), 1.0f);
 					for (int32 SubX = 0; SubX < SqrtSubsections; SubX++)
 					{
 						for (int32 SubY = 0; SubY < SqrtSubsections; SubY++)
 						{
-							SCOPE_CYCLE_COUNTER(STAT_FoliageGrassStartComp);
+							float MinDistanceToSubComp = MinDistanceToComp;
 
+							if (bCullSubsections && SqrtSubsections > 1)
+							{
+								FVector BoxMin;
+								BoxMin.X = LocalBox.Min.X + LocalExtentDiv.X * float(SubX);
+								BoxMin.Y = LocalBox.Min.Y + LocalExtentDiv.Y * float(SubY);
+								BoxMin.Z = LocalBox.Min.Z;
+
+								FVector BoxMax;
+								BoxMax.X = LocalBox.Min.X + LocalExtentDiv.X * float(SubX + 1);
+								BoxMax.Y = LocalBox.Min.Y + LocalExtentDiv.Y * float(SubY + 1);
+								BoxMax.Z = LocalBox.Max.Z;
+
+								FBox LocalSubBox(BoxMin, BoxMax);
+								FBox WorldSubBox = LocalSubBox.TransformBy(Component->ComponentToWorld);
+
+								MinDistanceToSubComp = Cameras.Num() ? MAX_flt : 0.0f;
+								for (auto& Pos : Cameras)
+								{
+									MinDistanceToSubComp = FMath::Min<float>(MinDistanceToSubComp, ComputeSquaredDistanceFromBoxToPoint(WorldSubBox.Min, WorldSubBox.Max, Pos));
+								}
+								MinDistanceToSubComp = FMath::Sqrt(MinDistanceToSubComp);
+							}
+
+							if (MinDistanceToSubComp > DiscardDistance)
+							{
+								continue;
+							}
+
+							FCachedLandscapeFoliage::FGrassComp NewComp;
+							NewComp.Key.BasedOn = Component;
+							NewComp.Key.Layer = LayerInfo;
+							NewComp.Key.SqrtSubsections = SqrtSubsections;
+							NewComp.Key.CachedMaxInstancesPerComponent = MaxInstancesPerComponent;
+							NewComp.Key.SubsectionX = SubX;
+							NewComp.Key.SubsectionY = SubY;
+
+							{
+								FCachedLandscapeFoliage::FGrassComp* Existing = FoliageCache.CachedGrassComps.Find(NewComp.Key);
+								if (Existing || MinDistanceToSubComp > MustHaveDistance)
+								{
+									if (Existing)
+									{
+										Existing->Touch();
+									}
+									continue;
+								}
+							}
+
+							if (NumCompsCreated)
+							{
+								continue; // one per frame, but we still want to touch the existing ones
+							}
+							NumCompsCreated++;
+
+							SCOPE_CYCLE_COUNTER(STAT_FoliageGrassStartComp);
 							int32 FolSeed = FCrc::StrCrc32((LayerInfo->LayerName.ToString() + Component->GetName() + FString::Printf(TEXT("%d %d"), SubX, SubY)).GetCharArray().GetData());
 							if (FolSeed == 0)
 							{
@@ -2847,6 +2850,9 @@ void ALandscapeProxy::UpdateFoliage(const TArray<FVector>& Cameras, ULandscapeCo
 								QUICK_SCOPE_CYCLE_COUNTER(STAT_GrassCreateComp);
 								HierarchicalInstancedStaticMeshComponent = ConstructObject<UHierarchicalInstancedStaticMeshComponent>(UHierarchicalInstancedStaticMeshComponent::StaticClass(), this); //, NAME_None, RF_Transactional);
 							}
+							NewComp.Foliage = HierarchicalInstancedStaticMeshComponent;
+							FoliageCache.CachedGrassComps.Add(NewComp);
+
 							HierarchicalInstancedStaticMeshComponent->Mobility = EComponentMobility::Static;
 
 							HierarchicalInstancedStaticMeshComponent->StaticMesh = LayerInfo->GrassMesh;
@@ -2858,7 +2864,7 @@ void ALandscapeProxy::UpdateFoliage(const TArray<FVector>& Cameras, ULandscapeCo
 							HierarchicalInstancedStaticMeshComponent->SetCanEverAffectNavigation(false);
 							HierarchicalInstancedStaticMeshComponent->InstancingRandomSeed = FolSeed;
 
-							if (!Cameras.Num())
+							if (!Cameras.Num() || bDisableGPUCull)
 							{
 								// if we don't have any cameras, then we are rendering landscape LOD materials or somesuch and we want to disable culling
 								HierarchicalInstancedStaticMeshComponent->InstanceStartCullDistance = 0;
@@ -2872,17 +2878,6 @@ void ALandscapeProxy::UpdateFoliage(const TArray<FVector>& Cameras, ULandscapeCo
 
 							//@todo - take the settings from a UFoliageType object.  For now, disable distance field lighting on grass so we don't hitch.
 							HierarchicalInstancedStaticMeshComponent->bAffectDistanceFieldLighting = false;
-	#if 0
-							HierarchicalInstancedStaticMeshComponent->CastShadow                     = Settings->CastShadow;
-							HierarchicalInstancedStaticMeshComponent->bCastDynamicShadow             = Settings->bCastDynamicShadow;
-							HierarchicalInstancedStaticMeshComponent->bCastStaticShadow              = Settings->bCastStaticShadow;
-							HierarchicalInstancedStaticMeshComponent->bAffectDynamicIndirectLighting = Settings->bAffectDynamicIndirectLighting;
-							HierarchicalInstancedStaticMeshComponent->bCastShadowAsTwoSided          = Settings->bCastShadowAsTwoSided;
-							HierarchicalInstancedStaticMeshComponent->bReceivesDecals                = Settings->bReceivesDecals;
-							HierarchicalInstancedStaticMeshComponent->bOverrideLightMapRes           = Settings->bOverrideLightMapRes;
-							HierarchicalInstancedStaticMeshComponent->OverriddenLightMapRes          = Settings->OverriddenLightMapRes;
-							HierarchicalInstancedStaticMeshComponent->BodyInstance.CopyBodyInstancePropertiesFrom(&Settings->BodyInstance);
-	#endif
 
 							{
 								QUICK_SCOPE_CYCLE_COUNTER(STAT_GrassAttachComp);
@@ -2893,21 +2888,18 @@ void ALandscapeProxy::UpdateFoliage(const TArray<FVector>& Cameras, ULandscapeCo
 								HierarchicalInstancedStaticMeshComponent->SetWorldTransform(DesiredTransform);
 
 								FoliageComponents.Add(HierarchicalInstancedStaticMeshComponent);
-								NewLayer.Foliage.Add(FCachedLandscapeFoliage::FPerGrassComp(HierarchicalInstancedStaticMeshComponent));
 							}
 
 							FAsyncGrassBuilder* Builder;
 
 							{
 								QUICK_SCOPE_CYCLE_COUNTER(STAT_GrassCreateBuilder);
-								Builder = new FAsyncGrassBuilder(this, Component, LayerInfo, HierarchicalInstancedStaticMeshComponent, SqrtSubsections, SubX, SubY, CopyFromBuilder);
+								Builder = new FAsyncGrassBuilder(this, Component, LayerInfo, HierarchicalInstancedStaticMeshComponent, SqrtSubsections, SubX, SubY);
 							}
 
-							CopyFromBuilder = Builder;
-
 							if (Builder->bHaveValidData)
 							{
-								FAsyncTask<FAsyncGrassTask>* Task = new FAsyncTask<FAsyncGrassTask>(Builder, Component, HierarchicalInstancedStaticMeshComponent, LayerInfo);
+								FAsyncTask<FAsyncGrassTask>* Task = new FAsyncTask<FAsyncGrassTask>(Builder, NewComp.Key, HierarchicalInstancedStaticMeshComponent);
 
 								Task->StartBackgroundTask();
 
@@ -2928,38 +2920,31 @@ void ALandscapeProxy::UpdateFoliage(const TArray<FVector>& Cameras, ULandscapeCo
 		}
 	}
 
-	TSet<UHierarchicalInstancedStaticMeshComponent *> StillUsed;
+	static TSet<UHierarchicalInstancedStaticMeshComponent *> StillUsed;
+	StillUsed.Empty(256);
 	{
 		// trim cached items based on time, pending and emptiness
 		double OldestToKeepTime = FPlatformTime::Seconds() - CVarMinTimeToKeepGrass.GetValueOnGameThread();
 		uint32 OldestToKeepFrame = GFrameNumber - CVarMinFramesToKeepGrass.GetValueOnGameThread();
-		for (int32 CacheIndex = 0; CacheIndex < FoliageCache.PerComponent.Num(); CacheIndex++)
-		{
-			FCachedLandscapeFoliage::FPerComponent& CacheItem = FoliageCache.PerComponent[CacheIndex];
-			for (int32 LayerIndex = 0; LayerIndex < CacheItem.Layers.Num(); LayerIndex++)
+		for (FCachedLandscapeFoliage::TGrassSet::TIterator Iter(FoliageCache.CachedGrassComps); Iter; ++Iter)
+		{
+			const FCachedLandscapeFoliage::FGrassComp& GrassItem = *Iter;
+			UHierarchicalInstancedStaticMeshComponent *Used = GrassItem.Foliage.Get();
+			bool bOld = 
+				!GrassItem.Pending &&
+				(
+					!GrassItem.Key.BasedOn.Get() ||
+					!GrassItem.Key.Layer.Get() ||
+					!Used ||
+					(GrassItem.LastUsedFrameNumber < OldestToKeepFrame && GrassItem.LastUsedTime < OldestToKeepTime)
+				);
+			if (bOld)
 			{
-				FCachedLandscapeFoliage::FPerLayer& LayerItem = CacheItem.Layers[LayerIndex];
-				for (int32 FoliageIndex = 0; FoliageIndex < LayerItem.Foliage.Num(); FoliageIndex++)
-				{
-					FCachedLandscapeFoliage::FPerGrassComp& GrassItem = LayerItem.Foliage[FoliageIndex];
-					UHierarchicalInstancedStaticMeshComponent *Used = GrassItem.Foliage.Get();
-					if (!Used || (!GrassItem.Pending && GrassItem.LastUsedFrameNumber < OldestToKeepFrame && GrassItem.LastUsedTime < OldestToKeepTime))
-					{
-						LayerItem.Foliage.RemoveAtSwap(FoliageIndex--);
-					}
-					else
-					{
-						StillUsed.Add(Used);
-					}
-				}
-				if (!LayerItem.Foliage.Num())
-				{
-					CacheItem.Layers.RemoveAtSwap(LayerIndex--);
-				}
+				Iter.RemoveCurrent();
 			}
-			if (!CacheItem.Layers.Num())
+			else if (Used)
 			{
-				FoliageCache.PerComponent.RemoveAtSwap(CacheIndex--);
+				StillUsed.Add(Used);
 			}
 		}
 	}
@@ -3008,26 +2993,11 @@ void ALandscapeProxy::UpdateFoliage(const TArray<FVector>& Cameras, ULandscapeCo
 					{
 						HierarchicalInstancedStaticMeshComponent->RecreateRenderState_Concurrent();
 					}
-					for (int32 CacheIndex = 0; CacheIndex < FoliageCache.PerComponent.Num(); CacheIndex++)
+					FCachedLandscapeFoliage::FGrassComp* Existing = FoliageCache.CachedGrassComps.Find(Inner.Key);
+					if (Existing)
 					{
-						FCachedLandscapeFoliage::FPerComponent& CacheItem = FoliageCache.PerComponent[CacheIndex];
-						for (int32 LayerIndex = 0; LayerIndex < CacheItem.Layers.Num(); LayerIndex++)
-						{
-							FCachedLandscapeFoliage::FPerLayer& LayerItem = CacheItem.Layers[LayerIndex];
-							for (int32 FoliageIndex = 0; FoliageIndex < LayerItem.Foliage.Num(); FoliageIndex++)
-							{
-								FCachedLandscapeFoliage::FPerGrassComp& GrassItem = LayerItem.Foliage[FoliageIndex];
-								UHierarchicalInstancedStaticMeshComponent *Used = GrassItem.Foliage.Get();
-								if (Used == HierarchicalInstancedStaticMeshComponent)
-								{
-									GrassItem.Pending = false;
-									GrassItem.Touch();
-									CacheIndex = MAX_int32 - 1;
-									LayerIndex = MAX_int32 - 1;
-									break;
-								}
-							}
-						}
+						Existing->Pending = false;
+						Existing->Touch();
 					}
 				}
 				delete Inner.Builder;
