nicer UI

diff --git a/Engine/Source/Runtime/Landscape/Classes/LandscapeProxy.h b/Engine/Source/Runtime/Landscape/Classes/LandscapeProxy.h
--- a/Engine/Source/Runtime/Landscape/Classes/LandscapeProxy.h
+++ b/Engine/Source/Runtime/Landscape/Classes/LandscapeProxy.h
@@ -384,7 +384,7 @@ public:
 	UPROPERTY(EditAnywhere, Category=LandscapeProxy)
 	TLazyObjectPtr<ALandscape> LandscapeActor;
 
-	UPROPERTY(EditAnywhere, Category=Lighting)
+	UPROPERTY(EditAnywhere, Category=Lighting, meta=(DisplayName = "Static Shadow")
 	uint32 bCastStaticShadow:1;
 
 	/** Whether this primitive should cast dynamic shadows as if it were a two sided material. */
