Fix Landscape copy/paste gizmo bounds being calculated incorrectly. This was causing it to randomly disappear.
Fix Landscape copy/paste gizmo not being selected automatically when entering the copy/paste tool, or importing data into the gizmo, or switching to the gizmo brush, or opening the landscape editor when the copy/paste tool or the gizmo brush was the active tool/brush. This was causing it to unexpectedly not be movable/rotatable/resizable in those situations.
Fix Landscape copy/paste gizmo not being selectable with a single click. Coupled with the previous bug, this was causing a lot of confusion!
A bit of cleanup of the gizmo code:
- the landscape edmode is no-longer responsible for rendering a (0% opacity) box around the gizmo to make it selectable, instead the gizmo itself is now selectable.
- the gizmo frustrum transform is now calculated significantly more sanely!

diff --git a/Engine/Source/Editor/LandscapeEditor/Private/LandscapeEdMode.cpp b/Engine/Source/Editor/LandscapeEditor/Private/LandscapeEdMode.cpp
--- a/Engine/Source/Editor/LandscapeEditor/Private/LandscapeEdMode.cpp
+++ b/Engine/Source/Editor/LandscapeEditor/Private/LandscapeEdMode.cpp
@@ -162,10 +162,7 @@ FEdModeLandscape::FEdModeLandscape()
 	, LandscapeRenderAddCollision(nullptr)
 	, CachedLandscapeMaterial(nullptr)
 	, bToolActive(false)
-	, GizmoMaterial(nullptr)
 {
-	GizmoMaterial = LoadObject<UMaterial>(NULL, TEXT("/Engine/EditorLandscapeResources/GizmoMaterial.GizmoMaterial"), NULL, LOAD_None, NULL);
-
 	GLayerDebugColorMaterial = LoadObject<UMaterial>(NULL, TEXT("/Engine/EditorLandscapeResources/LayerVisMaterial.LayerVisMaterial"), NULL, LOAD_None, NULL);
 	GSelectionColorMaterial = LoadObject<UMaterialInstanceConstant>(NULL, TEXT("/Engine/EditorLandscapeResources/SelectBrushMaterial_Selected.SelectBrushMaterial_Selected"), NULL, LOAD_None, NULL);
 	GSelectionRegionMaterial = LoadObject<UMaterialInstanceConstant>(NULL, TEXT("/Engine/EditorLandscapeResources/SelectBrushMaterial_SelectedRegion.SelectBrushMaterial_SelectedRegion"), NULL, LOAD_None, NULL);
@@ -240,7 +237,6 @@ void FEdModeLandscape::AddReferencedObjects(FReferenceCollector& Collector)
 
 	Collector.AddReferencedObject(UISettings);
 
-	Collector.AddReferencedObject(GizmoMaterial);
 	Collector.AddReferencedObject(GLayerDebugColorMaterial);
 	Collector.AddReferencedObject(GSelectionColorMaterial);
 	Collector.AddReferencedObject(GSelectionRegionMaterial);
@@ -1149,9 +1145,13 @@ bool FEdModeLandscape::InputKey(FEditorViewportClient* ViewportClient, FViewport
 	else
 	{
 		// Override Key Input for Selection Brush
-		if (CurrentBrush && CurrentBrush->InputKey(ViewportClient, Viewport, Key, Event) == true)
+		if (CurrentBrush)
 		{
-			return true;
+			TOptional<bool> BrushKeyOverride = CurrentBrush->InputKey(ViewportClient, Viewport, Key, Event);
+			if (BrushKeyOverride.IsSet())
+			{
+				return BrushKeyOverride.GetValue();
+			}
 		}
 
 		if (CurrentTool && CurrentTool->InputKey(ViewportClient, Viewport, Key, Event) == true)
@@ -2036,39 +2036,6 @@ void FEdModeLandscape::Render(const FSceneView* View, FViewport* Viewport, FPrim
 		PDI->DrawLine(LandscapeRenderAddCollision->Corners[3], LandscapeRenderAddCollision->Corners[0], FColor(0, 255, 128), SDPG_Foreground);
 	}
 
-	if (!GIsGizmoDragging && GLandscapeEditRenderMode & ELandscapeEditRenderMode::Gizmo && CurrentToolTarget.LandscapeInfo.IsValid() && CurrentGizmoActor.IsValid() && CurrentGizmoActor->TargetLandscapeInfo)
-	{
-		FDynamicMeshBuilder MeshBuilder;
-
-		for (int32 i = 0; i < 8; ++i)
-		{
-			MeshBuilder.AddVertex(CurrentGizmoActor->FrustumVerts[i], FVector2D(0, 0), FVector(1, 0, 0), FVector(0, 1, 0), FVector(0, 0, 1), FColor::White);
-		}
-
-		// Upper box.
-		MeshBuilder.AddTriangle(0, 2, 1);
-		MeshBuilder.AddTriangle(0, 3, 2);
-		// Lower box.
-		MeshBuilder.AddTriangle(4, 6, 5);
-		MeshBuilder.AddTriangle(4, 7, 6);
-		// Others
-		MeshBuilder.AddTriangle(1, 4, 0);
-		MeshBuilder.AddTriangle(1, 5, 4);
-
-		MeshBuilder.AddTriangle(3, 6, 2);
-		MeshBuilder.AddTriangle(3, 7, 6);
-
-		MeshBuilder.AddTriangle(2, 5, 1);
-		MeshBuilder.AddTriangle(2, 6, 5);
-
-		MeshBuilder.AddTriangle(0, 7, 3);
-		MeshBuilder.AddTriangle(0, 4, 7);
-
-		PDI->SetHitProxy(new HTranslucentActor(CurrentGizmoActor.Get(), NULL));
-		MeshBuilder.Draw(PDI, FMatrix::Identity, GizmoMaterial->GetRenderProxy(false), SDPG_World, true);
-		PDI->SetHitProxy(NULL);
-	}
-
 	// Override Rendering for Splines Tool
 	if (CurrentTool)
 	{
diff --git a/Engine/Source/Editor/LandscapeEditor/Private/LandscapeEdModeBrushes.cpp b/Engine/Source/Editor/LandscapeEditor/Private/LandscapeEdModeBrushes.cpp
--- a/Engine/Source/Editor/LandscapeEditor/Private/LandscapeEdModeBrushes.cpp
+++ b/Engine/Source/Editor/LandscapeEditor/Private/LandscapeEdModeBrushes.cpp
@@ -205,12 +205,6 @@ public:
 		LastMousePosition = FVector2D(LandscapeX, LandscapeY);
 	}
 
-	virtual bool InputKey(FEditorViewportClient* InViewportClient, FViewport* InViewport, FKey InKey, EInputEvent InEvent) override
-	{
-		bool bUpdate = false;
-		return bUpdate;
-	}
-
 	virtual FLandscapeBrushData ApplyBrush(const TArray<FLandscapeToolMousePosition>& InMousePositions) override
 	{
 		ULandscapeInfo* LandscapeInfo = EdMode->CurrentToolTarget.LandscapeInfo.Get();
@@ -433,13 +427,6 @@ public:
 		LastMousePosition = FVector2D(LandscapeX, LandscapeY);
 	}
 
-	virtual bool InputKey(FEditorViewportClient* InViewportClient, FViewport* InViewport, FKey InKey, EInputEvent InEvent) override
-	{
-		bool bUpdate = false;
-
-		return bUpdate;
-	}
-
 	virtual FLandscapeBrushData ApplyBrush(const TArray<FLandscapeToolMousePosition>& MousePositions) override
 	{
 		// Selection Brush only works for 
@@ -548,6 +535,17 @@ public:
 
 	virtual ELandscapeBrushType GetBrushType() override { return ELandscapeBrushType::Gizmo; }
 
+	virtual void EnterBrush() override
+	{
+		// Make sure gizmo actor is selected
+		ALandscapeGizmoActiveActor* Gizmo = EdMode->CurrentGizmoActor.Get();
+		if (Gizmo)
+		{
+			GEditor->SelectNone(false, true);
+			GEditor->SelectActor(Gizmo, true, false, true);
+		}
+	}
+
 	virtual void LeaveBrush() override
 	{
 		for (TSet<ULandscapeComponent*>::TIterator It(BrushMaterialComponents); It; ++It)
@@ -636,10 +634,25 @@ public:
 	{
 	}
 
-	virtual bool InputKey(FEditorViewportClient* InViewportClient, FViewport* InViewport, FKey InKey, EInputEvent InEvent) override
+	virtual TOptional<bool> InputKey(FEditorViewportClient* InViewportClient, FViewport* InViewport, FKey InKey, EInputEvent InEvent) override
 	{
-		bool bUpdate = false;
-		return bUpdate;
+
+		if (Key == EKeys::LeftMouseButton && Event == IE_Pressed)
+		{
+			int32 HitX = InViewport->GetMouseX();
+			int32 HitY = InViewport->GetMouseY();
+			HHitProxy* HitProxy = InViewport->GetHitProxy(HitX, HitY);
+
+			HActor* ActorHitProxy = HitProxyCast<HActor>(HitProxy);
+			if (ActorHitProxy && ActorHitProxy->Actor->IsA<ALandscapeGizmoActor>())
+			{
+				// don't treat clicks on a landscape gizmo as a tool invocation
+				return TOptional<bool>(false);
+			}
+		}
+
+		// default behaviour
+		return TOptional<bool>();
 	}
 
 	virtual FLandscapeBrushData ApplyBrush(const TArray<FLandscapeToolMousePosition>& MousePositions) override
diff --git a/Engine/Source/Editor/LandscapeEditor/Private/LandscapeEdModeComponentTools.cpp b/Engine/Source/Editor/LandscapeEditor/Private/LandscapeEdModeComponentTools.cpp
--- a/Engine/Source/Editor/LandscapeEditor/Private/LandscapeEdModeComponentTools.cpp
+++ b/Engine/Source/Editor/LandscapeEditor/Private/LandscapeEdModeComponentTools.cpp
@@ -1809,6 +1809,17 @@ public:
 	virtual const TCHAR* GetToolName() override { return TEXT("CopyPaste"); }
 	virtual FText GetDisplayName() override { return NSLOCTEXT("UnrealEd", "LandscapeMode_Region", "Region Copy/Paste"); };
 
+	virtual void EnterTool()
+	{
+		// Make sure gizmo actor is selected
+		ALandscapeGizmoActiveActor* Gizmo = EdMode->CurrentGizmoActor.Get();
+		if (Gizmo)
+		{
+			GEditor->SelectNone(false, true);
+			GEditor->SelectActor(Gizmo, true, false, true);
+		}
+	}
+
 	// Copy tool doesn't use any view information, so just do it as one function
 	void Copy()
 	{
diff --git a/Engine/Source/Editor/LandscapeEditor/Private/LandscapeEditorDetailCustomization_CopyPaste.cpp b/Engine/Source/Editor/LandscapeEditor/Private/LandscapeEditorDetailCustomization_CopyPaste.cpp
--- a/Engine/Source/Editor/LandscapeEditor/Private/LandscapeEditorDetailCustomization_CopyPaste.cpp
+++ b/Engine/Source/Editor/LandscapeEditor/Private/LandscapeEditorDetailCustomization_CopyPaste.cpp
@@ -355,6 +355,10 @@ FReply FLandscapeEditorDetailCustomization_CopyPaste::OnGizmoImportButtonClicked
 			}
 
 			Gizmo->Import(LandscapeEdMode->UISettings->GizmoImportSize.X, LandscapeEdMode->UISettings->GizmoImportSize.Y, (uint16*)Data.GetData(), LayerInfos, LayerDataPtrs.Num() ? LayerDataPtrs.GetData() : NULL);
+
+			// Make sure gizmo actor is selected
+			GEditor->SelectNone(false, true);
+			GEditor->SelectActor(Gizmo, true, false, true);
 		}
 	}
 
diff --git a/Engine/Source/Editor/LandscapeEditor/Public/LandscapeEdMode.h b/Engine/Source/Editor/LandscapeEditor/Public/LandscapeEdMode.h
--- a/Engine/Source/Editor/LandscapeEditor/Public/LandscapeEdMode.h
+++ b/Engine/Source/Editor/LandscapeEditor/Public/LandscapeEdMode.h
@@ -98,7 +98,7 @@ class FLandscapeBrush : public FGCObject
 public:
 	virtual void MouseMove(float LandscapeX, float LandscapeY) = 0;
 	virtual FLandscapeBrushData ApplyBrush(const TArray<FLandscapeToolMousePosition>& MousePositions) = 0;
-	virtual bool InputKey(FEditorViewportClient* InViewportClient, FViewport* InViewport, FKey InKey, EInputEvent InEvent) { return false; }
+	virtual TOptional<bool> InputKey(FEditorViewportClient* InViewportClient, FViewport* InViewport, FKey InKey, EInputEvent InEvent) { return false; }
 	virtual void Tick(FEditorViewportClient* ViewportClient, float DeltaTime) {};
 	virtual void BeginStroke(float LandscapeX, float LandscapeY, class FLandscapeTool* CurrentTool);
 	virtual void EndStroke();
@@ -693,7 +693,6 @@ private:
 	UMaterialInterface* CachedLandscapeMaterial;
 
 	bool bToolActive;
-	UMaterial* GizmoMaterial;
 
 	FDelegateHandle OnWorldChangeDelegateHandle;
 	FDelegateHandle OnMaterialCompilationFinishedDelegateHandle;
diff --git a/Engine/Source/Runtime/Landscape/Classes/LandscapeGizmoActiveActor.h b/Engine/Source/Runtime/Landscape/Classes/LandscapeGizmoActiveActor.h
--- a/Engine/Source/Runtime/Landscape/Classes/LandscapeGizmoActiveActor.h
+++ b/Engine/Source/Runtime/Landscape/Classes/LandscapeGizmoActiveActor.h
@@ -36,7 +36,7 @@ struct FGizmoSelectData
 
 		FGizmoSelectData()
 		#if WITH_EDITORONLY_DATA
-		: Ratio(0.f), HeightData(0.f), WeightDataMap()
+		: Ratio(0.0f), HeightData(0.0f), WeightDataMap()
 		#endif
 		{
 		}
diff --git a/Engine/Source/Runtime/Landscape/Private/LandscapeGizmoActor.cpp b/Engine/Source/Runtime/Landscape/Private/LandscapeGizmoActor.cpp
--- a/Engine/Source/Runtime/Landscape/Private/LandscapeGizmoActor.cpp
+++ b/Engine/Source/Runtime/Landscape/Private/LandscapeGizmoActor.cpp
@@ -123,12 +123,14 @@ public:
 	bool bHeightmapRendering;
 	FLandscapeGizmoMeshRenderProxy* HeightmapRenderProxy;
 	FMaterialRenderProxy* GizmoRenderProxy;
+	HHitProxy* HitProxy;
 
 	FLandscapeGizmoRenderSceneProxy(const ULandscapeGizmoRenderComponent* InComponent):
 		FPrimitiveSceneProxy(InComponent),
 		bHeightmapRendering(false),
-		HeightmapRenderProxy(NULL),
-		GizmoRenderProxy(NULL)
+		HeightmapRenderProxy(nullptr),
+		GizmoRenderProxy(nullptr),
+		HitProxy(nullptr)
 	{
 #if WITH_EDITOR	
 		ALandscapeGizmoActiveActor* Gizmo = Cast<ALandscapeGizmoActiveActor>(InComponent->GetOwner());
@@ -140,31 +142,35 @@ public:
 				SampleSizeX = Gizmo->SampleSizeX;
 				SampleSizeY = Gizmo->SampleSizeY;
 				bHeightmapRendering = (Gizmo->DataType & LGT_Height);
-				FMatrix LToW = LandscapeInfo->GetLandscapeProxy()->LandscapeActorToWorld().ToMatrixWithScale();
-				FMatrix WToL = LToW.InverseFast();
-				FVector BaseLocation = WToL.TransformPosition(Gizmo->GetActorLocation());
-				float ScaleXY = LandscapeInfo->DrawScale.X;
-				float ScaleZ = LandscapeInfo->DrawScale.Z;
-				const float W = Gizmo->GetWidth() / (2 * ScaleXY);
-				const float H = Gizmo->GetHeight() / (2 * ScaleXY);
-				const float L = Gizmo->GetLength() / ScaleZ;
-				FMatrix GizmoRT = FRotationTranslationMatrix(FRotator(0, Gizmo->GetActorRotation().Yaw, 0), FVector(BaseLocation.X, BaseLocation.Y, 0)) * LToW;
-
-				FrustumVerts[0] = Gizmo->FrustumVerts[0] = GizmoRT.TransformPosition(FVector( - W, - H, BaseLocation.Z + L ));
-				FrustumVerts[1] = Gizmo->FrustumVerts[1] = GizmoRT.TransformPosition(FVector( + W, - H, BaseLocation.Z + L ));
-				FrustumVerts[2] = Gizmo->FrustumVerts[2] = GizmoRT.TransformPosition(FVector( + W, + H, BaseLocation.Z + L ));
-				FrustumVerts[3] = Gizmo->FrustumVerts[3] = GizmoRT.TransformPosition(FVector( - W, + H, BaseLocation.Z + L ));
-
-				FrustumVerts[4] = Gizmo->FrustumVerts[4] = GizmoRT.TransformPosition(FVector( - W, - H, BaseLocation.Z ));
-				FrustumVerts[5] = Gizmo->FrustumVerts[5] = GizmoRT.TransformPosition(FVector( + W, - H, BaseLocation.Z ));
-				FrustumVerts[6] = Gizmo->FrustumVerts[6] = GizmoRT.TransformPosition(FVector( + W, + H, BaseLocation.Z ));
-				FrustumVerts[7] = Gizmo->FrustumVerts[7] = GizmoRT.TransformPosition(FVector( - W, + H, BaseLocation.Z ));
-
-				XAxis = GizmoRT.TransformPosition(FVector( + W,	0,		BaseLocation.Z + L ));
-				YAxis = GizmoRT.TransformPosition(FVector( 0,	+ H,	BaseLocation.Z + L ));
-				Origin = GizmoRT.TransformPosition(FVector( 0,	0,		BaseLocation.Z + L ));
-
-				MeshRT = FTranslationMatrix(FVector(- W + 0.5, - H + 0.5, 0)) * FRotationTranslationMatrix(FRotator(0, Gizmo->GetActorRotation().Yaw, 0), FVector(BaseLocation.X, BaseLocation.Y, 0)) * LToW;
+				FTransform LToW = LandscapeInfo->GetLandscapeProxy()->LandscapeActorToWorld();
+				const float W = Gizmo->Width / 2;
+				const float H = Gizmo->Height / 2;
+				const float L = Gizmo->LengthZ;
+				// The Gizmo's coordinate space is weird, it's partially relative to the landscape and partially relative to the world
+				const FVector GizmoLocation = Gizmo->GetActorLocation();
+				const FQuat   GizmoRotation = FRotator(0, Gizmo->GetActorRotation().Yaw, 0).Quaternion() * LToW.GetRotation();
+				const FVector GizmoScale3D  = Gizmo->GetActorScale3D();
+				const FTransform GizmoRT = FTransform(GizmoRotation, GizmoLocation, GizmoScale3D);
+
+				FrustumVerts[0] = Gizmo->FrustumVerts[0] = GizmoRT.TransformPosition(FVector( - W, - H, + L ));
+				FrustumVerts[1] = Gizmo->FrustumVerts[1] = GizmoRT.TransformPosition(FVector( + W, - H, + L ));
+				FrustumVerts[2] = Gizmo->FrustumVerts[2] = GizmoRT.TransformPosition(FVector( + W, + H, + L ));
+				FrustumVerts[3] = Gizmo->FrustumVerts[3] = GizmoRT.TransformPosition(FVector( - W, + H, + L ));
+
+				FrustumVerts[4] = Gizmo->FrustumVerts[4] = GizmoRT.TransformPosition(FVector( - W, - H,   0 ));
+				FrustumVerts[5] = Gizmo->FrustumVerts[5] = GizmoRT.TransformPosition(FVector( + W, - H,   0 ));
+				FrustumVerts[6] = Gizmo->FrustumVerts[6] = GizmoRT.TransformPosition(FVector( + W, + H,   0 ));
+				FrustumVerts[7] = Gizmo->FrustumVerts[7] = GizmoRT.TransformPosition(FVector( - W, + H,   0 ));
+
+				XAxis  = GizmoRT.TransformPosition(FVector( + W,   0, + L ));
+				YAxis  = GizmoRT.TransformPosition(FVector(   0, + H, + L ));
+				Origin = GizmoRT.TransformPosition(FVector(   0,   0, + L ));
+
+				const FMatrix WToL = LToW.ToMatrixWithScale().InverseFast();
+				const FVector BaseLocation = WToL.TransformPosition(Gizmo->GetActorLocation());
+				const float ScaleXY = LandscapeInfo->DrawScale.X;
+
+				MeshRT = FTranslationMatrix(FVector(-W / ScaleXY + 0.5, -H / ScaleXY + 0.5, 0)) * FRotationTranslationMatrix(FRotator(0, Gizmo->GetActorRotation().Yaw, 0), FVector(BaseLocation.X, BaseLocation.Y, 0)) * LToW.ToMatrixWithScale();
 				HeightmapRenderProxy = new FLandscapeGizmoMeshRenderProxy( Gizmo->GizmoMeshMaterial->GetRenderProxy(false), BaseLocation.Z + L, BaseLocation.Z, Gizmo->GizmoTexture, FLinearColor(Gizmo->TextureScale.X, Gizmo->TextureScale.Y, 0, 0), WToL );
 
 				GizmoRenderProxy = (Gizmo->DataType != LGT_None) ? Gizmo->GizmoDataMaterial->GetRenderProxy(false) : Gizmo->GizmoMaterial->GetRenderProxy(false);
@@ -207,6 +213,18 @@ public:
 		HeightmapRenderProxy = NULL;
 	}
 
+#if WITH_EDITOR
+	virtual HHitProxy* CreateHitProxies(UPrimitiveComponent* Component, TArray<TRefCountPtr<HHitProxy> >& OutHitProxies) override
+	{
+		ALandscapeGizmoActiveActor* Gizmo = CastChecked<ALandscapeGizmoActiveActor>(Component->GetOwner());
+		HitProxy = new HTranslucentActor(Gizmo, Component);
+		OutHitProxies.Add(HitProxy);
+
+		// by default we're not clickable, to allow the preview heightmap to be non-clickable (only the bounds frame)
+		return nullptr;
+	}
+#endif
+
 	virtual void GetDynamicMeshElements(const TArray<const FSceneView*>& Views, const FSceneViewFamily& ViewFamily, uint32 VisibilityMap, FMeshElementCollector& Collector) const override
 	{
 		//FMemMark Mark(FMemStack::Get());
@@ -265,11 +283,11 @@ public:
 							MeshBuilder.AddTriangle( Idx, Idx+3, Idx+2 );
 						}
 
-						MeshBuilder.GetMesh(FMatrix::Identity, GizmoRenderProxy , SDPG_World, true, false, ViewIndex, Collector);
+						MeshBuilder.GetMesh(FMatrix::Identity, GizmoRenderProxy, SDPG_World, true, false, false, ViewIndex, Collector, HitProxy);
 					}
 
 					if (bHeightmapRendering)
-					{		  		
+					{
 						FDynamicMeshBuilder MeshBuilder;
 
 						for (int32 Y = 0; Y < SampleSizeY; ++Y)
@@ -326,7 +344,6 @@ ULandscapeGizmoRenderComponent::ULandscapeGizmoRenderComponent(const FObjectInit
 	bHiddenInGame = true;
 	AlwaysLoadOnClient = false;
 	AlwaysLoadOnServer = false;
-	bSelectable = false;
 	SetCollisionProfileName(UCollisionProfile::BlockAll_ProfileName);
 }
 
@@ -339,14 +356,28 @@ FBoxSphereBounds ULandscapeGizmoRenderComponent::CalcBounds(const FTransform& Lo
 {
 #if WITH_EDITOR
 	ALandscapeGizmoActiveActor* Gizmo = Cast<ALandscapeGizmoActiveActor>(GetOwner());
-	if( Gizmo )
+	if (Gizmo)
 	{
-		return FBoxSphereBounds(Gizmo->FrustumVerts, 8);
+		ULandscapeInfo* LandscapeInfo = Gizmo->TargetLandscapeInfo;
+		if (LandscapeInfo && LandscapeInfo->GetLandscapeProxy())
+		{
+			FTransform LToW = LandscapeInfo->GetLandscapeProxy()->LandscapeActorToWorld();
+
+			// We calculate this ourselves, not from Gizmo->FrustrumVerts, as those haven't been updated yet
+			// The Gizmo's coordinate space is weird, it's partially relative to the landscape and partially relative to the world
+			const FVector GizmoLocation = Gizmo->GetActorLocation();
+			const FQuat   GizmoRotation = FRotator(0, Gizmo->GetActorRotation().Yaw, 0).Quaternion() * LToW.GetRotation();
+			const FVector GizmoScale3D = Gizmo->GetActorScale3D();
+			const FTransform GizmoRT = FTransform(GizmoRotation, GizmoLocation, GizmoScale3D);
+			const float W = Gizmo->Width / 2;
+			const float H = Gizmo->Height / 2;
+			const float L = Gizmo->LengthZ;
+			return FBoxSphereBounds(FBox(FVector(-W, -H, 0), FVector(+W, +H, +L))).TransformBy(GizmoRT);
+		}
 	}
 #endif
-	{
-		return Super::CalcBounds(LocalToWorld);
-	}
+
+	return Super::CalcBounds(LocalToWorld);
 }
 
 ALandscapeGizmoActor::ALandscapeGizmoActor(const FObjectInitializer& ObjectInitializer)
