Fix non-unity build error

diff --git a/Engine/Source/Runtime/Landscape/Private/LandscapeGizmoActor.cpp b/Engine/Source/Runtime/Landscape/Private/LandscapeGizmoActor.cpp
--- a/Engine/Source/Runtime/Landscape/Private/LandscapeGizmoActor.cpp
+++ b/Engine/Source/Runtime/Landscape/Private/LandscapeGizmoActor.cpp
@@ -11,6 +11,7 @@
 #include "LandscapeLayerInfoObject.h"
 #include "DynamicMeshBuilder.h"
 #include "Engine/CollisionProfile.h"
+#include "EngineUtils.h"
 
 namespace
 {
