UE4 - fix assert when a grass component generated no grass

diff --git a/Engine/Source/Runtime/Landscape/Private/Landscape.cpp b/Engine/Source/Runtime/Landscape/Private/Landscape.cpp
--- a/Engine/Source/Runtime/Landscape/Private/Landscape.cpp
+++ b/Engine/Source/Runtime/Landscape/Private/Landscape.cpp
@@ -2971,20 +2971,22 @@ void ALandscapeProxy::UpdateFoliage(const TArray<FVector>& Cameras, ULandscapeCo
 				UHierarchicalInstancedStaticMeshComponent* HierarchicalInstancedStaticMeshComponent = Inner.Foliage.Get();
 				if (HierarchicalInstancedStaticMeshComponent && StillUsed.Contains(HierarchicalInstancedStaticMeshComponent))
 				{
-					Exchange(HierarchicalInstancedStaticMeshComponent->WriteOncePrebuiltInstanceBuffer, Inner.Builder->InstanceBuffer);
-					HierarchicalInstancedStaticMeshComponent->AcceptPrebuiltTree(Inner.Builder->ClusterTree);
-
-					if (bForceSync)
-					{
-						HierarchicalInstancedStaticMeshComponent->RecreateRenderState_Concurrent();
-					}
-					FCachedLandscapeFoliage::FGrassComp* Existing = FoliageCache.CachedGrassComps.Find(Inner.Key);
-					if (Existing)
+					if (Inner.Builder->InstanceBuffer.Num())
 					{
-						Existing->Pending = false;
-						Existing->Touch();
+						Exchange(HierarchicalInstancedStaticMeshComponent->WriteOncePrebuiltInstanceBuffer, Inner.Builder->InstanceBuffer);
+						HierarchicalInstancedStaticMeshComponent->AcceptPrebuiltTree(Inner.Builder->ClusterTree);
+						if (bForceSync)
+						{
+							HierarchicalInstancedStaticMeshComponent->RecreateRenderState_Concurrent();
+						}
 					}
 				}
+				FCachedLandscapeFoliage::FGrassComp* Existing = FoliageCache.CachedGrassComps.Find(Inner.Key);
+				if (Existing)
+				{
+					Existing->Pending = false;
+					Existing->Touch();
+				}
 				delete Inner.Builder;
 				delete Task;
 				if (!bForceSync)
