Fixed copying data to the landscape gizmo when the "current level" doesn't have a proxy of the landscape in it

diff --git a/Engine/Source/Runtime/Landscape/Private/LandscapeGizmoActor.cpp b/Engine/Source/Runtime/Landscape/Private/LandscapeGizmoActor.cpp
--- a/Engine/Source/Runtime/Landscape/Private/LandscapeGizmoActor.cpp
+++ b/Engine/Source/Runtime/Landscape/Private/LandscapeGizmoActor.cpp
@@ -713,7 +713,7 @@ float ALandscapeGizmoActiveActor::GetNormalizedHeight(uint16 LandscapeHeight) co
 {
 	if (TargetLandscapeInfo)
 	{
-		ALandscapeProxy* Proxy = TargetLandscapeInfo->GetCurrentLevelLandscapeProxy(true);
+		ALandscapeProxy* Proxy = TargetLandscapeInfo->GetLandscapeProxy();
 		if (Proxy)
 		{
 			// Need to make it scale...?
@@ -732,7 +732,7 @@ float ALandscapeGizmoActiveActor::GetWorldHeight(float NormalizedHeight) const
 {
 	if (TargetLandscapeInfo)
 	{
-		ALandscapeProxy* Proxy = TargetLandscapeInfo->GetCurrentLevelLandscapeProxy(true);
+		ALandscapeProxy* Proxy = TargetLandscapeInfo->GetLandscapeProxy();
 		if (Proxy)
 		{
 			float ZScale = GetLength();
