Mac compile issues

diff --git a/Engine/Source/Runtime/Landscape/Private/LandscapeGrass.cpp b/Engine/Source/Runtime/Landscape/Private/LandscapeGrass.cpp
--- a/Engine/Source/Runtime/Landscape/Private/LandscapeGrass.cpp
+++ b/Engine/Source/Runtime/Landscape/Private/LandscapeGrass.cpp
@@ -1,6 +1,7 @@
 // Copyright 1998-2015 Epic Games, Inc. All Rights Reserved.
 
 #include "Landscape.h"
+#include "EngineUtils.h"
 #include "MaterialCompiler.h"
 #include "Materials/MaterialExpressionLandscapeGrassOutput.h"
 #include "EdGraph/EdGraphNode.h"
