UE4 - fixed crash caused by deleting a landscape with grass, then flushing the grass cache

diff --git a/Engine/Source/Runtime/Landscape/Private/Landscape.cpp b/Engine/Source/Runtime/Landscape/Private/Landscape.cpp
--- a/Engine/Source/Runtime/Landscape/Private/Landscape.cpp
+++ b/Engine/Source/Runtime/Landscape/Private/Landscape.cpp
@@ -3021,7 +3021,7 @@ static void FlushFoliage(const TArray<FString>& Args)
 	for( TObjectIterator<ALandscapeProxy> It; It; ++It )
 	{
 		ALandscapeProxy* Landscape = *It;
-		if( Landscape && !Landscape->IsTemplate())
+		if( Landscape && !Landscape->IsTemplate() && !Landscape->IsPendingKill())
 		{
 			Landscape->FlushFoliageComponents();
 		}
