UI tweak to expose a property

diff --git a/Engine/Source/Runtime/Landscape/Classes/LandscapeProxy.h b/Engine/Source/Runtime/Landscape/Classes/LandscapeProxy.h
--- a/Engine/Source/Runtime/Landscape/Classes/LandscapeProxy.h
+++ b/Engine/Source/Runtime/Landscape/Classes/LandscapeProxy.h
@@ -294,7 +294,7 @@ public:
 	}
 };
 
-UCLASS(NotPlaceable, hidecategories=(Display, Attachment, Physics, Debug, Lighting, LOD), showcategories=(Rendering, "Utilities|Transformation"), MinimalAPI)
+UCLASS(NotPlaceable, hidecategories=(Display, Attachment, Physics, Debug, Lighting, LOD), showcategories=(Lighting, Rendering, "Utilities|Transformation"), MinimalAPI)
 class ALandscapeProxy : public AActor, public FTickableGameObject
 {
 	GENERATED_UCLASS_BODY()
@@ -388,7 +388,7 @@ public:
 	uint32 bCastStaticShadow:1;
 
 	/** Whether this primitive should cast dynamic shadows as if it were a two sided material. */
-	UPROPERTY(EditAnywhere, AdvancedDisplay, Category=Lighting)
+	UPROPERTY(EditAnywhere, AdvancedDisplay, Category=Lighting, meta=(DisplayName = "Shadow Two Sided"))
 	uint32 bCastShadowAsTwoSided:1;
 
 	UPROPERTY()
