Make VisibilityLayer in editor-only #code_review: jack.porter

diff --git a/Engine/Source/Runtime/Landscape/Classes/LandscapeProxy.h b/Engine/Source/Runtime/Landscape/Classes/LandscapeProxy.h
--- a/Engine/Source/Runtime/Landscape/Classes/LandscapeProxy.h
+++ b/Engine/Source/Runtime/Landscape/Classes/LandscapeProxy.h
@@ -453,7 +453,10 @@ public:
 #endif
 
 public:
+
+#if WITH_EDITOR
 	LANDSCAPE_API static ULandscapeLayerInfoObject* VisibilityLayer;
+#endif
 
 	/** Map of material instance constants used to for the components. Key is generated with ULandscapeComponent::GetLayerAllocationKey() */
 	TMap<FString, UMaterialInstanceConstant*> MaterialInstanceConstantMap;
diff --git a/Engine/Source/Runtime/Landscape/Private/Landscape.cpp b/Engine/Source/Runtime/Landscape/Private/Landscape.cpp
--- a/Engine/Source/Runtime/Landscape/Private/Landscape.cpp
+++ b/Engine/Source/Runtime/Landscape/Private/Landscape.cpp
@@ -536,6 +536,7 @@ ALandscapeProxy::ALandscapeProxy(const FObjectInitializer& ObjectInitializer)
 	MaxPaintedLayersPerComponent = 0;
 #endif
 
+#if WITH_EDITOR
 	if (VisibilityLayer == NULL)
 	{
 		// Structure to hold one-time initialization
@@ -557,6 +558,7 @@ ALandscapeProxy::ALandscapeProxy(const FObjectInitializer& ObjectInitializer)
 #endif
 		VisibilityLayer->AddToRoot();
 	}
+#endif
 }
 
 ALandscape* ALandscapeProxy::GetLandscapeActor()
