UE4 - grass tweaks and bug fixes

diff --git a/Engine/Source/Runtime/Engine/Private/HierarchicalInstancedStaticMesh.cpp b/Engine/Source/Runtime/Engine/Private/HierarchicalInstancedStaticMesh.cpp
--- a/Engine/Source/Runtime/Engine/Private/HierarchicalInstancedStaticMesh.cpp
+++ b/Engine/Source/Runtime/Engine/Private/HierarchicalInstancedStaticMesh.cpp
@@ -53,7 +53,7 @@ static TAutoConsoleVariable<float> CVarFoliageGuardband(
 
 TAutoConsoleVariable<float> CVarFoliageMinimumScreenSize(
 	TEXT("foliage.MinimumScreenSize"),
-	0.0f,
+	0.000005f,
 	TEXT("This controls the screen size at which we cull foliage instances entirely."));
 
 TAutoConsoleVariable<float> CVarFoliageLODDistanceScale(
diff --git a/Engine/Source/Runtime/Landscape/Private/Landscape.cpp b/Engine/Source/Runtime/Landscape/Private/Landscape.cpp
--- a/Engine/Source/Runtime/Landscape/Private/Landscape.cpp
+++ b/Engine/Source/Runtime/Landscape/Private/Landscape.cpp
@@ -393,6 +393,11 @@ void ULandscapeInfo::UpdateDebugColorMaterial()
 	//GWarn->EndSlowTask();
 }
 
+static TAutoConsoleVariable<int32> CVarPreloadGrassmapsInEditor(
+	TEXT("grass.PreloadGrassmapsInEditor"),
+	0,
+	TEXT("1: Load grass maps in editor at load time uses lots of memory; 0: Load grass maps on demand; hitchy"));
+
 void ULandscapeComponent::PostLoad()
 {
 	Super::PostLoad();
@@ -480,8 +485,10 @@ void ULandscapeComponent::PostLoad()
 		}
 	}
 #endif
-
-	GenerateGrassMap();
+	if (CVarPreloadGrassmapsInEditor.GetValueOnGameThread() > 0)
+	{
+		GenerateGrassMap();
+	}
 }
 
 #endif // WITH_EDITOR
@@ -2230,6 +2237,11 @@ static TAutoConsoleVariable<int32> CVarMaxInstancesPerComponent(
 	65536,
 	TEXT("Used to control the number of hierarchical components created. More can be more efficient, but can be hitchy as new components come into range"));
 
+static TAutoConsoleVariable<int32> CVarMaxAsyncTasks(
+	TEXT("grass.MaxAsyncTasks"),
+	4,
+	TEXT("Used to control the number of hierarchical components created at a time."));
+
 static TAutoConsoleVariable<int32> CVarGrassEnable(
 	TEXT("grass.Enable"),
 	1,
@@ -2657,7 +2669,7 @@ struct FAsyncGrassBuilder : public FGrassBuilderBase
 
 void ALandscapeProxy::Tick(float DeltaSeconds)
 {
-	if (!IsPendingKillPending())
+	if (!IsPendingKillPending() && !HasAnyFlags(RF_NeedLoad | RF_NeedPostLoad | RF_NeedPostLoadSubobjects | RF_Unreachable | RF_PendingKill | RF_ClassDefaultObject | RF_AsyncLoading))
 	{
 		// this is NOT an actor tick, it is a FTickableGameObject tick
 		// the super tick is for an actor tick...
@@ -2718,187 +2730,192 @@ void ALandscapeProxy::UpdateFoliage(const TArray<FVector>& Cameras, ULandscapeCo
 		bool bCullSubsections = CVarCullSubsections.GetValueOnAnyThread() > 0;
 		bool bDisableGPUCull = CVarDisableGPUCull.GetValueOnAnyThread() > 0;
 		int32 MaxInstancesPerComponent = FMath::Max<int32>(1024, CVarMaxInstancesPerComponent.GetValueOnAnyThread());
+		int32 MaxTasks = CVarMaxAsyncTasks.GetValueOnAnyThread();
 
-		for (int32 ComponentIndex = 0; ComponentIndex < LandscapeComponents.Num(); ComponentIndex++)
+		if (GetWorld())
 		{
-			ULandscapeComponent* Component = LandscapeComponents[ComponentIndex];
-			if (OnlyComponent && OnlyComponent != Component)
-			{
-				continue;
-			}
-
-			FBoxSphereBounds WorldBounds = Component->CalcBounds(Component->ComponentToWorld);
-			float MinDistanceToComp = Cameras.Num() ? MAX_flt : 0.0f;
 
-			for (auto& Pos : Cameras)
+			for (int32 ComponentIndex = 0; ComponentIndex < LandscapeComponents.Num(); ComponentIndex++)
 			{
-				MinDistanceToComp = FMath::Min<float>(MinDistanceToComp, WorldBounds.ComputeSquaredDistanceFromBoxToPoint(Pos));
-			}
+				ULandscapeComponent* Component = LandscapeComponents[ComponentIndex];
+				if (OnlyComponent && OnlyComponent != Component)
+				{
+					continue;
+				}
 
-			MinDistanceToComp = FMath::Sqrt(MinDistanceToComp);
+				FBoxSphereBounds WorldBounds = Component->CalcBounds(Component->ComponentToWorld);
+				float MinDistanceToComp = Cameras.Num() ? MAX_flt : 0.0f;
 
-			for (FWeightmapLayerAllocationInfo& LayerAlloc : Component->WeightmapLayerAllocations)
-			{
-				ULandscapeLayerInfoObject* LayerInfo = LayerAlloc.LayerInfo;
+				for (auto& Pos : Cameras)
+				{
+					MinDistanceToComp = FMath::Min<float>(MinDistanceToComp, WorldBounds.ComputeSquaredDistanceFromBoxToPoint(Pos));
+				}
+
+				MinDistanceToComp = FMath::Sqrt(MinDistanceToComp);
 
-				if (LayerInfo && LayerInfo->GrassMesh && LayerInfo->GrassDensity > 0.0f && LayerInfo->EndCullDistance > 0)
+				for (FWeightmapLayerAllocationInfo& LayerAlloc : Component->WeightmapLayerAllocations)
 				{
-					bool bRandomRotation = LayerInfo->RandomRotation;
-					bool bAlign = LayerInfo->AlignToSurface;
-					float MustHaveDistance = GuardBand * (float)LayerInfo->EndCullDistance;
-					float DiscardDistance = DiscardGuardBand * (float)LayerInfo->EndCullDistance;
+					ULandscapeLayerInfoObject* LayerInfo = LayerAlloc.LayerInfo;
 
-					if (MinDistanceToComp > DiscardDistance)
+					if (LayerInfo && LayerInfo->GrassMesh && LayerInfo->GrassDensity > 0.0f && LayerInfo->EndCullDistance > 0)
 					{
-						continue;
-					}
+						bool bRandomRotation = LayerInfo->RandomRotation;
+						bool bAlign = LayerInfo->AlignToSurface;
+						float MustHaveDistance = GuardBand * (float)LayerInfo->EndCullDistance;
+						float DiscardDistance = DiscardGuardBand * (float)LayerInfo->EndCullDistance;
 
-					FGrassBuilderBase ForSubsectionMath(this, Component, LayerInfo);
+						if (MinDistanceToComp > DiscardDistance)
+						{
+							continue;
+						}
 
-					int32 SqrtSubsections = 1;
-				
-					if (ForSubsectionMath.bHaveValidData && ForSubsectionMath.SqrtMaxInstances > 0)
-					{
-						SqrtSubsections = FMath::Clamp<int32>(FMath::CeilToInt(float(ForSubsectionMath.SqrtMaxInstances) / FMath::Sqrt((float)MaxInstancesPerComponent)), 1, 16);
-					}
+						FGrassBuilderBase ForSubsectionMath(this, Component, LayerInfo);
 
-					int32 NumCompsCreated = 0;
-					FBox LocalBox = Component->CachedLocalBox;
-					FVector LocalExtentDiv = (LocalBox.Max - LocalBox.Min) * FVector(1.0f / float(SqrtSubsections), 1.0f / float(SqrtSubsections), 1.0f);
-					for (int32 SubX = 0; SubX < SqrtSubsections; SubX++)
-					{
-						for (int32 SubY = 0; SubY < SqrtSubsections; SubY++)
+						int32 SqrtSubsections = 1;
+				
+						if (ForSubsectionMath.bHaveValidData && ForSubsectionMath.SqrtMaxInstances > 0)
 						{
-							float MinDistanceToSubComp = MinDistanceToComp;
+							SqrtSubsections = FMath::Clamp<int32>(FMath::CeilToInt(float(ForSubsectionMath.SqrtMaxInstances) / FMath::Sqrt((float)MaxInstancesPerComponent)), 1, 16);
+						}
 
-							if (bCullSubsections && SqrtSubsections > 1)
+						int32 NumCompsCreated = 0;
+						FBox LocalBox = Component->CachedLocalBox;
+						FVector LocalExtentDiv = (LocalBox.Max - LocalBox.Min) * FVector(1.0f / float(SqrtSubsections), 1.0f / float(SqrtSubsections), 1.0f);
+						for (int32 SubX = 0; SubX < SqrtSubsections; SubX++)
+						{
+							for (int32 SubY = 0; SubY < SqrtSubsections; SubY++)
 							{
-								FVector BoxMin;
-								BoxMin.X = LocalBox.Min.X + LocalExtentDiv.X * float(SubX);
-								BoxMin.Y = LocalBox.Min.Y + LocalExtentDiv.Y * float(SubY);
-								BoxMin.Z = LocalBox.Min.Z;
+								float MinDistanceToSubComp = MinDistanceToComp;
 
-								FVector BoxMax;
-								BoxMax.X = LocalBox.Min.X + LocalExtentDiv.X * float(SubX + 1);
-								BoxMax.Y = LocalBox.Min.Y + LocalExtentDiv.Y * float(SubY + 1);
-								BoxMax.Z = LocalBox.Max.Z;
+								if (bCullSubsections && SqrtSubsections > 1)
+								{
+									FVector BoxMin;
+									BoxMin.X = LocalBox.Min.X + LocalExtentDiv.X * float(SubX);
+									BoxMin.Y = LocalBox.Min.Y + LocalExtentDiv.Y * float(SubY);
+									BoxMin.Z = LocalBox.Min.Z;
 
-								FBox LocalSubBox(BoxMin, BoxMax);
-								FBox WorldSubBox = LocalSubBox.TransformBy(Component->ComponentToWorld);
+									FVector BoxMax;
+									BoxMax.X = LocalBox.Min.X + LocalExtentDiv.X * float(SubX + 1);
+									BoxMax.Y = LocalBox.Min.Y + LocalExtentDiv.Y * float(SubY + 1);
+									BoxMax.Z = LocalBox.Max.Z;
 
-								MinDistanceToSubComp = Cameras.Num() ? MAX_flt : 0.0f;
-								for (auto& Pos : Cameras)
-								{
-									MinDistanceToSubComp = FMath::Min<float>(MinDistanceToSubComp, ComputeSquaredDistanceFromBoxToPoint(WorldSubBox.Min, WorldSubBox.Max, Pos));
+									FBox LocalSubBox(BoxMin, BoxMax);
+									FBox WorldSubBox = LocalSubBox.TransformBy(Component->ComponentToWorld);
+
+									MinDistanceToSubComp = Cameras.Num() ? MAX_flt : 0.0f;
+									for (auto& Pos : Cameras)
+									{
+										MinDistanceToSubComp = FMath::Min<float>(MinDistanceToSubComp, ComputeSquaredDistanceFromBoxToPoint(WorldSubBox.Min, WorldSubBox.Max, Pos));
+									}
+									MinDistanceToSubComp = FMath::Sqrt(MinDistanceToSubComp);
 								}
-								MinDistanceToSubComp = FMath::Sqrt(MinDistanceToSubComp);
-							}
 
-							if (MinDistanceToSubComp > DiscardDistance)
-							{
-								continue;
-							}
+								if (MinDistanceToSubComp > DiscardDistance)
+								{
+									continue;
+								}
 
-							FCachedLandscapeFoliage::FGrassComp NewComp;
-							NewComp.Key.BasedOn = Component;
-							NewComp.Key.Layer = LayerInfo;
-							NewComp.Key.SqrtSubsections = SqrtSubsections;
-							NewComp.Key.CachedMaxInstancesPerComponent = MaxInstancesPerComponent;
-							NewComp.Key.SubsectionX = SubX;
-							NewComp.Key.SubsectionY = SubY;
+								FCachedLandscapeFoliage::FGrassComp NewComp;
+								NewComp.Key.BasedOn = Component;
+								NewComp.Key.Layer = LayerInfo;
+								NewComp.Key.SqrtSubsections = SqrtSubsections;
+								NewComp.Key.CachedMaxInstancesPerComponent = MaxInstancesPerComponent;
+								NewComp.Key.SubsectionX = SubX;
+								NewComp.Key.SubsectionY = SubY;
 
-							{
-								FCachedLandscapeFoliage::FGrassComp* Existing = FoliageCache.CachedGrassComps.Find(NewComp.Key);
-								if (Existing || MinDistanceToSubComp > MustHaveDistance)
 								{
-									if (Existing)
+									FCachedLandscapeFoliage::FGrassComp* Existing = FoliageCache.CachedGrassComps.Find(NewComp.Key);
+									if (Existing || MinDistanceToSubComp > MustHaveDistance)
 									{
-										Existing->Touch();
+										if (Existing)
+										{
+											Existing->Touch();
+										}
+										continue;
 									}
-									continue;
 								}
-							}
 
-							if (NumCompsCreated)
-							{
-								continue; // one per frame, but we still want to touch the existing ones
-							}
-							NumCompsCreated++;
+								if (NumCompsCreated || (!bForceSync && AsyncFoliageTasks.Num() >= MaxTasks))
+								{
+									continue; // one per frame, but we still want to touch the existing ones
+								}
+								NumCompsCreated++;
 
-							SCOPE_CYCLE_COUNTER(STAT_FoliageGrassStartComp);
-							int32 FolSeed = FCrc::StrCrc32((LayerInfo->LayerName.ToString() + Component->GetName() + FString::Printf(TEXT("%d %d"), SubX, SubY)).GetCharArray().GetData());
-							if (FolSeed == 0)
-							{
-								FolSeed++;
-							}
+								SCOPE_CYCLE_COUNTER(STAT_FoliageGrassStartComp);
+								int32 FolSeed = FCrc::StrCrc32((LayerInfo->LayerName.ToString() + Component->GetName() + FString::Printf(TEXT("%d %d"), SubX, SubY)).GetCharArray().GetData());
+								if (FolSeed == 0)
+								{
+									FolSeed++;
+								}
 
-							UHierarchicalInstancedStaticMeshComponent* HierarchicalInstancedStaticMeshComponent;
-							{
-								QUICK_SCOPE_CYCLE_COUNTER(STAT_GrassCreateComp);
-								HierarchicalInstancedStaticMeshComponent = ConstructObject<UHierarchicalInstancedStaticMeshComponent>(UHierarchicalInstancedStaticMeshComponent::StaticClass(), this); //, NAME_None, RF_Transactional);
-							}
-							NewComp.Foliage = HierarchicalInstancedStaticMeshComponent;
-							FoliageCache.CachedGrassComps.Add(NewComp);
+								UHierarchicalInstancedStaticMeshComponent* HierarchicalInstancedStaticMeshComponent;
+								{
+									QUICK_SCOPE_CYCLE_COUNTER(STAT_GrassCreateComp);
+									HierarchicalInstancedStaticMeshComponent = ConstructObject<UHierarchicalInstancedStaticMeshComponent>(UHierarchicalInstancedStaticMeshComponent::StaticClass(), this); //, NAME_None, RF_Transactional);
+								}
+								NewComp.Foliage = HierarchicalInstancedStaticMeshComponent;
+								FoliageCache.CachedGrassComps.Add(NewComp);
 
-							HierarchicalInstancedStaticMeshComponent->Mobility = EComponentMobility::Static;
+								HierarchicalInstancedStaticMeshComponent->Mobility = EComponentMobility::Static;
 
-							HierarchicalInstancedStaticMeshComponent->StaticMesh = LayerInfo->GrassMesh;
-							HierarchicalInstancedStaticMeshComponent->bSelectable = false;
-							HierarchicalInstancedStaticMeshComponent->bHasPerInstanceHitProxies = true;
-							static FName NoCollision(TEXT("NoCollision"));
-							HierarchicalInstancedStaticMeshComponent->SetCollisionProfileName(NoCollision);
-							HierarchicalInstancedStaticMeshComponent->bDisableCollision = true;
-							HierarchicalInstancedStaticMeshComponent->SetCanEverAffectNavigation(false);
-							HierarchicalInstancedStaticMeshComponent->InstancingRandomSeed = FolSeed;
+								HierarchicalInstancedStaticMeshComponent->StaticMesh = LayerInfo->GrassMesh;
+								HierarchicalInstancedStaticMeshComponent->bSelectable = false;
+								HierarchicalInstancedStaticMeshComponent->bHasPerInstanceHitProxies = true;
+								static FName NoCollision(TEXT("NoCollision"));
+								HierarchicalInstancedStaticMeshComponent->SetCollisionProfileName(NoCollision);
+								HierarchicalInstancedStaticMeshComponent->bDisableCollision = true;
+								HierarchicalInstancedStaticMeshComponent->SetCanEverAffectNavigation(false);
+								HierarchicalInstancedStaticMeshComponent->InstancingRandomSeed = FolSeed;
 
-							if (!Cameras.Num() || bDisableGPUCull)
-							{
-								// if we don't have any cameras, then we are rendering landscape LOD materials or somesuch and we want to disable culling
-								HierarchicalInstancedStaticMeshComponent->InstanceStartCullDistance = 0;
-								HierarchicalInstancedStaticMeshComponent->InstanceEndCullDistance = 0;
-							}
-							else
-							{
-								HierarchicalInstancedStaticMeshComponent->InstanceStartCullDistance = LayerInfo->StartCullDistance;
-								HierarchicalInstancedStaticMeshComponent->InstanceEndCullDistance = LayerInfo->EndCullDistance;
-							}
+								if (!Cameras.Num() || bDisableGPUCull)
+								{
+									// if we don't have any cameras, then we are rendering landscape LOD materials or somesuch and we want to disable culling
+									HierarchicalInstancedStaticMeshComponent->InstanceStartCullDistance = 0;
+									HierarchicalInstancedStaticMeshComponent->InstanceEndCullDistance = 0;
+								}
+								else
+								{
+									HierarchicalInstancedStaticMeshComponent->InstanceStartCullDistance = LayerInfo->StartCullDistance;
+									HierarchicalInstancedStaticMeshComponent->InstanceEndCullDistance = LayerInfo->EndCullDistance;
+								}
 
-							//@todo - take the settings from a UFoliageType object.  For now, disable distance field lighting on grass so we don't hitch.
-							HierarchicalInstancedStaticMeshComponent->bAffectDistanceFieldLighting = false;
+								//@todo - take the settings from a UFoliageType object.  For now, disable distance field lighting on grass so we don't hitch.
+								HierarchicalInstancedStaticMeshComponent->bAffectDistanceFieldLighting = false;
 
-							{
-								QUICK_SCOPE_CYCLE_COUNTER(STAT_GrassAttachComp);
+								{
+									QUICK_SCOPE_CYCLE_COUNTER(STAT_GrassAttachComp);
 
-								HierarchicalInstancedStaticMeshComponent->AttachTo(GetRootComponent());
-								FTransform DesiredTransform = GetRootComponent()->ComponentToWorld;
-								DesiredTransform.RemoveScaling();
-								HierarchicalInstancedStaticMeshComponent->SetWorldTransform(DesiredTransform);
+									HierarchicalInstancedStaticMeshComponent->AttachTo(GetRootComponent());
+									FTransform DesiredTransform = GetRootComponent()->ComponentToWorld;
+									DesiredTransform.RemoveScaling();
+									HierarchicalInstancedStaticMeshComponent->SetWorldTransform(DesiredTransform);
 
-								FoliageComponents.Add(HierarchicalInstancedStaticMeshComponent);
-							}
+									FoliageComponents.Add(HierarchicalInstancedStaticMeshComponent);
+								}
 
-							FAsyncGrassBuilder* Builder;
+								FAsyncGrassBuilder* Builder;
 
-							{
-								QUICK_SCOPE_CYCLE_COUNTER(STAT_GrassCreateBuilder);
-								Builder = new FAsyncGrassBuilder(this, Component, LayerInfo, HierarchicalInstancedStaticMeshComponent, SqrtSubsections, SubX, SubY);
-							}
+								{
+									QUICK_SCOPE_CYCLE_COUNTER(STAT_GrassCreateBuilder);
+									Builder = new FAsyncGrassBuilder(this, Component, LayerInfo, HierarchicalInstancedStaticMeshComponent, SqrtSubsections, SubX, SubY);
+								}
 
-							if (Builder->bHaveValidData)
-							{
-								FAsyncTask<FAsyncGrassTask>* Task = new FAsyncTask<FAsyncGrassTask>(Builder, NewComp.Key, HierarchicalInstancedStaticMeshComponent);
+								if (Builder->bHaveValidData)
+								{
+									FAsyncTask<FAsyncGrassTask>* Task = new FAsyncTask<FAsyncGrassTask>(Builder, NewComp.Key, HierarchicalInstancedStaticMeshComponent);
 
-								Task->StartBackgroundTask();
+									Task->StartBackgroundTask();
 
-								AsyncFoliageTasks.Add(Task);
-							}
-							else
-							{
-								delete Builder;
-							}
-							{
-								QUICK_SCOPE_CYCLE_COUNTER(STAT_GrassRegisterComp);
-								HierarchicalInstancedStaticMeshComponent->RegisterComponent();
+									AsyncFoliageTasks.Add(Task);
+								}
+								else
+								{
+									delete Builder;
+								}
+								{
+									QUICK_SCOPE_CYCLE_COUNTER(STAT_GrassRegisterComp);
+									HierarchicalInstancedStaticMeshComponent->RegisterComponent();
+								}
 							}
 						}
 					}
@@ -2977,7 +2994,7 @@ void ALandscapeProxy::UpdateFoliage(const TArray<FVector>& Cameras, ULandscapeCo
 					{
 						Exchange(HierarchicalInstancedStaticMeshComponent->WriteOncePrebuiltInstanceBuffer, Inner.Builder->InstanceBuffer);
 						HierarchicalInstancedStaticMeshComponent->AcceptPrebuiltTree(Inner.Builder->ClusterTree);
-						if (bForceSync)
+						if (bForceSync && GetWorld())
 						{
 							HierarchicalInstancedStaticMeshComponent->RecreateRenderState_Concurrent();
 						}
diff --git a/Engine/Source/Runtime/Landscape/Private/LandscapeDataAccess.cpp b/Engine/Source/Runtime/Landscape/Private/LandscapeDataAccess.cpp
--- a/Engine/Source/Runtime/Landscape/Private/LandscapeDataAccess.cpp
+++ b/Engine/Source/Runtime/Landscape/Private/LandscapeDataAccess.cpp
@@ -68,8 +68,13 @@ LANDSCAPE_API void FLandscapeComponentDataInterface::UnlockRawHeightData() const
 	DataInterface->UnlockMip(Component->HeightmapTexture, MipLevel);
 }
 
-LANDSCAPE_API void FLandscapeComponentDataInterface::GetHeightmapTextureData(TArray<FColor>& OutData)
+LANDSCAPE_API void FLandscapeComponentDataInterface::GetHeightmapTextureData(TArray<FColor>& OutData, bool bOkToFail)
 {
+	if (bOkToFail && !HeightMipData)
+	{
+		OutData.Empty();
+		return;
+	}
 #if LANDSCAPE_VALIDATE_DATA_ACCESS
 	check(HeightMipData);
 #endif
diff --git a/Engine/Source/Runtime/Landscape/Private/LandscapeEdit.cpp b/Engine/Source/Runtime/Landscape/Private/LandscapeEdit.cpp
--- a/Engine/Source/Runtime/Landscape/Private/LandscapeEdit.cpp
+++ b/Engine/Source/Runtime/Landscape/Private/LandscapeEdit.cpp
@@ -5226,7 +5226,7 @@ void ULandscapeComponent::GenerateGrassMap()
 			GrassMap = MakeShareable(new FGrassMap(NumChannels));
 			TArray<FColor> LocalHeightCache;
 			FLandscapeComponentDataInterface CDI(this);
-			CDI.GetHeightmapTextureData(LocalHeightCache);
+			CDI.GetHeightmapTextureData(LocalHeightCache, true);
 			if (LocalHeightCache.Num())
 			{
 				check(LocalHeightCache.Num() == Stride * Stride);
diff --git a/Engine/Source/Runtime/Landscape/Public/LandscapeDataAccess.h b/Engine/Source/Runtime/Landscape/Public/LandscapeDataAccess.h
--- a/Engine/Source/Runtime/Landscape/Public/LandscapeDataAccess.h
+++ b/Engine/Source/Runtime/Landscape/Public/LandscapeDataAccess.h
@@ -222,7 +222,7 @@ struct FLandscapeComponentDataInterface
 	LANDSCAPE_API void UnlockRawHeightData() const;
 
 	/* Return the raw heightmap data exactly same size for Heightmap texture which belong to only this component */
-	LANDSCAPE_API void GetHeightmapTextureData(TArray<FColor>& OutData);
+	LANDSCAPE_API void GetHeightmapTextureData(TArray<FColor>& OutData, bool bOkToFail = false);
 
 	LANDSCAPE_API bool GetWeightmapTextureData(ULandscapeLayerInfoObject* LayerInfo, TArray<uint8>& OutData);
 
