Fix rendering of the landscape gizmo being offset after scaling it.
MeshRT is very wierdly computed, and the way vertices are scaled before being multiplied by it suggests things could be simplified...

diff --git a/Engine/Source/Runtime/Landscape/Private/LandscapeGizmoActor.cpp b/Engine/Source/Runtime/Landscape/Private/LandscapeGizmoActor.cpp
--- a/Engine/Source/Runtime/Landscape/Private/LandscapeGizmoActor.cpp
+++ b/Engine/Source/Runtime/Landscape/Private/LandscapeGizmoActor.cpp
@@ -171,7 +171,7 @@ public:
 				const FVector BaseLocation = WToL.TransformPosition(Gizmo->GetActorLocation());
 				const float ScaleXY = LandscapeInfo->DrawScale.X;
 
-				MeshRT = FTranslationMatrix(FVector(-W / ScaleXY + 0.5, -H / ScaleXY + 0.5, 0)) * FRotationTranslationMatrix(FRotator(0, Gizmo->GetActorRotation().Yaw, 0), FVector(BaseLocation.X, BaseLocation.Y, 0)) * LToW.ToMatrixWithScale();
+				MeshRT = FTranslationMatrix(FVector(-W / ScaleXY + 0.5, -H / ScaleXY + 0.5, 0) * GizmoScale3D) * FRotationTranslationMatrix(FRotator(0, Gizmo->GetActorRotation().Yaw, 0), FVector(BaseLocation.X, BaseLocation.Y, 0)) * LToW.ToMatrixWithScale();
 				HeightmapRenderProxy = new FLandscapeGizmoMeshRenderProxy( Gizmo->GizmoMeshMaterial->GetRenderProxy(false), BaseLocation.Z + L, BaseLocation.Z, Gizmo->GizmoTexture, FLinearColor(Gizmo->TextureScale.X, Gizmo->TextureScale.Y, 0, 0), WToL );
 
 				GizmoRenderProxy = (Gizmo->DataType != LGT_None) ? Gizmo->GizmoDataMaterial->GetRenderProxy(false) : Gizmo->GizmoMaterial->GetRenderProxy(false);
