Grass generation using material

diff --git a/Engine/Shaders/LandscapeGrassWeight.usf b/Engine/Shaders/LandscapeGrassWeight.usf
new file mode 100644
--- /dev/null
+++ b/Engine/Shaders/LandscapeGrassWeight.usf
@@ -0,0 +1,77 @@
+// Copyright 1998-2015 Epic Games, Inc. All Rights Reserved.
+
+/*=============================================================================
+	LandscapeGrassWeightVertexShader.usf: Vertex shader to dump the weight map for Landscape Grass
+=============================================================================*/
+
+#include "Common.usf"
+#include "Material.usf"
+#include "VertexFactory.usf"
+
+struct FLandscapeGrassWeightInterpolantsVSToPS
+{
+	FVertexFactoryInterpolantsVSToPS FactoryInterpolants;
+	float4 PixelPosition	: TEXCOORD8; // xyz = world position, w = clip z
+	float2 Height			: TEXCOORD9;
+	float4 Position : SV_POSITION;
+};
+
+#if VERTEXSHADER
+
+/** Vertex Shader */
+void VSMain(
+	FVertexFactoryInput Input,
+	out FLandscapeGrassWeightInterpolantsVSToPS Output
+	)
+{
+	FVertexFactoryIntermediates VFIntermediates = GetVertexFactoryIntermediates(Input);
+	float4 WorldPositionExcludingWPO = VertexFactoryGetWorldPosition(Input, VFIntermediates);
+	float4 WorldPosition = WorldPositionExcludingWPO;
+
+	half3x3 TangentToLocal = VertexFactoryGetTangentToLocal(Input, VFIntermediates);	
+	FMaterialVertexParameters VertexParameters = GetMaterialVertexParameters(Input, VFIntermediates, WorldPosition.xyz, TangentToLocal);
+
+	half3 WorldPositionOffset = GetMaterialWorldPositionOffset(VertexParameters);
+	
+	WorldPosition.xyz += WorldPositionOffset;
+	float4 RasterizedWorldPosition = VertexFactoryGetRasterizedWorldPosition(Input, VFIntermediates, WorldPosition);
+	Output.PixelPosition = WorldPosition;
+	Output.Position = mul(RasterizedWorldPosition, View.TranslatedWorldToClip);
+	Output.FactoryInterpolants = VertexFactoryGetInterpolantsVSToPS(Input, VFIntermediates, VertexParameters);
+
+	Output.Height = LandscapeVertexFactorySampleHeight(Input, VFIntermediates);
+
+	Output.PixelPosition.w = Output.Position.w;
+}
+
+#elif PIXELSHADER
+
+// Pixel Shader
+void PSMain( 
+	FVertexFactoryInterpolantsVSToPS Interpolants,
+	FLandscapeGrassWeightInterpolantsVSToPS GrassWeightInterpolants,
+	in float4 SvPosition : SV_Position,
+	OPTIONAL_IsFrontFace,
+	out half4 OutColor	: SV_Target0
+	)
+{
+	FMaterialPixelParameters MaterialParameters = GetMaterialPixelParameters(Interpolants, GrassWeightInterpolants.PixelPosition);
+
+	CalcMaterialParameters(MaterialParameters, SvPosition, bIsFrontFace, GrassWeightInterpolants.PixelPosition); 
+	OutColor.xy = GrassWeightInterpolants.Height;
+
+	#if NUM_MATERIAL_OUTPUTS_GETGRASSWEIGHT > 0
+		OutColor.z = GetGrassWeight0(MaterialParameters);
+		#if NUM_MATERIAL_OUTPUTS_GETGRASSWEIGHT > 1
+			OutColor.w = GetGrassWeight1(MaterialParameters);
+		#else
+			OutColor.w = 0;
+		#endif
+	#else
+		OutColor.z = 0;
+		OutColor.w = 0;
+	#endif
+
+}
+
+#endif
\ No newline at end of file
diff --git a/Engine/Shaders/LandscapeVertexFactory.usf b/Engine/Shaders/LandscapeVertexFactory.usf
--- a/Engine/Shaders/LandscapeVertexFactory.usf
+++ b/Engine/Shaders/LandscapeVertexFactory.usf
@@ -614,6 +614,34 @@ FVertexFactoryIntermediates GetVertexFactoryIntermediates(FVertexFactoryInput In
 	return Intermediates;
 }
 
+#if FEATURE_LEVEL >= FEATURE_LEVEL_SM4 && VERTEXSHADER
+float2 LandscapeVertexFactorySampleHeight(FVertexFactoryInput Input, FVertexFactoryIntermediates Intermediates)
+{
+	/*
+	FLandscapeTexCoords TexCoords = GetLandscapeTexCoords(Input, Intermediates);
+
+	float4 SampleValue = Texture2DSampleLevel(HeightmapTexture, HeightmapTextureSampler, TexCoords.HeightMapTexCoord, 0);
+	return SampleValue.xy;
+	*/
+
+	float2 xy = Intermediates.InputPosition.xy * LodValues.w;
+	float LODCalculated = CalcLOD(xy, Intermediates.InputPosition.zw);
+	float LodValue = floor(LODCalculated);
+
+	// Base to Actual LOD, Base to Next LOD
+	float2 CoordTranslate = float2( LandscapeParameters.SubsectionSizeVertsLayerUVPan.x * pow(2, -LodValue) - 1, max(LandscapeParameters.SubsectionSizeVertsLayerUVPan.x * pow(2, -LodValue-1), 2) - 1 ) * LandscapeParameters.SubsectionSizeVertsLayerUVPan.y;
+
+	// InputPositionLODAdjusted : Position for actual LOD in base LOD units
+	float2 ActualLODCoordsInt = floor(Intermediates.InputPosition.xy * pow(2, -(LodValue - LodValues.x)));
+
+	float2 InputPositionLODAdjusted = ActualLODCoordsInt / CoordTranslate.x;
+	float2 SampleCoords = InputPositionLODAdjusted * LandscapeParameters.HeightmapUVScaleBias.xy + LandscapeParameters.HeightmapUVScaleBias.zw + 0.5*LandscapeParameters.HeightmapUVScaleBias.xy + Intermediates.InputPosition.zw * LandscapeParameters.SubsectionOffsetParams.xy;
+	float4 SampleValue = Texture2DSampleLevel(HeightmapTexture, HeightmapTextureSampler, SampleCoords, LodValue-LodBias.z);
+	return SampleValue.rg;
+}
+#endif
+
+
 /**
 * Get the 3x3 tangent basis vectors for this vertex factory
 * this vertex factory will calculate the binormal on-the-fly
diff --git a/Engine/Source/Runtime/Core/Public/UObject/ObjectVersion.h b/Engine/Source/Runtime/Core/Public/UObject/ObjectVersion.h
--- a/Engine/Source/Runtime/Core/Public/UObject/ObjectVersion.h
+++ b/Engine/Source/Runtime/Core/Public/UObject/ObjectVersion.h
@@ -464,6 +464,8 @@ enum EUnrealEngineObjectUE4Version
 	VER_UE4_MERGED_ADD_MODIFIERS_RUNTIME_GENERATION_TO_4_7,
 	// After merging VER_UE4_ADD_MODIFIERS_RUNTIME_GENERATION into 4.7 branch
 	VER_UE4_AFTER_MERGING_ADD_MODIFIERS_RUNTIME_GENERATION_TO_4_7,
+	// Landscape grass weightmap data is now generated in the editor and serialized.
+	VER_UE4_SERIALIZE_LANDSCAPE_GRASS_DATA,
 
 	// -----<new versions can be added before this line>-------------------------------------------------
 	// - this needs to be the last line (see note below)
diff --git a/Engine/Source/Runtime/Landscape/Classes/LandscapeComponent.h b/Engine/Source/Runtime/Landscape/Classes/LandscapeComponent.h
--- a/Engine/Source/Runtime/Landscape/Classes/LandscapeComponent.h
+++ b/Engine/Source/Runtime/Landscape/Classes/LandscapeComponent.h
@@ -5,6 +5,7 @@
 #include "SceneTypes.h"
 #include "StaticLighting.h"
 #include "Components/PrimitiveComponent.h"
+#include "LandscapeGrassType.h"
 
 #include "LandscapeComponent.generated.h"
 
@@ -104,17 +105,23 @@ struct FWeightmapLayerAllocationInfo
 	FName GetLayerName() const;
 };
 
-struct FGrassMap
+struct FLandscapeComponentGrassData
 {
-	/** Total number of channels; at least 3 since we have at least the height (two bytes) and at least one layer channel **/
-	int32 ElementStride;
-	TArray<uint8> Data;
-	FGrassMap(int32 InElementStride)
-		: ElementStride(InElementStride)
+	TArray<uint16> HeightData;
+	TMap<ULandscapeGrassType*, TArray<uint8>> WeightData;
+
+	bool HasData()
 	{
+		return HeightData.Num() > 0;
+	}
+
+	SIZE_T GetAllocatedSize() const
+	{
+		return sizeof(*this) + HeightData.GetAllocatedSize() + WeightData.GetAllocatedSize();
 	}
-};
 
+	friend FArchive& operator<<(FArchive& Ar, FLandscapeComponentGrassData& Data);
+};
 
 UCLASS(hidecategories=(Display, Attachment, Physics, Debug, Collision, Movement, Rendering, PrimitiveComponent, Object, Transform), showcategories=("Rendering|Material"), MinimalAPI)
 class ULandscapeComponent : public UPrimitiveComponent
@@ -251,7 +258,7 @@ public:
 	FLandscapeComponentDerivedData PlatformData;
 
 	/** Grass data for generation **/
-	TSharedPtr<FGrassMap, ESPMode::ThreadSafe> GrassMap;
+	TSharedRef<FLandscapeComponentGrassData, ESPMode::ThreadSafe> GrassData;
 
 	virtual ~ULandscapeComponent();
 
@@ -315,8 +322,8 @@ public:
 	void GeneratePlatformVertexData();
 	UMaterialInstance* GeneratePlatformPixelData(TArray<class UTexture2D*>& WeightmapTextures, bool bIsCooking);
 
-	/** Creates and destroys grass data for cooked platforms, used temporarily when cooking */
-	void GenerateGrassMap();
+	/** Creates and destroys cooked grass data stored in the map */
+	void RenderGrassMap();
 	void RemoveGrassMap();
 
 #endif
diff --git a/Engine/Source/Runtime/Landscape/Classes/LandscapeLayerInfoObject.h b/Engine/Source/Runtime/Landscape/Classes/LandscapeLayerInfoObject.h
--- a/Engine/Source/Runtime/Landscape/Classes/LandscapeLayerInfoObject.h
+++ b/Engine/Source/Runtime/Landscape/Classes/LandscapeLayerInfoObject.h
@@ -18,33 +18,6 @@ class ULandscapeLayerInfoObject : public UObject
 	UPROPERTY(EditAnywhere, Category=LandscapeLayerInfoObject)
 	float Hardness;
 
-	UPROPERTY(EditAnywhere, Category=Grass)
-	class UStaticMesh* GrassMesh;
-
-	/* Instances per 10 square meters. */
-	UPROPERTY(EditAnywhere, Category=Grass)
-	float GrassDensity;
-
-	UPROPERTY(EditAnywhere, Category=Grass)
-	float PlacementJitter;
-
-	/* The distance where instances will begin to fade out if using a PerInstanceFadeAmount material node. 0 disables. */
-	UPROPERTY(EditAnywhere, Category=Grass)
-	int32 StartCullDistance;
-
-	/**
-	 * The distance where instances will have completely faded out when using a PerInstanceFadeAmount material node. 0 disables. 
-	 * When the entire cluster is beyond this distance, the cluster is completely culled and not rendered at all.
-	 */
-	UPROPERTY(EditAnywhere, Category = Grass)
-	int32 EndCullDistance;
-
-	UPROPERTY(EditAnywhere, Category = Grass)
-	bool RandomRotation;
-
-	UPROPERTY(EditAnywhere, Category = Grass)
-	bool AlignToSurface;
-
 #if WITH_EDITORONLY_DATA
 	UPROPERTY(VisibleAnywhere, Category=LandscapeLayerInfoObject)
 	uint32 bNoWeightBlend:1;
diff --git a/Engine/Source/Runtime/Landscape/Classes/LandscapeProxy.h b/Engine/Source/Runtime/Landscape/Classes/LandscapeProxy.h
--- a/Engine/Source/Runtime/Landscape/Classes/LandscapeProxy.h
+++ b/Engine/Source/Runtime/Landscape/Classes/LandscapeProxy.h
@@ -15,12 +15,12 @@
 #include "Components/HierarchicalInstancedStaticMeshComponent.h"
 #include "LandscapeComponent.h"
 #include "LandscapeLayerInfoObject.h"
+#include "LandscapeGrassType.h"
 #include "Tickable.h"
 
 #include "LandscapeProxy.generated.h"
 
 class ULandscapeMaterialInstanceConstant;
-class ULandscapeLayerInfoObject;
 class ULandscapeSplinesComponent;
 class ULandscapeHeightfieldCollisionComponent;
 class UMaterialInterface;
@@ -196,7 +196,7 @@ struct FCachedLandscapeFoliage
 	struct FGrassCompKey
 	{
 		TWeakObjectPtr<ULandscapeComponent> BasedOn;
-		TWeakObjectPtr<ULandscapeLayerInfoObject> Layer;
+		TWeakObjectPtr<ULandscapeGrassType> GrassType;
 		int32 SqrtSubsections;
 		int32 CachedMaxInstancesPerComponent;
 		int32 SubsectionX;
@@ -217,12 +217,12 @@ struct FCachedLandscapeFoliage
 				SubsectionX == Other.SubsectionX &&
 				SubsectionY == Other.SubsectionY &&
 				BasedOn == Other.BasedOn &&
-				Layer == Other.Layer;
+				GrassType == Other.GrassType;
 		}
 
 		friend uint32 GetTypeHash(const FGrassCompKey& Key)
 		{
-			return GetTypeHash(Key.BasedOn) ^ GetTypeHash(Key.Layer) ^ Key.SqrtSubsections ^ Key.CachedMaxInstancesPerComponent ^ (Key.SubsectionX >> 16) ^ (Key.SubsectionY >> 24);
+			return GetTypeHash(Key.BasedOn) ^ GetTypeHash(Key.GrassType) ^ Key.SqrtSubsections ^ Key.CachedMaxInstancesPerComponent ^ (Key.SubsectionX >> 16) ^ (Key.SubsectionY >> 24);
 		}
 
 	};
@@ -320,7 +320,7 @@ public:
 	bool bStaticSectionOffset;
 #endif
 
-	/** Max LOD level to use when rendering */
+	/** Max LOD level to use when rendering, -1 means the max available */
 	UPROPERTY(EditAnywhere, Category=LOD)
 	int32 MaxLODLevel;
 
@@ -499,6 +499,9 @@ public:
 	*/
 	LANDSCAPE_API void UpdateFoliage(const TArray<FVector>& Cameras, ULandscapeComponent* OnlyComponent = nullptr, bool bForceSync = false);
 
+	/* Get the list of grass types on this landscape */
+	TArray<const ULandscapeGrassType*> GetGrassTypes() const;
+
 	// Begin FTickableGameObject interface.
 	virtual void Tick(float DeltaTime) override;
 	virtual bool IsTickable() const override 
diff --git a/Engine/Source/Runtime/Landscape/Landscape.Build.cs b/Engine/Source/Runtime/Landscape/Landscape.Build.cs
--- a/Engine/Source/Runtime/Landscape/Landscape.Build.cs
+++ b/Engine/Source/Runtime/Landscape/Landscape.Build.cs
@@ -30,6 +30,7 @@ public class Landscape : ModuleRules
 				"RenderCore", 
 				"RHI",
 				"ShaderCore",
+                "Renderer",
                 "Foliage",
 			}
 		);
diff --git a/Engine/Source/Runtime/Landscape/Private/Landscape.cpp b/Engine/Source/Runtime/Landscape/Private/Landscape.cpp
--- a/Engine/Source/Runtime/Landscape/Private/Landscape.cpp
+++ b/Engine/Source/Runtime/Landscape/Private/Landscape.cpp
@@ -19,6 +19,7 @@ Landscape.cpp: Terrain rendering
 #include "LandscapeSplinesComponent.h"
 #include "LandscapeInfo.h"
 #include "LandscapeInfoMap.h"
+#include "LandscapeGrassType.h"
 #include "LandscapeLayerInfoObject.h"
 #include "LightMap.h"
 #include "ShadowMap.h"
@@ -32,7 +33,7 @@ Landscape.cpp: Terrain rendering
 #include "LandscapeMeshProxyComponent.h"
 #include "LandscapeMeshProxyActor.h"
 #include "ContentStreaming.h"
-
+#include "Materials/MaterialExpressionLandscapeGrassOutput.h"
 #include "Materials/MaterialExpressionLandscapeLayerWeight.h"
 #include "Materials/MaterialExpressionLandscapeLayerSwitch.h"
 #include "Materials/MaterialExpressionLandscapeLayerSample.h"
@@ -79,6 +80,7 @@ FAutoConsoleCommand CmdPrintNumLandscapeShadows(
 
 ULandscapeComponent::ULandscapeComponent(const FObjectInitializer& ObjectInitializer)
 	: Super(ObjectInitializer)
+, GrassData(MakeShareable(new FLandscapeComponentGrassData()))
 {
 	SetCollisionProfileName(UCollisionProfile::BlockAll_ProfileName);
 	bGenerateOverlapEvents = false;
@@ -107,6 +109,8 @@ ULandscapeComponent::ULandscapeComponent(const FObjectInitializer& ObjectInitial
 #endif
 
 	LpvBiasMultiplier = 0.0f; // Bias is 0 for landscape, since it's single sided
+
+	
 }
 
 ULandscapeComponent::~ULandscapeComponent()
@@ -144,10 +148,6 @@ void ULandscapeComponent::Serialize(FArchive& Ar)
 				HeightmapTexture = NULL;
 			}
 		}
-		if (Ar.UE4Ver() >= VER_UE4_LANDSCAPE_GRASS_COOKING)
-		{
-			GenerateGrassMap();
-		}
 	}
 #endif
 
@@ -156,6 +156,11 @@ void ULandscapeComponent::Serialize(FArchive& Ar)
 	Ar << LightMap;
 	Ar << ShadowMap;
 
+	if (Ar.UE4Ver() >= VER_UE4_SERIALIZE_LANDSCAPE_GRASS_DATA)
+	{
+		Ar << GrassData.Get();
+	}
+
 #if WITH_EDITOR
 	if (Ar.IsTransacting())
 	{
@@ -210,32 +215,15 @@ void ULandscapeComponent::Serialize(FArchive& Ar)
 				Ar << PlatformData;
 			}
 		}
-		if (Ar.UE4Ver() >= VER_UE4_LANDSCAPE_GRASS_COOKING)
-		{
-			// GrassMap
-			if (Ar.IsCooking() && Ar.IsSaving())
-			{
-				int32 NumChannels = GrassMap.IsValid() ? GrassMap->ElementStride : 0;
-				Ar << NumChannels;
-				if (NumChannels)
-				{
-					GrassMap->Data.BulkSerialize(Ar);
-				}
-#if WITH_EDITOR
-				RemoveGrassMap();
-#else
-				check(0); // why are we cooking without an editor
-#endif
-			}
-			else
+		if (Ar.UE4Ver() >= VER_UE4_LANDSCAPE_GRASS_COOKING && Ar.UE4Ver() < VER_UE4_SERIALIZE_LANDSCAPE_GRASS_DATA)
 			{
+			// deal with previous cooked FGrassMap data
 				int32 NumChannels = 0;
 				Ar << NumChannels;
 				if (NumChannels)
 				{
-					GrassMap = MakeShareable(new FGrassMap(NumChannels));
-					GrassMap->Data.BulkSerialize(Ar);
-				}
+				TArray<uint8> OldData;
+				OldData.BulkSerialize(Ar);
 			}
 		}
 	}
@@ -244,14 +232,7 @@ void ULandscapeComponent::Serialize(FArchive& Ar)
 
 SIZE_T ULandscapeComponent::GetResourceSize(EResourceSizeMode::Type Mode)
 {
-	SIZE_T ResSize = Super::GetResourceSize(Mode);
-
-	if (GrassMap.IsValid())
-	{
-		ResSize += GrassMap->Data.GetAllocatedSize();
-	}
-
-	return ResSize;
+	return Super::GetResourceSize(Mode) + GrassData->GetAllocatedSize();
 }
 
 #if WITH_EDITOR
@@ -487,10 +468,6 @@ void ULandscapeComponent::PostLoad()
 		}
 	}
 #endif
-	if (CVarPreloadGrassmapsInEditor.GetValueOnGameThread() > 0)
-	{
-		GenerateGrassMap();
-	}
 }
 
 #endif // WITH_EDITOR
@@ -2145,23 +2122,24 @@ void ALandscapeProxy::FlushFoliageComponents(const TSet<ULandscapeComponent*>* O
 		{
 			if (LandComp)
 			{
-				LandComp->RemoveGrassMap();
+				LandComp->RenderGrassMap();
 			}
 		}
 #endif
-		for (auto& CacheItem : FoliageCache.CachedGrassComps)
+		for (FCachedLandscapeFoliage::TGrassSet::TIterator Iter(FoliageCache.CachedGrassComps); Iter; ++Iter)
 		{
-			ULandscapeComponent* Component = CacheItem.Key.BasedOn.Get();
+			ULandscapeComponent* Component = (*Iter).Key.BasedOn.Get();
 			// if the weak pointer in the cache is invalid, we should kill them anyway
 			if (Component == nullptr || OnlyForComponents->Contains(Component))
 			{
-				UHierarchicalInstancedStaticMeshComponent *Used = CacheItem.Foliage.Get();
+				UHierarchicalInstancedStaticMeshComponent *Used = (*Iter).Foliage.Get();
 						if (Used)
 						{
 							SCOPE_CYCLE_COUNTER(STAT_FoliageGrassDestoryComp);
 							Used->ClearInstances();
 							Used->DestroyComponent();
 						}
+				Iter.RemoveCurrent();
 					}
 				}
 			}
@@ -2266,126 +2244,6 @@ static TAutoConsoleVariable<int32> CVarDisableGPUCull(
 	0,
 	TEXT("For debugging. Set this to zero to see where the grass is generated. Useful for tweaking the guard bands."));
 
-#if WITH_EDITOR && 0 // we will remove this shortly
-
-
-DECLARE_CYCLE_STAT(TEXT("Grass Load Source Art"),STAT_GrassLoadSourceArt,STATGROUP_Foliage);
-class FEditorGrassLayerData
-{
-	int32 Stride;
-	TArray<FColor> LocalHeightCache;
-	TArray<uint8> LocalCache;
-public:
-	FEditorGrassLayerData(ALandscapeProxy* Landscape, ULandscapeComponent* Component, ULandscapeLayerInfoObject* LayerInfo, UHierarchicalInstancedStaticMeshComponent* HierarchicalInstancedStaticMeshComponent, FEditorGrassLayerData* CopyFrom)
-	{
-		int32 SubsectionSizeQuads = Component->SubsectionSizeQuads;
-		int32 NumSubsections = Component->NumSubsections;
-		Stride = (SubsectionSizeQuads + 1) * NumSubsections;
-		if (Stride)
-		{
-			if (CopyFrom)
-			{
-				check(Stride == CopyFrom->Stride);
-				LocalHeightCache = CopyFrom->LocalHeightCache;
-				LocalCache = CopyFrom->LocalCache;
-			}
-			else
-			{
-				SCOPE_CYCLE_COUNTER(STAT_GrassLoadSourceArt);
-				FLandscapeComponentDataInterface CDI(Component);
-				CDI.GetHeightmapTextureData(LocalHeightCache);
-				if (!CDI.GetWeightmapTextureData(LayerInfo, LocalCache) || !LocalCache.Num() || LocalCache.Num() != LocalHeightCache.Num())
-				{
-					// no data for this layer for this component.
-					LocalHeightCache.Empty();
-					LocalCache.Empty();
-				}
-			}
-		}
-	}
-	bool HasData() const
-	{
-		return Stride && LocalCache.Num() && LocalHeightCache.Num() && LocalCache.Num() == LocalHeightCache.Num();
-	}
-	FORCEINLINE float GetHeight(int32 IdxX, int32 IdxY)
-	{
-		const FColor& SampleVal = LocalHeightCache[IdxX + Stride*IdxY];
-		return LandscapeDataAccess::GetLocalHeight((uint16(SampleVal.R) << 8) + uint16(SampleVal.G));
-	}
-	FORCEINLINE float GetLayer(int32 IdxX, int32 IdxY)
-	{
-		return (float)(LocalCache[IdxX + Stride*IdxY]) / 255.f;
-	}
-};
-
-typedef FEditorGrassLayerData FGrassLayerData;
-#else
-class FGameGrassLayerData
-{
-	int32 Stride;
-	int32 LayerOffset;
-	int32 ElementStride;
-	TSharedPtr<FGrassMap, ESPMode::ThreadSafe> GrassMap;
-	const uint8* GrassMapPtr;
-public:
-	FGameGrassLayerData(ALandscapeProxy* Landscape, ULandscapeComponent* Component, ULandscapeLayerInfoObject* LayerInfo, UHierarchicalInstancedStaticMeshComponent* HierarchicalInstancedStaticMeshComponent)
-	{
-		int32 SubsectionSizeQuads = Component->SubsectionSizeQuads;
-		int32 NumSubsections = Component->NumSubsections;
-		Stride = (SubsectionSizeQuads + 1) * NumSubsections;
-		GrassMapPtr = nullptr;
-		LayerOffset = 0;
-		ElementStride = 0;
-		if (Stride)
-		{
-#if WITH_EDITOR
-			if (!Component->GrassMap.IsValid())
-			{
-				check(IsInGameThread());
-				Component->GenerateGrassMap();
-			}
-#endif
-			GrassMap = Component->GrassMap;
-			if (GrassMap.IsValid())
-			{
-				check(GrassMap->Data.Num() == GrassMap->ElementStride * Stride * Stride);
-				int32 LayerIdx = INDEX_NONE;
-				for (int32 Idx = 0; Idx < Component->WeightmapLayerAllocations.Num(); Idx++)
-				{
-					if (Component->WeightmapLayerAllocations[Idx].LayerInfo == LayerInfo)
-					{
-						LayerIdx = Idx;
-						LayerOffset = int32(Component->WeightmapLayerAllocations[Idx].GrassMapChannelIndex);
-						check(LayerOffset > 1);
-						break;
-					}
-				}
-				if (LayerIdx >= 0 && LayerOffset > 1)
-				{
-					GrassMapPtr = &GrassMap->Data[0];
-					ElementStride = GrassMap->ElementStride;
-				}
-			}
-		}
-	}
-	bool HasData() const
-	{
-		return !!GrassMapPtr;
-	}
-	FORCEINLINE float GetHeight(int32 IdxX, int32 IdxY)
-	{
-		const uint8* Sample = &GrassMapPtr[(IdxX + Stride*IdxY) * ElementStride];
-		return LandscapeDataAccess::GetLocalHeight((uint16(Sample[0]) << 8) + uint16(Sample[1]));
-	}
-	FORCEINLINE float GetLayer(int32 IdxX, int32 IdxY)
-	{
-		return (float)(GrassMapPtr[(IdxX + Stride*IdxY) * ElementStride + LayerOffset]) / 255.f;
-	}
-};
-
-typedef FGameGrassLayerData FGrassLayerData;
-
-#endif
 
 struct FGrassBuilderBase
 {
@@ -2402,14 +2260,10 @@ struct FGrassBuilderBase
 	FVector Extent;
 	int32 SqrtMaxInstances;
 
-	int32 SubsectionSizeQuads;
-	int32 NumSubsections;
-	int32 Stride;
-
-	FGrassBuilderBase(ALandscapeProxy* Landscape, ULandscapeComponent* Component, ULandscapeLayerInfoObject* LayerInfo, int32 SqrtSubsections = 1, int32 SubX = 0, int32 SubY = 0)
+	FGrassBuilderBase(ALandscapeProxy* Landscape, ULandscapeComponent* Component, const ULandscapeGrassType* GrassType, int32 SqrtSubsections = 1, int32 SubX = 0, int32 SubY = 0)
 	{
 		bHaveValidData = true;
-		GrassDensity = LayerInfo->GrassDensity;
+		GrassDensity = GrassType->GrassDensity;
 
 		DrawScale = Landscape->GetRootComponent()->RelativeScale3D;
 		DrawLoc = Landscape->GetActorLocation();
@@ -2417,9 +2271,6 @@ struct FGrassBuilderBase
 
 		SectionBase = Component->GetSectionBase();
 		ComponentSizeQuads = Component->ComponentSizeQuads;
-		SubsectionSizeQuads = Component->SubsectionSizeQuads;
-		NumSubsections = Component->NumSubsections;
-		Stride = (SubsectionSizeQuads + 1) * NumSubsections;
 
 		Origin = FVector(DrawScale.X * float(SectionBase.X), DrawScale.Y * float(SectionBase.Y), 0.0f);
 		Extent = FVector(DrawScale.X * float(SectionBase.X + ComponentSizeQuads), DrawScale.Y * float(SectionBase.Y + ComponentSizeQuads), 0.0f) - Origin;
@@ -2445,9 +2296,45 @@ struct FGrassBuilderBase
 	}
 };
 
+// FLandscapeComponentGrassAccess - accessor wrapper for data for one GrassType from one Component
+struct FLandscapeComponentGrassAccess
+{
+	FLandscapeComponentGrassAccess(const ULandscapeComponent* InComponent, const ULandscapeGrassType* GrassType)
+	: GrassData(InComponent->GrassData)
+	, HeightData(InComponent->GrassData->HeightData)
+	, WeightData(InComponent->GrassData->WeightData.Find(GrassType))
+	, Stride(InComponent->ComponentSizeQuads + 1)
+	{}
+
+	bool IsValid()
+	{
+		return WeightData && WeightData->Num() == FMath::Square(Stride) && HeightData.Num() == FMath::Square(Stride);
+	}
+
+	FORCEINLINE float GetHeight(int32 IdxX, int32 IdxY)
+	{
+		return LandscapeDataAccess::GetLocalHeight(HeightData[IdxX + Stride*IdxY]);
+	}
+	FORCEINLINE float GetWeight(int32 IdxX, int32 IdxY)
+	{
+		return ((float)(*WeightData)[IdxX + Stride*IdxY]) / 255.f;
+	}
+
+	FORCEINLINE int32 GetStride()
+	{
+		return Stride;
+	}
+
+private:
+	TSharedRef<FLandscapeComponentGrassData, ESPMode::ThreadSafe> GrassData;
+	TArray<uint16>& HeightData;
+	TArray<uint8>* WeightData;
+	int32 Stride;
+};
+
 struct FAsyncGrassBuilder : public FGrassBuilderBase
 {
-	FGrassLayerData LayerData;
+	FLandscapeComponentGrassAccess GrassData;
 	bool RandomRotation;
 	bool AlignToSurface;
 	float PlacementJitter;
@@ -2464,19 +2351,19 @@ struct FAsyncGrassBuilder : public FGrassBuilderBase
 	FStaticMeshInstanceData InstanceBuffer;
 	TArray<FClusterNode> ClusterTree;
 
-	FAsyncGrassBuilder(ALandscapeProxy* Landscape, ULandscapeComponent* Component, ULandscapeLayerInfoObject* LayerInfo, UHierarchicalInstancedStaticMeshComponent* HierarchicalInstancedStaticMeshComponent, int32 SqrtSubsections, int32 SubX, int32 SubY)
-		: FGrassBuilderBase(Landscape, Component, LayerInfo, SqrtSubsections, SubX, SubY)
-		, LayerData(Landscape, Component, LayerInfo, HierarchicalInstancedStaticMeshComponent)
+	FAsyncGrassBuilder(ALandscapeProxy* Landscape, ULandscapeComponent* Component, const ULandscapeGrassType* GrassType, UHierarchicalInstancedStaticMeshComponent* HierarchicalInstancedStaticMeshComponent, int32 SqrtSubsections, int32 SubX, int32 SubY)
+		: FGrassBuilderBase(Landscape, Component, GrassType, SqrtSubsections, SubX, SubY)
+		, GrassData(Component, GrassType)
 	{
-		bHaveValidData = bHaveValidData && LayerData.HasData();
+		bHaveValidData = bHaveValidData && GrassData.IsValid();
 		XForm = LandscapeToWorld * HierarchicalInstancedStaticMeshComponent->ComponentToWorld.ToMatrixWithScale().Inverse();
-		PlacementJitter = LayerInfo->PlacementJitter;
-		RandomRotation = LayerInfo->RandomRotation;
-		AlignToSurface = LayerInfo->AlignToSurface;
+		PlacementJitter = GrassType->PlacementJitter;
+		RandomRotation = GrassType->RandomRotation;
+		AlignToSurface = GrassType->AlignToSurface;
 
 		RandomStream.Initialize(HierarchicalInstancedStaticMeshComponent->InstancingRandomSeed);
 
-		MeshBox = LayerInfo->GrassMesh->GetBounds().GetBox();
+		MeshBox = GrassType->GrassMesh->GetBounds().GetBox();
 
 		RasterTime = 0.0;
 		BuildTime = 0.0;
@@ -2628,10 +2515,10 @@ struct FAsyncGrassBuilder : public FGrassBuilderBase
 		int32 Y2 = FMath::CeilToInt(TestY);
 
 		// Min is to prevent the sampling of the final column from overflowing
-		int32 IdxX1 = FMath::Min<int32>(((X1 / SubsectionSizeQuads) * (SubsectionSizeQuads + 1)) + (X1 % SubsectionSizeQuads), Stride - 1);
-		int32 IdxY1 = FMath::Min<int32>(((Y1 / SubsectionSizeQuads) * (SubsectionSizeQuads + 1)) + (Y1 % SubsectionSizeQuads), Stride - 1);
-		int32 IdxX2 = FMath::Min<int32>(((X2 / SubsectionSizeQuads) * (SubsectionSizeQuads + 1)) + (X2 % SubsectionSizeQuads), Stride - 1);
-		int32 IdxY2 = FMath::Min<int32>(((Y2 / SubsectionSizeQuads) * (SubsectionSizeQuads + 1)) + (Y2 % SubsectionSizeQuads), Stride - 1);
+		int32 IdxX1 = FMath::Min<int32>(X1, GrassData.GetStride() - 1);
+		int32 IdxY1 = FMath::Min<int32>(Y1, GrassData.GetStride() - 1);
+		int32 IdxX2 = FMath::Min<int32>(X2, GrassData.GetStride() - 1);
+		int32 IdxY2 = FMath::Min<int32>(Y2, GrassData.GetStride() - 1);
 
 		float LerpX = FMath::Fractional(TestX);
 		float LerpY = FMath::Fractional(TestY);
@@ -2639,10 +2526,10 @@ struct FAsyncGrassBuilder : public FGrassBuilderBase
 		float Result;
 		{
 			// sample
-			float Sample11 = LayerData.GetLayer(IdxX1, IdxY1);
-			float Sample21 = LayerData.GetLayer(IdxX2, IdxY1);
-			float Sample12 = LayerData.GetLayer(IdxX1, IdxY2);
-			float Sample22 = LayerData.GetLayer(IdxX2, IdxY2);
+			float Sample11 = GrassData.GetWeight(IdxX1, IdxY1);
+			float Sample21 = GrassData.GetWeight(IdxX2, IdxY1);
+			float Sample12 = GrassData.GetWeight(IdxX1, IdxY2);
+			float Sample22 = GrassData.GetWeight(IdxX2, IdxY2);
 
 			// Bilinear interpolate
 			Result = FMath::Lerp(
@@ -2653,10 +2540,10 @@ struct FAsyncGrassBuilder : public FGrassBuilderBase
 
 		{
 			// sample
-			float Sample11 = LayerData.GetHeight(IdxX1, IdxY1);
-			float Sample21 = LayerData.GetHeight(IdxX2, IdxY1);
-			float Sample12 = LayerData.GetHeight(IdxX1, IdxY2);
-			float Sample22 = LayerData.GetHeight(IdxX2, IdxY2);
+			float Sample11 = GrassData.GetHeight(IdxX1, IdxY1);
+			float Sample21 = GrassData.GetHeight(IdxX2, IdxY1);
+			float Sample12 = GrassData.GetHeight(IdxX1, IdxY2);
+			float Sample22 = GrassData.GetHeight(IdxX2, IdxY2);
 
 			OutLocation->X = InLocation.X - DrawScale.X * float(LandscapeSectionOffset.X);
 			OutLocation->Y = InLocation.Y - DrawScale.Y * float(LandscapeSectionOffset.Y);
@@ -2723,12 +2610,32 @@ void ALandscapeProxy::Tick(float DeltaSeconds)
 	}
 }
 
+TArray<const ULandscapeGrassType*> ALandscapeProxy::GetGrassTypes() const
+{
+	TArray<const ULandscapeGrassType*> GrassTypes;
+	if (LandscapeMaterial)
+	{
+		TArray<const UMaterialExpressionLandscapeGrassOutput*> GrassExpressions;
+		LandscapeMaterial->GetMaterial()->GetAllExpressionsOfType<UMaterialExpressionLandscapeGrassOutput>(GrassExpressions);
+		if (GrassExpressions.Num() > 0)
+		{
+			for (auto& Type : GrassExpressions[0]->GrassTypes)
+			{
+				GrassTypes.Add(Type.GrassType);
+			}
+		}
+	}
+	return GrassTypes;
+}
+
 void ALandscapeProxy::UpdateFoliage(const TArray<FVector>& Cameras, ULandscapeComponent* OnlyComponent, bool bForceSync)
 {
 	SCOPE_CYCLE_COUNTER(STAT_GrassUpdate);
 
 	if (CVarGrassEnable.GetValueOnAnyThread() > 0)
 	{
+		TArray<const ULandscapeGrassType*> GrassTypes = GetGrassTypes();
+
 		float GuardBand = CVarGuardBandMultiplier.GetValueOnAnyThread();
 		float DiscardGuardBand = CVarGuardBandDiscardMultiplier.GetValueOnAnyThread();
 		bool bCullSubsections = CVarCullSubsections.GetValueOnAnyThread() > 0;
@@ -2739,187 +2646,193 @@ void ALandscapeProxy::UpdateFoliage(const TArray<FVector>& Cameras, ULandscapeCo
 		if (GetWorld())
 		{
 
-			for (int32 ComponentIndex = 0; ComponentIndex < LandscapeComponents.Num(); ComponentIndex++)
+		for (int32 ComponentIndex = 0; ComponentIndex < LandscapeComponents.Num(); ComponentIndex++)
+		{
+			ULandscapeComponent* Component = LandscapeComponents[ComponentIndex];
+			if (OnlyComponent && OnlyComponent != Component)
 			{
-				ULandscapeComponent* Component = LandscapeComponents[ComponentIndex];
-				if (OnlyComponent && OnlyComponent != Component)
-				{
-					continue;
-				}
+				continue;
+			}
 
-				FBoxSphereBounds WorldBounds = Component->CalcBounds(Component->ComponentToWorld);
-				float MinDistanceToComp = Cameras.Num() ? MAX_flt : 0.0f;
+			FBoxSphereBounds WorldBounds = Component->CalcBounds(Component->ComponentToWorld);
+			float MinDistanceToComp = Cameras.Num() ? MAX_flt : 0.0f;
 
-				for (auto& Pos : Cameras)
-				{
-					MinDistanceToComp = FMath::Min<float>(MinDistanceToComp, WorldBounds.ComputeSquaredDistanceFromBoxToPoint(Pos));
-				}
+			for (auto& Pos : Cameras)
+			{
+				MinDistanceToComp = FMath::Min<float>(MinDistanceToComp, WorldBounds.ComputeSquaredDistanceFromBoxToPoint(Pos));
+			}
 
-				MinDistanceToComp = FMath::Sqrt(MinDistanceToComp);
+			MinDistanceToComp = FMath::Sqrt(MinDistanceToComp);
 
-				for (FWeightmapLayerAllocationInfo& LayerAlloc : Component->WeightmapLayerAllocations)
+			for (auto GrassType : GrassTypes)
+			{
+				if (GrassType && GrassType->GrassMesh && GrassType->GrassDensity > 0.0f && GrassType->EndCullDistance > 0)
 				{
-					ULandscapeLayerInfoObject* LayerInfo = LayerAlloc.LayerInfo;
+					bool bRandomRotation = GrassType->RandomRotation;
+					bool bAlign = GrassType->AlignToSurface;
+					float MustHaveDistance = GuardBand * (float)GrassType->EndCullDistance;
+					float DiscardDistance = DiscardGuardBand * (float)GrassType->EndCullDistance;
 
-					if (LayerInfo && LayerInfo->GrassMesh && LayerInfo->GrassDensity > 0.0f && LayerInfo->EndCullDistance > 0)
+					if (MinDistanceToComp > DiscardDistance)
 					{
-						bool bRandomRotation = LayerInfo->RandomRotation;
-						bool bAlign = LayerInfo->AlignToSurface;
-						float MustHaveDistance = GuardBand * (float)LayerInfo->EndCullDistance;
-						float DiscardDistance = DiscardGuardBand * (float)LayerInfo->EndCullDistance;
-
-						if (MinDistanceToComp > DiscardDistance)
-						{
-							continue;
-						}
+						continue;
+					}
 
-						FGrassBuilderBase ForSubsectionMath(this, Component, LayerInfo);
+					FGrassBuilderBase ForSubsectionMath(this, Component, GrassType);
 
-						int32 SqrtSubsections = 1;
+					int32 SqrtSubsections = 1;
 				
-						if (ForSubsectionMath.bHaveValidData && ForSubsectionMath.SqrtMaxInstances > 0)
-						{
-							SqrtSubsections = FMath::Clamp<int32>(FMath::CeilToInt(float(ForSubsectionMath.SqrtMaxInstances) / FMath::Sqrt((float)MaxInstancesPerComponent)), 1, 16);
-						}
+					if (ForSubsectionMath.bHaveValidData && ForSubsectionMath.SqrtMaxInstances > 0)
+					{
+						SqrtSubsections = FMath::Clamp<int32>(FMath::CeilToInt(float(ForSubsectionMath.SqrtMaxInstances) / FMath::Sqrt((float)MaxInstancesPerComponent)), 1, 16);
+					}
 
-						int32 NumCompsCreated = 0;
-						FBox LocalBox = Component->CachedLocalBox;
-						FVector LocalExtentDiv = (LocalBox.Max - LocalBox.Min) * FVector(1.0f / float(SqrtSubsections), 1.0f / float(SqrtSubsections), 1.0f);
-						for (int32 SubX = 0; SubX < SqrtSubsections; SubX++)
+					int32 NumCompsCreated = 0;
+					FBox LocalBox = Component->CachedLocalBox;
+					FVector LocalExtentDiv = (LocalBox.Max - LocalBox.Min) * FVector(1.0f / float(SqrtSubsections), 1.0f / float(SqrtSubsections), 1.0f);
+					for (int32 SubX = 0; SubX < SqrtSubsections; SubX++)
+					{
+						for (int32 SubY = 0; SubY < SqrtSubsections; SubY++)
 						{
-							for (int32 SubY = 0; SubY < SqrtSubsections; SubY++)
-							{
-								float MinDistanceToSubComp = MinDistanceToComp;
-
-								if (bCullSubsections && SqrtSubsections > 1)
-								{
-									FVector BoxMin;
-									BoxMin.X = LocalBox.Min.X + LocalExtentDiv.X * float(SubX);
-									BoxMin.Y = LocalBox.Min.Y + LocalExtentDiv.Y * float(SubY);
-									BoxMin.Z = LocalBox.Min.Z;
+							float MinDistanceToSubComp = MinDistanceToComp;
 
-									FVector BoxMax;
-									BoxMax.X = LocalBox.Min.X + LocalExtentDiv.X * float(SubX + 1);
-									BoxMax.Y = LocalBox.Min.Y + LocalExtentDiv.Y * float(SubY + 1);
-									BoxMax.Z = LocalBox.Max.Z;
+							if (bCullSubsections && SqrtSubsections > 1)
+							{
+								FVector BoxMin;
+								BoxMin.X = LocalBox.Min.X + LocalExtentDiv.X * float(SubX);
+								BoxMin.Y = LocalBox.Min.Y + LocalExtentDiv.Y * float(SubY);
+								BoxMin.Z = LocalBox.Min.Z;
 
-									FBox LocalSubBox(BoxMin, BoxMax);
-									FBox WorldSubBox = LocalSubBox.TransformBy(Component->ComponentToWorld);
+								FVector BoxMax;
+								BoxMax.X = LocalBox.Min.X + LocalExtentDiv.X * float(SubX + 1);
+								BoxMax.Y = LocalBox.Min.Y + LocalExtentDiv.Y * float(SubY + 1);
+								BoxMax.Z = LocalBox.Max.Z;
 
-									MinDistanceToSubComp = Cameras.Num() ? MAX_flt : 0.0f;
-									for (auto& Pos : Cameras)
-									{
-										MinDistanceToSubComp = FMath::Min<float>(MinDistanceToSubComp, ComputeSquaredDistanceFromBoxToPoint(WorldSubBox.Min, WorldSubBox.Max, Pos));
-									}
-									MinDistanceToSubComp = FMath::Sqrt(MinDistanceToSubComp);
-								}
+								FBox LocalSubBox(BoxMin, BoxMax);
+								FBox WorldSubBox = LocalSubBox.TransformBy(Component->ComponentToWorld);
 
-								if (MinDistanceToSubComp > DiscardDistance)
+								MinDistanceToSubComp = Cameras.Num() ? MAX_flt : 0.0f;
+								for (auto& Pos : Cameras)
 								{
-									continue;
+									MinDistanceToSubComp = FMath::Min<float>(MinDistanceToSubComp, ComputeSquaredDistanceFromBoxToPoint(WorldSubBox.Min, WorldSubBox.Max, Pos));
 								}
+								MinDistanceToSubComp = FMath::Sqrt(MinDistanceToSubComp);
+							}
+
+							if (MinDistanceToSubComp > DiscardDistance)
+							{
+								continue;
+							}
 
-								FCachedLandscapeFoliage::FGrassComp NewComp;
-								NewComp.Key.BasedOn = Component;
-								NewComp.Key.Layer = LayerInfo;
-								NewComp.Key.SqrtSubsections = SqrtSubsections;
-								NewComp.Key.CachedMaxInstancesPerComponent = MaxInstancesPerComponent;
-								NewComp.Key.SubsectionX = SubX;
-								NewComp.Key.SubsectionY = SubY;
+							FCachedLandscapeFoliage::FGrassComp NewComp;
+							NewComp.Key.BasedOn = Component;
+							NewComp.Key.GrassType = GrassType;
+							NewComp.Key.SqrtSubsections = SqrtSubsections;
+							NewComp.Key.CachedMaxInstancesPerComponent = MaxInstancesPerComponent;
+							NewComp.Key.SubsectionX = SubX;
+							NewComp.Key.SubsectionY = SubY;
 
+							{
+								FCachedLandscapeFoliage::FGrassComp* Existing = FoliageCache.CachedGrassComps.Find(NewComp.Key);
+								if (Existing || MinDistanceToSubComp > MustHaveDistance)
 								{
-									FCachedLandscapeFoliage::FGrassComp* Existing = FoliageCache.CachedGrassComps.Find(NewComp.Key);
-									if (Existing || MinDistanceToSubComp > MustHaveDistance)
+									if (Existing)
 									{
-										if (Existing)
-										{
-											Existing->Touch();
-										}
-										continue;
+										Existing->Touch();
 									}
+									continue;
 								}
+							}
 
 								if (NumCompsCreated || (!bForceSync && AsyncFoliageTasks.Num() >= MaxTasks))
-								{
-									continue; // one per frame, but we still want to touch the existing ones
-								}
-								NumCompsCreated++;
+							{
+								continue; // one per frame, but we still want to touch the existing ones
+							}
+							NumCompsCreated++;
 
-								SCOPE_CYCLE_COUNTER(STAT_FoliageGrassStartComp);
-								int32 FolSeed = FCrc::StrCrc32((LayerInfo->LayerName.ToString() + Component->GetName() + FString::Printf(TEXT("%d %d"), SubX, SubY)).GetCharArray().GetData());
-								if (FolSeed == 0)
-								{
-									FolSeed++;
-								}
+							SCOPE_CYCLE_COUNTER(STAT_FoliageGrassStartComp);
+							int32 FolSeed = FCrc::StrCrc32((GrassType->GetName() + Component->GetName() + FString::Printf(TEXT("%d %d"), SubX, SubY)).GetCharArray().GetData());
+							if (FolSeed == 0)
+							{
+								FolSeed++;
+							}
 
-								UHierarchicalInstancedStaticMeshComponent* HierarchicalInstancedStaticMeshComponent;
-								{
-									QUICK_SCOPE_CYCLE_COUNTER(STAT_GrassCreateComp);
-									HierarchicalInstancedStaticMeshComponent = ConstructObject<UHierarchicalInstancedStaticMeshComponent>(UHierarchicalInstancedStaticMeshComponent::StaticClass(), this); //, NAME_None, RF_Transactional);
-								}
-								NewComp.Foliage = HierarchicalInstancedStaticMeshComponent;
-								FoliageCache.CachedGrassComps.Add(NewComp);
+							UHierarchicalInstancedStaticMeshComponent* HierarchicalInstancedStaticMeshComponent;
+							{
+								QUICK_SCOPE_CYCLE_COUNTER(STAT_GrassCreateComp);
+								HierarchicalInstancedStaticMeshComponent = ConstructObject<UHierarchicalInstancedStaticMeshComponent>(UHierarchicalInstancedStaticMeshComponent::StaticClass(), this); //, NAME_None, RF_Transactional);
+							}
+							NewComp.Foliage = HierarchicalInstancedStaticMeshComponent;
+							FoliageCache.CachedGrassComps.Add(NewComp);
 
-								HierarchicalInstancedStaticMeshComponent->Mobility = EComponentMobility::Static;
+							HierarchicalInstancedStaticMeshComponent->Mobility = EComponentMobility::Static;
 
-								HierarchicalInstancedStaticMeshComponent->StaticMesh = LayerInfo->GrassMesh;
-								HierarchicalInstancedStaticMeshComponent->bSelectable = false;
-								HierarchicalInstancedStaticMeshComponent->bHasPerInstanceHitProxies = true;
-								static FName NoCollision(TEXT("NoCollision"));
-								HierarchicalInstancedStaticMeshComponent->SetCollisionProfileName(NoCollision);
-								HierarchicalInstancedStaticMeshComponent->bDisableCollision = true;
-								HierarchicalInstancedStaticMeshComponent->SetCanEverAffectNavigation(false);
-								HierarchicalInstancedStaticMeshComponent->InstancingRandomSeed = FolSeed;
+							HierarchicalInstancedStaticMeshComponent->StaticMesh = GrassType->GrassMesh;
+							HierarchicalInstancedStaticMeshComponent->bSelectable = false;
+							HierarchicalInstancedStaticMeshComponent->bHasPerInstanceHitProxies = true;
+							static FName NoCollision(TEXT("NoCollision"));
+							HierarchicalInstancedStaticMeshComponent->SetCollisionProfileName(NoCollision);
+							HierarchicalInstancedStaticMeshComponent->bDisableCollision = true;
+							HierarchicalInstancedStaticMeshComponent->SetCanEverAffectNavigation(false);
+							HierarchicalInstancedStaticMeshComponent->InstancingRandomSeed = FolSeed;
 
-								if (!Cameras.Num() || bDisableGPUCull)
-								{
-									// if we don't have any cameras, then we are rendering landscape LOD materials or somesuch and we want to disable culling
-									HierarchicalInstancedStaticMeshComponent->InstanceStartCullDistance = 0;
-									HierarchicalInstancedStaticMeshComponent->InstanceEndCullDistance = 0;
-								}
-								else
-								{
-									HierarchicalInstancedStaticMeshComponent->InstanceStartCullDistance = LayerInfo->StartCullDistance;
-									HierarchicalInstancedStaticMeshComponent->InstanceEndCullDistance = LayerInfo->EndCullDistance;
-								}
+							if (!Cameras.Num() || bDisableGPUCull)
+							{
+								// if we don't have any cameras, then we are rendering landscape LOD materials or somesuch and we want to disable culling
+								HierarchicalInstancedStaticMeshComponent->InstanceStartCullDistance = 0;
+								HierarchicalInstancedStaticMeshComponent->InstanceEndCullDistance = 0;
+							}
+							else
+							{
+								HierarchicalInstancedStaticMeshComponent->InstanceStartCullDistance = GrassType->StartCullDistance;
+								HierarchicalInstancedStaticMeshComponent->InstanceEndCullDistance = GrassType->EndCullDistance;
+							}
 
-								//@todo - take the settings from a UFoliageType object.  For now, disable distance field lighting on grass so we don't hitch.
-								HierarchicalInstancedStaticMeshComponent->bAffectDistanceFieldLighting = false;
+							//@todo - take the settings from a UFoliageType object.  For now, disable distance field lighting on grass so we don't hitch.
+							HierarchicalInstancedStaticMeshComponent->bAffectDistanceFieldLighting = false;
 
-								{
-									QUICK_SCOPE_CYCLE_COUNTER(STAT_GrassAttachComp);
+							{
+								QUICK_SCOPE_CYCLE_COUNTER(STAT_GrassAttachComp);
 
-									HierarchicalInstancedStaticMeshComponent->AttachTo(GetRootComponent());
-									FTransform DesiredTransform = GetRootComponent()->ComponentToWorld;
-									DesiredTransform.RemoveScaling();
-									HierarchicalInstancedStaticMeshComponent->SetWorldTransform(DesiredTransform);
+								HierarchicalInstancedStaticMeshComponent->AttachTo(GetRootComponent());
+								FTransform DesiredTransform = GetRootComponent()->ComponentToWorld;
+								DesiredTransform.RemoveScaling();
+								HierarchicalInstancedStaticMeshComponent->SetWorldTransform(DesiredTransform);
 
-									FoliageComponents.Add(HierarchicalInstancedStaticMeshComponent);
-								}
+								FoliageComponents.Add(HierarchicalInstancedStaticMeshComponent);
+							}
 
-								FAsyncGrassBuilder* Builder;
+#if WITH_EDITOR
+							// render grass data if we don't have any
+							if (!Component->GrassData->HasData())
+							{
+								QUICK_SCOPE_CYCLE_COUNTER(STAT_GrassRenderToTexture);
+								Component->RenderGrassMap();
+							}
+#endif
 
-								{
-									QUICK_SCOPE_CYCLE_COUNTER(STAT_GrassCreateBuilder);
-									Builder = new FAsyncGrassBuilder(this, Component, LayerInfo, HierarchicalInstancedStaticMeshComponent, SqrtSubsections, SubX, SubY);
-								}
+							FAsyncGrassBuilder* Builder;
 
-								if (Builder->bHaveValidData)
-								{
-									FAsyncTask<FAsyncGrassTask>* Task = new FAsyncTask<FAsyncGrassTask>(Builder, NewComp.Key, HierarchicalInstancedStaticMeshComponent);
+							{
+								QUICK_SCOPE_CYCLE_COUNTER(STAT_GrassCreateBuilder);
+								Builder = new FAsyncGrassBuilder(this, Component, GrassType, HierarchicalInstancedStaticMeshComponent, SqrtSubsections, SubX, SubY);
+							}
 
-									Task->StartBackgroundTask();
+							if (Builder->bHaveValidData)
+							{
+								FAsyncTask<FAsyncGrassTask>* Task = new FAsyncTask<FAsyncGrassTask>(Builder, NewComp.Key, HierarchicalInstancedStaticMeshComponent);
 
-									AsyncFoliageTasks.Add(Task);
-								}
-								else
-								{
-									delete Builder;
-								}
-								{
-									QUICK_SCOPE_CYCLE_COUNTER(STAT_GrassRegisterComp);
-									HierarchicalInstancedStaticMeshComponent->RegisterComponent();
-								}
+								Task->StartBackgroundTask();
+
+								AsyncFoliageTasks.Add(Task);
+							}
+							else
+							{
+								delete Builder;
+							}
+							{
+								QUICK_SCOPE_CYCLE_COUNTER(STAT_GrassRegisterComp);
+								HierarchicalInstancedStaticMeshComponent->RegisterComponent();
 							}
 						}
 					}
@@ -2927,6 +2840,7 @@ void ALandscapeProxy::UpdateFoliage(const TArray<FVector>& Cameras, ULandscapeCo
 			}
 		}
 	}
+	}
 
 	static TSet<UHierarchicalInstancedStaticMeshComponent *> StillUsed;
 	StillUsed.Empty(256);
@@ -2942,7 +2856,7 @@ void ALandscapeProxy::UpdateFoliage(const TArray<FVector>& Cameras, ULandscapeCo
 				!GrassItem.Pending &&
 				(
 					!GrassItem.Key.BasedOn.Get() ||
-					!GrassItem.Key.Layer.Get() ||
+					!GrassItem.Key.GrassType.Get() ||
 					!Used ||
 					(GrassItem.LastUsedFrameNumber < OldestToKeepFrame && GrassItem.LastUsedTime < OldestToKeepTime)
 				);
@@ -3057,6 +2971,24 @@ static FAutoConsoleCommand FlushFoliageCmd(
 	FConsoleCommandWithArgsDelegate::CreateStatic(&FlushFoliage)
 	);
 
+
+#if WITH_EDITOR
+static void RenderGrass(const TArray<FString>& Args)
+{
+	for (TObjectIterator<ULandscapeComponent> It; It; ++It)
+	{
+		It->RenderGrassMap();
+	}
+}
+
+static FAutoConsoleCommand RenderFoliageCmd(
+	TEXT("grass.Render"),
+	TEXT("Render grass to texture. Debugging."),
+	FConsoleCommandWithArgsDelegate::CreateStatic(&RenderGrass)
+	);
+#endif
+
+
 //
 // ALandscapeMeshProxyActor
 //
diff --git a/Engine/Source/Runtime/Landscape/Private/LandscapeEdit.cpp b/Engine/Source/Runtime/Landscape/Private/LandscapeEdit.cpp
--- a/Engine/Source/Runtime/Landscape/Private/LandscapeEdit.cpp
+++ b/Engine/Source/Runtime/Landscape/Private/LandscapeEdit.cpp
@@ -3410,12 +3410,6 @@ ULandscapeLayerInfoObject::ULandscapeLayerInfoObject(const FObjectInitializer& O
 	: Super(ObjectInitializer)
 {
 	Hardness = 0.5f;
-	GrassDensity = 400;
-	StartCullDistance = 10000.0f;
-	EndCullDistance = 10000.0f;
-	PlacementJitter = 1.0f;
-	RandomRotation = true;
-	AlignToSurface = true;
 #if WITH_EDITORONLY_DATA
 	bNoWeightBlend = false;
 #endif // WITH_EDITORONLY_DATA
@@ -3836,7 +3830,7 @@ void ALandscape::PostEditChangeProperty(FPropertyChangedEvent& PropertyChangedEv
 	}
 	else if (GIsEditor && PropertyName == FName(TEXT("MaxLODLevel")))
 	{
-		MaxLODLevel = FMath::Clamp<int32>(MaxLODLevel, 0, FMath::CeilLogTwo(SubsectionSizeQuads + 1) - 1);
+		MaxLODLevel = FMath::Clamp<int32>(MaxLODLevel, -1, FMath::CeilLogTwo(SubsectionSizeQuads + 1) - 1);
 		bPropagateToProxies = true;
 	}
 	else if (PropertyName == FName(TEXT("LODDistanceFactor")))
@@ -5206,80 +5200,6 @@ void ULandscapeComponent::GeneratePlatformVertexData()
 	PlatformData.InitializeFromUncompressedData(NewPlatformData);
 }
 
-void ULandscapeComponent::GenerateGrassMap()
-{
-	RemoveGrassMap();
-	int32 Stride = (SubsectionSizeQuads + 1) * NumSubsections;
-	if (Stride)
-	{
-		int32 NumChannels = 2; // these are for the height
-		for (FWeightmapLayerAllocationInfo& Layer : WeightmapLayerAllocations)
-		{
-			ULandscapeLayerInfoObject* LayerInfo = Layer.LayerInfo;
-
-			if (LayerInfo && LayerInfo->GrassMesh && LayerInfo->GrassDensity > 0.0f && LayerInfo->EndCullDistance > 0)
-			{
-				Layer.GrassMapChannelIndex = (uint8)NumChannels++;
-			}
-			else
-			{
-				check(!Layer.GrassMapChannelIndex);
-			}
-		}
-		if (NumChannels > 2)
-		{
-			GrassMap = MakeShareable(new FGrassMap(NumChannels));
-			TArray<FColor> LocalHeightCache;
-			FLandscapeComponentDataInterface CDI(this);
-			CDI.GetHeightmapTextureData(LocalHeightCache, true);
-			if (LocalHeightCache.Num())
-			{
-				check(LocalHeightCache.Num() == Stride * Stride);
-				TArray<uint8>& Data = GrassMap->Data;
-				Data.AddZeroed(Stride * Stride * NumChannels);
-				for (int32 Index = 0; Index < LocalHeightCache.Num(); Index++)
-				{
-					const FColor& SampleVal = LocalHeightCache[Index];
-					uint8* Dest = &Data[Index * NumChannels];
-					*Dest++ = SampleVal.R;
-					*Dest = SampleVal.G;
-				}
-				TArray<uint8> LocalCache;
-				for (FWeightmapLayerAllocationInfo& Layer : WeightmapLayerAllocations)
-				{
-					if (Layer.GrassMapChannelIndex)
-					{
-						ULandscapeLayerInfoObject* LayerInfo = Layer.LayerInfo;
-						check(LayerInfo && LayerInfo->GrassMesh && LayerInfo->GrassDensity > 0.0f && LayerInfo->EndCullDistance > 0);
-						LocalCache.Reset();
-						if (CDI.GetWeightmapTextureData(LayerInfo, LocalCache) && LocalCache.Num() && LocalCache.Num() == LocalHeightCache.Num())
-						{
-							for (int32 Index = 0; Index < LocalCache.Num(); Index++)
-							{
-								Data[Index * NumChannels + Layer.GrassMapChannelIndex] = LocalCache[Index];
-							}
-						}
-					}
-				}
-			}
-			else
-			{
-				RemoveGrassMap();
-			}
-		}
-	}
-}
-
-void ULandscapeComponent::RemoveGrassMap()
-{
-	GrassMap.Reset();
-	for (FWeightmapLayerAllocationInfo& Layer : WeightmapLayerAllocations)
-	{
-		Layer.GrassMapChannelIndex = 0;
-	}
-}
-
-
 UTexture2D* ALandscapeProxy::CreateLandscapeTexture(int32 InSizeX, int32 InSizeY, TextureGroup InLODGroup, ETextureSourceFormat InFormat, UObject* OptionalOverrideOuter) const
 {
 	UObject* TexOuter = OptionalOverrideOuter ? OptionalOverrideOuter : GetOutermost();
diff --git a/Engine/Source/Runtime/Landscape/Private/LandscapeGrass.cpp b/Engine/Source/Runtime/Landscape/Private/LandscapeGrass.cpp
--- a/Engine/Source/Runtime/Landscape/Private/LandscapeGrass.cpp
+++ b/Engine/Source/Runtime/Landscape/Private/LandscapeGrass.cpp
@@ -1,13 +1,358 @@
 // Copyright 1998-2015 Epic Games, Inc. All Rights Reserved.
 
 #include "Landscape.h"
-#include "EngineUtils.h"
+#include "Engine/Engine.h"
+#include "EngineGlobals.h"
+#include "UnrealEngine.h"
+#include "Shader.h"
+#include "MeshMaterialShader.h"
+#include "EngineModule.h"
+#include "RendererInterface.h"
+#include "DrawingPolicy.h"
+#include "PrimitiveSceneInfo.h"
 #include "MaterialCompiler.h"
 #include "Materials/MaterialExpressionLandscapeGrassOutput.h"
 #include "EdGraph/EdGraphNode.h"
 
 #define LOCTEXT_NAMESPACE "Landscape"
 
+class FLandscapeGrassWeightVS : public FMeshMaterialShader
+{
+	DECLARE_SHADER_TYPE(FLandscapeGrassWeightVS, MeshMaterial);
+protected:
+
+	FLandscapeGrassWeightVS() {}
+	FLandscapeGrassWeightVS(const FMeshMaterialShaderType::CompiledShaderInitializerType& Initializer) :
+		FMeshMaterialShader(Initializer)
+	{}
+
+public:
+
+	static bool ShouldCache(EShaderPlatform Platform, const FMaterial* Material, const FVertexFactoryType* VertexFactoryType)
+	{
+		// Only compile for LandscapeVertexFactory
+		return (Material->IsUsedWithLandscape() || Material->IsSpecialEngineMaterial()) && IsFeatureLevelSupported(Platform, ERHIFeatureLevel::SM4) && (VertexFactoryType == FindVertexFactoryType(FName(TEXT("FLandscapeVertexFactory"), FNAME_Find)));
+	}
+
+	void SetParameters(FRHICommandList& RHICmdList, const FMaterialRenderProxy* MaterialRenderProxy, const FMaterial& MaterialResource, const FSceneView& View)
+	{
+		FMeshMaterialShader::SetParameters(RHICmdList, GetVertexShader(), MaterialRenderProxy, MaterialResource, View, ESceneRenderTargetsMode::DontSet);
+	}
+
+	void SetMesh(FRHICommandList& RHICmdList, const FVertexFactory* VertexFactory, const FSceneView& View, const FPrimitiveSceneProxy* Proxy, const FMeshBatchElement& BatchElement)
+	{
+		FMeshMaterialShader::SetMesh(RHICmdList, GetVertexShader(), VertexFactory, View, Proxy, BatchElement);
+	}
+};
+
+IMPLEMENT_MATERIAL_SHADER_TYPE(, FLandscapeGrassWeightVS, TEXT("LandscapeGrassWeight"), TEXT("VSMain"), SF_Vertex);
+
+
+class FLandscapeGrassWeightPS : public FMeshMaterialShader
+{
+	DECLARE_SHADER_TYPE(FLandscapeGrassWeightPS, MeshMaterial);
+public:
+
+	static bool ShouldCache(EShaderPlatform Platform, const FMaterial* Material, const FVertexFactoryType* VertexFactoryType)
+	{
+		// Only compile for LandscapeVertexFactory
+		return (Material->IsUsedWithLandscape() || Material->IsSpecialEngineMaterial()) && IsFeatureLevelSupported(Platform, ERHIFeatureLevel::SM4) && (VertexFactoryType == FindVertexFactoryType(FName(TEXT("FLandscapeVertexFactory"), FNAME_Find)));
+	}
+
+	FLandscapeGrassWeightPS(const ShaderMetaType::CompiledShaderInitializerType& Initializer) :
+		FMeshMaterialShader(Initializer)
+	{}
+
+	FLandscapeGrassWeightPS() {}
+
+	void SetParameters(FRHICommandList& RHICmdList, const FMaterialRenderProxy* MaterialRenderProxy, const FMaterial& MaterialResource, const FSceneView* View)
+	{
+		FMeshMaterialShader::SetParameters(RHICmdList, GetPixelShader(), MaterialRenderProxy, MaterialResource, *View, ESceneRenderTargetsMode::DontSet);
+	}
+
+	void SetMesh(FRHICommandList& RHICmdList, const FVertexFactory* VertexFactory, const FSceneView& View, const FPrimitiveSceneProxy* Proxy, const FMeshBatchElement& BatchElement)
+	{
+		FMeshMaterialShader::SetMesh(RHICmdList, GetPixelShader(), VertexFactory, View, Proxy, BatchElement);
+	}
+};
+
+IMPLEMENT_MATERIAL_SHADER_TYPE(, FLandscapeGrassWeightPS, TEXT("LandscapeGrassWeight"), TEXT("PSMain"), SF_Pixel);
+
+
+
+/**
+* Used to write out depth for opaque and masked materials during the depth-only pass.
+*/
+class FLandscapeGrassWeightDrawingPolicy : public FMeshDrawingPolicy
+{
+public:
+	FLandscapeGrassWeightDrawingPolicy(
+		const FVertexFactory* InVertexFactory,
+		const FMaterialRenderProxy* InMaterialRenderProxy,
+		const FMaterial& InMaterialResource
+		)
+		:
+		FMeshDrawingPolicy(InVertexFactory, InMaterialRenderProxy, InMaterialResource, false, false)
+	{
+		PixelShader = InMaterialResource.GetShader<FLandscapeGrassWeightPS>(InVertexFactory->GetType());
+		VertexShader = InMaterialResource.GetShader<FLandscapeGrassWeightVS>(VertexFactory->GetType());
+	}
+
+	// FMeshDrawingPolicy interface.
+	bool Matches(const FLandscapeGrassWeightDrawingPolicy& Other) const
+	{
+		return FMeshDrawingPolicy::Matches(Other)
+			&& VertexShader == Other.VertexShader
+			&& PixelShader == Other.PixelShader;
+	}
+
+	void SetSharedState(FRHICommandList& RHICmdList, const FSceneView* View, const ContextDataType PolicyContext) const
+	{
+		// Set the depth-only shader parameters for the material.
+		VertexShader->SetParameters(RHICmdList, MaterialRenderProxy, *MaterialResource, *View);
+		PixelShader->SetParameters(RHICmdList, MaterialRenderProxy, *MaterialResource, View);
+
+		// Set the shared mesh resources.
+		FMeshDrawingPolicy::SetSharedState(RHICmdList, View, PolicyContext);
+	}
+
+	FBoundShaderStateInput GetBoundShaderStateInput(ERHIFeatureLevel::Type InFeatureLevel)
+	{
+		return FBoundShaderStateInput(
+			FMeshDrawingPolicy::GetVertexDeclaration(),
+			VertexShader->GetVertexShader(),
+			FHullShaderRHIRef(),
+			FDomainShaderRHIRef(),
+			PixelShader->GetPixelShader(),
+			NULL);
+	}
+
+	void SetMeshRenderState(
+		FRHICommandList& RHICmdList,
+		const FSceneView& View,
+		const FPrimitiveSceneProxy* PrimitiveSceneProxy,
+		const FMeshBatch& Mesh,
+		int32 BatchElementIndex,
+		bool bBackFace,
+		const ElementDataType& ElementData,
+		const ContextDataType PolicyContext
+		) const
+	{
+		const FMeshBatchElement& BatchElement = Mesh.Elements[BatchElementIndex];
+		VertexShader->SetMesh(RHICmdList, VertexFactory, View, PrimitiveSceneProxy, BatchElement);
+		PixelShader->SetMesh(RHICmdList, VertexFactory, View, PrimitiveSceneProxy, BatchElement);
+		FMeshDrawingPolicy::SetMeshRenderState(RHICmdList, View, PrimitiveSceneProxy, Mesh, BatchElementIndex, bBackFace, ElementData, PolicyContext);
+	}
+
+	friend int32 CompareDrawingPolicy(const FLandscapeGrassWeightDrawingPolicy& A, const FLandscapeGrassWeightDrawingPolicy& B)
+	{
+		COMPAREDRAWINGPOLICYMEMBERS(VertexShader);
+		COMPAREDRAWINGPOLICYMEMBERS(PixelShader);
+		COMPAREDRAWINGPOLICYMEMBERS(VertexFactory);
+		COMPAREDRAWINGPOLICYMEMBERS(MaterialRenderProxy);
+		COMPAREDRAWINGPOLICYMEMBERS(bIsTwoSidedMaterial);
+		return 0;
+	}
+private:
+	FLandscapeGrassWeightVS* VertexShader;
+	FLandscapeGrassWeightPS* PixelShader;
+};
+
+
+class FLandscapeGrassWeightExporter
+{
+	static void RenderLandscapeComponentToTexture(
+		ULandscapeComponent* LandscapeComponent,
+		const FMatrix& ViewMatrix,
+		const FMatrix& ProjectionMatrix,
+		FIntPoint TargetSize,
+		float TargetGamma,
+		TArray<FColor>& OutSamples)
+	{
+		UTextureRenderTarget2D* RenderTargetTexture = new UTextureRenderTarget2D(FObjectInitializer());
+		check(RenderTargetTexture);
+		RenderTargetTexture->AddToRoot();
+		RenderTargetTexture->ClearColor = FLinearColor::Transparent;
+		RenderTargetTexture->TargetGamma = TargetGamma;
+		RenderTargetTexture->InitCustomFormat(TargetSize.X, TargetSize.Y, PF_B8G8R8A8, false); 
+		FTextureRenderTarget2DResource* RenderTargetResource = RenderTargetTexture->GameThread_GetRenderTargetResource()->GetTextureRenderTarget2DResource();
+		RenderTargetResource->GetRenderTargetTexture();
+
+		FSceneViewFamilyContext ViewFamily(
+			FSceneViewFamily::ConstructionValues(RenderTargetResource, NULL, FEngineShowFlags(ESFIM_Game))
+			.SetWorldTimes(FApp::GetCurrentTime() - GStartTime, FApp::GetDeltaTime(), FApp::GetCurrentTime() - GStartTime)
+			);
+
+		// To enable visualization mode
+		ViewFamily.EngineShowFlags.SetPostProcessing(true);
+		ViewFamily.EngineShowFlags.SetVisualizeBuffer(true);
+		ViewFamily.EngineShowFlags.SetTonemapper(false);
+
+		// Force LOD render
+#if WITH_EDITOR
+		ViewFamily.LandscapeLODOverride = 0;
+#endif
+		FSceneViewInitOptions ViewInitOptions;
+		ViewInitOptions.SetViewRectangle(FIntRect(0, 0, TargetSize.X, TargetSize.Y));
+		ViewInitOptions.ViewFamily = &ViewFamily;
+		ViewInitOptions.ViewMatrix = ViewMatrix;
+		ViewInitOptions.ProjectionMatrix = ProjectionMatrix;
+
+		FCanvas Canvas(RenderTargetResource, NULL, FApp::GetCurrentTime() - GStartTime, FApp::GetDeltaTime(), FApp::GetCurrentTime() - GStartTime, GMaxRHIFeatureLevel);
+		Canvas.Clear(FLinearColor::Transparent);
+
+		GetRendererModule().CreateAndInitSingleView(&ViewFamily, &ViewInitOptions);
+		
+		// render
+		ENQUEUE_UNIQUE_RENDER_COMMAND_TWOPARAMETER(
+			FDrawSceneCommand,
+			FSceneViewFamily&, ViewFamily, ViewFamily,
+			FPrimitiveSceneInfo*, PrimitiveSceneInfo, LandscapeComponent->SceneProxy->GetPrimitiveSceneInfo(),
+			{
+				RenderLandscapeComponentToTexture_RenderThread(RHICmdList, ViewFamily, PrimitiveSceneInfo);
+				FlushPendingDeleteRHIResources_RenderThread();
+			});
+
+		// Copy the contents of the remote texture to system memory
+		OutSamples.Init(TargetSize.X*TargetSize.Y);
+		FReadSurfaceDataFlags ReadSurfaceDataFlags;
+		ReadSurfaceDataFlags.SetLinearToGamma(false);
+		RenderTargetResource->ReadPixels(OutSamples, ReadSurfaceDataFlags, FIntRect(0, 0, TargetSize.X, TargetSize.Y));
+
+		RenderTargetTexture->RemoveFromRoot();
+		RenderTargetTexture = nullptr;
+	}
+
+	static void RenderLandscapeComponentToTexture_RenderThread(FRHICommandListImmediate& RHICmdList, FSceneViewFamily& ViewFamily, FPrimitiveSceneInfo* PrimitiveSceneInfo)
+	{
+		const FSceneView* View = ViewFamily.Views[0];
+		RHICmdList.SetViewport(View->ViewRect.Min.X, View->ViewRect.Min.Y, 0.0f, View->ViewRect.Max.X, View->ViewRect.Max.Y, 1.0f);
+		RHICmdList.SetBlendState(TStaticBlendState<>::GetRHI());
+		RHICmdList.SetRasterizerState(TStaticRasterizerState<FM_Solid, CM_None>::GetRHI());
+		RHICmdList.SetDepthStencilState(TStaticDepthStencilState<false, CF_Always>::GetRHI());
+		RHICmdList.SetScissorRect(false, 0, 0, 0, 0);
+
+		for (int32 StaticMeshIdx = 0; StaticMeshIdx < PrimitiveSceneInfo->StaticMeshes.Num(); StaticMeshIdx++)
+		{
+			FMeshBatch& Mesh = *(FMeshBatch*)(&PrimitiveSceneInfo->StaticMeshes[StaticMeshIdx]);
+
+			FLandscapeGrassWeightDrawingPolicy DrawingPolicy(Mesh.VertexFactory, Mesh.MaterialRenderProxy, *Mesh.MaterialRenderProxy->GetMaterial(GMaxRHIFeatureLevel));
+			RHICmdList.BuildAndSetLocalBoundShaderState(DrawingPolicy.GetBoundShaderStateInput(GMaxRHIFeatureLevel));
+			DrawingPolicy.SetSharedState(RHICmdList, View, FLandscapeGrassWeightDrawingPolicy::ContextDataType());
+
+			// The first batch element contains the grass batch for the entire component
+			DrawingPolicy.SetMeshRenderState(RHICmdList, *View, PrimitiveSceneInfo->Proxy, Mesh, 0, false, FMeshDrawingPolicy::ElementDataType(), FLandscapeGrassWeightDrawingPolicy::ContextDataType());
+			DrawingPolicy.DrawMesh(RHICmdList, Mesh, 0);
+		}
+	}
+
+public:
+	
+	static bool GetLandscapeComponentGrassMap(ULandscapeComponent* LandscapeComponent, TArray<FColor>& OutSamples)
+	{
+		check(LandscapeComponent);
+
+		if (!LandscapeComponent->SceneProxy)
+		{
+			return false;
+		}
+		ALandscapeProxy* LandscapeProxy = LandscapeComponent->GetLandscapeProxy();
+
+		FIntRect LandscapeRect = FIntRect(0, 0, LandscapeComponent->ComponentSizeQuads, LandscapeComponent->ComponentSizeQuads);
+		LandscapeRect += LandscapeComponent->GetSectionBase() - LandscapeProxy->LandscapeSectionOffset;
+
+		FVector MidPoint = FVector(LandscapeRect.Min, 0.f) + FVector(LandscapeRect.Size(), 0.f)*0.5f;
+		FVector LandscapeCenter = LandscapeProxy->GetTransform().TransformPosition(MidPoint);
+		FVector LandscapeExtent = FVector(LandscapeRect.Size() + FIntPoint(1,1), 0.f)*LandscapeProxy->GetActorScale()*0.5f;
+
+		FMatrix ViewMatrix = FTranslationMatrix(-LandscapeCenter);
+		ViewMatrix *= FInverseRotationMatrix(LandscapeProxy->GetActorRotation());
+		ViewMatrix *= FMatrix(FPlane(1, 0, 0, 0),
+			FPlane(0, -1, 0, 0),
+			FPlane(0, 0, -1, 0),
+			FPlane(0, 0, 0, 1));
+
+		const float ZOffset = WORLD_MAX;
+		FMatrix ProjectionMatrix = FReversedZOrthoMatrix(
+			LandscapeExtent.X,
+			LandscapeExtent.Y,
+			0.5f / ZOffset,
+			ZOffset);
+
+		FIntPoint TargetSize(LandscapeComponent->ComponentSizeQuads + 1, LandscapeComponent->ComponentSizeQuads + 1);
+
+		//TargetSize *= 8;
+
+		const float Gamma = 1.0f;
+		RenderLandscapeComponentToTexture(LandscapeComponent, ViewMatrix, ProjectionMatrix, TargetSize, Gamma, OutSamples);
+		return true;
+	}
+};
+
+#if WITH_EDITOR
+void ULandscapeComponent::RenderGrassMap()
+{
+	UMaterialInterface* Material = GetLandscapeMaterial();
+	if (GIsEditor && !GUsingNullRHI && SceneProxy && Material)
+	{
+		FLandscapeComponentGrassData* NewGrassData = new FLandscapeComponentGrassData();
+
+		TArray<FColor> Samples;
+		FLandscapeGrassWeightExporter::GetLandscapeComponentGrassMap(this, Samples);
+
+		TArray<const UMaterialExpressionLandscapeGrassOutput*> GrassExpressions;
+		Material->GetMaterial()->GetAllExpressionsOfType<UMaterialExpressionLandscapeGrassOutput>(GrassExpressions);
+		if (GrassExpressions.Num() > 0)
+		{
+			int NumTypes = FMath::Min<int32>(GrassExpressions[0]->GrassTypes.Num(), 2); // Only 2 layers currently supported
+			if (NumTypes > 0)
+			{
+				TArray<uint8>* GrassWeightArrays[2] = { nullptr, nullptr };
+
+				// Initialize arrays
+				NewGrassData->HeightData.Empty(Samples.Num());
+				for (int32 Index = 0; Index < NumTypes; Index++)
+				{
+					GrassWeightArrays[Index] = &NewGrassData->WeightData.Add(GrassExpressions[0]->GrassTypes[Index].GrassType);
+					GrassWeightArrays[Index]->Empty(Samples.Num());
+				}
+
+				// Fill data
+				for (FColor& Sample : Samples)
+				{
+					uint16 Height = (((uint16)Sample.R) << 8) + (uint16)(Sample.G);
+					check((Height >> 8) == Sample.R);
+					check((Height & 255) == Sample.G);
+					NewGrassData->HeightData.Add(Height);
+					GrassWeightArrays[0]->Add(Sample.B);
+					if (NumTypes>1)
+					{
+						GrassWeightArrays[1]->Add(Sample.A);
+					}
+				}
+			}
+		}
+
+		FArchive* Ar = IFileManager::Get().CreateFileWriter(TEXT("D:\\Heightmaps\\Rendered.r16"), FILEWRITE_EvenIfReadOnly);
+		if (Ar)
+		{
+			Ar->Serialize(NewGrassData->HeightData.GetData(), NewGrassData->HeightData.Num() * sizeof(uint16));
+			delete Ar;
+		}
+
+		// Assign the new data (thread-safe)
+		GrassData = MakeShareable(NewGrassData);
+	}
+}
+
+void ULandscapeComponent::RemoveGrassMap()
+{
+	GrassData = MakeShareable(new FLandscapeComponentGrassData());
+}
+
+#endif
+
+
 //
 // UMaterialExpressionLandscapeGrassOutput
 //
@@ -100,6 +445,7 @@ void UMaterialExpressionLandscapeGrassOutput::PostEditChangeProperty(FPropertyCh
 //
 // ULandscapeGrassType
 //
+
 ULandscapeGrassType::ULandscapeGrassType(const FObjectInitializer& ObjectInitializer)
 : Super(ObjectInitializer)
 {
@@ -146,5 +492,14 @@ void ULandscapeGrassType::PostEditChangeProperty(FPropertyChangedEvent& Property
 }
 #endif
 
+//
+// FLandscapeComponentGrassData
+//
+FArchive& operator<<(FArchive& Ar, FLandscapeComponentGrassData& Data)
+{
+	Data.HeightData.BulkSerialize(Ar);
+	// Each weight data array, being 1 byte will be serialized in bulk.
+	return Ar << Data.WeightData;
+}
 
-#undef LOCTEXT_NAMESPACE
+#undef LOCTEXT_NAMESPACE
\ No newline at end of file
diff --git a/Engine/Source/Runtime/Landscape/Private/LandscapeRender.cpp b/Engine/Source/Runtime/Landscape/Private/LandscapeRender.cpp
--- a/Engine/Source/Runtime/Landscape/Private/LandscapeRender.cpp
+++ b/Engine/Source/Runtime/Landscape/Private/LandscapeRender.cpp
@@ -1026,10 +1026,15 @@ void FLandscapeComponentSceneProxy::DrawStaticElements(FStaticPrimitiveDrawInter
 	int32 FirstLOD = (ForcedLOD >= 0) ? FMath::Min<int32>(ForcedLOD, MaxLOD) : FMath::Max<int32>(LODBias, 0);
 	int32 LastLOD = (ForcedLOD >= 0) ? FirstLOD : FMath::Min<int32>(MaxLOD, MaxLOD + LODBias);
 
-	StaticBatchParamArray.Empty((1 + LastLOD - FirstLOD) * (FMath::Square(NumSubsections) + 1));
+	int32 NumBatches = (1 + LastLOD - FirstLOD) * (FMath::Square(NumSubsections) + 1);
+	if (SharedBuffers->GrassIndexBuffer)
+	{
+		NumBatches++;
+	}
+	StaticBatchParamArray.Empty(NumBatches);
 
 	FMeshBatch MeshBatch;
-	MeshBatch.Elements.Empty((1 + LastLOD - FirstLOD) * (FMath::Square(NumSubsections) + 1));
+	MeshBatch.Elements.Empty(NumBatches);
 
 	FMaterialRenderProxy* RenderProxy = MaterialInterface->GetRenderProxy(false);
 
@@ -1044,6 +1049,26 @@ void FLandscapeComponentSceneProxy::DrawStaticElements(FStaticPrimitiveDrawInter
 	MeshBatch.Type = bCurrentRequiresAdjacencyInformation ? PT_12_ControlPointPatchList : PT_TriangleList;
 	MeshBatch.DepthPriorityGroup = SDPG_World;
 
+	if (SharedBuffers->GrassIndexBuffer)
+	{
+		// Combined grass rendering batch element
+		FMeshBatchElement* BatchElement = new(MeshBatch.Elements) FMeshBatchElement;
+		FLandscapeBatchElementParams* BatchElementParams = new(StaticBatchParamArray)FLandscapeBatchElementParams;
+		BatchElementParams->LocalToWorldNoScalingPtr = &LocalToWorldNoScaling;
+		BatchElement->UserData = BatchElementParams;
+		BatchElement->PrimitiveUniformBufferResource = &GetUniformBuffer();
+		BatchElementParams->LandscapeUniformShaderParametersResource = &LandscapeUniformShaderParameters;
+		BatchElementParams->SceneProxy = this;
+		BatchElementParams->SubX = -1;
+		BatchElementParams->SubY = -1;
+		BatchElementParams->CurrentLOD = 0;
+		BatchElement->IndexBuffer = SharedBuffers->GrassIndexBuffer;
+		BatchElement->NumPrimitives = (FMath::Square(NumSubsections) * FMath::Square(SubsectionSizeQuads) + 2 * NumSubsections * SubsectionSizeQuads + 1) * 2;
+		BatchElement->FirstIndex = 0;
+		BatchElement->MinVertexIndex = 0;
+		BatchElement->MaxVertexIndex = SharedBuffers->NumVertices-1;
+	}
+	
 	for (int32 LOD = FirstLOD; LOD <= LastLOD; LOD++)
 	{
 		int32 LodSubsectionSizeVerts = SubsectionSizeVerts >> LOD;
@@ -1117,12 +1142,13 @@ uint64 FLandscapeVertexFactory::GetStaticBatchElementVisibility(const class FSce
 uint64 FLandscapeComponentSceneProxy::GetStaticBatchElementVisibility(const class FSceneView& View, const struct FMeshBatch* Batch) const
 {
 	uint64 BatchesToRenderMask = 0;
+	int32 GrassBatchOffset = SharedBuffers->GrassIndexBuffer ? 1 : 0;
 
 	SCOPE_CYCLE_COUNTER(STAT_LandscapeStaticDrawLODTime);
 	if (ForcedLOD >= 0)
 	{
 		// When forcing LOD we only create one Batch Element, so not sure this loop is necessary?
-		for (int32 BatchElementIndex = 0; BatchElementIndex < Batch->Elements.Num(); BatchElementIndex++)
+		for (int32 BatchElementIndex = GrassBatchOffset; BatchElementIndex < Batch->Elements.Num(); BatchElementIndex++)
 		{
 			BatchesToRenderMask |= (((uint64)1) << BatchElementIndex);
 			INC_DWORD_STAT(STAT_LandscapeDrawCalls);
@@ -1164,7 +1190,7 @@ uint64 FLandscapeComponentSceneProxy::GetStaticBatchElementVisibility(const clas
 		if (bAllSameLOD && NumSubsections > 1 && !GLandscapeDebugOptions.bDisableCombine)
 		{
 			// choose the combined batch element
-			int32 BatchElementIndex = (CombinedLOD + 1)*BatchesPerLOD - 1;
+			int32 BatchElementIndex = GrassBatchOffset + (CombinedLOD + 1)*BatchesPerLOD - 1;
 			BatchesToRenderMask |= (((uint64)1) << BatchElementIndex);
 			INC_DWORD_STAT(STAT_LandscapeDrawCalls);
 			INC_DWORD_STAT_BY(STAT_LandscapeTriangles, Batch->Elements[BatchElementIndex].NumPrimitives);
@@ -1175,7 +1201,7 @@ uint64 FLandscapeComponentSceneProxy::GetStaticBatchElementVisibility(const clas
 			{
 				for (int32 SubX = 0; SubX < NumSubsections; SubX++)
 				{
-					int32 BatchElementIndex = CalculatedLods[SubX][SubY] * BatchesPerLOD + SubY*NumSubsections + SubX;
+					int32 BatchElementIndex = GrassBatchOffset + CalculatedLods[SubX][SubY] * BatchesPerLOD + SubY*NumSubsections + SubX;
 					BatchesToRenderMask |= (((uint64)1) << BatchElementIndex);
 					INC_DWORD_STAT(STAT_LandscapeDrawCalls);
 					INC_DWORD_STAT_BY(STAT_LandscapeTriangles, Batch->Elements[BatchElementIndex].NumPrimitives);
@@ -1656,27 +1682,30 @@ void FLandscapeVertexBuffer::InitRHI()
 {
 	// create a static vertex buffer
 	FRHIResourceCreateInfo CreateInfo;
-	VertexBufferRHI = RHICreateVertexBuffer(FMath::Square(SubsectionSizeVerts) * FMath::Square(NumSubsections) * sizeof(FLandscapeVertex), BUF_Static, CreateInfo);
-	FLandscapeVertex* Vertex = (FLandscapeVertex*)RHILockVertexBuffer(VertexBufferRHI, 0, FMath::Square(SubsectionSizeVerts) * FMath::Square(NumSubsections) * sizeof(FLandscapeVertex), RLM_WriteOnly);
-
+	VertexBufferRHI = RHICreateVertexBuffer(NumVertices * sizeof(FLandscapeVertex), BUF_Static, CreateInfo);
+	FLandscapeVertex* Vertex = (FLandscapeVertex*)RHILockVertexBuffer(VertexBufferRHI, 0, NumVertices * sizeof(FLandscapeVertex), RLM_WriteOnly);
+	int32 VertexIndex = 0;
 	for (int32 SubY = 0; SubY < NumSubsections; SubY++)
 	{
 		for (int32 SubX = 0; SubX < NumSubsections; SubX++)
 		{
-			for (int32 y = 0; y < SubsectionSizeVerts; y++)
+			int32 Cols = SubX == NumSubsections - 1 && FeatureLevel > ERHIFeatureLevel::ES3_1 ? SubsectionSizeVerts + 1 : SubsectionSizeVerts;
+			int32 Rows = SubY == NumSubsections - 1 && FeatureLevel > ERHIFeatureLevel::ES3_1 ? SubsectionSizeVerts + 1 : SubsectionSizeVerts;
+			for (int32 y = 0; y < Rows; y++)
 			{
-				for (int32 x = 0; x < SubsectionSizeVerts; x++)
+				for (int32 x = 0; x < Cols; x++)
 				{
 					Vertex->VertexX = x;
 					Vertex->VertexY = y;
 					Vertex->SubX = SubX;
 					Vertex->SubY = SubY;
 					Vertex++;
+					VertexIndex++;
 				}
 			}
 		}
 	}
-
+	check(NumVertices == VertexIndex);
 	RHIUnlockVertexBuffer(VertexBufferRHI);
 }
 
@@ -1832,25 +1861,26 @@ void FLandscapeSharedBuffers::CreateIndexBuffers(ERHIFeatureLevel::Type InFeatur
 		else
 		{
 			// non-ES2 version
+			int SubOffset = 0;
 			for (int32 SubY = 0; SubY < NumSubsections; SubY++)
 			{
 				for (int32 SubX = 0; SubX < NumSubsections; SubX++)
 				{
-					int32 SubOffset = (SubX + SubY * NumSubsections) * FMath::Square(SubsectionSizeVerts);
-
 					int32& MaxIndex = IndexRanges[Mip].MaxIndex[SubX][SubY];
 					int32& MinIndex = IndexRanges[Mip].MinIndex[SubX][SubY];
 					MaxIndex = 0;
 					MinIndex = MAX_int32;
 
+					int32 StrideX = SubX == NumSubsections - 1 ? SubsectionSizeVerts + 1 : SubsectionSizeVerts;
+					int32 StrideY = SubY == NumSubsections - 1 ? SubsectionSizeVerts + 1 : SubsectionSizeVerts;
 					for (int32 y = 0; y < LodSubsectionSizeQuads; y++)
 					{
 						for (int32 x = 0; x < LodSubsectionSizeQuads; x++)
 						{
-							INDEX_TYPE i00 = (x + 0) + (y + 0) * SubsectionSizeVerts + SubOffset;
-							INDEX_TYPE i10 = (x + 1) + (y + 0) * SubsectionSizeVerts + SubOffset;
-							INDEX_TYPE i11 = (x + 1) + (y + 1) * SubsectionSizeVerts + SubOffset;
-							INDEX_TYPE i01 = (x + 0) + (y + 1) * SubsectionSizeVerts + SubOffset;
+							INDEX_TYPE i00 = (x + 0) + (y + 0) * StrideX + SubOffset;
+							INDEX_TYPE i10 = (x + 1) + (y + 0) * StrideX + SubOffset;
+							INDEX_TYPE i11 = (x + 1) + (y + 1) * StrideX + SubOffset;
+							INDEX_TYPE i01 = (x + 0) + (y + 1) * StrideX + SubOffset;
 
 							NewIndices.Add(i00);
 							NewIndices.Add(i11);
@@ -1875,6 +1905,8 @@ void FLandscapeSharedBuffers::CreateIndexBuffers(ERHIFeatureLevel::Type InFeatur
 					// update min/max for full subsection
 					MaxIndexFull = FMath::Max<int32>(MaxIndexFull, MaxIndex);
 					MinIndexFull = FMath::Min<int32>(MinIndexFull, MinIndex);
+
+					SubOffset += StrideX * StrideY;
 				}
 			}
 
@@ -1900,6 +1932,56 @@ void FLandscapeSharedBuffers::CreateIndexBuffers(ERHIFeatureLevel::Type InFeatur
 	}
 }
 
+template <typename INDEX_TYPE>
+void FLandscapeSharedBuffers::CreateGrassIndexBuffer()
+{
+	int32 SubsectionSizeQuads = SubsectionSizeVerts - 1;
+
+	TArray<INDEX_TYPE> NewIndices;
+	int32 ExpectedNumIndices = (FMath::Square(NumSubsections) * FMath::Square(SubsectionSizeQuads) + 2 * NumSubsections * SubsectionSizeQuads + 1) * 6;
+	NewIndices.Empty(ExpectedNumIndices);
+
+	int32 SubOffset = 0;
+	for (int32 SubY = 0; SubY < NumSubsections; SubY++)
+	{
+		for (int32 SubX = 0; SubX < NumSubsections; SubX++)
+		{
+			int32 Cols = SubX == NumSubsections - 1 ? SubsectionSizeQuads + 1 : SubsectionSizeQuads;
+			int32 Rows = SubY == NumSubsections - 1 ? SubsectionSizeQuads + 1 : SubsectionSizeQuads;
+			int32 StrideX = SubX == NumSubsections - 1 ? SubsectionSizeVerts + 1 : SubsectionSizeVerts;
+			int32 StrideY = SubY == NumSubsections - 1 ? SubsectionSizeVerts + 1 : SubsectionSizeVerts;
+			for (int32 y = 0; y < Rows; y++)
+			{
+				for (int32 x = 0; x < Cols; x++)
+				{
+					INDEX_TYPE i00 = (x + 0) + (y + 0) * StrideX + SubOffset;
+					INDEX_TYPE i10 = (x + 1) + (y + 0) * StrideX + SubOffset;
+					INDEX_TYPE i11 = (x + 1) + (y + 1) * StrideX + SubOffset;
+					INDEX_TYPE i01 = (x + 0) + (y + 1) * StrideX + SubOffset;
+
+					NewIndices.Add(i00);
+					NewIndices.Add(i11);
+					NewIndices.Add(i10);
+
+					NewIndices.Add(i00);
+					NewIndices.Add(i01);
+					NewIndices.Add(i11);
+				}
+			}
+
+			SubOffset += StrideX * StrideY;
+		}
+	}
+
+	check(NewIndices.Num() == ExpectedNumIndices);
+
+	// Create and init new index buffer with index data
+	FRawStaticIndexBuffer16or32<INDEX_TYPE>* IndexBuffer = new FRawStaticIndexBuffer16or32<INDEX_TYPE>(false);
+	IndexBuffer->AssignNewBuffer(NewIndices);
+	IndexBuffer->InitResource();
+	GrassIndexBuffer = IndexBuffer;
+}
+
 FLandscapeSharedBuffers::FLandscapeSharedBuffers(int32 InSharedBuffersKey, int32 InSubsectionSizeQuads, int32 InNumSubsections, ERHIFeatureLevel::Type InFeatureLevel, bool bRequiresAdjacencyInformation)
 	: SharedBuffersKey(InSharedBuffersKey)
 	, NumIndexBuffers(FMath::CeilLogTwo(InSubsectionSizeQuads + 1))
@@ -1909,25 +1991,36 @@ FLandscapeSharedBuffers::FLandscapeSharedBuffers(int32 InSharedBuffersKey, int32
 	, VertexBuffer(nullptr)
 	, AdjacencyIndexBuffers(nullptr)
 	, bUse32BitIndices(false)
+	, GrassIndexBuffer(nullptr)
 {
+	NumVertices = FMath::Square(SubsectionSizeVerts) * FMath::Square(NumSubsections);
 	if (InFeatureLevel > ERHIFeatureLevel::ES3_1)
 	{
+		NumVertices += 2 * NumSubsections * SubsectionSizeVerts + 1;
 		// Vertex Buffer cannot be shared
-		VertexBuffer = new FLandscapeVertexBuffer(SubsectionSizeVerts, NumSubsections);
+		VertexBuffer = new FLandscapeVertexBuffer(InFeatureLevel, NumVertices, SubsectionSizeVerts, NumSubsections);
 	}
 	IndexBuffers = new FIndexBuffer*[NumIndexBuffers];
 	FMemory::Memzero(IndexBuffers, sizeof(FIndexBuffer*)* NumIndexBuffers);
 	IndexRanges = new FLandscapeIndexRanges[NumIndexBuffers]();
 
 	// See if we need to use 16 or 32-bit index buffers
-	if (FMath::Square(SubsectionSizeVerts) * FMath::Square(NumSubsections) > 65535)
+	if (NumVertices > 65535)
 	{
 		bUse32BitIndices = true;
 		CreateIndexBuffers<uint32>(InFeatureLevel, bRequiresAdjacencyInformation);
+		if (InFeatureLevel > ERHIFeatureLevel::ES3_1)
+		{
+			CreateGrassIndexBuffer<uint32>();
+		}
 	}
 	else
 	{
 		CreateIndexBuffers<uint16>(InFeatureLevel, bRequiresAdjacencyInformation);
+		if (InFeatureLevel > ERHIFeatureLevel::ES3_1)
+		{
+			CreateGrassIndexBuffer<uint16>();
+		}
 	}
 }
 
@@ -1943,6 +2036,12 @@ FLandscapeSharedBuffers::~FLandscapeSharedBuffers()
 	delete[] IndexBuffers;
 	delete[] IndexRanges;
 
+	if (GrassIndexBuffer)
+	{
+		GrassIndexBuffer->ReleaseResource();
+		delete GrassIndexBuffer;
+	}
+
 	if (AdjacencyIndexBuffers)
 	{
 		if (AdjacencyIndexBuffers->Release() == 0)
@@ -2000,7 +2099,7 @@ FLandscapeSharedAdjacencyIndexBuffer::FLandscapeSharedAdjacencyIndexBuffer(FLand
 	// Currently only support PN-AEN-Dominant Corner, which is the only mode for UE4 for now
 	IndexBuffers.Empty(Buffers->NumIndexBuffers);
 
-	bool b32BitIndex = FMath::Square(Buffers->SubsectionSizeVerts) * FMath::Square(Buffers->NumSubsections) > 65535;
+	bool b32BitIndex = Buffers->NumVertices > 65535;
 	for (int32 i = 0; i < Buffers->NumIndexBuffers; ++i)
 	{
 		if (b32BitIndex)
diff --git a/Engine/Source/Runtime/Landscape/Public/LandscapeRender.h b/Engine/Source/Runtime/Landscape/Public/LandscapeRender.h
--- a/Engine/Source/Runtime/Landscape/Public/LandscapeRender.h
+++ b/Engine/Source/Runtime/Landscape/Public/LandscapeRender.h
@@ -228,13 +228,17 @@ struct FLandscapeVertex
 //
 class FLandscapeVertexBuffer : public FVertexBuffer
 {
+	ERHIFeatureLevel::Type FeatureLevel;
+	int32 NumVertices;
 	int32 SubsectionSizeVerts;
 	int32 NumSubsections;
 public:
 
 	/** Constructor. */
-	FLandscapeVertexBuffer(int32 InSubsectionSizeVerts, int32 InNumSubsections)
-		: SubsectionSizeVerts(InSubsectionSizeVerts)
+	FLandscapeVertexBuffer(ERHIFeatureLevel::Type InFeatureLevel, int32 InNumVertices, int32 InSubsectionSizeVerts, int32 InNumSubsections)
+		: FeatureLevel(InFeatureLevel)
+		, NumVertices(InNumVertices)
+		, SubsectionSizeVerts(InSubsectionSizeVerts)
 		, NumSubsections(InNumSubsections)
 	{
 		InitResource();
@@ -279,6 +283,7 @@ public:
 		int32 MaxIndexFull;
 	};
 
+	int32 NumVertices;
 	int32 SharedBuffersKey;
 	int32 NumIndexBuffers;
 	int32 SubsectionSizeVerts;
@@ -290,11 +295,14 @@ public:
 	FLandscapeIndexRanges* IndexRanges;
 	FLandscapeSharedAdjacencyIndexBuffer* AdjacencyIndexBuffers;
 	bool bUse32BitIndices;
+	FIndexBuffer* GrassIndexBuffer;
 
 	FLandscapeSharedBuffers(int32 SharedBuffersKey, int32 SubsectionSizeQuads, int32 NumSubsections, ERHIFeatureLevel::Type InFeatureLevel, bool bRequiresAdjacencyInformation);
 
 	template <typename INDEX_TYPE>
 	void CreateIndexBuffers(ERHIFeatureLevel::Type InFeatureLevel, bool bRequiresAdjacencyInformation);
+	template <typename INDEX_TYPE>
+	void CreateGrassIndexBuffer();
 
 	virtual ~FLandscapeSharedBuffers();
 };
diff --git a/Engine/Source/Runtime/Renderer/Private/DrawingPolicy.h b/Engine/Source/Runtime/Renderer/Private/DrawingPolicy.h
deleted file mode 100644
diff --git a/Engine/Source/Runtime/Renderer/Private/PrimitiveSceneInfo.h b/Engine/Source/Runtime/Renderer/Private/PrimitiveSceneInfo.h
deleted file mode 100644
diff --git a/Engine/Source/Runtime/Renderer/Private/RendererPrivate.h b/Engine/Source/Runtime/Renderer/Private/RendererPrivate.h
--- a/Engine/Source/Runtime/Renderer/Private/RendererPrivate.h
+++ b/Engine/Source/Runtime/Renderer/Private/RendererPrivate.h
@@ -23,6 +23,7 @@ public:
 	virtual bool SupportsDynamicReloading() override { return true; }
 
 	virtual void BeginRenderingViewFamily(FCanvas* Canvas,FSceneViewFamily* ViewFamily) override;
+	virtual void CreateAndInitSingleView(class FSceneViewFamily* ViewFamily, const struct FSceneViewInitOptions* ViewInitOptions);
 	virtual FSceneInterface* AllocateScene(UWorld* World, bool bInRequiresHitProxies, bool bCreateFXSystem, ERHIFeatureLevel::Type InFeatureLevel) override;
 	virtual void RemoveScene(FSceneInterface* Scene) override;
 	virtual void UpdateStaticDrawListsForMaterials(const TArray<const FMaterial*>& Materials) override;
diff --git a/Engine/Source/Runtime/Renderer/Private/SceneRendering.cpp b/Engine/Source/Runtime/Renderer/Private/SceneRendering.cpp
--- a/Engine/Source/Runtime/Renderer/Private/SceneRendering.cpp
+++ b/Engine/Source/Runtime/Renderer/Private/SceneRendering.cpp
@@ -1330,6 +1330,23 @@ static void RenderViewFamily_RenderThread(FRHICommandListImmediate& RHICmdList,
 #endif
 }
 
+void FRendererModule::CreateAndInitSingleView(class FSceneViewFamily* ViewFamily, const struct FSceneViewInitOptions* ViewInitOptions)
+{
+	// Create and add the new view
+	FViewInfo* NewView = new FViewInfo(*ViewInitOptions);
+	ViewFamily->Views.Add(NewView);
+
+	// Ask the render thread to init its resources
+	ENQUEUE_UNIQUE_RENDER_COMMAND_ONEPARAMETER(
+		InitSetViewCommand,
+		FSceneViewFamily*, ViewFamily, ViewFamily,
+		{
+			SetRenderTarget(RHICmdList, ViewFamily->RenderTarget->GetRenderTargetTexture(), nullptr, ESimpleRenderTargetMode::EClearColorToWhite);
+			FViewInfo* View = (FViewInfo*)ViewFamily->Views[0];
+			View->InitRHIResources(nullptr);
+		});
+}
+
 void FRendererModule::BeginRenderingViewFamily(FCanvas* Canvas,FSceneViewFamily* ViewFamily)
 {
 	// Flush the canvas first.
diff --git a/Engine/Source/Runtime/Renderer/Private/ShaderBaseClasses.cpp b/Engine/Source/Runtime/Renderer/Private/ShaderBaseClasses.cpp
--- a/Engine/Source/Runtime/Renderer/Private/ShaderBaseClasses.cpp
+++ b/Engine/Source/Runtime/Renderer/Private/ShaderBaseClasses.cpp
@@ -419,7 +419,7 @@ void FMeshMaterialShader::SetMesh(
 }
 
 #define IMPLEMENT_MESH_MATERIAL_SHADER_SetMesh( ShaderRHIParamRef ) \
-	template void FMeshMaterialShader::SetMesh< ShaderRHIParamRef >( \
+	template RENDERER_API void FMeshMaterialShader::SetMesh< ShaderRHIParamRef >( \
 		FRHICommandList& RHICmdList,			\
 		const ShaderRHIParamRef ShaderRHI,		\
 		const FVertexFactory* VertexFactory,	\
diff --git a/Engine/Source/Runtime/Renderer/Public/DrawingPolicy.h b/Engine/Source/Runtime/Renderer/Public/DrawingPolicy.h
new file mode 100644
--- /dev/null
+++ b/Engine/Source/Runtime/Renderer/Public/DrawingPolicy.h
@@ -0,0 +1,143 @@
+// Copyright 1998-2015 Epic Games, Inc. All Rights Reserved.
+
+/*=============================================================================
+	DrawingPolicy.h: Drawing policy definitions.
+=============================================================================*/
+
+#ifndef __DRAWINGPOLICY_H__
+#define __DRAWINGPOLICY_H__
+
+#include "RHICommandList.h"
+
+
+/**
+ * A macro to compare members of two drawing policies(A and B), and return based on the result.
+ * If the members are the same, the macro continues execution rather than returning to the caller.
+ */
+#define COMPAREDRAWINGPOLICYMEMBERS(MemberName) \
+	if(A.MemberName < B.MemberName) { return -1; } \
+	else if(A.MemberName > B.MemberName) { return +1; }
+
+/**
+ * The base mesh drawing policy.  Subclasses are used to draw meshes with type-specific context variables.
+ * May be used either simply as a helper to render a dynamic mesh, or as a static instance shared between
+ * similar meshs.
+ */
+class RENDERER_API FMeshDrawingPolicy
+{
+public:
+	/** Per-element data required by the drawing policy that static mesh draw lists will cache. */
+	struct ElementDataType {};
+
+	/** Context data required by the drawing policy that is not known when caching policies in static mesh draw lists. */
+	struct ContextDataType {};
+
+	FMeshDrawingPolicy(
+		const FVertexFactory* InVertexFactory,
+		const FMaterialRenderProxy* InMaterialRenderProxy,
+		const FMaterial& InMaterialResource,
+		bool bInOverrideWithShaderComplexity = false,
+		bool bInTwoSidedOverride = false,
+		bool bInWireframeOverride = false
+		);
+
+	FMeshDrawingPolicy& operator = (const FMeshDrawingPolicy& Other)
+	{ 
+		VertexFactory = Other.VertexFactory;
+		MaterialRenderProxy = Other.MaterialRenderProxy;
+		MaterialResource = Other.MaterialResource;
+		bIsTwoSidedMaterial = Other.bIsTwoSidedMaterial;
+		bIsWireframeMaterial = Other.bIsWireframeMaterial;
+		bNeedsBackfacePass = Other.bNeedsBackfacePass;
+		bUsePositionOnlyVS = Other.bUsePositionOnlyVS;
+		bOverrideWithShaderComplexity = Other.bOverrideWithShaderComplexity;
+		return *this; 
+	}
+
+	uint32 GetTypeHash() const
+	{
+		return PointerHash(VertexFactory,PointerHash(MaterialRenderProxy));
+	}
+
+	bool Matches(const FMeshDrawingPolicy& OtherDrawer) const
+	{
+		return
+			VertexFactory == OtherDrawer.VertexFactory &&
+			MaterialRenderProxy == OtherDrawer.MaterialRenderProxy &&
+			bIsTwoSidedMaterial == OtherDrawer.bIsTwoSidedMaterial && 
+			bUsePositionOnlyVS == OtherDrawer.bUsePositionOnlyVS && 
+			bIsWireframeMaterial == OtherDrawer.bIsWireframeMaterial;
+	}
+
+	/**
+	 * Sets the render states for drawing a mesh.
+	 * @param PrimitiveSceneProxy - The primitive drawing the dynamic mesh.  If this is a view element, this will be NULL.
+	 */
+	void SetMeshRenderState(
+		FRHICommandList& RHICmdList, 
+		const FSceneView& View,
+		const FPrimitiveSceneProxy* PrimitiveSceneProxy,
+		const FMeshBatch& Mesh,
+		int32 BatchElementIndex,
+		bool bBackFace,
+		const ElementDataType& ElementData,
+		const ContextDataType PolicyContext
+		) const;
+
+	/**
+	 * Executes the draw commands for a mesh.
+	 */
+	void DrawMesh(FRHICommandList& RHICmdList, const FMeshBatch& Mesh, int32 BatchElementIndex) const;
+
+	/**
+	 * Executes the draw commands which can be shared between any meshes using this drawer.
+	 * @param CI - The command interface to execute the draw commands on.
+	 * @param View - The view of the scene being drawn.
+	 */
+	void SetSharedState(FRHICommandList& RHICmdList, const FSceneView* View, const ContextDataType PolicyContext) const;
+
+	/**
+	* Get the decl for this mesh policy type and vertexfactory
+	*/
+	const FVertexDeclarationRHIRef& GetVertexDeclaration() const;
+
+	friend int32 CompareDrawingPolicy(const FMeshDrawingPolicy& A,const FMeshDrawingPolicy& B)
+	{
+		COMPAREDRAWINGPOLICYMEMBERS(VertexFactory);
+		COMPAREDRAWINGPOLICYMEMBERS(MaterialRenderProxy);
+		COMPAREDRAWINGPOLICYMEMBERS(bIsTwoSidedMaterial);
+		return 0;
+	}
+
+	// Accessors.
+	bool IsTwoSided() const
+	{
+		return bIsTwoSidedMaterial;
+	}
+	bool IsWireframe() const
+	{
+		return bIsWireframeMaterial;
+	}
+	bool NeedsBackfacePass() const
+	{
+		return bNeedsBackfacePass;
+	}
+
+	const FVertexFactory* GetVertexFactory() const { return VertexFactory; }
+	const FMaterialRenderProxy* GetMaterialRenderProxy() const { return MaterialRenderProxy; }
+
+protected:
+	const FVertexFactory* VertexFactory;
+	const FMaterialRenderProxy* MaterialRenderProxy;
+	const FMaterial* MaterialResource;
+	uint32 bIsTwoSidedMaterial : 1;
+	uint32 bIsWireframeMaterial : 1;
+	uint32 bNeedsBackfacePass : 1;
+	uint32 bUsePositionOnlyVS : 1;
+	uint32 bOverrideWithShaderComplexity : 1;
+};
+
+
+uint32 GetTypeHash(const FBoundShaderStateRHIRef &Key);
+
+#endif
diff --git a/Engine/Source/Runtime/Renderer/Public/MeshMaterialShader.h b/Engine/Source/Runtime/Renderer/Public/MeshMaterialShader.h
--- a/Engine/Source/Runtime/Renderer/Public/MeshMaterialShader.h
+++ b/Engine/Source/Runtime/Renderer/Public/MeshMaterialShader.h
@@ -10,7 +10,7 @@
 #include "MeshMaterialShaderType.h"
 
 /** Base class of all shaders that need material and vertex factory parameters. */
-class FMeshMaterialShader : public FMaterialShader
+class RENDERER_API FMeshMaterialShader : public FMaterialShader
 {
 public:
 	FMeshMaterialShader() {}
diff --git a/Engine/Source/Runtime/Renderer/Public/PrimitiveSceneInfo.h b/Engine/Source/Runtime/Renderer/Public/PrimitiveSceneInfo.h
new file mode 100644
--- /dev/null
+++ b/Engine/Source/Runtime/Renderer/Public/PrimitiveSceneInfo.h
@@ -0,0 +1,371 @@
+// Copyright 1998-2015 Epic Games, Inc. All Rights Reserved.
+
+/*=============================================================================
+	PrimitiveSceneInfo.h: Primitive scene info definitions.
+=============================================================================*/
+
+#ifndef __PRIMITIVESCENEINFO_H__
+#define __PRIMITIVESCENEINFO_H__
+
+#include "GenericOctree.h"
+
+/** Data used to track a primitive's allocation in the volume texture atlas that stores indirect lighting. */
+class FIndirectLightingCacheAllocation
+{
+public:
+
+	FIndirectLightingCacheAllocation() :
+		Add(FVector(0, 0, 0)),
+		Scale(FVector(0, 0, 0)),
+		MinUV(FVector(0, 0, 0)),
+		MaxUV(FVector(0, 0, 0)),
+		MinTexel(FIntVector(-1, -1, -1)),
+		AllocationTexelSize(0),
+		TargetPosition(FVector(0, 0, 0)),
+		TargetDirectionalShadowing(1),
+		TargetSkyBentNormal(FVector4(0, 0, 1, 1)),
+		SingleSamplePosition(FVector(0, 0, 0)),
+		CurrentDirectionalShadowing(1),
+		CurrentSkyBentNormal(FVector4(0, 0, 1, 1)),
+		bHasEverUpdatedSingleSample(false),
+		bPointSample(true),
+		bIsDirty(false)
+	{
+		for (int32 VectorIndex = 0; VectorIndex < ARRAY_COUNT(TargetSamplePacked); VectorIndex++)
+		{
+			TargetSamplePacked[VectorIndex] = FVector4(0, 0, 0, 0);
+			SingleSamplePacked[VectorIndex] = FVector4(0, 0, 0, 0);
+		}
+	}
+
+	/** Add factor for calculating UVs from position. */
+	FVector Add;
+
+	/** Scale factor for calculating UVs from position. */
+	FVector Scale;
+
+	/** Used to clamp lookup UV to a valid range for pixels outside the object's bounding box. */
+	FVector MinUV;
+
+	/** Used to clamp lookup UV to a valid range for pixels outside the object's bounding box. */
+	FVector MaxUV;
+
+	/** Block index in the volume texture atlas, can represent unallocated. */
+	FIntVector MinTexel;
+
+	/** Size in texels of the allocation into the volume texture atlas. */
+	int32 AllocationTexelSize;
+
+	/** Position at the new single lighting sample. Used for interpolation over time. */
+	FVector TargetPosition;
+
+	/** SH sample at the new single lighting sample position. Used for interpolation over time. */
+	FVector4 TargetSamplePacked[sizeof(FSHVectorRGB2) / sizeof(FVector4)];
+
+	/** Target shadowing of the stationary directional light. */
+	float TargetDirectionalShadowing;
+
+	/** Target directional occlusion of the sky. */
+	FVector4 TargetSkyBentNormal;
+
+	/** Current position of the single lighting sample.  Used for interpolation over time. */
+	FVector SingleSamplePosition;
+
+	/** Current SH sample used when lighting the entire object with one sample. */
+	FVector4 SingleSamplePacked[sizeof(FSHVectorRGB2) / sizeof(FVector4)];
+
+	/** Current shadowing of the stationary directional light. */
+	float CurrentDirectionalShadowing;
+
+	/** Current directional occlusion of the sky. */
+	FVector4 CurrentSkyBentNormal;
+
+	/** Whether SingleSamplePacked has ever been populated with valid results, used to initialize. */
+	bool bHasEverUpdatedSingleSample;
+
+	/** Whether this allocation is a point sample and therefore was not put into the volume texture atlas. */
+	bool bPointSample;
+
+	/** Whether the primitive allocation is dirty and should be updated regardless of having moved. */
+	bool bIsDirty;
+
+	void SetDirty() 
+	{ 
+		bIsDirty = true; 
+	}
+
+	bool IsValid() const
+	{
+		return MinTexel.X >= 0 && MinTexel.Y >= 0 && MinTexel.Z >= 0 && AllocationTexelSize > 0;
+	}
+
+	void SetParameters(FIntVector InMinTexel, int32 InAllocationTexelSize, FVector InScale, FVector InAdd, FVector InMinUV, FVector InMaxUV, bool bInPointSample)
+	{
+		Add = InAdd;
+		Scale = InScale;
+		MinUV = InMinUV;
+		MaxUV = InMaxUV;
+		MinTexel = InMinTexel;
+		AllocationTexelSize = InAllocationTexelSize;
+		bIsDirty = false;
+		bPointSample = bInPointSample;
+	}
+};
+
+/** The information needed to determine whether a primitive is visible. */
+class FPrimitiveSceneInfoCompact
+{
+public:
+	FPrimitiveSceneInfo* PrimitiveSceneInfo;
+	FPrimitiveSceneProxy* Proxy;
+	FBoxSphereBounds Bounds;
+	float MinDrawDistance;
+	float MaxDrawDistance;
+	float LpvBiasMultiplier;
+	/** Used for precomputed visibility */
+	int32 VisibilityId;
+	uint32 bHasViewDependentDPG : 1;
+	uint32 bCastDynamicShadow : 1;
+	uint32 bAffectDynamicIndirectLighting : 1;
+	uint32 StaticDepthPriorityGroup : SDPG_NumBits;
+
+	/** Initializes the compact scene info from the primitive's full scene info. */
+	void Init(FPrimitiveSceneInfo* InPrimitiveSceneInfo);
+
+	/** Default constructor. */
+	FPrimitiveSceneInfoCompact():
+		PrimitiveSceneInfo(NULL),
+		Proxy(NULL)
+	{}
+
+	/** Initialization constructor. */
+	FPrimitiveSceneInfoCompact(FPrimitiveSceneInfo* InPrimitiveSceneInfo)
+	{
+		Init(InPrimitiveSceneInfo);
+	}
+};
+
+/** The type of the octree used by FScene to find primitives. */
+typedef TOctree<FPrimitiveSceneInfoCompact,struct FPrimitiveOctreeSemantics> FScenePrimitiveOctree;
+
+/**
+ * The renderer's internal state for a single UPrimitiveComponent.  This has a one to one mapping with FPrimitiveSceneProxy, which is in the engine module.
+ */
+class FPrimitiveSceneInfo : public FDeferredCleanupInterface
+{
+public:
+
+	/** The render proxy for the primitive. */
+	FPrimitiveSceneProxy* Proxy;
+
+	/** 
+	 * Id for the component this primitive belongs to.  
+	 * This will stay the same for the lifetime of the component, so it can be used to identify the component across re-registers.
+	 */
+	FPrimitiveComponentId PrimitiveComponentId;
+
+	/** 
+	 * Pointer to the primitive's last render time variable, which is written to by the RT and read by the GT.
+	 * The value of LastRenderTime will therefore not be deterministic due to race conditions, but the GT uses it in a way that allows this.
+	 * Storing a pointer to the UObject member variable only works because UPrimitiveComponent has a mechanism to ensure it does not get deleted before the proxy (DetachFence).
+	 * In general feedback from the renderer to the game thread like this should be avoided.
+	 */
+	float* ComponentLastRenderTime;
+
+	/** 
+	 * The root attachment component id for use with lighting, if valid.
+	 * If the root id is not valid, this is a parent primitive.
+	 */
+	FPrimitiveComponentId LightingAttachmentRoot;
+
+	/** 
+	 * The component id of the LOD parent if valid.
+	 */
+	FPrimitiveComponentId LODParentComponentId;
+
+	/** The primitive's static meshes. */
+	TIndirectArray<class FStaticMesh> StaticMeshes;
+
+	/** The identifier for the primitive in Scene->PrimitiveOctree. */
+	FOctreeElementId OctreeId;
+
+	/** 
+	 * Caches the primitive's indirect lighting cache allocation.
+	 * Note: This is only valid during the rendering of a frame, not just once the primitive is attached. 
+	 */
+	const FIndirectLightingCacheAllocation* IndirectLightingCacheAllocation;
+
+	/** 
+	 * Reflection capture proxy that was closest to this primitive, used for the forward shading rendering path. 
+	 */
+	const FReflectionCaptureProxy* CachedReflectionCaptureProxy;
+
+	/** Mapping from instance index in this primitive to index in the global distance field object buffers. */
+	TArray<int32, TInlineAllocator<1>> DistanceFieldInstanceIndices;
+
+	/** Whether the primitive is newly registered or moved and CachedReflectionCaptureProxy needs to be updated on the next render. */
+	uint32 bNeedsCachedReflectionCaptureUpdate : 1;
+
+	/** This primitive has the Motion Blur explicitly disabled */
+	uint32 bVelocityIsSupressed : 1;
+
+	/** The hit proxies used by the primitive. */
+	TArray<TRefCountPtr<HHitProxy> > HitProxies;
+
+	/** The hit proxy which is used to represent the primitive's dynamic elements. */
+	HHitProxy* DefaultDynamicHitProxy;
+
+	/** The ID of the hit proxy which is used to represent the primitive's dynamic elements. */
+	FHitProxyId DefaultDynamicHitProxyId;
+
+	/** The list of lights affecting this primitive. */
+	class FLightPrimitiveInteraction* LightList;
+
+	/** Last render time in seconds since level started play. */
+	float LastRenderTime;
+
+	/** Last time that the primitive became visible in seconds since level started play. */
+	float LastVisibilityChangeTime;
+
+	/** The scene the primitive is in. */
+	FScene* Scene;
+
+	/** Initialization constructor. */
+	FPrimitiveSceneInfo(UPrimitiveComponent* InPrimitive,FScene* InScene);
+
+	/** Destructor. */
+	~FPrimitiveSceneInfo();
+
+	/** Adds the primitive to the scene. */
+	void AddToScene(FRHICommandListImmediate& RHICmdList, bool bUpdateStaticDrawLists);
+
+	/** Removes the primitive from the scene. */
+	void RemoveFromScene(bool bUpdateStaticDrawLists);
+
+	/** return true if we need to call ConditionalUpdateStaticMeshes */
+	FORCEINLINE bool NeedsUpdateStaticMeshes()
+	{
+		return bNeedsStaticMeshUpdate;
+	}
+
+	/** Updates the primitive's static meshes in the scene. */
+	void UpdateStaticMeshes(FRHICommandListImmediate& RHICmdList);
+
+	/** Updates the primitive's static meshes in the scene. */
+	FORCEINLINE void ConditionalUpdateStaticMeshes(FRHICommandListImmediate& RHICmdList)
+	{
+		if (NeedsUpdateStaticMeshes())
+		{
+			UpdateStaticMeshes(RHICmdList);
+		}
+	}
+
+	/** Sets a flag to update the primitive's static meshes before it is next rendered. */
+	void BeginDeferredUpdateStaticMeshes();
+
+	/** Adds the primitive's static meshes to the scene. */
+	void AddStaticMeshes(FRHICommandListImmediate& RHICmdList);
+
+	/** Removes the primitive's static meshes from the scene. */
+	void RemoveStaticMeshes();
+
+	/** Set LOD Parent primitive information to the scene. */
+	void LinkLODParentComponent();
+
+	/** clear LOD parent primitive information from the scene. */
+	void UnlinkLODParentComponent();
+
+	/** Adds the primitive to the scene's attachment groups. */
+	void LinkAttachmentGroup();
+
+	/** Removes the primitive from the scene's attachment groups. */
+	void UnlinkAttachmentGroup();
+
+	/** 
+	 * Builds an array of all primitive scene info's in this primitive's attachment group. 
+	 * This only works on potential parents (!LightingAttachmentRoot.IsValid()) and will include the current primitive in the output array.
+	 */
+	void GatherLightingAttachmentGroupPrimitives(TArray<FPrimitiveSceneInfo*, SceneRenderingAllocator>& OutChildSceneInfos);
+
+	/** 
+	 * Builds a cumulative bounding box of this primitive and all the primitives in the same attachment group. 
+	 * This only works on potential parents (!LightingAttachmentRoot.IsValid()).
+	 */
+	FBoxSphereBounds GetAttachmentGroupBounds() const;
+
+	/** Size this class uses in bytes */
+	uint32 GetMemoryFootprint();
+
+	// FDeferredCleanupInterface implementation.
+	virtual void FinishCleanup() { delete this; }
+
+	/** 
+	 * Retrieves the index of the primitive in the scene's primitives array.
+	 * This index is only valid until a primitive is added to or removed from
+	 * the scene!
+	 */
+	inline int32 GetIndex() const { return PackedIndex; }
+
+	/* @return true if the object needs to be rendered in the velocity pass (is not moving like the world, needed for motionblur and TemporalAA) */
+	bool ShouldRenderVelocity(const FViewInfo& View, bool bCheckVisibility = true) const;
+	
+	/**
+	 * Shifts primitive position and all relevant data by an arbitrary delta.
+	 * Called on world origin changes
+	 * @param InOffset - The delta to shift by
+	 */
+	void ApplyWorldOffset(FVector InOffset);
+
+private:
+
+	/** Let FScene have direct access to the Id. */
+	friend class FScene;
+
+	/**
+	 * The index of the primitive in the scene's packed arrays. This value may
+	 * change as primitives are added and removed from the scene.
+	 */
+	int32 PackedIndex;
+
+	/** 
+	 * The UPrimitiveComponent this scene info is for, useful for quickly inspecting properties on the corresponding component while debugging.
+	 * This should not be dereferenced on the rendering thread.  The game thread can be modifying UObject members at any time.
+	 * Use PrimitiveComponentId instead when a component identifier is needed.
+	 */
+	const UPrimitiveComponent* ComponentForDebuggingOnly;
+
+	/** If this is TRUE, this primitive's static meshes need to be updated before it can be rendered. */
+	bool bNeedsStaticMeshUpdate;
+};
+
+/** Defines how the primitive is stored in the scene's primitive octree. */
+struct FPrimitiveOctreeSemantics
+{
+	enum { MaxElementsPerLeaf = 16 };
+	enum { MinInclusiveElementsPerNode = 7 };
+	enum { MaxNodeDepth = 12 };
+
+	typedef TInlineAllocator<MaxElementsPerLeaf> ElementAllocator;
+
+	FORCEINLINE static const FBoxSphereBounds& GetBoundingBox(const FPrimitiveSceneInfoCompact& PrimitiveSceneInfoCompact)
+	{
+		return PrimitiveSceneInfoCompact.Bounds;
+	}
+
+	FORCEINLINE static bool AreElementsEqual(const FPrimitiveSceneInfoCompact& A,const FPrimitiveSceneInfoCompact& B)
+	{
+		return A.PrimitiveSceneInfo == B.PrimitiveSceneInfo;
+	}
+
+	FORCEINLINE static void SetElementId(const FPrimitiveSceneInfoCompact& Element,FOctreeElementId Id)
+	{
+		Element.PrimitiveSceneInfo->OctreeId = Id;
+	}
+
+	FORCEINLINE static void ApplyOffset(FPrimitiveSceneInfoCompact& Element, FVector Offset)
+	{
+		Element.Bounds.Origin+= Offset;
+	}
+};
+
+#endif
diff --git a/Engine/Source/Runtime/Renderer/Public/RendererInterface.h b/Engine/Source/Runtime/Renderer/Public/RendererInterface.h
--- a/Engine/Source/Runtime/Renderer/Public/RendererInterface.h
+++ b/Engine/Source/Runtime/Renderer/Public/RendererInterface.h
@@ -478,6 +478,9 @@ public:
 	/** Call from the game thread to send a message to the rendering thread to being rendering this view family. */
 	virtual void BeginRenderingViewFamily(FCanvas* Canvas, FSceneViewFamily* ViewFamily) = 0;
 	
+	/** Call from the game thread to create and initalize a new FViewInfo with the specified options, and add it to the specified view family. */
+	virtual void CreateAndInitSingleView(class FSceneViewFamily* ViewFamily, const struct FSceneViewInitOptions* ViewInitOptions) = 0;
+
 	/**
 	 * Allocates a new instance of the private FScene implementation for the given world.
 	 * @param World - An optional world to associate with the scene.
diff --git a/Engine/Source/Runtime/Renderer/Renderer.Build.cs b/Engine/Source/Runtime/Renderer/Renderer.Build.cs
--- a/Engine/Source/Runtime/Renderer/Renderer.Build.cs
+++ b/Engine/Source/Runtime/Renderer/Renderer.Build.cs
@@ -15,6 +15,7 @@ public class Renderer : ModuleRules
 		);
 
 		PublicDependencyModuleNames.Add("Core");
+        PublicDependencyModuleNames.Add("Engine");
 
         // Renderer module builds faster without unity
         // Non-unity also provides faster iteration
@@ -25,7 +26,6 @@ public class Renderer : ModuleRules
 
 		PrivateDependencyModuleNames.AddRange(
 			new string[] {
-				"Engine", 
 				"CoreUObject", 
 				"RenderCore", 
 				"RHI", 
