Fix overflow problem with NumPrimitives calculation when

diff --git a/Engine/Source/Runtime/Landscape/Public/LandscapeRender.h b/Engine/Source/Runtime/Landscape/Public/LandscapeRender.h
--- a/Engine/Source/Runtime/Landscape/Public/LandscapeRender.h
+++ b/Engine/Source/Runtime/Landscape/Public/LandscapeRender.h
@@ -487,11 +487,11 @@ class FLandscapeComponentSceneProxy : public FPrimitiveSceneProxy, public FLands
 
 protected:
 	int8						MaxLOD;
-	int8						NumSubsections;
-	int16						SubsectionSizeQuads;
-	int16						SubsectionSizeVerts;
-	int16						ComponentSizeQuads;	// Size of component in quads
-	int16						ComponentSizeVerts;
+	int32						NumSubsections;
+	int32						SubsectionSizeQuads;
+	int32						SubsectionSizeVerts;
+	int32						ComponentSizeQuads;	// Size of component in quads
+	int32						ComponentSizeVerts;
 	uint8						StaticLightingLOD;
 	float						StaticLightingResolution;
 	FIntPoint					SectionBase;
