UE4 - Fix build issues

diff --git a/Engine/Source/Runtime/Landscape/Private/LandscapeGrass.cpp b/Engine/Source/Runtime/Landscape/Private/LandscapeGrass.cpp
--- a/Engine/Source/Runtime/Landscape/Private/LandscapeGrass.cpp
+++ b/Engine/Source/Runtime/Landscape/Private/LandscapeGrass.cpp
@@ -35,7 +35,7 @@ public:
 	static bool ShouldCache(EShaderPlatform Platform, const FMaterial* Material, const FVertexFactoryType* VertexFactoryType)
 	{
 		// Only compile for LandscapeVertexFactory
-		return (Material->IsUsedWithLandscape() || Material->IsSpecialEngineMaterial()) && IsFeatureLevelSupported(Platform, ERHIFeatureLevel::SM4) && (VertexFactoryType == FindVertexFactoryType(FName(TEXT("FLandscapeVertexFactory"), FNAME_Find)));
+		return (Material->IsUsedWithLandscape() || Material->IsSpecialEngineMaterial()) && IsFeatureLevelSupported(Platform, ERHIFeatureLevel::SM4) && (VertexFactoryType == FindVertexFactoryType(FName(TEXT("FLandscapeVertexFactory"), FNAME_Find))) && !IsConsolePlatform(Platform);
 	}
 
 	void SetParameters(FRHICommandList& RHICmdList, const FMaterialRenderProxy* MaterialRenderProxy, const FMaterial& MaterialResource, const FSceneView& View)
@@ -60,7 +60,7 @@ public:
 	static bool ShouldCache(EShaderPlatform Platform, const FMaterial* Material, const FVertexFactoryType* VertexFactoryType)
 	{
 		// Only compile for LandscapeVertexFactory
-		return (Material->IsUsedWithLandscape() || Material->IsSpecialEngineMaterial()) && IsFeatureLevelSupported(Platform, ERHIFeatureLevel::SM4) && (VertexFactoryType == FindVertexFactoryType(FName(TEXT("FLandscapeVertexFactory"), FNAME_Find)));
+		return (Material->IsUsedWithLandscape() || Material->IsSpecialEngineMaterial()) && IsFeatureLevelSupported(Platform, ERHIFeatureLevel::SM4) && (VertexFactoryType == FindVertexFactoryType(FName(TEXT("FLandscapeVertexFactory"), FNAME_Find))) && !IsConsolePlatform(Platform);
 	}
 
 	FLandscapeGrassWeightPS(const ShaderMetaType::CompiledShaderInitializerType& Initializer) :
