In landscape editor using the clear gizmo data button now also removes the copied data from the clipboard

diff --git a/Engine/Source/Runtime/Landscape/Private/LandscapeGizmoActor.cpp b/Engine/Source/Runtime/Landscape/Private/LandscapeGizmoActor.cpp
--- a/Engine/Source/Runtime/Landscape/Private/LandscapeGizmoActor.cpp
+++ b/Engine/Source/Runtime/Landscape/Private/LandscapeGizmoActor.cpp
@@ -659,6 +659,16 @@ void ALandscapeGizmoActiveActor::ClearGizmoData()
 	DataType = LGT_None;
 	SelectedData.Empty();
 	LayerInfos.Empty();
+
+	// If the clipboard contains copied gizmo data, clear it also
+	FString ClipboardString;
+	FPlatformMisc::ClipboardPaste(ClipboardString);
+	const TCHAR* Str = *ClipboardString;
+	if (FParse::Command(&Str, TEXT("GizmoData=")))
+	{
+		FPlatformMisc::ClipboardCopy(TEXT(""));
+	}
+
 	ReregisterAllComponents();
 }
 
