Made Landscape not blueprintable to hide the component details.

diff --git a/Engine/Source/Runtime/Landscape/Classes/LandscapeProxy.h b/Engine/Source/Runtime/Landscape/Classes/LandscapeProxy.h
--- a/Engine/Source/Runtime/Landscape/Classes/LandscapeProxy.h
+++ b/Engine/Source/Runtime/Landscape/Classes/LandscapeProxy.h
@@ -294,7 +294,7 @@ public:
 	}
 };
 
-UCLASS(NotPlaceable, hidecategories=(Display, Attachment, Physics, Debug, Lighting, LOD), showcategories=(Lighting, Rendering, "Utilities|Transformation"), MinimalAPI)
+UCLASS(NotPlaceable, NotBlueprintable, hidecategories=(Display, Attachment, Physics, Debug, Lighting, LOD), showcategories=(Lighting, Rendering, "Utilities|Transformation"), MinimalAPI)
 class ALandscapeProxy : public AActor, public FTickableGameObject
 {
 	GENERATED_UCLASS_BODY()
