UE4 - add grass variety array

diff --git a/Engine/Source/Runtime/Landscape/Classes/LandscapeGrassType.h b/Engine/Source/Runtime/Landscape/Classes/LandscapeGrassType.h
--- a/Engine/Source/Runtime/Landscape/Classes/LandscapeGrassType.h
+++ b/Engine/Source/Runtime/Landscape/Classes/LandscapeGrassType.h
@@ -4,10 +4,10 @@
 
 #include "LandscapeGrassType.generated.h"
 
-UCLASS(MinimalAPI)
-class ULandscapeGrassType : public UObject
+USTRUCT()
+struct FGrassVariety
 {
-	GENERATED_UCLASS_BODY()
+	GENERATED_USTRUCT_BODY()
 
 	UPROPERTY(EditAnywhere, Category=Grass)
 	class UStaticMesh* GrassMesh;
@@ -36,11 +36,47 @@ class ULandscapeGrassType : public UObject
 	UPROPERTY(EditAnywhere, Category = Grass)
 	bool AlignToSurface;
 
-#if WITH_EDITOR
+	FGrassVariety()
+	{
+		GrassMesh = nullptr;
+		GrassDensity = 400;
+		StartCullDistance = 10000.0f;
+		EndCullDistance = 10000.0f;
+		PlacementJitter = 1.0f;
+		RandomRotation = true;
+		AlignToSurface = true;
+	}
+};
+
+UCLASS(MinimalAPI)
+class ULandscapeGrassType : public UObject
+{
+	GENERATED_UCLASS_BODY()
+
+	UPROPERTY(EditAnywhere, Category = Grass)
+	TArray<FGrassVariety> GrassVarieties;
+
+	UPROPERTY()
+	class UStaticMesh* GrassMesh_DEPRECATED;
+	UPROPERTY()
+	float GrassDensity_DEPRECATED;
+	UPROPERTY()
+	float PlacementJitter_DEPRECATED;
+	UPROPERTY()
+	int32 StartCullDistance_DEPRECATED;
+	UPROPERTY()
+	int32 EndCullDistance_DEPRECATED;
+	UPROPERTY()
+	bool RandomRotation_DEPRECATED;
+	UPROPERTY()
+	bool AlignToSurface_DEPRECATED;
+
 	// Begin UObject Interface
+	virtual void PostLoad() override;
+#if WITH_EDITOR
 	virtual void PostEditChangeProperty(FPropertyChangedEvent& PropertyChangedEvent) override;
-	// End UObject Interface
 #endif
+	// End UObject Interface
 };
 
 
diff --git a/Engine/Source/Runtime/Landscape/Classes/LandscapeProxy.h b/Engine/Source/Runtime/Landscape/Classes/LandscapeProxy.h
--- a/Engine/Source/Runtime/Landscape/Classes/LandscapeProxy.h
+++ b/Engine/Source/Runtime/Landscape/Classes/LandscapeProxy.h
@@ -201,12 +201,16 @@ struct FCachedLandscapeFoliage
 		int32 CachedMaxInstancesPerComponent;
 		int32 SubsectionX;
 		int32 SubsectionY;
+		int32 NumVarieties;
+		int32 VarietyIndex;
 
 		FGrassCompKey()
 			: SqrtSubsections(0)
 			, CachedMaxInstancesPerComponent(0)
 			, SubsectionX(0)
 			, SubsectionY(0)
+			, NumVarieties(0)
+			, VarietyIndex(-1)
 		{
 		}
 		inline bool operator==(const FGrassCompKey& Other) const
@@ -217,12 +221,14 @@ struct FCachedLandscapeFoliage
 				SubsectionX == Other.SubsectionX &&
 				SubsectionY == Other.SubsectionY &&
 				BasedOn == Other.BasedOn &&
-				GrassType == Other.GrassType;
+				GrassType == Other.GrassType &&
+				NumVarieties == Other.NumVarieties &&
+				VarietyIndex == Other.VarietyIndex;
 		}
 
 		friend uint32 GetTypeHash(const FGrassCompKey& Key)
 		{
-			return GetTypeHash(Key.BasedOn) ^ GetTypeHash(Key.GrassType) ^ Key.SqrtSubsections ^ Key.CachedMaxInstancesPerComponent ^ (Key.SubsectionX >> 16) ^ (Key.SubsectionY >> 24);
+			return GetTypeHash(Key.BasedOn) ^ GetTypeHash(Key.GrassType) ^ Key.SqrtSubsections ^ Key.CachedMaxInstancesPerComponent ^ (Key.SubsectionX >> 16) ^ (Key.SubsectionY >> 24) ^ (Key.NumVarieties >> 3) ^ (Key.VarietyIndex >> 13);
 		}
 
 	};
diff --git a/Engine/Source/Runtime/Landscape/Private/Landscape.cpp b/Engine/Source/Runtime/Landscape/Private/Landscape.cpp
--- a/Engine/Source/Runtime/Landscape/Private/Landscape.cpp
+++ b/Engine/Source/Runtime/Landscape/Private/Landscape.cpp
@@ -2260,10 +2260,10 @@ struct FGrassBuilderBase
 	FVector Extent;
 	int32 SqrtMaxInstances;
 
-	FGrassBuilderBase(ALandscapeProxy* Landscape, ULandscapeComponent* Component, const ULandscapeGrassType* GrassType, int32 SqrtSubsections = 1, int32 SubX = 0, int32 SubY = 0)
+	FGrassBuilderBase(ALandscapeProxy* Landscape, ULandscapeComponent* Component, const FGrassVariety& GrassVariety, int32 SqrtSubsections = 1, int32 SubX = 0, int32 SubY = 0)
 	{
 		bHaveValidData = true;
-		GrassDensity = GrassType->GrassDensity;
+		GrassDensity = GrassVariety.GrassDensity;
 
 		DrawScale = Landscape->GetRootComponent()->RelativeScale3D;
 		DrawLoc = Landscape->GetActorLocation();
@@ -2351,19 +2351,19 @@ struct FAsyncGrassBuilder : public FGrassBuilderBase
 	FStaticMeshInstanceData InstanceBuffer;
 	TArray<FClusterNode> ClusterTree;
 
-	FAsyncGrassBuilder(ALandscapeProxy* Landscape, ULandscapeComponent* Component, const ULandscapeGrassType* GrassType, UHierarchicalInstancedStaticMeshComponent* HierarchicalInstancedStaticMeshComponent, int32 SqrtSubsections, int32 SubX, int32 SubY)
-		: FGrassBuilderBase(Landscape, Component, GrassType, SqrtSubsections, SubX, SubY)
+	FAsyncGrassBuilder(ALandscapeProxy* Landscape, ULandscapeComponent* Component, const ULandscapeGrassType* GrassType, const FGrassVariety& GrassVariety, UHierarchicalInstancedStaticMeshComponent* HierarchicalInstancedStaticMeshComponent, int32 SqrtSubsections, int32 SubX, int32 SubY)
+		: FGrassBuilderBase(Landscape, Component, GrassVariety, SqrtSubsections, SubX, SubY)
 		, GrassData(Component, GrassType)
 	{
 		bHaveValidData = bHaveValidData && GrassData.IsValid();
 		XForm = LandscapeToWorld * HierarchicalInstancedStaticMeshComponent->ComponentToWorld.ToMatrixWithScale().Inverse();
-		PlacementJitter = GrassType->PlacementJitter;
-		RandomRotation = GrassType->RandomRotation;
-		AlignToSurface = GrassType->AlignToSurface;
+		PlacementJitter = GrassVariety.PlacementJitter;
+		RandomRotation = GrassVariety.RandomRotation;
+		AlignToSurface = GrassVariety.AlignToSurface;
 
 		RandomStream.Initialize(HierarchicalInstancedStaticMeshComponent->InstancingRandomSeed);
 
-		MeshBox = GrassType->GrassMesh->GetBounds().GetBox();
+		MeshBox = GrassVariety.GrassMesh->GetBounds().GetBox();
 
 		RasterTime = 0.0;
 		BuildTime = 0.0;
@@ -2645,202 +2645,211 @@ void ALandscapeProxy::UpdateFoliage(const TArray<FVector>& Cameras, ULandscapeCo
 
 		if (GetWorld())
 		{
-
-		for (int32 ComponentIndex = 0; ComponentIndex < LandscapeComponents.Num(); ComponentIndex++)
-		{
-			ULandscapeComponent* Component = LandscapeComponents[ComponentIndex];
-			if (OnlyComponent && OnlyComponent != Component)
-			{
-				continue;
-			}
-
-			FBoxSphereBounds WorldBounds = Component->CalcBounds(Component->ComponentToWorld);
-			float MinDistanceToComp = Cameras.Num() ? MAX_flt : 0.0f;
-
-			for (auto& Pos : Cameras)
-			{
-				MinDistanceToComp = FMath::Min<float>(MinDistanceToComp, WorldBounds.ComputeSquaredDistanceFromBoxToPoint(Pos));
-			}
-
-			MinDistanceToComp = FMath::Sqrt(MinDistanceToComp);
-
-			for (auto GrassType : GrassTypes)
+			int32 NumCompsCreated = 0;
+			for (int32 ComponentIndex = 0; ComponentIndex < LandscapeComponents.Num(); ComponentIndex++)
 			{
-				if (GrassType && GrassType->GrassMesh && GrassType->GrassDensity > 0.0f && GrassType->EndCullDistance > 0)
+				ULandscapeComponent* Component = LandscapeComponents[ComponentIndex];
+				if (OnlyComponent && OnlyComponent != Component)
 				{
-					bool bRandomRotation = GrassType->RandomRotation;
-					bool bAlign = GrassType->AlignToSurface;
-					float MustHaveDistance = GuardBand * (float)GrassType->EndCullDistance;
-					float DiscardDistance = DiscardGuardBand * (float)GrassType->EndCullDistance;
+					continue;
+				}
 
-					if (MinDistanceToComp > DiscardDistance)
-					{
-						continue;
-					}
+				FBoxSphereBounds WorldBounds = Component->CalcBounds(Component->ComponentToWorld);
+				float MinDistanceToComp = Cameras.Num() ? MAX_flt : 0.0f;
 
-					FGrassBuilderBase ForSubsectionMath(this, Component, GrassType);
+				for (auto& Pos : Cameras)
+				{
+					MinDistanceToComp = FMath::Min<float>(MinDistanceToComp, WorldBounds.ComputeSquaredDistanceFromBoxToPoint(Pos));
+				}
 
-					int32 SqrtSubsections = 1;
-				
-					if (ForSubsectionMath.bHaveValidData && ForSubsectionMath.SqrtMaxInstances > 0)
-					{
-						SqrtSubsections = FMath::Clamp<int32>(FMath::CeilToInt(float(ForSubsectionMath.SqrtMaxInstances) / FMath::Sqrt((float)MaxInstancesPerComponent)), 1, 16);
-					}
+				MinDistanceToComp = FMath::Sqrt(MinDistanceToComp);
 
-					int32 NumCompsCreated = 0;
-					FBox LocalBox = Component->CachedLocalBox;
-					FVector LocalExtentDiv = (LocalBox.Max - LocalBox.Min) * FVector(1.0f / float(SqrtSubsections), 1.0f / float(SqrtSubsections), 1.0f);
-					for (int32 SubX = 0; SubX < SqrtSubsections; SubX++)
+				for (auto GrassType : GrassTypes)
+				{
+					if (GrassType)					
 					{
-						for (int32 SubY = 0; SubY < SqrtSubsections; SubY++)
+						int32 GrassVarietyIndex = -1;
+						for (auto& GrassVariety : GrassType->GrassVarieties)
 						{
-							float MinDistanceToSubComp = MinDistanceToComp;
-
-							if (bCullSubsections && SqrtSubsections > 1)
+							GrassVarietyIndex++;
+							if (GrassVariety.GrassMesh && GrassVariety.GrassDensity > 0.0f && GrassVariety.EndCullDistance > 0)
 							{
-								FVector BoxMin;
-								BoxMin.X = LocalBox.Min.X + LocalExtentDiv.X * float(SubX);
-								BoxMin.Y = LocalBox.Min.Y + LocalExtentDiv.Y * float(SubY);
-								BoxMin.Z = LocalBox.Min.Z;
-
-								FVector BoxMax;
-								BoxMax.X = LocalBox.Min.X + LocalExtentDiv.X * float(SubX + 1);
-								BoxMax.Y = LocalBox.Min.Y + LocalExtentDiv.Y * float(SubY + 1);
-								BoxMax.Z = LocalBox.Max.Z;
-
-								FBox LocalSubBox(BoxMin, BoxMax);
-								FBox WorldSubBox = LocalSubBox.TransformBy(Component->ComponentToWorld);
+								bool bRandomRotation = GrassVariety.RandomRotation;
+								bool bAlign = GrassVariety.AlignToSurface;
+								float MustHaveDistance = GuardBand * (float)GrassVariety.EndCullDistance;
+								float DiscardDistance = DiscardGuardBand * (float)GrassVariety.EndCullDistance;
 
-								MinDistanceToSubComp = Cameras.Num() ? MAX_flt : 0.0f;
-								for (auto& Pos : Cameras)
+								if (MinDistanceToComp > DiscardDistance)
 								{
-									MinDistanceToSubComp = FMath::Min<float>(MinDistanceToSubComp, ComputeSquaredDistanceFromBoxToPoint(WorldSubBox.Min, WorldSubBox.Max, Pos));
+									continue;
 								}
-								MinDistanceToSubComp = FMath::Sqrt(MinDistanceToSubComp);
-							}
 
-							if (MinDistanceToSubComp > DiscardDistance)
-							{
-								continue;
-							}
+								FGrassBuilderBase ForSubsectionMath(this, Component, GrassVariety);
 
-							FCachedLandscapeFoliage::FGrassComp NewComp;
-							NewComp.Key.BasedOn = Component;
-							NewComp.Key.GrassType = GrassType;
-							NewComp.Key.SqrtSubsections = SqrtSubsections;
-							NewComp.Key.CachedMaxInstancesPerComponent = MaxInstancesPerComponent;
-							NewComp.Key.SubsectionX = SubX;
-							NewComp.Key.SubsectionY = SubY;
+								int32 SqrtSubsections = 1;
+				
+								if (ForSubsectionMath.bHaveValidData && ForSubsectionMath.SqrtMaxInstances > 0)
+								{
+									SqrtSubsections = FMath::Clamp<int32>(FMath::CeilToInt(float(ForSubsectionMath.SqrtMaxInstances) / FMath::Sqrt((float)MaxInstancesPerComponent)), 1, 16);
+								}
 
-							{
-								FCachedLandscapeFoliage::FGrassComp* Existing = FoliageCache.CachedGrassComps.Find(NewComp.Key);
-								if (Existing || MinDistanceToSubComp > MustHaveDistance)
+								FBox LocalBox = Component->CachedLocalBox;
+								FVector LocalExtentDiv = (LocalBox.Max - LocalBox.Min) * FVector(1.0f / float(SqrtSubsections), 1.0f / float(SqrtSubsections), 1.0f);
+								for (int32 SubX = 0; SubX < SqrtSubsections; SubX++)
 								{
-									if (Existing)
+									for (int32 SubY = 0; SubY < SqrtSubsections; SubY++)
 									{
-										Existing->Touch();
+										float MinDistanceToSubComp = MinDistanceToComp;
+
+										if (bCullSubsections && SqrtSubsections > 1)
+										{
+											FVector BoxMin;
+											BoxMin.X = LocalBox.Min.X + LocalExtentDiv.X * float(SubX);
+											BoxMin.Y = LocalBox.Min.Y + LocalExtentDiv.Y * float(SubY);
+											BoxMin.Z = LocalBox.Min.Z;
+
+											FVector BoxMax;
+											BoxMax.X = LocalBox.Min.X + LocalExtentDiv.X * float(SubX + 1);
+											BoxMax.Y = LocalBox.Min.Y + LocalExtentDiv.Y * float(SubY + 1);
+											BoxMax.Z = LocalBox.Max.Z;
+
+											FBox LocalSubBox(BoxMin, BoxMax);
+											FBox WorldSubBox = LocalSubBox.TransformBy(Component->ComponentToWorld);
+
+											MinDistanceToSubComp = Cameras.Num() ? MAX_flt : 0.0f;
+											for (auto& Pos : Cameras)
+											{
+												MinDistanceToSubComp = FMath::Min<float>(MinDistanceToSubComp, ComputeSquaredDistanceFromBoxToPoint(WorldSubBox.Min, WorldSubBox.Max, Pos));
+											}
+											MinDistanceToSubComp = FMath::Sqrt(MinDistanceToSubComp);
+										}
+
+										if (MinDistanceToSubComp > DiscardDistance)
+										{
+											continue;
+										}
+
+										FCachedLandscapeFoliage::FGrassComp NewComp;
+										NewComp.Key.BasedOn = Component;
+										NewComp.Key.GrassType = GrassType;
+										NewComp.Key.SqrtSubsections = SqrtSubsections;
+										NewComp.Key.CachedMaxInstancesPerComponent = MaxInstancesPerComponent;
+										NewComp.Key.SubsectionX = SubX;
+										NewComp.Key.SubsectionY = SubY;
+										NewComp.Key.NumVarieties = GrassType->GrassVarieties.Num();
+										NewComp.Key.VarietyIndex = GrassVarietyIndex;
+
+										{
+											FCachedLandscapeFoliage::FGrassComp* Existing = FoliageCache.CachedGrassComps.Find(NewComp.Key);
+											if (Existing || MinDistanceToSubComp > MustHaveDistance)
+											{
+												if (Existing)
+												{
+													Existing->Touch();
+												}
+												continue;
+											}
+										}
+
+										if (NumCompsCreated || (!bForceSync && AsyncFoliageTasks.Num() >= MaxTasks))
+										{
+											continue; // one per frame, but we still want to touch the existing ones
+										}
+										NumCompsCreated++;
+
+										SCOPE_CYCLE_COUNTER(STAT_FoliageGrassStartComp);
+										int32 FolSeed = FCrc::StrCrc32((GrassType->GetName() + Component->GetName() + FString::Printf(TEXT("%d %d %d"), SubX, SubY, GrassVarietyIndex)).GetCharArray().GetData());
+										if (FolSeed == 0)
+										{
+											FolSeed++;
+										}
+
+										UHierarchicalInstancedStaticMeshComponent* HierarchicalInstancedStaticMeshComponent;
+										{
+											QUICK_SCOPE_CYCLE_COUNTER(STAT_GrassCreateComp);
+											HierarchicalInstancedStaticMeshComponent = ConstructObject<UHierarchicalInstancedStaticMeshComponent>(UHierarchicalInstancedStaticMeshComponent::StaticClass(), this); //, NAME_None, RF_Transactional);
+										}
+										NewComp.Foliage = HierarchicalInstancedStaticMeshComponent;
+										FoliageCache.CachedGrassComps.Add(NewComp);
+
+										HierarchicalInstancedStaticMeshComponent->Mobility = EComponentMobility::Static;
+
+										HierarchicalInstancedStaticMeshComponent->StaticMesh = GrassVariety.GrassMesh;
+										HierarchicalInstancedStaticMeshComponent->bSelectable = false;
+										HierarchicalInstancedStaticMeshComponent->bHasPerInstanceHitProxies = true;
+										static FName NoCollision(TEXT("NoCollision"));
+										HierarchicalInstancedStaticMeshComponent->SetCollisionProfileName(NoCollision);
+										HierarchicalInstancedStaticMeshComponent->bDisableCollision = true;
+										HierarchicalInstancedStaticMeshComponent->SetCanEverAffectNavigation(false);
+										HierarchicalInstancedStaticMeshComponent->InstancingRandomSeed = FolSeed;
+
+										if (!Cameras.Num() || bDisableGPUCull)
+										{
+											// if we don't have any cameras, then we are rendering landscape LOD materials or somesuch and we want to disable culling
+											HierarchicalInstancedStaticMeshComponent->InstanceStartCullDistance = 0;
+											HierarchicalInstancedStaticMeshComponent->InstanceEndCullDistance = 0;
+										}
+										else
+										{
+											HierarchicalInstancedStaticMeshComponent->InstanceStartCullDistance = GrassVariety.StartCullDistance;
+											HierarchicalInstancedStaticMeshComponent->InstanceEndCullDistance = GrassVariety.EndCullDistance;
+										}
+
+										//@todo - take the settings from a UFoliageType object.  For now, disable distance field lighting on grass so we don't hitch.
+										HierarchicalInstancedStaticMeshComponent->bAffectDistanceFieldLighting = false;
+
+										{
+											QUICK_SCOPE_CYCLE_COUNTER(STAT_GrassAttachComp);
+
+											HierarchicalInstancedStaticMeshComponent->AttachTo(GetRootComponent());
+											FTransform DesiredTransform = GetRootComponent()->ComponentToWorld;
+											DesiredTransform.RemoveScaling();
+											HierarchicalInstancedStaticMeshComponent->SetWorldTransform(DesiredTransform);
+
+											FoliageComponents.Add(HierarchicalInstancedStaticMeshComponent);
+										}
+
+			#if WITH_EDITOR
+										// render grass data if we don't have any
+										if (!Component->GrassData->HasData())
+										{
+											QUICK_SCOPE_CYCLE_COUNTER(STAT_GrassRenderToTexture);
+											Component->RenderGrassMap();
+										}
+			#endif
+
+										FAsyncGrassBuilder* Builder;
+
+										{
+											QUICK_SCOPE_CYCLE_COUNTER(STAT_GrassCreateBuilder);
+											Builder = new FAsyncGrassBuilder(this, Component, GrassType, GrassVariety, HierarchicalInstancedStaticMeshComponent, SqrtSubsections, SubX, SubY);
+										}
+
+										if (Builder->bHaveValidData)
+										{
+											FAsyncTask<FAsyncGrassTask>* Task = new FAsyncTask<FAsyncGrassTask>(Builder, NewComp.Key, HierarchicalInstancedStaticMeshComponent);
+
+											Task->StartBackgroundTask();
+
+											AsyncFoliageTasks.Add(Task);
+										}
+										else
+										{
+											delete Builder;
+										}
+										{
+											QUICK_SCOPE_CYCLE_COUNTER(STAT_GrassRegisterComp);
+											HierarchicalInstancedStaticMeshComponent->RegisterComponent();
+										}
 									}
-									continue;
 								}
 							}
-
-								if (NumCompsCreated || (!bForceSync && AsyncFoliageTasks.Num() >= MaxTasks))
-							{
-								continue; // one per frame, but we still want to touch the existing ones
-							}
-							NumCompsCreated++;
-
-							SCOPE_CYCLE_COUNTER(STAT_FoliageGrassStartComp);
-							int32 FolSeed = FCrc::StrCrc32((GrassType->GetName() + Component->GetName() + FString::Printf(TEXT("%d %d"), SubX, SubY)).GetCharArray().GetData());
-							if (FolSeed == 0)
-							{
-								FolSeed++;
-							}
-
-							UHierarchicalInstancedStaticMeshComponent* HierarchicalInstancedStaticMeshComponent;
-							{
-								QUICK_SCOPE_CYCLE_COUNTER(STAT_GrassCreateComp);
-								HierarchicalInstancedStaticMeshComponent = ConstructObject<UHierarchicalInstancedStaticMeshComponent>(UHierarchicalInstancedStaticMeshComponent::StaticClass(), this); //, NAME_None, RF_Transactional);
-							}
-							NewComp.Foliage = HierarchicalInstancedStaticMeshComponent;
-							FoliageCache.CachedGrassComps.Add(NewComp);
-
-							HierarchicalInstancedStaticMeshComponent->Mobility = EComponentMobility::Static;
-
-							HierarchicalInstancedStaticMeshComponent->StaticMesh = GrassType->GrassMesh;
-							HierarchicalInstancedStaticMeshComponent->bSelectable = false;
-							HierarchicalInstancedStaticMeshComponent->bHasPerInstanceHitProxies = true;
-							static FName NoCollision(TEXT("NoCollision"));
-							HierarchicalInstancedStaticMeshComponent->SetCollisionProfileName(NoCollision);
-							HierarchicalInstancedStaticMeshComponent->bDisableCollision = true;
-							HierarchicalInstancedStaticMeshComponent->SetCanEverAffectNavigation(false);
-							HierarchicalInstancedStaticMeshComponent->InstancingRandomSeed = FolSeed;
-
-							if (!Cameras.Num() || bDisableGPUCull)
-							{
-								// if we don't have any cameras, then we are rendering landscape LOD materials or somesuch and we want to disable culling
-								HierarchicalInstancedStaticMeshComponent->InstanceStartCullDistance = 0;
-								HierarchicalInstancedStaticMeshComponent->InstanceEndCullDistance = 0;
-							}
-							else
-							{
-								HierarchicalInstancedStaticMeshComponent->InstanceStartCullDistance = GrassType->StartCullDistance;
-								HierarchicalInstancedStaticMeshComponent->InstanceEndCullDistance = GrassType->EndCullDistance;
-							}
-
-							//@todo - take the settings from a UFoliageType object.  For now, disable distance field lighting on grass so we don't hitch.
-							HierarchicalInstancedStaticMeshComponent->bAffectDistanceFieldLighting = false;
-
-							{
-								QUICK_SCOPE_CYCLE_COUNTER(STAT_GrassAttachComp);
-
-								HierarchicalInstancedStaticMeshComponent->AttachTo(GetRootComponent());
-								FTransform DesiredTransform = GetRootComponent()->ComponentToWorld;
-								DesiredTransform.RemoveScaling();
-								HierarchicalInstancedStaticMeshComponent->SetWorldTransform(DesiredTransform);
-
-								FoliageComponents.Add(HierarchicalInstancedStaticMeshComponent);
-							}
-
-#if WITH_EDITOR
-							// render grass data if we don't have any
-							if (!Component->GrassData->HasData())
-							{
-								QUICK_SCOPE_CYCLE_COUNTER(STAT_GrassRenderToTexture);
-								Component->RenderGrassMap();
-							}
-#endif
-
-							FAsyncGrassBuilder* Builder;
-
-							{
-								QUICK_SCOPE_CYCLE_COUNTER(STAT_GrassCreateBuilder);
-								Builder = new FAsyncGrassBuilder(this, Component, GrassType, HierarchicalInstancedStaticMeshComponent, SqrtSubsections, SubX, SubY);
-							}
-
-							if (Builder->bHaveValidData)
-							{
-								FAsyncTask<FAsyncGrassTask>* Task = new FAsyncTask<FAsyncGrassTask>(Builder, NewComp.Key, HierarchicalInstancedStaticMeshComponent);
-
-								Task->StartBackgroundTask();
-
-								AsyncFoliageTasks.Add(Task);
-							}
-							else
-							{
-								delete Builder;
-							}
-							{
-								QUICK_SCOPE_CYCLE_COUNTER(STAT_GrassRegisterComp);
-								HierarchicalInstancedStaticMeshComponent->RegisterComponent();
-							}
 						}
 					}
 				}
 			}
 		}
 	}
-	}
 
 	static TSet<UHierarchicalInstancedStaticMeshComponent *> StillUsed;
 	StillUsed.Empty(256);
diff --git a/Engine/Source/Runtime/Landscape/Private/LandscapeGrass.cpp b/Engine/Source/Runtime/Landscape/Private/LandscapeGrass.cpp
--- a/Engine/Source/Runtime/Landscape/Private/LandscapeGrass.cpp
+++ b/Engine/Source/Runtime/Landscape/Private/LandscapeGrass.cpp
@@ -453,14 +453,34 @@ void UMaterialExpressionLandscapeGrassOutput::PostEditChangeProperty(FPropertyCh
 ULandscapeGrassType::ULandscapeGrassType(const FObjectInitializer& ObjectInitializer)
 : Super(ObjectInitializer)
 {
-	GrassDensity = 400;
-	StartCullDistance = 10000.0f;
-	EndCullDistance = 10000.0f;
-	PlacementJitter = 1.0f;
-	RandomRotation = true;
-	AlignToSurface = true;
+	GrassDensity_DEPRECATED = 400;
+	StartCullDistance_DEPRECATED = 10000.0f;
+	EndCullDistance_DEPRECATED = 10000.0f;
+	PlacementJitter_DEPRECATED = 1.0f;
+	RandomRotation_DEPRECATED = true;
+	AlignToSurface_DEPRECATED = true;
 }
 
+void ULandscapeGrassType::PostLoad()
+{
+	Super::PostLoad();
+	if (GrassMesh_DEPRECATED && !GrassVarieties.Num())
+	{
+		FGrassVariety Grass;
+		Grass.GrassMesh = GrassMesh_DEPRECATED;
+		Grass.GrassDensity = GrassDensity_DEPRECATED;
+		Grass.StartCullDistance = StartCullDistance_DEPRECATED;
+		Grass.EndCullDistance = EndCullDistance_DEPRECATED;
+		Grass.PlacementJitter = PlacementJitter_DEPRECATED;
+		Grass.RandomRotation = RandomRotation_DEPRECATED;
+		Grass.AlignToSurface = AlignToSurface_DEPRECATED;
+
+		GrassVarieties.Add(Grass);
+		GrassMesh_DEPRECATED = nullptr;
+	}
+}
+
+
 #if WITH_EDITOR
 void ULandscapeGrassType::PostEditChangeProperty(FPropertyChangedEvent& PropertyChangedEvent)
 {
