Support up to 10 grass layers.
Defer grass generation until texture streaming and material compilation has finished.

diff --git a/Engine/Shaders/LandscapeGrassWeight.usf b/Engine/Shaders/LandscapeGrassWeight.usf
--- a/Engine/Shaders/LandscapeGrassWeight.usf
+++ b/Engine/Shaders/LandscapeGrassWeight.usf
@@ -18,6 +18,8 @@ struct FLandscapeGrassWeightInterpolantsVSToPS
 
 #if VERTEXSHADER
 
+float2 RenderOffset;
+
 /** Vertex Shader */
 void VSMain(
 	FVertexFactoryInput Input,
@@ -41,11 +43,15 @@ void VSMain(
 
 	Output.Height = LandscapeVertexFactorySampleHeight(Input, VFIntermediates);
 
+	Output.Position.xy += RenderOffset;
+
 	Output.PixelPosition.w = Output.Position.w;
 }
 
 #elif PIXELSHADER
 
+int OutputPass;
+
 // Pixel Shader
 void PSMain( 
 	FVertexFactoryInterpolantsVSToPS Interpolants,
@@ -58,20 +64,69 @@ void PSMain(
 	FMaterialPixelParameters MaterialParameters = GetMaterialPixelParameters(Interpolants, GrassWeightInterpolants.PixelPosition);
 
 	CalcMaterialParameters(MaterialParameters, SvPosition, bIsFrontFace, GrassWeightInterpolants.PixelPosition); 
-	OutColor.xy = GrassWeightInterpolants.Height;
 
-	#if NUM_MATERIAL_OUTPUTS_GETGRASSWEIGHT > 0
-		OutColor.z = GetGrassWeight0(MaterialParameters);
-		#if NUM_MATERIAL_OUTPUTS_GETGRASSWEIGHT > 1
-			OutColor.w = GetGrassWeight1(MaterialParameters);
+#if NUM_MATERIAL_OUTPUTS_GETGRASSWEIGHT > 2
+	switch(OutputPass)
+	{
+	case 0:
+#endif
+		// Height, outputs 1 & 2
+		OutColor.xy = GrassWeightInterpolants.Height;
+		#if NUM_MATERIAL_OUTPUTS_GETGRASSWEIGHT > 0
+			OutColor.z = GetGrassWeight0(MaterialParameters);
+			#if NUM_MATERIAL_OUTPUTS_GETGRASSWEIGHT > 1
+				OutColor.w = GetGrassWeight1(MaterialParameters);
+			#else
+				OutColor.w = 0;
+			#endif
 		#else
+			OutColor.z = 0;
 			OutColor.w = 0;
 		#endif
-	#else
-		OutColor.z = 0;
-		OutColor.w = 0;
-	#endif
-
+#if NUM_MATERIAL_OUTPUTS_GETGRASSWEIGHT > 2
+		break;
+	case 1:
+		OutColor.x = GetGrassWeight2(MaterialParameters);
+		#if NUM_MATERIAL_OUTPUTS_GETGRASSWEIGHT > 3
+			OutColor.y = GetGrassWeight3(MaterialParameters);
+		#else
+			OutColor.y = 0;
+		#endif
+		#if NUM_MATERIAL_OUTPUTS_GETGRASSWEIGHT > 4
+			OutColor.z = GetGrassWeight4(MaterialParameters);
+		#else
+			OutColor.z = 0;
+		#endif
+		#if NUM_MATERIAL_OUTPUTS_GETGRASSWEIGHT > 5
+			OutColor.w = GetGrassWeight5(MaterialParameters);
+		#else
+			OutColor.w = 0;
+		#endif
+		break;
+#if NUM_MATERIAL_OUTPUTS_GETGRASSWEIGHT > 6
+	case 2:
+		OutColor.x = GetGrassWeight6(MaterialParameters);
+		#if NUM_MATERIAL_OUTPUTS_GETGRASSWEIGHT > 7
+			OutColor.y = GetGrassWeight7(MaterialParameters);
+		#else
+			OutColor.y = 0;
+		#endif
+		#if NUM_MATERIAL_OUTPUTS_GETGRASSWEIGHT > 8
+			OutColor.z = GetGrassWeight8(MaterialParameters);
+		#else
+			OutColor.z = 0;
+		#endif
+		#if NUM_MATERIAL_OUTPUTS_GETGRASSWEIGHT > 9
+			OutColor.w = GetGrassWeight9(MaterialParameters);
+		#else
+			OutColor.w = 0;
+		#endif
+		break;
+#endif // NUM_MATERIAL_OUTPUTS_GETGRASSWEIGHT > 6
+	default:
+		OutColor = float4(0,0,0,0);
+		break;
+	}
+#endif // NUM_MATERIAL_OUTPUTS_GETGRASSWEIGHT > 2
 }
-
 #endif
\ No newline at end of file
diff --git a/Engine/Source/Runtime/Engine/Private/Materials/MaterialShared.cpp b/Engine/Source/Runtime/Engine/Private/Materials/MaterialShared.cpp
--- a/Engine/Source/Runtime/Engine/Private/Materials/MaterialShared.cpp
+++ b/Engine/Source/Runtime/Engine/Private/Materials/MaterialShared.cpp
@@ -302,6 +302,15 @@ bool FMaterial::IsCompilationFinished()
 	return true;
 }
 
+bool FMaterial::HasValidGameThreadShaderMap()
+{
+	if(!GameThreadShaderMap || !GameThreadShaderMap->IsCompilationFinalized())
+	{
+		return false;
+	}
+	return true;
+}
+
 void FMaterial::CancelCompilation()
 {
 	TArray<int32> ShaderMapIdsToCancel;
@@ -1400,6 +1409,7 @@ bool FMaterial::BeginCompileShaderMap(
 		}
 		else
 		{
+			UE_LOG(LogMaterial, Display, TEXT("Add compile shader map id %d"), NewShaderMap->GetCompilingId());
 			OutstandingCompileShaderMapIds.Add( NewShaderMap->GetCompilingId() );
 			// Async compile, use NULL so that rendering will fall back to the default material.
 			OutShaderMap = NULL;
diff --git a/Engine/Source/Runtime/Engine/Public/MaterialShared.h b/Engine/Source/Runtime/Engine/Public/MaterialShared.h
--- a/Engine/Source/Runtime/Engine/Public/MaterialShared.h
+++ b/Engine/Source/Runtime/Engine/Public/MaterialShared.h
@@ -978,6 +978,14 @@ public:
 	 */
 	ENGINE_API bool IsCompilationFinished();
 
+	/**
+	* Checks if there is a valid GameThreadShaderMap, that is, the material can be rendered as intended.
+	*
+	* @return returns true if there is a GameThreadShaderMap.
+	*/
+	ENGINE_API bool HasValidGameThreadShaderMap();
+
+
 	EMaterialQualityLevel::Type GetQualityLevel() const 
 	{
 		return QualityLevel;
@@ -1040,7 +1048,7 @@ public:
 	void RemoveOutstandingCompileId(const int32 OldOutstandingCompileShaderMapId )
 	{
 		OutstandingCompileShaderMapIds.Remove( OldOutstandingCompileShaderMapId );
-		// UE_LOG(LogMaterial, Display, TEXT("Removing compile shader map id %d"), OldOutstandingCompileShaderMapId);
+		UE_LOG(LogMaterial, Display, TEXT("Removing compile shader map id %d"), OldOutstandingCompileShaderMapId);
 	}
 
 	void AddReferencedObjects(FReferenceCollector& Collector);
diff --git a/Engine/Source/Runtime/Landscape/Classes/LandscapeComponent.h b/Engine/Source/Runtime/Landscape/Classes/LandscapeComponent.h
--- a/Engine/Source/Runtime/Landscape/Classes/LandscapeComponent.h
+++ b/Engine/Source/Runtime/Landscape/Classes/LandscapeComponent.h
@@ -326,6 +326,9 @@ public:
 	void RenderGrassMap();
 	void RemoveGrassMap();
 
+	/* Can a valid grassmap currently be generated? */
+	bool CanRenderGrassMap();
+
 #endif
 
 	/** @todo document */
diff --git a/Engine/Source/Runtime/Landscape/Private/Landscape.cpp b/Engine/Source/Runtime/Landscape/Private/Landscape.cpp
--- a/Engine/Source/Runtime/Landscape/Private/Landscape.cpp
+++ b/Engine/Source/Runtime/Landscape/Private/Landscape.cpp
@@ -79,7 +79,7 @@ FAutoConsoleCommand CmdPrintNumLandscapeShadows(
 	);
 
 ULandscapeComponent::ULandscapeComponent(const FObjectInitializer& ObjectInitializer)
-	: Super(ObjectInitializer)
+: Super(ObjectInitializer)
 , GrassData(MakeShareable(new FLandscapeComponentGrassData()))
 {
 	SetCollisionProfileName(UCollisionProfile::BlockAll_ProfileName);
@@ -109,16 +109,12 @@ ULandscapeComponent::ULandscapeComponent(const FObjectInitializer& ObjectInitial
 #endif
 
 	LpvBiasMultiplier = 0.0f; // Bias is 0 for landscape, since it's single sided
-
-	
 }
 
 ULandscapeComponent::~ULandscapeComponent()
 {
-
 }
 
-
 void ULandscapeComponent::AddReferencedObjects(UObject* InThis, FReferenceCollector& Collector)
 {
 	ULandscapeComponent* This = CastChecked<ULandscapeComponent>(InThis);
@@ -216,12 +212,12 @@ void ULandscapeComponent::Serialize(FArchive& Ar)
 			}
 		}
 		if (Ar.UE4Ver() >= VER_UE4_LANDSCAPE_GRASS_COOKING && Ar.UE4Ver() < VER_UE4_SERIALIZE_LANDSCAPE_GRASS_DATA)
-			{
+		{
 			// deal with previous cooked FGrassMap data
-				int32 NumChannels = 0;
-				Ar << NumChannels;
-				if (NumChannels)
-				{
+			int32 NumChannels = 0;
+			Ar << NumChannels;
+			if (NumChannels)
+			{
 				TArray<uint8> OldData;
 				OldData.BulkSerialize(Ar);
 			}
@@ -344,14 +340,14 @@ void ULandscapeComponent::GetLayerDebugColorKey(int32& R, int32& G, int32& B) co
 #endif	//WITH_EDITOR
 
 ULandscapeMeshCollisionComponent::ULandscapeMeshCollisionComponent(const FObjectInitializer& ObjectInitializer)
-	: Super(ObjectInitializer)
+: Super(ObjectInitializer)
 {
 	// make landscape always create? 
 	bAlwaysCreatePhysicsState = true;
 }
 
 ULandscapeInfo::ULandscapeInfo(const FObjectInitializer& ObjectInitializer)
-	: Super(ObjectInitializer)
+: Super(ObjectInitializer)
 {
 #if WITH_EDITOR
 	bIsValid = false;
@@ -473,7 +469,7 @@ void ULandscapeComponent::PostLoad()
 #endif // WITH_EDITOR
 
 ALandscape::ALandscape(const FObjectInitializer& ObjectInitializer)
-	: Super(ObjectInitializer)
+: Super(ObjectInitializer)
 {
 	bIsProxy = false;
 
@@ -488,7 +484,7 @@ ALandscape* ALandscape::GetLandscapeActor()
 }
 
 ALandscapeProxy::ALandscapeProxy(const FObjectInitializer& ObjectInitializer)
-	: Super(ObjectInitializer)
+: Super(ObjectInitializer)
 {
 	bReplicates = false;
 	NetUpdateFrequency = 10.0f;
@@ -630,10 +626,10 @@ void ULandscapeComponent::GetGeneratedTexturesAndMaterialInstances(TArray<UObjec
 
 			auto* WeightmapPtr = LandscapeMIC->TextureParameterValues.FindByPredicate(
 				[](const FTextureParameterValue& ParamValue)
-				{
-					static const FName WeightmapParamName("Weightmap0");
-					return ParamValue.ParameterName == WeightmapParamName;
-				}
+			{
+				static const FName WeightmapParamName("Weightmap0");
+				return ParamValue.ParameterName == WeightmapParamName;
+			}
 			);
 
 			if (WeightmapPtr != nullptr &&
@@ -1283,12 +1279,12 @@ void ALandscapeProxy::PostLoad()
 		// Need to clean up invalid collision components
 		RecreateCollisionComponents();
 	}
-	
+
 	if (GetLinker() && GetLinker()->UE4Ver() < VER_UE4_LANDSCAPE_STATIC_SECTION_OFFSET)
 	{
 		bStaticSectionOffset = false;
 	}
-	
+
 	EditorLayerSettings.Remove(NULL);
 
 	if (EditorCachedLayerInfos_DEPRECATED.Num() > 0)
@@ -1591,10 +1587,10 @@ void ULandscapeInfo::RegisterActor(ALandscapeProxy* Proxy, bool bMapCheck)
 
 	if (!DrawScale.Equals(Proxy->GetRootComponent()->RelativeScale3D))
 	{
-		UE_LOG(LogLandscape, Warning, TEXT("Landscape proxy (%s) scale (%s) does not match to main actor scale (%s)."), 
+		UE_LOG(LogLandscape, Warning, TEXT("Landscape proxy (%s) scale (%s) does not match to main actor scale (%s)."),
 			*Proxy->GetName(), *Proxy->GetRootComponent()->RelativeScale3D.ToCompactString(), *DrawScale.ToCompactString());
 	}
-	
+
 	// register
 	ALandscape* Landscape = Cast<ALandscape>(Proxy);
 	if (Landscape)
@@ -1797,13 +1793,13 @@ void ULandscapeInfo::FixupProxiesWeightmaps()
 void AdjustLandscapeSectionOffsets(UWorld* InWorld, const TArray<ALandscapeProxy*> InLandscapeList)
 {
 	// We interested only in registered actors
-	TArray<ALandscapeProxy*> RegisteredLandscapeList = InLandscapeList.FilterByPredicate([](ALandscapeProxy* Proxy) { 
-		return Proxy->GetRootComponent()->IsRegistered(); 
+	TArray<ALandscapeProxy*> RegisteredLandscapeList = InLandscapeList.FilterByPredicate([](ALandscapeProxy* Proxy) {
+		return Proxy->GetRootComponent()->IsRegistered();
 	});
 
 	// Main Landscape actor should act as origin of global components grid
-	int32 LandcapeIndex = RegisteredLandscapeList.IndexOfByPredicate([](ALandscapeProxy* Proxy) { 
-		return Proxy->IsA<ALandscape>(); 
+	int32 LandcapeIndex = RegisteredLandscapeList.IndexOfByPredicate([](ALandscapeProxy* Proxy) {
+		return Proxy->IsA<ALandscape>();
 	});
 	ALandscapeProxy* StaticLandscape = RegisteredLandscapeList.IsValidIndex(LandcapeIndex) ? RegisteredLandscapeList[LandcapeIndex] : nullptr;
 
@@ -1811,19 +1807,19 @@ void AdjustLandscapeSectionOffsets(UWorld* InWorld, const TArray<ALandscapeProxy
 	{
 		StaticLandscape->SetAbsoluteSectionBase(FIntPoint::ZeroValue);
 	}
-	
+
 	// In case there is no main landscape actor loaded try use any landscape that already has static offsets
 	if (StaticLandscape == nullptr)
 	{
-		LandcapeIndex = RegisteredLandscapeList.IndexOfByPredicate([](ALandscapeProxy* Proxy) { 
-			return Proxy->bStaticSectionOffset; 
+		LandcapeIndex = RegisteredLandscapeList.IndexOfByPredicate([](ALandscapeProxy* Proxy) {
+			return Proxy->bStaticSectionOffset;
 		});
 		StaticLandscape = RegisteredLandscapeList.IsValidIndex(LandcapeIndex) ? RegisteredLandscapeList[LandcapeIndex] : nullptr;
 		// Otherwise offsets will stay variable
 	}
-		
+
 	ALandscapeProxy* OriginLandscape = StaticLandscape;
-		
+
 	for (ALandscapeProxy* Proxy : RegisteredLandscapeList)
 	{
 		if (OriginLandscape == nullptr)
@@ -2031,9 +2027,9 @@ void WorldDestroyEventFunction(UWorld* World)
 {
 	World->PerModuleDataObjects.RemoveAll(
 		[](UObject* Object)
-		{
-			return Object->IsA(ULandscapeInfoMap::StaticClass());
-		}
+	{
+		return Object->IsA(ULandscapeInfoMap::StaticClass());
+	}
 	);
 }
 
@@ -2107,11 +2103,11 @@ bool LandscapeMaterialsParameterSetUpdater(FStaticParameterSet &StaticParameterS
 }
 
 
-DECLARE_CYCLE_STAT(TEXT("Grass Async Build Time"),STAT_FoliageGrassAsyncBuildTime,STATGROUP_Foliage);
-DECLARE_CYCLE_STAT(TEXT("Grass Start Comp"),STAT_FoliageGrassStartComp,STATGROUP_Foliage);
-DECLARE_CYCLE_STAT(TEXT("Grass End Comp"),STAT_FoliageGrassEndComp,STATGROUP_Foliage);
-DECLARE_CYCLE_STAT(TEXT("Grass Destroy Comps"),STAT_FoliageGrassDestoryComp,STATGROUP_Foliage);
-DECLARE_CYCLE_STAT(TEXT("Grass Update"),STAT_GrassUpdate,STATGROUP_Foliage);
+DECLARE_CYCLE_STAT(TEXT("Grass Async Build Time"), STAT_FoliageGrassAsyncBuildTime, STATGROUP_Foliage);
+DECLARE_CYCLE_STAT(TEXT("Grass Start Comp"), STAT_FoliageGrassStartComp, STATGROUP_Foliage);
+DECLARE_CYCLE_STAT(TEXT("Grass End Comp"), STAT_FoliageGrassEndComp, STATGROUP_Foliage);
+DECLARE_CYCLE_STAT(TEXT("Grass Destroy Comps"), STAT_FoliageGrassDestoryComp, STATGROUP_Foliage);
+DECLARE_CYCLE_STAT(TEXT("Grass Update"), STAT_GrassUpdate, STATGROUP_Foliage);
 
 void ALandscapeProxy::FlushFoliageComponents(const TSet<ULandscapeComponent*>* OnlyForComponents)
 {
@@ -2133,16 +2129,16 @@ void ALandscapeProxy::FlushFoliageComponents(const TSet<ULandscapeComponent*>* O
 			if (Component == nullptr || OnlyForComponents->Contains(Component))
 			{
 				UHierarchicalInstancedStaticMeshComponent *Used = (*Iter).Foliage.Get();
-						if (Used)
-						{
-							SCOPE_CYCLE_COUNTER(STAT_FoliageGrassDestoryComp);
-							Used->ClearInstances();
-							Used->DestroyComponent();
-						}
-				Iter.RemoveCurrent();
-					}
+				if (Used)
+				{
+					SCOPE_CYCLE_COUNTER(STAT_FoliageGrassDestoryComp);
+					Used->ClearInstances();
+					Used->DestroyComponent();
 				}
+				Iter.RemoveCurrent();
 			}
+		}
+	}
 	else
 	{
 		// Clear old foliage component containers
@@ -2423,9 +2419,9 @@ struct FAsyncGrassBuilder : public FGrassBuilderBase
 				InstanceBuffer.AllocateInstances(NumKept);
 				int32 InstanceIndex = 0;
 				int32 OutInstanceIndex = 0;
-				for (int32 xStart = 0; xStart < SqrtMaxInstances; xStart++ )
+				for (int32 xStart = 0; xStart < SqrtMaxInstances; xStart++)
 				{
-					for (int32 yStart = 0; yStart < SqrtMaxInstances; yStart++ )
+					for (int32 yStart = 0; yStart < SqrtMaxInstances; yStart++)
 					{
 						const FInstanceLocal& Instance = Instances[InstanceIndex];
 						if (Instance.bKeep)
@@ -2444,7 +2440,7 @@ struct FAsyncGrassBuilder : public FGrassBuilderBase
 									FVector NewZ = ((PosX1 - PosX2) ^ (PosY1 - PosY2)).GetSafeNormal();
 									NewZ *= FMath::Sign(NewZ.Z);
 
-									const FVector NewX = (FVector(0,-1,0) ^ NewZ).GetSafeNormal();
+									const FVector NewX = (FVector(0, -1, 0) ^ NewZ).GetSafeNormal();
 									const FVector NewY = NewZ ^ NewX;
 
 									FMatrix Align = FMatrix(NewX, NewY, NewZ, FVector::ZeroVector);
@@ -2654,6 +2650,12 @@ void ALandscapeProxy::UpdateFoliage(const TArray<FVector>& Cameras, ULandscapeCo
 					continue;
 				}
 
+				// Don't try to place foliage for this component if it doesn't have data and it can't currently render it.
+				if (!Component->GrassData->HasData() && !Component->CanRenderGrassMap())
+				{
+					continue;
+				}
+
 				FBoxSphereBounds WorldBounds = Component->CalcBounds(Component->ComponentToWorld);
 				float MinDistanceToComp = Cameras.Num() ? MAX_flt : 0.0f;
 
@@ -2666,7 +2668,7 @@ void ALandscapeProxy::UpdateFoliage(const TArray<FVector>& Cameras, ULandscapeCo
 
 				for (auto GrassType : GrassTypes)
 				{
-					if (GrassType)					
+					if (GrassType)
 					{
 						int32 GrassVarietyIndex = -1;
 						for (auto& GrassVariety : GrassType->GrassVarieties)
@@ -2687,7 +2689,7 @@ void ALandscapeProxy::UpdateFoliage(const TArray<FVector>& Cameras, ULandscapeCo
 								FGrassBuilderBase ForSubsectionMath(this, Component, GrassVariety);
 
 								int32 SqrtSubsections = 1;
-				
+
 								if (ForSubsectionMath.bHaveValidData && ForSubsectionMath.SqrtMaxInstances > 0)
 								{
 									SqrtSubsections = FMath::Clamp<int32>(FMath::CeilToInt(float(ForSubsectionMath.SqrtMaxInstances) / FMath::Sqrt((float)MaxInstancesPerComponent)), 1, 16);
@@ -2809,14 +2811,14 @@ void ALandscapeProxy::UpdateFoliage(const TArray<FVector>& Cameras, ULandscapeCo
 											FoliageComponents.Add(HierarchicalInstancedStaticMeshComponent);
 										}
 
-			#if WITH_EDITOR
+#if WITH_EDITOR
 										// render grass data if we don't have any
 										if (!Component->GrassData->HasData())
 										{
 											QUICK_SCOPE_CYCLE_COUNTER(STAT_GrassRenderToTexture);
 											Component->RenderGrassMap();
 										}
-			#endif
+#endif
 
 										FAsyncGrassBuilder* Builder;
 
@@ -2861,13 +2863,13 @@ void ALandscapeProxy::UpdateFoliage(const TArray<FVector>& Cameras, ULandscapeCo
 		{
 			const FCachedLandscapeFoliage::FGrassComp& GrassItem = *Iter;
 			UHierarchicalInstancedStaticMeshComponent *Used = GrassItem.Foliage.Get();
-			bool bOld = 
+			bool bOld =
 				!GrassItem.Pending &&
 				(
-					!GrassItem.Key.BasedOn.Get() ||
-					!GrassItem.Key.GrassType.Get() ||
-					!Used ||
-					(GrassItem.LastUsedFrameNumber < OldestToKeepFrame && GrassItem.LastUsedTime < OldestToKeepTime)
+				!GrassItem.Key.BasedOn.Get() ||
+				!GrassItem.Key.GrassType.Get() ||
+				!Used ||
+				(GrassItem.LastUsedFrameNumber < OldestToKeepFrame && GrassItem.LastUsedTime < OldestToKeepTime)
 				);
 			if (bOld)
 			{
@@ -2964,10 +2966,10 @@ void FAsyncGrassTask::DoWork()
 
 static void FlushFoliage(const TArray<FString>& Args)
 {
-	for( TObjectIterator<ALandscapeProxy> It; It; ++It )
+	for (TObjectIterator<ALandscapeProxy> It; It; ++It)
 	{
 		ALandscapeProxy* Landscape = *It;
-		if( Landscape && !Landscape->IsTemplate() && !Landscape->IsPendingKill())
+		if (Landscape && !Landscape->IsTemplate() && !Landscape->IsPendingKill())
 		{
 			Landscape->FlushFoliageComponents();
 		}
diff --git a/Engine/Source/Runtime/Landscape/Private/LandscapeGrass.cpp b/Engine/Source/Runtime/Landscape/Private/LandscapeGrass.cpp
--- a/Engine/Source/Runtime/Landscape/Private/LandscapeGrass.cpp
+++ b/Engine/Source/Runtime/Landscape/Private/LandscapeGrass.cpp
@@ -5,8 +5,8 @@
 #include "EngineGlobals.h"
 #include "UnrealEngine.h"
 #include "EngineUtils.h"
-#include "CanvasTypes.h"
 #include "Shader.h"
+#include "ShaderParameterUtils.h"
 #include "MeshMaterialShader.h"
 #include "EngineModule.h"
 #include "RendererInterface.h"
@@ -23,13 +23,20 @@
 class FLandscapeGrassWeightVS : public FMeshMaterialShader
 {
 	DECLARE_SHADER_TYPE(FLandscapeGrassWeightVS, MeshMaterial);
+
+	FShaderParameter RenderOffsetParameter;
+
 protected:
 
-	FLandscapeGrassWeightVS() {}
-	FLandscapeGrassWeightVS(const FMeshMaterialShaderType::CompiledShaderInitializerType& Initializer) :
-		FMeshMaterialShader(Initializer)
+	FLandscapeGrassWeightVS()
 	{}
 
+	FLandscapeGrassWeightVS(const FMeshMaterialShaderType::CompiledShaderInitializerType& Initializer)
+	: FMeshMaterialShader(Initializer)
+	{
+		RenderOffsetParameter.Bind(Initializer.ParameterMap, TEXT("RenderOffset"));
+	}
+
 public:
 
 	static bool ShouldCache(EShaderPlatform Platform, const FMaterial* Material, const FVertexFactoryType* VertexFactoryType)
@@ -38,15 +45,23 @@ public:
 		return (Material->IsUsedWithLandscape() || Material->IsSpecialEngineMaterial()) && IsFeatureLevelSupported(Platform, ERHIFeatureLevel::SM4) && (VertexFactoryType == FindVertexFactoryType(FName(TEXT("FLandscapeVertexFactory"), FNAME_Find))) && !IsConsolePlatform(Platform);
 	}
 
-	void SetParameters(FRHICommandList& RHICmdList, const FMaterialRenderProxy* MaterialRenderProxy, const FMaterial& MaterialResource, const FSceneView& View)
+	void SetParameters(FRHICommandList& RHICmdList, const FMaterialRenderProxy* MaterialRenderProxy, const FMaterial& MaterialResource, const FSceneView& View, const FVector2D& RenderOffset)
 	{
 		FMeshMaterialShader::SetParameters(RHICmdList, GetVertexShader(), MaterialRenderProxy, MaterialResource, View, ESceneRenderTargetsMode::DontSet);
+		SetShaderValue(RHICmdList, GetVertexShader(), RenderOffsetParameter, RenderOffset);
 	}
 
 	void SetMesh(FRHICommandList& RHICmdList, const FVertexFactory* VertexFactory, const FSceneView& View, const FPrimitiveSceneProxy* Proxy, const FMeshBatchElement& BatchElement)
 	{
 		FMeshMaterialShader::SetMesh(RHICmdList, GetVertexShader(), VertexFactory, View, Proxy, BatchElement);
 	}
+
+	virtual bool Serialize(FArchive& Ar) override
+	{
+		bool bShaderHasOutdatedParameters = FMeshMaterialShader::Serialize(Ar);
+		Ar << RenderOffsetParameter;
+		return bShaderHasOutdatedParameters;
+	}
 };
 
 IMPLEMENT_MATERIAL_SHADER_TYPE(, FLandscapeGrassWeightVS, TEXT("LandscapeGrassWeight"), TEXT("VSMain"), SF_Vertex);
@@ -55,6 +70,7 @@ IMPLEMENT_MATERIAL_SHADER_TYPE(, FLandscapeGrassWeightVS, TEXT("LandscapeGrassWe
 class FLandscapeGrassWeightPS : public FMeshMaterialShader
 {
 	DECLARE_SHADER_TYPE(FLandscapeGrassWeightPS, MeshMaterial);
+	FShaderParameter OutputPassParameter;
 public:
 
 	static bool ShouldCache(EShaderPlatform Platform, const FMaterial* Material, const FVertexFactoryType* VertexFactoryType)
@@ -63,21 +79,35 @@ public:
 		return (Material->IsUsedWithLandscape() || Material->IsSpecialEngineMaterial()) && IsFeatureLevelSupported(Platform, ERHIFeatureLevel::SM4) && (VertexFactoryType == FindVertexFactoryType(FName(TEXT("FLandscapeVertexFactory"), FNAME_Find))) && !IsConsolePlatform(Platform);
 	}
 
-	FLandscapeGrassWeightPS(const ShaderMetaType::CompiledShaderInitializerType& Initializer) :
-		FMeshMaterialShader(Initializer)
-	{}
+	FLandscapeGrassWeightPS(const ShaderMetaType::CompiledShaderInitializerType& Initializer)
+	: FMeshMaterialShader(Initializer)
+	{
+		OutputPassParameter.Bind(Initializer.ParameterMap, TEXT("OutputPass"));
+	}
 
-	FLandscapeGrassWeightPS() {}
+	FLandscapeGrassWeightPS()
+	{}
 
-	void SetParameters(FRHICommandList& RHICmdList, const FMaterialRenderProxy* MaterialRenderProxy, const FMaterial& MaterialResource, const FSceneView* View)
+	void SetParameters(FRHICommandList& RHICmdList, const FMaterialRenderProxy* MaterialRenderProxy, const FMaterial& MaterialResource, const FSceneView* View, int32 OutputPass)
 	{
 		FMeshMaterialShader::SetParameters(RHICmdList, GetPixelShader(), MaterialRenderProxy, MaterialResource, *View, ESceneRenderTargetsMode::DontSet);
+		if (OutputPassParameter.IsBound())
+		{
+			SetShaderValue(RHICmdList, GetPixelShader(), OutputPassParameter, OutputPass);
+		}
 	}
 
 	void SetMesh(FRHICommandList& RHICmdList, const FVertexFactory* VertexFactory, const FSceneView& View, const FPrimitiveSceneProxy* Proxy, const FMeshBatchElement& BatchElement)
 	{
 		FMeshMaterialShader::SetMesh(RHICmdList, GetPixelShader(), VertexFactory, View, Proxy, BatchElement);
 	}
+
+	virtual bool Serialize(FArchive& Ar) override
+	{
+		bool bShaderHasOutdatedParameters = FMeshMaterialShader::Serialize(Ar);
+		Ar << OutputPassParameter;
+		return bShaderHasOutdatedParameters;
+	}
 };
 
 IMPLEMENT_MATERIAL_SHADER_TYPE(, FLandscapeGrassWeightPS, TEXT("LandscapeGrassWeight"), TEXT("PSMain"), SF_Pixel);
@@ -110,11 +140,11 @@ public:
 			&& PixelShader == Other.PixelShader;
 	}
 
-	void SetSharedState(FRHICommandList& RHICmdList, const FSceneView* View, const ContextDataType PolicyContext) const
+	void SetSharedState(FRHICommandList& RHICmdList, const FSceneView* View, const ContextDataType PolicyContext, int32 OutputPass, const FVector2D& RenderOffset) const
 	{
-		// Set the depth-only shader parameters for the material.
-		VertexShader->SetParameters(RHICmdList, MaterialRenderProxy, *MaterialResource, *View);
-		PixelShader->SetParameters(RHICmdList, MaterialRenderProxy, *MaterialResource, View);
+		// Set the shader parameters for the material.
+		VertexShader->SetParameters(RHICmdList, MaterialRenderProxy, *MaterialResource, *View, RenderOffset);
+		PixelShader->SetParameters(RHICmdList, MaterialRenderProxy, *MaterialResource, View, OutputPass);
 
 		// Set the shared mesh resources.
 		FMeshDrawingPolicy::SetSharedState(RHICmdList, View, PolicyContext);
@@ -162,73 +192,59 @@ private:
 	FLandscapeGrassWeightPS* PixelShader;
 };
 
-
-class FLandscapeGrassWeightExporter
+// data also accessible by render thread
+class FLandscapeGrassWeightExporter_RenderThread
 {
-	static void RenderLandscapeComponentToTexture(
-		ULandscapeComponent* LandscapeComponent,
-		const FMatrix& ViewMatrix,
-		const FMatrix& ProjectionMatrix,
-		FIntPoint TargetSize,
-		float TargetGamma,
-		TArray<FColor>& OutSamples)
+	FLandscapeGrassWeightExporter_RenderThread(int32 InNumPasses)
+	: RenderTargetResource(nullptr)
+	, NumPasses(InNumPasses)
+	{}
+
+	friend class FLandscapeGrassWeightExporter;
+
+public:
+	virtual ~FLandscapeGrassWeightExporter_RenderThread()
+	{}
+
+	struct FComponentInfo
 	{
-		UTextureRenderTarget2D* RenderTargetTexture = new UTextureRenderTarget2D(FObjectInitializer());
-		check(RenderTargetTexture);
-		RenderTargetTexture->AddToRoot();
-		RenderTargetTexture->ClearColor = FLinearColor::Transparent;
-		RenderTargetTexture->TargetGamma = TargetGamma;
-		RenderTargetTexture->InitCustomFormat(TargetSize.X, TargetSize.Y, PF_B8G8R8A8, false); 
-		FTextureRenderTarget2DResource* RenderTargetResource = RenderTargetTexture->GameThread_GetRenderTargetResource()->GetTextureRenderTarget2DResource();
-		RenderTargetResource->GetRenderTargetTexture();
-
-		FSceneViewFamilyContext ViewFamily(
-			FSceneViewFamily::ConstructionValues(RenderTargetResource, NULL, FEngineShowFlags(ESFIM_Game))
-			.SetWorldTimes(FApp::GetCurrentTime() - GStartTime, FApp::GetDeltaTime(), FApp::GetCurrentTime() - GStartTime)
-			);
-
-		// To enable visualization mode
-		ViewFamily.EngineShowFlags.SetPostProcessing(true);
-		ViewFamily.EngineShowFlags.SetVisualizeBuffer(true);
-		ViewFamily.EngineShowFlags.SetTonemapper(false);
-
-		// Force LOD render
+		ULandscapeComponent* Component;
+		FVector2D ViewOffset;
+		int32 PixelOffsetX;
+		FPrimitiveSceneInfo* PrimitiveSceneInfo;
+
+		FComponentInfo(ULandscapeComponent* InComponent, FVector2D& InViewOffset, int32 InPixelOffsetX)
+			: Component(InComponent)
+			, ViewOffset(InViewOffset)
+			, PixelOffsetX(InPixelOffsetX)
+			, PrimitiveSceneInfo(InComponent->SceneProxy->GetPrimitiveSceneInfo())
+		{}
+	};
+
+	FTextureRenderTarget2DResource* RenderTargetResource;
+	TArray<FComponentInfo> ComponentInfos;
+	FIntPoint TargetSize;
+	int32 NumPasses;
+	int32 PassOffsetX;
+	FMatrix ViewMatrix;
+	FMatrix ProjectionMatrix;
+
+	void RenderLandscapeComponentToTexture_RenderThread(FRHICommandListImmediate& RHICmdList)
+	{
+		FSceneViewFamilyContext ViewFamily(FSceneViewFamily::ConstructionValues(RenderTargetResource, NULL, FEngineShowFlags(ESFIM_Game)).SetWorldTimes(FApp::GetCurrentTime() - GStartTime, FApp::GetDeltaTime(), FApp::GetCurrentTime() - GStartTime));
+
 #if WITH_EDITOR
-		ViewFamily.LandscapeLODOverride = 0;
+		ViewFamily.LandscapeLODOverride = 0; 		// Force LOD render
 #endif
+
 		FSceneViewInitOptions ViewInitOptions;
 		ViewInitOptions.SetViewRectangle(FIntRect(0, 0, TargetSize.X, TargetSize.Y));
-		ViewInitOptions.ViewFamily = &ViewFamily;
 		ViewInitOptions.ViewMatrix = ViewMatrix;
 		ViewInitOptions.ProjectionMatrix = ProjectionMatrix;
+		ViewInitOptions.ViewFamily = &ViewFamily;
 
-		FCanvas Canvas(RenderTargetResource, NULL, FApp::GetCurrentTime() - GStartTime, FApp::GetDeltaTime(), FApp::GetCurrentTime() - GStartTime, GMaxRHIFeatureLevel);
-		Canvas.Clear(FLinearColor::Transparent);
-
-		GetRendererModule().CreateAndInitSingleView(&ViewFamily, &ViewInitOptions);
+		GetRendererModule().CreateAndInitSingleView(RHICmdList, &ViewFamily, &ViewInitOptions);
 		
-		// render
-		ENQUEUE_UNIQUE_RENDER_COMMAND_TWOPARAMETER(
-			FDrawSceneCommand,
-			FSceneViewFamily&, ViewFamily, ViewFamily,
-			FPrimitiveSceneInfo*, PrimitiveSceneInfo, LandscapeComponent->SceneProxy->GetPrimitiveSceneInfo(),
-			{
-				RenderLandscapeComponentToTexture_RenderThread(RHICmdList, ViewFamily, PrimitiveSceneInfo);
-				FlushPendingDeleteRHIResources_RenderThread();
-			});
-
-		// Copy the contents of the remote texture to system memory
-		OutSamples.Init(TargetSize.X*TargetSize.Y);
-		FReadSurfaceDataFlags ReadSurfaceDataFlags;
-		ReadSurfaceDataFlags.SetLinearToGamma(false);
-		RenderTargetResource->ReadPixels(OutSamples, ReadSurfaceDataFlags, FIntRect(0, 0, TargetSize.X, TargetSize.Y));
-
-		RenderTargetTexture->RemoveFromRoot();
-		RenderTargetTexture = nullptr;
-	}
-
-	static void RenderLandscapeComponentToTexture_RenderThread(FRHICommandListImmediate& RHICmdList, FSceneViewFamily& ViewFamily, FPrimitiveSceneInfo* PrimitiveSceneInfo)
-	{
 		const FSceneView* View = ViewFamily.Views[0];
 		RHICmdList.SetViewport(View->ViewRect.Min.X, View->ViewRect.Min.Y, 0.0f, View->ViewRect.Max.X, View->ViewRect.Max.Y, 1.0f);
 		RHICmdList.SetBlendState(TStaticBlendState<>::GetRHI());
@@ -236,40 +252,75 @@ class FLandscapeGrassWeightExporter
 		RHICmdList.SetDepthStencilState(TStaticDepthStencilState<false, CF_Always>::GetRHI());
 		RHICmdList.SetScissorRect(false, 0, 0, 0, 0);
 
-		for (int32 StaticMeshIdx = 0; StaticMeshIdx < PrimitiveSceneInfo->StaticMeshes.Num(); StaticMeshIdx++)
+		for (auto& ComponentInfo : ComponentInfos)
 		{
-			FMeshBatch& Mesh = *(FMeshBatch*)(&PrimitiveSceneInfo->StaticMeshes[StaticMeshIdx]);
+			FPrimitiveSceneInfo* PrimitiveSceneInfo = ComponentInfo.PrimitiveSceneInfo;
+			for (int32 PassIdx = 0; PassIdx < NumPasses; PassIdx++)
+			{
+				for (int32 StaticMeshIdx = 0; StaticMeshIdx < PrimitiveSceneInfo->StaticMeshes.Num(); StaticMeshIdx++)
+				{
+					FMeshBatch& Mesh = *(FMeshBatch*)(&PrimitiveSceneInfo->StaticMeshes[StaticMeshIdx]);
 
-			FLandscapeGrassWeightDrawingPolicy DrawingPolicy(Mesh.VertexFactory, Mesh.MaterialRenderProxy, *Mesh.MaterialRenderProxy->GetMaterial(GMaxRHIFeatureLevel));
-			RHICmdList.BuildAndSetLocalBoundShaderState(DrawingPolicy.GetBoundShaderStateInput(GMaxRHIFeatureLevel));
-			DrawingPolicy.SetSharedState(RHICmdList, View, FLandscapeGrassWeightDrawingPolicy::ContextDataType());
+					FLandscapeGrassWeightDrawingPolicy DrawingPolicy(Mesh.VertexFactory, Mesh.MaterialRenderProxy, *Mesh.MaterialRenderProxy->GetMaterial(GMaxRHIFeatureLevel));
+					RHICmdList.BuildAndSetLocalBoundShaderState(DrawingPolicy.GetBoundShaderStateInput(GMaxRHIFeatureLevel));
+					DrawingPolicy.SetSharedState(RHICmdList, View, FLandscapeGrassWeightDrawingPolicy::ContextDataType(), PassIdx, ComponentInfo.ViewOffset + FVector2D(PassOffsetX * PassIdx, 0));
 
-			// The first batch element contains the grass batch for the entire component
-			DrawingPolicy.SetMeshRenderState(RHICmdList, *View, PrimitiveSceneInfo->Proxy, Mesh, 0, false, FMeshDrawingPolicy::ElementDataType(), FLandscapeGrassWeightDrawingPolicy::ContextDataType());
-			DrawingPolicy.DrawMesh(RHICmdList, Mesh, 0);
+					// The first batch element contains the grass batch for the entire component
+					DrawingPolicy.SetMeshRenderState(RHICmdList, *View, PrimitiveSceneInfo->Proxy, Mesh, 0, false, FMeshDrawingPolicy::ElementDataType(), FLandscapeGrassWeightDrawingPolicy::ContextDataType());
+					DrawingPolicy.DrawMesh(RHICmdList, Mesh, 0);
+				}
+			}
 		}
 	}
+};
+
+
+
+class FLandscapeGrassWeightExporter : public FLandscapeGrassWeightExporter_RenderThread
+{
+	ALandscapeProxy* LandscapeProxy;
+	int32 ComponentSizeQuads;
+	int32 ComponentSizeVerts;
+	TArray<ULandscapeGrassType*> GrassTypes;
+	UTextureRenderTarget2D* RenderTargetTexture;
 
 public:
-	
-	static bool GetLandscapeComponentGrassMap(ULandscapeComponent* LandscapeComponent, TArray<FColor>& OutSamples)
+
+	FLandscapeGrassWeightExporter(ALandscapeProxy* InLandscapeProxy, const TArray<ULandscapeComponent*>& InLandscapeComponents, const TArray<ULandscapeGrassType*> InGrassTypes) 
+	:	FLandscapeGrassWeightExporter_RenderThread((InGrassTypes.Num() + 5) / 4)
+	,	LandscapeProxy(InLandscapeProxy)
+	,	ComponentSizeQuads(InLandscapeProxy->ComponentSizeQuads)
+	,	ComponentSizeVerts(ComponentSizeQuads + 1)
+	,	GrassTypes(InGrassTypes)
+	,	RenderTargetTexture(nullptr)
 	{
-		check(LandscapeComponent);
+		check(InLandscapeComponents.Num());
+
+		TargetSize = FIntPoint(ComponentSizeVerts * NumPasses * InLandscapeComponents.Num(), ComponentSizeVerts);
+		FIntPoint TargetSizeMinusOne(TargetSize - FIntPoint(1,1));
+		PassOffsetX = 2.0f * ComponentSizeVerts / TargetSize.X;
 
-		if (!LandscapeComponent->SceneProxy)
+		for (int32 Idx = 0; Idx<InLandscapeComponents.Num();Idx++)
 		{
-			return false;
-		}
-		ALandscapeProxy* LandscapeProxy = LandscapeComponent->GetLandscapeProxy();
+			ULandscapeComponent* Component = InLandscapeComponents[Idx];
+
+			FIntPoint ComponentOffset = (Component->GetSectionBase() - LandscapeProxy->LandscapeSectionOffset);
+			int32 PixelOffsetX = Idx * NumPasses * ComponentSizeVerts;
 
-		FIntRect LandscapeRect = FIntRect(0, 0, LandscapeComponent->ComponentSizeQuads, LandscapeComponent->ComponentSizeQuads);
-		LandscapeRect += LandscapeComponent->GetSectionBase() - LandscapeProxy->LandscapeSectionOffset;
+			FVector2D ViewOffset(-ComponentOffset.X,ComponentOffset.Y);
+			ViewOffset.X += PixelOffsetX;
+			ViewOffset /= (FVector2D(TargetSize) * 0.5f);
+
+			new(ComponentInfos)FComponentInfo(Component, ViewOffset, PixelOffsetX);
+		}
+		
+		// center of target area in world
+		FVector TargetCenter = LandscapeProxy->GetTransform().TransformPosition(FVector(TargetSizeMinusOne, 0.f)*0.5f);
 
-		FVector MidPoint = FVector(LandscapeRect.Min, 0.f) + FVector(LandscapeRect.Size(), 0.f)*0.5f;
-		FVector LandscapeCenter = LandscapeProxy->GetTransform().TransformPosition(MidPoint);
-		FVector LandscapeExtent = FVector(LandscapeRect.Size() + FIntPoint(1,1), 0.f)*LandscapeProxy->GetActorScale()*0.5f;
+		// extent of target in world space
+		FVector TargetExtent = FVector(TargetSize, 0.f)*LandscapeProxy->GetActorScale()*0.5f;
 
-		FMatrix ViewMatrix = FTranslationMatrix(-LandscapeCenter);
+		ViewMatrix = FTranslationMatrix(-TargetCenter);
 		ViewMatrix *= FInverseRotationMatrix(LandscapeProxy->GetActorRotation());
 		ViewMatrix *= FMatrix(FPlane(1, 0, 0, 0),
 			FPlane(0, -1, 0, 0),
@@ -277,75 +328,188 @@ public:
 			FPlane(0, 0, 0, 1));
 
 		const float ZOffset = WORLD_MAX;
-		FMatrix ProjectionMatrix = FReversedZOrthoMatrix(
-			LandscapeExtent.X,
-			LandscapeExtent.Y,
+		ProjectionMatrix = FReversedZOrthoMatrix(
+			TargetExtent.X,
+			TargetExtent.Y,
 			0.5f / ZOffset,
 			ZOffset);
 
-		FIntPoint TargetSize(LandscapeComponent->ComponentSizeQuads + 1, LandscapeComponent->ComponentSizeQuads + 1);
+		RenderTargetTexture = new UTextureRenderTarget2D(FObjectInitializer());
+		check(RenderTargetTexture);
+		RenderTargetTexture->ClearColor = FLinearColor::Transparent;
+		RenderTargetTexture->TargetGamma = 1.f;
+		RenderTargetTexture->InitCustomFormat(TargetSize.X, TargetSize.Y, PF_B8G8R8A8, false);
+		RenderTargetResource = RenderTargetTexture->GameThread_GetRenderTargetResource()->GetTextureRenderTarget2DResource();
 
-		//TargetSize *= 8;
+		// render
+		ENQUEUE_UNIQUE_RENDER_COMMAND_ONEPARAMETER(
+			FDrawSceneCommand,
+			FLandscapeGrassWeightExporter_RenderThread*, Exporter, this,
+			{
+				Exporter->RenderLandscapeComponentToTexture_RenderThread(RHICmdList);
+				FlushPendingDeleteRHIResources_RenderThread();
+			});
+	}
 
-		const float Gamma = 1.0f;
-		RenderLandscapeComponentToTexture(LandscapeComponent, ViewMatrix, ProjectionMatrix, TargetSize, Gamma, OutSamples);
-		return true;
+	void ApplyResults()
+	{
+		TArray<FColor> Samples;
+		Samples.Init(TargetSize.X*TargetSize.Y);
+
+		// Copy the contents of the remote texture to system memory
+		FReadSurfaceDataFlags ReadSurfaceDataFlags;
+		ReadSurfaceDataFlags.SetLinearToGamma(false);
+		RenderTargetResource->ReadPixels(Samples, ReadSurfaceDataFlags, FIntRect(0, 0, TargetSize.X, TargetSize.Y));
+
+		for (auto& ComponentInfo : ComponentInfos)
+		{
+			FLandscapeComponentGrassData* NewGrassData = new FLandscapeComponentGrassData();
+
+			NewGrassData->HeightData.Empty(FMath::Square(ComponentSizeVerts));
+
+			TArray<TArray<uint8>*> GrassWeightArrays;
+			GrassWeightArrays.Empty(GrassTypes.Num());
+			for (auto GrassType : GrassTypes)
+			{
+				int32 Index = GrassWeightArrays.Add(&NewGrassData->WeightData.Add(GrassType));
+				GrassWeightArrays[Index]->Empty(FMath::Square(ComponentSizeVerts));
+			}
+
+			for (int32 PassIdx = 0; PassIdx < NumPasses; PassIdx++)
+			{
+				FColor* SampleData = &Samples[ComponentInfo.PixelOffsetX + PassIdx*ComponentSizeVerts];
+				if (PassIdx == 0)
+				{
+					NewGrassData->HeightData.Empty(FMath::Square(ComponentSizeVerts));
+
+					for (int32 y = 0; y < ComponentSizeVerts; y++)
+					{
+						for (int32 x = 0; x < ComponentSizeVerts; x++)
+						{
+							FColor& Sample = SampleData[x + y * TargetSize.X];
+							uint16 Height = (((uint16)Sample.R) << 8) + (uint16)(Sample.G);
+							NewGrassData->HeightData.Add(Height);
+							GrassWeightArrays[0]->Add(Sample.B);
+							if (GrassTypes.Num() > 1)
+							{
+								GrassWeightArrays[1]->Add(Sample.A);
+							}
+						}
+					}
+				}
+				else
+				{								
+					for (int32 y = 0; y < ComponentSizeVerts; y++)
+					{
+						for (int32 x = 0; x < ComponentSizeVerts; x++)
+						{
+							FColor& Sample = SampleData[x + y * TargetSize.X];
+
+							int32 TypeIdx = PassIdx * 4 - 2;
+							GrassWeightArrays[TypeIdx++]->Add(Sample.R);
+							if (TypeIdx < GrassTypes.Num())
+							{
+								GrassWeightArrays[TypeIdx++]->Add(Sample.G);
+								if (TypeIdx < GrassTypes.Num())
+								{
+									GrassWeightArrays[TypeIdx++]->Add(Sample.B);
+									if (TypeIdx < GrassTypes.Num())
+									{
+										GrassWeightArrays[TypeIdx++]->Add(Sample.A);
+									}
+								}
+							}
+						}
+					}
+				}
+			}
+
+			// Assign the new data (thread-safe)
+			ComponentInfo.Component->GrassData = MakeShareable(NewGrassData);
+		}
+	}
+
+	void AddReferencedObjects(UObject* InThis, FReferenceCollector& Collector)
+	{
+		if (RenderTargetTexture)
+		{
+			Collector.AddReferencedObject(RenderTargetTexture);
+		}
+
+		if (LandscapeProxy)
+		{
+			Collector.AddReferencedObject(LandscapeProxy);
+		}
+
+		for (auto& Info : ComponentInfos)
+		{
+			if (Info.Component)
+			{
+				Collector.AddReferencedObject(Info.Component);
+			}
+		}
+
+		for (auto GrassType : GrassTypes)
+		{
+			if (GrassType)
+			{
+				Collector.AddReferencedObject(GrassType);
+			}
+		}
 	}
 };
 
 #if WITH_EDITOR
-void ULandscapeComponent::RenderGrassMap()
+
+bool ULandscapeComponent::CanRenderGrassMap()
 {
-	UMaterialInterface* Material = GetLandscapeMaterial();
-	if (GIsEditor && !GUsingNullRHI && SceneProxy && Material)
+	// Check we can render
+	if (!GIsEditor || GUsingNullRHI || !GetWorld() || GetWorld()->FeatureLevel < ERHIFeatureLevel::SM4 || !SceneProxy)
+	{
+		return false;
+	}
+		
+	// Check we can render the material
+	if (!MaterialInstance->GetMaterialResource(GetWorld()->FeatureLevel)->HasValidGameThreadShaderMap())
+	{
+		return false;
+	}
+	
+	// Check for valid heightmap that is fully streamed in
+	if (!HeightmapTexture || HeightmapTexture->ResidentMips != HeightmapTexture->GetNumMips())
 	{
-		FLandscapeComponentGrassData* NewGrassData = new FLandscapeComponentGrassData();
+		return false;
+	}
 
-		TArray<FColor> Samples;
-		FLandscapeGrassWeightExporter::GetLandscapeComponentGrassMap(this, Samples);
+	return true;
+}
 
+void ULandscapeComponent::RenderGrassMap()
+{
+	UMaterialInterface* Material = GetLandscapeMaterial();
+	if (CanRenderGrassMap())
+	{
 		TArray<const UMaterialExpressionLandscapeGrassOutput*> GrassExpressions;
 		Material->GetMaterial()->GetAllExpressionsOfType<UMaterialExpressionLandscapeGrassOutput>(GrassExpressions);
 		if (GrassExpressions.Num() > 0)
 		{
-			int NumTypes = FMath::Min<int32>(GrassExpressions[0]->GrassTypes.Num(), 2); // Only 2 layers currently supported
-			if (NumTypes > 0)
+			TArray<ULandscapeGrassType*> GrassTypes;
+			GrassTypes.Empty(GrassExpressions[0]->GrassTypes.Num());
+			for (auto& GrassTypeInput : GrassExpressions[0]->GrassTypes)
 			{
-				TArray<uint8>* GrassWeightArrays[2] = { nullptr, nullptr };
-
-				// Initialize arrays
-				NewGrassData->HeightData.Empty(Samples.Num());
-				for (int32 Index = 0; Index < NumTypes; Index++)
+				if (GrassTypeInput.GrassType)
 				{
-					GrassWeightArrays[Index] = &NewGrassData->WeightData.Add(GrassExpressions[0]->GrassTypes[Index].GrassType);
-					GrassWeightArrays[Index]->Empty(Samples.Num());
-				}
-
-				// Fill data
-				for (FColor& Sample : Samples)
-				{
-					uint16 Height = (((uint16)Sample.R) << 8) + (uint16)(Sample.G);
-					check((Height >> 8) == Sample.R);
-					check((Height & 255) == Sample.G);
-					NewGrassData->HeightData.Add(Height);
-					GrassWeightArrays[0]->Add(Sample.B);
-					if (NumTypes>1)
-					{
-						GrassWeightArrays[1]->Add(Sample.A);
-					}
+					GrassTypes.Add(GrassTypeInput.GrassType);
 				}
 			}
-		}
 
-		FArchive* Ar = IFileManager::Get().CreateFileWriter(TEXT("D:\\Heightmaps\\Rendered.r16"), FILEWRITE_EvenIfReadOnly);
-		if (Ar)
-		{
-			Ar->Serialize(NewGrassData->HeightData.GetData(), NewGrassData->HeightData.Num() * sizeof(uint16));
-			delete Ar;
+			TArray<ULandscapeComponent*> LandscapeComponents;
+			LandscapeComponents.Add(this);
+
+			FLandscapeGrassWeightExporter Exporter(GetLandscapeProxy(), LandscapeComponents, GrassTypes);
+			Exporter.ApplyResults();
 		}
 
-		// Assign the new data (thread-safe)
-		GrassData = MakeShareable(NewGrassData);
 	}
 }
 
diff --git a/Engine/Source/Runtime/Renderer/Private/RendererPrivate.h b/Engine/Source/Runtime/Renderer/Private/RendererPrivate.h
--- a/Engine/Source/Runtime/Renderer/Private/RendererPrivate.h
+++ b/Engine/Source/Runtime/Renderer/Private/RendererPrivate.h
@@ -23,7 +23,7 @@ public:
 	virtual bool SupportsDynamicReloading() override { return true; }
 
 	virtual void BeginRenderingViewFamily(FCanvas* Canvas,FSceneViewFamily* ViewFamily) override;
-	virtual void CreateAndInitSingleView(class FSceneViewFamily* ViewFamily, const struct FSceneViewInitOptions* ViewInitOptions);
+	virtual void CreateAndInitSingleView(FRHICommandListImmediate& RHICmdList, class FSceneViewFamily* ViewFamily, const struct FSceneViewInitOptions* ViewInitOptions);
 	virtual FSceneInterface* AllocateScene(UWorld* World, bool bInRequiresHitProxies, bool bCreateFXSystem, ERHIFeatureLevel::Type InFeatureLevel) override;
 	virtual void RemoveScene(FSceneInterface* Scene) override;
 	virtual void UpdateStaticDrawListsForMaterials(const TArray<const FMaterial*>& Materials) override;
diff --git a/Engine/Source/Runtime/Renderer/Private/SceneRendering.cpp b/Engine/Source/Runtime/Renderer/Private/SceneRendering.cpp
--- a/Engine/Source/Runtime/Renderer/Private/SceneRendering.cpp
+++ b/Engine/Source/Runtime/Renderer/Private/SceneRendering.cpp
@@ -1332,21 +1332,14 @@ static void RenderViewFamily_RenderThread(FRHICommandListImmediate& RHICmdList,
 #endif
 }
 
-void FRendererModule::CreateAndInitSingleView(class FSceneViewFamily* ViewFamily, const struct FSceneViewInitOptions* ViewInitOptions)
+void FRendererModule::CreateAndInitSingleView(FRHICommandListImmediate& RHICmdList, class FSceneViewFamily* ViewFamily, const struct FSceneViewInitOptions* ViewInitOptions)
 {
 	// Create and add the new view
 	FViewInfo* NewView = new FViewInfo(*ViewInitOptions);
 	ViewFamily->Views.Add(NewView);
-
-	// Ask the render thread to init its resources
-	ENQUEUE_UNIQUE_RENDER_COMMAND_ONEPARAMETER(
-		InitSetViewCommand,
-		FSceneViewFamily*, ViewFamily, ViewFamily,
-		{
-			SetRenderTarget(RHICmdList, ViewFamily->RenderTarget->GetRenderTargetTexture(), nullptr, ESimpleRenderTargetMode::EClearColorToWhite);
-			FViewInfo* View = (FViewInfo*)ViewFamily->Views[0];
-			View->InitRHIResources(nullptr);
-		});
+	SetRenderTarget(RHICmdList, ViewFamily->RenderTarget->GetRenderTargetTexture(), nullptr, ESimpleRenderTargetMode::EClearColorToWhite);
+	FViewInfo* View = (FViewInfo*)ViewFamily->Views[0];
+	View->InitRHIResources(nullptr);
 }
 
 void FRendererModule::BeginRenderingViewFamily(FCanvas* Canvas,FSceneViewFamily* ViewFamily)
diff --git a/Engine/Source/Runtime/Renderer/Public/RendererInterface.h b/Engine/Source/Runtime/Renderer/Public/RendererInterface.h
--- a/Engine/Source/Runtime/Renderer/Public/RendererInterface.h
+++ b/Engine/Source/Runtime/Renderer/Public/RendererInterface.h
@@ -478,8 +478,8 @@ public:
 	/** Call from the game thread to send a message to the rendering thread to being rendering this view family. */
 	virtual void BeginRenderingViewFamily(FCanvas* Canvas, FSceneViewFamily* ViewFamily) = 0;
 	
-	/** Call from the game thread to create and initalize a new FViewInfo with the specified options, and add it to the specified view family. */
-	virtual void CreateAndInitSingleView(class FSceneViewFamily* ViewFamily, const struct FSceneViewInitOptions* ViewInitOptions) = 0;
+	/** Call from the render thread to create and initalize a new FViewInfo with the specified options, and add it to the specified view family. */
+	virtual void CreateAndInitSingleView(FRHICommandListImmediate& RHICmdList, class FSceneViewFamily* ViewFamily, const struct FSceneViewInitOptions* ViewInitOptions) = 0;
 
 	/**
 	 * Allocates a new instance of the private FScene implementation for the given world.
