Fix non-WITH_EDITOR compilation

diff --git a/Engine/Source/Runtime/Landscape/Private/Landscape.cpp b/Engine/Source/Runtime/Landscape/Private/Landscape.cpp
--- a/Engine/Source/Runtime/Landscape/Private/Landscape.cpp
+++ b/Engine/Source/Runtime/Landscape/Private/Landscape.cpp
@@ -2651,7 +2651,11 @@ void ALandscapeProxy::UpdateFoliage(const TArray<FVector>& Cameras, ULandscapeCo
 				}
 
 				// Don't try to place foliage for this component if it doesn't have data and it can't currently render it.
-				if (!Component->GrassData->HasData() && !Component->CanRenderGrassMap())
+				if (!Component->GrassData->HasData() 
+#if WITH_EDITOR
+					&& !Component->CanRenderGrassMap()
+#endif
+					)
 				{
 					continue;
 				}
