Improved Map Check message for overlapping landscape components

diff --git a/Engine/Source/Runtime/Landscape/Private/Landscape.cpp b/Engine/Source/Runtime/Landscape/Private/Landscape.cpp
--- a/Engine/Source/Runtime/Landscape/Private/Landscape.cpp
+++ b/Engine/Source/Runtime/Landscape/Private/Landscape.cpp
@@ -1668,14 +1668,18 @@ void ULandscapeInfo::RegisterActorComponent(ULandscapeComponent* Component, bool
 	}
 	else if (bMapCheck)
 	{
-		ALandscapeProxy* Proxy = Component->GetLandscapeProxy();
+		ALandscapeProxy* OurProxy = Component->GetLandscapeProxy();
+		ALandscapeProxy* ExistingProxy = RegisteredComponent->GetLandscapeProxy();
 		FFormatNamedArguments Arguments;
-		Arguments.Add(TEXT("ProxyName"), FText::FromString(Proxy->GetName()));
+		Arguments.Add(TEXT("ProxyName1"), FText::FromString(OurProxy->GetName()));
+		Arguments.Add(TEXT("LevelName1"), FText::FromString(OurProxy->GetLevel()->GetOutermost()->GetName()));
+		Arguments.Add(TEXT("ProxyName2"), FText::FromString(ExistingProxy->GetName()));
+		Arguments.Add(TEXT("LevelName2"), FText::FromString(ExistingProxy->GetLevel()->GetOutermost()->GetName()));
 		Arguments.Add(TEXT("XLocation"), Component->GetSectionBase().X);
 		Arguments.Add(TEXT("YLocation"), Component->GetSectionBase().Y);
 		FMessageLog("MapCheck").Warning()
-			->AddToken(FUObjectToken::Create(Proxy))
-			->AddToken(FTextToken::Create(FText::Format(LOCTEXT("MapCheck_Message_LandscapeComponentPostLoad_Warning", "Landscape ({ProxyName}) has overlapping render components at location ({XLocation}, {YLocation})."), Arguments)))
+			->AddToken(FUObjectToken::Create(OurProxy))
+			->AddToken(FTextToken::Create(FText::Format(LOCTEXT("MapCheck_Message_LandscapeComponentPostLoad_Warning", "Landscape {ProxyName1} of {LevelName1} has overlapping render components with {ProxyName2} of {LevelName2} at location ({XLocation}, {YLocation})."), Arguments)))
 			->AddToken(FMapErrorToken::Create(FMapErrors::LandscapeComponentPostLoad_Warning));
 
 		// Show MapCheck window
