Fixed landscape copy/paste gizmo moving and rescaling when trying to copy from one landscape to another

diff --git a/Engine/Source/Runtime/Landscape/Private/LandscapeGizmoActor.cpp b/Engine/Source/Runtime/Landscape/Private/LandscapeGizmoActor.cpp
--- a/Engine/Source/Runtime/Landscape/Private/LandscapeGizmoActor.cpp
+++ b/Engine/Source/Runtime/Landscape/Private/LandscapeGizmoActor.cpp
@@ -640,7 +640,8 @@ void ALandscapeGizmoActiveActor::SetTargetLandscape(ULandscapeInfo* LandscapeInf
 		TargetLandscapeInfo = LandscapeInfo;
 	}
 
-	if (TargetLandscapeInfo && TargetLandscapeInfo != PrevInfo)
+	// if there's no copied data, try to move somewhere useful
+	if (TargetLandscapeInfo && TargetLandscapeInfo != PrevInfo && DataType == LGT_None)
 	{
 		MarginZ = TargetLandscapeInfo->DrawScale.Z * 3;
 		Width = Height = TargetLandscapeInfo->DrawScale.X * (TargetLandscapeInfo->ComponentSizeQuads+1);
