Fixed crash when pasting gizmo paint data to another landscape with a different material or that hasn't had its layer infos configured

diff --git a/Engine/Source/Editor/LandscapeEditor/Private/LandscapeEdModeComponentTools.cpp b/Engine/Source/Editor/LandscapeEditor/Private/LandscapeEdModeComponentTools.cpp
--- a/Engine/Source/Editor/LandscapeEditor/Private/LandscapeEdModeComponentTools.cpp
+++ b/Engine/Source/Editor/LandscapeEditor/Private/LandscapeEdModeComponentTools.cpp
@@ -1503,6 +1503,28 @@ public:
 				return;
 			}
 
+			// Automatically fill in any placeholder layers
+			// This gives a much better user experience when copying data to a newly created landscape
+			for (ULandscapeLayerInfoObject* LayerInfo : Gizmo->LayerInfos)
+			{
+				int32 LayerInfoIndex = LandscapeInfo->GetLayerInfoIndex(LayerInfo);
+				if (LayerInfoIndex == INDEX_NONE)
+				{
+					LayerInfoIndex = LandscapeInfo->GetLayerInfoIndex(LayerInfo->LayerName);
+					if (LayerInfoIndex != INDEX_NONE)
+					{
+						FLandscapeInfoLayerSettings& LayerSettings = LandscapeInfo->Layers[LayerInfoIndex];
+
+						if (LayerSettings.LayerInfoObj == nullptr)
+						{
+							LayerSettings.Owner = LandscapeInfo->GetLandscapeProxy(); // this isn't strictly accurate, but close enough
+							LayerSettings.LayerInfoObj = LayerInfo;
+							LayerSettings.bValid = true;
+						}
+					}
+				}
+			}
+
 			Gizmo->TargetLandscapeInfo = LandscapeInfo;
 			float ScaleXY = LandscapeInfo->DrawScale.X;
 
@@ -1687,9 +1709,12 @@ public:
 				}
 			}
 
-			for (int i = 0; i < Gizmo->LayerInfos.Num(); ++i)
+			for (ULandscapeLayerInfoObject* LayerInfo : Gizmo->LayerInfos)
 			{
-				WeightCache.AddDirtyLayer(Gizmo->LayerInfos[i]);
+				if (LandscapeInfo->GetLayerInfoIndex(LayerInfo) != INDEX_NONE)
+				{
+					WeightCache.AddDirtyLayer(LayerInfo);
+				}
 			}
 
 			if (bApplyToAll)
diff --git a/Engine/Source/Runtime/Landscape/Private/LandscapeEditInterface.cpp b/Engine/Source/Runtime/Landscape/Private/LandscapeEditInterface.cpp
--- a/Engine/Source/Runtime/Landscape/Private/LandscapeEditInterface.cpp
+++ b/Engine/Source/Runtime/Landscape/Private/LandscapeEditInterface.cpp
@@ -2483,6 +2483,15 @@ void FLandscapeEditDataInterface::SetAlphaData(const TSet<ULandscapeLayerInfoObj
 		return;
 	}
 
+	for (ULandscapeLayerInfoObject* LayerInfo : DirtyLayerInfos)
+	{
+		// The Data[] array passed in is indexed by LandscapeInfo->GetLayerInfoIndex(),
+		// so if we're trying to write a layer which isn't in the LandscapeInfo,
+		// its data is either missing or written where another layer's data should be.
+		// Either way it's *very bad*.
+		check(LandscapeInfo->GetLayerInfoIndex(LayerInfo) != INDEX_NONE);
+	}
+
 	if (Stride == 0)
 	{
 		Stride = (1+X2-X1) * LandscapeInfo->Layers.Num();
@@ -2529,9 +2538,8 @@ void FLandscapeEditDataInterface::SetAlphaData(const TSet<ULandscapeLayerInfoObj
 
 			NeedAllocationInfos.Reset();
 
-			for (auto It = DirtyLayerInfos.CreateConstIterator(); It; ++It)
+			for (ULandscapeLayerInfoObject* LayerInfo : DirtyLayerInfos)
 			{
-				ULandscapeLayerInfoObject* const LayerInfo = *It;
 				const bool bFound = Component->WeightmapLayerAllocations.ContainsByPredicate([LayerInfo](const FWeightmapLayerAllocationInfo& Allocation){ return Allocation.LayerInfo == LayerInfo; });
 				if (!bFound)
 				{
@@ -2555,14 +2563,14 @@ void FLandscapeEditDataInterface::SetAlphaData(const TSet<ULandscapeLayerInfoObj
 				if (PaintingRestriction != ELandscapeLayerPaintingRestriction::ExistingOnly)
 				{
 					Component->Modify();
-					for (int32 i = 0; i < NeedAllocationInfos.Num(); ++i)
+					for (ULandscapeLayerInfoObject* LayerInfoNeedingAllocation : NeedAllocationInfos)
 					{
 						if (PaintingRestriction == ELandscapeLayerPaintingRestriction::UseMaxLayers &&
 							LayerLimit > 0 && Component->WeightmapLayerAllocations.Num() >= LayerLimit)
 						{
 							break;
 						}
-						new (Component->WeightmapLayerAllocations) FWeightmapLayerAllocationInfo(NeedAllocationInfos[i]);
+						Component->WeightmapLayerAllocations.Emplace(LayerInfoNeedingAllocation);
 					}
 					Component->ReallocateWeightmaps(this);
 					Component->UpdateMaterialInstances();
@@ -2580,7 +2588,7 @@ void FLandscapeEditDataInterface::SetAlphaData(const TSet<ULandscapeLayerInfoObj
 			for (int32 WeightmapIdx = 0; WeightmapIdx < Component->WeightmapTextures.Num(); ++WeightmapIdx)
 			{
 				TexDataInfos[WeightmapIdx] = GetTextureDataInfo(Component->WeightmapTextures[WeightmapIdx]);
-			};
+			}
 
 			LayerDataInfos.Reset();        // Pointers to all layers' data
 			LayerDataInfos.AddUninitialized(Component->WeightmapLayerAllocations.Num());
