Automatically flush landscape grass when changing or recompiling the material
Automatically generate grass maps for all components, streaming in textures and weightmaps where needed
Status window when building grass maps
Grass data is now shared between Editor and PIE
Fixed issue where it didn't wait for streaming after updating grass map after painting landscape

diff --git a/Engine/Source/Editor/LandscapeEditor/Private/LandscapeEdModeTools.h b/Engine/Source/Editor/LandscapeEditor/Private/LandscapeEdModeTools.h
--- a/Engine/Source/Editor/LandscapeEditor/Private/LandscapeEdModeTools.h
+++ b/Engine/Source/Editor/LandscapeEditor/Private/LandscapeEdModeTools.h
@@ -619,7 +619,7 @@ struct FHeightmapAccessor
 			}
 
 			// Flush dynamic foliage (grass)
-			LandscapeInfo->GetLandscapeProxy()->FlushFoliageComponents(&Components);
+			LandscapeInfo->GetLandscapeProxy()->FlushGrassComponents(&Components);
 
 			// Notify foliage to move any attached instances
 			bool bUpdateFoliage = false;
@@ -792,7 +792,7 @@ struct FXYOffsetmapAccessor
 			}
 
 			// Flush dynamic foliage (grass)
-			LandscapeInfo->GetLandscapeProxy()->FlushFoliageComponents(&Components);
+			LandscapeInfo->GetLandscapeProxy()->FlushGrassComponents(&Components);
 
 			// Notify foliage to move any attached instances
 			bool bUpdateFoliage = false;
@@ -946,7 +946,7 @@ struct FAlphamapAccessor
 		if (LandscapeEdit.GetComponentsInRegion(X1, Y1, X2, Y2, &Components))
 		{
 			// Flush dynamic foliage (grass)
-			LandscapeInfo->GetLandscapeProxy()->FlushFoliageComponents(&Components);
+			LandscapeInfo->GetLandscapeProxy()->FlushGrassComponents(&Components);
 
 			LandscapeEdit.SetAlphaData(LayerInfo, X1, Y1, X2, Y2, Data, 0, PaintingRestriction, bBlendWeight, bUseTotalNormalize);
 			ModifiedComponents.Append(Components);
@@ -1053,7 +1053,7 @@ struct FFullWeightmapAccessor
 		if (LandscapeEdit.GetComponentsInRegion(X1, Y1, X2, Y2, &Components))
 		{
 			// Flush dynamic foliage (grass)
-			LandscapeInfo->GetLandscapeProxy()->FlushFoliageComponents(&Components);
+			LandscapeInfo->GetLandscapeProxy()->FlushGrassComponents(&Components);
 
 			LandscapeEdit.SetAlphaData(DirtyLayerInfos, X1, Y1, X2, Y2, Data, 0, PaintingRestriction);
 			ModifiedComponents.Append(Components);
diff --git a/Engine/Source/Editor/UnrealEd/Private/GrassRenderingNotification.cpp b/Engine/Source/Editor/UnrealEd/Private/GrassRenderingNotification.cpp
new file mode 100644
--- /dev/null
+++ b/Engine/Source/Editor/UnrealEd/Private/GrassRenderingNotification.cpp
@@ -0,0 +1,34 @@
+// Copyright 1998-2015 Epic Games, Inc. All Rights Reserved.
+
+#include "UnrealEd.h"
+#include "GlobalEditorNotification.h"
+#include "Landscape.h"
+#include "SNotificationList.h"
+
+/** Notification class for grassmmap rendering. */
+class FGrassRenderingNotificationImpl : public FGlobalEditorNotification
+{
+protected:
+	/** FGlobalEditorNotification interface */
+	virtual bool ShouldShowNotification(const bool bIsNotificationAlreadyActive) const override;
+	virtual void SetNotificationText(const TSharedPtr<SNotificationItem>& InNotificationItem) const override;
+};
+
+/** Global notification object. */
+FGrassRenderingNotificationImpl GGrassRenderingNotification;
+
+bool FGrassRenderingNotificationImpl::ShouldShowNotification(const bool bIsNotificationAlreadyActive) const
+{
+	return ALandscapeProxy::TotalComponentsNeedingGrassMapRender > 0;
+}
+
+void FGrassRenderingNotificationImpl::SetNotificationText(const TSharedPtr<SNotificationItem>& InNotificationItem) const
+{
+	if (ALandscapeProxy::TotalComponentsNeedingGrassMapRender > 0)
+	{
+		FFormatNamedArguments Args;
+		Args.Add(TEXT("OutstandingGrassMaps"), FText::AsNumber(ALandscapeProxy::TotalComponentsNeedingGrassMapRender));
+		const FText ProgressMessage = FText::Format(NSLOCTEXT("GrassMapRender", "GrassMapRenderFormat", "Building Grass Maps ({OutstandingGrassMaps})"), Args);
+		InNotificationItem->SetText(ProgressMessage);
+	}
+}
diff --git a/Engine/Source/Editor/WorldBrowser/Private/Tiles/WorldTileCollectionModel.cpp b/Engine/Source/Editor/WorldBrowser/Private/Tiles/WorldTileCollectionModel.cpp
--- a/Engine/Source/Editor/WorldBrowser/Private/Tiles/WorldTileCollectionModel.cpp
+++ b/Engine/Source/Editor/WorldBrowser/Private/Tiles/WorldTileCollectionModel.cpp
@@ -2008,9 +2008,9 @@ bool FWorldTileCollectionModel::GenerateLODLevels(FLevelModelList InLevelList, i
 
 			if (SimplificationDetails.bBakeGrassToLandscape)
 			{
-				Landscape->FlushFoliageComponents();
+				Landscape->FlushGrassComponents();
 				TArray<FVector> Cameras;
-				Landscape->UpdateFoliage(Cameras, nullptr, true);
+				Landscape->UpdateFoliage(Cameras, true);
 			}
 								
 			// This is texture resolution for a landscape mesh, probably needs to be calculated using landscape size
@@ -2025,7 +2025,7 @@ bool FWorldTileCollectionModel::GenerateLODLevels(FLevelModelList InLevelList, i
 		
 			if (SimplificationDetails.bBakeGrassToLandscape)
 			{
-				Landscape->FlushFoliageComponents(); // wipe this and let it fix itself later
+				Landscape->FlushGrassComponents(); // wipe this and let it fix itself later
 			}
 			FString LandscapeBaseAssetName = Landscape->GetName();
 			// Construct landscape material
diff --git a/Engine/Source/Runtime/Core/Public/UObject/ObjectVersion.h b/Engine/Source/Runtime/Core/Public/UObject/ObjectVersion.h
--- a/Engine/Source/Runtime/Core/Public/UObject/ObjectVersion.h
+++ b/Engine/Source/Runtime/Core/Public/UObject/ObjectVersion.h
@@ -468,6 +468,8 @@ enum EUnrealEngineObjectUE4Version
 	VER_UE4_SERIALIZE_LANDSCAPE_GRASS_DATA,
 	// New property to optionally prevent gpu emitters clearing existing particles on Init().
 	VER_UE4_OPTIONALLY_CLEAR_GPU_EMITTERS_ON_INIT,
+	// Also store the Material guid with the landscape grass data
+	VER_UE4_SERIALIZE_LANDSCAPE_GRASS_DATA_MATERIAL_GUID,
 
 	// -----<new versions can be added before this line>-------------------------------------------------
 	// - this needs to be the last line (see note below)
diff --git a/Engine/Source/Runtime/Landscape/Classes/LandscapeComponent.h b/Engine/Source/Runtime/Landscape/Classes/LandscapeComponent.h
--- a/Engine/Source/Runtime/Landscape/Classes/LandscapeComponent.h
+++ b/Engine/Source/Runtime/Landscape/Classes/LandscapeComponent.h
@@ -107,9 +107,16 @@ struct FWeightmapLayerAllocationInfo
 
 struct FLandscapeComponentGrassData
 {
+	FGuid MaterialStateId;
 	TArray<uint16> HeightData;
 	TMap<ULandscapeGrassType*, TArray<uint8>> WeightData;
 
+	FLandscapeComponentGrassData() {}
+
+	FLandscapeComponentGrassData(FGuid& InMaterialStateId)
+	: MaterialStateId(InMaterialStateId)
+	{}
+
 	bool HasData()
 	{
 		return HeightData.Num() > 0;
@@ -327,8 +334,10 @@ public:
 	void RemoveGrassMap();
 
 	/* Can a valid grassmap currently be generated? */
-	bool CanRenderGrassMap();
+	bool CanRenderGrassMap() const;
 
+	/* Is the grassmap data outdated, eg by a material */
+	bool IsGrassMapOutdated() const;
 #endif
 
 	/** @todo document */
diff --git a/Engine/Source/Runtime/Landscape/Classes/LandscapeProxy.h b/Engine/Source/Runtime/Landscape/Classes/LandscapeProxy.h
--- a/Engine/Source/Runtime/Landscape/Classes/LandscapeProxy.h
+++ b/Engine/Source/Runtime/Landscape/Classes/LandscapeProxy.h
@@ -494,19 +494,23 @@ public:
 	void SetLandscapeGuid(const FGuid& Guid) { LandscapeGuid = Guid; }
 	virtual ALandscape* GetLandscapeActor();
 
-	/** Flush the foliage cache */
-	LANDSCAPE_API void FlushFoliageComponents(const TSet<ULandscapeComponent*>* OnlyForComponents = nullptr);
+	/** Flush the grass cache */
+	LANDSCAPE_API void FlushGrassComponents(const TSet<ULandscapeComponent*>* OnlyForComponents = nullptr);
 
 	/** 
-		Update Foliage 
+		Update Grass 
 		* @param Cameras to use for culling, if empty, then NO culling
-		* @param OnlyComponent if non-null only do this component
 		* @param bForceSync if true, block and finish all work
 	*/
-	LANDSCAPE_API void UpdateFoliage(const TArray<FVector>& Cameras, ULandscapeComponent* OnlyComponent = nullptr, bool bForceSync = false);
+	LANDSCAPE_API void UpdateFoliage(const TArray<FVector>& Cameras, bool bForceSync = false);
 
 	/* Get the list of grass types on this landscape */
-	TArray<const ULandscapeGrassType*> GetGrassTypes() const;
+	TArray<ULandscapeGrassType*> GetGrassTypes() const;
+
+#if WITH_EDITOR
+	/** Render grass maps for the specified components */
+	void RenderGrassMaps(const TArray<ULandscapeComponent*>& LandscapeComponents, const TArray<ULandscapeGrassType*>& GrassTypes);
+#endif
 
 	// Begin FTickableGameObject interface.
 	virtual void Tick(float DeltaTime) override;
@@ -606,6 +610,10 @@ public:
 
 	/** Creates a Texture2D for use by this landscape proxy or one of it's components. If OptionalOverrideOuter is not specified, the level is used. */
 	LANDSCAPE_API UTexture2D* CreateLandscapeTexture(int32 InSizeX, int32 InSizeY, TextureGroup InLODGroup, ETextureSourceFormat InFormat, UObject* OptionalOverrideOuter = nullptr) const;
+
+	/* For the grassmap rendering notification */
+	int32 NumComponentsNeedingGrassMapRender;
+	LANDSCAPE_API static int32 TotalComponentsNeedingGrassMapRender;
 #endif
 };
 
diff --git a/Engine/Source/Runtime/Landscape/Private/Landscape.cpp b/Engine/Source/Runtime/Landscape/Private/Landscape.cpp
--- a/Engine/Source/Runtime/Landscape/Private/Landscape.cpp
+++ b/Engine/Source/Runtime/Landscape/Private/Landscape.cpp
@@ -154,7 +154,26 @@ void ULandscapeComponent::Serialize(FArchive& Ar)
 
 	if (Ar.UE4Ver() >= VER_UE4_SERIALIZE_LANDSCAPE_GRASS_DATA)
 	{
-		Ar << GrassData.Get();
+		// Share the shared ref so PIE can share this data
+		if (Ar.GetPortFlags() & PPF_DuplicateForPIE)
+		{
+			if (Ar.IsSaving())
+			{
+				PTRINT GrassDataPointer = (PTRINT)&GrassData;
+				Ar << GrassDataPointer;
+			}
+			else
+			{
+				PTRINT GrassDataPointer;
+				Ar << GrassDataPointer;
+				// Duplicate shared reference
+				GrassData = *(TSharedRef<FLandscapeComponentGrassData, ESPMode::ThreadSafe>*)GrassDataPointer;
+			}
+		}
+		else
+		{
+			Ar << GrassData.Get();
+		}
 	}
 
 #if WITH_EDITOR
@@ -1318,6 +1337,11 @@ void ALandscapeProxy::Destroyed()
 	{
 		SplineComponent->ModifySplines();
 	}
+
+#if WITH_EDITOR
+	TotalComponentsNeedingGrassMapRender -= NumComponentsNeedingGrassMapRender;
+	NumComponentsNeedingGrassMapRender = 0;
+#endif
 }
 #endif
 
@@ -2113,19 +2137,10 @@ DECLARE_CYCLE_STAT(TEXT("Grass End Comp"), STAT_FoliageGrassEndComp, STATGROUP_F
 DECLARE_CYCLE_STAT(TEXT("Grass Destroy Comps"), STAT_FoliageGrassDestoryComp, STATGROUP_Foliage);
 DECLARE_CYCLE_STAT(TEXT("Grass Update"), STAT_GrassUpdate, STATGROUP_Foliage);
 
-void ALandscapeProxy::FlushFoliageComponents(const TSet<ULandscapeComponent*>* OnlyForComponents)
+void ALandscapeProxy::FlushGrassComponents(const TSet<ULandscapeComponent*>* OnlyForComponents)
 {
 	if (OnlyForComponents)
 	{
-#if WITH_EDITOR
-		for (auto LandComp : *OnlyForComponents)
-		{
-			if (LandComp)
-			{
-				LandComp->RenderGrassMap();
-			}
-		}
-#endif
 		for (FCachedLandscapeFoliage::TGrassSet::TIterator Iter(FoliageCache.CachedGrassComps); Iter; ++Iter)
 		{
 			ULandscapeComponent* Component = (*Iter).Key.BasedOn.Get();
@@ -2142,6 +2157,15 @@ void ALandscapeProxy::FlushFoliageComponents(const TSet<ULandscapeComponent*>* O
 				Iter.RemoveCurrent();
 			}
 		}
+#if WITH_EDITOR
+		if (GIsEditor)
+		{
+		    for (ULandscapeComponent* Component : *OnlyForComponents)
+		    {
+			    Component->RemoveGrassMap();
+		    }
+		}
+#endif
 	}
 	else
 	{
@@ -2174,12 +2198,15 @@ void ALandscapeProxy::FlushFoliageComponents(const TSet<ULandscapeComponent*>* O
 		}
 
 #if WITH_EDITOR
-		// Clear GrassMaps
-		TInlineComponentArray<ULandscapeComponent*> LandComps;
-		GetComponents(LandComps);
-		for (ULandscapeComponent* Component : LandComps)
+		if (GIsEditor)
 		{
-			Component->RemoveGrassMap();
+		    // Clear GrassMaps
+		    TInlineComponentArray<ULandscapeComponent*> LandComps;
+		    GetComponents(LandComps);
+		    for (ULandscapeComponent* Component : LandComps)
+		    {
+			    Component->RemoveGrassMap();
+		    }
 		}
 #endif
 	}
@@ -2610,9 +2637,9 @@ void ALandscapeProxy::Tick(float DeltaSeconds)
 	}
 }
 
-TArray<const ULandscapeGrassType*> ALandscapeProxy::GetGrassTypes() const
+TArray<ULandscapeGrassType*> ALandscapeProxy::GetGrassTypes() const
 {
-	TArray<const ULandscapeGrassType*> GrassTypes;
+	TArray<ULandscapeGrassType*> GrassTypes;
 	if (LandscapeMaterial)
 	{
 		TArray<const UMaterialExpressionLandscapeGrassOutput*> GrassExpressions;
@@ -2628,13 +2655,17 @@ TArray<const ULandscapeGrassType*> ALandscapeProxy::GetGrassTypes() const
 	return GrassTypes;
 }
 
-void ALandscapeProxy::UpdateFoliage(const TArray<FVector>& Cameras, ULandscapeComponent* OnlyComponent, bool bForceSync)
+#if WITH_EDITOR
+int ALandscapeProxy::TotalComponentsNeedingGrassMapRender = 0;
+#endif
+
+void ALandscapeProxy::UpdateFoliage(const TArray<FVector>& Cameras, bool bForceSync)
 {
 	SCOPE_CYCLE_COUNTER(STAT_GrassUpdate);
 
 	if (CVarGrassEnable.GetValueOnAnyThread() > 0)
 	{
-		TArray<const ULandscapeGrassType*> GrassTypes = GetGrassTypes();
+		TArray<ULandscapeGrassType*> GrassTypes = GetGrassTypes();
 
 		float GuardBand = CVarGuardBandMultiplier.GetValueOnAnyThread();
 		float DiscardGuardBand = CVarGuardBandDiscardMultiplier.GetValueOnAnyThread();
@@ -2643,23 +2674,65 @@ void ALandscapeProxy::UpdateFoliage(const TArray<FVector>& Cameras, ULandscapeCo
 		int32 MaxInstancesPerComponent = FMath::Max<int32>(1024, CVarMaxInstancesPerComponent.GetValueOnAnyThread());
 		int32 MaxTasks = CVarMaxAsyncTasks.GetValueOnAnyThread();
 
-		if (GetWorld())
+		UWorld* World = GetWorld();
+		if (World)
 		{
+#if WITH_EDITOR
+			bool bStreamingRenderedGrassHeightmap = false;
+			TSet<ULandscapeComponent*> ComponentsNeedingGrassMapRender;
+			TSet<UTexture2D*> CurrentForcedStreamedTextures;
+			TSet<UTexture2D*> DesiredForceStreamedTextures;
+
+			if (!World->IsGameWorld())
+			{
+				// see if we need to flush grass for any components
+				TSet<ULandscapeComponent*> FlushComponents;
+				for (auto Component : LandscapeComponents)
+				{
+					// check textures currently needing force streaming
+					if (Component->HeightmapTexture->bForceMiplevelsToBeResident)
+					{
+						CurrentForcedStreamedTextures.Add(Component->HeightmapTexture);
+					}
+					for(auto WeightmapTexture : Component->WeightmapTextures)
+					{
+						if (WeightmapTexture->bForceMiplevelsToBeResident)
+						{
+							CurrentForcedStreamedTextures.Add(WeightmapTexture);
+						}
+					}
+
+					if (Component->IsGrassMapOutdated())
+					{
+						FlushComponents.Add(Component);
+						if (GrassTypes.Num() > 0)
+						{
+							ComponentsNeedingGrassMapRender.Add(Component);
+						}
+					}
+					else
+					if (!Component->GrassData->HasData())
+					{
+						if (GrassTypes.Num() > 0)
+						{
+							ComponentsNeedingGrassMapRender.Add(Component);
+						}
+					}
+				}
+				if (FlushComponents.Num())
+				{
+					FlushGrassComponents(&FlushComponents);
+				}
+			}
+#endif
+
 			int32 NumCompsCreated = 0;
 			for (int32 ComponentIndex = 0; ComponentIndex < LandscapeComponents.Num(); ComponentIndex++)
 			{
 				ULandscapeComponent* Component = LandscapeComponents[ComponentIndex];
-				if (OnlyComponent && OnlyComponent != Component)
-				{
-					continue;
-				}
 
-				// Don't try to place foliage for this component if it doesn't have data and it can't currently render it.
-				if (!Component->GrassData->HasData() 
-#if WITH_EDITOR
-					&& !Component->CanRenderGrassMap()
-#endif
-					)
+				// skip if we have no data and no way to generate it
+				if (World->IsGameWorld() && !Component->GrassData->HasData()) 
 				{
 					continue;
 				}
@@ -2765,6 +2838,30 @@ void ALandscapeProxy::UpdateFoliage(const TArray<FVector>& Cameras, ULandscapeCo
 										{
 											continue; // one per frame, but we still want to touch the existing ones
 										}
+
+#if WITH_EDITOR
+											// render grass data if we don't have any
+										if (!Component->GrassData->HasData())
+										{
+											// if we can't generate now, force our textures to be streamed in
+											if (!Component->CanRenderGrassMap())
+											{
+												DesiredForceStreamedTextures.Add(Component->HeightmapTexture);
+												for (auto WeightmapTexture : Component->WeightmapTextures)
+												{
+													DesiredForceStreamedTextures.Add(WeightmapTexture);
+													bStreamingRenderedGrassHeightmap = true;
+												}
+												continue;
+											}
+
+											QUICK_SCOPE_CYCLE_COUNTER(STAT_GrassRenderToTexture);
+											Component->RenderGrassMap();
+											ComponentsNeedingGrassMapRender.Remove(Component);
+											Component->MarkPackageDirty();
+										}
+#endif
+
 										NumCompsCreated++;
 
 										SCOPE_CYCLE_COUNTER(STAT_FoliageGrassStartComp);
@@ -2819,15 +2916,6 @@ void ALandscapeProxy::UpdateFoliage(const TArray<FVector>& Cameras, ULandscapeCo
 											FoliageComponents.Add(HierarchicalInstancedStaticMeshComponent);
 										}
 
-#if WITH_EDITOR
-										// render grass data if we don't have any
-										if (!Component->GrassData->HasData())
-										{
-											QUICK_SCOPE_CYCLE_COUNTER(STAT_GrassRenderToTexture);
-											Component->RenderGrassMap();
-										}
-#endif
-
 										FAsyncGrassBuilder* Builder;
 
 										{
@@ -2858,6 +2946,60 @@ void ALandscapeProxy::UpdateFoliage(const TArray<FVector>& Cameras, ULandscapeCo
 					}
 				}
 			}
+
+#if WITH_EDITOR
+			{
+				int32 NumComponentsRendered = 0;
+				if (GrassTypes.Num() > 0)
+				{
+					// try to render some grassmaps
+					TArray<ULandscapeComponent*> ComponentsToRender;
+					for (auto Component : ComponentsNeedingGrassMapRender)
+					{
+						if (Component->CanRenderGrassMap())
+						{
+							ComponentsToRender.Add(Component);
+							NumComponentsRendered++;
+							// We really want to throttle the number based on component size...
+							if (NumComponentsRendered >= 4)
+							{
+								break;
+							}
+						}
+						else if (!bStreamingRenderedGrassHeightmap)
+						{
+							// Force stream in other heightmaps but only if we're not waiting for the textures 
+							// near the camera to stream in
+							DesiredForceStreamedTextures.Add(Component->HeightmapTexture);
+							for (auto WeightmapTexture : Component->WeightmapTextures)
+							{
+								DesiredForceStreamedTextures.Add(WeightmapTexture);
+							}
+						}
+					}
+					if (ComponentsToRender.Num())
+					{
+						RenderGrassMaps(ComponentsToRender, GrassTypes);
+						MarkPackageDirty();
+					}
+				}
+
+				TotalComponentsNeedingGrassMapRender -= NumComponentsNeedingGrassMapRender;
+				NumComponentsNeedingGrassMapRender = ComponentsNeedingGrassMapRender.Num() - NumComponentsRendered;
+				TotalComponentsNeedingGrassMapRender += NumComponentsNeedingGrassMapRender;
+				
+				// Update resident flags
+				for (auto Texture : DesiredForceStreamedTextures.Difference(CurrentForcedStreamedTextures))
+				{
+					Texture->bForceMiplevelsToBeResident = true;
+				}
+				for (auto Texture : CurrentForcedStreamedTextures.Difference(DesiredForceStreamedTextures))
+				{
+					Texture->bForceMiplevelsToBeResident = false;
+				}
+
+			}
+#endif
 		}
 	}
 
@@ -2965,6 +3107,11 @@ ALandscapeProxy::~ALandscapeProxy()
 		delete Task;
 	}
 	AsyncFoliageTasks.Empty();
+
+#if WITH_EDITOR
+	TotalComponentsNeedingGrassMapRender -= NumComponentsNeedingGrassMapRender;
+	NumComponentsNeedingGrassMapRender = 0;
+#endif
 }
 
 void FAsyncGrassTask::DoWork()
@@ -2972,40 +3119,23 @@ void FAsyncGrassTask::DoWork()
 	Builder->Build();
 }
 
-static void FlushFoliage(const TArray<FString>& Args)
+static void FlushGrass(const TArray<FString>& Args)
 {
 	for (TObjectIterator<ALandscapeProxy> It; It; ++It)
 	{
 		ALandscapeProxy* Landscape = *It;
 		if (Landscape && !Landscape->IsTemplate() && !Landscape->IsPendingKill())
 		{
-			Landscape->FlushFoliageComponents();
+			Landscape->FlushGrassComponents();
 		}
 	}
 }
 
-static FAutoConsoleCommand FlushFoliageCmd(
+static FAutoConsoleCommand FlushGrassCmd(
 	TEXT("grass.FlushCache"),
 	TEXT("Flush the grass cache, debugging."),
-	FConsoleCommandWithArgsDelegate::CreateStatic(&FlushFoliage)
-	);
-
-
-#if WITH_EDITOR
-static void RenderGrass(const TArray<FString>& Args)
-{
-	for (TObjectIterator<ULandscapeComponent> It; It; ++It)
-	{
-		It->RenderGrassMap();
-	}
-}
-
-static FAutoConsoleCommand RenderFoliageCmd(
-	TEXT("grass.Render"),
-	TEXT("Render grass to texture. Debugging."),
-	FConsoleCommandWithArgsDelegate::CreateStatic(&RenderGrass)
+	FConsoleCommandWithArgsDelegate::CreateStatic(&FlushGrass)
 	);
-#endif
 
 
 //
diff --git a/Engine/Source/Runtime/Landscape/Private/LandscapeEdit.cpp b/Engine/Source/Runtime/Landscape/Private/LandscapeEdit.cpp
--- a/Engine/Source/Runtime/Landscape/Private/LandscapeEdit.cpp
+++ b/Engine/Source/Runtime/Landscape/Private/LandscapeEdit.cpp
@@ -348,7 +348,7 @@ void ULandscapeComponent::PostEditUndo()
 
 	TSet<ULandscapeComponent*> Components;
 	Components.Add(this);
-	GetLandscapeProxy()->FlushFoliageComponents(&Components);
+	GetLandscapeProxy()->FlushGrassComponents(&Components);
 }
 
 void ULandscapeComponent::FixupWeightmaps()
@@ -3477,7 +3477,7 @@ void ULandscapeLayerInfoObject::PostEditChangeProperty(FPropertyChangedEvent& Pr
 					{
 						if (Info->Layers[i].LayerInfoObj == this)
 						{
-							Proxy->FlushFoliageComponents();
+							Proxy->FlushGrassComponents();
 							break;
 						}
 					}
@@ -3530,7 +3530,7 @@ void ALandscapeProxy::RemoveXYOffsets()
 void ALandscapeProxy::RecreateCollisionComponents()
 {
 	// We can assume these are all junk; they recreate as needed
-	FlushFoliageComponents();
+	FlushGrassComponents();
 
 	// Clear old CollisionComponent containers
 	CollisionComponents.Empty();
diff --git a/Engine/Source/Runtime/Landscape/Private/LandscapeGrass.cpp b/Engine/Source/Runtime/Landscape/Private/LandscapeGrass.cpp
--- a/Engine/Source/Runtime/Landscape/Private/LandscapeGrass.cpp
+++ b/Engine/Source/Runtime/Landscape/Private/LandscapeGrass.cpp
@@ -225,7 +225,7 @@ public:
 	TArray<FComponentInfo> ComponentInfos;
 	FIntPoint TargetSize;
 	int32 NumPasses;
-	int32 PassOffsetX;
+	float PassOffsetX;
 	FMatrix ViewMatrix;
 	FMatrix ProjectionMatrix;
 
@@ -298,7 +298,7 @@ public:
 
 		TargetSize = FIntPoint(ComponentSizeVerts * NumPasses * InLandscapeComponents.Num(), ComponentSizeVerts);
 		FIntPoint TargetSizeMinusOne(TargetSize - FIntPoint(1,1));
-		PassOffsetX = 2.0f * ComponentSizeVerts / TargetSize.X;
+		PassOffsetX = 2.0f * (float)ComponentSizeVerts / (float)TargetSize.X;
 
 		for (int32 Idx = 0; Idx<InLandscapeComponents.Num();Idx++)
 		{
@@ -363,7 +363,7 @@ public:
 
 		for (auto& ComponentInfo : ComponentInfos)
 		{
-			FLandscapeComponentGrassData* NewGrassData = new FLandscapeComponentGrassData();
+			FLandscapeComponentGrassData* NewGrassData = new FLandscapeComponentGrassData(ComponentInfo.Component->MaterialInstance->GetMaterial()->StateId);
 
 			NewGrassData->HeightData.Empty(FMath::Square(ComponentSizeVerts));
 
@@ -461,10 +461,23 @@ public:
 
 #if WITH_EDITOR
 
-bool ULandscapeComponent::CanRenderGrassMap()
+bool ULandscapeComponent::IsGrassMapOutdated() const
+{
+	if (GrassData->HasData())
+	{
+		if (GrassData->MaterialStateId != MaterialInstance->GetMaterial()->StateId)
+		{
+			return true;
+		}
+	}
+	return false;
+}
+
+bool ULandscapeComponent::CanRenderGrassMap() const
 {
 	// Check we can render
-	if (!GIsEditor || GUsingNullRHI || !GetWorld() || GetWorld()->FeatureLevel < ERHIFeatureLevel::SM4 || !SceneProxy)
+	UWorld* World = GetWorld();
+	if (!GIsEditor || GUsingNullRHI || !World || World->IsGameWorld() || World->FeatureLevel < ERHIFeatureLevel::SM4 || !SceneProxy)
 	{
 		return false;
 	}
@@ -518,6 +531,12 @@ void ULandscapeComponent::RemoveGrassMap()
 	GrassData = MakeShareable(new FLandscapeComponentGrassData());
 }
 
+void ALandscapeProxy::RenderGrassMaps(const TArray<ULandscapeComponent*>& LandscapeComponents, const TArray<ULandscapeGrassType*>& GrassTypes)
+{
+	FLandscapeGrassWeightExporter Exporter(this, LandscapeComponents, GrassTypes);
+	Exporter.ApplyResults();
+}
+
 #endif
 
 
@@ -669,7 +688,7 @@ void ULandscapeGrassType::PostEditChangeProperty(FPropertyChangedEvent& Property
 					{
 						if (Output.GrassType == this)
 						{
-							Proxy->FlushFoliageComponents();
+							Proxy->FlushGrassComponents();
 							break;
 						}
 					}
@@ -685,6 +704,11 @@ void ULandscapeGrassType::PostEditChangeProperty(FPropertyChangedEvent& Property
 //
 FArchive& operator<<(FArchive& Ar, FLandscapeComponentGrassData& Data)
 {
+	if (Ar.UE4Ver() >= VER_UE4_SERIALIZE_LANDSCAPE_GRASS_DATA_MATERIAL_GUID)
+	{
+		Ar << Data.MaterialStateId;
+	}
+
 	Data.HeightData.BulkSerialize(Ar);
 	// Each weight data array, being 1 byte will be serialized in bulk.
 	return Ar << Data.WeightData;
diff --git a/Engine/Source/Runtime/Landscape/Private/LandscapeRender.cpp b/Engine/Source/Runtime/Landscape/Private/LandscapeRender.cpp
--- a/Engine/Source/Runtime/Landscape/Private/LandscapeRender.cpp
+++ b/Engine/Source/Runtime/Landscape/Private/LandscapeRender.cpp
@@ -1041,6 +1041,11 @@ void FLandscapeComponentSceneProxy::DrawStaticElements(FStaticPrimitiveDrawInter
 	// Could be different from bRequiresAdjacencyInformation during shader compilation
 	bool bCurrentRequiresAdjacencyInformation = RequiresAdjacencyInformation(RenderProxy, GetScene().GetFeatureLevel());
 
+	if (bCurrentRequiresAdjacencyInformation)
+	{
+		check(SharedBuffers->AdjacencyIndexBuffers);
+	}
+
 	MeshBatch.VertexFactory = VertexFactory;
 	MeshBatch.MaterialRenderProxy = RenderProxy;
 	MeshBatch.LCI = ComponentLightInfo.Get();
@@ -1083,17 +1088,14 @@ void FLandscapeComponentSceneProxy::DrawStaticElements(FStaticPrimitiveDrawInter
 					FMeshBatchElement* BatchElement = new(MeshBatch.Elements) FMeshBatchElement;
 					FLandscapeBatchElementParams* BatchElementParams = new(StaticBatchParamArray)FLandscapeBatchElementParams;
 					BatchElement->UserData = BatchElementParams;
-
 					BatchElement->PrimitiveUniformBufferResource = &GetUniformBuffer();
-
 					BatchElementParams->LandscapeUniformShaderParametersResource = &LandscapeUniformShaderParameters;
 					BatchElementParams->LocalToWorldNoScalingPtr = &LocalToWorldNoScaling;
 					BatchElementParams->SceneProxy = this;
 					BatchElementParams->SubX = SubX;
 					BatchElementParams->SubY = SubY;
 					BatchElementParams->CurrentLOD = LOD;
-
-					BatchElement->IndexBuffer = SharedBuffers->IndexBuffers[LOD];
+					BatchElement->IndexBuffer = bCurrentRequiresAdjacencyInformation ? SharedBuffers->AdjacencyIndexBuffers->IndexBuffers[LOD] : SharedBuffers->IndexBuffers[LOD];
 					BatchElement->NumPrimitives = FMath::Square((LodSubsectionSizeVerts - 1)) * 2;
 					BatchElement->FirstIndex = (SubX + SubY * NumSubsections) * BatchElement->NumPrimitives * 3;
 					BatchElement->MinVertexIndex = SharedBuffers->IndexRanges[LOD].MinIndex[SubX][SubY];
@@ -1113,17 +1115,7 @@ void FLandscapeComponentSceneProxy::DrawStaticElements(FStaticPrimitiveDrawInter
 		BatchElementParams->SubX = -1;
 		BatchElementParams->SubY = -1;
 		BatchElementParams->CurrentLOD = LOD;
-
-		if (bCurrentRequiresAdjacencyInformation)
-		{
-			check(SharedBuffers->AdjacencyIndexBuffers);
-			BatchElement->IndexBuffer = SharedBuffers->AdjacencyIndexBuffers->IndexBuffers[LOD];
-		}
-		else
-		{
-			BatchElement->IndexBuffer = SharedBuffers->IndexBuffers[LOD];
-		}
-
+		BatchElement->IndexBuffer = bCurrentRequiresAdjacencyInformation ? SharedBuffers->AdjacencyIndexBuffers->IndexBuffers[LOD] : SharedBuffers->IndexBuffers[LOD];
 		BatchElement->NumPrimitives = FMath::Square((LodSubsectionSizeVerts - 1)) * FMath::Square(NumSubsections) * 2;
 		BatchElement->FirstIndex = 0;
 		BatchElement->MinVertexIndex = SharedBuffers->IndexRanges[LOD].MinIndexFull;
@@ -2529,9 +2521,10 @@ void ULandscapeComponent::GetStreamingTextureInfo(TArray<FStreamingTexturePrimit
 	// Heightmap
 	if (HeightmapTexture)
 	{
-		float HeightmapTexelFactor = TexelFactor * (HeightmapTexture->GetSizeY() / (ComponentSizeQuads + 1));
 		FStreamingTexturePrimitiveInfo& StreamingHeightmap = *new(OutStreamingTextures)FStreamingTexturePrimitiveInfo;
 		StreamingHeightmap.Bounds = BoundingSphere;
+
+		float HeightmapTexelFactor = TexelFactor * (HeightmapTexture->GetSizeY() / (ComponentSizeQuads + 1));
 		StreamingHeightmap.TexelFactor = ForcedLOD >= 0 ? -13 + ForcedLOD : HeightmapTexelFactor; // Minus Value indicate ForcedLOD, 13 for 8k texture
 		StreamingHeightmap.Texture = HeightmapTexture;
 	}
