Give streaming priority to grass being rendered near the player

diff --git a/Engine/Source/Runtime/Landscape/Classes/LandscapeProxy.h b/Engine/Source/Runtime/Landscape/Classes/LandscapeProxy.h
--- a/Engine/Source/Runtime/Landscape/Classes/LandscapeProxy.h
+++ b/Engine/Source/Runtime/Landscape/Classes/LandscapeProxy.h
@@ -614,6 +614,11 @@ public:
 	/* For the grassmap rendering notification */
 	int32 NumComponentsNeedingGrassMapRender;
 	LANDSCAPE_API static int32 TotalComponentsNeedingGrassMapRender;
+
+	/* To throttle texture streaming when we're trying to render a grassmap */
+	int32 NumTexturesToStreamForVisibleGrassMapRender;
+	LANDSCAPE_API static int32 TotalTexturesToStreamForVisibleGrassMapRender;
+
 #endif
 };
 
diff --git a/Engine/Source/Runtime/Landscape/Private/Landscape.cpp b/Engine/Source/Runtime/Landscape/Private/Landscape.cpp
--- a/Engine/Source/Runtime/Landscape/Private/Landscape.cpp
+++ b/Engine/Source/Runtime/Landscape/Private/Landscape.cpp
@@ -391,11 +391,6 @@ void ULandscapeInfo::UpdateDebugColorMaterial()
 	//GWarn->EndSlowTask();
 }
 
-static TAutoConsoleVariable<int32> CVarPreloadGrassmapsInEditor(
-	TEXT("grass.PreloadGrassmapsInEditor"),
-	0,
-	TEXT("1: Load grass maps in editor at load time uses lots of memory; 0: Load grass maps on demand; hitchy"));
-
 void ULandscapeComponent::PostLoad()
 {
 	Super::PostLoad();
@@ -2656,9 +2651,16 @@ TArray<ULandscapeGrassType*> ALandscapeProxy::GetGrassTypes() const
 }
 
 #if WITH_EDITOR
-int ALandscapeProxy::TotalComponentsNeedingGrassMapRender = 0;
+int32 ALandscapeProxy::TotalComponentsNeedingGrassMapRender = 0;
+int32 ALandscapeProxy::TotalTexturesToStreamForVisibleGrassMapRender = 0;
 #endif
 
+static TAutoConsoleVariable<int32> CVarPrerenderGrassmaps(
+	TEXT("grass.PrerenderGrassmaps"),
+	1,
+	TEXT("1: Pre-render grass maps for all components in the editor; 0: Generate grass maps on demand while moving through the editor"));
+
+
 void ALandscapeProxy::UpdateFoliage(const TArray<FVector>& Cameras, bool bForceSync)
 {
 	SCOPE_CYCLE_COUNTER(STAT_GrassUpdate);
@@ -2678,7 +2680,7 @@ void ALandscapeProxy::UpdateFoliage(const TArray<FVector>& Cameras, bool bForceS
 		if (World)
 		{
 #if WITH_EDITOR
-			bool bStreamingRenderedGrassHeightmap = false;
+			int32 RequiredTexturesNotStreamedIn = 0;
 			TSet<ULandscapeComponent*> ComponentsNeedingGrassMapRender;
 			TSet<UTexture2D*> CurrentForcedStreamedTextures;
 			TSet<UTexture2D*> DesiredForceStreamedTextures;
@@ -2851,7 +2853,7 @@ void ALandscapeProxy::UpdateFoliage(const TArray<FVector>& Cameras, bool bForceS
 												{
 													DesiredForceStreamedTextures.Add(WeightmapTexture);
 												}
-												bStreamingRenderedGrassHeightmap = true;
+												RequiredTexturesNotStreamedIn++;
 												continue;
 											}
 
@@ -2948,9 +2950,14 @@ void ALandscapeProxy::UpdateFoliage(const TArray<FVector>& Cameras, bool bForceS
 			}
 
 #if WITH_EDITOR
+
+			TotalTexturesToStreamForVisibleGrassMapRender -= NumTexturesToStreamForVisibleGrassMapRender;
+			NumTexturesToStreamForVisibleGrassMapRender = RequiredTexturesNotStreamedIn;
+			TotalTexturesToStreamForVisibleGrassMapRender += NumTexturesToStreamForVisibleGrassMapRender;
+
 			{
 				int32 NumComponentsRendered = 0;
-				if (GrassTypes.Num() > 0)
+				if (GrassTypes.Num() > 0 && CVarPrerenderGrassmaps.GetValueOnAnyThread() > 0)
 				{
 					// try to render some grassmaps
 					TArray<ULandscapeComponent*> ComponentsToRender;
@@ -2966,7 +2973,7 @@ void ALandscapeProxy::UpdateFoliage(const TArray<FVector>& Cameras, bool bForceS
 								break;
 							}
 						}
-						else if (!bStreamingRenderedGrassHeightmap)
+						else if (TotalTexturesToStreamForVisibleGrassMapRender == 0)
 						{
 							// Force stream in other heightmaps but only if we're not waiting for the textures 
 							// near the camera to stream in
diff --git a/Engine/Source/Runtime/Landscape/Private/LandscapeGrass.cpp b/Engine/Source/Runtime/Landscape/Private/LandscapeGrass.cpp
--- a/Engine/Source/Runtime/Landscape/Private/LandscapeGrass.cpp
+++ b/Engine/Source/Runtime/Landscape/Private/LandscapeGrass.cpp
@@ -489,7 +489,8 @@ bool ULandscapeComponent::CanRenderGrassMap() const
 	}
 	
 	// Check for valid heightmap that is fully streamed in
-	if (!HeightmapTexture || HeightmapTexture->ResidentMips != HeightmapTexture->GetNumMips())
+	if (!HeightmapTexture || HeightmapTexture->ResidentMips != HeightmapTexture->GetNumMips()
+		|| !HeightmapTexture->Resource || ((FTexture2DResource*)HeightmapTexture->Resource)->GetCurrentFirstMip() > 0)
 	{
 		return false;
 	}
