Ignore all neighbor LOD calculations for LOD override used to render grass maps

diff --git a/Engine/Source/Runtime/Landscape/Private/LandscapeRender.cpp b/Engine/Source/Runtime/Landscape/Private/LandscapeRender.cpp
--- a/Engine/Source/Runtime/Landscape/Private/LandscapeRender.cpp
+++ b/Engine/Source/Runtime/Landscape/Private/LandscapeRender.cpp
@@ -1210,6 +1210,13 @@ uint64 FLandscapeComponentSceneProxy::GetStaticBatchElementVisibility(const clas
 
 float FLandscapeComponentSceneProxy::CalcDesiredLOD(const class FSceneView& View, const FVector2D& CameraLocalPos, int32 SubX, int32 SubY) const
 {
+#if WITH_EDITOR
+	if (View.Family->LandscapeLODOverride >= 0)
+	{
+		return View.Family->LandscapeLODOverride;
+	}
+#endif
+
 	// FLandscapeComponentSceneProxy::NumSubsections, SubsectionSizeQuads, MaxLOD, LODFalloff and LODDistance are the same for all components and so are safe to use in the neighbour LOD calculations
 	// HeightmapTexture, LODBias, ForcedLOD are component-specific with neighbor lookup
 	const bool bIsInThisComponent = (SubX >= 0 && SubX < NumSubsections && SubY >= 0 && SubY < NumSubsections);
@@ -1252,13 +1259,6 @@ float FLandscapeComponentSceneProxy::CalcDesiredLOD(const class FSceneView& View
 		fLOD = SubsectionForcedLOD;
 	}
 	else
-#if WITH_EDITOR
-	if (View.Family->LandscapeLODOverride >= 0)
-	{
-		fLOD = View.Family->LandscapeLODOverride;
-	}
-	else
-#endif
 	{
 		if (View.IsPerspectiveProjection())
 		{
