Add ability to deform a landscape using a given spline from blueprint

diff --git a/Engine/Source/Runtime/Landscape/Classes/LandscapeProxy.h b/Engine/Source/Runtime/Landscape/Classes/LandscapeProxy.h
--- a/Engine/Source/Runtime/Landscape/Classes/LandscapeProxy.h
+++ b/Engine/Source/Runtime/Landscape/Classes/LandscapeProxy.h
@@ -28,6 +28,7 @@ class UTexture2D;
 class ALandscape;
 class ALandscapeProxy;
 class ULandscapeComponent;
+class USplineComponent;
 
 /** Structure storing channel usage for weightmap textures */
 USTRUCT()
@@ -446,10 +447,6 @@ public:
 	UPROPERTY()
 	int32 NumSubsections;    // Number of subsections in X and Y axis
 
-	/** Change the Level of Detail distance factor */
-	UFUNCTION(BlueprintCallable, Category="Rendering")
-	virtual void ChangeLODDistanceFactor(float InLODDistanceFactor);
-
 	/** Hints navigation system whether this landscape will ever be navigated on. true by default, but make sure to set it to false for faraway, background landscapes */
 	UPROPERTY(EditAnywhere, Category=Landscape)
 	uint32 bUsedForNavigation:1;
@@ -474,6 +471,31 @@ public:
 	/** Map of weightmap usage */
 	TMap<UTexture2D*, struct FLandscapeWeightmapUsage> WeightmapUsageMap;
 
+	// Blueprint functions
+
+	/** Change the Level of Detail distance factor */
+	UFUNCTION(BlueprintCallable, Category = "Rendering")
+	virtual void ChangeLODDistanceFactor(float InLODDistanceFactor);
+
+	// Editor-time blueprint functions
+
+	/** Deform landscape using a given spline
+	 * @param StartWidth - Width of the spline at the start node, in Spline Component local space
+	 * @param EndWidth   - Width of the spline at the end node, in Spline Component local space
+	 * @param StartSideFalloff - Width of the falloff at either side of the spline at the start node, in Spline Component local space
+	 * @param EndSideFalloff - Width of the falloff at either side of the spline at the end node, in Spline Component local space
+	 * @param StartRoll - Roll applied to the spline at the start node, in degrees. 0 is flat
+	 * @param EndRoll - Roll applied to the spline at the end node, in degrees. 0 is flat
+	 * @param NumSubdivisions - Number of triangles to place along the spline when applying it to the landscape. Higher numbers give better results, but setting it too high will be slow and may cause artifacts
+	 * @param bRaiseHeights - Allow the landscape to be raised up to the level of the spline. If both bRaiseHeights and bLowerHeights are false, no height modification of the landscape will be performed
+	 * @param bLowerHeights - Allow the landscape to be lowered down to the level of the spline. If both bRaiseHeights and bLowerHeights are false, no height modification of the landscape will be performed
+	 * @param PaintLayer - LayerInfo to paint, or none to skip painting. The landscape must be configured with the same layer info in one of its layers or this will do nothing!
+	 */
+	UFUNCTION(BlueprintCallable, Category = "Landscape Editor")
+	void EditorApplySpline(USplineComponent* InSplineComponent, float StartWidth = 200, float EndWidth = 200, float StartSideFalloff = 200, float EndSideFalloff = 200, float StartRoll = 0, float EndRoll = 0, int32 NumSubdivisions = 20, bool bRaiseHeights = true, bool bLowerHeights = true, ULandscapeLayerInfoObject* PaintLayer = nullptr);
+
+	// End blueprint functions
+
 	// Begin AActor Interface
 	virtual void UnregisterAllComponents() override;
 	virtual void RegisterAllComponents() override;
diff --git a/Engine/Source/Runtime/Landscape/Classes/LandscapeSplineSegment.h b/Engine/Source/Runtime/Landscape/Classes/LandscapeSplineSegment.h
--- a/Engine/Source/Runtime/Landscape/Classes/LandscapeSplineSegment.h
+++ b/Engine/Source/Runtime/Landscape/Classes/LandscapeSplineSegment.h
@@ -247,10 +247,5 @@ public:
 
 	virtual bool OwnsComponent(const class USplineMeshComponent* SplineMeshComponent) const;
 
-protected:
-#if WITH_EDITOR
-	bool FixSelfIntersection(FVector FLandscapeSplineInterpPoint::* Side);
-#endif
-
 	friend class FEdModeLandscape;
 };
diff --git a/Engine/Source/Runtime/Landscape/Private/LandscapeBlueprintSupport.cpp b/Engine/Source/Runtime/Landscape/Private/LandscapeBlueprintSupport.cpp
new file mode 100644
--- /dev/null
+++ b/Engine/Source/Runtime/Landscape/Private/LandscapeBlueprintSupport.cpp
@@ -0,0 +1,25 @@
+// Copyright 1998-2015 Epic Games, Inc. All Rights Reserved.
+
+/*=============================================================================
+	LandscapeBlueprintSupport.cpp: Landscape blueprint functions
+  =============================================================================*/
+
+#include "Landscape.h"
+#include "LandscapeSplineRaster.h"
+#include "LandscapeProxy.h"
+#include "LandscapeSplineRaster.h"
+#include "Components/SplineComponent.h"
+#include "LandscapeSplineSegment.h"
+
+void ALandscapeProxy::EditorApplySpline(USplineComponent* SplineComponent, float StartWidth, float EndWidth, float StartSideFalloff, float EndSideFalloff, float StartRoll, float EndRoll, int32 NumSubdivisions, bool bRaiseHeights, bool bLowerHeights, ULandscapeLayerInfoObject* PaintLayer)
+{
+	check(GIsEditor);
+
+#if WITH_EDITOR
+	TArray<FLandscapeSplineInterpPoint> Points;
+	LandscapeSplineRaster::Pointify(SplineComponent->SplineInfo, Points, NumSubdivisions, 0.0f, 0.0f, StartWidth, EndWidth, StartSideFalloff, EndSideFalloff, StartRoll, EndRoll);
+
+	FTransform SplineToWorld = SplineComponent->GetComponentTransform();
+	LandscapeSplineRaster::RasterizeSegmentPoints(GetLandscapeInfo(), MoveTemp(Points), SplineToWorld, bRaiseHeights, bLowerHeights, PaintLayer);
+#endif
+}
diff --git a/Engine/Source/Runtime/Landscape/Private/LandscapeEdit.cpp b/Engine/Source/Runtime/Landscape/Private/LandscapeEdit.cpp
--- a/Engine/Source/Runtime/Landscape/Private/LandscapeEdit.cpp
+++ b/Engine/Source/Runtime/Landscape/Private/LandscapeEdit.cpp
@@ -22,8 +22,6 @@ LandscapeEdit.cpp: Landscape editing
 #include "LandscapeHeightfieldCollisionComponent.h"
 #include "LandscapeMeshCollisionComponent.h"
 #include "LandscapeSplinesComponent.h"
-#include "LandscapeSplineControlPoint.h"
-#include "LandscapeSplineSegment.h"
 #include "LandscapeGizmoActiveActor.h"
 #include "InstancedFoliageActor.h"
 #include "LevelUtils.h"
@@ -2675,503 +2673,6 @@ void ULandscapeInfo::ExportLayer(ULandscapeLayerInfoObject* LayerInfo, const FSt
 	GWarn->EndSlowTask();
 }
 
-//////////////////////////////////////////////////////////////////////////
-// Apply splines
-//////////////////////////////////////////////////////////////////////////
-#include "Raster.h"
-
-class FLandscapeSplineHeightsRasterPolicy
-{
-public:
-	// X = Side Alpha, Y = End Alpha, Z = Height
-	typedef FVector InterpolantType;
-
-	/** Initialization constructor. */
-	FLandscapeSplineHeightsRasterPolicy(TArray<uint16>& InData, int32 InMinX, int32 InMinY, int32 InMaxX, int32 InMaxY, bool InbRaiseTerrain, bool InbLowerTerrain) :
-		Data(InData),
-		MinX(InMinX),
-		MinY(InMinY),
-		MaxX(InMaxX),
-		MaxY(InMaxY),
-		bRaiseTerrain(InbRaiseTerrain),
-		bLowerTerrain(InbLowerTerrain)
-	{
-	}
-
-protected:
-
-	// FTriangleRasterizer policy interface.
-
-	int32 GetMinX() const { return MinX; }
-	int32 GetMaxX() const { return MaxX; }
-	int32 GetMinY() const { return MinY; }
-	int32 GetMaxY() const { return MaxY; }
-
-	void ProcessPixel(int32 X, int32 Y, const InterpolantType& Interpolant, bool BackFacing)
-	{
-		const float CosInterpX = (Interpolant.X >= 1 ? 1 : 0.5f - 0.5f * FMath::Cos(Interpolant.X * PI));
-		const float CosInterpY = (Interpolant.Y >= 1 ? 1 : 0.5f - 0.5f * FMath::Cos(Interpolant.Y * PI));
-		const float Alpha = CosInterpX * CosInterpY;
-		uint16& Dest = Data[(Y - MinY)*(1 + MaxX - MinX) + X - MinX];
-		float Value = FMath::Lerp((float)Dest, Interpolant.Z, Alpha);
-		uint16 DValue = (uint32)FMath::Clamp<float>(Value, 0, LandscapeDataAccess::MaxValue);
-		if ((bRaiseTerrain && DValue > Dest) ||
-			(bLowerTerrain && DValue < Dest))
-		{
-			Dest = DValue;
-		}
-	}
-
-private:
-	TArray<uint16>& Data;
-	int32 MinX, MinY, MaxX, MaxY;
-	uint32 bRaiseTerrain : 1, bLowerTerrain : 1;
-};
-
-class FLandscapeSplineBlendmaskRasterPolicy
-{
-public:
-	// X = Side Alpha, Y = End Alpha, Z = Blend Value
-	typedef FVector InterpolantType;
-
-	/** Initialization constructor. */
-	FLandscapeSplineBlendmaskRasterPolicy(TArray<uint8>& InData, int32 InMinX, int32 InMinY, int32 InMaxX, int32 InMaxY) :
-		Data(InData),
-		MinX(InMinX),
-		MinY(InMinY),
-		MaxX(InMaxX),
-		MaxY(InMaxY)
-	{
-	}
-
-protected:
-
-	// FTriangleRasterizer policy interface.
-
-	int32 GetMinX() const { return MinX; }
-	int32 GetMaxX() const { return MaxX; }
-	int32 GetMinY() const { return MinY; }
-	int32 GetMaxY() const { return MaxY; }
-
-	void ProcessPixel(int32 X, int32 Y, const InterpolantType& Interpolant, bool BackFacing)
-	{
-		const float CosInterpX = (Interpolant.X >= 1 ? 1 : 0.5f - 0.5f * FMath::Cos(Interpolant.X * PI));
-		const float CosInterpY = (Interpolant.Y >= 1 ? 1 : 0.5f - 0.5f * FMath::Cos(Interpolant.Y * PI));
-		const float Alpha = CosInterpX * CosInterpY;
-		uint8& Dest = Data[(Y - MinY)*(1 + MaxX - MinX) + X - MinX];
-		float Value = FMath::Lerp((float)Dest, Interpolant.Z, Alpha);
-		Dest = (uint32)FMath::Clamp<float>(Value, 0, LandscapeDataAccess::MaxValue);
-	}
-
-private:
-	TArray<uint8>& Data;
-	int32 MinX, MinY, MaxX, MaxY;
-};
-
-bool ULandscapeInfo::ApplySplines(bool bOnlySelected)
-{
-	bool bResult = false;
-
-	ALandscape* Landscape = LandscapeActor.Get();
-
-	bResult |= ApplySplinesInternal(bOnlySelected, Landscape);
-
-	for (auto It = Proxies.CreateIterator(); It; ++It)
-	{
-		ALandscapeProxy* LandscapeProxy = (*It);
-		bResult |= ApplySplinesInternal(bOnlySelected, LandscapeProxy);
-	}
-
-	return bResult;
-}
-
-bool ULandscapeInfo::ApplySplinesInternal(bool bOnlySelected, ALandscapeProxy* Landscape)
-{
-	if (!Landscape || !Landscape->SplineComponent || Landscape->SplineComponent->ControlPoints.Num() == 0 || Landscape->SplineComponent->Segments.Num() == 0)
-	{
-		return false;
-	}
-
-	FScopedTransaction Transaction(LOCTEXT("LandscapeSpline_ApplySplines", "Apply Splines to Landscape"));
-
-	const FTransform SplineToLandscape = Landscape->SplineComponent->ComponentToWorld.GetRelativeTransform(Landscape->LandscapeActorToWorld());
-
-	FLandscapeEditDataInterface LandscapeEdit(this);
-	TSet<ULandscapeComponent*> ModifiedComponents;
-
-	// I'd dearly love to use FIntRect in this code, but Landscape works with "Inclusive Max" and FIntRect is "Exclusive Max"
-	int32 LandscapeMinX, LandscapeMinY, LandscapeMaxX, LandscapeMaxY;
-	if (!GetLandscapeExtent(LandscapeMinX, LandscapeMinY, LandscapeMaxX, LandscapeMaxY))
-	{
-		return false;
-	}
-
-	for (auto It = Landscape->SplineComponent->ControlPoints.CreateConstIterator(); It; ++It)
-	{
-		const ULandscapeSplineControlPoint* ControlPoint = *It;
-
-		if (bOnlySelected && !ControlPoint->IsSplineSelected())
-		{
-			continue;
-		}
-
-		if (ControlPoint->GetPoints().Num() < 2)
-		{
-			continue;
-		}
-
-		FBox ControlPointBounds = ControlPoint->GetBounds();
-		ControlPointBounds = ControlPointBounds.TransformBy(SplineToLandscape.ToMatrixWithScale());
-
-		int32 MinX = FMath::CeilToInt(ControlPointBounds.Min.X);
-		int32 MinY = FMath::CeilToInt(ControlPointBounds.Min.Y);
-		int32 MaxX = FMath::FloorToInt(ControlPointBounds.Max.X);
-		int32 MaxY = FMath::FloorToInt(ControlPointBounds.Max.Y);
-
-		MinX = FMath::Max(MinX, LandscapeMinX);
-		MinY = FMath::Max(MinY, LandscapeMinY);
-		MaxX = FMath::Min(MaxX, LandscapeMaxX);
-		MaxY = FMath::Min(MaxY, LandscapeMaxY);
-
-		if (MinX > MaxX || MinY > MaxY)
-		{
-			// The control point's bounds don't intersect any data, so we skip it entirely
-			continue;
-		}
-
-		TArray<FLandscapeSplineInterpPoint> Points = ControlPoint->GetPoints();
-		for (int32 j = 0; j < Points.Num(); j++)
-		{
-			Points[j].Center = SplineToLandscape.TransformPosition(Points[j].Center);
-			Points[j].Left = SplineToLandscape.TransformPosition(Points[j].Left);
-			Points[j].Right = SplineToLandscape.TransformPosition(Points[j].Right);
-			Points[j].FalloffLeft = SplineToLandscape.TransformPosition(Points[j].FalloffLeft);
-			Points[j].FalloffRight = SplineToLandscape.TransformPosition(Points[j].FalloffRight);
-
-			// local-heights to texture value heights
-			Points[j].Left.Z = Points[j].Left.Z * LANDSCAPE_INV_ZSCALE + LandscapeDataAccess::MidValue;
-			Points[j].Right.Z = Points[j].Right.Z * LANDSCAPE_INV_ZSCALE + LandscapeDataAccess::MidValue;
-			Points[j].FalloffLeft.Z = Points[j].FalloffLeft.Z * LANDSCAPE_INV_ZSCALE + LandscapeDataAccess::MidValue;
-			Points[j].FalloffRight.Z = Points[j].FalloffRight.Z * LANDSCAPE_INV_ZSCALE + LandscapeDataAccess::MidValue;
-		}
-
-		// Heights raster
-		if (ControlPoint->bRaiseTerrain || ControlPoint->bLowerTerrain)
-		{
-			TArray<uint16> Data;
-			Data.AddZeroed((1 + MaxY - MinY) * (1 + MaxX - MinX));
-
-			int32 ValidMinX = MinX;
-			int32 ValidMinY = MinY;
-			int32 ValidMaxX = MaxX;
-			int32 ValidMaxY = MaxY;
-			LandscapeEdit.GetHeightData(ValidMinX, ValidMinY, ValidMaxX, ValidMaxY, Data.GetData(), 0);
-
-			if (ValidMinX > ValidMaxX || ValidMinY > ValidMaxY)
-			{
-				// The control point's bounds don't intersect any data, so we skip it entirely
-				// it wouldn't intersect any weightmap data either so we don't even bother trying
-				continue;
-			}
-
-			FLandscapeEditDataInterface::ShrinkData(Data, MinX, MinY, MaxX, MaxY, ValidMinX, ValidMinY, ValidMaxX, ValidMaxY);
-
-			MinX = ValidMinX;
-			MinY = ValidMinY;
-			MaxX = ValidMaxX;
-			MaxY = ValidMaxY;
-
-			FTriangleRasterizer<FLandscapeSplineHeightsRasterPolicy> Rasterizer(
-				FLandscapeSplineHeightsRasterPolicy(Data, MinX, MinY, MaxX, MaxY, ControlPoint->bRaiseTerrain, ControlPoint->bLowerTerrain));
-
-			const FVector Center3D = SplineToLandscape.TransformPosition(ControlPoint->Location);
-			const FVector2D CenterPos = FVector2D(Center3D);
-			const FVector Center = FVector(1.0f, Points[0].StartEndFalloff, Center3D.Z * LANDSCAPE_INV_ZSCALE + LandscapeDataAccess::MidValue);
-
-			for (int32 i = Points.Num() - 1, j = 0; j < Points.Num(); i = j++)
-			{
-				// Solid center
-				const FVector2D Right0Pos = FVector2D(Points[i].Right);
-				const FVector2D Left1Pos = FVector2D(Points[j].Left);
-				const FVector2D Right1Pos = FVector2D(Points[j].Right);
-				const FVector Right0 = FVector(1.0f, Points[i].StartEndFalloff, Points[i].Right.Z);
-				const FVector Left1 = FVector(1.0f, Points[j].StartEndFalloff, Points[j].Left.Z);
-				const FVector Right1 = FVector(1.0f, Points[j].StartEndFalloff, Points[j].Right.Z);
-
-				Rasterizer.DrawTriangle(Center, Right0, Left1, CenterPos, Right0Pos, Left1Pos, false);
-				Rasterizer.DrawTriangle(Center, Left1, Right1, CenterPos, Left1Pos, Right1Pos, false);
-
-				// Falloff
-				FVector2D FalloffRight0Pos = FVector2D(Points[i].FalloffRight);
-				FVector2D FalloffLeft1Pos = FVector2D(Points[j].FalloffLeft);
-				FVector FalloffRight0 = FVector(0.0f, Points[i].StartEndFalloff, Points[i].FalloffRight.Z);
-				FVector FalloffLeft1 = FVector(0.0f, Points[j].StartEndFalloff, Points[j].FalloffLeft.Z);
-				Rasterizer.DrawTriangle(Right0, FalloffRight0, Left1, Right0Pos, FalloffRight0Pos, Left1Pos, false);
-				Rasterizer.DrawTriangle(FalloffRight0, Left1, FalloffLeft1, FalloffRight0Pos, Left1Pos, FalloffLeft1Pos, false);
-			}
-
-			LandscapeEdit.SetHeightData(MinX, MinY, MaxX, MaxY, Data.GetData(), 0, true);
-			LandscapeEdit.GetComponentsInRegion(MinX, MinY, MaxX, MaxY, &ModifiedComponents);
-		}
-
-		// Blend layer raster
-		ULandscapeLayerInfoObject* LayerInfo = GetLayerInfoByName(ControlPoint->LayerName);
-		if (ControlPoint->LayerName != NAME_None && LayerInfo != NULL)
-		{
-			TArray<uint8> Data;
-			Data.AddZeroed((1 + MaxY - MinY) * (1 + MaxX - MinX));
-
-			int32 ValidMinX = MinX;
-			int32 ValidMinY = MinY;
-			int32 ValidMaxX = MaxX;
-			int32 ValidMaxY = MaxY;
-			LandscapeEdit.GetWeightData(LayerInfo, ValidMinX, ValidMinY, ValidMaxX, ValidMaxY, Data.GetData(), 0);
-
-			if (ValidMinX > ValidMaxX || ValidMinY > ValidMaxY)
-			{
-				// The control point's bounds don't intersect any data, so skip it
-				continue;
-			}
-
-			FLandscapeEditDataInterface::ShrinkData(Data, MinX, MinY, MaxX, MaxY, ValidMinX, ValidMinY, ValidMaxX, ValidMaxY);
-
-			MinX = ValidMinX;
-			MinY = ValidMinY;
-			MaxX = ValidMaxX;
-			MaxY = ValidMaxY;
-
-			FTriangleRasterizer<FLandscapeSplineBlendmaskRasterPolicy> Rasterizer(
-				FLandscapeSplineBlendmaskRasterPolicy(Data, MinX, MinY, MaxX, MaxY));
-
-			const float BlendValue = 255;
-
-			const FVector Center3D = SplineToLandscape.TransformPosition(ControlPoint->Location);
-			const FVector2D CenterPos = FVector2D(Center3D);
-			const FVector Center = FVector(1.0f, Points[0].StartEndFalloff, BlendValue);
-
-			for (int32 i = Points.Num() - 1, j = 0; j < Points.Num(); i = j++)
-			{
-				// Solid center
-				const FVector2D Right0Pos = FVector2D(Points[i].Right);
-				const FVector2D Left1Pos = FVector2D(Points[j].Left);
-				const FVector2D Right1Pos = FVector2D(Points[j].Right);
-				const FVector Right0 = FVector(1.0f, Points[i].StartEndFalloff, BlendValue);
-				const FVector Left1 = FVector(1.0f, Points[j].StartEndFalloff, BlendValue);
-				const FVector Right1 = FVector(1.0f, Points[j].StartEndFalloff, BlendValue);
-
-				Rasterizer.DrawTriangle(Center, Right0, Left1, CenterPos, Right0Pos, Left1Pos, false);
-				Rasterizer.DrawTriangle(Center, Left1, Right1, CenterPos, Left1Pos, Right1Pos, false);
-
-				// Falloff
-				FVector2D FalloffRight0Pos = FVector2D(Points[i].FalloffRight);
-				FVector2D FalloffLeft1Pos = FVector2D(Points[j].FalloffLeft);
-				FVector FalloffRight0 = FVector(0.0f, Points[i].StartEndFalloff, BlendValue);
-				FVector FalloffLeft1 = FVector(0.0f, Points[j].StartEndFalloff, BlendValue);
-				Rasterizer.DrawTriangle(Right0, FalloffRight0, Left1, Right0Pos, FalloffRight0Pos, Left1Pos, false);
-				Rasterizer.DrawTriangle(FalloffRight0, Left1, FalloffLeft1, FalloffRight0Pos, Left1Pos, FalloffLeft1Pos, false);
-			}
-
-			LandscapeEdit.SetAlphaData(LayerInfo, MinX, MinY, MaxX, MaxY, Data.GetData(), 0, ELandscapeLayerPaintingRestriction::None, true, false);
-			LandscapeEdit.GetComponentsInRegion(MinX, MinY, MaxX, MaxY, &ModifiedComponents);
-		}
-	}
-
-	for (auto It = Landscape->SplineComponent->Segments.CreateConstIterator(); It; ++It)
-	{
-		const ULandscapeSplineSegment* Segment = *It;
-
-		if (bOnlySelected && !Segment->IsSplineSelected())
-		{
-			continue;
-		}
-
-		FBox SegmentBounds = Segment->GetBounds();
-		SegmentBounds = SegmentBounds.TransformBy(SplineToLandscape.ToMatrixWithScale());
-
-		int32 MinX = FMath::CeilToInt(SegmentBounds.Min.X);
-		int32 MinY = FMath::CeilToInt(SegmentBounds.Min.Y);
-		int32 MaxX = FMath::FloorToInt(SegmentBounds.Max.X);
-		int32 MaxY = FMath::FloorToInt(SegmentBounds.Max.Y);
-
-		MinX = FMath::Max(MinX, LandscapeMinX);
-		MinY = FMath::Max(MinY, LandscapeMinY);
-		MaxX = FMath::Min(MaxX, LandscapeMaxX);
-		MaxY = FMath::Min(MaxY, LandscapeMaxY);
-
-		if (MinX > MaxX || MinY > MaxY)
-		{
-			// The segment's bounds don't intersect any data, so we skip it entirely
-			continue;
-		}
-
-		TArray<FLandscapeSplineInterpPoint> Points = Segment->GetPoints();
-		for (int32 j = 0; j < Points.Num(); j++)
-		{
-			Points[j].Center = SplineToLandscape.TransformPosition(Points[j].Center);
-			Points[j].Left = SplineToLandscape.TransformPosition(Points[j].Left);
-			Points[j].Right = SplineToLandscape.TransformPosition(Points[j].Right);
-			Points[j].FalloffLeft = SplineToLandscape.TransformPosition(Points[j].FalloffLeft);
-			Points[j].FalloffRight = SplineToLandscape.TransformPosition(Points[j].FalloffRight);
-
-			// local-heights to texture value heights
-			Points[j].Left.Z = Points[j].Left.Z * LANDSCAPE_INV_ZSCALE + LandscapeDataAccess::MidValue;
-			Points[j].Right.Z = Points[j].Right.Z * LANDSCAPE_INV_ZSCALE + LandscapeDataAccess::MidValue;
-			Points[j].FalloffLeft.Z = Points[j].FalloffLeft.Z * LANDSCAPE_INV_ZSCALE + LandscapeDataAccess::MidValue;
-			Points[j].FalloffRight.Z = Points[j].FalloffRight.Z * LANDSCAPE_INV_ZSCALE + LandscapeDataAccess::MidValue;
-		}
-
-		// Heights raster
-		if (Segment->bRaiseTerrain || Segment->bLowerTerrain)
-		{
-			TArray<uint16> Data;
-			Data.AddZeroed((1 + MaxY - MinY) * (1 + MaxX - MinX));
-
-			int32 ValidMinX = MinX;
-			int32 ValidMinY = MinY;
-			int32 ValidMaxX = MaxX;
-			int32 ValidMaxY = MaxY;
-			LandscapeEdit.GetHeightData(ValidMinX, ValidMinY, ValidMaxX, ValidMaxY, Data.GetData(), 0);
-
-			if (ValidMinX > ValidMaxX || ValidMinY > ValidMaxY)
-			{
-				// The segment's bounds don't intersect any data, so we skip it entirely
-				// it wouldn't intersect any weightmap data either so we don't even bother trying
-				continue;
-			}
-
-			FLandscapeEditDataInterface::ShrinkData(Data, MinX, MinY, MaxX, MaxY, ValidMinX, ValidMinY, ValidMaxX, ValidMaxY);
-
-			MinX = ValidMinX;
-			MinY = ValidMinY;
-			MaxX = ValidMaxX;
-			MaxY = ValidMaxY;
-
-			FTriangleRasterizer<FLandscapeSplineHeightsRasterPolicy> Rasterizer(
-				FLandscapeSplineHeightsRasterPolicy(Data, MinX, MinY, MaxX, MaxY, Segment->bRaiseTerrain, Segment->bLowerTerrain));
-
-			for (int32 j = 1; j < Points.Num(); j++)
-			{
-				// Middle
-				FVector2D Left0Pos = FVector2D(Points[j - 1].Left);
-				FVector2D Right0Pos = FVector2D(Points[j - 1].Right);
-				FVector2D Left1Pos = FVector2D(Points[j].Left);
-				FVector2D Right1Pos = FVector2D(Points[j].Right);
-				FVector Left0 = FVector(1.0f, Points[j - 1].StartEndFalloff, Points[j - 1].Left.Z);
-				FVector Right0 = FVector(1.0f, Points[j - 1].StartEndFalloff, Points[j - 1].Right.Z);
-				FVector Left1 = FVector(1.0f, Points[j].StartEndFalloff, Points[j].Left.Z);
-				FVector Right1 = FVector(1.0f, Points[j].StartEndFalloff, Points[j].Right.Z);
-				Rasterizer.DrawTriangle(Left0, Right0, Left1, Left0Pos, Right0Pos, Left1Pos, false);
-				Rasterizer.DrawTriangle(Right0, Left1, Right1, Right0Pos, Left1Pos, Right1Pos, false);
-
-				// Left Falloff
-				FVector2D FalloffLeft0Pos = FVector2D(Points[j - 1].FalloffLeft);
-				FVector2D FalloffLeft1Pos = FVector2D(Points[j].FalloffLeft);
-				FVector FalloffLeft0 = FVector(0.0f, Points[j - 1].StartEndFalloff, Points[j - 1].FalloffLeft.Z);
-				FVector FalloffLeft1 = FVector(0.0f, Points[j].StartEndFalloff, Points[j].FalloffLeft.Z);
-				Rasterizer.DrawTriangle(FalloffLeft0, Left0, FalloffLeft1, FalloffLeft0Pos, Left0Pos, FalloffLeft1Pos, false);
-				Rasterizer.DrawTriangle(Left0, FalloffLeft1, Left1, Left0Pos, FalloffLeft1Pos, Left1Pos, false);
-
-				// Right Falloff
-				FVector2D FalloffRight0Pos = FVector2D(Points[j - 1].FalloffRight);
-				FVector2D FalloffRight1Pos = FVector2D(Points[j].FalloffRight);
-				FVector FalloffRight0 = FVector(0.0f, Points[j - 1].StartEndFalloff, Points[j - 1].FalloffRight.Z);
-				FVector FalloffRight1 = FVector(0.0f, Points[j].StartEndFalloff, Points[j].FalloffRight.Z);
-				Rasterizer.DrawTriangle(Right0, FalloffRight0, Right1, Right0Pos, FalloffRight0Pos, Right1Pos, false);
-				Rasterizer.DrawTriangle(FalloffRight0, Right1, FalloffRight1, FalloffRight0Pos, Right1Pos, FalloffRight1Pos, false);
-			}
-
-			LandscapeEdit.SetHeightData(MinX, MinY, MaxX, MaxY, Data.GetData(), 0, true);
-			LandscapeEdit.GetComponentsInRegion(MinX, MinY, MaxX, MaxY, &ModifiedComponents);
-		}
-
-		// Blend layer raster
-		ULandscapeLayerInfoObject* LayerInfo = GetLayerInfoByName(Segment->LayerName);
-		if (Segment->LayerName != NAME_None && LayerInfo != NULL)
-		{
-			TArray<uint8> Data;
-			Data.AddZeroed((1 + MaxY - MinY) * (1 + MaxX - MinX));
-
-			int32 ValidMinX = MinX;
-			int32 ValidMinY = MinY;
-			int32 ValidMaxX = MaxX;
-			int32 ValidMaxY = MaxY;
-			LandscapeEdit.GetWeightData(LayerInfo, ValidMinX, ValidMinY, ValidMaxX, ValidMaxY, Data.GetData(), 0);
-
-			if (ValidMinX > ValidMaxX || ValidMinY > ValidMaxY)
-			{
-				// The segment's bounds don't intersect any data, so skip it
-				continue;
-			}
-
-			FLandscapeEditDataInterface::ShrinkData(Data, MinX, MinY, MaxX, MaxY, ValidMinX, ValidMinY, ValidMaxX, ValidMaxY);
-
-			MinX = ValidMinX;
-			MinY = ValidMinY;
-			MaxX = ValidMaxX;
-			MaxY = ValidMaxY;
-
-			FTriangleRasterizer<FLandscapeSplineBlendmaskRasterPolicy> Rasterizer(
-				FLandscapeSplineBlendmaskRasterPolicy(Data, MinX, MinY, MaxX, MaxY));
-
-			const float BlendValue = 255;
-
-			for (int32 j = 1; j < Points.Num(); j++)
-			{
-				// Middle
-				FVector2D Left0Pos = FVector2D(Points[j - 1].Left);
-				FVector2D Right0Pos = FVector2D(Points[j - 1].Right);
-				FVector2D Left1Pos = FVector2D(Points[j].Left);
-				FVector2D Right1Pos = FVector2D(Points[j].Right);
-				FVector Left0 = FVector(1.0f, Points[j - 1].StartEndFalloff, BlendValue);
-				FVector Right0 = FVector(1.0f, Points[j - 1].StartEndFalloff, BlendValue);
-				FVector Left1 = FVector(1.0f, Points[j].StartEndFalloff, BlendValue);
-				FVector Right1 = FVector(1.0f, Points[j].StartEndFalloff, BlendValue);
-				Rasterizer.DrawTriangle(Left0, Right0, Left1, Left0Pos, Right0Pos, Left1Pos, false);
-				Rasterizer.DrawTriangle(Right0, Left1, Right1, Right0Pos, Left1Pos, Right1Pos, false);
-
-				// Left Falloff
-				FVector2D FalloffLeft0Pos = FVector2D(Points[j - 1].FalloffLeft);
-				FVector2D FalloffLeft1Pos = FVector2D(Points[j].FalloffLeft);
-				FVector FalloffLeft0 = FVector(0.0f, Points[j - 1].StartEndFalloff, BlendValue);
-				FVector FalloffLeft1 = FVector(0.0f, Points[j].StartEndFalloff, BlendValue);
-				Rasterizer.DrawTriangle(FalloffLeft0, Left0, FalloffLeft1, FalloffLeft0Pos, Left0Pos, FalloffLeft1Pos, false);
-				Rasterizer.DrawTriangle(Left0, FalloffLeft1, Left1, Left0Pos, FalloffLeft1Pos, Left1Pos, false);
-
-				// Right Falloff
-				FVector2D FalloffRight0Pos = FVector2D(Points[j - 1].FalloffRight);
-				FVector2D FalloffRight1Pos = FVector2D(Points[j].FalloffRight);
-				FVector FalloffRight0 = FVector(0.0f, Points[j - 1].StartEndFalloff, BlendValue);
-				FVector FalloffRight1 = FVector(0.0f, Points[j].StartEndFalloff, BlendValue);
-				Rasterizer.DrawTriangle(Right0, FalloffRight0, Right1, Right0Pos, FalloffRight0Pos, Right1Pos, false);
-				Rasterizer.DrawTriangle(FalloffRight0, Right1, FalloffRight1, FalloffRight0Pos, Right1Pos, FalloffRight1Pos, false);
-			}
-
-			LandscapeEdit.SetAlphaData(LayerInfo, MinX, MinY, MaxX, MaxY, Data.GetData(), 0, ELandscapeLayerPaintingRestriction::None, true, false);
-			LandscapeEdit.GetComponentsInRegion(MinX, MinY, MaxX, MaxY, &ModifiedComponents);
-		}
-	}
-
-	LandscapeEdit.Flush();
-
-	for (ULandscapeComponent* Component : ModifiedComponents)
-	{
-		// Recreate collision for modified components and update the navmesh
-		ULandscapeHeightfieldCollisionComponent* CollisionComponent = Component->CollisionComponent.Get();
-		if (CollisionComponent)
-		{
-			CollisionComponent->RecreateCollision(false);
-			UNavigationSystem* NavSys = UNavigationSystem::GetCurrent(Component);
-			if (NavSys)
-			{
-				NavSys->UpdateNavOctree(CollisionComponent);
-			}
-		}
-	}
-
-	return true;
-}
-
-
 void ULandscapeInfo::DeleteLayer(ULandscapeLayerInfoObject* LayerInfo)
 {
 	GWarn->BeginSlowTask(LOCTEXT("BeginDeletingLayerTask", "Deleting Layer"), true);
diff --git a/Engine/Source/Runtime/Landscape/Private/LandscapeSplineRaster.cpp b/Engine/Source/Runtime/Landscape/Private/LandscapeSplineRaster.cpp
new file mode 100644
--- /dev/null
+++ b/Engine/Source/Runtime/Landscape/Private/LandscapeSplineRaster.cpp
@@ -0,0 +1,887 @@
+// Copyright 1998-2015 Epic Games, Inc. All Rights Reserved.
+
+/*=============================================================================
+	LandscapeSplineRaster.cpp:
+	Functions to rasterize a spline into landscape heights/weights
+  =============================================================================*/
+
+#include "Landscape.h"
+#include "LandscapeSplineRaster.h"
+#include "LandscapeProxy.h"
+#include "LandscapeEdit.h"
+#include "LandscapeHeightfieldCollisionComponent.h"
+#include "LandscapeSplinesComponent.h"
+#include "LandscapeSplineControlPoint.h"
+#include "LandscapeSplineSegment.h"
+#include "ScopedTransaction.h"
+#include "Raster.h"
+
+#define LOCTEXT_NAMESPACE "Landscape"
+
+//////////////////////////////////////////////////////////////////////////
+// Apply splines
+//////////////////////////////////////////////////////////////////////////
+
+class FLandscapeSplineHeightsRasterPolicy
+{
+public:
+	// X = Side Alpha, Y = End Alpha, Z = Height
+	typedef FVector InterpolantType;
+
+	/** Initialization constructor. */
+	FLandscapeSplineHeightsRasterPolicy(TArray<uint16>& InData, int32 InMinX, int32 InMinY, int32 InMaxX, int32 InMaxY, bool InbRaiseTerrain, bool InbLowerTerrain) :
+		Data(InData),
+		MinX(InMinX),
+		MinY(InMinY),
+		MaxX(InMaxX),
+		MaxY(InMaxY),
+		bRaiseTerrain(InbRaiseTerrain),
+		bLowerTerrain(InbLowerTerrain)
+	{
+	}
+
+protected:
+
+	// FTriangleRasterizer policy interface.
+
+	int32 GetMinX() const { return MinX; }
+	int32 GetMaxX() const { return MaxX; }
+	int32 GetMinY() const { return MinY; }
+	int32 GetMaxY() const { return MaxY; }
+
+	inline void ProcessPixel(int32 X, int32 Y, const InterpolantType& Interpolant, bool BackFacing)
+	{
+		const float CosInterpX = (Interpolant.X >= 1 ? 1 : 0.5f - 0.5f * FMath::Cos(Interpolant.X * PI));
+		const float CosInterpY = (Interpolant.Y >= 1 ? 1 : 0.5f - 0.5f * FMath::Cos(Interpolant.Y * PI));
+		const float Alpha = CosInterpX * CosInterpY;
+		uint16& Dest = Data[(Y - MinY)*(1 + MaxX - MinX) + X - MinX];
+		float Value = FMath::Lerp((float)Dest, Interpolant.Z, Alpha);
+		uint16 DValue = (uint32)FMath::Clamp<float>(Value, 0, LandscapeDataAccess::MaxValue);
+		if ((bRaiseTerrain && DValue > Dest) ||
+			(bLowerTerrain && DValue < Dest))
+		{
+			Dest = DValue;
+		}
+	}
+
+private:
+	TArray<uint16>& Data;
+	int32 MinX, MinY, MaxX, MaxY;
+	uint32 bRaiseTerrain : 1, bLowerTerrain : 1;
+};
+
+class FLandscapeSplineBlendmaskRasterPolicy
+{
+public:
+	// X = Side Alpha, Y = End Alpha, Z = Blend Value
+	typedef FVector InterpolantType;
+
+	/** Initialization constructor. */
+	FLandscapeSplineBlendmaskRasterPolicy(TArray<uint8>& InData, int32 InMinX, int32 InMinY, int32 InMaxX, int32 InMaxY) :
+		Data(InData),
+		MinX(InMinX),
+		MinY(InMinY),
+		MaxX(InMaxX),
+		MaxY(InMaxY)
+	{
+	}
+
+protected:
+
+	// FTriangleRasterizer policy interface.
+
+	int32 GetMinX() const { return MinX; }
+	int32 GetMaxX() const { return MaxX; }
+	int32 GetMinY() const { return MinY; }
+	int32 GetMaxY() const { return MaxY; }
+
+	inline void ProcessPixel(int32 X, int32 Y, const InterpolantType& Interpolant, bool BackFacing)
+	{
+		const float CosInterpX = (Interpolant.X >= 1 ? 1 : 0.5f - 0.5f * FMath::Cos(Interpolant.X * PI));
+		const float CosInterpY = (Interpolant.Y >= 1 ? 1 : 0.5f - 0.5f * FMath::Cos(Interpolant.Y * PI));
+		const float Alpha = CosInterpX * CosInterpY;
+		uint8& Dest = Data[(Y - MinY)*(1 + MaxX - MinX) + X - MinX];
+		float Value = FMath::Lerp((float)Dest, Interpolant.Z, Alpha);
+		Dest = (uint32)FMath::Clamp<float>(Value, 0, LandscapeDataAccess::MaxValue);
+	}
+
+private:
+	TArray<uint8>& Data;
+	int32 MinX, MinY, MaxX, MaxY;
+};
+
+bool ULandscapeInfo::ApplySplines(bool bOnlySelected)
+{
+	bool bResult = false;
+
+	ALandscape* Landscape = LandscapeActor.Get();
+
+	bResult |= ApplySplinesInternal(bOnlySelected, Landscape);
+
+	for (auto It = Proxies.CreateIterator(); It; ++It)
+	{
+		ALandscapeProxy* LandscapeProxy = (*It);
+		bResult |= ApplySplinesInternal(bOnlySelected, LandscapeProxy);
+	}
+
+	return bResult;
+}
+
+void RasterizeControlPointHeights(int32& MinX, int32& MinY, int32& MaxX, int32& MaxY, FLandscapeEditDataInterface& LandscapeEdit, FVector ControlPointLocation, const TArray<FLandscapeSplineInterpPoint>& Points, bool bRaiseTerrain, bool bLowerTerrain, TSet<ULandscapeComponent*>& ModifiedComponents)
+{
+	if (!(bRaiseTerrain || bLowerTerrain))
+	{
+		return;
+	}
+
+	if (MinX > MaxX || MinY > MaxY)
+	{
+		return;
+	}
+
+	TArray<uint16> Data;
+	Data.AddZeroed((1 + MaxY - MinY) * (1 + MaxX - MinX));
+
+	int32 ValidMinX = MinX;
+	int32 ValidMinY = MinY;
+	int32 ValidMaxX = MaxX;
+	int32 ValidMaxY = MaxY;
+	LandscapeEdit.GetHeightData(ValidMinX, ValidMinY, ValidMaxX, ValidMaxY, Data.GetData(), 0);
+
+	if (ValidMinX > ValidMaxX || ValidMinY > ValidMaxY)
+	{
+		// The control point's bounds don't intersect any data, so skip it
+		MinX = ValidMinX;
+		MinY = ValidMinY;
+		MaxX = ValidMaxX;
+		MaxY = ValidMaxY;
+
+		return;
+	}
+
+	FLandscapeEditDataInterface::ShrinkData(Data, MinX, MinY, MaxX, MaxY, ValidMinX, ValidMinY, ValidMaxX, ValidMaxY);
+
+	MinX = ValidMinX;
+	MinY = ValidMinY;
+	MaxX = ValidMaxX;
+	MaxY = ValidMaxY;
+
+	FTriangleRasterizer<FLandscapeSplineHeightsRasterPolicy> Rasterizer(
+		FLandscapeSplineHeightsRasterPolicy(Data, MinX, MinY, MaxX, MaxY, bRaiseTerrain, bLowerTerrain));
+
+	const FVector2D CenterPos = FVector2D(ControlPointLocation);
+	const FVector Center = FVector(1.0f, Points[0].StartEndFalloff, ControlPointLocation.Z * LANDSCAPE_INV_ZSCALE + LandscapeDataAccess::MidValue);
+
+	for (int32 i = Points.Num() - 1, j = 0; j < Points.Num(); i = j++)
+	{
+		// Solid center
+		const FVector2D Right0Pos = FVector2D(Points[i].Right);
+		const FVector2D Left1Pos = FVector2D(Points[j].Left);
+		const FVector2D Right1Pos = FVector2D(Points[j].Right);
+		const FVector Right0 = FVector(1.0f, Points[i].StartEndFalloff, Points[i].Right.Z);
+		const FVector Left1 = FVector(1.0f, Points[j].StartEndFalloff, Points[j].Left.Z);
+		const FVector Right1 = FVector(1.0f, Points[j].StartEndFalloff, Points[j].Right.Z);
+
+		Rasterizer.DrawTriangle(Center, Right0, Left1, CenterPos, Right0Pos, Left1Pos, false);
+		Rasterizer.DrawTriangle(Center, Left1, Right1, CenterPos, Left1Pos, Right1Pos, false);
+
+		// Falloff
+		FVector2D FalloffRight0Pos = FVector2D(Points[i].FalloffRight);
+		FVector2D FalloffLeft1Pos = FVector2D(Points[j].FalloffLeft);
+		FVector FalloffRight0 = FVector(0.0f, Points[i].StartEndFalloff, Points[i].FalloffRight.Z);
+		FVector FalloffLeft1 = FVector(0.0f, Points[j].StartEndFalloff, Points[j].FalloffLeft.Z);
+		Rasterizer.DrawTriangle(Right0, FalloffRight0, Left1, Right0Pos, FalloffRight0Pos, Left1Pos, false);
+		Rasterizer.DrawTriangle(FalloffRight0, Left1, FalloffLeft1, FalloffRight0Pos, Left1Pos, FalloffLeft1Pos, false);
+	}
+
+	LandscapeEdit.SetHeightData(MinX, MinY, MaxX, MaxY, Data.GetData(), 0, true);
+	LandscapeEdit.GetComponentsInRegion(MinX, MinY, MaxX, MaxY, &ModifiedComponents);
+}
+
+void RasterizeControlPointAlpha(int32& MinX, int32& MinY, int32& MaxX, int32& MaxY, FLandscapeEditDataInterface& LandscapeEdit, FVector ControlPointLocation, const TArray<FLandscapeSplineInterpPoint>& Points, ULandscapeLayerInfoObject* LayerInfo, TSet<ULandscapeComponent*>& ModifiedComponents)
+{
+	if (LayerInfo == nullptr)
+	{
+		return;
+	}
+
+	if (MinX > MaxX || MinY > MaxY)
+	{
+		return;
+	}
+
+	TArray<uint8> Data;
+	Data.AddZeroed((1 + MaxY - MinY) * (1 + MaxX - MinX));
+
+	int32 ValidMinX = MinX;
+	int32 ValidMinY = MinY;
+	int32 ValidMaxX = MaxX;
+	int32 ValidMaxY = MaxY;
+	LandscapeEdit.GetWeightData(LayerInfo, ValidMinX, ValidMinY, ValidMaxX, ValidMaxY, Data.GetData(), 0);
+
+	if (ValidMinX > ValidMaxX || ValidMinY > ValidMaxY)
+	{
+		// The control point's bounds don't intersect any data, so skip it
+		MinX = ValidMinX;
+		MinY = ValidMinY;
+		MaxX = ValidMaxX;
+		MaxY = ValidMaxY;
+
+		return;
+	}
+
+	FLandscapeEditDataInterface::ShrinkData(Data, MinX, MinY, MaxX, MaxY, ValidMinX, ValidMinY, ValidMaxX, ValidMaxY);
+
+	MinX = ValidMinX;
+	MinY = ValidMinY;
+	MaxX = ValidMaxX;
+	MaxY = ValidMaxY;
+
+	FTriangleRasterizer<FLandscapeSplineBlendmaskRasterPolicy> Rasterizer(
+		FLandscapeSplineBlendmaskRasterPolicy(Data, MinX, MinY, MaxX, MaxY));
+
+	const float BlendValue = 255;
+
+	const FVector2D CenterPos = FVector2D(ControlPointLocation);
+	const FVector Center = FVector(1.0f, Points[0].StartEndFalloff, BlendValue);
+
+	for (int32 i = Points.Num() - 1, j = 0; j < Points.Num(); i = j++)
+	{
+		// Solid center
+		const FVector2D Right0Pos = FVector2D(Points[i].Right);
+		const FVector2D Left1Pos = FVector2D(Points[j].Left);
+		const FVector2D Right1Pos = FVector2D(Points[j].Right);
+		const FVector Right0 = FVector(1.0f, Points[i].StartEndFalloff, BlendValue);
+		const FVector Left1 = FVector(1.0f, Points[j].StartEndFalloff, BlendValue);
+		const FVector Right1 = FVector(1.0f, Points[j].StartEndFalloff, BlendValue);
+
+		Rasterizer.DrawTriangle(Center, Right0, Left1, CenterPos, Right0Pos, Left1Pos, false);
+		Rasterizer.DrawTriangle(Center, Left1, Right1, CenterPos, Left1Pos, Right1Pos, false);
+
+		// Falloff
+		FVector2D FalloffRight0Pos = FVector2D(Points[i].FalloffRight);
+		FVector2D FalloffLeft1Pos = FVector2D(Points[j].FalloffLeft);
+		FVector FalloffRight0 = FVector(0.0f, Points[i].StartEndFalloff, BlendValue);
+		FVector FalloffLeft1 = FVector(0.0f, Points[j].StartEndFalloff, BlendValue);
+		Rasterizer.DrawTriangle(Right0, FalloffRight0, Left1, Right0Pos, FalloffRight0Pos, Left1Pos, false);
+		Rasterizer.DrawTriangle(FalloffRight0, Left1, FalloffLeft1, FalloffRight0Pos, Left1Pos, FalloffLeft1Pos, false);
+	}
+
+	LandscapeEdit.SetAlphaData(LayerInfo, MinX, MinY, MaxX, MaxY, Data.GetData(), 0, ELandscapeLayerPaintingRestriction::None, true, false);
+	LandscapeEdit.GetComponentsInRegion(MinX, MinY, MaxX, MaxY, &ModifiedComponents);
+}
+
+void RasterizeSegmentHeight(int32& MinX, int32& MinY, int32& MaxX, int32& MaxY, FLandscapeEditDataInterface& LandscapeEdit, const TArray<FLandscapeSplineInterpPoint>& Points, bool bRaiseTerrain, bool bLowerTerrain, TSet<ULandscapeComponent*>& ModifiedComponents)
+{
+	if (!(bRaiseTerrain || bLowerTerrain))
+	{
+		return;
+	}
+
+	if (MinX > MaxX || MinY > MaxY)
+	{
+		return;
+	}
+
+	TArray<uint16> Data;
+	Data.AddZeroed((1 + MaxY - MinY) * (1 + MaxX - MinX));
+
+	int32 ValidMinX = MinX;
+	int32 ValidMinY = MinY;
+	int32 ValidMaxX = MaxX;
+	int32 ValidMaxY = MaxY;
+	LandscapeEdit.GetHeightData(ValidMinX, ValidMinY, ValidMaxX, ValidMaxY, Data.GetData(), 0);
+
+	if (ValidMinX > ValidMaxX || ValidMinY > ValidMaxY)
+	{
+		// The segment's bounds don't intersect any data, so skip it
+		MinX = ValidMinX;
+		MinY = ValidMinY;
+		MaxX = ValidMaxX;
+		MaxY = ValidMaxY;
+
+		return;
+	}
+
+	FLandscapeEditDataInterface::ShrinkData(Data, MinX, MinY, MaxX, MaxY, ValidMinX, ValidMinY, ValidMaxX, ValidMaxY);
+
+	MinX = ValidMinX;
+	MinY = ValidMinY;
+	MaxX = ValidMaxX;
+	MaxY = ValidMaxY;
+
+	FTriangleRasterizer<FLandscapeSplineHeightsRasterPolicy> Rasterizer(
+		FLandscapeSplineHeightsRasterPolicy(Data, MinX, MinY, MaxX, MaxY, bRaiseTerrain, bLowerTerrain));
+
+	for (int32 j = 1; j < Points.Num(); j++)
+	{
+		// Middle
+		FVector2D Left0Pos = FVector2D(Points[j - 1].Left);
+		FVector2D Right0Pos = FVector2D(Points[j - 1].Right);
+		FVector2D Left1Pos = FVector2D(Points[j].Left);
+		FVector2D Right1Pos = FVector2D(Points[j].Right);
+		FVector Left0 = FVector(1.0f, Points[j - 1].StartEndFalloff, Points[j - 1].Left.Z);
+		FVector Right0 = FVector(1.0f, Points[j - 1].StartEndFalloff, Points[j - 1].Right.Z);
+		FVector Left1 = FVector(1.0f, Points[j].StartEndFalloff, Points[j].Left.Z);
+		FVector Right1 = FVector(1.0f, Points[j].StartEndFalloff, Points[j].Right.Z);
+		Rasterizer.DrawTriangle(Left0, Right0, Left1, Left0Pos, Right0Pos, Left1Pos, false);
+		Rasterizer.DrawTriangle(Right0, Left1, Right1, Right0Pos, Left1Pos, Right1Pos, false);
+
+		// Left Falloff
+		FVector2D FalloffLeft0Pos = FVector2D(Points[j - 1].FalloffLeft);
+		FVector2D FalloffLeft1Pos = FVector2D(Points[j].FalloffLeft);
+		FVector FalloffLeft0 = FVector(0.0f, Points[j - 1].StartEndFalloff, Points[j - 1].FalloffLeft.Z);
+		FVector FalloffLeft1 = FVector(0.0f, Points[j].StartEndFalloff, Points[j].FalloffLeft.Z);
+		Rasterizer.DrawTriangle(FalloffLeft0, Left0, FalloffLeft1, FalloffLeft0Pos, Left0Pos, FalloffLeft1Pos, false);
+		Rasterizer.DrawTriangle(Left0, FalloffLeft1, Left1, Left0Pos, FalloffLeft1Pos, Left1Pos, false);
+
+		// Right Falloff
+		FVector2D FalloffRight0Pos = FVector2D(Points[j - 1].FalloffRight);
+		FVector2D FalloffRight1Pos = FVector2D(Points[j].FalloffRight);
+		FVector FalloffRight0 = FVector(0.0f, Points[j - 1].StartEndFalloff, Points[j - 1].FalloffRight.Z);
+		FVector FalloffRight1 = FVector(0.0f, Points[j].StartEndFalloff, Points[j].FalloffRight.Z);
+		Rasterizer.DrawTriangle(Right0, FalloffRight0, Right1, Right0Pos, FalloffRight0Pos, Right1Pos, false);
+		Rasterizer.DrawTriangle(FalloffRight0, Right1, FalloffRight1, FalloffRight0Pos, Right1Pos, FalloffRight1Pos, false);
+	}
+
+	LandscapeEdit.SetHeightData(MinX, MinY, MaxX, MaxY, Data.GetData(), 0, true);
+	LandscapeEdit.GetComponentsInRegion(MinX, MinY, MaxX, MaxY, &ModifiedComponents);
+}
+
+
+void RasterizeSegmentAlpha(int32& MinX, int32& MinY, int32& MaxX, int32& MaxY, FLandscapeEditDataInterface& LandscapeEdit, const TArray<FLandscapeSplineInterpPoint>& Points, ULandscapeLayerInfoObject* LayerInfo, TSet<ULandscapeComponent*>& ModifiedComponents)
+{
+	if (LayerInfo == nullptr)
+	{
+		return;
+	}
+
+	if (MinX > MaxX || MinY > MaxY)
+	{
+		return;
+	}
+
+	TArray<uint8> Data;
+	Data.AddZeroed((1 + MaxY - MinY) * (1 + MaxX - MinX));
+
+	int32 ValidMinX = MinX;
+	int32 ValidMinY = MinY;
+	int32 ValidMaxX = MaxX;
+	int32 ValidMaxY = MaxY;
+	LandscapeEdit.GetWeightData(LayerInfo, ValidMinX, ValidMinY, ValidMaxX, ValidMaxY, Data.GetData(), 0);
+
+	if (ValidMinX > ValidMaxX || ValidMinY > ValidMaxY)
+	{
+		// The segment's bounds don't intersect any data, so skip it
+		MinX = ValidMinX;
+		MinY = ValidMinY;
+		MaxX = ValidMaxX;
+		MaxY = ValidMaxY;
+
+		return;
+	}
+
+	FLandscapeEditDataInterface::ShrinkData(Data, MinX, MinY, MaxX, MaxY, ValidMinX, ValidMinY, ValidMaxX, ValidMaxY);
+
+	MinX = ValidMinX;
+	MinY = ValidMinY;
+	MaxX = ValidMaxX;
+	MaxY = ValidMaxY;
+
+	FTriangleRasterizer<FLandscapeSplineBlendmaskRasterPolicy> Rasterizer(
+		FLandscapeSplineBlendmaskRasterPolicy(Data, MinX, MinY, MaxX, MaxY));
+
+	const float BlendValue = 255;
+
+	for (int32 j = 1; j < Points.Num(); j++)
+	{
+		// Middle
+		FVector2D Left0Pos = FVector2D(Points[j - 1].Left);
+		FVector2D Right0Pos = FVector2D(Points[j - 1].Right);
+		FVector2D Left1Pos = FVector2D(Points[j].Left);
+		FVector2D Right1Pos = FVector2D(Points[j].Right);
+		FVector Left0 = FVector(1.0f, Points[j - 1].StartEndFalloff, BlendValue);
+		FVector Right0 = FVector(1.0f, Points[j - 1].StartEndFalloff, BlendValue);
+		FVector Left1 = FVector(1.0f, Points[j].StartEndFalloff, BlendValue);
+		FVector Right1 = FVector(1.0f, Points[j].StartEndFalloff, BlendValue);
+		Rasterizer.DrawTriangle(Left0, Right0, Left1, Left0Pos, Right0Pos, Left1Pos, false);
+		Rasterizer.DrawTriangle(Right0, Left1, Right1, Right0Pos, Left1Pos, Right1Pos, false);
+
+		// Left Falloff
+		FVector2D FalloffLeft0Pos = FVector2D(Points[j - 1].FalloffLeft);
+		FVector2D FalloffLeft1Pos = FVector2D(Points[j].FalloffLeft);
+		FVector FalloffLeft0 = FVector(0.0f, Points[j - 1].StartEndFalloff, BlendValue);
+		FVector FalloffLeft1 = FVector(0.0f, Points[j].StartEndFalloff, BlendValue);
+		Rasterizer.DrawTriangle(FalloffLeft0, Left0, FalloffLeft1, FalloffLeft0Pos, Left0Pos, FalloffLeft1Pos, false);
+		Rasterizer.DrawTriangle(Left0, FalloffLeft1, Left1, Left0Pos, FalloffLeft1Pos, Left1Pos, false);
+
+		// Right Falloff
+		FVector2D FalloffRight0Pos = FVector2D(Points[j - 1].FalloffRight);
+		FVector2D FalloffRight1Pos = FVector2D(Points[j].FalloffRight);
+		FVector FalloffRight0 = FVector(0.0f, Points[j - 1].StartEndFalloff, BlendValue);
+		FVector FalloffRight1 = FVector(0.0f, Points[j].StartEndFalloff, BlendValue);
+		Rasterizer.DrawTriangle(Right0, FalloffRight0, Right1, Right0Pos, FalloffRight0Pos, Right1Pos, false);
+		Rasterizer.DrawTriangle(FalloffRight0, Right1, FalloffRight1, FalloffRight0Pos, Right1Pos, FalloffRight1Pos, false);
+	}
+
+	LandscapeEdit.SetAlphaData(LayerInfo, MinX, MinY, MaxX, MaxY, Data.GetData(), 0, ELandscapeLayerPaintingRestriction::None, true, false);
+	LandscapeEdit.GetComponentsInRegion(MinX, MinY, MaxX, MaxY, &ModifiedComponents);
+}
+
+bool ULandscapeInfo::ApplySplinesInternal(bool bOnlySelected, ALandscapeProxy* Landscape)
+{
+	if (!Landscape || !Landscape->SplineComponent || Landscape->SplineComponent->ControlPoints.Num() == 0 || Landscape->SplineComponent->Segments.Num() == 0)
+	{
+		return false;
+	}
+
+	FScopedTransaction Transaction(LOCTEXT("LandscapeSpline_ApplySplines", "Apply Splines to Landscape"));
+
+	const FTransform SplineToLandscape = Landscape->SplineComponent->ComponentToWorld.GetRelativeTransform(Landscape->LandscapeActorToWorld());
+
+	FLandscapeEditDataInterface LandscapeEdit(this);
+	TSet<ULandscapeComponent*> ModifiedComponents;
+
+	// I'd dearly love to use FIntRect in this code, but Landscape works with "Inclusive Max" and FIntRect is "Exclusive Max"
+	int32 LandscapeMinX, LandscapeMinY, LandscapeMaxX, LandscapeMaxY;
+	if (!GetLandscapeExtent(LandscapeMinX, LandscapeMinY, LandscapeMaxX, LandscapeMaxY))
+	{
+		return false;
+	}
+
+	for (const ULandscapeSplineControlPoint* ControlPoint : Landscape->SplineComponent->ControlPoints)
+	{
+		if (bOnlySelected && !ControlPoint->IsSplineSelected())
+		{
+			continue;
+		}
+
+		if (ControlPoint->GetPoints().Num() < 2)
+		{
+			continue;
+		}
+
+		FBox ControlPointBounds = ControlPoint->GetBounds();
+		ControlPointBounds = ControlPointBounds.TransformBy(SplineToLandscape.ToMatrixWithScale());
+
+		int32 MinX = FMath::CeilToInt(ControlPointBounds.Min.X);
+		int32 MinY = FMath::CeilToInt(ControlPointBounds.Min.Y);
+		int32 MaxX = FMath::FloorToInt(ControlPointBounds.Max.X);
+		int32 MaxY = FMath::FloorToInt(ControlPointBounds.Max.Y);
+
+		MinX = FMath::Max(MinX, LandscapeMinX);
+		MinY = FMath::Max(MinY, LandscapeMinY);
+		MaxX = FMath::Min(MaxX, LandscapeMaxX);
+		MaxY = FMath::Min(MaxY, LandscapeMaxY);
+
+		if (MinX > MaxX || MinY > MaxY)
+		{
+			// The control point's bounds don't intersect the landscape, so skip it entirely
+			continue;
+		}
+
+		TArray<FLandscapeSplineInterpPoint> Points = ControlPoint->GetPoints();
+		for (int32 j = 0; j < Points.Num(); j++)
+		{
+			Points[j].Center = SplineToLandscape.TransformPosition(Points[j].Center);
+			Points[j].Left = SplineToLandscape.TransformPosition(Points[j].Left);
+			Points[j].Right = SplineToLandscape.TransformPosition(Points[j].Right);
+			Points[j].FalloffLeft = SplineToLandscape.TransformPosition(Points[j].FalloffLeft);
+			Points[j].FalloffRight = SplineToLandscape.TransformPosition(Points[j].FalloffRight);
+
+			// local-heights to texture value heights
+			Points[j].Left.Z = Points[j].Left.Z * LANDSCAPE_INV_ZSCALE + LandscapeDataAccess::MidValue;
+			Points[j].Right.Z = Points[j].Right.Z * LANDSCAPE_INV_ZSCALE + LandscapeDataAccess::MidValue;
+			Points[j].FalloffLeft.Z = Points[j].FalloffLeft.Z * LANDSCAPE_INV_ZSCALE + LandscapeDataAccess::MidValue;
+			Points[j].FalloffRight.Z = Points[j].FalloffRight.Z * LANDSCAPE_INV_ZSCALE + LandscapeDataAccess::MidValue;
+		}
+
+		// Heights raster
+		if (ControlPoint->bRaiseTerrain || ControlPoint->bLowerTerrain)
+		{
+			const FVector Center3D = SplineToLandscape.TransformPosition(ControlPoint->Location);
+
+			RasterizeControlPointHeights(MinX, MinY, MaxX, MaxY, LandscapeEdit, Center3D, Points, ControlPoint->bRaiseTerrain, ControlPoint->bLowerTerrain, ModifiedComponents);
+		}
+
+		// Blend layer raster
+		ULandscapeLayerInfoObject* LayerInfo = GetLayerInfoByName(ControlPoint->LayerName);
+		if (ControlPoint->LayerName != NAME_None && LayerInfo != NULL)
+		{
+			const FVector Center3D = SplineToLandscape.TransformPosition(ControlPoint->Location);
+
+			RasterizeControlPointAlpha(MinX, MinY, MaxX, MaxY, LandscapeEdit, Center3D, Points, LayerInfo, ModifiedComponents);
+		}
+	}
+
+	for (const ULandscapeSplineSegment* Segment : Landscape->SplineComponent->Segments)
+	{
+		if (bOnlySelected && !Segment->IsSplineSelected())
+		{
+			continue;
+		}
+
+		FBox SegmentBounds = Segment->GetBounds();
+		SegmentBounds = SegmentBounds.TransformBy(SplineToLandscape.ToMatrixWithScale());
+
+		int32 MinX = FMath::CeilToInt(SegmentBounds.Min.X);
+		int32 MinY = FMath::CeilToInt(SegmentBounds.Min.Y);
+		int32 MaxX = FMath::FloorToInt(SegmentBounds.Max.X);
+		int32 MaxY = FMath::FloorToInt(SegmentBounds.Max.Y);
+
+		MinX = FMath::Max(MinX, LandscapeMinX);
+		MinY = FMath::Max(MinY, LandscapeMinY);
+		MaxX = FMath::Min(MaxX, LandscapeMaxX);
+		MaxY = FMath::Min(MaxY, LandscapeMaxY);
+
+		if (MinX > MaxX || MinY > MaxY)
+		{
+			// The segment's bounds don't intersect the landscape, so skip it entirely
+			continue;
+		}
+
+		TArray<FLandscapeSplineInterpPoint> Points = Segment->GetPoints();
+		for (int32 j = 0; j < Points.Num(); j++)
+		{
+			Points[j].Center = SplineToLandscape.TransformPosition(Points[j].Center);
+			Points[j].Left = SplineToLandscape.TransformPosition(Points[j].Left);
+			Points[j].Right = SplineToLandscape.TransformPosition(Points[j].Right);
+			Points[j].FalloffLeft = SplineToLandscape.TransformPosition(Points[j].FalloffLeft);
+			Points[j].FalloffRight = SplineToLandscape.TransformPosition(Points[j].FalloffRight);
+
+			// local-heights to texture value heights
+			Points[j].Left.Z = Points[j].Left.Z * LANDSCAPE_INV_ZSCALE + LandscapeDataAccess::MidValue;
+			Points[j].Right.Z = Points[j].Right.Z * LANDSCAPE_INV_ZSCALE + LandscapeDataAccess::MidValue;
+			Points[j].FalloffLeft.Z = Points[j].FalloffLeft.Z * LANDSCAPE_INV_ZSCALE + LandscapeDataAccess::MidValue;
+			Points[j].FalloffRight.Z = Points[j].FalloffRight.Z * LANDSCAPE_INV_ZSCALE + LandscapeDataAccess::MidValue;
+		}
+
+		// Heights raster
+		if (Segment->bRaiseTerrain || Segment->bLowerTerrain)
+		{
+			RasterizeSegmentHeight(MinX, MinY, MaxX, MaxY, LandscapeEdit, Points, Segment->bRaiseTerrain, Segment->bLowerTerrain, ModifiedComponents);
+
+			if (MinX > MaxX || MinY > MaxY)
+			{
+				// The segment's bounds don't intersect any data, so we skip it entirely
+				// it wouldn't intersect any weightmap data either so we don't even bother trying
+			}
+		}
+
+		// Blend layer raster
+		ULandscapeLayerInfoObject* LayerInfo = GetLayerInfoByName(Segment->LayerName);
+		if (Segment->LayerName != NAME_None && LayerInfo != NULL)
+		{
+			RasterizeSegmentAlpha(MinX, MinY, MaxX, MaxY, LandscapeEdit, Points, LayerInfo, ModifiedComponents);
+		}
+	}
+
+	LandscapeEdit.Flush();
+
+	for (ULandscapeComponent* Component : ModifiedComponents)
+	{
+		// Recreate collision for modified components and update the navmesh
+		ULandscapeHeightfieldCollisionComponent* CollisionComponent = Component->CollisionComponent.Get();
+		if (CollisionComponent)
+		{
+			CollisionComponent->RecreateCollision(false);
+			UNavigationSystem* NavSys = UNavigationSystem::GetCurrent(Component);
+			if (NavSys)
+			{
+				NavSys->UpdateNavOctree(CollisionComponent);
+			}
+		}
+	}
+
+	return true;
+}
+
+namespace LandscapeSplineRaster
+{
+	void RasterizeSegmentPoints(ULandscapeInfo* LandscapeInfo, TArray<FLandscapeSplineInterpPoint> Points, const FTransform& SplineToWorld, bool bRaiseTerrain, bool bLowerTerrain, ULandscapeLayerInfoObject* LayerInfo)
+	{
+		ALandscapeProxy* LandscapeProxy = LandscapeInfo->GetLandscapeProxy();
+		const FTransform SplineToLandscape = SplineToWorld.GetRelativeTransform(LandscapeProxy->LandscapeActorToWorld());
+
+		FLandscapeEditDataInterface LandscapeEdit(LandscapeInfo);
+		TSet<ULandscapeComponent*> ModifiedComponents;
+
+		// I'd dearly love to use FIntRect in this code, but Landscape works with "Inclusive Max" and FIntRect is "Exclusive Max"
+		int32 LandscapeMinX, LandscapeMinY, LandscapeMaxX, LandscapeMaxY;
+		if (!LandscapeInfo->GetLandscapeExtent(LandscapeMinX, LandscapeMinY, LandscapeMaxX, LandscapeMaxY))
+		{
+			return;
+		}
+
+		FBox SegmentBounds = FBox(0);
+		for (const FLandscapeSplineInterpPoint& Point : Points)
+		{
+			SegmentBounds += Point.FalloffLeft;
+			SegmentBounds += Point.FalloffRight;
+		}
+
+		SegmentBounds = SegmentBounds.TransformBy(SplineToLandscape.ToMatrixWithScale());
+
+		int32 MinX = FMath::CeilToInt(SegmentBounds.Min.X);
+		int32 MinY = FMath::CeilToInt(SegmentBounds.Min.Y);
+		int32 MaxX = FMath::FloorToInt(SegmentBounds.Max.X);
+		int32 MaxY = FMath::FloorToInt(SegmentBounds.Max.Y);
+
+		MinX = FMath::Max(MinX, LandscapeMinX);
+		MinY = FMath::Max(MinY, LandscapeMinY);
+		MaxX = FMath::Min(MaxX, LandscapeMaxX);
+		MaxY = FMath::Min(MaxY, LandscapeMaxY);
+
+		if (MinX > MaxX || MinY > MaxY)
+		{
+			// The segment's bounds don't intersect the landscape, so skip it entirely
+			return;
+		}
+
+		for (int32 j = 0; j < Points.Num(); j++)
+		{
+			Points[j].Center = SplineToLandscape.TransformPosition(Points[j].Center);
+			Points[j].Left = SplineToLandscape.TransformPosition(Points[j].Left);
+			Points[j].Right = SplineToLandscape.TransformPosition(Points[j].Right);
+			Points[j].FalloffLeft = SplineToLandscape.TransformPosition(Points[j].FalloffLeft);
+			Points[j].FalloffRight = SplineToLandscape.TransformPosition(Points[j].FalloffRight);
+
+			// local-heights to texture value heights
+			Points[j].Left.Z = Points[j].Left.Z * LANDSCAPE_INV_ZSCALE + LandscapeDataAccess::MidValue;
+			Points[j].Right.Z = Points[j].Right.Z * LANDSCAPE_INV_ZSCALE + LandscapeDataAccess::MidValue;
+			Points[j].FalloffLeft.Z = Points[j].FalloffLeft.Z * LANDSCAPE_INV_ZSCALE + LandscapeDataAccess::MidValue;
+			Points[j].FalloffRight.Z = Points[j].FalloffRight.Z * LANDSCAPE_INV_ZSCALE + LandscapeDataAccess::MidValue;
+		}
+
+		// Heights raster
+		if (bRaiseTerrain || bLowerTerrain)
+		{
+			RasterizeSegmentHeight(MinX, MinY, MaxX, MaxY, LandscapeEdit, Points, bRaiseTerrain, bLowerTerrain, ModifiedComponents);
+
+			if (MinX > MaxX || MinY > MaxY)
+			{
+				// The segment's bounds don't intersect any data, so we skip it entirely
+				// it wouldn't intersect any weightmap data either so we don't even bother trying
+			}
+		}
+
+		// Blend layer raster
+		if (LayerInfo != NULL)
+		{
+			RasterizeSegmentAlpha(MinX, MinY, MaxX, MaxY, LandscapeEdit, Points, LayerInfo, ModifiedComponents);
+		}
+
+		LandscapeEdit.Flush();
+
+		for (ULandscapeComponent* Component : ModifiedComponents)
+		{
+			// Recreate collision for modified components and update the navmesh
+			ULandscapeHeightfieldCollisionComponent* CollisionComponent = Component->CollisionComponent.Get();
+			if (CollisionComponent)
+			{
+				CollisionComponent->RecreateCollision(false);
+				UNavigationSystem* NavSys = UNavigationSystem::GetCurrent(Component);
+				if (NavSys)
+				{
+					NavSys->UpdateNavOctree(CollisionComponent);
+				}
+			}
+		}
+	}
+
+	static bool LineIntersect(const FVector2D& L1Start, const FVector2D& L1End, const FVector2D& L2Start, const FVector2D& L2End, FVector2D& Intersect, float Tolerance = KINDA_SMALL_NUMBER)
+	{
+		float tA = (L2End - L2Start) ^ (L2Start - L1Start);
+		float tB = (L1End - L1Start) ^ (L2Start - L1Start);
+		float Denom = (L2End - L2Start) ^ (L1End - L1Start);
+
+		if (FMath::IsNearlyZero(tA) && FMath::IsNearlyZero(tB))
+		{
+			// Lines are the same
+			Intersect = (L2Start + L2End) / 2;
+			return true;
+		}
+
+		if (FMath::IsNearlyZero(Denom))
+		{
+			// Lines are parallel
+			Intersect = (L2Start + L2End) / 2;
+			return false;
+		}
+
+		tA /= Denom;
+		tB /= Denom;
+
+		Intersect = L1Start + tA * (L1End - L1Start);
+
+		if (tA >= -Tolerance && tA <= (1.0f + Tolerance) && tB >= -Tolerance && tB <= (1.0f + Tolerance))
+		{
+			return true;
+		}
+
+		return false;
+	}
+
+	bool FixSelfIntersection(TArray<FLandscapeSplineInterpPoint>& Points, FVector FLandscapeSplineInterpPoint::* Side)
+	{
+		int32 StartSide = INDEX_NONE;
+		for (int32 i = 0; i < Points.Num(); i++)
+		{
+			bool bReversed = false;
+
+			if (i < Points.Num() - 1)
+			{
+				const FLandscapeSplineInterpPoint& CurrentPoint = Points[i];
+				const FLandscapeSplineInterpPoint& NextPoint = Points[i + 1];
+				const FVector Direction = (NextPoint.Center - CurrentPoint.Center).GetSafeNormal();
+				const FVector SideDirection = (NextPoint.*Side - CurrentPoint.*Side).GetSafeNormal();
+				bReversed = (SideDirection | Direction) < 0;
+			}
+
+			if (bReversed)
+			{
+				if (StartSide == INDEX_NONE)
+				{
+					StartSide = i;
+				}
+			}
+			else
+			{
+				if (StartSide != INDEX_NONE)
+				{
+					int32 EndSide = i;
+
+					// step startSide back until before the endSide point
+					while (StartSide > 0)
+					{
+						const float Projection = (Points[StartSide].*Side - Points[StartSide - 1].*Side) | (Points[EndSide].*Side - Points[StartSide - 1].*Side);
+						if (Projection >= 0)
+						{
+							break;
+						}
+						StartSide--;
+					}
+					// step endSide forwards until after the startSide point
+					while (EndSide < Points.Num() - 1)
+					{
+						const float Projection = (Points[EndSide].*Side - Points[EndSide + 1].*Side) | (Points[StartSide].*Side - Points[EndSide + 1].*Side);
+						if (Projection >= 0)
+						{
+							break;
+						}
+						EndSide++;
+					}
+
+					// Can't do anything if the start and end intersect, as they're both unalterable
+					if (StartSide == 0 && EndSide == Points.Num() - 1)
+					{
+						return false;
+					}
+
+					FVector2D Collapse;
+					if (StartSide == 0)
+					{
+						Collapse = FVector2D(Points[StartSide].*Side);
+						StartSide++;
+					}
+					else if (EndSide == Points.Num() - 1)
+					{
+						Collapse = FVector2D(Points[EndSide].*Side);
+						EndSide--;
+					}
+					else
+					{
+						LineIntersect(FVector2D(Points[StartSide - 1].*Side), FVector2D(Points[StartSide].*Side),
+							FVector2D(Points[EndSide + 1].*Side), FVector2D(Points[EndSide].*Side), Collapse);
+					}
+
+					for (int32 j = StartSide; j <= EndSide; j++)
+					{
+						(Points[j].*Side).X = Collapse.X;
+						(Points[j].*Side).Y = Collapse.Y;
+					}
+
+					StartSide = INDEX_NONE;
+					i = EndSide;
+				}
+			}
+		}
+
+		return true;
+	}
+
+	void Pointify(const FInterpCurveVector& SplineInfo, TArray<FLandscapeSplineInterpPoint>& Points, int32 NumSubdivisions,
+		float StartFalloffFraction, float EndFalloffFraction,
+		const float StartWidth, const float EndWidth,
+		const float StartSideFalloff, const float EndSideFalloff,
+		const float StartRollDegrees, const float EndRollDegrees)
+	{
+		// Stop the start and end fall-off overlapping
+		const float TotalFalloff = StartFalloffFraction + EndFalloffFraction;
+		if (TotalFalloff > 1.0f)
+		{
+			StartFalloffFraction /= TotalFalloff;
+			EndFalloffFraction /= TotalFalloff;
+		}
+
+		const float StartRoll = FMath::DegreesToRadians(StartRollDegrees);
+		const float EndRoll = FMath::DegreesToRadians(EndRollDegrees);
+
+		float OldKeyTime = 0;
+		for (int32 i = 0; i < SplineInfo.Points.Num(); i++)
+		{
+			const float NewKeyTime = SplineInfo.Points[i].InVal;
+			const float NewKeyCosInterp = 0.5f - 0.5f * FMath::Cos(NewKeyTime * PI);
+			const float NewKeyWidth = FMath::Lerp(StartWidth, EndWidth, NewKeyCosInterp);
+			const float NewKeyFalloff = FMath::Lerp(StartSideFalloff, EndSideFalloff, NewKeyCosInterp);
+			const float NewKeyRoll = FMath::Lerp(StartRoll, EndRoll, NewKeyCosInterp);
+			const FVector NewKeyPos = SplineInfo.Eval(NewKeyTime, FVector::ZeroVector);
+			const FVector NewKeyTangent = SplineInfo.EvalDerivative(NewKeyTime, FVector::ZeroVector).GetSafeNormal();
+			const FVector NewKeyBiNormal = FQuat(NewKeyTangent, -NewKeyRoll).RotateVector((NewKeyTangent ^ FVector(0, 0, -1)).GetSafeNormal());
+			const FVector NewKeyLeftPos = NewKeyPos - NewKeyBiNormal * NewKeyWidth;
+			const FVector NewKeyRightPos = NewKeyPos + NewKeyBiNormal * NewKeyWidth;
+			const FVector NewKeyFalloffLeftPos = NewKeyPos - NewKeyBiNormal * (NewKeyWidth + NewKeyFalloff);
+			const FVector NewKeyFalloffRightPos = NewKeyPos + NewKeyBiNormal * (NewKeyWidth + NewKeyFalloff);
+			const float NewKeyStartEndFalloff = FMath::Min((StartFalloffFraction > 0 ? NewKeyTime / StartFalloffFraction : 1.0f), (EndFalloffFraction > 0 ? (1 - NewKeyTime) / EndFalloffFraction : 1.0f));
+
+			// If not the first keypoint, interp from the last keypoint.
+			if (i > 0)
+			{
+				const int32 NumSteps = FMath::CeilToInt((NewKeyTime - OldKeyTime) * NumSubdivisions);
+				const float DrawSubstep = (NewKeyTime - OldKeyTime) / NumSteps;
+
+				// Add a point for each substep, except the ends because that's the point added outside the interp'ing.
+				for (int32 j = 1; j < NumSteps; j++)
+				{
+					const float NewTime = OldKeyTime + j*DrawSubstep;
+					const float NewCosInterp = 0.5f - 0.5f * FMath::Cos(NewTime * PI);
+					const float NewWidth = FMath::Lerp(StartWidth, EndWidth, NewCosInterp);
+					const float NewFalloff = FMath::Lerp(StartSideFalloff, EndSideFalloff, NewCosInterp);
+					const float NewRoll = FMath::Lerp(StartRoll, EndRoll, NewCosInterp);
+					const FVector NewPos = SplineInfo.Eval(NewTime, FVector::ZeroVector);
+					const FVector NewTangent = SplineInfo.EvalDerivative(NewTime, FVector::ZeroVector).GetSafeNormal();
+					const FVector NewBiNormal = FQuat(NewTangent, -NewRoll).RotateVector((NewTangent ^ FVector(0, 0, -1)).GetSafeNormal());
+					const FVector NewLeftPos = NewPos - NewBiNormal * NewWidth;
+					const FVector NewRightPos = NewPos + NewBiNormal * NewWidth;
+					const FVector NewFalloffLeftPos = NewPos - NewBiNormal * (NewWidth + NewFalloff);
+					const FVector NewFalloffRightPos = NewPos + NewBiNormal * (NewWidth + NewFalloff);
+					const float NewStartEndFalloff = FMath::Min((StartFalloffFraction > 0 ? NewTime / StartFalloffFraction : 1.0f), (EndFalloffFraction > 0 ? (1 - NewTime) / EndFalloffFraction : 1.0f));
+
+					Points.Emplace(NewPos, NewLeftPos, NewRightPos, NewFalloffLeftPos, NewFalloffRightPos, NewStartEndFalloff);
+				}
+			}
+
+			Points.Emplace(NewKeyPos, NewKeyLeftPos, NewKeyRightPos, NewKeyFalloffLeftPos, NewKeyFalloffRightPos, NewKeyStartEndFalloff);
+
+			OldKeyTime = NewKeyTime;
+		}
+
+		// Handle self-intersection errors due to tight turns
+		FixSelfIntersection(Points, &FLandscapeSplineInterpPoint::Left);
+		FixSelfIntersection(Points, &FLandscapeSplineInterpPoint::Right);
+		FixSelfIntersection(Points, &FLandscapeSplineInterpPoint::FalloffLeft);
+		FixSelfIntersection(Points, &FLandscapeSplineInterpPoint::FalloffRight);
+	}
+}
+
+#undef LOCTEXT_NAMESPACE
diff --git a/Engine/Source/Runtime/Landscape/Private/LandscapeSplineRaster.h b/Engine/Source/Runtime/Landscape/Private/LandscapeSplineRaster.h
new file mode 100644
--- /dev/null
+++ b/Engine/Source/Runtime/Landscape/Private/LandscapeSplineRaster.h
@@ -0,0 +1,22 @@
+// Copyright 1998-2015 Epic Games, Inc. All Rights Reserved.
+
+#pragma once
+
+class ULandscapeInfo;
+class ULandscapeLayerInfoObject;
+struct FLandscapeSplineInterpPoint;
+
+namespace LandscapeSplineRaster
+{
+#if WITH_EDITOR
+	bool FixSelfIntersection(TArray<FLandscapeSplineInterpPoint>& Points, FVector FLandscapeSplineInterpPoint::* Side);
+
+	void Pointify(const FInterpCurveVector& SplineInfo, TArray<FLandscapeSplineInterpPoint>& OutPoints, int32 NumSubdivisions,
+		float StartFalloffFraction, float EndFalloffFraction,
+		const float StartWidth, const float EndWidth,
+		const float StartSideFalloff, const float EndSideFalloff,
+		const float StartRollDegrees, const float EndRollDegrees);
+
+	void RasterizeSegmentPoints(ULandscapeInfo* LandscapeInfo, TArray<FLandscapeSplineInterpPoint> Points, const FTransform& SplineToWorld, bool bRaiseTerrain, bool bLowerTerrain, ULandscapeLayerInfoObject* LayerInfo);
+#endif
+}
diff --git a/Engine/Source/Runtime/Landscape/Private/LandscapeSplines.cpp b/Engine/Source/Runtime/Landscape/Private/LandscapeSplines.cpp
--- a/Engine/Source/Runtime/Landscape/Private/LandscapeSplines.cpp
+++ b/Engine/Source/Runtime/Landscape/Private/LandscapeSplines.cpp
@@ -1160,17 +1160,6 @@ static float ApproxLength(const FInterpCurveVector& SplineInfo, const float Star
 	return SplineLength;
 }
 
-/** Util that takes a 2D vector and rotates it by RotAngle (given in radians) */
-static FVector2D RotateVec2D(const FVector2D InVec, float RotAngle)
-{
-	FVector2D OutVec;
-	OutVec.X = (InVec.X * FMath::Cos(RotAngle)) - (InVec.Y * FMath::Sin(RotAngle));
-	OutVec.Y = (InVec.X * FMath::Sin(RotAngle)) + (InVec.Y * FMath::Cos(RotAngle));
-	return OutVec;
-}
-
-
-
 static ESplineMeshAxis::Type CrossAxis(ESplineMeshAxis::Type InForwardAxis, ESplineMeshAxis::Type InUpAxis)
 {
 	check(InForwardAxis != InUpAxis);
@@ -1201,27 +1190,18 @@ void ULandscapeSplineSegment::UpdateSplinePoints(bool bUpdateCollision)
 	// Set up BSpline
 	FVector StartLocation; FRotator StartRotation;
 	Connections[0].ControlPoint->GetConnectionLocationAndRotation(Connections[0].SocketName, StartLocation, StartRotation);
-	new(SplineInfo.Points) FInterpCurvePoint<FVector>(0.0f, StartLocation, StartRotation.Vector() * Connections[0].TangentLen, StartRotation.Vector() * Connections[0].TangentLen, CIM_CurveUser);
+	SplineInfo.Points.Emplace(0.0f, StartLocation, StartRotation.Vector() * Connections[0].TangentLen, StartRotation.Vector() * Connections[0].TangentLen, CIM_CurveUser);
 	FVector EndLocation; FRotator EndRotation;
 	Connections[1].ControlPoint->GetConnectionLocationAndRotation(Connections[1].SocketName, EndLocation, EndRotation);
-	new(SplineInfo.Points) FInterpCurvePoint<FVector>(1.0f, EndLocation, EndRotation.Vector() * -Connections[1].TangentLen, EndRotation.Vector() * -Connections[1].TangentLen, CIM_CurveUser);
+	SplineInfo.Points.Emplace(1.0f, EndLocation, EndRotation.Vector() * -Connections[1].TangentLen, EndRotation.Vector() * -Connections[1].TangentLen, CIM_CurveUser);
 
 	// Pointify
 
 	// Calculate spline length
 	const float SplineLength = ApproxLength(SplineInfo, 0.0f, 1.0f, 4);
 
-	float StartFalloffFraction = ((Connections[0].ControlPoint->ConnectedSegments.Num() > 1) ? 0 : (Connections[0].ControlPoint->EndFalloff / SplineLength));
-	float EndFalloffFraction = ((Connections[1].ControlPoint->ConnectedSegments.Num() > 1) ? 0 : (Connections[1].ControlPoint->EndFalloff / SplineLength));
-
-	// Stop the start and end fall-off overlapping
-	const float TotalFalloff = StartFalloffFraction + EndFalloffFraction;
-	if (TotalFalloff > 1.0f)
-	{
-		StartFalloffFraction /= TotalFalloff;
-		EndFalloffFraction /= TotalFalloff;
-	}
-
+	const float StartFalloffFraction = ((Connections[0].ControlPoint->ConnectedSegments.Num() > 1) ? 0 : (Connections[0].ControlPoint->EndFalloff / SplineLength));
+	const float EndFalloffFraction = ((Connections[1].ControlPoint->ConnectedSegments.Num() > 1) ? 0 : (Connections[1].ControlPoint->EndFalloff / SplineLength));
 	const float StartWidth = Connections[0].ControlPoint->Width;
 	const float EndWidth = Connections[1].ControlPoint->Width;
 	const float StartSideFalloff = Connections[0].ControlPoint->SideFalloff;
@@ -1234,60 +1214,7 @@ void ULandscapeSplineSegment::UpdateSplinePoints(bool bUpdateCollision)
 	int32 NumPoints = FMath::CeilToInt(SplineLength / OuterSplines->SplineResolution);
 	NumPoints = FMath::Clamp(NumPoints, 1, 1000);
 
-	float OldKeyTime = 0;
-	for (int32 i = 0; i < SplineInfo.Points.Num(); i++)
-	{
-		const float NewKeyTime = SplineInfo.Points[i].InVal;
-		const float NewKeyCosInterp = 0.5f - 0.5f * FMath::Cos(NewKeyTime * PI);
-		const float NewKeyWidth = FMath::Lerp(StartWidth, EndWidth, NewKeyCosInterp);
-		const float NewKeyFalloff = FMath::Lerp(StartSideFalloff, EndSideFalloff, NewKeyCosInterp);
-		const float NewKeyRoll = FMath::Lerp(StartRoll, EndRoll, NewKeyCosInterp);
-		const FVector NewKeyPos = SplineInfo.Eval(NewKeyTime, FVector::ZeroVector);
-		const FVector NewKeyTangent = SplineInfo.EvalDerivative(NewKeyTime, FVector::ZeroVector).GetSafeNormal();
-		const FVector NewKeyBiNormal = FQuat(NewKeyTangent, -NewKeyRoll).RotateVector((NewKeyTangent ^ FVector(0, 0, -1)).GetSafeNormal());
-		const FVector NewKeyLeftPos = NewKeyPos - NewKeyBiNormal * NewKeyWidth;
-		const FVector NewKeyRightPos = NewKeyPos + NewKeyBiNormal * NewKeyWidth;
-		const FVector NewKeyFalloffLeftPos = NewKeyPos - NewKeyBiNormal * (NewKeyWidth + NewKeyFalloff);
-		const FVector NewKeyFalloffRightPos = NewKeyPos + NewKeyBiNormal * (NewKeyWidth + NewKeyFalloff);
-		const float NewKeyStartEndFalloff = FMath::Min((StartFalloffFraction > 0 ? NewKeyTime / StartFalloffFraction : 1.0f), (EndFalloffFraction > 0 ? (1 - NewKeyTime) / EndFalloffFraction : 1.0f));
-
-		// If not the first keypoint, interp from the last keypoint.
-		if (i > 0)
-		{
-			int32 NumSteps = FMath::CeilToInt( (NewKeyTime - OldKeyTime) * NumPoints );
-			float DrawSubstep = (NewKeyTime - OldKeyTime) / NumSteps;
-
-			// Add a point for each substep, except the ends because that's the point added outside the interp'ing.
-			for (int32 j = 1; j < NumSteps; j++)
-			{
-				const float NewTime = OldKeyTime + j*DrawSubstep;
-				const float NewCosInterp = 0.5f - 0.5f * FMath::Cos(NewTime * PI);
-				const float NewWidth = FMath::Lerp(StartWidth, EndWidth, NewCosInterp);
-				const float NewFalloff = FMath::Lerp(StartSideFalloff, EndSideFalloff, NewCosInterp);
-				const float NewRoll = FMath::Lerp(StartRoll, EndRoll, NewCosInterp);
-				const FVector NewPos = SplineInfo.Eval(NewTime, FVector::ZeroVector);
-				const FVector NewTangent = SplineInfo.EvalDerivative(NewTime, FVector::ZeroVector).GetSafeNormal();
-				const FVector NewBiNormal = FQuat(NewTangent, -NewRoll).RotateVector((NewTangent ^ FVector(0, 0, -1)).GetSafeNormal());
-				const FVector NewLeftPos = NewPos - NewBiNormal * NewWidth;
-				const FVector NewRightPos = NewPos + NewBiNormal * NewWidth;
-				const FVector NewFalloffLeftPos = NewPos - NewBiNormal * (NewWidth + NewFalloff);
-				const FVector NewFalloffRightPos = NewPos + NewBiNormal * (NewWidth + NewFalloff);
-				const float NewStartEndFalloff = FMath::Min((StartFalloffFraction > 0 ? NewTime / StartFalloffFraction : 1.0f), (EndFalloffFraction > 0 ? (1 - NewTime) / EndFalloffFraction : 1.0f));
-
-				new(Points) FLandscapeSplineInterpPoint(NewPos, NewLeftPos, NewRightPos, NewFalloffLeftPos, NewFalloffRightPos, NewStartEndFalloff);
-			}
-		}
-
-		new(Points) FLandscapeSplineInterpPoint(NewKeyPos, NewKeyLeftPos, NewKeyRightPos, NewKeyFalloffLeftPos, NewKeyFalloffRightPos, NewKeyStartEndFalloff);
-
-		OldKeyTime = NewKeyTime;
-	}
-
-	// Handle self-intersection errors due to tight turns
-	FixSelfIntersection(&FLandscapeSplineInterpPoint::Left);
-	FixSelfIntersection(&FLandscapeSplineInterpPoint::Right);
-	FixSelfIntersection(&FLandscapeSplineInterpPoint::FalloffLeft);
-	FixSelfIntersection(&FLandscapeSplineInterpPoint::FalloffRight);
+	LandscapeSplineRaster::Pointify(SplineInfo, Points, NumPoints, StartFalloffFraction, EndFalloffFraction, StartWidth, EndWidth, StartSideFalloff, EndSideFalloff, StartRollDegrees, EndRollDegrees);
 
 	// Update Bounds
 	Bounds = FBox(0);
@@ -1481,7 +1408,7 @@ void ULandscapeSplineSegment::UpdateSplinePoints(bool bUpdateCollision)
 				break;
 			}
 			Offset *= Scale2D;
-			Offset = RotateVec2D(Offset, -Roll);
+			Offset = Offset.GetRotated(-Roll);
 
 			MeshComponent->SplineParams.StartPos = SplineInfo.Eval(RescaledT, FVector::ZeroVector);
 			MeshComponent->SplineParams.StartTangent = SplineInfo.EvalDerivative(RescaledT, FVector::ZeroVector) * (TEnd - RescaledT);
@@ -1529,7 +1456,7 @@ void ULandscapeSplineSegment::UpdateSplinePoints(bool bUpdateCollision)
 				break;
 			}
 			OffsetEnd *= Scale2DEnd;
-			OffsetEnd = RotateVec2D(OffsetEnd, -RollEnd);
+			OffsetEnd = OffsetEnd.GetRotated(-RollEnd);
 
 			MeshComponent->SplineParams.EndPos = SplineInfo.Eval(TEnd, FVector::ZeroVector);
 			MeshComponent->SplineParams.EndTangent = SplineInfo.EvalDerivative(TEnd, FVector::ZeroVector) * (TEnd - RescaledT);
@@ -1639,127 +1566,6 @@ void ULandscapeSplineSegment::DeleteSplinePoints()
 	}
 	MeshComponents.Empty();
 }
-
-static bool LineIntersect( const FVector2D& L1Start, const FVector2D& L1End, const FVector2D& L2Start, const FVector2D& L2End, FVector2D& Intersect, float Tolerance = KINDA_SMALL_NUMBER)
-{
-	float tA = (L2End - L2Start) ^ (L2Start - L1Start);
-	float tB = (L1End - L1Start) ^ (L2Start - L1Start);
-	float Denom = (L2End - L2Start) ^ (L1End - L1Start);
-
-	if( FMath::IsNearlyZero( tA ) && FMath::IsNearlyZero( tB )  )
-	{
-		// Lines are the same
-		Intersect = (L2Start + L2End) / 2;
-		return true;
-	}
-
-	if( FMath::IsNearlyZero(Denom) )
-	{
-		// Lines are parallel
-		Intersect = (L2Start + L2End) / 2;
-		return false;
-	}
-
-	tA /= Denom;
-	tB /= Denom;
-
-	Intersect = L1Start + tA * (L1End - L1Start);
-
-	if( tA >= -Tolerance && tA <= (1.0f + Tolerance ) && tB >= -Tolerance && tB <= (1.0f + Tolerance) )
-	{
-		return true;
-	}
-
-	return false;
-}
-
-bool ULandscapeSplineSegment::FixSelfIntersection(FVector FLandscapeSplineInterpPoint::* Side)
-{
-	int32 StartSide = INDEX_NONE;
-	for(int32 i = 0; i < Points.Num(); i++)
-	{
-		bool bReversed = false;
-
-		if (i < Points.Num() - 1)
-		{
-			const FLandscapeSplineInterpPoint& CurrentPoint = Points[i];
-			const FLandscapeSplineInterpPoint& NextPoint = Points[i+1];
-			const FVector Direction = (NextPoint.Center - CurrentPoint.Center).GetSafeNormal();
-			const FVector SideDirection = (NextPoint.*Side - CurrentPoint.*Side).GetSafeNormal();
-			bReversed = (SideDirection | Direction) < 0;
-		}
-
-		if (bReversed)
-		{
-			if (StartSide == INDEX_NONE)
-			{
-				StartSide = i;
-			}
-		}
-		else
-		{
-			if (StartSide != INDEX_NONE)
-			{
-				int32 EndSide = i;
-
-				// step startSide back until before the endSide point
-				while (StartSide > 0)
-				{
-					const float Projection = (Points[StartSide].*Side - Points[StartSide-1].*Side) | (Points[EndSide].*Side - Points[StartSide-1].*Side);
-					if (Projection >= 0)
-					{
-						break;
-					}
-					StartSide--;
-				}
-				// step endSide forwards until after the startSide point
-				while (EndSide < Points.Num() - 1)
-				{
-					const float Projection = (Points[EndSide].*Side - Points[EndSide+1].*Side) | (Points[StartSide].*Side - Points[EndSide+1].*Side);
-					if (Projection >= 0)
-					{
-						break;
-					}
-					EndSide++;
-				}
-
-				// Can't do anything if the start and end intersect, as they're both unalterable
-				if (StartSide == 0 && EndSide == Points.Num() - 1)
-				{
-					return false;
-				}
-
-				FVector2D Collapse;
-				if (StartSide == 0)
-				{
-					Collapse = FVector2D(Points[StartSide].*Side);
-					StartSide++;
-				}
-				else if (EndSide == Points.Num() - 1)
-				{
-					Collapse = FVector2D(Points[EndSide].*Side);
-					EndSide--;
-				}
-				else
-				{
-					LineIntersect(FVector2D(Points[StartSide-1].*Side), FVector2D(Points[StartSide].*Side),
-						FVector2D(Points[EndSide+1].*Side), FVector2D(Points[EndSide].*Side), Collapse);
-				}
-
-				for (int32 j = StartSide; j <= EndSide; j++)
-				{
-					(Points[j].*Side).X = Collapse.X;
-					(Points[j].*Side).Y = Collapse.Y;
-				}
-
-				StartSide = INDEX_NONE;
-				i = EndSide;
-			}
-		}
-	}
-
-	return true;
-}
 #endif
 
 void ULandscapeSplineSegment::FindNearest( const FVector& InLocation, float& t, FVector& OutLocation, FVector& OutTangent )
