Fix non-editor build

diff --git a/Engine/Source/Runtime/Landscape/Private/LandscapeSplineRaster.cpp b/Engine/Source/Runtime/Landscape/Private/LandscapeSplineRaster.cpp
--- a/Engine/Source/Runtime/Landscape/Private/LandscapeSplineRaster.cpp
+++ b/Engine/Source/Runtime/Landscape/Private/LandscapeSplineRaster.cpp
@@ -13,8 +13,10 @@
 #include "LandscapeSplinesComponent.h"
 #include "LandscapeSplineControlPoint.h"
 #include "LandscapeSplineSegment.h"
+#if WITH_EDITOR
 #include "ScopedTransaction.h"
 #include "Raster.h"
+#endif
 
 #define LOCTEXT_NAMESPACE "Landscape"
 
@@ -22,6 +24,7 @@
 // Apply splines
 //////////////////////////////////////////////////////////////////////////
 
+#if WITH_EDITOR
 class FLandscapeSplineHeightsRasterPolicy
 {
 public:
@@ -110,23 +113,6 @@ private:
 	int32 MinX, MinY, MaxX, MaxY;
 };
 
-bool ULandscapeInfo::ApplySplines(bool bOnlySelected)
-{
-	bool bResult = false;
-
-	ALandscape* Landscape = LandscapeActor.Get();
-
-	bResult |= ApplySplinesInternal(bOnlySelected, Landscape);
-
-	for (auto It = Proxies.CreateIterator(); It; ++It)
-	{
-		ALandscapeProxy* LandscapeProxy = (*It);
-		bResult |= ApplySplinesInternal(bOnlySelected, LandscapeProxy);
-	}
-
-	return bResult;
-}
-
 void RasterizeControlPointHeights(int32& MinX, int32& MinY, int32& MaxX, int32& MaxY, FLandscapeEditDataInterface& LandscapeEdit, FVector ControlPointLocation, const TArray<FLandscapeSplineInterpPoint>& Points, bool bRaiseTerrain, bool bLowerTerrain, TSet<ULandscapeComponent*>& ModifiedComponents)
 {
 	if (!(bRaiseTerrain || bLowerTerrain))
@@ -428,6 +414,23 @@ void RasterizeSegmentAlpha(int32& MinX, int32& MinY, int32& MaxX, int32& MaxY, F
 	LandscapeEdit.GetComponentsInRegion(MinX, MinY, MaxX, MaxY, &ModifiedComponents);
 }
 
+bool ULandscapeInfo::ApplySplines(bool bOnlySelected)
+{
+	bool bResult = false;
+
+	ALandscape* Landscape = LandscapeActor.Get();
+
+	bResult |= ApplySplinesInternal(bOnlySelected, Landscape);
+
+	for (auto It = Proxies.CreateIterator(); It; ++It)
+	{
+		ALandscapeProxy* LandscapeProxy = (*It);
+		bResult |= ApplySplinesInternal(bOnlySelected, LandscapeProxy);
+	}
+
+	return bResult;
+}
+
 bool ULandscapeInfo::ApplySplinesInternal(bool bOnlySelected, ALandscapeProxy* Landscape)
 {
 	if (!Landscape || !Landscape->SplineComponent || Landscape->SplineComponent->ControlPoints.Num() == 0 || Landscape->SplineComponent->Segments.Num() == 0)
@@ -883,5 +886,6 @@ namespace LandscapeSplineRaster
 		FixSelfIntersection(Points, &FLandscapeSplineInterpPoint::FalloffRight);
 	}
 }
+#endif
 
 #undef LOCTEXT_NAMESPACE
