Fix non-unity build

diff --git a/Engine/Source/Runtime/Landscape/Private/LandscapeSplines.cpp b/Engine/Source/Runtime/Landscape/Private/LandscapeSplines.cpp
--- a/Engine/Source/Runtime/Landscape/Private/LandscapeSplines.cpp
+++ b/Engine/Source/Runtime/Landscape/Private/LandscapeSplines.cpp
@@ -18,6 +18,9 @@
 #include "EngineGlobals.h"
 #include "Engine/StaticMesh.h"
 #include "Engine/StaticMeshSocket.h"
+#if WITH_EDITOR
+#include "LandscapeSplineRaster.h"
+#endif
 
 IMPLEMENT_HIT_PROXY(HLandscapeSplineProxy, HHitProxy);
 IMPLEMENT_HIT_PROXY(HLandscapeSplineProxy_Segment, HLandscapeSplineProxy);
