Fix non-unity build (again, sorry)

diff --git a/Engine/Source/Runtime/Landscape/Private/LandscapeSplineRaster.cpp b/Engine/Source/Runtime/Landscape/Private/LandscapeSplineRaster.cpp
--- a/Engine/Source/Runtime/Landscape/Private/LandscapeSplineRaster.cpp
+++ b/Engine/Source/Runtime/Landscape/Private/LandscapeSplineRaster.cpp
@@ -8,6 +8,7 @@
 #include "Landscape.h"
 #include "LandscapeSplineRaster.h"
 #include "LandscapeProxy.h"
+#include "LandscapeDataAccess.h"
 #include "LandscapeEdit.h"
 #include "LandscapeHeightfieldCollisionComponent.h"
 #include "LandscapeSplinesComponent.h"
