Fix material compiler problem with LandscapeVisibilityMask connected to anything other than Masked.

diff --git a/Engine/Source/Runtime/Landscape/Private/Materials/MaterialExpressionLandscapeVisibilityMask.cpp b/Engine/Source/Runtime/Landscape/Private/Materials/MaterialExpressionLandscapeVisibilityMask.cpp
--- a/Engine/Source/Runtime/Landscape/Private/Materials/MaterialExpressionLandscapeVisibilityMask.cpp
+++ b/Engine/Source/Runtime/Landscape/Private/Materials/MaterialExpressionLandscapeVisibilityMask.cpp
@@ -33,25 +33,22 @@ UMaterialExpressionLandscapeVisibilityMask::UMaterialExpressionLandscapeVisibili
 	MenuCategories.Add(ConstructorStatics.NAME_Landscape);
 }
 
-
 FGuid& UMaterialExpressionLandscapeVisibilityMask::GetParameterExpressionId()
 {
 	return ExpressionGUID;
 }
 
-
 int32 UMaterialExpressionLandscapeVisibilityMask::Compile(class FMaterialCompiler* Compiler, int32 OutputIndex, int32 MultiplexIndex)
 {
-	return Compiler->Sub(Compiler->Constant(1.f), Compiler->StaticTerrainLayerWeight(ParameterName, Compiler->Constant(0.f)));
+	int32 MaskLayerCode = Compiler->StaticTerrainLayerWeight(ParameterName, Compiler->Constant(0.f));
+	return MaskLayerCode == INDEX_NONE ? Compiler->Constant(1.f) : Compiler->Sub(Compiler->Constant(1.f), MaskLayerCode);
 }
 
-
 UTexture* UMaterialExpressionLandscapeVisibilityMask::GetReferencedTexture()
 {
 	return GEngine->WeightMapPlaceholderTexture;
 }
 
-
 void UMaterialExpressionLandscapeVisibilityMask::GetAllParameterNames(TArray<FName> &OutParameterNames, TArray<FGuid> &OutParameterIds)
 {
 	int32 CurrentSize = OutParameterNames.Num();
@@ -63,11 +60,9 @@ void UMaterialExpressionLandscapeVisibilityMask::GetAllParameterNames(TArray<FNa
 	}
 }
 
-
 void UMaterialExpressionLandscapeVisibilityMask::GetCaption(TArray<FString>& OutCaptions) const
 {
 	OutCaptions.Add(FString(TEXT("Landscape Visibility Mask")));
 }
 
-
 #undef LOCTEXT_NAMESPACE
