Trim down the amount of UObject constructors.

diff --git a/Engine/Source/Runtime/Foliage/Private/InstancedFoliage.cpp b/Engine/Source/Runtime/Foliage/Private/InstancedFoliage.cpp
--- a/Engine/Source/Runtime/Foliage/Private/InstancedFoliage.cpp
+++ b/Engine/Source/Runtime/Foliage/Private/InstancedFoliage.cpp
@@ -513,7 +513,7 @@ void FFoliageMeshInfo::AddInstance(AInstancedFoliageActor* InIFA, const UFoliage
 
 	if (Component == nullptr)
 	{
-		Component = ConstructObject<UHierarchicalInstancedStaticMeshComponent>(UHierarchicalInstancedStaticMeshComponent::StaticClass(), InIFA, NAME_None, RF_Transactional);
+		Component = NewObject<UHierarchicalInstancedStaticMeshComponent>(InIFA, NAME_None, RF_Transactional);
 
 		Component->StaticMesh = InSettings->GetStaticMesh();
 		Component->bSelectable = true;
@@ -1330,7 +1330,7 @@ FFoliageMeshInfo* AInstancedFoliageActor::AddMesh(UStaticMesh* InMesh, UFoliageT
 	else
 #endif
 	{
-		Settings = ConstructObject<UFoliageType_InstancedStaticMesh>(UFoliageType_InstancedStaticMesh::StaticClass(), this);
+		Settings = NewObject<UFoliageType_InstancedStaticMesh>(this);
 	}
 	Settings->Mesh = InMesh;
 
@@ -1684,7 +1684,7 @@ void AInstancedFoliageActor::Serialize(FArchive& Ar)
 			if (FoliageType == nullptr)
 			{
 				// If the Settings object was null, eg the user forgot to save their settings asset, create a new one.
-				FoliageType = ConstructObject<UFoliageType_InstancedStaticMesh>(UFoliageType_InstancedStaticMesh::StaticClass(), this);
+				FoliageType = NewObject<UFoliageType_InstancedStaticMesh>(this);
 			}
 
 			if (FoliageType->Mesh == nullptr)
diff --git a/Engine/Source/Runtime/GameplayAbilities/Private/AbilitySystemComponent.cpp b/Engine/Source/Runtime/GameplayAbilities/Private/AbilitySystemComponent.cpp
--- a/Engine/Source/Runtime/GameplayAbilities/Private/AbilitySystemComponent.cpp
+++ b/Engine/Source/Runtime/GameplayAbilities/Private/AbilitySystemComponent.cpp
@@ -83,7 +83,7 @@ const UAttributeSet* UAbilitySystemComponent::GetOrCreateAttributeSubobject(TSub
 		MyAttributes = GetAttributeSubobject(AttributeClass);
 		if (!MyAttributes)
 		{
-			UAttributeSet *Attributes = ConstructObject<UAttributeSet>(AttributeClass, OwningActor);
+			UAttributeSet *Attributes = NewObject<UAttributeSet>(OwningActor, AttributeClass);
 			SpawnedAttributes.AddUnique(Attributes);
 			MyAttributes = Attributes;
 		}
diff --git a/Engine/Source/Runtime/GameplayAbilities/Private/AbilitySystemComponent_Abilities.cpp b/Engine/Source/Runtime/GameplayAbilities/Private/AbilitySystemComponent_Abilities.cpp
--- a/Engine/Source/Runtime/GameplayAbilities/Private/AbilitySystemComponent_Abilities.cpp
+++ b/Engine/Source/Runtime/GameplayAbilities/Private/AbilitySystemComponent_Abilities.cpp
@@ -469,7 +469,7 @@ UGameplayAbility* UAbilitySystemComponent::CreateNewInstanceOfAbility(FGameplayA
 	AActor* OwnerActor = GetOwner();
 	check(OwnerActor);
 
-	UGameplayAbility * AbilityInstance = ConstructObject<UGameplayAbility>(Ability->GetClass(), OwnerActor);
+	UGameplayAbility * AbilityInstance = NewObject<UGameplayAbility>(OwnerActor, Ability->GetClass());
 	check(AbilityInstance);
 
 	// Add it to one of our instance lists so that it doesn't GC.
diff --git a/Engine/Source/Runtime/GameplayAbilities/Private/GameplayAbilitiesModule.cpp b/Engine/Source/Runtime/GameplayAbilities/Private/GameplayAbilitiesModule.cpp
--- a/Engine/Source/Runtime/GameplayAbilities/Private/GameplayAbilitiesModule.cpp
+++ b/Engine/Source/Runtime/GameplayAbilities/Private/GameplayAbilitiesModule.cpp
@@ -22,7 +22,7 @@ class FGameplayAbilitiesModule : public IGameplayAbilitiesModule
 
 			checkf(SingletonClass != NULL, TEXT("Ability config value AbilitySystemGlobalsClassName is not a valid class name."));
 
-			AbilitySystemGlobals = ConstructObject<UAbilitySystemGlobals>(SingletonClass, GetTransientPackage(), NAME_None, RF_RootSet);
+			AbilitySystemGlobals = NewObject<UAbilitySystemGlobals>(GetTransientPackage(), SingletonClass, NAME_None, RF_RootSet);
 		}
 
 		check(AbilitySystemGlobals);
diff --git a/Engine/Source/Runtime/GameplayTags/Private/GameplayTagsModule.cpp b/Engine/Source/Runtime/GameplayTags/Private/GameplayTagsModule.cpp
--- a/Engine/Source/Runtime/GameplayTags/Private/GameplayTagsModule.cpp
+++ b/Engine/Source/Runtime/GameplayTags/Private/GameplayTagsModule.cpp
@@ -28,7 +28,7 @@ DEFINE_LOG_CATEGORY(LogGameplayTags);
 void FGameplayTagsModule::StartupModule()
 {
 	// This code will execute after your module is loaded into memory (but after global variables are initialized, of course.)
-	GameplayTagsManager = ConstructObject<UGameplayTagsManager>(UGameplayTagsManager::StaticClass(),GetTransientPackage(),NAME_None,RF_RootSet);
+	GameplayTagsManager = NewObject<UGameplayTagsManager>(GetTransientPackage(), NAME_None, RF_RootSet);
 
 	TArray<FString> GameplayTagTables;
 	GConfig->GetArray(TEXT("GameplayTags"), TEXT("GameplayTagTableList"), GameplayTagTables, GEngineIni);
diff --git a/Engine/Source/Runtime/Landscape/Private/Landscape.cpp b/Engine/Source/Runtime/Landscape/Private/Landscape.cpp
--- a/Engine/Source/Runtime/Landscape/Private/Landscape.cpp
+++ b/Engine/Source/Runtime/Landscape/Private/Landscape.cpp
@@ -582,7 +582,7 @@ ULandscapeInfo* ALandscapeProxy::GetLandscapeInfo(bool bSpawnNewActor /*= true*/
 			ULandscapeInfo* LandscapeInfo = LandscapeInfoMap.Map.FindRef(LandscapeGuid);
 			if (!LandscapeInfo && bSpawnNewActor && !HasAnyFlags(RF_BeginDestroyed))
 			{
-				LandscapeInfo = ConstructObject<ULandscapeInfo>(ULandscapeInfo::StaticClass(), OwningWorld, NAME_None, RF_Transactional | RF_Transient);
+				LandscapeInfo = NewObject<ULandscapeInfo>(OwningWorld, NAME_None, RF_Transactional | RF_Transient);
 				LandscapeInfoMap.Modify(false);
 				LandscapeInfoMap.Map.Add(LandscapeGuid, LandscapeInfo);
 			}
@@ -2876,7 +2876,7 @@ void ALandscapeProxy::UpdateFoliage(const TArray<FVector>& Cameras, bool bForceS
 										UHierarchicalInstancedStaticMeshComponent* HierarchicalInstancedStaticMeshComponent;
 										{
 											QUICK_SCOPE_CYCLE_COUNTER(STAT_GrassCreateComp);
-											HierarchicalInstancedStaticMeshComponent = ConstructObject<UHierarchicalInstancedStaticMeshComponent>(UHierarchicalInstancedStaticMeshComponent::StaticClass(), this); //, NAME_None, RF_Transactional);
+											HierarchicalInstancedStaticMeshComponent = NewObject<UHierarchicalInstancedStaticMeshComponent>(this);
 										}
 										NewComp.Foliage = HierarchicalInstancedStaticMeshComponent;
 										FoliageCache.CachedGrassComps.Add(NewComp);
diff --git a/Engine/Source/Runtime/Landscape/Private/LandscapeEdit.cpp b/Engine/Source/Runtime/Landscape/Private/LandscapeEdit.cpp
--- a/Engine/Source/Runtime/Landscape/Private/LandscapeEdit.cpp
+++ b/Engine/Source/Runtime/Landscape/Private/LandscapeEdit.cpp
@@ -117,7 +117,7 @@ ULandscapeMaterialInstanceConstant* ALandscapeProxy::GetLayerThumbnailMIC(UMater
 		{
 			MICOuter = Proxy->GetOutermost();
 		}
-		CombinationMaterialInstance = ConstructObject<ULandscapeMaterialInstanceConstant>(ULandscapeMaterialInstanceConstant::StaticClass(), MICOuter);
+		CombinationMaterialInstance = NewObject<ULandscapeMaterialInstanceConstant>(MICOuter);
 		if (Proxy)
 		{
 			UE_LOG(LogLandscape, Log, TEXT("Looking for key %s, making new combination %s"), *LayerKey, *CombinationMaterialInstance->GetName());
@@ -148,7 +148,7 @@ ULandscapeMaterialInstanceConstant* ALandscapeProxy::GetLayerThumbnailMIC(UMater
 	}
 
 	// Create the instance for this component, that will use the layer combination instance.
-	ULandscapeMaterialInstanceConstant* MaterialInstance = ConstructObject<ULandscapeMaterialInstanceConstant>(ULandscapeMaterialInstanceConstant::StaticClass(), GetTransientPackage());
+	ULandscapeMaterialInstanceConstant* MaterialInstance = NewObject<ULandscapeMaterialInstanceConstant>();
 	MaterialInstance->SetParentEditorOnly(CombinationMaterialInstance);
 	MaterialInstance->bIsLayerThumbnail = true;
 
@@ -189,7 +189,7 @@ UMaterialInstanceConstant* ULandscapeComponent::GetCombinationMaterial(bool bMob
 		{
 			FlushRenderingCommands();
 
-			CombinationMaterialInstance = ConstructObject<ULandscapeMaterialInstanceConstant>(ULandscapeMaterialInstanceConstant::StaticClass(), GetOutermost());
+			CombinationMaterialInstance = NewObject<ULandscapeMaterialInstanceConstant>(GetOutermost());
 			UE_LOG(LogLandscape, Log, TEXT("Looking for key %s, making new combination %s"), *LayerKey, *CombinationMaterialInstance->GetName());
 			Proxy->MaterialInstanceConstantMap.Add(*LayerKey, CombinationMaterialInstance);
 			CombinationMaterialInstance->SetParentEditorOnly(LandscapeMaterial);
@@ -244,7 +244,7 @@ void ULandscapeComponent::UpdateMaterialInstances()
 		// Create the instance for this component, that will use the layer combination instance.
 		if (MaterialInstance == NULL || GetOutermost() != MaterialInstance->GetOutermost())
 		{
-			MaterialInstance = ConstructObject<ULandscapeMaterialInstanceConstant>(ULandscapeMaterialInstanceConstant::StaticClass(), GetOutermost());
+			MaterialInstance = NewObject<ULandscapeMaterialInstanceConstant>(GetOutermost());
 		}
 
 		// For undo
@@ -572,9 +572,9 @@ void ULandscapeComponent::UpdateCollisionHeightData(const FColor* HeightmapTextu
 			CollisionComp = NULL;
 		}
 
-		MeshCollisionComponent = XYOffsetmapTexture ? ConstructObject<ULandscapeMeshCollisionComponent>(ULandscapeMeshCollisionComponent::StaticClass(), Proxy, NAME_None, RF_Transactional) : NULL;
+		MeshCollisionComponent = XYOffsetmapTexture ? NewObject<ULandscapeMeshCollisionComponent>(Proxy, NAME_None, RF_Transactional) : NULL;
 		CollisionComp = XYOffsetmapTexture ? Cast<ULandscapeHeightfieldCollisionComponent>(MeshCollisionComponent) :
-			ConstructObject<ULandscapeHeightfieldCollisionComponent>(ULandscapeHeightfieldCollisionComponent::StaticClass(), Proxy, NAME_None, RF_Transactional);
+			NewObject<ULandscapeHeightfieldCollisionComponent>(Proxy, NAME_None, RF_Transactional);
 
 		CollisionComp->SetRelativeLocation(RelativeLocation);
 		CollisionComp->AttachTo(Proxy->GetRootComponent(), NAME_None);
@@ -1660,7 +1660,7 @@ ULandscapeLayerInfoObject* ALandscapeProxy::CreateLayerInfo(const TCHAR* LayerNa
 		Suffix++;
 	}
 	UPackage* Package = CreatePackage(NULL, *PackageName);
-	ULandscapeLayerInfoObject* LayerInfo = ConstructObject<ULandscapeLayerInfoObject>(ULandscapeLayerInfoObject::StaticClass(), Package, LayerObjectName, RF_Public | RF_Standalone | RF_Transactional);
+	ULandscapeLayerInfoObject* LayerInfo = NewObject<ULandscapeLayerInfoObject>(Package, LayerObjectName, RF_Public | RF_Standalone | RF_Transactional);
 	LayerInfo->LayerName = LayerName;
 
 	return LayerInfo;
@@ -1726,7 +1726,7 @@ void ALandscapeProxy::Import(FGuid Guid, int32 VertsX, int32 VertsY,
 			int32 BaseX = X * ComponentSizeQuads;
 			int32 BaseY = Y * ComponentSizeQuads;
 
-			ULandscapeComponent* LandscapeComponent = ConstructObject<ULandscapeComponent>(ULandscapeComponent::StaticClass(), this, NAME_None, RF_Transactional);
+			ULandscapeComponent* LandscapeComponent = NewObject<ULandscapeComponent>(this, NAME_None, RF_Transactional);
 			LandscapeComponent->SetRelativeLocation(FVector(BaseX, BaseY, 0));
 			LandscapeComponent->AttachTo(GetRootComponent(), NAME_None);
 			LandscapeComponents.Add(LandscapeComponent);
@@ -4510,7 +4510,7 @@ UMaterialInstance* ULandscapeComponent::GeneratePlatformPixelData(TArray<UTextur
 	else // for cooking
 	{
 		UMaterialInstanceConstant* CombinationMaterialInstance = GetCombinationMaterial(true);
-		UMaterialInstanceConstant* MobileMaterialInstance = ConstructObject<ULandscapeMaterialInstanceConstant>(ULandscapeMaterialInstanceConstant::StaticClass(), GetOutermost());
+		UMaterialInstanceConstant* MobileMaterialInstance = NewObject<ULandscapeMaterialInstanceConstant>(GetOutermost());
 
 		MobileMaterialInstance->SetParentEditorOnly(CombinationMaterialInstance);
 
@@ -4704,7 +4704,7 @@ void ULandscapeComponent::GeneratePlatformVertexData()
 UTexture2D* ALandscapeProxy::CreateLandscapeTexture(int32 InSizeX, int32 InSizeY, TextureGroup InLODGroup, ETextureSourceFormat InFormat, UObject* OptionalOverrideOuter) const
 {
 	UObject* TexOuter = OptionalOverrideOuter ? OptionalOverrideOuter : GetOutermost();
-	UTexture2D* NewTexture = ConstructObject<UTexture2D>(UTexture2D::StaticClass(), TexOuter);
+	UTexture2D* NewTexture = NewObject<UTexture2D>(TexOuter);
 	NewTexture->Source.Init2DWithMipChain(InSizeX, InSizeY, InFormat);
 	NewTexture->SRGB = false;
 	NewTexture->CompressionNone = true;
diff --git a/Engine/Source/Runtime/Landscape/Private/LandscapeSplines.cpp b/Engine/Source/Runtime/Landscape/Private/LandscapeSplines.cpp
--- a/Engine/Source/Runtime/Landscape/Private/LandscapeSplines.cpp
+++ b/Engine/Source/Runtime/Landscape/Private/LandscapeSplines.cpp
@@ -789,7 +789,7 @@ void ULandscapeSplineControlPoint::UpdateSplinePoints(bool bUpdateCollision, boo
 			{
 				Modify();
 				OuterSplines->GetOuter()->Modify();
-				MeshComponent = ConstructObject<UControlPointMeshComponent>(UControlPointMeshComponent::StaticClass(), OuterSplines->GetOuter(), NAME_None, RF_Transactional | RF_TextExportTransient);
+				MeshComponent = NewObject<UControlPointMeshComponent>(OuterSplines->GetOuter(), NAME_None, RF_Transactional | RF_TextExportTransient);
 				MeshComponent->bSelected = bSelected;
 				MeshComponent->AttachTo(OuterSplines);
 			}
@@ -1318,7 +1318,7 @@ void ULandscapeSplineSegment::UpdateSplinePoints(bool bUpdateCollision)
 			else
 			{
 				OuterSplines->GetOuter()->Modify();
-				MeshComponent = ConstructObject<USplineMeshComponent>(USplineMeshComponent::StaticClass(), OuterSplines->GetOuter(), NAME_None, RF_Transactional|RF_TextExportTransient);
+				MeshComponent = NewObject<USplineMeshComponent>(OuterSplines->GetOuter(), NAME_None, RF_Transactional | RF_TextExportTransient);
 				MeshComponent->bSelected = bSelected;
 				MeshComponent->AttachTo(OuterSplines);
 				MeshComponents.Add(MeshComponent);
diff --git a/Engine/Source/Runtime/Launch/Private/LaunchEngineLoop.cpp b/Engine/Source/Runtime/Launch/Private/LaunchEngineLoop.cpp
--- a/Engine/Source/Runtime/Launch/Private/LaunchEngineLoop.cpp
+++ b/Engine/Source/Runtime/Launch/Private/LaunchEngineLoop.cpp
@@ -1499,7 +1499,7 @@ int32 FEngineLoop::PreInit( const TCHAR* CmdLine )
 				{
 					UClass* EditorEngineClass = StaticLoadClass( UEditorEngine::StaticClass(), NULL, TEXT("engine-ini:/Script/Engine.Engine.EditorEngine"), NULL, LOAD_None, NULL );
 
-					GEngine = GEditor = ConstructObject<UEditorEngine>( EditorEngineClass );
+					GEngine = GEditor = NewObject<UEditorEngine>(GetTransientPackage(), EditorEngineClass);
 
 					GEngine->ParseCommandline();
 
@@ -1513,7 +1513,7 @@ int32 FEngineLoop::PreInit( const TCHAR* CmdLine )
 					UClass* EngineClass = StaticLoadClass( UEngine::StaticClass(), NULL, TEXT("engine-ini:/Script/Engine.Engine.GameEngine"), NULL, LOAD_None, NULL );
 
 					// must do this here so that the engine object that we create on the next line receives the correct property values
-					GEngine = ConstructObject<UEngine>( EngineClass );
+					GEngine = NewObject<UEngine>(GetTransientPackage(), EngineClass);
 					check(GEngine);
 
 					GEngine->ParseCommandline();
@@ -1529,7 +1529,7 @@ int32 FEngineLoop::PreInit( const TCHAR* CmdLine )
 			FAutomationTestFramework::GetInstance().RunSmokeTests();
 #endif 
 	
-			UCommandlet* Commandlet = ConstructObject<UCommandlet>( CommandletClass );
+			UCommandlet* Commandlet = NewObject<UCommandlet>(GetTransientPackage(), CommandletClass);
 			check(Commandlet);
 			Commandlet->AddToRoot();			
 
@@ -1945,14 +1945,14 @@ int32 FEngineLoop::Init()
 	{
 		// We're the game.
 		EngineClass = StaticLoadClass( UGameEngine::StaticClass(), NULL, TEXT("engine-ini:/Script/Engine.Engine.GameEngine"), NULL, LOAD_None, NULL );
-		GEngine = ConstructObject<UEngine>( EngineClass );
+		GEngine = NewObject<UEngine>(GetTransientPackage(), EngineClass);
 	}
 	else
 	{
 #if WITH_EDITOR
 		// We're UnrealEd.
 		EngineClass = StaticLoadClass( UUnrealEdEngine::StaticClass(), NULL, TEXT("engine-ini:/Script/Engine.Engine.UnrealEdEngine"), NULL, LOAD_None, NULL );
-		GEngine = GEditor = GUnrealEd = ConstructObject<UUnrealEdEngine>( EngineClass );
+		GEngine = GEditor = GUnrealEd = NewObject<UUnrealEdEngine>(GetTransientPackage(), EngineClass);
 #else
 		check(0);
 #endif
diff --git a/Engine/Source/Runtime/MovieSceneCore/Private/MovieScene.cpp b/Engine/Source/Runtime/MovieSceneCore/Private/MovieScene.cpp
--- a/Engine/Source/Runtime/MovieSceneCore/Private/MovieScene.cpp
+++ b/Engine/Source/Runtime/MovieSceneCore/Private/MovieScene.cpp
@@ -312,7 +312,7 @@ class UMovieSceneTrack* UMovieScene::AddTrack( TSubclassOf<UMovieSceneTrack> Tra
 		{
 			Modify();
 
-			CreatedType = ConstructObject<UMovieSceneTrack>( TrackClass, this, NAME_None, RF_Transactional );
+			CreatedType = NewObject<UMovieSceneTrack>(this, TrackClass, NAME_None, RF_Transactional);
 			
 			ObjectBinding.AddTrack( *CreatedType );
 		}
@@ -360,7 +360,7 @@ class UMovieSceneTrack* UMovieScene::AddMasterTrack( TSubclassOf<UMovieSceneTrac
 {
 	Modify();
 
-	UMovieSceneTrack* CreatedType = ConstructObject<UMovieSceneTrack>( TrackClass, this, NAME_None, RF_Transactional );
+	UMovieSceneTrack* CreatedType = NewObject<UMovieSceneTrack>(this, TrackClass, NAME_None, RF_Transactional);
 
 	MasterTracks.Add( CreatedType );
 	
diff --git a/Engine/Source/Runtime/MovieSceneCore/Private/RuntimeMovieScenePlayer.cpp b/Engine/Source/Runtime/MovieSceneCore/Private/RuntimeMovieScenePlayer.cpp
--- a/Engine/Source/Runtime/MovieSceneCore/Private/RuntimeMovieScenePlayer.cpp
+++ b/Engine/Source/Runtime/MovieSceneCore/Private/RuntimeMovieScenePlayer.cpp
@@ -6,8 +6,7 @@
 
 URuntimeMovieScenePlayer* URuntimeMovieScenePlayer::CreateRuntimeMovieScenePlayer( ULevel* Level, UMovieSceneBindings* MovieSceneBindings )
 {	
-	URuntimeMovieScenePlayer* NewRuntimeMovieScenePlayer = ConstructObject< URuntimeMovieScenePlayer >( 
-		URuntimeMovieScenePlayer::StaticClass(), 
+	URuntimeMovieScenePlayer* NewRuntimeMovieScenePlayer = NewObject<URuntimeMovieScenePlayer>(
 		(UObject*)GetTransientPackage(), 
 		NAME_None, 
 		RF_Transient );
diff --git a/Engine/Source/Runtime/MovieSceneCoreTypes/Private/Instances/MovieSceneAudioTrackInstance.cpp b/Engine/Source/Runtime/MovieSceneCoreTypes/Private/Instances/MovieSceneAudioTrackInstance.cpp
--- a/Engine/Source/Runtime/MovieSceneCoreTypes/Private/Instances/MovieSceneAudioTrackInstance.cpp
+++ b/Engine/Source/Runtime/MovieSceneCoreTypes/Private/Instances/MovieSceneAudioTrackInstance.cpp
@@ -165,7 +165,7 @@ UAudioComponent* FMovieSceneAudioTrackInstance::GetAudioComponent(AActor* Actor,
 	UAudioComponent*& AudioComponent = PlaybackAudioComponents[RowIndex].FindOrAdd(Actor);
 	if (AudioComponent == NULL)
 	{
-		USoundCue* TempPlaybackAudioCue = ConstructObject<USoundCue>( USoundCue::StaticClass() );
+		USoundCue* TempPlaybackAudioCue = NewObject<USoundCue>();
 		
 		AudioComponent = FAudioDevice::CreateComponent(TempPlaybackAudioCue, NULL, Actor, false, false);
 	}
diff --git a/Engine/Source/Runtime/MovieSceneCoreTypes/Private/MovieScene3DTransformTrack.cpp b/Engine/Source/Runtime/MovieSceneCoreTypes/Private/MovieScene3DTransformTrack.cpp
--- a/Engine/Source/Runtime/MovieSceneCoreTypes/Private/MovieScene3DTransformTrack.cpp
+++ b/Engine/Source/Runtime/MovieSceneCoreTypes/Private/MovieScene3DTransformTrack.cpp
@@ -64,7 +64,7 @@ UMovieScene3DTransformTrack::UMovieScene3DTransformTrack( const FObjectInitializ
 
 UMovieSceneSection* UMovieScene3DTransformTrack::CreateNewSection()
 {
-	return ConstructObject<UMovieSceneSection>( UMovieScene3DTransformSection::StaticClass(), this, NAME_None, RF_Transactional );
+	return NewObject<UMovieSceneSection>(this, UMovieScene3DTransformSection::StaticClass(), NAME_None, RF_Transactional);
 }
 
 TSharedPtr<IMovieSceneTrackInstance> UMovieScene3DTransformTrack::CreateInstance()
diff --git a/Engine/Source/Runtime/MovieSceneCoreTypes/Private/MovieSceneAnimationTrack.cpp b/Engine/Source/Runtime/MovieSceneCoreTypes/Private/MovieSceneAnimationTrack.cpp
--- a/Engine/Source/Runtime/MovieSceneCoreTypes/Private/MovieSceneAnimationTrack.cpp
+++ b/Engine/Source/Runtime/MovieSceneCoreTypes/Private/MovieSceneAnimationTrack.cpp
@@ -66,7 +66,7 @@ TRange<float> UMovieSceneAnimationTrack::GetSectionBoundaries() const
 void UMovieSceneAnimationTrack::AddNewAnimation(float KeyTime, UAnimSequence* AnimSequence)
 {
 	// add the section
-	UMovieSceneAnimationSection* NewSection = ConstructObject<UMovieSceneAnimationSection>( UMovieSceneAnimationSection::StaticClass(), this );
+	UMovieSceneAnimationSection* NewSection = NewObject<UMovieSceneAnimationSection>(this);
 	NewSection->InitialPlacement( AnimationSections, KeyTime, KeyTime + AnimSequence->SequenceLength, SupportsMultipleRows() );
 	NewSection->SetAnimationStartTime( KeyTime );
 	NewSection->SetAnimSequence(AnimSequence);
diff --git a/Engine/Source/Runtime/MovieSceneCoreTypes/Private/MovieSceneAudioTrack.cpp b/Engine/Source/Runtime/MovieSceneCoreTypes/Private/MovieSceneAudioTrack.cpp
--- a/Engine/Source/Runtime/MovieSceneCoreTypes/Private/MovieSceneAudioTrack.cpp
+++ b/Engine/Source/Runtime/MovieSceneCoreTypes/Private/MovieSceneAudioTrack.cpp
@@ -81,7 +81,7 @@ void UMovieSceneAudioTrack::AddNewSound(USoundBase* Sound, float Time)
 	}
 
 	// add the section
-	UMovieSceneAudioSection* NewSection = ConstructObject<UMovieSceneAudioSection>( UMovieSceneAudioSection::StaticClass(), this );
+	UMovieSceneAudioSection* NewSection = NewObject<UMovieSceneAudioSection>(this);
 	NewSection->InitialPlacement( AudioSections, Time, Time + DurationToUse, SupportsMultipleRows() );
 	NewSection->SetAudioStartTime(Time);
 	NewSection->SetSound(Sound);
diff --git a/Engine/Source/Runtime/MovieSceneCoreTypes/Private/MovieSceneBoolTrack.cpp b/Engine/Source/Runtime/MovieSceneCoreTypes/Private/MovieSceneBoolTrack.cpp
--- a/Engine/Source/Runtime/MovieSceneCoreTypes/Private/MovieSceneBoolTrack.cpp
+++ b/Engine/Source/Runtime/MovieSceneCoreTypes/Private/MovieSceneBoolTrack.cpp
@@ -14,7 +14,7 @@ UMovieSceneBoolTrack::UMovieSceneBoolTrack( const FObjectInitializer& ObjectInit
 
 UMovieSceneSection* UMovieSceneBoolTrack::CreateNewSection()
 {
-	return ConstructObject<UMovieSceneSection>( UMovieSceneBoolSection::StaticClass(), this, NAME_None, RF_Transactional );
+	return NewObject<UMovieSceneSection>(this, UMovieSceneBoolSection::StaticClass(), NAME_None, RF_Transactional);
 }
 
 TSharedPtr<IMovieSceneTrackInstance> UMovieSceneBoolTrack::CreateInstance()
diff --git a/Engine/Source/Runtime/MovieSceneCoreTypes/Private/MovieSceneColorTrack.cpp b/Engine/Source/Runtime/MovieSceneCoreTypes/Private/MovieSceneColorTrack.cpp
--- a/Engine/Source/Runtime/MovieSceneCoreTypes/Private/MovieSceneColorTrack.cpp
+++ b/Engine/Source/Runtime/MovieSceneCoreTypes/Private/MovieSceneColorTrack.cpp
@@ -14,7 +14,7 @@ UMovieSceneColorTrack::UMovieSceneColorTrack( const FObjectInitializer& ObjectIn
 
 UMovieSceneSection* UMovieSceneColorTrack::CreateNewSection()
 {
-	return ConstructObject<UMovieSceneSection>( UMovieSceneColorSection::StaticClass(), this, NAME_None, RF_Transactional );
+	return NewObject<UMovieSceneSection>(this, UMovieSceneColorSection::StaticClass(), NAME_None, RF_Transactional);
 }
 
 TSharedPtr<IMovieSceneTrackInstance> UMovieSceneColorTrack::CreateInstance()
diff --git a/Engine/Source/Runtime/MovieSceneCoreTypes/Private/MovieSceneDirectorTrack.cpp b/Engine/Source/Runtime/MovieSceneCoreTypes/Private/MovieSceneDirectorTrack.cpp
--- a/Engine/Source/Runtime/MovieSceneCoreTypes/Private/MovieSceneDirectorTrack.cpp
+++ b/Engine/Source/Runtime/MovieSceneCoreTypes/Private/MovieSceneDirectorTrack.cpp
@@ -58,7 +58,7 @@ bool UMovieSceneDirectorTrack::IsEmpty() const
 
 void UMovieSceneDirectorTrack::AddNewShot(FGuid CameraHandle, float Time)
 {
-	UMovieSceneShotSection* NewSection = ConstructObject<UMovieSceneShotSection>( UMovieSceneShotSection::StaticClass(), this );
+	UMovieSceneShotSection* NewSection = NewObject<UMovieSceneShotSection>(this);
 	NewSection->InitialPlacement(ShotSections, Time, Time + 1.f, SupportsMultipleRows()); //@todo find a better default end time
 	NewSection->SetCameraGuid(CameraHandle);
 	NewSection->SetTitle(LOCTEXT("AddNewShot", "New Shot"));
diff --git a/Engine/Source/Runtime/MovieSceneCoreTypes/Private/MovieSceneFloatTrack.cpp b/Engine/Source/Runtime/MovieSceneCoreTypes/Private/MovieSceneFloatTrack.cpp
--- a/Engine/Source/Runtime/MovieSceneCoreTypes/Private/MovieSceneFloatTrack.cpp
+++ b/Engine/Source/Runtime/MovieSceneCoreTypes/Private/MovieSceneFloatTrack.cpp
@@ -14,7 +14,7 @@ UMovieSceneFloatTrack::UMovieSceneFloatTrack( const FObjectInitializer& ObjectIn
 
 UMovieSceneSection* UMovieSceneFloatTrack::CreateNewSection()
 {
-	return ConstructObject<UMovieSceneSection>( UMovieSceneFloatSection::StaticClass(), this, NAME_None, RF_Transactional );
+	return NewObject<UMovieSceneSection>(this, UMovieSceneFloatSection::StaticClass(), NAME_None, RF_Transactional);
 }
 
 TSharedPtr<IMovieSceneTrackInstance> UMovieSceneFloatTrack::CreateInstance()
diff --git a/Engine/Source/Runtime/MovieSceneCoreTypes/Private/MovieSceneParticleTrack.cpp b/Engine/Source/Runtime/MovieSceneCoreTypes/Private/MovieSceneParticleTrack.cpp
--- a/Engine/Source/Runtime/MovieSceneCoreTypes/Private/MovieSceneParticleTrack.cpp
+++ b/Engine/Source/Runtime/MovieSceneCoreTypes/Private/MovieSceneParticleTrack.cpp
@@ -69,7 +69,7 @@ void UMovieSceneParticleTrack::AddNewParticleSystem(float KeyTime, bool bTrigger
 	// @todo Instead of a 0.1 second event, this should be 0 seconds, requires handling 0 size sections
 	float Duration = KeyType == EParticleKey::Trigger ? 0.1f : 1.f;
 
-	UMovieSceneParticleSection* NewSection = ConstructObject<UMovieSceneParticleSection>( UMovieSceneParticleSection::StaticClass(), this );
+	UMovieSceneParticleSection* NewSection = NewObject<UMovieSceneParticleSection>(this);
 	NewSection->InitialPlacement(ParticleSections, KeyTime, KeyTime + Duration, SupportsMultipleRows());
 	NewSection->SetKeyType(KeyType);
 
diff --git a/Engine/Source/Runtime/MovieSceneCoreTypes/Private/MovieSceneVectorTrack.cpp b/Engine/Source/Runtime/MovieSceneCoreTypes/Private/MovieSceneVectorTrack.cpp
--- a/Engine/Source/Runtime/MovieSceneCoreTypes/Private/MovieSceneVectorTrack.cpp
+++ b/Engine/Source/Runtime/MovieSceneCoreTypes/Private/MovieSceneVectorTrack.cpp
@@ -14,7 +14,7 @@ UMovieSceneVectorTrack::UMovieSceneVectorTrack( const FObjectInitializer& Object
 
 UMovieSceneSection* UMovieSceneVectorTrack::CreateNewSection()
 {
-	return ConstructObject<UMovieSceneSection>( UMovieSceneVectorSection::StaticClass(), this );
+	return NewObject<UMovieSceneSection>(this, UMovieSceneVectorSection::StaticClass());
 }
 
 TSharedPtr<IMovieSceneTrackInstance> UMovieSceneVectorTrack::CreateInstance()
diff --git a/Engine/Source/Runtime/MovieSceneCoreTypes/Private/SubMovieSceneTrack.cpp b/Engine/Source/Runtime/MovieSceneCoreTypes/Private/SubMovieSceneTrack.cpp
--- a/Engine/Source/Runtime/MovieSceneCoreTypes/Private/SubMovieSceneTrack.cpp
+++ b/Engine/Source/Runtime/MovieSceneCoreTypes/Private/SubMovieSceneTrack.cpp
@@ -13,7 +13,7 @@ USubMovieSceneTrack::USubMovieSceneTrack( const FObjectInitializer& ObjectInitia
 
 UMovieSceneSection* USubMovieSceneTrack::CreateNewSection()
 {
-	return ConstructObject<USubMovieSceneSection>( USubMovieSceneSection::StaticClass(), this, NAME_None, RF_Transactional );
+	return NewObject<USubMovieSceneSection>(this, NAME_None, RF_Transactional);
 }
 
 FName USubMovieSceneTrack::GetTrackName() const
diff --git a/Engine/Source/Runtime/Online/OnlineSubsystem/Private/NamedInterfaces.cpp b/Engine/Source/Runtime/Online/OnlineSubsystem/Private/NamedInterfaces.cpp
--- a/Engine/Source/Runtime/Online/OnlineSubsystem/Private/NamedInterfaces.cpp
+++ b/Engine/Source/Runtime/Online/OnlineSubsystem/Private/NamedInterfaces.cpp
@@ -22,7 +22,7 @@ void UNamedInterfaces::Initialize()
 			FNamedInterface& Interface = NamedInterfaces[AddIndex];
 			// Set the object and interface names
 			Interface.InterfaceName = Def.InterfaceName;
-			Interface.InterfaceObject = ConstructObject<UObject>(Class);
+			Interface.InterfaceObject = NewObject<UObject>(GetTransientPackage(), Class);
 			UE_LOG(LogOnline, Display,
 				TEXT("Created named interface (%s) of type (%s)"),
 				*Def.InterfaceName.ToString(),
diff --git a/Engine/Source/Runtime/Online/OnlineSubsystem/Private/OnlineSubsystemImpl.cpp b/Engine/Source/Runtime/Online/OnlineSubsystem/Private/OnlineSubsystemImpl.cpp
--- a/Engine/Source/Runtime/Online/OnlineSubsystem/Private/OnlineSubsystemImpl.cpp
+++ b/Engine/Source/Runtime/Online/OnlineSubsystem/Private/OnlineSubsystemImpl.cpp
@@ -32,7 +32,7 @@ FOnlineSubsystemImpl::~FOnlineSubsystemImpl()
 
 void FOnlineSubsystemImpl::InitNamedInterfaces()
 {
-	NamedInterfaces = ConstructObject<UNamedInterfaces>(UNamedInterfaces::StaticClass());
+	NamedInterfaces = NewObject<UNamedInterfaces>();
 	if (NamedInterfaces)
 	{
 		NamedInterfaces->Initialize();
diff --git a/Engine/Source/Runtime/Online/OnlineSubsystemUtils/Private/IpNetDriver.cpp b/Engine/Source/Runtime/Online/OnlineSubsystemUtils/Private/IpNetDriver.cpp
--- a/Engine/Source/Runtime/Online/OnlineSubsystemUtils/Private/IpNetDriver.cpp
+++ b/Engine/Source/Runtime/Online/OnlineSubsystemUtils/Private/IpNetDriver.cpp
@@ -138,7 +138,7 @@ bool UIpNetDriver::InitConnect( FNetworkNotify* InNotify, const FURL& ConnectURL
 	}
 
 	// Create new connection.
-	ServerConnection = ConstructObject<UNetConnection>(NetConnectionClass);
+	ServerConnection = NewObject<UNetConnection>(GetTransientPackage(), NetConnectionClass);
 	ServerConnection->InitLocalConnection( this, Socket, ConnectURL, USOCK_Pending);
 	UE_LOG(LogNet, Log, TEXT("Game client on port %i, rate %i"), ConnectURL.Port, ServerConnection->CurrentNetSpeed );
 
@@ -264,7 +264,7 @@ void UIpNetDriver::TickDispatch( float DeltaTime )
 
 				if (bAcceptingConnection)
 				{
-					Connection = ConstructObject<UIpConnection>(NetConnectionClass);
+					Connection = NewObject<UIpConnection>(GetTransientPackage(), NetConnectionClass);
                     check(Connection);
 					Connection->InitRemoteConnection( this, Socket,  FURL(), *FromAddr, USOCK_Open);
 					Notify->NotifyAcceptedConnection( Connection );
diff --git a/Engine/Source/Runtime/Online/OnlineSubsystemUtils/Private/OnlineSubsystemUtils.cpp b/Engine/Source/Runtime/Online/OnlineSubsystemUtils/Private/OnlineSubsystemUtils.cpp
--- a/Engine/Source/Runtime/Online/OnlineSubsystemUtils/Private/OnlineSubsystemUtils.cpp
+++ b/Engine/Source/Runtime/Online/OnlineSubsystemUtils/Private/OnlineSubsystemUtils.cpp
@@ -49,7 +49,7 @@ UAudioComponent* CreateVoiceAudioComponent(uint32 SampleRate)
 		if (AudioDevice != NULL)
 		{
 			USoundWaveStreaming* SoundStreaming = NULL;
-			SoundStreaming = ConstructObject<USoundWaveStreaming>(USoundWaveStreaming::StaticClass());
+			SoundStreaming = NewObject<USoundWaveStreaming>();
 			SoundStreaming->SampleRate = SampleRate;
 			SoundStreaming->NumChannels = 1;
 			SoundStreaming->Duration = INDEFINITELY_LOOPING_DURATION;
diff --git a/Engine/Source/Runtime/Online/OnlineSubsystemUtils/Private/PartyBeaconHost.cpp b/Engine/Source/Runtime/Online/OnlineSubsystemUtils/Private/PartyBeaconHost.cpp
--- a/Engine/Source/Runtime/Online/OnlineSubsystemUtils/Private/PartyBeaconHost.cpp
+++ b/Engine/Source/Runtime/Online/OnlineSubsystemUtils/Private/PartyBeaconHost.cpp
@@ -15,7 +15,7 @@ bool APartyBeaconHost::InitHostBeacon(int32 InTeamCount, int32 InTeamSize, int32
 	UE_LOG(LogBeacon, Verbose, TEXT("InitHostBeacon TeamCount:%d TeamSize:%d MaxSize:%d"), InTeamCount, InTeamSize, InMaxReservations);
 	if (InMaxReservations > 0)
 	{
-		State = ConstructObject<UPartyBeaconState>(GetPartyBeaconHostClass());
+		State = NewObject<UPartyBeaconState>(GetTransientPackage(), GetPartyBeaconHostClass());
 		if (State->InitState(InTeamCount, InTeamSize, InMaxReservations, InSessionName, InForceTeamNum))
 		{
 			return true;
diff --git a/Engine/Source/Runtime/UMG/Private/Animation/MovieScene2DTransformTrack.cpp b/Engine/Source/Runtime/UMG/Private/Animation/MovieScene2DTransformTrack.cpp
--- a/Engine/Source/Runtime/UMG/Private/Animation/MovieScene2DTransformTrack.cpp
+++ b/Engine/Source/Runtime/UMG/Private/Animation/MovieScene2DTransformTrack.cpp
@@ -14,7 +14,7 @@ UMovieScene2DTransformTrack::UMovieScene2DTransformTrack( const FObjectInitializ
 
 UMovieSceneSection* UMovieScene2DTransformTrack::CreateNewSection()
 {
-	return ConstructObject<UMovieSceneSection>( UMovieScene2DTransformSection::StaticClass(), this, NAME_None, RF_Transactional );
+	return NewObject<UMovieSceneSection>(this, UMovieScene2DTransformSection::StaticClass(), NAME_None, RF_Transactional);
 }
 
 TSharedPtr<IMovieSceneTrackInstance> UMovieScene2DTransformTrack::CreateInstance()
diff --git a/Engine/Source/Runtime/UMG/Private/Animation/MovieSceneMarginTrack.cpp b/Engine/Source/Runtime/UMG/Private/Animation/MovieSceneMarginTrack.cpp
--- a/Engine/Source/Runtime/UMG/Private/Animation/MovieSceneMarginTrack.cpp
+++ b/Engine/Source/Runtime/UMG/Private/Animation/MovieSceneMarginTrack.cpp
@@ -14,7 +14,7 @@ UMovieSceneMarginTrack::UMovieSceneMarginTrack(const FObjectInitializer& ObjectI
 
 UMovieSceneSection* UMovieSceneMarginTrack::CreateNewSection()
 {
-	return ConstructObject<UMovieSceneSection>( UMovieSceneMarginSection::StaticClass(), this, NAME_None, RF_Transactional );
+	return NewObject<UMovieSceneSection>(this, UMovieSceneMarginSection::StaticClass(), NAME_None, RF_Transactional);
 }
 
 TSharedPtr<IMovieSceneTrackInstance> UMovieSceneMarginTrack::CreateInstance()
diff --git a/Engine/Source/Runtime/UMG/Private/Animation/UMGSequencePlayer.cpp b/Engine/Source/Runtime/UMG/Private/Animation/UMGSequencePlayer.cpp
--- a/Engine/Source/Runtime/UMG/Private/Animation/UMGSequencePlayer.cpp
+++ b/Engine/Source/Runtime/UMG/Private/Animation/UMGSequencePlayer.cpp
@@ -25,7 +25,7 @@ void UUMGSequencePlayer::InitSequencePlayer( const UWidgetAnimation& InAnimation
 	// Cache the time range of the sequence to determine when we stop
 	TimeRange = MovieScene->GetTimeRange();
 
-	RuntimeBindings = ConstructObject<UMovieSceneBindings>( UMovieSceneBindings::StaticClass(), this );
+	RuntimeBindings = NewObject<UMovieSceneBindings>(this);
 	RuntimeBindings->SetRootMovieScene( MovieScene );
 
 	UWidgetTree* WidgetTree = UserWidget.WidgetTree;
diff --git a/Engine/Source/Runtime/UMG/Private/Animation/WidgetAnimation.cpp b/Engine/Source/Runtime/UMG/Private/Animation/WidgetAnimation.cpp
--- a/Engine/Source/Runtime/UMG/Private/Animation/WidgetAnimation.cpp
+++ b/Engine/Source/Runtime/UMG/Private/Animation/WidgetAnimation.cpp
@@ -36,8 +36,8 @@ UWidgetAnimation* UWidgetAnimation::GetNullAnimation()
 	static UWidgetAnimation* NullAnimation = nullptr;
 	if( !NullAnimation )
 	{
-		NullAnimation = ConstructObject<UWidgetAnimation>( UWidgetAnimation::StaticClass(), GetTransientPackage(), NAME_None, RF_RootSet );
-		NullAnimation->MovieScene = ConstructObject<UMovieScene>( UMovieScene::StaticClass(), NullAnimation, FName("No Animation"), RF_RootSet );
+		NullAnimation = NewObject<UWidgetAnimation>(GetTransientPackage(), NAME_None, RF_RootSet);
+		NullAnimation->MovieScene = NewObject<UMovieScene>(NullAnimation, FName("No Animation"), RF_RootSet);
 	}
 
 	return NullAnimation;
diff --git a/Engine/Source/Runtime/UMG/Private/Components/Border.cpp b/Engine/Source/Runtime/UMG/Private/Components/Border.cpp
--- a/Engine/Source/Runtime/UMG/Private/Components/Border.cpp
+++ b/Engine/Source/Runtime/UMG/Private/Components/Border.cpp
@@ -269,7 +269,7 @@ void UBorder::PostLoad()
 			UBorderSlot* BorderSlot = Cast<UBorderSlot>(PanelSlot);
 			if ( BorderSlot == NULL )
 			{
-				BorderSlot = ConstructObject<UBorderSlot>(UBorderSlot::StaticClass(), this);
+				BorderSlot = NewObject<UBorderSlot>(this);
 				BorderSlot->Content = GetContentSlot()->Content;
 				BorderSlot->Content->Slot = BorderSlot;
 				Slots[0] = BorderSlot;
diff --git a/Engine/Source/Runtime/UMG/Private/Components/Button.cpp b/Engine/Source/Runtime/UMG/Private/Components/Button.cpp
--- a/Engine/Source/Runtime/UMG/Private/Components/Button.cpp
+++ b/Engine/Source/Runtime/UMG/Private/Components/Button.cpp
@@ -120,7 +120,7 @@ void UButton::PostLoad()
 			UButtonSlot* ButtonSlot = Cast<UButtonSlot>(PanelSlot);
 			if ( ButtonSlot == NULL )
 			{
-				ButtonSlot = ConstructObject<UButtonSlot>(UButtonSlot::StaticClass(), this);
+				ButtonSlot = NewObject<UButtonSlot>(this);
 				ButtonSlot->Content = GetContentSlot()->Content;
 				ButtonSlot->Content->Slot = ButtonSlot;
 				Slots[0] = ButtonSlot;
diff --git a/Engine/Source/Runtime/UMG/Private/Components/PanelWidget.cpp b/Engine/Source/Runtime/UMG/Private/Components/PanelWidget.cpp
--- a/Engine/Source/Runtime/UMG/Private/Components/PanelWidget.cpp
+++ b/Engine/Source/Runtime/UMG/Private/Components/PanelWidget.cpp
@@ -112,7 +112,7 @@ UPanelSlot* UPanelWidget::AddChild(UWidget* Content)
 
 	Content->RemoveFromParent();
 
-	UPanelSlot* Slot = ConstructObject<UPanelSlot>(GetSlotClass(), this);
+	UPanelSlot* Slot = NewObject<UPanelSlot>(this);
 	Slot->SetFlags(RF_Transactional);
 	Slot->Content = Content;
 	Slot->Parent = this;
diff --git a/Engine/Source/Runtime/UMG/Private/Components/Widget.cpp b/Engine/Source/Runtime/UMG/Private/Components/Widget.cpp
--- a/Engine/Source/Runtime/UMG/Private/Components/Widget.cpp
+++ b/Engine/Source/Runtime/UMG/Private/Components/Widget.cpp
@@ -818,7 +818,7 @@ static UPropertyBinding* GenerateBinder(UDelegateProperty* DelegateProperty, UOb
 			TSubclassOf<UPropertyBinding> BinderClass = UWidget::FindBinderClassForDestination(ReturnProperty);
 			if ( BinderClass != nullptr )
 			{
-				UPropertyBinding* Binder = ConstructObject<UPropertyBinding>(BinderClass, Container);
+				UPropertyBinding* Binder = NewObject<UPropertyBinding>(Container, BinderClass);
 				Binder->SourceObject = SourceObject;
 				Binder->SourcePath = BindingPath;
 				Binder->Bind(ReturnProperty, ScriptDelegate);
diff --git a/Engine/Source/Runtime/UMG/Private/Components/WidgetComponent.cpp b/Engine/Source/Runtime/UMG/Private/Components/WidgetComponent.cpp
--- a/Engine/Source/Runtime/UMG/Private/Components/WidgetComponent.cpp
+++ b/Engine/Source/Runtime/UMG/Private/Components/WidgetComponent.cpp
@@ -716,7 +716,7 @@ void UWidgetComponent::UpdateRenderTarget()
 
 	if(!RenderTarget && DrawSize != FIntPoint::ZeroValue)
 	{
-		RenderTarget = ConstructObject<UTextureRenderTarget2D>(UTextureRenderTarget2D::StaticClass(), this);
+		RenderTarget = NewObject<UTextureRenderTarget2D>(this);
 
 		RenderTarget->ClearColor = BackgroundColor;
 
@@ -771,7 +771,7 @@ void UWidgetComponent::UpdateBodySetup( bool bDrawSizeChanged )
 {
 	if( !BodySetup || bDrawSizeChanged )
 	{
-		BodySetup = ConstructObject<UBodySetup>( UBodySetup::StaticClass(), this );
+		BodySetup = NewObject<UBodySetup>(this);
 		BodySetup->CollisionTraceFlag = CTF_UseSimpleAsComplex;
 		BodySetup->AggGeom.BoxElems.Add(FKBoxElem());
 
diff --git a/Engine/Source/Runtime/UMG/Private/UserWidget.cpp b/Engine/Source/Runtime/UMG/Private/UserWidget.cpp
--- a/Engine/Source/Runtime/UMG/Private/UserWidget.cpp
+++ b/Engine/Source/Runtime/UMG/Private/UserWidget.cpp
@@ -385,7 +385,7 @@ void UUserWidget::PlayAnimation(const UWidgetAnimation* InAnimation, float Start
 
 		if( !FoundPlayer )
 		{
-			UUMGSequencePlayer* NewPlayer = ConstructObject<UUMGSequencePlayer>( UUMGSequencePlayer::StaticClass(), this );
+			UUMGSequencePlayer* NewPlayer = NewObject<UUMGSequencePlayer>(this);
 			ActiveSequencePlayers.Add( NewPlayer );
 
 			NewPlayer->OnSequenceFinishedPlaying().AddUObject( this, &UUserWidget::OnAnimationFinishedPlaying );
diff --git a/Engine/Source/Runtime/UMG/Private/WidgetBlueprintLibrary.cpp b/Engine/Source/Runtime/UMG/Private/WidgetBlueprintLibrary.cpp
--- a/Engine/Source/Runtime/UMG/Private/WidgetBlueprintLibrary.cpp
+++ b/Engine/Source/Runtime/UMG/Private/WidgetBlueprintLibrary.cpp
@@ -37,11 +37,11 @@ UDragDropOperation* UWidgetBlueprintLibrary::CreateDragDropOperation(TSubclassOf
 {
 	if ( Operation )
 	{
-		return ConstructObject<UDragDropOperation>(Operation);
+		return NewObject<UDragDropOperation>(GetTransientPackage(), Operation);
 	}
 	else
 	{
-		return ConstructObject<UDragDropOperation>(UDragDropOperation::StaticClass());
+		return NewObject<UDragDropOperation>();
 	}
 }
 
diff --git a/Engine/Source/Runtime/UMG/Public/Blueprint/UserWidget.h b/Engine/Source/Runtime/UMG/Public/Blueprint/UserWidget.h
--- a/Engine/Source/Runtime/UMG/Public/Blueprint/UserWidget.h
+++ b/Engine/Source/Runtime/UMG/Public/Blueprint/UserWidget.h
@@ -694,7 +694,7 @@ T* CreateWidget(UWorld* World, UClass* UserWidgetClass)
 
 	// Assign the outer to the game instance if it exists, otherwise use the world
 	UObject* Outer = World->GetGameInstance() ? StaticCast<UObject*>(World->GetGameInstance()) : StaticCast<UObject*>(World);
-	UUserWidget* NewWidget = ConstructObject<UUserWidget>(UserWidgetClass, Outer);
+	UUserWidget* NewWidget = NewObject<UUserWidget>(Outer, UserWidgetClass);
 
 	if ( ULocalPlayer* Player = World->GetFirstLocalPlayerFromController() )
 	{
@@ -724,7 +724,7 @@ T* CreateWidget(APlayerController* OwningPlayer, UClass* UserWidgetClass)
 	// Assign the outer to the game instance if it exists, otherwise use the player controller's world
 	UWorld* World = OwningPlayer->GetWorld();
 	UObject* Outer = World->GetGameInstance() ? StaticCast<UObject*>(World->GetGameInstance()) : StaticCast<UObject*>(World);
-	UUserWidget* NewWidget = ConstructObject<UUserWidget>(UserWidgetClass, Outer);
+	UUserWidget* NewWidget = NewObject<UUserWidget>(Outer, UserWidgetClass);
 	
 	NewWidget->SetPlayerContext(FLocalPlayerContext(OwningPlayer));
 	NewWidget->Initialize();
diff --git a/Engine/Source/Runtime/UMG/Public/Blueprint/WidgetTree.h b/Engine/Source/Runtime/UMG/Public/Blueprint/WidgetTree.h
--- a/Engine/Source/Runtime/UMG/Public/Blueprint/WidgetTree.h
+++ b/Engine/Source/Runtime/UMG/Public/Blueprint/WidgetTree.h
@@ -105,14 +105,14 @@ public:
 	{
 		if ( WidgetType->IsChildOf(UUserWidget::StaticClass()) )
 		{
-			UUserWidget* Widget = ConstructObject<UUserWidget>(WidgetType, this);
+			UUserWidget* Widget = NewObject<UUserWidget>(this, WidgetType);
 			Widget->Initialize();
 			Widget->SetFlags(RF_Transactional);
 			return (T*)Widget;
 		}
 		else
 		{
-			UWidget* Widget = (UWidget*)ConstructObject<UWidget>(WidgetType, this);
+			UWidget* Widget = (UWidget*)NewObject<UWidget>(this, WidgetType);
 			Widget->SetFlags(RF_Transactional);
 			return (T*)Widget;
 		}
