Fix Mac non-unity build.
Why did this compile on Windows?

diff --git a/Engine/Source/Runtime/Landscape/Private/LandscapeSplineRaster.cpp b/Engine/Source/Runtime/Landscape/Private/LandscapeSplineRaster.cpp
--- a/Engine/Source/Runtime/Landscape/Private/LandscapeSplineRaster.cpp
+++ b/Engine/Source/Runtime/Landscape/Private/LandscapeSplineRaster.cpp
@@ -15,6 +15,7 @@
 #include "LandscapeSplineControlPoint.h"
 #include "LandscapeSplineSegment.h"
 #if WITH_EDITOR
+#include "AI/Navigation/NavigationSystem.h"
 #include "ScopedTransaction.h"
 #include "Raster.h"
 #endif
