Fix for crash cooking landscape for mobile in the editor

diff --git a/Engine/Source/Runtime/Landscape/Private/Landscape.cpp b/Engine/Source/Runtime/Landscape/Private/Landscape.cpp
--- a/Engine/Source/Runtime/Landscape/Private/Landscape.cpp
+++ b/Engine/Source/Runtime/Landscape/Private/Landscape.cpp
@@ -132,6 +132,10 @@ void ULandscapeComponent::AddReferencedObjects(UObject* InThis, FReferenceCollec
 void ULandscapeComponent::Serialize(FArchive& Ar)
 {
 #if ENABLE_LANDSCAPE_COOKING && WITH_EDITOR
+	bool bRestoreAfterCooking = false;
+	UMaterialInstanceConstant* BackupMaterialInstance = MaterialInstance;
+	UTexture2D* BackupHeightmapTexture = HeightmapTexture;
+
 	// Saving for cooking path
 	if (Ar.IsCooking() && Ar.IsSaving() && !HasAnyFlags(RF_ClassDefaultObject))
 	{
@@ -142,6 +146,7 @@ void ULandscapeComponent::Serialize(FArchive& Ar)
 				GeneratePlatformVertexData();
 				MaterialInstance = Cast<UMaterialInstanceConstant>(GeneratePlatformPixelData(WeightmapTextures, true));
 				HeightmapTexture = NULL;
+				bRestoreAfterCooking = true;
 			}
 		}
 	}
@@ -149,6 +154,14 @@ void ULandscapeComponent::Serialize(FArchive& Ar)
 
 	Super::Serialize(Ar);
 
+#if ENABLE_LANDSCAPE_COOKING && WITH_EDITOR
+	if (bRestoreAfterCooking)
+	{
+		MaterialInstance = BackupMaterialInstance;
+		HeightmapTexture = BackupHeightmapTexture;
+	}
+#endif
+
 	Ar << LightMap;
 	Ar << ShadowMap;
 
