[INTEGRATE] Change 2426756 by [REDACTED_PERSON][REDACTED_HANDLE] on 2015/01/30 19:05:11

	Fixed some calls to Super::DestroyComponent() that wern't passing the bPromoteChildren parameter.

diff --git a/Engine/Source/Runtime/Landscape/Private/Landscape.cpp b/Engine/Source/Runtime/Landscape/Private/Landscape.cpp
--- a/Engine/Source/Runtime/Landscape/Private/Landscape.cpp
+++ b/Engine/Source/Runtime/Landscape/Private/Landscape.cpp
@@ -797,7 +797,7 @@ void ULandscapeComponent::DestroyComponent(bool bPromoteChildren/*= false*/)
 		Proxy->LandscapeComponents.Remove(this);
 	}
 
-	Super::DestroyComponent();
+	Super::DestroyComponent(bPromoteChildren);
 }
 
 FBoxSphereBounds ULandscapeComponent::CalcBounds(const FTransform& LocalToWorld) const
diff --git a/Engine/Source/Runtime/Landscape/Private/LandscapeCollision.cpp b/Engine/Source/Runtime/Landscape/Private/LandscapeCollision.cpp
--- a/Engine/Source/Runtime/Landscape/Private/LandscapeCollision.cpp
+++ b/Engine/Source/Runtime/Landscape/Private/LandscapeCollision.cpp
@@ -868,7 +868,7 @@ void ULandscapeMeshCollisionComponent::DestroyComponent(bool bPromoteChildren/*=
 		Proxy->CollisionComponents.Remove(this);
 	}
 
-	Super::DestroyComponent();
+	Super::DestroyComponent(bPromoteChildren);
 }
 
 #if WITH_EDITOR
@@ -968,7 +968,7 @@ void ULandscapeHeightfieldCollisionComponent::DestroyComponent(bool bPromoteChil
 		Proxy->CollisionComponents.Remove(this);
 	}
 
-	Super::DestroyComponent();
+	Super::DestroyComponent(bPromoteChildren);
 }
 
 FBoxSphereBounds ULandscapeHeightfieldCollisionComponent::CalcBounds(const FTransform& LocalToWorld) const
diff --git a/Engine/Source/Runtime/UMG/Private/Components/WidgetComponent.cpp b/Engine/Source/Runtime/UMG/Private/Components/WidgetComponent.cpp
--- a/Engine/Source/Runtime/UMG/Private/Components/WidgetComponent.cpp
+++ b/Engine/Source/Runtime/UMG/Private/Components/WidgetComponent.cpp
@@ -437,7 +437,7 @@ void UWidgetComponent::OnUnregister()
 
 void UWidgetComponent::DestroyComponent(bool bPromoteChildren/*= false*/)
 {
-	Super::DestroyComponent();
+	Super::DestroyComponent(bPromoteChildren);
 
 	Renderer.Reset();
 
