Don't show the grassmap rendering notification or stream in any textures until all shaders are compiled.

[Merge CL 2431439 from UE4-GDC2015]

diff --git a/Engine/Source/Runtime/Landscape/Classes/LandscapeComponent.h b/Engine/Source/Runtime/Landscape/Classes/LandscapeComponent.h
--- a/Engine/Source/Runtime/Landscape/Classes/LandscapeComponent.h
+++ b/Engine/Source/Runtime/Landscape/Classes/LandscapeComponent.h
@@ -333,9 +333,12 @@ public:
 	void RenderGrassMap();
 	void RemoveGrassMap();
 
-	/* Can a valid grassmap currently be generated? */
+	/* Could a grassmap currently be generated, disregarding whether our textures are streamed in? */
 	bool CanRenderGrassMap() const;
 
+	/* Are the textures we need to render a grassmap currently streamed in? */
+	bool AreTexturesStreamedForGrassMapRender() const;
+
 	/* Is the grassmap data outdated, eg by a material */
 	bool IsGrassMapOutdated() const;
 #endif
diff --git a/Engine/Source/Runtime/Landscape/Private/Landscape.cpp b/Engine/Source/Runtime/Landscape/Private/Landscape.cpp
--- a/Engine/Source/Runtime/Landscape/Private/Landscape.cpp
+++ b/Engine/Source/Runtime/Landscape/Private/Landscape.cpp
@@ -2855,12 +2855,17 @@ void ALandscapeProxy::UpdateFoliage(const TArray<FVector>& Cameras, bool bForceS
 										}
 
 #if WITH_EDITOR
-											// render grass data if we don't have any
+										// render grass data if we don't have any
 										if (!Component->GrassData->HasData())
-										{
-											// if we can't generate now, force our textures to be streamed in
+										{		
 											if (!Component->CanRenderGrassMap())
 											{
+												// we can't currently render grassmaps (eg shaders not compiled)
+												continue;
+											}
+											else if (!Component->AreTexturesStreamedForGrassMapRender())
+											{
+												// we're ready to generate but our textures need streaming in
 												DesiredForceStreamedTextures.Add(Component->HeightmapTexture);
 												for (auto WeightmapTexture : Component->WeightmapTextures)
 												{
@@ -2970,6 +2975,7 @@ void ALandscapeProxy::UpdateFoliage(const TArray<FVector>& Cameras, bool bForceS
 
 			{
 				int32 NumComponentsRendered = 0;
+				int32 NumComponentsUnableToRender = 0;
 				if (GrassTypes.Num() > 0 && CVarPrerenderGrassmaps.GetValueOnAnyThread() > 0)
 				{
 					// try to render some grassmaps
@@ -2978,24 +2984,31 @@ void ALandscapeProxy::UpdateFoliage(const TArray<FVector>& Cameras, bool bForceS
 					{
 						if (Component->CanRenderGrassMap())
 						{
-							ComponentsToRender.Add(Component);
-							NumComponentsRendered++;
-							// We really want to throttle the number based on component size...
-							if (NumComponentsRendered >= 4)
+							if (Component->AreTexturesStreamedForGrassMapRender())
 							{
-								break;
+								// We really want to throttle the number based on component size.
+								if (NumComponentsRendered <= 4)
+								{
+									ComponentsToRender.Add(Component);
+									NumComponentsRendered++;
+								}
 							}
-						}
-						else if (TotalTexturesToStreamForVisibleGrassMapRender == 0)
-						{
-							// Force stream in other heightmaps but only if we're not waiting for the textures 
-							// near the camera to stream in
-							DesiredForceStreamedTextures.Add(Component->HeightmapTexture);
-							for (auto WeightmapTexture : Component->WeightmapTextures)
+							else
+							if (TotalTexturesToStreamForVisibleGrassMapRender == 0)
 							{
-								DesiredForceStreamedTextures.Add(WeightmapTexture);
+								// Force stream in other heightmaps but only if we're not waiting for the textures 
+								// near the camera to stream in
+								DesiredForceStreamedTextures.Add(Component->HeightmapTexture);
+								for (auto WeightmapTexture : Component->WeightmapTextures)
+								{
+									DesiredForceStreamedTextures.Add(WeightmapTexture);
+								}
 							}
 						}
+						else
+						{
+							NumComponentsUnableToRender++;
+						}
 					}
 					if (ComponentsToRender.Num())
 					{
@@ -3005,7 +3018,7 @@ void ALandscapeProxy::UpdateFoliage(const TArray<FVector>& Cameras, bool bForceS
 				}
 
 				TotalComponentsNeedingGrassMapRender -= NumComponentsNeedingGrassMapRender;
-				NumComponentsNeedingGrassMapRender = ComponentsNeedingGrassMapRender.Num() - NumComponentsRendered;
+				NumComponentsNeedingGrassMapRender = ComponentsNeedingGrassMapRender.Num() - NumComponentsRendered - NumComponentsUnableToRender;
 				TotalComponentsNeedingGrassMapRender += NumComponentsNeedingGrassMapRender;
 				
 				// Update resident flags
diff --git a/Engine/Source/Runtime/Landscape/Private/LandscapeGrass.cpp b/Engine/Source/Runtime/Landscape/Private/LandscapeGrass.cpp
--- a/Engine/Source/Runtime/Landscape/Private/LandscapeGrass.cpp
+++ b/Engine/Source/Runtime/Landscape/Private/LandscapeGrass.cpp
@@ -490,6 +490,11 @@ bool ULandscapeComponent::CanRenderGrassMap() const
 		return false;
 	}
 	
+	return true;
+}
+
+bool ULandscapeComponent::AreTexturesStreamedForGrassMapRender() const
+{
 	// Check for valid heightmap that is fully streamed in
 	if (!HeightmapTexture || HeightmapTexture->ResidentMips != HeightmapTexture->GetNumMips()
 		|| !HeightmapTexture->Resource || ((FTexture2DResource*)HeightmapTexture->Resource)->GetCurrentFirstMip() > 0)
