Moved all grass-related code out of Landscape.cpp into LandscapeGrass.cpp and renamed UpdateFoliage to UpdateGrass

[Integrated CL 2433264 from UE4-GDC2015]

diff --git a/Engine/Source/Editor/WorldBrowser/Private/Tiles/WorldTileCollectionModel.cpp b/Engine/Source/Editor/WorldBrowser/Private/Tiles/WorldTileCollectionModel.cpp
--- a/Engine/Source/Editor/WorldBrowser/Private/Tiles/WorldTileCollectionModel.cpp
+++ b/Engine/Source/Editor/WorldBrowser/Private/Tiles/WorldTileCollectionModel.cpp
@@ -2008,9 +2008,10 @@ bool FWorldTileCollectionModel::GenerateLODLevels(FLevelModelList InLevelList, i
 
 			if (SimplificationDetails.bBakeGrassToLandscape)
 			{
-				Landscape->FlushGrassComponents();
+				/* Flush existing grass components, but not grass maps */
+				Landscape->FlushGrassComponents(nullptr, false);
 				TArray<FVector> Cameras;
-				Landscape->UpdateFoliage(Cameras, true);
+				Landscape->UpdateGrass(Cameras, true);
 			}
 								
 			// This is texture resolution for a landscape mesh, probably needs to be calculated using landscape size
diff --git a/Engine/Source/Runtime/Landscape/Classes/LandscapeProxy.h b/Engine/Source/Runtime/Landscape/Classes/LandscapeProxy.h
--- a/Engine/Source/Runtime/Landscape/Classes/LandscapeProxy.h
+++ b/Engine/Source/Runtime/Landscape/Classes/LandscapeProxy.h
@@ -299,6 +299,8 @@ public:
 	{
 		return TEXT("FAsyncGrassTask");
 	}
+
+	~FAsyncGrassTask();
 };
 
 UCLASS(NotPlaceable, NotBlueprintable, hidecategories=(Display, Attachment, Physics, Debug, Lighting, LOD), showcategories=(Lighting, Rendering, "Utilities|Transformation"), MinimalAPI)
@@ -516,15 +518,18 @@ public:
 	void SetLandscapeGuid(const FGuid& Guid) { LandscapeGuid = Guid; }
 	virtual ALandscape* GetLandscapeActor();
 
+	/* Per-frame call to update dynamic grass placement and render grassmaps */
+	void TickGrass();
+
 	/** Flush the grass cache */
-	LANDSCAPE_API void FlushGrassComponents(const TSet<ULandscapeComponent*>* OnlyForComponents = nullptr);
+	LANDSCAPE_API void FlushGrassComponents(const TSet<ULandscapeComponent*>* OnlyForComponents = nullptr, bool bFlushGrassMaps = true);
 
 	/** 
 		Update Grass 
 		* @param Cameras to use for culling, if empty, then NO culling
 		* @param bForceSync if true, block and finish all work
 	*/
-	LANDSCAPE_API void UpdateFoliage(const TArray<FVector>& Cameras, bool bForceSync = false);
+	LANDSCAPE_API void UpdateGrass(const TArray<FVector>& Cameras, bool bForceSync = false);
 
 	/* Get the list of grass types on this landscape */
 	TArray<ULandscapeGrassType*> GetGrassTypes() const;
diff --git a/Engine/Source/Runtime/Landscape/Private/Landscape.cpp b/Engine/Source/Runtime/Landscape/Private/Landscape.cpp
--- a/Engine/Source/Runtime/Landscape/Private/Landscape.cpp
+++ b/Engine/Source/Runtime/Landscape/Private/Landscape.cpp
@@ -19,7 +19,6 @@ Landscape.cpp: Terrain rendering
 #include "LandscapeSplinesComponent.h"
 #include "LandscapeInfo.h"
 #include "LandscapeInfoMap.h"
-#include "LandscapeGrassType.h"
 #include "LandscapeLayerInfoObject.h"
 #include "LightMap.h"
 #include "ShadowMap.h"
@@ -32,8 +31,6 @@ Landscape.cpp: Terrain rendering
 #include "Engine/StaticMesh.h"
 #include "LandscapeMeshProxyComponent.h"
 #include "LandscapeMeshProxyActor.h"
-#include "ContentStreaming.h"
-#include "Materials/MaterialExpressionLandscapeGrassOutput.h"
 #include "Materials/MaterialExpressionLandscapeLayerWeight.h"
 #include "Materials/MaterialExpressionLandscapeLayerSwitch.h"
 #include "Materials/MaterialExpressionLandscapeLayerSample.h"
@@ -1349,6 +1346,8 @@ void ALandscapeProxy::Destroyed()
 #if WITH_EDITOR
 	TotalComponentsNeedingGrassMapRender -= NumComponentsNeedingGrassMapRender;
 	NumComponentsNeedingGrassMapRender = 0;
+	TotalTexturesToStreamForVisibleGrassMapRender -= NumTexturesToStreamForVisibleGrassMapRender;
+	NumTexturesToStreamForVisibleGrassMapRender = 0;
 #endif
 }
 #endif
@@ -2138,461 +2137,6 @@ bool LandscapeMaterialsParameterSetUpdater(FStaticParameterSet &StaticParameterS
 	return UpdateParameterSet<FStaticTerrainLayerWeightParameter, UMaterialExpressionLandscapeLayerWeight>(StaticParameterSet.TerrainLayerWeightParameters, ParentMaterial);
 }
 
-
-DECLARE_CYCLE_STAT(TEXT("Grass Async Build Time"), STAT_FoliageGrassAsyncBuildTime, STATGROUP_Foliage);
-DECLARE_CYCLE_STAT(TEXT("Grass Start Comp"), STAT_FoliageGrassStartComp, STATGROUP_Foliage);
-DECLARE_CYCLE_STAT(TEXT("Grass End Comp"), STAT_FoliageGrassEndComp, STATGROUP_Foliage);
-DECLARE_CYCLE_STAT(TEXT("Grass Destroy Comps"), STAT_FoliageGrassDestoryComp, STATGROUP_Foliage);
-DECLARE_CYCLE_STAT(TEXT("Grass Update"), STAT_GrassUpdate, STATGROUP_Foliage);
-
-void ALandscapeProxy::FlushGrassComponents(const TSet<ULandscapeComponent*>* OnlyForComponents)
-{
-	if (OnlyForComponents)
-	{
-		for (FCachedLandscapeFoliage::TGrassSet::TIterator Iter(FoliageCache.CachedGrassComps); Iter; ++Iter)
-		{
-			ULandscapeComponent* Component = (*Iter).Key.BasedOn.Get();
-			// if the weak pointer in the cache is invalid, we should kill them anyway
-			if (Component == nullptr || OnlyForComponents->Contains(Component))
-			{
-				UHierarchicalInstancedStaticMeshComponent *Used = (*Iter).Foliage.Get();
-				if (Used)
-				{
-					SCOPE_CYCLE_COUNTER(STAT_FoliageGrassDestoryComp);
-					Used->ClearInstances();
-					Used->DestroyComponent();
-				}
-				Iter.RemoveCurrent();
-			}
-		}
-#if WITH_EDITOR
-		if (GIsEditor)
-		{
-		    for (ULandscapeComponent* Component : *OnlyForComponents)
-		    {
-			    Component->RemoveGrassMap();
-		    }
-		}
-#endif
-	}
-	else
-	{
-		// Clear old foliage component containers
-		FoliageComponents.Empty();
-
-		// Might as well clear the cache...
-		FoliageCache.ClearCache();
-		// Destroy any owned foliage components
-		TInlineComponentArray<UHierarchicalInstancedStaticMeshComponent*> FoliageComps;
-		GetComponents(FoliageComps);
-		for (UHierarchicalInstancedStaticMeshComponent* Component : FoliageComps)
-		{
-			SCOPE_CYCLE_COUNTER(STAT_FoliageGrassDestoryComp);
-			Component->ClearInstances();
-			Component->DestroyComponent();
-		}
-
-		TArray<USceneComponent*> AttachedFoliageComponents = RootComponent->AttachChildren.FilterByPredicate(
-			[](USceneComponent* Component)
-		{
-			return Cast<UHierarchicalInstancedStaticMeshComponent>(Component);
-		});
-
-		// Destroy any attached but un-owned foliage components
-		for (USceneComponent* Component : AttachedFoliageComponents)
-		{
-			SCOPE_CYCLE_COUNTER(STAT_FoliageGrassDestoryComp);
-			Component->DestroyComponent();
-		}
-
-#if WITH_EDITOR
-		if (GIsEditor)
-		{
-		    // Clear GrassMaps
-		    TInlineComponentArray<ULandscapeComponent*> LandComps;
-		    GetComponents(LandComps);
-		    for (ULandscapeComponent* Component : LandComps)
-		    {
-			    Component->RemoveGrassMap();
-		    }
-		}
-#endif
-	}
-}
-
-//// Grass
-
-DEFINE_LOG_CATEGORY_STATIC(LogGrass, Log, All);
-
-static TAutoConsoleVariable<float> CVarGuardBandMultiplier(
-	TEXT("grass.GuardBandMultiplier"),
-	1.3f,
-	TEXT("Used to control discarding in the grass system. Approximate range, 1-4. Multiplied by the cull distance to control when we add grass components."));
-
-static TAutoConsoleVariable<float> CVarGuardBandDiscardMultiplier(
-	TEXT("grass.GuardBandDiscardMultiplier"),
-	1.4f,
-	TEXT("Used to control discarding in the grass system. Approximate range, 1-4. Multiplied by the cull distance to control when we discard grass components."));
-
-static TAutoConsoleVariable<int32> CVarMinFramesToKeepGrass(
-	TEXT("grass.MinFramesToKeepGrass"),
-	30,
-	TEXT("Minimum number of frames before cached grass can be discarded; used to prevent thrashing."));
-
-static TAutoConsoleVariable<float> CVarMinTimeToKeepGrass(
-	TEXT("grass.MinTimeToKeepGrass"),
-	5.0f,
-	TEXT("Minimum number of seconds before cached grass can be discarded; used to prevent thrashing."));
-
-static TAutoConsoleVariable<int32> CVarMinInstancesPerLeaf(
-	TEXT("grass.MinInstancesPerLeaf"),
-	256,
-	TEXT("Used to control the depth of hierarchical culling tree for grass. Clamped to 16-32000; a reasonable range."));
-
-static TAutoConsoleVariable<int32> CVarMaxInstancesPerComponent(
-	TEXT("grass.MaxInstancesPerComponent"),
-	65536,
-	TEXT("Used to control the number of hierarchical components created. More can be more efficient, but can be hitchy as new components come into range"));
-
-static TAutoConsoleVariable<int32> CVarMaxAsyncTasks(
-	TEXT("grass.MaxAsyncTasks"),
-	4,
-	TEXT("Used to control the number of hierarchical components created at a time."));
-
-static TAutoConsoleVariable<int32> CVarGrassEnable(
-	TEXT("grass.Enable"),
-	1,
-	TEXT("1: Enable Grass; 0: Disable Grass"));
-
-static TAutoConsoleVariable<int32> CVarUseStreamingManagerForCameras(
-	TEXT("grass.UseStreamingManagerForCameras"),
-	1,
-	TEXT("1: Use Streaming Manager; 0: Use ViewLocationsRenderedLastFrame"));
-
-static TAutoConsoleVariable<int32> CVarCullSubsections(
-	TEXT("grass.CullSubsections"),
-	1,
-	TEXT("1: Cull each foliage component; 0: Cull only based on the landscape component."));
-
-static TAutoConsoleVariable<int32> CVarDisableGPUCull(
-	TEXT("grass.DisableGPUCull"),
-	0,
-	TEXT("For debugging. Set this to zero to see where the grass is generated. Useful for tweaking the guard bands."));
-
-
-struct FGrassBuilderBase
-{
-	bool bHaveValidData;
-	float GrassDensity;
-	FVector DrawScale;
-	FVector DrawLoc;
-	FMatrix LandscapeToWorld;
-
-	FIntPoint SectionBase;
-	FIntPoint LandscapeSectionOffset;
-	int32 ComponentSizeQuads;
-	FVector Origin;
-	FVector Extent;
-	int32 SqrtMaxInstances;
-
-	FGrassBuilderBase(ALandscapeProxy* Landscape, ULandscapeComponent* Component, const FGrassVariety& GrassVariety, int32 SqrtSubsections = 1, int32 SubX = 0, int32 SubY = 0)
-	{
-		bHaveValidData = true;
-		GrassDensity = GrassVariety.GrassDensity;
-
-		DrawScale = Landscape->GetRootComponent()->RelativeScale3D;
-		DrawLoc = Landscape->GetActorLocation();
-		LandscapeSectionOffset = Landscape->LandscapeSectionOffset;
-
-		SectionBase = Component->GetSectionBase();
-		ComponentSizeQuads = Component->ComponentSizeQuads;
-
-		Origin = FVector(DrawScale.X * float(SectionBase.X), DrawScale.Y * float(SectionBase.Y), 0.0f);
-		Extent = FVector(DrawScale.X * float(SectionBase.X + ComponentSizeQuads), DrawScale.Y * float(SectionBase.Y + ComponentSizeQuads), 0.0f) - Origin;
-
-		SqrtMaxInstances = FMath::CeilToInt(FMath::Sqrt(FMath::Abs(Extent.X * Extent.Y * GrassDensity / 1000.0f / 1000.0f)));
-
-		if (!SqrtMaxInstances)
-		{
-			bHaveValidData = false;
-		}
-		const FRotator DrawRot = Landscape->GetActorRotation();
-		LandscapeToWorld = Landscape->GetRootComponent()->ComponentToWorld.ToMatrixNoScale();
-
-		if (bHaveValidData && SqrtSubsections != 1)
-		{
-			check(SqrtMaxInstances > 2 * SqrtSubsections);
-			SqrtMaxInstances /= SqrtSubsections;
-			check(SqrtMaxInstances > 0);
-
-			Extent /= float(SqrtSubsections);
-			Origin += Extent * FVector(float(SubX), float(SubY), 0.0f);
-		}
-	}
-};
-
-// FLandscapeComponentGrassAccess - accessor wrapper for data for one GrassType from one Component
-struct FLandscapeComponentGrassAccess
-{
-	FLandscapeComponentGrassAccess(const ULandscapeComponent* InComponent, const ULandscapeGrassType* GrassType)
-	: GrassData(InComponent->GrassData)
-	, HeightData(InComponent->GrassData->HeightData)
-	, WeightData(InComponent->GrassData->WeightData.Find(GrassType))
-	, Stride(InComponent->ComponentSizeQuads + 1)
-	{}
-
-	bool IsValid()
-	{
-		return WeightData && WeightData->Num() == FMath::Square(Stride) && HeightData.Num() == FMath::Square(Stride);
-	}
-
-	FORCEINLINE float GetHeight(int32 IdxX, int32 IdxY)
-	{
-		return LandscapeDataAccess::GetLocalHeight(HeightData[IdxX + Stride*IdxY]);
-	}
-	FORCEINLINE float GetWeight(int32 IdxX, int32 IdxY)
-	{
-		return ((float)(*WeightData)[IdxX + Stride*IdxY]) / 255.f;
-	}
-
-	FORCEINLINE int32 GetStride()
-	{
-		return Stride;
-	}
-
-private:
-	TSharedRef<FLandscapeComponentGrassData, ESPMode::ThreadSafe> GrassData;
-	TArray<uint16>& HeightData;
-	TArray<uint8>* WeightData;
-	int32 Stride;
-};
-
-struct FAsyncGrassBuilder : public FGrassBuilderBase
-{
-	FLandscapeComponentGrassAccess GrassData;
-	bool RandomRotation;
-	bool AlignToSurface;
-	float PlacementJitter;
-	FRandomStream RandomStream;
-	FMatrix XForm;
-	FBox MeshBox;
-
-	double RasterTime;
-	double BuildTime;
-	double InstanceTime;
-	int32 TotalInstances;
-
-	// output
-	FStaticMeshInstanceData InstanceBuffer;
-	TArray<FClusterNode> ClusterTree;
-
-	FAsyncGrassBuilder(ALandscapeProxy* Landscape, ULandscapeComponent* Component, const ULandscapeGrassType* GrassType, const FGrassVariety& GrassVariety, UHierarchicalInstancedStaticMeshComponent* HierarchicalInstancedStaticMeshComponent, int32 SqrtSubsections, int32 SubX, int32 SubY)
-		: FGrassBuilderBase(Landscape, Component, GrassVariety, SqrtSubsections, SubX, SubY)
-		, GrassData(Component, GrassType)
-	{
-		bHaveValidData = bHaveValidData && GrassData.IsValid();
-		XForm = LandscapeToWorld * HierarchicalInstancedStaticMeshComponent->ComponentToWorld.ToMatrixWithScale().Inverse();
-		PlacementJitter = GrassVariety.PlacementJitter;
-		RandomRotation = GrassVariety.RandomRotation;
-		AlignToSurface = GrassVariety.AlignToSurface;
-
-		RandomStream.Initialize(HierarchicalInstancedStaticMeshComponent->InstancingRandomSeed);
-
-		MeshBox = GrassVariety.GrassMesh->GetBounds().GetBox();
-
-		RasterTime = 0.0;
-		BuildTime = 0.0;
-		InstanceTime = 0.0;
-		TotalInstances = 0;
-	}
-
-	void Build()
-	{
-		SCOPE_CYCLE_COUNTER(STAT_FoliageGrassAsyncBuildTime);
-		check(bHaveValidData);
-		RasterTime -= FPlatformTime::Seconds();
-
-		float Div = 1.0f / float(SqrtMaxInstances);
-		float MaxJitter1D = FMath::Clamp<float>(PlacementJitter, 0.0f, .99f) * Div * .5f;
-		FVector MaxJitter(MaxJitter1D, MaxJitter1D, 0.0f);
-		MaxJitter *= Extent;
-		Origin += Extent * (Div * 0.5f);
-		struct FInstanceLocal
-		{
-			FVector Pos;
-			bool bKeep;
-		};
-		TArray<FInstanceLocal> Instances;
-		Instances.AddUninitialized(SqrtMaxInstances * SqrtMaxInstances);
-		int32 NumKept = 0;
-		{
-			int32 InstanceIndex = 0;
-			for (int32 xStart = 0; xStart < SqrtMaxInstances; xStart++)
-			{
-				for (int32 yStart = 0; yStart < SqrtMaxInstances; yStart++)
-				{
-					FVector Location(Origin.X + float(xStart) * Div * Extent.X, Origin.Y + float(yStart) * Div * Extent.Y, 0.0f);
-					Location += FVector(RandomStream.GetFraction() * 2.0f - 1.0f, RandomStream.GetFraction() * 2.0f - 1.0f, 0.0f) * MaxJitter;
-
-					FInstanceLocal& Instance = Instances[InstanceIndex];
-					float Weight = GetLayerWeightAtLocationLocal(Location, &Instance.Pos);
-					Instance.bKeep = Weight > 0.0f && Weight >= RandomStream.GetFraction();
-					if (Instance.bKeep)
-					{
-						NumKept++;
-					}
-					InstanceIndex++;
-				}
-			}
-		}
-		RasterTime += FPlatformTime::Seconds();
-
-		if (NumKept)
-		{
-			TArray<FMatrix> InstanceTransforms;
-			InstanceTransforms.AddUninitialized(NumKept);
-			TotalInstances += NumKept;
-
-			{
-				InstanceTime -= FPlatformTime::Seconds();
-				InstanceBuffer.AllocateInstances(NumKept);
-				int32 InstanceIndex = 0;
-				int32 OutInstanceIndex = 0;
-				for (int32 xStart = 0; xStart < SqrtMaxInstances; xStart++)
-				{
-					for (int32 yStart = 0; yStart < SqrtMaxInstances; yStart++)
-					{
-						const FInstanceLocal& Instance = Instances[InstanceIndex];
-						if (Instance.bKeep)
-						{
-							float Rot = RandomRotation ? RandomStream.GetFraction() * 360.0f : 0.0f;
-							FMatrix OutXForm;
-							if (AlignToSurface)
-							{
-								FVector PosX1 = xStart ? Instances[InstanceIndex - SqrtMaxInstances].Pos : Instance.Pos;
-								FVector PosX2 = (xStart + 1 < SqrtMaxInstances) ? Instances[InstanceIndex + SqrtMaxInstances].Pos : Instance.Pos;
-								FVector PosY1 = yStart ? Instances[InstanceIndex - 1].Pos : Instance.Pos;
-								FVector PosY2 = (yStart + 1 < SqrtMaxInstances) ? Instances[InstanceIndex + 1].Pos : Instance.Pos;
-
-								if (PosX1 != PosX2 && PosY1 != PosY2)
-								{
-									FVector NewZ = ((PosX1 - PosX2) ^ (PosY1 - PosY2)).GetSafeNormal();
-									NewZ *= FMath::Sign(NewZ.Z);
-
-									const FVector NewX = (FVector(0, -1, 0) ^ NewZ).GetSafeNormal();
-									const FVector NewY = NewZ ^ NewX;
-
-									FMatrix Align = FMatrix(NewX, NewY, NewZ, FVector::ZeroVector);
-									OutXForm = FRotationMatrix(FRotator(0.0f, Rot, 0.0f)) * Align * FTranslationMatrix(Instance.Pos) * XForm;
-								}
-								else
-								{
-									OutXForm = FRotationMatrix(FRotator(0.0f, Rot, 0.0f)) * FTranslationMatrix(Instance.Pos) * XForm;
-								}
-							}
-							else
-							{
-								OutXForm = FRotationMatrix(FRotator(0.0f, Rot, 0.0f)) * FTranslationMatrix(Instance.Pos) * XForm;
-							}
-							InstanceTransforms[OutInstanceIndex] = OutXForm;
-							FInstanceStream* RenderData = InstanceBuffer.GetInstanceWriteAddress(OutInstanceIndex++);
-							RenderData->SetInstance(OutXForm, RandomStream.GetFraction());
-						}
-						InstanceIndex++;
-					}
-				}
-				InstanceTime += FPlatformTime::Seconds();
-			}
-			{
-				BuildTime -= FPlatformTime::Seconds();
-				TArray<int32> SortedInstances;
-				TArray<int32> InstanceReorderTable;
-				UHierarchicalInstancedStaticMeshComponent::BuildTreeAnyThread(InstanceTransforms, MeshBox, ClusterTree, SortedInstances, InstanceReorderTable, FMath::Clamp<int32>(CVarMinInstancesPerLeaf.GetValueOnAnyThread(), 16, 256 * 128));
-
-				// in-place sort the instances
-				FInstanceStream SwapBuffer;
-				int32 FirstUnfixedIndex = 0;
-				for (int32 FirstUnfixedIndex = 0; FirstUnfixedIndex < NumKept; FirstUnfixedIndex++)
-				{
-					int32 LoadFrom = SortedInstances[FirstUnfixedIndex];
-					if (LoadFrom != FirstUnfixedIndex)
-					{
-						check(LoadFrom > FirstUnfixedIndex);
-						FMemory::Memcpy(&SwapBuffer, InstanceBuffer.GetInstanceWriteAddress(FirstUnfixedIndex), sizeof(FInstanceStream));
-						FMemory::Memcpy(InstanceBuffer.GetInstanceWriteAddress(FirstUnfixedIndex), InstanceBuffer.GetInstanceWriteAddress(LoadFrom), sizeof(FInstanceStream));
-						FMemory::Memcpy(InstanceBuffer.GetInstanceWriteAddress(LoadFrom), &SwapBuffer, sizeof(FInstanceStream));
-
-						int32 SwapGoesTo = InstanceReorderTable[FirstUnfixedIndex];
-						check(SwapGoesTo > FirstUnfixedIndex);
-						check(SortedInstances[SwapGoesTo] == FirstUnfixedIndex);
-						SortedInstances[SwapGoesTo] = LoadFrom;
-						InstanceReorderTable[LoadFrom] = SwapGoesTo;
-
-						InstanceReorderTable[FirstUnfixedIndex] = FirstUnfixedIndex;
-						SortedInstances[FirstUnfixedIndex] = FirstUnfixedIndex;
-					}
-				}
-
-				BuildTime += FPlatformTime::Seconds();
-			}
-		}
-	}
-	FORCEINLINE_DEBUGGABLE float GetLayerWeightAtLocationLocal(const FVector& InLocation, FVector* OutLocation)
-	{
-		// Find location
-		float TestX = InLocation.X / DrawScale.X - (float)SectionBase.X;
-		float TestY = InLocation.Y / DrawScale.Y - (float)SectionBase.Y;
-
-		// Find data
-		int32 X1 = FMath::FloorToInt(TestX);
-		int32 Y1 = FMath::FloorToInt(TestY);
-		int32 X2 = FMath::CeilToInt(TestX);
-		int32 Y2 = FMath::CeilToInt(TestY);
-
-		// Min is to prevent the sampling of the final column from overflowing
-		int32 IdxX1 = FMath::Min<int32>(X1, GrassData.GetStride() - 1);
-		int32 IdxY1 = FMath::Min<int32>(Y1, GrassData.GetStride() - 1);
-		int32 IdxX2 = FMath::Min<int32>(X2, GrassData.GetStride() - 1);
-		int32 IdxY2 = FMath::Min<int32>(Y2, GrassData.GetStride() - 1);
-
-		float LerpX = FMath::Fractional(TestX);
-		float LerpY = FMath::Fractional(TestY);
-
-		float Result;
-		{
-			// sample
-			float Sample11 = GrassData.GetWeight(IdxX1, IdxY1);
-			float Sample21 = GrassData.GetWeight(IdxX2, IdxY1);
-			float Sample12 = GrassData.GetWeight(IdxX1, IdxY2);
-			float Sample22 = GrassData.GetWeight(IdxX2, IdxY2);
-
-			// Bilinear interpolate
-			Result = FMath::Lerp(
-				FMath::Lerp(Sample11, Sample21, LerpX),
-				FMath::Lerp(Sample12, Sample22, LerpX),
-				LerpY);
-		}
-
-		{
-			// sample
-			float Sample11 = GrassData.GetHeight(IdxX1, IdxY1);
-			float Sample21 = GrassData.GetHeight(IdxX2, IdxY1);
-			float Sample12 = GrassData.GetHeight(IdxX1, IdxY2);
-			float Sample22 = GrassData.GetHeight(IdxX2, IdxY2);
-
-			OutLocation->X = InLocation.X - DrawScale.X * float(LandscapeSectionOffset.X);
-			OutLocation->Y = InLocation.Y - DrawScale.Y * float(LandscapeSectionOffset.Y);
-			// Bilinear interpolate
-			OutLocation->Z = DrawScale.Z * FMath::Lerp(
-				FMath::Lerp(Sample11, Sample21, LerpX),
-				FMath::Lerp(Sample12, Sample22, LerpX),
-				LerpY);
-		}
-		return Result;
-	}
-
-};
-
 void ALandscapeProxy::Tick(float DeltaSeconds)
 {
 	if (!IsPendingKillPending() && !HasAnyFlags(RF_NeedLoad | RF_NeedPostLoad | RF_NeedPostLoadSubobjects | RF_Unreachable | RF_PendingKill | RF_ClassDefaultObject | RF_AsyncLoading))
@@ -2601,531 +2145,7 @@ void ALandscapeProxy::Tick(float DeltaSeconds)
 		// the super tick is for an actor tick...
 		//Super::Tick(DeltaSeconds);
 
-		static TArray<FVector> OldCameras;
-		if (CVarUseStreamingManagerForCameras.GetValueOnGameThread() == 0)
-		{
-			UWorld* World = GetWorld();
-			if (!World)
-			{
-				return;
-			}
-
-			if (!OldCameras.Num() && !World->ViewLocationsRenderedLastFrame.Num())
-			{
-				// no cameras, no grass update
-				return;
-			}
-
-			// there is a bug here, which often leaves us with no cameras in the editor
-			const TArray<FVector>& Cameras = World->ViewLocationsRenderedLastFrame.Num() ? World->ViewLocationsRenderedLastFrame : OldCameras;
-
-			if (&Cameras != &OldCameras)
-			{
-				check(IsInGameThread());
-				OldCameras = Cameras;
-			}
-			UpdateFoliage(Cameras);
-		}
-		else
-		{
-			int32 Num = IStreamingManager::Get().GetNumViews();
-			if (!Num)
-			{
-				// no cameras, no grass update
-				return;
-			}
-			OldCameras.Reset(Num);
-			for (int32 Index = 0; Index < Num; Index++)
-			{
-				auto& ViewInfo = IStreamingManager::Get().GetViewInformation(Index);
-				OldCameras.Add(ViewInfo.ViewOrigin);
-			}
-			UpdateFoliage(OldCameras);
-		}
-	}
-}
-
-TArray<ULandscapeGrassType*> ALandscapeProxy::GetGrassTypes() const
-{
-	TArray<ULandscapeGrassType*> GrassTypes;
-	if (LandscapeMaterial)
-	{
-		TArray<const UMaterialExpressionLandscapeGrassOutput*> GrassExpressions;
-		LandscapeMaterial->GetMaterial()->GetAllExpressionsOfType<UMaterialExpressionLandscapeGrassOutput>(GrassExpressions);
-		if (GrassExpressions.Num() > 0)
-		{
-			for (auto& Type : GrassExpressions[0]->GrassTypes)
-			{
-				GrassTypes.Add(Type.GrassType);
-			}
-		}
-	}
-	return GrassTypes;
-}
-
-#if WITH_EDITOR
-int32 ALandscapeProxy::TotalComponentsNeedingGrassMapRender = 0;
-int32 ALandscapeProxy::TotalTexturesToStreamForVisibleGrassMapRender = 0;
-#endif
-
-static TAutoConsoleVariable<int32> CVarPrerenderGrassmaps(
-	TEXT("grass.PrerenderGrassmaps"),
-	1,
-	TEXT("1: Pre-render grass maps for all components in the editor; 0: Generate grass maps on demand while moving through the editor"));
-
-
-void ALandscapeProxy::UpdateFoliage(const TArray<FVector>& Cameras, bool bForceSync)
-{
-	SCOPE_CYCLE_COUNTER(STAT_GrassUpdate);
-
-	if (CVarGrassEnable.GetValueOnAnyThread() > 0)
-	{
-		TArray<ULandscapeGrassType*> GrassTypes = GetGrassTypes();
-
-		float GuardBand = CVarGuardBandMultiplier.GetValueOnAnyThread();
-		float DiscardGuardBand = CVarGuardBandDiscardMultiplier.GetValueOnAnyThread();
-		bool bCullSubsections = CVarCullSubsections.GetValueOnAnyThread() > 0;
-		bool bDisableGPUCull = CVarDisableGPUCull.GetValueOnAnyThread() > 0;
-		int32 MaxInstancesPerComponent = FMath::Max<int32>(1024, CVarMaxInstancesPerComponent.GetValueOnAnyThread());
-		int32 MaxTasks = CVarMaxAsyncTasks.GetValueOnAnyThread();
-
-		UWorld* World = GetWorld();
-		if (World)
-		{
-#if WITH_EDITOR
-			int32 RequiredTexturesNotStreamedIn = 0;
-			TSet<ULandscapeComponent*> ComponentsNeedingGrassMapRender;
-			TSet<UTexture2D*> CurrentForcedStreamedTextures;
-			TSet<UTexture2D*> DesiredForceStreamedTextures;
-
-			if (!World->IsGameWorld())
-			{
-				// see if we need to flush grass for any components
-				TSet<ULandscapeComponent*> FlushComponents;
-				for (auto Component : LandscapeComponents)
-				{
-					// check textures currently needing force streaming
-					if (Component->HeightmapTexture->bForceMiplevelsToBeResident)
-					{
-						CurrentForcedStreamedTextures.Add(Component->HeightmapTexture);
-					}
-					for(auto WeightmapTexture : Component->WeightmapTextures)
-					{
-						if (WeightmapTexture->bForceMiplevelsToBeResident)
-						{
-							CurrentForcedStreamedTextures.Add(WeightmapTexture);
-						}
-					}
-
-					if (Component->IsGrassMapOutdated())
-					{
-						FlushComponents.Add(Component);
-						if (GrassTypes.Num() > 0)
-						{
-							ComponentsNeedingGrassMapRender.Add(Component);
-						}
-					}
-					else
-					if (!Component->GrassData->HasData())
-					{
-						if (GrassTypes.Num() > 0)
-						{
-							ComponentsNeedingGrassMapRender.Add(Component);
-						}
-					}
-				}
-				if (FlushComponents.Num())
-				{
-					FlushGrassComponents(&FlushComponents);
-				}
-			}
-#endif
-
-			int32 NumCompsCreated = 0;
-			for (int32 ComponentIndex = 0; ComponentIndex < LandscapeComponents.Num(); ComponentIndex++)
-			{
-				ULandscapeComponent* Component = LandscapeComponents[ComponentIndex];
-
-				// skip if we have no data and no way to generate it
-				if (World->IsGameWorld() && !Component->GrassData->HasData()) 
-				{
-					continue;
-				}
-
-				FBoxSphereBounds WorldBounds = Component->CalcBounds(Component->ComponentToWorld);
-				float MinDistanceToComp = Cameras.Num() ? MAX_flt : 0.0f;
-
-				for (auto& Pos : Cameras)
-				{
-					MinDistanceToComp = FMath::Min<float>(MinDistanceToComp, WorldBounds.ComputeSquaredDistanceFromBoxToPoint(Pos));
-				}
-
-				MinDistanceToComp = FMath::Sqrt(MinDistanceToComp);
-
-				for (auto GrassType : GrassTypes)
-				{
-					if (GrassType)
-					{
-						int32 GrassVarietyIndex = -1;
-						for (auto& GrassVariety : GrassType->GrassVarieties)
-						{
-							GrassVarietyIndex++;
-							if (GrassVariety.GrassMesh && GrassVariety.GrassDensity > 0.0f && GrassVariety.EndCullDistance > 0)
-							{
-								bool bRandomRotation = GrassVariety.RandomRotation;
-								bool bAlign = GrassVariety.AlignToSurface;
-								float MustHaveDistance = GuardBand * (float)GrassVariety.EndCullDistance;
-								float DiscardDistance = DiscardGuardBand * (float)GrassVariety.EndCullDistance;
-
-								if (MinDistanceToComp > DiscardDistance)
-								{
-									continue;
-								}
-
-								FGrassBuilderBase ForSubsectionMath(this, Component, GrassVariety);
-
-								int32 SqrtSubsections = 1;
-
-								if (ForSubsectionMath.bHaveValidData && ForSubsectionMath.SqrtMaxInstances > 0)
-								{
-									SqrtSubsections = FMath::Clamp<int32>(FMath::CeilToInt(float(ForSubsectionMath.SqrtMaxInstances) / FMath::Sqrt((float)MaxInstancesPerComponent)), 1, 16);
-								}
-
-								FBox LocalBox = Component->CachedLocalBox;
-								FVector LocalExtentDiv = (LocalBox.Max - LocalBox.Min) * FVector(1.0f / float(SqrtSubsections), 1.0f / float(SqrtSubsections), 1.0f);
-								for (int32 SubX = 0; SubX < SqrtSubsections; SubX++)
-								{
-									for (int32 SubY = 0; SubY < SqrtSubsections; SubY++)
-									{
-										float MinDistanceToSubComp = MinDistanceToComp;
-
-										if (bCullSubsections && SqrtSubsections > 1)
-										{
-											FVector BoxMin;
-											BoxMin.X = LocalBox.Min.X + LocalExtentDiv.X * float(SubX);
-											BoxMin.Y = LocalBox.Min.Y + LocalExtentDiv.Y * float(SubY);
-											BoxMin.Z = LocalBox.Min.Z;
-
-											FVector BoxMax;
-											BoxMax.X = LocalBox.Min.X + LocalExtentDiv.X * float(SubX + 1);
-											BoxMax.Y = LocalBox.Min.Y + LocalExtentDiv.Y * float(SubY + 1);
-											BoxMax.Z = LocalBox.Max.Z;
-
-											FBox LocalSubBox(BoxMin, BoxMax);
-											FBox WorldSubBox = LocalSubBox.TransformBy(Component->ComponentToWorld);
-
-											MinDistanceToSubComp = Cameras.Num() ? MAX_flt : 0.0f;
-											for (auto& Pos : Cameras)
-											{
-												MinDistanceToSubComp = FMath::Min<float>(MinDistanceToSubComp, ComputeSquaredDistanceFromBoxToPoint(WorldSubBox.Min, WorldSubBox.Max, Pos));
-											}
-											MinDistanceToSubComp = FMath::Sqrt(MinDistanceToSubComp);
-										}
-
-										if (MinDistanceToSubComp > DiscardDistance)
-										{
-											continue;
-										}
-
-										FCachedLandscapeFoliage::FGrassComp NewComp;
-										NewComp.Key.BasedOn = Component;
-										NewComp.Key.GrassType = GrassType;
-										NewComp.Key.SqrtSubsections = SqrtSubsections;
-										NewComp.Key.CachedMaxInstancesPerComponent = MaxInstancesPerComponent;
-										NewComp.Key.SubsectionX = SubX;
-										NewComp.Key.SubsectionY = SubY;
-										NewComp.Key.NumVarieties = GrassType->GrassVarieties.Num();
-										NewComp.Key.VarietyIndex = GrassVarietyIndex;
-
-										{
-											FCachedLandscapeFoliage::FGrassComp* Existing = FoliageCache.CachedGrassComps.Find(NewComp.Key);
-											if (Existing || MinDistanceToSubComp > MustHaveDistance)
-											{
-												if (Existing)
-												{
-													Existing->Touch();
-												}
-												continue;
-											}
-										}
-
-										if (NumCompsCreated || (!bForceSync && AsyncFoliageTasks.Num() >= MaxTasks))
-										{
-											continue; // one per frame, but we still want to touch the existing ones
-										}
-
-#if WITH_EDITOR
-										// render grass data if we don't have any
-										if (!Component->GrassData->HasData())
-										{		
-											if (!Component->CanRenderGrassMap())
-											{
-												// we can't currently render grassmaps (eg shaders not compiled)
-												continue;
-											}
-											else if (!Component->AreTexturesStreamedForGrassMapRender())
-											{
-												// we're ready to generate but our textures need streaming in
-												DesiredForceStreamedTextures.Add(Component->HeightmapTexture);
-												for (auto WeightmapTexture : Component->WeightmapTextures)
-												{
-													DesiredForceStreamedTextures.Add(WeightmapTexture);
-												}
-												RequiredTexturesNotStreamedIn++;
-												continue;
-											}
-
-											QUICK_SCOPE_CYCLE_COUNTER(STAT_GrassRenderToTexture);
-											Component->RenderGrassMap();
-											ComponentsNeedingGrassMapRender.Remove(Component);
-											Component->MarkPackageDirty();
-										}
-#endif
-
-										NumCompsCreated++;
-
-										SCOPE_CYCLE_COUNTER(STAT_FoliageGrassStartComp);
-										int32 FolSeed = FCrc::StrCrc32((GrassType->GetName() + Component->GetName() + FString::Printf(TEXT("%d %d %d"), SubX, SubY, GrassVarietyIndex)).GetCharArray().GetData());
-										if (FolSeed == 0)
-										{
-											FolSeed++;
-										}
-
-										UHierarchicalInstancedStaticMeshComponent* HierarchicalInstancedStaticMeshComponent;
-										{
-											QUICK_SCOPE_CYCLE_COUNTER(STAT_GrassCreateComp);
-											HierarchicalInstancedStaticMeshComponent = NewObject<UHierarchicalInstancedStaticMeshComponent>(this);
-										}
-										NewComp.Foliage = HierarchicalInstancedStaticMeshComponent;
-										FoliageCache.CachedGrassComps.Add(NewComp);
-
-										HierarchicalInstancedStaticMeshComponent->Mobility = EComponentMobility::Static;
-
-										HierarchicalInstancedStaticMeshComponent->StaticMesh = GrassVariety.GrassMesh;
-										HierarchicalInstancedStaticMeshComponent->bSelectable = false;
-										HierarchicalInstancedStaticMeshComponent->bHasPerInstanceHitProxies = true;
-										static FName NoCollision(TEXT("NoCollision"));
-										HierarchicalInstancedStaticMeshComponent->SetCollisionProfileName(NoCollision);
-										HierarchicalInstancedStaticMeshComponent->bDisableCollision = true;
-										HierarchicalInstancedStaticMeshComponent->SetCanEverAffectNavigation(false);
-										HierarchicalInstancedStaticMeshComponent->InstancingRandomSeed = FolSeed;
-
-										if (!Cameras.Num() || bDisableGPUCull)
-										{
-											// if we don't have any cameras, then we are rendering landscape LOD materials or somesuch and we want to disable culling
-											HierarchicalInstancedStaticMeshComponent->InstanceStartCullDistance = 0;
-											HierarchicalInstancedStaticMeshComponent->InstanceEndCullDistance = 0;
-										}
-										else
-										{
-											HierarchicalInstancedStaticMeshComponent->InstanceStartCullDistance = GrassVariety.StartCullDistance;
-											HierarchicalInstancedStaticMeshComponent->InstanceEndCullDistance = GrassVariety.EndCullDistance;
-										}
-
-										//@todo - take the settings from a UFoliageType object.  For now, disable distance field lighting on grass so we don't hitch.
-										HierarchicalInstancedStaticMeshComponent->bAffectDistanceFieldLighting = false;
-
-										{
-											QUICK_SCOPE_CYCLE_COUNTER(STAT_GrassAttachComp);
-
-											HierarchicalInstancedStaticMeshComponent->AttachTo(GetRootComponent());
-											FTransform DesiredTransform = GetRootComponent()->ComponentToWorld;
-											DesiredTransform.RemoveScaling();
-											HierarchicalInstancedStaticMeshComponent->SetWorldTransform(DesiredTransform);
-
-											FoliageComponents.Add(HierarchicalInstancedStaticMeshComponent);
-										}
-
-										FAsyncGrassBuilder* Builder;
-
-										{
-											QUICK_SCOPE_CYCLE_COUNTER(STAT_GrassCreateBuilder);
-											Builder = new FAsyncGrassBuilder(this, Component, GrassType, GrassVariety, HierarchicalInstancedStaticMeshComponent, SqrtSubsections, SubX, SubY);
-										}
-
-										if (Builder->bHaveValidData)
-										{
-											FAsyncTask<FAsyncGrassTask>* Task = new FAsyncTask<FAsyncGrassTask>(Builder, NewComp.Key, HierarchicalInstancedStaticMeshComponent);
-
-											Task->StartBackgroundTask();
-
-											AsyncFoliageTasks.Add(Task);
-										}
-										else
-										{
-											delete Builder;
-										}
-										{
-											QUICK_SCOPE_CYCLE_COUNTER(STAT_GrassRegisterComp);
-											HierarchicalInstancedStaticMeshComponent->RegisterComponent();
-										}
-									}
-								}
-							}
-						}
-					}
-				}
-			}
-
-#if WITH_EDITOR
-
-			TotalTexturesToStreamForVisibleGrassMapRender -= NumTexturesToStreamForVisibleGrassMapRender;
-			NumTexturesToStreamForVisibleGrassMapRender = RequiredTexturesNotStreamedIn;
-			TotalTexturesToStreamForVisibleGrassMapRender += NumTexturesToStreamForVisibleGrassMapRender;
-
-			{
-				int32 NumComponentsRendered = 0;
-				int32 NumComponentsUnableToRender = 0;
-				if (GrassTypes.Num() > 0 && CVarPrerenderGrassmaps.GetValueOnAnyThread() > 0)
-				{
-					// try to render some grassmaps
-					TArray<ULandscapeComponent*> ComponentsToRender;
-					for (auto Component : ComponentsNeedingGrassMapRender)
-					{
-						if (Component->CanRenderGrassMap())
-						{
-							if (Component->AreTexturesStreamedForGrassMapRender())
-							{
-								// We really want to throttle the number based on component size.
-								if (NumComponentsRendered <= 4)
-								{
-									ComponentsToRender.Add(Component);
-									NumComponentsRendered++;
-								}
-							}
-							else
-							if (TotalTexturesToStreamForVisibleGrassMapRender == 0)
-							{
-								// Force stream in other heightmaps but only if we're not waiting for the textures 
-								// near the camera to stream in
-								DesiredForceStreamedTextures.Add(Component->HeightmapTexture);
-								for (auto WeightmapTexture : Component->WeightmapTextures)
-								{
-									DesiredForceStreamedTextures.Add(WeightmapTexture);
-								}
-							}
-						}
-						else
-						{
-							NumComponentsUnableToRender++;
-						}
-					}
-					if (ComponentsToRender.Num())
-					{
-						RenderGrassMaps(ComponentsToRender, GrassTypes);
-						MarkPackageDirty();
-					}
-				}
-
-				TotalComponentsNeedingGrassMapRender -= NumComponentsNeedingGrassMapRender;
-				NumComponentsNeedingGrassMapRender = ComponentsNeedingGrassMapRender.Num() - NumComponentsRendered - NumComponentsUnableToRender;
-				TotalComponentsNeedingGrassMapRender += NumComponentsNeedingGrassMapRender;
-				
-				// Update resident flags
-				for (auto Texture : DesiredForceStreamedTextures.Difference(CurrentForcedStreamedTextures))
-				{
-					Texture->bForceMiplevelsToBeResident = true;
-				}
-				for (auto Texture : CurrentForcedStreamedTextures.Difference(DesiredForceStreamedTextures))
-				{
-					Texture->bForceMiplevelsToBeResident = false;
-				}
-
-			}
-#endif
-		}
-	}
-
-	static TSet<UHierarchicalInstancedStaticMeshComponent *> StillUsed;
-	StillUsed.Empty(256);
-	{
-		// trim cached items based on time, pending and emptiness
-		double OldestToKeepTime = FPlatformTime::Seconds() - CVarMinTimeToKeepGrass.GetValueOnGameThread();
-		uint32 OldestToKeepFrame = GFrameNumber - CVarMinFramesToKeepGrass.GetValueOnGameThread();
-		for (FCachedLandscapeFoliage::TGrassSet::TIterator Iter(FoliageCache.CachedGrassComps); Iter; ++Iter)
-		{
-			const FCachedLandscapeFoliage::FGrassComp& GrassItem = *Iter;
-			UHierarchicalInstancedStaticMeshComponent *Used = GrassItem.Foliage.Get();
-			bool bOld =
-				!GrassItem.Pending &&
-				(
-				!GrassItem.Key.BasedOn.Get() ||
-				!GrassItem.Key.GrassType.Get() ||
-				!Used ||
-				(GrassItem.LastUsedFrameNumber < OldestToKeepFrame && GrassItem.LastUsedTime < OldestToKeepTime)
-				);
-			if (bOld)
-			{
-				Iter.RemoveCurrent();
-			}
-			else if (Used)
-			{
-				StillUsed.Add(Used);
-			}
-		}
-	}
-	{
-
-		// delete components that are no longer used
-		for (UActorComponent* ActorComponent : GetComponents())
-		{
-			UHierarchicalInstancedStaticMeshComponent* HComponent = Cast<UHierarchicalInstancedStaticMeshComponent>(ActorComponent);
-			if (HComponent && !StillUsed.Contains(HComponent))
-			{
-				{
-					SCOPE_CYCLE_COUNTER(STAT_FoliageGrassDestoryComp);
-					HComponent->ClearInstances();
-					HComponent->DestroyComponent();
-					FoliageComponents.Remove(HComponent);
-				}
-				if (!bForceSync)
-				{
-					break; // one per frame is fine
-				}
-			}
-		}
-	}
-	{
-		// finsh async tasks
-		for (int32 Index = 0; Index < AsyncFoliageTasks.Num(); Index++)
-		{
-			FAsyncTask<FAsyncGrassTask>* Task = AsyncFoliageTasks[Index];
-			if (bForceSync)
-			{
-				Task->EnsureCompletion();
-			}
-			if (Task->IsDone())
-			{
-				SCOPE_CYCLE_COUNTER(STAT_FoliageGrassEndComp);
-				FAsyncGrassTask& Inner = Task->GetTask();
-				AsyncFoliageTasks.RemoveAtSwap(Index--);
-				UHierarchicalInstancedStaticMeshComponent* HierarchicalInstancedStaticMeshComponent = Inner.Foliage.Get();
-				if (HierarchicalInstancedStaticMeshComponent && StillUsed.Contains(HierarchicalInstancedStaticMeshComponent))
-				{
-					if (Inner.Builder->InstanceBuffer.Num())
-					{
-						Exchange(HierarchicalInstancedStaticMeshComponent->WriteOncePrebuiltInstanceBuffer, Inner.Builder->InstanceBuffer);
-						HierarchicalInstancedStaticMeshComponent->AcceptPrebuiltTree(Inner.Builder->ClusterTree);
-						if (bForceSync && GetWorld())
-						{
-							HierarchicalInstancedStaticMeshComponent->RecreateRenderState_Concurrent();
-						}
-					}
-				}
-				FCachedLandscapeFoliage::FGrassComp* Existing = FoliageCache.CachedGrassComps.Find(Inner.Key);
-				if (Existing)
-				{
-					Existing->Pending = false;
-					Existing->Touch();
-				}
-				delete Inner.Builder;
-				delete Task;
-				if (!bForceSync)
-				{
-					break; // one per frame is fine
-				}
-			}
-		}
+		TickGrass();
 	}
 }
 
@@ -3136,7 +2156,6 @@ ALandscapeProxy::~ALandscapeProxy()
 		FAsyncTask<FAsyncGrassTask>* Task = AsyncFoliageTasks[Index];
 		Task->EnsureCompletion(true);
 		FAsyncGrassTask& Inner = Task->GetTask();
-		delete Inner.Builder;
 		delete Task;
 	}
 	AsyncFoliageTasks.Empty();
@@ -3144,33 +2163,10 @@ ALandscapeProxy::~ALandscapeProxy()
 #if WITH_EDITOR
 	TotalComponentsNeedingGrassMapRender -= NumComponentsNeedingGrassMapRender;
 	NumComponentsNeedingGrassMapRender = 0;
+	TotalTexturesToStreamForVisibleGrassMapRender -= NumTexturesToStreamForVisibleGrassMapRender;
+	NumTexturesToStreamForVisibleGrassMapRender = 0;
 #endif
 }
-
-void FAsyncGrassTask::DoWork()
-{
-	Builder->Build();
-}
-
-static void FlushGrass(const TArray<FString>& Args)
-{
-	for (TObjectIterator<ALandscapeProxy> It; It; ++It)
-	{
-		ALandscapeProxy* Landscape = *It;
-		if (Landscape && !Landscape->IsTemplate() && !Landscape->IsPendingKill())
-		{
-			Landscape->FlushGrassComponents();
-		}
-	}
-}
-
-static FAutoConsoleCommand FlushGrassCmd(
-	TEXT("grass.FlushCache"),
-	TEXT("Flush the grass cache, debugging."),
-	FConsoleCommandWithArgsDelegate::CreateStatic(&FlushGrass)
-	);
-
-
 //
 // ALandscapeMeshProxyActor
 //
diff --git a/Engine/Source/Runtime/Landscape/Private/LandscapeGrass.cpp b/Engine/Source/Runtime/Landscape/Private/LandscapeGrass.cpp
--- a/Engine/Source/Runtime/Landscape/Private/LandscapeGrass.cpp
+++ b/Engine/Source/Runtime/Landscape/Private/LandscapeGrass.cpp
@@ -14,12 +14,85 @@
 #include "GenericOctreePublic.h"
 #include "PrimitiveSceneInfo.h"
 #include "MaterialCompiler.h"
+#include "LandscapeGrassType.h"
 #include "Materials/MaterialExpressionLandscapeGrassOutput.h"
 #include "EdGraph/EdGraphNode.h"
 #include "Engine/TextureRenderTarget2D.h"
+#include "ContentStreaming.h"
 
 #define LOCTEXT_NAMESPACE "Landscape"
 
+DEFINE_LOG_CATEGORY_STATIC(LogGrass, Log, All);
+
+static TAutoConsoleVariable<float> CVarGuardBandMultiplier(
+	TEXT("grass.GuardBandMultiplier"),
+	1.3f,
+	TEXT("Used to control discarding in the grass system. Approximate range, 1-4. Multiplied by the cull distance to control when we add grass components."));
+
+static TAutoConsoleVariable<float> CVarGuardBandDiscardMultiplier(
+	TEXT("grass.GuardBandDiscardMultiplier"),
+	1.4f,
+	TEXT("Used to control discarding in the grass system. Approximate range, 1-4. Multiplied by the cull distance to control when we discard grass components."));
+
+static TAutoConsoleVariable<int32> CVarMinFramesToKeepGrass(
+	TEXT("grass.MinFramesToKeepGrass"),
+	30,
+	TEXT("Minimum number of frames before cached grass can be discarded; used to prevent thrashing."));
+
+static TAutoConsoleVariable<float> CVarMinTimeToKeepGrass(
+	TEXT("grass.MinTimeToKeepGrass"),
+	5.0f,
+	TEXT("Minimum number of seconds before cached grass can be discarded; used to prevent thrashing."));
+
+static TAutoConsoleVariable<int32> CVarMinInstancesPerLeaf(
+	TEXT("grass.MinInstancesPerLeaf"),
+	256,
+	TEXT("Used to control the depth of hierarchical culling tree for grass. Clamped to 16-32000; a reasonable range."));
+
+static TAutoConsoleVariable<int32> CVarMaxInstancesPerComponent(
+	TEXT("grass.MaxInstancesPerComponent"),
+	65536,
+	TEXT("Used to control the number of hierarchical components created. More can be more efficient, but can be hitchy as new components come into range"));
+
+static TAutoConsoleVariable<int32> CVarMaxAsyncTasks(
+	TEXT("grass.MaxAsyncTasks"),
+	4,
+	TEXT("Used to control the number of hierarchical components created at a time."));
+
+static TAutoConsoleVariable<int32> CVarGrassEnable(
+	TEXT("grass.Enable"),
+	1,
+	TEXT("1: Enable Grass; 0: Disable Grass"));
+
+static TAutoConsoleVariable<int32> CVarUseStreamingManagerForCameras(
+	TEXT("grass.UseStreamingManagerForCameras"),
+	1,
+	TEXT("1: Use Streaming Manager; 0: Use ViewLocationsRenderedLastFrame"));
+
+static TAutoConsoleVariable<int32> CVarCullSubsections(
+	TEXT("grass.CullSubsections"),
+	1,
+	TEXT("1: Cull each foliage component; 0: Cull only based on the landscape component."));
+
+static TAutoConsoleVariable<int32> CVarDisableGPUCull(
+	TEXT("grass.DisableGPUCull"),
+	0,
+	TEXT("For debugging. Set this to zero to see where the grass is generated. Useful for tweaking the guard bands."));
+
+static TAutoConsoleVariable<int32> CVarPrerenderGrassmaps(
+	TEXT("grass.PrerenderGrassmaps"),
+	1,
+	TEXT("1: Pre-render grass maps for all components in the editor; 0: Generate grass maps on demand while moving through the editor"));
+
+DECLARE_CYCLE_STAT(TEXT("Grass Async Build Time"), STAT_FoliageGrassAsyncBuildTime, STATGROUP_Foliage);
+DECLARE_CYCLE_STAT(TEXT("Grass Start Comp"), STAT_FoliageGrassStartComp, STATGROUP_Foliage);
+DECLARE_CYCLE_STAT(TEXT("Grass End Comp"), STAT_FoliageGrassEndComp, STATGROUP_Foliage);
+DECLARE_CYCLE_STAT(TEXT("Grass Destroy Comps"), STAT_FoliageGrassDestoryComp, STATGROUP_Foliage);
+DECLARE_CYCLE_STAT(TEXT("Grass Update"), STAT_GrassUpdate, STATGROUP_Foliage);
+
+//
+// Grass weightmap rendering
+//
 class FLandscapeGrassWeightVS : public FMeshMaterialShader
 {
 	DECLARE_SHADER_TYPE(FLandscapeGrassWeightVS, MeshMaterial);
@@ -66,7 +139,6 @@ public:
 
 IMPLEMENT_MATERIAL_SHADER_TYPE(, FLandscapeGrassWeightVS, TEXT("LandscapeGrassWeight"), TEXT("VSMain"), SF_Vertex);
 
-
 class FLandscapeGrassWeightPS : public FMeshMaterialShader
 {
 	DECLARE_SHADER_TYPE(FLandscapeGrassWeightPS, MeshMaterial);
@@ -113,9 +185,8 @@ public:
 IMPLEMENT_MATERIAL_SHADER_TYPE(, FLandscapeGrassWeightPS, TEXT("LandscapeGrassWeight"), TEXT("PSMain"), SF_Pixel);
 
 
-
 /**
-* Used to write out depth for opaque and masked materials during the depth-only pass.
+* Drawing policy used to write out landscape grass weightmap.
 */
 class FLandscapeGrassWeightDrawingPolicy : public FMeshDrawingPolicy
 {
@@ -276,8 +347,6 @@ public:
 	}
 };
 
-
-
 class FLandscapeGrassWeightExporter : public FLandscapeGrassWeightExporter_RenderThread
 {
 	ALandscapeProxy* LandscapeProxy;
@@ -547,7 +616,6 @@ void ALandscapeProxy::RenderGrassMaps(const TArray<ULandscapeComponent*>& Landsc
 
 #endif
 
-
 //
 // UMaterialExpressionLandscapeGrassOutput
 //
@@ -722,4 +790,946 @@ FArchive& operator<<(FArchive& Ar, FLandscapeComponentGrassData& Data)
 	return Ar << Data.WeightData;
 }
 
+//
+// ALandscapeProxy grass-related functions
+//
+
+void ALandscapeProxy::TickGrass()
+{
+	// Update foliage
+	static TArray<FVector> OldCameras;
+	if (CVarUseStreamingManagerForCameras.GetValueOnGameThread() == 0)
+	{
+		UWorld* World = GetWorld();
+		if (!World)
+		{
+			return;
+		}
+
+		if (!OldCameras.Num() && !World->ViewLocationsRenderedLastFrame.Num())
+		{
+			// no cameras, no grass update
+			return;
+		}
+
+		// there is a bug here, which often leaves us with no cameras in the editor
+		const TArray<FVector>& Cameras = World->ViewLocationsRenderedLastFrame.Num() ? World->ViewLocationsRenderedLastFrame : OldCameras;
+
+		if (&Cameras != &OldCameras)
+		{
+			check(IsInGameThread());
+			OldCameras = Cameras;
+		}
+		UpdateGrass(Cameras);
+	}
+	else
+	{
+		int32 Num = IStreamingManager::Get().GetNumViews();
+		if (!Num)
+		{
+			// no cameras, no grass update
+			return;
+		}
+		OldCameras.Reset(Num);
+		for (int32 Index = 0; Index < Num; Index++)
+		{
+			auto& ViewInfo = IStreamingManager::Get().GetViewInformation(Index);
+			OldCameras.Add(ViewInfo.ViewOrigin);
+		}
+		UpdateGrass(OldCameras);
+	}
+}
+
+struct FGrassBuilderBase
+{
+	bool bHaveValidData;
+	float GrassDensity;
+	FVector DrawScale;
+	FVector DrawLoc;
+	FMatrix LandscapeToWorld;
+
+	FIntPoint SectionBase;
+	FIntPoint LandscapeSectionOffset;
+	int32 ComponentSizeQuads;
+	FVector Origin;
+	FVector Extent;
+	int32 SqrtMaxInstances;
+
+	FGrassBuilderBase(ALandscapeProxy* Landscape, ULandscapeComponent* Component, const FGrassVariety& GrassVariety, int32 SqrtSubsections = 1, int32 SubX = 0, int32 SubY = 0)
+	{
+		bHaveValidData = true;
+		GrassDensity = GrassVariety.GrassDensity;
+
+		DrawScale = Landscape->GetRootComponent()->RelativeScale3D;
+		DrawLoc = Landscape->GetActorLocation();
+		LandscapeSectionOffset = Landscape->LandscapeSectionOffset;
+
+		SectionBase = Component->GetSectionBase();
+		ComponentSizeQuads = Component->ComponentSizeQuads;
+
+		Origin = FVector(DrawScale.X * float(SectionBase.X), DrawScale.Y * float(SectionBase.Y), 0.0f);
+		Extent = FVector(DrawScale.X * float(SectionBase.X + ComponentSizeQuads), DrawScale.Y * float(SectionBase.Y + ComponentSizeQuads), 0.0f) - Origin;
+
+		SqrtMaxInstances = FMath::CeilToInt(FMath::Sqrt(FMath::Abs(Extent.X * Extent.Y * GrassDensity / 1000.0f / 1000.0f)));
+
+		if (!SqrtMaxInstances)
+		{
+			bHaveValidData = false;
+		}
+		const FRotator DrawRot = Landscape->GetActorRotation();
+		LandscapeToWorld = Landscape->GetRootComponent()->ComponentToWorld.ToMatrixNoScale();
+
+		if (bHaveValidData && SqrtSubsections != 1)
+		{
+			check(SqrtMaxInstances > 2 * SqrtSubsections);
+			SqrtMaxInstances /= SqrtSubsections;
+			check(SqrtMaxInstances > 0);
+
+			Extent /= float(SqrtSubsections);
+			Origin += Extent * FVector(float(SubX), float(SubY), 0.0f);
+		}
+	}
+};
+
+// FLandscapeComponentGrassAccess - accessor wrapper for data for one GrassType from one Component
+struct FLandscapeComponentGrassAccess
+{
+	FLandscapeComponentGrassAccess(const ULandscapeComponent* InComponent, const ULandscapeGrassType* GrassType)
+	: GrassData(InComponent->GrassData)
+	, HeightData(InComponent->GrassData->HeightData)
+	, WeightData(InComponent->GrassData->WeightData.Find(GrassType))
+	, Stride(InComponent->ComponentSizeQuads + 1)
+	{}
+
+	bool IsValid()
+	{
+		return WeightData && WeightData->Num() == FMath::Square(Stride) && HeightData.Num() == FMath::Square(Stride);
+	}
+
+	FORCEINLINE float GetHeight(int32 IdxX, int32 IdxY)
+	{
+		return LandscapeDataAccess::GetLocalHeight(HeightData[IdxX + Stride*IdxY]);
+	}
+	FORCEINLINE float GetWeight(int32 IdxX, int32 IdxY)
+	{
+		return ((float)(*WeightData)[IdxX + Stride*IdxY]) / 255.f;
+	}
+
+	FORCEINLINE int32 GetStride()
+	{
+		return Stride;
+	}
+
+private:
+	TSharedRef<FLandscapeComponentGrassData, ESPMode::ThreadSafe> GrassData;
+	TArray<uint16>& HeightData;
+	TArray<uint8>* WeightData;
+	int32 Stride;
+};
+
+struct FAsyncGrassBuilder : public FGrassBuilderBase
+{
+	FLandscapeComponentGrassAccess GrassData;
+	bool RandomRotation;
+	bool AlignToSurface;
+	float PlacementJitter;
+	FRandomStream RandomStream;
+	FMatrix XForm;
+	FBox MeshBox;
+
+	double RasterTime;
+	double BuildTime;
+	double InstanceTime;
+	int32 TotalInstances;
+
+	// output
+	FStaticMeshInstanceData InstanceBuffer;
+	TArray<FClusterNode> ClusterTree;
+
+	FAsyncGrassBuilder(ALandscapeProxy* Landscape, ULandscapeComponent* Component, const ULandscapeGrassType* GrassType, const FGrassVariety& GrassVariety, UHierarchicalInstancedStaticMeshComponent* HierarchicalInstancedStaticMeshComponent, int32 SqrtSubsections, int32 SubX, int32 SubY)
+		: FGrassBuilderBase(Landscape, Component, GrassVariety, SqrtSubsections, SubX, SubY)
+		, GrassData(Component, GrassType)
+	{
+		bHaveValidData = bHaveValidData && GrassData.IsValid();
+		XForm = LandscapeToWorld * HierarchicalInstancedStaticMeshComponent->ComponentToWorld.ToMatrixWithScale().Inverse();
+		PlacementJitter = GrassVariety.PlacementJitter;
+		RandomRotation = GrassVariety.RandomRotation;
+		AlignToSurface = GrassVariety.AlignToSurface;
+
+		RandomStream.Initialize(HierarchicalInstancedStaticMeshComponent->InstancingRandomSeed);
+
+		MeshBox = GrassVariety.GrassMesh->GetBounds().GetBox();
+
+		RasterTime = 0.0;
+		BuildTime = 0.0;
+		InstanceTime = 0.0;
+		TotalInstances = 0;
+	}
+
+	void Build()
+	{
+		SCOPE_CYCLE_COUNTER(STAT_FoliageGrassAsyncBuildTime);
+		check(bHaveValidData);
+		RasterTime -= FPlatformTime::Seconds();
+
+		float Div = 1.0f / float(SqrtMaxInstances);
+		float MaxJitter1D = FMath::Clamp<float>(PlacementJitter, 0.0f, .99f) * Div * .5f;
+		FVector MaxJitter(MaxJitter1D, MaxJitter1D, 0.0f);
+		MaxJitter *= Extent;
+		Origin += Extent * (Div * 0.5f);
+		struct FInstanceLocal
+		{
+			FVector Pos;
+			bool bKeep;
+		};
+		TArray<FInstanceLocal> Instances;
+		Instances.AddUninitialized(SqrtMaxInstances * SqrtMaxInstances);
+		int32 NumKept = 0;
+		{
+			int32 InstanceIndex = 0;
+			for (int32 xStart = 0; xStart < SqrtMaxInstances; xStart++)
+			{
+				for (int32 yStart = 0; yStart < SqrtMaxInstances; yStart++)
+				{
+					FVector Location(Origin.X + float(xStart) * Div * Extent.X, Origin.Y + float(yStart) * Div * Extent.Y, 0.0f);
+					Location += FVector(RandomStream.GetFraction() * 2.0f - 1.0f, RandomStream.GetFraction() * 2.0f - 1.0f, 0.0f) * MaxJitter;
+
+					FInstanceLocal& Instance = Instances[InstanceIndex];
+					float Weight = GetLayerWeightAtLocationLocal(Location, &Instance.Pos);
+					Instance.bKeep = Weight > 0.0f && Weight >= RandomStream.GetFraction();
+					if (Instance.bKeep)
+					{
+						NumKept++;
+					}
+					InstanceIndex++;
+				}
+			}
+		}
+		RasterTime += FPlatformTime::Seconds();
+
+		if (NumKept)
+		{
+			TArray<FMatrix> InstanceTransforms;
+			InstanceTransforms.AddUninitialized(NumKept);
+			TotalInstances += NumKept;
+
+			{
+				InstanceTime -= FPlatformTime::Seconds();
+				InstanceBuffer.AllocateInstances(NumKept);
+				int32 InstanceIndex = 0;
+				int32 OutInstanceIndex = 0;
+				for (int32 xStart = 0; xStart < SqrtMaxInstances; xStart++)
+				{
+					for (int32 yStart = 0; yStart < SqrtMaxInstances; yStart++)
+					{
+						const FInstanceLocal& Instance = Instances[InstanceIndex];
+						if (Instance.bKeep)
+						{
+							float Rot = RandomRotation ? RandomStream.GetFraction() * 360.0f : 0.0f;
+							FMatrix OutXForm;
+							if (AlignToSurface)
+							{
+								FVector PosX1 = xStart ? Instances[InstanceIndex - SqrtMaxInstances].Pos : Instance.Pos;
+								FVector PosX2 = (xStart + 1 < SqrtMaxInstances) ? Instances[InstanceIndex + SqrtMaxInstances].Pos : Instance.Pos;
+								FVector PosY1 = yStart ? Instances[InstanceIndex - 1].Pos : Instance.Pos;
+								FVector PosY2 = (yStart + 1 < SqrtMaxInstances) ? Instances[InstanceIndex + 1].Pos : Instance.Pos;
+
+								if (PosX1 != PosX2 && PosY1 != PosY2)
+								{
+									FVector NewZ = ((PosX1 - PosX2) ^ (PosY1 - PosY2)).GetSafeNormal();
+									NewZ *= FMath::Sign(NewZ.Z);
+
+									const FVector NewX = (FVector(0, -1, 0) ^ NewZ).GetSafeNormal();
+									const FVector NewY = NewZ ^ NewX;
+
+									FMatrix Align = FMatrix(NewX, NewY, NewZ, FVector::ZeroVector);
+									OutXForm = FRotationMatrix(FRotator(0.0f, Rot, 0.0f)) * Align * FTranslationMatrix(Instance.Pos) * XForm;
+								}
+								else
+								{
+									OutXForm = FRotationMatrix(FRotator(0.0f, Rot, 0.0f)) * FTranslationMatrix(Instance.Pos) * XForm;
+								}
+							}
+							else
+							{
+								OutXForm = FRotationMatrix(FRotator(0.0f, Rot, 0.0f)) * FTranslationMatrix(Instance.Pos) * XForm;
+							}
+							InstanceTransforms[OutInstanceIndex] = OutXForm;
+							FInstanceStream* RenderData = InstanceBuffer.GetInstanceWriteAddress(OutInstanceIndex++);
+							RenderData->SetInstance(OutXForm, RandomStream.GetFraction());
+						}
+						InstanceIndex++;
+					}
+				}
+				InstanceTime += FPlatformTime::Seconds();
+			}
+			{
+				BuildTime -= FPlatformTime::Seconds();
+				TArray<int32> SortedInstances;
+				TArray<int32> InstanceReorderTable;
+				UHierarchicalInstancedStaticMeshComponent::BuildTreeAnyThread(InstanceTransforms, MeshBox, ClusterTree, SortedInstances, InstanceReorderTable, FMath::Clamp<int32>(CVarMinInstancesPerLeaf.GetValueOnAnyThread(), 16, 256 * 128));
+
+				// in-place sort the instances
+				FInstanceStream SwapBuffer;
+				int32 FirstUnfixedIndex = 0;
+				for (int32 FirstUnfixedIndex = 0; FirstUnfixedIndex < NumKept; FirstUnfixedIndex++)
+				{
+					int32 LoadFrom = SortedInstances[FirstUnfixedIndex];
+					if (LoadFrom != FirstUnfixedIndex)
+					{
+						check(LoadFrom > FirstUnfixedIndex);
+						FMemory::Memcpy(&SwapBuffer, InstanceBuffer.GetInstanceWriteAddress(FirstUnfixedIndex), sizeof(FInstanceStream));
+						FMemory::Memcpy(InstanceBuffer.GetInstanceWriteAddress(FirstUnfixedIndex), InstanceBuffer.GetInstanceWriteAddress(LoadFrom), sizeof(FInstanceStream));
+						FMemory::Memcpy(InstanceBuffer.GetInstanceWriteAddress(LoadFrom), &SwapBuffer, sizeof(FInstanceStream));
+
+						int32 SwapGoesTo = InstanceReorderTable[FirstUnfixedIndex];
+						check(SwapGoesTo > FirstUnfixedIndex);
+						check(SortedInstances[SwapGoesTo] == FirstUnfixedIndex);
+						SortedInstances[SwapGoesTo] = LoadFrom;
+						InstanceReorderTable[LoadFrom] = SwapGoesTo;
+
+						InstanceReorderTable[FirstUnfixedIndex] = FirstUnfixedIndex;
+						SortedInstances[FirstUnfixedIndex] = FirstUnfixedIndex;
+					}
+				}
+
+				BuildTime += FPlatformTime::Seconds();
+			}
+		}
+	}
+	FORCEINLINE_DEBUGGABLE float GetLayerWeightAtLocationLocal(const FVector& InLocation, FVector* OutLocation)
+	{
+		// Find location
+		float TestX = InLocation.X / DrawScale.X - (float)SectionBase.X;
+		float TestY = InLocation.Y / DrawScale.Y - (float)SectionBase.Y;
+
+		// Find data
+		int32 X1 = FMath::FloorToInt(TestX);
+		int32 Y1 = FMath::FloorToInt(TestY);
+		int32 X2 = FMath::CeilToInt(TestX);
+		int32 Y2 = FMath::CeilToInt(TestY);
+
+		// Min is to prevent the sampling of the final column from overflowing
+		int32 IdxX1 = FMath::Min<int32>(X1, GrassData.GetStride() - 1);
+		int32 IdxY1 = FMath::Min<int32>(Y1, GrassData.GetStride() - 1);
+		int32 IdxX2 = FMath::Min<int32>(X2, GrassData.GetStride() - 1);
+		int32 IdxY2 = FMath::Min<int32>(Y2, GrassData.GetStride() - 1);
+
+		float LerpX = FMath::Fractional(TestX);
+		float LerpY = FMath::Fractional(TestY);
+
+		float Result;
+		{
+			// sample
+			float Sample11 = GrassData.GetWeight(IdxX1, IdxY1);
+			float Sample21 = GrassData.GetWeight(IdxX2, IdxY1);
+			float Sample12 = GrassData.GetWeight(IdxX1, IdxY2);
+			float Sample22 = GrassData.GetWeight(IdxX2, IdxY2);
+
+			// Bilinear interpolate
+			Result = FMath::Lerp(
+				FMath::Lerp(Sample11, Sample21, LerpX),
+				FMath::Lerp(Sample12, Sample22, LerpX),
+				LerpY);
+		}
+
+		{
+			// sample
+			float Sample11 = GrassData.GetHeight(IdxX1, IdxY1);
+			float Sample21 = GrassData.GetHeight(IdxX2, IdxY1);
+			float Sample12 = GrassData.GetHeight(IdxX1, IdxY2);
+			float Sample22 = GrassData.GetHeight(IdxX2, IdxY2);
+
+			OutLocation->X = InLocation.X - DrawScale.X * float(LandscapeSectionOffset.X);
+			OutLocation->Y = InLocation.Y - DrawScale.Y * float(LandscapeSectionOffset.Y);
+			// Bilinear interpolate
+			OutLocation->Z = DrawScale.Z * FMath::Lerp(
+				FMath::Lerp(Sample11, Sample21, LerpX),
+				FMath::Lerp(Sample12, Sample22, LerpX),
+				LerpY);
+		}
+		return Result;
+	}
+};
+
+void ALandscapeProxy::FlushGrassComponents(const TSet<ULandscapeComponent*>* OnlyForComponents, bool bFlushGrassMaps)
+{
+	if (OnlyForComponents)
+	{
+		for (FCachedLandscapeFoliage::TGrassSet::TIterator Iter(FoliageCache.CachedGrassComps); Iter; ++Iter)
+		{
+			ULandscapeComponent* Component = (*Iter).Key.BasedOn.Get();
+			// if the weak pointer in the cache is invalid, we should kill them anyway
+			if (Component == nullptr || OnlyForComponents->Contains(Component))
+			{
+				UHierarchicalInstancedStaticMeshComponent *Used = (*Iter).Foliage.Get();
+				if (Used)
+				{
+					SCOPE_CYCLE_COUNTER(STAT_FoliageGrassDestoryComp);
+					Used->ClearInstances();
+					Used->DestroyComponent();
+				}
+				Iter.RemoveCurrent();
+			}
+		}
+#if WITH_EDITOR
+		if (GIsEditor && bFlushGrassMaps)
+		{
+			for (ULandscapeComponent* Component : *OnlyForComponents)
+			{
+				Component->RemoveGrassMap();
+			}
+		}
+#endif
+	}
+	else
+	{
+		// Clear old foliage component containers
+		FoliageComponents.Empty();
+
+		// Might as well clear the cache...
+		FoliageCache.ClearCache();
+		// Destroy any owned foliage components
+		TInlineComponentArray<UHierarchicalInstancedStaticMeshComponent*> FoliageComps;
+		GetComponents(FoliageComps);
+		for (UHierarchicalInstancedStaticMeshComponent* Component : FoliageComps)
+		{
+			SCOPE_CYCLE_COUNTER(STAT_FoliageGrassDestoryComp);
+			Component->ClearInstances();
+			Component->DestroyComponent();
+		}
+
+		TArray<USceneComponent*> AttachedFoliageComponents = RootComponent->AttachChildren.FilterByPredicate(
+			[](USceneComponent* Component)
+		{
+			return Cast<UHierarchicalInstancedStaticMeshComponent>(Component);
+		});
+
+		// Destroy any attached but un-owned foliage components
+		for (USceneComponent* Component : AttachedFoliageComponents)
+		{
+			SCOPE_CYCLE_COUNTER(STAT_FoliageGrassDestoryComp);
+			Component->DestroyComponent();
+		}
+
+#if WITH_EDITOR
+		if (GIsEditor && bFlushGrassMaps)
+		{
+			// Clear GrassMaps
+			TInlineComponentArray<ULandscapeComponent*> LandComps;
+			GetComponents(LandComps);
+			for (ULandscapeComponent* Component : LandComps)
+			{
+				Component->RemoveGrassMap();
+			}
+		}
+#endif
+	}
+}
+
+TArray<ULandscapeGrassType*> ALandscapeProxy::GetGrassTypes() const
+{
+	TArray<ULandscapeGrassType*> GrassTypes;
+	if (LandscapeMaterial)
+	{
+		TArray<const UMaterialExpressionLandscapeGrassOutput*> GrassExpressions;
+		LandscapeMaterial->GetMaterial()->GetAllExpressionsOfType<UMaterialExpressionLandscapeGrassOutput>(GrassExpressions);
+		if (GrassExpressions.Num() > 0)
+		{
+			for (auto& Type : GrassExpressions[0]->GrassTypes)
+			{
+				GrassTypes.Add(Type.GrassType);
+			}
+		}
+	}
+	return GrassTypes;
+}
+
+#if WITH_EDITOR
+int32 ALandscapeProxy::TotalComponentsNeedingGrassMapRender = 0;
+int32 ALandscapeProxy::TotalTexturesToStreamForVisibleGrassMapRender = 0;
+#endif
+
+void ALandscapeProxy::UpdateGrass(const TArray<FVector>& Cameras, bool bForceSync)
+{
+	SCOPE_CYCLE_COUNTER(STAT_GrassUpdate);
+
+	if (CVarGrassEnable.GetValueOnAnyThread() > 0)
+	{
+		TArray<ULandscapeGrassType*> GrassTypes = GetGrassTypes();
+
+		float GuardBand = CVarGuardBandMultiplier.GetValueOnAnyThread();
+		float DiscardGuardBand = CVarGuardBandDiscardMultiplier.GetValueOnAnyThread();
+		bool bCullSubsections = CVarCullSubsections.GetValueOnAnyThread() > 0;
+		bool bDisableGPUCull = CVarDisableGPUCull.GetValueOnAnyThread() > 0;
+		int32 MaxInstancesPerComponent = FMath::Max<int32>(1024, CVarMaxInstancesPerComponent.GetValueOnAnyThread());
+		int32 MaxTasks = CVarMaxAsyncTasks.GetValueOnAnyThread();
+
+		UWorld* World = GetWorld();
+		if (World)
+		{
+#if WITH_EDITOR
+			int32 RequiredTexturesNotStreamedIn = 0;
+			TSet<ULandscapeComponent*> ComponentsNeedingGrassMapRender;
+			TSet<UTexture2D*> CurrentForcedStreamedTextures;
+			TSet<UTexture2D*> DesiredForceStreamedTextures;
+
+			if (!World->IsGameWorld())
+			{
+				// see if we need to flush grass for any components
+				TSet<ULandscapeComponent*> FlushComponents;
+				for (auto Component : LandscapeComponents)
+				{
+					// check textures currently needing force streaming
+					if (Component->HeightmapTexture->bForceMiplevelsToBeResident)
+					{
+						CurrentForcedStreamedTextures.Add(Component->HeightmapTexture);
+					}
+					for (auto WeightmapTexture : Component->WeightmapTextures)
+					{
+						if (WeightmapTexture->bForceMiplevelsToBeResident)
+						{
+							CurrentForcedStreamedTextures.Add(WeightmapTexture);
+						}
+					}
+
+					if (Component->IsGrassMapOutdated())
+					{
+						FlushComponents.Add(Component);
+						if (GrassTypes.Num() > 0)
+						{
+							ComponentsNeedingGrassMapRender.Add(Component);
+						}
+					}
+					else
+					if (!Component->GrassData->HasData())
+					{
+						if (GrassTypes.Num() > 0)
+						{
+							ComponentsNeedingGrassMapRender.Add(Component);
+						}
+					}
+				}
+				if (FlushComponents.Num())
+				{
+					FlushGrassComponents(&FlushComponents);
+				}
+			}
+#endif
+
+			int32 NumCompsCreated = 0;
+			for (int32 ComponentIndex = 0; ComponentIndex < LandscapeComponents.Num(); ComponentIndex++)
+			{
+				ULandscapeComponent* Component = LandscapeComponents[ComponentIndex];
+
+				// skip if we have no data and no way to generate it
+				if (World->IsGameWorld() && !Component->GrassData->HasData())
+				{
+					continue;
+				}
+
+				FBoxSphereBounds WorldBounds = Component->CalcBounds(Component->ComponentToWorld);
+				float MinDistanceToComp = Cameras.Num() ? MAX_flt : 0.0f;
+
+				for (auto& Pos : Cameras)
+				{
+					MinDistanceToComp = FMath::Min<float>(MinDistanceToComp, WorldBounds.ComputeSquaredDistanceFromBoxToPoint(Pos));
+				}
+
+				MinDistanceToComp = FMath::Sqrt(MinDistanceToComp);
+
+				for (auto GrassType : GrassTypes)
+				{
+					if (GrassType)
+					{
+						int32 GrassVarietyIndex = -1;
+						for (auto& GrassVariety : GrassType->GrassVarieties)
+						{
+							GrassVarietyIndex++;
+							if (GrassVariety.GrassMesh && GrassVariety.GrassDensity > 0.0f && GrassVariety.EndCullDistance > 0)
+							{
+								bool bRandomRotation = GrassVariety.RandomRotation;
+								bool bAlign = GrassVariety.AlignToSurface;
+								float MustHaveDistance = GuardBand * (float)GrassVariety.EndCullDistance;
+								float DiscardDistance = DiscardGuardBand * (float)GrassVariety.EndCullDistance;
+
+								if (MinDistanceToComp > DiscardDistance)
+								{
+									continue;
+								}
+
+								FGrassBuilderBase ForSubsectionMath(this, Component, GrassVariety);
+
+								int32 SqrtSubsections = 1;
+
+								if (ForSubsectionMath.bHaveValidData && ForSubsectionMath.SqrtMaxInstances > 0)
+								{
+									SqrtSubsections = FMath::Clamp<int32>(FMath::CeilToInt(float(ForSubsectionMath.SqrtMaxInstances) / FMath::Sqrt((float)MaxInstancesPerComponent)), 1, 16);
+								}
+
+								FBox LocalBox = Component->CachedLocalBox;
+								FVector LocalExtentDiv = (LocalBox.Max - LocalBox.Min) * FVector(1.0f / float(SqrtSubsections), 1.0f / float(SqrtSubsections), 1.0f);
+								for (int32 SubX = 0; SubX < SqrtSubsections; SubX++)
+								{
+									for (int32 SubY = 0; SubY < SqrtSubsections; SubY++)
+									{
+										float MinDistanceToSubComp = MinDistanceToComp;
+
+										if (bCullSubsections && SqrtSubsections > 1)
+										{
+											FVector BoxMin;
+											BoxMin.X = LocalBox.Min.X + LocalExtentDiv.X * float(SubX);
+											BoxMin.Y = LocalBox.Min.Y + LocalExtentDiv.Y * float(SubY);
+											BoxMin.Z = LocalBox.Min.Z;
+
+											FVector BoxMax;
+											BoxMax.X = LocalBox.Min.X + LocalExtentDiv.X * float(SubX + 1);
+											BoxMax.Y = LocalBox.Min.Y + LocalExtentDiv.Y * float(SubY + 1);
+											BoxMax.Z = LocalBox.Max.Z;
+
+											FBox LocalSubBox(BoxMin, BoxMax);
+											FBox WorldSubBox = LocalSubBox.TransformBy(Component->ComponentToWorld);
+
+											MinDistanceToSubComp = Cameras.Num() ? MAX_flt : 0.0f;
+											for (auto& Pos : Cameras)
+											{
+												MinDistanceToSubComp = FMath::Min<float>(MinDistanceToSubComp, ComputeSquaredDistanceFromBoxToPoint(WorldSubBox.Min, WorldSubBox.Max, Pos));
+											}
+											MinDistanceToSubComp = FMath::Sqrt(MinDistanceToSubComp);
+										}
+
+										if (MinDistanceToSubComp > DiscardDistance)
+										{
+											continue;
+										}
+
+										FCachedLandscapeFoliage::FGrassComp NewComp;
+										NewComp.Key.BasedOn = Component;
+										NewComp.Key.GrassType = GrassType;
+										NewComp.Key.SqrtSubsections = SqrtSubsections;
+										NewComp.Key.CachedMaxInstancesPerComponent = MaxInstancesPerComponent;
+										NewComp.Key.SubsectionX = SubX;
+										NewComp.Key.SubsectionY = SubY;
+										NewComp.Key.NumVarieties = GrassType->GrassVarieties.Num();
+										NewComp.Key.VarietyIndex = GrassVarietyIndex;
+
+										{
+											FCachedLandscapeFoliage::FGrassComp* Existing = FoliageCache.CachedGrassComps.Find(NewComp.Key);
+											if (Existing || MinDistanceToSubComp > MustHaveDistance)
+											{
+												if (Existing)
+												{
+													Existing->Touch();
+												}
+												continue;
+											}
+										}
+
+										if (NumCompsCreated || (!bForceSync && AsyncFoliageTasks.Num() >= MaxTasks))
+										{
+											continue; // one per frame, but we still want to touch the existing ones
+										}
+
+#if WITH_EDITOR
+										// render grass data if we don't have any
+										if (!Component->GrassData->HasData())
+										{
+											if (!Component->CanRenderGrassMap())
+											{
+												// we can't currently render grassmaps (eg shaders not compiled)
+												continue;
+											}
+											else if (!Component->AreTexturesStreamedForGrassMapRender())
+											{
+												// we're ready to generate but our textures need streaming in
+												DesiredForceStreamedTextures.Add(Component->HeightmapTexture);
+												for (auto WeightmapTexture : Component->WeightmapTextures)
+												{
+													DesiredForceStreamedTextures.Add(WeightmapTexture);
+												}
+												RequiredTexturesNotStreamedIn++;
+												continue;
+											}
+
+											QUICK_SCOPE_CYCLE_COUNTER(STAT_GrassRenderToTexture);
+											Component->RenderGrassMap();
+											ComponentsNeedingGrassMapRender.Remove(Component);
+											Component->MarkPackageDirty();
+										}
+#endif
+
+										NumCompsCreated++;
+
+										SCOPE_CYCLE_COUNTER(STAT_FoliageGrassStartComp);
+										int32 FolSeed = FCrc::StrCrc32((GrassType->GetName() + Component->GetName() + FString::Printf(TEXT("%d %d %d"), SubX, SubY, GrassVarietyIndex)).GetCharArray().GetData());
+										if (FolSeed == 0)
+										{
+											FolSeed++;
+										}
+
+										UHierarchicalInstancedStaticMeshComponent* HierarchicalInstancedStaticMeshComponent;
+										{
+											QUICK_SCOPE_CYCLE_COUNTER(STAT_GrassCreateComp);
+											HierarchicalInstancedStaticMeshComponent = NewObject<UHierarchicalInstancedStaticMeshComponent>(this); //, NAME_None, RF_Transactional);
+										}
+										NewComp.Foliage = HierarchicalInstancedStaticMeshComponent;
+										FoliageCache.CachedGrassComps.Add(NewComp);
+
+										HierarchicalInstancedStaticMeshComponent->Mobility = EComponentMobility::Static;
+
+										HierarchicalInstancedStaticMeshComponent->StaticMesh = GrassVariety.GrassMesh;
+										HierarchicalInstancedStaticMeshComponent->bSelectable = false;
+										HierarchicalInstancedStaticMeshComponent->bHasPerInstanceHitProxies = true;
+										static FName NoCollision(TEXT("NoCollision"));
+										HierarchicalInstancedStaticMeshComponent->SetCollisionProfileName(NoCollision);
+										HierarchicalInstancedStaticMeshComponent->bDisableCollision = true;
+										HierarchicalInstancedStaticMeshComponent->SetCanEverAffectNavigation(false);
+										HierarchicalInstancedStaticMeshComponent->InstancingRandomSeed = FolSeed;
+
+										if (!Cameras.Num() || bDisableGPUCull)
+										{
+											// if we don't have any cameras, then we are rendering landscape LOD materials or somesuch and we want to disable culling
+											HierarchicalInstancedStaticMeshComponent->InstanceStartCullDistance = 0;
+											HierarchicalInstancedStaticMeshComponent->InstanceEndCullDistance = 0;
+										}
+										else
+										{
+											HierarchicalInstancedStaticMeshComponent->InstanceStartCullDistance = GrassVariety.StartCullDistance;
+											HierarchicalInstancedStaticMeshComponent->InstanceEndCullDistance = GrassVariety.EndCullDistance;
+										}
+
+										//@todo - take the settings from a UFoliageType object.  For now, disable distance field lighting on grass so we don't hitch.
+										HierarchicalInstancedStaticMeshComponent->bAffectDistanceFieldLighting = false;
+
+										{
+											QUICK_SCOPE_CYCLE_COUNTER(STAT_GrassAttachComp);
+
+											HierarchicalInstancedStaticMeshComponent->AttachTo(GetRootComponent());
+											FTransform DesiredTransform = GetRootComponent()->ComponentToWorld;
+											DesiredTransform.RemoveScaling();
+											HierarchicalInstancedStaticMeshComponent->SetWorldTransform(DesiredTransform);
+
+											FoliageComponents.Add(HierarchicalInstancedStaticMeshComponent);
+										}
+
+										FAsyncGrassBuilder* Builder;
+
+										{
+											QUICK_SCOPE_CYCLE_COUNTER(STAT_GrassCreateBuilder);
+											Builder = new FAsyncGrassBuilder(this, Component, GrassType, GrassVariety, HierarchicalInstancedStaticMeshComponent, SqrtSubsections, SubX, SubY);
+										}
+
+										if (Builder->bHaveValidData)
+										{
+											FAsyncTask<FAsyncGrassTask>* Task = new FAsyncTask<FAsyncGrassTask>(Builder, NewComp.Key, HierarchicalInstancedStaticMeshComponent);
+
+											Task->StartBackgroundTask();
+
+											AsyncFoliageTasks.Add(Task);
+										}
+										else
+										{
+											delete Builder;
+										}
+										{
+											QUICK_SCOPE_CYCLE_COUNTER(STAT_GrassRegisterComp);
+											HierarchicalInstancedStaticMeshComponent->RegisterComponent();
+										}
+									}
+								}
+							}
+						}
+					}
+				}
+			}
+
+#if WITH_EDITOR
+
+			TotalTexturesToStreamForVisibleGrassMapRender -= NumTexturesToStreamForVisibleGrassMapRender;
+			NumTexturesToStreamForVisibleGrassMapRender = RequiredTexturesNotStreamedIn;
+			TotalTexturesToStreamForVisibleGrassMapRender += NumTexturesToStreamForVisibleGrassMapRender;
+
+			{
+				int32 NumComponentsRendered = 0;
+				int32 NumComponentsUnableToRender = 0;
+				if (GrassTypes.Num() > 0 && CVarPrerenderGrassmaps.GetValueOnAnyThread() > 0)
+				{
+					// try to render some grassmaps
+					TArray<ULandscapeComponent*> ComponentsToRender;
+					for (auto Component : ComponentsNeedingGrassMapRender)
+					{
+						if (Component->CanRenderGrassMap())
+						{
+							if (Component->AreTexturesStreamedForGrassMapRender())
+							{
+								// We really want to throttle the number based on component size.
+								if (NumComponentsRendered <= 4)
+								{
+									ComponentsToRender.Add(Component);
+									NumComponentsRendered++;
+								}
+							}
+							else
+							if (TotalTexturesToStreamForVisibleGrassMapRender == 0)
+							{
+								// Force stream in other heightmaps but only if we're not waiting for the textures 
+								// near the camera to stream in
+								DesiredForceStreamedTextures.Add(Component->HeightmapTexture);
+								for (auto WeightmapTexture : Component->WeightmapTextures)
+								{
+									DesiredForceStreamedTextures.Add(WeightmapTexture);
+								}
+							}
+						}
+						else
+						{
+							NumComponentsUnableToRender++;
+						}
+					}
+					if (ComponentsToRender.Num())
+					{
+						RenderGrassMaps(ComponentsToRender, GrassTypes);
+						MarkPackageDirty();
+					}
+				}
+
+				TotalComponentsNeedingGrassMapRender -= NumComponentsNeedingGrassMapRender;
+				NumComponentsNeedingGrassMapRender = ComponentsNeedingGrassMapRender.Num() - NumComponentsRendered - NumComponentsUnableToRender;
+				TotalComponentsNeedingGrassMapRender += NumComponentsNeedingGrassMapRender;
+
+				// Update resident flags
+				for (auto Texture : DesiredForceStreamedTextures.Difference(CurrentForcedStreamedTextures))
+				{
+					Texture->bForceMiplevelsToBeResident = true;
+				}
+				for (auto Texture : CurrentForcedStreamedTextures.Difference(DesiredForceStreamedTextures))
+				{
+					Texture->bForceMiplevelsToBeResident = false;
+				}
+
+			}
+#endif
+		}
+	}
+
+	static TSet<UHierarchicalInstancedStaticMeshComponent *> StillUsed;
+	StillUsed.Empty(256);
+	{
+		// trim cached items based on time, pending and emptiness
+		double OldestToKeepTime = FPlatformTime::Seconds() - CVarMinTimeToKeepGrass.GetValueOnGameThread();
+		uint32 OldestToKeepFrame = GFrameNumber - CVarMinFramesToKeepGrass.GetValueOnGameThread();
+		for (FCachedLandscapeFoliage::TGrassSet::TIterator Iter(FoliageCache.CachedGrassComps); Iter; ++Iter)
+		{
+			const FCachedLandscapeFoliage::FGrassComp& GrassItem = *Iter;
+			UHierarchicalInstancedStaticMeshComponent *Used = GrassItem.Foliage.Get();
+			bool bOld =
+				!GrassItem.Pending &&
+				(
+				!GrassItem.Key.BasedOn.Get() ||
+				!GrassItem.Key.GrassType.Get() ||
+				!Used ||
+				(GrassItem.LastUsedFrameNumber < OldestToKeepFrame && GrassItem.LastUsedTime < OldestToKeepTime)
+				);
+			if (bOld)
+			{
+				Iter.RemoveCurrent();
+			}
+			else if (Used)
+			{
+				StillUsed.Add(Used);
+			}
+		}
+	}
+	{
+		// delete components that are no longer used
+		for (UActorComponent* ActorComponent : GetComponents())
+		{
+			UHierarchicalInstancedStaticMeshComponent* HComponent = Cast<UHierarchicalInstancedStaticMeshComponent>(ActorComponent);
+			if (HComponent && !StillUsed.Contains(HComponent))
+			{
+				{
+					SCOPE_CYCLE_COUNTER(STAT_FoliageGrassDestoryComp);
+					HComponent->ClearInstances();
+					HComponent->DestroyComponent();
+					FoliageComponents.Remove(HComponent);
+				}
+				if (!bForceSync)
+				{
+					break; // one per frame is fine
+				}
+			}
+		}
+	}
+	{
+		// finish async tasks
+		for (int32 Index = 0; Index < AsyncFoliageTasks.Num(); Index++)
+		{
+			FAsyncTask<FAsyncGrassTask>* Task = AsyncFoliageTasks[Index];
+			if (bForceSync)
+			{
+				Task->EnsureCompletion();
+			}
+			if (Task->IsDone())
+			{
+				SCOPE_CYCLE_COUNTER(STAT_FoliageGrassEndComp);
+				FAsyncGrassTask& Inner = Task->GetTask();
+				AsyncFoliageTasks.RemoveAtSwap(Index--);
+				UHierarchicalInstancedStaticMeshComponent* HierarchicalInstancedStaticMeshComponent = Inner.Foliage.Get();
+				if (HierarchicalInstancedStaticMeshComponent && StillUsed.Contains(HierarchicalInstancedStaticMeshComponent))
+				{
+					if (Inner.Builder->InstanceBuffer.Num())
+					{
+						Exchange(HierarchicalInstancedStaticMeshComponent->WriteOncePrebuiltInstanceBuffer, Inner.Builder->InstanceBuffer);
+						HierarchicalInstancedStaticMeshComponent->AcceptPrebuiltTree(Inner.Builder->ClusterTree);
+						if (bForceSync && GetWorld())
+						{
+							HierarchicalInstancedStaticMeshComponent->RecreateRenderState_Concurrent();
+						}
+					}
+				}
+				FCachedLandscapeFoliage::FGrassComp* Existing = FoliageCache.CachedGrassComps.Find(Inner.Key);
+				if (Existing)
+				{
+					Existing->Pending = false;
+					Existing->Touch();
+				}
+				delete Task;
+				if (!bForceSync)
+				{
+					break; // one per frame is fine
+				}
+			}
+		}
+	}
+}
+
+void FAsyncGrassTask::DoWork()
+{
+	Builder->Build();
+}
+
+FAsyncGrassTask::~FAsyncGrassTask()
+{
+	delete Builder;
+}
+
+static void FlushGrass(const TArray<FString>& Args)
+{
+	for (TObjectIterator<ALandscapeProxy> It; It; ++It)
+	{
+		ALandscapeProxy* Landscape = *It;
+		if (Landscape && !Landscape->IsTemplate() && !Landscape->IsPendingKill())
+		{
+			Landscape->FlushGrassComponents();
+		}
+	}
+}
+
+static FAutoConsoleCommand FlushGrassCmd(
+	TEXT("grass.FlushCache"),
+	TEXT("Flush the grass cache, debugging."),
+	FConsoleCommandWithArgsDelegate::CreateStatic(&FlushGrass)
+	);
+
+
 #undef LOCTEXT_NAMESPACE
\ No newline at end of file
