UObject constructors create subobjects using UObject::Create* instead of FObjectInitializer::Create*

diff --git a/Engine/Plugins/2D/Paper2D/Source/Paper2D/Private/PaperCharacter.cpp b/Engine/Plugins/2D/Paper2D/Source/Paper2D/Private/PaperCharacter.cpp
--- a/Engine/Plugins/2D/Paper2D/Source/Paper2D/Private/PaperCharacter.cpp
+++ b/Engine/Plugins/2D/Paper2D/Source/Paper2D/Private/PaperCharacter.cpp
@@ -11,7 +11,7 @@ APaperCharacter::APaperCharacter(const FObjectInitializer& ObjectInitializer)
 	: Super(ObjectInitializer.DoNotCreateDefaultSubobject(ACharacter::MeshComponentName))
 {
 	// Try to create the sprite component
-	Sprite = ObjectInitializer.CreateOptionalDefaultSubobject<UPaperFlipbookComponent>(this, APaperCharacter::SpriteComponentName);
+	Sprite = CreateOptionalDefaultSubobject<UPaperFlipbookComponent>(APaperCharacter::SpriteComponentName);
 	if (Sprite)
 	{
 		Sprite->AlwaysLoadOnClient = true;
diff --git a/Engine/Plugins/2D/Paper2D/Source/Paper2D/Private/PaperFlipbookActor.cpp b/Engine/Plugins/2D/Paper2D/Source/Paper2D/Private/PaperFlipbookActor.cpp
--- a/Engine/Plugins/2D/Paper2D/Source/Paper2D/Private/PaperFlipbookActor.cpp
+++ b/Engine/Plugins/2D/Paper2D/Source/Paper2D/Private/PaperFlipbookActor.cpp
@@ -9,7 +9,7 @@
 APaperFlipbookActor::APaperFlipbookActor(const FObjectInitializer& ObjectInitializer)
 	: Super(ObjectInitializer)
 {
-	RenderComponent = ObjectInitializer.CreateDefaultSubobject<UPaperFlipbookComponent>(this, TEXT("RenderComponent"));
+	RenderComponent = CreateDefaultSubobject<UPaperFlipbookComponent>(TEXT("RenderComponent"));
 
 	RootComponent = RenderComponent;
 }
diff --git a/Engine/Plugins/2D/Paper2D/Source/Paper2D/Private/PaperSpriteActor.cpp b/Engine/Plugins/2D/Paper2D/Source/Paper2D/Private/PaperSpriteActor.cpp
--- a/Engine/Plugins/2D/Paper2D/Source/Paper2D/Private/PaperSpriteActor.cpp
+++ b/Engine/Plugins/2D/Paper2D/Source/Paper2D/Private/PaperSpriteActor.cpp
@@ -9,7 +9,7 @@
 APaperSpriteActor::APaperSpriteActor(const FObjectInitializer& ObjectInitializer)
 	: Super(ObjectInitializer)
 {
-	RenderComponent = ObjectInitializer.CreateDefaultSubobject<UPaperSpriteComponent>(this, TEXT("RenderComponent"));
+	RenderComponent = CreateDefaultSubobject<UPaperSpriteComponent>(TEXT("RenderComponent"));
 	RenderComponent->SetCollisionProfileName(UCollisionProfile::BlockAll_ProfileName);
 	RenderComponent->Mobility = EComponentMobility::Static;
 
diff --git a/Engine/Plugins/2D/Paper2D/Source/Paper2D/Private/PaperTileMapActor.cpp b/Engine/Plugins/2D/Paper2D/Source/Paper2D/Private/PaperTileMapActor.cpp
--- a/Engine/Plugins/2D/Paper2D/Source/Paper2D/Private/PaperTileMapActor.cpp
+++ b/Engine/Plugins/2D/Paper2D/Source/Paper2D/Private/PaperTileMapActor.cpp
@@ -8,7 +8,7 @@
 APaperTileMapActor::APaperTileMapActor(const FObjectInitializer& ObjectInitializer)
 	: Super(ObjectInitializer)
 {
-	RenderComponent = ObjectInitializer.CreateDefaultSubobject<UPaperTileMapComponent>(this, TEXT("RenderComponent"));
+	RenderComponent = CreateDefaultSubobject<UPaperTileMapComponent>(TEXT("RenderComponent"));
 
 	RootComponent = RenderComponent;
 }
diff --git a/Engine/Plugins/2D/Paper2D/Source/Paper2D/Private/Terrain/PaperTerrainActor.cpp b/Engine/Plugins/2D/Paper2D/Source/Paper2D/Private/Terrain/PaperTerrainActor.cpp
--- a/Engine/Plugins/2D/Paper2D/Source/Paper2D/Private/Terrain/PaperTerrainActor.cpp
+++ b/Engine/Plugins/2D/Paper2D/Source/Paper2D/Private/Terrain/PaperTerrainActor.cpp
@@ -8,9 +8,9 @@
 APaperTerrainActor::APaperTerrainActor(const FObjectInitializer& ObjectInitializer)
 	: Super(ObjectInitializer)
 {
-	DummyRoot = ObjectInitializer.CreateDefaultSubobject<USceneComponent>(this, TEXT("RootComponent"));
-	SplineComponent = ObjectInitializer.CreateDefaultSubobject<UPaperTerrainSplineComponent>(this, TEXT("SplineComponent"));
- 	RenderComponent = ObjectInitializer.CreateDefaultSubobject<UPaperTerrainComponent>(this, TEXT("RenderComponent"));
+	DummyRoot = CreateDefaultSubobject<USceneComponent>(TEXT("RootComponent"));
+	SplineComponent = CreateDefaultSubobject<UPaperTerrainSplineComponent>(TEXT("SplineComponent"));
+ 	RenderComponent = CreateDefaultSubobject<UPaperTerrainComponent>(TEXT("RenderComponent"));
  
 	SplineComponent->AttachParent = DummyRoot;
 	RenderComponent->AttachParent = DummyRoot;
diff --git a/Engine/Plugins/Runtime/CableComponent/Source/CableComponent/Private/CableActor.cpp b/Engine/Plugins/Runtime/CableComponent/Source/CableComponent/Private/CableActor.cpp
--- a/Engine/Plugins/Runtime/CableComponent/Source/CableComponent/Private/CableActor.cpp
+++ b/Engine/Plugins/Runtime/CableComponent/Source/CableComponent/Private/CableActor.cpp
@@ -6,6 +6,6 @@
 ACableActor::ACableActor(const FObjectInitializer& ObjectInitializer)
 	: Super(ObjectInitializer)
 {
-	CableComponent = ObjectInitializer.CreateDefaultSubobject<UCableComponent>(this, TEXT("CableComponent0"));
+	CableComponent = CreateDefaultSubobject<UCableComponent>(TEXT("CableComponent0"));
 	RootComponent = CableComponent;
 }
\ No newline at end of file
diff --git a/Engine/Source/Developer/FunctionalTesting/Private/FunctionalTest.cpp b/Engine/Source/Developer/FunctionalTesting/Private/FunctionalTest.cpp
--- a/Engine/Source/Developer/FunctionalTesting/Private/FunctionalTest.cpp
+++ b/Engine/Source/Developer/FunctionalTesting/Private/FunctionalTest.cpp
@@ -16,7 +16,7 @@ AFunctionalTest::AFunctionalTest( const FObjectInitializer& ObjectInitializer )
 	PrimaryActorTick.bCanEverTick = true;
 	PrimaryActorTick.bStartWithTickEnabled = false;
 		
-	SpriteComponent = ObjectInitializer.CreateDefaultSubobject<UBillboardComponent>(this, TEXT("Sprite"));
+	SpriteComponent = CreateDefaultSubobject<UBillboardComponent>(TEXT("Sprite"));
 	if (SpriteComponent)
 	{
 		SpriteComponent->bHiddenInGame = false;
diff --git a/Engine/Source/Developer/GameplayDebugger/Private/GameplayDebuggingReplicator.cpp b/Engine/Source/Developer/GameplayDebugger/Private/GameplayDebuggingReplicator.cpp
--- a/Engine/Source/Developer/GameplayDebugger/Private/GameplayDebuggingReplicator.cpp
+++ b/Engine/Source/Developer/GameplayDebugger/Private/GameplayDebuggingReplicator.cpp
@@ -41,7 +41,7 @@ AGameplayDebuggingReplicator::AGameplayDebuggingReplicator(const FObjectInitiali
 	PrimaryActorTick.bCanEverTick = true;
 	PrimaryActorTick.bStartWithTickEnabled = false;
 	
-	USceneComponent* SceneComponent = ObjectInitializer.CreateDefaultSubobject<USceneComponent>(this, TEXT("SceneComponent"));
+	USceneComponent* SceneComponent = CreateDefaultSubobject<USceneComponent>(TEXT("SceneComponent"));
 	RootComponent = SceneComponent;
 
 #if WITH_EDITOR
diff --git a/Engine/Source/Editor/LandscapeEditor/Private/ActorFactoryLandscape.cpp b/Engine/Source/Editor/LandscapeEditor/Private/ActorFactoryLandscape.cpp
--- a/Engine/Source/Editor/LandscapeEditor/Private/ActorFactoryLandscape.cpp
+++ b/Engine/Source/Editor/LandscapeEditor/Private/ActorFactoryLandscape.cpp
@@ -45,12 +45,12 @@ ALandscapePlaceholder::ALandscapePlaceholder(const FObjectInitializer& ObjectIni
 	};
 	static FConstructorStatics ConstructorStatics;
 
-	USceneComponent* SceneComponent = ObjectInitializer.CreateDefaultSubobject<USceneComponent>(this, TEXT("RootComponent0"));
+	USceneComponent* SceneComponent = CreateDefaultSubobject<USceneComponent>(TEXT("RootComponent0"));
 	RootComponent = SceneComponent;
 	RootComponent->Mobility = EComponentMobility::Static;
 
 #if WITH_EDITORONLY_DATA
-	UBillboardComponent* SpriteComponent = ObjectInitializer.CreateEditorOnlyDefaultSubobject<UBillboardComponent>(this, TEXT("Sprite"));
+	UBillboardComponent* SpriteComponent = CreateEditorOnlyDefaultSubobject<UBillboardComponent>(TEXT("Sprite"));
 	if (SpriteComponent)
 	{
 		SpriteComponent->Sprite = ConstructorStatics.TerrainTexture.Get();
diff --git a/Engine/Source/Editor/UnrealEd/Private/GroupActor.cpp b/Engine/Source/Editor/UnrealEd/Private/GroupActor.cpp
--- a/Engine/Source/Editor/UnrealEd/Private/GroupActor.cpp
+++ b/Engine/Source/Editor/UnrealEd/Private/GroupActor.cpp
@@ -15,7 +15,7 @@ AGroupActor::AGroupActor(const FObjectInitializer& ObjectInitializer)
 {
 	bLocked = true;
 
-	USceneComponent* GroupComponent = ObjectInitializer.CreateDefaultSubobject<USceneComponent>(this, TEXT("GroupComponent"));
+	USceneComponent* GroupComponent = CreateDefaultSubobject<USceneComponent>(TEXT("GroupComponent"));
 	RootComponent = GroupComponent;
 }
 
diff --git a/Engine/Source/Runtime/AIModule/Private/AIController.cpp b/Engine/Source/Runtime/AIModule/Private/AIController.cpp
--- a/Engine/Source/Runtime/AIModule/Private/AIController.cpp
+++ b/Engine/Source/Runtime/AIModule/Private/AIController.cpp
@@ -38,10 +38,10 @@ DEFINE_LOG_CATEGORY(LogAINavigation);
 AAIController::AAIController(const FObjectInitializer& ObjectInitializer)
 	: Super(ObjectInitializer)
 {
-	PathFollowingComponent = ObjectInitializer.CreateDefaultSubobject<UPathFollowingComponent>(this, TEXT("PathFollowingComponent"));
+	PathFollowingComponent = CreateDefaultSubobject<UPathFollowingComponent>(TEXT("PathFollowingComponent"));
 	PathFollowingComponent->OnMoveFinished.AddUObject(this, &AAIController::OnMoveCompleted);
 
-	ActionsComp = ObjectInitializer.CreateDefaultSubobject<UPawnActionsComponent>(this, "ActionsComp");
+	ActionsComp = CreateDefaultSubobject<UPawnActionsComponent>("ActionsComp");
 
 	bSkipExtraLOSChecks = true;
 	bWantsPlayerState = false;
diff --git a/Engine/Source/Runtime/AIModule/Private/EnvironmentQuery/EQSTestingPawn.cpp b/Engine/Source/Runtime/AIModule/Private/EnvironmentQuery/EQSTestingPawn.cpp
--- a/Engine/Source/Runtime/AIModule/Private/EnvironmentQuery/EQSTestingPawn.cpp
+++ b/Engine/Source/Runtime/AIModule/Private/EnvironmentQuery/EQSTestingPawn.cpp
@@ -37,7 +37,7 @@ AEQSTestingPawn::AEQSTestingPawn(const FObjectInitializer& ObjectInitializer)
 	GetCapsuleComponent()->SetCollisionProfileName(CollisionProfileName);
 
 #if WITH_EDITORONLY_DATA
-	EdRenderComp = ObjectInitializer.CreateEditorOnlyDefaultSubobject<UEQSRenderingComponent>(this, TEXT("EQSRender"));
+	EdRenderComp = CreateEditorOnlyDefaultSubobject<UEQSRenderingComponent>(TEXT("EQSRender"));
 	if (HasAnyFlags(RF_ClassDefaultObject) == false)
 	{
 		UArrowComponent* ArrowComponent = FindComponentByClass<UArrowComponent>();
@@ -47,7 +47,7 @@ AEQSTestingPawn::AEQSTestingPawn(const FObjectInitializer& ObjectInitializer)
 			ArrowComponent->bIsScreenSizeScaled = true;
 		}
 
-		UBillboardComponent* SpriteComponent = ObjectInitializer.CreateEditorOnlyDefaultSubobject<UBillboardComponent>(this, TEXT("Sprite"));
+		UBillboardComponent* SpriteComponent = CreateEditorOnlyDefaultSubobject<UBillboardComponent>(TEXT("Sprite"));
 		if (!IsRunningCommandlet() && (SpriteComponent != nullptr))
 		{
 			struct FConstructorStatics
diff --git a/Engine/Source/Runtime/GameplayAbilities/Private/Abilities/GameplayAbilityWorldReticle_ActorVisualization.cpp b/Engine/Source/Runtime/GameplayAbilities/Private/Abilities/GameplayAbilityWorldReticle_ActorVisualization.cpp
--- a/Engine/Source/Runtime/GameplayAbilities/Private/Abilities/GameplayAbilityWorldReticle_ActorVisualization.cpp
+++ b/Engine/Source/Runtime/GameplayAbilities/Private/Abilities/GameplayAbilityWorldReticle_ActorVisualization.cpp
@@ -16,7 +16,7 @@ AGameplayAbilityWorldReticle_ActorVisualization::AGameplayAbilityWorldReticle_Ac
 	//PrimaryActorTick.bCanEverTick = true;
 	//PrimaryActorTick.TickGroup = TG_PrePhysics;
 
-	CollisionComponent = ObjectInitializer.CreateDefaultSubobject<UCapsuleComponent>(this, TEXT("CollisionCapsule0"));
+	CollisionComponent = CreateDefaultSubobject<UCapsuleComponent>(TEXT("CollisionCapsule0"));
 	CollisionComponent->InitCapsuleSize(0.f, 0.f);
 	CollisionComponent->AlwaysLoadOnClient = true;
 	CollisionComponent->bAbsoluteScale = true;
diff --git a/Engine/Source/Runtime/GameplayAbilities/Private/AbilitySystemTestPawn.cpp b/Engine/Source/Runtime/GameplayAbilities/Private/AbilitySystemTestPawn.cpp
--- a/Engine/Source/Runtime/GameplayAbilities/Private/AbilitySystemTestPawn.cpp
+++ b/Engine/Source/Runtime/GameplayAbilities/Private/AbilitySystemTestPawn.cpp
@@ -10,7 +10,7 @@ FName  AAbilitySystemTestPawn::AbilitySystemComponentName(TEXT("AbilitySystemCom
 AAbilitySystemTestPawn::AAbilitySystemTestPawn(const FObjectInitializer& ObjectInitializer)
 : Super(ObjectInitializer)
 {
-	AbilitySystemComponent = ObjectInitializer.CreateDefaultSubobject<UAbilitySystemComponent>(this, AAbilitySystemTestPawn::AbilitySystemComponentName);
+	AbilitySystemComponent = CreateDefaultSubobject<UAbilitySystemComponent>(AAbilitySystemTestPawn::AbilitySystemComponentName);
 	AbilitySystemComponent->SetIsReplicated(true);
 
 	//DefaultAbilitySet = NULL;
diff --git a/Engine/Source/Runtime/Landscape/Private/Landscape.cpp b/Engine/Source/Runtime/Landscape/Private/Landscape.cpp
--- a/Engine/Source/Runtime/Landscape/Private/Landscape.cpp
+++ b/Engine/Source/Runtime/Landscape/Private/Landscape.cpp
@@ -518,7 +518,7 @@ ALandscapeProxy::ALandscapeProxy(const FObjectInitializer& ObjectInitializer)
 	// by default we want to see the Landscape shadows even in the far shadow cascades
 	bCastFarShadow = true;
 
-	USceneComponent* SceneComponent = ObjectInitializer.CreateDefaultSubobject<USceneComponent>(this, TEXT("RootComponent0"));
+	USceneComponent* SceneComponent = CreateDefaultSubobject<USceneComponent>(TEXT("RootComponent0"));
 	RootComponent = SceneComponent;
 	RootComponent->RelativeScale3D = FVector(128.0f, 128.0f, 256.0f); // Old default scale, preserved for compatibility. See ULandscapeEditorObject::NewLandscape_Scale
 	RootComponent->Mobility = EComponentMobility::Static;
diff --git a/Engine/Source/Runtime/Landscape/Private/LandscapeGizmoActor.cpp b/Engine/Source/Runtime/Landscape/Private/LandscapeGizmoActor.cpp
--- a/Engine/Source/Runtime/Landscape/Private/LandscapeGizmoActor.cpp
+++ b/Engine/Source/Runtime/Landscape/Private/LandscapeGizmoActor.cpp
@@ -385,7 +385,7 @@ ALandscapeGizmoActor::ALandscapeGizmoActor(const FObjectInitializer& ObjectIniti
 	: Super(ObjectInitializer)
 {
 #if WITH_EDITORONLY_DATA
-	SpriteComponent = ObjectInitializer.CreateEditorOnlyDefaultSubobject<UBillboardComponent>(this, TEXT("Sprite"));
+	SpriteComponent = CreateEditorOnlyDefaultSubobject<UBillboardComponent>(TEXT("Sprite"));
 	if (!IsRunningCommandlet() && (SpriteComponent != nullptr))
 	{
 		// Structure to hold one-time initialization
@@ -412,7 +412,7 @@ ALandscapeGizmoActor::ALandscapeGizmoActor(const FObjectInitializer& ObjectIniti
 	}
 #endif
 
-	USceneComponent* SceneComponent = ObjectInitializer.CreateDefaultSubobject<USceneComponent>(this, TEXT("RootComponent0"));
+	USceneComponent* SceneComponent = CreateDefaultSubobject<USceneComponent>(TEXT("RootComponent0"));
 	RootComponent = SceneComponent;
 	RootComponent->Mobility = EComponentMobility::Static;
 
@@ -481,7 +481,7 @@ ALandscapeGizmoActiveActor::ALandscapeGizmoActiveActor(const FObjectInitializer&
 	}
 #endif // WITH_EDITORONLY_DATA
 
-	ULandscapeGizmoRenderComponent* LandscapeGizmoRenderComponent = ObjectInitializer.CreateDefaultSubobject<ULandscapeGizmoRenderComponent>(this, TEXT("GizmoRendererComponent0"));
+	ULandscapeGizmoRenderComponent* LandscapeGizmoRenderComponent = CreateDefaultSubobject<ULandscapeGizmoRenderComponent>(TEXT("GizmoRendererComponent0"));
 	LandscapeGizmoRenderComponent->SetCollisionProfileName(UCollisionProfile::BlockAllDynamic_ProfileName);
 
 	RootComponent = LandscapeGizmoRenderComponent;
