Landscape baking of diffuse color for GI from Jack

diff --git a/Engine/Source/Editor/LandscapeEditor/Private/LandscapeEdModeTools.h b/Engine/Source/Editor/LandscapeEditor/Private/LandscapeEdModeTools.h
--- a/Engine/Source/Editor/LandscapeEditor/Private/LandscapeEdModeTools.h
+++ b/Engine/Source/Editor/LandscapeEditor/Private/LandscapeEdModeTools.h
@@ -619,7 +619,7 @@ struct FHeightmapAccessor
 			}
 
 			// Flush dynamic foliage (grass)
-			ALandscapeProxy::DistributeFlushGrassComponents(Components);
+			ALandscapeProxy::InvalidateGeneratedComponentData(Components);
 
 			// Notify foliage to move any attached instances
 			bool bUpdateFoliage = false;
@@ -792,7 +792,7 @@ struct FXYOffsetmapAccessor
 			}
 
 			// Flush dynamic foliage (grass)
-			ALandscapeProxy::DistributeFlushGrassComponents(Components);
+			ALandscapeProxy::InvalidateGeneratedComponentData(Components);
 
 			// Notify foliage to move any attached instances
 			bool bUpdateFoliage = false;
@@ -949,7 +949,7 @@ struct FAlphamapAccessor
 		if (LandscapeEdit.GetComponentsInRegion(X1, Y1, X2, Y2, &Components))
 		{
 			// Flush dynamic foliage (grass)
-			ALandscapeProxy::DistributeFlushGrassComponents(Components);
+			ALandscapeProxy::InvalidateGeneratedComponentData(Components);
 
 			LandscapeEdit.SetAlphaData(LayerInfo, X1, Y1, X2, Y2, Data, 0, PaintingRestriction, bBlendWeight, bUseTotalNormalize);
 			ModifiedComponents.Append(Components);
@@ -1056,7 +1056,7 @@ struct FFullWeightmapAccessor
 		if (LandscapeEdit.GetComponentsInRegion(X1, Y1, X2, Y2, &Components))
 		{
 			// Flush dynamic foliage (grass)
-			ALandscapeProxy::DistributeFlushGrassComponents(Components);
+			ALandscapeProxy::InvalidateGeneratedComponentData(Components);
 
 			LandscapeEdit.SetAlphaData(DirtyLayerInfos, X1, Y1, X2, Y2, Data, 0, PaintingRestriction);
 			ModifiedComponents.Append(Components);
diff --git a/Engine/Source/Editor/UnrealEd/Private/EditorExporters.cpp b/Engine/Source/Editor/UnrealEd/Private/EditorExporters.cpp
--- a/Engine/Source/Editor/UnrealEd/Private/EditorExporters.cpp
+++ b/Engine/Source/Editor/UnrealEd/Private/EditorExporters.cpp
@@ -3009,4 +3009,48 @@ namespace MaterialExportUtils
 		Material->PostEditChange();
 		return Material;
 	}
+
+	bool ExportBaseColor(ULandscapeComponent* LandscapeComponent, int32 TextureSize, TArray<FColor>& OutSamples)
+	{
+		ALandscapeProxy* LandscapeProxy = LandscapeComponent->GetLandscapeProxy();
+
+		FIntPoint ComponentOrigin = LandscapeComponent->GetSectionBase() - LandscapeProxy->LandscapeSectionOffset;
+		FIntPoint ComponentSize(LandscapeComponent->ComponentSizeQuads, LandscapeComponent->ComponentSizeQuads);
+		FVector MidPoint = FVector(ComponentOrigin, 0.f) + FVector(ComponentSize, 0.f)*0.5f;
+
+		FVector LandscapeCenter = LandscapeProxy->GetTransform().TransformPosition(MidPoint);
+		FVector LandscapeExtent = FVector(ComponentSize, 0.f)*LandscapeProxy->GetActorScale()*0.5f;
+
+		FVector ViewOrigin = LandscapeCenter;
+		FMatrix ViewRotationMatrix = FInverseRotationMatrix(LandscapeProxy->GetActorRotation());
+		ViewRotationMatrix *= FMatrix(FPlane(1, 0, 0, 0),
+			FPlane(0, -1, 0, 0),
+			FPlane(0, 0, -1, 0),
+			FPlane(0, 0, 0, 1));
+
+		const float ZOffset = WORLD_MAX;
+		FMatrix ProjectionMatrix = FReversedZOrthoMatrix(
+			LandscapeExtent.X,
+			LandscapeExtent.Y,
+			0.5f / ZOffset,
+			ZOffset);
+
+		FSceneInterface* Scene = LandscapeProxy->GetWorld()->Scene;
+
+		// Hide all but the component
+		TSet<FPrimitiveComponentId> HiddenPrimitives;
+		for (auto PrimitiveComponentId : Scene->GetScenePrimitiveComponentIds())
+		{
+			HiddenPrimitives.Add(PrimitiveComponentId);
+		}
+		HiddenPrimitives.Remove(LandscapeComponent->SceneProxy->GetPrimitiveComponentId());
+				
+		FIntPoint TargetSize(TextureSize, TextureSize);
+
+		// Render diffuse texture using BufferVisualizationMode=BaseColor
+		static const FName BaseColorName("BaseColor");
+		const float BaseColorGamma = 2.2f;
+		RenderSceneToTexture(Scene, BaseColorName, ViewOrigin, ViewRotationMatrix, ProjectionMatrix, HiddenPrimitives, TargetSize, BaseColorGamma, OutSamples);
+		return true;
+	}
 }
diff --git a/Engine/Source/Editor/UnrealEd/Private/LandscapeTextureBakingNotification.cpp b/Engine/Source/Editor/UnrealEd/Private/LandscapeTextureBakingNotification.cpp
new file mode 100644
--- /dev/null
+++ b/Engine/Source/Editor/UnrealEd/Private/LandscapeTextureBakingNotification.cpp
@@ -0,0 +1,34 @@
+// Copyright 1998-2015 Epic Games, Inc. All Rights Reserved.
+
+#include "UnrealEd.h"
+#include "GlobalEditorNotification.h"
+#include "Landscape.h"
+#include "SNotificationList.h"
+
+/** Notification class for grassmmap rendering. */
+class FLandscapeTextureBakingNotificationImpl : public FGlobalEditorNotification
+{
+protected:
+	/** FGlobalEditorNotification interface */
+	virtual bool ShouldShowNotification(const bool bIsNotificationAlreadyActive) const override;
+	virtual void SetNotificationText(const TSharedPtr<SNotificationItem>& InNotificationItem) const override;
+};
+
+/** Global notification object. */
+FLandscapeTextureBakingNotificationImpl GLandscapeTextureBakingNotification;
+
+bool FLandscapeTextureBakingNotificationImpl::ShouldShowNotification(const bool bIsNotificationAlreadyActive) const
+{
+	return ALandscapeProxy::TotalComponentsNeedingTextureBaking > 0;
+}
+
+void FLandscapeTextureBakingNotificationImpl::SetNotificationText(const TSharedPtr<SNotificationItem>& InNotificationItem) const
+{
+	if (ALandscapeProxy::TotalComponentsNeedingTextureBaking > 0)
+	{
+		FFormatNamedArguments Args;
+		Args.Add(TEXT("OutstandingTextures"), FText::AsNumber(ALandscapeProxy::TotalComponentsNeedingTextureBaking));
+		const FText ProgressMessage = FText::Format(NSLOCTEXT("TextureBaking", "TextureBakingFormat", "Baking Landscape Textures ({OutstandingTextures})"), Args);
+		InNotificationItem->SetText(ProgressMessage);
+	}
+}
diff --git a/Engine/Source/Editor/UnrealEd/Public/MaterialExportUtils.h b/Engine/Source/Editor/UnrealEd/Public/MaterialExportUtils.h
--- a/Engine/Source/Editor/UnrealEd/Public/MaterialExportUtils.h
+++ b/Engine/Source/Editor/UnrealEd/Public/MaterialExportUtils.h
@@ -3,6 +3,7 @@
 
 // Forward declarations
 class ALandscapeProxy;
+class ULandscapeComponent;
 
 namespace MaterialExportUtils
 {
@@ -73,5 +74,29 @@ namespace MaterialExportUtils
 	 * @return						Returns a pointer to the constructed UMaterial object.
 	 */
 	UNREALED_API UMaterial* CreateMaterial(const FFlattenMaterial& InFlattenMaterial, UPackage* InOuter, const FString& BaseName, EObjectFlags Flags, TArray<UObject*>& OutGeneratedAssets);
+
+
+	/**
+ 	 * Generates a texture from an array of samples 
+	 *
+	 * @param Outer					Outer for the material and texture objects, if NULL new packages will be created for each asset
+	 * @param AssetLongName			Long asset path for the new texture
+	 * @param Size					Resolution of the texture to generate (must match the number of samples)
+	 * @param Samples				Color data for the texture
+	 * @param CompressionSettings	Compression settings for the new texture
+	 * @param LODGroup				LODGroup for the new texture
+	 * @param Flags					ObjectFlags for the new texture
+	 * @param bSRGB					Whether to set the bSRGB flag on the new texture
+	 * @return						The new texture.
+	 */
+	UNREALED_API UTexture2D* CreateTexture(UPackage* Outer, const FString& AssetLongName, FIntPoint Size, const TArray<FColor>& Samples, TextureCompressionSettings CompressionSettings, TextureGroup LODGroup, EObjectFlags Flags, bool bSRGB);
+
+	/**
+	* Creates bakes textures for a ULandscapeComponent
+	*
+	* @param LandscapeComponent		The component to bake textures for
+	* @return						Whether operation was successful
+	*/
+	UNREALED_API bool ExportBaseColor(ULandscapeComponent* LandscapeComponent, int32 TextureSize, TArray<FColor>& OutSamples);
 }
 
diff --git a/Engine/Source/Runtime/Engine/Public/SceneInterface.h b/Engine/Source/Runtime/Engine/Public/SceneInterface.h
--- a/Engine/Source/Runtime/Engine/Public/SceneInterface.h
+++ b/Engine/Source/Runtime/Engine/Public/SceneInterface.h
@@ -8,6 +8,7 @@ class FMaterialShaderMap;
 class ULightComponent;
 class USkyLightComponent;
 class FAtmosphericFogSceneInfo;
+class FPrimitiveComponentId;
 
 /**
  * An interface to the private scene manager implementation of a scene.  Use GetRendererModule().AllocateScene to create.
@@ -313,6 +314,11 @@ public:
 		return GetFeatureLevel() >= ERHIFeatureLevel::SM4;
 	}
 
+	/**
+	 * Returns the FPrimitiveComponentId for all primitives in the scene
+	 */
+	virtual TArray<FPrimitiveComponentId> GetScenePrimitiveComponentIds() const { return TArray<FPrimitiveComponentId>(); }
+
 protected:
 	virtual ~FSceneInterface() {}
 };
diff --git a/Engine/Source/Runtime/Landscape/Classes/LandscapeComponent.h b/Engine/Source/Runtime/Landscape/Classes/LandscapeComponent.h
--- a/Engine/Source/Runtime/Landscape/Classes/LandscapeComponent.h
+++ b/Engine/Source/Runtime/Landscape/Classes/LandscapeComponent.h
@@ -236,6 +236,14 @@ public:
 	UPROPERTY()
 	FGuid StateId;
 
+	/** The Material Guid that used when baking, to detect material recompilations */
+	UPROPERTY()
+	FGuid BakedTextureMaterialGuid;
+
+	/** Pre-baked Base Color texture for use by distance field GI */
+	UPROPERTY(EditAnywhere, BlueprintReadOnly, Category = BakedTextures)
+	UTexture2D* GIBakedBaseColorTexture;
+
 #if WITH_EDITORONLY_DATA
 	/** LOD level Bias to use when lighting buidling via lightmass, -1 Means automatic LOD calculation based on ForcedLOD + LODBias */
 	UPROPERTY(EditAnywhere, Category=LandscapeComponent)
diff --git a/Engine/Source/Runtime/Landscape/Classes/LandscapeProxy.h b/Engine/Source/Runtime/Landscape/Classes/LandscapeProxy.h
--- a/Engine/Source/Runtime/Landscape/Classes/LandscapeProxy.h
+++ b/Engine/Source/Runtime/Landscape/Classes/LandscapeProxy.h
@@ -524,9 +524,6 @@ public:
 	/** Flush the grass cache */
 	LANDSCAPE_API void FlushGrassComponents(const TSet<ULandscapeComponent*>* OnlyForComponents = nullptr, bool bFlushGrassMaps = true);
 
-	/* Static version to distribute to the appropriate proxy for each component */
-	LANDSCAPE_API static void DistributeFlushGrassComponents(const TSet<ULandscapeComponent*>& Components, bool bFlushGrassMaps = true);
-
 	/** 
 		Update Grass 
 		* @param Cameras to use for culling, if empty, then NO culling
@@ -537,9 +534,18 @@ public:
 	/* Get the list of grass types on this landscape */
 	TArray<ULandscapeGrassType*> GetGrassTypes() const;
 
+	/* Invalidate the precomputed grass and baked texture data for the specified components */
+	LANDSCAPE_API static void InvalidateGeneratedComponentData(const TSet<ULandscapeComponent*>& Components);
+
 #if WITH_EDITOR
 	/** Render grass maps for the specified components */
 	void RenderGrassMaps(const TArray<ULandscapeComponent*>& LandscapeComponents, const TArray<ULandscapeGrassType*>& GrassTypes);
+
+	/** Update any textures baked from the landscape as necessary */
+	void UpdateBakedTextures();
+
+	/** Frame counter to count down to the next time we check to update baked textures, so we don't check every frame */
+	int32 UpdateBakedTexturesCountdown;
 #endif
 
 	// Begin FTickableGameObject interface.
@@ -649,6 +655,9 @@ public:
 	int32 NumTexturesToStreamForVisibleGrassMapRender;
 	LANDSCAPE_API static int32 TotalTexturesToStreamForVisibleGrassMapRender;
 
+	/* For the texture baking notification */
+	int32 NumComponentsNeedingTextureBaking;
+	LANDSCAPE_API static int32 TotalComponentsNeedingTextureBaking;
 #endif
 };
 
diff --git a/Engine/Source/Runtime/Landscape/Private/Landscape.cpp b/Engine/Source/Runtime/Landscape/Private/Landscape.cpp
--- a/Engine/Source/Runtime/Landscape/Private/Landscape.cpp
+++ b/Engine/Source/Runtime/Landscape/Private/Landscape.cpp
@@ -37,6 +37,10 @@ Landscape.cpp: Terrain rendering
 #include "Materials/MaterialExpressionLandscapeLayerBlend.h"
 #include "Materials/MaterialExpressionLandscapeVisibilityMask.h"
 
+#if WITH_EDITOR
+#include "MaterialExportUtils.h"
+#endif
+
 /** Landscape stats */
 
 DEFINE_STAT(STAT_LandscapeDynamicDrawTime);
@@ -2145,6 +2149,15 @@ void ALandscapeProxy::Tick(float DeltaSeconds)
 		// the super tick is for an actor tick...
 		//Super::Tick(DeltaSeconds);
 
+#if WITH_EDITOR
+		UWorld* World = GetWorld();
+
+		if (GIsEditor && World && !World->IsPlayInEditor())
+		{
+			UpdateBakedTextures();
+		}
+#endif
+
 		TickGrass();
 	}
 }
@@ -2191,4 +2204,138 @@ ULandscapeMeshProxyComponent::ULandscapeMeshProxyComponent(const FObjectInitiali
 {
 }
 
+#if WITH_EDITOR
+void ALandscapeProxy::UpdateBakedTextures()
+{
+	// See if we can render
+	UWorld* World = GetWorld();
+	if (!GIsEditor || GUsingNullRHI || !World || World->IsGameWorld() || World->FeatureLevel < ERHIFeatureLevel::SM4)
+	{
+		return;
+	}
+
+	if (UpdateBakedTexturesCountdown-- > 0)
+	{
+		return;
+	}
+
+	// Group components by heightmap texture
+	TMap<UTexture2D*, TArray<ULandscapeComponent*>> ComponentsByHeightmap;
+	for (ULandscapeComponent* Component : LandscapeComponents)
+	{
+		ComponentsByHeightmap.FindOrAdd(Component->HeightmapTexture).Add(Component);
+	}
+
+	TotalComponentsNeedingTextureBaking -= NumComponentsNeedingTextureBaking;
+	NumComponentsNeedingTextureBaking = 0;
+	int32 NumGenerated = 0;
+
+	for (auto It = ComponentsByHeightmap.CreateConstIterator(); It; ++It)
+	{
+		bool bHaveComponentsToUpdate = false;
+		for (ULandscapeComponent* Component : It.Value())
+		{
+			// not registered; ignore this component
+			if (!Component->SceneProxy)
+			{
+				continue;
+			}
+
+			// Check we can render the material
+			UMaterialInstance* MaterialInstance = Component->MaterialInstance;
+
+			if (!MaterialInstance || !MaterialInstance->GetMaterialResource(World->FeatureLevel)->HasValidGameThreadShaderMap())
+			{
+				// Cannot render this component yet as it doesn't have a material; abandon the atlas for this heightmap
+				bHaveComponentsToUpdate = false;
+				break;
+			}
+
+			if (Component->BakedTextureMaterialGuid == MaterialInstance->GetMaterial()->StateId)
+			{
+				// up to date
+				continue;
+			}
+
+			bHaveComponentsToUpdate = true;
+		}
+
+		if (bHaveComponentsToUpdate)
+		{
+			// We throttle, baking only one atlas per frame
+			if (NumGenerated > 0)
+			{
+				NumComponentsNeedingTextureBaking += It.Value().Num();
+			}
+			else
+			{
+				UTexture2D* HeightmapTexture = It.Key();
+				// 1/8 the res of the heightmap
+				FIntPoint AtlasSize(HeightmapTexture->GetSizeX() >> 3, HeightmapTexture->GetSizeY() >> 3);
+
+				TArray<FColor> AtlasSamples;
+				AtlasSamples.AddZeroed(AtlasSize.X * AtlasSize.Y);
+
+				for (ULandscapeComponent* Component : It.Value())
+				{
+					// not registered; ignore this component
+					if (!Component->SceneProxy)
+					{
+						continue;
+					}
+
+					int32 ComponentSamples = (SubsectionSizeQuads + 1) * NumSubsections;
+					check(FMath::IsPowerOfTwo(ComponentSamples));
+
+					int32 BakeSize = ComponentSamples >> 3;
+					TArray<FColor> Samples;
+					if (MaterialExportUtils::ExportBaseColor(Component, BakeSize, Samples))
+					{
+						int32 AtlasOffsetX = FMath::RoundToInt(Component->HeightmapScaleBias.Z * (float)HeightmapTexture->GetSizeX()) >> 3;
+						int32 AtlasOffsetY = FMath::RoundToInt(Component->HeightmapScaleBias.W * (float)HeightmapTexture->GetSizeY()) >> 3;
+						for (int32 y = 0; y < BakeSize; y++)
+						{
+							FMemory::Memcpy(&AtlasSamples[(y + AtlasOffsetY)*AtlasSize.Y + AtlasOffsetX], &Samples[y*BakeSize], sizeof(FColor)* BakeSize);
+						}
+						NumGenerated++;
+					}
+				}
+				UTexture2D* AtlasTexture = MaterialExportUtils::CreateTexture(GetOutermost(), HeightmapTexture->GetName() + TEXT("_BaseColor"), AtlasSize, AtlasSamples, TC_Default, TEXTUREGROUP_World, RF_Public, true);
+				AtlasTexture->MarkPackageDirty();
+
+				for (ULandscapeComponent* Component : It.Value())
+				{
+					Component->BakedTextureMaterialGuid = Component->MaterialInstance->GetMaterial()->StateId;
+					Component->GIBakedBaseColorTexture = AtlasTexture;
+					Component->MarkRenderStateDirty();
+				}
+			}
+		}
+	}
+
+	TotalComponentsNeedingTextureBaking += NumComponentsNeedingTextureBaking;
+
+	if (NumGenerated == 0)
+	{
+		// Don't check if we need to update anything for another 60 frames
+		UpdateBakedTexturesCountdown = 60;
+	}
+}
+#endif
+
+void ALandscapeProxy::InvalidateGeneratedComponentData(const TSet<ULandscapeComponent*>& Components)
+{
+	TMap<ALandscapeProxy*, TSet<ULandscapeComponent*>> ByProxy;
+	for (auto Component : Components)
+	{
+		Component->BakedTextureMaterialGuid.Invalidate();
+
+		ByProxy.FindOrAdd(Component->GetLandscapeProxy()).Add(Component);
+	}
+	for (auto It = ByProxy.CreateConstIterator(); It; ++It)
+	{
+		It.Key()->FlushGrassComponents(&It.Value());
+	}
+}
+
 #undef LOCTEXT_NAMESPACE
\ No newline at end of file
diff --git a/Engine/Source/Runtime/Landscape/Private/LandscapeGrass.cpp b/Engine/Source/Runtime/Landscape/Private/LandscapeGrass.cpp
--- a/Engine/Source/Runtime/Landscape/Private/LandscapeGrass.cpp
+++ b/Engine/Source/Runtime/Landscape/Private/LandscapeGrass.cpp
@@ -1153,20 +1153,6 @@ struct FAsyncGrassBuilder : public FGrassBuilderBase
 	}
 };
 
-/* Static version to distribute to the appropriate proxy for each component */
-void ALandscapeProxy::DistributeFlushGrassComponents(const TSet<ULandscapeComponent*>& Components, bool bFlushGrassMaps)
-{
-	TMap<ALandscapeProxy*, TSet<ULandscapeComponent*>> ByProxy;
-	for (auto Component : Components)
-	{
-		ByProxy.FindOrAdd(Component->GetLandscapeProxy()).Add(Component);
-	}
-	for (auto It = ByProxy.CreateConstIterator(); It; ++It)
-	{
-		It.Key()->FlushGrassComponents(&It.Value(), bFlushGrassMaps);
-	}
-}
-
 void ALandscapeProxy::FlushGrassComponents(const TSet<ULandscapeComponent*>* OnlyForComponents, bool bFlushGrassMaps)
 {
 	if (OnlyForComponents)
@@ -1263,6 +1249,7 @@ TArray<ULandscapeGrassType*> ALandscapeProxy::GetGrassTypes() const
 #if WITH_EDITOR
 int32 ALandscapeProxy::TotalComponentsNeedingGrassMapRender = 0;
 int32 ALandscapeProxy::TotalTexturesToStreamForVisibleGrassMapRender = 0;
+int32 ALandscapeProxy::TotalComponentsNeedingTextureBaking = 0;
 #endif
 
 void ALandscapeProxy::UpdateGrass(const TArray<FVector>& Cameras, bool bForceSync)
diff --git a/Engine/Source/Runtime/Renderer/Private/ScenePrivate.h b/Engine/Source/Runtime/Renderer/Private/ScenePrivate.h
--- a/Engine/Source/Runtime/Renderer/Private/ScenePrivate.h
+++ b/Engine/Source/Runtime/Renderer/Private/ScenePrivate.h
@@ -1658,6 +1658,11 @@ public:
 		return SkyLight && !SkyLight->bHasStaticLighting && GSupportsRenderTargetFormat_PF_FloatRGBA;
 	}
 
+	virtual TArray<FPrimitiveComponentId> GetScenePrimitiveComponentIds() const override
+	{
+		return PrimitiveComponentIds;
+	}
+
 private:
 
 	/**
