Fix for Crash in FLandscapeComponentSceneProxy::GetStaticBatchElementVisibility with tiny component size
UE-11046

diff --git a/Engine/Source/Runtime/Landscape/Private/LandscapeRender.cpp b/Engine/Source/Runtime/Landscape/Private/LandscapeRender.cpp
--- a/Engine/Source/Runtime/Landscape/Private/LandscapeRender.cpp
+++ b/Engine/Source/Runtime/Landscape/Private/LandscapeRender.cpp
@@ -1251,7 +1251,7 @@ float FLandscapeComponentSceneProxy::CalcDesiredLOD(const class FSceneView& View
 		SubsectionLODBias   = Neighbors[3] ? Neighbors[3]->LODBias : 0;
 	}
 
-	const int32 MinStreamedLOD = SubsectionHeightmapTexture ? ((FTexture2DResource*)SubsectionHeightmapTexture->Resource)->GetCurrentFirstMip() : 0;
+	const int32 MinStreamedLOD = SubsectionHeightmapTexture ? FMath::Min<int32>(((FTexture2DResource*)SubsectionHeightmapTexture->Resource)->GetCurrentFirstMip(), FMath::CeilLogTwo(SubsectionSizeVerts) - 1) : 0;
 
 	float fLOD = FLT_MAX;
 
