Fix for crash 10 seconds after loading landscape map

diff --git a/Engine/Source/Runtime/Landscape/Private/Landscape.cpp b/Engine/Source/Runtime/Landscape/Private/Landscape.cpp
--- a/Engine/Source/Runtime/Landscape/Private/Landscape.cpp
+++ b/Engine/Source/Runtime/Landscape/Private/Landscape.cpp
@@ -2295,7 +2295,7 @@ void ALandscapeProxy::UpdateBakedTextures()
 						int32 AtlasOffsetY = FMath::RoundToInt(Component->HeightmapScaleBias.W * (float)HeightmapTexture->GetSizeY()) >> 3;
 						for (int32 y = 0; y < BakeSize; y++)
 						{
-							FMemory::Memcpy(&AtlasSamples[(y + AtlasOffsetY)*AtlasSize.Y + AtlasOffsetX], &Samples[y*BakeSize], sizeof(FColor)* BakeSize);
+							FMemory::Memcpy(&AtlasSamples[(y + AtlasOffsetY)*AtlasSize.X + AtlasOffsetX], &Samples[y*BakeSize], sizeof(FColor)* BakeSize);
 						}
 						NumGenerated++;
 					}
