Fixed: Gaps between landscape components
UE-8932

diff --git a/Engine/Source/Runtime/Landscape/Private/Landscape.cpp b/Engine/Source/Runtime/Landscape/Private/Landscape.cpp
--- a/Engine/Source/Runtime/Landscape/Private/Landscape.cpp
+++ b/Engine/Source/Runtime/Landscape/Private/Landscape.cpp
@@ -432,6 +432,18 @@ void ULandscapeComponent::PostLoad()
 #if WITH_EDITOR
 	if (GIsEditor && !HasAnyFlags(RF_ClassDefaultObject))
 	{
+		// This is to ensure that component relative location is exact section base offset value
+		float CheckRelativeLocationX = float(SectionBaseX - LandscapeProxy->LandscapeSectionOffset.X);
+		float CheckRelativeLocationY = float(SectionBaseY - LandscapeProxy->LandscapeSectionOffset.Y);
+		if (CheckRelativeLocationX != RelativeLocation.X || 
+			CheckRelativeLocationY != RelativeLocation.Y)
+		{
+			UE_LOG(LogLandscape, Warning, TEXT("LandscapeComponent RelativeLocation disagrees with its section base, attempted automated fix: '%s', %f,%f vs %f,%f."),
+				*GetFullName(), RelativeLocation.X, RelativeLocation.Y, CheckRelativeLocationX, CheckRelativeLocationY);
+			RelativeLocation.X = CheckRelativeLocationX;
+			RelativeLocation.Y = CheckRelativeLocationY;
+		}
+						
 		// Remove standalone flags from data textures to ensure data is unloaded in the editor when reverting an unsaved level.
 		// Previous version of landscape set these flags on creation.
 		if (HeightmapTexture && HeightmapTexture->HasAnyFlags(RF_Standalone))
diff --git a/Engine/Source/Runtime/Landscape/Private/LandscapeCollision.cpp b/Engine/Source/Runtime/Landscape/Private/LandscapeCollision.cpp
--- a/Engine/Source/Runtime/Landscape/Private/LandscapeCollision.cpp
+++ b/Engine/Source/Runtime/Landscape/Private/LandscapeCollision.cpp
@@ -1305,6 +1305,22 @@ void ULandscapeHeightfieldCollisionComponent::PostLoad()
 	{
 		bShouldSaveCookedDataToDDC[0] = true;
 		bShouldSaveCookedDataToDDC[1] = true;
+
+		ALandscapeProxy* LandscapeProxy = GetLandscapeProxy();
+		if (ensure(LandscapeProxy) && GIsEditor)
+		{
+			// This is to ensure that component relative location is exact section base offset value
+			float CheckRelativeLocationX = float(SectionBaseX - LandscapeProxy->LandscapeSectionOffset.X);
+			float CheckRelativeLocationY = float(SectionBaseY - LandscapeProxy->LandscapeSectionOffset.Y);
+			if (CheckRelativeLocationX != RelativeLocation.X || 
+				CheckRelativeLocationY != RelativeLocation.Y)
+			{
+				UE_LOG(LogLandscape, Warning, TEXT("ULandscapeHeightfieldCollisionComponent RelativeLocation disagrees with its section base, attempted automated fix: '%s', %f,%f vs %f,%f."),
+					*GetFullName(), RelativeLocation.X, RelativeLocation.Y, CheckRelativeLocationX, CheckRelativeLocationY);
+				RelativeLocation.X = CheckRelativeLocationX;
+				RelativeLocation.Y = CheckRelativeLocationY;
+			}
+		}
 	}
 #endif//WITH_EDITOR
 }
