UE4 - merge GDC branch, code [REDACTED_HANDLE] to main

diff --git a/Engine/Source/Runtime/Landscape/Private/Landscape.cpp b/Engine/Source/Runtime/Landscape/Private/Landscape.cpp
--- a/Engine/Source/Runtime/Landscape/Private/Landscape.cpp
+++ b/Engine/Source/Runtime/Landscape/Private/Landscape.cpp
@@ -1727,6 +1727,9 @@ void ULandscapeInfo::RegisterActorComponent(ULandscapeComponent* Component, bool
 		FMessageLog("MapCheck").Warning()
 			->AddToken(FUObjectToken::Create(OurProxy))
 			->AddToken(FTextToken::Create(FText::Format(LOCTEXT("MapCheck_Message_LandscapeComponentPostLoad_Warning", "Landscape {ProxyName1} of {LevelName1} has overlapping render components with {ProxyName2} of {LevelName2} at location ({XLocation}, {YLocation})."), Arguments)))
+#if WITH_EDITOR
+			->AddToken(FActionToken::Create(LOCTEXT("MapCheck_RemoveDuplicateLandscapeComponent", "Delete Duplicate"), LOCTEXT("MapCheck_RemoveDuplicateLandscapeComponentDesc", "Deletes the duplicate landscape component."), FOnActionTokenExecuted::CreateUObject(OurProxy, &ALandscapeProxy::RemoveOverlappingComponent, Component), true))
+#endif
 			->AddToken(FMapErrorToken::Create(FMapErrors::LandscapeComponentPostLoad_Warning));
 
 		// Show MapCheck window
@@ -2217,6 +2220,23 @@ ULandscapeMeshProxyComponent::ULandscapeMeshProxyComponent(const FObjectInitiali
 }
 
 #if WITH_EDITOR
+
+void ULandscapeComponent::SerializeStateHashes(FArchive& Ar)
+{
+	if (MaterialInstance)
+	{
+		Ar << MaterialInstance->GetMaterial()->StateId;
+	}
+
+	FGuid HeightmapGuid = HeightmapTexture->Source.GetId();
+	Ar << HeightmapGuid;
+	for (auto WeightmapTexture : WeightmapTextures)
+	{
+		FGuid WeightmapGuid = WeightmapTexture->Source.GetId();
+		Ar << WeightmapGuid;
+	}
+}
+
 void ALandscapeProxy::UpdateBakedTextures()
 {
 	// See if we can render
@@ -2231,11 +2251,30 @@ void ALandscapeProxy::UpdateBakedTextures()
 		return;
 	}
 
+	// Stores the components and their state hash data for a single atlas
+	struct FBakedTextureSourceInfo
+	{
+		// pointer as FMemoryWriter caches the address of the FBufferArchive, and this struct could be relocated on a realloc.
+		FBufferArchive* ComponentStateAr;
+		TArray<ULandscapeComponent*> Components;
+
+		FBakedTextureSourceInfo()
+		{
+			ComponentStateAr = new FBufferArchive();
+		}
+		~FBakedTextureSourceInfo()
+		{
+			delete ComponentStateAr;
+		}
+	};
+
 	// Group components by heightmap texture
-	TMap<UTexture2D*, TArray<ULandscapeComponent*>> ComponentsByHeightmap;
+	TMap<UTexture2D*, FBakedTextureSourceInfo> ComponentsByHeightmap;
 	for (ULandscapeComponent* Component : LandscapeComponents)
 	{
-		ComponentsByHeightmap.FindOrAdd(Component->HeightmapTexture).Add(Component);
+		FBakedTextureSourceInfo& Info = ComponentsByHeightmap.FindOrAdd(Component->HeightmapTexture);
+		Info.Components.Add(Component);
+		Component->SerializeStateHashes(*Info.ComponentStateAr);
 	}
 
 	TotalComponentsNeedingTextureBaking -= NumComponentsNeedingTextureBaking;
@@ -2244,8 +2283,10 @@ void ALandscapeProxy::UpdateBakedTextures()
 
 	for (auto It = ComponentsByHeightmap.CreateConstIterator(); It; ++It)
 	{
-		bool bHaveComponentsToUpdate = false;
-		for (ULandscapeComponent* Component : It.Value())
+		const FBakedTextureSourceInfo& Info = It.Value();
+
+		bool bCanBake = true;
+		for (ULandscapeComponent* Component : Info.Components)
 		{
 			// not registered; ignore this component
 			if (!Component->SceneProxy)
@@ -2255,71 +2296,88 @@ void ALandscapeProxy::UpdateBakedTextures()
 
 			// Check we can render the material
 			UMaterialInstance* MaterialInstance = Component->MaterialInstance;
-
-			if (!MaterialInstance || !MaterialInstance->GetMaterialResource(World->FeatureLevel)->HasValidGameThreadShaderMap())
+			if (!MaterialInstance)
 			{
 				// Cannot render this component yet as it doesn't have a material; abandon the atlas for this heightmap
-				bHaveComponentsToUpdate = false;
+				bCanBake = false;
 				break;
 			}
 
-			if (Component->BakedTextureMaterialGuid == MaterialInstance->GetMaterial()->StateId)
+			FMaterialResource* MaterialResource = MaterialInstance->GetMaterialResource(World->FeatureLevel);
+			if (!MaterialResource || !MaterialResource->HasValidGameThreadShaderMap())
 			{
-				// up to date
-				continue;
+				// Cannot render this component yet as its shaders aren't compiled; abandon the atlas for this heightmap
+				bCanBake = false;
+				break;
 			}
-
-			bHaveComponentsToUpdate = true;
 		}
 
-		if (bHaveComponentsToUpdate)
+		if (bCanBake)
 		{
-			// We throttle, baking only one atlas per frame
-			if (NumGenerated > 0)
+			// Calculate a combined Guid-like ID we can use for this component
+			uint32 Hash[5];
+			FSHA1::HashBuffer(Info.ComponentStateAr->GetData(), Info.ComponentStateAr->Num(), (uint8*)Hash);
+			FGuid CombinedStateId = FGuid(Hash[0] ^ Hash[4], Hash[1], Hash[2], Hash[3]);
+
+			bool bNeedsBake = false;
+			for (ULandscapeComponent* Component : Info.Components)
 			{
-				NumComponentsNeedingTextureBaking += It.Value().Num();
+				if (Component->BakedTextureMaterialGuid != CombinedStateId)
+				{
+					bNeedsBake = true;
+					break;
+				}
 			}
-			else
+			
+			if (bNeedsBake)
 			{
-				UTexture2D* HeightmapTexture = It.Key();
-				// 1/8 the res of the heightmap
-				FIntPoint AtlasSize(HeightmapTexture->GetSizeX() >> 3, HeightmapTexture->GetSizeY() >> 3);
+				// We throttle, baking only one atlas per frame
+				if (NumGenerated > 0)
+				{
+					NumComponentsNeedingTextureBaking += Info.Components.Num();
+				}
+				else
+				{
+					UTexture2D* HeightmapTexture = It.Key();
+					// 1/8 the res of the heightmap
+					FIntPoint AtlasSize(HeightmapTexture->GetSizeX() >> 3, HeightmapTexture->GetSizeY() >> 3);
 
-				TArray<FColor> AtlasSamples;
-				AtlasSamples.AddZeroed(AtlasSize.X * AtlasSize.Y);
+					TArray<FColor> AtlasSamples;
+					AtlasSamples.AddZeroed(AtlasSize.X * AtlasSize.Y);
 
-				for (ULandscapeComponent* Component : It.Value())
-				{
-					// not registered; ignore this component
-					if (!Component->SceneProxy)
+					for (ULandscapeComponent* Component : Info.Components)
 					{
-						continue;
-					}
+						// not registered; ignore this component
+						if (!Component->SceneProxy)
+						{
+							continue;
+						}
 
-					int32 ComponentSamples = (SubsectionSizeQuads + 1) * NumSubsections;
-					check(FMath::IsPowerOfTwo(ComponentSamples));
+						int32 ComponentSamples = (SubsectionSizeQuads + 1) * NumSubsections;
+						check(FMath::IsPowerOfTwo(ComponentSamples));
 
-					int32 BakeSize = ComponentSamples >> 3;
-					TArray<FColor> Samples;
-					if (MaterialExportUtils::ExportBaseColor(Component, BakeSize, Samples))
-					{
-						int32 AtlasOffsetX = FMath::RoundToInt(Component->HeightmapScaleBias.Z * (float)HeightmapTexture->GetSizeX()) >> 3;
-						int32 AtlasOffsetY = FMath::RoundToInt(Component->HeightmapScaleBias.W * (float)HeightmapTexture->GetSizeY()) >> 3;
-						for (int32 y = 0; y < BakeSize; y++)
+						int32 BakeSize = ComponentSamples >> 3;
+						TArray<FColor> Samples;
+						if (MaterialExportUtils::ExportBaseColor(Component, BakeSize, Samples))
 						{
-							FMemory::Memcpy(&AtlasSamples[(y + AtlasOffsetY)*AtlasSize.X + AtlasOffsetX], &Samples[y*BakeSize], sizeof(FColor)* BakeSize);
+							int32 AtlasOffsetX = FMath::RoundToInt(Component->HeightmapScaleBias.Z * (float)HeightmapTexture->GetSizeX()) >> 3;
+							int32 AtlasOffsetY = FMath::RoundToInt(Component->HeightmapScaleBias.W * (float)HeightmapTexture->GetSizeY()) >> 3;
+							for (int32 y = 0; y < BakeSize; y++)
+							{
+								FMemory::Memcpy(&AtlasSamples[(y + AtlasOffsetY)*AtlasSize.Y + AtlasOffsetX], &Samples[y*BakeSize], sizeof(FColor)* BakeSize);
+							}
+							NumGenerated++;
 						}
-						NumGenerated++;
 					}
-				}
-				UTexture2D* AtlasTexture = MaterialExportUtils::CreateTexture(GetOutermost(), HeightmapTexture->GetName() + TEXT("_BaseColor"), AtlasSize, AtlasSamples, TC_Default, TEXTUREGROUP_World, RF_Public, true);
-				AtlasTexture->MarkPackageDirty();
+					UTexture2D* AtlasTexture = MaterialExportUtils::CreateTexture(GetOutermost(), HeightmapTexture->GetName() + TEXT("_BaseColor"), AtlasSize, AtlasSamples, TC_Default, TEXTUREGROUP_World, RF_Public, true, CombinedStateId);
+					AtlasTexture->MarkPackageDirty();
 
-				for (ULandscapeComponent* Component : It.Value())
-				{
-					Component->BakedTextureMaterialGuid = Component->MaterialInstance->GetMaterial()->StateId;
-					Component->GIBakedBaseColorTexture = AtlasTexture;
-					Component->MarkRenderStateDirty();
+					for (ULandscapeComponent* Component : Info.Components)
+					{
+						Component->BakedTextureMaterialGuid = CombinedStateId;
+						Component->GIBakedBaseColorTexture = AtlasTexture;
+						Component->MarkRenderStateDirty();
+					}
 				}
 			}
 		}
diff --git a/Engine/Source/Runtime/Landscape/Private/LandscapeCollision.cpp b/Engine/Source/Runtime/Landscape/Private/LandscapeCollision.cpp
--- a/Engine/Source/Runtime/Landscape/Private/LandscapeCollision.cpp
+++ b/Engine/Source/Runtime/Landscape/Private/LandscapeCollision.cpp
@@ -262,6 +262,8 @@ void ULandscapeHeightfieldCollisionComponent::CreateCollisionObject()
 	// If we have not created a heightfield yet - do it now.
 	if (!IsValidRef(HeightfieldRef))
 	{
+		UWorld* World = GetWorld();
+
 		FPhysXHeightfieldRef* ExistingHeightfieldRef = nullptr;
 		bool bCheckDDC = true;
 
@@ -293,6 +295,10 @@ void ULandscapeHeightfieldCollisionComponent::CreateCollisionObject()
 			// Prepare heightfield data
 			static FName PhysicsFormatName(FPlatformProperties::GetPhysicsFormat());
 			CookCollsionData(PhysicsFormatName, false, bCheckDDC, CookedCollisionData, CookedPhysicalMaterials);
+
+			// The World will clean up any speculatively-loaded data we didn't end up using.
+			SpeculativeDDCRequest.Reset();
+
 #endif //WITH_EDITOR
 
 			if (CookedCollisionData.Num())
@@ -312,14 +318,14 @@ void ULandscapeHeightfieldCollisionComponent::CreateCollisionObject()
 
 				// Release cooked collison data
 				// In cooked builds created collision object will never be deleted while component is alive, so we don't need this data anymore
-				if (FPlatformProperties::RequiresCookedData() || GetWorld()->IsGameWorld())
+				if (FPlatformProperties::RequiresCookedData() || World->IsGameWorld())
 				{
 					CookedCollisionData.Empty();
 				}
 
 #if WITH_EDITOR
 				// Create heightfield for the landscape editor (no holes in it)
-				if (!GetWorld()->IsGameWorld())
+				if (!World->IsGameWorld())
 				{
 					TArray<UPhysicalMaterial*> CookedMaterialsEd;
 					if (CookCollsionData(PhysicsFormatName, true, bCheckDDC, CookedCollisionDataEd, CookedMaterialsEd))
@@ -336,6 +342,26 @@ void ULandscapeHeightfieldCollisionComponent::CreateCollisionObject()
 }
 
 #if WITH_EDITOR
+void ULandscapeHeightfieldCollisionComponent::SpeculativelyLoadAsyncDDCCollsionData()
+{
+#if WITH_PHYSX
+	if (GetLinkerUE4Version() >= VER_UE4_LANDSCAPE_SERIALIZE_PHYSICS_MATERIALS)
+	{
+		UWorld* World = GetWorld();
+		if (World && HeightfieldGuid.IsValid() && CookedPhysicalMaterials.Num() > 0 && GSharedHeightfieldRefs.FindRef(HeightfieldGuid) == nullptr)
+		{
+			static FName PhysicsFormatName(FPlatformProperties::GetPhysicsFormat());
+
+			FString Key = GetHFDDCKeyString(PhysicsFormatName, false, HeightfieldGuid);
+			uint32 Handle = GetDerivedDataCacheRef().GetAsynchronous(*Key);
+			check(!SpeculativeDDCRequest.IsValid());
+			SpeculativeDDCRequest = MakeShareable(new FAsyncPreRegisterDDCRequest(Key, Handle));
+			World->AsyncPreRegisterDDCRequests.Add(SpeculativeDDCRequest);
+		}
+	}
+#endif
+}
+
 bool ULandscapeHeightfieldCollisionComponent::CookCollsionData(const FName& Format, bool bUseDefMaterial, bool bCheckDDC, TArray<uint8>& OutCookedData, TArray<UPhysicalMaterial*>& OutMaterials) const
 {
 #if WITH_PHYSX
@@ -347,7 +373,23 @@ bool ULandscapeHeightfieldCollisionComponent::CookCollsionData(const FName& Form
 		// Ensure that content was saved with physical materials before using DDC data
 		if (GetLinkerUE4Version() >= VER_UE4_LANDSCAPE_SERIALIZE_PHYSICS_MATERIALS)
 		{
-			if (GetDerivedDataCacheRef().GetSynchronous(*GetHFDDCKeyString(Format, bUseDefMaterial, HeightfieldGuid), OutCookedData))
+			FString DDCKey = GetHFDDCKeyString(Format, bUseDefMaterial, HeightfieldGuid);
+
+			// Check if the speculatively-loaded data loaded and is what we wanted
+			if (SpeculativeDDCRequest.IsValid() && DDCKey == SpeculativeDDCRequest->GetKey())
+			{
+				SpeculativeDDCRequest->WaitAsynchronousCompletion();
+				bool bSuccess = SpeculativeDDCRequest->GetAsynchronousResults(OutCookedData);
+				// World will clean up remaining reference
+				SpeculativeDDCRequest.Reset();
+				if (bSuccess)
+				{
+					bShouldSaveCookedDataToDDC[CookedDataIndex] = false;
+					return true;
+				}
+			}
+
+			if (GetDerivedDataCacheRef().GetSynchronous(*DDCKey, OutCookedData))
 			{
 				bShouldSaveCookedDataToDDC[CookedDataIndex] = false;
 				return true;
@@ -485,7 +527,23 @@ bool ULandscapeMeshCollisionComponent::CookCollsionData(const FName& Format, boo
 		// Ensure that content was saved with physical materials before using DDC data
 		if (GetLinkerUE4Version() >= VER_UE4_LANDSCAPE_SERIALIZE_PHYSICS_MATERIALS)
 		{
-			if (GetDerivedDataCacheRef().GetSynchronous(*GetHFDDCKeyString(Format, bUseDefMaterial, HeightfieldGuid), OutCookedData))
+			FString DDCKey = GetHFDDCKeyString(Format, bUseDefMaterial, HeightfieldGuid);
+
+			// Check if the speculatively-loaded data loaded and is what we wanted
+			if (SpeculativeDDCRequest.IsValid() && DDCKey == SpeculativeDDCRequest->GetKey())
+			{
+				SpeculativeDDCRequest->WaitAsynchronousCompletion();
+				bool bSuccess = SpeculativeDDCRequest->GetAsynchronousResults(OutCookedData);
+				// World will clean up remaining reference
+				SpeculativeDDCRequest.Reset();
+				if (bSuccess)
+				{
+					bShouldSaveCookedDataToDDC[CookedDataIndex] = false;
+					return true;
+				}
+			}
+
+			if (GetDerivedDataCacheRef().GetSynchronous(*DDCKey, OutCookedData))
 			{
 				bShouldSaveCookedDataToDDC[CookedDataIndex] = false;
 				return true;
@@ -1259,7 +1317,7 @@ bool ULandscapeHeightfieldCollisionComponent::ComponentIsTouchingSelectionFrustu
 }
 #endif
 
-bool ULandscapeHeightfieldCollisionComponent::DoCustomNavigableGeometryExport(FNavigableGeometryExport* GeomExport) const
+bool ULandscapeHeightfieldCollisionComponent::DoCustomNavigableGeometryExport(FNavigableGeometryExport& GeomExport) const
 {
 	check(IsInGameThread());
 #if WITH_PHYSX
@@ -1267,14 +1325,66 @@ bool ULandscapeHeightfieldCollisionComponent::DoCustomNavigableGeometryExport(FN
 	{
 		FTransform HFToW = ComponentToWorld;
 		HFToW.MultiplyScale3D(FVector(CollisionScale, CollisionScale, LANDSCAPE_ZSCALE));
-
-		GeomExport->ExportPxHeightField(HeightfieldRef->RBHeightfield, HFToW);
+		GeomExport.ExportPxHeightField(HeightfieldRef->RBHeightfield, HFToW);
 	}
 #endif// WITH_PHYSX
 	return false;
 }
 
-bool ULandscapeMeshCollisionComponent::DoCustomNavigableGeometryExport(struct FNavigableGeometryExport* GeomExport) const
+void ULandscapeHeightfieldCollisionComponent::GatherGeometrySlice(FNavigableGeometryExport& GeomExport, const FBox& SliceBox) const
+{
+	// note that this function can get called off game thread
+	if (CachedHeightFieldSamples.IsEmpty() == false)
+	{
+		FTransform HFToW = ComponentToWorld;
+		HFToW.MultiplyScale3D(FVector(CollisionScale, CollisionScale, LANDSCAPE_ZSCALE));
+
+		GeomExport.ExportHeightFieldSlice(CachedHeightFieldSamples, HeightfieldRowsCount, HeightfieldColumnsCount, HFToW, SliceBox);
+	}
+}
+
+ENavDataGatheringMode ULandscapeHeightfieldCollisionComponent::GetGeometryGatheringMode() const
+{ 
+	ALandscapeProxy* Proxy = GetLandscapeProxy();
+	return Proxy ? Proxy->NavigationGeometryGatheringMode : ENavDataGatheringMode::Default;
+}
+
+void ULandscapeHeightfieldCollisionComponent::PrepareGeometryExportSync()
+{
+	//check(IsInGameThread());
+#if WITH_PHYSX
+	if (IsValidRef(HeightfieldRef) && HeightfieldRef->RBHeightfield != nullptr && CachedHeightFieldSamples.IsEmpty())
+	{
+		const UWorld* World = GetWorld();
+
+		if (World != nullptr)
+		{
+			HeightfieldRowsCount = HeightfieldRef->RBHeightfield->getNbRows();
+			HeightfieldColumnsCount = HeightfieldRef->RBHeightfield->getNbColumns();
+				
+			if (CachedHeightFieldSamples.Heights.Num() != HeightfieldRowsCount * HeightfieldRowsCount)
+			{
+				QUICK_SCOPE_CYCLE_COUNTER(STAT_NavMesh_ExportPxHeightField_saveCells);
+
+				CachedHeightFieldSamples.Heights.SetNumUninitialized(HeightfieldRowsCount * HeightfieldRowsCount);
+
+				TArray<PxHeightFieldSample> HFSamples;
+				HFSamples.SetNumUninitialized(HeightfieldRowsCount * HeightfieldRowsCount);
+				HeightfieldRef->RBHeightfield->saveCells(HFSamples.GetData(), HFSamples.Num()*HFSamples.GetTypeSize());
+
+				for (int32 SampleIndex = 0; SampleIndex < HFSamples.Num(); ++SampleIndex)
+				{
+					const PxHeightFieldSample& Sample = HFSamples[SampleIndex];
+					CachedHeightFieldSamples.Heights[SampleIndex] = Sample.height;
+					CachedHeightFieldSamples.Holes.Add((Sample.materialIndex0 == PxHeightFieldMaterial::eHOLE));
+				}
+			}
+		}
+	}
+#endif// WITH_PHYSX
+}
+
+bool ULandscapeMeshCollisionComponent::DoCustomNavigableGeometryExport(FNavigableGeometryExport& GeomExport) const
 {
 	check(IsInGameThread());
 #if WITH_PHYSX
@@ -1285,11 +1395,11 @@ bool ULandscapeMeshCollisionComponent::DoCustomNavigableGeometryExport(struct FN
 
 		if (MeshRef->RBTriangleMesh->getTriangleMeshFlags() & PxTriangleMeshFlag::eHAS_16BIT_TRIANGLE_INDICES)
 		{
-			GeomExport->ExportPxTriMesh16Bit(MeshRef->RBTriangleMesh, MeshToW);
+			GeomExport.ExportPxTriMesh16Bit(MeshRef->RBTriangleMesh, MeshToW);
 		}
 		else
 		{
-			GeomExport->ExportPxTriMesh32Bit(MeshRef->RBTriangleMesh, MeshToW);
+			GeomExport.ExportPxTriMesh32Bit(MeshRef->RBTriangleMesh, MeshToW);
 		}
 	}
 #endif// WITH_PHYSX
@@ -1321,6 +1431,12 @@ void ULandscapeHeightfieldCollisionComponent::PostLoad()
 				RelativeLocation.Y = CheckRelativeLocationY;
 			}
 		}
+
+		UWorld* World = GetWorld();
+		if (World && World->IsGameWorld())
+		{
+			SpeculativelyLoadAsyncDDCCollsionData();
+		}
 	}
 #endif//WITH_EDITOR
 }
@@ -1797,4 +1913,7 @@ ULandscapeHeightfieldCollisionComponent::ULandscapeHeightfieldCollisionComponent
 	Mobility = EComponentMobility::Static;
 	bCanEverAffectNavigation = true;
 	bHasCustomNavigableGeometry = EHasCustomNavigableGeometry::Yes;
+
+	HeightfieldRowsCount = -1;
+	HeightfieldColumnsCount = -1;
 }
diff --git a/Engine/Source/Runtime/Landscape/Private/LandscapeEdit.cpp b/Engine/Source/Runtime/Landscape/Private/LandscapeEdit.cpp
--- a/Engine/Source/Runtime/Landscape/Private/LandscapeEdit.cpp
+++ b/Engine/Source/Runtime/Landscape/Private/LandscapeEdit.cpp
@@ -1508,39 +1508,43 @@ float ULandscapeComponent::GetLayerWeightAtLocation(const FVector& InLocation, U
 
 	// Find location
 	// TODO: Root landscape isn't always loaded, would Proxy suffice?
-	ALandscape* Landscape = GetLandscapeActor();
-	const FVector DrawScale = Landscape->GetRootComponent()->RelativeScale3D;
-	float TestX = (InLocation.X - Landscape->GetActorLocation().X) / DrawScale.X - (float)GetSectionBase().X;
-	float TestY = (InLocation.Y - Landscape->GetActorLocation().Y) / DrawScale.Y - (float)GetSectionBase().Y;
-
-	// Find data
-	int32 X1 = FMath::FloorToInt(TestX);
-	int32 Y1 = FMath::FloorToInt(TestY);
-	int32 X2 = FMath::CeilToInt(TestX);
-	int32 Y2 = FMath::CeilToInt(TestY);
-
-	int32 Stride = (SubsectionSizeQuads + 1) * NumSubsections;
-
-	// Min is to prevent the sampling of the final column from overflowing
-	int32 IdxX1 = FMath::Min<int32>(((X1 / SubsectionSizeQuads) * (SubsectionSizeQuads + 1)) + (X1 % SubsectionSizeQuads), Stride - 1);
-	int32 IdxY1 = FMath::Min<int32>(((Y1 / SubsectionSizeQuads) * (SubsectionSizeQuads + 1)) + (Y1 % SubsectionSizeQuads), Stride - 1);
-	int32 IdxX2 = FMath::Min<int32>(((X2 / SubsectionSizeQuads) * (SubsectionSizeQuads + 1)) + (X2 % SubsectionSizeQuads), Stride - 1);
-	int32 IdxY2 = FMath::Min<int32>(((Y2 / SubsectionSizeQuads) * (SubsectionSizeQuads + 1)) + (Y2 % SubsectionSizeQuads), Stride - 1);
-
-	// sample
-	float Sample11 = (float)((*LayerCache)[IdxX1 + Stride*IdxY1]) / 255.0f;
-	float Sample21 = (float)((*LayerCache)[IdxX2 + Stride*IdxY1]) / 255.0f;
-	float Sample12 = (float)((*LayerCache)[IdxX1 + Stride*IdxY2]) / 255.0f;
-	float Sample22 = (float)((*LayerCache)[IdxX2 + Stride*IdxY2]) / 255.0f;
-
-	float LerpX = FMath::Fractional(TestX);
-	float LerpY = FMath::Fractional(TestY);
+	if (ALandscape* Landscape = GetLandscapeActor())
+	{
+		const FVector DrawScale = Landscape->GetRootComponent()->RelativeScale3D;
+		float TestX = (InLocation.X - Landscape->GetActorLocation().X) / DrawScale.X - (float)GetSectionBase().X;
+		float TestY = (InLocation.Y - Landscape->GetActorLocation().Y) / DrawScale.Y - (float)GetSectionBase().Y;
+
+		// Find data
+		int32 X1 = FMath::FloorToInt(TestX);
+		int32 Y1 = FMath::FloorToInt(TestY);
+		int32 X2 = FMath::CeilToInt(TestX);
+		int32 Y2 = FMath::CeilToInt(TestY);
+
+		int32 Stride = (SubsectionSizeQuads + 1) * NumSubsections;
+
+		// Min is to prevent the sampling of the final column from overflowing
+		int32 IdxX1 = FMath::Min<int32>(((X1 / SubsectionSizeQuads) * (SubsectionSizeQuads + 1)) + (X1 % SubsectionSizeQuads), Stride - 1);
+		int32 IdxY1 = FMath::Min<int32>(((Y1 / SubsectionSizeQuads) * (SubsectionSizeQuads + 1)) + (Y1 % SubsectionSizeQuads), Stride - 1);
+		int32 IdxX2 = FMath::Min<int32>(((X2 / SubsectionSizeQuads) * (SubsectionSizeQuads + 1)) + (X2 % SubsectionSizeQuads), Stride - 1);
+		int32 IdxY2 = FMath::Min<int32>(((Y2 / SubsectionSizeQuads) * (SubsectionSizeQuads + 1)) + (Y2 % SubsectionSizeQuads), Stride - 1);
+
+		// sample
+		float Sample11 = (float)((*LayerCache)[IdxX1 + Stride*IdxY1]) / 255.0f;
+		float Sample21 = (float)((*LayerCache)[IdxX2 + Stride*IdxY1]) / 255.0f;
+		float Sample12 = (float)((*LayerCache)[IdxX1 + Stride*IdxY2]) / 255.0f;
+		float Sample22 = (float)((*LayerCache)[IdxX2 + Stride*IdxY2]) / 255.0f;
+
+		float LerpX = FMath::Fractional(TestX);
+		float LerpY = FMath::Fractional(TestY);
 
 		// Bilinear interpolate
-	return FMath::Lerp(
+		return FMath::Lerp(
 			FMath::Lerp(Sample11, Sample21, LerpX),
 			FMath::Lerp(Sample12, Sample22, LerpX),
 			LerpY);
+	}
+	
+	return 0.f;	//if landscape is null we just return 0 instead of crashing. Seen cases where this happens, seems like a bug?
 }
 
 void ULandscapeComponent::GetComponentExtent(int32& MinX, int32& MinY, int32& MaxX, int32& MaxY) const
@@ -4686,6 +4690,20 @@ UTexture2D* ALandscapeProxy::CreateLandscapeTexture(int32 InSizeX, int32 InSizeY
 	return NewTexture;
 }
 
+void ALandscapeProxy::RemoveOverlappingComponent(ULandscapeComponent* Component)
+{
+	Modify();
+	Component->Modify();
+	if (Component->CollisionComponent.IsValid() && (Component->CollisionComponent->RenderComponent.Get() == Component || Component->CollisionComponent->RenderComponent.IsNull()))
+	{
+		Component->CollisionComponent->Modify();
+		CollisionComponents.Remove(Component->CollisionComponent.Get());
+		Component->CollisionComponent.Get()->DestroyComponent();
+	}
+	LandscapeComponents.Remove(Component);
+	Component->DestroyComponent();
+}
+
 #endif //WITH_EDITOR
 
 #undef LOCTEXT_NAMESPACE
diff --git a/Engine/Source/Runtime/Landscape/Private/LandscapeGrass.cpp b/Engine/Source/Runtime/Landscape/Private/LandscapeGrass.cpp
--- a/Engine/Source/Runtime/Landscape/Private/LandscapeGrass.cpp
+++ b/Engine/Source/Runtime/Landscape/Private/LandscapeGrass.cpp
@@ -45,11 +45,6 @@ static TAutoConsoleVariable<float> CVarMinTimeToKeepGrass(
 	5.0f,
 	TEXT("Minimum number of seconds before cached grass can be discarded; used to prevent thrashing."));
 
-static TAutoConsoleVariable<int32> CVarMinInstancesPerLeaf(
-	TEXT("grass.MinInstancesPerLeaf"),
-	256,
-	TEXT("Used to control the depth of hierarchical culling tree for grass. Clamped to 16-32000; a reasonable range."));
-
 static TAutoConsoleVariable<int32> CVarMaxInstancesPerComponent(
 	TEXT("grass.MaxInstancesPerComponent"),
 	65536,
@@ -60,6 +55,16 @@ static TAutoConsoleVariable<int32> CVarMaxAsyncTasks(
 	4,
 	TEXT("Used to control the number of hierarchical components created at a time."));
 
+static TAutoConsoleVariable<int32> CVarUseHaltonDistribution(
+	TEXT("grass.UseHaltonDistribution"),
+	0,
+	TEXT("Used to control the distribution of grass instances. If non-zero, use a halton sequence."));
+
+static TAutoConsoleVariable<float> CVarGrassDensityScale(
+	TEXT("grass.densityScale"),
+	1,
+	TEXT("Multiplier on all grass densities.  Do a grass.flushcache or grass.flushcachepie afterwards as appropriate."));
+
 static TAutoConsoleVariable<int32> CVarGrassEnable(
 	TEXT("grass.Enable"),
 	1,
@@ -859,7 +864,9 @@ struct FGrassBuilderBase
 	FGrassBuilderBase(ALandscapeProxy* Landscape, ULandscapeComponent* Component, const FGrassVariety& GrassVariety, int32 SqrtSubsections = 1, int32 SubX = 0, int32 SubY = 0)
 	{
 		bHaveValidData = true;
-		GrassDensity = GrassVariety.GrassDensity;
+
+		const float DensityScale = CVarGrassDensityScale.GetValueOnAnyThread();
+		GrassDensity = GrassVariety.GrassDensity * DensityScale;
 
 		DrawScale = Landscape->GetRootComponent()->RelativeScale3D;
 		DrawLoc = Landscape->GetActorLocation();
@@ -928,6 +935,21 @@ private:
 	int32 Stride;
 };
 
+template<uint32 Base>
+static FORCEINLINE float Halton(uint32 Index)
+{
+	float Result = 0.0f;
+	float InvBase = 1.0f / Base;
+	float Fraction = InvBase;
+	while( Index > 0 )
+	{
+		Result += ( Index % Base ) * Fraction;
+		Index /= Base;
+		Fraction *= InvBase;
+	}
+	return Result;
+}
+
 struct FAsyncGrassBuilder : public FGrassBuilderBase
 {
 	FLandscapeComponentGrassAccess GrassData;
@@ -937,17 +959,19 @@ struct FAsyncGrassBuilder : public FGrassBuilderBase
 	FRandomStream RandomStream;
 	FMatrix XForm;
 	FBox MeshBox;
+	int32 DesiredInstancesPerLeaf;
 
 	double RasterTime;
 	double BuildTime;
 	double InstanceTime;
 	int32 TotalInstances;
+	uint32 HaltonBaseIndex;
 
 	// output
 	FStaticMeshInstanceData InstanceBuffer;
 	TArray<FClusterNode> ClusterTree;
 
-	FAsyncGrassBuilder(ALandscapeProxy* Landscape, ULandscapeComponent* Component, const ULandscapeGrassType* GrassType, const FGrassVariety& GrassVariety, UHierarchicalInstancedStaticMeshComponent* HierarchicalInstancedStaticMeshComponent, int32 SqrtSubsections, int32 SubX, int32 SubY)
+	FAsyncGrassBuilder(ALandscapeProxy* Landscape, ULandscapeComponent* Component, const ULandscapeGrassType* GrassType, const FGrassVariety& GrassVariety, UHierarchicalInstancedStaticMeshComponent* HierarchicalInstancedStaticMeshComponent, int32 SqrtSubsections, int32 SubX, int32 SubY, uint32 InHaltonBaseIndex)
 		: FGrassBuilderBase(Landscape, Component, GrassVariety, SqrtSubsections, SubX, SubY)
 		, GrassData(Component, GrassType)
 	{
@@ -960,11 +984,10 @@ struct FAsyncGrassBuilder : public FGrassBuilderBase
 		RandomStream.Initialize(HierarchicalInstancedStaticMeshComponent->InstancingRandomSeed);
 
 		MeshBox = GrassVariety.GrassMesh->GetBounds().GetBox();
+		DesiredInstancesPerLeaf = HierarchicalInstancedStaticMeshComponent->DesiredInstancesPerLeaf();
 
-		RasterTime = 0.0;
-		BuildTime = 0.0;
-		InstanceTime = 0.0;
 		TotalInstances = 0;
+		HaltonBaseIndex = InHaltonBaseIndex;
 	}
 
 	void Build()
@@ -974,131 +997,203 @@ struct FAsyncGrassBuilder : public FGrassBuilderBase
 		RasterTime -= FPlatformTime::Seconds();
 
 		float Div = 1.0f / float(SqrtMaxInstances);
-		float MaxJitter1D = FMath::Clamp<float>(PlacementJitter, 0.0f, .99f) * Div * .5f;
-		FVector MaxJitter(MaxJitter1D, MaxJitter1D, 0.0f);
-		MaxJitter *= Extent;
-		Origin += Extent * (Div * 0.5f);
-		struct FInstanceLocal
+		TArray<FMatrix> InstanceTransforms;
+		if (HaltonBaseIndex)
 		{
-			FVector Pos;
-			bool bKeep;
-		};
-		TArray<FInstanceLocal> Instances;
-		Instances.AddUninitialized(SqrtMaxInstances * SqrtMaxInstances);
-		int32 NumKept = 0;
-		{
-			int32 InstanceIndex = 0;
-			for (int32 xStart = 0; xStart < SqrtMaxInstances; xStart++)
+			if (Extent.X < 0)
+			{
+				Origin.X += Extent.X;
+				Extent.X *= -1.0f;
+			}
+			if (Extent.Y < 0)
 			{
-				for (int32 yStart = 0; yStart < SqrtMaxInstances; yStart++)
+				Origin.Y += Extent.Y;
+				Extent.Y *= -1.0f;
+			}
+			int32 MaxNum = SqrtMaxInstances * SqrtMaxInstances;
+			InstanceTransforms.Reserve(MaxNum);
+			FVector DivExtent(Extent * Div);
+			for (int32 InstanceIndex = 0; InstanceIndex < MaxNum; InstanceIndex++)
+			{
+				float HaltonX = Halton<2>(InstanceIndex + HaltonBaseIndex);
+				float HaltonY = Halton<3>(InstanceIndex + HaltonBaseIndex);
+				FVector Location(Origin.X + HaltonX * Extent.X, Origin.Y + HaltonY * Extent.Y, 0.0f);
+				FVector LocationWithHeight;
+				float Weight = GetLayerWeightAtLocationLocal(Location, &LocationWithHeight);
+				bool bKeep = Weight > 0.0f && Weight >= RandomStream.GetFraction();
+				if (bKeep)
 				{
-					FVector Location(Origin.X + float(xStart) * Div * Extent.X, Origin.Y + float(yStart) * Div * Extent.Y, 0.0f);
-					Location += FVector(RandomStream.GetFraction() * 2.0f - 1.0f, RandomStream.GetFraction() * 2.0f - 1.0f, 0.0f) * MaxJitter;
+					float Rot = RandomRotation ? RandomStream.GetFraction() * 360.0f : 0.0f;
+					FMatrix OutXForm;
+					if (AlignToSurface)
+					{
+						FVector LocationWithHeightDX;
+						FVector LocationDX(Location);
+						LocationDX.X = FMath::Clamp<float>(LocationDX.X + (HaltonX < 0.5f ? DivExtent.X : -DivExtent.X), Origin.X, Origin.X + Extent.X);
+						GetLayerWeightAtLocationLocal(LocationDX, &LocationWithHeightDX, false);
 
-					FInstanceLocal& Instance = Instances[InstanceIndex];
-					float Weight = GetLayerWeightAtLocationLocal(Location, &Instance.Pos);
-					Instance.bKeep = Weight > 0.0f && Weight >= RandomStream.GetFraction();
-					if (Instance.bKeep)
+						FVector LocationWithHeightDY;
+						FVector LocationDY(Location);
+						LocationDY.Y = FMath::Clamp<float>(LocationDX.Y + (HaltonY < 0.5f ? DivExtent.Y : -DivExtent.Y), Origin.Y, Origin.Y + Extent.Y);
+						GetLayerWeightAtLocationLocal(LocationDY, &LocationWithHeightDY, false);
+
+						if (LocationWithHeight != LocationWithHeightDX && LocationWithHeight != LocationWithHeightDY)
+						{
+							FVector NewZ = ((LocationWithHeight - LocationWithHeightDX) ^ (LocationWithHeight - LocationWithHeightDY)).GetSafeNormal();
+							NewZ *= FMath::Sign(NewZ.Z);
+
+							const FVector NewX = (FVector(0, -1, 0) ^ NewZ).GetSafeNormal();
+							const FVector NewY = NewZ ^ NewX;
+
+							FMatrix Align = FMatrix(NewX, NewY, NewZ, FVector::ZeroVector);
+							OutXForm = FRotationMatrix(FRotator(0.0f, Rot, 0.0f)) * Align * FTranslationMatrix(LocationWithHeight) * XForm;
+						}
+						else
+						{
+							OutXForm = FRotationMatrix(FRotator(0.0f, Rot, 0.0f)) * FTranslationMatrix(LocationWithHeight) * XForm;
+						}
+					}
+					else
 					{
-						NumKept++;
+						OutXForm = FRotationMatrix(FRotator(0.0f, Rot, 0.0f)) * FTranslationMatrix(LocationWithHeight) * XForm;
 					}
-					InstanceIndex++;
+					InstanceTransforms.Add(OutXForm);
+				}
+			}
+			if (InstanceTransforms.Num())
+			{
+				TotalInstances += InstanceTransforms.Num();
+				InstanceBuffer.AllocateInstances(InstanceTransforms.Num());
+				for (int32 InstanceIndex = 0; InstanceIndex < InstanceTransforms.Num(); InstanceIndex++)
+				{
+					const FMatrix& OutXForm = InstanceTransforms[InstanceIndex];
+					FInstanceStream* RenderData = InstanceBuffer.GetInstanceWriteAddress(InstanceIndex);
+					RenderData->SetInstance(OutXForm, RandomStream.GetFraction());
 				}
 			}
 		}
-		RasterTime += FPlatformTime::Seconds();
-
-		if (NumKept)
+		else
 		{
-			TArray<FMatrix> InstanceTransforms;
-			InstanceTransforms.AddUninitialized(NumKept);
-			TotalInstances += NumKept;
-
+			int32 NumKept = 0;
+			float MaxJitter1D = FMath::Clamp<float>(PlacementJitter, 0.0f, .99f) * Div * .5f;
+			FVector MaxJitter(MaxJitter1D, MaxJitter1D, 0.0f);
+			MaxJitter *= Extent;
+			Origin += Extent * (Div * 0.5f);
+			struct FInstanceLocal
+			{
+				FVector Pos;
+				bool bKeep;
+			};
+			TArray<FInstanceLocal> Instances;
+			Instances.AddUninitialized(SqrtMaxInstances * SqrtMaxInstances);
 			{
-				InstanceTime -= FPlatformTime::Seconds();
-				InstanceBuffer.AllocateInstances(NumKept);
 				int32 InstanceIndex = 0;
-				int32 OutInstanceIndex = 0;
 				for (int32 xStart = 0; xStart < SqrtMaxInstances; xStart++)
 				{
 					for (int32 yStart = 0; yStart < SqrtMaxInstances; yStart++)
 					{
-						const FInstanceLocal& Instance = Instances[InstanceIndex];
+						FVector Location(Origin.X + float(xStart) * Div * Extent.X, Origin.Y + float(yStart) * Div * Extent.Y, 0.0f);
+						Location += FVector(RandomStream.GetFraction() * 2.0f - 1.0f, RandomStream.GetFraction() * 2.0f - 1.0f, 0.0f) * MaxJitter;
+
+						FInstanceLocal& Instance = Instances[InstanceIndex];
+						float Weight = GetLayerWeightAtLocationLocal(Location, &Instance.Pos);
+						Instance.bKeep = Weight > 0.0f && Weight >= RandomStream.GetFraction();
 						if (Instance.bKeep)
 						{
-							float Rot = RandomRotation ? RandomStream.GetFraction() * 360.0f : 0.0f;
-							FMatrix OutXForm;
-							if (AlignToSurface)
+							NumKept++;
+						}
+						InstanceIndex++;
+					}
+				}
+			}
+			if (NumKept)
+			{
+				InstanceTransforms.AddUninitialized(NumKept);
+				TotalInstances += NumKept;
+				{
+					InstanceBuffer.AllocateInstances(NumKept);
+					int32 InstanceIndex = 0;
+					int32 OutInstanceIndex = 0;
+					for (int32 xStart = 0; xStart < SqrtMaxInstances; xStart++)
+					{
+						for (int32 yStart = 0; yStart < SqrtMaxInstances; yStart++)
+						{
+							const FInstanceLocal& Instance = Instances[InstanceIndex];
+							if (Instance.bKeep)
 							{
-								FVector PosX1 = xStart ? Instances[InstanceIndex - SqrtMaxInstances].Pos : Instance.Pos;
-								FVector PosX2 = (xStart + 1 < SqrtMaxInstances) ? Instances[InstanceIndex + SqrtMaxInstances].Pos : Instance.Pos;
-								FVector PosY1 = yStart ? Instances[InstanceIndex - 1].Pos : Instance.Pos;
-								FVector PosY2 = (yStart + 1 < SqrtMaxInstances) ? Instances[InstanceIndex + 1].Pos : Instance.Pos;
-
-								if (PosX1 != PosX2 && PosY1 != PosY2)
+								float Rot = RandomRotation ? RandomStream.GetFraction() * 360.0f : 0.0f;
+								FMatrix OutXForm;
+								if (AlignToSurface)
 								{
-									FVector NewZ = ((PosX1 - PosX2) ^ (PosY1 - PosY2)).GetSafeNormal();
-									NewZ *= FMath::Sign(NewZ.Z);
+									FVector PosX1 = xStart ? Instances[InstanceIndex - SqrtMaxInstances].Pos : Instance.Pos;
+									FVector PosX2 = (xStart + 1 < SqrtMaxInstances) ? Instances[InstanceIndex + SqrtMaxInstances].Pos : Instance.Pos;
+									FVector PosY1 = yStart ? Instances[InstanceIndex - 1].Pos : Instance.Pos;
+									FVector PosY2 = (yStart + 1 < SqrtMaxInstances) ? Instances[InstanceIndex + 1].Pos : Instance.Pos;
 
-									const FVector NewX = (FVector(0, -1, 0) ^ NewZ).GetSafeNormal();
-									const FVector NewY = NewZ ^ NewX;
+									if (PosX1 != PosX2 && PosY1 != PosY2)
+									{
+										FVector NewZ = ((PosX1 - PosX2) ^ (PosY1 - PosY2)).GetSafeNormal();
+										NewZ *= FMath::Sign(NewZ.Z);
+
+										const FVector NewX = (FVector(0, -1, 0) ^ NewZ).GetSafeNormal();
+										const FVector NewY = NewZ ^ NewX;
 
-									FMatrix Align = FMatrix(NewX, NewY, NewZ, FVector::ZeroVector);
-									OutXForm = FRotationMatrix(FRotator(0.0f, Rot, 0.0f)) * Align * FTranslationMatrix(Instance.Pos) * XForm;
+										FMatrix Align = FMatrix(NewX, NewY, NewZ, FVector::ZeroVector);
+										OutXForm = FRotationMatrix(FRotator(0.0f, Rot, 0.0f)) * Align * FTranslationMatrix(Instance.Pos) * XForm;
+									}
+									else
+									{
+										OutXForm = FRotationMatrix(FRotator(0.0f, Rot, 0.0f)) * FTranslationMatrix(Instance.Pos) * XForm;
+									}
 								}
 								else
 								{
 									OutXForm = FRotationMatrix(FRotator(0.0f, Rot, 0.0f)) * FTranslationMatrix(Instance.Pos) * XForm;
 								}
+								InstanceTransforms[OutInstanceIndex] = OutXForm;
+								FInstanceStream* RenderData = InstanceBuffer.GetInstanceWriteAddress(OutInstanceIndex++);
+								RenderData->SetInstance(OutXForm, RandomStream.GetFraction());
 							}
-							else
-							{
-								OutXForm = FRotationMatrix(FRotator(0.0f, Rot, 0.0f)) * FTranslationMatrix(Instance.Pos) * XForm;
-							}
-							InstanceTransforms[OutInstanceIndex] = OutXForm;
-							FInstanceStream* RenderData = InstanceBuffer.GetInstanceWriteAddress(OutInstanceIndex++);
-							RenderData->SetInstance(OutXForm, RandomStream.GetFraction());
+							InstanceIndex++;
 						}
-						InstanceIndex++;
 					}
 				}
-				InstanceTime += FPlatformTime::Seconds();
 			}
+		}
+
+		int32 NumInstances = InstanceTransforms.Num();
+		if (NumInstances)
+		{
+			TArray<int32> SortedInstances;
+			TArray<int32> InstanceReorderTable;
+			UHierarchicalInstancedStaticMeshComponent::BuildTreeAnyThread(InstanceTransforms, MeshBox, ClusterTree, SortedInstances, InstanceReorderTable, DesiredInstancesPerLeaf);
+
+			// in-place sort the instances
+			FInstanceStream SwapBuffer;
+			int32 FirstUnfixedIndex = 0;
+			for (int32 FirstUnfixedIndex = 0; FirstUnfixedIndex < NumInstances; FirstUnfixedIndex++)
 			{
-				BuildTime -= FPlatformTime::Seconds();
-				TArray<int32> SortedInstances;
-				TArray<int32> InstanceReorderTable;
-				UHierarchicalInstancedStaticMeshComponent::BuildTreeAnyThread(InstanceTransforms, MeshBox, ClusterTree, SortedInstances, InstanceReorderTable, FMath::Clamp<int32>(CVarMinInstancesPerLeaf.GetValueOnAnyThread(), 16, 256 * 128));
-
-				// in-place sort the instances
-				FInstanceStream SwapBuffer;
-				int32 FirstUnfixedIndex = 0;
-				for (int32 FirstUnfixedIndex = 0; FirstUnfixedIndex < NumKept; FirstUnfixedIndex++)
+				int32 LoadFrom = SortedInstances[FirstUnfixedIndex];
+				if (LoadFrom != FirstUnfixedIndex)
 				{
-					int32 LoadFrom = SortedInstances[FirstUnfixedIndex];
-					if (LoadFrom != FirstUnfixedIndex)
-					{
-						check(LoadFrom > FirstUnfixedIndex);
-						FMemory::Memcpy(&SwapBuffer, InstanceBuffer.GetInstanceWriteAddress(FirstUnfixedIndex), sizeof(FInstanceStream));
-						FMemory::Memcpy(InstanceBuffer.GetInstanceWriteAddress(FirstUnfixedIndex), InstanceBuffer.GetInstanceWriteAddress(LoadFrom), sizeof(FInstanceStream));
-						FMemory::Memcpy(InstanceBuffer.GetInstanceWriteAddress(LoadFrom), &SwapBuffer, sizeof(FInstanceStream));
-
-						int32 SwapGoesTo = InstanceReorderTable[FirstUnfixedIndex];
-						check(SwapGoesTo > FirstUnfixedIndex);
-						check(SortedInstances[SwapGoesTo] == FirstUnfixedIndex);
-						SortedInstances[SwapGoesTo] = LoadFrom;
-						InstanceReorderTable[LoadFrom] = SwapGoesTo;
-
-						InstanceReorderTable[FirstUnfixedIndex] = FirstUnfixedIndex;
-						SortedInstances[FirstUnfixedIndex] = FirstUnfixedIndex;
-					}
+					check(LoadFrom > FirstUnfixedIndex);
+					FMemory::Memcpy(&SwapBuffer, InstanceBuffer.GetInstanceWriteAddress(FirstUnfixedIndex), sizeof(FInstanceStream));
+					FMemory::Memcpy(InstanceBuffer.GetInstanceWriteAddress(FirstUnfixedIndex), InstanceBuffer.GetInstanceWriteAddress(LoadFrom), sizeof(FInstanceStream));
+					FMemory::Memcpy(InstanceBuffer.GetInstanceWriteAddress(LoadFrom), &SwapBuffer, sizeof(FInstanceStream));
+
+					int32 SwapGoesTo = InstanceReorderTable[FirstUnfixedIndex];
+					check(SwapGoesTo > FirstUnfixedIndex);
+					check(SortedInstances[SwapGoesTo] == FirstUnfixedIndex);
+					SortedInstances[SwapGoesTo] = LoadFrom;
+					InstanceReorderTable[LoadFrom] = SwapGoesTo;
+
+					InstanceReorderTable[FirstUnfixedIndex] = FirstUnfixedIndex;
+					SortedInstances[FirstUnfixedIndex] = FirstUnfixedIndex;
 				}
-
-				BuildTime += FPlatformTime::Seconds();
 			}
 		}
 	}
-	FORCEINLINE_DEBUGGABLE float GetLayerWeightAtLocationLocal(const FVector& InLocation, FVector* OutLocation)
+	FORCEINLINE_DEBUGGABLE float GetLayerWeightAtLocationLocal(const FVector& InLocation, FVector* OutLocation, bool bWeight = true)
 	{
 		// Find location
 		float TestX = InLocation.X / DrawScale.X - (float)SectionBase.X;
@@ -1119,7 +1214,8 @@ struct FAsyncGrassBuilder : public FGrassBuilderBase
 		float LerpX = FMath::Fractional(TestX);
 		float LerpY = FMath::Fractional(TestY);
 
-		float Result;
+		float Result = 0.0f;
+		if (bWeight)
 		{
 			// sample
 			float Sample11 = GrassData.GetWeight(IdxX1, IdxY1);
@@ -1168,6 +1264,7 @@ void ALandscapeProxy::FlushGrassComponents(const TSet<ULandscapeComponent*>* Onl
 				{
 					SCOPE_CYCLE_COUNTER(STAT_FoliageGrassDestoryComp);
 					Used->ClearInstances();
+					Used->DetachFromParent(false, false);
 					Used->DestroyComponent();
 				}
 				Iter.RemoveCurrent();
@@ -1197,6 +1294,7 @@ void ALandscapeProxy::FlushGrassComponents(const TSet<ULandscapeComponent*>* Onl
 		{
 			SCOPE_CYCLE_COUNTER(STAT_FoliageGrassDestoryComp);
 			Component->ClearInstances();
+			Component->DetachFromParent(false, false);
 			Component->DestroyComponent();
 		}
 
@@ -1210,6 +1308,8 @@ void ALandscapeProxy::FlushGrassComponents(const TSet<ULandscapeComponent*>* Onl
 		for (USceneComponent* Component : AttachedFoliageComponents)
 		{
 			SCOPE_CYCLE_COUNTER(STAT_FoliageGrassDestoryComp);
+			CastChecked<UHierarchicalInstancedStaticMeshComponent>(Component)->ClearInstances();
+			Component->DetachFromParent(false, false);
 			Component->DestroyComponent();
 		}
 
@@ -1345,6 +1445,7 @@ void ALandscapeProxy::UpdateGrass(const TArray<FVector>& Cameras, bool bForceSyn
 					if (GrassType)
 					{
 						int32 GrassVarietyIndex = -1;
+						uint32 HaltonBaseIndex = 1;
 						for (auto& GrassVariety : GrassType->GrassVarieties)
 						{
 							GrassVarietyIndex++;
@@ -1355,7 +1456,9 @@ void ALandscapeProxy::UpdateGrass(const TArray<FVector>& Cameras, bool bForceSyn
 								float MustHaveDistance = GuardBand * (float)GrassVariety.EndCullDistance;
 								float DiscardDistance = DiscardGuardBand * (float)GrassVariety.EndCullDistance;
 
-								if (MinDistanceToComp > DiscardDistance)
+								bool bUseHalton = !GrassVariety.bUseGrid;
+
+								if (!bUseHalton && MinDistanceToComp > DiscardDistance)
 								{
 									continue;
 								}
@@ -1368,6 +1471,13 @@ void ALandscapeProxy::UpdateGrass(const TArray<FVector>& Cameras, bool bForceSyn
 								{
 									SqrtSubsections = FMath::Clamp<int32>(FMath::CeilToInt(float(ForSubsectionMath.SqrtMaxInstances) / FMath::Sqrt((float)MaxInstancesPerComponent)), 1, 16);
 								}
+								int32 MaxInstancesSub = FMath::Square(ForSubsectionMath.SqrtMaxInstances / SqrtSubsections);
+
+								if (bUseHalton && MinDistanceToComp > DiscardDistance)
+								{
+									HaltonBaseIndex += MaxInstancesSub * SqrtSubsections * SqrtSubsections;
+									continue;
+								}
 
 								FBox LocalBox = Component->CachedLocalBox;
 								FVector LocalExtentDiv = (LocalBox.Max - LocalBox.Min) * FVector(1.0f / float(SqrtSubsections), 1.0f / float(SqrtSubsections), 1.0f);
@@ -1400,6 +1510,11 @@ void ALandscapeProxy::UpdateGrass(const TArray<FVector>& Cameras, bool bForceSyn
 											MinDistanceToSubComp = FMath::Sqrt(MinDistanceToSubComp);
 										}
 
+										if (bUseHalton)
+										{
+											HaltonBaseIndex += MaxInstancesSub;  // we are going to pre-increment this for all of the continues...however we need to subtract later if we actually do this sub
+										}
+
 										if (MinDistanceToSubComp > DiscardDistance)
 										{
 											continue;
@@ -1472,7 +1587,7 @@ void ALandscapeProxy::UpdateGrass(const TArray<FVector>& Cameras, bool bForceSyn
 										UHierarchicalInstancedStaticMeshComponent* HierarchicalInstancedStaticMeshComponent;
 										{
 											QUICK_SCOPE_CYCLE_COUNTER(STAT_GrassCreateComp);
-											HierarchicalInstancedStaticMeshComponent = NewObject<UHierarchicalInstancedStaticMeshComponent>(this); //, NAME_None, RF_Transactional);
+											HierarchicalInstancedStaticMeshComponent = NewObject<UHierarchicalInstancedStaticMeshComponent>(this);
 										}
 										NewComp.Foliage = HierarchicalInstancedStaticMeshComponent;
 										FoliageCache.CachedGrassComps.Add(NewComp);
@@ -1480,6 +1595,7 @@ void ALandscapeProxy::UpdateGrass(const TArray<FVector>& Cameras, bool bForceSyn
 										HierarchicalInstancedStaticMeshComponent->Mobility = EComponentMobility::Static;
 
 										HierarchicalInstancedStaticMeshComponent->StaticMesh = GrassVariety.GrassMesh;
+										HierarchicalInstancedStaticMeshComponent->MinLOD = GrassVariety.MinLOD;
 										HierarchicalInstancedStaticMeshComponent->bSelectable = false;
 										HierarchicalInstancedStaticMeshComponent->bHasPerInstanceHitProxies = true;
 										static FName NoCollision(TEXT("NoCollision"));
@@ -1518,7 +1634,14 @@ void ALandscapeProxy::UpdateGrass(const TArray<FVector>& Cameras, bool bForceSyn
 
 										{
 											QUICK_SCOPE_CYCLE_COUNTER(STAT_GrassCreateBuilder);
-											Builder = new FAsyncGrassBuilder(this, Component, GrassType, GrassVariety, HierarchicalInstancedStaticMeshComponent, SqrtSubsections, SubX, SubY);
+
+											uint32 HaltonIndexForSub = 0;
+											if (bUseHalton)
+											{
+												check(HaltonBaseIndex > (uint32)MaxInstancesSub);
+												HaltonIndexForSub = HaltonBaseIndex - (uint32)MaxInstancesSub;
+											}
+											Builder = new FAsyncGrassBuilder(this, Component, GrassType, GrassVariety, HierarchicalInstancedStaticMeshComponent, SqrtSubsections, SubX, SubY, HaltonIndexForSub);
 										}
 
 										if (Builder->bHaveValidData)
@@ -1652,6 +1775,7 @@ void ALandscapeProxy::UpdateGrass(const TArray<FVector>& Cameras, bool bForceSyn
 				{
 					SCOPE_CYCLE_COUNTER(STAT_FoliageGrassDestoryComp);
 					HComponent->ClearInstances();
+					HComponent->DetachFromParent(false, false);
 					HComponent->DestroyComponent();
 					FoliageComponents.Remove(HComponent);
 				}
@@ -1681,10 +1805,12 @@ void ALandscapeProxy::UpdateGrass(const TArray<FVector>& Cameras, bool bForceSyn
 				{
 					if (Inner.Builder->InstanceBuffer.Num())
 					{
-						Exchange(HierarchicalInstancedStaticMeshComponent->WriteOncePrebuiltInstanceBuffer, Inner.Builder->InstanceBuffer);
+						QUICK_SCOPE_CYCLE_COUNTER(STAT_FoliageGrassEndComp_AcceptPrebuiltTree);
+						FMemory::Memswap(&HierarchicalInstancedStaticMeshComponent->WriteOncePrebuiltInstanceBuffer, &Inner.Builder->InstanceBuffer, sizeof(FStaticMeshInstanceData));
 						HierarchicalInstancedStaticMeshComponent->AcceptPrebuiltTree(Inner.Builder->ClusterTree);
 						if (bForceSync && GetWorld())
 						{
+							QUICK_SCOPE_CYCLE_COUNTER(STAT_FoliageGrassEndComp_SyncUpdate);
 							HierarchicalInstancedStaticMeshComponent->RecreateRenderState_Concurrent();
 						}
 					}
@@ -1727,11 +1853,30 @@ static void FlushGrass(const TArray<FString>& Args)
 	}
 }
 
+static void FlushGrassPIE(const TArray<FString>& Args)
+{
+	for (TObjectIterator<ALandscapeProxy> It; It; ++It)
+	{
+		ALandscapeProxy* Landscape = *It;
+		if (Landscape && !Landscape->IsTemplate() && !Landscape->IsPendingKill())
+		{
+			Landscape->FlushGrassComponents(nullptr, false);
+		}
+	}
+}
+
 static FAutoConsoleCommand FlushGrassCmd(
 	TEXT("grass.FlushCache"),
 	TEXT("Flush the grass cache, debugging."),
 	FConsoleCommandWithArgsDelegate::CreateStatic(&FlushGrass)
 	);
 
+static FAutoConsoleCommand FlushGrassCmdPIE(
+	TEXT("grass.FlushCachePIE"),
+	TEXT("Flush the grass cache, debugging."),
+	FConsoleCommandWithArgsDelegate::CreateStatic(&FlushGrassPIE)
+	);
+
+
 
 #undef LOCTEXT_NAMESPACE
\ No newline at end of file
diff --git a/Engine/Source/Runtime/Launch/Private/LaunchEngineLoop.cpp b/Engine/Source/Runtime/Launch/Private/LaunchEngineLoop.cpp
--- a/Engine/Source/Runtime/Launch/Private/LaunchEngineLoop.cpp
+++ b/Engine/Source/Runtime/Launch/Private/LaunchEngineLoop.cpp
@@ -2183,7 +2183,12 @@ void FEngineLoop::Tick()
 	ensure(GetMoviePlayer()->IsLoadingFinished() && !GetMoviePlayer()->IsMovieCurrentlyPlaying());
 
 	// early in the Tick() to get the callbacks for cvar changes called
-	IConsoleManager::Get().CallAllConsoleVariableSinks();
+	{
+#if WITH_ENGINE
+		QUICK_SCOPE_CYCLE_COUNTER(STAT_FEngineLoop_Tick_CallAllConsoleVariableSinks);
+#endif
+		IConsoleManager::Get().CallAllConsoleVariableSinks();
+	}
 
 	{ 
 		SCOPE_CYCLE_COUNTER( STAT_FrameTime );	
@@ -2197,8 +2202,13 @@ void FEngineLoop::Tick()
 			RHICmdList.BeginFrame();
 		});
 
-		// Flush debug output which has been buffered by other threads.
-		GLog->FlushThreadedLogs();
+		{
+#if WITH_ENGINE
+			QUICK_SCOPE_CYCLE_COUNTER(STAT_FEngineLoop_FlushThreadedLogs);
+#endif
+			// Flush debug output which has been buffered by other threads.
+			GLog->FlushThreadedLogs();
+		}
 
 		// Exit if frame limit is reached in benchmark mode.
 		if( (FApp::IsBenchmarking() && MaxFrameCounter && (GFrameCounter > MaxFrameCounter))
@@ -2208,11 +2218,24 @@ void FEngineLoop::Tick()
 			FPlatformMisc::RequestExit(0);
 		}
 
-		// Set FApp::CurrentTime, FApp::DeltaTime and potentially wait to enforce max tick rate.
-		GEngine->UpdateTimeAndHandleMaxTickRate();
+		{
+#if WITH_ENGINE
+			QUICK_SCOPE_CYCLE_COUNTER(STAT_FEngineLoop_UpdateTimeAndHandleMaxTickRate);
+#endif
+			// Set FApp::CurrentTime, FApp::DeltaTime and potentially wait to enforce max tick rate.
+			GEngine->UpdateTimeAndHandleMaxTickRate();
+		}
 
-		GEngine->TickFPSChart( FApp::GetDeltaTime() );
+		{
+#if WITH_ENGINE
+			QUICK_SCOPE_CYCLE_COUNTER(STAT_FEngineLoop_TickFPSChart);
+#endif
+			GEngine->TickFPSChart( FApp::GetDeltaTime() );
+		}
 
+#if WITH_ENGINE
+		QUICK_SCOPE_CYCLE_COUNTER(STAT_FEngineLoop_Malloc_UpdateStats);
+#endif
 		// Update platform memory and memory allocator stats.
 		FPlatformMemory::UpdateStats();
 		GMalloc->UpdateStats();
@@ -2245,16 +2268,26 @@ void FEngineLoop::Tick()
 			FPlatformMisc::PumpMessages(true);
 		}
 
-		// Idle mode prevents ticking and rendering completely
-		const bool bIdleMode = ShouldUseIdleMode();
-		if (bIdleMode)
+		bool bIdleMode;
 		{
-			// Yield CPU time
-			FPlatformProcess::Sleep(.1f);
+
+#if WITH_ENGINE
+			QUICK_SCOPE_CYCLE_COUNTER(STAT_FEngineLoop_Idle);
+#endif
+			// Idle mode prevents ticking and rendering completely
+			bIdleMode = ShouldUseIdleMode();
+			if (bIdleMode)
+			{
+				// Yield CPU time
+				FPlatformProcess::Sleep(.1f);
+			}
 		}
 
 		if (FSlateApplication::IsInitialized() && !bIdleMode)
 		{
+#if WITH_ENGINE
+			QUICK_SCOPE_CYCLE_COUNTER(STAT_FEngineLoop_Tick_SlateInput);
+#endif
 			FSlateApplication& SlateApp = FSlateApplication::Get();
 			SlateApp.PollGameDeviceState();
 			// Gives widgets a chance to process any accumulated input
@@ -2266,21 +2299,35 @@ void FEngineLoop::Tick()
 		// If a movie that is blocking the game thread has been playing,
 		// wait for it to finish before we continue to tick or tick again
 		// We do this right after GEngine->Tick() because that is where user code would initiate a load / movie.
-		GetMoviePlayer()->WaitForMovieToFinish();
+		{
+#if WITH_ENGINE
+			QUICK_SCOPE_CYCLE_COUNTER(STAT_FEngineLoop_WaitForMovieToFinish);
+#endif
+			GetMoviePlayer()->WaitForMovieToFinish();
+		}
 		
 		if (GShaderCompilingManager)
 		{
 			// Process any asynchronous shader compile results that are ready, limit execution time
+#if WITH_ENGINE
+			QUICK_SCOPE_CYCLE_COUNTER(STAT_FEngineLoop_Tick_GShaderCompilingManager);
+#endif
 			GShaderCompilingManager->ProcessAsyncResults(true, false);
 		}
 
 		if (GDistanceFieldAsyncQueue)
 		{
+#if WITH_ENGINE
+			QUICK_SCOPE_CYCLE_COUNTER(STAT_FEngineLoop_Tick_GDistanceFieldAsyncQueue);
+#endif
 			GDistanceFieldAsyncQueue->ProcessAsyncTasks();
 		}
 
 		if (FSlateApplication::IsInitialized() && !bIdleMode)
 		{
+#if WITH_ENGINE
+			QUICK_SCOPE_CYCLE_COUNTER(STAT_FEngineLoop_ProcessPlayerControllersSlateOperations);
+#endif
 			check(!IsRunningDedicatedServer());
 			
 			// Process slate operations accumulated in the world ticks.
@@ -2299,6 +2346,9 @@ void FEngineLoop::Tick()
 #if WITH_EDITOR
 		if( GIsEditor )
 		{
+#if WITH_ENGINE
+			QUICK_SCOPE_CYCLE_COUNTER(STAT_FEngineLoop_Tick_AutomationController);
+#endif
 			static FName AutomationController("AutomationController");
 			//Check if module loaded to support the change to allow this to be hot compilable.
 			if (FModuleManager::Get().IsModuleLoaded(AutomationController))
@@ -2310,11 +2360,16 @@ void FEngineLoop::Tick()
 
 #if WITH_ENGINE
 #if WITH_AUTOMATION_WORKER
-		//Check if module loaded to support the change to allow this to be hot compilable.
-		static const FName AutomationWorkerModuleName = TEXT("AutomationWorker");
-		if (FModuleManager::Get().IsModuleLoaded(AutomationWorkerModuleName))
 		{
-			FModuleManager::GetModuleChecked<IAutomationWorkerModule>(AutomationWorkerModuleName).Tick();
+#if WITH_ENGINE
+			QUICK_SCOPE_CYCLE_COUNTER(STAT_FEngineLoop_Tick_AutomationWorker);
+#endif
+			//Check if module loaded to support the change to allow this to be hot compilable.
+			static const FName AutomationWorkerModuleName = TEXT("AutomationWorker");
+			if (FModuleManager::Get().IsModuleLoaded(AutomationWorkerModuleName))
+			{
+				FModuleManager::GetModuleChecked<IAutomationWorkerModule>(AutomationWorkerModuleName).Tick();
+			}
 		}
 #endif
 #endif //WITH_ENGINE
diff --git a/Engine/Source/Runtime/Online/HTTP/Private/Curl/CurlHttpManager.cpp b/Engine/Source/Runtime/Online/HTTP/Private/Curl/CurlHttpManager.cpp
--- a/Engine/Source/Runtime/Online/HTTP/Private/Curl/CurlHttpManager.cpp
+++ b/Engine/Source/Runtime/Online/HTTP/Private/Curl/CurlHttpManager.cpp
@@ -209,6 +209,7 @@ void FCurlHttpManager::RemoveRequest(const TSharedRef<class IHttpRequest>& Reque
 
 bool FCurlHttpManager::Tick(float DeltaSeconds)
 {
+	QUICK_SCOPE_CYCLE_COUNTER(STAT_FCurlHttpManager_Tick);
 	check(MultiHandle);
 	if (Requests.Num() > 0)
 	{
diff --git a/Engine/Source/Runtime/RenderCore/Private/RenderingThread.cpp b/Engine/Source/Runtime/RenderCore/Private/RenderingThread.cpp
--- a/Engine/Source/Runtime/RenderCore/Private/RenderingThread.cpp
+++ b/Engine/Source/Runtime/RenderCore/Private/RenderingThread.cpp
@@ -108,6 +108,7 @@ FSuspendRenderingThread::FSuspendRenderingThread( bool bInRecreateThread )
 				}
 				else
 				{
+					QUICK_SCOPE_CYCLE_COUNTER(STAT_FSuspendRenderingThread);
 					FTaskGraphInterface::Get().WaitUntilTaskCompletes(CompleteHandle, ENamedThreads::GameThread);
 				}
 				check(GIsRenderingThreadSuspended);
@@ -523,6 +524,7 @@ void StartRenderingThread()
 		DECLARE_CYCLE_STAT(TEXT("Wait For RHIThread"), STAT_WaitForRHIThread, STATGROUP_TaskGraphTasks);
 
 		FGraphEventRef CompletionEvent = TGraphTask<FNullGraphTask>::CreateTask(NULL, ENamedThreads::GameThread).ConstructAndDispatchWhenReady(GET_STATID(STAT_WaitForRHIThread), ENamedThreads::RHIThread);
+		QUICK_SCOPE_CYCLE_COUNTER(STAT_StartRenderingThread);
 		FTaskGraphInterface::Get().WaitUntilTaskCompletes(CompletionEvent, ENamedThreads::GameThread_Local);
 		GRHIThread = FRHIThread::Get().Thread;
 		check(GRHIThread);
@@ -596,6 +598,7 @@ void StopRenderingThread()
 			{
 				DECLARE_CYCLE_STAT(TEXT("Wait For RHIThread Finish"), STAT_WaitForRHIThreadFinish, STATGROUP_TaskGraphTasks);
 				FGraphEventRef FlushTask = TGraphTask<FNullGraphTask>::CreateTask(NULL, ENamedThreads::GameThread).ConstructAndDispatchWhenReady(GET_STATID(STAT_WaitForRHIThreadFinish), ENamedThreads::RHIThread);
+				QUICK_SCOPE_CYCLE_COUNTER(STAT_StopRenderingThread_RHIThread);
 				FTaskGraphInterface::Get().WaitUntilTaskCompletes(FlushTask, ENamedThreads::GameThread_Local);
 				GRHIThread = nullptr;
 			}
@@ -622,6 +625,7 @@ void StopRenderingThread()
 				}
 				else
 				{
+					QUICK_SCOPE_CYCLE_COUNTER(STAT_StopRenderingThread);
 					FTaskGraphInterface::Get().WaitUntilTaskCompletes(QuitTask, ENamedThreads::GameThread_Local);
 				}
 			}
@@ -796,6 +800,7 @@ void FRenderCommandFence::Wait(bool bProcessGameThreadTasks) const
 		// windows needs to pump messages
 		if (bProcessGameThreadTasks)
 		{
+			QUICK_SCOPE_CYCLE_COUNTER(STAT_FRenderCommandFence_Wait);
 			FTaskGraphInterface::Get().WaitUntilTaskCompletes(CompletionEvent, ENamedThreads::GameThread);
 		}
 #endif
@@ -896,6 +901,8 @@ FPendingCleanupObjects::FPendingCleanupObjects()
 
 FPendingCleanupObjects::~FPendingCleanupObjects()
 {
+	QUICK_SCOPE_CYCLE_COUNTER(STAT_FPendingCleanupObjects_Destruct);
+
 	for(int32 ObjectIndex = 0;ObjectIndex < CleanupArray.Num();ObjectIndex++)
 	{
 		CleanupArray[ObjectIndex]->FinishCleanup();
diff --git a/Engine/Source/Runtime/Renderer/Private/BasePassRendering.cpp b/Engine/Source/Runtime/Renderer/Private/BasePassRendering.cpp
--- a/Engine/Source/Runtime/Renderer/Private/BasePassRendering.cpp
+++ b/Engine/Source/Runtime/Renderer/Private/BasePassRendering.cpp
@@ -118,7 +118,7 @@ void FTranslucentLightingParameters::Set(FRHICommandList& RHICmdList, FShader* S
 			TStaticSamplerState<SF_Point, AM_Clamp, AM_Clamp, AM_Clamp>::GetRHI(), 
 			View->HZB->GetRenderTargetItem().ShaderResourceTexture );
 
-		TRefCountPtr<IPooledRenderTarget>* PrevSceneColorRT = &GSceneRenderTargets.GetSceneColor();
+		TRefCountPtr<IPooledRenderTarget>* PrevSceneColorRT = &GSystemTextures.BlackDummy;
 
 		FSceneViewState* ViewState = (FSceneViewState*)View->State;
 		if( ViewState && ViewState->TemporalAAHistoryRT && !View->bCameraCut )
diff --git a/Engine/Source/Runtime/Renderer/Private/CompositionLighting/PostProcessAmbientOcclusion.h b/Engine/Source/Runtime/Renderer/Private/CompositionLighting/PostProcessAmbientOcclusion.h
--- a/Engine/Source/Runtime/Renderer/Private/CompositionLighting/PostProcessAmbientOcclusion.h
+++ b/Engine/Source/Runtime/Renderer/Private/CompositionLighting/PostProcessAmbientOcclusion.h
@@ -180,8 +180,9 @@ void FCameraMotionParameters::Set(FRHICommandList& RHICmdList, const FSceneView&
 	PrevProj.M[2][0] = 0.0f;
 	PrevProj.M[2][1] = 0.0f;
 
-	FMatrix ViewProj = ( View.ViewMatrices.ViewMatrix * Proj ).GetTransposed();
-	FMatrix PrevViewProj = ( ViewState->PrevViewMatrices.ViewMatrix * PrevProj ).GetTransposed();
+	FVector DeltaTranslation = ViewState->PrevViewMatrices.PreViewTranslation - View.ViewMatrices.PreViewTranslation;
+	FMatrix ViewProj = ( View.ViewMatrices.TranslatedViewMatrix * Proj ).GetTransposed();
+	FMatrix PrevViewProj = ( FTranslationMatrix(DeltaTranslation) * ViewState->PrevViewMatrices.TranslatedViewMatrix * PrevProj ).GetTransposed();
 
 	double InvViewProj[16];
 	Inverse4x4( InvViewProj, (float*)ViewProj.M );
diff --git a/Engine/Source/Runtime/Renderer/Private/DeferredShadingRenderer.cpp b/Engine/Source/Runtime/Renderer/Private/DeferredShadingRenderer.cpp
--- a/Engine/Source/Runtime/Renderer/Private/DeferredShadingRenderer.cpp
+++ b/Engine/Source/Runtime/Renderer/Private/DeferredShadingRenderer.cpp
@@ -644,12 +644,15 @@ void FDeferredShadingSceneRenderer::Render(FRHICommandListImmediate& RHICmdList)
 	}
 	SCOPED_DRAW_EVENT(RHICmdList, Scene);
 	
-	// Initialize global system textures (pass-through if already initialized).
-	GSystemTextures.InitializeTextures(RHICmdList, FeatureLevel);
+	{
+		QUICK_SCOPE_CYCLE_COUNTER(STAT_FDeferredShadingSceneRenderer_Render_Init);
 
-	// Allocate the maximum scene render target space for the current view family.
-	GSceneRenderTargets.Allocate(ViewFamily);
+		// Initialize global system textures (pass-through if already initialized).
+		GSystemTextures.InitializeTextures(RHICmdList, FeatureLevel);
 
+		// Allocate the maximum scene render target space for the current view family.
+		GSceneRenderTargets.Allocate(ViewFamily);
+	}
 	// Find the visible primitives.
 	InitViews(RHICmdList);
 
@@ -662,6 +665,7 @@ void FDeferredShadingSceneRenderer::Render(FRHICommandListImmediate& RHICmdList)
 
 	if (ShouldPrepareForDistanceFieldAO() || ShouldPrepareForDistanceFieldShadows())
 	{
+		QUICK_SCOPE_CYCLE_COUNTER(STAT_FDeferredShadingSceneRenderer_DistanceFieldAO_Init);
 		GDistanceFieldVolumeTextureAtlas.UpdateAllocations();
 		UpdateGlobalDistanceFieldObjectBuffers(RHICmdList);
 
@@ -730,8 +734,11 @@ void FDeferredShadingSceneRenderer::Render(FRHICommandListImmediate& RHICmdList)
 	const bool bIsOcclusionTesting = DoOcclusionQueries(FeatureLevel) && (!bIsWireframe || bIsViewFrozen || bHasViewParent);
 
 	// Dynamic vertex and index buffers need to be committed before rendering.
-	FGlobalDynamicVertexBuffer::Get().Commit();
-	FGlobalDynamicIndexBuffer::Get().Commit();
+	{
+		QUICK_SCOPE_CYCLE_COUNTER(STAT_FDeferredShadingSceneRenderer_FGlobalDynamicVertexBuffer_Commit);
+		FGlobalDynamicVertexBuffer::Get().Commit();
+		FGlobalDynamicIndexBuffer::Get().Commit();
+	}
 
 	// Notify the FX system that the scene is about to be rendered.
 	if (Scene->FXSystem)
@@ -754,12 +761,17 @@ void FDeferredShadingSceneRenderer::Render(FRHICommandListImmediate& RHICmdList)
 	
 	const bool bShouldRenderVelocities = ShouldRenderVelocities();
 	const bool bUseVelocityGBuffer = FVelocityRendering::OutputsToGBuffer();
-	GSceneRenderTargets.PreallocGBufferTargets(bShouldRenderVelocities && bUseVelocityGBuffer);
-	GSceneRenderTargets.AllocGBufferTargets();
+
+	{
+		QUICK_SCOPE_CYCLE_COUNTER(STAT_FDeferredShadingSceneRenderer_AllocGBufferTargets);
+		GSceneRenderTargets.PreallocGBufferTargets(bShouldRenderVelocities && bUseVelocityGBuffer);
+		GSceneRenderTargets.AllocGBufferTargets();
+	}
 	
 	// Clear LPVs for all views
 	if ( FeatureLevel >= ERHIFeatureLevel::SM5 )
 	{
+		QUICK_SCOPE_CYCLE_COUNTER(STAT_FDeferredShadingSceneRenderer_ClearLPVs);
 		ClearLPVs(RHICmdList);
 	}
 
@@ -770,6 +782,7 @@ void FDeferredShadingSceneRenderer::Render(FRHICommandListImmediate& RHICmdList)
 
 	if(bDBuffer)
 	{
+		QUICK_SCOPE_CYCLE_COUNTER(STAT_FDeferredShadingSceneRenderer_DBuffer);
 		GSceneRenderTargets.ResolveSceneDepthTexture(RHICmdList);
 		GSceneRenderTargets.ResolveSceneDepthToAuxiliaryTexture(RHICmdList);
 
@@ -785,6 +798,7 @@ void FDeferredShadingSceneRenderer::Render(FRHICommandListImmediate& RHICmdList)
 	bool bIsGBufferCurrent = false;
 	if (bRequiresRHIClear)
 	{
+		QUICK_SCOPE_CYCLE_COUNTER(STAT_FDeferredShadingSceneRenderer_SetAndClearViewGBuffer);
 		// set GBuffer to be current, and clear it
 		SetAndClearViewGBuffer(RHICmdList, Views[0], !bDepthWasCleared);
 
@@ -812,24 +826,25 @@ void FDeferredShadingSceneRenderer::Render(FRHICommandListImmediate& RHICmdList)
 	RenderBasePass(RHICmdList);
 	RHICmdList.SetCurrentStat(GET_STATID(STAT_CLM_AfterBasePass));
 
-	if(ViewFamily.EngineShowFlags.VisualizeLightCulling)
-	{
-		// clear out emissive and baked lighting (not too efficient but simple and only needed for this debug view)
-		GSceneRenderTargets.BeginRenderingSceneColor(RHICmdList);
-		RHICmdList.Clear(true, FLinearColor(0, 0, 0, 0), false, 0, false, 0, FIntRect());
-	}
+	  if(ViewFamily.EngineShowFlags.VisualizeLightCulling)
+	  {
+		  // clear out emissive and baked lighting (not too efficient but simple and only needed for this debug view)
+		  GSceneRenderTargets.BeginRenderingSceneColor(RHICmdList);
+		  RHICmdList.Clear(true, FLinearColor(0, 0, 0, 0), false, 0, false, 0, FIntRect());
+	  }
 
-	GSceneRenderTargets.DBufferA.SafeRelease();
-	GSceneRenderTargets.DBufferB.SafeRelease();
-	GSceneRenderTargets.DBufferC.SafeRelease();
+	  GSceneRenderTargets.DBufferA.SafeRelease();
+	  GSceneRenderTargets.DBufferB.SafeRelease();
+	  GSceneRenderTargets.DBufferC.SafeRelease();
 
-	// only temporarily available after early z pass and until base pass
-	check(!GSceneRenderTargets.DBufferA);
-	check(!GSceneRenderTargets.DBufferB);
-	check(!GSceneRenderTargets.DBufferC);
+	  // only temporarily available after early z pass and until base pass
+	  check(!GSceneRenderTargets.DBufferA);
+	  check(!GSceneRenderTargets.DBufferB);
+	  check(!GSceneRenderTargets.DBufferC);
 
 	if (bRequiresFarZQuadClear)
 	{
+		QUICK_SCOPE_CYCLE_COUNTER(STAT_FDeferredShadingSceneRenderer_ClearGBufferAtMaxZ);
 		// Clears view by drawing quad at maximum Z
 		// TODO: if all the platforms have fast color clears, we can replace this with an RHICmdList.Clear.
 		ClearGBufferAtMaxZ(RHICmdList);
@@ -837,14 +852,17 @@ void FDeferredShadingSceneRenderer::Render(FRHICommandListImmediate& RHICmdList)
 		bRequiresFarZQuadClear = false;
 	}
 	
-	GSceneRenderTargets.ResolveSceneColor(RHICmdList, FResolveRect(0, 0, ViewFamily.FamilySizeX, ViewFamily.FamilySizeY));
-	GSceneRenderTargets.ResolveSceneDepthTexture(RHICmdList);
-	GSceneRenderTargets.ResolveSceneDepthToAuxiliaryTexture(RHICmdList);
+	{
+		QUICK_SCOPE_CYCLE_COUNTER(STAT_FDeferredShadingSceneRenderer_Resolve_After_Basepass);
 
-	GSceneRenderTargets.FinishRenderingGBuffer(RHICmdList);
-	
-	RenderCustomDepthPass(RHICmdList);
+		GSceneRenderTargets.ResolveSceneColor(RHICmdList, FResolveRect(0, 0, ViewFamily.FamilySizeX, ViewFamily.FamilySizeY));
+		GSceneRenderTargets.ResolveSceneDepthTexture(RHICmdList);
+		GSceneRenderTargets.ResolveSceneDepthToAuxiliaryTexture(RHICmdList);
+
+		GSceneRenderTargets.FinishRenderingGBuffer(RHICmdList);
 
+		RenderCustomDepthPass(RHICmdList);
+	}
 	// Notify the FX system that opaque primitives have been rendered and we now have a valid depth buffer.
 	if (Scene->FXSystem && Views.IsValidIndex(0))
 	{
@@ -858,7 +876,10 @@ void FDeferredShadingSceneRenderer::Render(FRHICommandListImmediate& RHICmdList)
 
 	// Update the quarter-sized depth buffer with the current contents of the scene depth texture.
 	// This needs to happen before occlusion tests, which makes use of the small depth buffer.
-	UpdateDownsampledDepthSurface(RHICmdList);
+	{
+		QUICK_SCOPE_CYCLE_COUNTER(STAT_FDeferredShadingSceneRenderer_UpdateDownsampledDepthSurface);
+		UpdateDownsampledDepthSurface(RHICmdList);
+	}
 
 	for (int32 ViewIndex = 0; ViewIndex < Views.Num(); ViewIndex++)
 	{
@@ -899,6 +920,8 @@ void FDeferredShadingSceneRenderer::Render(FRHICommandListImmediate& RHICmdList)
 		&& bGBuffer
 		)
 	{
+		QUICK_SCOPE_CYCLE_COUNTER(STAT_FDeferredShadingSceneRenderer_Lighting);
+
 		GRenderTargetPool.AddPhaseEvent(TEXT("Lighting"));
 
 		// Pre-lighting composition lighting stage
@@ -958,12 +981,14 @@ void FDeferredShadingSceneRenderer::Render(FRHICommandListImmediate& RHICmdList)
 	// Draw Lightshafts
 	if (ViewFamily.EngineShowFlags.LightShafts)
 	{
+		QUICK_SCOPE_CYCLE_COUNTER(STAT_FDeferredShadingSceneRenderer_RenderLightShaftOcclusion);
 		LightShaftOutput = RenderLightShaftOcclusion(RHICmdList);
 	}
 
 	// Draw atmosphere
 	if(ShouldRenderAtmosphere(ViewFamily))
 	{
+		QUICK_SCOPE_CYCLE_COUNTER(STAT_FDeferredShadingSceneRenderer_RenderAtmosphere);
 		if (Scene->AtmosphericFog)
 		{
 			// Update RenderFlag based on LightShaftTexture is valid or not
@@ -991,6 +1016,7 @@ void FDeferredShadingSceneRenderer::Render(FRHICommandListImmediate& RHICmdList)
 	// Draw fog.
 	if(ShouldRenderFog(ViewFamily))
 	{
+		QUICK_SCOPE_CYCLE_COUNTER(STAT_FDeferredShadingSceneRenderer_RenderFog);
 		RenderFog(RHICmdList, LightShaftOutput);
 	}
 
@@ -1018,6 +1044,7 @@ void FDeferredShadingSceneRenderer::Render(FRHICommandListImmediate& RHICmdList)
 	}
 	if (ViewFamily.EngineShowFlags.LightShafts)
 	{
+		QUICK_SCOPE_CYCLE_COUNTER(STAT_FDeferredShadingSceneRenderer_RenderLightShaftBloom);
 		RenderLightShaftBloom(RHICmdList);
 	}
 
@@ -1063,13 +1090,17 @@ void FDeferredShadingSceneRenderer::Render(FRHICommandListImmediate& RHICmdList)
 	//grab the new transform out of the proxies for next frame
 	if (VelocityRT)
 	{
+		QUICK_SCOPE_CYCLE_COUNTER(STAT_FDeferredShadingSceneRenderer_UpdateMotionBlurCache);
 		Scene->MotionBlurInfoData.UpdateMotionBlurCache(Scene);
 	}
 
 	VelocityRT.SafeRelease();
 
-	RenderFinish(RHICmdList);
-	RHICmdList.SetCurrentStat(GET_STATID(STAT_CLM_AfterFrame));
+	{
+		QUICK_SCOPE_CYCLE_COUNTER(STAT_FDeferredShadingSceneRenderer_RenderFinish);
+		RenderFinish(RHICmdList);
+		RHICmdList.SetCurrentStat(GET_STATID(STAT_CLM_AfterFrame));
+	}
 }
 
 bool FDeferredShadingSceneRenderer::RenderPrePassViewDynamic(FRHICommandList& RHICmdList, const FViewInfo& View)
@@ -1323,6 +1354,7 @@ bool FDeferredShadingSceneRenderer::RenderBasePass(FRHICommandListImmediate& RHI
 
 	if (FVelocityRendering::OutputsToGBuffer())
 	{
+		QUICK_SCOPE_CYCLE_COUNTER(STAT_FDeferredShadingSceneRenderer_RenderBasePass_GPrevPerBoneMotionBlur_LockData);
 		GPrevPerBoneMotionBlur.LockData();
 	}
 
@@ -1361,6 +1393,7 @@ bool FDeferredShadingSceneRenderer::RenderBasePass(FRHICommandListImmediate& RHI
 
 	if (FVelocityRendering::OutputsToGBuffer())
 	{
+		QUICK_SCOPE_CYCLE_COUNTER(STAT_FDeferredShadingSceneRenderer_RenderBasePass_GPrevPerBoneMotionBlur_UnlockData);
 		GPrevPerBoneMotionBlur.UnlockData();
 	}
 
diff --git a/Engine/Source/Runtime/Renderer/Private/ForwardShadingRenderer.cpp b/Engine/Source/Runtime/Renderer/Private/ForwardShadingRenderer.cpp
--- a/Engine/Source/Runtime/Renderer/Private/ForwardShadingRenderer.cpp
+++ b/Engine/Source/Runtime/Renderer/Private/ForwardShadingRenderer.cpp
@@ -77,6 +77,8 @@ void FForwardShadingSceneRenderer::InitViews(FRHICommandListImmediate& RHICmdLis
 */
 void FForwardShadingSceneRenderer::Render(FRHICommandListImmediate& RHICmdList)
 {
+	QUICK_SCOPE_CYCLE_COUNTER(STAT_FForwardShadingSceneRenderer_Render);
+
 	if(!ViewFamily.EngineShowFlags.Rendering)
 	{
 		return;
diff --git a/Engine/Source/Runtime/Renderer/Private/PostProcess/PostProcessMotionBlur.cpp b/Engine/Source/Runtime/Renderer/Private/PostProcess/PostProcessMotionBlur.cpp
--- a/Engine/Source/Runtime/Renderer/Private/PostProcess/PostProcessMotionBlur.cpp
+++ b/Engine/Source/Runtime/Renderer/Private/PostProcess/PostProcessMotionBlur.cpp
@@ -9,6 +9,7 @@
 #include "SceneFilterRendering.h"
 #include "PostProcessAmbientOcclusion.h"
 #include "PostProcessMotionBlur.h"
+#include "PostProcessAmbientOcclusion.h"
 #include "PostProcessing.h"
 #include "SceneUtils.h"
 
@@ -979,7 +980,7 @@ void FRCPassPostProcessVelocityScatter::Process(FRenderingCompositePassContext&
 
 	TRefCountPtr<IPooledRenderTarget> DepthTarget;
 	
-	FPooledRenderTargetDesc Desc( FPooledRenderTargetDesc::Create2DDesc( DestRect.Size(), PF_ShadowDepth, TexCreate_None, TexCreate_DepthStencilTargetable, false ) );
+	FPooledRenderTargetDesc Desc( FPooledRenderTargetDesc::Create2DDesc( DestSize, PF_ShadowDepth, TexCreate_None, TexCreate_DepthStencilTargetable, false ) );
 	GRenderTargetPool.FindFreeElement( Desc, DepthTarget, TEXT("VelocityScatterDepth") );
 
 	// Set the view family's render target/viewport.
diff --git a/Engine/Source/Runtime/Renderer/Private/PostProcess/PostProcessTemporalAA.cpp b/Engine/Source/Runtime/Renderer/Private/PostProcess/PostProcessTemporalAA.cpp
--- a/Engine/Source/Runtime/Renderer/Private/PostProcess/PostProcessTemporalAA.cpp
+++ b/Engine/Source/Runtime/Renderer/Private/PostProcess/PostProcessTemporalAA.cpp
@@ -450,8 +450,6 @@ void FRCPassPostProcessLightShaftTemporalAA::Process(FRenderingCompositePassCont
 		return;
 	}
 
-	SCOPED_DRAW_EVENT(Context.RHICmdList, LSTemporalAA);
-
 	const FViewInfo& View = Context.View;
 	FSceneViewState* ViewState = Context.ViewState;
 
@@ -717,11 +715,8 @@ void FRCPassPostProcessTemporalAA::Process(FRenderingCompositePassContext& Conte
 	View.ViewProjectionMatrix = View.ViewMatrices.ViewMatrix * View.ViewMatrices.ProjMatrix;
 	View.InvViewProjectionMatrix = View.ViewMatrices.GetInvProjMatrix() * View.InvViewMatrix;
 
-	/** The view transform, starting from world-space points translated by -ViewOrigin. */
-	FMatrix TranslatedViewMatrix = FTranslationMatrix(-View.ViewMatrices.PreViewTranslation) * View.ViewMatrices.ViewMatrix;
-
 	// Compute a transform from view origin centered world-space to clip space.
-	View.ViewMatrices.TranslatedViewProjectionMatrix = TranslatedViewMatrix * View.ViewMatrices.ProjMatrix;
+	View.ViewMatrices.TranslatedViewProjectionMatrix = View.ViewMatrices.TranslatedViewMatrix * View.ViewMatrices.ProjMatrix;
 	View.ViewMatrices.InvTranslatedViewProjectionMatrix = View.ViewMatrices.TranslatedViewProjectionMatrix.Inverse();
 
 	View.InitRHIResources(nullptr);
diff --git a/Engine/Source/Runtime/Renderer/Private/PostProcess/PostProcessTonemap.cpp b/Engine/Source/Runtime/Renderer/Private/PostProcess/PostProcessTonemap.cpp
--- a/Engine/Source/Runtime/Renderer/Private/PostProcess/PostProcessTonemap.cpp
+++ b/Engine/Source/Runtime/Renderer/Private/PostProcess/PostProcessTonemap.cpp
@@ -91,29 +91,37 @@ static uint32 TonemapperConfBitmaskPC[22] = {
 	TonemapperContrast +
 	TonemapperColorMatrix +
 	TonemapperShadowTint +
-	TonemapperVignette +
+	TonemapperGrainJitter +
+	TonemapperGrainIntensity +
 	TonemapperGrainQuantization +
+	TonemapperVignette +
+	TonemapperVignetteColor +
 	TonemapperColorFringe +
 	TonemapperColorGrading +
+	Tonemapper709 +
 	0,
 
 	TonemapperBloom + 
 	TonemapperContrast + 
 	TonemapperColorMatrix + 
 	TonemapperShadowTint +
+	TonemapperVignette +
 	TonemapperGrainQuantization +
+	TonemapperColorFringe +
 	TonemapperColorGrading +
 	0,
 
 	TonemapperBloom + 
 	TonemapperContrast + 
 	TonemapperColorMatrix +
+	TonemapperShadowTint +
 	TonemapperGrainQuantization +
 	TonemapperColorGrading +
 	0,
 
-	TonemapperBloom	+ 
+	TonemapperBloom + 
 	TonemapperContrast +
+	TonemapperColorMatrix +
 	TonemapperGrainQuantization +
 	TonemapperColorGrading +
 	0,
@@ -136,26 +144,34 @@ static uint32 TonemapperConfBitmaskPC[22] = {
 	TonemapperContrast +
 	TonemapperColorMatrix +
 	TonemapperShadowTint +
+	TonemapperGrainJitter +
+	TonemapperGrainIntensity +
 	TonemapperVignette +
+	TonemapperVignetteColor +
 	TonemapperColorFringe +
 	TonemapperColorGrading +
+	Tonemapper709 +
 	0,
 
 	TonemapperBloom + 
 	TonemapperContrast + 
 	TonemapperColorMatrix + 
 	TonemapperShadowTint +
+	TonemapperVignette +
+	TonemapperColorFringe +
 	TonemapperColorGrading +
 	0,
 
 	TonemapperBloom + 
 	TonemapperContrast + 
 	TonemapperColorMatrix +
+	TonemapperShadowTint +
 	TonemapperColorGrading +
 	0,
 
-	TonemapperBloom	+ 
+	TonemapperBloom + 
 	TonemapperContrast +
+	TonemapperColorMatrix +
 	TonemapperColorGrading +
 	0,
 	
@@ -711,32 +727,7 @@ static uint32 TonemapperGenerateBitmaskPC(const FViewInfo* RESTRICT View, bool b
 
 		if(Value > 0)
 		{
-			// Remove settings not compatible.
-			Bitmask &= ~(TonemapperColorMatrix|TonemapperShadowTint|TonemapperContrast|TonemapperColorGrading|TonemapperMsaa);
-			Bitmask |= TonemapperPhoto | Tonemapper709;
-		}
-	}
-
-	{
-		static TConsoleVariableData<float>* CVar = IConsoleManager::Get().FindTConsoleVariableDataFloat(TEXT("r.TonemapperGamma"));
-		float Value = CVar->GetValueOnRenderThread();
-		if(Value > 0.0f)
-		{
-			// Remove settings not compatible.
-			Bitmask &= ~(TonemapperColorMatrix|TonemapperShadowTint|TonemapperContrast|TonemapperColorGrading|TonemapperMsaa);
-			Bitmask |= TonemapperPhoto | TonemapperGamma;
-		}
-	}
-
-	{
-		static TConsoleVariableData<int32>* CVar = IConsoleManager::Get().FindTConsoleVariableDataInt(TEXT("r.TonemapperPhoto"));
-		int32 Value = CVar->GetValueOnRenderThread();
-
-		if(Value > 0)
-		{
-			// Remove settings not compatible.
-			Bitmask &= ~(TonemapperColorMatrix|TonemapperShadowTint|TonemapperContrast|TonemapperColorGrading|TonemapperMsaa);
-			Bitmask |= TonemapperPhoto;
+			Bitmask |= Tonemapper709;
 		}
 	}
 
@@ -757,6 +748,18 @@ static uint32 TonemapperGenerateBitmaskPC(const FViewInfo* RESTRICT View, bool b
 		Bitmask |= TonemapperColorGrading;
 	}
 
+	{
+		static TConsoleVariableData<int32>* CVar = IConsoleManager::Get().FindTConsoleVariableDataInt(TEXT("r.TonemapperPhoto"));
+		int32 Value = CVar->GetValueOnRenderThread();
+
+		if(Value > 0)
+		{
+			// Remove settings not compatible.
+			Bitmask &= ~(TonemapperColorMatrix|TonemapperShadowTint|TonemapperContrast|TonemapperColorGrading|TonemapperMsaa);
+			Bitmask |= TonemapperPhoto;
+		}
+	}
+
 	return Bitmask + TonemapperGenerateBitmaskPost(View);
 }
 
@@ -816,6 +819,10 @@ void FilmPostSetConstants(FVector4* RESTRICT const Constants, const uint32 Confi
 	uint32 UseShadowTint  = TonemapperIsDefined(ConfigBitmask, TonemapperShadowTint);
 	uint32 UseContrast    = TonemapperIsDefined(ConfigBitmask, TonemapperContrast);
 
+	// DON'T CHECK INTO MAIN: NEEDS A PERMUTATION TO AVOID NON-BACKWARDS COMPATIBLE CHANGES
+	// This corrects so the linear segment is linear in log space!
+	float In18Grey = FMath::Log2(1.0f+0.18f);
+
 	// Must insure inputs are in correct range (else possible generation of NaNs).
 	float InExposure = 1.0f;
 	FVector InWhitePoint(FinalPostProcessSettings->FilmWhitePoint);
@@ -826,9 +833,9 @@ void FilmPostSetConstants(FVector4* RESTRICT const Constants, const uint32 Confi
 	FVector InMatrixB(FinalPostProcessSettings->FilmChannelMixerBlue);
 	float InContrast = FMath::Clamp(FinalPostProcessSettings->FilmContrast, 0.0f, 1.0f) + 1.0f;
 	float InDynamicRange = powf(2.0f, FMath::Clamp(FinalPostProcessSettings->FilmDynamicRange, 1.0f, 4.0f));
-	float InToe = (1.0f - FMath::Clamp(FinalPostProcessSettings->FilmToeAmount, 0.0f, 1.0f)) * 0.18f;
-	InToe = FMath::Clamp(InToe, 0.18f/8.0f, 0.18f * (15.0f/16.0f));
-	float InHeal = 1.0f - (FMath::Max(1.0f/32.0f, 1.0f - FMath::Clamp(FinalPostProcessSettings->FilmHealAmount, 0.0f, 1.0f)) * (1.0f - 0.18f)); 
+	float InToe = (1.0f - FMath::Clamp(FinalPostProcessSettings->FilmToeAmount, 0.0f, 1.0f)) * In18Grey;
+	InToe = FMath::Clamp(InToe, In18Grey/8.0f, In18Grey * (15.0f/16.0f));
+	float InHeal = 1.0f - (FMath::Max(1.0f/32.0f, 1.0f - FMath::Clamp(FinalPostProcessSettings->FilmHealAmount, 0.0f, 1.0f)) * (1.0f - In18Grey)); 
 	FVector InShadowTint(FinalPostProcessSettings->FilmShadowTint);
 	float InShadowTintBlend = FMath::Clamp(FinalPostProcessSettings->FilmShadowTintBlend, 0.0f, 1.0f) * 64.0f;
 
@@ -908,7 +915,7 @@ void FilmPostSetConstants(FVector4* RESTRICT const Constants, const uint32 Confi
 	float OutColorCurveCm2;
 
 	// Line for linear section.
-	float FilmLineOffset = 0.18f - 0.18f*InContrast;
+	float FilmLineOffset = In18Grey - In18Grey*InContrast;
 	float FilmXAtY0 = -FilmLineOffset/InContrast;
 	float FilmXAtY1 = (1.0f - FilmLineOffset) / InContrast;
 	float FilmXS = FilmXAtY1 - FilmXAtY0;
@@ -1084,6 +1091,13 @@ BEGIN_UNIFORM_BUFFER_STRUCT(FBloomDirtMaskParameters,)
 END_UNIFORM_BUFFER_STRUCT(FBloomDirtMaskParameters)
 IMPLEMENT_UNIFORM_BUFFER_STRUCT(FBloomDirtMaskParameters,TEXT("BloomDirtMask"));
 
+
+static TAutoConsoleVariable<float> CVarGamma(
+	TEXT("r.Gamma"),
+	1.0f,
+	TEXT("Gamma on output"),
+	ECVF_RenderThreadSafe);
+
 /**
  * Encapsulates the post processing tonemapper pixel shader.
  */
diff --git a/Engine/Source/Runtime/Renderer/Private/PostProcess/PostProcessTonemap.h b/Engine/Source/Runtime/Renderer/Private/PostProcess/PostProcessTonemap.h
--- a/Engine/Source/Runtime/Renderer/Private/PostProcess/PostProcessTonemap.h
+++ b/Engine/Source/Runtime/Renderer/Private/PostProcess/PostProcessTonemap.h
@@ -146,8 +146,19 @@ public:
 			// for scene color fringe
 			// from percent to fraction
 			float Offset = Context.View.FinalPostProcessSettings.SceneFringeIntensity * 0.01f;
+			//FVector4 Value(1.0f - Offset * 0.5f, 1.0f - Offset, 0.0f, 0.0f);
+
+			// Wavelength of primaries in nm
+			const float PrimaryR = 611.3f;
+			const float PrimaryG = 549.1f;
+			const float PrimaryB = 464.3f;
+
+			// Simple lens chromatic aberration is roughly linear in wavelength
+			float ScaleR = 0.007f * ( PrimaryR - PrimaryB );
+			float ScaleG = 0.007f * ( PrimaryG - PrimaryB );
+			FVector4 Value( 1.0f / ( 1.0f + Offset * ScaleG ), 1.0f / ( 1.0f + Offset * ScaleR ), 0.0f, 0.0f);
+
 			// we only get bigger to not leak in content from outside
-			FVector4 Value(1.0f - Offset * 0.5f, 1.0f - Offset, 0.0f, 0.0f);
 			SetShaderValue(Context.RHICmdList, ShaderRHI, FringeUVParams, Value);
 		}
 	}
diff --git a/Engine/Source/Runtime/Renderer/Private/PrimitiveSceneInfo.cpp b/Engine/Source/Runtime/Renderer/Private/PrimitiveSceneInfo.cpp
--- a/Engine/Source/Runtime/Renderer/Private/PrimitiveSceneInfo.cpp
+++ b/Engine/Source/Runtime/Renderer/Private/PrimitiveSceneInfo.cpp
@@ -306,6 +306,8 @@ void FPrimitiveSceneInfo::UpdateStaticMeshes(FRHICommandListImmediate& RHICmdLis
 	checkSlow(bNeedsStaticMeshUpdate);
 	bNeedsStaticMeshUpdate = false;
 
+	QUICK_SCOPE_CYCLE_COUNTER(STAT_FPrimitiveSceneInfo_UpdateStaticMeshes);
+
 	// Remove the primitive's static meshes from the draw lists they're currently in, and re-add them to the appropriate draw lists.
 	for (int32 MeshIndex = 0; MeshIndex < StaticMeshes.Num(); MeshIndex++)
 	{
diff --git a/Engine/Source/Runtime/Renderer/Private/RendererScene.cpp b/Engine/Source/Runtime/Renderer/Private/RendererScene.cpp
--- a/Engine/Source/Runtime/Renderer/Private/RendererScene.cpp
+++ b/Engine/Source/Runtime/Renderer/Private/RendererScene.cpp
@@ -1495,41 +1495,46 @@ const TArray<FWindSourceSceneProxy*>& FScene::GetWindSources_RenderThread() cons
 	return WindSources;
 }
 
-FVector4 FScene::GetWindParameters(const FVector& Position) const
+void FScene::GetWindParameters(const FVector& Position, FVector& OutDirection, float& OutSpeed, float& OutMinGustAmt, float& OutMaxGustAmt) const
 {
+	FWindSourceSceneProxy::FWindData AccumWindData;
+	AccumWindData.PrepareForAccumulate();
+
 	int32 NumActiveWindSources = 0;
 	FVector4 AccumulatedDirectionAndSpeed(0,0,0,0);
 	float TotalWeight = 0.0f;
 	for (int32 i = 0; i < WindSources.Num(); i++)
 	{
+		
 		FVector4 CurrentDirectionAndSpeed;
 		float Weight;
 		const FWindSourceSceneProxy* CurrentSource = WindSources[i];
-		if (CurrentSource->GetWindParameters(Position, CurrentDirectionAndSpeed, Weight))
+		FWindSourceSceneProxy::FWindData CurrentSourceData;
+		if (CurrentSource->GetWindParameters(Position, CurrentSourceData, Weight))
 		{
-			AccumulatedDirectionAndSpeed.X += CurrentDirectionAndSpeed.X * Weight;
-			AccumulatedDirectionAndSpeed.Y += CurrentDirectionAndSpeed.Y * Weight;
-			AccumulatedDirectionAndSpeed.Z += CurrentDirectionAndSpeed.Z * Weight;
-			AccumulatedDirectionAndSpeed.W += CurrentDirectionAndSpeed.W * Weight;
+			AccumWindData.AddWeighted(CurrentSourceData, Weight);
 			TotalWeight += Weight;
 			NumActiveWindSources++;
 		}
 	}
 
-	if (TotalWeight > 0)
+	AccumWindData.NormalizeByTotalWeight(TotalWeight);
+
+	if (NumActiveWindSources == 0)
 	{
-		AccumulatedDirectionAndSpeed.X /= TotalWeight;
-		AccumulatedDirectionAndSpeed.Y /= TotalWeight;
-		AccumulatedDirectionAndSpeed.Z /= TotalWeight;
-		AccumulatedDirectionAndSpeed.W /= TotalWeight;
+		AccumWindData.Direction = FVector(1.0f, 0.0f, 0.0f);
 	}
-
-	// Normalize averaged direction and speed
-	return NumActiveWindSources > 0 ? AccumulatedDirectionAndSpeed / NumActiveWindSources : FVector4(0,0,0,0);
+	OutDirection	= AccumWindData.Direction;
+	OutSpeed		= AccumWindData.Speed;
+	OutMinGustAmt	= AccumWindData.MinGustAmt;
+	OutMaxGustAmt	= AccumWindData.MaxGustAmt;
 }
 
-FVector4 FScene::GetDirectionalWindParameters(void) const
+void FScene::GetDirectionalWindParameters(FVector& OutDirection, float& OutSpeed, float& OutMinGustAmt, float& OutMaxGustAmt) const
 {
+	FWindSourceSceneProxy::FWindData AccumWindData;
+	AccumWindData.PrepareForAccumulate();
+
 	int32 NumActiveWindSources = 0;
 	FVector4 AccumulatedDirectionAndSpeed(0,0,0,0);
 	float TotalWeight = 0.0f;
@@ -1538,27 +1543,25 @@ FVector4 FScene::GetDirectionalWindParameters(void) const
 		FVector4 CurrentDirectionAndSpeed;
 		float Weight;
 		const FWindSourceSceneProxy* CurrentSource = WindSources[i];
-		if (CurrentSource->GetDirectionalWindParameters(CurrentDirectionAndSpeed, Weight))
+		FWindSourceSceneProxy::FWindData CurrentSourceData;
+		if (CurrentSource->GetDirectionalWindParameters(CurrentSourceData, Weight))
 		{
-			AccumulatedDirectionAndSpeed.X += CurrentDirectionAndSpeed.X * Weight;
-			AccumulatedDirectionAndSpeed.Y += CurrentDirectionAndSpeed.Y * Weight;
-			AccumulatedDirectionAndSpeed.Z += CurrentDirectionAndSpeed.Z * Weight;
-			AccumulatedDirectionAndSpeed.W += CurrentDirectionAndSpeed.W * Weight;
+			AccumWindData.AddWeighted(CurrentSourceData, Weight);			
 			TotalWeight += Weight;
 			NumActiveWindSources++;
 		}
 	}
 
-	if (TotalWeight > 0)
+	AccumWindData.NormalizeByTotalWeight(TotalWeight);	
+
+	if (NumActiveWindSources == 0)
 	{
-		AccumulatedDirectionAndSpeed.X /= TotalWeight;
-		AccumulatedDirectionAndSpeed.Y /= TotalWeight;
-		AccumulatedDirectionAndSpeed.Z /= TotalWeight;
-		AccumulatedDirectionAndSpeed.W /= TotalWeight;
+		AccumWindData.Direction = FVector(1.0f, 0.0f, 0.0f);
 	}
-
-	// Normalize averaged direction and speed
-	return NumActiveWindSources > 0 ? AccumulatedDirectionAndSpeed / NumActiveWindSources : FVector4(0,0,1,0);
+	OutDirection = AccumWindData.Direction;
+	OutSpeed = AccumWindData.Speed;
+	OutMinGustAmt = AccumWindData.MinGustAmt;
+	OutMaxGustAmt = AccumWindData.MaxGustAmt;
 }
 
 void FScene::AddSpeedTreeWind(FVertexFactory* VertexFactory, const UStaticMesh* StaticMesh)
@@ -1579,7 +1582,6 @@ void FScene::AddSpeedTreeWind(FVertexFactory* VertexFactory, const UStaticMesh*
 				}
 				else
 				{
-					UE_LOG(LogRenderer, Log, TEXT("Adding SpeedTree wind for static mesh %s"), *StaticMesh->GetName());
 					FSpeedTreeWindComputation* WindComputation = new FSpeedTreeWindComputation;
 					WindComputation->Wind = *(StaticMesh->SpeedTreeWind.Get( ));
 					WindComputation->UniformBuffer.InitResource();
@@ -1631,14 +1633,20 @@ void FScene::RemoveSpeedTreeWind_RenderThread(class FVertexFactory* VertexFactor
 
 void FScene::UpdateSpeedTreeWind(double CurrentTime)
 {
-	#define SET_SPEEDTREE_TABLE_FLOAT4V(name, offset) UniformParameters.name = *(FVector4*)(WindShaderValues + FSpeedTreeWind::offset);
+#define SET_SPEEDTREE_TABLE_FLOAT4V(name, offset) \
+	UniformParameters.name = *(FVector4*)(WindShaderValues + FSpeedTreeWind::offset); \
+	UniformParameters.Prev##name = *(FVector4*)(WindShaderValues + FSpeedTreeWind::offset + FSpeedTreeWind::NUM_SHADER_VALUES);
 
 	ENQUEUE_UNIQUE_RENDER_COMMAND_TWOPARAMETER(
 		FUpdateSpeedTreeWindCommand,
 		FScene*,Scene,this,
 		double,CurrentTime,CurrentTime,
 		{
-			FVector4 WindInfo = Scene->GetDirectionalWindParameters();
+			FVector WindDirection;
+			float WindSpeed;
+			float WindMinGustAmt;
+			float WindMaxGustAmt;
+			Scene->GetDirectionalWindParameters(WindDirection, WindSpeed, WindMinGustAmt, WindMaxGustAmt);
 
 			for (TMap<const UStaticMesh*, FSpeedTreeWindComputation*>::TIterator It(Scene->SpeedTreeWindComputationMap); It; ++It )
 			{
@@ -1665,8 +1673,10 @@ void FScene::UpdateSpeedTreeWind(double CurrentTime)
 				}
 
 				// advance the wind object
-				WindComputation->Wind.SetDirection(FVector(WindInfo));
-				WindComputation->Wind.SetStrength(WindInfo.W);
+				WindComputation->Wind.SetDirection(WindDirection);
+				WindComputation->Wind.SetStrength(WindSpeed);
+				WindComputation->Wind.SetGustMin(WindMinGustAmt);
+				WindComputation->Wind.SetGustMax(WindMaxGustAmt);
 				WindComputation->Wind.Advance(true, CurrentTime);
 
 				// copy data into uniform buffer
@@ -1697,8 +1707,7 @@ void FScene::UpdateSpeedTreeWind(double CurrentTime)
 				WindComputation->UniformBuffer.SetContents(UniformParameters);
 			}
 		});
-	
-	#undef SET_SPEEDTREE_TABLE_FLOAT4V
+#undef SET_SPEEDTREE_TABLE_FLOAT4V
 }
 
 FUniformBufferRHIParamRef FScene::GetSpeedTreeUniformBuffer(const FVertexFactory* VertexFactory)
@@ -2362,8 +2371,8 @@ public:
 		static TArray<class FWindSourceSceneProxy*> NullWindSources;
 		return NullWindSources;
 	}
-	virtual FVector4 GetWindParameters(const FVector& Position) const { return FVector4(0,0,1,0); }
-	virtual FVector4 GetDirectionalWindParameters() const { return FVector4(0,0,1,0); }
+	virtual void GetWindParameters(const FVector& Position, FVector& OutDirection, float& OutSpeed, float& OutMinGustAmt, float& OutMaxGustAmt) const { OutDirection = FVector(1.0f, 0.0f, 0.0f); OutSpeed = 0.0f; OutMinGustAmt = 0.0f; OutMaxGustAmt = 0.0f; }
+	virtual void GetDirectionalWindParameters(FVector& OutDirection, float& OutSpeed, float& OutMinGustAmt, float& OutMaxGustAmt) const { OutDirection = FVector(1.0f, 0.0f, 0.0f); OutSpeed = 0.0f; OutMinGustAmt = 0.0f; OutMaxGustAmt = 0.0f; }
 	virtual void AddSpeedTreeWind(class FVertexFactory* VertexFactory, const class UStaticMesh* StaticMesh) {}
 	virtual void RemoveSpeedTreeWind(class FVertexFactory* VertexFactory, const class UStaticMesh* StaticMesh) {}
 	virtual void RemoveSpeedTreeWind_RenderThread(class FVertexFactory* VertexFactory, const class UStaticMesh* StaticMesh) {}
diff --git a/Engine/Source/Runtime/Renderer/Private/ScenePrivate.h b/Engine/Source/Runtime/Renderer/Private/ScenePrivate.h
--- a/Engine/Source/Runtime/Renderer/Private/ScenePrivate.h
+++ b/Engine/Source/Runtime/Renderer/Private/ScenePrivate.h
@@ -1553,8 +1553,8 @@ public:
 	virtual void AddWindSource(UWindDirectionalSourceComponent* WindComponent);
 	virtual void RemoveWindSource(UWindDirectionalSourceComponent* WindComponent);
 	virtual const TArray<FWindSourceSceneProxy*>& GetWindSources_RenderThread() const;
-	virtual FVector4 GetWindParameters(const FVector& Position) const;
-	virtual FVector4 GetDirectionalWindParameters() const;
+	virtual void GetWindParameters(const FVector& Position, FVector& OutDirection, float& OutSpeed, float& OutMinGustAmt, float& OutMaxGustAmt) const override;
+	virtual void GetDirectionalWindParameters(FVector& OutDirection, float& OutSpeed, float& OutMinGustAmt, float& OutMaxGustAmt) const override;
 	virtual void AddSpeedTreeWind(FVertexFactory* VertexFactory, const UStaticMesh* StaticMesh);
 	virtual void RemoveSpeedTreeWind(FVertexFactory* VertexFactory, const UStaticMesh* StaticMesh);
 	virtual void RemoveSpeedTreeWind_RenderThread(FVertexFactory* VertexFactory, const UStaticMesh* StaticMesh);
diff --git a/Engine/Source/Runtime/Renderer/Private/SceneRendering.cpp b/Engine/Source/Runtime/Renderer/Private/SceneRendering.cpp
--- a/Engine/Source/Runtime/Renderer/Private/SceneRendering.cpp
+++ b/Engine/Source/Runtime/Renderer/Private/SceneRendering.cpp
@@ -569,6 +569,7 @@ TUniformBufferRef<FViewUniformShaderParameters> FViewInfo::CreateUniformBuffer(
 	static const auto CVar = IConsoleManager::Get().FindConsoleVariable(TEXT("r.DiffuseFromCaptures"));
 	const bool bUseLightmaps = CVar->GetInt() == 0;
 
+	ViewUniformShaderParameters.CameraCut = bCameraCut ? 1 : 0;
 	ViewUniformShaderParameters.UseLightmaps = bUseLightmaps ? 1 : 0;
 
 	if(State)
@@ -1347,27 +1348,34 @@ static void RenderViewFamily_RenderThread(FRHICommandListImmediate& RHICmdList,
 		}
 
 #if STATS
-		// Update scene memory stats that couldn't be tracked continuously
-		SET_MEMORY_STAT(STAT_StaticDrawListMemory, FStaticMeshDrawListBase::TotalBytesUsed);
-		SET_MEMORY_STAT(STAT_RenderingSceneMemory, SceneRenderer->Scene->GetSizeBytes());
-
-		SIZE_T ViewStateMemory = 0;
-		for (int32 ViewIndex = 0; ViewIndex < SceneRenderer->Views.Num(); ViewIndex++)
 		{
-			if (SceneRenderer->Views[ViewIndex].State)
+			QUICK_SCOPE_CYCLE_COUNTER(STAT_RenderViewFamily_RenderThread_MemStats);
+
+			// Update scene memory stats that couldn't be tracked continuously
+			SET_MEMORY_STAT(STAT_StaticDrawListMemory, FStaticMeshDrawListBase::TotalBytesUsed);
+			SET_MEMORY_STAT(STAT_RenderingSceneMemory, SceneRenderer->Scene->GetSizeBytes());
+
+			SIZE_T ViewStateMemory = 0;
+			for (int32 ViewIndex = 0; ViewIndex < SceneRenderer->Views.Num(); ViewIndex++)
 			{
-				ViewStateMemory += SceneRenderer->Views[ViewIndex].State->GetSizeBytes();
+				if (SceneRenderer->Views[ViewIndex].State)
+				{
+					ViewStateMemory += SceneRenderer->Views[ViewIndex].State->GetSizeBytes();
+				}
 			}
+			SET_MEMORY_STAT(STAT_ViewStateMemory, ViewStateMemory);
+			SET_MEMORY_STAT(STAT_RenderingMemStackMemory, FMemStack::Get().GetByteCount());
+			SET_MEMORY_STAT(STAT_LightInteractionMemory, FLightPrimitiveInteraction::GetMemoryPoolSize());
 		}
-		SET_MEMORY_STAT(STAT_ViewStateMemory, ViewStateMemory);
-		SET_MEMORY_STAT(STAT_RenderingMemStackMemory, FMemStack::Get().GetByteCount());
-		SET_MEMORY_STAT(STAT_LightInteractionMemory, FLightPrimitiveInteraction::GetMemoryPoolSize());
 #endif
 
 		GRenderTargetPool.SetEventRecordingActive(false);
 
 		// Delete the scene renderer.
-		delete SceneRenderer;
+		{
+			QUICK_SCOPE_CYCLE_COUNTER(STAT_DeleteSceneRenderer);
+			delete SceneRenderer;
+		}
 		{
 			QUICK_SCOPE_CYCLE_COUNTER(STAT_DeleteSceneRenderer_Dispatch);
 			FRHICommandListExecutor::GetImmediateCommandList().ImmediateFlush(EImmediateFlushType::DispatchToRHIThread); // we want to make sure this all gets to the GPU this frame and doesn't hang around
@@ -1375,6 +1383,7 @@ static void RenderViewFamily_RenderThread(FRHICommandListImmediate& RHICmdList,
 	}
 
 #if STATS
+	QUICK_SCOPE_CYCLE_COUNTER(STAT_RenderViewFamily_RenderThread_RHIGetGPUFrameCycles);
 	if (FPlatformProperties::SupportsWindowedMode() == false)
 	{
 		/** Update STATS with the total GPU time taken to render the last frame. */
diff --git a/Engine/Source/Runtime/Windows/D3D11RHI/Private/D3D11Device.cpp b/Engine/Source/Runtime/Windows/D3D11RHI/Private/D3D11Device.cpp
--- a/Engine/Source/Runtime/Windows/D3D11RHI/Private/D3D11Device.cpp
+++ b/Engine/Source/Runtime/Windows/D3D11RHI/Private/D3D11Device.cpp
@@ -27,6 +27,13 @@ bool D3D11RHI_ShouldAllowAsyncResourceCreation()
 
 IMPLEMENT_MODULE(FD3D11DynamicRHIModule, D3D11RHI);
 
+TAutoConsoleVariable<int32> CVarD3D11ZeroBufferSizeInMB(
+	TEXT("d3d11.ZeroBufferSizeInMB"),
+	4,
+	TEXT("The D3D11 RHI needs a static allocation of zeroes to use when streaming textures asynchronously. It should be large enough to support the largest mipmap you need to stream. The default is 4MB."),
+	ECVF_ReadOnly
+	);
+
 FD3D11DynamicRHI::FD3D11DynamicRHI(IDXGIFactory* InDXGIFactory,D3D_FEATURE_LEVEL InFeatureLevel, int32 InChosenAdapter) :
 	DXGIFactory(InDXGIFactory),
 	bDeviceRemoved(false),
@@ -56,7 +63,7 @@ FD3D11DynamicRHI::FD3D11DynamicRHI(IDXGIFactory* InDXGIFactory,D3D_FEATURE_LEVEL
 
 	// Allocate a buffer of zeroes. This is used when we need to pass D3D memory
 	// that we don't care about and will overwrite with valid data in the future.
-	ZeroBufferSize = 4 * (1 << 20); // 4MB in the zero buffer.
+	ZeroBufferSize = FMath::Max(CVarD3D11ZeroBufferSizeInMB.GetValueOnAnyThread(), 0) * (1 << 20);
 	ZeroBuffer = FMemory::Malloc(ZeroBufferSize);
 	FMemory::Memzero(ZeroBuffer,ZeroBufferSize);
 
diff --git a/Engine/Source/Runtime/Windows/D3D11RHI/Private/D3D11Texture.cpp b/Engine/Source/Runtime/Windows/D3D11RHI/Private/D3D11Texture.cpp
--- a/Engine/Source/Runtime/Windows/D3D11RHI/Private/D3D11Texture.cpp
+++ b/Engine/Source/Runtime/Windows/D3D11RHI/Private/D3D11Texture.cpp
@@ -74,13 +74,14 @@ void UpdateD3D11TextureStats(uint32 BindFlags, uint32 MiscFlags, int64 TextureSi
 		return;
 	}
 	
+	int64 AlignedSize = (TextureSize > 0) ? Align(TextureSize, 1024) / 1024 : -(Align(-TextureSize, 1024) / 1024);
 	if(ShouldCountAsTextureMemory(BindFlags))
 	{
-		FPlatformAtomics::InterlockedAdd(&GCurrentTextureMemorySize, Align(TextureSize, 1024) / 1024);
+		FPlatformAtomics::InterlockedAdd(&GCurrentTextureMemorySize, AlignedSize);
 	}
 	else
 	{
-		FPlatformAtomics::InterlockedAdd(&GCurrentRendertargetMemorySize, Align(TextureSize, 1024) / 1024);
+		FPlatformAtomics::InterlockedAdd(&GCurrentRendertargetMemorySize, AlignedSize);
 	}
 
 	bool bCubeMap = (MiscFlags & D3D11_RESOURCE_MISC_TEXTURECUBE) != 0;
