Fixed landscape tools to work with cross-level foliage
Fixed landscape retopologize tool that was losing foliage links

diff --git a/Engine/Source/Editor/LandscapeEditor/Private/LandscapeEdMode.cpp b/Engine/Source/Editor/LandscapeEditor/Private/LandscapeEdMode.cpp
--- a/Engine/Source/Editor/LandscapeEditor/Private/LandscapeEdMode.cpp
+++ b/Engine/Source/Editor/LandscapeEditor/Private/LandscapeEdMode.cpp
@@ -2616,13 +2616,8 @@ ALandscape* FEdModeLandscape::ChangeComponentSetting(int32 NumComponentsX, int32
 
 						if (OldCollisionComponent && NewCollisionComponent)
 						{
-							AInstancedFoliageActor* OldFoliageActor = AInstancedFoliageActor::GetInstancedFoliageActorForLevel(OldCollisionComponent->GetLandscapeProxy()->GetLevel());
-							if (OldFoliageActor)
-							{
-								OldFoliageActor->MoveInstancesToNewComponent(OldCollisionComponent, NewCollisionComponent);
-								AInstancedFoliageActor* NewFoliageActor = AInstancedFoliageActor::GetInstancedFoliageActorForLevel(NewCollisionComponent->GetTypedOuter<ULevel>(), true);
-								NewCollisionComponent->SnapFoliageInstances(*NewFoliageActor, FBox(FVector(-WORLD_MAX), FVector(WORLD_MAX)));
-							}
+							AInstancedFoliageActor::MoveInstancesToNewComponent(OldCollisionComponent->GetWorld(), OldCollisionComponent, NewCollisionComponent);
+							NewCollisionComponent->SnapFoliageInstances(FBox(FVector(-WORLD_MAX), FVector(WORLD_MAX)));
 						}
 					}
 				}
diff --git a/Engine/Source/Editor/LandscapeEditor/Private/LandscapeEdModeComponentTools.cpp b/Engine/Source/Editor/LandscapeEditor/Private/LandscapeEdModeComponentTools.cpp
--- a/Engine/Source/Editor/LandscapeEditor/Private/LandscapeEdModeComponentTools.cpp
+++ b/Engine/Source/Editor/LandscapeEditor/Private/LandscapeEdModeComponentTools.cpp
@@ -683,11 +683,7 @@ public:
 					// Need to move or recreate all related data (Height map, Weight map, maybe collision components, allocation info)
 
 					// Move any foliage associated
-					AInstancedFoliageActor* OldIFA = AInstancedFoliageActor::GetInstancedFoliageActorForLevel(Component->GetLandscapeProxy()->GetLevel());
-					if (OldIFA)
-					{
-						OldIFA->MoveInstancesForComponentToCurrentLevel(Component);
-					}
+					AInstancedFoliageActor::MoveInstancesForComponentToCurrentLevel(Component);
 
 					Component->GetLandscapeProxy()->CollisionComponents.Remove(Component);
 					Component->UnregisterComponent();
diff --git a/Engine/Source/Editor/LandscapeEditor/Private/LandscapeEdModeTools.h b/Engine/Source/Editor/LandscapeEditor/Private/LandscapeEdModeTools.h
--- a/Engine/Source/Editor/LandscapeEditor/Private/LandscapeEdModeTools.h
+++ b/Engine/Source/Editor/LandscapeEditor/Private/LandscapeEdModeTools.h
@@ -658,11 +658,7 @@ struct FHeightmapAccessor
 				for (int32 Index = 0; Index < CollisionComponents.Num(); ++Index)
 				{
 					ULandscapeHeightfieldCollisionComponent* CollisionComponent = CollisionComponents[Index];
-					AInstancedFoliageActor* IFA = AInstancedFoliageActor::GetInstancedFoliageActorForLevel(CollisionComponent->GetComponentLevel());
-					if (IFA)
-					{
-						CollisionComponent->SnapFoliageInstances(*IFA, PreUpdateLocalBoxes[Index].TransformBy(LandscapeInfo->GetLandscapeProxy()->LandscapeActorToWorld().ToMatrixWithScale()).ExpandBy(1.0f));
-					}
+					CollisionComponent->SnapFoliageInstances(PreUpdateLocalBoxes[Index].TransformBy(LandscapeInfo->GetLandscapeProxy()->LandscapeActorToWorld().ToMatrixWithScale()).ExpandBy(1.0f));
 				}
 			}
 			else
@@ -798,8 +794,7 @@ struct FXYOffsetmapAccessor
 			bool bUpdateFoliage = false;
 			for (ULandscapeComponent* Component : Components)
 			{
-				AInstancedFoliageActor* IFA = AInstancedFoliageActor::GetInstancedFoliageActorForLevel(Component->GetComponentLevel());
-				if (IFA)
+				if (AInstancedFoliageActor::HasFoliageAttached(Component))
 				{
 					bUpdateFoliage = true;
 					break;
@@ -828,11 +823,7 @@ struct FXYOffsetmapAccessor
 				for (int32 Index = 0; Index < CollisionComponents.Num(); ++Index)
 				{
 					ULandscapeHeightfieldCollisionComponent* CollisionComponent = CollisionComponents[Index];
-					AInstancedFoliageActor* IFA = AInstancedFoliageActor::GetInstancedFoliageActorForLevel(CollisionComponent->GetComponentLevel());
-					if (IFA)
-					{
-						CollisionComponent->SnapFoliageInstances(*IFA, PreUpdateLocalBoxes[Index].TransformBy(LandscapeInfo->GetLandscapeProxy()->LandscapeActorToWorld().ToMatrixWithScale()).ExpandBy(1.0f));
-					}
+					CollisionComponent->SnapFoliageInstances(PreUpdateLocalBoxes[Index].TransformBy(LandscapeInfo->GetLandscapeProxy()->LandscapeActorToWorld().ToMatrixWithScale()).ExpandBy(1.0f));
 				}
 			}
 			else
diff --git a/Engine/Source/Runtime/Foliage/Private/InstancedFoliage.cpp b/Engine/Source/Runtime/Foliage/Private/InstancedFoliage.cpp
--- a/Engine/Source/Runtime/Foliage/Private/InstancedFoliage.cpp
+++ b/Engine/Source/Runtime/Foliage/Private/InstancedFoliage.cpp
@@ -1287,33 +1287,47 @@ void AInstancedFoliageActor::DeleteInstancesForProceduralFoliageComponent(const
 
 void AInstancedFoliageActor::MoveInstancesForComponentToCurrentLevel(UActorComponent* InComponent)
 {
-	AInstancedFoliageActor* NewIFA = AInstancedFoliageActor::GetInstancedFoliageActorForCurrentLevel(InComponent->GetWorld(), true);
-	const auto SourceBaseId = InstanceBaseCache.GetInstanceBaseId(InComponent);
+	if (HasFoliageAttached(InComponent))
+	{
+		// Quit early if there are no foliage instances painted on this component
+		return;
+	}
 	
-	for (auto& MeshPair : FoliageMeshes)
+	UWorld* InWorld = InComponent->GetWorld();
+	AInstancedFoliageActor* NewIFA = AInstancedFoliageActor::GetInstancedFoliageActorForCurrentLevel(InWorld, true);
+	NewIFA->Modify();
+	
+	for (TActorIterator<AInstancedFoliageActor> It(InWorld); It; ++It)
 	{
-		FFoliageMeshInfo& MeshInfo = *MeshPair.Value;
-		UFoliageType* FoliageType = MeshPair.Key;
-
-		// Duplicate the foliage type if it's not shared
-		if (FoliageType->GetOutermost() == GetOutermost())
-		{
-			FoliageType = (UFoliageType*)StaticDuplicateObject(FoliageType, NewIFA, nullptr, RF_AllFlags & ~(RF_Standalone | RF_Public));
-		}
+		AInstancedFoliageActor* IFA = (*It);
 		
-		const auto* InstanceSet = MeshInfo.ComponentHash.Find(SourceBaseId);
-		if (InstanceSet)
+		const auto SourceBaseId = IFA->InstanceBaseCache.GetInstanceBaseId(InComponent);
+		if (SourceBaseId != FFoliageInstanceBaseCache::InvalidBaseId && IFA != NewIFA)
 		{
-			FFoliageMeshInfo* NewMeshInfo = NewIFA->FindOrAddMesh(FoliageType);
+			IFA->Modify();
 
-			// Add the foliage to the new level
-			for (int32 InstanceIndex : *InstanceSet)
+			for (auto& MeshPair : IFA->FoliageMeshes)
 			{
-				NewMeshInfo->AddInstance(NewIFA, FoliageType, MeshInfo.Instances[InstanceIndex], InComponent);
-			}
+				FFoliageMeshInfo& MeshInfo = *MeshPair.Value;
+				UFoliageType* FoliageType = MeshPair.Key;
 
-			// Remove from old level
-			MeshInfo.RemoveInstances(this, InstanceSet->Array());
+				const auto* InstanceSet = MeshInfo.ComponentHash.Find(SourceBaseId);
+				if (InstanceSet)
+				{
+					// Duplicate the foliage type if it's not shared
+					FFoliageMeshInfo* TargetMeshInfo = nullptr;
+					UFoliageType* TargetFoliageType = NewIFA->AddFoliageType(FoliageType, &TargetMeshInfo);
+
+					// Add the foliage to the new level
+					for (int32 InstanceIndex : *InstanceSet)
+					{
+						TargetMeshInfo->AddInstance(NewIFA, TargetFoliageType, MeshInfo.Instances[InstanceIndex], InComponent);
+					}
+
+					// Remove from old level
+					MeshInfo.RemoveInstances(IFA, InstanceSet->Array());
+				}
+			}
 		}
 	}
 }
@@ -1370,6 +1384,15 @@ void AInstancedFoliageActor::MoveInstancesToNewComponent(UPrimitiveComponent* In
 	}
 }
 
+void AInstancedFoliageActor::MoveInstancesToNewComponent(UWorld* InWorld, UPrimitiveComponent* InOldComponent, UPrimitiveComponent* InNewComponent)
+{
+	for (TActorIterator<AInstancedFoliageActor> It(InWorld); It; ++It)
+	{
+		AInstancedFoliageActor* IFA = (*It);
+		IFA->MoveInstancesToNewComponent(InOldComponent, InNewComponent);
+	}
+}
+
 void AInstancedFoliageActor::MoveSelectedInstancesToLevel(ULevel* InTargetLevel)
 {
 	if (InTargetLevel == GetLevel() || !HasSelectedInstances())
@@ -1774,6 +1797,21 @@ bool AInstancedFoliageActor::GetSelectionLocation(FVector& OutLocation) const
 	return false;
 }
 
+bool AInstancedFoliageActor::HasFoliageAttached(UActorComponent* InComponent)
+{
+	for (TActorIterator<AInstancedFoliageActor> It(InComponent->GetWorld()); It; ++It)
+	{
+		AInstancedFoliageActor* IFA = (*It);
+		if (IFA->InstanceBaseCache.GetInstanceBaseId(InComponent) != FFoliageInstanceBaseCache::InvalidBaseId)
+		{
+			return true;
+		}
+	}
+
+	return false;
+}
+
+
 void AInstancedFoliageActor::MapRebuild()
 {
 	// Map rebuild may have modified the BSP's ModelComponents and thrown the previous ones away.
diff --git a/Engine/Source/Runtime/Foliage/Public/FoliageInstanceBase.h b/Engine/Source/Runtime/Foliage/Public/FoliageInstanceBase.h
--- a/Engine/Source/Runtime/Foliage/Public/FoliageInstanceBase.h
+++ b/Engine/Source/Runtime/Foliage/Public/FoliageInstanceBase.h
@@ -61,7 +61,7 @@ struct FFoliageInstanceBaseCache
 	friend FArchive& operator<<(FArchive& Ar, FFoliageInstanceBaseCache& InstanceBaseCache);
 		
 	//
-	static FFoliageInstanceBaseId InvalidBaseId;
+	FOLIAGE_API static FFoliageInstanceBaseId InvalidBaseId;
 
 	// ID generator
 	// serialized
diff --git a/Engine/Source/Runtime/Foliage/Public/InstancedFoliageActor.h b/Engine/Source/Runtime/Foliage/Public/InstancedFoliageActor.h
--- a/Engine/Source/Runtime/Foliage/Public/InstancedFoliageActor.h
+++ b/Engine/Source/Runtime/Foliage/Public/InstancedFoliageActor.h
@@ -98,12 +98,13 @@ public:
 	FOLIAGE_API void MapRebuild();
 
 	// Moves instances based on the specified component to the current streaming level
-	FOLIAGE_API void MoveInstancesForComponentToCurrentLevel(UActorComponent* InComponent);
+	static FOLIAGE_API void MoveInstancesForComponentToCurrentLevel(UActorComponent* InComponent);
 
 	// Change all instances based on one component to a new component (possible in another level).
 	// The instances keep the same world locations
 	FOLIAGE_API void MoveInstancesToNewComponent(UPrimitiveComponent* InOldComponent, UPrimitiveComponent* InNewComponent);
-
+	static FOLIAGE_API void MoveInstancesToNewComponent(UWorld* InWorld, UPrimitiveComponent* InOldComponent, UPrimitiveComponent* InNewComponent);
+	
 	// Move selected instances to a foliage actor in target level
 	FOLIAGE_API void MoveSelectedInstancesToLevel(ULevel* InTargetLevel);
 	
@@ -147,6 +148,9 @@ public:
 	// Returns the location for the widget
 	FOLIAGE_API bool GetSelectionLocation(FVector& OutLocation) const;
 
+	/** Whether there any foliage instances painted on specified component */
+	static FOLIAGE_API bool HasFoliageAttached(UActorComponent* InComponent);
+
 	/* Called to notify InstancedFoliageActor that a UFoliageType has been modified */
 	void NotifyFoliageTypeChanged(UFoliageType* FoliageType);
 
diff --git a/Engine/Source/Runtime/Landscape/Classes/LandscapeHeightfieldCollisionComponent.h b/Engine/Source/Runtime/Landscape/Classes/LandscapeHeightfieldCollisionComponent.h
--- a/Engine/Source/Runtime/Landscape/Classes/LandscapeHeightfieldCollisionComponent.h
+++ b/Engine/Source/Runtime/Landscape/Classes/LandscapeHeightfieldCollisionComponent.h
@@ -219,7 +219,7 @@ class ULandscapeHeightfieldCollisionComponent : public UPrimitiveComponent
 
 #if WITH_EDITORONLY_DATA
 	// Called from editor code to manage foliage instances on landscape.
-	LANDSCAPE_API void SnapFoliageInstances(AInstancedFoliageActor &IFA, const FBox& InInstanceBox);
+	LANDSCAPE_API void SnapFoliageInstances(const FBox& InInstanceBox);
 #endif
 };
 
diff --git a/Engine/Source/Runtime/Landscape/Private/LandscapeCollision.cpp b/Engine/Source/Runtime/Landscape/Private/LandscapeCollision.cpp
--- a/Engine/Source/Runtime/Landscape/Private/LandscapeCollision.cpp
+++ b/Engine/Source/Runtime/Landscape/Private/LandscapeCollision.cpp
@@ -1070,101 +1070,111 @@ void ULandscapeHeightfieldCollisionComponent::RecreateCollision(bool bUpdateAddC
 }
 
 #if WITH_EDITORONLY_DATA
-void ULandscapeHeightfieldCollisionComponent::SnapFoliageInstances(AInstancedFoliageActor &IFA, const FBox& InInstanceBox)
+void ULandscapeHeightfieldCollisionComponent::SnapFoliageInstances(const FBox& InInstanceBox)
 {
-	for (auto& MeshPair : IFA.FoliageMeshes)
+	UWorld* ComponentWorld = GetWorld();
+	for (TActorIterator<AInstancedFoliageActor> It(ComponentWorld); It; ++It)
 	{
-		// Find the per-mesh info matching the mesh.
-		UFoliageType* Settings = MeshPair.Key;
-		FFoliageMeshInfo& MeshInfo = *MeshPair.Value;
-		const auto BaseId = IFA.InstanceBaseCache.GetInstanceBaseId(this);  // can we move this out?
-
-		const auto* InstanceSet = MeshInfo.ComponentHash.Find(BaseId);
-		if (InstanceSet)
+		AInstancedFoliageActor& IFA = *(*It);
+		const auto BaseId = IFA.InstanceBaseCache.GetInstanceBaseId(this);
+		if (BaseId == FFoliageInstanceBaseCache::InvalidBaseId)
 		{
-			float TraceExtentSize = Bounds.SphereRadius * 2.f + 10.f; // extend a little
-			FVector TraceVector = GetOwner()->GetRootComponent()->ComponentToWorld.GetUnitAxis(EAxis::Z) * TraceExtentSize;
-
-			bool bFirst = true;
-			TArray<int32> InstancesToRemove;
-			for (int32 InstanceIndex : *InstanceSet)
+			continue;
+		}
+			
+		for (auto& MeshPair : IFA.FoliageMeshes)
+		{
+			// Find the per-mesh info matching the mesh.
+			UFoliageType* Settings = MeshPair.Key;
+			FFoliageMeshInfo& MeshInfo = *MeshPair.Value;
+			
+			const auto* InstanceSet = MeshInfo.ComponentHash.Find(BaseId);
+			if (InstanceSet)
 			{
-				FFoliageInstance& Instance = MeshInfo.Instances[InstanceIndex];
-
-				// Test location should remove any Z offset
-				FVector TestLocation = FMath::Abs(Instance.ZOffset) > KINDA_SMALL_NUMBER
-					? (FVector)Instance.GetInstanceWorldTransform().TransformPosition(FVector(0, 0, -Instance.ZOffset))
-					: Instance.Location;
+				float TraceExtentSize = Bounds.SphereRadius * 2.f + 10.f; // extend a little
+				FVector TraceVector = GetOwner()->GetRootComponent()->ComponentToWorld.GetUnitAxis(EAxis::Z) * TraceExtentSize;
 
-				if (InInstanceBox.IsInside(TestLocation))
+				bool bFirst = true;
+				TArray<int32> InstancesToRemove;
+				for (int32 InstanceIndex : *InstanceSet)
 				{
-					if (bFirst)
-					{
-						bFirst = false;
-						Modify();
-					}
-
-					FVector Start = TestLocation + TraceVector;
-					FVector End = TestLocation - TraceVector;
+					FFoliageInstance& Instance = MeshInfo.Instances[InstanceIndex];
 
-					static FName TraceTag = FName(TEXT("FoliageSnapToLandscape"));
-					TArray<FHitResult> Results;
-					UWorld* World = GetWorld();
-					check(World);
-					// Editor specific landscape heightfield uses ECC_Visibility collision channel
-					World->LineTraceMultiByObjectType(Results, Start, End, FCollisionObjectQueryParams(ECollisionChannel::ECC_Visibility), FCollisionQueryParams(TraceTag, true));
+					// Test location should remove any Z offset
+					FVector TestLocation = FMath::Abs(Instance.ZOffset) > KINDA_SMALL_NUMBER
+						? (FVector)Instance.GetInstanceWorldTransform().TransformPosition(FVector(0, 0, -Instance.ZOffset))
+						: Instance.Location;
 
-					bool bFoundHit = false;
-					for (const FHitResult& Hit : Results)
+					if (InInstanceBox.IsInside(TestLocation))
 					{
-						if (Hit.Component == this)
+						if (bFirst)
 						{
-							bFoundHit = true;
-							if ((TestLocation - Hit.Location).SizeSquared() > KINDA_SMALL_NUMBER)
-							{
-								// Remove instance location from the hash. Do not need to update ComponentHash as we re-add below.
-								MeshInfo.InstanceHash->RemoveInstance(Instance.Location, InstanceIndex);
+							bFirst = false;
+							Modify();
+						}
 
-								// Update the instance editor data
-								Instance.Location = Hit.Location;
+						FVector Start = TestLocation + TraceVector;
+						FVector End = TestLocation - TraceVector;
 
-								if (Instance.Flags & FOLIAGE_AlignToNormal)
-								{
-									// Remove previous alignment and align to new normal.
-									Instance.Rotation = Instance.PreAlignRotation;
-									Instance.AlignToNormal(Hit.Normal, Settings->AlignMaxAngle);
-								}
+						static FName TraceTag = FName(TEXT("FoliageSnapToLandscape"));
+						TArray<FHitResult> Results;
+						UWorld* World = GetWorld();
+						check(World);
+						// Editor specific landscape heightfield uses ECC_Visibility collision channel
+						World->LineTraceMultiByObjectType(Results, Start, End, FCollisionObjectQueryParams(ECollisionChannel::ECC_Visibility), FCollisionQueryParams(TraceTag, true));
 
-								// Reapply the Z offset in local space
-								if (FMath::Abs(Instance.ZOffset) > KINDA_SMALL_NUMBER)
+						bool bFoundHit = false;
+						for (const FHitResult& Hit : Results)
+						{
+							if (Hit.Component == this)
+							{
+								bFoundHit = true;
+								if ((TestLocation - Hit.Location).SizeSquared() > KINDA_SMALL_NUMBER)
 								{
-									Instance.Location = Instance.GetInstanceWorldTransform().TransformPosition(FVector(0, 0, Instance.ZOffset));
+									// Remove instance location from the hash. Do not need to update ComponentHash as we re-add below.
+									MeshInfo.InstanceHash->RemoveInstance(Instance.Location, InstanceIndex);
+
+									// Update the instance editor data
+									Instance.Location = Hit.Location;
+
+									if (Instance.Flags & FOLIAGE_AlignToNormal)
+									{
+										// Remove previous alignment and align to new normal.
+										Instance.Rotation = Instance.PreAlignRotation;
+										Instance.AlignToNormal(Hit.Normal, Settings->AlignMaxAngle);
+									}
+
+									// Reapply the Z offset in local space
+									if (FMath::Abs(Instance.ZOffset) > KINDA_SMALL_NUMBER)
+									{
+										Instance.Location = Instance.GetInstanceWorldTransform().TransformPosition(FVector(0, 0, Instance.ZOffset));
+									}
+
+									// Todo: add do validation with other parameters such as max/min height etc.
+
+									check(MeshInfo.Component);
+									MeshInfo.Component->Modify();
+									MeshInfo.Component->UpdateInstanceTransform(InstanceIndex, Instance.GetInstanceWorldTransform(), true);
+									MeshInfo.Component->InvalidateLightingCache();
+
+									// Re-add the new instance location to the hash
+									MeshInfo.InstanceHash->InsertInstance(Instance.Location, InstanceIndex);
 								}
-
-								// Todo: add do validation with other parameters such as max/min height etc.
-
-								check(MeshInfo.Component);
-								MeshInfo.Component->Modify();
-								MeshInfo.Component->UpdateInstanceTransform(InstanceIndex, Instance.GetInstanceWorldTransform(), true);
-								MeshInfo.Component->InvalidateLightingCache();
-
-								// Re-add the new instance location to the hash
-								MeshInfo.InstanceHash->InsertInstance(Instance.Location, InstanceIndex);
+								break;
 							}
-							break;
 						}
-					}
 
-					if (!bFoundHit)
-					{
-						// Couldn't find new spot - remove instance
-						InstancesToRemove.Add(InstanceIndex);
+						if (!bFoundHit)
+						{
+							// Couldn't find new spot - remove instance
+							InstancesToRemove.Add(InstanceIndex);
+						}
 					}
 				}
-			}
 
-			// Remove any unused instances
-			MeshInfo.RemoveInstances(&IFA, InstancesToRemove);
+				// Remove any unused instances
+				MeshInfo.RemoveInstances(&IFA, InstancesToRemove);
+			}
 		}
 	}
 }
diff --git a/Engine/Source/Runtime/Landscape/Private/LandscapeEdit.cpp b/Engine/Source/Runtime/Landscape/Private/LandscapeEdit.cpp
--- a/Engine/Source/Runtime/Landscape/Private/LandscapeEdit.cpp
+++ b/Engine/Source/Runtime/Landscape/Private/LandscapeEdit.cpp
@@ -748,19 +748,15 @@ void ULandscapeComponent::UpdateCollisionHeightData(const FColor* HeightmapTextu
 		CollisionComp->bCanEverAffectNavigation = LandscapeProxy ? LandscapeProxy->bUsedForNavigation : false;
 	}
 
+	// Move any foliage instances if we created a new collision component.
+	if (OldCollisionComponent && OldCollisionComponent != CollisionComp)
+	{
+		AInstancedFoliageActor::MoveInstancesToNewComponent(Proxy->GetWorld(), OldCollisionComponent, CollisionComp);
+	}
+
 	if (bRebuild && CollisionProxy)
 	{
 		CollisionProxy->RegisterAllComponents();
-
-		// Move any foliage instances if we created a new collision component.
-		if (OldCollisionComponent && OldCollisionComponent != CollisionComp)
-		{
-			AInstancedFoliageActor* IFA = AInstancedFoliageActor::GetInstancedFoliageActorForCurrentLevel(OldCollisionComponent->GetWorld());
-			if (IFA)
-			{
-				IFA->MoveInstancesToNewComponent(OldCollisionComponent, CollisionComp);
-			}
-		}
 	}
 
 	// Set new collision component to pointer
