UE4 - fix grass occlusion crash

This is a safe fix and without it grass will crash instantly

diff --git a/Engine/Source/Runtime/Engine/Private/HierarchicalInstancedStaticMesh.cpp b/Engine/Source/Runtime/Engine/Private/HierarchicalInstancedStaticMesh.cpp
--- a/Engine/Source/Runtime/Engine/Private/HierarchicalInstancedStaticMesh.cpp
+++ b/Engine/Source/Runtime/Engine/Private/HierarchicalInstancedStaticMesh.cpp
@@ -178,6 +178,7 @@ class FClusterBuilder
 			Clusters.Add(FRunPair(Start, NumRange));
 			return;
 		}
+		check(NumRange >= 2);
 		SortPairs.Reset();
 		int32 BestAxis = -1;
 		float BestAxisValue = -1.0f;
@@ -1834,6 +1835,8 @@ void UHierarchicalInstancedStaticMeshComponent::BuildTreeAnyThread(
 	int32 MaxInstancesPerLeaf
 	)
 {
+	check(MaxInstancesPerLeaf > 0);
+
 	TUniquePtr<FClusterBuilder> Builder(new FClusterBuilder(InstanceTransforms, MeshBox, MaxInstancesPerLeaf));
 	Builder->Build();
 	OutOcclusionLayerNum = Builder->Result->OutOcclusionLayerNum;
diff --git a/Engine/Source/Runtime/Landscape/Private/LandscapeGrass.cpp b/Engine/Source/Runtime/Landscape/Private/LandscapeGrass.cpp
--- a/Engine/Source/Runtime/Landscape/Private/LandscapeGrass.cpp
+++ b/Engine/Source/Runtime/Landscape/Private/LandscapeGrass.cpp
@@ -986,6 +986,7 @@ struct FAsyncGrassBuilder : public FGrassBuilderBase
 
 		MeshBox = GrassVariety.GrassMesh->GetBounds().GetBox();
 		DesiredInstancesPerLeaf = HierarchicalInstancedStaticMeshComponent->DesiredInstancesPerLeaf();
+		check(DesiredInstancesPerLeaf > 0);
 
 		TotalInstances = 0;
 		HaltonBaseIndex = InHaltonBaseIndex;
@@ -1168,7 +1169,7 @@ struct FAsyncGrassBuilder : public FGrassBuilderBase
 		{
 			TArray<int32> SortedInstances;
 			TArray<int32> InstanceReorderTable;
-			UHierarchicalInstancedStaticMeshComponent::BuildTreeAnyThread(InstanceTransforms, MeshBox, ClusterTree, SortedInstances, InstanceReorderTable, DesiredInstancesPerLeaf, OutOcclusionLayerNum);
+			UHierarchicalInstancedStaticMeshComponent::BuildTreeAnyThread(InstanceTransforms, MeshBox, ClusterTree, SortedInstances, InstanceReorderTable, OutOcclusionLayerNum, DesiredInstancesPerLeaf);
 
 			// in-place sort the instances
 			FInstanceStream SwapBuffer;
