Fix poor style of landscape code assigning a literal 0 to a TLazyObjectPtr instead of null

diff --git a/Engine/Source/Runtime/Landscape/Private/Landscape.cpp b/Engine/Source/Runtime/Landscape/Private/Landscape.cpp
--- a/Engine/Source/Runtime/Landscape/Private/Landscape.cpp
+++ b/Engine/Source/Runtime/Landscape/Private/Landscape.cpp
@@ -1673,18 +1673,18 @@ void ULandscapeInfo::UnregisterActor(ALandscapeProxy* Proxy)
 	if (Landscape)
 	{
 		check(LandscapeActor.Get() == Landscape);
-		LandscapeActor = 0;
+		LandscapeActor = nullptr;
 
 		// update proxies reference to landscape actor
 		for (auto It = Proxies.CreateConstIterator(); It; ++It)
 		{
-			(*It)->LandscapeActor = 0;
+			(*It)->LandscapeActor = nullptr;
 		}
 	}
 	else
 	{
 		Proxies.Remove(Proxy);
-		Proxy->LandscapeActor = 0;
+		Proxy->LandscapeActor = nullptr;
 	}
 
 	// remove proxy components from the XY map
diff --git a/Engine/Source/Runtime/Landscape/Private/LandscapeEdit.cpp b/Engine/Source/Runtime/Landscape/Private/LandscapeEdit.cpp
--- a/Engine/Source/Runtime/Landscape/Private/LandscapeEdit.cpp
+++ b/Engine/Source/Runtime/Landscape/Private/LandscapeEdit.cpp
@@ -3124,7 +3124,7 @@ void ALandscapeProxy::PostEditChangeProperty(FPropertyChangedEvent& PropertyChan
 			}
 			else
 			{
-				LandscapeActor = 0;
+				LandscapeActor = nullptr;
 			}
 		}
 		else if (PropertyName == FName(TEXT("LandscapeMaterial")) || PropertyName == FName(TEXT("LandscapeHoleMaterial")))
