Landscape section offset fix when upgrading world composition maps from 4.6 to 4.7

diff --git a/Engine/Source/Runtime/Landscape/Private/Landscape.cpp b/Engine/Source/Runtime/Landscape/Private/Landscape.cpp
--- a/Engine/Source/Runtime/Landscape/Private/Landscape.cpp
+++ b/Engine/Source/Runtime/Landscape/Private/Landscape.cpp
@@ -1791,6 +1791,26 @@ void ULandscapeInfo::FixupProxiesTransform()
 		return;
 	}
 
+	// Make sure section offset of all proxies is multiple of ALandscapeProxy::ComponentSizeQuads
+	for (auto It = Proxies.CreateConstIterator(); It; ++It)
+	{
+		ALandscapeProxy* Proxy = *It;
+		FIntPoint LandscapeSectionOffset = Proxy->LandscapeSectionOffset - Landscape->LandscapeSectionOffset;
+		FIntPoint LandscapeSectionOffsetRem(
+			LandscapeSectionOffset.X % Proxy->ComponentSizeQuads, 
+			LandscapeSectionOffset.Y % Proxy->ComponentSizeQuads);
+
+		if (LandscapeSectionOffsetRem.X != 0 || LandscapeSectionOffsetRem.Y != 0)
+		{
+			FIntPoint NewLandscapeSectionOffset = Proxy->LandscapeSectionOffset - LandscapeSectionOffsetRem;
+			
+			UE_LOG(LogLandscape, Warning, TEXT("Landscape section base is not multiple of component size, attempted automated fix: '%s', %d,%d vs %d,%d."),
+					*Proxy->GetFullName(), Proxy->LandscapeSectionOffset.X, Proxy->LandscapeSectionOffset.Y, NewLandscapeSectionOffset.X, NewLandscapeSectionOffset.Y);
+
+			Proxy->SetAbsoluteSectionBase(NewLandscapeSectionOffset);
+		}
+	}
+
 	FTransform LandscapeTM = Landscape->LandscapeActorToWorld();
 	// Update transformations of all linked landscape proxies
 	for (auto It = Proxies.CreateConstIterator(); It; ++It)
