Update LandscapeGUID when user assigns landscape proxy to actor.

diff --git a/Engine/Source/Runtime/Landscape/Private/LandscapeEdit.cpp b/Engine/Source/Runtime/Landscape/Private/LandscapeEdit.cpp
--- a/Engine/Source/Runtime/Landscape/Private/LandscapeEdit.cpp
+++ b/Engine/Source/Runtime/Landscape/Private/LandscapeEdit.cpp
@@ -3116,6 +3116,7 @@ void ALandscapeProxy::PostEditChangeProperty(FPropertyChangedEvent& PropertyChan
 		{
 			if (LandscapeActor && IsValidLandscapeActor(LandscapeActor.Get()))
 			{
+				LandscapeGuid = LandscapeActor->LandscapeGuid;
 				// defer LandscapeInfo setup
 				if (GIsEditor && GetWorld() && !GetWorld()->IsPlayInEditor())
 				{
