Fixed: ULandscapeMeshCollisionComponent was using wrong DDC key for collision data
UE-11447

diff --git a/Engine/Source/Runtime/Landscape/Private/LandscapeCollision.cpp b/Engine/Source/Runtime/Landscape/Private/LandscapeCollision.cpp
--- a/Engine/Source/Runtime/Landscape/Private/LandscapeCollision.cpp
+++ b/Engine/Source/Runtime/Landscape/Private/LandscapeCollision.cpp
@@ -527,7 +527,7 @@ bool ULandscapeMeshCollisionComponent::CookCollsionData(const FName& Format, boo
 		// Ensure that content was saved with physical materials before using DDC data
 		if (GetLinkerUE4Version() >= VER_UE4_LANDSCAPE_SERIALIZE_PHYSICS_MATERIALS)
 		{
-			FString DDCKey = GetHFDDCKeyString(Format, bUseDefMaterial, HeightfieldGuid);
+			FString DDCKey = GetHFDDCKeyString(Format, bUseDefMaterial, MeshGuid);
 
 			// Check if the speculatively-loaded data loaded and is what we wanted
 			if (SpeculativeDDCRequest.IsValid() && DDCKey == SpeculativeDDCRequest->GetKey())
@@ -687,7 +687,7 @@ bool ULandscapeMeshCollisionComponent::CookCollsionData(const FName& Format, boo
 
 		if (bShouldSaveCookedDataToDDC[CookedDataIndex])
 		{
-			GetDerivedDataCacheRef().Put(*GetHFDDCKeyString(Format, bUseDefMaterial, HeightfieldGuid), OutCookedData);
+			GetDerivedDataCacheRef().Put(*GetHFDDCKeyString(Format, bUseDefMaterial, MeshGuid), OutCookedData);
 			bShouldSaveCookedDataToDDC[CookedDataIndex] = false;
 		}
 	}
