Fix crashes in FLandscapeComponentSceneProxy::GetStaticBatchElementVisibility when using ForcedLOD or LODBias on landscape components during texture streaming  UE-11321

diff --git a/Engine/Source/Runtime/Landscape/Private/LandscapeRender.cpp b/Engine/Source/Runtime/Landscape/Private/LandscapeRender.cpp
--- a/Engine/Source/Runtime/Landscape/Private/LandscapeRender.cpp
+++ b/Engine/Source/Runtime/Landscape/Private/LandscapeRender.cpp
@@ -1025,7 +1025,7 @@ namespace
 void FLandscapeComponentSceneProxy::DrawStaticElements(FStaticPrimitiveDrawInterface* PDI)
 {
 	int32 FirstLOD = (ForcedLOD >= 0) ? FMath::Min<int32>(ForcedLOD, MaxLOD) : FMath::Max<int32>(LODBias, 0);
-	int32 LastLOD = (ForcedLOD >= 0) ? FirstLOD : FMath::Min<int32>(MaxLOD, MaxLOD + LODBias);
+	int32 LastLOD = (ForcedLOD >= 0) ? FirstLOD : MaxLOD;	// we always need to go to MaxLOD regardless of LODBias as we could need the lowest LODs due to streaming.
 
 	int32 NumBatches = (1 + LastLOD - FirstLOD) * (FMath::Square(NumSubsections) + 1);
 	if (SharedBuffers->GrassIndexBuffer)
@@ -1140,13 +1140,12 @@ uint64 FLandscapeComponentSceneProxy::GetStaticBatchElementVisibility(const clas
 	SCOPE_CYCLE_COUNTER(STAT_LandscapeStaticDrawLODTime);
 	if (ForcedLOD >= 0)
 	{
-		// When forcing LOD we only create one Batch Element, so not sure this loop is necessary?
-		for (int32 BatchElementIndex = GrassBatchOffset; BatchElementIndex < Batch->Elements.Num(); BatchElementIndex++)
-		{
-			BatchesToRenderMask |= (((uint64)1) << BatchElementIndex);
-			INC_DWORD_STAT(STAT_LandscapeDrawCalls);
-			INC_DWORD_STAT_BY(STAT_LandscapeTriangles, Batch->Elements[BatchElementIndex].NumPrimitives);
-		}
+		// When forcing LOD we only create one Batch Element (excluding the grass rendering batch)
+		ensure(Batch->Elements.Num() - GrassBatchOffset == 1);
+		int32 BatchElementIndex = GrassBatchOffset;
+		BatchesToRenderMask |= (((uint64)1) << BatchElementIndex);
+		INC_DWORD_STAT(STAT_LandscapeDrawCalls);
+		INC_DWORD_STAT_BY(STAT_LandscapeTriangles, Batch->Elements[BatchElementIndex].NumPrimitives);
 	}
 	else
 	{
@@ -1159,34 +1158,34 @@ uint64 FLandscapeComponentSceneProxy::GetStaticBatchElementVisibility(const clas
 		int32 CombinedLOD = -1;
 		int32 bAllSameLOD = true;
 
+		// Components with positive LODBias don't generate batch elements for unused LODs.
+		int32 LODBiasOffset = FMath::Max<int32>(LODBias, 0);
+
 		for (int32 SubY = 0; SubY < NumSubsections; SubY++)
 		{
 			for (int32 SubX = 0; SubX < NumSubsections; SubX++)
 			{
-				int32 TempLOD = CalcLODForSubsection(View, SubX, SubY, CameraLocalPos);
-
-				if (LODBias > 0)
-				{
-					TempLOD = FMath::Max<int32>(TempLOD - LODBias, 0);
-				}
-
+				int32 ThisSubsectionLOD = CalcLODForSubsection(View, SubX, SubY, CameraLocalPos);
 				// check if all LODs are the same.
-				if (TempLOD != CombinedLOD && CombinedLOD != -1)
+				if (ThisSubsectionLOD != CombinedLOD && CombinedLOD != -1)
 				{
 					bAllSameLOD = false;
 				}
-				CombinedLOD = TempLOD;
-				CalculatedLods[SubX][SubY] = TempLOD;
+				CombinedLOD = ThisSubsectionLOD;
+				CalculatedLods[SubX][SubY] = ThisSubsectionLOD;
 			}
 		}
 
 		if (bAllSameLOD && NumSubsections > 1 && !GLandscapeDebugOptions.bDisableCombine)
 		{
 			// choose the combined batch element
-			int32 BatchElementIndex = GrassBatchOffset + (CombinedLOD + 1)*BatchesPerLOD - 1;
-			BatchesToRenderMask |= (((uint64)1) << BatchElementIndex);
-			INC_DWORD_STAT(STAT_LandscapeDrawCalls);
-			INC_DWORD_STAT_BY(STAT_LandscapeTriangles, Batch->Elements[BatchElementIndex].NumPrimitives);
+			int32 BatchElementIndex = GrassBatchOffset + (CombinedLOD - LODBiasOffset + 1)*BatchesPerLOD - 1;
+			if (ensure(Batch->Elements.IsValidIndex(BatchElementIndex)))
+			{
+				BatchesToRenderMask |= (((uint64)1) << BatchElementIndex);
+				INC_DWORD_STAT(STAT_LandscapeDrawCalls);
+				INC_DWORD_STAT_BY(STAT_LandscapeTriangles, Batch->Elements[BatchElementIndex].NumPrimitives);
+			}
 		}
 		else
 		{
@@ -1194,14 +1193,16 @@ uint64 FLandscapeComponentSceneProxy::GetStaticBatchElementVisibility(const clas
 			{
 				for (int32 SubX = 0; SubX < NumSubsections; SubX++)
 				{
-					int32 BatchElementIndex = GrassBatchOffset + CalculatedLods[SubX][SubY] * BatchesPerLOD + SubY*NumSubsections + SubX;
-					BatchesToRenderMask |= (((uint64)1) << BatchElementIndex);
-					INC_DWORD_STAT(STAT_LandscapeDrawCalls);
-					INC_DWORD_STAT_BY(STAT_LandscapeTriangles, Batch->Elements[BatchElementIndex].NumPrimitives);
+					int32 BatchElementIndex = GrassBatchOffset + (CalculatedLods[SubX][SubY] - LODBiasOffset) * BatchesPerLOD + SubY*NumSubsections + SubX;
+					if (ensure(Batch->Elements.IsValidIndex(BatchElementIndex)))
+					{
+						BatchesToRenderMask |= (((uint64)1) << BatchElementIndex);
+						INC_DWORD_STAT(STAT_LandscapeDrawCalls);
+						INC_DWORD_STAT_BY(STAT_LandscapeTriangles, Batch->Elements[BatchElementIndex].NumPrimitives);
+					}
 				}
 			}
 		}
-
 	}
 
 	INC_DWORD_STAT(STAT_LandscapeComponents);
@@ -1293,9 +1294,9 @@ float FLandscapeComponentSceneProxy::CalcDesiredLOD(const class FSceneView& View
 				break;
 			}
 		}
-	}
 
-	fLOD = FMath::Clamp<float>(fLOD, SubsectionLODBias, FMath::Min<int32>(MaxLOD, MaxLOD + SubsectionLODBias));
+		fLOD = FMath::Clamp<float>(fLOD, SubsectionLODBias, FMath::Min<int32>(MaxLOD, MaxLOD + SubsectionLODBias));
+	}
 
 	// ultimately due to texture streaming we sometimes need to go past MaxLOD
 	fLOD = FMath::Max<float>(fLOD, MinStreamedLOD);
