Fix ensure/missing landscape when using the viewport Override LOD greater than the landscape's maximum LOD.

diff --git a/Engine/Source/Runtime/Landscape/Private/LandscapeRender.cpp b/Engine/Source/Runtime/Landscape/Private/LandscapeRender.cpp
--- a/Engine/Source/Runtime/Landscape/Private/LandscapeRender.cpp
+++ b/Engine/Source/Runtime/Landscape/Private/LandscapeRender.cpp
@@ -1215,7 +1215,7 @@ float FLandscapeComponentSceneProxy::CalcDesiredLOD(const class FSceneView& View
 #if WITH_EDITOR
 	if (View.Family->LandscapeLODOverride >= 0)
 	{
-		return View.Family->LandscapeLODOverride;
+		return FMath::Min<int32>(FMath::CeilLogTwo(SubsectionSizeVerts)-1, View.Family->LandscapeLODOverride);
 	}
 #endif
 
