Expose TranslucencySortPriority for landscape spline meshes

diff --git a/Engine/Source/Runtime/Landscape/Classes/LandscapeSplineControlPoint.h b/Engine/Source/Runtime/Landscape/Classes/LandscapeSplineControlPoint.h
--- a/Engine/Source/Runtime/Landscape/Classes/LandscapeSplineControlPoint.h
+++ b/Engine/Source/Runtime/Landscape/Classes/LandscapeSplineControlPoint.h
@@ -81,6 +81,16 @@ class ULandscapeSplineControlPoint : public UObject
 	UPROPERTY(EditAnywhere, AdvancedDisplay, Category=Mesh, meta = (DisplayName = "Max Draw Distance"))
 	float LDMaxDrawDistance;
 
+	/**
+	 * Translucent objects with a lower sort priority draw behind objects with a higher priority.
+	 * Translucent objects with the same priority are rendered from back-to-front based on their bounds origin.
+	 *
+	 * Ignored if the object is not translucent.  The default priority is zero.
+	 * Warning: This should never be set to a non-default value unless you know what you are doing, as it will prevent the renderer from sorting correctly.
+	 */
+	UPROPERTY(EditAnywhere, AdvancedDisplay, Category=Mesh)
+	int32 TranslucencySortPriority;
+
 	/**
 	 * Name of blend layer to paint when applying spline to landscape
 	 * If "none", no layer is painted
diff --git a/Engine/Source/Runtime/Landscape/Classes/LandscapeSplineSegment.h b/Engine/Source/Runtime/Landscape/Classes/LandscapeSplineSegment.h
--- a/Engine/Source/Runtime/Landscape/Classes/LandscapeSplineSegment.h
+++ b/Engine/Source/Runtime/Landscape/Classes/LandscapeSplineSegment.h
@@ -169,6 +169,16 @@ class ULandscapeSplineSegment : public UObject
 	UPROPERTY(EditAnywhere, AdvancedDisplay, Category = LandscapeSplineMeshes, meta = (DisplayName = "Max Draw Distance"))
 	float LDMaxDrawDistance;
 
+	/**
+	 * Translucent objects with a lower sort priority draw behind objects with a higher priority.
+	 * Translucent objects with the same priority are rendered from back-to-front based on their bounds origin.
+	 *
+	 * Ignored if the object is not translucent.  The default priority is zero.
+	 * Warning: This should never be set to a non-default value unless you know what you are doing, as it will prevent the renderer from sorting correctly.
+	 */
+	UPROPERTY(EditAnywhere, AdvancedDisplay, Category=LandscapeSplineMeshes)
+	int32 TranslucencySortPriority;
+
 	/** If the spline is above the terrain, whether to raise the terrain up to the level of the spline when applying it to the landscape. */
 	UPROPERTY(EditAnywhere, Category=LandscapeSplineSegment)
 	uint32 bRaiseTerrain:1;
diff --git a/Engine/Source/Runtime/Landscape/Private/LandscapeSplines.cpp b/Engine/Source/Runtime/Landscape/Private/LandscapeSplines.cpp
--- a/Engine/Source/Runtime/Landscape/Private/LandscapeSplines.cpp
+++ b/Engine/Source/Runtime/Landscape/Private/LandscapeSplines.cpp
@@ -527,6 +527,7 @@ ULandscapeSplineControlPoint::ULandscapeSplineControlPoint(const FObjectInitiali
 	MeshScale = FVector(1);
 
 	LDMaxDrawDistance = 0;
+	TranslucencySortPriority = 0;
 
 	LayerName = NAME_None;
 	bRaiseTerrain = true;
@@ -831,6 +832,13 @@ void ULandscapeSplineControlPoint::UpdateSplinePoints(bool bUpdateCollision, boo
 			}
 		}
 
+		if (MeshComponent->TranslucencySortPriority != TranslucencySortPriority)
+		{
+			MeshComponent->Modify();
+			MeshComponent->TranslucencySortPriority = TranslucencySortPriority;
+			MeshComponent->MarkRenderStateDirty();
+		}
+
 		if (MeshComponent->LDMaxDrawDistance != LDMaxDrawDistance)
 		{
 			MeshComponent->Modify();
@@ -1006,6 +1014,7 @@ ULandscapeSplineSegment::ULandscapeSplineSegment(const FObjectInitializer& Objec
 	// SplineMesh properties
 	SplineMeshes.Empty();
 	LDMaxDrawDistance = 0;
+	TranslucencySortPriority = 0;
 	bEnableCollision = true;
 	bCastShadow = true;
 
@@ -1498,6 +1507,7 @@ void ULandscapeSplineSegment::UpdateSplinePoints(bool bUpdateCollision)
 				MeshComponent->LDMaxDrawDistance = LDMaxDrawDistance;
 				MeshComponent->CachedMaxDrawDistance = 0;
 			}
+			MeshComponent->TranslucencySortPriority = TranslucencySortPriority;
 
 			MeshComponent->SetCastShadow(bCastShadow);
 			MeshComponent->InvalidateLightingCache();
