Fix crash launching PIE when sublevel contains landscape in hidden level

diff --git a/Engine/Source/Runtime/Landscape/Private/LandscapeCollision.cpp b/Engine/Source/Runtime/Landscape/Private/LandscapeCollision.cpp
--- a/Engine/Source/Runtime/Landscape/Private/LandscapeCollision.cpp
+++ b/Engine/Source/Runtime/Landscape/Private/LandscapeCollision.cpp
@@ -1233,11 +1233,11 @@ void ULandscapeHeightfieldCollisionComponent::Serialize(FArchive& Ar)
 		else
 		{
 #if WITH_EDITORONLY_DATA			
-			// do not need this data in PIE, so don't bother duplicating it
-			if (!(Ar.GetPortFlags() & PPF_DuplicateForPIE)) 
+			// For PIE, we won't need the source height data if we already have a shared reference to the heightfield
+			if (!(Ar.GetPortFlags() & PPF_DuplicateForPIE) || !HeightfieldGuid.IsValid() || GSharedMeshRefs.FindRef(HeightfieldGuid) == nullptr)
 			{
-			CollisionHeightData.Serialize(Ar, this);
-			DominantLayerData.Serialize(Ar, this);
+				CollisionHeightData.Serialize(Ar, this);
+				DominantLayerData.Serialize(Ar, this);
 			}
 #endif//WITH_EDITORONLY_DATA
 		}
