Fix spelling error in CookCollisionData function name

diff --git a/Engine/Source/Runtime/Landscape/Classes/LandscapeHeightfieldCollisionComponent.h b/Engine/Source/Runtime/Landscape/Classes/LandscapeHeightfieldCollisionComponent.h
--- a/Engine/Source/Runtime/Landscape/Classes/LandscapeHeightfieldCollisionComponent.h
+++ b/Engine/Source/Runtime/Landscape/Classes/LandscapeHeightfieldCollisionComponent.h
@@ -195,7 +195,7 @@ class ULandscapeHeightfieldCollisionComponent : public UPrimitiveComponent
 	/** 
 	 * Cooks raw height data into collision object binary stream
 	 */
-	virtual bool CookCollsionData(const FName& Format, bool bUseOnlyDefMaterial, bool bCheckDDC, TArray<uint8>& OutCookedData, TArray<UPhysicalMaterial*>& OutMaterials) const;
+	virtual bool CookCollisionData(const FName& Format, bool bUseOnlyDefMaterial, bool bCheckDDC, TArray<uint8>& OutCookedData, TArray<UPhysicalMaterial*>& OutMaterials) const;
 
 	/** Modify a sub-region of the PhysX heightfield. Note that this does not update the physical material */
 	void UpdateHeightfieldRegion(int32 ComponentX1, int32 ComponentY1, int32 ComponentX2, int32 ComponentY2);
diff --git a/Engine/Source/Runtime/Landscape/Classes/LandscapeMeshCollisionComponent.h b/Engine/Source/Runtime/Landscape/Classes/LandscapeMeshCollisionComponent.h
--- a/Engine/Source/Runtime/Landscape/Classes/LandscapeMeshCollisionComponent.h
+++ b/Engine/Source/Runtime/Landscape/Classes/LandscapeMeshCollisionComponent.h
@@ -88,7 +88,7 @@ class ULandscapeMeshCollisionComponent : public ULandscapeHeightfieldCollisionCo
 	virtual void ExportCustomProperties(FOutputDevice& Out, uint32 Indent) override;
 	virtual void ImportCustomProperties(const TCHAR* SourceText, FFeedbackContext* Warn) override;
 
-	virtual bool CookCollsionData(const FName& Format, bool bUseOnlyDefMaterial, bool bCheckDDC, TArray<uint8>& OutCookedData, TArray<UPhysicalMaterial*>& OutMaterials) const override;
+	virtual bool CookCollisionData(const FName& Format, bool bUseOnlyDefMaterial, bool bCheckDDC, TArray<uint8>& OutCookedData, TArray<UPhysicalMaterial*>& OutMaterials) const override;
 #endif
 	// End UObject Interface.
 
diff --git a/Engine/Source/Runtime/Landscape/Private/LandscapeCollision.cpp b/Engine/Source/Runtime/Landscape/Private/LandscapeCollision.cpp
--- a/Engine/Source/Runtime/Landscape/Private/LandscapeCollision.cpp
+++ b/Engine/Source/Runtime/Landscape/Private/LandscapeCollision.cpp
@@ -294,7 +294,7 @@ void ULandscapeHeightfieldCollisionComponent::CreateCollisionObject()
 #if WITH_EDITOR
 			// Prepare heightfield data
 			static FName PhysicsFormatName(FPlatformProperties::GetPhysicsFormat());
-			CookCollsionData(PhysicsFormatName, false, bCheckDDC, CookedCollisionData, CookedPhysicalMaterials);
+			CookCollisionData(PhysicsFormatName, false, bCheckDDC, CookedCollisionData, CookedPhysicalMaterials);
 
 			// The World will clean up any speculatively-loaded data we didn't end up using.
 			SpeculativeDDCRequest.Reset();
@@ -328,7 +328,7 @@ void ULandscapeHeightfieldCollisionComponent::CreateCollisionObject()
 				if (!World->IsGameWorld())
 				{
 					TArray<UPhysicalMaterial*> CookedMaterialsEd;
-					if (CookCollsionData(PhysicsFormatName, true, bCheckDDC, CookedCollisionDataEd, CookedMaterialsEd))
+					if (CookCollisionData(PhysicsFormatName, true, bCheckDDC, CookedCollisionDataEd, CookedMaterialsEd))
 					{
 						FPhysXInputStream HeightFieldStream(CookedCollisionDataEd.GetData(), CookedCollisionDataEd.Num());
 						HeightfieldRef->RBHeightfieldEd = GPhysXSDK->createHeightField(HeightFieldStream);
@@ -362,7 +362,7 @@ void ULandscapeHeightfieldCollisionComponent::SpeculativelyLoadAsyncDDCCollsionD
 #endif
 }
 
-bool ULandscapeHeightfieldCollisionComponent::CookCollsionData(const FName& Format, bool bUseDefMaterial, bool bCheckDDC, TArray<uint8>& OutCookedData, TArray<UPhysicalMaterial*>& OutMaterials) const
+bool ULandscapeHeightfieldCollisionComponent::CookCollisionData(const FName& Format, bool bUseDefMaterial, bool bCheckDDC, TArray<uint8>& OutCookedData, TArray<UPhysicalMaterial*>& OutMaterials) const
 {
 #if WITH_PHYSX
 	// we have 2 versions of collision objects
@@ -515,7 +515,7 @@ bool ULandscapeHeightfieldCollisionComponent::CookCollsionData(const FName& Form
 	return false;
 }
 
-bool ULandscapeMeshCollisionComponent::CookCollsionData(const FName& Format, bool bUseDefMaterial, bool bCheckDDC, TArray<uint8>& OutCookedData, TArray<UPhysicalMaterial*>& OutMaterials) const
+bool ULandscapeMeshCollisionComponent::CookCollisionData(const FName& Format, bool bUseDefMaterial, bool bCheckDDC, TArray<uint8>& OutCookedData, TArray<UPhysicalMaterial*>& OutMaterials) const
 {
 #if WITH_PHYSX
 	
@@ -733,7 +733,7 @@ void ULandscapeMeshCollisionComponent::CreateCollisionObject()
 #if WITH_EDITOR
 			// Create cooked physics data
 			static FName PhysicsFormatName(FPlatformProperties::GetPhysicsFormat());
-			CookCollsionData(PhysicsFormatName, false, bCheckDDC, CookedCollisionData, CookedPhysicalMaterials);
+			CookCollisionData(PhysicsFormatName, false, bCheckDDC, CookedCollisionData, CookedPhysicalMaterials);
 #endif //WITH_EDITOR
 
 			if (CookedCollisionData.Num())
@@ -761,7 +761,7 @@ void ULandscapeMeshCollisionComponent::CreateCollisionObject()
 				if (!GetWorld()->IsGameWorld())
 				{
 					TArray<UPhysicalMaterial*> CookedMaterialsEd;
-					if (CookCollsionData(PhysicsFormatName, true, bCheckDDC, CookedCollisionDataEd, CookedMaterialsEd))
+					if (CookCollisionData(PhysicsFormatName, true, bCheckDDC, CookedCollisionDataEd, CookedMaterialsEd))
 					{
 						FPhysXInputStream MeshStream(CookedCollisionDataEd.GetData(), CookedCollisionDataEd.Num());
 						MeshRef->RBTriangleMeshEd = GPhysXSDK->createTriangleMesh(MeshStream);
@@ -1200,7 +1200,7 @@ void ULandscapeHeightfieldCollisionComponent::Serialize(FArchive& Ar)
 		if (Ar.IsCooking() && !HasAnyFlags(RF_ClassDefaultObject))
 		{
 			FName Format = Ar.CookingTarget()->GetPhysicsFormat(nullptr);
-			CookCollsionData(Format, false, true, CookedCollisionData, CookedPhysicalMaterials);
+			CookCollisionData(Format, false, true, CookedCollisionData, CookedPhysicalMaterials);
 			GetDerivedDataCacheRef().Put(*GetHFDDCKeyString(Format, false, HeightfieldGuid), CookedCollisionData);
 		}
 	}
