Fix ensure selecting landscape view LOD override exceeding landscape's MaxLOD setting.

diff --git a/Engine/Source/Runtime/Landscape/Private/LandscapeRender.cpp b/Engine/Source/Runtime/Landscape/Private/LandscapeRender.cpp
--- a/Engine/Source/Runtime/Landscape/Private/LandscapeRender.cpp
+++ b/Engine/Source/Runtime/Landscape/Private/LandscapeRender.cpp
@@ -535,6 +535,8 @@ FLandscapeComponentSceneProxy::FLandscapeComponentSceneProxy(ULandscapeComponent
 	: FPrimitiveSceneProxy(InComponent)
 	, FLandscapeNeighborInfo(InComponent->GetWorld(), InComponent->GetLandscapeProxy()->GetLandscapeGuid(), InComponent->GetSectionBase() / InComponent->ComponentSizeQuads, InComponent->HeightmapTexture, InComponent->ForcedLOD, InComponent->LODBias)
 	, MaxLOD(FMath::CeilLogTwo(InComponent->SubsectionSizeQuads + 1) - 1)
+	, FirstLOD(0)
+	, LastLOD(FMath::CeilLogTwo(InComponent->SubsectionSizeQuads + 1) - 1)
 	, NumSubsections(InComponent->NumSubsections)
 	, SubsectionSizeQuads(InComponent->SubsectionSizeQuads)
 	, SubsectionSizeVerts(InComponent->SubsectionSizeQuads + 1)
@@ -586,6 +588,9 @@ FLandscapeComponentSceneProxy::FLandscapeComponentSceneProxy(ULandscapeComponent
 		MaxLOD = FMath::Min<int8>(MaxLOD, InComponent->GetLandscapeProxy()->MaxLODLevel);
 	}
 
+	FirstLOD = (ForcedLOD >= 0) ? FMath::Min<int32>(ForcedLOD, MaxLOD) : FMath::Max<int32>(LODBias, 0);
+	LastLOD = (ForcedLOD >= 0) ? FirstLOD : MaxLOD;	// we always need to go to MaxLOD regardless of LODBias as we could need the lowest LODs due to streaming.
+
 	float LODDistanceFactor;
 	switch (LODFalloff)
 	{
@@ -1024,9 +1029,6 @@ namespace
 
 void FLandscapeComponentSceneProxy::DrawStaticElements(FStaticPrimitiveDrawInterface* PDI)
 {
-	int32 FirstLOD = (ForcedLOD >= 0) ? FMath::Min<int32>(ForcedLOD, MaxLOD) : FMath::Max<int32>(LODBias, 0);
-	int32 LastLOD = (ForcedLOD >= 0) ? FirstLOD : MaxLOD;	// we always need to go to MaxLOD regardless of LODBias as we could need the lowest LODs due to streaming.
-
 	int32 NumBatches = (1 + LastLOD - FirstLOD) * (FMath::Square(NumSubsections) + 1);
 	if (SharedBuffers->GrassIndexBuffer)
 	{
@@ -1215,7 +1217,7 @@ float FLandscapeComponentSceneProxy::CalcDesiredLOD(const class FSceneView& View
 #if WITH_EDITOR
 	if (View.Family->LandscapeLODOverride >= 0)
 	{
-		return FMath::Min<int32>(FMath::CeilLogTwo(SubsectionSizeVerts)-1, View.Family->LandscapeLODOverride);
+		return FMath::Clamp<int32>(View.Family->LandscapeLODOverride, FirstLOD, LastLOD);
 	}
 #endif
 
diff --git a/Engine/Source/Runtime/Landscape/Public/LandscapeRender.h b/Engine/Source/Runtime/Landscape/Public/LandscapeRender.h
--- a/Engine/Source/Runtime/Landscape/Public/LandscapeRender.h
+++ b/Engine/Source/Runtime/Landscape/Public/LandscapeRender.h
@@ -486,7 +486,10 @@ class FLandscapeComponentSceneProxy : public FPrimitiveSceneProxy, public FLands
 	};
 
 protected:
-	int8						MaxLOD;
+	int8						MaxLOD;		// Maximum LOD level, user override possible
+	int32						FirstLOD;	// First LOD we have batch elements for
+	int32						LastLOD;	// Last LOD we have batch elements for
+
 	/** 
 	 * Number of subsections within the component in each dimension, this can be 1 or 2.  
 	 * Subsections exist to improve the speed at which LOD transitions can take place over distance. 
