New feature: Landscape splines can store their meshes in landscape streaming levels

diff --git a/Engine/Config/BaseEngine.ini b/Engine/Config/BaseEngine.ini
--- a/Engine/Config/BaseEngine.ini
+++ b/Engine/Config/BaseEngine.ini
@@ -1175,6 +1175,9 @@ bUseStreamingPause=false
 ;Exposing SplineMeshComponents to Blueprint
 +TaggedPropertyRedirects=(ClassName="SplineMeshComponent",OldPropertyName="SplineXDir",NewPropertyName="SplineUpDir")
 
++TaggedPropertyRedirects=(ClassName="LandscapeSplineControlPoint",OldPropertyName="MeshComponent",NewPropertyName="LocalMeshComponent")
++TaggedPropertyRedirects=(ClassName="LandscapeSplineSegment",OldPropertyName="MeshComponents",NewPropertyName="LocalMeshComponents")
+
 ;Foliage refactor
 +ActiveClassRedirects=(OldClassName="InstancedFoliageSettings", NewClassName="/Script/Foliage.FoliageType_InstancedStaticMesh")
 +ActiveClassRedirects=(OldClassName="FoliageType", NewClassName="/Script/Foliage.FoliageType")
diff --git a/Engine/Source/Editor/LandscapeEditor/Private/LandscapeEdModeSplineTools.cpp b/Engine/Source/Editor/LandscapeEditor/Private/LandscapeEdModeSplineTools.cpp
--- a/Engine/Source/Editor/LandscapeEditor/Private/LandscapeEdModeSplineTools.cpp
+++ b/Engine/Source/Editor/LandscapeEditor/Private/LandscapeEdModeSplineTools.cpp
@@ -276,31 +276,32 @@ public:
 		NewSegment->AutoFlipTangents();
 
 		// set up other segment options
+		ULandscapeSplineSegment* CopyFromSegment = nullptr;
 		if (Start->ConnectedSegments.Num() > 0)
 		{
-			NewSegment->LayerName = Start->ConnectedSegments[0].Segment->LayerName;
-			NewSegment->SplineMeshes = Start->ConnectedSegments[0].Segment->SplineMeshes;
-			NewSegment->LDMaxDrawDistance = Start->ConnectedSegments[0].Segment->LDMaxDrawDistance;
-			NewSegment->bRaiseTerrain = Start->ConnectedSegments[0].Segment->bRaiseTerrain;
-			NewSegment->bLowerTerrain = Start->ConnectedSegments[0].Segment->bLowerTerrain;
-			NewSegment->bEnableCollision = Start->ConnectedSegments[0].Segment->bEnableCollision;
-			NewSegment->bCastShadow = Start->ConnectedSegments[0].Segment->bCastShadow;
+			CopyFromSegment = Start->ConnectedSegments[0].Segment;
 		}
 		else if (End->ConnectedSegments.Num() > 0)
 		{
-			NewSegment->LayerName = End->ConnectedSegments[0].Segment->LayerName;
-			NewSegment->SplineMeshes = End->ConnectedSegments[0].Segment->SplineMeshes;
-			NewSegment->LDMaxDrawDistance = End->ConnectedSegments[0].Segment->LDMaxDrawDistance;
-			NewSegment->bRaiseTerrain = End->ConnectedSegments[0].Segment->bRaiseTerrain;
-			NewSegment->bLowerTerrain = End->ConnectedSegments[0].Segment->bLowerTerrain;
-			NewSegment->bEnableCollision = End->ConnectedSegments[0].Segment->bEnableCollision;
-			NewSegment->bCastShadow = End->ConnectedSegments[0].Segment->bCastShadow;
+			CopyFromSegment = End->ConnectedSegments[0].Segment;
 		}
 		else
 		{
 			// Use defaults
 		}
 
+		if (CopyFromSegment != nullptr)
+		{
+			NewSegment->LayerName         = CopyFromSegment->LayerName;
+			NewSegment->SplineMeshes      = CopyFromSegment->SplineMeshes;
+			NewSegment->LDMaxDrawDistance = CopyFromSegment->LDMaxDrawDistance;
+			NewSegment->bRaiseTerrain     = CopyFromSegment->bRaiseTerrain;
+			NewSegment->bLowerTerrain     = CopyFromSegment->bLowerTerrain;
+			NewSegment->bPlaceSplineMeshesInStreamingLevels = CopyFromSegment->bPlaceSplineMeshesInStreamingLevels;
+			NewSegment->bEnableCollision  = CopyFromSegment->bEnableCollision;
+			NewSegment->bCastShadow       = CopyFromSegment->bCastShadow;
+		}
+
 		Start->ConnectedSegments.Add(FLandscapeSplineConnection(NewSegment, 0));
 		End->ConnectedSegments.Add(FLandscapeSplineConnection(NewSegment, 1));
 
@@ -341,17 +342,18 @@ public:
 		if (SelectedSplineControlPoints.Num() > 0)
 		{
 			ULandscapeSplineControlPoint* FirstPoint = *SelectedSplineControlPoints.CreateConstIterator();
-			NewControlPoint->Rotation = (NewControlPoint->Location - FirstPoint->Location).Rotation();
-			NewControlPoint->Width = FirstPoint->Width;
+			NewControlPoint->Rotation    = (NewControlPoint->Location - FirstPoint->Location).Rotation();
+			NewControlPoint->Width       = FirstPoint->Width;
 			NewControlPoint->SideFalloff = FirstPoint->SideFalloff;
-			NewControlPoint->EndFalloff = FirstPoint->EndFalloff;
+			NewControlPoint->EndFalloff  = FirstPoint->EndFalloff;
 
 			if (bCopyMeshToNewControlPoint)
 			{
-				NewControlPoint->Mesh = FirstPoint->Mesh;
-				NewControlPoint->MeshScale = FirstPoint->MeshScale;
+				NewControlPoint->Mesh             = FirstPoint->Mesh;
+				NewControlPoint->MeshScale        = FirstPoint->MeshScale;
+				NewControlPoint->bPlaceSplineMeshesInStreamingLevels = FirstPoint->bPlaceSplineMeshesInStreamingLevels;
 				NewControlPoint->bEnableCollision = FirstPoint->bEnableCollision;
-				NewControlPoint->bCastShadow = FirstPoint->bCastShadow;
+				NewControlPoint->bCastShadow      = FirstPoint->bCastShadow;
 			}
 
 			for (ULandscapeSplineControlPoint* ControlPoint : SelectedSplineControlPoints)
@@ -595,6 +597,197 @@ public:
 		}
 	}
 
+	void MoveSelectedToLevel()
+	{
+		TSet<ALandscapeProxy*> FromProxies;
+		ALandscapeProxy* ToLandscape = nullptr;
+
+		for (ULandscapeSplineControlPoint* ControlPoint : SelectedSplineControlPoints)
+		{
+			ULandscapeSplinesComponent* LandscapeSplinesComp = ControlPoint->GetOuterULandscapeSplinesComponent();
+			ALandscapeProxy* FromProxy = LandscapeSplinesComp ? Cast<ALandscapeProxy>(LandscapeSplinesComp->GetOuter()) : nullptr;
+			if (FromProxy)
+			{
+				if (!ToLandscape)
+				{
+					ULandscapeInfo* LandscapeInfo = FromProxy->GetLandscapeInfo(false);
+					check(LandscapeInfo);
+					ToLandscape = LandscapeInfo->GetCurrentLevelLandscapeProxy(true);
+					if (!ToLandscape)
+					{
+						// No Landscape Proxy, don't support for creating only for Spline now
+						return;
+					}
+				}
+
+				if (ToLandscape != FromProxy)
+				{
+					ToLandscape->Modify();
+					if (ToLandscape->SplineComponent == nullptr)
+					{
+						CreateSplineComponent(ToLandscape, FromProxy->SplineComponent->RelativeScale3D);
+					}
+					ToLandscape->SplineComponent->Modify();
+
+					const FTransform OldToNewTransform =
+						FromProxy->SplineComponent->ComponentToWorld.GetRelativeTransform(ToLandscape->SplineComponent->ComponentToWorld);
+
+					if (FromProxies.Find(FromProxy) == nullptr)
+					{
+						FromProxies.Add(FromProxy);
+						FromProxy->Modify();
+						FromProxy->SplineComponent->Modify();
+						FromProxy->SplineComponent->MarkRenderStateDirty();
+					}
+
+					// Handle control point mesh
+					if (ControlPoint->bPlaceSplineMeshesInStreamingLevels)
+					{
+						// Mark previously local component as Foreign
+						if (ControlPoint->LocalMeshComponent)
+						{
+							auto* MeshComponent = ControlPoint->LocalMeshComponent;
+							verifySlow(FromProxy->SplineComponent->MeshComponentLocalOwnersMap.Remove(MeshComponent) == 1);
+							FromProxy->SplineComponent->AddForeignMeshComponent(ControlPoint, MeshComponent);
+						}
+						ControlPoint->LocalMeshComponent = nullptr;
+
+						// Mark previously foreign component as local
+						auto* MeshComponent = ToLandscape->SplineComponent->GetForeignMeshComponent(ControlPoint);
+						if (MeshComponent)
+						{
+							ToLandscape->SplineComponent->RemoveForeignMeshComponent(ControlPoint, MeshComponent);
+							ToLandscape->SplineComponent->MeshComponentLocalOwnersMap.Add(MeshComponent, ControlPoint);
+						}
+						ControlPoint->LocalMeshComponent = MeshComponent;
+					}
+					else
+					{
+						// non-streaming case
+						if (ControlPoint->LocalMeshComponent)
+						{
+							UControlPointMeshComponent* MeshComponent = ControlPoint->LocalMeshComponent;
+							MeshComponent->Modify();
+							MeshComponent->UnregisterComponent();
+							MeshComponent->DetachFromParent(true);
+							MeshComponent->InvalidateLightingCache();
+							MeshComponent->Rename(nullptr, ToLandscape);
+							MeshComponent->AttachTo(ToLandscape->SplineComponent, NAME_None, EAttachLocation::KeepWorldPosition);
+
+							verifySlow(FromProxy->SplineComponent->MeshComponentLocalOwnersMap.Remove(MeshComponent) == 1);
+							ToLandscape->SplineComponent->MeshComponentLocalOwnersMap.Add(MeshComponent, ControlPoint);
+						}
+					}
+
+					// Move control point to new level
+					FromProxy->SplineComponent->ControlPoints.Remove(ControlPoint);
+					ControlPoint->Rename(nullptr, ToLandscape->SplineComponent);
+					ToLandscape->SplineComponent->ControlPoints.Add(ControlPoint);
+
+					ControlPoint->Location = OldToNewTransform.TransformPosition(ControlPoint->Location);
+
+					ControlPoint->UpdateSplinePoints(true, false);
+				}
+			}
+		}
+
+		for (ULandscapeSplineSegment* Segment : SelectedSplineSegments)
+		{
+			ULandscapeSplinesComponent* LandscapeSplinesComp = Segment->GetOuterULandscapeSplinesComponent();
+			ALandscapeProxy* FromProxy = LandscapeSplinesComp ? Cast<ALandscapeProxy>(LandscapeSplinesComp->GetOuter()) : nullptr;
+			if (FromProxy)
+			{
+				if (!ToLandscape)
+				{
+					ULandscapeInfo* LandscapeInfo = FromProxy->GetLandscapeInfo(false);
+					check(LandscapeInfo);
+					ToLandscape = LandscapeInfo->GetCurrentLevelLandscapeProxy(true);
+					if (!ToLandscape)
+					{
+						// No Landscape Proxy, don't support for creating only for Spline now
+						return;
+					}
+				}
+
+				if (ToLandscape != FromProxy)
+				{
+					ToLandscape->Modify();
+					if (ToLandscape->SplineComponent == nullptr)
+					{
+						CreateSplineComponent(ToLandscape, FromProxy->SplineComponent->RelativeScale3D);
+					}
+					ToLandscape->SplineComponent->Modify();
+
+					if (FromProxies.Find(FromProxy) == nullptr)
+					{
+						FromProxies.Add(FromProxy);
+						FromProxy->Modify();
+						FromProxy->SplineComponent->Modify();
+						FromProxy->SplineComponent->MarkRenderStateDirty();
+					}
+
+					// Handle spline meshes
+					if (Segment->bPlaceSplineMeshesInStreamingLevels)
+					{
+						// Mark previously local components as Foreign
+						for (auto* MeshComponent : Segment->LocalMeshComponents)
+						{
+							verifySlow(FromProxy->SplineComponent->MeshComponentLocalOwnersMap.Remove(MeshComponent) == 1);
+							FromProxy->SplineComponent->AddForeignMeshComponent(Segment, MeshComponent);
+						}
+						Segment->LocalMeshComponents.Empty();
+
+						// Mark previously foreign components as local
+						TArray<USplineMeshComponent*> MeshComponents = ToLandscape->SplineComponent->GetForeignMeshComponents(Segment);
+						ToLandscape->SplineComponent->RemoveAllForeignMeshComponents(Segment);
+						for (auto* MeshComponent : MeshComponents)
+						{
+							ToLandscape->SplineComponent->MeshComponentLocalOwnersMap.Add(MeshComponent, Segment);
+						}
+						Segment->LocalMeshComponents = MoveTemp(MeshComponents);
+					}
+					else
+					{
+						// non-streaming case
+						for (auto* MeshComponent : Segment->LocalMeshComponents)
+						{
+							MeshComponent->Modify();
+							MeshComponent->UnregisterComponent();
+							MeshComponent->DetachFromParent(true);
+							MeshComponent->InvalidateLightingCache();
+							MeshComponent->Rename(nullptr, ToLandscape);
+							MeshComponent->AttachTo(ToLandscape->SplineComponent, NAME_None, EAttachLocation::KeepWorldPosition);
+
+							verifySlow(FromProxy->SplineComponent->MeshComponentLocalOwnersMap.Remove(MeshComponent) == 1);
+							ToLandscape->SplineComponent->MeshComponentLocalOwnersMap.Add(MeshComponent, Segment);
+						}
+					}
+
+					// Move segment to new level
+					FromProxy->SplineComponent->Segments.Remove(Segment);
+					Segment->Rename(nullptr, ToLandscape->SplineComponent);
+					ToLandscape->SplineComponent->Segments.Add(Segment);
+
+					Segment->UpdateSplinePoints();
+				}
+			}
+		}
+
+		if (ToLandscape && ToLandscape->SplineComponent)
+		{
+			if (!ToLandscape->SplineComponent->IsRegistered())
+			{
+				ToLandscape->SplineComponent->RegisterComponent();
+			}
+			else
+			{
+				ToLandscape->SplineComponent->MarkRenderStateDirty();
+			}
+		}
+
+		GUnrealEd->RedrawLevelEditingViewports();
+	}
+
 	void ShowSplineProperties()
 	{
 		TArray<UObject*> Objects;
@@ -634,7 +827,7 @@ public:
 
 			if (!Landscape->SplineComponent)
 			{
-				CreateSplineComponent(Landscape, FVector(1.f) / Landscape->GetRootComponent()->RelativeScale3D);
+				CreateSplineComponent(Landscape, FVector(1.0f) / Landscape->GetRootComponent()->RelativeScale3D);
 			}
 
 			const FTransform LandscapeToSpline = Landscape->LandscapeActorToWorld().GetRelativeTransform(Landscape->SplineComponent->ComponentToWorld);
@@ -702,28 +895,22 @@ public:
 			{
 				HActor* ActorProxy = (HActor*)HitProxy;
 				AActor* Actor = ActorProxy->Actor;
-				ULandscapeSplinesComponent* SplineComponent = Actor->FindComponentByClass<ULandscapeSplinesComponent>();
-				if (SplineComponent != NULL)
+				const UMeshComponent* MeshComponent = Cast<UMeshComponent>(ActorProxy->PrimComponent);
+				if (MeshComponent)
 				{
-					if (const UControlPointMeshComponent* ControlPointMeshComponent = Cast<const UControlPointMeshComponent>(ActorProxy->PrimComponent))
+					ULandscapeSplinesComponent* SplineComponent = Actor->FindComponentByClass<ULandscapeSplinesComponent>();
+					if (SplineComponent)
 					{
-						for (ULandscapeSplineControlPoint* ControlPoint : SplineComponent->ControlPoints)
+						UObject* ComponentOwner = SplineComponent->GetOwnerForMeshComponent(MeshComponent);
+						if (ComponentOwner)
 						{
-							if (ControlPoint->OwnsComponent(ControlPointMeshComponent))
+							if (ULandscapeSplineControlPoint* ControlPoint = Cast<ULandscapeSplineControlPoint>(ComponentOwner))
 							{
 								ClickedControlPoint = ControlPoint;
-								break;
 							}
-						}
-					}
-					else if (const USplineMeshComponent* SplineMeshComponent = Cast<const USplineMeshComponent>(ActorProxy->PrimComponent))
-					{
-						for (ULandscapeSplineSegment* SplineSegment : SplineComponent->Segments)
-						{
-							if (SplineSegment->OwnsComponent(SplineMeshComponent))
+							else if (ULandscapeSplineSegment* SplineSegment = Cast<ULandscapeSplineSegment>(ComponentOwner))
 							{
 								ClickedSplineSegment = SplineSegment;
-								break;
 							}
 						}
 					}
@@ -898,18 +1085,19 @@ public:
 					else if (HitProxy->IsA(HActor::StaticGetType()))
 					{
 						HActor* ActorProxy = (HActor*)HitProxy;
-						ULandscapeSplinesComponent* SplineComponent = ActorProxy->Actor->FindComponentByClass<ULandscapeSplinesComponent>();
-						if (SplineComponent != NULL)
+						AActor* Actor = ActorProxy->Actor;
+						const UMeshComponent* MeshComponent = Cast<UMeshComponent>(ActorProxy->PrimComponent);
+						if (MeshComponent)
 						{
-							const UControlPointMeshComponent* ControlPointMeshComponent = Cast<const UControlPointMeshComponent>(ActorProxy->PrimComponent);
-							if (ControlPointMeshComponent != NULL)
+							ULandscapeSplinesComponent* SplineComponent = Actor->FindComponentByClass<ULandscapeSplinesComponent>();
+							if (SplineComponent)
 							{
-								for (ULandscapeSplineControlPoint* ControlPoint : SplineComponent->ControlPoints)
+								UObject* ComponentOwner = SplineComponent->GetOwnerForMeshComponent(MeshComponent);
+								if (ComponentOwner)
 								{
-									if (ControlPoint->OwnsComponent(ControlPointMeshComponent))
+									if (ULandscapeSplineControlPoint* ControlPoint = Cast<ULandscapeSplineControlPoint>(ComponentOwner))
 									{
 										ClickedControlPoint = ControlPoint;
-										break;
 									}
 								}
 							}
@@ -956,19 +1144,20 @@ public:
 					else if (HitProxy->IsA(HActor::StaticGetType()))
 					{
 						HActor* ActorProxy = (HActor*)HitProxy;
-						const USplineMeshComponent* SplineMeshComponent = Cast<const USplineMeshComponent>(ActorProxy->PrimComponent);
-						if (SplineMeshComponent != NULL)
+						AActor* Actor = ActorProxy->Actor;
+						const UMeshComponent* MeshComponent = Cast<UMeshComponent>(ActorProxy->PrimComponent);
+						if (MeshComponent)
 						{
-							ULandscapeSplinesComponent* SplineComponent = ActorProxy->Actor->FindComponentByClass<ULandscapeSplinesComponent>();
-							if (SplineComponent != NULL)
+							ULandscapeSplinesComponent* SplineComponent = Actor->FindComponentByClass<ULandscapeSplinesComponent>();
+							if (SplineComponent)
 							{
-								for (ULandscapeSplineSegment* SplineSegment : SplineComponent->Segments)
+								UObject* ComponentOwner = SplineComponent->GetOwnerForMeshComponent(MeshComponent);
+								if (ComponentOwner)
 								{
-									if (SplineSegment->OwnsComponent(SplineMeshComponent))
+									if (ULandscapeSplineSegment* SplineSegment = Cast<ULandscapeSplineSegment>(ComponentOwner))
 									{
 										ClickedSplineSegment = SplineSegment;
 										LandscapeToSpline = ActorProxy->Actor->ActorToWorld().GetRelativeTransform(SplineComponent->ComponentToWorld);
-										break;
 									}
 								}
 							}
@@ -1237,11 +1426,19 @@ public:
 
 		for (const FLandscapeListInfo& Info : EdMode->GetLandscapeList())
 		{
-			ALandscapeProxy* Landscape = Info.Info->GetCurrentLevelLandscapeProxy(true);
-			if (Landscape != NULL && Landscape->SplineComponent != NULL)
+			ALandscape* Landscape = Info.Info->LandscapeActor.Get();
+			if (Landscape && Landscape->SplineComponent)
 			{
 				Landscape->SplineComponent->ShowSplineEditorMesh(true);
 			}
+
+			for (ALandscapeProxy* LandscapeProxy : Info.Info->Proxies)
+			{
+				if (LandscapeProxy->SplineComponent)
+				{
+					LandscapeProxy->SplineComponent->ShowSplineEditorMesh(true);
+				}
+			}
 		}
 	}
 
@@ -1252,11 +1449,19 @@ public:
 
 		for (const FLandscapeListInfo& Info : EdMode->GetLandscapeList())
 		{
-			ALandscapeProxy* Landscape = Info.Info->GetCurrentLevelLandscapeProxy(true);
-			if (Landscape != NULL && Landscape->SplineComponent != NULL)
+			ALandscape* Landscape = Info.Info->LandscapeActor.Get();
+			if (Landscape && Landscape->SplineComponent)
 			{
 				Landscape->SplineComponent->ShowSplineEditorMesh(false);
 			}
+
+			for (ALandscapeProxy* LandscapeProxy : Info.Info->Proxies)
+			{
+				if (LandscapeProxy->SplineComponent)
+				{
+					LandscapeProxy->SplineComponent->ShowSplineEditorMesh(false);
+				}
+			}
 		}
 	}
 
@@ -1282,7 +1487,7 @@ public:
 					PDI->DrawLine(StartPos, HandlePos, FColor::White, SDPG_Foreground);
 
 					if (PDI->IsHitTesting()) PDI->SetHitProxy(new HLandscapeSplineProxy_Tangent(Connection.Segment, Connection.End));
-					PDI->DrawPoint(HandlePos, FColor::White, 10.f, SDPG_Foreground);
+					PDI->DrawPoint(HandlePos, FColor::White, 10.0f, SDPG_Foreground);
 					if (PDI->IsHitTesting()) PDI->SetHitProxy(NULL);
 				}
 			}
@@ -1305,7 +1510,7 @@ public:
 
 					PDI->DrawLine(EndPos, EndHandlePos, FColor::White, SDPG_Foreground);
 					if (PDI->IsHitTesting()) PDI->SetHitProxy(new HLandscapeSplineProxy_Tangent(Segment, !!End));
-					PDI->DrawPoint(EndHandlePos, FColor::White, 10.f, SDPG_Foreground);
+					PDI->DrawPoint(EndHandlePos, FColor::White, 10.0f, SDPG_Foreground);
 					if (PDI->IsHitTesting()) PDI->SetHitProxy(NULL);
 				}
 			}
@@ -1564,7 +1769,7 @@ public:
 		}
 		if (!Landscape->SplineComponent)
 		{
-			CreateSplineComponent(Landscape, FVector(1.f) / Landscape->GetRootComponent()->RelativeScale3D);
+			CreateSplineComponent(Landscape, FVector(1.0f) / Landscape->GetRootComponent()->RelativeScale3D);
 		}
 		Landscape->SplineComponent->Modify();
 
@@ -1671,140 +1876,10 @@ void FEdModeLandscape::SplineMoveToCurrentLevel()
 		SplinesTool->SelectAdjacentControlPoints();
 		SplinesTool->SelectConnected();
 
-		TSet<ALandscapeProxy*> FromProxies;
-		ALandscapeProxy* Landscape = NULL;
-
-		for (ULandscapeSplineControlPoint* ControlPoint : SplinesTool->SelectedSplineControlPoints)
-		{
-			ULandscapeSplinesComponent* LandscapeSplinesComp = ControlPoint->GetOuterULandscapeSplinesComponent();
-			ALandscapeProxy* FromProxy = LandscapeSplinesComp ? Cast<ALandscapeProxy>(LandscapeSplinesComp->GetOuter()) : NULL;
-			if (FromProxy)
-			{
-				if (!Landscape)
-				{
-					ULandscapeInfo* LandscapeInfo = FromProxy->GetLandscapeInfo(false);
-					check(LandscapeInfo);
-					Landscape = LandscapeInfo->GetCurrentLevelLandscapeProxy(true);
-					if (!Landscape)
-					{
-						// No Landscape Proxy, don't support for creating only for Spline now
-						return;
-					}
-				}
-
-				if (Landscape != FromProxy)
-				{
-					Landscape->Modify();
-					if (Landscape->SplineComponent == NULL)
-					{
-						SplinesTool->CreateSplineComponent(Landscape, FromProxy->SplineComponent->RelativeScale3D);
-					}
-					Landscape->SplineComponent->Modify();
-
-					const FTransform OldToNewTransform =
-						FromProxy->SplineComponent->ComponentToWorld.GetRelativeTransform(Landscape->SplineComponent->ComponentToWorld);
-
-					if (FromProxies.Find(FromProxy) == NULL)
-					{
-						FromProxies.Add(FromProxy);
-						FromProxy->Modify();
-						FromProxy->SplineComponent->Modify();
-						FromProxy->SplineComponent->MarkRenderStateDirty();
-					}
-
-					// Add/Remove Control Point/Segment/MeshComponents
-					FromProxy->SplineComponent->ControlPoints.Remove(ControlPoint);
-					ControlPoint->Rename(NULL, Landscape->SplineComponent);
-					Landscape->SplineComponent->ControlPoints.Add(ControlPoint);
-
-					ControlPoint->Location = OldToNewTransform.TransformPosition(ControlPoint->Location);
-
-					if (ControlPoint->MeshComponent)
-					{
-						UControlPointMeshComponent* ControlPointMeshComp = ControlPoint->MeshComponent;
-						ControlPointMeshComp->Modify();
-						ControlPointMeshComp->UnregisterComponent();
-						ControlPointMeshComp->DetachFromParent(true);
-						ControlPointMeshComp->InvalidateLightingCache();
-						ControlPointMeshComp->Rename(NULL, Landscape);
-						ControlPointMeshComp->AttachTo(Landscape->SplineComponent, NAME_None, EAttachLocation::KeepWorldPosition);
-					}
-
-					ControlPoint->UpdateSplinePoints(true, false);
-				}
-			}
-		}
-
-		for (ULandscapeSplineSegment* Segment : SplinesTool->SelectedSplineSegments)
-		{
-			ULandscapeSplinesComponent* LandscapeSplinesComp = Segment->GetOuterULandscapeSplinesComponent();
-			ALandscapeProxy* FromProxy = LandscapeSplinesComp ? Cast<ALandscapeProxy>(LandscapeSplinesComp->GetOuter()) : NULL;
-			if (FromProxy)
-			{
-				if (!Landscape)
-				{
-					ULandscapeInfo* LandscapeInfo = FromProxy->GetLandscapeInfo(false);
-					check(LandscapeInfo);
-					Landscape = LandscapeInfo->GetCurrentLevelLandscapeProxy(true);
-					if (!Landscape)
-					{
-						// No Landscape Proxy, don't support for creating only for Spline now
-						return;
-					}
-				}
-
-				if (Landscape != FromProxy)
-				{
-					Landscape->Modify();
-					if (Landscape->SplineComponent == NULL)
-					{
-						SplinesTool->CreateSplineComponent(Landscape, FromProxy->SplineComponent->RelativeScale3D);
-					}
-					Landscape->SplineComponent->Modify();
-
-					if (FromProxies.Find(FromProxy) == NULL)
-					{
-						FromProxies.Add(FromProxy);
-						FromProxy->Modify();
-						FromProxy->SplineComponent->Modify();
-						FromProxy->SplineComponent->MarkRenderStateDirty();
-					}
-
-					// Add/Remove Control Point/Segment/MeshComponents
-					FromProxy->SplineComponent->Segments.Remove(Segment);
-					Segment->Rename(NULL, Landscape->SplineComponent);
-					Landscape->SplineComponent->Segments.Add(Segment);
-
-					for (USplineMeshComponent* SplineMeshComp : Segment->MeshComponents)
-					{
-						SplineMeshComp->Modify();
-						SplineMeshComp->UnregisterComponent();
-						SplineMeshComp->DetachFromParent(true);
-						SplineMeshComp->InvalidateLightingCache();
-						SplineMeshComp->Rename(NULL, Landscape);
-						SplineMeshComp->AttachTo(Landscape->SplineComponent, NAME_None, EAttachLocation::KeepWorldPosition);
-					}
-
-					Segment->UpdateSplinePoints();
-				}
-			}
-		}
-
-		if (Landscape && Landscape->SplineComponent)
-		{
-			if (!Landscape->SplineComponent->IsRegistered())
-			{
-				Landscape->SplineComponent->RegisterComponent();
-			}
-			else
-			{
-				Landscape->SplineComponent->MarkRenderStateDirty();
-			}
-		}
+		SplinesTool->MoveSelectedToLevel();
 
 		SplinesTool->ClearSelection();
 		SplinesTool->UpdatePropertiesWindows();
-		GUnrealEd->RedrawLevelEditingViewports();
 	}
 }
 
diff --git a/Engine/Source/Runtime/Core/Public/UObject/ObjectVersion.h b/Engine/Source/Runtime/Core/Public/UObject/ObjectVersion.h
--- a/Engine/Source/Runtime/Core/Public/UObject/ObjectVersion.h
+++ b/Engine/Source/Runtime/Core/Public/UObject/ObjectVersion.h
@@ -486,6 +486,8 @@ enum EUnrealEngineObjectUE4Version
 	VER_UE4_PACKAGE_SUMMARY_HAS_COMPATIBLE_ENGINE_VERSION,
 	// Track UCS modified properties on Actor Components
 	VER_UE4_TRACK_UCS_MODIFIED_PROPERTIES,
+	// Allowed landscape spline meshes to be stored into landscape streaming levels rather than the spline's level
+	VER_UE4_LANDSCAPE_SPLINE_CROSS_LEVEL_MESHES,
 
 	// -----<new versions can be added before this line>-------------------------------------------------
 	// - this needs to be the last line (see note below)
diff --git a/Engine/Source/Runtime/Landscape/Classes/LandscapeInfo.h b/Engine/Source/Runtime/Landscape/Classes/LandscapeInfo.h
--- a/Engine/Source/Runtime/Landscape/Classes/LandscapeInfo.h
+++ b/Engine/Source/Runtime/Landscape/Classes/LandscapeInfo.h
@@ -189,7 +189,14 @@ public:
 	LANDSCAPE_API int32 GetLayerInfoIndex(ULandscapeLayerInfoObject* LayerInfo, ALandscapeProxy* Owner = NULL) const;
 	LANDSCAPE_API bool UpdateLayerInfoMap(ALandscapeProxy* Proxy = NULL, bool bInvalidate = false);
 
-	/** 
+	/**
+	 *  Returns the landscape proxy of this landscape info in the given level (if it exists)
+	 *  @param  Level  Level to look in
+	 *	@return        Landscape or landscape proxy found in the given level, or null if none
+	 */
+	LANDSCAPE_API ALandscapeProxy* GetLandscapeProxyForLevel(ULevel* Level) const;
+
+	/**
 	 *  Returns landscape which is spawned in the current level that was previously added to this landscape info object
 	 *  @param	bRegistered		Whether to consider only registered(visible) landscapes
 	 *	@return					Landscape or landscape proxy found in the current level 
diff --git a/Engine/Source/Runtime/Landscape/Classes/LandscapeSplineControlPoint.h b/Engine/Source/Runtime/Landscape/Classes/LandscapeSplineControlPoint.h
--- a/Engine/Source/Runtime/Landscape/Classes/LandscapeSplineControlPoint.h
+++ b/Engine/Source/Runtime/Landscape/Classes/LandscapeSplineControlPoint.h
@@ -2,11 +2,15 @@
 
 #pragma once
 
+// for FLandscapeSplineInterpPoint
 #include "LandscapeSplineSegment.h"
 
 #include "LandscapeSplineControlPoint.generated.h"
 
-struct FLandscapeSplineInterpPoint;
+// forward declarations
+class ULandscapeSplineSegment;
+class UControlPointMeshComponent;
+struct FLandscapeSplineSegmentConnection;
 
 USTRUCT()
 struct FLandscapeSplineConnection
@@ -15,7 +19,7 @@ struct FLandscapeSplineConnection
 
 	FLandscapeSplineConnection() {}
 
-	FLandscapeSplineConnection(class ULandscapeSplineSegment* InSegment, int32 InEnd)
+	FLandscapeSplineConnection(ULandscapeSplineSegment* InSegment, int32 InEnd)
 		: Segment(InSegment)
 		, End(InEnd)
 	{
@@ -28,14 +32,14 @@ struct FLandscapeSplineConnection
 
 	// Segment connected to this control point
 	UPROPERTY()
-	class ULandscapeSplineSegment* Segment;
+	ULandscapeSplineSegment* Segment;
 
 	// Which end of the segment is connected to this control point
 	UPROPERTY()
 	uint32 End:1;
 
-	LANDSCAPE_API struct FLandscapeSplineSegmentConnection& GetNearConnection() const;
-	LANDSCAPE_API struct FLandscapeSplineSegmentConnection& GetFarConnection() const;
+	LANDSCAPE_API FLandscapeSplineSegmentConnection& GetNearConnection() const;
+	LANDSCAPE_API FLandscapeSplineSegmentConnection& GetFarConnection() const;
 };
 
 UCLASS(Within=LandscapeSplinesComponent,autoExpandCategories=LandscapeSplineControlPoint,MinimalAPI)
@@ -67,11 +71,11 @@ class ULandscapeSplineControlPoint : public UObject
 #if WITH_EDITORONLY_DATA
 	/** Mesh to use on the control point */
 	UPROPERTY(EditAnywhere, Category=Mesh)
-	class UStaticMesh* Mesh;
+	UStaticMesh* Mesh;
 
 	/** Overrides mesh's materials */
 	UPROPERTY(EditAnywhere, Category=Mesh)
-	TArray<class UMaterialInterface*> MaterialOverrides;
+	TArray<UMaterialInterface*> MaterialOverrides;
 
 	/** Scale of the control point mesh */
 	UPROPERTY(EditAnywhere, Category=Mesh)
@@ -106,6 +110,10 @@ class ULandscapeSplineControlPoint : public UObject
 	UPROPERTY(EditAnywhere, Category=LandscapeSplineControlPoint)
 	uint32 bLowerTerrain:1;
 
+	/** Whether control point mesh should be placed in landscape proxy streaming level (true) or the spline's level (false) */
+	UPROPERTY(EditAnywhere, Category=Mesh)
+	uint32 bPlaceSplineMeshesInStreamingLevels : 1;
+
 	/** Whether to enable collision for the Control Point Mesh. */
 	UPROPERTY(EditAnywhere, Category=Mesh)
 	uint32 bEnableCollision:1;
@@ -128,22 +136,33 @@ public:
 	TArray<FLandscapeSplineConnection> ConnectedSegments;
 
 protected:
-	/** Control point mesh */
-	UPROPERTY(TextExportTransient)
-	class UControlPointMeshComponent* MeshComponent;
-
 	/** Spline points */
 	UPROPERTY()
-	TArray<struct FLandscapeSplineInterpPoint> Points;
+	TArray<FLandscapeSplineInterpPoint> Points;
 
 	/** Bounds of points */
 	UPROPERTY()
 	FBox Bounds;
 
+	/** Control point mesh */
+	UPROPERTY(TextExportTransient)
+	UControlPointMeshComponent* LocalMeshComponent;
+
+#if WITH_EDITORONLY_DATA
+	/** World reference for if mesh component is stored in another streaming level */
+	UPROPERTY(TextExportTransient, NonPIEDuplicateTransient)
+	TAssetPtr<UWorld> ForeignWorld;
+
+	/** Key for tracking whether this segment has been modified relative to the mesh component stored in another streaming level */
+	UPROPERTY(TextExportTransient, NonPIEDuplicateTransient)
+	FGuid ModificationKey;
+#endif
+
 public:
 	const FBox& GetBounds() const { return Bounds; }
-	const TArray<struct FLandscapeSplineInterpPoint>& GetPoints() const { return Points; }
+	const TArray<FLandscapeSplineInterpPoint>& GetPoints() const { return Points; }
 
+#if WITH_EDITOR
 	// Get the name of the best connection point (socket) to use for a particular destination
 	virtual FName GetBestConnectionTo(FVector Destination) const;
 
@@ -153,7 +172,6 @@ public:
 	// Get the location and rotation of a connection point (socket) in spline space
 	virtual void GetConnectionLocationAndRotation(FName SocketName, OUT FVector& OutLocation, OUT FRotator& OutRotation) const;
 
-#if WITH_EDITOR
 	bool IsSplineSelected() const { return bSelected; }
 	virtual void SetSplineSelected(bool bInSelected);
 
@@ -165,14 +183,20 @@ public:
 
 	virtual void AutoSetConnections(bool bIncludingValid);
 
+	TMap<ULandscapeSplinesComponent*, UControlPointMeshComponent*> GetForeignMeshComponents();
+
 	/** Update spline points */
 	virtual void UpdateSplinePoints(bool bUpdateCollision = true, bool bUpdateAttachedSegments = true);
 
 	/** Delete spline points */
 	virtual void DeleteSplinePoints();
+
+	const TAssetPtr<UWorld>& GetForeignWorld() const { return ForeignWorld; }
+	FGuid GetModificationKey() const { return ModificationKey; }
 #endif // WITH_EDITOR
 
 	// Begin UObject Interface
+	virtual void Serialize(FArchive& Ar) override;
 	virtual void PostLoad() override;
 #if WITH_EDITOR
 	virtual void PostEditUndo() override;
@@ -182,10 +206,5 @@ public:
 	// End UObject Interface
 #endif // WITH_EDITOR
 
-	void RegisterComponents();
-	void UnregisterComponents();
-
-	virtual bool OwnsComponent(const class UControlPointMeshComponent* ControlPointMeshComponent) const;
-
-	friend class FEdModeLandscape;
+	friend class FLandscapeToolSplines;
 };
diff --git a/Engine/Source/Runtime/Landscape/Classes/LandscapeSplineSegment.h b/Engine/Source/Runtime/Landscape/Classes/LandscapeSplineSegment.h
--- a/Engine/Source/Runtime/Landscape/Classes/LandscapeSplineSegment.h
+++ b/Engine/Source/Runtime/Landscape/Classes/LandscapeSplineSegment.h
@@ -159,7 +159,7 @@ class ULandscapeSplineSegment : public UObject
 
 	/** Spline meshes from this list are used in random order along the spline. */
 	UPROPERTY(EditAnywhere, Category=LandscapeSplineMeshes)
-	TArray<struct FLandscapeSplineMeshEntry> SplineMeshes;
+	TArray<FLandscapeSplineMeshEntry> SplineMeshes;
 
 	/** Random seed used for choosing which spline meshes to use. */
 	UPROPERTY(EditAnywhere, Category=LandscapeSplineMeshes)
@@ -187,6 +187,10 @@ class ULandscapeSplineSegment : public UObject
 	UPROPERTY(EditAnywhere, Category=LandscapeSplineSegment)
 	uint32 bLowerTerrain:1;
 
+	/** Whether spline meshes should be placed in landscape proxy streaming levels (true) or the spline's level (false) */
+	UPROPERTY(EditAnywhere, Category = LandscapeSplineMeshes)
+	uint32 bPlaceSplineMeshesInStreamingLevels : 1;
+
 	/** Whether to generate collision for the Spline Meshes. */
 	UPROPERTY(EditAnywhere, Category=LandscapeSplineMeshes)
 	uint32 bEnableCollision:1;
@@ -219,7 +223,17 @@ protected:
 
 	/** Spline meshes */
 	UPROPERTY(TextExportTransient)
-	TArray<class USplineMeshComponent*> MeshComponents;
+	TArray<USplineMeshComponent*> LocalMeshComponents;
+
+#if WITH_EDITORONLY_DATA
+	/** World references for mesh components stored in other streaming levels */
+	UPROPERTY(TextExportTransient, NonPIEDuplicateTransient)
+	TArray<TAssetPtr<UWorld>> ForeignWorlds;
+
+	/** Key for tracking whether this segment has been modified relative to the mesh components stored in other streaming levels */
+	UPROPERTY(TextExportTransient, NonPIEDuplicateTransient)
+	FGuid ModificationKey;
+#endif
 
 public:
 	const FBox& GetBounds() const { return Bounds; }
@@ -230,12 +244,19 @@ public:
 	virtual void SetSplineSelected(bool bInSelected);
 
 	virtual void AutoFlipTangents();
+
+	TMap<ULandscapeSplinesComponent*, TArray<USplineMeshComponent*>> GetForeignMeshComponents();
+
 	virtual void UpdateSplinePoints(bool bUpdateCollision = true);
+
 	void UpdateSplineEditorMesh();
 	virtual void DeleteSplinePoints();
+
+	const TArray<TAssetPtr<UWorld>>& GetForeignWorlds() const { return ForeignWorlds; }
+	FGuid GetModificationKey() const { return ModificationKey; }
 #endif
 
-	virtual void FindNearest( const FVector& InLocation, float& t, FVector& OutLocation, FVector& OutTangent );
+	virtual void FindNearest(const FVector& InLocation, float& t, FVector& OutLocation, FVector& OutTangent);
 
 	// Begin UObject Interface
 	virtual void Serialize(FArchive& Ar) override;
@@ -252,10 +273,5 @@ protected:
 public:
 	// End UObject Interface
 
-	void RegisterComponents();
-	void UnregisterComponents();
-
-	virtual bool OwnsComponent(const class USplineMeshComponent* SplineMeshComponent) const;
-
-	friend class FEdModeLandscape;
+	friend class FLandscapeToolSplines;
 };
diff --git a/Engine/Source/Runtime/Landscape/Classes/LandscapeSplinesComponent.h b/Engine/Source/Runtime/Landscape/Classes/LandscapeSplinesComponent.h
--- a/Engine/Source/Runtime/Landscape/Classes/LandscapeSplinesComponent.h
+++ b/Engine/Source/Runtime/Landscape/Classes/LandscapeSplinesComponent.h
@@ -1,9 +1,62 @@
 // Copyright 1998-2015 Epic Games, Inc. All Rights Reserved.
 
-
 #pragma once
+
 #include "LandscapeSplinesComponent.generated.h"
 
+// Forward declarations
+class ULandscapeSplineControlPoint;
+class ULandscapeSplineSegment;
+class USplineMeshComponent;
+class UControlPointMeshComponent;
+
+USTRUCT()
+struct FForeignControlPointData
+{
+	GENERATED_USTRUCT_BODY()
+
+#if WITH_EDITORONLY_DATA
+	UPROPERTY()
+	FGuid ModificationKey;
+
+	UPROPERTY()
+	UControlPointMeshComponent* MeshComponent;
+#endif
+
+	friend FArchive& operator<<(FArchive& Ar, FForeignControlPointData& Value);
+};
+
+USTRUCT()
+struct FForeignSplineSegmentData
+{
+	GENERATED_USTRUCT_BODY()
+
+#if WITH_EDITORONLY_DATA
+		UPROPERTY()
+		FGuid ModificationKey;
+
+	UPROPERTY()
+		TArray<USplineMeshComponent*> MeshComponents;
+#endif
+
+	friend FArchive& operator<<(FArchive& Ar, FForeignSplineSegmentData& Value);
+};
+
+USTRUCT()
+struct FForeignWorldSplineData
+{
+	GENERATED_USTRUCT_BODY()
+
+#if WITH_EDITORONLY_DATA
+	TMap<TLazyObjectPtr<ULandscapeSplineControlPoint>, FForeignControlPointData> ForeignControlPointDataMap;
+	TMap<TLazyObjectPtr<ULandscapeSplineSegment>, FForeignSplineSegmentData> ForeignSplineSegmentDataMap;
+#endif
+
+	friend FArchive& operator<<(FArchive& Ar, FForeignWorldSplineData& Value);
+
+	bool IsEmpty();
+};
+
 UCLASS(MinimalAPI)
 class ULandscapeSplinesComponent : public UPrimitiveComponent
 {
@@ -20,23 +73,30 @@ class ULandscapeSplinesComponent : public UPrimitiveComponent
 
 	/** Sprite used to draw control points */
 	UPROPERTY()
-	class UTexture2D* ControlPointSprite;
+	UTexture2D* ControlPointSprite;
 
 	/** Mesh used to draw splines that have no mesh */
 	UPROPERTY()
-	class UStaticMesh* SplineEditorMesh;
+	UStaticMesh* SplineEditorMesh;
 
 	/** Whether we are in-editor and showing spline editor meshes */
 	UPROPERTY(NonTransactional, Transient)
 	uint32 bShowSplineEditorMesh:1;
 #endif
 
+protected:
 	UPROPERTY(TextExportTransient)
-	TArray<class ULandscapeSplineControlPoint*> ControlPoints;
+	TArray<ULandscapeSplineControlPoint*> ControlPoints;
 
 	UPROPERTY(TextExportTransient)
-	TArray<class ULandscapeSplineSegment*> Segments;
+	TArray<ULandscapeSplineSegment*> Segments;
+
+#if WITH_EDITORONLY_DATA
+	TMap<TAssetPtr<UWorld>, FForeignWorldSplineData> ForeignWorldSplineDataMap;
 
+	TMap<UMeshComponent*, UObject*> MeshComponentLocalOwnersMap;
+	TMap<UMeshComponent*, TLazyObjectPtr<UObject>> MeshComponentForeignOwnersMap;
+#endif
 
 public:
 	void CheckSplinesValid();
@@ -44,9 +104,42 @@ public:
 
 #if WITH_EDITOR
 	virtual void ShowSplineEditorMesh(bool bShow);
+
+	// Rebuilds all spline points and meshes for all spline control points and segments in this splines component
+	// @param bBuildCollision Building collision data is very slow, so during interactive changes pass false to improve user experience (make sure to call with true when done)
+	virtual void RebuildAllSplines(bool bBuildCollision = true);
+
+	// returns a suitable ULandscapeSplinesComponent to place streaming meshes into, given a location
+	// falls back to "this" if it can't find another suitable, so never returns nullptr
+	// @param bCreate whether to create a component if a suitable actor is found but it has no splines component yet
+	ULandscapeSplinesComponent* GetStreamingSplinesComponentByLocation(const FVector& LocalLocation, bool bCreate = true);
+
+	// returns the matching ULandscapeSplinesComponent for a given level, *can return null*
+	// @param bCreate whether to create a component if a suitable actor is found but it has no splines component yet
+	ULandscapeSplinesComponent* GetStreamingSplinesComponentForLevel(ULevel* Level, bool bCreate = true);
+
+	// gathers and returns all currently existing 
+	TArray<ULandscapeSplinesComponent*> GetAllStreamingSplinesComponents();
+
+	virtual void UpdateModificationKey(ULandscapeSplineSegment* Owner);
+	virtual void UpdateModificationKey(ULandscapeSplineControlPoint* Owner);
+	virtual void AddForeignMeshComponent(ULandscapeSplineSegment* Owner, USplineMeshComponent* Component);
+	virtual void RemoveForeignMeshComponent(ULandscapeSplineSegment* Owner, USplineMeshComponent* Component);
+	virtual void RemoveAllForeignMeshComponents(ULandscapeSplineSegment* Owner);
+	virtual void AddForeignMeshComponent(ULandscapeSplineControlPoint* Owner, UControlPointMeshComponent* Component);
+	virtual void RemoveForeignMeshComponent(ULandscapeSplineControlPoint* Owner, UControlPointMeshComponent* Component);
+	virtual void DestroyOrphanedForeignMeshComponents(UWorld* OwnerWorld);
+	virtual UControlPointMeshComponent*   GetForeignMeshComponent(ULandscapeSplineControlPoint* Owner);
+	virtual TArray<USplineMeshComponent*> GetForeignMeshComponents(ULandscapeSplineSegment* Owner);
+
+	virtual UObject* GetOwnerForMeshComponent(const UMeshComponent* SplineMeshComponent);
+
+	void AutoFixMeshComponentErrors(UWorld* OtherWorld);
 #endif
 
 	// Begin UObject Interface
+
+	virtual void Serialize(FArchive& Ar) override;
 	virtual void PostLoad() override;
 #if WITH_EDITOR
 	virtual void PostEditChangeProperty(FPropertyChangedEvent& PropertyChangedEvent) override;
@@ -55,6 +148,9 @@ public:
 	// End UObject Interface
 
 	// Begin UActorComponent interface
+#if WITH_EDITOR
+	virtual void CheckForErrors() override;
+#endif
 	virtual void OnRegister() override;
 	// End UActorComponent interface
 	
@@ -64,4 +160,14 @@ public:
 #endif
 	virtual FBoxSphereBounds CalcBounds(const FTransform& LocalToWorld) const override;
 	// End UPrimitiveComponent interface.
+
+	// many friends
+	friend class FLandscapeToolSplines;
+	friend class FLandscapeSplinesSceneProxy;
+	friend class ULandscapeSplineControlPoint;
+	friend class ULandscapeSplineSegment;
+#if WITH_EDITOR
+	// TODO - move this out of ULandscapeInfo
+	friend bool ULandscapeInfo::ApplySplinesInternal(bool bOnlySelected, ALandscapeProxy* Landscape);
+#endif
 };
diff --git a/Engine/Source/Runtime/Landscape/Private/Landscape.cpp b/Engine/Source/Runtime/Landscape/Private/Landscape.cpp
--- a/Engine/Source/Runtime/Landscape/Private/Landscape.cpp
+++ b/Engine/Source/Runtime/Landscape/Private/Landscape.cpp
@@ -1541,6 +1541,26 @@ void ALandscape::PreSave()
 	//}
 }
 
+ALandscapeProxy* ULandscapeInfo::GetLandscapeProxyForLevel(ULevel* Level) const
+{
+	ALandscape* Landscape = LandscapeActor.Get();
+
+	if (Landscape && Landscape->GetLevel() == Level)
+	{
+		return Landscape;
+	}
+
+	for (ALandscapeProxy* LandscapeProxy : Proxies)
+	{
+		if (LandscapeProxy && LandscapeProxy->GetLevel() == Level)
+		{
+			return LandscapeProxy;
+		}
+	}
+
+	return nullptr;
+}
+
 ALandscapeProxy* ULandscapeInfo::GetCurrentLevelLandscapeProxy(bool bRegistered) const
 {
 	ALandscape* Landscape = LandscapeActor.Get();
diff --git a/Engine/Source/Runtime/Landscape/Private/LandscapeSplines.cpp b/Engine/Source/Runtime/Landscape/Private/LandscapeSplines.cpp
--- a/Engine/Source/Runtime/Landscape/Private/LandscapeSplines.cpp
+++ b/Engine/Source/Runtime/Landscape/Private/LandscapeSplines.cpp
@@ -20,6 +20,8 @@
 #include "Engine/StaticMeshSocket.h"
 #if WITH_EDITOR
 #include "LandscapeSplineRaster.h"
+#include "MessageLog.h"
+#include "UObjectToken.h"
 #endif
 
 IMPLEMENT_HIT_PROXY(HLandscapeSplineProxy, HHitProxy);
@@ -27,6 +29,8 @@ IMPLEMENT_HIT_PROXY(HLandscapeSplineProxy_Segment, HLandscapeSplineProxy);
 IMPLEMENT_HIT_PROXY(HLandscapeSplineProxy_ControlPoint, HLandscapeSplineProxy);
 IMPLEMENT_HIT_PROXY(HLandscapeSplineProxy_Tangent, HLandscapeSplineProxy);
 
+#define LOCTEXT_NAMESPACE "Landscape.Splines"
+
 //////////////////////////////////////////////////////////////////////////
 // LANDSCAPE SPLINES SCENE PROXY
 
@@ -79,7 +83,7 @@ public:
 		{
 			FSegmentProxy SegmentProxy;
 			SegmentProxy.Owner = Segment;
-			SegmentProxy.HitProxy = NULL;
+			SegmentProxy.HitProxy = nullptr;
 			SegmentProxy.Points = Segment->GetPoints();
 			SegmentProxy.bSelected = Segment->IsSplineSelected();
 			Segments.Add(SegmentProxy);
@@ -90,7 +94,7 @@ public:
 		{
 			FControlPointProxy ControlPointProxy;
 			ControlPointProxy.Owner = ControlPoint;
-			ControlPointProxy.HitProxy = NULL;
+			ControlPointProxy.HitProxy = nullptr;
 			ControlPointProxy.Location = ControlPoint->Location;
 			ControlPointProxy.Points = ControlPoint->GetPoints();
 			ControlPointProxy.SpriteScale = FMath::Clamp<float>(ControlPoint->Width != 0 ? ControlPoint->Width / 2 : ControlPoint->SideFalloff / 4, 10, 1000);
@@ -112,7 +116,7 @@ public:
 			ControlPoint.HitProxy = new HLandscapeSplineProxy_ControlPoint(ControlPoint.Owner);
 			OutHitProxies.Add(ControlPoint.HitProxy);
 		}
-		return NULL;
+		return nullptr;
 	}
 
 	virtual void GetDynamicMeshElements(const TArray<const FSceneView*>& Views, const FSceneViewFamily& ViewFamily, uint32 VisibilityMap, FMeshElementCollector& Collector) const override
@@ -162,7 +166,7 @@ public:
 						PDI->DrawLine(OldPoint.Left, NewPoint.Left, SegmentColor, GetDepthPriorityGroup(View), 0.0f, DepthBias);
 						PDI->DrawLine(OldPoint.Right, NewPoint.Right, SegmentColor, GetDepthPriorityGroup(View), 0.0f, DepthBias);
 
-						PDI->SetHitProxy(NULL);
+						PDI->SetHitProxy(nullptr);
 
 						// draw falloff sides
 						if (bDrawFalloff)
@@ -248,7 +252,7 @@ public:
 							// draw sides
 							PDI->DrawLine(OldPoint.Right, NewPoint.Left, ControlPointColor, GetDepthPriorityGroup(View), 0.0f, DepthBias);
 
-							PDI->SetHitProxy(NULL);
+							PDI->SetHitProxy(nullptr);
 
 							// draw falloff sides
 							if (bDrawFalloff)
@@ -273,7 +277,7 @@ public:
 					}
 				}
 
-				PDI->SetHitProxy(NULL);
+				PDI->SetHitProxy(nullptr);
 			}
 		}
 	}
@@ -352,6 +356,7 @@ void ULandscapeSplinesComponent::CheckSplinesValid()
 	Segments.Remove(nullptr);
 
 	// Check for cross-spline connections, as this is a potential source of nulls
+	// this may be allowed in future, but is not currently
 	for (ULandscapeSplineControlPoint* ControlPoint : ControlPoints)
 	{
 		ensure(ControlPoint->GetOuterULandscapeSplinesComponent() == this);
@@ -446,11 +451,251 @@ bool ULandscapeSplinesComponent::ModifySplines(bool bAlwaysMarkDirty /*= true*/)
 	return bSavedToTransactionBuffer;
 }
 
+FArchive& operator<<(FArchive& Ar, FForeignControlPointData& Value)
+{
+#if WITH_EDITORONLY_DATA
+	if (!Ar.IsFilterEditorOnly())
+	{
+		Ar << Value.ModificationKey << Value.MeshComponent;
+	}
+#endif
+	return Ar;
+}
+
+FArchive& operator<<(FArchive& Ar, FForeignSplineSegmentData& Value)
+{
+#if WITH_EDITORONLY_DATA
+	if (!Ar.IsFilterEditorOnly())
+	{
+		Ar << Value.ModificationKey << Value.MeshComponents;
+	}
+#endif
+	return Ar;
+}
+
+FArchive& operator<<(FArchive& Ar, FForeignWorldSplineData& Value)
+{
+#if WITH_EDITORONLY_DATA
+	if (!Ar.IsFilterEditorOnly())
+	{
+		Ar << Value.ForeignSplineSegmentDataMap;
+	}
+#endif
+	return Ar;
+}
+
+void ULandscapeSplinesComponent::Serialize(FArchive& Ar)
+{
+	Super::Serialize(Ar);
+
+#if WITH_EDITORONLY_DATA
+	if (Ar.UE4Ver() >= VER_UE4_LANDSCAPE_SPLINE_CROSS_LEVEL_MESHES &&
+		!Ar.IsFilterEditorOnly())
+	{
+		Ar << ForeignWorldSplineDataMap;
+	}
+
+	if (!Ar.IsPersistent())
+	{
+		Ar << MeshComponentLocalOwnersMap;
+		Ar << MeshComponentForeignOwnersMap;
+	}
+#endif
+}
+
+#if WITH_EDITOR
+void ULandscapeSplinesComponent::AutoFixMeshComponentErrors(UWorld* OtherWorld)
+{
+	UWorld* ThisOuterWorld = GetTypedOuter<UWorld>();
+
+	TAssetPtr<UWorld> OtherWorldAssetPtr = OtherWorld;
+	ULandscapeSplinesComponent* StreamingSplinesComponent = GetStreamingSplinesComponentForLevel(OtherWorld->PersistentLevel);
+	auto* ForeignWorldSplineData = StreamingSplinesComponent ? StreamingSplinesComponent->ForeignWorldSplineDataMap.Find(ThisOuterWorld) : nullptr;
+
+	// Fix control point meshes
+	for (ULandscapeSplineControlPoint* ControlPoint : ControlPoints)
+	{
+		if (ControlPoint->GetForeignWorld() == OtherWorld)
+		{
+			auto* ForeignControlPointData = ForeignWorldSplineData ? ForeignWorldSplineData->ForeignControlPointDataMap.Find(ControlPoint) : nullptr;
+			if (!ForeignControlPointData || ForeignControlPointData->ModificationKey != ControlPoint->GetModificationKey())
+			{
+				// We don't pass true for update segments to avoid them being updated multiple times
+				ControlPoint->UpdateSplinePoints(true, false);
+			}
+		}
+	}
+
+	// Fix spline segment meshes
+	for (ULandscapeSplineSegment* Segment : Segments)
+	{
+		if (Segment->GetForeignWorlds().Contains(OtherWorld))
+		{
+			auto* ForeignSplineSegmentData = ForeignWorldSplineData ? ForeignWorldSplineData->ForeignSplineSegmentDataMap.Find(Segment) : nullptr;
+			if (!ForeignSplineSegmentData || ForeignSplineSegmentData->ModificationKey != Segment->GetModificationKey())
+			{
+				Segment->UpdateSplinePoints(true);
+			}
+		}
+	}
+
+	if (StreamingSplinesComponent)
+	{
+		StreamingSplinesComponent->DestroyOrphanedForeignMeshComponents(ThisOuterWorld);
+	}
+}
+
+void ULandscapeSplinesComponent::CheckForErrors()
+{
+	Super::CheckForErrors();
+
+	UWorld* ThisOuterWorld = GetTypedOuter<UWorld>();
+
+	TSet<UWorld*> OutdatedWorlds;
+	TMap<UWorld*, FForeignWorldSplineData*> ForeignWorldSplineDataMapCache;
+
+	// Check control point meshes
+	for (ULandscapeSplineControlPoint* ControlPoint : ControlPoints)
+	{
+		UWorld* ForeignWorld = ControlPoint->GetForeignWorld().Get();
+		if (ForeignWorld && !OutdatedWorlds.Contains(ForeignWorld))
+		{
+			auto** ForeignWorldSplineDataCachedPtr = ForeignWorldSplineDataMapCache.Find(ForeignWorld);
+			auto* ForeignWorldSplineData = ForeignWorldSplineDataCachedPtr ? *ForeignWorldSplineDataCachedPtr : nullptr;
+			if (!ForeignWorldSplineDataCachedPtr)
+			{
+				ULandscapeSplinesComponent* StreamingSplinesComponent = GetStreamingSplinesComponentForLevel(ForeignWorld->PersistentLevel);
+				ForeignWorldSplineData = StreamingSplinesComponent ? StreamingSplinesComponent->ForeignWorldSplineDataMap.Find(ThisOuterWorld) : nullptr;
+				ForeignWorldSplineDataMapCache.Add(ForeignWorld, ForeignWorldSplineData);
+			}
+			auto* ForeignControlPointData = ForeignWorldSplineData ? ForeignWorldSplineData->ForeignControlPointDataMap.Find(ControlPoint) : nullptr;
+			if (!ForeignControlPointData || ForeignControlPointData->ModificationKey != ControlPoint->GetModificationKey())
+			{
+				OutdatedWorlds.Add(ForeignWorld);
+			}
+		}
+	}
+
+	// Check spline segment meshes
+	for (ULandscapeSplineSegment* Segment : Segments)
+	{
+		for (auto& ForeignWorldAssetPtr : Segment->GetForeignWorlds())
+		{
+			UWorld* ForeignWorld = ForeignWorldAssetPtr.Get();
+
+			if (ForeignWorld && !OutdatedWorlds.Contains(ForeignWorld))
+			{
+				auto** ForeignWorldSplineDataCachedPtr = ForeignWorldSplineDataMapCache.Find(ForeignWorld);
+				auto* ForeignWorldSplineData = ForeignWorldSplineDataCachedPtr ? *ForeignWorldSplineDataCachedPtr : nullptr;
+				if (!ForeignWorldSplineDataCachedPtr)
+				{
+					ULandscapeSplinesComponent* StreamingSplinesComponent = GetStreamingSplinesComponentForLevel(ForeignWorld->PersistentLevel);
+					ForeignWorldSplineData = StreamingSplinesComponent ? StreamingSplinesComponent->ForeignWorldSplineDataMap.Find(ThisOuterWorld) : nullptr;
+					ForeignWorldSplineDataMapCache.Add(ForeignWorld, ForeignWorldSplineData);
+				}
+				auto* ForeignSplineSegmentData = ForeignWorldSplineData ? ForeignWorldSplineData->ForeignSplineSegmentDataMap.Find(Segment) : nullptr;
+				if (!ForeignSplineSegmentData || ForeignSplineSegmentData->ModificationKey != Segment->GetModificationKey())
+				{
+					OutdatedWorlds.Add(ForeignWorld);
+				}
+			}
+		}
+	}
+	ForeignWorldSplineDataMapCache.Empty();
+
+	for (UWorld* OutdatedWorld : OutdatedWorlds)
+	{
+		FFormatNamedArguments Arguments;
+		Arguments.Add(TEXT("MeshMap"), FText::FromName(OutdatedWorld->GetFName()));
+		Arguments.Add(TEXT("SplineMap"), FText::FromName(ThisOuterWorld->GetFName()));
+
+		FMessageLog("MapCheck").Error()
+			->AddToken(FUObjectToken::Create(GetOwner()))
+			->AddToken(FTextToken::Create(FText::Format(LOCTEXT("MapCheck_Message_MeshesOutDated", "Meshes in {MeshMap} out of date compared to landscape spline in {SplineMap}"), Arguments)))
+			->AddToken(FActionToken::Create(LOCTEXT("MapCheck_ActionName_MeshesOutDated", "Rebuild landscape splines"), FText(),
+			FOnActionTokenExecuted::CreateUObject(this, &ULandscapeSplinesComponent::AutoFixMeshComponentErrors, OutdatedWorld), true));
+	}
+
+	// check for orphaned components
+	for (auto& ForeignWorldSplineDataPair : ForeignWorldSplineDataMap)
+	{
+		auto& ForeignWorldAssetPtr = ForeignWorldSplineDataPair.Key;
+		auto& ForeignWorldSplineData = ForeignWorldSplineDataPair.Value;
+
+		// World is not loaded
+		if (ForeignWorldAssetPtr.IsPending())
+		{
+			continue;
+		}
+
+		UWorld* ForeignWorld = ForeignWorldAssetPtr.Get();
+		for (auto& ForeignSplineSegmentDataPair : ForeignWorldSplineData.ForeignSplineSegmentDataMap)
+		{
+			const ULandscapeSplineSegment* ForeignSplineSegment = ForeignSplineSegmentDataPair.Key.Get();
+			auto& ForeignSplineSegmentData = ForeignSplineSegmentDataPair.Value;
+
+			// No such segment or segment doesn't match our meshes
+			if (!ForeignSplineSegment)
+			{
+				FFormatNamedArguments Arguments;
+				Arguments.Add(TEXT("MeshMap"),   FText::FromName(ThisOuterWorld->GetFName()));
+				Arguments.Add(TEXT("SplineMap"), FText::FromName(ForeignWorld->GetFName()));
+
+				FMessageLog("MapCheck").Error()
+					->AddToken(FUObjectToken::Create(GetOwner()))
+					->AddToken(FTextToken::Create(FText::Format(LOCTEXT("MapCheck_Message_OrphanedMeshes", "{MeshMap} contains orphaned meshes due to mismatch with landscape splines in {SplineMap}"), Arguments)))
+					->AddToken(FActionToken::Create(LOCTEXT("MapCheck_ActionName_OrphanedMeshes", "Clean up orphaned meshes"), FText(),
+					FOnActionTokenExecuted::CreateUObject(this, &ULandscapeSplinesComponent::DestroyOrphanedForeignMeshComponents, ForeignWorld), true));
+
+				break;
+			}
+		}
+	}
+}
+#endif
+
 void ULandscapeSplinesComponent::PostLoad()
 {
 	Super::PostLoad();
 
+#if WITH_EDITOR
+	if (GIsEditor)
+	{
+		// Build MeshComponentForeignOwnersMap (Component->Spline) from ForeignWorldSplineDataMap (World->Spline->Component)
+		for (auto& ForeignWorldSplineDataPair : ForeignWorldSplineDataMap)
+		{
+			auto& ForeignWorld = ForeignWorldSplineDataPair.Key;
+			auto& ForeignWorldSplineData = ForeignWorldSplineDataPair.Value;
+
+			for (auto& ForeignControlPointDataPair : ForeignWorldSplineData.ForeignControlPointDataMap)
+			{
+				TLazyObjectPtr<ULandscapeSplineControlPoint> ForeignControlPoint = ForeignControlPointDataPair.Key;
+				auto& ForeignControlPointData = ForeignControlPointDataPair.Value;
+				MeshComponentForeignOwnersMap.Add(ForeignControlPointData.MeshComponent, ForeignControlPoint);
+			}
+
+			for (auto& ForeignSplineSegmentDataPair : ForeignWorldSplineData.ForeignSplineSegmentDataMap)
+			{
+				TLazyObjectPtr<ULandscapeSplineSegment> ForeignSplineSegment = ForeignSplineSegmentDataPair.Key;
+				auto& ForeignSplineSegmentData = ForeignSplineSegmentDataPair.Value;
+
+				for (auto* MeshComponent : ForeignSplineSegmentData.MeshComponents)
+				{
+					MeshComponentForeignOwnersMap.Add(MeshComponent, ForeignSplineSegment);
+				}
+			}
+		}
+	}
+#endif
+
 	CheckSplinesValid();
+
+#if WITH_EDITOR
+	if (GIsEditor)
+	{
+		CheckForErrors();
+	}
+#endif
 }
 
 #if WITH_EDITOR
@@ -465,36 +710,376 @@ void ULandscapeSplinesComponent::PostEditChangeProperty(FPropertyChangedEvent& P
 	if (!bHackIsUndoingSplines)
 	{
 		const bool bUpdateCollision = PropertyChangedEvent.ChangeType != EPropertyChangeType::Interactive;
+		RebuildAllSplines(bUpdateCollision);
+	}
+}
+void ULandscapeSplinesComponent::PostEditUndo()
+{
+	bHackIsUndoingSplines = true;
+	Super::PostEditUndo();
+	bHackIsUndoingSplines = false;
+
+	MarkRenderStateDirty();
+}
+
+void ULandscapeSplinesComponent::RebuildAllSplines(bool bUpdateCollision)
+{
+	for (ULandscapeSplineControlPoint* ControlPoint : ControlPoints)
+	{
+		ControlPoint->UpdateSplinePoints(true, false);
+	}
+
+	for (ULandscapeSplineSegment* Segment : Segments)
+	{
+		Segment->UpdateSplinePoints(true);
+	}
+}
+
+void ULandscapeSplinesComponent::ShowSplineEditorMesh(bool bShow)
+{
+	bShowSplineEditorMesh = bShow;
+
+	for (ULandscapeSplineSegment* Segment : Segments)
+	{
+		Segment->UpdateSplineEditorMesh();
+	}
+
+	MarkRenderStateDirty();
+}
+
+bool FForeignWorldSplineData::IsEmpty()
+{
+	return ForeignControlPointDataMap.Num() == 0 && ForeignSplineSegmentDataMap.Num() == 0;
+}
+
+ULandscapeSplinesComponent* ULandscapeSplinesComponent::GetStreamingSplinesComponentByLocation(const FVector& LocalLocation, bool bCreate /* = true*/)
+{
+	ALandscapeProxy* OuterLandscape = Cast<ALandscapeProxy>(GetOwner());
+	if (OuterLandscape)
+	{
+		FVector LandscapeLocalLocation = ComponentToWorld.GetRelativeTransform(OuterLandscape->LandscapeActorToWorld()).TransformPosition(LocalLocation);
+		const int32 ComponentIndexX = (LandscapeLocalLocation.X >= 0.0f) ? FMath::FloorToInt(LandscapeLocalLocation.X / OuterLandscape->ComponentSizeQuads) : FMath::CeilToInt(LandscapeLocalLocation.X / OuterLandscape->ComponentSizeQuads);
+		const int32 ComponentIndexY = (LandscapeLocalLocation.Y >= 0.0f) ? FMath::FloorToInt(LandscapeLocalLocation.Y / OuterLandscape->ComponentSizeQuads) : FMath::CeilToInt(LandscapeLocalLocation.Y / OuterLandscape->ComponentSizeQuads);
+		ULandscapeComponent* LandscapeComponent = OuterLandscape->GetLandscapeInfo()->XYtoComponentMap.FindRef(FIntPoint(ComponentIndexX, ComponentIndexY));
+		if (LandscapeComponent)
+		{
+			ALandscapeProxy* ComponentLandscapeProxy = LandscapeComponent->GetLandscapeProxy();
+			if (!ComponentLandscapeProxy->SplineComponent && bCreate)
+			{
+				ComponentLandscapeProxy->Modify();
+				ComponentLandscapeProxy->SplineComponent = NewObject<ULandscapeSplinesComponent>(ComponentLandscapeProxy, NAME_None, RF_Transactional);
+				ComponentLandscapeProxy->SplineComponent->RelativeScale3D = RelativeScale3D;
+				ComponentLandscapeProxy->SplineComponent->AttachTo(ComponentLandscapeProxy->GetRootComponent());
+			}
+			if (ComponentLandscapeProxy->SplineComponent)
+			{
+				return ComponentLandscapeProxy->SplineComponent;
+			}
+		}
+	}
+
+	return this;
+}
+
+ULandscapeSplinesComponent* ULandscapeSplinesComponent::GetStreamingSplinesComponentForLevel(ULevel* Level, bool bCreate /* = true*/)
+{
+	ALandscapeProxy* OuterLandscape = Cast<ALandscapeProxy>(GetOwner());
+	if (OuterLandscape)
+	{
+		ULandscapeInfo* LandscapeInfo = OuterLandscape->GetLandscapeInfo();
+		check(LandscapeInfo);
+
+		ALandscapeProxy* Proxy = LandscapeInfo->GetLandscapeProxyForLevel(Level);
+		if (Proxy)
+		{
+			if (!Proxy->SplineComponent && bCreate)
+			{
+				Proxy->Modify();
+				Proxy->SplineComponent = NewObject<ULandscapeSplinesComponent>(Proxy, NAME_None, RF_Transactional);
+				Proxy->SplineComponent->RelativeScale3D = RelativeScale3D;
+				Proxy->SplineComponent->AttachTo(Proxy->GetRootComponent());
+			}
+			return Proxy->SplineComponent;
+		}
+	}
+
+	return nullptr;
+}
+
+TArray<ULandscapeSplinesComponent*> ULandscapeSplinesComponent::GetAllStreamingSplinesComponents()
+{
+	ALandscapeProxy* OuterLandscape = Cast<ALandscapeProxy>(GetOwner());
+	if (OuterLandscape)
+	{
+		ULandscapeInfo* LandscapeInfo = OuterLandscape->GetLandscapeInfo();
+		check(LandscapeInfo);
+
+		TArray<ULandscapeSplinesComponent*> SplinesComponents;
+		SplinesComponents.Reserve(LandscapeInfo->Proxies.Num() + 1);
+
+		ALandscape* RootLandscape = LandscapeInfo->LandscapeActor.Get();
+		if (RootLandscape && RootLandscape->SplineComponent)
+		{
+			SplinesComponents.Add(RootLandscape->SplineComponent);
+		}
+		for (ALandscapeProxy* LandscapeProxy : LandscapeInfo->Proxies)
+		{
+			if (LandscapeProxy && LandscapeProxy->SplineComponent)
+			{
+				SplinesComponents.Add(LandscapeProxy->SplineComponent);
+			}
+		}
+
+		return SplinesComponents;
+	}
+
+	return {};
+}
+
+void ULandscapeSplinesComponent::UpdateModificationKey(ULandscapeSplineSegment* Owner)
+{
+	UWorld* OwnerWorld = Owner->GetTypedOuter<UWorld>();
+	checkSlow(OwnerWorld != GetTypedOuter<UWorld>());
+
+	auto* ForeignWorldSplineData = ForeignWorldSplineDataMap.Find(OwnerWorld);
+	checkSlow(ForeignWorldSplineData);
+
+	if (ForeignWorldSplineData)
+	{
+		auto* ForeignSplineSegmentData = ForeignWorldSplineData->ForeignSplineSegmentDataMap.Find(Owner);
+		ForeignSplineSegmentData->ModificationKey = Owner->GetModificationKey();
+	}
+}
+
+void ULandscapeSplinesComponent::UpdateModificationKey(ULandscapeSplineControlPoint* Owner)
+{
+	UWorld* OwnerWorld = Owner->GetTypedOuter<UWorld>();
+	checkSlow(OwnerWorld != GetTypedOuter<UWorld>());
+
+	auto* ForeignWorldSplineData = ForeignWorldSplineDataMap.Find(OwnerWorld);
+	checkSlow(ForeignWorldSplineData);
+
+	if (ForeignWorldSplineData)
+	{
+		auto* ForeignControlPointData = ForeignWorldSplineData->ForeignControlPointDataMap.Find(Owner);
+		ForeignControlPointData->ModificationKey = Owner->GetModificationKey();
+	}
+}
+
+void ULandscapeSplinesComponent::AddForeignMeshComponent(ULandscapeSplineSegment* Owner, USplineMeshComponent* Component)
+{
+#if DO_GUARD_SLOW
+	UWorld* ThisOuterWorld = GetTypedOuter<UWorld>();
+	UWorld* ComponentOuterWorld = Component->GetTypedOuter<UWorld>();
+	checkSlow(ComponentOuterWorld == ThisOuterWorld);
+#endif
+
+	UWorld* OwnerWorld = Owner->GetTypedOuter<UWorld>();
+	checkSlow(OwnerWorld != ThisOuterWorld);
+
+	auto& ForeignWorldSplineData = ForeignWorldSplineDataMap.FindOrAdd(OwnerWorld);
+	auto& ForeignSplineSegmentData = ForeignWorldSplineData.ForeignSplineSegmentDataMap.FindOrAdd(Owner);
+	ForeignSplineSegmentData.MeshComponents.Add(Component);
+	ForeignSplineSegmentData.ModificationKey = Owner->GetModificationKey();
+
+	MeshComponentForeignOwnersMap.Add(Component, Owner);
+}
+
+void ULandscapeSplinesComponent::RemoveForeignMeshComponent(ULandscapeSplineSegment* Owner, USplineMeshComponent* Component)
+{
+#if DO_GUARD_SLOW
+	UWorld* ThisOuterWorld = GetTypedOuter<UWorld>();
+	UWorld* ComponentOuterWorld = Component->GetTypedOuter<UWorld>();
+	checkSlow(ComponentOuterWorld == ThisOuterWorld);
+#endif
+
+	UWorld* OwnerWorld = Owner->GetTypedOuter<UWorld>();
+	checkSlow(OwnerWorld != ThisOuterWorld);
+
+	auto* ForeignWorldSplineData = ForeignWorldSplineDataMap.Find(OwnerWorld);
+	checkSlow(ForeignWorldSplineData);
+	checkSlow(MeshComponentForeignOwnersMap.FindRef(Component) == Owner);
+	verifySlow(MeshComponentForeignOwnersMap.Remove(Component) == 1);
+
+	if (ForeignWorldSplineData)
+	{
+		auto* ForeignSplineSegmentData = ForeignWorldSplineData->ForeignSplineSegmentDataMap.Find(Owner);
+		verifySlow(ForeignSplineSegmentData->MeshComponents.RemoveSingle(Component) == 1);
+		if (ForeignSplineSegmentData->MeshComponents.Num() == 0)
+		{
+			verifySlow(ForeignWorldSplineData->ForeignSplineSegmentDataMap.Remove(Owner) == 1);
+			if (ForeignWorldSplineData->IsEmpty())
+			{
+				verifySlow(ForeignWorldSplineDataMap.Remove(OwnerWorld) == 1);
+			}
+		}
+		else
+		{
+			ForeignSplineSegmentData->ModificationKey = Owner->GetModificationKey();
+		}
+	}
+}
+
+void ULandscapeSplinesComponent::RemoveAllForeignMeshComponents(ULandscapeSplineSegment* Owner)
+{
+	UWorld* OwnerWorld = Owner->GetTypedOuter<UWorld>();
+	checkSlow(OwnerWorld != GetTypedOuter<UWorld>());
+
+	auto* ForeignWorldSplineData = ForeignWorldSplineDataMap.Find(OwnerWorld);
+	checkSlow(ForeignWorldSplineData);
+
+	if (ForeignWorldSplineData)
+	{
+		auto* ForeignSplineSegmentData = ForeignWorldSplineData->ForeignSplineSegmentDataMap.Find(Owner);
+
+		for (auto* MeshComponent : ForeignSplineSegmentData->MeshComponents)
+		{
+			checkSlow(MeshComponentForeignOwnersMap.FindRef(MeshComponent) == Owner);
+			verifySlow(MeshComponentForeignOwnersMap.Remove(MeshComponent) == 1);
+		}
+		ForeignSplineSegmentData->MeshComponents.Empty();
+		verifySlow(ForeignWorldSplineData->ForeignSplineSegmentDataMap.Remove(Owner) == 1);
+		if (ForeignWorldSplineData->IsEmpty())
+		{
+			verifySlow(ForeignWorldSplineDataMap.Remove(OwnerWorld) == 1);
+		}
+	}
+}
+
+void ULandscapeSplinesComponent::AddForeignMeshComponent(ULandscapeSplineControlPoint* Owner, UControlPointMeshComponent* Component)
+{
+#if DO_GUARD_SLOW
+	UWorld* ThisOuterWorld = GetTypedOuter<UWorld>();
+	UWorld* ComponentOuterWorld = Component->GetTypedOuter<UWorld>();
+	checkSlow(ComponentOuterWorld == ThisOuterWorld);
+#endif
+
+	UWorld* OwnerWorld = Owner->GetTypedOuter<UWorld>();
+	checkSlow(OwnerWorld != ThisOuterWorld);
+
+	auto& ForeignWorldSplineData = ForeignWorldSplineDataMap.FindOrAdd(OwnerWorld);
+	checkSlow(!ForeignWorldSplineData.ForeignControlPointDataMap.Find(Owner));
+	auto& ForeignControlPointData = ForeignWorldSplineData.ForeignControlPointDataMap.Add(Owner);
+
+	ForeignControlPointData.MeshComponent = Component;
+	ForeignControlPointData.ModificationKey = Owner->GetModificationKey();
+
+	MeshComponentForeignOwnersMap.Add(Component, Owner);
+}
+
+void ULandscapeSplinesComponent::RemoveForeignMeshComponent(ULandscapeSplineControlPoint* Owner, UControlPointMeshComponent* Component)
+{
+#if DO_GUARD_SLOW
+	UWorld* ThisOuterWorld = GetTypedOuter<UWorld>();
+	UWorld* ComponentOuterWorld = Component->GetTypedOuter<UWorld>();
+	checkSlow(ComponentOuterWorld == ThisOuterWorld);
+#endif
+
+	UWorld* OwnerWorld = Owner->GetTypedOuter<UWorld>();
+	checkSlow(OwnerWorld != ThisOuterWorld);
+
+	auto* ForeignWorldSplineData = ForeignWorldSplineDataMap.Find(OwnerWorld);
+	checkSlow(ForeignWorldSplineData);
+	checkSlow(MeshComponentForeignOwnersMap.FindRef(Component) == Owner);
+	verifySlow(MeshComponentForeignOwnersMap.Remove(Component) == 1);
+
+	if (ForeignWorldSplineData)
+	{
+		auto* ForeignControlPointData = ForeignWorldSplineData->ForeignControlPointDataMap.Find(Owner);
+		checkSlow(ForeignControlPointData);
+		checkSlow(ForeignControlPointData->MeshComponent == Component);
+
+		verifySlow(ForeignWorldSplineData->ForeignControlPointDataMap.Remove(Owner) == 1);
+		if (ForeignWorldSplineData->IsEmpty())
+		{
+			verifySlow(ForeignWorldSplineDataMap.Remove(OwnerWorld) == 1);
+		}
+	}
+}
+
+void ULandscapeSplinesComponent::DestroyOrphanedForeignMeshComponents(UWorld* OwnerWorld)
+{
+	auto* ForeignWorldSplineData = ForeignWorldSplineDataMap.Find(OwnerWorld);
+
+	if (ForeignWorldSplineData)
+	{
+		for (auto ForeignSplineSegmentDataIt = ForeignWorldSplineData->ForeignSplineSegmentDataMap.CreateIterator(); ForeignSplineSegmentDataIt; ++ForeignSplineSegmentDataIt)
+		{
+			const auto& ForeignSplineSegment = ForeignSplineSegmentDataIt->Key;
+			auto& ForeignSplineSegmentData = ForeignSplineSegmentDataIt->Value;
+
+			if (!ForeignSplineSegment)
+			{
+				for (auto* MeshComponent : ForeignSplineSegmentData.MeshComponents)
+				{
+					checkSlow(!MeshComponentForeignOwnersMap.FindRef(MeshComponent).IsValid());
+					verifySlow(MeshComponentForeignOwnersMap.Remove(MeshComponent) == 1);
+					MeshComponent->DestroyComponent();
+				}
+				ForeignSplineSegmentData.MeshComponents.Empty();
+				ForeignSplineSegmentDataIt.RemoveCurrent();
+			}
+		}
+
+		if (ForeignWorldSplineData->IsEmpty())
+		{
+			verifySlow(ForeignWorldSplineDataMap.Remove(OwnerWorld) == 1);
+		}
+	}
+}
 
-		for (ULandscapeSplineControlPoint* ControlPoint : ControlPoints)
-		{
-			ControlPoint->UpdateSplinePoints(bUpdateCollision);
-		}
-		for (ULandscapeSplineSegment* Segment : Segments)
+UControlPointMeshComponent* ULandscapeSplinesComponent::GetForeignMeshComponent(ULandscapeSplineControlPoint* Owner)
+{
+	UWorld* OwnerWorld = Owner->GetTypedOuter<UWorld>();
+	checkSlow(OwnerWorld != GetTypedOuter<UWorld>());
+
+	auto* ForeignWorldSplineData = ForeignWorldSplineDataMap.Find(OwnerWorld);
+	if (ForeignWorldSplineData)
+	{
+		auto* ForeignControlPointData = ForeignWorldSplineData->ForeignControlPointDataMap.Find(Owner);
+		if (ForeignControlPointData)
 		{
-			Segment->UpdateSplinePoints(bUpdateCollision);
+			return ForeignControlPointData->MeshComponent;
 		}
 	}
+
+	return nullptr;
 }
-void ULandscapeSplinesComponent::PostEditUndo()
+
+TArray<USplineMeshComponent*> ULandscapeSplinesComponent::GetForeignMeshComponents(ULandscapeSplineSegment* Owner)
 {
-	bHackIsUndoingSplines = true;
-	Super::PostEditUndo();
-	bHackIsUndoingSplines = false;
+	UWorld* OwnerWorld = Owner->GetTypedOuter<UWorld>();
+	checkSlow(OwnerWorld != GetTypedOuter<UWorld>());
 
-	MarkRenderStateDirty();
+	auto* ForeignWorldSplineData = ForeignWorldSplineDataMap.Find(OwnerWorld);
+	if (ForeignWorldSplineData)
+	{
+		auto* ForeignSplineSegmentData = ForeignWorldSplineData->ForeignSplineSegmentDataMap.Find(Owner);
+		if (ForeignSplineSegmentData)
+		{
+			return ForeignSplineSegmentData->MeshComponents;
+		}
+	}
+
+	return {};
 }
 
-void ULandscapeSplinesComponent::ShowSplineEditorMesh(bool bShow)
+UObject* ULandscapeSplinesComponent::GetOwnerForMeshComponent(const UMeshComponent* SplineMeshComponent)
 {
-	bShowSplineEditorMesh = bShow;
+	UObject* LocalOwner = MeshComponentLocalOwnersMap.FindRef(SplineMeshComponent);
+	if (LocalOwner)
+	{
+		return LocalOwner;
+	}
 
-	for (ULandscapeSplineSegment* Segment : Segments)
+	TLazyObjectPtr<UObject>* ForeignOwner = MeshComponentForeignOwnersMap.Find(SplineMeshComponent);
+	if (ForeignOwner)
 	{
-		Segment->UpdateSplineEditorMesh();
+		// this will be null if ForeignOwner isn't currently loaded
+		return ForeignOwner->Get();
 	}
 
-	MarkRenderStateDirty();
+	return nullptr;
 }
 #endif // WITH_EDITOR
 
@@ -523,7 +1108,7 @@ ULandscapeSplineControlPoint::ULandscapeSplineControlPoint(const FObjectInitiali
 	EndFalloff = 2000;
 
 #if WITH_EDITORONLY_DATA
-	Mesh = NULL;
+	Mesh = nullptr;
 	MeshScale = FVector(1);
 
 	LDMaxDrawDistance = 0;
@@ -533,7 +1118,8 @@ ULandscapeSplineControlPoint::ULandscapeSplineControlPoint(const FObjectInitiali
 	bRaiseTerrain = true;
 	bLowerTerrain = true;
 
-	MeshComponent = NULL;
+	LocalMeshComponent = nullptr;
+	bPlaceSplineMeshesInStreamingLevels = true;
 	bEnableCollision = true;
 	bCastShadow = true;
 
@@ -542,21 +1128,45 @@ ULandscapeSplineControlPoint::ULandscapeSplineControlPoint(const FObjectInitiali
 #endif
 }
 
+void ULandscapeSplineControlPoint::Serialize(FArchive& Ar)
+{
+	Super::Serialize(Ar);
+
+#if WITH_EDITOR
+	if (Ar.UE4Ver() < VER_UE4_LANDSCAPE_SPLINE_CROSS_LEVEL_MESHES)
+	{
+		bPlaceSplineMeshesInStreamingLevels = false;
+	}
+#endif
+}
+
 void ULandscapeSplineControlPoint::PostLoad()
 {
 	Super::PostLoad();
 
 #if WITH_EDITOR
-	if (MeshComponent != NULL)
+	if (GIsEditor)
 	{
-		// Fix collision profile
-		const FName CollisionProfile = bEnableCollision ? UCollisionProfile::BlockAll_ProfileName : UCollisionProfile::NoCollision_ProfileName;
-		if (MeshComponent->GetCollisionProfileName() != CollisionProfile)
+		if (LocalMeshComponent != nullptr)
 		{
-			MeshComponent->SetCollisionProfileName(CollisionProfile);
+			ULandscapeSplinesComponent* OuterSplines = GetOuterULandscapeSplinesComponent();
+			OuterSplines->MeshComponentLocalOwnersMap.Add(LocalMeshComponent, this);
 		}
+	}
+
+	if (GetLinkerUE4Version() < VER_UE4_LANDSCAPE_SPLINE_CROSS_LEVEL_MESHES)
+	{
+		// Fix collision profile
+		if (LocalMeshComponent != nullptr) // ForeignMeshComponents didn't exist yet
+		{
+			const FName CollisionProfile = bEnableCollision ? UCollisionProfile::BlockAll_ProfileName : UCollisionProfile::NoCollision_ProfileName;
+			if (LocalMeshComponent->GetCollisionProfileName() != CollisionProfile)
+			{
+				LocalMeshComponent->SetCollisionProfileName(CollisionProfile);
+			}
 
-		MeshComponent->SetFlags(RF_TextExportTransient);
+			LocalMeshComponent->SetFlags(RF_TextExportTransient);
+		}
 	}
 #endif
 }
@@ -571,20 +1181,14 @@ FLandscapeSplineSegmentConnection& FLandscapeSplineConnection::GetFarConnection(
 	return Segment->Connections[1 - End];
 }
 
+#if WITH_EDITOR
 FName ULandscapeSplineControlPoint::GetBestConnectionTo(FVector Destination) const
 {
 	FName BestSocket = NAME_None;
 	float BestScore = -FLT_MAX;
 
-#if WITH_EDITORONLY_DATA
-	if (Mesh != NULL)
+	if (Mesh != nullptr)
 	{
-#else
-	if (MeshComponent != NULL)
-	{
-		const UStaticMesh* Mesh = MeshComponent->StaticMesh;
-		const FVector MeshScale = MeshComponent->RelativeScale3D;
-#endif
 		for (const UStaticMeshSocket* Socket : Mesh->Sockets)
 		{
 			FTransform SocketTransform = FTransform(Socket->RelativeRotation, Socket->RelativeLocation) * FTransform(Rotation, Location, MeshScale);
@@ -610,16 +1214,10 @@ void ULandscapeSplineControlPoint::GetConnectionLocalLocationAndRotation(FName S
 	OutLocation = FVector::ZeroVector;
 	OutRotation = FRotator::ZeroRotator;
 
-#if WITH_EDITORONLY_DATA
-	if (Mesh != NULL)
-	{
-#else
-	if (MeshComponent != NULL)
+	if (Mesh != nullptr)
 	{
-		UStaticMesh* Mesh = MeshComponent->StaticMesh;
-#endif
 		const UStaticMeshSocket* Socket = Mesh->FindSocket(SocketName);
-		if (Socket != NULL)
+		if (Socket != nullptr)
 		{
 			OutLocation = Socket->RelativeLocation;
 			OutRotation = Socket->RelativeRotation;
@@ -632,17 +1230,10 @@ void ULandscapeSplineControlPoint::GetConnectionLocationAndRotation(FName Socket
 	OutLocation = Location;
 	OutRotation = Rotation;
 
-#if WITH_EDITORONLY_DATA
-	if (Mesh != NULL)
+	if (Mesh != nullptr)
 	{
-#else
-	if (MeshComponent != NULL)
-	{
-		UStaticMesh* Mesh = MeshComponent->StaticMesh;
-		const FVector MeshScale = MeshComponent->RelativeScale3D;
-#endif
 		const UStaticMeshSocket* Socket = Mesh->FindSocket(SocketName);
-		if (Socket != NULL)
+		if (Socket != nullptr)
 		{
 			FTransform SocketTransform = FTransform(Socket->RelativeRotation, Socket->RelativeLocation) * FTransform(Rotation, Location, MeshScale);
 			OutLocation = SocketTransform.GetTranslation();
@@ -651,14 +1242,22 @@ void ULandscapeSplineControlPoint::GetConnectionLocationAndRotation(FName Socket
 	}
 }
 
-#if WITH_EDITOR
 void ULandscapeSplineControlPoint::SetSplineSelected(bool bInSelected)
 {
 	bSelected = bInSelected;
 	GetOuterULandscapeSplinesComponent()->MarkRenderStateDirty();
 
-	if (MeshComponent != NULL)
+	if (LocalMeshComponent != nullptr)
+	{
+		LocalMeshComponent->bSelected = bInSelected;
+		LocalMeshComponent->PushSelectionToProxy();
+	}
+
+	auto ForeignMeshComponentsMap = GetForeignMeshComponents();
+	for (auto& ForeignMeshComponentsPair : ForeignMeshComponentsMap)
 	{
+		ULandscapeSplinesComponent* MeshComponentOuterSplines = ForeignMeshComponentsPair.Key;
+		auto* MeshComponent = ForeignMeshComponentsPair.Value;
 		MeshComponent->bSelected = bInSelected;
 		MeshComponent->PushSelectionToProxy();
 	}
@@ -716,8 +1315,8 @@ void ULandscapeSplineControlPoint::AutoSetConnections(bool bIncludingValid)
 	{
 		FLandscapeSplineSegmentConnection& NearConnection = Connection.GetNearConnection();
 		if (bIncludingValid ||
-			(Mesh != NULL && Mesh->FindSocket(NearConnection.SocketName) == NULL) ||
-			(Mesh == NULL && NearConnection.SocketName != NAME_None))
+			(Mesh != nullptr && Mesh->FindSocket(NearConnection.SocketName) == nullptr) ||
+			(Mesh == nullptr && NearConnection.SocketName != NAME_None))
 		{
 			FLandscapeSplineSegmentConnection& FarConnection = Connection.GetFarConnection();
 			FVector EndLocation; FRotator EndRotation;
@@ -742,85 +1341,95 @@ void ULandscapeSplineControlPoint::AutoSetConnections(bool bIncludingValid)
 }
 #endif
 
-void ULandscapeSplineControlPoint::RegisterComponents()
+#if WITH_EDITOR
+TMap<ULandscapeSplinesComponent*, UControlPointMeshComponent*> ULandscapeSplineControlPoint::GetForeignMeshComponents()
 {
-	if (MeshComponent != NULL)
-	{
-		MeshComponent->RegisterComponent();
-	}
-}
+	TMap<ULandscapeSplinesComponent*, UControlPointMeshComponent*> ForeignMeshComponentsMap;
 
-void ULandscapeSplineControlPoint::UnregisterComponents()
-{
-	if (MeshComponent != NULL)
+	ULandscapeSplinesComponent* OuterSplines = GetOuterULandscapeSplinesComponent();
+	TArray<ULandscapeSplinesComponent*> SplineComponents = OuterSplines->GetAllStreamingSplinesComponents();
+
+	for (ULandscapeSplinesComponent* SplineComponent : SplineComponents)
 	{
-		MeshComponent->UnregisterComponent();
+		if (SplineComponent != OuterSplines)
+		{
+			auto* ForeignMeshComponent = SplineComponent->GetForeignMeshComponent(this);
+			if (ForeignMeshComponent)
+			{
+				ForeignMeshComponent->Modify();
+				ForeignMeshComponentsMap.Add(SplineComponent, ForeignMeshComponent);
+			}
+		}
 	}
-}
 
-bool ULandscapeSplineControlPoint::OwnsComponent(const class UControlPointMeshComponent* StaticMeshComponent) const
-{
-	return MeshComponent == StaticMeshComponent;
+	return ForeignMeshComponentsMap;
 }
 
-#if WITH_EDITOR
 void ULandscapeSplineControlPoint::UpdateSplinePoints(bool bUpdateCollision, bool bUpdateAttachedSegments)
 {
+	Modify();
+
 	ULandscapeSplinesComponent* OuterSplines = GetOuterULandscapeSplinesComponent();
 
-	if ((MeshComponent == NULL && Mesh != NULL) ||
-		(MeshComponent != NULL &&
-			(MeshComponent->StaticMesh != Mesh ||
-			MeshComponent->RelativeLocation != Location ||
-			MeshComponent->RelativeRotation != Rotation ||
-			MeshComponent->RelativeScale3D  != MeshScale))
-		)
-	{
-		OuterSplines->Modify();
+	auto ForeignMeshComponentsMap = GetForeignMeshComponents();
+
+	ModificationKey = FGuid::NewGuid();
+
+	UControlPointMeshComponent* MeshComponent = LocalMeshComponent;
+	ULandscapeSplinesComponent* MeshComponentOuterSplines = OuterSplines;
 
-		if (MeshComponent)
+	if (Mesh != nullptr)
+	{
+		// Attempt to place mesh components into the appropriate landscape streaming levels based on the components under the spline
+		if (bPlaceSplineMeshesInStreamingLevels)
 		{
-			MeshComponent->Modify();
-			MeshComponent->UnregisterComponent();
+			MeshComponentOuterSplines = OuterSplines->GetStreamingSplinesComponentByLocation(Location);
+
+			if (MeshComponentOuterSplines != OuterSplines)
+			{
+				MeshComponent = MeshComponentOuterSplines->GetForeignMeshComponent(this);
+				MeshComponentOuterSplines->Modify();
+				MeshComponentOuterSplines->UpdateModificationKey(this);
+			}
 		}
 
-		if (Mesh != NULL)
+		// Create mesh component if needed
+		if (MeshComponent == nullptr)
 		{
-			if (MeshComponent == NULL)
+			AActor* MeshComponentOuterActor = MeshComponentOuterSplines->GetOwner();
+			MeshComponentOuterSplines->Modify();
+			MeshComponentOuterActor->Modify();
+			MeshComponent = NewObject<UControlPointMeshComponent>(MeshComponentOuterActor, NAME_None, RF_Transactional | RF_TextExportTransient);
+			MeshComponent->bSelected = bSelected;
+			MeshComponent->AttachTo(MeshComponentOuterSplines);
+			if (MeshComponentOuterSplines == OuterSplines)
 			{
-				Modify();
-				OuterSplines->GetOuter()->Modify();
-				MeshComponent = NewObject<UControlPointMeshComponent>(OuterSplines->GetOuter(), NAME_None, RF_Transactional | RF_TextExportTransient);
-				MeshComponent->bSelected = bSelected;
-				MeshComponent->AttachTo(OuterSplines);
+				MeshComponentOuterSplines->MeshComponentLocalOwnersMap.Add(MeshComponent, this);
 			}
-
-			if (MeshComponent->StaticMesh != Mesh)
+			else
 			{
-				MeshComponent->SetStaticMesh(Mesh);
-
-				AutoSetConnections(false);
+				MeshComponentOuterSplines->AddForeignMeshComponent(this, MeshComponent);
+				ForeignWorld = MeshComponentOuterSplines->GetTypedOuter<UWorld>();
 			}
+		}
 
-			MeshComponent->RelativeLocation = Location;
-			MeshComponent->RelativeRotation = Rotation;
-			MeshComponent->RelativeScale3D  = MeshScale;
+		if (MeshComponent->RelativeLocation != Location ||
+			MeshComponent->RelativeRotation != Rotation ||
+			MeshComponent->RelativeScale3D != MeshScale)
+		{
+			MeshComponent->Modify();
+			MeshComponent->SetRelativeTransform(FTransform(Rotation, Location, MeshScale));
 			MeshComponent->InvalidateLightingCache();
-			MeshComponent->RegisterComponent();
 		}
-		else
+
+		if (MeshComponent->StaticMesh != Mesh)
 		{
-			Modify();
-			MeshComponent->GetOuter()->Modify();
-			MeshComponent->DestroyComponent();
-			MeshComponent = NULL;
+			MeshComponent->Modify();
+			MeshComponent->SetStaticMesh(Mesh);
 
 			AutoSetConnections(false);
 		}
-	}
 
-	if (MeshComponent != NULL)
-	{
 		if (MeshComponent->OverrideMaterials != MaterialOverrides)
 		{
 			MeshComponent->Modify();
@@ -860,8 +1469,43 @@ void ULandscapeSplineControlPoint::UpdateSplinePoints(bool bUpdateCollision, boo
 			MeshComponent->BodyInstance.SetCollisionProfileName(CollisionProfile);
 		}
 	}
+	else
+	{
+		MeshComponent = nullptr;
+		ForeignWorld = NULL;
+	}
+
+	// Destroy any unused components
+	bool bDestroyedAnyComponents = false;
+	if (LocalMeshComponent && LocalMeshComponent != MeshComponent)
+	{
+		OuterSplines->Modify();
+		LocalMeshComponent->Modify();
+		checkSlow(OuterSplines->MeshComponentLocalOwnersMap.FindRef(LocalMeshComponent) == this);
+		verifySlow(OuterSplines->MeshComponentLocalOwnersMap.Remove(LocalMeshComponent) == 1);
+		LocalMeshComponent->DestroyComponent();
+		LocalMeshComponent = nullptr;
+	}
+	for (auto& ForeignMeshComponentsPair : ForeignMeshComponentsMap)
+	{
+		ULandscapeSplinesComponent* MeshComponentOuterSplines = ForeignMeshComponentsPair.Key;
+		auto* ForeignMeshComponent = ForeignMeshComponentsPair.Value;
+		if (ForeignMeshComponent != MeshComponent)
+		{
+			MeshComponentOuterSplines->Modify();
+			ForeignMeshComponent->Modify();
+			MeshComponentOuterSplines->RemoveForeignMeshComponent(this, ForeignMeshComponent);
+			ForeignMeshComponent->DestroyComponent();
+		}
+	}
+	ForeignMeshComponentsMap.Empty();
+	if (bDestroyedAnyComponents)
+	{
+		AutoSetConnections(false);
+	}
 
-	if (Mesh != NULL)
+	// Update "Points" array
+	if (Mesh != nullptr)
 	{
 		Points.Reset(ConnectedSegments.Num());
 
@@ -878,7 +1522,7 @@ void ULandscapeSplineControlPoint::UpdateSplinePoints(bool bUpdateCollision, boo
 			const FVector FalloffLeftPos = StartLocation - BiNormal * (Width + SideFalloff);
 			const FVector FalloffRightPos = StartLocation + BiNormal * (Width + SideFalloff);
 
-			new(Points) FLandscapeSplineInterpPoint(StartLocation, LeftPos, RightPos, FalloffLeftPos, FalloffRightPos, 1.0f);
+			Points.Emplace(StartLocation, LeftPos, RightPos, FalloffLeftPos, FalloffRightPos, 1.0f);
 		}
 
 		const FVector CPLocation = Location;
@@ -899,7 +1543,7 @@ void ULandscapeSplineControlPoint::UpdateSplinePoints(bool bUpdateCollision, boo
 		const FVector FalloffLeftPos = StartLocation - BiNormal * (Width + SideFalloff);
 		const FVector FalloffRightPos = StartLocation + BiNormal * (Width + SideFalloff);
 
-		new(Points) FLandscapeSplineInterpPoint(StartLocation, LeftPos, RightPos, FalloffLeftPos, FalloffRightPos, 1.0f);
+		Points.Emplace(StartLocation, LeftPos, RightPos, FalloffLeftPos, FalloffRightPos, 1.0f);
 	}
 
 	// Update Bounds
@@ -932,11 +1576,24 @@ void ULandscapeSplineControlPoint::DeleteSplinePoints()
 
 	OuterSplines->MarkRenderStateDirty();
 
-	if (MeshComponent != NULL)
+	if (LocalMeshComponent != nullptr)
+	{
+		OuterSplines->Modify();
+		checkSlow(OuterSplines->MeshComponentLocalOwnersMap.FindRef(LocalMeshComponent) == this);
+		verifySlow(OuterSplines->MeshComponentLocalOwnersMap.Remove(LocalMeshComponent) == 1);
+		LocalMeshComponent->DestroyComponent();
+		LocalMeshComponent = nullptr;
+	}
+
+	auto ForeignMeshComponentsMap = GetForeignMeshComponents();
+	for (auto& ForeignMeshComponentsPair : ForeignMeshComponentsMap)
 	{
-		MeshComponent->GetOuter()->Modify();
+		ULandscapeSplinesComponent* MeshComponentOuterSplines = ForeignMeshComponentsPair.Key;
+		MeshComponentOuterSplines->Modify();
+		auto* MeshComponent = ForeignMeshComponentsPair.Value;
+		MeshComponent->Modify();
+		MeshComponentOuterSplines->RemoveForeignMeshComponent(this, MeshComponent);
 		MeshComponent->DestroyComponent();
-		MeshComponent = NULL;
 	}
 }
 
@@ -953,16 +1610,17 @@ void ULandscapeSplineControlPoint::PostDuplicate(bool bDuplicateForPIE)
 {
 	if (!bDuplicateForPIE)
 	{
-		ULandscapeSplinesComponent* OuterSplines = CastChecked<ULandscapeSplinesComponent>(GetOuter());
-
-		if (MeshComponent != NULL)
+		// if we get duplicated but our local mesh doesn't, then clear our reference to the mesh - it's not ours
+		if (LocalMeshComponent != nullptr)
 		{
-			if (MeshComponent->GetOuter() != OuterSplines->GetOuter())
+			ULandscapeSplinesComponent* OuterSplines = CastChecked<ULandscapeSplinesComponent>(GetOuter());
+			if (LocalMeshComponent->GetOuter() != OuterSplines->GetOwner())
 			{
-				MeshComponent = DuplicateObject<UControlPointMeshComponent>(MeshComponent, OuterSplines->GetOuter(), *MeshComponent->GetName());
-				MeshComponent->AttachTo(OuterSplines);
+				LocalMeshComponent = nullptr;
 			}
 		}
+
+		UpdateSplinePoints();
 	}
 
 	Super::PostDuplicate(bDuplicateForPIE);
@@ -1001,9 +1659,9 @@ ULandscapeSplineSegment::ULandscapeSplineSegment(const FObjectInitializer& Objec
 	: Super(ObjectInitializer)
 {
 
-	Connections[0].ControlPoint = NULL;
+	Connections[0].ControlPoint = nullptr;
 	Connections[0].TangentLen = 0;
-	Connections[1].ControlPoint = NULL;
+	Connections[1].ControlPoint = nullptr;
 	Connections[1].TangentLen = 0;
 
 #if WITH_EDITORONLY_DATA
@@ -1015,6 +1673,7 @@ ULandscapeSplineSegment::ULandscapeSplineSegment(const FObjectInitializer& Objec
 	SplineMeshes.Empty();
 	LDMaxDrawDistance = 0;
 	TranslucencySortPriority = 0;
+	bPlaceSplineMeshesInStreamingLevels = true;
 	bEnableCollision = true;
 	bCastShadow = true;
 
@@ -1058,6 +1717,11 @@ void ULandscapeSplineSegment::Serialize(FArchive& Ar)
 			}
 		}
 	}
+
+	if (Ar.UE4Ver() < VER_UE4_LANDSCAPE_SPLINE_CROSS_LEVEL_MESHES)
+	{
+		bPlaceSplineMeshesInStreamingLevels = false;
+	}
 #endif
 }
 
@@ -1066,40 +1730,52 @@ void ULandscapeSplineSegment::PostLoad()
 	Super::PostLoad();
 
 #if WITH_EDITOR
-	if (GIsEditor &&
-		GetLinkerUE4Version() < VER_UE4_ADDED_LANDSCAPE_SPLINE_EDITOR_MESH &&
-		MeshComponents.Num() == 0)
-	{
-		UpdateSplinePoints();
-	}
-
 	if (GIsEditor)
 	{
-		// Regenerate spline mesh components if any of the meshes used have gone missing
-		// Otherwise the spline will have no editor mesh and won't be selectable
+		if (GetLinkerUE4Version() < VER_UE4_ADDED_LANDSCAPE_SPLINE_EDITOR_MESH &&
+			LocalMeshComponents.Num() == 0) // ForeignMeshComponents didn't exist yet
+		{
+			UpdateSplinePoints();
+		}
+
+		// Replace null meshes with the editor mesh
+		// Otherwise the spline will have no mesh and won't be easily selectable
 		ULandscapeSplinesComponent* OuterSplines = GetOuterULandscapeSplinesComponent();
-		for (auto* MeshComponent : MeshComponents)
+		if (OuterSplines->SplineEditorMesh != nullptr)
 		{
-			MeshComponent->ConditionalPostLoad();
-			if (MeshComponent->StaticMesh == NULL)
+			for (auto* MeshComponent : LocalMeshComponents)
 			{
-				UpdateSplinePoints();
-				break;
+				if (MeshComponent->StaticMesh == nullptr)
+				{
+					MeshComponent->ConditionalPostLoad();
+					MeshComponent->SetStaticMesh(OuterSplines->SplineEditorMesh);
+					MeshComponent->SetHiddenInGame(true);
+					MeshComponent->SetVisibility(OuterSplines->bShowSplineEditorMesh);
+					MeshComponent->BodyInstance.SetCollisionProfileName(UCollisionProfile::NoCollision_ProfileName);
+				}
 			}
 		}
+
+		for (auto* MeshComponent : LocalMeshComponents)
+		{
+			OuterSplines->MeshComponentLocalOwnersMap.Add(MeshComponent, this);
+		}
 	}
 
-	for (auto* MeshComponent : MeshComponents)
+	if (GetLinkerUE4Version() < VER_UE4_LANDSCAPE_SPLINE_CROSS_LEVEL_MESHES)
 	{
 		// Fix collision profile
-		const bool bUsingEditorMesh = MeshComponent->bHiddenInGame;
-		const FName CollisionProfile = (bEnableCollision && !bUsingEditorMesh) ? UCollisionProfile::BlockAll_ProfileName : UCollisionProfile::NoCollision_ProfileName;
-		if (MeshComponent->GetCollisionProfileName() != CollisionProfile)
+		for (auto* MeshComponent : LocalMeshComponents) // ForeignMeshComponents didn't exist yet
 		{
-			MeshComponent->SetCollisionProfileName(CollisionProfile);
-		}
+			const bool bUsingEditorMesh = MeshComponent->bHiddenInGame;
+			const FName CollisionProfile = (bEnableCollision && !bUsingEditorMesh) ? UCollisionProfile::BlockAll_ProfileName : UCollisionProfile::NoCollision_ProfileName;
+			if (MeshComponent->GetCollisionProfileName() != CollisionProfile)
+			{
+				MeshComponent->SetCollisionProfileName(CollisionProfile);
+			}
 
-		MeshComponent->SetFlags(RF_TextExportTransient);
+			MeshComponent->SetFlags(RF_TextExportTransient);
+		}
 	}
 #endif
 }
@@ -1111,11 +1787,22 @@ void ULandscapeSplineSegment::SetSplineSelected(bool bInSelected)
 	bSelected = bInSelected;
 	GetOuterULandscapeSplinesComponent()->MarkRenderStateDirty();
 
-	for (auto* MeshComponent : MeshComponents)
+	for (auto* MeshComponent : LocalMeshComponents)
 	{
 		MeshComponent->bSelected = bInSelected;
 		MeshComponent->PushSelectionToProxy();
 	}
+
+	auto ForeignMeshComponentsMap = GetForeignMeshComponents();
+	for (auto& ForeignMeshComponentsPair : ForeignMeshComponentsMap)
+	{
+		ULandscapeSplinesComponent* MeshComponentOuterSplines = ForeignMeshComponentsPair.Key;
+		for (auto* MeshComponent : ForeignMeshComponentsPair.Value)
+		{
+			MeshComponent->bSelected = bInSelected;
+			MeshComponent->PushSelectionToProxy();
+		}
+	}
 }
 
 void ULandscapeSplineSegment::AutoFlipTangents()
@@ -1137,27 +1824,6 @@ void ULandscapeSplineSegment::AutoFlipTangents()
 }
 #endif
 
-void ULandscapeSplineSegment::RegisterComponents()
-{
-	for (auto* MeshComponent : MeshComponents)
-	{
-		MeshComponent->RegisterComponent();
-	}
-}
-
-void ULandscapeSplineSegment::UnregisterComponents()
-{
-	for (auto* MeshComponent : MeshComponents)
-	{
-		MeshComponent->UnregisterComponent();
-	}
-}
-
-bool ULandscapeSplineSegment::OwnsComponent(const class USplineMeshComponent* SplineMeshComponent) const
-{
-	return MeshComponents.FindByKey(SplineMeshComponent) != NULL;
-}
-
 static float ApproxLength(const FInterpCurveVector& SplineInfo, const float Start = 0.0f, const float End = 1.0f, const int32 ApproxSections = 4)
 {
 	float SplineLength = 0;
@@ -1180,21 +1846,47 @@ static ESplineMeshAxis::Type CrossAxis(ESplineMeshAxis::Type InForwardAxis, ESpl
 
 bool FLandscapeSplineMeshEntry::IsValid() const
 {
-	return Mesh != NULL && ForwardAxis != UpAxis && Scale.GetAbsMin() > KINDA_SMALL_NUMBER;
+	return Mesh != nullptr && ForwardAxis != UpAxis && Scale.GetAbsMin() > KINDA_SMALL_NUMBER;
 }
 
 #if WITH_EDITOR
+TMap<ULandscapeSplinesComponent*, TArray<USplineMeshComponent*>> ULandscapeSplineSegment::GetForeignMeshComponents()
+{
+	TMap<ULandscapeSplinesComponent*, TArray<USplineMeshComponent*>> ForeignMeshComponentsMap;
+
+	ULandscapeSplinesComponent* OuterSplines = GetOuterULandscapeSplinesComponent();
+	TArray<ULandscapeSplinesComponent*> SplineComponents = OuterSplines->GetAllStreamingSplinesComponents();
+
+	for (ULandscapeSplinesComponent* SplineComponent : SplineComponents)
+	{
+		if (SplineComponent != OuterSplines)
+		{
+			auto ForeignMeshComponents = SplineComponent->GetForeignMeshComponents(this);
+			if (ForeignMeshComponents.Num() > 0)
+			{
+				for (auto* ForeignMeshComponent : ForeignMeshComponents)
+				{
+					ForeignMeshComponent->Modify();
+				}
+				ForeignMeshComponentsMap.Add(SplineComponent, MoveTemp(ForeignMeshComponents));
+			}
+		}
+	}
+
+	return ForeignMeshComponentsMap;
+}
+
 void ULandscapeSplineSegment::UpdateSplinePoints(bool bUpdateCollision)
 {
 	Modify();
 
-	ULandscapeSplinesComponent* OuterSplines = CastChecked<ULandscapeSplinesComponent>(GetOuter());
+	ULandscapeSplinesComponent* OuterSplines = GetOuterULandscapeSplinesComponent();
 
 	SplineInfo.Points.Empty(2);
 	Points.Reset();
 
-	if (Connections[0].ControlPoint == NULL
-		|| Connections[1].ControlPoint == NULL)
+	if (Connections[0].ControlPoint == nullptr
+		|| Connections[1].ControlPoint == nullptr)
 	{
 		return;
 	}
@@ -1250,23 +1942,45 @@ void ULandscapeSplineSegment::UpdateSplinePoints(bool bUpdateCollision)
 	}
 
 	// Editor mesh
-	FLandscapeSplineMeshEntry SplineEditorMeshEntry;
-	SplineEditorMeshEntry.Mesh = OuterSplines->SplineEditorMesh;
-	SplineEditorMeshEntry.Scale.X = 3;
-	SplineEditorMeshEntry.Offset.Y = 0.5f;
 	bool bUsingEditorMesh = false;
-	if (UsableMeshes.Num() == 0 && OuterSplines->SplineEditorMesh != NULL)
+	if (UsableMeshes.Num() == 0 && OuterSplines->SplineEditorMesh != nullptr)
 	{
-		bUsingEditorMesh = true;
+		FLandscapeSplineMeshEntry SplineEditorMeshEntry;
+		SplineEditorMeshEntry.Mesh = OuterSplines->SplineEditorMesh;
+		SplineEditorMeshEntry.Scale.X = 3;
+		SplineEditorMeshEntry.Offset.Y = 0.5f;
 		UsableMeshes.Add(&SplineEditorMeshEntry);
+		bUsingEditorMesh = true;
 	}
 
 	OuterSplines->Modify();
-	for (auto* MeshComponent : MeshComponents)
+
+	TArray<USplineMeshComponent*> MeshComponents;
+
+	TArray<USplineMeshComponent*> OldLocalMeshComponents = MoveTemp(LocalMeshComponents);
+	LocalMeshComponents.Reserve(20);
+
+	auto ForeignMeshComponentsMap = GetForeignMeshComponents();
+
+	// Unregister components
+	for (auto* MeshComponent : OldLocalMeshComponents)
 	{
 		MeshComponent->Modify();
+		MeshComponent->UnregisterComponent();
+	}
+	for (auto& ForeignMeshComponentsPair : ForeignMeshComponentsMap)
+	{
+		ForeignMeshComponentsPair.Key->Modify();
+		ForeignMeshComponentsPair.Key->GetOwner()->Modify();
+		for (auto* MeshComponent : ForeignMeshComponentsPair.Value)
+		{
+			MeshComponent->Modify();
+			MeshComponent->UnregisterComponent();
+		}
 	}
-	UnregisterComponents();
+
+	ModificationKey = FGuid::NewGuid();
+	ForeignWorlds.Reset();
 
 	if (SplineLength > 0 && (StartWidth > 0 || EndWidth > 0) && UsableMeshes.Num() > 0)
 	{
@@ -1319,20 +2033,58 @@ void ULandscapeSplineSegment::UpdateSplinePoints(bool bUpdateCollision)
 				break;
 			}
 
-			USplineMeshComponent* MeshComponent = NULL;
-			if (iMesh < MeshComponents.Num())
+			ULandscapeSplinesComponent* MeshComponentOuterSplines = OuterSplines;
+
+			// Attempt to place mesh components into the appropriate landscape streaming levels based on the components under the spline
+			if (bPlaceSplineMeshesInStreamingLevels && !bUsingEditorMesh)
+			{
+				// Only "approx" because we rescale T for the 2nd pass based on how well our chosen meshes fit, but it should be good enough
+				FVector ApproxMeshLocation = SplineInfo.Eval(T + MeshT / 2, FVector::ZeroVector);
+				MeshComponentOuterSplines = OuterSplines->GetStreamingSplinesComponentByLocation(ApproxMeshLocation);
+				MeshComponentOuterSplines->Modify();
+			}
+
+			USplineMeshComponent* MeshComponent = nullptr;
+			if (MeshComponentOuterSplines == OuterSplines)
 			{
-				MeshComponent = MeshComponents[iMesh];
+				if (OldLocalMeshComponents.Num() > 0)
+				{
+					MeshComponent = OldLocalMeshComponents.Pop(false);
+					LocalMeshComponents.Add(MeshComponent);
+				}
 			}
 			else
 			{
-				OuterSplines->GetOuter()->Modify();
-				MeshComponent = NewObject<USplineMeshComponent>(OuterSplines->GetOuter(), NAME_None, RF_Transactional | RF_TextExportTransient);
+				TArray<USplineMeshComponent*>* ForeignMeshComponents = ForeignMeshComponentsMap.Find(MeshComponentOuterSplines);
+				if (ForeignMeshComponents && ForeignMeshComponents->Num() > 0)
+				{
+					MeshComponentOuterSplines->UpdateModificationKey(this);
+					MeshComponent = ForeignMeshComponents->Pop(false);
+					ForeignWorlds.AddUnique(MeshComponentOuterSplines->GetTypedOuter<UWorld>());
+				}
+			}
+
+			if (!MeshComponent)
+			{
+				AActor* MeshComponentOuterActor = MeshComponentOuterSplines->GetOwner();
+				MeshComponentOuterActor->Modify();
+				MeshComponent = NewObject<USplineMeshComponent>(MeshComponentOuterActor, NAME_None, RF_Transactional | RF_TextExportTransient);
 				MeshComponent->bSelected = bSelected;
-				MeshComponent->AttachTo(OuterSplines);
-				MeshComponents.Add(MeshComponent);
+				MeshComponent->AttachTo(MeshComponentOuterSplines);
+				if (MeshComponentOuterSplines == OuterSplines)
+				{
+					LocalMeshComponents.Add(MeshComponent);
+					MeshComponentOuterSplines->MeshComponentLocalOwnersMap.Add(MeshComponent, this);
+				}
+				else
+				{
+					MeshComponentOuterSplines->AddForeignMeshComponent(this, MeshComponent);
+					ForeignWorlds.AddUnique(MeshComponentOuterSplines->GetTypedOuter<UWorld>());
+				}
 			}
 
+			MeshComponents.Add(MeshComponent);
+
 			MeshComponent->SetStaticMesh(Mesh);
 
 			MeshComponent->OverrideMaterials = MeshEntry->MaterialOverrides;
@@ -1349,18 +2101,27 @@ void ULandscapeSplineSegment::UpdateSplinePoints(bool bUpdateCollision)
 			iMesh++;
 			T += MeshT;
 		}
-		MeshSettings.Add(FMeshSettings(T, NULL));
+		// Add terminating key
+		MeshSettings.Add(FMeshSettings(T, nullptr));
+
+		// Destroy old unwanted components now
+		for (UMeshComponent* MeshComponent : OldLocalMeshComponents)
+		{
+			verifySlow(OuterSplines->MeshComponentLocalOwnersMap.Remove(MeshComponent) == 1);
+			MeshComponent->DestroyComponent();
+		}
+		OldLocalMeshComponents.Empty();
 
-		if (iMesh < MeshComponents.Num())
+		for (auto& ForeignMeshComponentsPair : ForeignMeshComponentsMap)
 		{
-			for (int32 i = iMesh; i < MeshComponents.Num(); i++)
+			ULandscapeSplinesComponent* MeshComponentOuterSplines = ForeignMeshComponentsPair.Key;
+			for (auto* MeshComponent : ForeignMeshComponentsPair.Value)
 			{
-				MeshComponents[i]->GetOuter()->Modify();
-				MeshComponents[i]->DestroyComponent();
+				MeshComponentOuterSplines->RemoveForeignMeshComponent(this, MeshComponent);
+				MeshComponent->DestroyComponent();
 			}
-
-			MeshComponents.RemoveAt(iMesh, MeshComponents.Num() - iMesh);
 		}
+		ForeignMeshComponentsMap.Empty();
 
 		// Second pass:
 		// Rescale components to fit a whole number to the spline, set up final parameters
@@ -1479,6 +2240,13 @@ void ULandscapeSplineSegment::UpdateSplinePoints(bool bUpdateCollision)
 			MeshComponent->SplineUpDir = FVector(0,0,1); // Up, to be consistent between joined meshes. We rotate it to horizontal using roll if using Z Forward X Up or X Forward Y Up
 			MeshComponent->ForwardAxis = MeshEntry->ForwardAxis;
 
+			if (MeshComponent->AttachParent != OuterSplines)
+			{
+				FTransform RelativeTransform = OuterSplines->ComponentToWorld.GetRelativeTransform(MeshComponent->AttachParent->ComponentToWorld);
+				MeshComponent->SplineParams.StartPos = RelativeTransform.TransformPosition(MeshComponent->SplineParams.StartPos);
+				MeshComponent->SplineParams.EndPos   = RelativeTransform.TransformPosition(MeshComponent->SplineParams.EndPos);
+			}
+
 			if (USplineMeshComponent::GetAxisValue(MeshEntry->Scale, MeshEntry->ForwardAxis) < 0)
 			{
 				Swap(MeshComponent->SplineParams.StartPos, MeshComponent->SplineParams.EndPos);
@@ -1531,19 +2299,33 @@ void ULandscapeSplineSegment::UpdateSplinePoints(bool bUpdateCollision)
 		}
 
 		// Finally, register components
-		RegisterComponents();
+		for (auto* MeshComponent : MeshComponents)
+		{
+			MeshComponent->RegisterComponent();
+		}
 	}
 	else
 	{
-		if (MeshComponents.Num() > 0)
+		// Spline needs no mesh components (0 length or no meshes to use) so destroy any we have
+		for (auto* MeshComponent : OldLocalMeshComponents)
 		{
-			for (auto* MeshComponent : MeshComponents)
+			verifySlow(OuterSplines->MeshComponentLocalOwnersMap.Remove(MeshComponent) == 1);
+			MeshComponent->DestroyComponent();
+		}
+		OldLocalMeshComponents.Empty();
+		for (auto& ForeignMeshComponentsPair : ForeignMeshComponentsMap)
+		{
+			for (USplineMeshComponent* MeshComponent : ForeignMeshComponentsPair.Value)
 			{
-				MeshComponent->GetOuter()->Modify();
+				ULandscapeSplinesComponent* MeshComponentOuterSplines = dynamic_cast<ULandscapeSplinesComponent*>(MeshComponent->GetAttachParent());
+				if (MeshComponentOuterSplines)
+				{
+					MeshComponentOuterSplines->RemoveForeignMeshComponent(this, MeshComponent);
+				}
 				MeshComponent->DestroyComponent();
 			}
-			MeshComponents.Empty();
 		}
+		ForeignMeshComponentsMap.Empty();
 	}
 }
 
@@ -1551,20 +2333,32 @@ void ULandscapeSplineSegment::UpdateSplineEditorMesh()
 {
 	ULandscapeSplinesComponent* OuterSplines = CastChecked<ULandscapeSplinesComponent>(GetOuter());
 
-	for (auto* MeshComponent : MeshComponents)
+	for (auto* MeshComponent : LocalMeshComponents)
 	{
 		if (MeshComponent->bHiddenInGame)
 		{
 			MeshComponent->SetVisibility(OuterSplines->bShowSplineEditorMesh);
 		}
 	}
+
+	auto ForeignMeshComponentsMap = GetForeignMeshComponents();
+	for (auto& ForeignMeshComponentsPair : ForeignMeshComponentsMap)
+	{
+		for (auto* MeshComponent : ForeignMeshComponentsPair.Value)
+		{
+			if (MeshComponent->bHiddenInGame)
+			{
+				MeshComponent->SetVisibility(OuterSplines->bShowSplineEditorMesh);
+			}
+		}
+	}
 }
 
 void ULandscapeSplineSegment::DeleteSplinePoints()
 {
 	Modify();
 
-	ULandscapeSplinesComponent* OuterSplines = CastChecked<ULandscapeSplinesComponent>(GetOuter());
+	ULandscapeSplinesComponent* OuterSplines = GetOuterULandscapeSplinesComponent();
 
 	SplineInfo.Reset();
 	Points.Reset();
@@ -1572,12 +2366,31 @@ void ULandscapeSplineSegment::DeleteSplinePoints()
 
 	OuterSplines->MarkRenderStateDirty();
 
-	for (auto* MeshComponent : MeshComponents)
+	// Destroy mesh components
+	OuterSplines->GetOwner()->Modify();
+	for (auto* MeshComponent : LocalMeshComponents)
 	{
-		MeshComponent->GetOuter()->Modify();
+		MeshComponent->Modify();
 		MeshComponent->DestroyComponent();
 	}
-	MeshComponents.Empty();
+	LocalMeshComponents.Empty();
+
+	auto ForeignMeshComponentsMap = GetForeignMeshComponents();
+	for (auto& ForeignMeshComponentsPair : ForeignMeshComponentsMap)
+	{
+		ULandscapeSplinesComponent* MeshComponentOuterSplines = ForeignMeshComponentsPair.Key;
+		MeshComponentOuterSplines->Modify();
+		MeshComponentOuterSplines->GetOwner()->Modify();
+		for (auto* MeshComponent : ForeignMeshComponentsPair.Value)
+		{
+			MeshComponent->Modify();
+			MeshComponentOuterSplines->RemoveForeignMeshComponent(this, MeshComponent);
+			MeshComponent->DestroyComponent();
+		}
+	}
+
+	ModificationKey.Invalidate();
+	ForeignWorlds.Empty();
 }
 #endif
 
@@ -1593,13 +2406,13 @@ bool ULandscapeSplineSegment::Modify(bool bAlwaysMarkDirty /*= true*/)
 {
 	bool bSavedToTransactionBuffer = Super::Modify(bAlwaysMarkDirty);
 
-	for (auto* MeshComponent : MeshComponents)
-	{
-		if (MeshComponent)
-		{
-			bSavedToTransactionBuffer = MeshComponent->Modify(bAlwaysMarkDirty) || bSavedToTransactionBuffer;
-		}
-	}
+	//for (auto MeshComponent : MeshComponents)
+	//{
+	//	if (MeshComponent)
+	//	{
+	//		bSavedToTransactionBuffer = MeshComponent->Modify(bAlwaysMarkDirty) || bSavedToTransactionBuffer;
+	//	}
+	//}
 
 	return bSavedToTransactionBuffer;
 }
@@ -1618,16 +2431,19 @@ void ULandscapeSplineSegment::PostDuplicate(bool bDuplicateForPIE)
 {
 	if (!bDuplicateForPIE)
 	{
-		ULandscapeSplinesComponent* OuterSplines = CastChecked<ULandscapeSplinesComponent>(GetOuter());
-
-		for (auto*& MeshComponent : MeshComponents)
+		// if we get duplicated but our local meshes don't, then clear our reference to the meshes - they're not ours
+		if (LocalMeshComponents.Num() > 0)
 		{
-			if (MeshComponent->GetOuter() != OuterSplines->GetOuter())
+			ULandscapeSplinesComponent* OuterSplines = CastChecked<ULandscapeSplinesComponent>(GetOuter());
+
+			// we assume all meshes are duplicated or none are, to avoid testing every one
+			if (LocalMeshComponents[0]->GetOuter() != OuterSplines->GetOwner())
 			{
-				MeshComponent = DuplicateObject<USplineMeshComponent>(MeshComponent, OuterSplines->GetOuter(), *MeshComponent->GetName());
-				MeshComponent->AttachTo(OuterSplines);
+				LocalMeshComponents.Empty();
 			}
 		}
+
+		UpdateSplinePoints();
 	}
 
 	Super::PostDuplicate(bDuplicateForPIE);
@@ -1639,7 +2455,7 @@ void ULandscapeSplineSegment::PostEditImport()
 
 	GetOuterULandscapeSplinesComponent()->Segments.AddUnique(this);
 
-	if (Connections[0].ControlPoint != NULL)
+	if (Connections[0].ControlPoint != nullptr)
 	{
 		Connections[0].ControlPoint->ConnectedSegments.AddUnique(FLandscapeSplineConnection(this, 0));
 		Connections[1].ControlPoint->ConnectedSegments.AddUnique(FLandscapeSplineConnection(this, 1));
@@ -1672,4 +2488,4 @@ void ULandscapeSplineSegment::PostEditChangeProperty(FPropertyChangedEvent& Prop
 #endif // WITH_EDITOR
 
 
-
+#undef LOCTEXT_NAMESPACE
