Landscape: Fixed Mac non-unity

diff --git a/Engine/Source/Runtime/Landscape/Private/LandscapeCollision.cpp b/Engine/Source/Runtime/Landscape/Private/LandscapeCollision.cpp
--- a/Engine/Source/Runtime/Landscape/Private/LandscapeCollision.cpp
+++ b/Engine/Source/Runtime/Landscape/Private/LandscapeCollision.cpp
@@ -1,9 +1,5 @@
 // Copyright 1998-2015 Epic Games, Inc. All Rights Reserved.
 
-/*=============================================================================
-LandscapeCollision.cpp: Landscape collision
-=============================================================================*/
-
 #include "Landscape.h"
 #include "PhysicsPublic.h"
 #include "LandscapeDataAccess.h"
@@ -29,6 +25,8 @@ LandscapeCollision.cpp: Landscape collision
 #include "Engine/Engine.h"
 #include "EngineGlobals.h"
 #include "InstancedFoliageActor.h"
+#include "EngineUtils.h"
+
 
 TMap<FGuid, ULandscapeHeightfieldCollisionComponent::FPhysXHeightfieldRef* > GSharedHeightfieldRefs;
 
