Global landscape map considers landscape actors only in visible levels
Fixes UE-12859

diff --git a/Engine/Source/Runtime/Landscape/Private/Landscape.cpp b/Engine/Source/Runtime/Landscape/Private/Landscape.cpp
--- a/Engine/Source/Runtime/Landscape/Private/Landscape.cpp
+++ b/Engine/Source/Runtime/Landscape/Private/Landscape.cpp
@@ -1958,6 +1958,8 @@ void ULandscapeInfo::RecreateLandscapeInfo(UWorld* InWorld, bool bMapCheck)
 	{
 		ALandscapeProxy* Proxy = *It;
 		if (Proxy
+			&& Proxy->GetLevel()
+			&& Proxy->GetLevel()->bIsVisible
 			&& Proxy->HasAnyFlags(RF_BeginDestroyed) == false
 			&& Proxy->IsPendingKill() == false
 			&& Proxy->IsPendingKillPending() == false)
