Fixed: Landscape Physical Material always registers default
UE-11421

diff --git a/Engine/Source/Runtime/Landscape/Private/LandscapeEdit.cpp b/Engine/Source/Runtime/Landscape/Private/LandscapeEdit.cpp
--- a/Engine/Source/Runtime/Landscape/Private/LandscapeEdit.cpp
+++ b/Engine/Source/Runtime/Landscape/Private/LandscapeEdit.cpp
@@ -3430,6 +3430,8 @@ void ALandscapeProxy::ChangedPhysMaterial()
 			if (CollisionComponent)
 			{
 				Comp->UpdateCollisionLayerData();
+				// Physical materials cooked into collision object, so we need to recreate it
+				CollisionComponent->RecreateCollision(false);
 			}
 		}
 	}
