Fixed mouse move during sculpting on mirrored landscapes

diff --git a/Engine/Source/Runtime/Landscape/Private/LandscapeCollision.cpp b/Engine/Source/Runtime/Landscape/Private/LandscapeCollision.cpp
--- a/Engine/Source/Runtime/Landscape/Private/LandscapeCollision.cpp
+++ b/Engine/Source/Runtime/Landscape/Private/LandscapeCollision.cpp
@@ -998,7 +998,7 @@ void ULandscapeHeightfieldCollisionComponent::UpdateHeightfieldRegion(int32 Comp
 		//
 		// Reset geometry of heightfield shape. Required by the modifySamples
 		//
-		FVector LandscapeScale = GetComponentToWorld().GetScale3D();
+		FVector LandscapeScale = GetComponentToWorld().GetScale3D().GetAbs();
 		// Create the geometry
 		PxHeightFieldGeometry LandscapeComponentGeom(HeightfieldRef->RBHeightfieldEd, PxMeshGeometryFlags(), LandscapeScale.Z * LANDSCAPE_ZSCALE, LandscapeScale.Y * CollisionScale, LandscapeScale.X * CollisionScale);
 
