Remove the need for LandscapeHoleMaterial.
It will still exist as a property, but if it's not set it will use LandscapeMaterial automatically (overriding the blend mode to masked if it's opaque).

diff --git a/Engine/Source/Runtime/Landscape/Classes/LandscapeComponent.h b/Engine/Source/Runtime/Landscape/Classes/LandscapeComponent.h
--- a/Engine/Source/Runtime/Landscape/Classes/LandscapeComponent.h
+++ b/Engine/Source/Runtime/Landscape/Classes/LandscapeComponent.h
@@ -166,7 +166,7 @@ class ULandscapeComponent : public UPrimitiveComponent
 	UPROPERTY(EditAnywhere, BlueprintReadWrite, Category=LandscapeComponent)
 	class UMaterialInterface* OverrideMaterial;
 
-	UPROPERTY(EditAnywhere, BlueprintReadWrite, Category=LandscapeComponent)
+	UPROPERTY(EditAnywhere, BlueprintReadWrite, Category=LandscapeComponent, AdvancedDisplay)
 	class UMaterialInterface* OverrideHoleMaterial;
 
 	UPROPERTY(TextExportTransient)
@@ -508,7 +508,7 @@ public:
 	/**
 	 * Generate a key for this component's layer allocations to use with MaterialInstanceConstantMap.
 	 */
-	FString GetLayerAllocationKey(bool bMobile = false) const;
+	FString GetLayerAllocationKey(UMaterialInterface* LandscapeMaterial, bool bMobile = false) const;
 
 	/** @todo document */
 	void GetLayerDebugColorKey(int32& R, int32& G, int32& B) const;
diff --git a/Engine/Source/Runtime/Landscape/Classes/LandscapeProxy.h b/Engine/Source/Runtime/Landscape/Classes/LandscapeProxy.h
--- a/Engine/Source/Runtime/Landscape/Classes/LandscapeProxy.h
+++ b/Engine/Source/Runtime/Landscape/Classes/LandscapeProxy.h
@@ -361,8 +361,8 @@ public:
 	UPROPERTY(EditAnywhere, Category=Landscape)
 	UMaterialInterface* LandscapeMaterial;
 
-	/** Material used to render landscape components with holes. Should be a BLEND_Masked version of LandscapeMaterial. */
-	UPROPERTY(EditAnywhere, Category=Landscape)
+	/** Material used to render landscape components with holes. If not set, LandscapeMaterial will be used (blend mode will be overridden to Masked if it is set to Opaque) */
+	UPROPERTY(EditAnywhere, Category=Landscape, AdvancedDisplay)
 	UMaterialInterface* LandscapeHoleMaterial;
 
 	UPROPERTY(EditAnywhere, Category=LOD)
diff --git a/Engine/Source/Runtime/Landscape/Private/Landscape.cpp b/Engine/Source/Runtime/Landscape/Private/Landscape.cpp
--- a/Engine/Source/Runtime/Landscape/Private/Landscape.cpp
+++ b/Engine/Source/Runtime/Landscape/Private/Landscape.cpp
@@ -290,7 +290,7 @@ UMaterialInterface* ULandscapeComponent::GetLandscapeHoleMaterial() const
 	{
 		return Proxy->GetLandscapeHoleMaterial();
 	}
-	return UMaterial::GetDefaultMaterial(MD_Surface);
+	return nullptr;
 }
 
 bool ULandscapeComponent::ComponentHasVisibilityPainted() const
@@ -306,9 +306,8 @@ bool ULandscapeComponent::ComponentHasVisibilityPainted() const
 	return false;
 }
 
-FString ULandscapeComponent::GetLayerAllocationKey(bool bMobile /* = false */) const
+FString ULandscapeComponent::GetLayerAllocationKey(UMaterialInterface* LandscapeMaterial, bool bMobile /*= false*/) const
 {
-	UMaterialInterface* LandscapeMaterial = GetLandscapeMaterial();
 	if (!LandscapeMaterial)
 	{
 		return FString();
@@ -320,7 +319,7 @@ FString ULandscapeComponent::GetLayerAllocationKey(bool bMobile /* = false */) c
 	TArray<FString> LayerStrings;
 	for (int32 LayerIdx = 0; LayerIdx < WeightmapLayerAllocations.Num(); LayerIdx++)
 	{
-		new(LayerStrings)FString(*FString::Printf(TEXT("_%s_%d"), *WeightmapLayerAllocations[LayerIdx].GetLayerName().ToString(), bMobile ? 0 : WeightmapLayerAllocations[LayerIdx].WeightmapTextureIndex));
+		LayerStrings.Add(FString::Printf(TEXT("_%s_%d"), *WeightmapLayerAllocations[LayerIdx].GetLayerName().ToString(), bMobile ? 0 : WeightmapLayerAllocations[LayerIdx].WeightmapTextureIndex));
 	}
 	/**
 	 * Generate a key for this component's layer allocations to use with MaterialInstanceConstantMap.
@@ -1494,11 +1493,11 @@ UMaterialInterface* ALandscapeProxy::GetLandscapeHoleMaterial() const
 	{
 		return LandscapeHoleMaterial;
 	}
-	else if (LandscapeActor)
+	else if (ALandscape* Landscape = LandscapeActor.Get())
 	{
-		return LandscapeActor->GetLandscapeHoleMaterial();
+		return Landscape->GetLandscapeHoleMaterial();
 	}
-	return UMaterial::GetDefaultMaterial(MD_Surface);
+	return nullptr;
 }
 
 UMaterialInterface* ALandscape::GetLandscapeMaterial() const
@@ -1516,7 +1515,7 @@ UMaterialInterface* ALandscape::GetLandscapeHoleMaterial() const
 	{
 		return LandscapeHoleMaterial;
 	}
-	return UMaterial::GetDefaultMaterial(MD_Surface);
+	return nullptr;
 }
 
 void ALandscape::PreSave()
diff --git a/Engine/Source/Runtime/Landscape/Private/LandscapeEdit.cpp b/Engine/Source/Runtime/Landscape/Private/LandscapeEdit.cpp
--- a/Engine/Source/Runtime/Landscape/Private/LandscapeEdit.cpp
+++ b/Engine/Source/Runtime/Landscape/Private/LandscapeEdit.cpp
@@ -167,32 +167,36 @@ UMaterialInstanceConstant* ULandscapeComponent::GetCombinationMaterial(bool bMob
 
 	ALandscapeProxy* Proxy = GetLandscapeProxy();
 
-	UMaterialInterface* LandscapeMaterial = ComponentHasVisibilityPainted() ? GetLandscapeHoleMaterial() : GetLandscapeMaterial();
+	const bool bComponentHasHoles = ComponentHasVisibilityPainted();
+	UMaterialInterface* const LandscapeMaterial = GetLandscapeMaterial();
+	UMaterialInterface* const HoleMaterial = bComponentHasHoles ? GetLandscapeHoleMaterial() : nullptr;
+	UMaterialInterface* const MaterialToUse = bComponentHasHoles && HoleMaterial ? HoleMaterial : LandscapeMaterial;
+	const bool bOverrideBlendMode = bComponentHasHoles && !HoleMaterial && LandscapeMaterial->GetBlendMode() == BLEND_Opaque;
 
-	if (LandscapeMaterial != NULL)
+	if (ensure(MaterialToUse != nullptr))
 	{
 		// Ensure top level UMaterial has appropriate usage flags set.
 		bool bNeedsRecompile;
-		UMaterial* ParentUMaterial = LandscapeMaterial->GetMaterial();
+		UMaterial* ParentUMaterial = MaterialToUse->GetMaterial();
 		if (ParentUMaterial && ParentUMaterial != UMaterial::GetDefaultMaterial(MD_Surface))
 		{
 			ParentUMaterial->SetMaterialUsage(bNeedsRecompile, MATUSAGE_Landscape);
 			ParentUMaterial->SetMaterialUsage(bNeedsRecompile, MATUSAGE_StaticLighting);
 		}
 
-		FString LayerKey = GetLayerAllocationKey(bMobile);
+		FString LayerKey = GetLayerAllocationKey(MaterialToUse, bMobile);
 		//UE_LOG(LogLandscape, Log, TEXT("Looking for key %s"), *LayerKey);
 
 		// Find or set a matching MIC in the Landscape's map.
 		UMaterialInstanceConstant* CombinationMaterialInstance = Proxy->MaterialInstanceConstantMap.FindRef(*LayerKey);
-		if (CombinationMaterialInstance == NULL || CombinationMaterialInstance->Parent != LandscapeMaterial || GetOutermost() != CombinationMaterialInstance->GetOutermost())
+		if (CombinationMaterialInstance == nullptr || CombinationMaterialInstance->Parent != MaterialToUse || GetOutermost() != CombinationMaterialInstance->GetOutermost())
 		{
 			FlushRenderingCommands();
 
 			CombinationMaterialInstance = NewObject<ULandscapeMaterialInstanceConstant>(GetOutermost());
 			UE_LOG(LogLandscape, Log, TEXT("Looking for key %s, making new combination %s"), *LayerKey, *CombinationMaterialInstance->GetName());
 			Proxy->MaterialInstanceConstantMap.Add(*LayerKey, CombinationMaterialInstance);
-			CombinationMaterialInstance->SetParentEditorOnly(LandscapeMaterial);
+			CombinationMaterialInstance->SetParentEditorOnly(MaterialToUse);
 
 			FStaticParameterSet StaticParameters;
 			CombinationMaterialInstance->GetStaticParameterValues(StaticParameters);
@@ -222,6 +226,12 @@ UMaterialInstanceConstant* ULandscapeComponent::GetCombinationMaterial(bool bMob
 				}
 			}
 
+			CombinationMaterialInstance->BasePropertyOverrides.bOverride_BlendMode = bOverrideBlendMode;
+			if (bOverrideBlendMode)
+			{
+				CombinationMaterialInstance->BasePropertyOverrides.BlendMode = bComponentHasHoles ? BLEND_Masked : BLEND_Opaque;
+			}
+
 			CombinationMaterialInstance->UpdateStaticPermutation(StaticParameters);
 
 			CombinationMaterialInstance->PostEditChange();
@@ -229,7 +239,7 @@ UMaterialInstanceConstant* ULandscapeComponent::GetCombinationMaterial(bool bMob
 
 		return CombinationMaterialInstance;
 	}
-	return NULL;
+	return nullptr;
 }
 
 void ULandscapeComponent::UpdateMaterialInstances()
@@ -447,7 +457,7 @@ void ULandscapeComponent::FixupWeightmaps()
 				UMaterialInstanceConstant* CombinationMaterialInstance = Cast<UMaterialInstanceConstant>(MaterialInstance->Parent);
 				if (CombinationMaterialInstance)
 				{
-					Proxy->MaterialInstanceConstantMap.Add(*GetLayerAllocationKey(), CombinationMaterialInstance);
+					Proxy->MaterialInstanceConstantMap.Add(*GetLayerAllocationKey(CombinationMaterialInstance->Parent), CombinationMaterialInstance);
 				}
 			}
 		}
