Fix shadow variables

diff --git a/Engine/Source/Developer/MeshUtilities/Private/MeshUtilities.cpp b/Engine/Source/Developer/MeshUtilities/Private/MeshUtilities.cpp
--- a/Engine/Source/Developer/MeshUtilities/Private/MeshUtilities.cpp
+++ b/Engine/Source/Developer/MeshUtilities/Private/MeshUtilities.cpp
@@ -2007,7 +2007,7 @@ public:
 		const TArray<FMeshWedge>	&Wedges,
 		const TArray<FMeshFace>		&Faces,
 		const TArray<FVector>		&Points,
-		bool						&bComputeNormals,
+		bool						&bInComputeNormals,
 		TArray<FVector>				&VertexTangentsX,
 		TArray<FVector>				&VertexTangentsY,
 		TArray<FVector>				&VertexTangentsZ
@@ -2016,7 +2016,7 @@ public:
 		wedges			(Wedges),
 		faces			(Faces),
 		points			(Points),
-		bComputeNormals	(bComputeNormals),
+		bComputeNormals	(bInComputeNormals),
 		TangentsX		(VertexTangentsX),
 		TangentsY		(VertexTangentsY),
 		TangentsZ		(VertexTangentsZ)
diff --git a/Engine/Source/Developer/TextureCompressor/Private/TextureCompressorModule.cpp b/Engine/Source/Developer/TextureCompressor/Private/TextureCompressorModule.cpp
--- a/Engine/Source/Developer/TextureCompressor/Private/TextureCompressorModule.cpp
+++ b/Engine/Source/Developer/TextureCompressor/Private/TextureCompressorModule.cpp
@@ -998,8 +998,8 @@ static inline float ComputeTexelArea(uint32 x, uint32 y, float InvSideExtentMul2
  */
 static void GenerateAngularFilteredMip(FImage* DestMip, FImage& SrcMip, float ConeAngle)
 {
-	int32 Extent = DestMip->SizeX;
-	float InvSideExtent = 1.0f / Extent;
+	int32 MipExtent = DestMip->SizeX;
+	float MipInvSideExtent = 1.0f / MipExtent;
 
 	TArray<float> TexelAreaArray;
 	TexelAreaArray.AddUninitialized(SrcMip.SizeX * SrcMip.SizeY);
@@ -1009,7 +1009,7 @@ static void GenerateAngularFilteredMip(FImage* DestMip, FImage& SrcMip, float Co
 	{
 		for(int32 x = 0; x < SrcMip.SizeX; ++x)
 		{
-			TexelAreaArray[x + y * SrcMip.SizeX] = ComputeTexelArea(x, y, InvSideExtent * 2);
+			TexelAreaArray[x + y * SrcMip.SizeX] = ComputeTexelArea(x, y, MipInvSideExtent * 2);
 		}
 	}
 
@@ -1060,7 +1060,7 @@ static void GenerateAngularFilteredMip(FImage* DestMip, FImage& SrcMip, float Co
 
 		for (int32 Face = 0; Face < 6; ++Face)
 		{
-			auto* AsyncTask = new(AsyncTasks) FAsyncGenerateMipsPerFaceTask(Face, DestMip, Extent, ConeAngle, TexelAreaArray.GetData(), &SrcMip);
+			auto* AsyncTask = new(AsyncTasks) FAsyncGenerateMipsPerFaceTask(Face, DestMip, MipExtent, ConeAngle, TexelAreaArray.GetData(), &SrcMip);
 			AsyncTask->StartBackgroundTask();
 		}
 
@@ -1075,11 +1075,11 @@ static void GenerateAngularFilteredMip(FImage* DestMip, FImage& SrcMip, float Co
 		for (int32 Face = 0; Face < 6; ++Face)
 		{
 			FImageView2D DestMipView(*DestMip, Face);
-			for (int32 y = 0; y < Extent; ++y)
+			for (int32 y = 0; y < MipExtent; ++y)
 			{
-				for (int32 x = 0; x < Extent; ++x)
+				for (int32 x = 0; x < MipExtent; ++x)
 				{
-					FVector DirectionWS = ComputeWSCubeDirectionAtTexelCenter(Face, x, y, InvSideExtent);
+					FVector DirectionWS = ComputeWSCubeDirectionAtTexelCenter(Face, x, y, MipInvSideExtent);
 					DestMipView.Access(x, y) = IntegrateAngularArea(SrcMip, DirectionWS, ConeAngle, TexelAreaArray.GetData());
 				}
 			}
diff --git a/Engine/Source/Runtime/CoreUObject/Private/Serialization/AsyncLoading.cpp b/Engine/Source/Runtime/CoreUObject/Private/Serialization/AsyncLoading.cpp
--- a/Engine/Source/Runtime/CoreUObject/Private/Serialization/AsyncLoading.cpp
+++ b/Engine/Source/Runtime/CoreUObject/Private/Serialization/AsyncLoading.cpp
@@ -1388,7 +1388,7 @@ EAsyncPackageState::Type FAsyncPackage::PostLoadObjects()
 	return Result;
 }
 
-EAsyncPackageState::Type FAsyncPackage::PostLoadDeferredObjects(double InTickStartTime, bool bUseTimeLimit, float& InOutTimeLimit)
+EAsyncPackageState::Type FAsyncPackage::PostLoadDeferredObjects(double InTickStartTime, bool bInUseTimeLimit, float& InOutTimeLimit)
 {
 	SCOPE_CYCLE_COUNTER(STAT_FAsyncPackage_PostLoadObjectsGameThread);
 
@@ -1398,7 +1398,7 @@ EAsyncPackageState::Type FAsyncPackage::PostLoadDeferredObjects(double InTickSta
 	LastObjectWorkWasPerformedOn = nullptr;
 	LastTypeOfWorkPerformed = TEXT("postloading_gamethread");
 
-	while (DeferredPostLoadIndex < DeferredPostLoadObjects.Num() && !::IsTimeLimitExceeded(InTickStartTime, bUseTimeLimit, InOutTimeLimit, LastTypeOfWorkPerformed, LastObjectWorkWasPerformedOn))
+	while (DeferredPostLoadIndex < DeferredPostLoadObjects.Num() && !::IsTimeLimitExceeded(InTickStartTime, bInUseTimeLimit, InOutTimeLimit, LastTypeOfWorkPerformed, LastObjectWorkWasPerformedOn))
 	{
 		UObject* Object = DeferredPostLoadObjects[DeferredPostLoadIndex++];
 		check(Object);
diff --git a/Engine/Source/Runtime/Engine/Private/Components/ActorComponent.cpp b/Engine/Source/Runtime/Engine/Private/Components/ActorComponent.cpp
--- a/Engine/Source/Runtime/Engine/Private/Components/ActorComponent.cpp
+++ b/Engine/Source/Runtime/Engine/Private/Components/ActorComponent.cpp
@@ -1227,10 +1227,10 @@ void UActorComponent::MarkForNeededEndOfFrameUpdate()
 		return;
 	}
 
-	UWorld* World = GetWorld();
-	if (World)
+	UWorld* ComponentWorld = GetWorld();
+	if (ComponentWorld)
 	{
-		World->MarkActorComponentForNeededEndOfFrameUpdate(this, RequiresGameThreadEndOfFrameUpdates());
+		ComponentWorld->MarkActorComponentForNeededEndOfFrameUpdate(this, RequiresGameThreadEndOfFrameUpdates());
 	}
 	else if (!HasAnyFlags(RF_Unreachable))
 	{
@@ -1246,11 +1246,11 @@ void UActorComponent::MarkForNeededEndOfFrameRecreate()
 		return;
 	}
 
-	UWorld* World = GetWorld();
-	if (World)
+	UWorld* ComponentWorld = GetWorld();
+	if (ComponentWorld)
 	{
 		// by convention, recreates are always done on the gamethread
-		World->MarkActorComponentForNeededEndOfFrameUpdate(this, true);
+		ComponentWorld->MarkActorComponentForNeededEndOfFrameUpdate(this, true);
 	}
 	else if (!HasAnyFlags(RF_Unreachable))
 	{
diff --git a/Engine/Source/Runtime/Engine/Private/PlayerController.cpp b/Engine/Source/Runtime/Engine/Private/PlayerController.cpp
--- a/Engine/Source/Runtime/Engine/Private/PlayerController.cpp
+++ b/Engine/Source/Runtime/Engine/Private/PlayerController.cpp
@@ -3479,7 +3479,6 @@ public:
 
 	FDynamicForceFeedbackDetails ForceFeedbackDetails;
 
-	FLatentActionInfo LatentInfo;
 	/** Function to execute on completion */
 	FName ExecutionFunction;
 	/** Link to fire on completion */
diff --git a/Engine/Source/Runtime/GameplayAbilities/Private/AbilitySystemComponent_Abilities.cpp b/Engine/Source/Runtime/GameplayAbilities/Private/AbilitySystemComponent_Abilities.cpp
--- a/Engine/Source/Runtime/GameplayAbilities/Private/AbilitySystemComponent_Abilities.cpp
+++ b/Engine/Source/Runtime/GameplayAbilities/Private/AbilitySystemComponent_Abilities.cpp
@@ -2541,21 +2541,21 @@ void UAbilitySystemComponent::ClientSetReplicatedEvent_Implementation(EAbilityGe
 
 // -------
 
-void UAbilitySystemComponent::ServerSetReplicatedTargetData_Implementation(FGameplayAbilitySpecHandle AbilityHandle, FPredictionKey AbilityOriginalPredictionKey, FGameplayAbilityTargetDataHandle ReplicatedTargetData, FGameplayTag ApplicationTag, FPredictionKey CurrentPredictionKey)
+void UAbilitySystemComponent::ServerSetReplicatedTargetData_Implementation(FGameplayAbilitySpecHandle AbilityHandle, FPredictionKey AbilityOriginalPredictionKey, FGameplayAbilityTargetDataHandle ReplicatedTargetDataHandle, FGameplayTag ApplicationTag, FPredictionKey CurrentPredictionKey)
 {
 	FScopedPredictionWindow ScopedPrediction(this, CurrentPredictionKey);
 
 	// Always adds to cache to store the new data
 	FAbilityReplicatedDataCache& ReplicatedData = AbilityTargetDataMap.FindOrAdd(FGameplayAbilitySpecHandleAndPredictionKey(AbilityHandle, AbilityOriginalPredictionKey));
 
-	ReplicatedData.TargetData = ReplicatedTargetData;
+	ReplicatedData.TargetData = ReplicatedTargetDataHandle;
 	ReplicatedData.ApplicationTag = ApplicationTag;
 	ReplicatedData.bTargetConfirmed = true;
 	ReplicatedData.bTargetCancelled = false;
 	ReplicatedData.TargetSetDelegate.Broadcast(ReplicatedData.TargetData, ReplicatedData.ApplicationTag);
 }
 
-bool UAbilitySystemComponent::ServerSetReplicatedTargetData_Validate(FGameplayAbilitySpecHandle AbilityHandle, FPredictionKey AbilityOriginalPredictionKey, FGameplayAbilityTargetDataHandle ReplicatedTargetData, FGameplayTag ApplicationTag, FPredictionKey CurrentPredictionKey)
+bool UAbilitySystemComponent::ServerSetReplicatedTargetData_Validate(FGameplayAbilitySpecHandle AbilityHandle, FPredictionKey AbilityOriginalPredictionKey, FGameplayAbilityTargetDataHandle ReplicatedTargetDataHandle, FGameplayTag ApplicationTag, FPredictionKey CurrentPredictionKey)
 {
 	return true;
 }
diff --git a/Engine/Source/Runtime/GameplayAbilities/Public/AbilitySystemComponent.h b/Engine/Source/Runtime/GameplayAbilities/Public/AbilitySystemComponent.h
--- a/Engine/Source/Runtime/GameplayAbilities/Public/AbilitySystemComponent.h
+++ b/Engine/Source/Runtime/GameplayAbilities/Public/AbilitySystemComponent.h
@@ -939,7 +939,7 @@ public:
 	
 	/**  */
 	UFUNCTION(Server, reliable, WithValidation)
-	void ServerSetReplicatedTargetData(FGameplayAbilitySpecHandle AbilityHandle, FPredictionKey AbilityOriginalPredictionKey, FGameplayAbilityTargetDataHandle ReplicatedTargetData, FGameplayTag ApplicationTag, FPredictionKey CurrentPredictionKey);
+	void ServerSetReplicatedTargetData(FGameplayAbilitySpecHandle AbilityHandle, FPredictionKey AbilityOriginalPredictionKey, FGameplayAbilityTargetDataHandle ReplicatedTargetDataHandle, FGameplayTag ApplicationTag, FPredictionKey CurrentPredictionKey);
 
 	/** Replicates to the server that targeting has been cancelled */
 	UFUNCTION(Server, reliable, WithValidation)
diff --git a/Engine/Source/Runtime/Landscape/Private/LandscapeBlueprintSupport.cpp b/Engine/Source/Runtime/Landscape/Private/LandscapeBlueprintSupport.cpp
--- a/Engine/Source/Runtime/Landscape/Private/LandscapeBlueprintSupport.cpp
+++ b/Engine/Source/Runtime/Landscape/Private/LandscapeBlueprintSupport.cpp
@@ -11,15 +11,15 @@
 #include "Components/SplineComponent.h"
 #include "LandscapeSplineSegment.h"
 
-void ALandscapeProxy::EditorApplySpline(USplineComponent* SplineComponent, float StartWidth, float EndWidth, float StartSideFalloff, float EndSideFalloff, float StartRoll, float EndRoll, int32 NumSubdivisions, bool bRaiseHeights, bool bLowerHeights, ULandscapeLayerInfoObject* PaintLayer)
+void ALandscapeProxy::EditorApplySpline(USplineComponent* InSplineComponent, float StartWidth, float EndWidth, float StartSideFalloff, float EndSideFalloff, float StartRoll, float EndRoll, int32 NumSubdivisions, bool bRaiseHeights, bool bLowerHeights, ULandscapeLayerInfoObject* PaintLayer)
 {
 	check(GIsEditor);
 
 #if WITH_EDITOR
 	TArray<FLandscapeSplineInterpPoint> Points;
-	LandscapeSplineRaster::Pointify(SplineComponent->SplineInfo, Points, NumSubdivisions, 0.0f, 0.0f, StartWidth, EndWidth, StartSideFalloff, EndSideFalloff, StartRoll, EndRoll);
+	LandscapeSplineRaster::Pointify(InSplineComponent->SplineInfo, Points, NumSubdivisions, 0.0f, 0.0f, StartWidth, EndWidth, StartSideFalloff, EndSideFalloff, StartRoll, EndRoll);
 
-	FTransform SplineToWorld = SplineComponent->GetComponentTransform();
+	FTransform SplineToWorld = InSplineComponent->GetComponentTransform();
 	LandscapeSplineRaster::RasterizeSegmentPoints(GetLandscapeInfo(), MoveTemp(Points), SplineToWorld, bRaiseHeights, bLowerHeights, PaintLayer);
 #endif
 }
diff --git a/Engine/Source/Runtime/Landscape/Private/LandscapeEditInterface.cpp b/Engine/Source/Runtime/Landscape/Private/LandscapeEditInterface.cpp
--- a/Engine/Source/Runtime/Landscape/Private/LandscapeEditInterface.cpp
+++ b/Engine/Source/Runtime/Landscape/Private/LandscapeEditInterface.cpp
@@ -2851,20 +2851,20 @@ void FLandscapeEditDataInterface::GetWeightDataTemplFast(ULandscapeLayerInfoObje
 								for( int32 LayerIdx=0;LayerIdx<Component->WeightmapLayerAllocations.Num();LayerIdx++ )
 								{
 									int32 Idx = Component->WeightmapLayerAllocations[LayerIdx].WeightmapTextureIndex;
-									UTexture2D* WeightmapTexture = Component->WeightmapTextures[Idx];
-									uint8* WeightmapTextureData = (uint8*)TexDataInfos[Idx]->GetMipData(0);
-									uint8 WeightmapChannelOffset = ChannelOffsets[Component->WeightmapLayerAllocations[LayerIdx].WeightmapTextureChannel];
+									UTexture2D* ComponentWeightmapTexture = Component->WeightmapTextures[Idx];
+									uint8* ComponentWeightmapTextureData = (uint8*)TexDataInfos[Idx]->GetMipData(0);
+									uint8 ComponentWeightmapChannelOffset = ChannelOffsets[Component->WeightmapLayerAllocations[LayerIdx].WeightmapTextureChannel];
 
 									// Find the texture data corresponding to this vertex
-									int32 SizeU = WeightmapTexture->Source.GetSizeX();
-									int32 SizeV = WeightmapTexture->Source.GetSizeY();
+									int32 SizeU = ComponentWeightmapTexture->Source.GetSizeX();
+									int32 SizeV = ComponentWeightmapTexture->Source.GetSizeY();
 									int32 WeightmapOffsetX = Component->WeightmapScaleBias.Z * (float)SizeU;
 									int32 WeightmapOffsetY = Component->WeightmapScaleBias.W * (float)SizeV;
 
 									int32 TexX = WeightmapOffsetX + (SubsectionSizeQuads+1) * SubIndexX + SubX;
 									int32 TexY = WeightmapOffsetY + (SubsectionSizeQuads+1) * SubIndexY + SubY;
 
-									uint8 Weight = WeightmapTextureData[ 4 * (TexX + TexY * SizeU) + WeightmapChannelOffset ];
+									uint8 Weight = ComponentWeightmapTextureData[ 4 * (TexX + TexY * SizeU) + ComponentWeightmapChannelOffset ];
 
 									// Find index in LayerInfos
 									{
@@ -3436,20 +3436,20 @@ void FLandscapeEditDataInterface::GetWeightDataTempl(ULandscapeLayerInfoObject*
 								for( int32 LayerIdx=0;LayerIdx<Component->WeightmapLayerAllocations.Num();LayerIdx++ )
 								{
 									int32 Idx = Component->WeightmapLayerAllocations[LayerIdx].WeightmapTextureIndex;
-									UTexture2D* WeightmapTexture = Component->WeightmapTextures[Idx];
-									uint8* WeightmapTextureData = (uint8*)TexDataInfos[Idx]->GetMipData(0);
-									uint8 WeightmapChannelOffset = ChannelOffsets[Component->WeightmapLayerAllocations[LayerIdx].WeightmapTextureChannel];
+									UTexture2D* ComponentWeightmapTexture = Component->WeightmapTextures[Idx];
+									uint8* ComponentWeightmapTextureData = (uint8*)TexDataInfos[Idx]->GetMipData(0);
+									uint8 ComponentWeightmapChannelOffset = ChannelOffsets[Component->WeightmapLayerAllocations[LayerIdx].WeightmapTextureChannel];
 
 									// Find the texture data corresponding to this vertex
-									int32 SizeU = WeightmapTexture->Source.GetSizeX();
-									int32 SizeV = WeightmapTexture->Source.GetSizeY();
+									int32 SizeU = ComponentWeightmapTexture->Source.GetSizeX();
+									int32 SizeV = ComponentWeightmapTexture->Source.GetSizeY();
 									int32 WeightmapOffsetX = Component->WeightmapScaleBias.Z * (float)SizeU;
 									int32 WeightmapOffsetY = Component->WeightmapScaleBias.W * (float)SizeV;
 
 									int32 TexX = WeightmapOffsetX + (SubsectionSizeQuads+1) * SubIndexX + SubX;
 									int32 TexY = WeightmapOffsetY + (SubsectionSizeQuads+1) * SubIndexY + SubY;
 
-									uint8 Weight = WeightmapTextureData[ 4 * (TexX + TexY * SizeU) + WeightmapChannelOffset ];
+									uint8 Weight = ComponentWeightmapTextureData[ 4 * (TexX + TexY * SizeU) + ComponentWeightmapChannelOffset ];
 
 									// Find index in LayerInfos
 									{
diff --git a/Engine/Source/Runtime/MediaAssets/Private/Assets/MediaTexture.cpp b/Engine/Source/Runtime/MediaAssets/Private/Assets/MediaTexture.cpp
--- a/Engine/Source/Runtime/MediaAssets/Private/Assets/MediaTexture.cpp
+++ b/Engine/Source/Runtime/MediaAssets/Private/Assets/MediaTexture.cpp
@@ -104,9 +104,9 @@ void UMediaTexture::FinishDestroy()
 
 FString UMediaTexture::GetDesc()
 {
-	TSharedPtr<IMediaPlayer> MediaPlayer = GetPlayer();
+	TSharedPtr<IMediaPlayer> MediaPlayerPtr = GetPlayer();
 
-	if (!MediaPlayer.IsValid())
+	if (!MediaPlayerPtr.IsValid())
 	{
 		return FString();
 	}
diff --git a/Engine/Source/Runtime/MoviePlayer/Private/DefaultGameMoviePlayer.cpp b/Engine/Source/Runtime/MoviePlayer/Private/DefaultGameMoviePlayer.cpp
--- a/Engine/Source/Runtime/MoviePlayer/Private/DefaultGameMoviePlayer.cpp
+++ b/Engine/Source/Runtime/MoviePlayer/Private/DefaultGameMoviePlayer.cpp
@@ -177,12 +177,12 @@ bool FDefaultGameMoviePlayer::PlayMovie()
 		
 		LastPlayTime = FPlatformTime::Seconds();
 
-        bool bInitialized = true;
+        bool bIsInitialized = true;
 		if (MovieStreamingIsPrepared())
 		{
-			bInitialized = MovieStreamer->Init(LoadingScreenAttributes.MoviePaths);
+			bIsInitialized = MovieStreamer->Init(LoadingScreenAttributes.MoviePaths);
 		}
-        if (bInitialized)
+        if (bIsInitialized)
         {
 			MovieStreamingIsDone.Set(MovieStreamingIsPrepared() ? 0 : 1);
 			LoadingIsDone.Set(0);
diff --git a/Engine/Source/Runtime/Renderer/Private/SystemTextures.cpp b/Engine/Source/Runtime/Renderer/Private/SystemTextures.cpp
--- a/Engine/Source/Runtime/Renderer/Private/SystemTextures.cpp
+++ b/Engine/Source/Runtime/Renderer/Private/SystemTextures.cpp
@@ -231,21 +231,23 @@ void FSystemTextures::InitializeTextures(FRHICommandListImmediate& RHICmdList, E
 			Bases[Pos] = FColor(FMath::Quantize8SignedByte(c), FMath::Quantize8SignedByte(s), FMath::Quantize8SignedByte(-s), FMath::Quantize8SignedByte(c));
 		}
 
-		FPooledRenderTargetDesc Desc(FPooledRenderTargetDesc::Create2DDesc(FIntPoint(64, 64), PF_B8G8R8A8, TexCreate_HideInVisualizeTexture, TexCreate_None, false));
-		GRenderTargetPool.FindFreeElement(Desc, SSAORandomization, TEXT("SSAORandomization"));
-		// Write the contents of the texture.
-		uint32 DestStride;
-		uint8* DestBuffer = (uint8*)RHICmdList.LockTexture2D((FTexture2DRHIRef&)SSAORandomization->GetRenderTargetItem().ShaderResourceTexture, 0, RLM_WriteOnly, DestStride, false);
-
-		for(int32 y = 0; y < Desc.Extent.Y; ++y)
 		{
-			for(int32 x = 0; x < Desc.Extent.X; ++x)
+			FPooledRenderTargetDesc Desc(FPooledRenderTargetDesc::Create2DDesc(FIntPoint(64, 64), PF_B8G8R8A8, TexCreate_HideInVisualizeTexture, TexCreate_None, false));
+			GRenderTargetPool.FindFreeElement(Desc, SSAORandomization, TEXT("SSAORandomization"));
+			// Write the contents of the texture.
+			uint32 DestStride;
+			uint8* DestBuffer = (uint8*)RHICmdList.LockTexture2D((FTexture2DRHIRef&)SSAORandomization->GetRenderTargetItem().ShaderResourceTexture, 0, RLM_WriteOnly, DestStride, false);
+
+			for(int32 y = 0; y < Desc.Extent.Y; ++y)
 			{
-				FColor* Dest = (FColor*)(DestBuffer + x * sizeof(uint32) + y * DestStride);
+				for(int32 x = 0; x < Desc.Extent.X; ++x)
+				{
+					FColor* Dest = (FColor*)(DestBuffer + x * sizeof(uint32) + y * DestStride);
 
-				uint32 Index = (x % 4) + (y % 4) * 4; 
+					uint32 Index = (x % 4) + (y % 4) * 4; 
 
-				*Dest = Bases[Index];
+					*Dest = Bases[Index];
+				}
 			}
 		}
 		RHICmdList.UnlockTexture2D((FTexture2DRHIRef&)SSAORandomization->GetRenderTargetItem().ShaderResourceTexture, 0, false);
