Fixed case where landscape was not properly using default phisical material when one of the painted layers had empty material

diff --git a/Engine/Source/Runtime/Landscape/Private/LandscapeEdit.cpp b/Engine/Source/Runtime/Landscape/Private/LandscapeEdit.cpp
--- a/Engine/Source/Runtime/Landscape/Private/LandscapeEdit.cpp
+++ b/Engine/Source/Runtime/Landscape/Private/LandscapeEdit.cpp
@@ -803,18 +803,18 @@ void ULandscapeComponent::UpdateCollisionLayerData(TArray<FColor*>& WeightmapTex
 		{
 			FWeightmapLayerAllocationInfo& AllocInfo = WeightmapLayerAllocations[AllocIdx];
 			ULandscapeLayerInfoObject* LayerInfo = AllocInfo.LayerInfo;
-			if (AllocInfo.LayerInfo == ALandscapeProxy::VisibilityLayer || (LayerInfo && LayerInfo->PhysMaterial))
+			if (LayerInfo == ALandscapeProxy::VisibilityLayer || LayerInfo != nullptr)
 			{
-				int32 Idx = CandidateLayers.Add(AllocInfo.LayerInfo);
+				int32 Idx = CandidateLayers.Add(LayerInfo);
 				CandidateDataPtrs.Add(((uint8*)WeightmapTextureMipData[AllocInfo.WeightmapTextureIndex]) + ChannelOffsets[AllocInfo.WeightmapTextureChannel]);
 
 				// Check if we still match the collision component.
-				if (!CollisionComponent->ComponentLayerInfos.IsValidIndex(Idx) || CollisionComponent->ComponentLayerInfos[Idx] != AllocInfo.LayerInfo)
+				if (!CollisionComponent->ComponentLayerInfos.IsValidIndex(Idx) || CollisionComponent->ComponentLayerInfos[Idx] != LayerInfo)
 				{
 					bExistingLayerMismatch = true;
 				}
 
-				if (AllocInfo.LayerInfo == ALandscapeProxy::VisibilityLayer)
+				if (LayerInfo == ALandscapeProxy::VisibilityLayer)
 				{
 					DataLayerIdx = Idx;
 					bExistingLayerMismatch = true; // always rebuild whole component for hole
