Added postponing for landscape textures baking while user uses landscape tools
UE-14108

diff --git a/Engine/Source/Editor/LandscapeEditor/Private/LandscapeEdModeTools.h b/Engine/Source/Editor/LandscapeEditor/Private/LandscapeEdModeTools.h
--- a/Engine/Source/Editor/LandscapeEditor/Private/LandscapeEdModeTools.h
+++ b/Engine/Source/Editor/LandscapeEditor/Private/LandscapeEdModeTools.h
@@ -1313,6 +1313,9 @@ public:
 				MousePositions.Empty(1);
 			}
 			TimeSinceLastMouseMove += DeltaTime;
+
+			// Prevent landscape from baking textures while tool stroke is active
+			EdMode->CurrentToolTarget.LandscapeInfo->PostponeTextureBaking();
 		}
 	}
 
diff --git a/Engine/Source/Runtime/Landscape/Classes/LandscapeInfo.h b/Engine/Source/Runtime/Landscape/Classes/LandscapeInfo.h
--- a/Engine/Source/Runtime/Landscape/Classes/LandscapeInfo.h
+++ b/Engine/Source/Runtime/Landscape/Classes/LandscapeInfo.h
@@ -250,5 +250,8 @@ public:
 	LANDSCAPE_API void RecreateCollisionComponents();
 
 	LANDSCAPE_API void RemoveXYOffsets();
+
+	/** Postpones landscape textures baking, usually used during landscape painting to avoid hitches */
+	LANDSCAPE_API void PostponeTextureBaking();
 #endif
 };
diff --git a/Engine/Source/Runtime/Landscape/Private/LandscapeEdit.cpp b/Engine/Source/Runtime/Landscape/Private/LandscapeEdit.cpp
--- a/Engine/Source/Runtime/Landscape/Private/LandscapeEdit.cpp
+++ b/Engine/Source/Runtime/Landscape/Private/LandscapeEdit.cpp
@@ -3076,6 +3076,22 @@ void ULandscapeInfo::RemoveXYOffsets()
 	}
 }
 
+void ULandscapeInfo::PostponeTextureBaking()
+{
+	const int32 PostponeValue = 60; //frames
+	
+	ALandscape* Landscape = LandscapeActor.Get();
+	if (Landscape)
+	{
+		Landscape->UpdateBakedTexturesCountdown = PostponeValue;
+	}
+
+	for (ALandscapeProxy* Proxy : Proxies)
+	{
+		Proxy->UpdateBakedTexturesCountdown = PostponeValue;
+	}
+}
+
 namespace
 {
 	inline float AdjustStaticLightingResolution(float StaticLightingResolution, int32 NumSubsections, int32 SubsectionSizeQuads, int32 ComponentSizeQuads)
