Fix index buffer overflow crash when rendering tessellation on landscapes with more than 1 subsection. (UE-14187)

diff --git a/Engine/Source/Runtime/Landscape/Private/LandscapeRender.cpp b/Engine/Source/Runtime/Landscape/Private/LandscapeRender.cpp
--- a/Engine/Source/Runtime/Landscape/Private/LandscapeRender.cpp
+++ b/Engine/Source/Runtime/Landscape/Private/LandscapeRender.cpp
@@ -1098,9 +1098,18 @@ void FLandscapeComponentSceneProxy::DrawStaticElements(FStaticPrimitiveDrawInter
 					BatchElementParams->SubX = SubX;
 					BatchElementParams->SubY = SubY;
 					BatchElementParams->CurrentLOD = LOD;
-					BatchElement->IndexBuffer = bCurrentRequiresAdjacencyInformation ? SharedBuffers->AdjacencyIndexBuffers->IndexBuffers[LOD] : SharedBuffers->IndexBuffers[LOD];
-					BatchElement->NumPrimitives = FMath::Square((LodSubsectionSizeVerts - 1)) * 2;
-					BatchElement->FirstIndex = (SubX + SubY * NumSubsections) * BatchElement->NumPrimitives * 3;
+					uint32 NumPrimitives = FMath::Square((LodSubsectionSizeVerts - 1)) * 2;
+					if (bCurrentRequiresAdjacencyInformation)
+					{
+						BatchElement->IndexBuffer = SharedBuffers->AdjacencyIndexBuffers->IndexBuffers[LOD];
+						BatchElement->FirstIndex = (SubX + SubY * NumSubsections) * NumPrimitives * 12;
+					}
+					else
+					{
+						BatchElement->IndexBuffer = SharedBuffers->IndexBuffers[LOD];
+						BatchElement->FirstIndex = (SubX + SubY * NumSubsections) * NumPrimitives * 3;
+					}
+					BatchElement->NumPrimitives = NumPrimitives;
 					BatchElement->MinVertexIndex = SharedBuffers->IndexRanges[LOD].MinIndex[SubX][SubY];
 					BatchElement->MaxVertexIndex = SharedBuffers->IndexRanges[LOD].MaxIndex[SubX][SubY];
 				}
@@ -1382,21 +1391,20 @@ void FLandscapeComponentSceneProxy::GetDynamicMeshElements(const TArray<const FS
 					BatchElementParams.SubY = SubY;
 					BatchElementParams.CurrentLOD = CurrentLOD;
 
+					int32 LodSubsectionSizeVerts = (SubsectionSizeVerts >> CurrentLOD);
+					uint32 NumPrimitives = FMath::Square((LodSubsectionSizeVerts - 1)) * 2;
 					if (bCurrentRequiresAdjacencyInformation)
 					{
 						check(SharedBuffers->AdjacencyIndexBuffers);
 						BatchElement.IndexBuffer = SharedBuffers->AdjacencyIndexBuffers->IndexBuffers[CurrentLOD];
+						BatchElement.FirstIndex = (SubX + SubY * NumSubsections) * NumPrimitives * 12;
 					}
 					else
 					{
 						BatchElement.IndexBuffer = SharedBuffers->IndexBuffers[CurrentLOD];
+						BatchElement.FirstIndex = (SubX + SubY * NumSubsections) * NumPrimitives * 3;
 					}
-
-					int32 LodSubsectionSizeVerts = (SubsectionSizeVerts >> CurrentLOD);
-
-					int32 NumPrimitives = FMath::Square((LodSubsectionSizeVerts - 1)) * 2;
 					BatchElement.NumPrimitives = NumPrimitives;
-					BatchElement.FirstIndex = (SubX + SubY * NumSubsections) * NumPrimitives * 3;
 					BatchElement.MinVertexIndex = SharedBuffers->IndexRanges[CurrentLOD].MinIndex[SubX][SubY];
 					BatchElement.MaxVertexIndex = SharedBuffers->IndexRanges[CurrentLOD].MaxIndex[SubX][SubY];
 
@@ -1408,7 +1416,7 @@ void FLandscapeComponentSceneProxy::GetDynamicMeshElements(const TArray<const FS
 					// Tools never use tessellation
 					BatchElementTools.IndexBuffer = SharedBuffers->IndexBuffers[CurrentLOD];
 					BatchElementTools.NumPrimitives = NumPrimitives;
-					BatchElementTools.FirstIndex = BatchElement.FirstIndex;
+					BatchElementTools.FirstIndex = (SubX + SubY * NumSubsections) * NumPrimitives * 3;
 					BatchElementTools.MinVertexIndex = SharedBuffers->IndexRanges[CurrentLOD].MinIndex[SubX][SubY];
 					BatchElementTools.MaxVertexIndex = SharedBuffers->IndexRanges[CurrentLOD].MaxIndex[SubX][SubY];
 #endif
@@ -2051,35 +2059,45 @@ FLandscapeSharedBuffers::~FLandscapeSharedBuffers()
 }
 
 template<typename IndexType>
-static void BuildLandscapeAdjacencyIndexBuffer(int32 LODSubsectionSizeQuads, const FRawStaticIndexBuffer16or32<IndexType>* Indices, TArray<IndexType>& OutPnAenIndices)
+static void BuildLandscapeAdjacencyIndexBuffer(int32 LODSubsectionSizeQuads, int32 NumSubsections, const FRawStaticIndexBuffer16or32<IndexType>* Indices, TArray<IndexType>& OutPnAenIndices)
 {
 	if (Indices && Indices->Num())
 	{
 		// Landscape use regular grid, so only expand Index buffer works
 		// PN AEN Dominant Corner
 		uint32 TriCount = LODSubsectionSizeQuads*LODSubsectionSizeQuads * 2;
-		uint32 ExpandedCount = 12 * TriCount;
+
+		uint32 ExpandedCount = 12 * TriCount * NumSubsections * NumSubsections;
+
 		OutPnAenIndices.Empty(ExpandedCount);
 		OutPnAenIndices.AddUninitialized(ExpandedCount);
 
-		for (uint32 TriIdx = 0; TriIdx < TriCount; ++TriIdx)
+		for (int32 SubY = 0; SubY < NumSubsections; SubY++)
 		{
-			uint32 OutStartIdx = TriIdx * 12;
-			uint32 InStartIdx = TriIdx * 3;
-			OutPnAenIndices[OutStartIdx + 0] = Indices->Get(InStartIdx + 0);
-			OutPnAenIndices[OutStartIdx + 1] = Indices->Get(InStartIdx + 1);
-			OutPnAenIndices[OutStartIdx + 2] = Indices->Get(InStartIdx + 2);
-
-			OutPnAenIndices[OutStartIdx + 3] = Indices->Get(InStartIdx + 0);
-			OutPnAenIndices[OutStartIdx + 4] = Indices->Get(InStartIdx + 1);
-			OutPnAenIndices[OutStartIdx + 5] = Indices->Get(InStartIdx + 1);
-			OutPnAenIndices[OutStartIdx + 6] = Indices->Get(InStartIdx + 2);
-			OutPnAenIndices[OutStartIdx + 7] = Indices->Get(InStartIdx + 2);
-			OutPnAenIndices[OutStartIdx + 8] = Indices->Get(InStartIdx + 0);
+			for (int32 SubX = 0; SubX < NumSubsections; SubX++)
+			{
+				uint32 SubsectionTriIndex = (SubX + SubY * NumSubsections) * TriCount;
 
-			OutPnAenIndices[OutStartIdx + 9] = Indices->Get(InStartIdx + 0);
-			OutPnAenIndices[OutStartIdx + 10] = Indices->Get(InStartIdx + 1);
-			OutPnAenIndices[OutStartIdx + 11] = Indices->Get(InStartIdx + 2);
+				for (uint32 TriIdx = SubsectionTriIndex; TriIdx < SubsectionTriIndex + TriCount; ++TriIdx)
+				{
+					uint32 OutStartIdx = TriIdx * 12;
+					uint32 InStartIdx = TriIdx * 3;
+					OutPnAenIndices[OutStartIdx + 0] = Indices->Get(InStartIdx + 0);
+					OutPnAenIndices[OutStartIdx + 1] = Indices->Get(InStartIdx + 1);
+					OutPnAenIndices[OutStartIdx + 2] = Indices->Get(InStartIdx + 2);
+
+					OutPnAenIndices[OutStartIdx + 3] = Indices->Get(InStartIdx + 0);
+					OutPnAenIndices[OutStartIdx + 4] = Indices->Get(InStartIdx + 1);
+					OutPnAenIndices[OutStartIdx + 5] = Indices->Get(InStartIdx + 1);
+					OutPnAenIndices[OutStartIdx + 6] = Indices->Get(InStartIdx + 2);
+					OutPnAenIndices[OutStartIdx + 7] = Indices->Get(InStartIdx + 2);
+					OutPnAenIndices[OutStartIdx + 8] = Indices->Get(InStartIdx + 0);
+
+					OutPnAenIndices[OutStartIdx + 9] = Indices->Get(InStartIdx + 0);
+					OutPnAenIndices[OutStartIdx + 10] = Indices->Get(InStartIdx + 1);
+					OutPnAenIndices[OutStartIdx + 11] = Indices->Get(InStartIdx + 2);
+				}
+			}
 		}
 	}
 	else
@@ -2088,6 +2106,7 @@ static void BuildLandscapeAdjacencyIndexBuffer(int32 LODSubsectionSizeQuads, con
 	}
 }
 
+
 FLandscapeSharedAdjacencyIndexBuffer::FLandscapeSharedAdjacencyIndexBuffer(FLandscapeSharedBuffers* Buffers)
 {
 	ensure(Buffers && Buffers->IndexBuffers);
@@ -2101,7 +2120,7 @@ FLandscapeSharedAdjacencyIndexBuffer::FLandscapeSharedAdjacencyIndexBuffer(FLand
 		if (b32BitIndex)
 		{
 			TArray<uint32> OutPnAenIndices;
-			BuildLandscapeAdjacencyIndexBuffer<uint32>((Buffers->SubsectionSizeVerts >> i) - 1, (FRawStaticIndexBuffer16or32<uint32>*)Buffers->IndexBuffers[i], OutPnAenIndices);
+			BuildLandscapeAdjacencyIndexBuffer<uint32>((Buffers->SubsectionSizeVerts >> i) - 1, Buffers->NumSubsections, (FRawStaticIndexBuffer16or32<uint32>*)Buffers->IndexBuffers[i], OutPnAenIndices);
 
 			FRawStaticIndexBuffer16or32<uint32>* IndexBuffer = new FRawStaticIndexBuffer16or32<uint32>();
 			IndexBuffer->AssignNewBuffer(OutPnAenIndices);
@@ -2110,7 +2129,7 @@ FLandscapeSharedAdjacencyIndexBuffer::FLandscapeSharedAdjacencyIndexBuffer(FLand
 		else
 		{
 			TArray<uint16> OutPnAenIndices;
-			BuildLandscapeAdjacencyIndexBuffer<uint16>((Buffers->SubsectionSizeVerts >> i) - 1, (FRawStaticIndexBuffer16or32<uint16>*)Buffers->IndexBuffers[i], OutPnAenIndices);
+			BuildLandscapeAdjacencyIndexBuffer<uint16>((Buffers->SubsectionSizeVerts >> i) - 1, Buffers->NumSubsections, (FRawStaticIndexBuffer16or32<uint16>*)Buffers->IndexBuffers[i], OutPnAenIndices);
 
 			FRawStaticIndexBuffer16or32<uint16>* IndexBuffer = new FRawStaticIndexBuffer16or32<uint16>();
 			IndexBuffer->AssignNewBuffer(OutPnAenIndices);
