Fix issues with landscape ES2 data being overwriting PC data causing issues with component weightmaps when launching or using the mobile previewer multiple times. (UE-10421 / UE-12441)

diff --git a/Engine/Source/Runtime/Core/Public/UObject/ObjectVersion.h b/Engine/Source/Runtime/Core/Public/UObject/ObjectVersion.h
--- a/Engine/Source/Runtime/Core/Public/UObject/ObjectVersion.h
+++ b/Engine/Source/Runtime/Core/Public/UObject/ObjectVersion.h
@@ -502,6 +502,8 @@ enum EUnrealEngineObjectUE4Version
 	VER_UE4_INTERPCURVE_SUPPORTS_LOOPING,
 	// Material Instances overriding base material LOD transitions
 	VER_UE4_MATERIAL_INSTANCE_BASE_PROPERTY_OVERRIDES_DITHERED_LOD_TRANSITION,
+	// Serialize ES2 textures separately rather than overwriting the properties used on other platforms
+	VER_UE4_SERIALIZE_LANDSCAPE_ES2_TEXTURES,
 
 	// -----<new versions can be added before this line>-------------------------------------------------
 	// - this needs to be the last line (see note below)
diff --git a/Engine/Source/Runtime/Landscape/Classes/LandscapeComponent.h b/Engine/Source/Runtime/Landscape/Classes/LandscapeComponent.h
--- a/Engine/Source/Runtime/Landscape/Classes/LandscapeComponent.h
+++ b/Engine/Source/Runtime/Landscape/Classes/LandscapeComponent.h
@@ -263,19 +263,23 @@ public:
 	/** Pointer to data shared with the render thread, used by the editor tools */
 	struct FLandscapeEditToolRenderData* EditToolRenderData;
 
-	/** Runtime-generated editor data for ES2 emulation */
+	/** Hash of source for ES2 generated data. Used for mobile preview to determine if we need to re-generate ES2 pixel data. */
 	UPROPERTY(Transient, DuplicateTransient)
-	UTexture2D* MobileNormalmapTexture;
-
-	/** Runtime-generated editor data for ES2 emulation */
-	UPROPERTY(Transient, DuplicateTransient)
-	UMaterialInterface* MobileMaterialInterface;
+	FGuid MobilePixelDataSourceHash;
 #endif
 
 	/** For ES2 */
 	UPROPERTY()
 	uint8 MobileBlendableLayerMask;
 
+	/** Material interface used for ES2. Serialized only when cooking or loading cooked builds. */
+	UPROPERTY(Transient, DuplicateTransient)
+	UMaterialInterface* MobileMaterialInterface;
+
+	/** Generated weight/normal map texture used for ES2. Serialized only when cooking or loading cooked builds. */
+	UPROPERTY(Transient, DuplicateTransient)
+	UTexture2D* MobileWeightNormalmapTexture;
+
 public:
 	/** Platform Data where don't support texture sampling in vertex buffer */
 	FLandscapeComponentDerivedData PlatformData;
@@ -343,7 +347,7 @@ public:
 	void ReplaceLayer(ULandscapeLayerInfoObject* FromLayerInfo, ULandscapeLayerInfoObject* ToLayerInfo, struct FLandscapeEditDataInterface* LandscapeEdit);
 	
 	void GeneratePlatformVertexData();
-	UMaterialInstance* GeneratePlatformPixelData(TArray<class UTexture2D*>& WeightmapTextures, bool bIsCooking);
+	void GeneratePlatformPixelData(bool bIsCooking);
 
 	/** Creates and destroys cooked grass data stored in the map */
 	void RenderGrassMap();
diff --git a/Engine/Source/Runtime/Landscape/Private/Landscape.cpp b/Engine/Source/Runtime/Landscape/Private/Landscape.cpp
--- a/Engine/Source/Runtime/Landscape/Private/Landscape.cpp
+++ b/Engine/Source/Runtime/Landscape/Private/Landscape.cpp
@@ -133,10 +133,6 @@ void ULandscapeComponent::AddReferencedObjects(UObject* InThis, FReferenceCollec
 void ULandscapeComponent::Serialize(FArchive& Ar)
 {
 #if ENABLE_LANDSCAPE_COOKING && WITH_EDITOR
-	bool bRestoreAfterCooking = false;
-	UMaterialInstanceConstant* BackupMaterialInstance = MaterialInstance;
-	UTexture2D* BackupHeightmapTexture = HeightmapTexture;
-
 	// Saving for cooking path
 	if (Ar.IsCooking() && Ar.IsSaving() && !HasAnyFlags(RF_ClassDefaultObject))
 	{
@@ -145,23 +141,36 @@ void ULandscapeComponent::Serialize(FArchive& Ar)
 			if (!PlatformData.HasValidPlatformData())
 			{
 				GeneratePlatformVertexData();
-				MaterialInstance = Cast<UMaterialInstanceConstant>(GeneratePlatformPixelData(WeightmapTextures, true));
-				HeightmapTexture = NULL;
-				bRestoreAfterCooking = true;
+				GeneratePlatformPixelData(true);
 			}
 		}
 	}
 #endif
 
-	Super::Serialize(Ar);
+	if (Ar.IsCooking() && !Ar.CookingTarget()->SupportsFeature(ETargetPlatformFeatures::VertexShaderTextureSampling))
+	{
+		// These properties are not used for ES2 so we back them up and clear them before serializing them.
+		UTexture2D* BackupHeightmapTexture = nullptr;
+		UTexture2D* BackupXYOffsetmapTexture = nullptr;
+		UMaterialInstanceConstant* BackupMaterialInstance = nullptr;
+		TArray<UTexture2D*> BackupWeightmapTextures;
+			
+		Exchange(HeightmapTexture, BackupHeightmapTexture);
+		Exchange(BackupXYOffsetmapTexture, XYOffsetmapTexture);
+		Exchange(BackupMaterialInstance, MaterialInstance);
+		Exchange(BackupWeightmapTextures, WeightmapTextures);
 
-#if ENABLE_LANDSCAPE_COOKING && WITH_EDITOR
-	if (bRestoreAfterCooking)
+		Super::Serialize(Ar);
+
+		Exchange(HeightmapTexture, BackupHeightmapTexture);
+		Exchange(BackupXYOffsetmapTexture, XYOffsetmapTexture);
+		Exchange(BackupMaterialInstance, MaterialInstance);
+		Exchange(BackupWeightmapTextures, WeightmapTextures);
+	}
+	else
 	{
-		MaterialInstance = BackupMaterialInstance;
-		HeightmapTexture = BackupHeightmapTexture;
+		Super::Serialize(Ar);
 	}
-#endif
 
 	Ar << LightMap;
 	Ar << ShadowMap;
@@ -232,6 +241,11 @@ void ULandscapeComponent::Serialize(FArchive& Ar)
 			{
 				Ar << PlatformData;
 			}
+			if (Ar.UE4Ver() >= VER_UE4_SERIALIZE_LANDSCAPE_ES2_TEXTURES)
+			{
+				Ar << MobileMaterialInterface;
+				Ar << MobileWeightNormalmapTexture;
+			}
 		}
 		else if (!FPlatformProperties::SupportsVertexShaderTextureSampling())
 		{
@@ -243,6 +257,11 @@ void ULandscapeComponent::Serialize(FArchive& Ar)
 			{
 				Ar << PlatformData;
 			}
+			if (Ar.UE4Ver() >= VER_UE4_SERIALIZE_LANDSCAPE_ES2_TEXTURES)
+			{
+				Ar << MobileMaterialInterface;
+				Ar << MobileWeightNormalmapTexture;
+			}
 		}
 		if (Ar.UE4Ver() >= VER_UE4_LANDSCAPE_GRASS_COOKING && Ar.UE4Ver() < VER_UE4_SERIALIZE_LANDSCAPE_GRASS_DATA)
 		{
diff --git a/Engine/Source/Runtime/Landscape/Private/LandscapeEdit.cpp b/Engine/Source/Runtime/Landscape/Private/LandscapeEdit.cpp
--- a/Engine/Source/Runtime/Landscape/Private/LandscapeEdit.cpp
+++ b/Engine/Source/Runtime/Landscape/Private/LandscapeEdit.cpp
@@ -4391,11 +4391,30 @@ struct FMobileLayerAllocation
 	}
 };
 
-UMaterialInstance* ULandscapeComponent::GeneratePlatformPixelData(TArray<UTexture2D*>& InWeightmapTextures, bool bIsCooking)
+void ULandscapeComponent::GeneratePlatformPixelData(bool bIsCooking)
 {
-	if (IsTemplate() || HeightmapTexture == NULL)
+	check(!IsTemplate())
+
+	if (!bIsCooking)
 	{
-		return MaterialInstance;
+		// Calculate hash of source data and skip generation if the data we have in memory is unchanged
+		FBufferArchive ComponentStateAr;
+		SerializeStateHashes(ComponentStateAr);
+
+		uint32 Hash[5];
+		FSHA1::HashBuffer(ComponentStateAr.GetData(), ComponentStateAr.Num(), (uint8*)Hash);
+		FGuid NewSourceHash = FGuid(Hash[0] ^ Hash[4], Hash[1], Hash[2], Hash[3]);
+	
+		// Skip generation if the source hash matches
+		if (MobilePixelDataSourceHash.IsValid() && 
+			MobilePixelDataSourceHash == NewSourceHash &&
+			MobileMaterialInterface != nullptr &&
+			MobileWeightNormalmapTexture != nullptr)
+		{
+			return;
+		}
+			
+		MobilePixelDataSourceHash = NewSourceHash;
 	}
 
 	TArray<FMobileLayerAllocation> MobileLayerAllocations;
@@ -4442,13 +4461,13 @@ UMaterialInstance* ULandscapeComponent::GeneratePlatformPixelData(TArray<UTextur
 	}
 
 	int32 WeightmapSize = (SubsectionSizeQuads + 1) * NumSubsections;
-	UTexture2D* WeightmapTexture = GetLandscapeProxy()->CreateLandscapeTexture(WeightmapSize, WeightmapSize, TEXTUREGROUP_Terrain_Weightmap, TSF_BGRA8);
-	CreateEmptyTextureMips(WeightmapTexture);
+	UTexture2D* NewWeightNormalmapTexture = GetLandscapeProxy()->CreateLandscapeTexture(WeightmapSize, WeightmapSize, TEXTUREGROUP_Terrain_Weightmap, TSF_BGRA8);
+	CreateEmptyTextureMips(NewWeightNormalmapTexture);
 
 	{
 		FLandscapeEditDataInterface LandscapeEdit(GetLandscapeInfo(false));
 
-		if (InWeightmapTextures.Num() > 0)
+		if (WeightmapTextures.Num() > 0)
 		{
 			int32 CurrentIdx = 0;
 			for (const auto& MobileAllocation : MobileLayerAllocations)
@@ -4456,9 +4475,7 @@ UMaterialInstance* ULandscapeComponent::GeneratePlatformPixelData(TArray<UTextur
 				// Only for valid Layers
 				if (MobileAllocation.Allocation.LayerInfo)
 				{
-					uint8 TextureIndex = MobileAllocation.Allocation.WeightmapTextureIndex;
-					uint8 TextureChannel = MobileAllocation.Allocation.WeightmapTextureChannel;
-					LandscapeEdit.CopyTextureChannel(WeightmapTexture, CurrentIdx, InWeightmapTextures[TextureIndex], TextureChannel);
+					LandscapeEdit.CopyTextureFromWeightmap(NewWeightNormalmapTexture, CurrentIdx, this, MobileAllocation.Allocation.LayerInfo);
 					CurrentIdx++;
 					if (CurrentIdx >= 2) // Only support 2 layers in texture
 					{
@@ -4468,13 +4485,15 @@ UMaterialInstance* ULandscapeComponent::GeneratePlatformPixelData(TArray<UTextur
 			}
 		}
 
-		LandscapeEdit.CopyTextureFromHeightmap(WeightmapTexture, 2, this, 2);
-		LandscapeEdit.CopyTextureFromHeightmap(WeightmapTexture, 3, this, 3);
+		// copy normals into B/A channels.
+		LandscapeEdit.CopyTextureFromHeightmap(NewWeightNormalmapTexture, 2, this, 2);
+		LandscapeEdit.CopyTextureFromHeightmap(NewWeightNormalmapTexture, 3, this, 3);
 	}
 
-	WeightmapTexture->PostEditChange();
-	InWeightmapTextures.Empty();
-	InWeightmapTextures.Add(WeightmapTexture);
+	NewWeightNormalmapTexture->PostEditChange();
+
+	MobileWeightNormalmapTexture = NewWeightNormalmapTexture;
+
 
 	FLinearColor Masks[5];
 	Masks[0] = FLinearColor(1, 0, 0, 0);
@@ -4485,7 +4504,7 @@ UMaterialInstance* ULandscapeComponent::GeneratePlatformPixelData(TArray<UTextur
 
 	if (!bIsCooking)
 	{
-		UMaterialInstanceDynamic* MobileMaterialInstance = UMaterialInstanceDynamic::Create(MaterialInstance, GetOutermost());
+		UMaterialInstanceDynamic* NewMobileMaterialInstance = UMaterialInstanceDynamic::Create(MaterialInstance, GetOutermost());
 
 		MobileBlendableLayerMask = 0;
 
@@ -4497,19 +4516,19 @@ UMaterialInstance* ULandscapeComponent::GeneratePlatformPixelData(TArray<UTextur
 			if (Allocation.LayerInfo)
 			{
 				FName LayerName = Allocation.LayerInfo == ALandscapeProxy::VisibilityLayer ? UMaterialExpressionLandscapeVisibilityMask::ParameterName : Allocation.LayerInfo->LayerName;
-				MobileMaterialInstance->SetVectorParameterValue(FName(*FString::Printf(TEXT("LayerMask_%s"), *LayerName.ToString())), Masks[FMath::Min(4, CurrentIdx)]);
+				NewMobileMaterialInstance->SetVectorParameterValue(FName(*FString::Printf(TEXT("LayerMask_%s"), *LayerName.ToString())), Masks[FMath::Min(4, CurrentIdx)]);
 				MobileBlendableLayerMask |= (!Allocation.LayerInfo->bNoWeightBlend ? (1 << CurrentIdx) : 0);
 				CurrentIdx++;
 			}
 		}
-		return MobileMaterialInstance;
+		MobileMaterialInterface = NewMobileMaterialInstance;
 	}
 	else // for cooking
 	{
 		UMaterialInstanceConstant* CombinationMaterialInstance = GetCombinationMaterial(true);
-		UMaterialInstanceConstant* MobileMaterialInstance = NewObject<ULandscapeMaterialInstanceConstant>(GetOutermost());
+		UMaterialInstanceConstant* NewMobileMaterialInstance = NewObject<ULandscapeMaterialInstanceConstant>(GetOutermost());
 
-		MobileMaterialInstance->SetParentEditorOnly(CombinationMaterialInstance);
+		NewMobileMaterialInstance->SetParentEditorOnly(CombinationMaterialInstance);
 
 		MobileBlendableLayerMask = 0;
 
@@ -4521,15 +4540,15 @@ UMaterialInstance* ULandscapeComponent::GeneratePlatformPixelData(TArray<UTextur
 			if (Allocation.LayerInfo)
 			{
 				FName LayerName = Allocation.LayerInfo == ALandscapeProxy::VisibilityLayer ? UMaterialExpressionLandscapeVisibilityMask::ParameterName : Allocation.LayerInfo->LayerName;
-				MobileMaterialInstance->SetVectorParameterValueEditorOnly(FName(*FString::Printf(TEXT("LayerMask_%s"), *LayerName.ToString())), Masks[FMath::Min(4, CurrentIdx)]);
+				NewMobileMaterialInstance->SetVectorParameterValueEditorOnly(FName(*FString::Printf(TEXT("LayerMask_%s"), *LayerName.ToString())), Masks[FMath::Min(4, CurrentIdx)]);
 				MobileBlendableLayerMask |= (!Allocation.LayerInfo->bNoWeightBlend ? (1 << CurrentIdx) : 0);
 				CurrentIdx++;
 			}
 		}
 
-		MobileMaterialInstance->PostEditChange();
+		NewMobileMaterialInstance->PostEditChange();
 
-		return MobileMaterialInstance;
+		MobileMaterialInterface = NewMobileMaterialInstance;
 	}
 }
 
diff --git a/Engine/Source/Runtime/Landscape/Private/LandscapeEditInterface.cpp b/Engine/Source/Runtime/Landscape/Private/LandscapeEditInterface.cpp
--- a/Engine/Source/Runtime/Landscape/Private/LandscapeEditInterface.cpp
+++ b/Engine/Source/Runtime/Landscape/Private/LandscapeEditInterface.cpp
@@ -3604,6 +3604,33 @@ void FLandscapeEditDataInterface::CopyTextureFromHeightmap(UTexture2D* Dest, int
 	}
 }
 
+void FLandscapeEditDataInterface::CopyTextureFromWeightmap(UTexture2D* Dest, int32 DestChannel, ULandscapeComponent* Comp, ULandscapeLayerInfoObject* LayerInfo)
+{
+	FLandscapeTextureDataInfo* DestDataInfo = GetTextureDataInfo(Dest);
+	int32 MipSize = Dest->Source.GetSizeX();
+	check(Dest->Source.GetSizeX() == Dest->Source.GetSizeY());
+
+	// Channel remapping
+	int32 ChannelOffsets[4] = { (int32)STRUCT_OFFSET(FColor, R), (int32)STRUCT_OFFSET(FColor, G), (int32)STRUCT_OFFSET(FColor, B), (int32)STRUCT_OFFSET(FColor, A) };
+
+	for (int32 MipIdx = 0; MipIdx < DestDataInfo->NumMips(); MipIdx++)
+	{
+		FLandscapeComponentDataInterface DataInterface(Comp, MipIdx);
+		TArray<uint8> WeightData;
+		DataInterface.GetWeightmapTextureData(LayerInfo, WeightData);
+
+		uint8* DestTextureData = (uint8*)DestDataInfo->GetMipData(MipIdx) + ChannelOffsets[DestChannel];
+
+		for (int32 i = 0; i < FMath::Square(MipSize); i++)
+		{
+			DestTextureData[i * 4] = WeightData[i];
+		}
+
+		DestDataInfo->AddMipUpdateRegion(MipIdx, 0, 0, MipSize - 1, MipSize - 1);
+		MipSize >>= 1;
+	}
+}
+
 void FLandscapeEditDataInterface::ZeroTextureChannel(UTexture2D* Dest, int32 DestChannel)
 {
 	FLandscapeTextureDataInfo* DestDataInfo = GetTextureDataInfo(Dest);
diff --git a/Engine/Source/Runtime/Landscape/Private/LandscapeRenderMobile.cpp b/Engine/Source/Runtime/Landscape/Private/LandscapeRenderMobile.cpp
--- a/Engine/Source/Runtime/Landscape/Private/LandscapeRenderMobile.cpp
+++ b/Engine/Source/Runtime/Landscape/Private/LandscapeRenderMobile.cpp
@@ -216,20 +216,20 @@ void FLandscapeVertexBufferMobile::InitRHI()
 FLandscapeComponentSceneProxyMobile::FLandscapeComponentSceneProxyMobile(ULandscapeComponent* InComponent, FLandscapeEditToolRenderData* InEditToolRenderData)
 	:	FLandscapeComponentSceneProxy(InComponent, InEditToolRenderData)
 {
-		check(InComponent && InComponent->PlatformData.HasValidPlatformData());
-		InComponent->PlatformData.GetUncompressedData(PlatformData);
+	check(InComponent && InComponent->PlatformData.HasValidPlatformData());
+	InComponent->PlatformData.GetUncompressedData(PlatformData);
+
 #if WITH_EDITOR
-	if (InComponent->HeightmapTexture != NULL)
-	{
-		MaterialInterface = InComponent->GeneratePlatformPixelData(WeightmapTextures, false);
-		InComponent->MobileMaterialInterface = MaterialInterface;
-		InComponent->MobileNormalmapTexture = WeightmapTextures[0];
-	}
+	InComponent->GeneratePlatformPixelData(false);
 #endif
-	if (WeightmapTextures.Num() > 0)
-	{
-		NormalmapTexture = WeightmapTextures[0]; // Use Weightmap0 as Normal/Weightmap Texture
-	}
+	check(InComponent->MobileMaterialInterface);
+	check(InComponent->MobileWeightNormalmapTexture);
+	
+	MaterialInterface = InComponent->MobileMaterialInterface;
+
+	WeightmapTextures.Empty(1);
+	WeightmapTextures.Add(InComponent->MobileWeightNormalmapTexture);
+	NormalmapTexture = InComponent->MobileWeightNormalmapTexture;
 
 	BlendableLayerMask = InComponent->MobileBlendableLayerMask;
 }
diff --git a/Engine/Source/Runtime/Landscape/Public/LandscapeEdit.h b/Engine/Source/Runtime/Landscape/Public/LandscapeEdit.h
--- a/Engine/Source/Runtime/Landscape/Public/LandscapeEdit.h
+++ b/Engine/Source/Runtime/Landscape/Public/LandscapeEdit.h
@@ -163,6 +163,7 @@ struct LANDSCAPE_API FLandscapeEditDataInterface
 	void CopyTextureChannel(UTexture2D* Dest, int32 DestChannel, UTexture2D* Src, int32 SrcChannel);
 	void ZeroTextureChannel(UTexture2D* Dest, int32 DestChannel);
 	void CopyTextureFromHeightmap(UTexture2D* Dest, int32 DestChannel, ULandscapeComponent* Comp, int32 SrcChannel);
+	void CopyTextureFromWeightmap(UTexture2D* Dest, int32 DestChannel, ULandscapeComponent* Comp, ULandscapeLayerInfoObject* LayerInfo);
 
 	template<typename TData>
 	void SetTextureValueTempl(UTexture2D* Dest, TData Value);
