Properly detect that all required landscape textures are streamed-in before rendereing grass maps

diff --git a/Engine/Source/Runtime/Landscape/Private/LandscapeGrass.cpp b/Engine/Source/Runtime/Landscape/Private/LandscapeGrass.cpp
--- a/Engine/Source/Runtime/Landscape/Private/LandscapeGrass.cpp
+++ b/Engine/Source/Runtime/Landscape/Private/LandscapeGrass.cpp
@@ -569,14 +569,32 @@ bool ULandscapeComponent::CanRenderGrassMap() const
 	return true;
 }
 
+static bool IsTextureStreamedForGrassMapRender(UTexture2D* InTexture)
+{
+	if (!InTexture || InTexture->ResidentMips != InTexture->GetNumMips()
+		|| !InTexture->Resource || ((FTexture2DResource*)InTexture->Resource)->GetCurrentFirstMip() > 0)
+	{
+		return false;
+	}
+	return true;
+}
+
 bool ULandscapeComponent::AreTexturesStreamedForGrassMapRender() const
 {
 	// Check for valid heightmap that is fully streamed in
-	if (!HeightmapTexture || HeightmapTexture->ResidentMips != HeightmapTexture->GetNumMips()
-		|| !HeightmapTexture->Resource || ((FTexture2DResource*)HeightmapTexture->Resource)->GetCurrentFirstMip() > 0)
+	if (!IsTextureStreamedForGrassMapRender(HeightmapTexture))
 	{
 		return false;
 	}
+	
+	// Check for valid weightmaps that is fully streamed in
+	for (auto WeightmapTexture : WeightmapTextures)
+	{
+		if (!IsTextureStreamedForGrassMapRender(WeightmapTexture))
+		{
+			return false;
+		}
+	}
 
 	return true;
 }
