Restore landscape physical materials and hole collision after changing the CollisionMipLeve property or running RecreateLandscapeCollision console command. (UE-9875)

diff --git a/Engine/Source/Runtime/Landscape/Private/LandscapeEdit.cpp b/Engine/Source/Runtime/Landscape/Private/LandscapeEdit.cpp
--- a/Engine/Source/Runtime/Landscape/Private/LandscapeEdit.cpp
+++ b/Engine/Source/Runtime/Landscape/Private/LandscapeEdit.cpp
@@ -764,14 +764,19 @@ void ULandscapeComponent::UpdateCollisionHeightData(const FColor* HeightmapTextu
 		AInstancedFoliageActor::MoveInstancesToNewComponent(Proxy->GetWorld(), OldCollisionComponent, CollisionComp);
 	}
 
+	// Set new collision component to pointer
+	CollisionComponent = CollisionComp;
+
+	if (bRebuild)
+	{
+		UpdateCollisionLayerData();
+	}
+
 	if (bRebuild && CollisionProxy)
 	{
 		CollisionProxy->RegisterAllComponents();
 	}
 
-	// Set new collision component to pointer
-	CollisionComponent = CollisionComp;
-
 	if (ChangeType || CreatedNew)
 	{
 		Proxy->RegisterAllComponents();
@@ -3043,7 +3048,8 @@ void ALandscapeProxy::RecreateCollisionComponents()
 			Comp->CollisionMipLevel = CollisionMipLevel;
 			TArray<uint8> CollisionMipData;
 			Comp->HeightmapTexture->Source.GetMipData(CollisionMipData, CollisionMipLevel);
-			Comp->UpdateCollisionHeightData((FColor*)CollisionMipData.GetData(), 0, 0, MAX_int32, MAX_int32, true, NULL, true); // Rebuild for new CollisionMipLevel
+			// Rebuild all collision
+			Comp->UpdateCollisionHeightData((FColor*)CollisionMipData.GetData(), 0, 0, MAX_int32, MAX_int32, true, NULL, true); 
 		}
 	}
 }
