Fix crash generating grass for landscapes with XY offset map (UE-15252)

diff --git a/Engine/Source/Runtime/Landscape/Private/LandscapeGrass.cpp b/Engine/Source/Runtime/Landscape/Private/LandscapeGrass.cpp
--- a/Engine/Source/Runtime/Landscape/Private/LandscapeGrass.cpp
+++ b/Engine/Source/Runtime/Landscape/Private/LandscapeGrass.cpp
@@ -99,6 +99,16 @@ DECLARE_CYCLE_STAT(TEXT("Grass Update"), STAT_GrassUpdate, STATGROUP_Foliage);
 //
 // Grass weightmap rendering
 //
+
+static bool ShouldCacheLandscapeGrassShaders(EShaderPlatform Platform, const FMaterial* Material, const FVertexFactoryType* VertexFactoryType)
+{
+	// We only need grass weight shaders for Landscape vertex factories on desktop platforms
+	return (Material->IsUsedWithLandscape() || Material->IsSpecialEngineMaterial()) &&
+		IsFeatureLevelSupported(Platform, ERHIFeatureLevel::SM4) &&
+		((VertexFactoryType == FindVertexFactoryType(FName(TEXT("FLandscapeVertexFactory"), FNAME_Find))) || (VertexFactoryType == FindVertexFactoryType(FName(TEXT("FLandscapeXYOffsetVertexFactory"), FNAME_Find))))
+		&& !IsConsolePlatform(Platform);
+}
+
 class FLandscapeGrassWeightVS : public FMeshMaterialShader
 {
 	DECLARE_SHADER_TYPE(FLandscapeGrassWeightVS, MeshMaterial);
@@ -120,8 +130,7 @@ public:
 
 	static bool ShouldCache(EShaderPlatform Platform, const FMaterial* Material, const FVertexFactoryType* VertexFactoryType)
 	{
-		// Only compile for LandscapeVertexFactory
-		return (Material->IsUsedWithLandscape() || Material->IsSpecialEngineMaterial()) && IsFeatureLevelSupported(Platform, ERHIFeatureLevel::SM4) && (VertexFactoryType == FindVertexFactoryType(FName(TEXT("FLandscapeVertexFactory"), FNAME_Find))) && !IsConsolePlatform(Platform);
+		return ShouldCacheLandscapeGrassShaders(Platform, Material, VertexFactoryType);
 	}
 
 	void SetParameters(FRHICommandList& RHICmdList, const FMaterialRenderProxy* MaterialRenderProxy, const FMaterial& MaterialResource, const FSceneView& View, const FVector2D& RenderOffset)
@@ -153,8 +162,7 @@ public:
 
 	static bool ShouldCache(EShaderPlatform Platform, const FMaterial* Material, const FVertexFactoryType* VertexFactoryType)
 	{
-		// Only compile for LandscapeVertexFactory
-		return (Material->IsUsedWithLandscape() || Material->IsSpecialEngineMaterial()) && IsFeatureLevelSupported(Platform, ERHIFeatureLevel::SM4) && (VertexFactoryType == FindVertexFactoryType(FName(TEXT("FLandscapeVertexFactory"), FNAME_Find))) && !IsConsolePlatform(Platform);
+		return ShouldCacheLandscapeGrassShaders(Platform, Material, VertexFactoryType);
 	}
 
 	FLandscapeGrassWeightPS(const ShaderMetaType::CompiledShaderInitializerType& Initializer)
