Fixed: Landscape LODDistanceFactor not applied to levels that were not loaded when it is changed
UE-11825

diff --git a/Engine/Source/Runtime/Landscape/Classes/LandscapeProxy.h b/Engine/Source/Runtime/Landscape/Classes/LandscapeProxy.h
--- a/Engine/Source/Runtime/Landscape/Classes/LandscapeProxy.h
+++ b/Engine/Source/Runtime/Landscape/Classes/LandscapeProxy.h
@@ -609,7 +609,9 @@ public:
 
 	// Copy properties from parent Landscape actor
 	LANDSCAPE_API void GetSharedProperties(ALandscapeProxy* Landscape);
-
+	// Assign only mismatched properties and mark proxy package dirty
+	LANDSCAPE_API void ConditionalAssignCommonProperties(ALandscape* Landscape);
+	
 	/** Get the LandcapeActor-to-world transform with respect to landscape section offset*/
 	LANDSCAPE_API FTransform LandscapeActorToWorld() const;
 	
diff --git a/Engine/Source/Runtime/Landscape/Private/Landscape.cpp b/Engine/Source/Runtime/Landscape/Private/Landscape.cpp
--- a/Engine/Source/Runtime/Landscape/Private/Landscape.cpp
+++ b/Engine/Source/Runtime/Landscape/Private/Landscape.cpp
@@ -1425,6 +1425,39 @@ void ALandscapeProxy::GetSharedProperties(ALandscapeProxy* Landscape)
 	}
 }
 
+void ALandscapeProxy::ConditionalAssignCommonProperties(ALandscape* Landscape)
+{
+	if (Landscape == nullptr)
+	{
+		return;
+	}
+	
+	bool bUpdated = false;
+	
+	if (MaxLODLevel != Landscape->MaxLODLevel)
+	{
+		MaxLODLevel = Landscape->MaxLODLevel;
+		bUpdated = true;
+	}
+	
+	if (LODDistanceFactor != Landscape->LODDistanceFactor)
+	{
+		LODDistanceFactor = Landscape->LODDistanceFactor;
+		bUpdated = true;
+	}
+	
+	if (LODFalloff != Landscape->LODFalloff)
+	{
+		LODFalloff = Landscape->LODFalloff;
+		bUpdated = true;
+	}
+
+	if (bUpdated)
+	{
+		MarkPackageDirty();
+	}
+}
+
 FTransform ALandscapeProxy::LandscapeActorToWorld() const
 {
 	FTransform TM = ActorToWorld();
@@ -1684,12 +1717,14 @@ void ULandscapeInfo::RegisterActor(ALandscapeProxy* Proxy, bool bMapCheck)
 		for (auto It = Proxies.CreateConstIterator(); It; ++It)
 		{
 			(*It)->LandscapeActor = LandscapeActor;
+			(*It)->ConditionalAssignCommonProperties(Landscape);
 		}
 	}
 	else
 	{
 		Proxies.Add(Proxy);
 		Proxy->LandscapeActor = LandscapeActor;
+		Proxy->ConditionalAssignCommonProperties(LandscapeActor.Get());
 	}
 
 	//
