Fix null landscape layers not being deleted correctly
- this was causing a crash later when trying to use the paint/visibility tools

diff --git a/Engine/Source/Runtime/Landscape/Private/LandscapeEdit.cpp b/Engine/Source/Runtime/Landscape/Private/LandscapeEdit.cpp
--- a/Engine/Source/Runtime/Landscape/Private/LandscapeEdit.cpp
+++ b/Engine/Source/Runtime/Landscape/Private/LandscapeEdit.cpp
@@ -368,38 +368,35 @@ void ULandscapeComponent::FixupWeightmaps()
 
 		if (Info)
 		{
-			TArray<ULandscapeLayerInfoObject*> DeletedLayers;
+			TArray<ULandscapeLayerInfoObject*> LayersToDelete;
 			bool bFixedLayerDeletion = false;
 
-			if (Info->Layers.Num() && Cast<ALandscape>(Proxy))
+			// LayerInfo Validation check...
+			for (const auto& Allocation : WeightmapLayerAllocations)
 			{
-				// LayerName Validation check...
-				for (int32 LayerIdx = 0; LayerIdx < WeightmapLayerAllocations.Num(); LayerIdx++)
+				if (!Allocation.LayerInfo
+					|| (Allocation.LayerInfo != ALandscapeProxy::VisibilityLayer && Info->GetLayerInfoIndex(Allocation.LayerInfo) == INDEX_NONE))
 				{
-					if (!WeightmapLayerAllocations[LayerIdx].LayerInfo
-						|| (WeightmapLayerAllocations[LayerIdx].LayerInfo != ALandscapeProxy::VisibilityLayer && Info->GetLayerInfoIndex(WeightmapLayerAllocations[LayerIdx].LayerInfo) == INDEX_NONE))
+					if (!bFixedLayerDeletion)
 					{
-						if (!bFixedLayerDeletion)
-						{
-							FFormatNamedArguments Arguments;
-							Arguments.Add(TEXT("LandscapeName"), FText::FromString(GetName()));
-							FMessageLog("MapCheck").Warning()
-								->AddToken(FTextToken::Create(FText::Format(LOCTEXT("MapCheck_Message_FixedUpDeletedLayerWeightmap", "{LandscapeName} : Fixed up deleted layer weightmap"), Arguments)))
-								->AddToken(FMapErrorToken::Create(FMapErrors::FixedUpDeletedLayerWeightmap));
-						}
-
-						bFixedLayerDeletion = true;
-						DeletedLayers.Add(WeightmapLayerAllocations[LayerIdx].LayerInfo);
+						FFormatNamedArguments Arguments;
+						Arguments.Add(TEXT("LandscapeName"), FText::FromString(GetName()));
+						FMessageLog("MapCheck").Warning()
+							->AddToken(FTextToken::Create(FText::Format(LOCTEXT("MapCheck_Message_FixedUpDeletedLayerWeightmap", "{LandscapeName} : Fixed up deleted layer weightmap"), Arguments)))
+							->AddToken(FMapErrorToken::Create(FMapErrors::FixedUpDeletedLayerWeightmap));
 					}
+
+					bFixedLayerDeletion = true;
+					LayersToDelete.Add(Allocation.LayerInfo);
 				}
 			}
 
 			if (bFixedLayerDeletion)
 			{
 				FLandscapeEditDataInterface LandscapeEdit(Info);
-				for (int32 Idx = 0; Idx < DeletedLayers.Num(); ++Idx)
+				for (int32 Idx = 0; Idx < LayersToDelete.Num(); ++Idx)
 				{
-					DeleteLayer(DeletedLayers[Idx], &LandscapeEdit);
+					DeleteLayer(LayersToDelete[Idx], &LandscapeEdit);
 				}
 			}
 
