Cleaned up ULandscapeComponent::DeleteLayer

diff --git a/Engine/Source/Runtime/Landscape/Private/LandscapeEditInterface.cpp b/Engine/Source/Runtime/Landscape/Private/LandscapeEditInterface.cpp
--- a/Engine/Source/Runtime/Landscape/Private/LandscapeEditInterface.cpp
+++ b/Engine/Source/Runtime/Landscape/Private/LandscapeEditInterface.cpp
@@ -1535,31 +1535,16 @@ void FLandscapeEditDataInterface::GetHeightDataFast(const int32 X1, const int32
 	}
 }
 
-void ULandscapeComponent::DeleteLayer( ULandscapeLayerInfoObject* LayerInfo, struct FLandscapeEditDataInterface* LandscapeEdit )
+void ULandscapeComponent::DeleteLayer(ULandscapeLayerInfoObject* LayerInfo, FLandscapeEditDataInterface* LandscapeEdit)
 {
-	ULandscapeComponent* Component = this; //Landscape->LandscapeComponents(ComponentIdx);
+	ULandscapeComponent* Component = this;
 
 	// Find the index for this layer in this component.
-	int32 DeleteLayerIdx = INDEX_NONE;
-
-	for( int32 LayerIdx=0;LayerIdx<Component->WeightmapLayerAllocations.Num();LayerIdx++ )
-	{
-		FWeightmapLayerAllocationInfo& Allocation = Component->WeightmapLayerAllocations[LayerIdx];
-		if( Allocation.LayerInfo == LayerInfo )
-		{
-			DeleteLayerIdx = LayerIdx;
-		}
-	}
-	if( DeleteLayerIdx == INDEX_NONE )
+	const int32 DeleteLayerIdx = Component->WeightmapLayerAllocations.IndexOfByPredicate(
+		[LayerInfo](const FWeightmapLayerAllocationInfo& Allocation) { return Allocation.LayerInfo == LayerInfo; });
+	if (DeleteLayerIdx == INDEX_NONE)
 	{
 		// Layer not used for this component.
-		//continue;
-		return;
-	}
-
-	ULandscapeInfo* Info = GetLandscapeInfo();
-	if (!Info)
-	{
 		return;
 	}
 
@@ -1568,12 +1553,12 @@ void ULandscapeComponent::DeleteLayer( ULandscapeLayerInfoObject* LayerInfo, str
 
 	// See if we'll be able to remove the texture completely.
 	bool bCanRemoveLayerTexture = true;
-	for( int32 LayerIdx=0;LayerIdx<Component->WeightmapLayerAllocations.Num();LayerIdx++ )
+	for (int32 LayerIdx = 0; LayerIdx < Component->WeightmapLayerAllocations.Num(); LayerIdx++)
 	{
 		FWeightmapLayerAllocationInfo& Allocation = Component->WeightmapLayerAllocations[LayerIdx];
 
 		// check if we will be able to remove the texture also
-		if( LayerIdx!=DeleteLayerIdx && Allocation.WeightmapTextureIndex == DeleteLayerWeightmapTextureIndex )
+		if (LayerIdx != DeleteLayerIdx && Allocation.WeightmapTextureIndex == DeleteLayerWeightmapTextureIndex)
 		{
 			bCanRemoveLayerTexture = false;
 		}
@@ -1582,70 +1567,67 @@ void ULandscapeComponent::DeleteLayer( ULandscapeLayerInfoObject* LayerInfo, str
 	// See if the deleted layer is a NoWeightBlend layer - if not, we don't have to worry about normalization
 	bool bDeleteLayerIsNoWeightBlend = (LayerInfo && LayerInfo->bNoWeightBlend);
 
-	if( !bDeleteLayerIsNoWeightBlend )
+	if (!bDeleteLayerIsNoWeightBlend)
 	{
 		// Lock data for all the weightmaps
 		TArray<FLandscapeTextureDataInfo*> TexDataInfos;
-		for( int32 WeightmapIdx=0;WeightmapIdx < Component->WeightmapTextures.Num();WeightmapIdx++ )
+		for (int32 WeightmapIdx = 0; WeightmapIdx < Component->WeightmapTextures.Num(); WeightmapIdx++)
 		{
-			if (LandscapeEdit)
-			{
-				TexDataInfos.Add(LandscapeEdit->GetTextureDataInfo(Component->WeightmapTextures[WeightmapIdx]));
-			}
+			TexDataInfos.Add(LandscapeEdit->GetTextureDataInfo(Component->WeightmapTextures[WeightmapIdx]));
 		}
 
 		// Channel remapping
-		int32 ChannelOffsets[4] = {(int32)STRUCT_OFFSET(FColor,R),(int32)STRUCT_OFFSET(FColor,G),(int32)STRUCT_OFFSET(FColor,B),(int32)STRUCT_OFFSET(FColor,A)};
+		const int32 ChannelOffsets[4] = {(int32)STRUCT_OFFSET(FColor,R), (int32)STRUCT_OFFSET(FColor,G), (int32)STRUCT_OFFSET(FColor,B), (int32)STRUCT_OFFSET(FColor,A)};
 
 		TArray<bool> LayerNoWeightBlends;	// Array of NoWeightBlend flags
 		TArray<uint8*> LayerDataPtrs;		// Pointers to all layers' data 
 
 		// Get the data for each layer
-		for( int32 LayerIdx=0;LayerIdx<Component->WeightmapLayerAllocations.Num();LayerIdx++ )
+		for (int32 LayerIdx = 0; LayerIdx < Component->WeightmapLayerAllocations.Num(); LayerIdx++)
 		{
 			FWeightmapLayerAllocationInfo& Allocation = Component->WeightmapLayerAllocations[LayerIdx];
-			LayerDataPtrs.Add( (uint8*)TexDataInfos[Allocation.WeightmapTextureIndex]->GetMipData(0) + ChannelOffsets[Allocation.WeightmapTextureChannel] );
+			LayerDataPtrs.Add((uint8*)TexDataInfos[Allocation.WeightmapTextureIndex]->GetMipData(0) + ChannelOffsets[Allocation.WeightmapTextureChannel]);
 
 			// Find the layer info and record if it is a bNoWeightBlend layer.
-			LayerNoWeightBlends.Add( Allocation.LayerInfo && Allocation.LayerInfo->bNoWeightBlend );
+			LayerNoWeightBlends.Add(Allocation.LayerInfo && Allocation.LayerInfo->bNoWeightBlend);
 		}
 
 		// Find the texture data corresponding to this vertex
-		int32 SizeU = (SubsectionSizeQuads+1) * NumSubsections; //Component->WeightmapTextures(0)->SizeX;		// not exactly correct.  
-		int32 SizeV = (SubsectionSizeQuads+1) * NumSubsections; //Component->WeightmapTextures(0)->SizeY;
-		int32 WeightmapOffsetX = Component->WeightmapScaleBias.Z * (float)SizeU;
-		int32 WeightmapOffsetY = Component->WeightmapScaleBias.W * (float)SizeV;
+		const int32 SizeU = (SubsectionSizeQuads + 1) * NumSubsections;
+		const int32 SizeV = (SubsectionSizeQuads + 1) * NumSubsections;
+		const int32 WeightmapOffsetX = Component->WeightmapScaleBias.Z * (float)SizeU;
+		const int32 WeightmapOffsetY = Component->WeightmapScaleBias.W * (float)SizeV;
 
-		for( int32 SubIndexY=0;SubIndexY<NumSubsections;SubIndexY++ )
+		for (int32 SubIndexY = 0; SubIndexY < NumSubsections; SubIndexY++)
 		{
-			for( int32 SubIndexX=0;SubIndexX<NumSubsections;SubIndexX++ )
+			for (int32 SubIndexX = 0; SubIndexX < NumSubsections; SubIndexX++)
 			{
-				for( int32 SubY=0;SubY<=SubsectionSizeQuads;SubY++ )
+				for (int32 SubY = 0; SubY <= SubsectionSizeQuads; SubY++)
 				{
-					for( int32 SubX=0;SubX<=SubsectionSizeQuads;SubX++ )
+					for (int32 SubX = 0; SubX <= SubsectionSizeQuads; SubX++)
 					{
-						int32 TexX = WeightmapOffsetX + (SubsectionSizeQuads+1) * SubIndexX + SubX;
-						int32 TexY = WeightmapOffsetY + (SubsectionSizeQuads+1) * SubIndexY + SubY;
-						int32 TexDataIndex = 4 * (TexX + TexY * SizeU);
+						const int32 TexX = WeightmapOffsetX + (SubsectionSizeQuads + 1) * SubIndexX + SubX;
+						const int32 TexY = WeightmapOffsetY + (SubsectionSizeQuads + 1) * SubIndexY + SubY;
+						const int32 TexDataIndex = 4 * (TexX + TexY * SizeU);
 
 						// Calculate the sum of other layer weights
 						int32 OtherLayerWeightSum = 0;
-						for( int32 LayerIdx=0;LayerIdx<Component->WeightmapLayerAllocations.Num();LayerIdx++ )
+						for (int32 LayerIdx = 0; LayerIdx < Component->WeightmapLayerAllocations.Num(); LayerIdx++)
 						{
-							if( LayerIdx != DeleteLayerIdx && LayerNoWeightBlends[LayerIdx]==false )
+							if (LayerIdx != DeleteLayerIdx && LayerNoWeightBlends[LayerIdx] == false)
 							{
 								OtherLayerWeightSum += LayerDataPtrs[LayerIdx][TexDataIndex];
 							}
 						}
 
-						if( OtherLayerWeightSum == 0 )
+						if (OtherLayerWeightSum == 0)
 						{
 							// Set the first other weight-blend layer we can find to 255 to avoid a black hole
 							// This isn't ideal but it's the best option
 							// There's nothing we can easily do if this was the only weight-blend layer on this component
-							for( int32 LayerIdx=0;LayerIdx<Component->WeightmapLayerAllocations.Num();LayerIdx++ )
+							for (int32 LayerIdx = 0; LayerIdx < Component->WeightmapLayerAllocations.Num(); LayerIdx++)
 							{
-								if( LayerIdx != DeleteLayerIdx && LayerNoWeightBlends[LayerIdx]==false )
+								if (LayerIdx != DeleteLayerIdx && LayerNoWeightBlends[LayerIdx] == false)
 								{
 									uint8& Weight = LayerDataPtrs[LayerIdx][TexDataIndex];
 									Weight = 255;
@@ -1656,12 +1638,12 @@ void ULandscapeComponent::DeleteLayer( ULandscapeLayerInfoObject* LayerInfo, str
 						else
 						{
 							// Adjust other layer weights
-							for( int32 LayerIdx=0;LayerIdx<Component->WeightmapLayerAllocations.Num();LayerIdx++ )
+							for (int32 LayerIdx = 0; LayerIdx < Component->WeightmapLayerAllocations.Num(); LayerIdx++)
 							{
-								if( LayerIdx != DeleteLayerIdx && LayerNoWeightBlends[LayerIdx]==false )
+								if (LayerIdx != DeleteLayerIdx && LayerNoWeightBlends[LayerIdx] == false)
 								{
 									uint8& Weight = LayerDataPtrs[LayerIdx][TexDataIndex];
-									Weight = FMath::Clamp<int32>( FMath::RoundToInt(255.0f * (float)Weight/(float)OtherLayerWeightSum), 0, 255 );
+									Weight = FMath::Clamp<int32>(FMath::RoundToInt(255.0f * (float)Weight / (float)OtherLayerWeightSum), 0, 255);
 								}
 							}
 						}
@@ -1671,9 +1653,9 @@ void ULandscapeComponent::DeleteLayer( ULandscapeLayerInfoObject* LayerInfo, str
 		}
 
 		// Update all the textures and mips
-		for( int32 Idx=0;Idx<Component->WeightmapTextures.Num();Idx++)
+		for (int32 Idx = 0; Idx < Component->WeightmapTextures.Num(); Idx++)
 		{
-			if( bCanRemoveLayerTexture && Idx==DeleteLayerWeightmapTextureIndex )
+			if (bCanRemoveLayerTexture && Idx == DeleteLayerWeightmapTextureIndex)
 			{
 				// We're going to remove this texture anyway, so don't bother updating
 				continue;
@@ -1682,34 +1664,33 @@ void ULandscapeComponent::DeleteLayer( ULandscapeLayerInfoObject* LayerInfo, str
 			UTexture2D* WeightmapTexture = Component->WeightmapTextures[Idx];
 			FLandscapeTextureDataInfo* WeightmapDataInfo = TexDataInfos[Idx];
 
-			int32 NumMips = WeightmapTexture->Source.GetNumMips();
+			const int32 NumMips = WeightmapTexture->Source.GetNumMips();
 			TArray<FColor*> WeightmapTextureMipData;
 			WeightmapTextureMipData.AddUninitialized(NumMips);
-			for( int32 MipIdx=0;MipIdx<NumMips;MipIdx++ )
+			for (int32 MipIdx = 0; MipIdx < NumMips; MipIdx++)
 			{
 				WeightmapTextureMipData[MipIdx] = (FColor*)WeightmapDataInfo->GetMipData(MipIdx);
 			}
 
 			ULandscapeComponent::UpdateWeightmapMips(Component->NumSubsections, Component->SubsectionSizeQuads, WeightmapTexture, WeightmapTextureMipData, 0, 0, MAX_int32, MAX_int32, WeightmapDataInfo);
 
-			WeightmapDataInfo->AddMipUpdateRegion(0,0,0,WeightmapTexture->Source.GetSizeX()-1,WeightmapTexture->Source.GetSizeY()-1);
+			WeightmapDataInfo->AddMipUpdateRegion(0, 0, 0, WeightmapTexture->Source.GetSizeX() - 1, WeightmapTexture->Source.GetSizeY() - 1);
 		}
 	}
 
-	ALandscapeProxy* Proxy = GetLandscapeProxy(); //CastChecked<ALandscapeProxy>(Component->GetOuter());
 	// Mark the channel as unallocated, so we can reuse it later
+	ALandscapeProxy* Proxy = GetLandscapeProxy();
 	FLandscapeWeightmapUsage* Usage = Proxy->WeightmapUsageMap.Find(Component->WeightmapTextures[DeleteLayerAllocation.WeightmapTextureIndex]);
-	//check(Usage);
-	if (Usage)
+	if (Usage) // can be null if WeightmapUsageMap hasn't been built yet
 	{
-		Usage->ChannelUsage[DeleteLayerAllocation.WeightmapTextureChannel] = NULL;
+		Usage->ChannelUsage[DeleteLayerAllocation.WeightmapTextureChannel] = nullptr;
 	}
 
 	// Remove the layer
 	Component->WeightmapLayerAllocations.RemoveAt(DeleteLayerIdx);
 
 	// If this layer was the last usage for this channel in this layer, we can remove it.
-	if( bCanRemoveLayerTexture )
+	if (bCanRemoveLayerTexture)
 	{
 		Component->WeightmapTextures[DeleteLayerWeightmapTextureIndex]->SetFlags(RF_Transactional);
 		Component->WeightmapTextures[DeleteLayerWeightmapTextureIndex]->Modify();
@@ -1719,38 +1700,38 @@ void ULandscapeComponent::DeleteLayer( ULandscapeLayerInfoObject* LayerInfo, str
 		Component->WeightmapTextures.RemoveAt(DeleteLayerWeightmapTextureIndex);
 
 		// Adjust WeightmapTextureIndex index for other layers
-		for( int32 LayerIdx=0;LayerIdx<Component->WeightmapLayerAllocations.Num();LayerIdx++ )
+		for (int32 LayerIdx = 0; LayerIdx < Component->WeightmapLayerAllocations.Num(); LayerIdx++)
 		{
 			FWeightmapLayerAllocationInfo& Allocation = Component->WeightmapLayerAllocations[LayerIdx];
 
-			if( Allocation.WeightmapTextureIndex > DeleteLayerWeightmapTextureIndex )
+			if (Allocation.WeightmapTextureIndex > DeleteLayerWeightmapTextureIndex)
 			{
 				Allocation.WeightmapTextureIndex--;
 			}
 
-			check( Allocation.WeightmapTextureIndex < Component->WeightmapTextures.Num() );
+			check(Allocation.WeightmapTextureIndex < Component->WeightmapTextures.Num());
 		}
 	}
 
 	// Update the shaders for this component
 	Component->UpdateMaterialInstances();
+
+	// Update dominant layer info stored in collision component
+	TArray<FColor*> CollisionWeightmapMipData;
+	for (int32 WeightmapIdx = 0; WeightmapIdx < Component->WeightmapTextures.Num(); WeightmapIdx++)
+	{
+		CollisionWeightmapMipData.Add((FColor*)LandscapeEdit->GetTextureDataInfo(Component->WeightmapTextures[WeightmapIdx])->GetMipData(Component->CollisionMipLevel));
+	}
+	Component->UpdateCollisionLayerData(CollisionWeightmapMipData);
 }
 
-void FLandscapeEditDataInterface::DeleteLayer( ULandscapeLayerInfoObject* LayerInfo )
+void FLandscapeEditDataInterface::DeleteLayer(ULandscapeLayerInfoObject* LayerInfo)
 {
 	if (!LandscapeInfo) return;
-	for( auto It = LandscapeInfo->XYtoComponentMap.CreateIterator(); It; ++It )
+	for (auto& XYComponentPair : LandscapeInfo->XYtoComponentMap)
 	{
-		ULandscapeComponent* Component = It.Value();
+		ULandscapeComponent* Component = XYComponentPair.Value;
 		Component->DeleteLayer(LayerInfo, this);
-
-		// Update dominant layer info stored in collision component
-		TArray<FColor*> CollisionWeightmapMipData;
-		for( int32 WeightmapIdx=0;WeightmapIdx < Component->WeightmapTextures.Num();WeightmapIdx++ )
-		{
-			CollisionWeightmapMipData.Add( (FColor*)GetTextureDataInfo(Component->WeightmapTextures[WeightmapIdx])->GetMipData(Component->CollisionMipLevel) );
-		}
-		Component->UpdateCollisionLayerData(CollisionWeightmapMipData);
 	}
 }
 
