LandscapeRenderer.cpp was missing an include for LandscapeLight.h on non unity builds.

diff --git a/Engine/Source/Runtime/Landscape/Private/LandscapeRender.cpp b/Engine/Source/Runtime/Landscape/Private/LandscapeRender.cpp
--- a/Engine/Source/Runtime/Landscape/Private/LandscapeRender.cpp
+++ b/Engine/Source/Runtime/Landscape/Private/LandscapeRender.cpp
@@ -22,6 +22,7 @@ LandscapeRender.cpp: New terrain rendering
 #include "Engine/Engine.h"
 #include "EngineGlobals.h"
 #include "UnrealEngine.h"
+#include "LandscapeLight.h"
 
 IMPLEMENT_UNIFORM_BUFFER_STRUCT(FLandscapeUniformShaderParameters, TEXT("LandscapeParameters"));
 
