Fixed crash generating grassmaps for landscape materials with PN Tessellation (UE-15619)

diff --git a/Engine/Source/Runtime/Landscape/Private/LandscapeGrass.cpp b/Engine/Source/Runtime/Landscape/Private/LandscapeGrass.cpp
--- a/Engine/Source/Runtime/Landscape/Private/LandscapeGrass.cpp
+++ b/Engine/Source/Runtime/Landscape/Private/LandscapeGrass.cpp
@@ -347,7 +347,9 @@ public:
 			{
 				for (int32 StaticMeshIdx = 0; StaticMeshIdx < PrimitiveSceneInfo->StaticMeshes.Num(); StaticMeshIdx++)
 				{
-					FMeshBatch& Mesh = *(FMeshBatch*)(&PrimitiveSceneInfo->StaticMeshes[StaticMeshIdx]);
+					FMeshBatch Mesh = *(FMeshBatch*)(&PrimitiveSceneInfo->StaticMeshes[StaticMeshIdx]);
+					// Grass maps don't render with tessellation enabled but the StaticMesh is set up to use PT_12_ControlPointPatchList
+					Mesh.Type = PT_TriangleList; 
 
 					FLandscapeGrassWeightDrawingPolicy DrawingPolicy(Mesh.VertexFactory, Mesh.MaterialRenderProxy, *Mesh.MaterialRenderProxy->GetMaterial(GMaxRHIFeatureLevel));
 					RHICmdList.BuildAndSetLocalBoundShaderState(DrawingPolicy.GetBoundShaderStateInput(GMaxRHIFeatureLevel));
