Fix for missing collision on components with XY offsets (retopologize tool) when CollisionMipLevel is modified or the RecreateLandscapeCollision console command is used. (UE-15615)

diff --git a/Engine/Source/Runtime/Landscape/Private/LandscapeEdit.cpp b/Engine/Source/Runtime/Landscape/Private/LandscapeEdit.cpp
--- a/Engine/Source/Runtime/Landscape/Private/LandscapeEdit.cpp
+++ b/Engine/Source/Runtime/Landscape/Private/LandscapeEdit.cpp
@@ -3045,8 +3045,14 @@ void ALandscapeProxy::RecreateCollisionComponents()
 			Comp->CollisionMipLevel = CollisionMipLevel;
 			TArray<uint8> CollisionMipData;
 			Comp->HeightmapTexture->Source.GetMipData(CollisionMipData, CollisionMipLevel);
+			TArray<uint8> XYOffsetMipData;
+			if (Comp->XYOffsetmapTexture)
+			{
+				Comp->XYOffsetmapTexture->Source.GetMipData(XYOffsetMipData, CollisionMipLevel);
+			}
+
 			// Rebuild all collision
-			Comp->UpdateCollisionHeightData((FColor*)CollisionMipData.GetData(), 0, 0, MAX_int32, MAX_int32, true, NULL, true); 
+			Comp->UpdateCollisionHeightData((FColor*)CollisionMipData.GetData(), 0, 0, MAX_int32, MAX_int32, true, XYOffsetMipData.Num() ? (FColor*)XYOffsetMipData.GetData() : nullptr, true);
 		}
 	}
 }
