Fix crash in ULandscapeInfo::UpdateAddCollision
Also cleaned up a lot of use of the AddCollision map, though it should probably be moved completely into the "Add Component" tool

diff --git a/Engine/Source/Editor/LandscapeEditor/Private/LandscapeEdModeComponentTools.cpp b/Engine/Source/Editor/LandscapeEditor/Private/LandscapeEdModeComponentTools.cpp
--- a/Engine/Source/Editor/LandscapeEditor/Private/LandscapeEdModeComponentTools.cpp
+++ b/Engine/Source/Editor/LandscapeEditor/Private/LandscapeEdModeComponentTools.cpp
@@ -811,12 +811,12 @@ public:
 			{
 				for (int32 ComponentIndexX = ComponentIndexX1; ComponentIndexX <= ComponentIndexX2; ComponentIndexX++)
 				{
-					ULandscapeComponent* Component = LandscapeInfo->XYtoComponentMap.FindRef(FIntPoint(ComponentIndexX, ComponentIndexY));
-					if (!Component)
+					ULandscapeComponent* LandscapeComponent = LandscapeInfo->XYtoComponentMap.FindRef(FIntPoint(ComponentIndexX, ComponentIndexY));
+					if (!LandscapeComponent)
 					{
 						// Add New component...
 						FIntPoint ComponentBase = FIntPoint(ComponentIndexX, ComponentIndexY)*Landscape->ComponentSizeQuads;
-						ULandscapeComponent* LandscapeComponent = NewObject<ULandscapeComponent>(Landscape, NAME_None, RF_Transactional);
+						LandscapeComponent = NewObject<ULandscapeComponent>(Landscape, NAME_None, RF_Transactional);
 						Landscape->LandscapeComponents.Add(LandscapeComponent);
 						NewComponents.Add(LandscapeComponent);
 						LandscapeComponent->Init(
@@ -841,6 +841,9 @@ public:
 						HeightData.AddZeroed(FMath::Square(ComponentVerts));
 						LandscapeComponent->InitHeightmapData(HeightData, true);
 						LandscapeComponent->UpdateMaterialInstances();
+
+						LandscapeInfo->XYtoComponentMap.Add(FIntPoint(ComponentIndexX, ComponentIndexY), LandscapeComponent);
+						LandscapeInfo->XYtoAddCollisionMap.Remove(FIntPoint(ComponentIndexX, ComponentIndexY));
 					}
 				}
 			}
@@ -860,13 +863,13 @@ public:
 			HeightCache.SetCachedData(X1, Y1, X2, Y2, Data);
 			HeightCache.Flush();
 
-			for (int32 Idx = 0; Idx < NewComponents.Num(); Idx++)
+			for (ULandscapeComponent* NewComponent : NewComponents)
 			{
 				// Update Collision
-				NewComponents[Idx]->UpdateCachedBounds();
-				NewComponents[Idx]->UpdateBounds();
-				NewComponents[Idx]->MarkRenderStateDirty();
-				ULandscapeHeightfieldCollisionComponent* CollisionComp = NewComponents[Idx]->CollisionComponent.Get();
+				NewComponent->UpdateCachedBounds();
+				NewComponent->UpdateBounds();
+				NewComponent->MarkRenderStateDirty();
+				ULandscapeHeightfieldCollisionComponent* CollisionComp = NewComponent->CollisionComponent.Get();
 				if (CollisionComp && !bHasXYOffset)
 				{
 					CollisionComp->MarkRenderStateDirty();
@@ -904,6 +907,13 @@ public:
 	virtual void SetEditRenderType() override { GLandscapeEditRenderMode = ELandscapeEditRenderMode::None | (GLandscapeEditRenderMode & ELandscapeEditRenderMode::BitMaskForMask); }
 	virtual bool SupportsMask() override { return false; }
 
+	virtual void EnterTool() override
+	{
+		FLandscapeToolBase<FLandscapeToolStrokeAddComponent>::EnterTool();
+		ULandscapeInfo* LandscapeInfo = EdMode->CurrentToolTarget.LandscapeInfo.Get();
+		LandscapeInfo->UpdateAllAddCollisions(); // Todo - as this is only used by this tool, move it into this tool?
+	}
+
 	virtual void ExitTool() override
 	{
 		FLandscapeToolBase<FLandscapeToolStrokeAddComponent>::ExitTool();
@@ -993,7 +1003,6 @@ public:
 				}
 			}
 
-			TArray<FIntPoint> DeletedNeighborKeys;
 			// Check which ones are need for height map change
 			for (TSet<ULandscapeComponent*>::TIterator It(SelectedComponents); It; ++It)
 			{
@@ -1004,7 +1013,7 @@ public:
 
 				// Reset neighbors LOD information
 				FIntPoint ComponentBase = Component->GetSectionBase() / Component->ComponentSizeQuads;
-				FIntPoint LandscapeKey[8] =
+				FIntPoint NeighborKeys[8] =
 				{
 					ComponentBase + FIntPoint(-1, -1),
 					ComponentBase + FIntPoint(+0, -1),
@@ -1016,13 +1025,15 @@ public:
 					ComponentBase + FIntPoint(+1, +1)
 				};
 
-				for (int32 Idx = 0; Idx < 8; ++Idx)
+				for (const FIntPoint& NeighborKey : NeighborKeys)
 				{
-					ULandscapeComponent* NeighborComp = LandscapeInfo->XYtoComponentMap.FindRef(LandscapeKey[Idx]);
+					ULandscapeComponent* NeighborComp = LandscapeInfo->XYtoComponentMap.FindRef(NeighborKey);
 					if (NeighborComp)
 					{
 						NeighborComp->Modify();
 						NeighborComp->InvalidateLightingCache();
+
+						// is this really needed? It can happen multiple times per component and even for components about to be deleted!
 						FComponentReregisterContext ReregisterContext(NeighborComp);
 					}
 				}
@@ -1060,16 +1071,6 @@ public:
 					Component->XYOffsetmapTexture->ClearFlags(RF_Standalone);
 				}
 
-				FIntPoint Key = Component->GetSectionBase() / Component->ComponentSizeQuads;
-				DeletedNeighborKeys.AddUnique(Key + FIntPoint(-1, -1));
-				DeletedNeighborKeys.AddUnique(Key + FIntPoint(+0, -1));
-				DeletedNeighborKeys.AddUnique(Key + FIntPoint(+1, -1));
-				DeletedNeighborKeys.AddUnique(Key + FIntPoint(-1, +0));
-				DeletedNeighborKeys.AddUnique(Key + FIntPoint(+1, +0));
-				DeletedNeighborKeys.AddUnique(Key + FIntPoint(-1, +1));
-				DeletedNeighborKeys.AddUnique(Key + FIntPoint(+0, +1));
-				DeletedNeighborKeys.AddUnique(Key + FIntPoint(+1, +1));
-
 				ULandscapeHeightfieldCollisionComponent* CollisionComp = Component->CollisionComponent.Get();
 				if (CollisionComp)
 				{
@@ -1078,26 +1079,6 @@ public:
 				Component->DestroyComponent();
 			}
 
-			// Update AddCollisions...
-			for (int32 i = 0; i < DeletedNeighborKeys.Num(); ++i)
-			{
-				LandscapeInfo->XYtoAddCollisionMap.Remove(DeletedNeighborKeys[i]);
-			}
-
-			for (int32 i = 0; i < DeletedNeighborKeys.Num(); ++i)
-			{
-				ULandscapeComponent* Component = LandscapeInfo->XYtoComponentMap.FindRef(DeletedNeighborKeys[i]);
-				if (Component)
-				{
-					ULandscapeHeightfieldCollisionComponent* CollisionComp = Component->CollisionComponent.Get();
-					if (CollisionComp)
-					{
-						CollisionComp->UpdateAddCollisions();
-					}
-				}
-			}
-
-
 			// Remove Selection
 			LandscapeInfo->ClearSelectedRegion(true);
 			//EdMode->SetMaskEnable(Landscape->SelectedRegion.Num());
diff --git a/Engine/Source/Editor/LandscapeEditor/Private/LandscapeEdModeRampTool.cpp b/Engine/Source/Editor/LandscapeEditor/Private/LandscapeEdModeRampTool.cpp
--- a/Engine/Source/Editor/LandscapeEditor/Private/LandscapeEdModeRampTool.cpp
+++ b/Engine/Source/Editor/LandscapeEditor/Private/LandscapeEdModeRampTool.cpp
@@ -554,7 +554,7 @@ public:
 					ULandscapeHeightfieldCollisionComponent* CollisionComponent = Component->CollisionComponent.Get();
 					if (CollisionComponent)
 					{
-						CollisionComponent->RecreateCollision(false);
+						CollisionComponent->RecreateCollision();
 						UNavigationSystem* NavSys = UNavigationSystem::GetCurrent(Component);
 						if (NavSys)
 						{
diff --git a/Engine/Source/Editor/LandscapeEditor/Private/LandscapeEdModeTools.h b/Engine/Source/Editor/LandscapeEditor/Private/LandscapeEdModeTools.h
--- a/Engine/Source/Editor/LandscapeEditor/Private/LandscapeEdModeTools.h
+++ b/Engine/Source/Editor/LandscapeEditor/Private/LandscapeEdModeTools.h
@@ -689,7 +689,7 @@ struct FHeightmapAccessor
 			ULandscapeHeightfieldCollisionComponent* CollisionComponent = (*It)->CollisionComponent.Get();
 			if (CollisionComponent)
 			{
-				CollisionComponent->RecreateCollision(true);
+				CollisionComponent->RecreateCollision();
 				UNavigationSystem::UpdateNavOctree(CollisionComponent);
 			}
 		}
@@ -910,7 +910,7 @@ struct FAlphamapAccessor
 			ULandscapeHeightfieldCollisionComponent* CollisionComponent = Component->CollisionComponent.Get();
 			if (CollisionComponent)
 			{
-				CollisionComponent->RecreateCollision(false);
+				CollisionComponent->RecreateCollision();
 
 				// We need to trigger navigation mesh build, in case user have painted holes on a landscape
 				if (LayerInfo == ALandscapeProxy::VisibilityLayer)
@@ -1015,7 +1015,7 @@ struct FFullWeightmapAccessor
 			ULandscapeHeightfieldCollisionComponent* CollisionComponent = Component->CollisionComponent.Get();
 			if (CollisionComponent)
 			{
-				CollisionComponent->RecreateCollision(false);
+				CollisionComponent->RecreateCollision();
 
 				// We need to trigger navigation mesh build, in case user have painted holes on a landscape
 				if (LandscapeInfo->GetLayerInfoIndex(ALandscapeProxy::VisibilityLayer) != INDEX_NONE)
diff --git a/Engine/Source/Runtime/Landscape/Classes/LandscapeHeightfieldCollisionComponent.h b/Engine/Source/Runtime/Landscape/Classes/LandscapeHeightfieldCollisionComponent.h
--- a/Engine/Source/Runtime/Landscape/Classes/LandscapeHeightfieldCollisionComponent.h
+++ b/Engine/Source/Runtime/Landscape/Classes/LandscapeHeightfieldCollisionComponent.h
@@ -183,9 +183,6 @@ class ULandscapeHeightfieldCollisionComponent : public UPrimitiveComponent
 	virtual void PostEditUndo() override;
 	// End UObject Interface.
 
-	// Update Collision object for add LandscapeComponent tool
-	LANDSCAPE_API void UpdateAddCollisions();
-
 	// @todo document
 	class ULandscapeInfo* GetLandscapeInfo(bool bSpawnNewActor = true) const;
 
@@ -215,7 +212,7 @@ class ULandscapeHeightfieldCollisionComponent : public UPrimitiveComponent
 	LANDSCAPE_API void SetSectionBase(FIntPoint InSectionBase);
 
 	/** Recreate heightfield and restart physics */
-	LANDSCAPE_API virtual void RecreateCollision(bool bUpdateAddCollision = true);
+	LANDSCAPE_API virtual void RecreateCollision();
 
 #if WITH_EDITORONLY_DATA
 	// Called from editor code to manage foliage instances on landscape.
diff --git a/Engine/Source/Runtime/Landscape/Classes/LandscapeInfo.h b/Engine/Source/Runtime/Landscape/Classes/LandscapeInfo.h
--- a/Engine/Source/Runtime/Landscape/Classes/LandscapeInfo.h
+++ b/Engine/Source/Runtime/Landscape/Classes/LandscapeInfo.h
@@ -13,26 +13,17 @@ struct FLandscapeEditorLayerSettings;
 struct FLandscapeDataInterface;
 
 /** Structure storing Collision for LandscapeComponent Add */
-USTRUCT()
+#if WITH_EDITORONLY_DATA
 struct FLandscapeAddCollision
 {
-	GENERATED_USTRUCT_BODY()
-
-#if WITH_EDITORONLY_DATA
-	UPROPERTY()
 	FVector Corners[4];
 
-#endif // WITH_EDITORONLY_DATA
-
-
-		FLandscapeAddCollision()
-		{
-			#if WITH_EDITORONLY_DATA
-			Corners[0] = Corners[1] = Corners[2] = Corners[3] = FVector::ZeroVector;
-			#endif // WITH_EDITORONLY_DATA
-		}
-	
+	FLandscapeAddCollision()
+	{
+		Corners[0] = Corners[1] = Corners[2] = Corners[3] = FVector::ZeroVector;
+	}
 };
+#endif // WITH_EDITORONLY_DATA
 
 USTRUCT()
 struct FLandscapeInfoLayerSettings
@@ -129,8 +120,11 @@ public:
 	/** Map of the offsets (in component space) to the component. Valid in editor only. */
 	TMap<FIntPoint, ULandscapeComponent*> XYtoComponentMap;
 
-	/** Map of the offsets to the newly added collision components. Only available near valid LandscapeComponents. Valid in editor only. */
+#if WITH_EDITORONLY_DATA
+	/** Lookup map used by the "add component" tool. Only available near valid LandscapeComponents.
+	    only for use by the "add component" tool. Todo - move into the tool? */
 	TMap<FIntPoint, FLandscapeAddCollision> XYtoAddCollisionMap;
+#endif
 
 	TSet<ALandscapeProxy*> Proxies;
 
@@ -178,8 +172,9 @@ public:
 	LANDSCAPE_API void SortSelectedComponents();
 	LANDSCAPE_API void ClearSelectedRegion(bool bIsComponentwise = true);
 
+	// only for use by the "add component" tool. Todo - move into the tool?
 	LANDSCAPE_API void UpdateAllAddCollisions();
-	void UpdateAddCollision(FIntPoint LandscapeKey);
+	LANDSCAPE_API void UpdateAddCollision(FIntPoint LandscapeKey);
 
 	LANDSCAPE_API FLandscapeEditorLayerSettings& GetLayerEditorSettings(ULandscapeLayerInfoObject* LayerInfo) const;
 	LANDSCAPE_API void CreateLayerEditorSettingsFor(ULandscapeLayerInfoObject* LayerInfo);
diff --git a/Engine/Source/Runtime/Landscape/Classes/LandscapeMeshCollisionComponent.h b/Engine/Source/Runtime/Landscape/Classes/LandscapeMeshCollisionComponent.h
--- a/Engine/Source/Runtime/Landscape/Classes/LandscapeMeshCollisionComponent.h
+++ b/Engine/Source/Runtime/Landscape/Classes/LandscapeMeshCollisionComponent.h
@@ -94,7 +94,7 @@ class ULandscapeMeshCollisionComponent : public ULandscapeHeightfieldCollisionCo
 
 	// Begin ULandscapeHeightfieldCollisionComponent Interface
 	virtual void CreateCollisionObject() override;
-	virtual void RecreateCollision(bool bUpdateAddCollision = true) override;
+	virtual void RecreateCollision() override;
 	// End ULandscapeHeightfieldCollisionComponent Interface
 };
 
diff --git a/Engine/Source/Runtime/Landscape/Private/Landscape.cpp b/Engine/Source/Runtime/Landscape/Private/Landscape.cpp
--- a/Engine/Source/Runtime/Landscape/Private/Landscape.cpp
+++ b/Engine/Source/Runtime/Landscape/Private/Landscape.cpp
@@ -900,14 +900,12 @@ FArchive& operator<<(FArchive& Ar, FLandscapeWeightmapUsage& U)
 	return Ar << U.ChannelUsage[0] << U.ChannelUsage[1] << U.ChannelUsage[2] << U.ChannelUsage[3];
 }
 
+#if WITH_EDITORONLY_DATA
 FArchive& operator<<(FArchive& Ar, FLandscapeAddCollision& U)
 {
-#if WITH_EDITORONLY_DATA
 	return Ar << U.Corners[0] << U.Corners[1] << U.Corners[2] << U.Corners[3];
-#else
-	return Ar;
-#endif // WITH_EDITORONLY_DATA
 }
+#endif // WITH_EDITORONLY_DATA
 
 FArchive& operator<<(FArchive& Ar, FLandscapeLayerStruct*& L)
 {
@@ -932,7 +930,9 @@ void ULandscapeInfo::Serialize(FArchive& Ar)
 	if (Ar.IsTransacting())
 	{
 		Ar << XYtoComponentMap;
+#if WITH_EDITORONLY_DATA
 		Ar << XYtoAddCollisionMap;
+#endif
 		Ar << SelectedComponents;
 		Ar << SelectedRegion;
 		Ar << SelectedRegionComponents;
@@ -1518,9 +1518,6 @@ void ALandscapeProxy::RecreateComponentsState()
 		if (Comp)
 		{
 			Comp->UpdateComponentToWorld();
-#if WITH_EDITOR
-			Comp->UpdateAddCollisions();
-#endif
 			Comp->RecreatePhysicsState();
 		}
 	}
diff --git a/Engine/Source/Runtime/Landscape/Private/LandscapeCollision.cpp b/Engine/Source/Runtime/Landscape/Private/LandscapeCollision.cpp
--- a/Engine/Source/Runtime/Landscape/Private/LandscapeCollision.cpp
+++ b/Engine/Source/Runtime/Landscape/Private/LandscapeCollision.cpp
@@ -938,7 +938,7 @@ void ULandscapeHeightfieldCollisionComponent::UpdateHeightfieldRegion(int32 Comp
 		// If we're currently sharing this data with a PIE session, we need to make a new heightfield.
 		if (HeightfieldRef->GetRefCount() > 1)
 		{
-			RecreateCollision(false);
+			RecreateCollision();
 			return;
 		}
 
@@ -1052,18 +1052,12 @@ void ULandscapeMeshCollisionComponent::BeginDestroy()
 	Super::BeginDestroy();
 }
 
-void ULandscapeHeightfieldCollisionComponent::RecreateCollision(bool bUpdateAddCollision/*= true*/)
+void ULandscapeHeightfieldCollisionComponent::RecreateCollision()
 {
 	if (!HasAnyFlags(RF_ClassDefaultObject))
 	{
 		HeightfieldRef = NULL;
 		HeightfieldGuid = FGuid();
-#if WITH_EDITOR
-		if (bUpdateAddCollision)
-		{
-			UpdateAddCollisions();
-		}
-#endif
 
 		RecreatePhysicsState();
 	}
@@ -1180,7 +1174,7 @@ void ULandscapeHeightfieldCollisionComponent::SnapFoliageInstances(const FBox& I
 }
 #endif // WITH_EDITORONLY_DATA
 
-void ULandscapeMeshCollisionComponent::RecreateCollision(bool bUpdateAddCollision/*= true*/)
+void ULandscapeMeshCollisionComponent::RecreateCollision()
 {
 	if (!HasAnyFlags(RF_ClassDefaultObject))
 	{
@@ -1188,7 +1182,7 @@ void ULandscapeMeshCollisionComponent::RecreateCollision(bool bUpdateAddCollisio
 		MeshGuid = FGuid();
 	}
 
-	Super::RecreateCollision(bUpdateAddCollision);
+	Super::RecreateCollision();
 }
 
 void ULandscapeHeightfieldCollisionComponent::Serialize(FArchive& Ar)
@@ -1289,7 +1283,7 @@ void ULandscapeHeightfieldCollisionComponent::PostEditImport()
 	// Reinitialize physics after paste
 	if (CollisionSizeQuads > 0)
 	{
-		RecreateCollision(false);
+		RecreateCollision();
 	}
 }
 
@@ -1300,7 +1294,7 @@ void ULandscapeHeightfieldCollisionComponent::PostEditUndo()
 	// Reinitialize physics after undo
 	if (CollisionSizeQuads > 0)
 	{
-		RecreateCollision(false);
+		RecreateCollision();
 	}
 
 	UNavigationSystem::UpdateNavOctree(this);
@@ -1475,51 +1469,37 @@ void ULandscapeHeightfieldCollisionComponent::PreSave()
 #if WITH_EDITOR
 void ULandscapeInfo::UpdateAllAddCollisions()
 {
+	XYtoAddCollisionMap.Reset();
+
 	for (auto It = XYtoComponentMap.CreateIterator(); It; ++It)
 	{
-		ULandscapeComponent* Comp = It.Value();
-		if (Comp)
+		const ULandscapeComponent* const Component = It.Value();
+		if (ensure(Component))
 		{
-			ULandscapeHeightfieldCollisionComponent* CollisionComp = Comp->CollisionComponent.Get();
-			if (CollisionComp)
-			{
-				CollisionComp->UpdateAddCollisions();
-			}
-		}
-	}
-}
+			const FIntPoint ComponentBase = Component->GetSectionBase() / ComponentSizeQuads;
 
-void ULandscapeHeightfieldCollisionComponent::UpdateAddCollisions()
-{
-	ULandscapeInfo* Info = GetLandscapeInfo();
-	if (Info)
-	{
-		ALandscapeProxy* Proxy = GetLandscapeProxy();
-		FIntPoint ComponentBase = GetSectionBase() / Proxy->ComponentSizeQuads;
-
-		FIntPoint NeighborsKeys[8] =
-		{
-			ComponentBase + FIntPoint(-1, -1),
-			ComponentBase + FIntPoint(+0, -1),
-			ComponentBase + FIntPoint(+1, -1),
-			ComponentBase + FIntPoint(-1, +0),
-			ComponentBase + FIntPoint(+1, +0),
-			ComponentBase + FIntPoint(-1, +1),
-			ComponentBase + FIntPoint(+0, +1),
-			ComponentBase + FIntPoint(+1, +1)
-		};
-
-		// Search for Neighbors...
-		for (int32 i = 0; i < 8; ++i)
-		{
-			ULandscapeComponent* Comp = Info->XYtoComponentMap.FindRef(NeighborsKeys[i]);
-			if (!Comp || !Comp->CollisionComponent.IsValid())
+			const FIntPoint NeighborsKeys[8] =
 			{
-				Info->UpdateAddCollision(NeighborsKeys[i]);
-			}
-			else
+				ComponentBase + FIntPoint(-1, -1),
+				ComponentBase + FIntPoint(+0, -1),
+				ComponentBase + FIntPoint(+1, -1),
+				ComponentBase + FIntPoint(-1, +0),
+				ComponentBase + FIntPoint(+1, +0),
+				ComponentBase + FIntPoint(-1, +1),
+				ComponentBase + FIntPoint(+0, +1),
+				ComponentBase + FIntPoint(+1, +1)
+			};
+
+			// Search for Neighbors...
+			for (int32 i = 0; i < 8; ++i)
 			{
-				Info->XYtoAddCollisionMap.Remove(NeighborsKeys[i]);
+				ULandscapeComponent* NeighborComponent = XYtoComponentMap.FindRef(NeighborsKeys[i]);
+
+				// UpdateAddCollision() treats a null CollisionComponent as an empty hole
+				if (!NeighborComponent || !NeighborComponent->CollisionComponent.IsValid())
+				{
+					UpdateAddCollision(NeighborsKeys[i]);
+				}
 			}
 		}
 	}
@@ -1527,13 +1507,7 @@ void ULandscapeHeightfieldCollisionComponent::UpdateAddCollisions()
 
 void ULandscapeInfo::UpdateAddCollision(FIntPoint LandscapeKey)
 {
-	FLandscapeAddCollision* AddCollision = XYtoAddCollisionMap.Find(LandscapeKey);
-	if (!AddCollision)
-	{
-		AddCollision = &XYtoAddCollisionMap.Add(LandscapeKey, FLandscapeAddCollision());
-	}
-
-	check(AddCollision);
+	FLandscapeAddCollision& AddCollision = XYtoAddCollisionMap.FindOrAdd(LandscapeKey);
 
 	// 8 Neighbors... 
 	// 0 1 2 
@@ -1551,6 +1525,8 @@ void ULandscapeInfo::UpdateAddCollision(FIntPoint LandscapeKey)
 		LandscapeKey + FIntPoint(+1, +1)
 	};
 
+	// Todo: Use data accessor not collision
+
 	ULandscapeHeightfieldCollisionComponent* NeighborCollisions[8];
 	// Search for Neighbors...
 	for (int32 i = 0; i < 8; ++i)
@@ -1651,14 +1627,14 @@ void ULandscapeInfo::UpdateAddCollision(FIntPoint LandscapeKey)
 	FillCornerValues(CornerSet, HeightCorner);
 	//check(CornerSet == 15);
 
-	FIntPoint SectionBase = LandscapeKey*ComponentSizeQuads;
+	FIntPoint SectionBase = LandscapeKey * ComponentSizeQuads;
 
 	// Transform Height to Vectors...
-	FMatrix LtoW = GetLandscapeProxy()->LandscapeActorToWorld().ToMatrixWithScale();
-	AddCollision->Corners[0] = LtoW.TransformPosition(FVector(SectionBase.X, SectionBase.Y, LandscapeDataAccess::GetLocalHeight(HeightCorner[0])));
-	AddCollision->Corners[1] = LtoW.TransformPosition(FVector(SectionBase.X + ComponentSizeQuads, SectionBase.Y, LandscapeDataAccess::GetLocalHeight(HeightCorner[1])));
-	AddCollision->Corners[2] = LtoW.TransformPosition(FVector(SectionBase.X, SectionBase.Y + ComponentSizeQuads, LandscapeDataAccess::GetLocalHeight(HeightCorner[2])));
-	AddCollision->Corners[3] = LtoW.TransformPosition(FVector(SectionBase.X + ComponentSizeQuads, SectionBase.Y + ComponentSizeQuads, LandscapeDataAccess::GetLocalHeight(HeightCorner[3])));
+	FTransform LtoW = GetLandscapeProxy()->LandscapeActorToWorld();
+	AddCollision.Corners[0] = LtoW.TransformPosition(FVector(SectionBase.X                     , SectionBase.Y                     , LandscapeDataAccess::GetLocalHeight(HeightCorner[0])));
+	AddCollision.Corners[1] = LtoW.TransformPosition(FVector(SectionBase.X + ComponentSizeQuads, SectionBase.Y                     , LandscapeDataAccess::GetLocalHeight(HeightCorner[1])));
+	AddCollision.Corners[2] = LtoW.TransformPosition(FVector(SectionBase.X                     , SectionBase.Y + ComponentSizeQuads, LandscapeDataAccess::GetLocalHeight(HeightCorner[2])));
+	AddCollision.Corners[3] = LtoW.TransformPosition(FVector(SectionBase.X + ComponentSizeQuads, SectionBase.Y + ComponentSizeQuads, LandscapeDataAccess::GetLocalHeight(HeightCorner[3])));
 }
 
 void ULandscapeHeightfieldCollisionComponent::ExportCustomProperties(FOutputDevice& Out, uint32 Indent)
diff --git a/Engine/Source/Runtime/Landscape/Private/LandscapeEdit.cpp b/Engine/Source/Runtime/Landscape/Private/LandscapeEdit.cpp
--- a/Engine/Source/Runtime/Landscape/Private/LandscapeEdit.cpp
+++ b/Engine/Source/Runtime/Landscape/Private/LandscapeEdit.cpp
@@ -733,7 +733,7 @@ void ULandscapeComponent::UpdateCollisionHeightData(const FColor* HeightmapTextu
 		if (MeshCollisionComponent)
 		{
 			// Will be done once for XY Offset data update in FXYOffsetmapAccessor() destructor with UpdateCachedBounds()
-			//MeshCollisionComponent->RecreateCollision(false);
+			//MeshCollisionComponent->RecreateCollision();
 		}
 		else if (CollisionMipLevel == 0)
 		{
@@ -2842,7 +2842,7 @@ void ALandscapeProxy::EditorApplyScale(const FVector& DeltaScale, const FVector*
 	{
 		if (Comp)
 		{
-			Comp->RecreateCollision(false);
+			Comp->RecreateCollision();
 		}
 	}
 }
@@ -2856,7 +2856,7 @@ void ALandscapeProxy::EditorApplyMirror(const FVector& MirrorScale, const FVecto
 	{
 		if (Comp)
 		{
-			Comp->RecreateCollision(false);
+			Comp->RecreateCollision();
 		}
 	}
 }
@@ -3283,7 +3283,7 @@ void ALandscapeProxy::PostEditChangeChainProperty(FPropertyChangedChainEvent& Pr
 				ULandscapeHeightfieldCollisionComponent* Comp = CollisionComponents[ComponentIndex];
 				if (Comp)
 				{
-					Comp->RecreateCollision(false);
+					Comp->RecreateCollision();
 				}
 			}
 		}
@@ -3466,7 +3466,7 @@ void ALandscapeProxy::ChangedPhysMaterial()
 			{
 				Comp->UpdateCollisionLayerData();
 				// Physical materials cooked into collision object, so we need to recreate it
-				CollisionComponent->RecreateCollision(false);
+				CollisionComponent->RecreateCollision();
 			}
 		}
 	}
diff --git a/Engine/Source/Runtime/Landscape/Private/LandscapeSplineRaster.cpp b/Engine/Source/Runtime/Landscape/Private/LandscapeSplineRaster.cpp
--- a/Engine/Source/Runtime/Landscape/Private/LandscapeSplineRaster.cpp
+++ b/Engine/Source/Runtime/Landscape/Private/LandscapeSplineRaster.cpp
@@ -589,7 +589,7 @@ bool ULandscapeInfo::ApplySplinesInternal(bool bOnlySelected, ALandscapeProxy* L
 		ULandscapeHeightfieldCollisionComponent* CollisionComponent = Component->CollisionComponent.Get();
 		if (CollisionComponent)
 		{
-			CollisionComponent->RecreateCollision(false);
+			CollisionComponent->RecreateCollision();
 			UNavigationSystem* NavSys = UNavigationSystem::GetCurrent(Component);
 			if (NavSys)
 			{
@@ -684,7 +684,7 @@ namespace LandscapeSplineRaster
 			ULandscapeHeightfieldCollisionComponent* CollisionComponent = Component->CollisionComponent.Get();
 			if (CollisionComponent)
 			{
-				CollisionComponent->RecreateCollision(false);
+				CollisionComponent->RecreateCollision();
 				UNavigationSystem* NavSys = UNavigationSystem::GetCurrent(Component);
 				if (NavSys)
 				{
