Fix for incorrect grassmap rasterization under OpenGL on AMD GPUs (UE-15648)

diff --git a/Engine/Source/Runtime/Landscape/Private/LandscapeGrass.cpp b/Engine/Source/Runtime/Landscape/Private/LandscapeGrass.cpp
--- a/Engine/Source/Runtime/Landscape/Private/LandscapeGrass.cpp
+++ b/Engine/Source/Runtime/Landscape/Private/LandscapeGrass.cpp
@@ -20,6 +20,7 @@
 #include "Engine/TextureRenderTarget2D.h"
 #include "ContentStreaming.h"
 #include "LandscapeDataAccess.h"
+#include "LandscapeRender.h"
 
 #define LOCTEXT_NAMESPACE "Landscape"
 
@@ -297,13 +298,13 @@ public:
 		ULandscapeComponent* Component;
 		FVector2D ViewOffset;
 		int32 PixelOffsetX;
-		FPrimitiveSceneInfo* PrimitiveSceneInfo;
+		FLandscapeComponentSceneProxy* SceneProxy;
 
 		FComponentInfo(ULandscapeComponent* InComponent, FVector2D& InViewOffset, int32 InPixelOffsetX)
 			: Component(InComponent)
 			, ViewOffset(InViewOffset)
 			, PixelOffsetX(InPixelOffsetX)
-			, PrimitiveSceneInfo(InComponent->SceneProxy->GetPrimitiveSceneInfo())
+			, SceneProxy((FLandscapeComponentSceneProxy*)InComponent->SceneProxy)
 		{}
 	};
 
@@ -342,23 +343,17 @@ public:
 
 		for (auto& ComponentInfo : ComponentInfos)
 		{
-			FPrimitiveSceneInfo* PrimitiveSceneInfo = ComponentInfo.PrimitiveSceneInfo;
+			const FMeshBatch& Mesh = ComponentInfo.SceneProxy->GetGrassMeshBatch();
+
 			for (int32 PassIdx = 0; PassIdx < NumPasses; PassIdx++)
 			{
-				for (int32 StaticMeshIdx = 0; StaticMeshIdx < PrimitiveSceneInfo->StaticMeshes.Num(); StaticMeshIdx++)
-				{
-					FMeshBatch Mesh = *(FMeshBatch*)(&PrimitiveSceneInfo->StaticMeshes[StaticMeshIdx]);
-					// Grass maps don't render with tessellation enabled but the StaticMesh is set up to use PT_12_ControlPointPatchList
-					Mesh.Type = PT_TriangleList; 
+				FLandscapeGrassWeightDrawingPolicy DrawingPolicy(Mesh.VertexFactory, Mesh.MaterialRenderProxy, *Mesh.MaterialRenderProxy->GetMaterial(GMaxRHIFeatureLevel));
+				RHICmdList.BuildAndSetLocalBoundShaderState(DrawingPolicy.GetBoundShaderStateInput(GMaxRHIFeatureLevel));
+				DrawingPolicy.SetSharedState(RHICmdList, View, FLandscapeGrassWeightDrawingPolicy::ContextDataType(), PassIdx, ComponentInfo.ViewOffset + FVector2D(PassOffsetX * PassIdx, 0));
 
-					FLandscapeGrassWeightDrawingPolicy DrawingPolicy(Mesh.VertexFactory, Mesh.MaterialRenderProxy, *Mesh.MaterialRenderProxy->GetMaterial(GMaxRHIFeatureLevel));
-					RHICmdList.BuildAndSetLocalBoundShaderState(DrawingPolicy.GetBoundShaderStateInput(GMaxRHIFeatureLevel));
-					DrawingPolicy.SetSharedState(RHICmdList, View, FLandscapeGrassWeightDrawingPolicy::ContextDataType(), PassIdx, ComponentInfo.ViewOffset + FVector2D(PassOffsetX * PassIdx, 0));
-
-					// The first batch element contains the grass batch for the entire component
-					DrawingPolicy.SetMeshRenderState(RHICmdList, *View, PrimitiveSceneInfo->Proxy, Mesh, 0, false, 0.0f, FMeshDrawingPolicy::ElementDataType(), FLandscapeGrassWeightDrawingPolicy::ContextDataType());
-					DrawingPolicy.DrawMesh(RHICmdList, Mesh, 0);
-				}
+				// The first batch element contains the grass batch for the entire component
+				DrawingPolicy.SetMeshRenderState(RHICmdList, *View, ComponentInfo.SceneProxy, Mesh, 0, false, 0.0f, FMeshDrawingPolicy::ElementDataType(), FLandscapeGrassWeightDrawingPolicy::ContextDataType());
+				DrawingPolicy.DrawMesh(RHICmdList, Mesh, 0);
 			}
 		}
 	}
diff --git a/Engine/Source/Runtime/Landscape/Private/LandscapeRender.cpp b/Engine/Source/Runtime/Landscape/Private/LandscapeRender.cpp
--- a/Engine/Source/Runtime/Landscape/Private/LandscapeRender.cpp
+++ b/Engine/Source/Runtime/Landscape/Private/LandscapeRender.cpp
@@ -659,7 +659,7 @@ void FLandscapeComponentSceneProxy::CreateRenderThreadResources()
 	check(HeightmapTexture != nullptr);
 
 	if (IsComponentLevelVisible())
-		{
+	{
 		RegisterNeighbors();
 	}
 
@@ -724,6 +724,37 @@ void FLandscapeComponentSceneProxy::CreateRenderThreadResources()
 
 	// Assign LandscapeUniformShaderParameters
 	LandscapeUniformShaderParameters.InitResource();
+
+	// Create MeshBatch for grass rendering
+	{
+		GrassMeshBatch.Elements.Empty(1);
+
+		FMaterialRenderProxy* RenderProxy = MaterialInterface->GetRenderProxy(false);
+		GrassMeshBatch.VertexFactory = VertexFactory;
+		GrassMeshBatch.MaterialRenderProxy = RenderProxy;
+		GrassMeshBatch.LCI = nullptr;
+		GrassMeshBatch.ReverseCulling = false;
+		GrassMeshBatch.CastShadow = false;
+		GrassMeshBatch.Type = PT_PointList;
+		GrassMeshBatch.DepthPriorityGroup = SDPG_World;
+
+		// Combined grass rendering batch element
+		FMeshBatchElement* BatchElement = new(GrassMeshBatch.Elements) FMeshBatchElement;
+		FLandscapeBatchElementParams* BatchElementParams = &GrassBatchParams;
+		BatchElementParams->LocalToWorldNoScalingPtr = &LocalToWorldNoScaling;
+		BatchElement->UserData = BatchElementParams;
+		BatchElement->PrimitiveUniformBufferResource = &GetUniformBuffer();
+		BatchElementParams->LandscapeUniformShaderParametersResource = &LandscapeUniformShaderParameters;
+		BatchElementParams->SceneProxy = this;
+		BatchElementParams->SubX = -1;
+		BatchElementParams->SubY = -1;
+		BatchElementParams->CurrentLOD = 0;
+		BatchElement->IndexBuffer = SharedBuffers->GrassIndexBuffer;
+		BatchElement->NumPrimitives = FMath::Square(NumSubsections) * FMath::Square(SubsectionSizeVerts);
+		BatchElement->FirstIndex = 0;
+		BatchElement->MinVertexIndex = 0;
+		BatchElement->MaxVertexIndex = SharedBuffers->NumVertices - 1;
+	}
 }
 
 void FLandscapeComponentSceneProxy::OnLevelAddedToWorld()
@@ -731,9 +762,8 @@ void FLandscapeComponentSceneProxy::OnLevelAddedToWorld()
 	RegisterNeighbors();
 }
 
-
 FLandscapeComponentSceneProxy::~FLandscapeComponentSceneProxy()
-			{
+{
 	UnregisterNeighbors();
 
 	// Free the subsection uniform buffer
@@ -1031,10 +1061,6 @@ namespace
 void FLandscapeComponentSceneProxy::DrawStaticElements(FStaticPrimitiveDrawInterface* PDI)
 {
 	int32 NumBatches = (1 + LastLOD - FirstLOD) * (FMath::Square(NumSubsections) + 1);
-	if (SharedBuffers->GrassIndexBuffer)
-	{
-		NumBatches++;
-	}
 	StaticBatchParamArray.Empty(NumBatches);
 
 	FMeshBatch MeshBatch;
@@ -1058,26 +1084,6 @@ void FLandscapeComponentSceneProxy::DrawStaticElements(FStaticPrimitiveDrawInter
 	MeshBatch.Type = bCurrentRequiresAdjacencyInformation ? PT_12_ControlPointPatchList : PT_TriangleList;
 	MeshBatch.DepthPriorityGroup = SDPG_World;
 
-	if (SharedBuffers->GrassIndexBuffer)
-	{
-		// Combined grass rendering batch element
-		FMeshBatchElement* BatchElement = new(MeshBatch.Elements) FMeshBatchElement;
-		FLandscapeBatchElementParams* BatchElementParams = new(StaticBatchParamArray)FLandscapeBatchElementParams;
-		BatchElementParams->LocalToWorldNoScalingPtr = &LocalToWorldNoScaling;
-		BatchElement->UserData = BatchElementParams;
-		BatchElement->PrimitiveUniformBufferResource = &GetUniformBuffer();
-		BatchElementParams->LandscapeUniformShaderParametersResource = &LandscapeUniformShaderParameters;
-		BatchElementParams->SceneProxy = this;
-		BatchElementParams->SubX = -1;
-		BatchElementParams->SubY = -1;
-		BatchElementParams->CurrentLOD = 0;
-		BatchElement->IndexBuffer = SharedBuffers->GrassIndexBuffer;
-		BatchElement->NumPrimitives = (FMath::Square(NumSubsections) * FMath::Square(SubsectionSizeQuads) + 2 * NumSubsections * SubsectionSizeQuads + 1) * 2;
-		BatchElement->FirstIndex = 0;
-		BatchElement->MinVertexIndex = 0;
-		BatchElement->MaxVertexIndex = SharedBuffers->NumVertices-1;
-	}
-	
 	for (int32 LOD = FirstLOD; LOD <= LastLOD; LOD++)
 	{
 		int32 LodSubsectionSizeVerts = SubsectionSizeVerts >> LOD;
@@ -1147,14 +1153,13 @@ uint64 FLandscapeVertexFactory::GetStaticBatchElementVisibility(const class FSce
 uint64 FLandscapeComponentSceneProxy::GetStaticBatchElementVisibility(const class FSceneView& View, const struct FMeshBatch* Batch) const
 {
 	uint64 BatchesToRenderMask = 0;
-	int32 GrassBatchOffset = SharedBuffers->GrassIndexBuffer ? 1 : 0;
 
 	SCOPE_CYCLE_COUNTER(STAT_LandscapeStaticDrawLODTime);
 	if (ForcedLOD >= 0)
 	{
-		// When forcing LOD we only create one Batch Element (excluding the grass rendering batch)
-		ensure(Batch->Elements.Num() - GrassBatchOffset == 1);
-		int32 BatchElementIndex = GrassBatchOffset;
+		// When forcing LOD we only create one Batch Element
+		ensure(Batch->Elements.Num() == 1);
+		int32 BatchElementIndex = 0;
 		BatchesToRenderMask |= (((uint64)1) << BatchElementIndex);
 		INC_DWORD_STAT(STAT_LandscapeDrawCalls);
 		INC_DWORD_STAT_BY(STAT_LandscapeTriangles, Batch->Elements[BatchElementIndex].NumPrimitives);
@@ -1191,7 +1196,7 @@ uint64 FLandscapeComponentSceneProxy::GetStaticBatchElementVisibility(const clas
 		if (bAllSameLOD && NumSubsections > 1 && !GLandscapeDebugOptions.bDisableCombine)
 		{
 			// choose the combined batch element
-			int32 BatchElementIndex = GrassBatchOffset + (CombinedLOD - LODBiasOffset + 1)*BatchesPerLOD - 1;
+			int32 BatchElementIndex = (CombinedLOD - LODBiasOffset + 1)*BatchesPerLOD - 1;
 			if (ensure(Batch->Elements.IsValidIndex(BatchElementIndex)))
 			{
 				BatchesToRenderMask |= (((uint64)1) << BatchElementIndex);
@@ -1205,7 +1210,7 @@ uint64 FLandscapeComponentSceneProxy::GetStaticBatchElementVisibility(const clas
 			{
 				for (int32 SubX = 0; SubX < NumSubsections; SubX++)
 				{
-					int32 BatchElementIndex = GrassBatchOffset + (CalculatedLods[SubX][SubY] - LODBiasOffset) * BatchesPerLOD + SubY*NumSubsections + SubX;
+					int32 BatchElementIndex = (CalculatedLods[SubX][SubY] - LODBiasOffset) * BatchesPerLOD + SubY*NumSubsections + SubX;
 					if (ensure(Batch->Elements.IsValidIndex(BatchElementIndex)))
 					{
 						BatchesToRenderMask |= (((uint64)1) << BatchElementIndex);
@@ -1676,6 +1681,7 @@ void FLandscapeComponentSceneProxy::GetDynamicMeshElements(const TArray<const FS
 	INC_DWORD_STAT_BY(STAT_LandscapeTriangles, NumTriangles * NumPasses);
 }
 
+
 //
 // FLandscapeVertexBuffer
 //
@@ -1694,11 +1700,9 @@ void FLandscapeVertexBuffer::InitRHI()
 	{
 		for (int32 SubX = 0; SubX < NumSubsections; SubX++)
 		{
-			int32 Cols = SubX == NumSubsections - 1 && FeatureLevel > ERHIFeatureLevel::ES3_1 ? SubsectionSizeVerts + 1 : SubsectionSizeVerts;
-			int32 Rows = SubY == NumSubsections - 1 && FeatureLevel > ERHIFeatureLevel::ES3_1 ? SubsectionSizeVerts + 1 : SubsectionSizeVerts;
-			for (int32 y = 0; y < Rows; y++)
+			for (int32 y = 0; y < SubsectionSizeVerts; y++)
 			{
-				for (int32 x = 0; x < Cols; x++)
+				for (int32 x = 0; x < SubsectionSizeVerts; x++)
 				{
 					Vertex->VertexX = x;
 					Vertex->VertexY = y;
@@ -1876,16 +1880,14 @@ void FLandscapeSharedBuffers::CreateIndexBuffers(ERHIFeatureLevel::Type InFeatur
 					MaxIndex = 0;
 					MinIndex = MAX_int32;
 
-					int32 StrideX = SubX == NumSubsections - 1 ? SubsectionSizeVerts + 1 : SubsectionSizeVerts;
-					int32 StrideY = SubY == NumSubsections - 1 ? SubsectionSizeVerts + 1 : SubsectionSizeVerts;
 					for (int32 y = 0; y < LodSubsectionSizeQuads; y++)
 					{
 						for (int32 x = 0; x < LodSubsectionSizeQuads; x++)
 						{
-							INDEX_TYPE i00 = (x + 0) + (y + 0) * StrideX + SubOffset;
-							INDEX_TYPE i10 = (x + 1) + (y + 0) * StrideX + SubOffset;
-							INDEX_TYPE i11 = (x + 1) + (y + 1) * StrideX + SubOffset;
-							INDEX_TYPE i01 = (x + 0) + (y + 1) * StrideX + SubOffset;
+							INDEX_TYPE i00 = (x + 0) + (y + 0) * SubsectionSizeVerts + SubOffset;
+							INDEX_TYPE i10 = (x + 1) + (y + 0) * SubsectionSizeVerts + SubOffset;
+							INDEX_TYPE i11 = (x + 1) + (y + 1) * SubsectionSizeVerts + SubOffset;
+							INDEX_TYPE i01 = (x + 0) + (y + 1) * SubsectionSizeVerts + SubOffset;
 
 							NewIndices.Add(i00);
 							NewIndices.Add(i11);
@@ -1911,7 +1913,7 @@ void FLandscapeSharedBuffers::CreateIndexBuffers(ERHIFeatureLevel::Type InFeatur
 					MaxIndexFull = FMath::Max<int32>(MaxIndexFull, MaxIndex);
 					MinIndexFull = FMath::Min<int32>(MinIndexFull, MinIndex);
 
-					SubOffset += StrideX * StrideY;
+					SubOffset += FMath::Square(SubsectionSizeVerts);
 				}
 			}
 
@@ -1940,10 +1942,8 @@ void FLandscapeSharedBuffers::CreateIndexBuffers(ERHIFeatureLevel::Type InFeatur
 template <typename INDEX_TYPE>
 void FLandscapeSharedBuffers::CreateGrassIndexBuffer()
 {
-	int32 SubsectionSizeQuads = SubsectionSizeVerts - 1;
-
 	TArray<INDEX_TYPE> NewIndices;
-	int32 ExpectedNumIndices = (FMath::Square(NumSubsections) * FMath::Square(SubsectionSizeQuads) + 2 * NumSubsections * SubsectionSizeQuads + 1) * 6;
+	int32 ExpectedNumIndices = (FMath::Square(NumSubsections) * FMath::Square(SubsectionSizeVerts));
 	NewIndices.Empty(ExpectedNumIndices);
 
 	int32 SubOffset = 0;
@@ -1951,30 +1951,15 @@ void FLandscapeSharedBuffers::CreateGrassIndexBuffer()
 	{
 		for (int32 SubX = 0; SubX < NumSubsections; SubX++)
 		{
-			int32 Cols = SubX == NumSubsections - 1 ? SubsectionSizeQuads + 1 : SubsectionSizeQuads;
-			int32 Rows = SubY == NumSubsections - 1 ? SubsectionSizeQuads + 1 : SubsectionSizeQuads;
-			int32 StrideX = SubX == NumSubsections - 1 ? SubsectionSizeVerts + 1 : SubsectionSizeVerts;
-			int32 StrideY = SubY == NumSubsections - 1 ? SubsectionSizeVerts + 1 : SubsectionSizeVerts;
-			for (int32 y = 0; y < Rows; y++)
+			for (int32 y = 0; y < SubsectionSizeVerts; y++)
 			{
-				for (int32 x = 0; x < Cols; x++)
+				for (int32 x = 0; x < SubsectionSizeVerts; x++)
 				{
-					INDEX_TYPE i00 = (x + 0) + (y + 0) * StrideX + SubOffset;
-					INDEX_TYPE i10 = (x + 1) + (y + 0) * StrideX + SubOffset;
-					INDEX_TYPE i11 = (x + 1) + (y + 1) * StrideX + SubOffset;
-					INDEX_TYPE i01 = (x + 0) + (y + 1) * StrideX + SubOffset;
-
-					NewIndices.Add(i00);
-					NewIndices.Add(i11);
-					NewIndices.Add(i10);
-
-					NewIndices.Add(i00);
-					NewIndices.Add(i01);
-					NewIndices.Add(i11);
+					NewIndices.Add(x + y * SubsectionSizeVerts + SubOffset);
 				}
 			}
 
-			SubOffset += StrideX * StrideY;
+			SubOffset += FMath::Square(SubsectionSizeVerts);
 		}
 	}
 
@@ -2001,7 +1986,6 @@ FLandscapeSharedBuffers::FLandscapeSharedBuffers(int32 InSharedBuffersKey, int32
 	NumVertices = FMath::Square(SubsectionSizeVerts) * FMath::Square(NumSubsections);
 	if (InFeatureLevel > ERHIFeatureLevel::ES3_1)
 	{
-		NumVertices += 2 * NumSubsections * SubsectionSizeVerts + 1;
 		// Vertex Buffer cannot be shared
 		VertexBuffer = new FLandscapeVertexBuffer(InFeatureLevel, NumVertices, SubsectionSizeVerts, NumSubsections);
 	}
diff --git a/Engine/Source/Runtime/Landscape/Public/LandscapeRender.h b/Engine/Source/Runtime/Landscape/Public/LandscapeRender.h
--- a/Engine/Source/Runtime/Landscape/Public/LandscapeRender.h
+++ b/Engine/Source/Runtime/Landscape/Public/LandscapeRender.h
@@ -515,6 +515,10 @@ protected:
 	// Storage for static draw list batch params
 	TArray<FLandscapeBatchElementParams> StaticBatchParamArray;
 
+	// Precomputed grass rendering MeshBatch
+	FMeshBatch					GrassMeshBatch;
+	FLandscapeBatchElementParams GrassBatchParams;
+
 	// Precomputed values
 	float					LODDistance;
 	float					DistDiff;
@@ -588,6 +592,8 @@ public:
 	int32 CalcLODForSubsection(const FSceneView& View, int32 SubX, int32 SubY, const FVector2D& CameraLocalPos) const;
 	void CalcLODParamsForSubsection(const FSceneView& View, const FVector2D& CameraLocalPos, int32 SubX, int32 SubY, int32 BatchLOD, float& OutfLOD, FVector4& OutNeighborLODs) const;
 	uint64 GetStaticBatchElementVisibility(const FSceneView& View, const FMeshBatch* Batch) const;
+	const FMeshBatch& GetGrassMeshBatch() const { return GrassMeshBatch; }
+
 
 	// FLandcapeSceneProxy
 	void ChangeLODDistanceFactor_RenderThread(float InLODDistanceFactor);
