Fix to CL 2554857

diff --git a/Engine/Source/Runtime/Landscape/Private/LandscapeRender.cpp b/Engine/Source/Runtime/Landscape/Private/LandscapeRender.cpp
--- a/Engine/Source/Runtime/Landscape/Private/LandscapeRender.cpp
+++ b/Engine/Source/Runtime/Landscape/Private/LandscapeRender.cpp
@@ -726,6 +726,7 @@ void FLandscapeComponentSceneProxy::CreateRenderThreadResources()
 	LandscapeUniformShaderParameters.InitResource();
 
 	// Create MeshBatch for grass rendering
+	if(SharedBuffers->GrassIndexBuffer)
 	{
 		GrassMeshBatch.Elements.Empty(1);
 
