Fix for crash  resimulating a Procedural Foliage Volume with foliage type that has a landscape layer set. UE-15936

diff --git a/Engine/Source/Runtime/Landscape/Private/LandscapeEdit.cpp b/Engine/Source/Runtime/Landscape/Private/LandscapeEdit.cpp
--- a/Engine/Source/Runtime/Landscape/Private/LandscapeEdit.cpp
+++ b/Engine/Source/Runtime/Landscape/Private/LandscapeEdit.cpp
@@ -1522,6 +1522,12 @@ float ULandscapeComponent::GetLayerWeightAtLocation(const FVector& InLocation, U
 		float TestX = (InLocation.X - Landscape->GetActorLocation().X) / DrawScale.X - (float)GetSectionBase().X;
 		float TestY = (InLocation.Y - Landscape->GetActorLocation().Y) / DrawScale.Y - (float)GetSectionBase().Y;
 
+		// Abort if the test location is not on this component
+		if (TestX < 0 || TestY < 0 || TestX > ComponentSizeQuads || TestY > ComponentSizeQuads)
+		{
+			return 0.0f;
+		}
+
 		// Find data
 		int32 X1 = FMath::FloorToInt(TestX);
 		int32 Y1 = FMath::FloorToInt(TestY);
