Fixed issue where undoing landscape painting left any placed foliage floating in mid air [UE-16142]

diff --git a/Engine/Source/Runtime/Landscape/Private/LandscapeCollision.cpp b/Engine/Source/Runtime/Landscape/Private/LandscapeCollision.cpp
--- a/Engine/Source/Runtime/Landscape/Private/LandscapeCollision.cpp
+++ b/Engine/Source/Runtime/Landscape/Private/LandscapeCollision.cpp
@@ -1069,14 +1069,14 @@ void ULandscapeHeightfieldCollisionComponent::SnapFoliageInstances(const FBox& I
 	UWorld* ComponentWorld = GetWorld();
 	for (TActorIterator<AInstancedFoliageActor> It(ComponentWorld); It; ++It)
 	{
-		AInstancedFoliageActor& IFA = *(*It);
-		const auto BaseId = IFA.InstanceBaseCache.GetInstanceBaseId(this);
+		AInstancedFoliageActor* IFA = *It;
+		const auto BaseId = IFA->InstanceBaseCache.GetInstanceBaseId(this);
 		if (BaseId == FFoliageInstanceBaseCache::InvalidBaseId)
 		{
 			continue;
 		}
 			
-		for (auto& MeshPair : IFA.FoliageMeshes)
+		for (auto& MeshPair : IFA->FoliageMeshes)
 		{
 			// Find the per-mesh info matching the mesh.
 			UFoliageType* Settings = MeshPair.Key;
@@ -1125,6 +1125,8 @@ void ULandscapeHeightfieldCollisionComponent::SnapFoliageInstances(const FBox& I
 								bFoundHit = true;
 								if ((TestLocation - Hit.Location).SizeSquared() > KINDA_SMALL_NUMBER)
 								{
+									IFA->Modify();
+
 									// Remove instance location from the hash. Do not need to update ComponentHash as we re-add below.
 									MeshInfo.InstanceHash->RemoveInstance(Instance.Location, InstanceIndex);
 
@@ -1167,7 +1169,7 @@ void ULandscapeHeightfieldCollisionComponent::SnapFoliageInstances(const FBox& I
 				}
 
 				// Remove any unused instances
-				MeshInfo.RemoveInstances(&IFA, InstancesToRemove);
+				MeshInfo.RemoveInstances(IFA, InstancesToRemove);
 			}
 		}
 	}
