Fix crash assigning control point mesh
Fix control point mesh not being visible
Fix control point mesh in same level not persisting over save/load
(control point meshes in another level won't persist until a later fix)

diff --git a/Engine/Source/Runtime/Landscape/Private/LandscapeSplines.cpp b/Engine/Source/Runtime/Landscape/Private/LandscapeSplines.cpp
--- a/Engine/Source/Runtime/Landscape/Private/LandscapeSplines.cpp
+++ b/Engine/Source/Runtime/Landscape/Private/LandscapeSplines.cpp
@@ -1388,12 +1388,16 @@ void ULandscapeSplineControlPoint::UpdateSplinePoints(bool bUpdateCollision, boo
 			if (MeshComponentOuterSplines != OuterSplines)
 			{
 				MeshComponent = MeshComponentOuterSplines->GetForeignMeshComponent(this);
-				MeshComponentOuterSplines->Modify();
-				MeshComponentOuterSplines->UpdateModificationKey(this);
+				if (MeshComponent)
+				{
+					MeshComponentOuterSplines->Modify();
+					MeshComponentOuterSplines->UpdateModificationKey(this);
+				}
 			}
 		}
 
 		// Create mesh component if needed
+		bool bComponentNeedsRegistering = false;
 		if (MeshComponent == nullptr)
 		{
 			AActor* MeshComponentOuterActor = MeshComponentOuterSplines->GetOwner();
@@ -1402,9 +1406,12 @@ void ULandscapeSplineControlPoint::UpdateSplinePoints(bool bUpdateCollision, boo
 			MeshComponent = NewObject<UControlPointMeshComponent>(MeshComponentOuterActor, NAME_None, RF_Transactional | RF_TextExportTransient);
 			MeshComponent->bSelected = bSelected;
 			MeshComponent->AttachTo(MeshComponentOuterSplines);
+			bComponentNeedsRegistering = true;
+
 			if (MeshComponentOuterSplines == OuterSplines)
 			{
 				MeshComponentOuterSplines->MeshComponentLocalOwnersMap.Add(MeshComponent, this);
+				LocalMeshComponent = MeshComponent;
 			}
 			else
 			{
@@ -1468,11 +1475,16 @@ void ULandscapeSplineControlPoint::UpdateSplinePoints(bool bUpdateCollision, boo
 			MeshComponent->Modify();
 			MeshComponent->BodyInstance.SetCollisionProfileName(CollisionProfile);
 		}
+
+		if (bComponentNeedsRegistering)
+		{
+			MeshComponent->RegisterComponent();
+		}
 	}
 	else
 	{
 		MeshComponent = nullptr;
-		ForeignWorld = NULL;
+		ForeignWorld = nullptr;
 	}
 
 	// Destroy any unused components
