Fixed spline meshes stored in foreign levels disappearing in cooked builds

diff --git a/Engine/Source/Runtime/Landscape/Classes/LandscapeSplinesComponent.h b/Engine/Source/Runtime/Landscape/Classes/LandscapeSplinesComponent.h
--- a/Engine/Source/Runtime/Landscape/Classes/LandscapeSplinesComponent.h
+++ b/Engine/Source/Runtime/Landscape/Classes/LandscapeSplinesComponent.h
@@ -32,11 +32,11 @@ struct FForeignSplineSegmentData
 	GENERATED_USTRUCT_BODY()
 
 #if WITH_EDITORONLY_DATA
-		UPROPERTY()
-		FGuid ModificationKey;
+	UPROPERTY()
+	FGuid ModificationKey;
 
 	UPROPERTY()
-		TArray<USplineMeshComponent*> MeshComponents;
+	TArray<USplineMeshComponent*> MeshComponents;
 #endif
 
 	friend FArchive& operator<<(FArchive& Ar, FForeignSplineSegmentData& Value);
@@ -92,12 +92,19 @@ protected:
 	TArray<ULandscapeSplineSegment*> Segments;
 
 #if WITH_EDITORONLY_DATA
+	// Serialized
 	TMap<TAssetPtr<UWorld>, FForeignWorldSplineData> ForeignWorldSplineDataMap;
 
+	// Transient - rebuilt on load
 	TMap<UMeshComponent*, UObject*> MeshComponentLocalOwnersMap;
 	TMap<UMeshComponent*, TLazyObjectPtr<UObject>> MeshComponentForeignOwnersMap;
 #endif
 
+	// References to components owned by landscape splines in other levels
+	// for cooked build (uncooked keeps references via ForeignWorldSplineDataMap)
+	UPROPERTY()
+	TArray<UMeshComponent*> CookedForeignMeshComponents;
+
 public:
 	void CheckSplinesValid();
 	bool ModifySplines(bool bAlwaysMarkDirty = true);
diff --git a/Engine/Source/Runtime/Landscape/Private/LandscapeSplines.cpp b/Engine/Source/Runtime/Landscape/Private/LandscapeSplines.cpp
--- a/Engine/Source/Runtime/Landscape/Private/LandscapeSplines.cpp
+++ b/Engine/Source/Runtime/Landscape/Private/LandscapeSplines.cpp
@@ -486,6 +486,31 @@ FArchive& operator<<(FArchive& Ar, FForeignWorldSplineData& Value)
 
 void ULandscapeSplinesComponent::Serialize(FArchive& Ar)
 {
+#if WITH_EDITORONLY_DATA
+	// Cooking is a save-time operation, so has to be done before Super::Serialize
+	if (Ar.IsCooking())
+	{
+		CookedForeignMeshComponents.Reset();
+
+		for (const auto& ForeignWorldSplineDataPair : ForeignWorldSplineDataMap)
+		{
+			const auto& ForeignWorldSplineData = ForeignWorldSplineDataPair.Value;
+
+			for (const auto& ForeignControlPointDataPair : ForeignWorldSplineData.ForeignControlPointDataMap)
+			{
+				const auto& ForeignControlPointData = ForeignControlPointDataPair.Value;
+				CookedForeignMeshComponents.Add(ForeignControlPointData.MeshComponent);
+			}
+
+			for (const auto& ForeignSplineSegmentDataPair : ForeignWorldSplineData.ForeignSplineSegmentDataMap)
+			{
+				const auto& ForeignSplineSegmentData = ForeignSplineSegmentDataPair.Value;
+				CookedForeignMeshComponents.Append(ForeignSplineSegmentData.MeshComponents);
+			}
+		}
+	}
+#endif
+
 	Super::Serialize(Ar);
 
 #if WITH_EDITORONLY_DATA
