Fixed landscape spline cross-level control point mesh components not surviving save/load in the editor

diff --git a/Engine/Source/Runtime/Landscape/Classes/LandscapeSplinesComponent.h b/Engine/Source/Runtime/Landscape/Classes/LandscapeSplinesComponent.h
--- a/Engine/Source/Runtime/Landscape/Classes/LandscapeSplinesComponent.h
+++ b/Engine/Source/Runtime/Landscape/Classes/LandscapeSplinesComponent.h
@@ -10,6 +10,8 @@ class ULandscapeSplineSegment;
 class USplineMeshComponent;
 class UControlPointMeshComponent;
 
+// structs for ForeignWorldSplineDataMap
+// these are editor-only, but we don't have the concept of an editor-only USTRUCT
 USTRUCT()
 struct FForeignControlPointData
 {
@@ -22,8 +24,6 @@ struct FForeignControlPointData
 	UPROPERTY()
 	UControlPointMeshComponent* MeshComponent;
 #endif
-
-	friend FArchive& operator<<(FArchive& Ar, FForeignControlPointData& Value);
 };
 
 USTRUCT()
@@ -38,8 +38,6 @@ struct FForeignSplineSegmentData
 	UPROPERTY()
 	TArray<USplineMeshComponent*> MeshComponents;
 #endif
-
-	friend FArchive& operator<<(FArchive& Ar, FForeignSplineSegmentData& Value);
 };
 
 USTRUCT()
@@ -48,15 +46,22 @@ struct FForeignWorldSplineData
 	GENERATED_USTRUCT_BODY()
 
 #if WITH_EDITORONLY_DATA
+	UPROPERTY()
 	TMap<TLazyObjectPtr<ULandscapeSplineControlPoint>, FForeignControlPointData> ForeignControlPointDataMap;
+
+	UPROPERTY()
 	TMap<TLazyObjectPtr<ULandscapeSplineSegment>, FForeignSplineSegmentData> ForeignSplineSegmentDataMap;
 #endif
 
-	friend FArchive& operator<<(FArchive& Ar, FForeignWorldSplineData& Value);
-
+#if WITH_EDITOR
 	bool IsEmpty();
+#endif
 };
 
+//////////////////////////////////////////////////////////////////////////
+// ULandscapeSplinesComponent
+//////////////////////////////////////////////////////////////////////////
+
 UCLASS(MinimalAPI)
 class ULandscapeSplinesComponent : public UPrimitiveComponent
 {
@@ -93,6 +98,7 @@ protected:
 
 #if WITH_EDITORONLY_DATA
 	// Serialized
+	UPROPERTY(TextExportTransient)
 	TMap<TAssetPtr<UWorld>, FForeignWorldSplineData> ForeignWorldSplineDataMap;
 
 	// Transient - rebuilt on load
@@ -102,7 +108,7 @@ protected:
 
 	// References to components owned by landscape splines in other levels
 	// for cooked build (uncooked keeps references via ForeignWorldSplineDataMap)
-	UPROPERTY()
+	UPROPERTY(TextExportTransient)
 	TArray<UMeshComponent*> CookedForeignMeshComponents;
 
 public:
diff --git a/Engine/Source/Runtime/Landscape/Private/LandscapeModule.cpp b/Engine/Source/Runtime/Landscape/Private/LandscapeModule.cpp
--- a/Engine/Source/Runtime/Landscape/Private/LandscapeModule.cpp
+++ b/Engine/Source/Runtime/Landscape/Private/LandscapeModule.cpp
@@ -6,6 +6,10 @@
 #include "UObjectHash.h"
 #include "Engine/World.h"
 #include "Engine/Level.h"
+#include "LandscapeVersion.h"
+
+// Register the custom version with core
+FCustomVersionRegistration GRegisterLandscapeCustomVersion(FLandscapeCustomVersion::GUID, FLandscapeCustomVersion::LatestVersion, TEXT("Landscape"));
 
 class FLandscapeModule : public IModuleInterface
 {
diff --git a/Engine/Source/Runtime/Landscape/Private/LandscapeSplines.cpp b/Engine/Source/Runtime/Landscape/Private/LandscapeSplines.cpp
--- a/Engine/Source/Runtime/Landscape/Private/LandscapeSplines.cpp
+++ b/Engine/Source/Runtime/Landscape/Private/LandscapeSplines.cpp
@@ -23,6 +23,7 @@
 #include "MessageLog.h"
 #include "UObjectToken.h"
 #endif
+#include "LandscapeVersion.h"
 
 IMPLEMENT_HIT_PROXY(HLandscapeSplineProxy, HHitProxy);
 IMPLEMENT_HIT_PROXY(HLandscapeSplineProxy_Segment, HLandscapeSplineProxy);
@@ -451,38 +452,27 @@ bool ULandscapeSplinesComponent::ModifySplines(bool bAlwaysMarkDirty /*= true*/)
 	return bSavedToTransactionBuffer;
 }
 
-FArchive& operator<<(FArchive& Ar, FForeignControlPointData& Value)
-{
 #if WITH_EDITORONLY_DATA
-	if (!Ar.IsFilterEditorOnly())
-	{
-		Ar << Value.ModificationKey << Value.MeshComponent;
-	}
-#endif
-	return Ar;
-}
-
+// legacy ForeignWorldSplineDataMap serialization
 FArchive& operator<<(FArchive& Ar, FForeignSplineSegmentData& Value)
 {
-#if WITH_EDITORONLY_DATA
 	if (!Ar.IsFilterEditorOnly())
 	{
 		Ar << Value.ModificationKey << Value.MeshComponents;
 	}
-#endif
 	return Ar;
 }
 
 FArchive& operator<<(FArchive& Ar, FForeignWorldSplineData& Value)
 {
-#if WITH_EDITORONLY_DATA
 	if (!Ar.IsFilterEditorOnly())
 	{
+		// note: ForeignControlPointDataMap is missing in legacy serialization
 		Ar << Value.ForeignSplineSegmentDataMap;
 	}
-#endif
 	return Ar;
 }
+#endif
 
 void ULandscapeSplinesComponent::Serialize(FArchive& Ar)
 {
@@ -517,7 +507,12 @@ void ULandscapeSplinesComponent::Serialize(FArchive& Ar)
 	if (Ar.UE4Ver() >= VER_UE4_LANDSCAPE_SPLINE_CROSS_LEVEL_MESHES &&
 		!Ar.IsFilterEditorOnly())
 	{
-		Ar << ForeignWorldSplineDataMap;
+		Ar.UsingCustomVersion(FLandscapeCustomVersion::GUID);
+
+		if (Ar.CustomVer(FLandscapeCustomVersion::GUID) < FLandscapeCustomVersion::NewSplineCrossLevelMeshSerialization)
+		{
+			Ar << ForeignWorldSplineDataMap;
+		}
 	}
 
 	if (!Ar.IsPersistent())
diff --git a/Engine/Source/Runtime/Landscape/Private/LandscapeVersion.h b/Engine/Source/Runtime/Landscape/Private/LandscapeVersion.h
new file mode 100644
--- /dev/null
+++ b/Engine/Source/Runtime/Landscape/Private/LandscapeVersion.h
@@ -0,0 +1,22 @@
+// Copyright 1998-2015 Epic Games, Inc. All Rights Reserved.
+
+#pragma once
+
+// Custom version for landscape serialization
+namespace FLandscapeCustomVersion
+{
+	enum Type
+	{
+		// Before any version changes were made in the plugin
+		BeforeCustomVersionWasAdded = 0,
+		// Changed to TMap properties instead of manual serialization and fixed landscape spline control point cross-level mesh components not being serialized
+		NewSplineCrossLevelMeshSerialization,
+
+		// -----<new versions can be added above this line>-------------------------------------------------
+		VersionPlusOne,
+		LatestVersion = VersionPlusOne - 1
+	};
+
+	// The GUID for this custom version
+	const static FGuid GUID = FGuid(0x23AFE18E, 0x4CE14E58, 0x8D61C252, 0xB953BEB7);
+};
