Fixed: Changes to Landscape Thickness not propagated to PhysX
UE-16517

diff --git a/Engine/Source/Runtime/Landscape/Private/LandscapeEdit.cpp b/Engine/Source/Runtime/Landscape/Private/LandscapeEdit.cpp
--- a/Engine/Source/Runtime/Landscape/Private/LandscapeEdit.cpp
+++ b/Engine/Source/Runtime/Landscape/Private/LandscapeEdit.cpp
@@ -3202,7 +3202,7 @@ void ALandscapeProxy::PostEditChangeProperty(FPropertyChangedEvent& PropertyChan
 	{
 		ChangedPhysMaterial();
 	}
-	else if (GIsEditor && (PropertyName == FName(TEXT("CollisionMipLevel"))))
+	else if (GIsEditor && (PropertyName == FName(TEXT("CollisionMipLevel")) || PropertyName == FName(TEXT("CollisionThickness"))))
 	{
 		RecreateCollisionComponents();
 	}
