Fix for crash in FMaterialUpdateContext when assigning a landscape material to a LandscapeProxy. (UE-16900)

diff --git a/Engine/Source/Runtime/Landscape/Private/LandscapeEdit.cpp b/Engine/Source/Runtime/Landscape/Private/LandscapeEdit.cpp
--- a/Engine/Source/Runtime/Landscape/Private/LandscapeEdit.cpp
+++ b/Engine/Source/Runtime/Landscape/Private/LandscapeEdit.cpp
@@ -3168,18 +3168,20 @@ void ALandscapeProxy::PostEditChangeProperty(FPropertyChangedEvent& PropertyChan
 		}
 		else if (PropertyName == FName(TEXT("LandscapeMaterial")) || PropertyName == FName(TEXT("LandscapeHoleMaterial")))
 		{
-			FMaterialUpdateContext MaterialUpdateContext;
-			GetLandscapeInfo()->UpdateLayerInfoMap(/*this*/);
-
-			// Clear the parents out of combination material instances
-			for (TMap<FString, UMaterialInstanceConstant*>::TIterator It(MaterialInstanceConstantMap); It; ++It)
 			{
-				It.Value()->SetParentEditorOnly(NULL);
-				MaterialUpdateContext.AddMaterial(It.Value()->GetMaterial());
-			}
+				FMaterialUpdateContext MaterialUpdateContext;
+				GetLandscapeInfo()->UpdateLayerInfoMap(/*this*/);
+
+				// Clear the parents out of combination material instances
+				for (TMap<FString, UMaterialInstanceConstant*>::TIterator It(MaterialInstanceConstantMap); It; ++It)
+				{
+					It.Value()->SetParentEditorOnly(NULL);
+					MaterialUpdateContext.AddMaterial(It.Value()->GetMaterial());
+				}
 
-			// Remove our references to any material instances
-			MaterialInstanceConstantMap.Empty();
+				// Remove our references to any material instances
+				MaterialInstanceConstantMap.Empty();
+			}
 
 			for (int32 ComponentIndex = 0; ComponentIndex < LandscapeComponents.Num(); ComponentIndex++)
 			{
