Fix landscape physical material crashes due to mismatch between saved level and DDC [UE-17367]

diff --git a/Engine/Source/Runtime/Landscape/Classes/LandscapeHeightfieldCollisionComponent.h b/Engine/Source/Runtime/Landscape/Classes/LandscapeHeightfieldCollisionComponent.h
--- a/Engine/Source/Runtime/Landscape/Classes/LandscapeHeightfieldCollisionComponent.h
+++ b/Engine/Source/Runtime/Landscape/Classes/LandscapeHeightfieldCollisionComponent.h
@@ -192,7 +192,7 @@ class ULandscapeHeightfieldCollisionComponent : public UPrimitiveComponent
 	/** 
 	 * Cooks raw height data into collision object binary stream
 	 */
-	virtual bool CookCollisionData(const FName& Format, bool bUseOnlyDefMaterial, bool bCheckDDC, TArray<uint8>& OutCookedData, TArray<UPhysicalMaterial*>& OutMaterials) const;
+	virtual bool CookCollisionData(const FName& Format, bool bUseOnlyDefMaterial, bool bCheckDDC, TArray<uint8>& OutCookedData, TArray<UPhysicalMaterial*>& InOutMaterials) const;
 
 	/** Modify a sub-region of the PhysX heightfield. Note that this does not update the physical material */
 	void UpdateHeightfieldRegion(int32 ComponentX1, int32 ComponentY1, int32 ComponentX2, int32 ComponentY2);
diff --git a/Engine/Source/Runtime/Landscape/Classes/LandscapeMeshCollisionComponent.h b/Engine/Source/Runtime/Landscape/Classes/LandscapeMeshCollisionComponent.h
--- a/Engine/Source/Runtime/Landscape/Classes/LandscapeMeshCollisionComponent.h
+++ b/Engine/Source/Runtime/Landscape/Classes/LandscapeMeshCollisionComponent.h
@@ -88,7 +88,7 @@ class ULandscapeMeshCollisionComponent : public ULandscapeHeightfieldCollisionCo
 	virtual void ExportCustomProperties(FOutputDevice& Out, uint32 Indent) override;
 	virtual void ImportCustomProperties(const TCHAR* SourceText, FFeedbackContext* Warn) override;
 
-	virtual bool CookCollisionData(const FName& Format, bool bUseOnlyDefMaterial, bool bCheckDDC, TArray<uint8>& OutCookedData, TArray<UPhysicalMaterial*>& OutMaterials) const override;
+	virtual bool CookCollisionData(const FName& Format, bool bUseOnlyDefMaterial, bool bCheckDDC, TArray<uint8>& OutCookedData, TArray<UPhysicalMaterial*>& InOutMaterials) const override;
 #endif
 	// End UObject Interface.
 
diff --git a/Engine/Source/Runtime/Landscape/Private/LandscapeCollision.cpp b/Engine/Source/Runtime/Landscape/Private/LandscapeCollision.cpp
--- a/Engine/Source/Runtime/Landscape/Private/LandscapeCollision.cpp
+++ b/Engine/Source/Runtime/Landscape/Private/LandscapeCollision.cpp
@@ -77,13 +77,40 @@ ULandscapeMeshCollisionComponent::FPhysXMeshRef::~FPhysXMeshRef()
 	GSharedMeshRefs.Remove(Guid);
 }
 
-// Generate a new guid to force a recache of landscape collisoon derived data
+// Generate a new guid to force a recache of landscape collison derived data
 #define LANDSCAPE_COLLISION_DERIVEDDATA_VER	TEXT("5DF9E1AAB7CC4DCCB2965BA1A78DFE8")
 
-static FString GetHFDDCKeyString(const FName& Format, bool bDefMaterial, const FGuid& StateId)
+static FString GetHFDDCKeyString(const FName& Format, bool bDefMaterial, const FGuid& StateId, const TArray<UPhysicalMaterial*>& PhysicalMaterials)
 {
+	FGuid CombinedStateId;
+
+	if (bDefMaterial)
+	{
+		CombinedStateId = StateId;
+	}
+	else
+	{
+		// Build a combined state ID based on both the heightfield state and all physical materials.
+		FBufferArchive CombinedStateAr;
+
+		// Add main heightfield state
+		FGuid HeightfieldState = StateId;
+		CombinedStateAr << HeightfieldState;
+
+		// Add physical materials
+		for (UPhysicalMaterial* PhysicalMaterial : PhysicalMaterials)
+		{
+			FString PhysicalMaterialName = PhysicalMaterial->GetPathName().ToUpper();
+			CombinedStateAr << PhysicalMaterialName;
+		}
+
+		uint32 Hash[5];
+		FSHA1::HashBuffer(CombinedStateAr.GetData(), CombinedStateAr.Num(), (uint8*)Hash);
+		CombinedStateId = FGuid(Hash[0] ^ Hash[4], Hash[1], Hash[2], Hash[3]);
+	}
+
 	const FString KeyPrefix = FString::Printf(TEXT("%s_%s"), *Format.ToString(), (bDefMaterial ? TEXT("VIS") : TEXT("FULL")));
-	return FDerivedDataCacheInterface::BuildCacheKey(*KeyPrefix, LANDSCAPE_COLLISION_DERIVEDDATA_VER, *StateId.ToString());
+	return FDerivedDataCacheInterface::BuildCacheKey(*KeyPrefix, LANDSCAPE_COLLISION_DERIVEDDATA_VER, *CombinedStateId.ToString());
 }
 
 ECollisionEnabled::Type ULandscapeHeightfieldCollisionComponent::GetCollisionEnabled() const
@@ -278,13 +305,6 @@ void ULandscapeHeightfieldCollisionComponent::CreateCollisionObject()
 			ExistingHeightfieldRef = GSharedHeightfieldRefs.FindRef(HeightfieldGuid);
 		}
 
-		// This should only occur if a level prior to VER_UE4_LANDSCAPE_COLLISION_DATA_COOKING 
-		// was resaved using a commandlet and not saved in the editor.
-		if (CookedPhysicalMaterials.Num() == 0)
-		{
-			bCheckDDC = false;
-		}
-
 		if (ExistingHeightfieldRef)
 		{
 			HeightfieldRef = ExistingHeightfieldRef;
@@ -292,6 +312,13 @@ void ULandscapeHeightfieldCollisionComponent::CreateCollisionObject()
 		else
 		{
 #if WITH_EDITOR
+		    // This should only occur if a level prior to VER_UE4_LANDSCAPE_COLLISION_DATA_COOKING 
+		    // was resaved using a commandlet and not saved in the editor, or if a PhysicalMaterial asset was deleted.
+		    if (CookedPhysicalMaterials.Num() == 0 || CookedPhysicalMaterials.Contains(nullptr))
+		    {
+			    bCheckDDC = false;
+		    }
+
 			// Prepare heightfield data
 			static FName PhysicsFormatName(FPlatformProperties::GetPhysicsFormat());
 			CookCollisionData(PhysicsFormatName, false, bCheckDDC, CookedCollisionData, CookedPhysicalMaterials);
@@ -352,7 +379,7 @@ void ULandscapeHeightfieldCollisionComponent::SpeculativelyLoadAsyncDDCCollsionD
 		{
 			static FName PhysicsFormatName(FPlatformProperties::GetPhysicsFormat());
 
-			FString Key = GetHFDDCKeyString(PhysicsFormatName, false, HeightfieldGuid);
+			FString Key = GetHFDDCKeyString(PhysicsFormatName, false, HeightfieldGuid, CookedPhysicalMaterials);
 			uint32 Handle = GetDerivedDataCacheRef().GetAsynchronous(*Key);
 			check(!SpeculativeDDCRequest.IsValid());
 			SpeculativeDDCRequest = MakeShareable(new FAsyncPreRegisterDDCRequest(Key, Handle));
@@ -362,7 +389,7 @@ void ULandscapeHeightfieldCollisionComponent::SpeculativelyLoadAsyncDDCCollsionD
 #endif
 }
 
-bool ULandscapeHeightfieldCollisionComponent::CookCollisionData(const FName& Format, bool bUseDefMaterial, bool bCheckDDC, TArray<uint8>& OutCookedData, TArray<UPhysicalMaterial*>& OutMaterials) const
+bool ULandscapeHeightfieldCollisionComponent::CookCollisionData(const FName& Format, bool bUseDefMaterial, bool bCheckDDC, TArray<uint8>& OutCookedData, TArray<UPhysicalMaterial*>& InOutMaterials) const
 {
 #if WITH_PHYSX
 	// we have 2 versions of collision objects
@@ -373,7 +400,7 @@ bool ULandscapeHeightfieldCollisionComponent::CookCollisionData(const FName& For
 		// Ensure that content was saved with physical materials before using DDC data
 		if (GetLinkerUE4Version() >= VER_UE4_LANDSCAPE_SERIALIZE_PHYSICS_MATERIALS)
 		{
-			FString DDCKey = GetHFDDCKeyString(Format, bUseDefMaterial, HeightfieldGuid);
+			FString DDCKey = GetHFDDCKeyString(Format, bUseDefMaterial, HeightfieldGuid, InOutMaterials);
 
 			// Check if the speculatively-loaded data loaded and is what we wanted
 			if (SpeculativeDDCRequest.IsValid() && DDCKey == SpeculativeDDCRequest->GetKey())
@@ -422,7 +449,7 @@ bool ULandscapeHeightfieldCollisionComponent::CookCollisionData(const FName& For
 	}
 
 	// List of materials which is actually used by heightfield
-	OutMaterials.Empty();
+	InOutMaterials.Empty();
 
 	TArray<PxHeightFieldSample> Samples;
 	const int32 NumSamples = FMath::Square(CollisionSizeVerts);
@@ -458,7 +485,7 @@ bool ULandscapeHeightfieldCollisionComponent::CookCollisionData(const FName& For
 						else
 						{
 							UPhysicalMaterial* DominantMaterial = Layer && Layer->PhysMaterial ? Layer->PhysMaterial : DefMaterial;
-							MaterialIndex = OutMaterials.AddUnique(DominantMaterial);
+							MaterialIndex = InOutMaterials.AddUnique(DominantMaterial);
 						}
 					}
 				}
@@ -478,9 +505,9 @@ bool ULandscapeHeightfieldCollisionComponent::CookCollisionData(const FName& For
 	}
 
 	// Add the default physical material to be used used when we have no dominant data.
-	if (OutMaterials.Num() == 0)
+	if (InOutMaterials.Num() == 0)
 	{
-		OutMaterials.Add(DefMaterial);
+		InOutMaterials.Add(DefMaterial);
 	}
 
 	//
@@ -499,14 +526,14 @@ bool ULandscapeHeightfieldCollisionComponent::CookCollisionData(const FName& For
 
 		if (bShouldSaveCookedDataToDDC[CookedDataIndex])
 		{
-			GetDerivedDataCacheRef().Put(*GetHFDDCKeyString(Format, bUseDefMaterial, HeightfieldGuid), OutCookedData);
+			GetDerivedDataCacheRef().Put(*GetHFDDCKeyString(Format, bUseDefMaterial, HeightfieldGuid, InOutMaterials), OutCookedData);
 			bShouldSaveCookedDataToDDC[CookedDataIndex] = false;
 		}
 	}
 	else
 	{
 		OutCookedData.Empty();
-		OutMaterials.Empty();
+		InOutMaterials.Empty();
 	}
 
 	return Result;
@@ -515,7 +542,7 @@ bool ULandscapeHeightfieldCollisionComponent::CookCollisionData(const FName& For
 	return false;
 }
 
-bool ULandscapeMeshCollisionComponent::CookCollisionData(const FName& Format, bool bUseDefMaterial, bool bCheckDDC, TArray<uint8>& OutCookedData, TArray<UPhysicalMaterial*>& OutMaterials) const
+bool ULandscapeMeshCollisionComponent::CookCollisionData(const FName& Format, bool bUseDefMaterial, bool bCheckDDC, TArray<uint8>& OutCookedData, TArray<UPhysicalMaterial*>& InOutMaterials) const
 {
 #if WITH_PHYSX
 	
@@ -527,7 +554,7 @@ bool ULandscapeMeshCollisionComponent::CookCollisionData(const FName& Format, bo
 		// Ensure that content was saved with physical materials before using DDC data
 		if (GetLinkerUE4Version() >= VER_UE4_LANDSCAPE_SERIALIZE_PHYSICS_MATERIALS)
 		{
-			FString DDCKey = GetHFDDCKeyString(Format, bUseDefMaterial, MeshGuid);
+			FString DDCKey = GetHFDDCKeyString(Format, bUseDefMaterial, MeshGuid, InOutMaterials);
 
 			// Check if the speculatively-loaded data loaded and is what we wanted
 			if (SpeculativeDDCRequest.IsValid() && DDCKey == SpeculativeDDCRequest->GetKey())
@@ -555,7 +582,7 @@ bool ULandscapeMeshCollisionComponent::CookCollisionData(const FName& Format, bo
 	UPhysicalMaterial* DefMaterial = (Proxy && Proxy->DefaultPhysMaterial != nullptr) ? Proxy->DefaultPhysMaterial : GEngine->DefaultPhysMaterial;
 
 	// List of materials which is actually used by trimesh
-	OutMaterials.Empty();
+	InOutMaterials.Empty();
 
 	TArray<FVector>			Vertices;
 	TArray<FTriIndices>		Indices;
@@ -614,7 +641,7 @@ bool ULandscapeMeshCollisionComponent::CookCollisionData(const FName& Format, bo
 					else
 					{
 						UPhysicalMaterial* DominantMaterial = Layer && Layer->PhysMaterial ? Layer->PhysMaterial : DefMaterial;
-						MaterialIndex = OutMaterials.AddUnique(DominantMaterial);
+						MaterialIndex = InOutMaterials.AddUnique(DominantMaterial);
 					}
 				}
 			}
@@ -669,9 +696,9 @@ bool ULandscapeMeshCollisionComponent::CookCollisionData(const FName& Format, bo
 	}
 
 	// Add the default physical material to be used used when we have no dominant data.
-	if (OutMaterials.Num() == 0)
+	if (InOutMaterials.Num() == 0)
 	{
-		OutMaterials.Add(DefMaterial);
+		InOutMaterials.Add(DefMaterial);
 	}
 
 	bool bFlipNormals = true;
@@ -687,14 +714,14 @@ bool ULandscapeMeshCollisionComponent::CookCollisionData(const FName& Format, bo
 
 		if (bShouldSaveCookedDataToDDC[CookedDataIndex])
 		{
-			GetDerivedDataCacheRef().Put(*GetHFDDCKeyString(Format, bUseDefMaterial, MeshGuid), OutCookedData);
+			GetDerivedDataCacheRef().Put(*GetHFDDCKeyString(Format, bUseDefMaterial, MeshGuid, InOutMaterials), OutCookedData);
 			bShouldSaveCookedDataToDDC[CookedDataIndex] = false;
 		}
 	}
 	else
 	{
 		OutCookedData.Empty();
-		OutMaterials.Empty();
+		InOutMaterials.Empty();
 	}
 
 	return Result;
@@ -731,6 +758,13 @@ void ULandscapeMeshCollisionComponent::CreateCollisionObject()
 		else
 		{
 #if WITH_EDITOR
+		    // This should only occur if a level prior to VER_UE4_LANDSCAPE_COLLISION_DATA_COOKING 
+		    // was resaved using a commandlet and not saved in the editor, or if a PhysicalMaterial asset was deleted.
+		    if (CookedPhysicalMaterials.Num() == 0 || CookedPhysicalMaterials.Contains(nullptr))
+		    {
+			    bCheckDDC = false;
+		    }
+
 			// Create cooked physics data
 			static FName PhysicsFormatName(FPlatformProperties::GetPhysicsFormat());
 			CookCollisionData(PhysicsFormatName, false, bCheckDDC, CookedCollisionData, CookedPhysicalMaterials);
@@ -1197,7 +1231,7 @@ void ULandscapeHeightfieldCollisionComponent::Serialize(FArchive& Ar)
 		{
 			FName Format = Ar.CookingTarget()->GetPhysicsFormat(nullptr);
 			CookCollisionData(Format, false, true, CookedCollisionData, CookedPhysicalMaterials);
-			GetDerivedDataCacheRef().Put(*GetHFDDCKeyString(Format, false, HeightfieldGuid), CookedCollisionData);
+			GetDerivedDataCacheRef().Put(*GetHFDDCKeyString(Format, false, HeightfieldGuid, CookedPhysicalMaterials), CookedCollisionData);
 		}
 	}
 #endif// WITH_EDITOR
@@ -1457,12 +1491,12 @@ void ULandscapeHeightfieldCollisionComponent::PreSave()
 		static FName PhysicsFormatName(FPlatformProperties::GetPhysicsFormat());
 		if (CookedCollisionData.Num())
 		{
-			GetDerivedDataCacheRef().Put(*GetHFDDCKeyString(PhysicsFormatName, false, HeightfieldGuid), CookedCollisionData);
+			GetDerivedDataCacheRef().Put(*GetHFDDCKeyString(PhysicsFormatName, false, HeightfieldGuid, CookedPhysicalMaterials), CookedCollisionData);
 		}
 
 		if (CookedCollisionDataEd.Num())
 		{
-			GetDerivedDataCacheRef().Put(*GetHFDDCKeyString(PhysicsFormatName, true, HeightfieldGuid), CookedCollisionDataEd);
+			GetDerivedDataCacheRef().Put(*GetHFDDCKeyString(PhysicsFormatName, true, HeightfieldGuid, TArray<UPhysicalMaterial*>()), CookedCollisionDataEd);
 		}
 #endif// WITH_EDITOR
 	}
