Prevent and warn user trying to use default engine material as a landscape hole material. Overriding blend mode caused a crash. UE-17601

diff --git a/Engine/Source/Runtime/Landscape/Landscape.Build.cs b/Engine/Source/Runtime/Landscape/Landscape.Build.cs
--- a/Engine/Source/Runtime/Landscape/Landscape.Build.cs
+++ b/Engine/Source/Runtime/Landscape/Landscape.Build.cs
@@ -57,6 +57,8 @@ public class Landscape : ModuleRules
 				new string[] {
                     "UnrealEd",
                     "MaterialUtilities", 
+    				"SlateCore",
+    				"Slate",
                 }
 			);
 
diff --git a/Engine/Source/Runtime/Landscape/Private/LandscapeEdit.cpp b/Engine/Source/Runtime/Landscape/Private/LandscapeEdit.cpp
--- a/Engine/Source/Runtime/Landscape/Private/LandscapeEdit.cpp
+++ b/Engine/Source/Runtime/Landscape/Private/LandscapeEdit.cpp
@@ -37,6 +37,9 @@ LandscapeEdit.cpp: Landscape editing
 #include "EngineGlobals.h"
 #include "ShowFlags.h"
 #include "ConvexVolume.h"
+#include "SlateBasics.h"  // For AddNotification
+#include "SNotificationList.h"
+#include "NotificationManager.h"
 #endif
 #include "ComponentReregisterContext.h"
 
@@ -171,7 +174,28 @@ UMaterialInstanceConstant* ULandscapeComponent::GetCombinationMaterial(bool bMob
 	UMaterialInterface* const LandscapeMaterial = GetLandscapeMaterial();
 	UMaterialInterface* const HoleMaterial = bComponentHasHoles ? GetLandscapeHoleMaterial() : nullptr;
 	UMaterialInterface* const MaterialToUse = bComponentHasHoles && HoleMaterial ? HoleMaterial : LandscapeMaterial;
-	const bool bOverrideBlendMode = bComponentHasHoles && !HoleMaterial && LandscapeMaterial->GetBlendMode() == BLEND_Opaque;
+	bool bOverrideBlendMode = bComponentHasHoles && !HoleMaterial && LandscapeMaterial->GetBlendMode() == BLEND_Opaque;
+
+	if (bOverrideBlendMode)
+	{
+		UMaterial* Material = LandscapeMaterial->GetMaterial();
+		if (Material && Material->bUsedAsSpecialEngineMaterial)
+		{
+			bOverrideBlendMode = false;
+#if WITH_EDITOR
+			static TWeakPtr<SNotificationItem> ExistingNotification;
+			if (!ExistingNotification.IsValid())
+			{
+				// let the user know why they are not seeing holes
+				FNotificationInfo Info(LOCTEXT("AssignLandscapeMaterial", "You must assign a regular, non-engine material to your landscape in order to see holes created with the visibility tool."));
+				Info.ExpireDuration = 5.0f;
+				Info.bUseSuccessFailIcons = true;
+				ExistingNotification = TWeakPtr<SNotificationItem>(FSlateNotificationManager::Get().AddNotification(Info));
+			}
+#endif
+			return nullptr;
+		}
+	}
 
 	if (ensure(MaterialToUse != nullptr))
 	{
@@ -3175,7 +3199,8 @@ void ALandscapeProxy::PostEditChangeProperty(FPropertyChangedEvent& PropertyChan
 				// Clear the parents out of combination material instances
 				for (TMap<FString, UMaterialInstanceConstant*>::TIterator It(MaterialInstanceConstantMap); It; ++It)
 				{
-					It.Value()->SetParentEditorOnly(NULL);
+					It.Value()->BasePropertyOverrides.bOverride_BlendMode = false;
+					It.Value()->SetParentEditorOnly(nullptr);
 					MaterialUpdateContext.AddMaterial(It.Value()->GetMaterial());
 				}
 
@@ -3325,7 +3350,8 @@ void ALandscape::PostEditChangeProperty(FPropertyChangedEvent& PropertyChangedEv
 		// Clear the parents out of combination material instances
 		for (TMap<FString, UMaterialInstanceConstant*>::TIterator It(MaterialInstanceConstantMap); It; ++It)
 		{
-			It.Value()->SetParentEditorOnly(NULL);
+			It.Value()->BasePropertyOverrides.bOverride_BlendMode = false;
+			It.Value()->SetParentEditorOnly(nullptr);
 			MaterialUpdateContext.AddMaterial(It.Value()->GetMaterial());
 		}
 
