Fixed: Bake Grass to Landscape only bakes top left corner of texture maps UE-17832

diff --git a/Engine/Source/Runtime/Landscape/Private/LandscapeGrass.cpp b/Engine/Source/Runtime/Landscape/Private/LandscapeGrass.cpp
--- a/Engine/Source/Runtime/Landscape/Private/LandscapeGrass.cpp
+++ b/Engine/Source/Runtime/Landscape/Private/LandscapeGrass.cpp
@@ -1568,7 +1568,7 @@ void ALandscapeProxy::UpdateGrass(const TArray<FVector>& Cameras, bool bForceSyn
 											}
 										}
 
-										if (NumCompsCreated || (!bForceSync && AsyncFoliageTasks.Num() >= MaxTasks))
+										if (!bForceSync && (NumCompsCreated || AsyncFoliageTasks.Num() >= MaxTasks))
 										{
 											continue; // one per frame, but we still want to touch the existing ones
 										}
