Fixed a couple of GWorlds in Landscape

diff --git a/Engine/Source/Runtime/Landscape/Private/LandscapeEdit.cpp b/Engine/Source/Runtime/Landscape/Private/LandscapeEdit.cpp
--- a/Engine/Source/Runtime/Landscape/Private/LandscapeEdit.cpp
+++ b/Engine/Source/Runtime/Landscape/Private/LandscapeEdit.cpp
@@ -2977,7 +2977,7 @@ void ULandscapeLayerInfoObject::PostEditChangeProperty(FPropertyChangedEvent& Pr
 		else if (PropertyName == NAME_PhysMaterial)
 		{
 			// Only care current world object
-			for (TActorIterator<ALandscapeProxy> It(GWorld); It; ++It)
+			for (TActorIterator<ALandscapeProxy> It(GetWorld()); It; ++It)
 			{
 				ALandscapeProxy* Proxy = *It;
 				ULandscapeInfo* Info = Proxy->GetLandscapeInfo(false);
diff --git a/Engine/Source/Runtime/Landscape/Private/LandscapeGrass.cpp b/Engine/Source/Runtime/Landscape/Private/LandscapeGrass.cpp
--- a/Engine/Source/Runtime/Landscape/Private/LandscapeGrass.cpp
+++ b/Engine/Source/Runtime/Landscape/Private/LandscapeGrass.cpp
@@ -778,7 +778,7 @@ void ULandscapeGrassType::PostEditChangeProperty(FPropertyChangedEvent& Property
 	if (GIsEditor)
 	{
 		// Only care current world object
-		for (TActorIterator<ALandscapeProxy> It(GWorld); It; ++It)
+		for (TActorIterator<ALandscapeProxy> It(GetWorld()); It; ++It)
 		{
 			ALandscapeProxy* Proxy = *It;
 			const UMaterialInterface* MaterialInterface = Proxy->LandscapeMaterial;
