Fixed problems with landscape cooking where launching on mobile multiple times from the editor would not correctly cook changes made in the editor session
(UE-18292)

diff --git a/Engine/Source/Runtime/Engine/Private/Materials/MaterialInterface.cpp b/Engine/Source/Runtime/Engine/Private/Materials/MaterialInterface.cpp
--- a/Engine/Source/Runtime/Engine/Private/Materials/MaterialInterface.cpp
+++ b/Engine/Source/Runtime/Engine/Private/Materials/MaterialInterface.cpp
@@ -35,8 +35,11 @@ UMaterialInterface::UMaterialInterface(const FObjectInitializer& ObjectInitializ
 		InitDefaultMaterials();
 		AssertDefaultMaterialsExist();
 
-		SamplerTypeEnum = FindObject<UEnum>(NULL, TEXT("/Script/Engine.EMaterialSamplerType"));
-		check(SamplerTypeEnum);
+		if (SamplerTypeEnum == nullptr)
+		{
+			SamplerTypeEnum = FindObject<UEnum>(NULL, TEXT("/Script/Engine.EMaterialSamplerType"));
+			check(SamplerTypeEnum);
+		}
 	}
 }
 
diff --git a/Engine/Source/Runtime/Landscape/Classes/LandscapeComponent.h b/Engine/Source/Runtime/Landscape/Classes/LandscapeComponent.h
--- a/Engine/Source/Runtime/Landscape/Classes/LandscapeComponent.h
+++ b/Engine/Source/Runtime/Landscape/Classes/LandscapeComponent.h
@@ -263,9 +263,10 @@ public:
 	/** Pointer to data shared with the render thread, used by the editor tools */
 	struct FLandscapeEditToolRenderData* EditToolRenderData;
 
-	/** Hash of source for ES2 generated data. Used for mobile preview to determine if we need to re-generate ES2 pixel data. */
+	/** Hash of source for ES2 generated data. Used for mobile preview and cook-in-editor
+	 * to determine if we need to re-generate ES2 pixel data. */
 	UPROPERTY(Transient, DuplicateTransient)
-	FGuid MobilePixelDataSourceHash;
+	FGuid MobileDataSourceHash;
 #endif
 
 	/** For ES2 */
@@ -297,6 +298,7 @@ public:
 	static void AddReferencedObjects(UObject* InThis, FReferenceCollector& Collector);
 	virtual void PostDuplicate(bool bDuplicateForPIE) override;
 #if WITH_EDITOR
+	virtual void BeginCacheForCookedPlatformData(const ITargetPlatform* TargetPlatform) override;
 	virtual void PostLoad() override;
 	virtual void PostEditUndo() override;
 	virtual void PreEditChange(UProperty* PropertyThatWillChange) override;
@@ -347,7 +349,7 @@ public:
 	void ReplaceLayer(ULandscapeLayerInfoObject* FromLayerInfo, ULandscapeLayerInfoObject* ToLayerInfo, struct FLandscapeEditDataInterface* LandscapeEdit);
 	
 	void GeneratePlatformVertexData();
-	void GeneratePlatformPixelData(bool bIsCooking);
+	void GeneratePlatformPixelData();
 
 	/** Creates and destroys cooked grass data stored in the map */
 	void RenderGrassMap();
@@ -364,6 +366,9 @@ public:
 
 	/* Serialize all hashes/guids that record the current state of this component */
 	void SerializeStateHashes(FArchive& Ar);
+
+	/** Generate mobile data if it's missing or outdated */
+	void CheckGenerateLandscapePlatformData(bool bIsCooking);
 #endif
 
 	/** @todo document */
diff --git a/Engine/Source/Runtime/Landscape/Private/Landscape.cpp b/Engine/Source/Runtime/Landscape/Private/Landscape.cpp
--- a/Engine/Source/Runtime/Landscape/Private/Landscape.cpp
+++ b/Engine/Source/Runtime/Landscape/Private/Landscape.cpp
@@ -130,23 +130,70 @@ void ULandscapeComponent::AddReferencedObjects(UObject* InThis, FReferenceCollec
 	Super::AddReferencedObjects(This, Collector);
 }
 
-void ULandscapeComponent::Serialize(FArchive& Ar)
+#if WITH_EDITOR
+
+void ULandscapeComponent::BeginCacheForCookedPlatformData(const ITargetPlatform* TargetPlatform)
 {
-#if ENABLE_LANDSCAPE_COOKING && WITH_EDITOR
-	// Saving for cooking path
-	if (Ar.IsCooking() && Ar.IsSaving() && !HasAnyFlags(RF_ClassDefaultObject))
+	if (!TargetPlatform->SupportsFeature(ETargetPlatformFeatures::VertexShaderTextureSampling))
+	{
+		CheckGenerateLandscapePlatformData(true);
+	}
+}
+
+void ULandscapeComponent::CheckGenerateLandscapePlatformData(bool bIsCooking)
+{
+#if ENABLE_LANDSCAPE_COOKING
+	// Calculate hash of source data and skip generation if the data we have in memory is unchanged
+	FBufferArchive ComponentStateAr;
+	SerializeStateHashes(ComponentStateAr);
+	uint32 Hash[5];
+	FSHA1::HashBuffer(ComponentStateAr.GetData(), ComponentStateAr.Num(), (uint8*)Hash);
+
+	FGuid NewSourceHash = FGuid(Hash[0] ^ Hash[4], Hash[1], Hash[2], Hash[3]);
+
+	bool bGenerateVertexData = true;
+	bool bGeneratePixelData = true;
+
+	// Skip generation if the source hash matches
+	if (MobileDataSourceHash.IsValid() && MobileDataSourceHash == NewSourceHash)
 	{
-		if (!Ar.CookingTarget()->SupportsFeature(ETargetPlatformFeatures::VertexShaderTextureSampling))
+		if (MobileMaterialInterface != nullptr && MobileWeightNormalmapTexture != nullptr)
 		{
-			if (!PlatformData.HasValidPlatformData())
-			{
-				GeneratePlatformVertexData();
-				GeneratePlatformPixelData(true);
-			}
+			bGeneratePixelData = false;
+		}
+
+		if (PlatformData.HasValidPlatformData())
+		{
+			bGenerateVertexData = false;
+		}
+		else
+		if (PlatformData.LoadFromDDC(NewSourceHash))
+		{
+			bGenerateVertexData = false;
 		}
 	}
+
+	if (bGenerateVertexData)
+	{
+		GeneratePlatformVertexData();
+		if (bIsCooking)
+		{
+			PlatformData.SaveToDDC(NewSourceHash);
+		}
+	}
+
+	if (bGeneratePixelData)
+	{
+		GeneratePlatformPixelData();
+	}
+
+	MobileDataSourceHash = NewSourceHash;
+#endif
+}
 #endif
 
+void ULandscapeComponent::Serialize(FArchive& Ar)
+{
 	if (Ar.IsCooking() && !Ar.CookingTarget()->SupportsFeature(ETargetPlatformFeatures::VertexShaderTextureSampling))
 	{
 		// These properties are not used for ES2 so we back them up and clear them before serializing them.
@@ -235,7 +282,7 @@ void ULandscapeComponent::Serialize(FArchive& Ar)
 		// Saving for cooking path
 		if (Ar.IsCooking() && !Ar.CookingTarget()->SupportsFeature(ETargetPlatformFeatures::VertexShaderTextureSampling))
 		{
-			bool bValid = PlatformData.HasValidPlatformData();
+			bool bValid = ensure(PlatformData.HasValidPlatformData());
 			Ar << bValid;
 			if (bValid)
 			{
@@ -784,22 +831,9 @@ FPrimitiveSceneProxy* ULandscapeComponent::CreateSceneProxy()
 	else // i.e. (FeatureLevel <= ERHIFeatureLevel::ES3_1)
 	{
 #if WITH_EDITOR
-		// We need to cook platform data for ES2 preview in editor
+		// See if we need to cook platform data for ES2 preview in editor
 		// We can't always pre-cook it in PostLoad/Serialize etc because landscape can be edited
-		if (!PlatformData.HasValidPlatformData())
-		{
-			// Try to reload the ES2 landscape data from the DDC
-			if (!PlatformData.LoadFromDDC(StateId))
-			{
-				// Generate and save to the derived data cache
-				GeneratePlatformVertexData();
-
-				if (PlatformData.HasValidPlatformData())
-				{
-					PlatformData.SaveToDDC(StateId);
-				}
-			}
-		}
+		CheckGenerateLandscapePlatformData(false);
 
 		if (PlatformData.HasValidPlatformData())
 		{
@@ -2140,8 +2174,11 @@ void FLandscapeComponentDerivedData::GetUncompressedData(TArray<uint8>& OutUncom
 
 	verify(FCompression::UncompressMemory((ECompressionFlags)COMPRESS_ZLIB, OutUncompressedData.GetData(), UncompressedSize, CompressedData.GetData(), CompressedSize));
 
-	// free the compressed data now that we have the uncompressed version
-	CompressedLandscapeData.Empty();
+	if (FPlatformProperties::RequiresCookedData())
+	{
+		// free the compressed data now that we have the uncompressed version if running with cooked data
+		CompressedLandscapeData.Empty();
+	}
 }
 
 FArchive& operator<<(FArchive& Ar, FLandscapeComponentDerivedData& Data)
diff --git a/Engine/Source/Runtime/Landscape/Private/LandscapeEdit.cpp b/Engine/Source/Runtime/Landscape/Private/LandscapeEdit.cpp
--- a/Engine/Source/Runtime/Landscape/Private/LandscapeEdit.cpp
+++ b/Engine/Source/Runtime/Landscape/Private/LandscapeEdit.cpp
@@ -4434,32 +4434,10 @@ struct FMobileLayerAllocation
 	}
 };
 
-void ULandscapeComponent::GeneratePlatformPixelData(bool bIsCooking)
+void ULandscapeComponent::GeneratePlatformPixelData()
 {
 	check(!IsTemplate())
 
-	if (!bIsCooking)
-	{
-		// Calculate hash of source data and skip generation if the data we have in memory is unchanged
-		FBufferArchive ComponentStateAr;
-		SerializeStateHashes(ComponentStateAr);
-
-		uint32 Hash[5];
-		FSHA1::HashBuffer(ComponentStateAr.GetData(), ComponentStateAr.Num(), (uint8*)Hash);
-		FGuid NewSourceHash = FGuid(Hash[0] ^ Hash[4], Hash[1], Hash[2], Hash[3]);
-	
-		// Skip generation if the source hash matches
-		if (MobilePixelDataSourceHash.IsValid() && 
-			MobilePixelDataSourceHash == NewSourceHash &&
-			MobileMaterialInterface != nullptr &&
-			MobileWeightNormalmapTexture != nullptr)
-		{
-			return;
-		}
-			
-		MobilePixelDataSourceHash = NewSourceHash;
-	}
-
 	TArray<FMobileLayerAllocation> MobileLayerAllocations;
 	MobileLayerAllocations.Reserve(WeightmapLayerAllocations.Num());
 	for (const auto& Allocation : WeightmapLayerAllocations)
@@ -4537,7 +4515,6 @@ void ULandscapeComponent::GeneratePlatformPixelData(bool bIsCooking)
 
 	MobileWeightNormalmapTexture = NewWeightNormalmapTexture;
 
-
 	FLinearColor Masks[5];
 	Masks[0] = FLinearColor(1, 0, 0, 0);
 	Masks[1] = FLinearColor(0, 1, 0, 0);
@@ -4545,54 +4522,30 @@ void ULandscapeComponent::GeneratePlatformPixelData(bool bIsCooking)
 	Masks[3] = FLinearColor(0, 0, 0, 1);
 	Masks[4] = FLinearColor(0, 0, 0, 0); // mask out layers 4+ altogether
 
-	if (!bIsCooking)
-	{
-		UMaterialInstanceDynamic* NewMobileMaterialInstance = UMaterialInstanceDynamic::Create(MaterialInstance, GetOutermost());
-
-		MobileBlendableLayerMask = 0;
-
-		// Set the layer mask
-		int32 CurrentIdx = 0;
-		for (const auto& MobileAllocation : MobileLayerAllocations)
-		{
-			const FWeightmapLayerAllocationInfo& Allocation = MobileAllocation.Allocation;
-			if (Allocation.LayerInfo)
-			{
-				FName LayerName = Allocation.LayerInfo == ALandscapeProxy::VisibilityLayer ? UMaterialExpressionLandscapeVisibilityMask::ParameterName : Allocation.LayerInfo->LayerName;
-				NewMobileMaterialInstance->SetVectorParameterValue(FName(*FString::Printf(TEXT("LayerMask_%s"), *LayerName.ToString())), Masks[FMath::Min(4, CurrentIdx)]);
-				MobileBlendableLayerMask |= (!Allocation.LayerInfo->bNoWeightBlend ? (1 << CurrentIdx) : 0);
-				CurrentIdx++;
-			}
-		}
-		MobileMaterialInterface = NewMobileMaterialInstance;
-	}
-	else // for cooking
-	{
-		UMaterialInstanceConstant* CombinationMaterialInstance = GetCombinationMaterial(true);
-		UMaterialInstanceConstant* NewMobileMaterialInstance = NewObject<ULandscapeMaterialInstanceConstant>(GetOutermost());
+	UMaterialInstanceConstant* CombinationMaterialInstance = GetCombinationMaterial(true);
+	UMaterialInstanceConstant* NewMobileMaterialInstance = NewObject<ULandscapeMaterialInstanceConstant>(GetOutermost());
 
-		NewMobileMaterialInstance->SetParentEditorOnly(CombinationMaterialInstance);
+	NewMobileMaterialInstance->SetParentEditorOnly(CombinationMaterialInstance);
 
-		MobileBlendableLayerMask = 0;
+	MobileBlendableLayerMask = 0;
 
-		// Set the layer mask
-		int32 CurrentIdx = 0;
-		for (const auto& MobileAllocation : MobileLayerAllocations)
+	// Set the layer mask
+	int32 CurrentIdx = 0;
+	for (const auto& MobileAllocation : MobileLayerAllocations)
+	{
+		const FWeightmapLayerAllocationInfo& Allocation = MobileAllocation.Allocation;
+		if (Allocation.LayerInfo)
 		{
-			const FWeightmapLayerAllocationInfo& Allocation = MobileAllocation.Allocation;
-			if (Allocation.LayerInfo)
-			{
-				FName LayerName = Allocation.LayerInfo == ALandscapeProxy::VisibilityLayer ? UMaterialExpressionLandscapeVisibilityMask::ParameterName : Allocation.LayerInfo->LayerName;
-				NewMobileMaterialInstance->SetVectorParameterValueEditorOnly(FName(*FString::Printf(TEXT("LayerMask_%s"), *LayerName.ToString())), Masks[FMath::Min(4, CurrentIdx)]);
-				MobileBlendableLayerMask |= (!Allocation.LayerInfo->bNoWeightBlend ? (1 << CurrentIdx) : 0);
-				CurrentIdx++;
-			}
+			FName LayerName = Allocation.LayerInfo == ALandscapeProxy::VisibilityLayer ? UMaterialExpressionLandscapeVisibilityMask::ParameterName : Allocation.LayerInfo->LayerName;
+			NewMobileMaterialInstance->SetVectorParameterValueEditorOnly(FName(*FString::Printf(TEXT("LayerMask_%s"), *LayerName.ToString())), Masks[FMath::Min(4, CurrentIdx)]);
+			MobileBlendableLayerMask |= (!Allocation.LayerInfo->bNoWeightBlend ? (1 << CurrentIdx) : 0);
+			CurrentIdx++;
 		}
+	}
 
-		NewMobileMaterialInstance->PostEditChange();
+	NewMobileMaterialInstance->PostEditChange();
 
-		MobileMaterialInterface = NewMobileMaterialInstance;
-	}
+	MobileMaterialInterface = NewMobileMaterialInstance;
 }
 
 //
@@ -4620,13 +4573,16 @@ void ULandscapeComponent::GeneratePlatformVertexData()
 	NewPlatformData.AddZeroed(NewPlatformDataSize);
 
 	// Get the required mip data
+	TArray<TArray<uint8>> HeightmapMipRawData;
 	TArray<FColor*> HeightmapMipData;
 	for (int32 MipIdx = 0; MipIdx < FMath::Min(LANDSCAPE_MAX_ES_LOD, HeightmapTexture->Source.GetNumMips()); MipIdx++)
 	{
 		int32 MipSubsectionSizeVerts = (SubsectionSizeVerts) >> MipIdx;
 		if (MipSubsectionSizeVerts > 1)
 		{
-			HeightmapMipData.Add((FColor*)HeightmapTexture->Source.LockMip(MipIdx));
+			new(HeightmapMipRawData) TArray<uint8>();
+			HeightmapTexture->Source.GetMipData(HeightmapMipRawData.Last(), MipIdx);
+			HeightmapMipData.Add((FColor*)HeightmapMipRawData.Last().GetData());
 		}
 	}
 
@@ -4751,11 +4707,6 @@ void ULandscapeComponent::GeneratePlatformVertexData()
 		DstVert++;
 	}
 
-	for (int32 MipIdx = 0; MipIdx < HeightmapTexture->Source.GetNumMips(); MipIdx++)
-	{
-		HeightmapTexture->Source.UnlockMip(MipIdx);
-	}
-
 	// Copy to PlatformData as Compressed
 	PlatformData.InitializeFromUncompressedData(NewPlatformData);
 }
diff --git a/Engine/Source/Runtime/Landscape/Private/LandscapeRenderMobile.cpp b/Engine/Source/Runtime/Landscape/Private/LandscapeRenderMobile.cpp
--- a/Engine/Source/Runtime/Landscape/Private/LandscapeRenderMobile.cpp
+++ b/Engine/Source/Runtime/Landscape/Private/LandscapeRenderMobile.cpp
@@ -220,9 +220,6 @@ FLandscapeComponentSceneProxyMobile::FLandscapeComponentSceneProxyMobile(ULandsc
 	check(InComponent && InComponent->PlatformData.HasValidPlatformData());
 	InComponent->PlatformData.GetUncompressedData(PlatformData);
 
-#if WITH_EDITOR
-	InComponent->GeneratePlatformPixelData(false);
-#endif
 	check(InComponent->MobileMaterialInterface);
 	check(InComponent->MobileWeightNormalmapTexture);
 	
