Fix Pri-1 crash generating landscape data during Mobile Preview (UE-18747)

diff --git a/Engine/Source/Runtime/Landscape/Private/LandscapeEdit.cpp b/Engine/Source/Runtime/Landscape/Private/LandscapeEdit.cpp
--- a/Engine/Source/Runtime/Landscape/Private/LandscapeEdit.cpp
+++ b/Engine/Source/Runtime/Landscape/Private/LandscapeEdit.cpp
@@ -4522,30 +4522,59 @@ void ULandscapeComponent::GeneratePlatformPixelData()
 	Masks[3] = FLinearColor(0, 0, 0, 1);
 	Masks[4] = FLinearColor(0, 0, 0, 0); // mask out layers 4+ altogether
 
-	UMaterialInstanceConstant* CombinationMaterialInstance = GetCombinationMaterial(true);
-	UMaterialInstanceConstant* NewMobileMaterialInstance = NewObject<ULandscapeMaterialInstanceConstant>(GetOutermost());
-
-	NewMobileMaterialInstance->SetParentEditorOnly(CombinationMaterialInstance);
+	if (!GIsEditor)
+	{
+		// This path is used by game mode running with uncooked data, eg Mobile Preview.
+		// Game mode cannot create MICs, so we use a MaterialInstanceDynamic here.
+		UMaterialInstanceDynamic* NewMobileMaterialInstance = UMaterialInstanceDynamic::Create(MaterialInstance, GetOutermost());
 
-	MobileBlendableLayerMask = 0;
+		MobileBlendableLayerMask = 0;
 
-	// Set the layer mask
-	int32 CurrentIdx = 0;
-	for (const auto& MobileAllocation : MobileLayerAllocations)
-	{
-		const FWeightmapLayerAllocationInfo& Allocation = MobileAllocation.Allocation;
-		if (Allocation.LayerInfo)
+		// Set the layer mask
+		int32 CurrentIdx = 0;
+		for (const auto& MobileAllocation : MobileLayerAllocations)
 		{
-			FName LayerName = Allocation.LayerInfo == ALandscapeProxy::VisibilityLayer ? UMaterialExpressionLandscapeVisibilityMask::ParameterName : Allocation.LayerInfo->LayerName;
-			NewMobileMaterialInstance->SetVectorParameterValueEditorOnly(FName(*FString::Printf(TEXT("LayerMask_%s"), *LayerName.ToString())), Masks[FMath::Min(4, CurrentIdx)]);
-			MobileBlendableLayerMask |= (!Allocation.LayerInfo->bNoWeightBlend ? (1 << CurrentIdx) : 0);
-			CurrentIdx++;
+			const FWeightmapLayerAllocationInfo& Allocation = MobileAllocation.Allocation;
+			if (Allocation.LayerInfo)
+			{
+				FName LayerName = Allocation.LayerInfo == ALandscapeProxy::VisibilityLayer ? UMaterialExpressionLandscapeVisibilityMask::ParameterName : Allocation.LayerInfo->LayerName;
+				NewMobileMaterialInstance->SetVectorParameterValue(FName(*FString::Printf(TEXT("LayerMask_%s"), *LayerName.ToString())), Masks[FMath::Min(4, CurrentIdx)]);
+				MobileBlendableLayerMask |= (!Allocation.LayerInfo->bNoWeightBlend ? (1 << CurrentIdx) : 0);
+				CurrentIdx++;
+			}
 		}
+		MobileMaterialInterface = NewMobileMaterialInstance;
 	}
+	else
+	{
+		// When cooking, we need to make a persistent MIC. In the editor we also do so in
+		// case we start a Cook in Editor operation, which will reuse the MIC we create now.
+
+		UMaterialInstanceConstant* CombinationMaterialInstance = GetCombinationMaterial(true);
+		UMaterialInstanceConstant* NewMobileMaterialInstance = NewObject<ULandscapeMaterialInstanceConstant>(GetOutermost());
 
-	NewMobileMaterialInstance->PostEditChange();
+		NewMobileMaterialInstance->SetParentEditorOnly(CombinationMaterialInstance);
 
-	MobileMaterialInterface = NewMobileMaterialInstance;
+		MobileBlendableLayerMask = 0;
+
+		// Set the layer mask
+		int32 CurrentIdx = 0;
+		for (const auto& MobileAllocation : MobileLayerAllocations)
+		{
+			const FWeightmapLayerAllocationInfo& Allocation = MobileAllocation.Allocation;
+			if (Allocation.LayerInfo)
+			{
+				FName LayerName = Allocation.LayerInfo == ALandscapeProxy::VisibilityLayer ? UMaterialExpressionLandscapeVisibilityMask::ParameterName : Allocation.LayerInfo->LayerName;
+				NewMobileMaterialInstance->SetVectorParameterValueEditorOnly(FName(*FString::Printf(TEXT("LayerMask_%s"), *LayerName.ToString())), Masks[FMath::Min(4, CurrentIdx)]);
+				MobileBlendableLayerMask |= (!Allocation.LayerInfo->bNoWeightBlend ? (1 << CurrentIdx) : 0);
+				CurrentIdx++;
+			}
+		}
+
+		NewMobileMaterialInstance->PostEditChange();
+
+		MobileMaterialInterface = NewMobileMaterialInstance;
+	}
 }
 
 //
