Fixed changing the mesh of a landscape spline control point not working

diff --git a/Engine/Source/Runtime/Landscape/Private/LandscapeSplines.cpp b/Engine/Source/Runtime/Landscape/Private/LandscapeSplines.cpp
--- a/Engine/Source/Runtime/Landscape/Private/LandscapeSplines.cpp
+++ b/Engine/Source/Runtime/Landscape/Private/LandscapeSplines.cpp
@@ -1453,6 +1453,8 @@ void ULandscapeSplineControlPoint::UpdateSplinePoints(bool bUpdateCollision, boo
 		if (MeshComponent->StaticMesh != Mesh)
 		{
 			MeshComponent->Modify();
+			MeshComponent->UnregisterComponent();
+			bComponentNeedsRegistering = true;
 			MeshComponent->SetStaticMesh(Mesh);
 
 			AutoSetConnections(false);
