Fixed vertical mesh offset on landscape splines being width-dependent and not blendable between different heights.
There is now a "Mesh Vertical Offset" on spline *control points* (to allow it to blend between different values along a spline). The old "Mesh Offest" on the segment mesh entry has been renamed "Center Adjust", as that's what it actually did and that behaviour is still useful.

diff --git a/Engine/Config/BaseEngine.ini b/Engine/Config/BaseEngine.ini
--- a/Engine/Config/BaseEngine.ini
+++ b/Engine/Config/BaseEngine.ini
@@ -1269,9 +1269,13 @@ bUseStreamingPause=false
 ;Exposing SplineMeshComponents to Blueprint
 +TaggedPropertyRedirects=(ClassName="SplineMeshComponent",OldPropertyName="SplineXDir",NewPropertyName="SplineUpDir")
 
+;Landscape Spline Mesh Streaming
 +TaggedPropertyRedirects=(ClassName="LandscapeSplineControlPoint",OldPropertyName="MeshComponent",NewPropertyName="LocalMeshComponent")
 +TaggedPropertyRedirects=(ClassName="LandscapeSplineSegment",OldPropertyName="MeshComponents",NewPropertyName="LocalMeshComponents")
 
+; Renaming to avoid confusion
++TaggedPropertyRedirects=(ClassName="LandscapeSplineMeshEntry", OldPropertyName="Offset", NewPropertyName="CenterAdjust")
+
 ;Foliage refactor
 +ActiveClassRedirects=(OldClassName="InstancedFoliageSettings", NewClassName="/Script/Foliage.FoliageType_InstancedStaticMesh")
 +ActiveClassRedirects=(OldClassName="FoliageType", NewClassName="/Script/Foliage.FoliageType")
diff --git a/Engine/Source/Runtime/Landscape/Classes/LandscapeSplineControlPoint.h b/Engine/Source/Runtime/Landscape/Classes/LandscapeSplineControlPoint.h
--- a/Engine/Source/Runtime/Landscape/Classes/LandscapeSplineControlPoint.h
+++ b/Engine/Source/Runtime/Landscape/Classes/LandscapeSplineControlPoint.h
@@ -49,26 +49,45 @@ class ULandscapeSplineControlPoint : public UObject
 
 // Directly editable data:
 	/** Location in Landscape-space */
-	UPROPERTY(EditAnywhere, Category=LandscapeSplineControlPoint)
+	UPROPERTY(EditAnywhere, Category=LandscapeSpline)
 	FVector Location;
 
 	/** Rotation of tangent vector at this point (in landscape-space) */
-	UPROPERTY(EditAnywhere, Category=LandscapeSplineControlPoint)
+	UPROPERTY(EditAnywhere, Category=LandscapeSpline)
 	FRotator Rotation;
 
 	/** Width of the spline at this point. */
-	UPROPERTY(EditAnywhere, Category=LandscapeSplineControlPoint)
+	UPROPERTY(EditAnywhere, Category=LandscapeSpline)
 	float Width;
 
 	/** Falloff at the sides of the spline at this point. */
-	UPROPERTY(EditAnywhere, Category=LandscapeSplineControlPoint)
+	UPROPERTY(EditAnywhere, Category=LandscapeSpline)
 	float SideFalloff;
 
 	/** Falloff at the start/end of the spline (if this point is a start or end point, otherwise ignored). */
-	UPROPERTY(EditAnywhere, Category=LandscapeSplineControlPoint)
+	UPROPERTY(EditAnywhere, Category=LandscapeSpline)
 	float EndFalloff;
 
 #if WITH_EDITORONLY_DATA
+	/** Vertical offset of the spline segment mesh. Useful for a river's surface, among other things. */
+	UPROPERTY(EditAnywhere, Category=LandscapeSpline, meta=(DisplayName="Mesh Vertical Offset"))
+	float SegmentMeshOffset;
+
+	/**
+	 * Name of blend layer to paint when applying spline to landscape
+	 * If "none", no layer is painted
+	 */
+	UPROPERTY(EditAnywhere, Category=LandscapeDeformation)
+	FName LayerName;
+
+	/** If the spline is above the terrain, whether to raise the terrain up to the level of the spline when applying it to the landscape. */
+	UPROPERTY(EditAnywhere, Category=LandscapeDeformation)
+	uint32 bRaiseTerrain:1;
+
+	/** If the spline is below the terrain, whether to lower the terrain down to the level of the spline when applying it to the landscape. */
+	UPROPERTY(EditAnywhere, Category=LandscapeDeformation)
+	uint32 bLowerTerrain:1;
+
 	/** Mesh to use on the control point */
 	UPROPERTY(EditAnywhere, Category=Mesh)
 	UStaticMesh* Mesh;
@@ -81,8 +100,16 @@ class ULandscapeSplineControlPoint : public UObject
 	UPROPERTY(EditAnywhere, Category=Mesh)
 	FVector MeshScale;
 
+	/** Whether to enable collision for the Control Point Mesh. */
+	UPROPERTY(EditAnywhere, Category=Mesh)
+	uint32 bEnableCollision:1;
+
+	/** Whether the Control Point Mesh should cast a shadow. */
+	UPROPERTY(EditAnywhere, Category=Mesh)
+	uint32 bCastShadow:1;
+
 	/**  Max draw distance for the mesh used on this control point */
-	UPROPERTY(EditAnywhere, AdvancedDisplay, Category=Mesh, meta = (DisplayName = "Max Draw Distance"))
+	UPROPERTY(EditAnywhere, Category=Mesh, AdvancedDisplay, meta=(DisplayName="Max Draw Distance"))
 	float LDMaxDrawDistance;
 
 	/**
@@ -92,36 +119,13 @@ class ULandscapeSplineControlPoint : public UObject
 	 * Ignored if the object is not translucent.  The default priority is zero.
 	 * Warning: This should never be set to a non-default value unless you know what you are doing, as it will prevent the renderer from sorting correctly.
 	 */
-	UPROPERTY(EditAnywhere, AdvancedDisplay, Category=Mesh)
+	UPROPERTY(EditAnywhere, Category=Mesh, AdvancedDisplay)
 	int32 TranslucencySortPriority;
 
-	/**
-	 * Name of blend layer to paint when applying spline to landscape
-	 * If "none", no layer is painted
-	 */
-	UPROPERTY(EditAnywhere, Category=LandscapeSplineControlPoint)
-	FName LayerName;
-
-	/** If the spline is above the terrain, whether to raise the terrain up to the level of the spline when applying it to the landscape. */
-	UPROPERTY(EditAnywhere, Category=LandscapeSplineControlPoint)
-	uint32 bRaiseTerrain:1;
-
-	/** If the spline is below the terrain, whether to lower the terrain down to the level of the spline when applying it to the landscape. */
-	UPROPERTY(EditAnywhere, Category=LandscapeSplineControlPoint)
-	uint32 bLowerTerrain:1;
-
 	/** Whether control point mesh should be placed in landscape proxy streaming level (true) or the spline's level (false) */
-	UPROPERTY(EditAnywhere, Category=Mesh)
+	UPROPERTY(EditAnywhere, Category=Mesh, AdvancedDisplay)
 	uint32 bPlaceSplineMeshesInStreamingLevels : 1;
 
-	/** Whether to enable collision for the Control Point Mesh. */
-	UPROPERTY(EditAnywhere, Category=Mesh)
-	uint32 bEnableCollision:1;
-
-	/** Whether the Control Point Mesh should cast a shadow. */
-	UPROPERTY(EditAnywhere, Category=Mesh)
-	uint32 bCastShadow:1;
-
 protected:
 	UPROPERTY(Transient)
 	uint32 bSelected : 1;
diff --git a/Engine/Source/Runtime/Landscape/Classes/LandscapeSplineSegment.h b/Engine/Source/Runtime/Landscape/Classes/LandscapeSplineSegment.h
--- a/Engine/Source/Runtime/Landscape/Classes/LandscapeSplineSegment.h
+++ b/Engine/Source/Runtime/Landscape/Classes/LandscapeSplineSegment.h
@@ -1,9 +1,13 @@
 // Copyright 1998-2015 Epic Games, Inc. All Rights Reserved.
 
 #pragma once
+
 #include "Components/SplineMeshComponent.h"
 #include "LandscapeSplineSegment.generated.h"
 
+//Forward declarations
+class ULandscapeSplineControlPoint;
+
 USTRUCT()
 struct FLandscapeSplineInterpPoint
 {
@@ -55,7 +59,7 @@ struct FLandscapeSplineSegmentConnection
 
 	// Control point connected to this end of the segment
 	UPROPERTY()
-	class ULandscapeSplineControlPoint* ControlPoint;
+	ULandscapeSplineControlPoint* ControlPoint;
 
 	// Tangent length of the connection
 	UPROPERTY(EditAnywhere, Category=LandscapeSplineSegmentConnection)
@@ -66,7 +70,7 @@ struct FLandscapeSplineSegmentConnection
 	FName SocketName;
 
 	FLandscapeSplineSegmentConnection()
-		: ControlPoint(NULL)
+		: ControlPoint(nullptr)
 		, TangentLen(0.0f)
 		, SocketName(NAME_None)
 	{
@@ -89,19 +93,19 @@ struct FLandscapeSplineMeshEntry
 
 	/** Mesh to use on the spline */
 	UPROPERTY(EditAnywhere, Category=LandscapeSplineMeshEntry)
-	class UStaticMesh* Mesh;
+	UStaticMesh* Mesh;
 
 	/** Overrides mesh's materials */
-	UPROPERTY(EditAnywhere, Category=LandscapeSplineMeshEntry)
-	TArray<class UMaterialInterface*> MaterialOverrides;
+	UPROPERTY(EditAnywhere, Category=LandscapeSplineMeshEntry, AdvancedDisplay)
+	TArray<UMaterialInterface*> MaterialOverrides;
 
-	/** Whether to center the mesh horizontally on the spline */
-	UPROPERTY(EditAnywhere, Category=LandscapeSplineMeshEntry)
+	/** Whether to automatically center the mesh horizontally on the spline */
+	UPROPERTY(EditAnywhere, Category=LandscapeSplineMeshEntry, meta=(DisplayName="Center Horizontally"))
 	uint32 bCenterH:1;
 
-	/** X/Y offset of the mesh relative to the spline */
-	UPROPERTY(EditAnywhere, Category=LandscapeSplineMeshEntry)
-	FVector2D Offset;
+	/** Tweak to center the mesh correctly on the spline */
+	UPROPERTY(EditAnywhere, Category=LandscapeSplineMeshEntry, AdvancedDisplay, meta=(DisplayName="Center Adjust"))
+	FVector2D CenterAdjust;
 
 	/** Whether to scale the mesh to fit the width of the spline */
 	UPROPERTY(EditAnywhere, Category=LandscapeSplineMeshEntry)
@@ -113,7 +117,7 @@ struct FLandscapeSplineMeshEntry
 
 	/** Orientation of the spline mesh, X=Up or Y=Up */
 	UPROPERTY()
-	TEnumAsByte<enum LandscapeSplineMeshOrientation> Orientation_DEPRECATED;
+	TEnumAsByte<LandscapeSplineMeshOrientation> Orientation_DEPRECATED;
 
 	/** Chooses the forward axis for the spline mesh orientation */
 	UPROPERTY(EditAnywhere, Category=LandscapeSplineMeshEntry)
@@ -124,10 +128,10 @@ struct FLandscapeSplineMeshEntry
 	TEnumAsByte<ESplineMeshAxis::Type> UpAxis;
 
 	FLandscapeSplineMeshEntry() :
-		Mesh(NULL),
+		Mesh(nullptr),
 		MaterialOverrides(),
 		bCenterH(true),
-		Offset(0, 0),
+		CenterAdjust(0, 0),
 		bScaleToWidth(true),
 		Scale(1,1,1),
 		Orientation_DEPRECATED(LSMO_YUp),
@@ -154,19 +158,35 @@ class ULandscapeSplineSegment : public UObject
 	 * Name of blend layer to paint when applying spline to landscape
 	 * If "none", no layer is painted
 	 */
-	UPROPERTY(EditAnywhere, Category=LandscapeSplineSegment)
+	UPROPERTY(EditAnywhere, Category=LandscapeDeformation)
 	FName LayerName;
 
+	/** If the spline is above the terrain, whether to raise the terrain up to the level of the spline when applying it to the landscape. */
+	UPROPERTY(EditAnywhere, Category=LandscapeDeformation)
+	uint32 bRaiseTerrain:1;
+
+	/** If the spline is below the terrain, whether to lower the terrain down to the level of the spline when applying it to the landscape. */
+	UPROPERTY(EditAnywhere, Category=LandscapeDeformation)
+	uint32 bLowerTerrain:1;
+
 	/** Spline meshes from this list are used in random order along the spline. */
 	UPROPERTY(EditAnywhere, Category=LandscapeSplineMeshes)
 	TArray<FLandscapeSplineMeshEntry> SplineMeshes;
 
-	/** Random seed used for choosing which spline meshes to use. */
+	/** Whether to generate collision for the Spline Meshes. */
+	UPROPERTY(EditAnywhere, Category=LandscapeSplineMeshes)
+	uint32 bEnableCollision:1;
+
+	/** Whether the Spline Meshes should cast a shadow. */
 	UPROPERTY(EditAnywhere, Category=LandscapeSplineMeshes)
+	uint32 bCastShadow:1;
+
+	/** Random seed used for choosing which order to use spline meshes. Ignored if only one mesh is set. */
+	UPROPERTY(EditAnywhere, Category=LandscapeSplineMeshes, AdvancedDisplay)
 	int32 RandomSeed;
 
 	/**  Max draw distance for all the mesh pieces used in this spline */
-	UPROPERTY(EditAnywhere, AdvancedDisplay, Category = LandscapeSplineMeshes, meta = (DisplayName = "Max Draw Distance"))
+	UPROPERTY(EditAnywhere, Category=LandscapeSplineMeshes, AdvancedDisplay, meta=(DisplayName="Max Draw Distance"))
 	float LDMaxDrawDistance;
 
 	/**
@@ -176,29 +196,13 @@ class ULandscapeSplineSegment : public UObject
 	 * Ignored if the object is not translucent.  The default priority is zero.
 	 * Warning: This should never be set to a non-default value unless you know what you are doing, as it will prevent the renderer from sorting correctly.
 	 */
-	UPROPERTY(EditAnywhere, AdvancedDisplay, Category=LandscapeSplineMeshes)
+	UPROPERTY(EditAnywhere, Category=LandscapeSplineMeshes, AdvancedDisplay)
 	int32 TranslucencySortPriority;
 
-	/** If the spline is above the terrain, whether to raise the terrain up to the level of the spline when applying it to the landscape. */
-	UPROPERTY(EditAnywhere, Category=LandscapeSplineSegment)
-	uint32 bRaiseTerrain:1;
-
-	/** If the spline is below the terrain, whether to lower the terrain down to the level of the spline when applying it to the landscape. */
-	UPROPERTY(EditAnywhere, Category=LandscapeSplineSegment)
-	uint32 bLowerTerrain:1;
-
 	/** Whether spline meshes should be placed in landscape proxy streaming levels (true) or the spline's level (false) */
-	UPROPERTY(EditAnywhere, Category = LandscapeSplineMeshes)
+	UPROPERTY(EditAnywhere, Category=LandscapeSplineMeshes, AdvancedDisplay)
 	uint32 bPlaceSplineMeshesInStreamingLevels : 1;
 
-	/** Whether to generate collision for the Spline Meshes. */
-	UPROPERTY(EditAnywhere, Category=LandscapeSplineMeshes)
-	uint32 bEnableCollision:1;
-
-	/** Whether the Spline Meshes should cast a shadow. */
-	UPROPERTY(EditAnywhere, Category=LandscapeSplineMeshes)
-	uint32 bCastShadow:1;
-
 protected:
 	UPROPERTY(Transient)
 	uint32 bSelected : 1;
diff --git a/Engine/Source/Runtime/Landscape/Private/LandscapeSplines.cpp b/Engine/Source/Runtime/Landscape/Private/LandscapeSplines.cpp
--- a/Engine/Source/Runtime/Landscape/Private/LandscapeSplines.cpp
+++ b/Engine/Source/Runtime/Landscape/Private/LandscapeSplines.cpp
@@ -1949,6 +1949,8 @@ void ULandscapeSplineSegment::UpdateSplinePoints(bool bUpdateCollision)
 	const float EndRollDegrees = EndRotation.Roll * (Connections[1].TangentLen > 0 ? -1 : 1);
 	const float StartRoll = FMath::DegreesToRadians(StartRollDegrees);
 	const float EndRoll = FMath::DegreesToRadians(EndRollDegrees);
+	const float StartMeshOffset = Connections[0].ControlPoint->SegmentMeshOffset;
+	const float EndMeshOffset = Connections[1].ControlPoint->SegmentMeshOffset;
 
 	int32 NumPoints = FMath::CeilToInt(SplineLength / OuterSplines->SplineResolution);
 	NumPoints = FMath::Clamp(NumPoints, 1, 1000);
@@ -1984,7 +1986,7 @@ void ULandscapeSplineSegment::UpdateSplinePoints(bool bUpdateCollision)
 		SplineEditorMeshEntry.Mesh = OuterSplines->SplineEditorMesh;
 		SplineEditorMeshEntry.MaterialOverrides = {};
 		SplineEditorMeshEntry.bCenterH = true;
-		SplineEditorMeshEntry.Offset = {0.0f, 0.5f};
+		SplineEditorMeshEntry.CenterAdjust = {0.0f, 0.5f};
 		SplineEditorMeshEntry.bScaleToWidth = true;
 		SplineEditorMeshEntry.Scale = {3, 1, 1};
 		SplineEditorMeshEntry.ForwardAxis = ESplineMeshAxis::X;
@@ -2184,6 +2186,7 @@ void ULandscapeSplineSegment::UpdateSplinePoints(bool bUpdateCollision)
 			                                (MeshEntry->ForwardAxis == ESplineMeshAxis::Y && MeshEntry->UpAxis == ESplineMeshAxis::Z) ||
 			                                (MeshEntry->ForwardAxis == ESplineMeshAxis::Z && MeshEntry->UpAxis == ESplineMeshAxis::X);
 			const float Roll = FMath::Lerp(StartRoll, EndRoll, CosInterp) + (bDoOrientationRoll ? -HALF_PI : 0);
+			const float MeshOffset = FMath::Lerp(StartMeshOffset, EndMeshOffset, CosInterp);
 
 			FVector Scale = MeshEntry->Scale;
 			if (MeshEntry->bScaleToWidth)
@@ -2191,7 +2194,7 @@ void ULandscapeSplineSegment::UpdateSplinePoints(bool bUpdateCollision)
 				Scale *= Width / USplineMeshComponent::GetAxisValue(MeshBounds.BoxExtent, SideAxis);
 			}
 
-			FVector2D Offset = MeshEntry->Offset;
+			FVector2D Offset = MeshEntry->CenterAdjust;
 			if (MeshEntry->bCenterH)
 			{
 				if (bDoOrientationRoll)
@@ -2221,6 +2224,7 @@ void ULandscapeSplineSegment::UpdateSplinePoints(bool bUpdateCollision)
 				break;
 			}
 			Offset *= Scale2D;
+			Offset.Y += MeshOffset;
 			Offset = Offset.GetRotated(-Roll);
 
 			MeshComponent->SplineParams.StartPos = SplineInfo.Eval(RescaledT, FVector::ZeroVector);
@@ -2232,6 +2236,7 @@ void ULandscapeSplineSegment::UpdateSplinePoints(bool bUpdateCollision)
 			const float CosInterpEnd = 0.5f - 0.5f * FMath::Cos(TEnd * PI);
 			const float WidthEnd = FMath::Lerp(StartWidth, EndWidth, CosInterpEnd);
 			const float RollEnd = FMath::Lerp(StartRoll, EndRoll, CosInterpEnd) + (bDoOrientationRoll ? -HALF_PI : 0);
+			const float MeshOffsetEnd = FMath::Lerp(StartMeshOffset, EndMeshOffset, CosInterpEnd);
 
 			FVector ScaleEnd = MeshEntry->Scale;
 			if (MeshEntry->bScaleToWidth)
@@ -2239,7 +2244,7 @@ void ULandscapeSplineSegment::UpdateSplinePoints(bool bUpdateCollision)
 				ScaleEnd *= WidthEnd / USplineMeshComponent::GetAxisValue(MeshBounds.BoxExtent, SideAxis);
 			}
 
-			FVector2D OffsetEnd = MeshEntry->Offset;
+			FVector2D OffsetEnd = MeshEntry->CenterAdjust;
 			if (MeshEntry->bCenterH)
 			{
 				if (bDoOrientationRoll)
@@ -2269,6 +2274,7 @@ void ULandscapeSplineSegment::UpdateSplinePoints(bool bUpdateCollision)
 				break;
 			}
 			OffsetEnd *= Scale2DEnd;
+			OffsetEnd.Y += MeshOffsetEnd;
 			OffsetEnd = OffsetEnd.GetRotated(-RollEnd);
 
 			MeshComponent->SplineParams.EndPos = SplineInfo.Eval(TEnd, FVector::ZeroVector);
