Fix crash pressing "+" element in landscape grass type asset editor (UE-19474)

diff --git a/Engine/Source/Runtime/Landscape/Private/LandscapeEdit.cpp b/Engine/Source/Runtime/Landscape/Private/LandscapeEdit.cpp
--- a/Engine/Source/Runtime/Landscape/Private/LandscapeEdit.cpp
+++ b/Engine/Source/Runtime/Landscape/Private/LandscapeEdit.cpp
@@ -2976,19 +2976,21 @@ void ULandscapeLayerInfoObject::PostEditChangeProperty(FPropertyChangedEvent& Pr
 		}
 		else if (PropertyName == NAME_PhysMaterial)
 		{
-			// Only care current world object
-			for (TActorIterator<ALandscapeProxy> It(GetWorld()); It; ++It)
+			for (TObjectIterator<ALandscapeProxy> It; It; ++It)
 			{
 				ALandscapeProxy* Proxy = *It;
-				ULandscapeInfo* Info = Proxy->GetLandscapeInfo(false);
-				if (Info)
+				if (Proxy->GetWorld() && !Proxy->GetWorld()->IsPlayInEditor())
 				{
-					for (int32 i = 0; i < Info->Layers.Num(); ++i)
+					ULandscapeInfo* Info = Proxy->GetLandscapeInfo(false);
+					if (Info)
 					{
-						if (Info->Layers[i].LayerInfoObj == this)
+						for (int32 i = 0; i < Info->Layers.Num(); ++i)
 						{
-							Proxy->ChangedPhysMaterial();
-							break;
+							if (Info->Layers[i].LayerInfoObj == this)
+							{
+								Proxy->ChangedPhysMaterial();
+								break;
+							}
 						}
 					}
 				}
diff --git a/Engine/Source/Runtime/Landscape/Private/LandscapeGrass.cpp b/Engine/Source/Runtime/Landscape/Private/LandscapeGrass.cpp
--- a/Engine/Source/Runtime/Landscape/Private/LandscapeGrass.cpp
+++ b/Engine/Source/Runtime/Landscape/Private/LandscapeGrass.cpp
@@ -777,25 +777,27 @@ void ULandscapeGrassType::PostEditChangeProperty(FPropertyChangedEvent& Property
 
 	if (GIsEditor)
 	{
-		// Only care current world object
-		for (TActorIterator<ALandscapeProxy> It(GetWorld()); It; ++It)
+		for (TObjectIterator<ALandscapeProxy> It; It; ++It)
 		{
 			ALandscapeProxy* Proxy = *It;
-			const UMaterialInterface* MaterialInterface = Proxy->LandscapeMaterial;
-			if (MaterialInterface)
+			if (Proxy->GetWorld() && !Proxy->GetWorld()->IsPlayInEditor())
 			{
-				TArray<const UMaterialExpressionLandscapeGrassOutput*> GrassExpressions;
-				MaterialInterface->GetMaterial()->GetAllExpressionsOfType<UMaterialExpressionLandscapeGrassOutput>(GrassExpressions);
-
-				// Should only be one grass type node
-				if (GrassExpressions.Num() > 0)
+				const UMaterialInterface* MaterialInterface = Proxy->LandscapeMaterial;
+				if (MaterialInterface)
 				{
-					for (auto& Output : GrassExpressions[0]->GrassTypes)
+					TArray<const UMaterialExpressionLandscapeGrassOutput*> GrassExpressions;
+					MaterialInterface->GetMaterial()->GetAllExpressionsOfType<UMaterialExpressionLandscapeGrassOutput>(GrassExpressions);
+
+					// Should only be one grass type node
+					if (GrassExpressions.Num() > 0)
 					{
-						if (Output.GrassType == this)
+						for (auto& Output : GrassExpressions[0]->GrassTypes)
 						{
-							Proxy->FlushGrassComponents();
-							break;
+							if (Output.GrassType == this)
+							{
+								Proxy->FlushGrassComponents();
+								break;
+							}
 						}
 					}
 				}
