Fixed static objects rendering as fully lit when they should be fully shadowed.

diff --git a/Engine/Source/Runtime/Engine/Private/ModelLight.cpp b/Engine/Source/Runtime/Engine/Private/ModelLight.cpp
--- a/Engine/Source/Runtime/Engine/Private/ModelLight.cpp
+++ b/Engine/Source/Runtime/Engine/Private/ModelLight.cpp
@@ -1177,7 +1177,8 @@ void UModel::ApplyStaticLighting()
 		// We always create a light map if the surface either has any non-zero lighting data, or if the surface has a shadow map.  The runtime
 		// shaders are always expecting a light map in the case of a shadow map, even if the lighting is entirely zero.  This is simply to reduce
 		// the number of shader permutations to support in the very unlikely case of a unshadowed surfaces that has lighting values of zero.
-		const bool bNeedsLightMap = bHasNonZeroData || SurfaceGroup.ShadowMappedLights.Num() > 0 || GroupQuantizedData->bHasSkyShadowing;
+		const bool bHasRelevantLights = SurfaceGroup.Surfaces.ContainsByPredicate([](const FSurfaceStaticLightingGroup::FSurfaceInfo& SurfaceInfo) { return SurfaceInfo.SurfaceStaticLighting->RelevantLights.Num() > 0; });
+		const bool bNeedsLightMap = bHasNonZeroData || SurfaceGroup.ShadowMappedLights.Num() > 0 || bHasRelevantLights || GroupQuantizedData->bHasSkyShadowing;
 		if (bNeedsLightMap)
 		{
 			LightMap = FLightMap2D::AllocateLightMap(this, GroupQuantizedData, GroupLightmapBounds, PaddingType, LMF_None);
diff --git a/Engine/Source/Runtime/Engine/Private/StaticMeshLight.cpp b/Engine/Source/Runtime/Engine/Private/StaticMeshLight.cpp
--- a/Engine/Source/Runtime/Engine/Private/StaticMeshLight.cpp
+++ b/Engine/Source/Runtime/Engine/Private/StaticMeshLight.cpp
@@ -225,7 +225,7 @@ void FStaticMeshStaticLightingTextureMapping::Apply(FQuantizedLightmapData* Quan
 		// We always create a light map if the surface either has any non-zero lighting data, or if the surface has a shadow map.  The runtime
 		// shaders are always expecting a light map in the case of a shadow map, even if the lighting is entirely zero.  This is simply to reduce
 		// the number of shader permutations to support in the very unlikely case of a unshadowed surfaces that has lighting values of zero.
-		const bool bNeedsLightMap = bHasNonZeroData || ShadowMapData.Num() > 0 || (QuantizedData != NULL && QuantizedData->bHasSkyShadowing);
+		const bool bNeedsLightMap = bHasNonZeroData || ShadowMapData.Num() > 0 || Mesh->RelevantLights.Num() > 0 || (QuantizedData != NULL && QuantizedData->bHasSkyShadowing);
 		if (bNeedsLightMap)
 		{
 			// Create a light-map for the primitive.
diff --git a/Engine/Source/Runtime/Landscape/Private/LandscapeLight.cpp b/Engine/Source/Runtime/Landscape/Private/LandscapeLight.cpp
--- a/Engine/Source/Runtime/Landscape/Private/LandscapeLight.cpp
+++ b/Engine/Source/Runtime/Landscape/Private/LandscapeLight.cpp
@@ -46,7 +46,7 @@ void FLandscapeStaticLightingTextureMapping::Apply(FQuantizedLightmapData* Quant
 	// We always create a light map if the surface either has any non-zero lighting data, or if the surface has a shadow map.  The runtime
 	// shaders are always expecting a light map in the case of a shadow map, even if the lighting is entirely zero.  This is simply to reduce
 	// the number of shader permutations to support in the very unlikely case of a unshadowed surfaces that has lighting values of zero.
-	const bool bNeedsLightMap = bHasNonZeroData || ShadowMapData.Num() > 0 || (QuantizedData != NULL && QuantizedData->bHasSkyShadowing);
+	const bool bNeedsLightMap = bHasNonZeroData || ShadowMapData.Num() > 0 || Mesh->RelevantLights.Num() > 0 || (QuantizedData != NULL && QuantizedData->bHasSkyShadowing);
 	if (bNeedsLightMap)
 	{
 		// Create a light-map for the primitive.
@@ -81,7 +81,7 @@ void FLandscapeStaticLightingTextureMapping::Apply(FQuantizedLightmapData* Quant
 	// Build the list of statically irrelevant lights.
 	// TODO: This should be stored per LOD.
 	LandscapeComponent->IrrelevantLights.Empty();
-	for(int32 LightIndex = 0;LightIndex < Mesh->RelevantLights.Num();LightIndex++)
+	for (int32 LightIndex = 0; LightIndex < Mesh->RelevantLights.Num(); LightIndex++)
 	{
 		const ULightComponent* Light = Mesh->RelevantLights[LightIndex];
 
@@ -90,7 +90,7 @@ void FLandscapeStaticLightingTextureMapping::Apply(FQuantizedLightmapData* Quant
 
 		// Add the light to the statically irrelevant light list if it is in the potentially relevant light list, but didn't contribute to the light-map.
 		if(!bIsInLightMap)
-		{	
+		{
 			LandscapeComponent->IrrelevantLights.AddUnique(Light->LightGuid);
 		}
 	}
diff --git a/Engine/Source/Runtime/Landscape/Private/LandscapeRender.cpp b/Engine/Source/Runtime/Landscape/Private/LandscapeRender.cpp
--- a/Engine/Source/Runtime/Landscape/Private/LandscapeRender.cpp
+++ b/Engine/Source/Runtime/Landscape/Private/LandscapeRender.cpp
@@ -619,7 +619,7 @@ FLandscapeComponentSceneProxy::FLandscapeComponentSceneProxy(ULandscapeComponent
 	ComponentLightInfo = MakeUnique<FLandscapeLCI>(InComponent);
 	check(ComponentLightInfo);
 
-	const bool bHasStaticLighting = InComponent->LightMap != nullptr || InComponent->ShadowMap != nullptr;
+	const bool bHasStaticLighting = InComponent->bHasCachedStaticLighting;
 
 	// Check material usage
 	if (MaterialInterface == nullptr ||
