Do not register landscape components during PostLoad event
Do not interact with DDC when landscape component has invalid collision guid
^ Those two cases could happen when new collision components created during PostLoad

Merging using UE4-Orion_to_UE4

diff --git a/Engine/Source/Runtime/Landscape/Private/LandscapeCollision.cpp b/Engine/Source/Runtime/Landscape/Private/LandscapeCollision.cpp
--- a/Engine/Source/Runtime/Landscape/Private/LandscapeCollision.cpp
+++ b/Engine/Source/Runtime/Landscape/Private/LandscapeCollision.cpp
@@ -83,6 +83,8 @@ ULandscapeMeshCollisionComponent::FPhysXMeshRef::~FPhysXMeshRef()
 static FString GetHFDDCKeyString(const FName& Format, bool bDefMaterial, const FGuid& StateId, const TArray<UPhysicalMaterial*>& PhysicalMaterials)
 {
 	FGuid CombinedStateId;
+	
+	ensure(StateId.IsValid());
 
 	if (bDefMaterial)
 	{
@@ -395,7 +397,7 @@ bool ULandscapeHeightfieldCollisionComponent::CookCollisionData(const FName& For
 	// we have 2 versions of collision objects
 	const int32 CookedDataIndex = bUseDefMaterial ? 0 : 1;
 
-	if (bCheckDDC)
+	if (bCheckDDC && HeightfieldGuid.IsValid())
 	{
 		// Ensure that content was saved with physical materials before using DDC data
 		if (GetLinkerUE4Version() >= VER_UE4_LANDSCAPE_SERIALIZE_PHYSICS_MATERIALS)
@@ -524,7 +526,7 @@ bool ULandscapeHeightfieldCollisionComponent::CookCollisionData(const FName& For
 		OutCookedData.SetNumUninitialized(OutData.Num());
 		FMemory::Memcpy(OutCookedData.GetData(), OutData.GetData(), OutData.Num());
 
-		if (bShouldSaveCookedDataToDDC[CookedDataIndex])
+		if (bShouldSaveCookedDataToDDC[CookedDataIndex] && HeightfieldGuid.IsValid())
 		{
 			GetDerivedDataCacheRef().Put(*GetHFDDCKeyString(Format, bUseDefMaterial, HeightfieldGuid, InOutMaterials), OutCookedData);
 			bShouldSaveCookedDataToDDC[CookedDataIndex] = false;
@@ -552,7 +554,7 @@ bool ULandscapeMeshCollisionComponent::CookCollisionData(const FName& Format, bo
 	if (bCheckDDC)
 	{
 		// Ensure that content was saved with physical materials before using DDC data
-		if (GetLinkerUE4Version() >= VER_UE4_LANDSCAPE_SERIALIZE_PHYSICS_MATERIALS)
+		if (GetLinkerUE4Version() >= VER_UE4_LANDSCAPE_SERIALIZE_PHYSICS_MATERIALS && MeshGuid.IsValid())
 		{
 			FString DDCKey = GetHFDDCKeyString(Format, bUseDefMaterial, MeshGuid, InOutMaterials);
 
@@ -712,7 +714,7 @@ bool ULandscapeMeshCollisionComponent::CookCollisionData(const FName& Format, bo
 		OutCookedData.SetNumUninitialized(OutData.Num());
 		FMemory::Memcpy(OutCookedData.GetData(), OutData.GetData(), OutData.Num());
 
-		if (bShouldSaveCookedDataToDDC[CookedDataIndex])
+		if (bShouldSaveCookedDataToDDC[CookedDataIndex] && MeshGuid.IsValid())
 		{
 			GetDerivedDataCacheRef().Put(*GetHFDDCKeyString(Format, bUseDefMaterial, MeshGuid, InOutMaterials), OutCookedData);
 			bShouldSaveCookedDataToDDC[CookedDataIndex] = false;
@@ -1233,7 +1235,10 @@ void ULandscapeHeightfieldCollisionComponent::Serialize(FArchive& Ar)
 		{
 			FName Format = Ar.CookingTarget()->GetPhysicsFormat(nullptr);
 			CookCollisionData(Format, false, true, CookedCollisionData, CookedPhysicalMaterials);
-			GetDerivedDataCacheRef().Put(*GetHFDDCKeyString(Format, false, HeightfieldGuid, CookedPhysicalMaterials), CookedCollisionData);
+			if (HeightfieldGuid.IsValid())
+			{
+				GetDerivedDataCacheRef().Put(*GetHFDDCKeyString(Format, false, HeightfieldGuid, CookedPhysicalMaterials), CookedCollisionData);
+			}
 		}
 	}
 #endif// WITH_EDITOR
@@ -1491,12 +1496,12 @@ void ULandscapeHeightfieldCollisionComponent::PreSave()
 	{
 #if WITH_EDITOR
 		static FName PhysicsFormatName(FPlatformProperties::GetPhysicsFormat());
-		if (CookedCollisionData.Num())
+		if (CookedCollisionData.Num() && HeightfieldGuid.IsValid())
 		{
 			GetDerivedDataCacheRef().Put(*GetHFDDCKeyString(PhysicsFormatName, false, HeightfieldGuid, CookedPhysicalMaterials), CookedCollisionData);
 		}
 
-		if (CookedCollisionDataEd.Num())
+		if (CookedCollisionDataEd.Num() && HeightfieldGuid.IsValid())
 		{
 			GetDerivedDataCacheRef().Put(*GetHFDDCKeyString(PhysicsFormatName, true, HeightfieldGuid, TArray<UPhysicalMaterial*>()), CookedCollisionDataEd);
 		}
diff --git a/Engine/Source/Runtime/Landscape/Private/LandscapeEdit.cpp b/Engine/Source/Runtime/Landscape/Private/LandscapeEdit.cpp
--- a/Engine/Source/Runtime/Landscape/Private/LandscapeEdit.cpp
+++ b/Engine/Source/Runtime/Landscape/Private/LandscapeEdit.cpp
@@ -793,12 +793,12 @@ void ULandscapeComponent::UpdateCollisionHeightData(const FColor* HeightmapTextu
 		UpdateCollisionLayerData();
 	}
 
-	if (bRebuild && CollisionProxy)
+	if (bRebuild && CollisionProxy && CollisionProxy->GetRootComponent()->IsRegistered())
 	{
 		CollisionProxy->RegisterAllComponents();
 	}
 
-	if (ChangeType || CreatedNew)
+	if (ChangeType || CreatedNew && Proxy->GetRootComponent()->IsRegistered())
 	{
 		Proxy->RegisterAllComponents();
 	}
