Fix Linux build CIS.

diff --git a/Engine/Source/Runtime/Landscape/Private/LandscapeEdit.cpp b/Engine/Source/Runtime/Landscape/Private/LandscapeEdit.cpp
--- a/Engine/Source/Runtime/Landscape/Private/LandscapeEdit.cpp
+++ b/Engine/Source/Runtime/Landscape/Private/LandscapeEdit.cpp
@@ -798,7 +798,7 @@ void ULandscapeComponent::UpdateCollisionHeightData(const FColor* HeightmapTextu
 		CollisionProxy->RegisterAllComponents();
 	}
 
-	if (ChangeType || CreatedNew && Proxy->GetRootComponent()->IsRegistered())
+	if ((ChangeType || CreatedNew) && Proxy->GetRootComponent()->IsRegistered())
 	{
 		Proxy->RegisterAllComponents();
 	}
