Make landscape layer blend work with material attributes

diff --git a/Engine/Source/Runtime/Landscape/Classes/Materials/MaterialExpressionLandscapeLayerBlend.h b/Engine/Source/Runtime/Landscape/Classes/Materials/MaterialExpressionLandscapeLayerBlend.h
--- a/Engine/Source/Runtime/Landscape/Classes/Materials/MaterialExpressionLandscapeLayerBlend.h
+++ b/Engine/Source/Runtime/Landscape/Classes/Materials/MaterialExpressionLandscapeLayerBlend.h
@@ -71,6 +71,8 @@ class UMaterialExpressionLandscapeLayerBlend : public UMaterialExpression
 	// End UObject Interface
 
 	// Begin UMaterialExpression Interface
+	virtual uint32 GetInputType(int32 InputIndex) override { return MCT_Unknown; }
+	virtual uint32 GetOutputType(int32 InputIndex) override { return MCT_Unknown; }
 #endif
 	virtual bool IsResultMaterialAttributes(int32 OutputIndex) override;
 	virtual int32 Compile(class FMaterialCompiler* Compiler, int32 OutputIndex, int32 MultiplexIndex) override;
