Fixed issue where landscape grass with more than 5 layers only showed the last layer
Fixed issue where a None entry in the Landscape Grass Output node would cause grass meshes to be assigned to the wrong layers

UE-18502

diff --git a/Engine/Source/Runtime/Landscape/Private/LandscapeGrass.cpp b/Engine/Source/Runtime/Landscape/Private/LandscapeGrass.cpp
--- a/Engine/Source/Runtime/Landscape/Private/LandscapeGrass.cpp
+++ b/Engine/Source/Runtime/Landscape/Private/LandscapeGrass.cpp
@@ -454,8 +454,16 @@ public:
 			GrassWeightArrays.Empty(GrassTypes.Num());
 			for (auto GrassType : GrassTypes)
 			{
-				int32 Index = GrassWeightArrays.Add(&NewGrassData->WeightData.Add(GrassType));
-				GrassWeightArrays[Index]->Empty(FMath::Square(ComponentSizeVerts));
+				NewGrassData->WeightData.Add(GrassType);
+			}
+
+			// need a second loop because the WeightData map will reallocate its arrays as grass types are added
+			for (auto GrassType : GrassTypes)
+			{
+				TArray<uint8>* DataArray = NewGrassData->WeightData.Find(GrassType);
+				check(DataArray);
+				DataArray->Empty(FMath::Square(ComponentSizeVerts));
+				GrassWeightArrays.Add(DataArray);
 			}
 
 			for (int32 PassIdx = 0; PassIdx < NumPasses; PassIdx++)
@@ -507,6 +515,9 @@ public:
 				}
 			}
 
+			// remove null grass type if we had one (can occur if the node has null entries)
+			NewGrassData->WeightData.Remove(nullptr);
+
 			// Assign the new data (thread-safe)
 			ComponentInfo.Component->GrassData = MakeShareable(NewGrassData);
 		}
@@ -617,10 +628,7 @@ void ULandscapeComponent::RenderGrassMap()
 			GrassTypes.Empty(GrassExpressions[0]->GrassTypes.Num());
 			for (auto& GrassTypeInput : GrassExpressions[0]->GrassTypes)
 			{
-				if (GrassTypeInput.GrassType)
-				{
-					GrassTypes.Add(GrassTypeInput.GrassType);
-				}
+				GrassTypes.Add(GrassTypeInput.GrassType);
 			}
 
 			TArray<ULandscapeComponent*> LandscapeComponents;
