Fixed crash after deleting a landscape spline segment with meshes

diff --git a/Engine/Source/Runtime/Landscape/Private/LandscapeSplines.cpp b/Engine/Source/Runtime/Landscape/Private/LandscapeSplines.cpp
--- a/Engine/Source/Runtime/Landscape/Private/LandscapeSplines.cpp
+++ b/Engine/Source/Runtime/Landscape/Private/LandscapeSplines.cpp
@@ -1614,6 +1614,7 @@ void ULandscapeSplineControlPoint::DeleteSplinePoints()
 	if (LocalMeshComponent != nullptr)
 	{
 		OuterSplines->Modify();
+		LocalMeshComponent->Modify();
 		checkSlow(OuterSplines->MeshComponentLocalOwnersMap.FindRef(LocalMeshComponent) == this);
 		verifySlow(OuterSplines->MeshComponentLocalOwnersMap.Remove(LocalMeshComponent) == 1);
 		LocalMeshComponent->DestroyComponent();
@@ -2149,6 +2150,7 @@ void ULandscapeSplineSegment::UpdateSplinePoints(bool bUpdateCollision)
 		// Destroy old unwanted components now
 		for (UMeshComponent* MeshComponent : OldLocalMeshComponents)
 		{
+			checkSlow(OuterSplines->MeshComponentLocalOwnersMap.FindRef(MeshComponent) == this);
 			verifySlow(OuterSplines->MeshComponentLocalOwnersMap.Remove(MeshComponent) == 1);
 			MeshComponent->DestroyComponent();
 		}
@@ -2355,6 +2357,7 @@ void ULandscapeSplineSegment::UpdateSplinePoints(bool bUpdateCollision)
 		// Spline needs no mesh components (0 length or no meshes to use) so destroy any we have
 		for (auto* MeshComponent : OldLocalMeshComponents)
 		{
+			checkSlow(OuterSplines->MeshComponentLocalOwnersMap.FindRef(MeshComponent) == this);
 			verifySlow(OuterSplines->MeshComponentLocalOwnersMap.Remove(MeshComponent) == 1);
 			MeshComponent->DestroyComponent();
 		}
@@ -2413,13 +2416,18 @@ void ULandscapeSplineSegment::DeleteSplinePoints()
 	OuterSplines->MarkRenderStateDirty();
 
 	// Destroy mesh components
-	OuterSplines->GetOwner()->Modify();
-	for (auto* MeshComponent : LocalMeshComponents)
+	if (LocalMeshComponents.Num() > 0)
 	{
-		MeshComponent->Modify();
-		MeshComponent->DestroyComponent();
+		OuterSplines->Modify();
+		for (auto* LocalMeshComponent : LocalMeshComponents)
+		{
+			checkSlow(OuterSplines->MeshComponentLocalOwnersMap.FindRef(LocalMeshComponent) == this);
+			verifySlow(OuterSplines->MeshComponentLocalOwnersMap.Remove(LocalMeshComponent) == 1);
+			LocalMeshComponent->Modify();
+			LocalMeshComponent->DestroyComponent();
+		}
+		LocalMeshComponents.Empty();
 	}
-	LocalMeshComponents.Empty();
 
 	auto ForeignMeshComponentsMap = GetForeignMeshComponents();
 	for (auto& ForeignMeshComponentsPair : ForeignMeshComponentsMap)
