Rename some variables for clarity

diff --git a/Engine/Source/Runtime/Landscape/Private/LandscapeSplines.cpp b/Engine/Source/Runtime/Landscape/Private/LandscapeSplines.cpp
--- a/Engine/Source/Runtime/Landscape/Private/LandscapeSplines.cpp
+++ b/Engine/Source/Runtime/Landscape/Private/LandscapeSplines.cpp
@@ -1779,38 +1779,38 @@ void ULandscapeSplineSegment::PostLoad()
 		ULandscapeSplinesComponent* OuterSplines = GetOuterULandscapeSplinesComponent();
 		if (OuterSplines->SplineEditorMesh != nullptr)
 		{
-			for (auto* MeshComponent : LocalMeshComponents)
+			for (auto* LocalMeshComponent : LocalMeshComponents)
 			{
-				if (MeshComponent->StaticMesh == nullptr)
+				if (LocalMeshComponent->StaticMesh == nullptr)
 				{
-					MeshComponent->ConditionalPostLoad();
-					MeshComponent->SetStaticMesh(OuterSplines->SplineEditorMesh);
-					MeshComponent->SetHiddenInGame(true);
-					MeshComponent->SetVisibility(OuterSplines->bShowSplineEditorMesh);
-					MeshComponent->BodyInstance.SetCollisionProfileName(UCollisionProfile::NoCollision_ProfileName);
+					LocalMeshComponent->ConditionalPostLoad();
+					LocalMeshComponent->SetStaticMesh(OuterSplines->SplineEditorMesh);
+					LocalMeshComponent->SetHiddenInGame(true);
+					LocalMeshComponent->SetVisibility(OuterSplines->bShowSplineEditorMesh);
+					LocalMeshComponent->BodyInstance.SetCollisionProfileName(UCollisionProfile::NoCollision_ProfileName);
 				}
 			}
 		}
 
-		for (auto* MeshComponent : LocalMeshComponents)
+		for (auto* LocalMeshComponent : LocalMeshComponents)
 		{
-			OuterSplines->MeshComponentLocalOwnersMap.Add(MeshComponent, this);
+			OuterSplines->MeshComponentLocalOwnersMap.Add(LocalMeshComponent, this);
 		}
 	}
 
 	if (GetLinkerUE4Version() < VER_UE4_LANDSCAPE_SPLINE_CROSS_LEVEL_MESHES)
 	{
 		// Fix collision profile
-		for (auto* MeshComponent : LocalMeshComponents) // ForeignMeshComponents didn't exist yet
+		for (auto* LocalMeshComponent : LocalMeshComponents) // ForeignMeshComponents didn't exist yet
 		{
-			const bool bUsingEditorMesh = MeshComponent->bHiddenInGame;
+			const bool bUsingEditorMesh = LocalMeshComponent->bHiddenInGame;
 			const FName CollisionProfile = (bEnableCollision && !bUsingEditorMesh) ? UCollisionProfile::BlockAll_ProfileName : UCollisionProfile::NoCollision_ProfileName;
-			if (MeshComponent->GetCollisionProfileName() != CollisionProfile)
+			if (LocalMeshComponent->GetCollisionProfileName() != CollisionProfile)
 			{
-				MeshComponent->SetCollisionProfileName(CollisionProfile);
+				LocalMeshComponent->SetCollisionProfileName(CollisionProfile);
 			}
 
-			MeshComponent->SetFlags(RF_TextExportTransient);
+			LocalMeshComponent->SetFlags(RF_TextExportTransient);
 		}
 	}
 #endif
@@ -1823,20 +1823,20 @@ void ULandscapeSplineSegment::SetSplineSelected(bool bInSelected)
 	bSelected = bInSelected;
 	GetOuterULandscapeSplinesComponent()->MarkRenderStateDirty();
 
-	for (auto* MeshComponent : LocalMeshComponents)
+	for (auto* LocalMeshComponent : LocalMeshComponents)
 	{
-		MeshComponent->bSelected = bInSelected;
-		MeshComponent->PushSelectionToProxy();
+		LocalMeshComponent->bSelected = bInSelected;
+		LocalMeshComponent->PushSelectionToProxy();
 	}
 
 	auto ForeignMeshComponentsMap = GetForeignMeshComponents();
 	for (auto& ForeignMeshComponentsPair : ForeignMeshComponentsMap)
 	{
 		ULandscapeSplinesComponent* MeshComponentOuterSplines = ForeignMeshComponentsPair.Key;
-		for (auto* MeshComponent : ForeignMeshComponentsPair.Value)
+		for (auto* ForeignMeshComponent : ForeignMeshComponentsPair.Value)
 		{
-			MeshComponent->bSelected = bInSelected;
-			MeshComponent->PushSelectionToProxy();
+			ForeignMeshComponent->bSelected = bInSelected;
+			ForeignMeshComponent->PushSelectionToProxy();
 		}
 	}
 }
@@ -2006,19 +2006,19 @@ void ULandscapeSplineSegment::UpdateSplinePoints(bool bUpdateCollision)
 	auto ForeignMeshComponentsMap = GetForeignMeshComponents();
 
 	// Unregister components
-	for (auto* MeshComponent : OldLocalMeshComponents)
+	for (auto* LocalMeshComponent : OldLocalMeshComponents)
 	{
-		MeshComponent->Modify();
-		MeshComponent->UnregisterComponent();
+		LocalMeshComponent->Modify();
+		LocalMeshComponent->UnregisterComponent();
 	}
 	for (auto& ForeignMeshComponentsPair : ForeignMeshComponentsMap)
 	{
 		ForeignMeshComponentsPair.Key->Modify();
 		ForeignMeshComponentsPair.Key->GetOwner()->Modify();
-		for (auto* MeshComponent : ForeignMeshComponentsPair.Value)
+		for (auto* ForeignMeshComponent : ForeignMeshComponentsPair.Value)
 		{
-			MeshComponent->Modify();
-			MeshComponent->UnregisterComponent();
+			ForeignMeshComponent->Modify();
+			ForeignMeshComponent->UnregisterComponent();
 		}
 	}
 
@@ -2148,21 +2148,21 @@ void ULandscapeSplineSegment::UpdateSplinePoints(bool bUpdateCollision)
 		MeshSettings.Add(FMeshSettings(T, nullptr));
 
 		// Destroy old unwanted components now
-		for (UMeshComponent* MeshComponent : OldLocalMeshComponents)
+		for (UMeshComponent* LocalMeshComponent : OldLocalMeshComponents)
 		{
-			checkSlow(OuterSplines->MeshComponentLocalOwnersMap.FindRef(MeshComponent) == this);
-			verifySlow(OuterSplines->MeshComponentLocalOwnersMap.Remove(MeshComponent) == 1);
-			MeshComponent->DestroyComponent();
+			checkSlow(OuterSplines->MeshComponentLocalOwnersMap.FindRef(LocalMeshComponent) == this);
+			verifySlow(OuterSplines->MeshComponentLocalOwnersMap.Remove(LocalMeshComponent) == 1);
+			LocalMeshComponent->DestroyComponent();
 		}
 		OldLocalMeshComponents.Empty();
 
 		for (auto& ForeignMeshComponentsPair : ForeignMeshComponentsMap)
 		{
 			ULandscapeSplinesComponent* MeshComponentOuterSplines = ForeignMeshComponentsPair.Key;
-			for (auto* MeshComponent : ForeignMeshComponentsPair.Value)
+			for (auto* ForeignMeshComponent : ForeignMeshComponentsPair.Value)
 			{
-				MeshComponentOuterSplines->RemoveForeignMeshComponent(this, MeshComponent);
-				MeshComponent->DestroyComponent();
+				MeshComponentOuterSplines->RemoveForeignMeshComponent(this, ForeignMeshComponent);
+				ForeignMeshComponent->DestroyComponent();
 			}
 		}
 		ForeignMeshComponentsMap.Empty();
@@ -2355,11 +2355,11 @@ void ULandscapeSplineSegment::UpdateSplinePoints(bool bUpdateCollision)
 	else
 	{
 		// Spline needs no mesh components (0 length or no meshes to use) so destroy any we have
-		for (auto* MeshComponent : OldLocalMeshComponents)
+		for (auto* LocalMeshComponent : OldLocalMeshComponents)
 		{
-			checkSlow(OuterSplines->MeshComponentLocalOwnersMap.FindRef(MeshComponent) == this);
-			verifySlow(OuterSplines->MeshComponentLocalOwnersMap.Remove(MeshComponent) == 1);
-			MeshComponent->DestroyComponent();
+			checkSlow(OuterSplines->MeshComponentLocalOwnersMap.FindRef(LocalMeshComponent) == this);
+			verifySlow(OuterSplines->MeshComponentLocalOwnersMap.Remove(LocalMeshComponent) == 1);
+			LocalMeshComponent->DestroyComponent();
 		}
 		OldLocalMeshComponents.Empty();
 		for (auto& ForeignMeshComponentsPair : ForeignMeshComponentsMap)
@@ -2382,22 +2382,22 @@ void ULandscapeSplineSegment::UpdateSplineEditorMesh()
 {
 	ULandscapeSplinesComponent* OuterSplines = CastChecked<ULandscapeSplinesComponent>(GetOuter());
 
-	for (auto* MeshComponent : LocalMeshComponents)
+	for (auto* LocalMeshComponent : LocalMeshComponents)
 	{
-		if (MeshComponent->bHiddenInGame)
+		if (LocalMeshComponent->bHiddenInGame)
 		{
-			MeshComponent->SetVisibility(OuterSplines->bShowSplineEditorMesh);
+			LocalMeshComponent->SetVisibility(OuterSplines->bShowSplineEditorMesh);
 		}
 	}
 
 	auto ForeignMeshComponentsMap = GetForeignMeshComponents();
 	for (auto& ForeignMeshComponentsPair : ForeignMeshComponentsMap)
 	{
-		for (auto* MeshComponent : ForeignMeshComponentsPair.Value)
+		for (auto* ForeignMeshComponent : ForeignMeshComponentsPair.Value)
 		{
-			if (MeshComponent->bHiddenInGame)
+			if (ForeignMeshComponent->bHiddenInGame)
 			{
-				MeshComponent->SetVisibility(OuterSplines->bShowSplineEditorMesh);
+				ForeignMeshComponent->SetVisibility(OuterSplines->bShowSplineEditorMesh);
 			}
 		}
 	}
@@ -2435,11 +2435,11 @@ void ULandscapeSplineSegment::DeleteSplinePoints()
 		ULandscapeSplinesComponent* MeshComponentOuterSplines = ForeignMeshComponentsPair.Key;
 		MeshComponentOuterSplines->Modify();
 		MeshComponentOuterSplines->GetOwner()->Modify();
-		for (auto* MeshComponent : ForeignMeshComponentsPair.Value)
+		for (auto* ForeignMeshComponent : ForeignMeshComponentsPair.Value)
 		{
-			MeshComponent->Modify();
-			MeshComponentOuterSplines->RemoveForeignMeshComponent(this, MeshComponent);
-			MeshComponent->DestroyComponent();
+			ForeignMeshComponent->Modify();
+			MeshComponentOuterSplines->RemoveForeignMeshComponent(this, ForeignMeshComponent);
+			ForeignMeshComponent->DestroyComponent();
 		}
 	}
 
