Cleaning out programmer-only comments that were harvested by Doxygen.

diff --git a/Engine/Source/Runtime/Engine/Classes/Particles/TypeData/ParticleModuleTypeDataGpu.h b/Engine/Source/Runtime/Engine/Classes/Particles/TypeData/ParticleModuleTypeDataGpu.h
--- a/Engine/Source/Runtime/Engine/Classes/Particles/TypeData/ParticleModuleTypeDataGpu.h
+++ b/Engine/Source/Runtime/Engine/Classes/Particles/TypeData/ParticleModuleTypeDataGpu.h
@@ -414,15 +414,15 @@ class UParticleModuleTypeDataGpu : public UParticleModuleTypeDataBase
 	UPROPERTY(EditAnywhere, Category = ParticleModuleTypeDataGpu)
 	uint32 bClearExistingParticlesOnInit:1;
 
-	// Begin UObject Interface
+	//~ Begin UObject Interface
 	virtual void PostLoad() override;
 	virtual void BeginDestroy() override;
-	// End UObject Interface
+	//~ End UObject Interface
 
-	// Begin UParticleModuleTypeDataBase Interface
+	//~ Begin UParticleModuleTypeDataBase Interface
 	virtual void Build( struct FParticleEmitterBuildInfo& EmitterBuildInfo ) override;
 	virtual bool RequiresBuild() const override { return true; }
 	virtual FParticleEmitterInstance* CreateInstance(UParticleEmitter* InEmitterParent, UParticleSystemComponent* InComponent) override;
-	// End UParticleModuleTypeDataBase Interface
+	//~ End UParticleModuleTypeDataBase Interface
 };
 
diff --git a/Engine/Source/Runtime/Engine/Classes/Particles/TypeData/ParticleModuleTypeDataMesh.h b/Engine/Source/Runtime/Engine/Classes/Particles/TypeData/ParticleModuleTypeDataMesh.h
--- a/Engine/Source/Runtime/Engine/Classes/Particles/TypeData/ParticleModuleTypeDataMesh.h
+++ b/Engine/Source/Runtime/Engine/Classes/Particles/TypeData/ParticleModuleTypeDataMesh.h
@@ -184,27 +184,27 @@ class UParticleModuleTypeDataMesh : public UParticleModuleTypeDataBase
 	UPROPERTY(EditAnywhere, Category = Collision)
 	uint32 bCollisionsConsiderPartilceSize : 1;
 
-	// Begin UObject Interface
+	//~ Begin UObject Interface
 #if WITH_EDITOR
 	virtual void	PostEditChangeProperty(FPropertyChangedEvent& PropertyChangedEvent) override;
 #endif // WITH_EDITOR
 
 	virtual void	Serialize(FArchive& Ar) override;
 
-	// End UObject Interface
+	//~ End UObject Interface
 
 	void CreateDistribution();
 
-	// Begin UParticleModule Interface
+	//~ Begin UParticleModule Interface
 	virtual void	SetToSensibleDefaults(UParticleEmitter* Owner) override;
-	// End UParticleModule Interface
+	//~ End UParticleModule Interface
 
-	// Begin UParticleModuleTypeDataBase Interface
+	//~ Begin UParticleModuleTypeDataBase Interface
 	virtual FParticleEmitterInstance*	CreateInstance(UParticleEmitter* InEmitterParent, UParticleSystemComponent* InComponent) override;
 	virtual bool	SupportsSpecificScreenAlignmentFlags() const override {	return true;	}	
 	virtual bool	SupportsSubUV() const override { return true; }
 	virtual bool	IsAMeshEmitter() const override { return true; }
-	// End UParticleModuleTypeDataBase Interface
+	//~ End UParticleModuleTypeDataBase Interface
 };
 
 
diff --git a/Engine/Source/Runtime/Engine/Classes/Particles/VectorField/ParticleModuleVectorFieldGlobal.h b/Engine/Source/Runtime/Engine/Classes/Particles/VectorField/ParticleModuleVectorFieldGlobal.h
--- a/Engine/Source/Runtime/Engine/Classes/Particles/VectorField/ParticleModuleVectorFieldGlobal.h
+++ b/Engine/Source/Runtime/Engine/Classes/Particles/VectorField/ParticleModuleVectorFieldGlobal.h
@@ -25,8 +25,8 @@ class UParticleModuleVectorFieldGlobal : public UParticleModuleVectorFieldBase
 	UPROPERTY(EditAnywhere, Category=VectorField, meta=(ClampMin = "0.0", UIMin = "0.0", UIMax="1.0", editcondition = "bOverrideGlobalVectorFieldTightness"))
 	float GlobalVectorFieldTightness;
 
-	// Begin UParticleModule Interface
+	//~ Begin UParticleModule Interface
 	virtual void CompileModule(struct FParticleEmitterBuildInfo& EmitterInfo) override;
-	// Begin UParticleModule Interface
+	//~ Begin UParticleModule Interface
 };
 
diff --git a/Engine/Source/Runtime/Engine/Classes/Particles/Velocity/ParticleModuleVelocity.h b/Engine/Source/Runtime/Engine/Classes/Particles/Velocity/ParticleModuleVelocity.h
--- a/Engine/Source/Runtime/Engine/Classes/Particles/Velocity/ParticleModuleVelocity.h
+++ b/Engine/Source/Runtime/Engine/Classes/Particles/Velocity/ParticleModuleVelocity.h
@@ -28,16 +28,16 @@ class UParticleModuleVelocity : public UParticleModuleVelocityBase
 	/** Initializes the default values for this property */
 	void InitializeDefaults();
 	
-	// Begin UObject Interface
+	//~ Begin UObject Interface
 #if WITH_EDITOR
 	virtual void PostEditChangeProperty(FPropertyChangedEvent& PropertyChangedEvent) override;
 #endif // WITH_EDITOR
 	virtual void PostInitProperties() override;
-	// End UObject Interface
+	//~ End UObject Interface
 
-	// Begin UParticleModule Interface
+	//~ Begin UParticleModule Interface
 	virtual void Spawn(FParticleEmitterInstance* Owner, int32 Offset, float SpawnTime, FBaseParticle* ParticleBase) override;
-	// Begin UParticleModule Interface
+	//~ Begin UParticleModule Interface
 
 	/**
 	 *	Extended version of spawn, allows for using a random stream for distribution value retrieval
diff --git a/Engine/Source/Runtime/Engine/Classes/Particles/Velocity/ParticleModuleVelocityCone.h b/Engine/Source/Runtime/Engine/Classes/Particles/Velocity/ParticleModuleVelocityCone.h
--- a/Engine/Source/Runtime/Engine/Classes/Particles/Velocity/ParticleModuleVelocityCone.h
+++ b/Engine/Source/Runtime/Engine/Classes/Particles/Velocity/ParticleModuleVelocityCone.h
@@ -25,17 +25,17 @@ class UParticleModuleVelocityCone : public UParticleModuleVelocityBase
 	/** Initializes the default values for this property */
 	void InitializeDefaults();
 
-	// Begin UObject Interface
+	//~ Begin UObject Interface
 #if WITH_EDITOR
 	virtual void PostEditChangeProperty(FPropertyChangedEvent& PropertyChangedEvent) override;
 #endif // WITH_EDITOR
 	virtual void PostInitProperties() override;
-	// End UObject Interface
+	//~ End UObject Interface
 
-	// Begin UParticleModule Interface
+	//~ Begin UParticleModule Interface
 	virtual void Spawn(FParticleEmitterInstance* Owner, int32 Offset, float SpawnTime, FBaseParticle* ParticleBase) override;
 	virtual void	Render3DPreview(FParticleEmitterInstance* Owner, const FSceneView* View,FPrimitiveDrawInterface* PDI) override;
-	// Begin UParticleModule Interface
+	//~ Begin UParticleModule Interface
 	
 	/**
 	 *	Extended version of spawn, allows for using a random stream for distribution value retrieval
diff --git a/Engine/Source/Runtime/Engine/Classes/Particles/Velocity/ParticleModuleVelocityInheritParent.h b/Engine/Source/Runtime/Engine/Classes/Particles/Velocity/ParticleModuleVelocityInheritParent.h
--- a/Engine/Source/Runtime/Engine/Classes/Particles/Velocity/ParticleModuleVelocityInheritParent.h
+++ b/Engine/Source/Runtime/Engine/Classes/Particles/Velocity/ParticleModuleVelocityInheritParent.h
@@ -25,11 +25,11 @@ class UParticleModuleVelocityInheritParent : public UParticleModuleVelocityBase
 	virtual void PostEditChangeProperty(FPropertyChangedEvent& PropertyChangedEvent) override;
 #endif // WITH_EDITOR
 	virtual void PostInitProperties() override;
-	// End UObject Interface
+	//~ End UObject Interface
 
-	// Begin UParticleModule Interface
+	//~ Begin UParticleModule Interface
 	virtual void Spawn(FParticleEmitterInstance* Owner, int32 Offset, float SpawnTime, FBaseParticle* ParticleBase) override;
-	// Begin UParticleModule Interface
+	//~ Begin UParticleModule Interface
 };
 
 
diff --git a/Engine/Source/Runtime/Engine/Classes/Particles/Velocity/ParticleModuleVelocityOverLifetime.h b/Engine/Source/Runtime/Engine/Classes/Particles/Velocity/ParticleModuleVelocityOverLifetime.h
--- a/Engine/Source/Runtime/Engine/Classes/Particles/Velocity/ParticleModuleVelocityOverLifetime.h
+++ b/Engine/Source/Runtime/Engine/Classes/Particles/Velocity/ParticleModuleVelocityOverLifetime.h
@@ -32,12 +32,12 @@ class UParticleModuleVelocityOverLifetime : public UParticleModuleVelocityBase
 	virtual void PostEditChangeProperty(FPropertyChangedEvent& PropertyChangedEvent) override;
 #endif // WITH_EDITOR
 	virtual void PostInitProperties() override;
-	// End UObject Interface
+	//~ End UObject Interface
 
-	// Begin UParticleModule Interface
+	//~ Begin UParticleModule Interface
 	virtual void Spawn(FParticleEmitterInstance* Owner, int32 Offset, float SpawnTime, FBaseParticle* ParticleBase) override;
 	virtual void	Update(FParticleEmitterInstance* Owner, int32 Offset, float DeltaTime) override;
-	// Begin UParticleModule Interface
+	//~ Begin UParticleModule Interface
 };
 
 
diff --git a/Engine/Source/Runtime/Engine/Classes/PhysicsEngine/BodySetup.h b/Engine/Source/Runtime/Engine/Classes/PhysicsEngine/BodySetup.h
--- a/Engine/Source/Runtime/Engine/Classes/PhysicsEngine/BodySetup.h
+++ b/Engine/Source/Runtime/Engine/Classes/PhysicsEngine/BodySetup.h
@@ -133,7 +133,7 @@ class UBodySetup : public UObject
 	bool bHasCookedCollisionData;
 
 public:
-	// Begin UObject interface.
+	//~ Begin UObject Interface.
 	virtual void Serialize(FArchive& Ar) override;
 	virtual void BeginDestroy() override;
 	virtual void FinishDestroy() override;
@@ -143,10 +143,10 @@ public:
 	virtual void PostEditUndo() override;
 #endif // WITH_EDITOR
 	virtual SIZE_T GetResourceSize(EResourceSizeMode::Type Mode) override;
-	// End UObject interface.
+	//~ End UObject Interface.
 
 	//
-	// UBodySetup interface.
+	//~ Begin UBodySetup Interface.
 	//
 	ENGINE_API void CopyBodyPropertiesFrom(const UBodySetup* FromSetup);
 
diff --git a/Engine/Source/Runtime/Engine/Classes/PhysicsEngine/PhysicsConstraintComponent.h b/Engine/Source/Runtime/Engine/Classes/PhysicsEngine/PhysicsConstraintComponent.h
--- a/Engine/Source/Runtime/Engine/Classes/PhysicsEngine/PhysicsConstraintComponent.h
+++ b/Engine/Source/Runtime/Engine/Classes/PhysicsEngine/PhysicsConstraintComponent.h
@@ -74,11 +74,11 @@ public:
 	virtual void InitializeComponent() override;
 	//End ActorComponent interface
 
-	// Begin SceneComponent interface
+	//~ Begin SceneComponent Interface
 #if WITH_EDITOR
 	virtual void PostEditComponentMove(bool bFinished) override;
 #endif // WITH_EDITOR
-	// End SceneComponent interface
+	//~ End SceneComponent Interface
 
 	/** Get the body frame. Works without constraint being created */
 	ENGINE_API FTransform GetBodyTransform(EConstraintFrame::Type Frame) const;
diff --git a/Engine/Source/Runtime/Engine/Classes/PhysicsEngine/PhysicsConstraintTemplate.h b/Engine/Source/Runtime/Engine/Classes/PhysicsEngine/PhysicsConstraintTemplate.h
--- a/Engine/Source/Runtime/Engine/Classes/PhysicsEngine/PhysicsConstraintTemplate.h
+++ b/Engine/Source/Runtime/Engine/Classes/PhysicsEngine/PhysicsConstraintTemplate.h
@@ -1,12 +1,12 @@
 // Copyright 1998-2015 Epic Games, Inc. All Rights Reserved.
 
-//=============================================================================
+//~=============================================================================
 // Complete constraint definition used by rigid body physics.
 // 
 // Defaults here will give you a ball and socket joint.
 // Positions are in Physics scale.
 // When adding stuff here, make sure to update URB_ConstraintSetup::CopyConstraintParamsFrom
-//=============================================================================
+//~=============================================================================
 
 #pragma once 
 
@@ -94,9 +94,9 @@ class UPhysicsConstraintTemplate : public UObject
 	FConstraintInstance DefaultInstance;
 
 
-	// Begin UObject interface.
+	//~ Begin UObject Interface.
 	virtual void Serialize(FArchive& Ar) override;
-	// End UObject interface.
+	//~ End UObject Interface.
 
 private:
 	// Only needed for old content! Pre VER_UE4_ALL_PROPS_TO_CONSTRAINTINSTANCE
diff --git a/Engine/Source/Runtime/Engine/Classes/PhysicsEngine/RadialForceComponent.h b/Engine/Source/Runtime/Engine/Classes/PhysicsEngine/RadialForceComponent.h
--- a/Engine/Source/Runtime/Engine/Classes/PhysicsEngine/RadialForceComponent.h
+++ b/Engine/Source/Runtime/Engine/Classes/PhysicsEngine/RadialForceComponent.h
@@ -64,16 +64,16 @@ protected:
 	FCollisionObjectQueryParams CollisionObjectQueryParams;
 
 protected:
-	// Begin UActorComponent interface.
+	//~ Begin UActorComponent Interface.
 	virtual void TickComponent(float DeltaTime, enum ELevelTick TickType, FActorComponentTickFunction *ThisTickFunction) override;
-	// End UActorComponent interface.
+	//~ End UActorComponent Interface.
 
-	// Begin UObject interface.
+	//~ Begin UObject Interface.
 	virtual void PostLoad() override;
 #if WITH_EDITOR
 	virtual void PostEditChangeProperty(FPropertyChangedEvent& PropertyChangedEvent) override;
 #endif
-	// End UObject interface.
+	//~ End UObject Interface.
 
 	/** Update CollisionObjectQueryParams from ObjectTypesToAffect */
 	void UpdateCollisionObjectQueryParams();
diff --git a/Engine/Source/Runtime/Engine/Classes/Sound/AudioVolume.h b/Engine/Source/Runtime/Engine/Classes/Sound/AudioVolume.h
--- a/Engine/Source/Runtime/Engine/Classes/Sound/AudioVolume.h
+++ b/Engine/Source/Runtime/Engine/Classes/Sound/AudioVolume.h
@@ -182,18 +182,18 @@ class AAudioVolume : public AVolume
 
 
 
-	// Begin UObject Interface
+	//~ Begin UObject Interface
 #if WITH_EDITOR
 	virtual void PostEditChangeProperty(FPropertyChangedEvent& PropertyChangedEvent) override;
 #endif // WITH_EDITOR
-	// End UObject Interface
+	//~ End UObject Interface
 
-	// Begin AActor Interface
+	//~ Begin AActor Interface
 	virtual void PostUnregisterAllComponents() override;
 protected:
 	virtual void PostRegisterAllComponents() override;
 public:
-	// End AActor Interface
+	//~ End AActor Interface
 };
 
 
diff --git a/Engine/Source/Runtime/Engine/Classes/Sound/DialogueWave.h b/Engine/Source/Runtime/Engine/Classes/Sound/DialogueWave.h
--- a/Engine/Source/Runtime/Engine/Classes/Sound/DialogueWave.h
+++ b/Engine/Source/Runtime/Engine/Classes/Sound/DialogueWave.h
@@ -63,7 +63,7 @@ class UDialogueWave : public UObject
 	FGuid LocalizationGUID;
 
 public:
-	// Begin UObject interface. 
+	//~ Begin UObject Interface. 
 	virtual void Serialize( FArchive& Ar ) override;
 	virtual bool IsReadyForFinishDestroy() override;
 	virtual FString GetDesc() override;
@@ -76,13 +76,13 @@ public:
 	virtual void PostEditChangeChainProperty(FPropertyChangedChainEvent& PropertyChangedEvent) override;
 #endif
 
-	// End UObject interface. 
+	//~ End UObject Interface. 
 
-	// Begin UDialogueWave interface.
+	//~ Begin UDialogueWave Interface.
 	ENGINE_API bool SupportsContext(const FDialogueContext& Context) const;
 	ENGINE_API USoundBase* GetWaveFromContext(const FDialogueContext& Context) const;
 	ENGINE_API FString GetContextLocalizationKey( const FDialogueContext& Context ) const;
-	// End UDialogueWave interface.
+	//~ End UDialogueWave Interface.
 
 protected:
 	void UpdateMappingProxy(FDialogueContextMapping& ContextMapping);
diff --git a/Engine/Source/Runtime/Engine/Classes/Sound/SoundCue.h b/Engine/Source/Runtime/Engine/Classes/Sound/SoundCue.h
--- a/Engine/Source/Runtime/Engine/Classes/Sound/SoundCue.h
+++ b/Engine/Source/Runtime/Engine/Classes/Sound/SoundCue.h
@@ -72,7 +72,7 @@ private:
 
 public:
 
-	// Begin UObject interface.
+	//~ Begin UObject Interface.
 	virtual SIZE_T GetResourceSize(EResourceSizeMode::Type Mode) override;
 	virtual FString GetDesc() override;
 #if WITH_EDITOR
@@ -82,9 +82,9 @@ public:
 #endif
 	virtual void PostLoad() override;
 	virtual void Serialize(FArchive& Ar) override;
-	// End UObject interface.
+	//~ End UObject Interface.
 
-	// Begin USoundBase interface.
+	//~ Begin USoundBase Interface.
 	virtual bool IsPlayable() const override;
 	virtual void Parse( class FAudioDevice* AudioDevice, const UPTRINT NodeWaveInstanceHash, FActiveSound& ActiveSound, const FSoundParseParameters& ParseParams, TArray<FWaveInstance*>& WaveInstances ) override;
 	virtual float GetVolumeMultiplier() override;
@@ -92,7 +92,7 @@ public:
 	virtual float GetMaxAudibleDistance() override;
 	virtual float GetDuration() override;
 	virtual const FAttenuationSettings* GetAttenuationSettingsToApply() const override;
-	// End USoundBase interface.
+	//~ End USoundBase Interface.
 
 	/** Construct and initialize a node within this Cue */
 	template<class T>
diff --git a/Engine/Source/Runtime/Engine/Classes/Sound/SoundNode.h b/Engine/Source/Runtime/Engine/Classes/Sound/SoundNode.h
--- a/Engine/Source/Runtime/Engine/Classes/Sound/SoundNode.h
+++ b/Engine/Source/Runtime/Engine/Classes/Sound/SoundNode.h
@@ -59,16 +59,16 @@ class ENGINE_API USoundNode : public UObject
 #endif
 
 public:
-	// Begin UObject Interface
+	//~ Begin UObject Interface
 #if WITH_EDITOR
 	virtual void PostEditChangeProperty(FPropertyChangedEvent& PropertyChangedEvent) override;
 	static void AddReferencedObjects(UObject* InThis, FReferenceCollector& Collector);
 #endif //WITH_EDITOR
 	virtual void Serialize(FArchive& Ar) override;
-	// End UObject Interface
+	//~ End UObject Interface
 
 	//
-	// USoundNode interface. 
+	//~ Begin USoundNode Interface. 
 	//
 
 	/**
diff --git a/Engine/Source/Runtime/Engine/Classes/Sound/SoundNodeAttenuation.h b/Engine/Source/Runtime/Engine/Classes/Sound/SoundNodeAttenuation.h
--- a/Engine/Source/Runtime/Engine/Classes/Sound/SoundNodeAttenuation.h
+++ b/Engine/Source/Runtime/Engine/Classes/Sound/SoundNodeAttenuation.h
@@ -27,10 +27,10 @@ class USoundNodeAttenuation : public USoundNode
 	uint32 bOverrideAttenuation:1;
 
 public:
-	// Begin USoundNode interface. 
+	//~ Begin USoundNode Interface. 
 	virtual void ParseNodes( class FAudioDevice* AudioDevice, const UPTRINT NodeWaveInstanceHash, FActiveSound& ActiveSound, const FSoundParseParameters& ParseParams, TArray<FWaveInstance*>& WaveInstances ) override;
 	virtual float MaxAudibleDistance( float CurrentMaxDistance ) override;
-	// End USoundNode interface. 
+	//~ End USoundNode Interface. 
 
 	ENGINE_API FAttenuationSettings* GetAttenuationSettingsToApply();
 };
diff --git a/Engine/Source/Runtime/Engine/Classes/Sound/SoundNodeBranch.h b/Engine/Source/Runtime/Engine/Classes/Sound/SoundNodeBranch.h
--- a/Engine/Source/Runtime/Engine/Classes/Sound/SoundNodeBranch.h
+++ b/Engine/Source/Runtime/Engine/Classes/Sound/SoundNodeBranch.h
@@ -27,7 +27,7 @@ private:
 	};
 
 public:
-	// Begin USoundNode interface.
+	//~ Begin USoundNode Interface.
 	virtual void ParseNodes( FAudioDevice* AudioDevice, const UPTRINT NodeWaveInstanceHash, FActiveSound& ActiveSound, const FSoundParseParameters& ParseParams, TArray<FWaveInstance*>& WaveInstances ) override;
 	virtual int32 GetMaxChildNodes( void ) const override 
 	{ 
@@ -47,7 +47,7 @@ public:
 	virtual FText GetInputPinName(int32 PinIndex) const override;
 	virtual FText GetTitle() const override;
 #endif //WITH_EDITOR
-	// End USoundNode interface.
+	//~ End USoundNode Interface.
 };
 
 
diff --git a/Engine/Source/Runtime/Engine/Classes/Sound/SoundNodeDialoguePlayer.h b/Engine/Source/Runtime/Engine/Classes/Sound/SoundNodeDialoguePlayer.h
--- a/Engine/Source/Runtime/Engine/Classes/Sound/SoundNodeDialoguePlayer.h
+++ b/Engine/Source/Runtime/Engine/Classes/Sound/SoundNodeDialoguePlayer.h
@@ -21,11 +21,11 @@ class USoundNodeDialoguePlayer : public USoundNode
 	uint32 bLooping:1;
 
 public:	
-	// Begin USoundNode Interface
+	//~ Begin USoundNode Interface
 	virtual int32 GetMaxChildNodes() const override;
 	virtual float GetDuration() override;
 	virtual void ParseNodes( FAudioDevice* AudioDevice, const UPTRINT NodeWaveInstanceHash, FActiveSound& ActiveSound, const FSoundParseParameters& ParseParams, TArray<FWaveInstance*>& WaveInstances ) override;
-	// End USoundNode Interface
+	//~ End USoundNode Interface
 
 };
 
diff --git a/Engine/Source/Runtime/Engine/Classes/Sound/SoundNodeEnveloper.h b/Engine/Source/Runtime/Engine/Classes/Sound/SoundNodeEnveloper.h
--- a/Engine/Source/Runtime/Engine/Classes/Sound/SoundNodeEnveloper.h
+++ b/Engine/Source/Runtime/Engine/Classes/Sound/SoundNodeEnveloper.h
@@ -68,19 +68,19 @@ class USoundNodeEnveloper : public USoundNode
 
 
 public:	
-	// Begin UObject interface. 
+	//~ Begin UObject Interface. 
 #if WITH_EDITOR
 	virtual void PostEditChangeProperty( struct FPropertyChangedEvent& PropertyChangedEvent) override;
 #endif // WITH_EDITOR
 	virtual void PostInitProperties() override;
 	virtual void Serialize(FArchive& Ar) override;
-	// End UObject interface. 
+	//~ End UObject Interface. 
 
 
-	// Begin USoundNode interface. 
+	//~ Begin USoundNode Interface. 
 	virtual void ParseNodes( FAudioDevice* AudioDevice, const UPTRINT NodeWaveInstanceHash, FActiveSound& ActiveSound, const FSoundParseParameters& ParseParams, TArray<FWaveInstance*>& WaveInstances ) override;
 	virtual float GetDuration( void ) override;
-	// End USoundNode interface. 
+	//~ End USoundNode Interface. 
 
 };
 
diff --git a/Engine/Source/Runtime/Engine/Classes/Sound/SoundNodeMature.h b/Engine/Source/Runtime/Engine/Classes/Sound/SoundNodeMature.h
--- a/Engine/Source/Runtime/Engine/Classes/Sound/SoundNodeMature.h
+++ b/Engine/Source/Runtime/Engine/Classes/Sound/SoundNodeMature.h
@@ -16,14 +16,14 @@ class USoundNodeMature : public USoundNode
 
 
 public:
-	// Begin UObject Interface
+	//~ Begin UObject Interface
 	virtual void PostLoad() override;
-	// End UObject Interface
+	//~ End UObject Interface
 
-	// Begin USoundNode interface.
+	//~ Begin USoundNode Interface.
 	virtual void ParseNodes( FAudioDevice* AudioDevice, const UPTRINT NodeWaveInstanceHash, FActiveSound& ActiveSound, const FSoundParseParameters& ParseParams, TArray<FWaveInstance*>& WaveInstances ) override;
 	virtual void CreateStartingConnectors( void ) override;
 	virtual int32 GetMaxChildNodes() const override;
-	// End USoundNode interface.
+	//~ End USoundNode Interface.
 };
 
diff --git a/Engine/Source/Runtime/Engine/Classes/Sound/SoundNodeQualityLevel.h b/Engine/Source/Runtime/Engine/Classes/Sound/SoundNodeQualityLevel.h
--- a/Engine/Source/Runtime/Engine/Classes/Sound/SoundNodeQualityLevel.h
+++ b/Engine/Source/Runtime/Engine/Classes/Sound/SoundNodeQualityLevel.h
@@ -17,19 +17,19 @@ class USoundNodeQualityLevel : public USoundNode
 public:
 
 #if WITH_EDITOR
-	// Begin UObject Interface
+	//~ Begin UObject Interface
 	virtual void PostLoad() override;
-	// End UObject Interface
+	//~ End UObject Interface
 #endif
 
-	// Begin USoundNode interface.
+	//~ Begin USoundNode Interface.
 	virtual void ParseNodes( FAudioDevice* AudioDevice, const UPTRINT NodeWaveInstanceHash, FActiveSound& ActiveSound, const FSoundParseParameters& ParseParams, TArray<FWaveInstance*>& WaveInstances ) override;
 	virtual int32 GetMaxChildNodes() const override;
 	virtual int32 GetMinChildNodes() const override;
 #if WITH_EDITOR
 	virtual FText GetInputPinName(int32 PinIndex) const override;
 #endif
-	// End USoundNode interface.
+	//~ End USoundNode Interface.
 
 #if WITH_EDITOR
 	void ReconcileNode(bool bReconstructNode);
diff --git a/Engine/Source/Runtime/Engine/Classes/Sound/SoundNodeWaveParam.h b/Engine/Source/Runtime/Engine/Classes/Sound/SoundNodeWaveParam.h
--- a/Engine/Source/Runtime/Engine/Classes/Sound/SoundNodeWaveParam.h
+++ b/Engine/Source/Runtime/Engine/Classes/Sound/SoundNodeWaveParam.h
@@ -18,9 +18,9 @@ class USoundNodeWaveParam : public USoundNode
 
 
 public:	
-	// Begin USoundNode Interface
+	//~ Begin USoundNode Interface
 	virtual float GetDuration( void ) override;
 	virtual void ParseNodes( FAudioDevice* AudioDevice, const UPTRINT NodeWaveInstanceHash, FActiveSound& ActiveSound, const FSoundParseParameters& ParseParams, TArray<FWaveInstance*>& WaveInstances ) override;
-	// End USoundNode Interface
+	//~ End USoundNode Interface
 };
 
diff --git a/Engine/Source/Runtime/Engine/Classes/Sound/SoundNodeWavePlayer.h b/Engine/Source/Runtime/Engine/Classes/Sound/SoundNodeWavePlayer.h
--- a/Engine/Source/Runtime/Engine/Classes/Sound/SoundNodeWavePlayer.h
+++ b/Engine/Source/Runtime/Engine/Classes/Sound/SoundNodeWavePlayer.h
@@ -30,24 +30,24 @@ public:
 	ENGINE_API USoundWave* GetSoundWave() const { return SoundWave; }
 	ENGINE_API void SetSoundWave(USoundWave* SoundWave);
 
-	// Begin UObject Interface
+	//~ Begin UObject Interface
 #if WITH_EDITOR
 	virtual void PostEditChangeProperty(FPropertyChangedEvent& PropertyChangedEvent) override;
 #endif
-	// End UObject Interface
+	//~ End UObject Interface
 
-	// Begin USoundNode Interface
+	//~ Begin USoundNode Interface
 	virtual int32 GetMaxChildNodes() const override;
 	virtual float GetDuration() override;
 	virtual void ParseNodes( FAudioDevice* AudioDevice, const UPTRINT NodeWaveInstanceHash, FActiveSound& ActiveSound, const FSoundParseParameters& ParseParams, TArray<FWaveInstance*>& WaveInstances ) override;
 #if WITH_EDITOR
 	virtual FText GetTitle() const override;
 #endif
-	// End USoundNode Interface
+	//~ End USoundNode Interface
 
-	// Begin USoundNodeAssetReferencer Interface
+	//~ Begin USoundNodeAssetReferencer Interface
 	virtual void LoadAsset() override;
-	// End USoundNode Interface
+	//~ End USoundNode Interface
 
 };
 
diff --git a/Engine/Source/Runtime/Engine/Classes/Sound/SoundWaveProcedural.h b/Engine/Source/Runtime/Engine/Classes/Sound/SoundWaveProcedural.h
--- a/Engine/Source/Runtime/Engine/Classes/Sound/SoundWaveProcedural.h
+++ b/Engine/Source/Runtime/Engine/Classes/Sound/SoundWaveProcedural.h
@@ -27,19 +27,19 @@ private:
 
 public:
 
-	// Begin UObject interface. 
+	//~ Begin UObject Interface. 
 	virtual void Serialize( FArchive& Ar ) override;
 	virtual void GetAssetRegistryTags(TArray<FAssetRegistryTag>& OutTags) const override;
 	virtual SIZE_T GetResourceSize(EResourceSizeMode::Type Mode) override;
-	// End UObject interface. 
+	//~ End UObject Interface. 
 
-	// Begin USoundWave interface.
+	//~ Begin USoundWave Interface.
 	virtual int32 GeneratePCMData(uint8* PCMData, const int32 SamplesNeeded) override;
 	virtual FByteBulkData* GetCompressedData(FName Format) override;
 	virtual void InitAudioResource( FByteBulkData& CompressedData ) override;
 	virtual bool InitAudioResource(FName Format) override;
 	virtual int32 GetResourceSizeForFormat(FName Format) override;
-	// End USoundWave interface.
+	//~ End USoundWave Interface.
 
 
 	/** Add data to the FIFO that feeds the audio device. */
diff --git a/Engine/Source/Runtime/Engine/Classes/VectorField/VectorFieldAnimated.h b/Engine/Source/Runtime/Engine/Classes/VectorField/VectorFieldAnimated.h
--- a/Engine/Source/Runtime/Engine/Classes/VectorField/VectorFieldAnimated.h
+++ b/Engine/Source/Runtime/Engine/Classes/VectorField/VectorFieldAnimated.h
@@ -77,15 +77,15 @@ class UVectorFieldAnimated : public UVectorField
 	float NoiseMax;
 
 
-	// Begin UObject interface.
+	//~ Begin UObject Interface.
 #if WITH_EDITOR
 	virtual void PostEditChangeProperty(FPropertyChangedEvent& PropertyChangedEvent) override;
 #endif // WITH_EDITOR
-	// End UObject interface
+	//~ End UObject Interface
 
-	// Begin UVectorField Interface
+	//~ Begin UVectorField Interface
 	virtual void InitInstance(class FVectorFieldInstance* Instance, bool bPreviewInstance) override;
-	// End UVectorField Interface
+	//~ End UVectorField Interface
 
 };
 
diff --git a/Engine/Source/Runtime/Engine/Classes/VectorField/VectorFieldStatic.h b/Engine/Source/Runtime/Engine/Classes/VectorField/VectorFieldStatic.h
--- a/Engine/Source/Runtime/Engine/Classes/VectorField/VectorFieldStatic.h
+++ b/Engine/Source/Runtime/Engine/Classes/VectorField/VectorFieldStatic.h
@@ -41,7 +41,7 @@ public:
 	class UAssetImportData* AssetImportData;
 #endif // WITH_EDITORONLY_DATA
 
-	// Begin UObject interface.
+	//~ Begin UObject Interface.
 	virtual void PostLoad() override;
 	virtual void BeginDestroy() override;
 #if WITH_EDITOR
@@ -52,11 +52,11 @@ public:
 	virtual void GetAssetRegistryTags(TArray<FAssetRegistryTag>& OutTags) const override;
 	virtual void PostInitProperties() override;
 #endif
-	// End UObject interface.
+	//~ End UObject Interface.
 
-	// Begin UVectorField Interface
+	//~ Begin UVectorField Interface
 	virtual void InitInstance(class FVectorFieldInstance* Instance, bool bPreviewInstance) override;
-	// End UVectorField Interface
+	//~ End UVectorField Interface
 
 	/**
 	 * Initialize resources.
diff --git a/Engine/Source/Runtime/Engine/Private/SkeletalRenderCPUSkin.h b/Engine/Source/Runtime/Engine/Private/SkeletalRenderCPUSkin.h
--- a/Engine/Source/Runtime/Engine/Private/SkeletalRenderCPUSkin.h
+++ b/Engine/Source/Runtime/Engine/Private/SkeletalRenderCPUSkin.h
@@ -142,7 +142,7 @@ public:
 	FSkeletalMeshObjectCPUSkin(USkinnedMeshComponent* InMeshComponent, FSkeletalMeshResource* InSkeletalMeshResource, ERHIFeatureLevel::Type InFeatureLevel);
 	virtual ~FSkeletalMeshObjectCPUSkin();
 
-	// Begin FSkeletalMeshObject interface
+	//~ Begin FSkeletalMeshObject Interface
 	virtual void InitResources() override;
 	virtual void ReleaseResources() override;
 	virtual void Update(int32 LODIndex,USkinnedMeshComponent* InMeshComponent,const TArray<FActiveVertexAnim>& ActiveVertexAnims) override;
@@ -194,7 +194,7 @@ public:
 	}
 
 	virtual void DrawVertexElements(FPrimitiveDrawInterface* PDI, const FTransform& ToWorldSpace, bool bDrawNormals, bool bDrawTangents, bool bDrawBinormals) const override;
-	// End FSkeletalMeshObject interface
+	//~ End FSkeletalMeshObject Interface
 
 private:
 	/** vertex data for rendering a single LOD */
diff --git a/Engine/Source/Runtime/Engine/Private/SkeletalRenderGPUSkin.h b/Engine/Source/Runtime/Engine/Private/SkeletalRenderGPUSkin.h
--- a/Engine/Source/Runtime/Engine/Private/SkeletalRenderGPUSkin.h
+++ b/Engine/Source/Runtime/Engine/Private/SkeletalRenderGPUSkin.h
@@ -295,7 +295,7 @@ public:
 	FSkeletalMeshObjectGPUSkin(USkinnedMeshComponent* InMeshComponent, FSkeletalMeshResource* InSkeletalMeshResource, ERHIFeatureLevel::Type InFeatureLevel);
 	virtual ~FSkeletalMeshObjectGPUSkin();
 
-	// Begin FSkeletalMeshObject interface
+	//~ Begin FSkeletalMeshObject Interface
 	virtual void InitResources() override;
 	virtual void ReleaseResources() override;
 	virtual void Update(int32 LODIndex,USkinnedMeshComponent* InMeshComponent,const TArray<FActiveVertexAnim>& ActiveVertexAnims) override;
@@ -343,7 +343,7 @@ public:
 
 		return ResourceSize;
 	}
-	// End FSkeletalMeshObject interface
+	//~ End FSkeletalMeshObject Interface
 
 	/** 
 	 * Vertex buffers that can be used for GPU skinning factories 
diff --git a/Engine/Source/Runtime/Engine/Public/UnrealClient.h b/Engine/Source/Runtime/Engine/Public/UnrealClient.h
--- a/Engine/Source/Runtime/Engine/Public/UnrealClient.h
+++ b/Engine/Source/Runtime/Engine/Public/UnrealClient.h
@@ -249,7 +249,7 @@ public:
 	// Destructor
 	virtual ~FViewport(){}
 
-	// FViewport interface.
+	//~ Begin FViewport Interface.
 	virtual void* GetWindow() = 0;
 	virtual void MoveWindow(int32 NewPosX, int32 NewPosY, int32 NewSizeX, int32 NewSizeY) = 0;
 
@@ -419,7 +419,7 @@ public:
 	 **/
 	ENGINE_API virtual void SetViewportClient( FViewportClient* InViewportClient );
 
-	// FRenderTarget interface.
+	//~ Begin FRenderTarget Interface.
 	virtual FIntPoint GetSizeXY() const override { return FIntPoint(SizeX, SizeY); }
 
 	// Accessors.
@@ -545,10 +545,10 @@ protected:
 		/** Invalidates the cached hit proxy map. */
 		void Invalidate();
 
-		// FHitProxyConsumer interface.
+		//~ Begin FHitProxyConsumer Interface.
 		virtual void AddHitProxy(HHitProxy* HitProxy) override;
 
-		// FRenderTarget interface.
+		//~ Begin FRenderTarget Interface.
 		virtual FIntPoint GetSizeXY() const override { return FIntPoint(SizeX, SizeY); }
 
 		/** FGCObject interface */
@@ -616,7 +616,7 @@ protected:
 
 	/** Triggers the taking of a high res screen shot for this viewport. */
 	bool bTakeHighResScreenShot;
-	// FRenderResource interface.
+	//~ Begin FRenderResource Interface.
 	ENGINE_API virtual void InitDynamicRHI() override;
 	ENGINE_API virtual void ReleaseDynamicRHI() override;
 	ENGINE_API virtual void InitRHI() override;
@@ -945,7 +945,7 @@ public:
 
 	virtual ~FDummyViewport();
 
-	// Begin FViewport interface
+	//~ Begin FViewport Interface
 	virtual void BeginRenderFrame(FRHICommandListImmediate& RHICmdList) override
 	{
 		check( IsInRenderingThread() );
@@ -973,9 +973,9 @@ public:
 	virtual void DeferInvalidateHitProxy() override { }
 	virtual FViewportFrame* GetViewportFrame() override { return 0; }
 	virtual FCanvas* GetDebugCanvas() override { return DebugCanvas; }
-	// End FViewport interface
+	//~ End FViewport Interface
 
-	// Begin FRenderResource interface
+	//~ Begin FRenderResource Interface
 	virtual void InitDynamicRHI() override
 	{
 		FTexture2DRHIRef ShaderResourceTextureRHI;
@@ -990,7 +990,7 @@ public:
 	virtual void InitResource() override{ FViewport::InitResource(); }
 	virtual void ReleaseResource() override { FViewport::ReleaseResource(); }
 	virtual FString GetFriendlyName() const override { return FString(TEXT("FDummyViewport"));}
-	// End FRenderResource interface
+	//~ End FRenderResource Interface
 private:
 	FCanvas* DebugCanvas;
 };
diff --git a/Engine/Source/Runtime/Foliage/Private/InteractiveFoliageComponent.h b/Engine/Source/Runtime/Foliage/Private/InteractiveFoliageComponent.h
--- a/Engine/Source/Runtime/Foliage/Private/InteractiveFoliageComponent.h
+++ b/Engine/Source/Runtime/Foliage/Private/InteractiveFoliageComponent.h
@@ -13,13 +13,13 @@ class UInteractiveFoliageComponent : public UStaticMeshComponent
 public:
 	class FInteractiveFoliageSceneProxy* FoliageSceneProxy;
 
-	// Begin UPrimitiveComponent Interface
+	//~ Begin UPrimitiveComponent Interface
 	virtual FPrimitiveSceneProxy* CreateSceneProxy() override;
-	// End UPrimitiveComponent Interface
+	//~ End UPrimitiveComponent Interface
 
-	// Begin UActorComponent Interface
+	//~ Begin UActorComponent Interface
 	virtual void DestroyRenderState_Concurrent() override;
-	// End UActorComponent Interface
+	//~ End UActorComponent Interface
 
 	friend class AInteractiveFoliageActor;
 };
diff --git a/Engine/Source/Runtime/Foliage/Public/InstancedFoliageActor.h b/Engine/Source/Runtime/Foliage/Public/InstancedFoliageActor.h
--- a/Engine/Source/Runtime/Foliage/Public/InstancedFoliageActor.h
+++ b/Engine/Source/Runtime/Foliage/Public/InstancedFoliageActor.h
@@ -37,13 +37,13 @@ public:
 	TMap<UFoliageType*, TUniqueObj<FFoliageMeshInfo>> FoliageMeshes;
 
 public:
-	// Begin UObject interface.
+	//~ Begin UObject Interface.
 	virtual void Serialize(FArchive& Ar) override;
 	virtual void PostLoad() override;
 	static void AddReferencedObjects(UObject* InThis, FReferenceCollector& Collector);
-	// End UObject interface. 
+	//~ End UObject Interface. 
 
-	// Begin AActor interface.
+	//~ Begin AActor Interface.
 	// we don't want to have our components automatically destroyed by the Blueprint code
 	virtual void RerunConstructionScripts() override {}
 	virtual bool IsLevelBoundsRelevant() const override { return false; }
@@ -57,7 +57,7 @@ public:
 	virtual void Destroyed() override;
 	FOLIAGE_API void CleanupDeletedFoliageType();
 #endif
-	// End AActor interface.
+	//~ End AActor Interface.
 
 
 	// Performs a reverse lookup from a mesh to a local foliage type (i.e. the foliage type owned exclusively by this IFA)
diff --git a/Engine/Source/Runtime/Foliage/Public/InteractiveFoliageActor.h b/Engine/Source/Runtime/Foliage/Public/InteractiveFoliageActor.h
--- a/Engine/Source/Runtime/Foliage/Public/InteractiveFoliageActor.h
+++ b/Engine/Source/Runtime/Foliage/Public/InteractiveFoliageActor.h
@@ -1,7 +1,7 @@
 // Copyright 1998-2015 Epic Games, Inc. All Rights Reserved.
 
-//=============================================================================
-//=============================================================================
+//~=============================================================================
+//~=============================================================================
 
 #pragma once
 #include "Engine/StaticMeshActor.h"
@@ -87,15 +87,15 @@ protected:
 	UFUNCTION()
 	void CapsuleTouched(AActor* Other, UPrimitiveComponent* OtherComp, int32 OtherBodyIndex, bool bFromSweep, const FHitResult& OverlapInfo);
 public:
-	// Begin AActor interface
+	//~ Begin AActor Interface
 	virtual void Tick(float DeltaSeconds) override;
 	virtual void PostActorCreated() override;
 	virtual float TakeDamage(float DamageAmount, struct FDamageEvent const& DamageEvent, class AController* EventInstigator, class AActor* DamageCauser) override;
-	// End AActor interface
+	//~ End AActor Interface
 
-	// Begin UObject interface
+	//~ Begin UObject Interface
 	virtual void PostLoad() override;
-	// End UObject interface
+	//~ End UObject Interface
 
 };
 
diff --git a/Engine/Source/Runtime/GameMenuBuilder/Public/SGameMenuItemWidget.h b/Engine/Source/Runtime/GameMenuBuilder/Public/SGameMenuItemWidget.h
--- a/Engine/Source/Runtime/GameMenuBuilder/Public/SGameMenuItemWidget.h
+++ b/Engine/Source/Runtime/GameMenuBuilder/Public/SGameMenuItemWidget.h
@@ -41,7 +41,7 @@ public:
 	/** needed for every widget */
 	void Construct(const FArguments& InArgs);
 
-	// Begin SWidget interface
+	//~ Begin SWidget Interface
 	virtual bool SupportsKeyboardFocus() const override { return true; }
 	
 	virtual FReply OnMouseButtonDown(const FGeometry& MyGeometry, const FPointerEvent& MouseEvent) override;
@@ -49,7 +49,7 @@ public:
 	virtual FReply OnMouseButtonUp(const FGeometry& MyGeometry, const FPointerEvent& MouseEvent) override;
 
 	virtual FReply OnMouseMove(const FGeometry& MyGeometry, const FPointerEvent& MouseEvent) override;
-	// End SWidget interface
+	//~ End SWidget Interface
 
 	/** Sets this menu item as active (selected) */
 	void SetMenuItemActive(bool bIsMenuItemActive);
diff --git a/Engine/Source/Runtime/GeometryCache/Classes/GeometryCacheComponent.h b/Engine/Source/Runtime/GeometryCache/Classes/GeometryCacheComponent.h
--- a/Engine/Source/Runtime/GeometryCache/Classes/GeometryCacheComponent.h
+++ b/Engine/Source/Runtime/GeometryCache/Classes/GeometryCacheComponent.h
@@ -42,34 +42,34 @@ class GEOMETRYCACHE_API UGeometryCacheComponent : public UMeshComponent
 	/** Required for access to (protected) TrackSections */
 	friend FGeometryCacheSceneProxy;
 		
-	// Begin UObject Interface
+	//~ Begin UObject Interface
 	virtual void BeginDestroy() override;
 #if WITH_EDITOR
 	virtual void PreEditUndo() override;
 	virtual void PostEditUndo() override;
 	virtual void PostEditChangeProperty(FPropertyChangedEvent& PropertyChangedEvent) override;
 #endif // WITH_EDITOR
-	// End UObject Interface
+	//~ End UObject Interface
 
-	// Begin UActorComponent Interface.
+	//~ Begin UActorComponent Interface.
 	virtual void OnRegister() override;
 	virtual void OnUnregister() override;
 	virtual void TickComponent(float DeltaTime, enum ELevelTick TickType, FActorComponentTickFunction *ThisTickFunction) override;
-	// End UActorComponent Interface.
+	//~ End UActorComponent Interface.
 
-	// Begin USceneComponent Interface.
+	//~ Begin USceneComponent Interface.
 	virtual FBoxSphereBounds CalcBounds(const FTransform& LocalToWorld) const override;
 	/** Update LocalBounds member from the local box of each section */
 	void UpdateLocalBounds();
-	// Begin USceneComponent Interface.	
+	//~ Begin USceneComponent Interface.	
 
-	// Begin UPrimitiveComponent Interface.
+	//~ Begin UPrimitiveComponent Interface.
 	virtual FPrimitiveSceneProxy* CreateSceneProxy() override;
-	// End UPrimitiveComponent Interface.
+	//~ End UPrimitiveComponent Interface.
 
-	// Begin UMeshComponent Interface.
+	//~ Begin UMeshComponent Interface.
 	virtual int32 GetNumMaterials() const override;
-	// End UMeshComponent Interface.
+	//~ End UMeshComponent Interface.
 
 	/**
 	* CreateTrackSection, Create/replace a track section.
diff --git a/Engine/Source/Runtime/GeometryCache/Classes/GeometryCacheTrack.h b/Engine/Source/Runtime/GeometryCache/Classes/GeometryCacheTrack.h
--- a/Engine/Source/Runtime/GeometryCache/Classes/GeometryCacheTrack.h
+++ b/Engine/Source/Runtime/GeometryCache/Classes/GeometryCacheTrack.h
@@ -13,10 +13,10 @@ class UGeometryCacheTrack : public UObject
 
 	virtual ~UGeometryCacheTrack();
 	
-	// Begin UObject interface.
+	//~ Begin UObject Interface.
 	virtual void Serialize(FArchive& Ar) override;
 	virtual SIZE_T GetResourceSize(EResourceSizeMode::Type Mode) override;
-	// End UObject interface
+	//~ End UObject Interface
 			
 	/**
 	* UpdateMatrixData
diff --git a/Engine/Source/Runtime/GeometryCache/Classes/GeometryCacheTrackFlipbookAnimation.h b/Engine/Source/Runtime/GeometryCache/Classes/GeometryCacheTrackFlipbookAnimation.h
--- a/Engine/Source/Runtime/GeometryCache/Classes/GeometryCacheTrackFlipbookAnimation.h
+++ b/Engine/Source/Runtime/GeometryCache/Classes/GeometryCacheTrackFlipbookAnimation.h
@@ -13,16 +13,16 @@ class GEOMETRYCACHE_API UGeometryCacheTrack_FlipbookAnimation : public UGeometry
 
 	virtual ~UGeometryCacheTrack_FlipbookAnimation();
 
-	// Begin UObject interface.
+	//~ Begin UObject Interface.
 	virtual SIZE_T GetResourceSize(EResourceSizeMode::Type Mode) override;
 	virtual void Serialize(FArchive& Ar) override;
 	virtual void BeginDestroy() override;
-	// End UObject interface.
+	//~ End UObject Interface.
 
-	// Begin UGeometryCacheTrack interface.
+	//~ Begin UGeometryCacheTrack Interface.
 	virtual const bool UpdateMeshData(const float Time, const bool bLooping, int32& InOutMeshSampleIndex, FGeometryCacheMeshData*& OutMeshData) override;
 	virtual const float GetMaxSampleTime() const override;
-	// End UGeometryCacheTrack interface.
+	//~ End UGeometryCacheTrack Interface.
 
 	/**
 	* Add a GeometryCacheMeshData sample to the Track
diff --git a/Engine/Source/Runtime/GeometryCache/Classes/GeometryCacheTrackTransformAnimation.h b/Engine/Source/Runtime/GeometryCache/Classes/GeometryCacheTrackTransformAnimation.h
--- a/Engine/Source/Runtime/GeometryCache/Classes/GeometryCacheTrackTransformAnimation.h
+++ b/Engine/Source/Runtime/GeometryCache/Classes/GeometryCacheTrackTransformAnimation.h
@@ -11,14 +11,14 @@ class GEOMETRYCACHE_API UGeometryCacheTrack_TransformAnimation : public UGeometr
 {
 	GENERATED_UCLASS_BODY()
 		
-	// Begin UObject interface.
+	//~ Begin UObject Interface.
 	virtual SIZE_T GetResourceSize(EResourceSizeMode::Type Mode) override;
 	virtual void Serialize(FArchive& Ar) override;
-	// End UObject interface.
+	//~ End UObject Interface.
 	
-	// Begin UGeometryCacheTrack interface.
+	//~ Begin UGeometryCacheTrack Interface.
 	virtual const bool UpdateMeshData(const float Time, const bool bLooping, int32& InOutMeshSampleIndex, FGeometryCacheMeshData*& OutMeshData) override;
-	// End UGeometryCacheTrack interface.
+	//~ End UGeometryCacheTrack Interface.
 
 	/**
 	* Sets/updates the MeshData for this track
diff --git a/Engine/Source/Runtime/Landscape/Classes/Landscape.h b/Engine/Source/Runtime/Landscape/Classes/Landscape.h
--- a/Engine/Source/Runtime/Landscape/Classes/Landscape.h
+++ b/Engine/Source/Runtime/Landscape/Classes/Landscape.h
@@ -35,18 +35,18 @@ class ALandscape : public ALandscapeProxy
 	static FIntPoint MakeKey( int32 X, int32 Y ) { return FIntPoint(X, Y); }
 	static void UnpackKey( FIntPoint Key, int32& OutX, int32& OutY ) { OutX = Key.X; OutY = Key.Y; }
 
-	// Begin AActor Interface
+	//~ Begin AActor Interface
 #if WITH_EDITOR
 	virtual void CheckForErrors() override;
 #endif
-	// End AActor Interface
+	//~ End AActor Interface
 
-	// Begin ALandscapeProxy Interface
+	//~ Begin ALandscapeProxy Interface
 	LANDSCAPE_API virtual ALandscape* GetLandscapeActor() override;
 #if WITH_EDITOR
 	virtual UMaterialInterface* GetLandscapeMaterial() const override;
 	virtual UMaterialInterface* GetLandscapeHoleMaterial() const override;
-	// End ALandscapeProxy Interface
+	//~ End ALandscapeProxy Interface
 
 	LANDSCAPE_API bool HasAllComponent(); // determine all component is in this actor
 	
@@ -62,13 +62,13 @@ class ALandscape : public ALandscapeProxy
 
 	static void SplitHeightmap(ULandscapeComponent* Comp, bool bMoveToCurrentLevel = false);
 	
-	// Begin UObject interface.
+	//~ Begin UObject Interface.
 	virtual void PreSave() override;
 	virtual void PostEditChangeProperty(FPropertyChangedEvent& PropertyChangedEvent) override;
 	virtual void PostEditMove(bool bFinished) override;
 #endif
 	virtual void PostLoad() override;
-	// End UObject Interface
+	//~ End UObject Interface
 };
 
 LANDSCAPE_API DECLARE_LOG_CATEGORY_EXTERN(LogLandscape, Warning, All);
diff --git a/Engine/Source/Runtime/Landscape/Classes/LandscapeComponent.h b/Engine/Source/Runtime/Landscape/Classes/LandscapeComponent.h
--- a/Engine/Source/Runtime/Landscape/Classes/LandscapeComponent.h
+++ b/Engine/Source/Runtime/Landscape/Classes/LandscapeComponent.h
@@ -290,7 +290,7 @@ public:
 
 	virtual ~ULandscapeComponent();
 
-	// Begin UObject interface.	
+	//~ Begin UObject Interface.	
 	virtual void PostInitProperties() override;	
 	virtual void Serialize(FArchive& Ar) override;
 	virtual SIZE_T GetResourceSize(EResourceSizeMode::Type Mode) override;
@@ -303,13 +303,13 @@ public:
 	virtual void PostEditUndo() override;
 	virtual void PreEditChange(UProperty* PropertyThatWillChange) override;
 	virtual void PostEditChangeProperty(FPropertyChangedEvent& PropertyChangedEvent) override;
-	// End UObject interface
+	//~ End UObject Interface
 
 	/** Fix up component layers, weightmaps
 	 */
 	LANDSCAPE_API void FixupWeightmaps();
 	
-	// Begin UPrimitiveComponent interface.
+	//~ Begin UPrimitiveComponent Interface.
 	virtual bool GetLightMapResolution( int32& Width, int32& Height ) const override;
 	virtual void GetLightAndShadowMapMemoryUsage( int32& LightMapMemoryUsage, int32& ShadowMapMemoryUsage ) const override;
 	virtual void GetStaticLightingInfo(FStaticLightingPrimitiveInfo& OutPrimitiveInfo,const TArray<ULightComponent*>& InRelevantLights,const FLightingBuildOptions& Options) override;
@@ -326,17 +326,17 @@ public:
 	virtual bool ComponentIsTouchingSelectionBox(const FBox& InSelBBox, const FEngineShowFlags& ShowFlags, const bool bConsiderOnlyBSP, const bool bMustEncompassEntireComponent) const override;
 	virtual bool ComponentIsTouchingSelectionFrustum(const FConvexVolume& InFrustum, const FEngineShowFlags& ShowFlags, const bool bConsiderOnlyBSP, const bool bMustEncompassEntireComponent) const override;
 #endif
-	// End UPrimitiveComponent interface.
+	//~ End UPrimitiveComponent Interface.
 
-	// Begin USceneComponent interface.
+	//~ Begin USceneComponent Interface.
 	virtual void DestroyComponent(bool bPromoteChildren = false) override;
 	virtual FBoxSphereBounds CalcBounds(const FTransform& LocalToWorld) const override;
-	// End USceneComponent interface.
+	//~ End USceneComponent Interface.
 
-	// Begin UActorComponent interface.
+	//~ Begin UActorComponent Interface.
 	virtual void OnRegister() override;
 	virtual void OnUnregister() override;
-	// End UActorComponent interface.
+	//~ End UActorComponent Interface.
 
 
 #if WITH_EDITOR
diff --git a/Engine/Source/Runtime/Landscape/Classes/LandscapeGizmoRenderComponent.h b/Engine/Source/Runtime/Landscape/Classes/LandscapeGizmoRenderComponent.h
--- a/Engine/Source/Runtime/Landscape/Classes/LandscapeGizmoRenderComponent.h
+++ b/Engine/Source/Runtime/Landscape/Classes/LandscapeGizmoRenderComponent.h
@@ -10,13 +10,13 @@ class ULandscapeGizmoRenderComponent : public UPrimitiveComponent
 	GENERATED_UCLASS_BODY()
 
 
-	// Begin UPrimitiveComponent Interface
+	//~ Begin UPrimitiveComponent Interface
 	virtual FPrimitiveSceneProxy* CreateSceneProxy() override;
-	// End UPrimitiveComponent Interface
+	//~ End UPrimitiveComponent Interface
 
-	// Begin USceneComponent interface.
+	//~ Begin USceneComponent Interface.
 	virtual FBoxSphereBounds CalcBounds(const FTransform& LocalToWorld) const override;
-	// End USceneComponent interface.
+	//~ End USceneComponent Interface.
 };
 
 
diff --git a/Engine/Source/Runtime/Landscape/Classes/LandscapeHeightfieldCollisionComponent.h b/Engine/Source/Runtime/Landscape/Classes/LandscapeHeightfieldCollisionComponent.h
--- a/Engine/Source/Runtime/Landscape/Classes/LandscapeHeightfieldCollisionComponent.h
+++ b/Engine/Source/Runtime/Landscape/Classes/LandscapeHeightfieldCollisionComponent.h
@@ -141,12 +141,12 @@ class ULandscapeHeightfieldCollisionComponent : public UPrimitiveComponent
 		QF_NoCollision = 128,			// This quad has no collision.
 	};
 
-	// Begin UActorComponent interface.
+	//~ Begin UActorComponent Interface.
 	virtual void CreatePhysicsState() override;
 	virtual void ApplyWorldOffset(const FVector& InOffset, bool bWorldShift) override;
-	// End UActorComponent interface.
+	//~ End UActorComponent Interface.
 
-	// Begin USceneComponent interface.
+	//~ Begin USceneComponent Interface.
 	virtual void DestroyComponent(bool bPromoteChildren = false) override;
 	virtual FBoxSphereBounds CalcBounds(const FTransform &BoundTransform) const override;
 
@@ -154,9 +154,9 @@ class ULandscapeHeightfieldCollisionComponent : public UPrimitiveComponent
 	virtual ECollisionResponse GetCollisionResponseToChannel(ECollisionChannel Channel) const override;
 	virtual ECollisionChannel GetCollisionObjectType() const override;
 	virtual const FCollisionResponseContainer& GetCollisionResponseToChannels() const override;
-	// End USceneComponent interface.
+	//~ End USceneComponent Interface.
 
-	// Begin UPrimitiveComponent interface
+	//~ Begin UPrimitiveComponent Interface
 	virtual bool DoCustomNavigableGeometryExport(FNavigableGeometryExport& GeomExport) const override;
 #if WITH_EDITOR
 	virtual bool ComponentIsTouchingSelectionBox(const FBox& InSelBBox, const FEngineShowFlags& ShowFlags, const bool bConsiderOnlyBSP, const bool bMustEncompassEntireComponent) const override;
@@ -164,14 +164,14 @@ class ULandscapeHeightfieldCollisionComponent : public UPrimitiveComponent
 #endif
 	//End UPrimitiveComponent interface
 
-	// Begin INavRelevantInterface interface
+	//~ Begin INavRelevantInterface Interface
 	virtual bool SupportsGatheringGeometrySlices() const override { return true; }
 	virtual void GatherGeometrySlice(FNavigableGeometryExport& GeomExport, const FBox& SliceBox) const override;
 	virtual ENavDataGatheringMode GetGeometryGatheringMode() const override;
 	virtual void PrepareGeometryExportSync() override;
-	// End INavRelevantInterface interface
+	//~ End INavRelevantInterface Interface
 
-	// Begin UObject Interface.
+	//~ Begin UObject Interface.
 	virtual void Serialize(FArchive& Ar) override;
 	virtual void BeginDestroy() override;
 	virtual void PostLoad() override;
@@ -181,7 +181,7 @@ class ULandscapeHeightfieldCollisionComponent : public UPrimitiveComponent
 	virtual void ImportCustomProperties(const TCHAR* SourceText, FFeedbackContext* Warn) override;
 	virtual void PostEditImport() override;
 	virtual void PostEditUndo() override;
-	// End UObject Interface.
+	//~ End UObject Interface.
 
 	// @todo document
 	class ULandscapeInfo* GetLandscapeInfo(bool bSpawnNewActor = true) const;
diff --git a/Engine/Source/Runtime/Landscape/Classes/LandscapeInfo.h b/Engine/Source/Runtime/Landscape/Classes/LandscapeInfo.h
--- a/Engine/Source/Runtime/Landscape/Classes/LandscapeInfo.h
+++ b/Engine/Source/Runtime/Landscape/Classes/LandscapeInfo.h
@@ -139,12 +139,12 @@ public:
 
 	TMap<FIntPoint,float> SelectedRegion;
 
-	// Begin UObject Interface.
+	//~ Begin UObject Interface.
 	virtual void Serialize(FArchive& Ar) override;
 	virtual void BeginDestroy() override;
 #if WITH_EDITOR
 	virtual void PostEditUndo() override;
-	// End UObject Interface
+	//~ End UObject Interface
 
 	// @todo document 
 	// all below.
diff --git a/Engine/Source/Runtime/Landscape/Classes/LandscapeMeshCollisionComponent.h b/Engine/Source/Runtime/Landscape/Classes/LandscapeMeshCollisionComponent.h
--- a/Engine/Source/Runtime/Landscape/Classes/LandscapeMeshCollisionComponent.h
+++ b/Engine/Source/Runtime/Landscape/Classes/LandscapeMeshCollisionComponent.h
@@ -68,20 +68,20 @@ class ULandscapeMeshCollisionComponent : public ULandscapeHeightfieldCollisionCo
 	/** Physics engine version of heightfield data. */
 	TRefCountPtr<struct FPhysXMeshRef>			MeshRef;
 
-	// Begin UActorComponent interface.
+	//~ Begin UActorComponent Interface.
 	virtual void CreatePhysicsState() override;
 	virtual void ApplyWorldOffset(const FVector& InOffset, bool bWorldShift) override;
-	// End UActorComponent interface.
+	//~ End UActorComponent Interface.
 
-	// Begin USceneComponent interface.
+	//~ Begin USceneComponent Interface.
 	virtual void DestroyComponent(bool bPromoteChildren = false) override;
-	// End USceneComponent interface.
+	//~ End USceneComponent Interface.
 
-	// Begin UPrimitiveComponent interface
+	//~ Begin UPrimitiveComponent Interface
 	virtual bool DoCustomNavigableGeometryExport(FNavigableGeometryExport& GeomExport) const override;
 	//End UPrimitiveComponent interface
 
-	// Begin UObject Interface.
+	//~ Begin UObject Interface.
 	virtual void Serialize(FArchive& Ar) override;
 	virtual void BeginDestroy() override;
 #if WITH_EDITOR
@@ -90,11 +90,11 @@ class ULandscapeMeshCollisionComponent : public ULandscapeHeightfieldCollisionCo
 
 	virtual bool CookCollisionData(const FName& Format, bool bUseOnlyDefMaterial, bool bCheckDDC, TArray<uint8>& OutCookedData, TArray<UPhysicalMaterial*>& InOutMaterials) const override;
 #endif
-	// End UObject Interface.
+	//~ End UObject Interface.
 
-	// Begin ULandscapeHeightfieldCollisionComponent Interface
+	//~ Begin ULandscapeHeightfieldCollisionComponent Interface
 	virtual void CreateCollisionObject() override;
 	virtual void RecreateCollision() override;
-	// End ULandscapeHeightfieldCollisionComponent Interface
+	//~ End ULandscapeHeightfieldCollisionComponent Interface
 };
 
diff --git a/Engine/Source/Runtime/Landscape/Classes/LandscapeProxy.h b/Engine/Source/Runtime/Landscape/Classes/LandscapeProxy.h
--- a/Engine/Source/Runtime/Landscape/Classes/LandscapeProxy.h
+++ b/Engine/Source/Runtime/Landscape/Classes/LandscapeProxy.h
@@ -503,7 +503,7 @@ public:
 
 	// End blueprint functions
 
-	// Begin AActor Interface
+	//~ Begin AActor Interface
 	virtual void UnregisterAllComponents() override;
 	virtual void RegisterAllComponents() override;
 	virtual void RerunConstructionScripts() override {}
@@ -516,7 +516,7 @@ public:
 	virtual void PostEditMove(bool bFinished) override;
 	virtual bool ShouldImport(FString* ActorPropString, bool IsMovingLevel) override;
 	virtual bool ShouldExport() override;
-	// End AActor Interface
+	//~ End AActor Interface
 #endif	//WITH_EDITOR
 
 	FGuid GetLandscapeGuid() const { return LandscapeGuid; }
@@ -553,7 +553,7 @@ public:
 	int32 UpdateBakedTexturesCountdown;
 #endif
 
-	// Begin FTickableGameObject interface.
+	//~ Begin FTickableGameObject Interface.
 	virtual void Tick(float DeltaTime) override;
 	virtual bool IsTickable() const override 
 	{ 
@@ -572,7 +572,7 @@ public:
 		return GetStatID();
 	}
 
-	// Begin UObject interface.
+	//~ Begin UObject Interface.
 	virtual void Serialize(FArchive& Ar) override;
 	static void AddReferencedObjects(UObject* InThis, FReferenceCollector& Collector);
 	virtual void PostLoad() override;
@@ -583,7 +583,7 @@ public:
 	virtual void PreEditUndo() override;
 	virtual void PostEditUndo() override;
 	virtual void PostEditImport() override;
-	// End UObject Interface
+	//~ End UObject Interface
 
 	LANDSCAPE_API static TArray<FName> GetLayersFromMaterial(UMaterialInterface* Material);
 	LANDSCAPE_API TArray<FName> GetLayersFromMaterial() const;
diff --git a/Engine/Source/Runtime/Landscape/Classes/LandscapeSplinesComponent.h b/Engine/Source/Runtime/Landscape/Classes/LandscapeSplinesComponent.h
--- a/Engine/Source/Runtime/Landscape/Classes/LandscapeSplinesComponent.h
+++ b/Engine/Source/Runtime/Landscape/Classes/LandscapeSplinesComponent.h
@@ -150,7 +150,7 @@ public:
 	void AutoFixMeshComponentErrors(UWorld* OtherWorld);
 #endif
 
-	// Begin UObject Interface
+	//~ Begin UObject Interface
 
 	virtual void Serialize(FArchive& Ar) override;
 	virtual void PostLoad() override;
@@ -158,21 +158,21 @@ public:
 	virtual void PostEditChangeProperty(FPropertyChangedEvent& PropertyChangedEvent) override;
 	virtual void PostEditUndo() override;
 #endif // WITH_EDITOR
-	// End UObject Interface
+	//~ End UObject Interface
 
-	// Begin UActorComponent interface
+	//~ Begin UActorComponent Interface
 #if WITH_EDITOR
 	virtual void CheckForErrors() override;
 #endif
 	virtual void OnRegister() override;
-	// End UActorComponent interface
+	//~ End UActorComponent Interface
 	
-	// Begin UPrimitiveComponent interface.
+	//~ Begin UPrimitiveComponent Interface.
 #if WITH_EDITOR
 	virtual FPrimitiveSceneProxy* CreateSceneProxy() override;
 #endif
 	virtual FBoxSphereBounds CalcBounds(const FTransform& LocalToWorld) const override;
-	// End UPrimitiveComponent interface.
+	//~ End UPrimitiveComponent Interface.
 
 	// many friends
 	friend class FLandscapeToolSplines;
diff --git a/Engine/Source/Runtime/Niagara/Public/NiagaraComponent.h b/Engine/Source/Runtime/Niagara/Public/NiagaraComponent.h
--- a/Engine/Source/Runtime/Niagara/Public/NiagaraComponent.h
+++ b/Engine/Source/Runtime/Niagara/Public/NiagaraComponent.h
@@ -24,7 +24,7 @@ private:
 
 	TSharedPtr<FNiagaraEffectInstance> EffectInstance;
 
-	// Begin UActorComponent interface.
+	//~ Begin UActorComponent Interface.
 protected:
 	virtual void OnRegister() override;
 	virtual void OnUnregister()  override;
@@ -32,13 +32,13 @@ protected:
 public:
 	virtual void TickComponent(float DeltaTime, enum ELevelTick TickType, FActorComponentTickFunction* ThisTickFunction) override;
 	virtual const UObject* AdditionalStatObject() const override;
-	// End UActorComponent interface.
+	//~ End UActorComponent Interface.
 
-	// Begin UPrimitiveComponent Interface
+	//~ Begin UPrimitiveComponent Interface
 	virtual int32 GetNumMaterials() const override;
 	virtual FBoxSphereBounds CalcBounds(const FTransform& LocalToWorld) const override;
 	virtual FPrimitiveSceneProxy* CreateSceneProxy() override;
-	// End UPrimitiveComponent Interface
+	//~ End UPrimitiveComponent Interface
 
 	void SetAsset(UNiagaraEffect* InAsset);
 	UNiagaraEffect* GetAsset() const { return Asset; }
@@ -46,11 +46,11 @@ public:
 	TSharedPtr<FNiagaraEffectInstance> GetEffectInstance()	const { return EffectInstance; }
 	void SetEffectInstance(TSharedPtr<FNiagaraEffectInstance> InInstance)	{ EffectInstance = InInstance; }
 
-	// Begin UObject interface.
+	//~ Begin UObject Interface.
 #if WITH_EDITOR
 	virtual void PostEditChangeProperty(FPropertyChangedEvent& PropertyChangedEvent) override;
 #endif // WITH_EDITOR
-	// End UObject interface.
+	//~ End UObject Interface.
 
 	static const TArray<FNiagaraVariableInfo>& GetSystemConstants();
 };
@@ -79,7 +79,7 @@ public:
 private:
 	void ReleaseRenderThreadResources();
 
-	// FPrimitiveSceneProxy interface.
+	//~ Begin FPrimitiveSceneProxy Interface.
 	virtual void CreateRenderThreadResources() override;
 
 	virtual void OnActorPositionChanged() override;
diff --git a/Engine/Source/Runtime/Online/Android/OnlineSubsystemGooglePlay/Private/OnlineAchievementsInterfaceGooglePlay.h b/Engine/Source/Runtime/Online/Android/OnlineSubsystemGooglePlay/Private/OnlineAchievementsInterfaceGooglePlay.h
--- a/Engine/Source/Runtime/Online/Android/OnlineSubsystemGooglePlay/Private/OnlineAchievementsInterfaceGooglePlay.h
+++ b/Engine/Source/Runtime/Online/Android/OnlineSubsystemGooglePlay/Private/OnlineAchievementsInterfaceGooglePlay.h
@@ -68,7 +68,7 @@ PACKAGE_SCOPE:
 
 public:
 
-	// Begin IOnlineAchievements interface
+	//~ Begin IOnlineAchievements Interface
 	virtual void WriteAchievements(const FUniqueNetId& PlayerId, FOnlineAchievementsWriteRef& WriteObject, const FOnAchievementsWrittenDelegate& Delegate = FOnAchievementsWrittenDelegate()) override;
 	virtual void QueryAchievements(const FUniqueNetId& PlayerId, const FOnQueryAchievementsCompleteDelegate& Delegate = FOnQueryAchievementsCompleteDelegate()) override;
 	virtual void QueryAchievementDescriptions(const FUniqueNetId& PlayerId, const FOnQueryAchievementsCompleteDelegate& Delegate = FOnQueryAchievementsCompleteDelegate()) override ;
@@ -78,7 +78,7 @@ public:
 #if !UE_BUILD_SHIPPING
 	virtual bool ResetAchievements( const FUniqueNetId& PlayerId ) override;
 #endif // !UE_BUILD_SHIPPING
-	// End IOnlineAchievements interface
+	//~ End IOnlineAchievements Interface
 
 	/**
 	 * Constructor
diff --git a/Engine/Source/Runtime/Online/Android/OnlineSubsystemGooglePlay/Private/OnlineIdentityInterfaceGooglePlay.h b/Engine/Source/Runtime/Online/Android/OnlineSubsystemGooglePlay/Private/OnlineIdentityInterfaceGooglePlay.h
--- a/Engine/Source/Runtime/Online/Android/OnlineSubsystemGooglePlay/Private/OnlineIdentityInterfaceGooglePlay.h
+++ b/Engine/Source/Runtime/Online/Android/OnlineSubsystemGooglePlay/Private/OnlineIdentityInterfaceGooglePlay.h
@@ -47,7 +47,7 @@ PACKAGE_SCOPE:
 
 public:
 
-	// Begin IOnlineIdentity interface
+	//~ Begin IOnlineIdentity Interface
 	virtual TSharedPtr<FUserOnlineAccount> GetUserAccount(const FUniqueNetId& UserId) const override;
 	virtual TArray<TSharedPtr<FUserOnlineAccount> > GetAllUserAccounts() const override;
 	virtual bool Login(int32 LocalUserNum, const FOnlineAccountCredentials& AccountCredentials) override;
@@ -63,13 +63,13 @@ public:
 	virtual FString GetAuthToken(int32 LocalUserNum) const override;
 	virtual void GetUserPrivilege(const FUniqueNetId& UserId, EUserPrivileges::Type Privilege, const FOnGetUserPrivilegeCompleteDelegate& Delegate) override;
 	virtual FPlatformUserId GetPlatformUserIdFromUniqueNetId(const FUniqueNetId& NetId) override;
-	// End IOnlineIdentity interface
+	//~ End IOnlineIdentity Interface
 
 public:
 
-	// Begin IOnlineIdentity interface
+	//~ Begin IOnlineIdentity Interface
 	void Tick(float DeltaTime);
-	// End IOnlineIdentity interface
+	//~ End IOnlineIdentity Interface
 	
 	//platform specific
 	void OnLogin(const bool InLoggedIn,  const FString& InPlayerId, const FString& InPlayerAlias);
diff --git a/Engine/Source/Runtime/Online/Android/OnlineSubsystemGooglePlay/Public/OnlineSubsystemGooglePlay.h b/Engine/Source/Runtime/Online/Android/OnlineSubsystemGooglePlay/Public/OnlineSubsystemGooglePlay.h
--- a/Engine/Source/Runtime/Online/Android/OnlineSubsystemGooglePlay/Public/OnlineSubsystemGooglePlay.h
+++ b/Engine/Source/Runtime/Online/Android/OnlineSubsystemGooglePlay/Public/OnlineSubsystemGooglePlay.h
@@ -29,7 +29,7 @@ public:
 	FOnlineSubsystemGooglePlay();
 	virtual ~FOnlineSubsystemGooglePlay() {}
 
-	// Begin IOnlineSubsystem Interface
+	//~ Begin IOnlineSubsystem Interface
 	virtual IOnlineSessionPtr GetSessionInterface() const override;
 	virtual IOnlineFriendsPtr GetFriendsInterface() const override;
 	virtual IOnlinePartyPtr GetPartyInterface() const override;
@@ -66,7 +66,7 @@ public:
 	virtual bool Shutdown() override;
 	virtual FString GetAppId() const override;
 	virtual bool Exec(class UWorld* InWorld, const TCHAR* Cmd, FOutputDevice& Ar) override;
-	// End IOnlineSubsystem Interface
+	//~ End IOnlineSubsystem Interface
 
 	virtual bool Tick(float DeltaTime) override;
 	
diff --git a/Engine/Source/Runtime/Online/BuildPatchServices/Private/BuildPatchServicesModule.h b/Engine/Source/Runtime/Online/BuildPatchServices/Private/BuildPatchServicesModule.h
--- a/Engine/Source/Runtime/Online/BuildPatchServices/Private/BuildPatchServicesModule.h
+++ b/Engine/Source/Runtime/Online/BuildPatchServices/Private/BuildPatchServicesModule.h
@@ -66,12 +66,12 @@ class FBuildPatchServicesModule
 {
 public:
 
-	// Begin IModuleInterface interface
+	//~ Begin IModuleInterface Interface
 	virtual void StartupModule() override;
 	virtual void ShutdownModule() override;
-	// End IModuleInterface interface
+	//~ End IModuleInterface Interface
 
-	// Begin IBuildPatchServicesModule interface
+	//~ Begin IBuildPatchServicesModule Interface
 	virtual IBuildManifestPtr LoadManifestFromFile(const FString& Filename) override;
 	virtual IBuildManifestPtr MakeManifestFromData(const TArray<uint8>& ManifestData) override;
 	virtual bool SaveManifestToFile(const FString& Filename, IBuildManifestRef Manifest, bool bUseBinary = true) override;
@@ -91,7 +91,7 @@ public:
 	virtual bool EnumerateManifestData(FString ManifestFilePath, FString OutputFile, const bool bIncludeSizes) override;
 #endif
 	virtual IBuildManifestPtr MakeManifestFromJSON(const FString& ManifestJSON) override;
-	// End IBuildPatchServicesModule interface
+	//~ End IBuildPatchServicesModule Interface
 
 	/**
 	 * Gets the directory used for staging intermediate files.
diff --git a/Engine/Source/Runtime/Online/HTTP/Private/Curl/CurlHttp.h b/Engine/Source/Runtime/Online/HTTP/Private/Curl/CurlHttp.h
--- a/Engine/Source/Runtime/Online/HTTP/Private/Curl/CurlHttp.h
+++ b/Engine/Source/Runtime/Online/HTTP/Private/Curl/CurlHttp.h
@@ -106,7 +106,7 @@ public:
 	// implementation friends
 	friend class FCurlHttpResponse;
 
-	// Begin IHttpBase interface
+	//~ Begin IHttpBase Interface
 	virtual FString GetURL() override;
 	virtual FString GetURLParameter(const FString& ParameterName) override;
 	virtual FString GetHeader(const FString& HeaderName) override;
@@ -114,9 +114,9 @@ public:
 	virtual FString GetContentType() override;
 	virtual int32 GetContentLength() override;
 	virtual const TArray<uint8>& GetContent() override;
-	// End IHttpBase interface
+	//~ End IHttpBase Interface
 
-	// Begin IHttpRequest interface
+	//~ Begin IHttpRequest Interface
 	virtual FString GetVerb() override;
 	virtual void SetVerb(const FString& InVerb) override;
 	virtual void SetURL(const FString& InURL) override;
@@ -131,7 +131,7 @@ public:
 	virtual const FHttpResponsePtr GetResponse() const override;
 	virtual void Tick(float DeltaSeconds) override;
 	virtual float GetElapsedTime() override;
-	// End IHttpRequest interface
+	//~ End IHttpRequest Interface
 
 	/**
 	 * Returns libcurl's easy handle - needed for HTTP manager.
@@ -336,7 +336,7 @@ public:
 	friend class FCurlHttpRequest;
 
 
-	// Begin IHttpBase interface
+	//~ Begin IHttpBase Interface
 	virtual FString GetURL() override;
 	virtual FString GetURLParameter(const FString& ParameterName) override;
 	virtual FString GetHeader(const FString& HeaderName) override;
@@ -344,12 +344,12 @@ public:
 	virtual FString GetContentType() override;
 	virtual int32 GetContentLength() override;
 	virtual const TArray<uint8>& GetContent() override;
-	// End IHttpBase interface
+	//~ End IHttpBase Interface
 
-	// Begin IHttpResponse interface
+	//~ Begin IHttpResponse Interface
 	virtual int32 GetResponseCode() override;
 	virtual FString GetContentAsString() override;
-	// End IHttpResponse interface
+	//~ End IHttpResponse Interface
 
 	/**
 	 * Check whether a response is ready or not.
diff --git a/Engine/Source/Runtime/Online/HTTP/Private/Curl/CurlHttpManager.h b/Engine/Source/Runtime/Online/HTTP/Private/Curl/CurlHttpManager.h
--- a/Engine/Source/Runtime/Online/HTTP/Private/Curl/CurlHttpManager.h
+++ b/Engine/Source/Runtime/Online/HTTP/Private/Curl/CurlHttpManager.h
@@ -27,11 +27,11 @@ protected:
 
 public:
 
-	// Begin HttpManager interface
+	//~ Begin HttpManager Interface
 	virtual void AddRequest(const TSharedRef<class IHttpRequest>& Request) override;
 	virtual void RemoveRequest(const TSharedRef<class IHttpRequest>& Request) override;
 	virtual bool Tick(float DeltaSeconds) override;
-	// End HttpManager interface
+	//~ End HttpManager Interface
 
 	FCurlHttpManager();
 
diff --git a/Engine/Source/Runtime/Online/HTTP/Private/HTML5/HTML5HTTP.h b/Engine/Source/Runtime/Online/HTTP/Private/HTML5/HTML5HTTP.h
--- a/Engine/Source/Runtime/Online/HTTP/Private/HTML5/HTML5HTTP.h
+++ b/Engine/Source/Runtime/Online/HTTP/Private/HTML5/HTML5HTTP.h
@@ -35,7 +35,7 @@ public:
 	// implementation friends
 	friend class FHTML5HttpResponse;
 
-	// Begin IHttpBase interface
+	//~ Begin IHttpBase Interface
 	virtual FString GetURL() override;
 	virtual FString GetURLParameter(const FString& ParameterName) override;
 	virtual FString GetHeader(const FString& HeaderName) override;
@@ -43,9 +43,9 @@ public:
 	virtual FString GetContentType() override;
 	virtual int32 GetContentLength() override;
 	virtual const TArray<uint8>& GetContent() override;
-	// End IHttpBase interface
+	//~ End IHttpBase Interface
 
-	// Begin IHttpRequest interface
+	//~ Begin IHttpRequest Interface
 	virtual FString GetVerb() override;
 	virtual void SetVerb(const FString& InVerb) override;
 	virtual void SetURL(const FString& InURL) override;
@@ -60,7 +60,7 @@ public:
 	virtual const FHttpResponsePtr GetResponse() const override;
 	virtual void Tick(float DeltaSeconds) override;
 	virtual float GetElapsedTime() override;
-	// End IHttpRequest interface
+	//~ End IHttpRequest Interface
 
 	/**
 	* Marks request as completed.
@@ -153,7 +153,7 @@ public:
 	friend class FHTML5HttpRequest;
 
 
-	// Begin IHttpBase interface
+	//~ Begin IHttpBase Interface
 	virtual FString GetURL() override;
 	virtual FString GetURLParameter(const FString& ParameterName) override;
 	virtual FString GetHeader(const FString& HeaderName) override;
@@ -161,12 +161,12 @@ public:
 	virtual FString GetContentType() override;
 	virtual int32 GetContentLength() override;
 	virtual const TArray<uint8>& GetContent() override;
-	// End IHttpBase interface
+	//~ End IHttpBase Interface
 
-	// Begin IHttpResponse interface
+	//~ Begin IHttpResponse Interface
 	virtual int32 GetResponseCode() override;
 	virtual FString GetContentAsString() override;
-	// End IHttpResponse interface
+	//~ End IHttpResponse Interface
 
 	/**
 	* Check whether a response is ready or not.
diff --git a/Engine/Source/Runtime/Online/HTTP/Private/NullHttp.h b/Engine/Source/Runtime/Online/HTTP/Private/NullHttp.h
--- a/Engine/Source/Runtime/Online/HTTP/Private/NullHttp.h
+++ b/Engine/Source/Runtime/Online/HTTP/Private/NullHttp.h
@@ -65,7 +65,7 @@ class FNullHttpResponse : public IHttpResponse
 	virtual FString GetContentType() override;
 	virtual int32 GetContentLength() override;
 	virtual const TArray<uint8>& GetContent() override;
-	// Begin IHttpResponse interface
+	//~ Begin IHttpResponse Interface
 	virtual int32 GetResponseCode() override;
 	virtual FString GetContentAsString() override;
 
diff --git a/Engine/Source/Runtime/Online/IOS/OnlineSubsystemIOS/Private/OnlineAchievementsInterfaceIOS.h b/Engine/Source/Runtime/Online/IOS/OnlineSubsystemIOS/Private/OnlineAchievementsInterfaceIOS.h
--- a/Engine/Source/Runtime/Online/IOS/OnlineSubsystemIOS/Private/OnlineAchievementsInterfaceIOS.h
+++ b/Engine/Source/Runtime/Online/IOS/OnlineSubsystemIOS/Private/OnlineAchievementsInterfaceIOS.h
@@ -26,7 +26,7 @@ private:
 
 public:
 
-	// Begin IOnlineAchievements interface
+	//~ Begin IOnlineAchievements Interface
 	virtual void WriteAchievements(const FUniqueNetId& PlayerId, FOnlineAchievementsWriteRef& WriteObject, const FOnAchievementsWrittenDelegate& Delegate = FOnAchievementsWrittenDelegate()) override;
 	virtual void QueryAchievements(const FUniqueNetId& PlayerId, const FOnQueryAchievementsCompleteDelegate& Delegate = FOnQueryAchievementsCompleteDelegate()) override;
 	virtual void QueryAchievementDescriptions( const FUniqueNetId& PlayerId, const FOnQueryAchievementsCompleteDelegate& Delegate = FOnQueryAchievementsCompleteDelegate() ) override;
@@ -36,7 +36,7 @@ public:
 #if !UE_BUILD_SHIPPING
 	virtual bool ResetAchievements( const FUniqueNetId& PlayerId ) override;
 #endif // !UE_BUILD_SHIPPING
-	// End IOnlineAchievements interface
+	//~ End IOnlineAchievements Interface
 
 
 	/**
diff --git a/Engine/Source/Runtime/Online/IOS/OnlineSubsystemIOS/Private/OnlineExternalUIInterfaceIOS.h b/Engine/Source/Runtime/Online/IOS/OnlineSubsystemIOS/Private/OnlineExternalUIInterfaceIOS.h
--- a/Engine/Source/Runtime/Online/IOS/OnlineSubsystemIOS/Private/OnlineExternalUIInterfaceIOS.h
+++ b/Engine/Source/Runtime/Online/IOS/OnlineSubsystemIOS/Private/OnlineExternalUIInterfaceIOS.h
@@ -15,7 +15,7 @@ PACKAGE_SCOPE:
 
 public:
 
-	// Begin IOnlineExternalUI interface
+	//~ Begin IOnlineExternalUI Interface
 	virtual bool ShowLoginUI(const int ControllerIndex, bool bShowOnlineOnly, const FOnLoginUIClosedDelegate& Delegate) override;
 	virtual bool ShowFriendsUI(int32 LocalUserNum) override;
 	virtual bool ShowInviteUI(int32 LocalUserNum, FName SessionMame = GameSessionName) override;
@@ -24,7 +24,7 @@ public:
 	virtual bool ShowWebURL(const FString& WebURL) override;
 	virtual bool ShowProfileUI(const FUniqueNetId& Requestor, const FUniqueNetId& Requestee, const FOnProfileUIClosedDelegate& Delegate) override;
 	virtual bool ShowAccountUpgradeUI(const FUniqueNetId& UniqueId) override;
-	// End IOnlineExternalUI interface
+	//~ End IOnlineExternalUI Interface
 	
 private:
 	FOnlineSubsystemIOS* Subsystem;
diff --git a/Engine/Source/Runtime/Online/IOS/OnlineSubsystemIOS/Private/OnlineIdentityInterfaceIOS.h b/Engine/Source/Runtime/Online/IOS/OnlineSubsystemIOS/Private/OnlineIdentityInterfaceIOS.h
--- a/Engine/Source/Runtime/Online/IOS/OnlineSubsystemIOS/Private/OnlineIdentityInterfaceIOS.h
+++ b/Engine/Source/Runtime/Online/IOS/OnlineSubsystemIOS/Private/OnlineIdentityInterfaceIOS.h
@@ -30,7 +30,7 @@ PACKAGE_SCOPE:
 
 public:
 
-	// Begin IOnlineIdentity interface
+	//~ Begin IOnlineIdentity Interface
 	virtual bool Login(int32 LocalUserNum, const FOnlineAccountCredentials& AccountCredentials) override;
 	virtual bool Logout(int32 LocalUserNum) override;
 	virtual bool AutoLogin(int32 LocalUserNum) override;
@@ -46,7 +46,7 @@ public:
 	virtual FString GetAuthToken(int32 LocalUserNum) const override;
 	virtual void GetUserPrivilege(const FUniqueNetId& UserId, EUserPrivileges::Type Privilege, const FOnGetUserPrivilegeCompleteDelegate& Delegate) override;
 	virtual FPlatformUserId GetPlatformUserIdFromUniqueNetId(const FUniqueNetId& UniqueNetId) override;
-	// End IOnlineIdentity interface
+	//~ End IOnlineIdentity Interface
 
 
 public:
diff --git a/Engine/Source/Runtime/Online/IOS/OnlineSubsystemIOS/Private/OnlineLeaderboardsInterfaceIOS.h b/Engine/Source/Runtime/Online/IOS/OnlineSubsystemIOS/Private/OnlineLeaderboardsInterfaceIOS.h
--- a/Engine/Source/Runtime/Online/IOS/OnlineSubsystemIOS/Private/OnlineLeaderboardsInterfaceIOS.h
+++ b/Engine/Source/Runtime/Online/IOS/OnlineSubsystemIOS/Private/OnlineLeaderboardsInterfaceIOS.h
@@ -24,14 +24,14 @@ public:
 
 	virtual ~FOnlineLeaderboardsIOS();
 
-	// Begin IOnlineLeaderboards interface
+	//~ Begin IOnlineLeaderboards Interface
 	virtual bool ReadLeaderboards(const TArray< TSharedRef<const FUniqueNetId> >& Players, FOnlineLeaderboardReadRef& ReadObject) override;
 	virtual bool ReadLeaderboardsForFriends(int32 LocalUserNum, FOnlineLeaderboardReadRef& ReadObject) override;
 	virtual void FreeStats(FOnlineLeaderboardRead& ReadObject) override;
 	virtual bool WriteLeaderboards(const FName& SessionName, const FUniqueNetId& Player, FOnlineLeaderboardWrite& WriteObject) override;
 	virtual bool FlushLeaderboards(const FName& SessionName) override;
 	virtual bool WriteOnlinePlayerRatings(const FName& SessionName, int32 LeaderboardId, const TArray<FOnlinePlayerScore>& PlayerScores) override;
-	// End IOnlineLeaderboards interface
+	//~ End IOnlineLeaderboards Interface
 
 };
 
diff --git a/Engine/Source/Runtime/Online/IOS/OnlineSubsystemIOS/Private/OnlineSessionInterfaceIOS.h b/Engine/Source/Runtime/Online/IOS/OnlineSubsystemIOS/Private/OnlineSessionInterfaceIOS.h
--- a/Engine/Source/Runtime/Online/IOS/OnlineSubsystemIOS/Private/OnlineSessionInterfaceIOS.h
+++ b/Engine/Source/Runtime/Online/IOS/OnlineSubsystemIOS/Private/OnlineSessionInterfaceIOS.h
@@ -100,7 +100,7 @@ PACKAGE_SCOPE:
 	void Tick(float DeltaTime);
 
 
-	// Begin IOnlineSession interface
+	//~ Begin IOnlineSession Interface
 	class FNamedOnlineSession* AddNamedSession(FName SessionName, const FOnlineSessionSettings& SessionSettings) override;
 
 	class FNamedOnlineSession* AddNamedSession(FName SessionName, const FOnlineSession& Session) override;
@@ -184,7 +184,7 @@ public:
 	virtual int32 GetNumSessions() override;
 
 	virtual void DumpSessionState() override;
-	// End IOnlineSession interface
+	//~ End IOnlineSession Interface
 };
 
 typedef TSharedPtr<FOnlineSessionIOS, ESPMode::ThreadSafe> FOnlineSessionIOSPtr;
diff --git a/Engine/Source/Runtime/Online/IOS/OnlineSubsystemIOS/Private/OnlineStoreInterfaceIOS.h b/Engine/Source/Runtime/Online/IOS/OnlineSubsystemIOS/Private/OnlineStoreInterfaceIOS.h
--- a/Engine/Source/Runtime/Online/IOS/OnlineSubsystemIOS/Private/OnlineStoreInterfaceIOS.h
+++ b/Engine/Source/Runtime/Online/IOS/OnlineSubsystemIOS/Private/OnlineStoreInterfaceIOS.h
@@ -63,12 +63,12 @@ public:
 	/** Destructor */
 	virtual ~FOnlineStoreInterfaceIOS();
 
-	// Begin IOnlineStore interface
+	//~ Begin IOnlineStore Interface
 	virtual bool QueryForAvailablePurchases(const TArray<FString>& ProductIDs, FOnlineProductInformationReadRef& InReadObject) override;
 	virtual bool BeginPurchase(const FInAppPurchaseProductRequest& ProductRequest, FOnlineInAppPurchaseTransactionRef& InReadObject) override;
 	virtual bool IsAllowedToMakePurchases() override;
 	virtual bool RestorePurchases(FOnlineInAppPurchaseRestoreReadRef& InReadObject) override;
-	// End IOnlineStore interface
+	//~ End IOnlineStore Interface
 
 	/**
 	 * Process a response from the StoreKit
diff --git a/Engine/Source/Runtime/Online/IOS/OnlineSubsystemIOS/Public/OnlineSubsystemIOS.h b/Engine/Source/Runtime/Online/IOS/OnlineSubsystemIOS/Public/OnlineSubsystemIOS.h
--- a/Engine/Source/Runtime/Online/IOS/OnlineSubsystemIOS/Public/OnlineSubsystemIOS.h
+++ b/Engine/Source/Runtime/Online/IOS/OnlineSubsystemIOS/Public/OnlineSubsystemIOS.h
@@ -18,7 +18,7 @@ public:
 	FOnlineSubsystemIOS();
 	virtual ~FOnlineSubsystemIOS() {};
 
-	// Begin IOnlineSubsystem Interface
+	//~ Begin IOnlineSubsystem Interface
 	virtual IOnlineSessionPtr GetSessionInterface() const override;
 	virtual IOnlineFriendsPtr GetFriendsInterface() const override;
 	virtual IOnlinePartyPtr GetPartyInterface() const override;
@@ -48,7 +48,7 @@ public:
 	virtual FString GetAppId() const override;
 	virtual bool Exec(class UWorld* InWorld, const TCHAR* Cmd, FOutputDevice& Ar) override;
 	virtual bool Tick(float DeltaTime) override;
-	// End IOnlineSubsystem Interface
+	//~ End IOnlineSubsystem Interface
 
 PACKAGE_SCOPE:
 
diff --git a/Engine/Source/Runtime/Online/OnlineFramework/Lobby/Public/LobbyBeaconClient.h b/Engine/Source/Runtime/Online/OnlineFramework/Lobby/Public/LobbyBeaconClient.h
--- a/Engine/Source/Runtime/Online/OnlineFramework/Lobby/Public/LobbyBeaconClient.h
+++ b/Engine/Source/Runtime/Online/OnlineFramework/Lobby/Public/LobbyBeaconClient.h
@@ -54,9 +54,9 @@ class LOBBY_API ALobbyBeaconClient : public AOnlineBeaconClient
 	UPROPERTY(Replicated)
 	ALobbyBeaconPlayerState* PlayerState;
 	
-	// Begin AOnlineBeaconClient Interface
+	//~ Begin AOnlineBeaconClient Interface
 	virtual void OnConnected() override;
-	// End AOnlineBeaconClient Interface
+	//~ End AOnlineBeaconClient Interface
 
 	/**
 	 * Initiate a connection to the lobby host beacon
diff --git a/Engine/Source/Runtime/Online/OnlineFramework/Lobby/Public/LobbyBeaconHost.h b/Engine/Source/Runtime/Online/OnlineFramework/Lobby/Public/LobbyBeaconHost.h
--- a/Engine/Source/Runtime/Online/OnlineFramework/Lobby/Public/LobbyBeaconHost.h
+++ b/Engine/Source/Runtime/Online/OnlineFramework/Lobby/Public/LobbyBeaconHost.h
@@ -20,9 +20,9 @@ class LOBBY_API ALobbyBeaconHost : public AOnlineBeaconHostObject
 {
 	GENERATED_UCLASS_BODY()
 
-	// Begin AOnlineBeaconHostObject Interface 
+	//~ Begin AOnlineBeaconHostObject Interface 
 	void NotifyClientDisconnected(AOnlineBeaconClient* LeavingClientActor);
-	// End AOnlineBeaconHostObject Interface 
+	//~ End AOnlineBeaconHostObject Interface 
 
 	/**
 	 * Initialize the lobby beacon, creating an object to maintain state
diff --git a/Engine/Source/Runtime/Online/OnlineFramework/Party/Public/PartyGameState.h b/Engine/Source/Runtime/Online/OnlineFramework/Party/Public/PartyGameState.h
--- a/Engine/Source/Runtime/Online/OnlineFramework/Party/Public/PartyGameState.h
+++ b/Engine/Source/Runtime/Online/OnlineFramework/Party/Public/PartyGameState.h
@@ -107,9 +107,9 @@ class PARTY_API UPartyGameState : public UObject
 {
 	GENERATED_UCLASS_BODY()
 
-	// Begin UObject Interface
+	//~ Begin UObject Interface
 	virtual void BeginDestroy() override;
-	// End UObject Interface
+	//~ End UObject Interface
 
 private:
 
@@ -160,10 +160,10 @@ protected:
 
 public:
 
-	// UObject interface begin
+	//~ Begin UObject Interface begin
 	virtual UWorld* GetWorld() const override;
 	static void AddReferencedObjects(UObject* InThis, FReferenceCollector& Collector);
-	// UObject interface end
+	//~ Begin UObject Interface end
 
 	/**
 	 * Notification that the game is about to travel to another map/server
diff --git a/Engine/Source/Runtime/Online/OnlineSubsystemFacebook/Private/IOS/OnlineIdentityFacebook.h b/Engine/Source/Runtime/Online/OnlineSubsystemFacebook/Private/IOS/OnlineIdentityFacebook.h
--- a/Engine/Source/Runtime/Online/OnlineSubsystemFacebook/Private/IOS/OnlineIdentityFacebook.h
+++ b/Engine/Source/Runtime/Online/OnlineSubsystemFacebook/Private/IOS/OnlineIdentityFacebook.h
@@ -65,7 +65,7 @@ class FOnlineIdentityFacebook : public IOnlineIdentity
 
 public:
 
-	// Begin IOnlineIdentity interface	
+	//~ Begin IOnlineIdentity Interface	
 	virtual bool Login(int32 LocalUserNum, const FOnlineAccountCredentials& AccountCredentials) override;
 	virtual bool Logout(int32 LocalUserNum) override;
 	virtual bool AutoLogin(int32 LocalUserNum) override;
@@ -81,7 +81,7 @@ public:
 	virtual FString GetAuthToken(int32 LocalUserNum) const override;
 	virtual void GetUserPrivilege(const FUniqueNetId& UserId, EUserPrivileges::Type Privilege, const FOnGetUserPrivilegeCompleteDelegate& Delegate) override;
 	virtual FPlatformUserId GetPlatformUserIdFromUniqueNetId(const FUniqueNetId& UniqueNetId) override;
-	// End IOnlineIdentity interface
+	//~ End IOnlineIdentity Interface
 
 
 public:
diff --git a/Engine/Source/Runtime/Online/OnlineSubsystemFacebook/Private/IOS/OnlineSharingFacebook.h b/Engine/Source/Runtime/Online/OnlineSubsystemFacebook/Private/IOS/OnlineSharingFacebook.h
--- a/Engine/Source/Runtime/Online/OnlineSubsystemFacebook/Private/IOS/OnlineSharingFacebook.h
+++ b/Engine/Source/Runtime/Online/OnlineSubsystemFacebook/Private/IOS/OnlineSharingFacebook.h
@@ -16,14 +16,14 @@ class FOnlineSharingFacebook : public IOnlineSharing
 
 public:
 
-	// Begin IOnlineSharing Interface
+	//~ Begin IOnlineSharing Interface
 	virtual bool ReadNewsFeed(int32 LocalUserNum, int32 NumPostsToRead) override;
 	virtual bool RequestNewReadPermissions(int32 LocalUserNum, EOnlineSharingReadCategory::Type NewPermissions) override;
 	virtual bool ShareStatusUpdate(int32 LocalUserNum, const FOnlineStatusUpdate& StatusUpdate) override;
 	virtual bool RequestNewPublishPermissions(int32 LocalUserNum, EOnlineSharingPublishingCategory::Type NewPermissions, EOnlineStatusUpdatePrivacy::Type Privacy) override;
 	virtual EOnlineCachedResult::Type GetCachedNewsFeed(int32 LocalUserNum, int32 NewsFeedIdx, FOnlineStatusUpdate& OutNewsFeed) override;
 	virtual EOnlineCachedResult::Type GetCachedNewsFeeds(int32 LocalUserNum, TArray<FOnlineStatusUpdate>& OutNewsFeeds) override;
-	// End IOnlineSharing Interface
+	//~ End IOnlineSharing Interface
 
 	
 public:
diff --git a/Engine/Source/Runtime/Online/OnlineSubsystemFacebook/Private/OnlineSubsystemFacebookTypes.h b/Engine/Source/Runtime/Online/OnlineSubsystemFacebook/Private/OnlineSubsystemFacebookTypes.h
--- a/Engine/Source/Runtime/Online/OnlineSubsystemFacebook/Private/OnlineSubsystemFacebookTypes.h
+++ b/Engine/Source/Runtime/Online/OnlineSubsystemFacebook/Private/OnlineSubsystemFacebookTypes.h
@@ -42,7 +42,7 @@ public:
 	{
 	}
 
-	// Begin FUniqueNetId interface
+	//~ Begin FUniqueNetId Interface
 	virtual const uint8* GetBytes() const override
 	{
 		return (uint8*)&UniqueNetId;
@@ -72,7 +72,7 @@ public:
 		return FString::Printf(TEXT("0%I64X"), UniqueNetId);
 	}
 
-	// End FUniqueNetId interface
+	//~ End FUniqueNetId Interface
 
 
 public:
diff --git a/Engine/Source/Runtime/Online/OnlineSubsystemSteam/Classes/SteamNetConnection.h b/Engine/Source/Runtime/Online/OnlineSubsystemSteam/Classes/SteamNetConnection.h
--- a/Engine/Source/Runtime/Online/OnlineSubsystemSteam/Classes/SteamNetConnection.h
+++ b/Engine/Source/Runtime/Online/OnlineSubsystemSteam/Classes/SteamNetConnection.h
@@ -13,11 +13,11 @@ class USteamNetConnection : public UIpConnection
 {
     GENERATED_UCLASS_BODY()
 
-	// Begin UIpConnection Interface
+	//~ Begin UIpConnection Interface
 	virtual void InitRemoteConnection(class UNetDriver* InDriver, class FSocket* InSocket, const FURL& InURL, const class FInternetAddr& InRemoteAddr, EConnectionState InState, int32 InMaxPacket = 0, int32 InPacketOverhead = 0) override;
 	virtual void InitLocalConnection(class UNetDriver* InDriver, class FSocket* InSocket, const FURL& InURL, EConnectionState InState, int32 InMaxPacket = 0, int32 InPacketOverhead = 0) override;
 	virtual void CleanUp() override;
-	// End UIpConnection Interface
+	//~ End UIpConnection Interface
 
 	/** Is this net connection passthrough to IpConnection */
 	UPROPERTY()
diff --git a/Engine/Source/Runtime/Online/OnlineSubsystemSteam/Classes/SteamNetDriver.h b/Engine/Source/Runtime/Online/OnlineSubsystemSteam/Classes/SteamNetDriver.h
--- a/Engine/Source/Runtime/Online/OnlineSubsystemSteam/Classes/SteamNetDriver.h
+++ b/Engine/Source/Runtime/Online/OnlineSubsystemSteam/Classes/SteamNetDriver.h
@@ -20,11 +20,11 @@ class USteamNetDriver : public UIpNetDriver
 	/** Should this net driver behave as a passthrough to normal IP */
 	bool bIsPassthrough;
 
-	// Begin UObject Interface
+	//~ Begin UObject Interface
 	virtual void PostInitProperties() override;
-	// End UObject Interface
+	//~ End UObject Interface
 
-	// Begin UIpNetDriver Interface
+	//~ Begin UIpNetDriver Interface
 
 	virtual class ISocketSubsystem* GetSocketSubsystem() override;
 	virtual bool IsAvailable() const override;
@@ -35,6 +35,6 @@ class USteamNetDriver : public UIpNetDriver
 	virtual void TickFlush(float DeltaSeconds) override;
 	virtual bool IsNetResourceValid() override;
 
-	// End UIpNetDriver Interface
+	//~ End UIpNetDriver Interface
 
 };
diff --git a/Engine/Source/Runtime/Online/OnlineSubsystemUtils/Classes/IpConnection.h b/Engine/Source/Runtime/Online/OnlineSubsystemUtils/Classes/IpConnection.h
--- a/Engine/Source/Runtime/Online/OnlineSubsystemUtils/Classes/IpConnection.h
+++ b/Engine/Source/Runtime/Online/OnlineSubsystemUtils/Classes/IpConnection.h
@@ -16,7 +16,7 @@ class ONLINESUBSYSTEMUTILS_API UIpConnection : public UNetConnection
 	class FSocket*				Socket;
 	class FResolveInfo*			ResolveInfo;
 
-	// Begin NetConnection Interface
+	//~ Begin NetConnection Interface
 	virtual void InitBase(UNetDriver* InDriver, class FSocket* InSocket, const FURL& InURL, EConnectionState InState, int32 InMaxPacket = 0, int32 InPacketOverhead = 0) override;
 	virtual void InitRemoteConnection(UNetDriver* InDriver, class FSocket* InSocket, const FURL& InURL, const class FInternetAddr& InRemoteAddr, EConnectionState InState, int32 InMaxPacket = 0, int32 InPacketOverhead = 0) override;
 	virtual void InitLocalConnection(UNetDriver* InDriver, class FSocket* InSocket, const FURL& InURL, EConnectionState InState, int32 InMaxPacket = 0, int32 InPacketOverhead = 0) override;
@@ -26,5 +26,5 @@ class ONLINESUBSYSTEMUTILS_API UIpConnection : public UNetConnection
 	virtual int32 GetAddrAsInt(void) override;
 	virtual int32 GetAddrPort(void) override;
 	virtual FString RemoteAddressToString() override;
-	// End NetConnection Interface
+	//~ End NetConnection Interface
 };
diff --git a/Engine/Source/Runtime/Online/OnlineSubsystemUtils/Classes/IpNetDriver.h b/Engine/Source/Runtime/Online/OnlineSubsystemUtils/Classes/IpNetDriver.h
--- a/Engine/Source/Runtime/Online/OnlineSubsystemUtils/Classes/IpNetDriver.h
+++ b/Engine/Source/Runtime/Online/OnlineSubsystemUtils/Classes/IpNetDriver.h
@@ -30,7 +30,7 @@ class ONLINESUBSYSTEMUTILS_API UIpNetDriver : public UNetDriver
 	/** Underlying socket communication */
 	FSocket* Socket;
 
-	// Begin UNetDriver interface.
+	//~ Begin UNetDriver Interface.
 	virtual bool IsAvailable() const override;
 	virtual bool InitBase(bool bInitAsClient, FNetworkNotify* InNotify, const FURL& URL, bool bReuseAddressAndPort, FString& Error) override;
 	virtual bool InitConnect( FNetworkNotify* InNotify, const FURL& ConnectURL, FString& Error ) override;
@@ -44,9 +44,9 @@ class ONLINESUBSYSTEMUTILS_API UIpNetDriver : public UNetDriver
 	{
 		return Socket != NULL;
 	}
-	// End UNetDriver Interface
+	//~ End UNetDriver Interface
 
-	// Begin UIpNetDriver interface.
+	//~ Begin UIpNetDriver Interface.
 	virtual FSocket * CreateSocket();
 
 	/**
@@ -57,11 +57,11 @@ class ONLINESUBSYSTEMUTILS_API UIpNetDriver : public UNetDriver
 	 * @return The port number to use for client sockets. Base implementation returns 0.
 	 */
 	virtual int GetClientPort();
-	// End UIpNetDriver interface.
+	//~ End UIpNetDriver Interface.
 
-	// Begin FExec Interface
+	//~ Begin FExec Interface
 	virtual bool Exec( UWorld* InWorld, const TCHAR* Cmd, FOutputDevice& Ar=*GLog ) override;
-	// End FExec Interface
+	//~ End FExec Interface
 
 	/**
 	 * Exec command handlers
diff --git a/Engine/Source/Runtime/Online/OnlineSubsystemUtils/Classes/TestBeaconClient.h b/Engine/Source/Runtime/Online/OnlineSubsystemUtils/Classes/TestBeaconClient.h
--- a/Engine/Source/Runtime/Online/OnlineSubsystemUtils/Classes/TestBeaconClient.h
+++ b/Engine/Source/Runtime/Online/OnlineSubsystemUtils/Classes/TestBeaconClient.h
@@ -12,9 +12,9 @@ class ONLINESUBSYSTEMUTILS_API ATestBeaconClient : public AOnlineBeaconClient
 {
 	GENERATED_UCLASS_BODY()
 
-	// Begin AOnlineBeaconClient Interface
+	//~ Begin AOnlineBeaconClient Interface
 	virtual void OnFailure() override;
-	// End AOnlineBeaconClient Interface
+	//~ End AOnlineBeaconClient Interface
 
 	/** Send a ping RPC to the client */
 	UFUNCTION(client, reliable)
diff --git a/Engine/Source/Runtime/Online/OnlineSubsystemUtils/Classes/TestBeaconHost.h b/Engine/Source/Runtime/Online/OnlineSubsystemUtils/Classes/TestBeaconHost.h
--- a/Engine/Source/Runtime/Online/OnlineSubsystemUtils/Classes/TestBeaconHost.h
+++ b/Engine/Source/Runtime/Online/OnlineSubsystemUtils/Classes/TestBeaconHost.h
@@ -12,10 +12,10 @@ class ONLINESUBSYSTEMUTILS_API ATestBeaconHost : public AOnlineBeaconHostObject
 {
 	GENERATED_UCLASS_BODY()
 
-	// Begin AOnlineBeaconHost Interface 
+	//~ Begin AOnlineBeaconHost Interface 
 	virtual AOnlineBeaconClient* SpawnBeaconActor(class UNetConnection* ClientConnection) override;
 	virtual void OnClientConnected(class AOnlineBeaconClient* NewClientActor, class UNetConnection* ClientConnection) override;
-	// End AOnlineBeaconHost Interface 
+	//~ End AOnlineBeaconHost Interface 
 
 	virtual bool Init();
 };
diff --git a/Engine/Source/Runtime/Online/OnlineSubsystemUtils/Public/OnlineBeacon.h b/Engine/Source/Runtime/Online/OnlineSubsystemUtils/Public/OnlineBeacon.h
--- a/Engine/Source/Runtime/Online/OnlineSubsystemUtils/Public/OnlineBeacon.h
+++ b/Engine/Source/Runtime/Online/OnlineSubsystemUtils/Public/OnlineBeacon.h
@@ -29,7 +29,7 @@ class ONLINESUBSYSTEMUTILS_API AOnlineBeacon : public AActor, public FNetworkNot
 {
 	GENERATED_UCLASS_BODY()
 
-	// Begin AActor Interface
+	//~ Begin AActor Interface
 	virtual void OnActorChannelOpen(FInBunch& InBunch, UNetConnection* Connection) override;
 	virtual bool IsRelevancyOwnerFor(const AActor* ReplicatedActor, const AActor* ActorOwner, const AActor* ConnectionActor) const override;
 	virtual bool IsNetRelevantFor(const AActor* RealViewer, const AActor* ViewTarget, const FVector& SrcLocation) const override;
@@ -37,14 +37,14 @@ class ONLINESUBSYSTEMUTILS_API AOnlineBeacon : public AActor, public FNetworkNot
 	virtual UNetConnection* GetNetConnection() const override { return nullptr; }
 	virtual bool IsLevelBoundsRelevant() const override { return false; }
 	virtual void EndPlay(const EEndPlayReason::Type EndPlayReason);
-	// End AActor Interface
+	//~ End AActor Interface
 
-	// Begin FNetworkNotify Interface
+	//~ Begin FNetworkNotify Interface
 	virtual EAcceptConnection::Type NotifyAcceptingConnection() override;
 	virtual void NotifyAcceptedConnection(UNetConnection* Connection) override;
 	virtual bool NotifyAcceptingChannel(UChannel* Channel) override;
 	virtual void NotifyControlMessage(UNetConnection* Connection, uint8 MessageType, FInBunch& Bunch) override;
-	// End FNetworkNotify Interface
+	//~ End FNetworkNotify Interface
 	
     /**
 	 * Get the current state of the beacon
diff --git a/Engine/Source/Runtime/Online/OnlineSubsystemUtils/Public/OnlineBeaconClient.h b/Engine/Source/Runtime/Online/OnlineSubsystemUtils/Public/OnlineBeaconClient.h
--- a/Engine/Source/Runtime/Online/OnlineSubsystemUtils/Public/OnlineBeaconClient.h
+++ b/Engine/Source/Runtime/Online/OnlineSubsystemUtils/Public/OnlineBeaconClient.h
@@ -39,22 +39,22 @@ class ONLINESUBSYSTEMUTILS_API AOnlineBeaconClient : public AOnlineBeacon
 {
 	GENERATED_UCLASS_BODY()
 
-	// Begin AActor Interface
+	//~ Begin AActor Interface
 	virtual bool UseShortConnectTimeout() const override;
 	virtual void OnNetCleanup(UNetConnection* Connection) override;
 	virtual const AActor* GetNetOwner() const override;
 	virtual UNetConnection* GetNetConnection() const override;
 	virtual bool DestroyNetworkActorHandled() override;
-	// End AActor Interface
+	//~ End AActor Interface
 
-	// Begin FNetworkNotify Interface
+	//~ Begin FNetworkNotify Interface
 	virtual void NotifyControlMessage(UNetConnection* Connection, uint8 MessageType, FInBunch& Bunch) override;
-	// End FNetworkNotify Interface
+	//~ End FNetworkNotify Interface
 
-	// Begin OnlineBeacon Interface
+	//~ Begin OnlineBeacon Interface
 	virtual void OnFailure() override;
 	virtual void DestroyBeacon() override;
-	// End OnlineBeacon Interface
+	//~ End OnlineBeacon Interface
 
 	/**
 	 * Initialize the client beacon with connection endpoint
diff --git a/Engine/Source/Runtime/Online/OnlineSubsystemUtils/Public/OnlineBeaconHost.h b/Engine/Source/Runtime/Online/OnlineSubsystemUtils/Public/OnlineBeaconHost.h
--- a/Engine/Source/Runtime/Online/OnlineSubsystemUtils/Public/OnlineBeaconHost.h
+++ b/Engine/Source/Runtime/Online/OnlineSubsystemUtils/Public/OnlineBeaconHost.h
@@ -26,17 +26,17 @@ class ONLINESUBSYSTEMUTILS_API AOnlineBeaconHost : public AOnlineBeacon
 	UPROPERTY(Config)
 	int32 ListenPort;
 
-	// Begin AActor Interface
+	//~ Begin AActor Interface
 	virtual void OnNetCleanup(UNetConnection* Connection) override;
-	// End AActor Interface
+	//~ End AActor Interface
 
-	// Begin OnlineBeacon Interface
+	//~ Begin OnlineBeacon Interface
 	virtual void HandleNetworkFailure(UWorld* World, UNetDriver* NetDriver, ENetworkFailure::Type FailureType, const FString& ErrorString) override;
-	// End OnlineBeacon Interface
+	//~ End OnlineBeacon Interface
 
-	// Begin FNetworkNotify Interface
+	//~ Begin FNetworkNotify Interface
 	virtual void NotifyControlMessage(UNetConnection* Connection, uint8 MessageType, FInBunch& Bunch) override;
-	// End FNetworkNotify Interface
+	//~ End FNetworkNotify Interface
 
 	/**
 	 * Initialize the host beacon on a specified port
diff --git a/Engine/Source/Runtime/Online/OnlineSubsystemUtils/Public/PartyBeaconHost.h b/Engine/Source/Runtime/Online/OnlineSubsystemUtils/Public/PartyBeaconHost.h
--- a/Engine/Source/Runtime/Online/OnlineSubsystemUtils/Public/PartyBeaconHost.h
+++ b/Engine/Source/Runtime/Online/OnlineSubsystemUtils/Public/PartyBeaconHost.h
@@ -46,12 +46,12 @@ class ONLINESUBSYSTEMUTILS_API APartyBeaconHost : public AOnlineBeaconHostObject
 {
 	GENERATED_UCLASS_BODY()
 
-	// Begin AActor Interface
+	//~ Begin AActor Interface
 	virtual void Tick(float DeltaTime) override;
-	// End AActor Interface
+	//~ End AActor Interface
 
-	// Begin AOnlineBeaconHostObject Interface 
-	// End AOnlineBeaconHost Interface 
+	//~ Begin AOnlineBeaconHostObject Interface 
+	//~ End AOnlineBeaconHost Interface 
 
 	/**
 	 * Initialize the party host beacon
diff --git a/Engine/Source/Runtime/PakFile/Private/SignedArchiveReader.h b/Engine/Source/Runtime/PakFile/Private/SignedArchiveReader.h
--- a/Engine/Source/Runtime/PakFile/Private/SignedArchiveReader.h
+++ b/Engine/Source/Runtime/PakFile/Private/SignedArchiveReader.h
@@ -151,11 +151,11 @@ public:
 	FChunkCacheWorker(FArchive* InReader);
 	virtual ~FChunkCacheWorker();
 
-	// Begin FRunnable interface.
+	//~ Begin FRunnable Interface.
 	virtual bool Init();
 	virtual uint32 Run();
 	virtual void Stop();
-	// End FRunnable interface
+	//~ End FRunnable Interface
 
 	/** 
 	 * Requests a chunk to be loaded and verified
@@ -256,7 +256,7 @@ public:
 	FSignedArchiveReader(FArchive* InPakReader, FChunkCacheWorker* InSignatureChecker);
 	virtual ~FSignedArchiveReader();
 
-	// Begin FArchive interface
+	//~ Begin FArchive Interface
 	virtual void Serialize(void* Data, int64 Length) override;
 	virtual int64 Tell() override
 	{
@@ -270,5 +270,5 @@ public:
 	{
 		PakOffset = InPos;
 	}
-	// End FArchive interface
+	//~ End FArchive Interface
 };
diff --git a/Engine/Source/Runtime/PakFile/Public/IPlatformFilePak.h b/Engine/Source/Runtime/PakFile/Public/IPlatformFilePak.h
--- a/Engine/Source/Runtime/PakFile/Public/IPlatformFilePak.h
+++ b/Engine/Source/Runtime/PakFile/Public/IPlatformFilePak.h
@@ -716,7 +716,7 @@ public:
 		}
 	}
 
-	// BEGIN IFileHandle Interface
+	//~ Begin IFileHandle Interface
 	virtual int64 Tell() override
 	{
 		return ReadPos;
@@ -774,7 +774,7 @@ public:
 	{
 		return Reader.FileSize();
 	}
-	/// END IFileHandle Interface
+	///~ End IFileHandle Interface
 };
 
 /**
@@ -984,7 +984,7 @@ public:
 		return FindFileInPakFiles(Paks, Filename, OutPakFile);
 	}
 
-	// BEGIN IPlatformFile Interface
+	//~ Begin IPlatformFile Interface
 	virtual bool FileExists(const TCHAR* Filename) override
 	{
 		// Check pak files first.
@@ -1481,7 +1481,7 @@ public:
 			return LowerLevel->ConvertToAbsolutePathForExternalAppForWrite(Filename);
 		}
 	}
-	// END IPlatformFile Interface
+	//~ End IPlatformFile Interface
 
 	// BEGIN Console commands
 #if !UE_BUILD_SHIPPING
diff --git a/Engine/Source/Runtime/PerfCounters/Private/PerfCounters.h b/Engine/Source/Runtime/PerfCounters/Private/PerfCounters.h
--- a/Engine/Source/Runtime/PerfCounters/Private/PerfCounters.h
+++ b/Engine/Source/Runtime/PerfCounters/Private/PerfCounters.h
@@ -25,18 +25,18 @@ public:
 	/** FTickerObjectBase */
 	virtual bool Tick(float DeltaTime) override;
 
-	// Begin Exec Interface
+	//~ Begin Exec Interface
 	virtual bool Exec(UWorld* InWorld, const TCHAR* Cmd, FOutputDevice& Ar);
-	// End Exec Interface
+	//~ End Exec Interface
 
-	// IPerfCounters interface
+	//~ Begin IPerfCounters Interface
 	const FString& GetInstanceName() const override { return UniqueInstanceId; }
 	virtual double GetNumber(const FString& Name, double DefaultValue = 0.0) override;
 	virtual void SetNumber(const FString& Name, double Value, uint32 Flags) override;
 	virtual void SetString(const FString& Name, const FString& Value, uint32 Flags) override;
 	virtual void SetJson(const FString& Name, const FProduceJsonCounterValue& InCallback, uint32 Flags) override;
 	virtual FPerfCounterExecCommandCallback& OnPerfCounterExecCommand() override { return ExecCmdCallback; }
-	// IPerfCounters interface end
+	//~ Begin IPerfCounters Interface end
 
 private:
 	/** Convert all perf counters into a json type */
