Fixed landscape Grass wasting memory storing empty grass maps in areas with no grass

diff --git a/Engine/Source/Runtime/Landscape/Private/LandscapeGrass.cpp b/Engine/Source/Runtime/Landscape/Private/LandscapeGrass.cpp
--- a/Engine/Source/Runtime/Landscape/Private/LandscapeGrass.cpp
+++ b/Engine/Source/Runtime/Landscape/Private/LandscapeGrass.cpp
@@ -518,6 +518,15 @@ public:
 			// remove null grass type if we had one (can occur if the node has null entries)
 			NewGrassData->WeightData.Remove(nullptr);
 
+			// Remove any grass data that is entirely weight 0
+			for (auto Iter(NewGrassData->WeightData.CreateIterator()); Iter; ++Iter)
+			{
+				if (Iter.Value().IndexOfByPredicate([&](const int8& Weight){ return Weight != 0; }) == INDEX_NONE)
+				{
+					Iter.RemoveCurrent();
+				}
+			}
+
 			// Assign the new data (thread-safe)
 			ComponentInfo.Component->GrassData = MakeShareable(NewGrassData);
 		}
