Mergine engine changes to //UE4/Main @ 2748007

diff --git a/Engine/Source/Runtime/GameplayAbilities/Private/Abilities/GameplayAbility.cpp b/Engine/Source/Runtime/GameplayAbilities/Private/Abilities/GameplayAbility.cpp
--- a/Engine/Source/Runtime/GameplayAbilities/Private/Abilities/GameplayAbility.cpp
+++ b/Engine/Source/Runtime/GameplayAbilities/Private/Abilities/GameplayAbility.cpp
@@ -22,22 +22,26 @@ UGameplayAbility::UGameplayAbility(const FObjectInitializer& ObjectInitializer)
 	{
 		static FName FuncName = FName(TEXT("K2_ShouldAbilityRespondToEvent"));
 		UFunction* ShouldRespondFunction = GetClass()->FindFunctionByName(FuncName);
-		bHasBlueprintShouldAbilityRespondToEvent = ShouldRespondFunction && ShouldRespondFunction->GetOuter()->IsA(UBlueprintGeneratedClass::StaticClass());
+		bHasBlueprintShouldAbilityRespondToEvent = ShouldRespondFunction && ensure(ShouldRespondFunction->GetOuter()) && ShouldRespondFunction->GetOuter()->IsA(UBlueprintGeneratedClass::StaticClass());
 	}
 	{
 		static FName FuncName = FName(TEXT("K2_CanActivateAbility"));
 		UFunction* CanActivateFunction = GetClass()->FindFunctionByName(FuncName);
-		bHasBlueprintCanUse = CanActivateFunction && CanActivateFunction->GetOuter()->IsA(UBlueprintGeneratedClass::StaticClass());
+		bHasBlueprintCanUse = CanActivateFunction && ensure(CanActivateFunction->GetOuter()) && CanActivateFunction->GetOuter()->IsA(UBlueprintGeneratedClass::StaticClass());
 	}
 	{
 		static FName FuncName = FName(TEXT("K2_ActivateAbility"));
 		UFunction* ActivateFunction = GetClass()->FindFunctionByName(FuncName);
-		bHasBlueprintActivate = ActivateFunction && ActivateFunction->GetOuter()->IsA(UBlueprintGeneratedClass::StaticClass());
+		// FIXME: temp to work around crash
+		if (ActivateFunction && ActivateFunction->IsValidLowLevelFast())
+		{
+			bHasBlueprintActivate = ActivateFunction && ensure(ActivateFunction->GetOuter()) && ActivateFunction->GetOuter()->IsA(UBlueprintGeneratedClass::StaticClass());
+		}
 	}
 	{
 		static FName FuncName = FName(TEXT("K2_ActivateAbilityFromEvent"));
 		UFunction* ActivateFunction = GetClass()->FindFunctionByName(FuncName);
-		bHasBlueprintActivateFromEvent = ActivateFunction && ActivateFunction->GetOuter()->IsA(UBlueprintGeneratedClass::StaticClass());
+		bHasBlueprintActivateFromEvent = ActivateFunction && ensure(ActivateFunction->GetOuter()) && ActivateFunction->GetOuter()->IsA(UBlueprintGeneratedClass::StaticClass());
 	}
 	
 #if WITH_EDITOR
@@ -866,10 +870,21 @@ FGameplayEffectSpecHandle UGameplayAbility::GetOutgoingGameplayEffectSpec(const
 	return FGameplayEffectSpecHandle(nullptr);
 }
 
+int32 AbilitySystemShowMakeOutgoingGameplayEffectSpecs = 0;
+static FAutoConsoleVariableRef CVarAbilitySystemShowMakeOutgoingGameplayEffectSpecs(TEXT("AbilitySystem.ShowClientMakeOutgoingSpecs"), AbilitySystemShowMakeOutgoingGameplayEffectSpecs, TEXT("Displays all GameplayEffect specs created on non authority clients"), ECVF_Default );
+
 FGameplayEffectSpecHandle UGameplayAbility::MakeOutgoingGameplayEffectSpec(const FGameplayAbilitySpecHandle Handle, const FGameplayAbilityActorInfo* ActorInfo, const FGameplayAbilityActivationInfo ActivationInfo, TSubclassOf<UGameplayEffect> GameplayEffectClass, float Level) const
 {
 	check(ActorInfo);
 
+	
+#if !(UE_BUILD_SHIPPING || UE_BUILD_TEST)
+	if (AbilitySystemShowMakeOutgoingGameplayEffectSpecs && HasAuthority(&ActivationInfo) == false)
+	{
+		ABILITY_LOG(Warning, TEXT("%s, MakeOutgoingGameplayEffectSpec: %s"), *ActorInfo->AbilitySystemComponent->GetFullName(),  *GameplayEffectClass->GetName()); 
+	}
+#endif
+
 	FGameplayEffectSpecHandle NewHandle = ActorInfo->AbilitySystemComponent->MakeOutgoingSpec(GameplayEffectClass, Level, GetEffectContext(Handle, ActorInfo));
 	if (NewHandle.IsValid())
 	{
@@ -1051,12 +1066,21 @@ void UGameplayAbility::OnTaskActivated(UGameplayTask& Task)
 
 void UGameplayAbility::ConfirmTaskByInstanceName(FName InstanceName, bool bEndTask)
 {
-	TArray<TWeakObjectPtr<UGameplayTask>> NamedTasks = ActiveTasks.FilterByPredicate<FGameplayTaskInstanceNamePredicate>(FGameplayTaskInstanceNamePredicate(InstanceName));
+	TArray<UGameplayTask*, TInlineAllocator<8> > NamedTasks;
+
+	for (UGameplayTask* Task : ActiveTasks)
+	{
+		if (Task && Task->GetInstanceName() == InstanceName)
+		{
+			NamedTasks.Add(Task);
+		}
+	}
+	
 	for (int32 i = NamedTasks.Num() - 1; i >= 0; --i)
 	{
-		if (NamedTasks[i].IsValid())
+		UGameplayTask* CurrentTask = NamedTasks[i];
+		if (CurrentTask && CurrentTask->IsPendingKill() == false)
 		{
-			UGameplayTask* CurrentTask = NamedTasks[i].Get();
 			CurrentTask->ExternalConfirm(bEndTask);
 		}
 	}
@@ -1064,30 +1088,58 @@ void UGameplayAbility::ConfirmTaskByInstanceName(FName InstanceName, bool bEndTa
 
 void UGameplayAbility::EndOrCancelTasksByInstanceName()
 {
+	// Static array for avoiding memory allocations
+	TArray<UGameplayTask*, TInlineAllocator<8> > NamedTasks;
+
+	// Call Endtask on everything in EndTaskInstanceNames list
 	for (int32 j = 0; j < EndTaskInstanceNames.Num(); ++j)
 	{
 		FName InstanceName = EndTaskInstanceNames[j];
-		TArray<TWeakObjectPtr<UGameplayTask>> NamedTasks = ActiveTasks.FilterByPredicate<FGameplayTaskInstanceNamePredicate>(FGameplayTaskInstanceNamePredicate(InstanceName));
+		NamedTasks.Reset();
+
+		// Find every current task that needs to end before ending any
+		for (UGameplayTask* Task : ActiveTasks)
+		{
+			if (Task && Task->GetInstanceName() == InstanceName)
+			{
+				NamedTasks.Add(Task);
+			}
+		}
+		
+		// End each one individually. Not ending a task may do "anything" including killing other tasks or the ability itself
 		for (int32 i = NamedTasks.Num() - 1; i >= 0; --i)
 		{
-			if (NamedTasks[i].IsValid())
+			UGameplayTask* CurrentTask = NamedTasks[i];
+			if (CurrentTask && CurrentTask->IsPendingKill() == false)
 			{
-				UGameplayTask* CurrentTask = NamedTasks[i].Get();
 				CurrentTask->EndTask();
 			}
 		}
 	}
 	EndTaskInstanceNames.Empty();
 
+
+	// Call ExternalCancel on everything in CancelTaskInstanceNames list
 	for (int32 j = 0; j < CancelTaskInstanceNames.Num(); ++j)
 	{
 		FName InstanceName = CancelTaskInstanceNames[j];
-		TArray<TWeakObjectPtr<UGameplayTask>> NamedTasks = ActiveTasks.FilterByPredicate<FGameplayTaskInstanceNamePredicate>(FGameplayTaskInstanceNamePredicate(InstanceName));
+		NamedTasks.Reset();
+		
+		// Find every current task that needs to cancel before cancelling any
+		for (UGameplayTask* Task : ActiveTasks)
+		{
+			if (Task && Task->GetInstanceName() == InstanceName)
+			{
+				NamedTasks.Add(Task);
+			}
+		}
+
+		// Cancel each one individually.  Not canceling a task may do "anything" including killing other tasks or the ability itself
 		for (int32 i = NamedTasks.Num() - 1; i >= 0; --i)
 		{
-			if (NamedTasks[i].IsValid())
+			UGameplayTask* CurrentTask = NamedTasks[i];
+			if (CurrentTask && CurrentTask->IsPendingKill() == false)
 			{
-				UGameplayTask* CurrentTask = NamedTasks[i].Get();
 				CurrentTask->ExternalCancel();
 			}
 		}
@@ -1422,6 +1474,12 @@ TArray<FActiveGameplayEffectHandle> UGameplayAbility::ApplyGameplayEffectToTarge
 {
 	TArray<FActiveGameplayEffectHandle> EffectHandles;
 
+	if (HasAuthority(&ActivationInfo) == false && UAbilitySystemGlobals::Get().ShouldPredictTargetGameplayEffects() == false)
+	{
+		// Early out to avoid making effect specs that we can't apply
+		return EffectHandles;
+	}
+
 	// This batches all created cues together
 	FScopedGameplayCueSendContext GameplayCueSendContext;
 
diff --git a/Engine/Source/Runtime/GameplayAbilities/Private/AbilitySystemComponent.cpp b/Engine/Source/Runtime/GameplayAbilities/Private/AbilitySystemComponent.cpp
--- a/Engine/Source/Runtime/GameplayAbilities/Private/AbilitySystemComponent.cpp
+++ b/Engine/Source/Runtime/GameplayAbilities/Private/AbilitySystemComponent.cpp
@@ -1371,7 +1371,7 @@ void UAbilitySystemComponent::OnPeriodicGameplayEffectExecuteOnSelf(UAbilitySyst
 	OnPeriodicGameplayEffectExecuteDelegateOnSelf.Broadcast(Source, SpecExecuted, ActiveHandle);
 }
 
-TArray<TWeakObjectPtr<UGameplayTask> >&	UAbilitySystemComponent::GetAbilityActiveTasks(UGameplayAbility* Ability)
+TArray<UGameplayTask*>&	UAbilitySystemComponent::GetAbilityActiveTasks(UGameplayAbility* Ability)
 {
 	return Ability->ActiveTasks;
 }
@@ -1723,7 +1723,7 @@ void UAbilitySystemComponent::DisplayDebug(class UCanvas* Canvas, const class FD
 					Canvas->SetDrawColor(FColor::White);
 					for (auto TaskPtr : Instance->ActiveTasks)
 					{
-						UGameplayTask* Task = TaskPtr.Get();
+						UGameplayTask* Task = TaskPtr;
 						if (Task)
 						{
 							YL = Canvas->DrawText(GEngine->GetTinyFont(), FString::Printf(TEXT("%s"), *Task->GetDebugString()), XPos + 7.f, YPos);
diff --git a/Engine/Source/Runtime/GameplayAbilities/Private/GameplayEffect.cpp b/Engine/Source/Runtime/GameplayAbilities/Private/GameplayEffect.cpp
--- a/Engine/Source/Runtime/GameplayAbilities/Private/GameplayEffect.cpp
+++ b/Engine/Source/Runtime/GameplayAbilities/Private/GameplayEffect.cpp
@@ -843,15 +843,7 @@ void FGameplayEffectSpec::GetAllAssetTags(OUT FGameplayTagContainer& Container)
 
 void FGameplayEffectSpec::SetSetByCallerMagnitude(FName DataName, float Magnitude)
 {
-#if !(UE_BUILD_SHIPPING || UE_BUILD_TEST)
-	const float* CurrentValue = SetByCallerMagnitudes.Find(DataName);
-	if (CurrentValue)
-	{
-		ABILITY_LOG(Error, TEXT("FGameplayEffectSpec::SetMagnitude called on Data %s for Def %s when this magnitude was already set. Current Value: %.2f"), *DataName.ToString(), *Def->GetName(), *CurrentValue);
-	}
-#endif
-
-	SetByCallerMagnitudes.Add(DataName) = Magnitude;
+	SetByCallerMagnitudes.FindOrAdd(DataName) = Magnitude;
 }
 
 float FGameplayEffectSpec::GetSetByCallerMagnitude(FName DataName, bool WarnIfNotFound, float DefaultIfNotFound) const
diff --git a/Engine/Source/Runtime/GameplayAbilities/Public/Abilities/GameplayAbility.h b/Engine/Source/Runtime/GameplayAbilities/Public/Abilities/GameplayAbility.h
--- a/Engine/Source/Runtime/GameplayAbilities/Public/Abilities/GameplayAbility.h
+++ b/Engine/Source/Runtime/GameplayAbilities/Public/Abilities/GameplayAbility.h
@@ -735,7 +735,8 @@ protected:
 	//
 	// ----------------------------------------------------------------------------------------------------------------
 	
-	TArray<TWeakObjectPtr<UGameplayTask> >	ActiveTasks;
+	UPROPERTY()
+	TArray<UGameplayTask*>	ActiveTasks;
 
 	// ----------------------------------------------------------------------------------------------------------------
 	//
diff --git a/Engine/Source/Runtime/GameplayAbilities/Public/AbilitySystemComponent.h b/Engine/Source/Runtime/GameplayAbilities/Public/AbilitySystemComponent.h
--- a/Engine/Source/Runtime/GameplayAbilities/Public/AbilitySystemComponent.h
+++ b/Engine/Source/Runtime/GameplayAbilities/Public/AbilitySystemComponent.h
@@ -1123,7 +1123,7 @@ protected:
 
 	void OnAttributeGameplayEffectSpecExected(const FGameplayAttribute &Attribute, const struct FGameplayEffectSpec &Spec, struct FGameplayModifierEvaluatedData &Data);
 		
-	TArray<TWeakObjectPtr<UGameplayTask> >&	GetAbilityActiveTasks(UGameplayAbility* Ability);
+	TArray<UGameplayTask*>&	GetAbilityActiveTasks(UGameplayAbility* Ability);
 	// --------------------------------------------
 	
 	// Contains all of the gameplay effects that are currently active on this component
diff --git a/Engine/Source/Runtime/GameplayTasks/Classes/GameplayTask.h b/Engine/Source/Runtime/GameplayTasks/Classes/GameplayTask.h
--- a/Engine/Source/Runtime/GameplayTasks/Classes/GameplayTask.h
+++ b/Engine/Source/Runtime/GameplayTasks/Classes/GameplayTask.h
@@ -364,9 +364,9 @@ struct FGameplayTaskInstanceNamePredicate
 		InstanceName = DesiredInstanceName;
 	}
 
-	bool operator()(const TWeakObjectPtr<UGameplayTask> A) const
+	bool operator()(const UGameplayTask* A) const
 	{
-		return (A.IsValid() && !A.Get()->GetInstanceName().IsNone() && A.Get()->GetInstanceName().IsValid() && (A.Get()->GetInstanceName() == InstanceName));
+		return (A && !A->GetInstanceName().IsNone() && A->GetInstanceName().IsValid() && (A->GetInstanceName() == InstanceName));
 	}
 
 	FName InstanceName;
@@ -380,9 +380,9 @@ struct FGameplayTaskInstanceClassPredicate
 		TaskClass = Class;
 	}
 
-	bool operator()(const TWeakObjectPtr<UGameplayTask> A) const
+	bool operator()(const UGameplayTask* A) const
 	{
-		return (A.IsValid() && (A.Get()->GetClass() == TaskClass));
+		return (A && (A->GetClass() == TaskClass));
 	}
 
 	TSubclassOf<UGameplayTask> TaskClass;
diff --git a/Engine/Source/Runtime/Landscape/Classes/LandscapeProxy.h b/Engine/Source/Runtime/Landscape/Classes/LandscapeProxy.h
--- a/Engine/Source/Runtime/Landscape/Classes/LandscapeProxy.h
+++ b/Engine/Source/Runtime/Landscape/Classes/LandscapeProxy.h
@@ -126,6 +126,20 @@ struct FLandscapeLayerStruct
 	}
 };
 
+UENUM()
+enum class ELandscapeImportAlphamapType : uint8
+{
+	// Three layers blended 50/30/20 represented as 0.5, 0.3, and 0.2 in the alpha maps
+	// All alpha maps for blended layers total to 1.0
+	// This is the style used by UE4 internally for blended layers
+	Additive,
+
+	// Three layers blended 50/30/20 represented as 0.5, 0.6, and 1.0 in the alpha maps
+	// Each alpha map only specifies the remainder from previous layers, so the last layer used will always be 1.0
+	// Some other tools use this format
+	Layered,
+};
+
 /** Structure storing Layer Data for import */
 USTRUCT()
 struct FLandscapeImportLayerInfo
@@ -633,10 +647,9 @@ public:
 
 	LANDSCAPE_API static ULandscapeMaterialInstanceConstant* GetLayerThumbnailMIC(UMaterialInterface* LandscapeMaterial, FName LayerName, UTexture2D* ThumbnailWeightmap, UTexture2D* ThumbnailHeightmap, ALandscapeProxy* Proxy);
 
-	LANDSCAPE_API void Import(FGuid Guid, int32 VertsX, int32 VertsY, 
-							int32 ComponentSizeQuads, int32 NumSubsections, int32 SubsectionSizeQuads, 
-							const uint16* HeightData, const TCHAR* HeightmapFileName, 
-							const TArray<FLandscapeImportLayerInfo>& ImportLayerInfos);
+	LANDSCAPE_API void Import(FGuid Guid, int32 MinX, int32 MinY, int32 MaxX, int32 MaxY, int32 NumSubsections, int32 SubsectionSizeQuads,
+							const uint16* HeightData, const TCHAR* HeightmapFileName,
+							const TArray<FLandscapeImportLayerInfo>& ImportLayerInfos, ELandscapeImportAlphamapType ImportLayerType);
 
 	/**
 	 * Exports landscape into raw mesh
diff --git a/Engine/Source/Runtime/Landscape/Private/LandscapeEdit.cpp b/Engine/Source/Runtime/Landscape/Private/LandscapeEdit.cpp
--- a/Engine/Source/Runtime/Landscape/Private/LandscapeEdit.cpp
+++ b/Engine/Source/Runtime/Landscape/Private/LandscapeEdit.cpp
@@ -1734,41 +1734,44 @@ ULandscapeLayerInfoObject* ALandscapeProxy::CreateLayerInfo(const TCHAR* LayerNa
 }
 
 #define HEIGHTDATA(X,Y) (HeightData[ FMath::Clamp<int32>(Y,0,VertsY) * VertsX + FMath::Clamp<int32>(X,0,VertsX) ])
-void ALandscapeProxy::Import(FGuid Guid, int32 VertsX, int32 VertsY,
-	int32 InComponentSizeQuads, int32 InNumSubsections, int32 InSubsectionSizeQuads,
-	const uint16* HeightData, const TCHAR* HeightmapFileName,
-	const TArray<FLandscapeImportLayerInfo>& ImportLayerInfos)
+LANDSCAPE_API void ALandscapeProxy::Import(
+	const FGuid Guid,
+	const int32 MinX, const int32 MinY, const int32 MaxX, const int32 MaxY,
+	const int32 InNumSubsections, const int32 InSubsectionSizeQuads,
+	const uint16* const HeightData, const TCHAR* const HeightmapFileName,
+	const TArray<FLandscapeImportLayerInfo>& ImportLayerInfos, const ELandscapeImportAlphamapType ImportLayerType)
 {
 	GWarn->BeginSlowTask(LOCTEXT("BeingImportingLandscapeTask", "Importing Landscape"), true);
 
-	ComponentSizeQuads = InComponentSizeQuads;
+	const int32 VertsX = MaxX - MinX + 1;
+	const int32 VertsY = MaxY - MinY + 1;
+
+	ComponentSizeQuads = InNumSubsections * InSubsectionSizeQuads;
 	NumSubsections = InNumSubsections;
 	SubsectionSizeQuads = InSubsectionSizeQuads;
 	LandscapeGuid = Guid;
 
-	MarkPackageDirty();
+	Modify();
 
 	// Create and initialize landscape info object
 	GetLandscapeInfo(true)->RegisterActor(this);
 
-	int32 NumPatchesX = (VertsX - 1);
-	int32 NumPatchesY = (VertsY - 1);
+	const int32 NumPatchesX = (VertsX - 1);
+	const int32 NumPatchesY = (VertsY - 1);
 
-	int32 NumSectionsX = NumPatchesX / ComponentSizeQuads;
-	int32 NumSectionsY = NumPatchesY / ComponentSizeQuads;
+	const int32 NumComponentsX = NumPatchesX / ComponentSizeQuads;
+	const int32 NumComponentsY = NumPatchesY / ComponentSizeQuads;
 
-	LandscapeComponents.Empty(NumSectionsX * NumSectionsY);
+	// currently only support importing into a new/blank landscape actor/proxy
+	check(LandscapeComponents.Num() == 0);
+	LandscapeComponents.Empty(NumComponentsX * NumComponentsY);
 
-	for (int32 Y = 0; Y < NumSectionsY; Y++)
+	for (int32 Y = 0; Y < NumComponentsY; Y++)
 	{
-		for (int32 X = 0; X < NumSectionsX; X++)
+		for (int32 X = 0; X < NumComponentsX; X++)
 		{
-			// The number of quads
-			int32 NumQuadsX = NumPatchesX;
-			int32 NumQuadsY = NumPatchesY;
-
-			int32 BaseX = X * ComponentSizeQuads;
-			int32 BaseY = Y * ComponentSizeQuads;
+			const int32 BaseX = MinX + X * ComponentSizeQuads;
+			const int32 BaseY = MinY + Y * ComponentSizeQuads;
 
 			ULandscapeComponent* LandscapeComponent = NewObject<ULandscapeComponent>(this, NAME_None, RF_Transactional);
 			LandscapeComponent->SetRelativeLocation(FVector(BaseX, BaseY, 0));
@@ -1787,17 +1790,14 @@ void ALandscapeProxy::Import(FGuid Guid, int32 VertsX, int32 VertsY,
 		}
 	}
 
-#define MAX_HEIGHTMAP_TEXTURE_SIZE 512
-
-	int32 ComponentSizeVerts = NumSubsections * (SubsectionSizeQuads + 1);
-	int32 ComponentsPerHeightmap = MAX_HEIGHTMAP_TEXTURE_SIZE / ComponentSizeVerts;
-
 	// Ensure that we don't pack so many heightmaps into a texture that their lowest LOD isn't guaranteed to be resident
-	ComponentsPerHeightmap = FMath::Min(ComponentsPerHeightmap, 1 << (UTexture2D::GetMinTextureResidentMipCount() - 2));
+#define MAX_HEIGHTMAP_TEXTURE_SIZE 512
+	const int32 ComponentSizeVerts = NumSubsections * (SubsectionSizeQuads + 1);
+	const int32 ComponentsPerHeightmap = FMath::Min(MAX_HEIGHTMAP_TEXTURE_SIZE / ComponentSizeVerts, 1 << (UTexture2D::GetMinTextureResidentMipCount() - 2));
 
 	// Count how many heightmaps we need and the X dimension of the final heightmap
 	int32 NumHeightmapsX = 1;
-	int32 FinalComponentsX = NumSectionsX;
+	int32 FinalComponentsX = NumComponentsX;
 	while (FinalComponentsX > ComponentsPerHeightmap)
 	{
 		FinalComponentsX -= ComponentsPerHeightmap;
@@ -1805,7 +1805,7 @@ void ALandscapeProxy::Import(FGuid Guid, int32 VertsX, int32 VertsY,
 	}
 	// Count how many heightmaps we need and the Y dimension of the final heightmap
 	int32 NumHeightmapsY = 1;
-	int32 FinalComponentsY = NumSectionsY;
+	int32 FinalComponentsY = NumComponentsY;
 	while (FinalComponentsY > ComponentsPerHeightmap)
 	{
 		FinalComponentsY -= ComponentsPerHeightmap;
@@ -1821,8 +1821,8 @@ void ALandscapeProxy::Import(FGuid Guid, int32 VertsX, int32 VertsY,
 			FHeightmapInfo& HeightmapInfo = HeightmapInfos[HeightmapInfos.AddZeroed()];
 
 			// make sure the heightmap UVs are powers of two.
-			HeightmapInfo.HeightmapSizeU = (1 << FMath::CeilLogTwo(((HmX == NumHeightmapsX - 1) ? FinalComponentsX : ComponentsPerHeightmap)*ComponentSizeVerts));
-			HeightmapInfo.HeightmapSizeV = (1 << FMath::CeilLogTwo(((HmY == NumHeightmapsY - 1) ? FinalComponentsY : ComponentsPerHeightmap)*ComponentSizeVerts));
+			HeightmapInfo.HeightmapSizeU = (1 << FMath::CeilLogTwo(((HmX == NumHeightmapsX - 1) ? FinalComponentsX : ComponentsPerHeightmap) * ComponentSizeVerts));
+			HeightmapInfo.HeightmapSizeV = (1 << FMath::CeilLogTwo(((HmY == NumHeightmapsY - 1) ? FinalComponentsY : ComponentsPerHeightmap) * ComponentSizeVerts));
 
 			// Construct the heightmap textures
 			HeightmapInfo.HeightmapTexture = CreateLandscapeTexture(HeightmapInfo.HeightmapSizeU, HeightmapInfo.HeightmapSizeV, TEXTUREGROUP_Terrain_Heightmap, TSF_BGRA8);
@@ -1848,41 +1848,41 @@ void ALandscapeProxy::Import(FGuid Guid, int32 VertsX, int32 VertsY,
 	const FVector DrawScale3D = GetRootComponent()->RelativeScale3D;
 
 	// Calculate the normals for each of the two triangles per quad.
-	FVector* VertexNormals = new FVector[(NumPatchesX + 1)*(NumPatchesY + 1)];
-	FMemory::Memzero(VertexNormals, (NumPatchesX + 1)*(NumPatchesY + 1)*sizeof(FVector));
+	TArray<FVector> VertexNormals;
+	VertexNormals.AddZeroed(VertsX * VertsY);
 	for (int32 QuadY = 0; QuadY < NumPatchesY; QuadY++)
 	{
 		for (int32 QuadX = 0; QuadX < NumPatchesX; QuadX++)
 		{
-			FVector Vert00 = FVector(0.0f, 0.0f, ((float)HEIGHTDATA(QuadX + 0, QuadY + 0) - 32768.0f)*LANDSCAPE_ZSCALE) * DrawScale3D;
-			FVector Vert01 = FVector(0.0f, 1.0f, ((float)HEIGHTDATA(QuadX + 0, QuadY + 1) - 32768.0f)*LANDSCAPE_ZSCALE) * DrawScale3D;
-			FVector Vert10 = FVector(1.0f, 0.0f, ((float)HEIGHTDATA(QuadX + 1, QuadY + 0) - 32768.0f)*LANDSCAPE_ZSCALE) * DrawScale3D;
-			FVector Vert11 = FVector(1.0f, 1.0f, ((float)HEIGHTDATA(QuadX + 1, QuadY + 1) - 32768.0f)*LANDSCAPE_ZSCALE) * DrawScale3D;
+			const FVector Vert00 = FVector(0.0f, 0.0f, ((float)HEIGHTDATA(QuadX + 0, QuadY + 0) - 32768.0f)*LANDSCAPE_ZSCALE) * DrawScale3D;
+			const FVector Vert01 = FVector(0.0f, 1.0f, ((float)HEIGHTDATA(QuadX + 0, QuadY + 1) - 32768.0f)*LANDSCAPE_ZSCALE) * DrawScale3D;
+			const FVector Vert10 = FVector(1.0f, 0.0f, ((float)HEIGHTDATA(QuadX + 1, QuadY + 0) - 32768.0f)*LANDSCAPE_ZSCALE) * DrawScale3D;
+			const FVector Vert11 = FVector(1.0f, 1.0f, ((float)HEIGHTDATA(QuadX + 1, QuadY + 1) - 32768.0f)*LANDSCAPE_ZSCALE) * DrawScale3D;
 
-			FVector FaceNormal1 = ((Vert00 - Vert10) ^ (Vert10 - Vert11)).GetSafeNormal();
-			FVector FaceNormal2 = ((Vert11 - Vert01) ^ (Vert01 - Vert00)).GetSafeNormal();
+			const FVector FaceNormal1 = ((Vert00 - Vert10) ^ (Vert10 - Vert11)).GetSafeNormal();
+			const FVector FaceNormal2 = ((Vert11 - Vert01) ^ (Vert01 - Vert00)).GetSafeNormal();
 
 			// contribute to the vertex normals.
-			VertexNormals[(QuadX + 1 + (NumPatchesX + 1)*(QuadY + 0))] += FaceNormal1;
-			VertexNormals[(QuadX + 0 + (NumPatchesX + 1)*(QuadY + 1))] += FaceNormal2;
-			VertexNormals[(QuadX + 0 + (NumPatchesX + 1)*(QuadY + 0))] += FaceNormal1 + FaceNormal2;
-			VertexNormals[(QuadX + 1 + (NumPatchesX + 1)*(QuadY + 1))] += FaceNormal1 + FaceNormal2;
+			VertexNormals[(QuadX + 1 + VertsX * (QuadY + 0))] += FaceNormal1;
+			VertexNormals[(QuadX + 0 + VertsX * (QuadY + 1))] += FaceNormal2;
+			VertexNormals[(QuadX + 0 + VertsX * (QuadY + 0))] += FaceNormal1 + FaceNormal2;
+			VertexNormals[(QuadX + 1 + VertsX * (QuadY + 1))] += FaceNormal1 + FaceNormal2;
 		}
 	}
 
 	// Weight values for each layer for each component.
-	TArray<TArray<TArray<uint8> > > ComponentWeightValues;
-	ComponentWeightValues.AddZeroed(NumSectionsX*NumSectionsY);
+	TArray<TArray<TArray<uint8>>> ComponentWeightValues;
+	ComponentWeightValues.AddZeroed(NumComponentsX * NumComponentsY);
 
-	for (int32 ComponentY = 0; ComponentY < NumSectionsY; ComponentY++)
+	for (int32 ComponentY = 0; ComponentY < NumComponentsY; ComponentY++)
 	{
-		for (int32 ComponentX = 0; ComponentX < NumSectionsX; ComponentX++)
+		for (int32 ComponentX = 0; ComponentX < NumComponentsX; ComponentX++)
 		{
-			ULandscapeComponent* LandscapeComponent = LandscapeComponents[ComponentX + ComponentY*NumSectionsX];
-			TArray<TArray<uint8> >& WeightValues = ComponentWeightValues[ComponentX + ComponentY*NumSectionsX];
+			ULandscapeComponent* const LandscapeComponent = LandscapeComponents[ComponentX + ComponentY*NumComponentsX];
+			TArray<TArray<uint8>>& WeightValues = ComponentWeightValues[ComponentX + ComponentY*NumComponentsX];
 
 			// Import alphamap data into local array and check for unused layers for this component.
-			TArray<FLandscapeComponentAlphaInfo> EditingAlphaLayerData;
+			TArray<FLandscapeComponentAlphaInfo, TInlineAllocator<16>> EditingAlphaLayerData;
 			for (int32 LayerIndex = 0; LayerIndex < ImportLayerInfos.Num(); LayerIndex++)
 			{
 				FLandscapeComponentAlphaInfo* NewAlphaInfo = new(EditingAlphaLayerData) FLandscapeComponentAlphaInfo(LandscapeComponent, LayerIndex);
@@ -1891,8 +1891,8 @@ void ALandscapeProxy::Import(FGuid Guid, int32 VertsX, int32 VertsY,
 				{
 					for (int32 AlphaY = 0; AlphaY <= LandscapeComponent->ComponentSizeQuads; AlphaY++)
 					{
-						const uint8* OldAlphaRowStart = &ImportLayerInfos[LayerIndex].LayerData[(AlphaY + LandscapeComponent->GetSectionBase().Y) * VertsX + (LandscapeComponent->GetSectionBase().X)];
-						uint8* NewAlphaRowStart = &NewAlphaInfo->AlphaValues[AlphaY * (LandscapeComponent->ComponentSizeQuads + 1)];
+						const uint8* const OldAlphaRowStart = &ImportLayerInfos[LayerIndex].LayerData[(AlphaY + LandscapeComponent->GetSectionBase().Y - MinY) * VertsX + (LandscapeComponent->GetSectionBase().X - MinX)];
+						uint8* const NewAlphaRowStart = &NewAlphaInfo->AlphaValues[AlphaY * (LandscapeComponent->ComponentSizeQuads + 1)];
 						FMemory::Memcpy(NewAlphaRowStart, OldAlphaRowStart, LandscapeComponent->ComponentSizeQuads + 1);
 					}
 				}
@@ -1915,7 +1915,7 @@ void ALandscapeProxy::Import(FGuid Guid, int32 VertsX, int32 VertsY,
 			WeightValues.AddZeroed(EditingAlphaLayerData.Num());
 			LandscapeComponent->WeightmapLayerAllocations.Empty(EditingAlphaLayerData.Num());
 
-			TArray<bool> IsNoBlendArray;
+			TArray<bool, TInlineAllocator<16>> IsNoBlendArray;
 			IsNoBlendArray.Empty(EditingAlphaLayerData.Num());
 			IsNoBlendArray.AddZeroed(EditingAlphaLayerData.Num());
 
@@ -1930,37 +1930,40 @@ void ALandscapeProxy::Import(FGuid Guid, int32 VertsX, int32 VertsY,
 			// Discard the temporary alpha data
 			EditingAlphaLayerData.Empty();
 
-			// For each layer...
-			for (int32 WeightLayerIndex = WeightValues.Num() - 1; WeightLayerIndex >= 0; WeightLayerIndex--)
+			if (ImportLayerType == ELandscapeImportAlphamapType::Layered)
 			{
-				// ... multiply all lower layers'...
-				for (int32 BelowWeightLayerIndex = WeightLayerIndex - 1; BelowWeightLayerIndex >= 0; BelowWeightLayerIndex--)
+				// For each layer...
+				for (int32 WeightLayerIndex = WeightValues.Num() - 1; WeightLayerIndex >= 0; WeightLayerIndex--)
 				{
-					int32 TotalWeight = 0;
-
-					if (IsNoBlendArray[BelowWeightLayerIndex])
+					// ... multiply all lower layers'...
+					for (int32 BelowWeightLayerIndex = WeightLayerIndex - 1; BelowWeightLayerIndex >= 0; BelowWeightLayerIndex--)
 					{
-						continue; // skip no blend
-					}
+						int32 TotalWeight = 0;
 
-					// ... values by...
-					for (int32 Idx = 0; Idx < WeightValues[WeightLayerIndex].Num(); Idx++)
-					{
-						// ... one-minus the current layer's values
-						int32 NewValue = (int32)WeightValues[BelowWeightLayerIndex][Idx] * (int32)(255 - WeightValues[WeightLayerIndex][Idx]) / 255;
-						WeightValues[BelowWeightLayerIndex][Idx] = (uint8)NewValue;
-						TotalWeight += NewValue;
-					}
+						if (IsNoBlendArray[BelowWeightLayerIndex])
+						{
+							continue; // skip no blend
+						}
 
-					if (TotalWeight == 0)
-					{
-						// Remove the layer as it has no contribution
-						WeightValues.RemoveAt(BelowWeightLayerIndex);
-						LandscapeComponent->WeightmapLayerAllocations.RemoveAt(BelowWeightLayerIndex);
-						IsNoBlendArray.RemoveAt(BelowWeightLayerIndex);
+						// ... values by...
+						for (int32 Idx = 0; Idx < WeightValues[WeightLayerIndex].Num(); Idx++)
+						{
+							// ... one-minus the current layer's values
+							int32 NewValue = (int32)WeightValues[BelowWeightLayerIndex][Idx] * (int32)(255 - WeightValues[WeightLayerIndex][Idx]) / 255;
+							WeightValues[BelowWeightLayerIndex][Idx] = (uint8)NewValue;
+							TotalWeight += NewValue;
+						}
 
-						// The current layer has been re-numbered
-						WeightLayerIndex--;
+						if (TotalWeight == 0)
+						{
+							// Remove the layer as it has no contribution
+							WeightValues.RemoveAt(BelowWeightLayerIndex);
+							LandscapeComponent->WeightmapLayerAllocations.RemoveAt(BelowWeightLayerIndex);
+							IsNoBlendArray.RemoveAt(BelowWeightLayerIndex);
+
+							// The current layer has been re-numbered
+							WeightLayerIndex--;
+						}
 					}
 				}
 			}
@@ -2022,29 +2025,29 @@ void ALandscapeProxy::Import(FGuid Guid, int32 VertsX, int32 VertsY,
 	// Remember where we have spare texture channels.
 	TArray<FWeightmapTextureAllocation> TextureAllocations;
 
-	for (int32 ComponentY = 0; ComponentY < NumSectionsY; ComponentY++)
+	for (int32 ComponentY = 0; ComponentY < NumComponentsY; ComponentY++)
 	{
-		int32 HmY = ComponentY / ComponentsPerHeightmap;
-		int32 HeightmapOffsetY = (ComponentY - ComponentsPerHeightmap*HmY) * NumSubsections * (SubsectionSizeQuads + 1);
+		const int32 HmY = ComponentY / ComponentsPerHeightmap;
+		const int32 HeightmapOffsetY = (ComponentY - ComponentsPerHeightmap*HmY) * NumSubsections * (SubsectionSizeQuads + 1);
 
-		for (int32 ComponentX = 0; ComponentX < NumSectionsX; ComponentX++)
+		for (int32 ComponentX = 0; ComponentX < NumComponentsX; ComponentX++)
 		{
-			int32 HmX = ComponentX / ComponentsPerHeightmap;
-			FHeightmapInfo& HeightmapInfo = HeightmapInfos[HmX + HmY * NumHeightmapsX];
+			const int32 HmX = ComponentX / ComponentsPerHeightmap;
+			const FHeightmapInfo& HeightmapInfo = HeightmapInfos[HmX + HmY * NumHeightmapsX];
 
-			ULandscapeComponent* LandscapeComponent = LandscapeComponents[ComponentX + ComponentY*NumSectionsX];
+			ULandscapeComponent* LandscapeComponent = LandscapeComponents[ComponentX + ComponentY*NumComponentsX];
 
 			// Lookup array of weight values for this component.
-			TArray<TArray<uint8> >& WeightValues = ComponentWeightValues[ComponentX + ComponentY*NumSectionsX];
+			const TArray<TArray<uint8>>& WeightValues = ComponentWeightValues[ComponentX + ComponentY*NumComponentsX];
 
 			// Heightmap offsets
-			int32 HeightmapOffsetX = (ComponentX - ComponentsPerHeightmap*HmX) * NumSubsections * (SubsectionSizeQuads + 1);
+			const int32 HeightmapOffsetX = (ComponentX - ComponentsPerHeightmap*HmX) * NumSubsections * (SubsectionSizeQuads + 1);
 
 			LandscapeComponent->HeightmapScaleBias = FVector4(1.0f / (float)HeightmapInfo.HeightmapSizeU, 1.0f / (float)HeightmapInfo.HeightmapSizeV, (float)((HeightmapOffsetX)) / (float)HeightmapInfo.HeightmapSizeU, ((float)(HeightmapOffsetY)) / (float)HeightmapInfo.HeightmapSizeV);
 			LandscapeComponent->HeightmapTexture = HeightmapInfo.HeightmapTexture;
 
 			// Weightmap is sized the same as the component
-			int32 WeightmapSize = (SubsectionSizeQuads + 1) * NumSubsections;
+			const int32 WeightmapSize = (SubsectionSizeQuads + 1) * NumSubsections;
 			// Should be power of two
 			check(FMath::IsPowerOfTwo(WeightmapSize));
 
@@ -2060,7 +2063,7 @@ void ALandscapeProxy::Import(FGuid Guid, int32 VertsX, int32 VertsY,
 			int32 LayerIndex = 0;
 			while (LayerIndex < WeightValues.Num())
 			{
-				int32 RemainingLayers = WeightValues.Num() - LayerIndex;
+				const int32 RemainingLayers = WeightValues.Num() - LayerIndex;
 
 				int32 BestAllocationIndex = -1;
 
@@ -2073,7 +2076,7 @@ void ALandscapeProxy::Import(FGuid Guid, int32 VertsX, int32 VertsY,
 						if (TextureAllocations[TryAllocIdx].ChannelsInUse + RemainingLayers <= 4)
 						{
 							FWeightmapTextureAllocation& TryAllocation = TextureAllocations[TryAllocIdx];
-							int32 TryDistSquared = FMath::Square(TryAllocation.X - ComponentX) + FMath::Square(TryAllocation.Y - ComponentY);
+							const int32 TryDistSquared = FMath::Square(TryAllocation.X - ComponentX) + FMath::Square(TryAllocation.Y - ComponentY);
 							if (TryDistSquared < BestDistSquared)
 							{
 								BestDistSquared = TryDistSquared;
@@ -2120,11 +2123,11 @@ void ALandscapeProxy::Import(FGuid Guid, int32 VertsX, int32 VertsY,
 				else
 				{
 					// We couldn't find a suitable place for these layers, so lets make a new one.
-					UTexture2D* WeightmapTexture = CreateLandscapeTexture(WeightmapSize, WeightmapSize, TEXTUREGROUP_Terrain_Weightmap, TSF_BGRA8);
-					FColor* MipData = (FColor*)WeightmapTexture->Source.LockMip(0);
+					UTexture2D* const WeightmapTexture = CreateLandscapeTexture(WeightmapSize, WeightmapSize, TEXTUREGROUP_Terrain_Weightmap, TSF_BGRA8);
+					FColor* const MipData = (FColor*)WeightmapTexture->Source.LockMip(0);
 
-					int32 ThisAllocationLayers = FMath::Min<int32>(RemainingLayers, 4);
-					new(TextureAllocations)FWeightmapTextureAllocation(ComponentX, ComponentY, ThisAllocationLayers, WeightmapTexture, MipData);
+					const int32 ThisAllocationLayers = FMath::Min<int32>(RemainingLayers, 4);
+					new(TextureAllocations) FWeightmapTextureAllocation(ComponentX, ComponentY, ThisAllocationLayers, WeightmapTexture, MipData);
 					FLandscapeWeightmapUsage& WeightmapUsage = WeightmapUsageMap.Add(WeightmapTexture, FLandscapeWeightmapUsage());
 
 					UE_LOG(LogLandscape, Log, TEXT("  ==> Storing %d channels in new texture %s"), ThisAllocationLayers, *WeightmapTexture->GetName());
@@ -2164,8 +2167,7 @@ void ALandscapeProxy::Import(FGuid Guid, int32 VertsX, int32 VertsY,
 			}
 			check(WeightmapTextureDataPointers.Num() == WeightValues.Num());
 
-			FVector* LocalVerts = new FVector[FMath::Square(ComponentSizeQuads + 1)];
-
+			FBox LocalBox(0);
 			for (int32 SubsectionY = 0; SubsectionY < NumSubsections; SubsectionY++)
 			{
 				for (int32 SubsectionX = 0; SubsectionX < NumSubsections; SubsectionX++)
@@ -2175,21 +2177,21 @@ void ALandscapeProxy::Import(FGuid Guid, int32 VertsX, int32 VertsY,
 						for (int32 SubX = 0; SubX <= SubsectionSizeQuads; SubX++)
 						{
 							// X/Y of the vertex we're looking at in component's coordinates.
-							int32 CompX = SubsectionSizeQuads * SubsectionX + SubX;
-							int32 CompY = SubsectionSizeQuads * SubsectionY + SubY;
+							const int32 CompX = SubsectionSizeQuads * SubsectionX + SubX;
+							const int32 CompY = SubsectionSizeQuads * SubsectionY + SubY;
 
 							// X/Y of the vertex we're looking indexed into the texture data
-							int32 TexX = (SubsectionSizeQuads + 1) * SubsectionX + SubX;
-							int32 TexY = (SubsectionSizeQuads + 1) * SubsectionY + SubY;
+							const int32 TexX = (SubsectionSizeQuads + 1) * SubsectionX + SubX;
+							const int32 TexY = (SubsectionSizeQuads + 1) * SubsectionY + SubY;
 
-							int32 WeightSrcDataIdx = CompY * (ComponentSizeQuads + 1) + CompX;
-							int32 HeightTexDataIdx = (HeightmapOffsetX + TexX) + (HeightmapOffsetY + TexY) * (HeightmapInfo.HeightmapSizeU);
+							const int32 WeightSrcDataIdx = CompY * (ComponentSizeQuads + 1) + CompX;
+							const int32 HeightTexDataIdx = (HeightmapOffsetX + TexX) + (HeightmapOffsetY + TexY) * (HeightmapInfo.HeightmapSizeU);
 
-							int32 WeightTexDataIdx = (TexX)+(TexY)* (WeightmapSize);
+							const int32 WeightTexDataIdx = (TexX)+(TexY)* (WeightmapSize);
 
 							// copy height and normal data
-							uint16 HeightValue = HEIGHTDATA(CompX + LandscapeComponent->GetSectionBase().X, CompY + LandscapeComponent->GetSectionBase().Y);
-							FVector Normal = VertexNormals[CompX + LandscapeComponent->GetSectionBase().X + (NumPatchesX + 1)*(CompY + LandscapeComponent->GetSectionBase().Y)].GetSafeNormal();
+							const uint16 HeightValue = HEIGHTDATA(CompX + LandscapeComponent->GetSectionBase().X - MinX, CompY + LandscapeComponent->GetSectionBase().Y - MinY);
+							const FVector Normal = VertexNormals[CompX + LandscapeComponent->GetSectionBase().X - MinX + VertsX * (CompY + LandscapeComponent->GetSectionBase().Y - MinY)].GetSafeNormal();
 
 							HeightmapInfo.HeightmapTextureMipData[0][HeightTexDataIdx].R = HeightValue >> 8;
 							HeightmapInfo.HeightmapTextureMipData[0][HeightTexDataIdx].G = HeightValue & 255;
@@ -2202,15 +2204,14 @@ void ALandscapeProxy::Import(FGuid Guid, int32 VertsX, int32 VertsY,
 							}
 
 							// Get local space verts
-							FVector LocalVertex(CompX, CompY, LandscapeDataAccess::GetLocalHeight(HeightValue));
-							LocalVerts[(LandscapeComponent->ComponentSizeQuads + 1) * CompY + CompX] = LocalVertex;
+							const FVector LocalVertex(CompX, CompY, LandscapeDataAccess::GetLocalHeight(HeightValue));
+							LocalBox += LocalVertex;
 						}
 					}
 				}
 			}
 
-			LandscapeComponent->CachedLocalBox = FBox(LocalVerts, FMath::Square(ComponentSizeQuads + 1));
-			delete[] LocalVerts;
+			LandscapeComponent->CachedLocalBox = LocalBox;
 
 			// Update MaterialInstance
 			LandscapeComponent->UpdateMaterialInstances();
@@ -2220,8 +2221,8 @@ void ALandscapeProxy::Import(FGuid Guid, int32 VertsX, int32 VertsY,
 	// Unlock the weightmaps' base mips
 	for (int32 AllocationIndex = 0; AllocationIndex < TextureAllocations.Num(); AllocationIndex++)
 	{
-		UTexture2D* WeightmapTexture = TextureAllocations[AllocationIndex].Texture;
-		FColor* BaseMipData = TextureAllocations[AllocationIndex].TextureData;
+		UTexture2D* const WeightmapTexture = TextureAllocations[AllocationIndex].Texture;
+		FColor* const BaseMipData = TextureAllocations[AllocationIndex].TextureData;
 
 		// Generate mips for weightmaps
 		ULandscapeComponent::GenerateWeightmapMips(NumSubsections, SubsectionSizeQuads, WeightmapTexture, BaseMipData);
@@ -2230,19 +2231,17 @@ void ALandscapeProxy::Import(FGuid Guid, int32 VertsX, int32 VertsY,
 		WeightmapTexture->PostEditChange();
 	}
 
-	delete[] VertexNormals;
-
 	// Generate mipmaps for the components, and create the collision components
-	for (int32 ComponentY = 0; ComponentY < NumSectionsY; ComponentY++)
+	for (int32 ComponentY = 0; ComponentY < NumComponentsY; ComponentY++)
 	{
-		for (int32 ComponentX = 0; ComponentX < NumSectionsX; ComponentX++)
+		for (int32 ComponentX = 0; ComponentX < NumComponentsX; ComponentX++)
 		{
-			int32 HmX = ComponentX / ComponentsPerHeightmap;
-			int32 HmY = ComponentY / ComponentsPerHeightmap;
+			const int32 HmX = ComponentX / ComponentsPerHeightmap;
+			const int32 HmY = ComponentY / ComponentsPerHeightmap;
 			FHeightmapInfo& HeightmapInfo = HeightmapInfos[HmX + HmY * NumHeightmapsX];
 
-			ULandscapeComponent* LandscapeComponent = LandscapeComponents[ComponentX + ComponentY*NumSectionsX];
-			LandscapeComponent->GenerateHeightmapMips(HeightmapInfo.HeightmapTextureMipData, ComponentX == NumSectionsX - 1 ? MAX_int32 : 0, ComponentY == NumSectionsY - 1 ? MAX_int32 : 0);
+			ULandscapeComponent* LandscapeComponent = LandscapeComponents[ComponentX + ComponentY*NumComponentsX];
+			LandscapeComponent->GenerateHeightmapMips(HeightmapInfo.HeightmapTextureMipData, ComponentX == NumComponentsX - 1 ? MAX_int32 : 0, ComponentY == NumComponentsY - 1 ? MAX_int32 : 0);
 			LandscapeComponent->UpdateCollisionHeightData(HeightmapInfo.HeightmapTextureMipData[LandscapeComponent->CollisionMipLevel]);
 			LandscapeComponent->UpdateCollisionLayerData();
 		}
@@ -2260,19 +2259,19 @@ void ALandscapeProxy::Import(FGuid Guid, int32 VertsX, int32 VertsY,
 		while (MipSizeU > 1 && MipSizeV > 1)
 		{
 			HeightmapInfo.HeightmapTextureMipData.Add((FColor*)HeightmapInfo.HeightmapTexture->Source.LockMip(Mip));
-			int32 PrevMipSizeU = (HeightmapInfo.HeightmapTexture->Source.GetSizeX()) >> (Mip - 1);
-			int32 PrevMipSizeV = (HeightmapInfo.HeightmapTexture->Source.GetSizeY()) >> (Mip - 1);
+			const int32 PrevMipSizeU = (HeightmapInfo.HeightmapTexture->Source.GetSizeX()) >> (Mip - 1);
+			const int32 PrevMipSizeV = (HeightmapInfo.HeightmapTexture->Source.GetSizeY()) >> (Mip - 1);
 
 			for (int32 Y = 0; Y < MipSizeV; Y++)
 			{
 				for (int32 X = 0; X < MipSizeU; X++)
 				{
-					FColor* TexData = &(HeightmapInfo.HeightmapTextureMipData[Mip])[X + Y * MipSizeU];
+					FColor* const TexData = &(HeightmapInfo.HeightmapTextureMipData[Mip])[X + Y * MipSizeU];
 
-					FColor *PreMipTexData00 = &(HeightmapInfo.HeightmapTextureMipData[Mip - 1])[(X * 2 + 0) + (Y * 2 + 0)  * PrevMipSizeU];
-					FColor *PreMipTexData01 = &(HeightmapInfo.HeightmapTextureMipData[Mip - 1])[(X * 2 + 0) + (Y * 2 + 1)  * PrevMipSizeU];
-					FColor *PreMipTexData10 = &(HeightmapInfo.HeightmapTextureMipData[Mip - 1])[(X * 2 + 1) + (Y * 2 + 0)  * PrevMipSizeU];
-					FColor *PreMipTexData11 = &(HeightmapInfo.HeightmapTextureMipData[Mip - 1])[(X * 2 + 1) + (Y * 2 + 1)  * PrevMipSizeU];
+					const FColor* const PreMipTexData00 = &(HeightmapInfo.HeightmapTextureMipData[Mip - 1])[(X * 2 + 0) + (Y * 2 + 0)  * PrevMipSizeU];
+					const FColor* const PreMipTexData01 = &(HeightmapInfo.HeightmapTextureMipData[Mip - 1])[(X * 2 + 0) + (Y * 2 + 1)  * PrevMipSizeU];
+					const FColor* const PreMipTexData10 = &(HeightmapInfo.HeightmapTextureMipData[Mip - 1])[(X * 2 + 1) + (Y * 2 + 0)  * PrevMipSizeU];
+					const FColor* const PreMipTexData11 = &(HeightmapInfo.HeightmapTextureMipData[Mip - 1])[(X * 2 + 1) + (Y * 2 + 1)  * PrevMipSizeU];
 
 					TexData->R = (((int32)PreMipTexData00->R + (int32)PreMipTexData01->R + (int32)PreMipTexData10->R + (int32)PreMipTexData11->R) >> 2);
 					TexData->G = (((int32)PreMipTexData00->G + (int32)PreMipTexData01->G + (int32)PreMipTexData10->G + (int32)PreMipTexData11->G) >> 2);
diff --git a/Engine/Source/Runtime/NetworkReplayStreaming/HttpNetworkReplayStreaming/Private/HttpNetworkReplayStreaming.cpp b/Engine/Source/Runtime/NetworkReplayStreaming/HttpNetworkReplayStreaming/Private/HttpNetworkReplayStreaming.cpp
--- a/Engine/Source/Runtime/NetworkReplayStreaming/HttpNetworkReplayStreaming/Private/HttpNetworkReplayStreaming.cpp
+++ b/Engine/Source/Runtime/NetworkReplayStreaming/HttpNetworkReplayStreaming/Private/HttpNetworkReplayStreaming.cpp
@@ -9,7 +9,7 @@ static TAutoConsoleVariable<FString> CVarMetaFilterOverride( TEXT( "httpReplay.M
 class FNetworkReplayListItem : public FOnlineJsonSerializable
 {
 public:
-	FNetworkReplayListItem() {}
+	FNetworkReplayListItem() : SizeInBytes( 0 ), DemoTimeInMs( 0 ), NumViewers( 0 ), bIsLive( false ), Changelist( 0 ), bShouldKeep( false ) {}
 	virtual ~FNetworkReplayListItem() {}
 
 	FString		AppName;
@@ -21,6 +21,7 @@ public:
 	int32		NumViewers;
 	bool		bIsLive;
 	int32		Changelist;
+	bool		bShouldKeep;
 
 	// FOnlineJsonSerializable
 	BEGIN_ONLINE_JSON_SERIALIZER
@@ -33,6 +34,7 @@ public:
 		ONLINE_JSON_SERIALIZE( "NumViewers",	NumViewers );
 		ONLINE_JSON_SERIALIZE( "bIsLive",		bIsLive );
 		ONLINE_JSON_SERIALIZE( "Changelist",	Changelist );
+		ONLINE_JSON_SERIALIZE( "shouldKeep",	bShouldKeep );
 	END_ONLINE_JSON_SERIALIZER
 };
 
@@ -542,9 +544,7 @@ void FHttpNetworkReplayStreamer::SearchEvents(const FString& EventGroup, const F
 	HttpRequest->SetURL(FString::Printf(TEXT("%sevent?group=%s"), *ServerURL, *EventGroup));
 	HttpRequest->SetVerb(TEXT("GET"));
 
-	HttpRequest->OnProcessRequestComplete().BindRaw(this, &FHttpNetworkReplayStreamer::HttpEnumerateSessionsFinished);
-
-	EnumerateStreamsDelegate = Delegate;
+	HttpRequest->OnProcessRequestComplete().BindRaw(this, &FHttpNetworkReplayStreamer::HttpEnumerateSessionsFinished, Delegate);
 
 	AddRequestToQueue(EQueuedHttpRequestType::EnumeratingSessions, HttpRequest);
 }
@@ -995,12 +995,11 @@ void FHttpNetworkReplayStreamer::DeleteFinishedStream( const FString& StreamName
 
 void FHttpNetworkReplayStreamer::EnumerateStreams( const FNetworkReplayVersion& InReplayVersion, const FString& UserString, const FString& MetaString, const FOnEnumerateStreamsComplete& Delegate )
 {
-	if ( IsHttpRequestInFlight() )
-	{
-		Delegate.ExecuteIfBound( TArray<FNetworkReplayStreamInfo>() );
-		return;
-	}
+	EnumerateStreams( InReplayVersion, UserString, MetaString, TArray< FString >(), Delegate );
+}
 
+void FHttpNetworkReplayStreamer::EnumerateStreams( const FNetworkReplayVersion& InReplayVersion, const FString& UserString, const FString& MetaString, const TArray< FString >& ExtraParms, const FOnEnumerateStreamsComplete& Delegate )
+{
 	TSharedRef<class IHttpRequest> HttpRequest = FHttpModule::Get().CreateRequest();
 
 	// Build base URL
@@ -1031,34 +1030,30 @@ void FHttpNetworkReplayStreamer::EnumerateStreams( const FNetworkReplayVersion&
 		URL += FString::Printf( TEXT( "&user=%s" ), *UserString );
 	}
 
+	// Add any extra parms now
+	for ( int i = 0; i < ExtraParms.Num(); i++ )
+	{
+		URL += FString::Printf( TEXT( "&%s" ), *ExtraParms[i] );
+	}
+
 	HttpRequest->SetURL( URL );
 
 	HttpRequest->SetVerb( TEXT( "GET" ) );
 
-	HttpRequest->OnProcessRequestComplete().BindRaw( this, &FHttpNetworkReplayStreamer::HttpEnumerateSessionsFinished );
-
-	EnumerateStreamsDelegate = Delegate;
+	HttpRequest->OnProcessRequestComplete().BindRaw( this, &FHttpNetworkReplayStreamer::HttpEnumerateSessionsFinished, Delegate );
 
 	AddRequestToQueue( EQueuedHttpRequestType::EnumeratingSessions, HttpRequest );
 }
 
 void FHttpNetworkReplayStreamer::EnumerateRecentStreams( const FNetworkReplayVersion& InReplayVersion, const FString& InRecentViewer, const FOnEnumerateStreamsComplete& Delegate )
 {
-	if ( IsHttpRequestInFlight() )
-	{
-		Delegate.ExecuteIfBound( TArray<FNetworkReplayStreamInfo>() );
-		return;
-	}
-
 	TSharedRef<class IHttpRequest> HttpRequest = FHttpModule::Get().CreateRequest();
 
 	// Enumerate all of the sessions
 	HttpRequest->SetURL( FString::Printf( TEXT( "%sreplay?app=%s&version=%u&cl=%u&recent=%s" ), *ServerURL, *InReplayVersion.AppString, InReplayVersion.NetworkVersion, InReplayVersion.Changelist, *InRecentViewer ) );
 	HttpRequest->SetVerb( TEXT( "GET" ) );
 
-	HttpRequest->OnProcessRequestComplete().BindRaw( this, &FHttpNetworkReplayStreamer::HttpEnumerateSessionsFinished );
-
-	EnumerateStreamsDelegate = Delegate;
+	HttpRequest->OnProcessRequestComplete().BindRaw( this, &FHttpNetworkReplayStreamer::HttpEnumerateSessionsFinished, Delegate );
 
 	AddRequestToQueue( EQueuedHttpRequestType::EnumeratingSessions, HttpRequest );
 }
@@ -1547,7 +1542,7 @@ void FHttpNetworkReplayStreamer::HttpRefreshViewerFinished( FHttpRequestPtr Http
 	}
 }
 
-void FHttpNetworkReplayStreamer::HttpEnumerateSessionsFinished( FHttpRequestPtr HttpRequest, FHttpResponsePtr HttpResponse, bool bSucceeded )
+void FHttpNetworkReplayStreamer::HttpEnumerateSessionsFinished( FHttpRequestPtr HttpRequest, FHttpResponsePtr HttpResponse, bool bSucceeded, FOnEnumerateStreamsComplete Delegate )
 {
 	check( InFlightHttpRequest.IsValid() );
 	check( InFlightHttpRequest->Request == HttpRequest );
@@ -1567,7 +1562,7 @@ void FHttpNetworkReplayStreamer::HttpEnumerateSessionsFinished( FHttpRequestPtr
 		if ( !ReplayList.FromJson( JsonString ) )
 		{
 			UE_LOG( LogHttpReplay, Warning, TEXT( "FHttpNetworkReplayStreamer::HttpEnumerateSessionsFinished. FromJson FAILED" ) );
-			EnumerateStreamsDelegate.ExecuteIfBound( TArray<FNetworkReplayStreamInfo>() );		// FIXME: Notify failure here
+			Delegate.ExecuteIfBound( TArray<FNetworkReplayStreamInfo>() );		// FIXME: Notify failure here
 			return;
 		}
 
@@ -1583,19 +1578,18 @@ void FHttpNetworkReplayStreamer::HttpEnumerateSessionsFinished( FHttpRequestPtr
 			NewStream.NumViewers	= ReplayList.Replays[i].NumViewers;
 			NewStream.bIsLive		= ReplayList.Replays[i].bIsLive;
 			NewStream.Changelist	= ReplayList.Replays[i].Changelist;
+			NewStream.bShouldKeep	= ReplayList.Replays[i].bShouldKeep;
 
 			Streams.Add( NewStream );
 		}
 
-		EnumerateStreamsDelegate.ExecuteIfBound( Streams );
+		Delegate.ExecuteIfBound( Streams );
 	}
 	else
 	{
 		UE_LOG( LogHttpReplay, Warning, TEXT( "FHttpNetworkReplayStreamer::HttpEnumerateSessionsFinished. FAILED, %s" ), *BuildRequestErrorString( HttpRequest, HttpResponse ) );
-		EnumerateStreamsDelegate.ExecuteIfBound( TArray<FNetworkReplayStreamInfo>() );		// FIXME: Notify failure here
+		Delegate.ExecuteIfBound( TArray<FNetworkReplayStreamInfo>() );		// FIXME: Notify failure here
 	}
-
-	EnumerateStreamsDelegate = FOnEnumerateStreamsComplete();
 }
 
 void FHttpNetworkReplayStreamer::HttpEnumerateCheckpointsFinished( FHttpRequestPtr HttpRequest, FHttpResponsePtr HttpResponse, bool bSucceeded )
diff --git a/Engine/Source/Runtime/NetworkReplayStreaming/HttpNetworkReplayStreaming/Public/HttpNetworkReplayStreaming.h b/Engine/Source/Runtime/NetworkReplayStreaming/HttpNetworkReplayStreaming/Public/HttpNetworkReplayStreaming.h
--- a/Engine/Source/Runtime/NetworkReplayStreaming/HttpNetworkReplayStreaming/Public/HttpNetworkReplayStreaming.h
+++ b/Engine/Source/Runtime/NetworkReplayStreaming/HttpNetworkReplayStreaming/Public/HttpNetworkReplayStreaming.h
@@ -166,6 +166,7 @@ public:
 	virtual bool		IsLive() const override;
 	virtual void		DeleteFinishedStream( const FString& StreamName, const FOnDeleteFinishedStreamComplete& Delegate ) const override;
 	virtual void		EnumerateStreams( const FNetworkReplayVersion& ReplayVersion, const FString& UserString, const FString& MetaString, const FOnEnumerateStreamsComplete& Delegate ) override;
+	virtual void		EnumerateStreams( const FNetworkReplayVersion& InReplayVersion, const FString& UserString, const FString& MetaString, const TArray< FString >& ExtraParms, const FOnEnumerateStreamsComplete& Delegate ) override;
 	virtual void		EnumerateEvents( const FString& Group, const FEnumerateEventsCompleteDelegate& EnumerationCompleteDelegate ) override;
 	virtual void		EnumerateEvents( const FString& ReplayName, const FString& Group, const FEnumerateEventsCompleteDelegate& EnumerationCompleteDelegate ) override;
 	virtual void		EnumerateRecentStreams( const FNetworkReplayVersion& ReplayVersion, const FString& RecentViewer, const FOnEnumerateStreamsComplete& Delegate ) override;
@@ -218,7 +219,7 @@ public:
 	void HttpUploadStreamFinished( FHttpRequestPtr HttpRequest, FHttpResponsePtr HttpResponse, bool bSucceeded );
 	void HttpUploadCheckpointFinished( FHttpRequestPtr HttpRequest, FHttpResponsePtr HttpResponse, bool bSucceeded );
 	void HttpUploadCustomEventFinished( FHttpRequestPtr HttpRequest, FHttpResponsePtr HttpResponse, bool bSucceeded );
-	void HttpEnumerateSessionsFinished( FHttpRequestPtr HttpRequest, FHttpResponsePtr HttpResponse, bool bSucceeded );
+	void HttpEnumerateSessionsFinished( FHttpRequestPtr HttpRequest, FHttpResponsePtr HttpResponse, bool bSucceeded, FOnEnumerateStreamsComplete Delegate );
 	void HttpEnumerateCheckpointsFinished( FHttpRequestPtr HttpRequest, FHttpResponsePtr HttpResponse, bool bSucceeded );
 	void HttpEnumerateEventsFinished( FHttpRequestPtr HttpRequest, FHttpResponsePtr HttpResponse, bool bSucceeded, FEnumerateEventsCompleteDelegate EnumerateEventsDelegate );
 	void HttpAddUserFinished( FHttpRequestPtr HttpRequest, FHttpResponsePtr HttpResponse, bool bSucceeded );
@@ -255,7 +256,6 @@ public:
 	ENetworkReplayError::Type		StreamerLastError;
 
 	FOnStreamReadyDelegate			StartStreamingDelegate;		// Delegate passed in to StartStreaming
-	FOnEnumerateStreamsComplete		EnumerateStreamsDelegate;
 	FOnCheckpointReadyDelegate		GotoCheckpointDelegate;
 	int32							DownloadCheckpointIndex;
 	int64							LastGotoTimeInMS;
diff --git a/Engine/Source/Runtime/NetworkReplayStreaming/InMemoryNetworkReplayStreaming/Private/InMemoryNetworkReplayStreaming.cpp b/Engine/Source/Runtime/NetworkReplayStreaming/InMemoryNetworkReplayStreaming/Private/InMemoryNetworkReplayStreaming.cpp
--- a/Engine/Source/Runtime/NetworkReplayStreaming/InMemoryNetworkReplayStreaming/Private/InMemoryNetworkReplayStreaming.cpp
+++ b/Engine/Source/Runtime/NetworkReplayStreaming/InMemoryNetworkReplayStreaming/Private/InMemoryNetworkReplayStreaming.cpp
@@ -195,10 +195,15 @@ void FInMemoryNetworkReplayStreamer::DeleteFinishedStream( const FString& Stream
 }
 
 void FInMemoryNetworkReplayStreamer::EnumerateStreams( const FNetworkReplayVersion& ReplayVersion, const FString& UserString, const FString& MetaString, const FOnEnumerateStreamsComplete& Delegate )
+{
+	EnumerateStreams( ReplayVersion, UserString, MetaString, TArray< FString >(), Delegate );
+}
+
+void FInMemoryNetworkReplayStreamer::EnumerateStreams( const FNetworkReplayVersion& ReplayVersion, const FString& UserString, const FString& MetaString, const TArray< FString >& ExtraParms, const FOnEnumerateStreamsComplete& Delegate )
 {
 	TArray<FNetworkReplayStreamInfo> Results;
 
-	for (const auto& StreamPair : OwningFactory->Replays)
+	for ( const auto& StreamPair : OwningFactory->Replays )
 	{
 		// Check version. NetworkVersion and changelist of 0 will ignore version check.
 		const bool NetworkVersionMatches = ReplayVersion.NetworkVersion == StreamPair.Value->NetworkVersion;
@@ -209,14 +214,14 @@ void FInMemoryNetworkReplayStreamer::EnumerateStreams( const FNetworkReplayVersi
 
 		if ( NetworkVersionPasses && ChangelistPasses )
 		{
-			Results.Add(StreamPair.Value->StreamInfo);
+			Results.Add( StreamPair.Value->StreamInfo );
 		}
 	}
 
-	Delegate.ExecuteIfBound(Results);
+	Delegate.ExecuteIfBound( Results );
 }
 
-void FInMemoryNetworkReplayStreamer::AddUserToReplay(const FString& UserString)
+void FInMemoryNetworkReplayStreamer::AddUserToReplay( const FString& UserString )
 {
 	UE_LOG(LogMemoryReplay, Log, TEXT("FInMemoryNetworkReplayStreamer::AddUserToReplay is currently unsupported."));
 }
diff --git a/Engine/Source/Runtime/NetworkReplayStreaming/InMemoryNetworkReplayStreaming/Public/InMemoryNetworkReplayStreaming.h b/Engine/Source/Runtime/NetworkReplayStreaming/InMemoryNetworkReplayStreaming/Public/InMemoryNetworkReplayStreaming.h
--- a/Engine/Source/Runtime/NetworkReplayStreaming/InMemoryNetworkReplayStreaming/Public/InMemoryNetworkReplayStreaming.h
+++ b/Engine/Source/Runtime/NetworkReplayStreaming/InMemoryNetworkReplayStreaming/Public/InMemoryNetworkReplayStreaming.h
@@ -66,6 +66,7 @@ public:
 	virtual bool IsLive() const override;
 	virtual void DeleteFinishedStream( const FString& StreamName, const FOnDeleteFinishedStreamComplete& Delegate) const override;
 	virtual void EnumerateStreams( const FNetworkReplayVersion& ReplayVersion, const FString& UserString, const FString& MetaString, const FOnEnumerateStreamsComplete& Delegate ) override;
+	virtual void EnumerateStreams( const FNetworkReplayVersion& InReplayVersion, const FString& UserString, const FString& MetaString, const TArray< FString >& ExtraParms, const FOnEnumerateStreamsComplete& Delegate ) override;
 	virtual void EnumerateRecentStreams( const FNetworkReplayVersion& ReplayVersion, const FString& RecentViewer, const FOnEnumerateStreamsComplete& Delegate ) override {}
 	virtual ENetworkReplayError::Type GetLastError() const override { return ENetworkReplayError::None; }
 	virtual void AddUserToReplay(const FString& UserString) override;
diff --git a/Engine/Source/Runtime/NetworkReplayStreaming/NetworkReplayStreaming/Public/NetworkReplayStreaming.h b/Engine/Source/Runtime/NetworkReplayStreaming/NetworkReplayStreaming/Public/NetworkReplayStreaming.h
--- a/Engine/Source/Runtime/NetworkReplayStreaming/NetworkReplayStreaming/Public/NetworkReplayStreaming.h
+++ b/Engine/Source/Runtime/NetworkReplayStreaming/NetworkReplayStreaming/Public/NetworkReplayStreaming.h
@@ -47,7 +47,7 @@ public:
 /** Struct to store information about a stream, returned from search results. */
 struct FNetworkReplayStreamInfo
 {
-	FNetworkReplayStreamInfo() : SizeInBytes( 0 ), LengthInMS( 0 ), NumViewers( 0 ), bIsLive( false ), Changelist( 0 ) {}
+	FNetworkReplayStreamInfo() : SizeInBytes( 0 ), LengthInMS( 0 ), NumViewers( 0 ), bIsLive( false ), Changelist( 0 ), bShouldKeep( false ) {}
 
 	/** The name of the stream (generally this is auto generated, refer to friendly name for UI) */
 	FString Name;
@@ -72,6 +72,9 @@ struct FNetworkReplayStreamInfo
 
 	/** The changelist of the replay */
 	int32 Changelist;
+
+	/** Debug feature that allows us to mark replays to not be deleted. True if this replay has been marked as such */
+	bool bShouldKeep;
 };
 
 namespace ENetworkReplayError
@@ -200,6 +203,12 @@ public:
 	 */
 	virtual void EnumerateStreams( const FNetworkReplayVersion& ReplayVersion, const FString& UserString, const FString& MetaString, const FOnEnumerateStreamsComplete& Delegate ) = 0;
 
+	/**
+	* Retrieves the streams that are available for viewing. May execute asynchronously.
+	* Allows the caller to pass in a custom list of query parameters
+	*/
+	virtual void EnumerateStreams( const FNetworkReplayVersion& InReplayVersion, const FString& UserString, const FString& MetaString, const TArray< FString >& ExtraParms, const FOnEnumerateStreamsComplete& Delegate ) = 0;
+
 	/**
 	 * Retrieves the streams that have been recently viewed. May execute asynchronously.
 	 *
diff --git a/Engine/Source/Runtime/NetworkReplayStreaming/NullNetworkReplayStreaming/Private/NullNetworkReplayStreaming.cpp b/Engine/Source/Runtime/NetworkReplayStreaming/NullNetworkReplayStreaming/Private/NullNetworkReplayStreaming.cpp
--- a/Engine/Source/Runtime/NetworkReplayStreaming/NullNetworkReplayStreaming/Private/NullNetworkReplayStreaming.cpp
+++ b/Engine/Source/Runtime/NetworkReplayStreaming/NullNetworkReplayStreaming/Private/NullNetworkReplayStreaming.cpp
@@ -328,6 +328,11 @@ void FNullNetworkReplayStreamer::DeleteFinishedStream( const FString& StreamName
 }
 
 void FNullNetworkReplayStreamer::EnumerateStreams( const FNetworkReplayVersion& ReplayVersion, const FString& UserString, const FString& MetaString, const FOnEnumerateStreamsComplete& Delegate )
+{
+	EnumerateStreams( ReplayVersion, UserString, MetaString, TArray< FString >(), Delegate );
+}
+
+void FNullNetworkReplayStreamer::EnumerateStreams( const FNetworkReplayVersion& ReplayVersion, const FString& UserString, const FString& MetaString, const TArray< FString >& ExtraParms, const FOnEnumerateStreamsComplete& Delegate )
 {
 	// Simply returns a stream for each folder in the Saved/Demos directory
 	const FString WildCardPath = GetDemoPath() + TEXT( "*" );
@@ -337,18 +342,18 @@ void FNullNetworkReplayStreamer::EnumerateStreams( const FNetworkReplayVersion&
 
 	TArray<FNetworkReplayStreamInfo> Results;
 
-	for (const FString& Directory : DirectoryNames)
+	for ( const FString& Directory : DirectoryNames )
 	{
 		// Assume there will be one file with a .demo extension in the directory
-		const FString FullDemoFilePath = GetDemoFilename(Directory);
+		const FString FullDemoFilePath = GetDemoFilename( Directory );
 
 		FNetworkReplayStreamInfo Info;
 		Info.SizeInBytes = IFileManager::Get().FileSize( *FullDemoFilePath );
-		
+
 		// Read stored info for this replay
 		FNullReplayInfo StoredReplayInfo = ReadReplayInfo( Directory );
-		
-		if (!StoredReplayInfo.bIsValid)
+
+		if ( !StoredReplayInfo.bIsValid )
 		{
 			continue;
 		}
@@ -368,11 +373,11 @@ void FNullNetworkReplayStreamer::EnumerateStreams( const FNetworkReplayVersion&
 			Info.LengthInMS = StoredReplayInfo.LengthInMS;
 			Info.FriendlyName = StoredReplayInfo.FriendlyName;
 
-			Results.Add(Info);
+			Results.Add( Info );
 		}
 	}
 
-	Delegate.ExecuteIfBound(Results);
+	Delegate.ExecuteIfBound( Results );
 }
 
 void FNullNetworkReplayStreamer::AddUserToReplay(const FString& UserString)
diff --git a/Engine/Source/Runtime/NetworkReplayStreaming/NullNetworkReplayStreaming/Public/NullNetworkReplayStreaming.h b/Engine/Source/Runtime/NetworkReplayStreaming/NullNetworkReplayStreaming/Public/NullNetworkReplayStreaming.h
--- a/Engine/Source/Runtime/NetworkReplayStreaming/NullNetworkReplayStreaming/Public/NullNetworkReplayStreaming.h
+++ b/Engine/Source/Runtime/NetworkReplayStreaming/NullNetworkReplayStreaming/Public/NullNetworkReplayStreaming.h
@@ -62,6 +62,7 @@ public:
 	virtual bool IsLive() const override;
 	virtual void DeleteFinishedStream( const FString& StreamName, const FOnDeleteFinishedStreamComplete& Delegate) const override;
 	virtual void EnumerateStreams( const FNetworkReplayVersion& ReplayVersion, const FString& UserString, const FString& MetaString, const FOnEnumerateStreamsComplete& Delegate ) override;
+	virtual void EnumerateStreams( const FNetworkReplayVersion& InReplayVersion, const FString& UserString, const FString& MetaString, const TArray< FString >& ExtraParms, const FOnEnumerateStreamsComplete& Delegate ) override;
 	virtual void EnumerateRecentStreams( const FNetworkReplayVersion& ReplayVersion, const FString& RecentViewer, const FOnEnumerateStreamsComplete& Delegate ) override {}
 	virtual ENetworkReplayError::Type GetLastError() const override { return ENetworkReplayError::None; }
 	virtual void AddUserToReplay(const FString& UserString) override;
diff --git a/Engine/Source/Runtime/Online/Hotfix/Hotfix.Build.cs b/Engine/Source/Runtime/Online/Hotfix/Hotfix.Build.cs
new file mode 100644
--- /dev/null
+++ b/Engine/Source/Runtime/Online/Hotfix/Hotfix.Build.cs
@@ -0,0 +1,17 @@
+// Copyright 1998-2015 Epic Games, Inc. All Rights Reserved.
+
+using UnrealBuildTool;
+
+public class Hotfix : ModuleRules
+{
+	public Hotfix(TargetInfo Target)
+    {
+        PrivateDependencyModuleNames.AddRange(
+			new string[] { 
+				"Core",
+				"CoreUObject",
+				"OnlineSubsystem",
+			}
+			);
+    }
+}
diff --git a/Engine/Source/Runtime/Online/Hotfix/Private/HotfixModule.cpp b/Engine/Source/Runtime/Online/Hotfix/Private/HotfixModule.cpp
new file mode 100644
--- /dev/null
+++ b/Engine/Source/Runtime/Online/Hotfix/Private/HotfixModule.cpp
@@ -0,0 +1,14 @@
+// Copyright 1998-2015 Epic Games, Inc. All Rights Reserved.
+
+#include "HotfixPrivatePCH.h"
+#include "HotfixModule.h"
+
+IMPLEMENT_MODULE(FHotfixModule, Hotfix);
+
+void FHotfixModule::StartupModule()
+{
+}
+
+void FHotfixModule::ShutdownModule()
+{
+}
diff --git a/Engine/Source/Runtime/Online/Hotfix/Private/HotfixPrivatePCH.h b/Engine/Source/Runtime/Online/Hotfix/Private/HotfixPrivatePCH.h
new file mode 100644
--- /dev/null
+++ b/Engine/Source/Runtime/Online/Hotfix/Private/HotfixPrivatePCH.h
@@ -0,0 +1,7 @@
+// Copyright 1998-2015 Epic Games, Inc. All Rights Reserved.
+
+#pragma once
+
+#include "Core.h"
+#include "CoreUObject.h"
+#include "Online.h"
diff --git a/Engine/Source/Runtime/Online/Hotfix/Private/OnlineHotfixManager.cpp b/Engine/Source/Runtime/Online/Hotfix/Private/OnlineHotfixManager.cpp
new file mode 100644
--- /dev/null
+++ b/Engine/Source/Runtime/Online/Hotfix/Private/OnlineHotfixManager.cpp
@@ -0,0 +1,420 @@
+// Copyright 1998-2015 Epic Games, Inc. All Rights Reserved.
+
+#include "HotfixPrivatePCH.h"
+#include "OnlineHotfixManager.h"
+#include "Internationalization.h"
+
+DEFINE_LOG_CATEGORY(LogHotfixManager);
+
+UOnlineHotfixManager::UOnlineHotfixManager() :
+	Super()
+{
+	OnEnumerateFilesCompleteDelegate = FOnEnumerateFilesCompleteDelegate::CreateUObject(this, &UOnlineHotfixManager::OnEnumerateFilesComplete);
+	OnReadFileCompleteDelegate = FOnReadFileCompleteDelegate::CreateUObject(this, &UOnlineHotfixManager::OnReadFileComplete);
+	// So we only try to apply files for this platform
+	PlatformPrefix = ANSI_TO_TCHAR(FPlatformProperties::PlatformName());
+	PlatformPrefix += TEXT("_");
+}
+
+void UOnlineHotfixManager::Init()
+{
+	// Build the name of the loc file that we'll care about
+	// It can change at runtime so build it just before fetching the data
+	GameLocName = FInternationalization::Get().GetCurrentCulture()->GetTwoLetterISOLanguageName() + TEXT("_Game.locres");
+	OnlineTitleFile = Online::GetTitleFileInterface(OSSName.Len() ? FName(*OSSName, FNAME_Find) : NAME_None);
+	if (OnlineTitleFile.IsValid())
+	{
+		OnEnumerateFilesCompleteDelegateHandle = OnlineTitleFile->AddOnEnumerateFilesCompleteDelegate_Handle(OnEnumerateFilesCompleteDelegate);
+		OnReadFileCompleteDelegateHandle = OnlineTitleFile->AddOnReadFileCompleteDelegate_Handle(OnReadFileCompleteDelegate);
+	}
+}
+
+void UOnlineHotfixManager::Cleanup()
+{
+	PendingHotfixFiles.Empty();
+	if (OnlineTitleFile.IsValid())
+	{
+		// Make sure to give back the memory used when reading the hotfix files
+		OnlineTitleFile->ClearFiles();
+		OnlineTitleFile->ClearOnEnumerateFilesCompleteDelegate_Handle(OnEnumerateFilesCompleteDelegateHandle);
+		OnlineTitleFile->ClearOnReadFileCompleteDelegate_Handle(OnReadFileCompleteDelegateHandle);
+	}
+	OnlineTitleFile = nullptr;
+}
+
+void UOnlineHotfixManager::StartHotfixProcess()
+{
+	// Patching the editor this way seems like a bad idea
+	if (!IsRunningGame())
+	{
+		UE_LOG(LogHotfixManager, Warning, TEXT("Hotfixing skipped due to IsRunningGame() == false"));
+		TriggerHotfixComplete(EHotfixResult::SuccessNoChange);
+		return;
+	}
+
+	Init();
+	// Kick off an enumeration of the files that are available to download
+	if (OnlineTitleFile.IsValid())
+	{
+		OnlineTitleFile->EnumerateFiles();
+	}
+	else
+	{
+		UE_LOG(LogHotfixManager, Error, TEXT("Failed to start the hotfixing process due to no OnlineTitleInterface present for OSS(%s)"), *OSSName);
+		TriggerHotfixComplete(EHotfixResult::Failed);
+	}
+}
+
+void UOnlineHotfixManager::OnEnumerateFilesComplete(bool bWasSuccessful)
+{
+	if (bWasSuccessful)
+	{
+		check(OnlineTitleFile.IsValid());
+		// Cache our current set so we can compare for differences
+		LastHotfixFileList = HotfixFileList;
+		HotfixFileList.Empty();
+		// Get the new header data
+		OnlineTitleFile->GetFileList(HotfixFileList);
+		FilterHotfixFiles();
+		// Sort after filtering so that the comparison below doesn't fail to different order from the server
+		HotfixFileList.Sort();
+		if (LastHotfixFileList.Num() == 0 || LastHotfixFileList != HotfixFileList)
+		{
+			UnmountHotfixFiles();
+			ReadHotfixFiles();
+		}
+		else
+		{
+			UE_LOG(LogHotfixManager, Display, TEXT("Returned hotfix data is the same as last application, skipping the apply phase"));
+			TriggerHotfixComplete(EHotfixResult::SuccessNoChange);
+		}
+	}
+	else
+	{
+		UE_LOG(LogHotfixManager, Error, TEXT("Enumeration of hotfix files failed"));
+		TriggerHotfixComplete(EHotfixResult::Failed);
+	}
+}
+
+void UOnlineHotfixManager::FilterHotfixFiles()
+{
+	for (int32 Idx = 0; Idx < HotfixFileList.Num(); Idx++)
+	{
+		if (!WantsHotfixProcessing(HotfixFileList[Idx]))
+		{
+			HotfixFileList.RemoveAt(Idx, 1, false);
+			Idx--;
+		}
+	}
+}
+
+void UOnlineHotfixManager::ReadHotfixFiles()
+{
+	if (HotfixFileList.Num())
+	{
+		check(OnlineTitleFile.IsValid());
+		// Kick off a read for each file
+		// Do this in two passes so already cached files don't trigger completion
+		for (auto FileHeader : HotfixFileList)
+		{
+			PendingHotfixFiles.Add(FileHeader.DLName);
+		}
+		for (auto FileHeader : HotfixFileList)
+		{
+			OnlineTitleFile->ReadFile(FileHeader.DLName);
+		}
+	}
+	else
+	{
+		UE_LOG(LogHotfixManager, Display, TEXT("No hotfix files need to be downloaded"));
+		TriggerHotfixComplete(EHotfixResult::Success);
+	}
+}
+
+void UOnlineHotfixManager::OnReadFileComplete(bool bWasSuccessful, const FString& FileName)
+{
+	if (PendingHotfixFiles.Contains(FileName))
+	{
+		if (bWasSuccessful)
+		{
+			UE_LOG(LogHotfixManager, Log, TEXT("Hotfix file (%s) downloaded"), *GetFriendlyNameFromDLName(FileName));
+			PendingHotfixFiles.Remove(FileName);
+			if (PendingHotfixFiles.Num() == 0)
+			{
+				// All files have been downloaded so now apply the files
+				ApplyHotfix();
+			}
+		}
+		else
+		{
+			UE_LOG(LogHotfixManager, Error, TEXT("Hotfix file (%s) failed to download"), *GetFriendlyNameFromDLName(FileName));
+			TriggerHotfixComplete(EHotfixResult::Failed);
+		}
+	}
+}
+
+void UOnlineHotfixManager::ApplyHotfix()
+{
+	for (auto FileHeader : HotfixFileList)
+	{
+		if (!ApplyHotfixProcessing(FileHeader))
+		{
+			UE_LOG(LogHotfixManager, Error, TEXT("Couldn't apply hotfix file (%s)"), *FileHeader.FileName);
+			TriggerHotfixComplete(EHotfixResult::Failed);
+			return;
+		}
+	}
+	UE_LOG(LogHotfixManager, Display, TEXT("Hotfix data has been successfully applied"));
+	TriggerHotfixComplete(EHotfixResult::Success);
+}
+
+void UOnlineHotfixManager::TriggerHotfixComplete(EHotfixResult HotfixResult)
+{
+	TriggerOnHotfixCompleteDelegates(HotfixResult);
+	if (HotfixResult == EHotfixResult::Failed)
+	{
+		HotfixFileList.Empty();
+		UnmountHotfixFiles();
+	}
+	Cleanup();
+}
+
+bool UOnlineHotfixManager::WantsHotfixProcessing(const FCloudFileHeader& FileHeader)
+{
+	const FString Extension = FPaths::GetExtension(FileHeader.FileName);
+	if (Extension == TEXT("INI"))
+	{
+		return FileHeader.FileName.StartsWith(PlatformPrefix);
+	}
+	else if (Extension == TEXT("PAK"))
+	{
+		return FileHeader.FileName.StartsWith(PlatformPrefix);
+	}
+	return FileHeader.FileName == GameLocName;
+}
+
+bool UOnlineHotfixManager::ApplyHotfixProcessing(const FCloudFileHeader& FileHeader)
+{
+	bool bSuccess = false;
+	const FString Extension = FPaths::GetExtension(FileHeader.FileName);
+	if (Extension == TEXT("INI"))
+	{
+		TArray<uint8> FileData;
+		if (OnlineTitleFile->GetFileContents(FileHeader.DLName, FileData))
+		{
+			// Convert to a FString
+			FileData.Add(0);
+			const FString IniData = (ANSICHAR*)FileData.GetData();
+			bSuccess = HotfixIniFile(FileHeader.FileName, IniData);
+		}
+	}
+	else if (Extension == TEXT("LOCRES"))
+	{
+		HotfixLocFile(FileHeader);
+	}
+	else if (Extension == TEXT("PAK"))
+	{
+		bSuccess = HotfixPakFile(FileHeader);
+	}
+	OnlineTitleFile->ClearFile(FileHeader.FileName);
+	return bSuccess;
+}
+
+FConfigFile* UOnlineHotfixManager::GetConfigFile(const FString& IniName)
+{
+	FString StrippedIniName;
+	if (IniName.StartsWith(PlatformPrefix))
+	{
+		StrippedIniName = IniName.Right(IniName.Len() - PlatformPrefix.Len());
+	}
+	if (StrippedIniName.StartsWith(TEXT("Default")))
+	{
+		StrippedIniName = IniName.Right(StrippedIniName.Len() - 7);
+	}
+	FConfigFile* ConfigFile = nullptr;
+	// Look for the first matching INI file entry
+	for (TMap<FString, FConfigFile>::TIterator It(*GConfig); It; ++It)
+	{
+		if (It.Key().EndsWith(StrippedIniName))
+		{
+			ConfigFile = &It.Value();
+			break;
+		}
+	}
+	// If not found, add this file to the config cache
+	if (ConfigFile == nullptr)
+	{
+		const FString IniNameWithPath = FPaths::GeneratedConfigDir() + StrippedIniName;
+		ConfigFile = new FConfigFile();
+		GConfig->SetFile(*IniNameWithPath, ConfigFile);
+	}
+	check(ConfigFile);
+	// We never want to save these merged files
+	ConfigFile->NoSave = true;
+	return ConfigFile;
+}
+
+bool UOnlineHotfixManager::HotfixIniFile(const FString& FileName, const FString& IniData)
+{
+	FConfigFile* ConfigFile = GetConfigFile(FileName);
+	// Merge the string into the config file
+	ConfigFile->CombineFromBuffer(*IniData);
+	TArray<UClass*> Classes;
+	TArray<UObject*> PerObjectConfigObjects;
+	int32 StartIndex = 0;
+	int32 EndIndex = 0;
+	// Find the set of object classes that were affected
+	while (StartIndex >= 0 && StartIndex < IniData.Len() && EndIndex >= StartIndex)
+	{
+		// Find the next section header
+		StartIndex = IniData.Find(TEXT("["), ESearchCase::IgnoreCase, ESearchDir::FromStart, StartIndex);
+		if (StartIndex > -1)
+		{
+			// Find the ending section identifier
+			EndIndex = IniData.Find(TEXT("]"), ESearchCase::IgnoreCase, ESearchDir::FromStart, StartIndex);
+			if (EndIndex > StartIndex)
+			{
+				int32 PerObjectNameIndex = IniData.Find(TEXT(" "), ESearchCase::IgnoreCase, ESearchDir::FromStart, StartIndex);
+				// Per object config entries will have a space in the name, but classes won't
+				if (PerObjectNameIndex == -1)
+				{
+					if (IniData.StartsWith(TEXT("[/Script/"), ESearchCase::IgnoreCase))
+					{
+						const int32 ScriptSectionTag = 9;
+						// Snip the text out and try to find the class for that
+						const FString PackageClassName = IniData.Mid(StartIndex + ScriptSectionTag, EndIndex - StartIndex - ScriptSectionTag);
+						// Find the class for this so we know what to update
+						UClass* Class = FindObject<UClass>(nullptr, *PackageClassName, true);
+						if (Class)
+						{
+							// Add this to the list to check against
+							Classes.Add(Class);
+						}
+					}
+				}
+				// Handle the per object config case by finding the object for reload
+				else
+				{
+					const FString PerObjectName = IniData.Mid(StartIndex + 1, PerObjectNameIndex - 1);
+					// Explicitly search the transient package (won't update non-transient objects)
+					UObject* PerObject = FindObject<UObject>(ANY_PACKAGE, *PerObjectName, false);
+					if (PerObject != nullptr)
+					{
+						PerObjectConfigObjects.Add(PerObject);
+					}
+				}
+				StartIndex = EndIndex;
+			}
+		}
+	}
+
+	int32 NumObjectsReloaded = 0;
+	const double StartTime = FPlatformTime::Seconds();
+	if (Classes.Num())
+	{
+		// Now that we have a list of classes to update, we can iterate objects and reload
+		for (FObjectIterator It; It; ++It)
+		{
+			UClass* Class = It->GetClass();
+			if (Class->HasAnyClassFlags(CLASS_Config))
+			{
+				// Check to see if this class is in our list (yes, potentially n^2, but not in practice)
+				for (int32 ClassIndex = 0; ClassIndex < Classes.Num(); ClassIndex++)
+				{
+					if (It->IsA(Classes[ClassIndex]))
+					{
+						// Force a reload of the config vars
+						It->ReloadConfig();
+						NumObjectsReloaded++;
+						break;
+					}
+				}
+			}
+		}
+	}
+	// Reload any PerObjectConfig objects that were affected
+	for (auto ReloadObject : PerObjectConfigObjects)
+	{
+		ReloadObject->ReloadConfig();
+		NumObjectsReloaded++;
+	}
+	UE_LOG(LogHotfixManager, Log, TEXT("Updating config from %s took %f seconds and reloaded %d objects"),
+		*FileName, FPlatformTime::Seconds() - StartTime, NumObjectsReloaded);
+	return true;
+}
+
+void UOnlineHotfixManager::HotfixLocFile(const FCloudFileHeader& FileHeader)
+{
+	const double StartTime = FPlatformTime::Seconds();
+	FString LocFilePath = FString::Printf(TEXT("%s/%s"), *GetCachedDirectory(), *FileHeader.DLName);
+	FTextLocalizationManager::Get().UpdateFromLocalizationResource(LocFilePath);
+	UE_LOG(LogHotfixManager, Log, TEXT("Updating loc from %s took %f seconds"), *FileHeader.FileName, FPlatformTime::Seconds() - StartTime);
+}
+
+bool UOnlineHotfixManager::HotfixPakFile(const FCloudFileHeader& FileHeader)
+{
+	if (!FCoreDelegates::OnMountPak.IsBound())
+	{
+		UE_LOG(LogHotfixManager, Error, TEXT("PAK file (%s) could not be mounted because OnMountPak is not bound"), *FileHeader.FileName);
+		return false;
+	}
+	FString PakLocation = FString::Printf(TEXT("%s/%s"), *GetCachedDirectory(), *FileHeader.DLName);
+	if (FCoreDelegates::OnMountPak.Execute(PakLocation, 0))
+	{
+		MountedPakFiles.Add(FileHeader.DLName);
+		UE_LOG(LogHotfixManager, Log, TEXT("Hotfix mounted PAK file (%s)"), *FileHeader.FileName);
+		// @todo joeg - Merge any INI files that are part of this PAK file into the config cache
+		// Do we want/need this? Is it clear whether this will give us the right info?
+		return true;
+	}
+	return false;
+}
+
+bool UOnlineHotfixManager::HotfixPakIniFile(const FString& FileName)
+{
+	FConfigFile* ConfigFile = GetConfigFile(FileName);
+	ConfigFile->Combine(FileName);
+	UE_LOG(LogHotfixManager, Log, TEXT("Hotfix merged INI (%s) found in a PAK file"), *FileName);
+
+	FName IniFileName(*FileName, FNAME_Find);
+	int32 NumObjectsReloaded = 0;
+	const double StartTime = FPlatformTime::Seconds();
+	// Now that we have a list of classes to update, we can iterate objects and
+	// reload if they match the INI file that was changed
+	for (FObjectIterator It; It; ++It)
+	{
+		UClass* Class = It->GetClass();
+		if (Class->HasAnyClassFlags(CLASS_Config) &&
+			Class->ClassConfigName == IniFileName)
+		{
+			// Force a reload of the config vars
+			It->ReloadConfig();
+			NumObjectsReloaded++;
+		}
+	}
+	UE_LOG(LogHotfixManager, Log, TEXT("Updating config from %s took %f seconds reloading %d objects"),
+		*FileName, FPlatformTime::Seconds() - StartTime, NumObjectsReloaded);
+	return true;
+}
+
+const FString UOnlineHotfixManager::GetFriendlyNameFromDLName(const FString& DLName) const
+{
+	for (auto Header : HotfixFileList)
+	{
+		if (Header.DLName == DLName)
+		{
+			return Header.FileName;
+		}
+	}
+	return FString();
+}
+
+void UOnlineHotfixManager::UnmountHotfixFiles()
+{
+	// Unmount any hotfix files since we need to download them again
+	for (auto PakFile : MountedPakFiles)
+	{
+		FCoreDelegates::OnUnmountPak.Execute(PakFile);
+	}
+	MountedPakFiles.Empty();
+}
diff --git a/Engine/Source/Runtime/Online/Hotfix/Public/HotfixModule.h b/Engine/Source/Runtime/Online/Hotfix/Public/HotfixModule.h
new file mode 100644
--- /dev/null
+++ b/Engine/Source/Runtime/Online/Hotfix/Public/HotfixModule.h
@@ -0,0 +1,17 @@
+// Copyright 1998-2015 Epic Games, Inc. All Rights Reserved.
+
+#pragma once
+
+#include "ModuleManager.h"
+
+/**
+ * Module for Hotfix base implementation
+ */
+class FHotfixModule : 
+	public IModuleInterface
+{
+private:
+// IModuleInterface
+	virtual void StartupModule() override;
+	virtual void ShutdownModule() override;
+};
diff --git a/Engine/Source/Runtime/Online/OnlineFramework/Party/Private/PartyGameState.cpp b/Engine/Source/Runtime/Online/OnlineFramework/Party/Private/PartyGameState.cpp
--- a/Engine/Source/Runtime/Online/OnlineFramework/Party/Private/PartyGameState.cpp
+++ b/Engine/Source/Runtime/Online/OnlineFramework/Party/Private/PartyGameState.cpp
@@ -300,6 +300,11 @@ UPartyMemberState* UPartyGameState::CreateNewPartyMember(const FUniqueNetId& InM
 				NewPartyMemberState = NewObject<UPartyMemberState>(this, PartyMemberStateClass);
 				NewPartyMemberState->UniqueId.SetUniqueNetId(PartyMember->GetUserId());
 				NewPartyMemberState->DisplayName = FText::FromString(PartyMember->GetDisplayName());
+
+				if ( InitMemberFunctor )
+				{
+					InitMemberFunctor( PartyMember, NewPartyMemberState );
+				}
 			}
 			else
 			{
@@ -319,7 +324,12 @@ UPartyMemberState* UPartyGameState::CreateNewPartyMember(const FUniqueNetId& InM
 	return NewPartyMemberState;
 }
 
-void UPartyGameState::HandlePartyConfigChanged(const TSharedRef<FPartyConfiguration>& InPartyConfig)
+void UPartyGameState::SetInitMemberFunctor( TFunction<void( TSharedPtr<FOnlinePartyMember>, UPartyMemberState* )>&& InFunctor )
+{
+	InitMemberFunctor = MoveTemp( InFunctor );
+}
+
+void UPartyGameState::HandlePartyConfigChanged( const TSharedRef<FPartyConfiguration>& InPartyConfig )
 {
 	UE_LOG(LogParty, VeryVerbose, TEXT("[%s] HandlePartyConfigChanged"), PartyInfo.IsValid() ? *PartyInfo->PartyId->ToString() : TEXT("INVALID"));
 	if (PartyInfo.IsValid())
diff --git a/Engine/Source/Runtime/Online/OnlineFramework/Party/Public/Party.h b/Engine/Source/Runtime/Online/OnlineFramework/Party/Public/Party.h
--- a/Engine/Source/Runtime/Online/OnlineFramework/Party/Public/Party.h
+++ b/Engine/Source/Runtime/Online/OnlineFramework/Party/Public/Party.h
@@ -49,9 +49,9 @@ struct FPartyDetails
 		return PartyJoinInfo->GetLeaderId();
 	}
 
-	inline const FString& GetClientId() const
+	inline const FString& GetAppId() const
 	{
-		return PartyJoinInfo->GetClientId();
+		return PartyJoinInfo->GetAppId();
 	}
 
 	FString ToString() const
@@ -60,7 +60,7 @@ struct FPartyDetails
 			TEXT("PartyId: %s LeaderId: %s ResKey: %s Client: %s"), 
 			*GetPartyId()->ToDebugString(), 
 			*GetLeaderId()->ToDebugString(),
-			*GetClientId());
+			*GetAppId());
 	}
 
 	TSharedRef<const IOnlinePartyJoinInfo> PartyJoinInfo;
diff --git a/Engine/Source/Runtime/Online/OnlineFramework/Party/Public/PartyGameState.h b/Engine/Source/Runtime/Online/OnlineFramework/Party/Public/PartyGameState.h
--- a/Engine/Source/Runtime/Online/OnlineFramework/Party/Public/PartyGameState.h
+++ b/Engine/Source/Runtime/Online/OnlineFramework/Party/Public/PartyGameState.h
@@ -165,6 +165,8 @@ public:
 	static void AddReferencedObjects(UObject* InThis, FReferenceCollector& Collector);
 	//~ Begin UObject Interface end
 
+	void SetInitMemberFunctor( TFunction<void( TSharedPtr<FOnlinePartyMember>, UPartyMemberState* )>&& InFunctor );
+
 	/**
 	 * Notification that the game is about to travel to another map/server
 	 */
@@ -611,6 +613,8 @@ private:
 	/** Scratch copy of child USTRUCT for handling replication comparisons */
 	FPartyState* PartyStateRefScratch;
 
+	TFunction<void( TSharedPtr<FOnlinePartyMember>, UPartyMemberState* )> InitMemberFunctor;
+
 	friend UParty;
 	friend UPartyMemberState;
 };
diff --git a/Engine/Source/Runtime/Online/OnlineFramework/Qos/Private/QosStats.h b/Engine/Source/Runtime/Online/OnlineFramework/Qos/Private/QosStats.h
--- a/Engine/Source/Runtime/Online/OnlineFramework/Qos/Private/QosStats.h
+++ b/Engine/Source/Runtime/Online/OnlineFramework/Qos/Private/QosStats.h
@@ -129,24 +129,24 @@ private:
 	};
 
 	// Events
-	static const FString FQosDatacenterStats::QosStats_DatacenterEvent;
+	static const FString QosStats_DatacenterEvent;
 
 	// Common attribution
-	static const FString FQosDatacenterStats::QosStats_SessionId;
-	static const FString FQosDatacenterStats::QosStats_Version;
+	static const FString QosStats_SessionId;
+	static const FString QosStats_Version;
 
 	// Header stats
-	static const FString FQosDatacenterStats::QosStats_Timestamp;
-	static const FString FQosDatacenterStats::QosStats_TotalTime;
+	static const FString QosStats_Timestamp;
+	static const FString QosStats_TotalTime;
 
 	// Qos stats
-	static const FString FQosDatacenterStats::QosStats_DeterminationType;
-	static const FString FQosDatacenterStats::QosStats_DatacenterId;
-	static const FString FQosDatacenterStats::QosStats_NumRegions;
-	static const FString FQosDatacenterStats::QosStats_RegionDetails;
-	static const FString FQosDatacenterStats::QosStats_NumResults;
-	static const FString FQosDatacenterStats::QosStats_NumSuccessCount;
-	static const FString FQosDatacenterStats::QosStats_SearchDetails;
+	static const FString QosStats_DeterminationType;
+	static const FString QosStats_DatacenterId;
+	static const FString QosStats_NumRegions;
+	static const FString QosStats_RegionDetails;
+	static const FString QosStats_NumResults;
+	static const FString QosStats_NumSuccessCount;
+	static const FString QosStats_SearchDetails;
 
 	/** Version of the stats for separation */
 	int32 StatsVersion;
diff --git a/Engine/Source/Runtime/Online/OnlineSubsystem/Public/Interfaces/OnlinePartyInterface.h b/Engine/Source/Runtime/Online/OnlineSubsystem/Public/Interfaces/OnlinePartyInterface.h
--- a/Engine/Source/Runtime/Online/OnlineSubsystem/Public/Interfaces/OnlinePartyInterface.h
+++ b/Engine/Source/Runtime/Online/OnlineSubsystem/Public/Interfaces/OnlinePartyInterface.h
@@ -108,7 +108,7 @@ public:
 	/**
 	 * @return id of the client app associated with the sender of the party invite 
 	 */
-	virtual const FString& GetClientId() const = 0;
+	virtual const FString& GetAppId() const = 0;
 
 	/**
 	 * @return id of the client app associated with the sender of the party invite 
diff --git a/Engine/Source/Runtime/Online/OnlineSubsystem/Public/Interfaces/OnlinePresenceInterface.h b/Engine/Source/Runtime/Online/OnlineSubsystem/Public/Interfaces/OnlinePresenceInterface.h
--- a/Engine/Source/Runtime/Online/OnlineSubsystem/Public/Interfaces/OnlinePresenceInterface.h
+++ b/Engine/Source/Runtime/Online/OnlineSubsystem/Public/Interfaces/OnlinePresenceInterface.h
@@ -16,8 +16,8 @@ const FString DefaultPresenceKey = "RichPresence";
 /** Custom presence data that is not seen by users but can be polled */
 const FString CustomPresenceDataKey = "CustomData";
 
-/** Id of the client that sent the presence update */
-const FString DefaultClientIdKey = "ClientId";
+/** Name of the client that sent the presence update */
+const FString DefaultAppIdKey = "AppId";
 
 /** Override Id of the client to set the presence state to */
 const FString OverrideClientIdKey = "OverrideClientId";
diff --git a/Engine/Source/Runtime/Online/OnlineSubsystem/Public/OnlineSubsystemTypes.h b/Engine/Source/Runtime/Online/OnlineSubsystem/Public/OnlineSubsystemTypes.h
--- a/Engine/Source/Runtime/Online/OnlineSubsystem/Public/OnlineSubsystemTypes.h
+++ b/Engine/Source/Runtime/Online/OnlineSubsystem/Public/OnlineSubsystemTypes.h
@@ -1017,12 +1017,16 @@ struct FCloudFileHeader
 {	
 	/** Hash value, if applicable, of the given file contents */
     FString Hash;
+	/** The hash algorithm used to sign this file */
+	FName HashType;
 	/** Filename as downloaded */
     FString DLName;
 	/** Logical filename, maps to the downloaded filename */
     FString FileName;
 	/** File size */
     int32 FileSize;
+	/** The full URL to download the file if it is stored in a CDN or separate host site */
+	FString URL;
 
     /** Constructors */
     FCloudFileHeader() :
@@ -1034,6 +1038,21 @@ struct FCloudFileHeader
 		FileName(InFileName),
 		FileSize(InFileSize)
 	{}
+
+	bool operator==(const FCloudFileHeader& Other) const
+	{
+		return FileSize == Other.FileSize &&
+			Hash == Other.Hash &&
+			HashType == Other.HashType &&
+			DLName == Other.DLName &&
+			FileName == Other.FileName &&
+			URL == Other.URL;
+	}
+
+	bool operator<(const FCloudFileHeader& Other) const
+	{
+		return FileName.Compare(Other.FileName, ESearchCase::IgnoreCase) < 0;
+	}
 };
 
 /** Holds the data used in downloading a file asynchronously from the online service */
diff --git a/Engine/Source/Runtime/Online/XMPP/Private/XmppConnection.cpp b/Engine/Source/Runtime/Online/XMPP/Private/XmppConnection.cpp
new file mode 100644
--- /dev/null
+++ b/Engine/Source/Runtime/Online/XMPP/Private/XmppConnection.cpp
@@ -0,0 +1,39 @@
+// Copyright 1998-2015 Epic Games, Inc. All Rights Reserved.
+
+#include "XmppPrivatePCH.h"
+#include "XmppConnection.h"
+
+bool FXmppUserJid::ParseResource(FString& OutAppId, FString& OutPlatform) const
+{
+	OutAppId.Empty();
+	OutPlatform.Empty();
+
+	TArray<FString> ParsedResource;
+	if (Resource.ParseIntoArray(ParsedResource, TEXT(":"), false) > 1)
+	{
+		if (ParsedResource[0].StartsWith(TEXT("V")))
+		{
+			uint32 Version = FCString::Atoi((*ParsedResource[0]) + 1);
+			if (Version == 2)
+			{
+				if (ParsedResource.Num() >= 3)
+				{
+					OutAppId = ParsedResource[1];
+					OutPlatform = ParsedResource[2];
+					return true;
+				}
+			}
+		}
+	}
+	else
+	{
+		FString ClientId, NotUsed;
+		if (Resource.Split(TEXT("-"), &ClientId, &NotUsed) &&
+			!ClientId.IsEmpty())
+		{
+			OutAppId = ClientId;
+			return true;
+		}
+	}
+	return false;
+}
diff --git a/Engine/Source/Runtime/Online/XMPP/Private/XmppJingle/XmppConnectionJingle.cpp b/Engine/Source/Runtime/Online/XMPP/Private/XmppJingle/XmppConnectionJingle.cpp
--- a/Engine/Source/Runtime/Online/XMPP/Private/XmppJingle/XmppConnectionJingle.cpp
+++ b/Engine/Source/Runtime/Online/XMPP/Private/XmppJingle/XmppConnectionJingle.cpp
@@ -33,6 +33,8 @@ DEFINE_STAT(STAT_XmppMucOpRequests);
 
 #if WITH_XMPP_JINGLE
 
+#define XMPP_RESOURCE_VERSION 2
+
 /**
  * Thread to create the xmpp pump/connection
  * Spawned during login and destroyed on logout
@@ -650,7 +652,11 @@ void FXmppConnectionJingle::SetServer(const FXmppServer& InServer)
 	// in order to ensure unique connections from user/client combination
 	// add random number to the client resource identifier
 	ServerConfig = InServer;
-	ServerConfig.ClientResource += FString(TEXT("-")) + FGuid::NewGuid().ToString(EGuidFormats::Digits);
+	
+	ServerConfig.ClientResource  = FString::Printf(TEXT("V%u"), XMPP_RESOURCE_VERSION) + FString(TEXT(":"));
+	ServerConfig.ClientResource += ServerConfig.AppId + FString(TEXT(":"));
+	ServerConfig.ClientResource += ServerConfig.Platform + FString(TEXT(":"));
+	ServerConfig.ClientResource += FGuid::NewGuid().ToString(EGuidFormats::Digits);
 }
 
 const FXmppServer& FXmppConnectionJingle::GetServer() const
diff --git a/Engine/Source/Runtime/Online/XMPP/Private/XmppJingle/XmppMultiUserChatJingle.cpp b/Engine/Source/Runtime/Online/XMPP/Private/XmppJingle/XmppMultiUserChatJingle.cpp
--- a/Engine/Source/Runtime/Online/XMPP/Private/XmppJingle/XmppMultiUserChatJingle.cpp
+++ b/Engine/Source/Runtime/Online/XMPP/Private/XmppJingle/XmppMultiUserChatJingle.cpp
@@ -1454,7 +1454,7 @@ static void ConvertToChatMember(FXmppChatMember& OutChatMember, const buzz::Xmpp
 	if (InChatMemberJingle.presence() != NULL)
 	{
 		//@todo sz1
-		//FXmppPresenceJingle::ConvertToPresence(OutChatMember.UserPresence, *InChatMemberJingle.presence());
+		//FXmppPresenceJingle::ConvertToPresence(OutChatMember.UserPresence, *InChatMemberJingle.presence(), OutChatMember.MemberJid);
 	}
 }
 
diff --git a/Engine/Source/Runtime/Online/XMPP/Private/XmppJingle/XmppPresenceJingle.cpp b/Engine/Source/Runtime/Online/XMPP/Private/XmppJingle/XmppPresenceJingle.cpp
--- a/Engine/Source/Runtime/Online/XMPP/Private/XmppJingle/XmppPresenceJingle.cpp
+++ b/Engine/Source/Runtime/Online/XMPP/Private/XmppJingle/XmppPresenceJingle.cpp
@@ -303,7 +303,7 @@ FXmppPresenceJingle::~FXmppPresenceJingle()
 	}
 }
 
-void FXmppPresenceJingle::ConvertToPresence(FXmppUserPresence& OutPresence, const buzz::PresenceStatus& InStatus)
+void FXmppPresenceJingle::ConvertToPresence(FXmppUserPresence& OutPresence, const buzz::PresenceStatus& InStatus, const FXmppUserJid& InJid)
 {
 	OutPresence.bIsAvailable = InStatus.available();
 	OutPresence.StatusStr = UTF8_TO_TCHAR(InStatus.status().c_str());
@@ -346,15 +346,8 @@ void FXmppPresenceJingle::ConvertToPresence(FXmppUserPresence& OutPresence, cons
 			break;
 		}
 	}
-	OutPresence.ClientResource = UTF8_TO_TCHAR(InStatus.jid().resource().c_str());
-
-	// strip off the guid portion of the full resource id <clientid>-<guid>
-	FString ClientId, NotUsed;
-	if (OutPresence.ClientResource.Split(TEXT("-"), &ClientId, &NotUsed) &&
-		!ClientId.IsEmpty())
-	{
-		OutPresence.ClientResource = ClientId;
-	}
+	
+	InJid.ParseResource(OutPresence.AppId, OutPresence.Platform);
 }
 
 void FXmppPresenceJingle::ConvertFromPresence(buzz::PresenceStatus& OutStatus, const FXmppUserPresence& InPresence)
@@ -395,7 +388,8 @@ static void DebugPrintPresence(const FXmppUserPresence& Presence)
 	UE_LOG(LogXmpp, Verbose, TEXT("   Status = %s"), EXmppPresenceStatus::ToString(Presence.Status));
 	UE_LOG(LogXmpp, Verbose, TEXT("   bIsAvailable = %s"), Presence.bIsAvailable ? TEXT("true") : TEXT("false"));
 	UE_LOG(LogXmpp, Verbose, TEXT("   SentTime = %s"), *Presence.SentTime.ToString());
-	UE_LOG(LogXmpp, Verbose, TEXT("   ClientResource = %s"), *Presence.ClientResource);
+	UE_LOG(LogXmpp, Verbose, TEXT("   AppId = %s"), *Presence.AppId);
+	UE_LOG(LogXmpp, Verbose, TEXT("   Platform = %s"), *Presence.Platform);
 	UE_LOG(LogXmpp, Verbose, TEXT("   StatusStr = %s"), *Presence.StatusStr);
 }
 
@@ -559,7 +553,7 @@ void FXmppPresenceJingle::OnSignalPresenceUpdate(const buzz::PresenceStatus& InS
 			FScopeLock Lock(&RosterLock);
 		
 			FXmppUserPresenceJingle& RosterEntry = RosterPresence.FindOrAdd(UserJid.GetFullPath());
-			ConvertToPresence(*RosterEntry.Presence, InStatus);
+			ConvertToPresence(*RosterEntry.Presence, InStatus, UserJid);
 			FXmppJingle::ConvertToJid(RosterEntry.UserJid, InStatus.jid());
 
 			UE_LOG(LogXmpp, Verbose, TEXT("Received presence for user [%s]"), *UserJid.GetFullPath());
diff --git a/Engine/Source/Runtime/Online/XMPP/Private/XmppJingle/XmppPresenceJingle.h b/Engine/Source/Runtime/Online/XMPP/Private/XmppJingle/XmppPresenceJingle.h
--- a/Engine/Source/Runtime/Online/XMPP/Private/XmppJingle/XmppPresenceJingle.h
+++ b/Engine/Source/Runtime/Online/XMPP/Private/XmppJingle/XmppPresenceJingle.h
@@ -52,7 +52,7 @@ public:
 	virtual ~FXmppPresenceJingle();
 
 	static void ConvertFromPresence(buzz::PresenceStatus& OutStatus, const FXmppUserPresence& InPresence);
-	static void ConvertToPresence(FXmppUserPresence& OutPresence, const buzz::PresenceStatus& InStatus);
+	static void ConvertToPresence(FXmppUserPresence& OutPresence, const buzz::PresenceStatus& InStatus, const FXmppUserJid& InJid);
 
 private:
 
diff --git a/Engine/Source/Runtime/Online/XMPP/Private/XmppModule.cpp b/Engine/Source/Runtime/Online/XMPP/Private/XmppModule.cpp
--- a/Engine/Source/Runtime/Online/XMPP/Private/XmppModule.cpp
+++ b/Engine/Source/Runtime/Online/XMPP/Private/XmppModule.cpp
@@ -51,7 +51,7 @@ bool FXmppModule::HandleXmppCommand( const TCHAR* Cmd, FOutputDevice& Ar )
 {
 	FXmppServer XmppServer;
 	XmppServer.bUseSSL = true;
-	XmppServer.ClientResource = TEXT("ue_xmpp_test");
+	XmppServer.AppId = TEXT("ue_xmpp_test");
 
 	//@todo sz1 - for debug only
 	FString ConfigOverride;
diff --git a/Engine/Source/Runtime/Online/XMPP/Public/XmppConnection.h b/Engine/Source/Runtime/Online/XMPP/Public/XmppConnection.h
--- a/Engine/Source/Runtime/Online/XMPP/Public/XmppConnection.h
+++ b/Engine/Source/Runtime/Online/XMPP/Public/XmppConnection.h
@@ -2,6 +2,8 @@
 
 #pragma once
 
+#include "ModuleManager.h"
+
 /** Possible XMPP login states */
 namespace EXmppLoginStatus
 {
@@ -32,8 +34,12 @@ public:
 	int32 ServerPort;
 	/** domain for user jid */
 	FString Domain;
-	/** client id user is logging in from */
+	/** client id user is logging in from (constructed from other fields) */
 	FString ClientResource;
+	/** app id associated with this client */
+	FString AppId;
+	/** platform the player is using */
+	FString Platform;
 	/** true to enable SSL handshake for connection */
 	bool bUseSSL;
 	/** seconds between sending server ping requests */
@@ -47,7 +53,7 @@ public:
 /**
  * Jid for identifying user on the current connection as well as roster members
  */
-class FXmppUserJid
+class XMPP_API FXmppUserJid
 {
 public:
 
@@ -68,6 +74,16 @@ public:
 	/** client resource user is logged in from */
 	FString Resource;
 
+	/**
+	 * Get the components that comprise the resource
+	 *
+	 * @param OutAppId The app id the user is using
+	 * @param OutPlatform The platform the user is using
+	 *
+	 * @return Whether the Resource was successfully parsed or not
+	 */
+	bool ParseResource(FString& OutAppId, FString& OutPlatform) const;
+
 	bool operator==(const FXmppUserJid& Other)
 	{
 		return 
diff --git a/Engine/Source/Runtime/Online/XMPP/Public/XmppPresence.h b/Engine/Source/Runtime/Online/XMPP/Public/XmppPresence.h
--- a/Engine/Source/Runtime/Online/XMPP/Public/XmppPresence.h
+++ b/Engine/Source/Runtime/Online/XMPP/Public/XmppPresence.h
@@ -63,8 +63,10 @@ public:
 	bool bIsAvailable;
 	/** time when presence was sent by the user */
 	FDateTime SentTime;
-	/** client id user is logged in from */
-	FString ClientResource;
+	/** app id user is logged in from */
+	FString AppId;
+	/** platform associated with this client */
+	FString Platform;
 	/** string that will be parsed for further displayed presence info */
 	FString StatusStr;
 };
diff --git a/Engine/Source/Runtime/PacketHandlers/PacketHandler/Public/PacketHandler.h b/Engine/Source/Runtime/PacketHandlers/PacketHandler/Public/PacketHandler.h
--- a/Engine/Source/Runtime/PacketHandlers/PacketHandler/Public/PacketHandler.h
+++ b/Engine/Source/Runtime/PacketHandlers/PacketHandler/Public/PacketHandler.h
@@ -101,6 +101,10 @@ public:
 	/* Gets a packet from the the buffered packet queue for sending */
 	BufferedPacket* GetQueuedPacket();
 
+
+	/** Whether or not the PacketHandler system is enabled */
+	bool bEnabled;
+
 	/* Mode of the handler, Client or Server */
 	Handler::Mode Mode;
 
diff --git a/Engine/Source/Runtime/Renderer/Private/ScenePrivate.h b/Engine/Source/Runtime/Renderer/Private/ScenePrivate.h
--- a/Engine/Source/Runtime/Renderer/Private/ScenePrivate.h
+++ b/Engine/Source/Runtime/Renderer/Private/ScenePrivate.h
@@ -870,6 +870,10 @@ public:
 	{
 		return UniqueID;
 	}
+	uint32 GetOcclusionFrameCounter() const
+	{
+		return OcclusionFrameCounter;
+	}
 
 
 
diff --git a/Engine/Source/Runtime/Renderer/Private/SceneRendering.cpp b/Engine/Source/Runtime/Renderer/Private/SceneRendering.cpp
--- a/Engine/Source/Runtime/Renderer/Private/SceneRendering.cpp
+++ b/Engine/Source/Runtime/Renderer/Private/SceneRendering.cpp
@@ -132,14 +132,15 @@ static TAutoConsoleVariable<int32> CVarRHICmdSpewParallelListBalance(
 	0,
 	TEXT("For debugging, spews the size of the parallel command lists. This stalls and otherwise wrecks performance.\n")
 	TEXT(" 0: off (default)\n")
-	TEXT(" 1: enabled"));
+	TEXT(" 1: enabled (default)"));
 
 static TAutoConsoleVariable<int32> CVarRHICmdBalanceParallelLists(
 	TEXT("r.RHICmdBalanceParallelLists"),
 	1,
 	TEXT("Allows to enable a preprocess of the drawlists to try to balance the load equally among the command lists.\n")
 	TEXT(" 0: off \n")
-	TEXT(" 1: enabled (default)"));
+	TEXT(" 1: enabled")
+	TEXT(" 2: experiemental, uses previous frame results (does not do anything in split screen etc)"));
 
 static TAutoConsoleVariable<int32> CVarRHICmdMinCmdlistForParallelSubmit(
 	TEXT("r.RHICmdMinCmdlistForParallelSubmit"),
@@ -162,7 +163,9 @@ FParallelCommandListSet::FParallelCommandListSet(const FViewInfo& InView, FRHICo
 	Width = CVarRHICmdWidth.GetValueOnRenderThread();
 	MinDrawsPerCommandList = CVarRHICmdMinDrawsPerParallelCmdList.GetValueOnRenderThread();
 	bSpewBalance = !!CVarRHICmdSpewParallelListBalance.GetValueOnRenderThread();
-	bBalanceCommands = !!CVarRHICmdBalanceParallelLists.GetValueOnRenderThread();
+	int32 IntBalance = CVarRHICmdBalanceParallelLists.GetValueOnRenderThread();
+	bBalanceCommands = !!IntBalance;
+	bBalanceCommandsWithLastFrame = IntBalance > 1;
 	CommandLists.Reserve(Width * 8);
 	Events.Reserve(Width * 8);
 	NumDrawsIfKnown.Reserve(Width * 8);
@@ -235,6 +238,8 @@ void FParallelCommandListSet::Dispatch()
 	CommandLists.Reset();
 	Snapshot = nullptr;
 	Events.Reset();
+	QUICK_SCOPE_CYCLE_COUNTER(STAT_FParallelCommandListSet_Dispatch_ServiceLocalQueue);
+	FTaskGraphInterface::Get().ProcessThreadUntilIdle(ENamedThreads::RenderThread_Local);
 }
 
 FParallelCommandListSet::~FParallelCommandListSet()
@@ -1818,4 +1823,4 @@ bool IsMobileHDRMosaic()
 		default:
 			return !(GSupportsHDR32bppEncodeModeIntrinsic && GSupportsShaderFramebufferFetch);
 	}
-}
+}
\ No newline at end of file
diff --git a/Engine/Source/Runtime/Renderer/Private/SceneRendering.h b/Engine/Source/Runtime/Renderer/Private/SceneRendering.h
--- a/Engine/Source/Runtime/Renderer/Private/SceneRendering.h
+++ b/Engine/Source/Runtime/Renderer/Private/SceneRendering.h
@@ -328,6 +328,7 @@ public:
 	bool bBalanceCommands;
 	// see r.RHICmdSpewParallelListBalance
 	bool bSpewBalance;
+	bool bBalanceCommandsWithLastFrame;
 public:
 	TArray<FRHICommandList*,SceneRenderingAllocator> CommandLists;
 	TArray<FGraphEventRef,SceneRenderingAllocator> Events;
diff --git a/Engine/Source/Runtime/Renderer/Private/StaticMeshDrawList.h b/Engine/Source/Runtime/Renderer/Private/StaticMeshDrawList.h
--- a/Engine/Source/Runtime/Renderer/Private/StaticMeshDrawList.h
+++ b/Engine/Source/Runtime/Renderer/Private/StaticMeshDrawList.h
@@ -181,11 +181,14 @@ private:
 
 		TStaticMeshDrawList* DrawList;
 
+		uint32 VisibleCount;
+
 		/** Initialization constructor. */
 		FDrawingPolicyLink(TStaticMeshDrawList* InDrawList, const DrawingPolicyType& InDrawingPolicy, ERHIFeatureLevel::Type InFeatureLevel) :
 			DrawingPolicy(InDrawingPolicy),
 			FeatureLevel(InFeatureLevel),
-			DrawList(InDrawList)
+			DrawList(InDrawList),
+			VisibleCount(0)
 		{
 			CreateBoundShaderState();
 		}
@@ -366,6 +369,9 @@ private:
 	
 	/** All drawing policy element sets in the draw list, hashed by drawing policy. */
 	TDrawingPolicySet DrawingPolicySet;
+
+	uint32 FrameNumberForVisibleCount;
+	uint32 ViewStateUniqueId;
 };
 
 /** Helper struct for sorting */
