Copying //UE4/Dev-Mobile to Dev-Main (//UE4/Dev-Main)

==========================
MAJOR FEATURES + CHANGES
==========================

Change 2792618 on 2015/12/07 by [REDACTED_PERSON]

	Baking landscape material world-position-offset into collision

Change 2794726 on 2015/12/08 by [REDACTED_PERSON]

	Remove unused property from landscape material position offset baking implementation

Change 2797437 on 2015/12/10 by [REDACTED_PERSON]

	Fixed: Crash launching on Nexus 5 running 4.4.2 (disabled hardware instancing on this specific device/driver combo)

Change 2797524 on 2015/12/10 by [REDACTED_PERSON]

	Fixed: Materials not showing up in HUD on mobile (only when mobile HDR is off)

Change 2797562 on 2015/12/10 by [REDACTED_PERSON]

	Fixed: Reduntant seek for world composition related data inside level package in a cooked builds

Change 2808493 on 2015/12/18 by [REDACTED_PERSON]

	Fix a couple of bugs in using landscape material world position offset for collision / grass

Change 2808495 on 2015/12/18 by [REDACTED_PERSON]

	Added support for simple collision to landscape
	Added support for the pawn/vis collision view modes to landscape

Change 2822883 on 2016/01/11 by [REDACTED_PERSON]

	Add support for landscape material world position offset for collision mip / simple collision mip other than 0

Change 2824705 on 2016/01/12 by [REDACTED_PERSON]

	Fix various minor issues with Landscape collision baking from material:
	* Stopped including collision mips in grass index buffer when not baking collision
	* Added collision mips to the key for sharing index buffer
	* Grass system no longer tries to render collision mips when not baking collision
	* Changes to CollisionMipLevel / SimpleCollisionMipLevel weren't being properly propogated from landscape to render components
	** this was causing crashes by the collision baking trying to use the wrong index buffer
	* And lastly, corrected some comments :)

diff --git a/Engine/Shaders/LandscapeGrassWeight.usf b/Engine/Shaders/LandscapeGrassWeight.usf
--- a/Engine/Shaders/LandscapeGrassWeight.usf
+++ b/Engine/Shaders/LandscapeGrassWeight.usf
@@ -4,6 +4,8 @@
 	LandscapeGrassWeightVertexShader.usf: Vertex shader to dump the weight map for Landscape Grass
 =============================================================================*/
 
+#define RENDERING_LANDSCAPE_GRASS 1
+
 #include "Common.usf"
 #include "Material.usf"
 #include "VertexFactory.usf"
@@ -11,7 +13,6 @@
 struct FLandscapeGrassWeightInterpolantsVSToPS
 {
 	FVertexFactoryInterpolantsVSToPS FactoryInterpolants;
-	float4 PixelPosition	: TEXCOORD8; // xyz = world position, w = clip z
 	float  Height			: TEXCOORD9;
 	float4 Position : SV_POSITION;
 };
@@ -31,13 +32,11 @@ void VSMain(
 	FVertexFactoryIntermediates VFIntermediates = GetVertexFactoryIntermediates(Input);
 	float4 WorldPositionExcludingWPO = VertexFactoryGetWorldPosition(Input, VFIntermediates);
 	float4 WorldPosition = WorldPositionExcludingWPO;
-
-	half3x3 TangentToLocal = VertexFactoryGetTangentToLocal(Input, VFIntermediates);	
-	FMaterialVertexParameters VertexParameters = GetMaterialVertexParameters(Input, VFIntermediates, WorldPosition.xyz, TangentToLocal);
-
 	float4 RasterizedWorldPosition = VertexFactoryGetRasterizedWorldPosition(Input, VFIntermediates, WorldPosition);
-	Output.PixelPosition = WorldPosition;
 	Output.Position = mul(RasterizedWorldPosition, ResolvedView.TranslatedWorldToClip);
+
+	half3x3 TangentToLocal = VertexFactoryGetTangentToLocal(Input, VFIntermediates);
+	FMaterialVertexParameters VertexParameters = GetMaterialVertexParameters(Input, VFIntermediates, WorldPosition.xyz, TangentToLocal);
 	Output.FactoryInterpolants = VertexFactoryGetInterpolantsVSToPS(Input, VFIntermediates, VertexParameters);
 
 	float2 Height2 = LandscapeVertexFactorySampleHeight(Input, VFIntermediates);
@@ -48,8 +47,6 @@ void VSMain(
 	Output.Height = Height;
 
 	Output.Position.xy += RenderOffset;
-
-	Output.PixelPosition.w = Output.Position.w;
 }
 
 #elif PIXELSHADER
diff --git a/Engine/Shaders/LandscapeVertexFactory.usf b/Engine/Shaders/LandscapeVertexFactory.usf
--- a/Engine/Shaders/LandscapeVertexFactory.usf
+++ b/Engine/Shaders/LandscapeVertexFactory.usf
@@ -636,7 +636,7 @@ FVertexFactoryIntermediates GetVertexFactoryIntermediates(FVertexFactoryInput In
 	return Intermediates;
 }
 
-#if FEATURE_LEVEL >= FEATURE_LEVEL_SM4 && VERTEXSHADER
+#if RENDERING_LANDSCAPE_GRASS && FEATURE_LEVEL >= FEATURE_LEVEL_SM4 && VERTEXSHADER
 float2 LandscapeVertexFactorySampleHeight(FVertexFactoryInput Input, FVertexFactoryIntermediates Intermediates)
 {
 	/*
@@ -649,14 +649,12 @@ float2 LandscapeVertexFactorySampleHeight(FVertexFactoryInput Input, FVertexFact
 	float2 xy = Intermediates.InputPosition.xy * LodValues.w;
 	float LODCalculated = CalcLOD(xy, Intermediates.InputPosition.zw);
 	float LodValue = floor(LODCalculated);
-
-	// Base to Actual LOD, Base to Next LOD
-	float2 CoordTranslate = float2( LandscapeParameters.SubsectionSizeVertsLayerUVPan.x * pow(2, -LodValue) - 1, max(LandscapeParameters.SubsectionSizeVertsLayerUVPan.x * pow(2, -LodValue-1), 2) - 1 ) * LandscapeParameters.SubsectionSizeVertsLayerUVPan.y;
+	float CoordTranslate_x = (LandscapeParameters.SubsectionSizeVertsLayerUVPan.x * pow(2, -LodValue) - 1) * LandscapeParameters.SubsectionSizeVertsLayerUVPan.y;
 
 	// InputPositionLODAdjusted : Position for actual LOD in base LOD units
 	float2 ActualLODCoordsInt = floor(Intermediates.InputPosition.xy * pow(2, -(LodValue - LodValues.x)));
 
-	float2 InputPositionLODAdjusted = ActualLODCoordsInt / CoordTranslate.x;
+	float2 InputPositionLODAdjusted = ActualLODCoordsInt / CoordTranslate_x;
 	float2 SampleCoords = InputPositionLODAdjusted * LandscapeParameters.HeightmapUVScaleBias.xy + LandscapeParameters.HeightmapUVScaleBias.zw + 0.5*LandscapeParameters.HeightmapUVScaleBias.xy + Intermediates.InputPosition.zw * LandscapeParameters.SubsectionOffsetParams.xy;
 	float4 SampleValue = Texture2DSampleLevel(HeightmapTexture, HeightmapTextureSampler, SampleCoords, LodValue-LodBias.z);
 	return SampleValue.rg;
@@ -679,7 +677,19 @@ float3x3 VertexFactoryGetTangentToLocal( FVertexFactoryInput Input, FVertexFacto
 
 float4 VertexFactoryGetRasterizedWorldPosition(FVertexFactoryInput Input, FVertexFactoryIntermediates Intermediates, float4 InWorldPosition)
 {
+#if RENDERING_LANDSCAPE_GRASS
+	float2 xy = Intermediates.InputPosition.xy * LodValues.w;
+	float LODCalculated = CalcLOD(xy, Intermediates.InputPosition.zw);
+	float LodValue = floor(LODCalculated);
+	float CoordTranslate_x = (LandscapeParameters.SubsectionSizeVertsLayerUVPan.x * pow(2, -LodValue) - 1) * LandscapeParameters.SubsectionSizeVertsLayerUVPan.y;
+
+	float4 LocalPosition = mul(InWorldPosition, Primitive.WorldToLocal);
+	LocalPosition.xy *= CoordTranslate_x;
+
+	return mul(LocalPosition, Primitive.LocalToWorld);
+#else
 	return InWorldPosition;
+#endif
 }
 
 FVertexFactoryInterpolantsVSToPS VertexFactoryGetInterpolantsVSToPS(FVertexFactoryInput Input, FVertexFactoryIntermediates Intermediates, FMaterialVertexParameters VertexParameters)
diff --git a/Engine/Source/Runtime/Engine/Private/SceneView.cpp b/Engine/Source/Runtime/Engine/Private/SceneView.cpp
--- a/Engine/Source/Runtime/Engine/Private/SceneView.cpp
+++ b/Engine/Source/Runtime/Engine/Private/SceneView.cpp
@@ -546,7 +546,7 @@ FSceneView::FSceneView(const FSceneViewInitOptions& InitOptions)
 
 	// OpenGL Gamma space output in GLSL flips Y when rendering directly to the back buffer (so not needed on PC, as we never render directly into the back buffer)
 	auto ShaderPlatform = GShaderPlatformForFeatureLevel[FeatureLevel];
-	bool bUsingForwardRenderer = (Family && Family->Scene) ? !Family->Scene->ShouldUseDeferredRenderer() : false;
+	bool bUsingForwardRenderer = FSceneInterface::ShouldUseDeferredRenderer(FeatureLevel) == false;
 	bool bPlatformRequiresReverseCulling = (IsOpenGLPlatform(ShaderPlatform) && bUsingForwardRenderer && !IsPCPlatform(ShaderPlatform));
 	static auto* MobileHDRCvar = IConsoleManager::Get().FindTConsoleVariableDataInt(TEXT("r.MobileHDR"));
 	check(MobileHDRCvar);
@@ -561,7 +561,7 @@ FSceneView::FSceneView(const FSceneViewInitOptions& InitOptions)
 
 	TranslucentSortAxis = GetDefault<URendererSettings>()->TranslucentSortAxis;
 
-	// As the world is only accessable from the game thread, bIsGameView should be explicitly
+	// As the world is only accessible from the game thread, bIsGameView should be explicitly
 	// set on any other thread.
 	if(IsInGameThread())
 	{
diff --git a/Engine/Source/Runtime/Engine/Private/WorldComposition.cpp b/Engine/Source/Runtime/Engine/Private/WorldComposition.cpp
--- a/Engine/Source/Runtime/Engine/Private/WorldComposition.cpp
+++ b/Engine/Source/Runtime/Engine/Private/WorldComposition.cpp
@@ -868,10 +868,12 @@ void UWorldComposition::OnLevelPostLoad(ULevel* InLevel)
 	if (LevelPackage && InLevel->OwningWorld)
 	{
 		FWorldTileInfo Info;
-		if (InLevel->OwningWorld->WorldComposition)
+		UWorld* World = InLevel->OwningWorld;
+				
+		if (World->WorldComposition)
 		{
 			// Assign WorldLevelInfo previously loaded by world composition
-			FWorldCompositionTile* Tile = InLevel->OwningWorld->WorldComposition->FindTileByName(LevelPackage->GetFName());
+			FWorldCompositionTile* Tile = World->WorldComposition->FindTileByName(LevelPackage->GetFName());
 			if (Tile)
 			{
 				Info = Tile->Info;
@@ -879,9 +881,11 @@ void UWorldComposition::OnLevelPostLoad(ULevel* InLevel)
 		}
 		else
 		{
-			// Preserve FWorldTileInfo during standalone level loading
+#if WITH_EDITOR
+			// Preserve FWorldTileInfo in case sub-level was loaded in the editor outside of world composition
 			FString PackageFilename = FPackageName::LongPackageNameToFilename(LevelPackage->GetName(), FPackageName::GetMapPackageExtension());
 			FWorldTileInfo::Read(PackageFilename, Info);
+#endif //WITH_EDITOR
 		}
 		
 		const bool bIsDefault = (Info == FWorldTileInfo());
diff --git a/Engine/Source/Runtime/Engine/Public/SceneInterface.h b/Engine/Source/Runtime/Engine/Public/SceneInterface.h
--- a/Engine/Source/Runtime/Engine/Public/SceneInterface.h
+++ b/Engine/Source/Runtime/Engine/Public/SceneInterface.h
@@ -320,9 +320,14 @@ public:
 	virtual ERHIFeatureLevel::Type GetFeatureLevel() const { return GMaxRHIFeatureLevel; }
 	EShaderPlatform GetShaderPlatform() const { return GShaderPlatformForFeatureLevel[GetFeatureLevel()]; }
 
+	static bool ShouldUseDeferredRenderer(ERHIFeatureLevel::Type InFeatureLevel)
+	{
+		return InFeatureLevel >= ERHIFeatureLevel::SM4;
+	}
+
 	bool ShouldUseDeferredRenderer() const
 	{
-		return GetFeatureLevel() >= ERHIFeatureLevel::SM4;
+		return ShouldUseDeferredRenderer(GetFeatureLevel());
 	}
 
 	/**
diff --git a/Engine/Source/Runtime/Landscape/Classes/LandscapeComponent.h b/Engine/Source/Runtime/Landscape/Classes/LandscapeComponent.h
--- a/Engine/Source/Runtime/Landscape/Classes/LandscapeComponent.h
+++ b/Engine/Source/Runtime/Landscape/Classes/LandscapeComponent.h
@@ -123,6 +123,10 @@ struct FLandscapeComponentGrassData
 	FQuat RotationForWPO;
 
 	TArray<uint16> HeightData;
+#if WITH_EDITORONLY_DATA
+	TArray<uint16> CollisionHeightData;
+	TArray<uint16> SimpleCollisionHeightData;
+#endif
 	TMap<ULandscapeGrassType*, TArray<uint8>> WeightData;
 
 	FLandscapeComponentGrassData() {}
@@ -141,7 +145,13 @@ struct FLandscapeComponentGrassData
 		{
 			WeightSize += It.Value().GetAllocatedSize();
 		}
-		return sizeof(*this) + HeightData.GetAllocatedSize() + WeightData.GetAllocatedSize() + WeightSize;
+		return sizeof(*this)
+			+ HeightData.GetAllocatedSize()
+#if WITH_EDITORONLY_DATA
+			+ CollisionHeightData.GetAllocatedSize()
+			+ SimpleCollisionHeightData.GetAllocatedSize()
+#endif
+			+ WeightData.GetAllocatedSize() + WeightSize;
 	}
 
 	friend FArchive& operator<<(FArchive& Ar, FLandscapeComponentGrassData& Data);
@@ -238,6 +248,10 @@ public:
 	UPROPERTY(EditAnywhere, Category=LandscapeComponent)
 	int32 CollisionMipLevel;
 
+	/** Heightfield mipmap used to generate simple collision */
+	UPROPERTY(EditAnywhere, Category=LandscapeComponent)
+	int32 SimpleCollisionMipLevel;
+
 	/** StaticLightingResolution overriding per component, default value 0 means no overriding */
 	UPROPERTY(EditAnywhere, Category=LandscapeComponent)
 	float StaticLightingResolution;
@@ -498,12 +512,17 @@ public:
 	static void UpdateDataMips(int32 InNumSubsections, int32 InSubsectionSizeQuads, UTexture2D* Texture, TArray<uint8*>& TextureMipData, int32 ComponentX1=0, int32 ComponentY1=0, int32 ComponentX2=MAX_int32, int32 ComponentY2=MAX_int32, FLandscapeTextureDataInfo* TextureDataInfo=nullptr);
 
 	/**
-	 * Create or updatescollision component height data
+	 * Create or updates collision component height data
 	 * @param HeightmapTextureMipData: heightmap data
 	 * @param ComponentX1, ComponentY1, ComponentX2, ComponentY2: region to update
-	 * @param Whether to update bounds from render component.
+	 * @param bUpdateBounds: Whether to update bounds from render component.
+	 * @param XYOffsetTextureMipData: xy-offset map data
 	 */
-	void UpdateCollisionHeightData(const FColor* HeightmapTextureMipData, int32 ComponentX1=0, int32 ComponentY1=0, int32 ComponentX2=MAX_int32, int32 ComponentY2=MAX_int32, bool bUpdateBounds=false, const FColor* XYOffsetTextureMipData = NULL, bool bRebuild=false);
+	void UpdateCollisionHeightData(const FColor* HeightmapTextureMipData, const FColor* SimpleCollisionHeightmapTextureData, int32 ComponentX1=0, int32 ComponentY1=0, int32 ComponentX2=MAX_int32, int32 ComponentY2=MAX_int32, bool bUpdateBounds=false, const FColor* XYOffsetTextureMipData=nullptr);
+
+	/** Updates collision component height data for the entire component, locking and unlocking heightmap textures
+	 * @param: bRebuild: If true, recreates the collision component */
+	void UpdateCollisionData(bool bRebuild);
 
 	/**
 	 * Update collision component dominant layer data
@@ -511,7 +530,7 @@ public:
 	 * @param ComponentX1, ComponentY1, ComponentX2, ComponentY2: region to update
 	 * @param Whether to update bounds from render component.
 	 */
-	void UpdateCollisionLayerData(TArray<FColor*>& WeightmapTextureMipData, int32 ComponentX1=0, int32 ComponentY1=0, int32 ComponentX2=MAX_int32, int32 ComponentY2=MAX_int32);
+	void UpdateCollisionLayerData(const FColor* const* WeightmapTextureMipData, const FColor* const* const SimpleCollisionWeightmapTextureMipData, int32 ComponentX1=0, int32 ComponentY1=0, int32 ComponentX2=MAX_int32, int32 ComponentY2=MAX_int32);
 
 	/**
 	 * Update collision component dominant layer data for the whole component, locking and unlocking the weightmap textures.
diff --git a/Engine/Source/Runtime/Landscape/Classes/LandscapeHeightfieldCollisionComponent.h b/Engine/Source/Runtime/Landscape/Classes/LandscapeHeightfieldCollisionComponent.h
--- a/Engine/Source/Runtime/Landscape/Classes/LandscapeHeightfieldCollisionComponent.h
+++ b/Engine/Source/Runtime/Landscape/Classes/LandscapeHeightfieldCollisionComponent.h
@@ -3,7 +3,12 @@
 
 #pragma once
 
+// Forward declarations
+class ULandscapeComponent;
 class ULandscapeLayerInfoObject;
+class ULandscapeInfo;
+class ALandscape;
+class ALandscapeProxy;
 class AInstancedFoliageActor;
 
 #include "AI/Navigation/NavigationTypes.h"
@@ -41,6 +46,10 @@ class ULandscapeHeightfieldCollisionComponent : public UPrimitiveComponent
 	UPROPERTY()
 	float CollisionScale;
 
+	/** Size of component's "simple collision" in collision quads */
+	UPROPERTY()
+	int32 SimpleCollisionSizeQuads;
+
 	/** The flags for each collision quad. See ECollisionQuadFlags. */
 	UPROPERTY()
 	TArray<uint8> CollisionQuadFlags;
@@ -55,7 +64,7 @@ class ULandscapeHeightfieldCollisionComponent : public UPrimitiveComponent
 
 	/** Reference to render component */
 	UPROPERTY()
-	TLazyObjectPtr<class ULandscapeComponent> RenderComponent;
+	TLazyObjectPtr<ULandscapeComponent> RenderComponent;
 
 	struct FPhysXHeightfieldRef : public FRefCountedObject
 	{
@@ -63,31 +72,38 @@ class ULandscapeHeightfieldCollisionComponent : public UPrimitiveComponent
 
 #if WITH_PHYSX
 		/** List of PxMaterials used on this landscape */
-		TArray<physx::PxMaterial*>	UsedPhysicalMaterialArray;
-		class physx::PxHeightField*			RBHeightfield;
+		TArray<physx::PxMaterial*> UsedPhysicalMaterialArray;
+		physx::PxHeightField* RBHeightfield;
+		physx::PxHeightField* RBHeightfieldSimple;
 #if WITH_EDITOR
-		class physx::PxHeightField*			RBHeightfieldEd; // Used only by landscape editor, does not have holes in it
+		physx::PxHeightField* RBHeightfieldEd; // Used only by landscape editor, does not have holes in it
 #endif	//WITH_EDITOR
 #endif	//WITH_PHYSX
 
 		/** tors **/
-		FPhysXHeightfieldRef() 
+		FPhysXHeightfieldRef()
 #if WITH_PHYSX
-			:	RBHeightfield(NULL)
+			: RBHeightfield(nullptr)
+			, RBHeightfieldSimple(nullptr)
 #if WITH_EDITOR
-			,	RBHeightfieldEd(NULL)
+			, RBHeightfieldEd(nullptr)
 #endif	//WITH_EDITOR
 #endif	//WITH_PHYSX
-		{}
+		{
+		}
+
 		FPhysXHeightfieldRef(FGuid& InGuid)
-			:	Guid(InGuid)
+			: Guid(InGuid)
 #if WITH_PHYSX
-			,	RBHeightfield(NULL)
+			, RBHeightfield(nullptr)
+			, RBHeightfieldSimple(nullptr)
 #if WITH_EDITOR
-			,	RBHeightfieldEd(NULL)
+			, RBHeightfieldEd(nullptr)
 #endif	//WITH_EDITOR
 #endif	//WITH_PHYSX
-		{}
+		{
+		}
+
 		virtual ~FPhysXHeightfieldRef();
 	};
 	
@@ -127,14 +143,14 @@ class ULandscapeHeightfieldCollisionComponent : public UPrimitiveComponent
 	TArray<UPhysicalMaterial*>					CookedPhysicalMaterials;
 	
 	/** Physics engine version of heightfield data. */
-	TRefCountPtr<struct FPhysXHeightfieldRef>	HeightfieldRef;
+	TRefCountPtr<FPhysXHeightfieldRef>	HeightfieldRef;
 	
 	/** Cached PxHeightFieldSamples values for navmesh generation. Note that it's being used only if navigation octree is set up for lazy geometry exporting */
 	int32 HeightfieldRowsCount;
 	int32 HeightfieldColumnsCount;
 	mutable FNavHeightfieldSamples CachedHeightFieldSamples;
 
-	enum ECollisionQuadFlags
+	enum ECollisionQuadFlags : uint8
 	{
 		QF_PhysicalMaterialMask = 63,	// Mask value for the physical material index, stored in the lower 6 bits.
 		QF_EdgeTurned = 64,				// This quad's diagonal has been turned.
@@ -184,7 +200,7 @@ class ULandscapeHeightfieldCollisionComponent : public UPrimitiveComponent
 	//~ End UObject Interface.
 
 	// @todo document
-	class ULandscapeInfo* GetLandscapeInfo(bool bSpawnNewActor = true) const;
+	ULandscapeInfo* GetLandscapeInfo(bool bSpawnNewActor = true) const;
 
 	/**  We speculatively async load collision object from DDC to remove hitch when streaming */
 	void SpeculativelyLoadAsyncDDCCollsionData();
@@ -202,8 +218,8 @@ class ULandscapeHeightfieldCollisionComponent : public UPrimitiveComponent
 	virtual void CreateCollisionObject();
 
 	/** Return the landscape actor associated with this component. */
-	class ALandscape* GetLandscapeActor() const;
-	LANDSCAPE_API class ALandscapeProxy* GetLandscapeProxy() const;
+	ALandscape* GetLandscapeActor() const;
+	LANDSCAPE_API ALandscapeProxy* GetLandscapeProxy() const;
 
 	/** @return Component section base as FIntPoint */
 	LANDSCAPE_API FIntPoint GetSectionBase() const; 
diff --git a/Engine/Source/Runtime/Landscape/Classes/LandscapeProxy.h b/Engine/Source/Runtime/Landscape/Classes/LandscapeProxy.h
--- a/Engine/Source/Runtime/Landscape/Classes/LandscapeProxy.h
+++ b/Engine/Source/Runtime/Landscape/Classes/LandscapeProxy.h
@@ -353,7 +353,7 @@ public:
 
 #if WITH_EDITORONLY_DATA
 	/** LOD level to use when exporting the landscape to obj or FBX */
-	UPROPERTY(EditAnywhere, Category=LOD)
+	UPROPERTY(EditAnywhere, Category=LOD, AdvancedDisplay)
 	int32 ExportLOD;
 #endif
 
@@ -435,11 +435,16 @@ public:
 	FLightmassPrimitiveSettings LightmassSettings;
 
 	// Landscape LOD to use for collision tests. Higher numbers use less memory and process faster, but are much less accurate
-	UPROPERTY(EditAnywhere, Category=Landscape)
+	UPROPERTY(EditAnywhere, Category=Collision)
 	int32 CollisionMipLevel;
 
+	// If set higher than the "Collision Mip Level", this specifies the Landscape LOD to use for "simple collision" tests, otherwise the "Collision Mip Level" is used for both simple and complex collision.
+	// Does not work with an XY offset map (mesh collision)
+	UPROPERTY(EditAnywhere, Category=Collision)
+	int32 SimpleCollisionMipLevel;
+
 	/** Thickness of the collision surface, in unreal units */
-	UPROPERTY(EditAnywhere, Category=Landscape)
+	UPROPERTY(EditAnywhere, Category=Collision)
 	float CollisionThickness;
 
 	/** Collision profile settings for this landscape */
@@ -448,9 +453,8 @@ public:
 
 	/** Whether to bake the landscape material's vertical world position offset into the collision heightfield.
 		Note: Only z (vertical) offset is supported. XY offsets are ignored.
-		Does not work with CollisionMipLevel > 0
 		Does not work with an XY offset map (mesh collision) */
-	UPROPERTY(EditAnywhere, AdvancedDisplay, Category=Landscape)
+	UPROPERTY(EditAnywhere, AdvancedDisplay, Category=Collision)
 	bool bBakeMaterialPositionOffsetIntoCollision;
 
 #if WITH_EDITORONLY_DATA
diff --git a/Engine/Source/Runtime/Landscape/Private/LandscapeCollision.cpp b/Engine/Source/Runtime/Landscape/Private/LandscapeCollision.cpp
--- a/Engine/Source/Runtime/Landscape/Private/LandscapeCollision.cpp
+++ b/Engine/Source/Runtime/Landscape/Private/LandscapeCollision.cpp
@@ -177,6 +177,9 @@
 
 			PxTransform PhysXLandscapeComponentTransform = U2PTransform(FTransform(LandscapeComponentMatrix));
 
+			const bool bCreateSimpleCollision = SimpleCollisionSizeQuads > 0;
+			const float SimpleCollisionScale = bCreateSimpleCollision ? CollisionScale * CollisionSizeQuads / SimpleCollisionSizeQuads : 0;
+
 			// Create the geometry
 			PxHeightFieldGeometry LandscapeComponentGeom(HeightfieldRef->RBHeightfield, PxMeshGeometryFlags(), LandscapeScale.Z * LANDSCAPE_ZSCALE, LandscapeScale.Y * CollisionScale, LandscapeScale.X * CollisionScale);
 
@@ -194,14 +197,33 @@
 				CreateShapeFilterData(GetCollisionObjectType(), FMaskFilter(0), GetUniqueID(), GetCollisionResponseToChannels(), 0, 0, PQueryFilterData, PSimFilterData, true, false, true);
 
 				// Heightfield is used for simple and complex collision
-				PQueryFilterData.word3 |= (EPDF_SimpleCollision | EPDF_ComplexCollision);
-				PSimFilterData.word3 |= (EPDF_SimpleCollision | EPDF_ComplexCollision);
+				PQueryFilterData.word3 |= bCreateSimpleCollision ? EPDF_ComplexCollision : (EPDF_SimpleCollision | EPDF_ComplexCollision);
+				PSimFilterData.word3 |= bCreateSimpleCollision ? EPDF_ComplexCollision : (EPDF_SimpleCollision | EPDF_ComplexCollision);
 				HeightFieldShapeSync->setQueryFilterData(PQueryFilterData);
 				HeightFieldShapeSync->setSimulationFilterData(PSimFilterData);
 				HeightFieldShapeSync->setFlag(PxShapeFlag::eSCENE_QUERY_SHAPE, true);
 				HeightFieldShapeSync->setFlag(PxShapeFlag::eSIMULATION_SHAPE, true);
 				HeightFieldShapeSync->setFlag(PxShapeFlag::eVISUALIZATION, true);
 
+				if (bCreateSimpleCollision)
+				{
+					PxHeightFieldGeometry LandscapeComponentGeomSimple(HeightfieldRef->RBHeightfieldSimple, PxMeshGeometryFlags(), LandscapeScale.Z * LANDSCAPE_ZSCALE, LandscapeScale.Y * SimpleCollisionScale, LandscapeScale.X * SimpleCollisionScale);
+					check(LandscapeComponentGeomSimple.isValid());
+					PxShape* HeightFieldShapeSimpleSync = HeightFieldActorSync->createShape(LandscapeComponentGeomSimple, HeightfieldRef->UsedPhysicalMaterialArray.GetData(), HeightfieldRef->UsedPhysicalMaterialArray.Num());
+					check(HeightFieldShapeSimpleSync);
+
+					// Setup filtering
+					PxFilterData PQueryFilterDataSimple = PQueryFilterData;
+					PxFilterData PSimFilterDataSimple = PSimFilterData;
+					PQueryFilterDataSimple.word3 = (PQueryFilterDataSimple.word3 & ~EPDF_ComplexCollision) | EPDF_SimpleCollision;
+					PSimFilterDataSimple.word3 = (PSimFilterDataSimple.word3 & ~EPDF_ComplexCollision) | EPDF_SimpleCollision;
+					HeightFieldShapeSimpleSync->setQueryFilterData(PQueryFilterDataSimple);
+					HeightFieldShapeSimpleSync->setSimulationFilterData(PSimFilterDataSimple);
+					HeightFieldShapeSimpleSync->setFlag(PxShapeFlag::eSCENE_QUERY_SHAPE, true);
+					HeightFieldShapeSimpleSync->setFlag(PxShapeFlag::eSIMULATION_SHAPE, true);
+					HeightFieldShapeSimpleSync->setFlag(PxShapeFlag::eVISUALIZATION, true);
+				}
+
 #if WITH_EDITOR
 				// Create a shape for a heightfield which is used only by the landscape editor
 				if (!GetWorld()->IsGameWorld())
@@ -241,6 +263,26 @@
 					HeightFieldShapeAsync->setFlag(PxShapeFlag::eSCENE_QUERY_SHAPE, false);
 					HeightFieldShapeAsync->setFlag(PxShapeFlag::eSIMULATION_SHAPE, true);
 					HeightFieldShapeAsync->setFlag(PxShapeFlag::eVISUALIZATION, true);
+
+					if (bCreateSimpleCollision)
+					{
+						PxHeightFieldGeometry LandscapeComponentGeomSimple(HeightfieldRef->RBHeightfieldSimple, PxMeshGeometryFlags(), LandscapeScale.Z * LANDSCAPE_ZSCALE, LandscapeScale.Y * SimpleCollisionScale, LandscapeScale.X * SimpleCollisionScale);
+						check(LandscapeComponentGeomSimple.isValid());
+						PxShape* HeightFieldShapeSimpleAsync = HeightFieldActorAsync->createShape(LandscapeComponentGeomSimple, HeightfieldRef->UsedPhysicalMaterialArray.GetData(), HeightfieldRef->UsedPhysicalMaterialArray.Num());
+						check(HeightFieldShapeSimpleAsync);
+
+						// Setup filtering
+						PxFilterData PQueryFilterDataSimple = PQueryFilterData;
+						PxFilterData PSimFilterDataSimple = PSimFilterData;
+						PQueryFilterDataSimple.word3 = (PQueryFilterDataSimple.word3 & ~EPDF_ComplexCollision) | EPDF_SimpleCollision;
+						PSimFilterDataSimple.word3 = (PSimFilterDataSimple.word3 & ~EPDF_ComplexCollision) | EPDF_SimpleCollision;
+						HeightFieldShapeSimpleAsync->setQueryFilterData(PQueryFilterDataSimple);
+						HeightFieldShapeSimpleAsync->setSimulationFilterData(PSimFilterDataSimple);
+						// Only perform scene queries in the synchronous scene for static shapes
+						HeightFieldShapeSimpleAsync->setFlag(PxShapeFlag::eSCENE_QUERY_SHAPE, false);
+						HeightFieldShapeSimpleAsync->setFlag(PxShapeFlag::eSIMULATION_SHAPE, true);
+						HeightFieldShapeSimpleAsync->setFlag(PxShapeFlag::eVISUALIZATION, true);
+					}
 				}
 
 				// Set body instance data
@@ -337,6 +379,10 @@
 				{
 					FPhysXInputStream HeightFieldStream(CookedCollisionData.GetData(), CookedCollisionData.Num());
 					HeightfieldRef->RBHeightfield = GPhysXSDK->createHeightField(HeightFieldStream);
+					if (SimpleCollisionSizeQuads > 0)
+					{
+						HeightfieldRef->RBHeightfieldSimple = GPhysXSDK->createHeightField(HeightFieldStream);
+					}
 				}
 
 				for (UPhysicalMaterial* PhysicalMaterial : CookedPhysicalMaterials)
@@ -390,6 +436,61 @@
 #endif
 }
 
+#if WITH_PHYSX
+TArray<PxHeightFieldSample> ConvertHeightfieldDataForPhysx(const ULandscapeHeightfieldCollisionComponent* const Component, const int32 CollisionSizeVerts, const bool bIsMirrored, const uint16* Heights, const bool bUseDefMaterial, const uint8* DominantLayers, UPhysicalMaterial* const DefMaterial, TArray<UPhysicalMaterial*> &InOutMaterials)
+{
+	const int32 NumSamples = FMath::Square(CollisionSizeVerts);
+
+	TArray<PxHeightFieldSample> Samples;
+	Samples.Reserve(NumSamples);
+	Samples.AddZeroed(NumSamples);
+
+	for (int32 RowIndex = 0; RowIndex < CollisionSizeVerts; RowIndex++)
+	{
+		for (int32 ColIndex = 0; ColIndex < CollisionSizeVerts; ColIndex++)
+		{
+			int32 SrcSampleIndex = (ColIndex * CollisionSizeVerts) + (bIsMirrored ? RowIndex : (CollisionSizeVerts - RowIndex - 1));
+			int32 DstSampleIndex = (RowIndex * CollisionSizeVerts) + ColIndex;
+
+			PxHeightFieldSample& Sample = Samples[DstSampleIndex];
+			Sample.height = ((int32)Heights[SrcSampleIndex] - 32768);
+
+			// Materials are not relevant on the last row/column because they are per-triangle and the last row/column don't own any
+			if (RowIndex < CollisionSizeVerts - 1 &&
+				ColIndex < CollisionSizeVerts - 1)
+			{
+				int32 MaterialIndex = 0; // Default physical material.
+				if (!bUseDefMaterial && DominantLayers)
+				{
+					uint8 DominantLayerIdx = DominantLayers[SrcSampleIndex];
+					if (Component->ComponentLayerInfos.IsValidIndex(DominantLayerIdx))
+					{
+						ULandscapeLayerInfoObject* Layer = Component->ComponentLayerInfos[DominantLayerIdx];
+						if (Layer == ALandscapeProxy::VisibilityLayer)
+						{
+							// If it's a hole, override with the hole flag.
+							MaterialIndex = PxHeightFieldMaterial::eHOLE;
+						}
+						else
+						{
+							UPhysicalMaterial* DominantMaterial = Layer && Layer->PhysMaterial ? Layer->PhysMaterial : DefMaterial;
+							MaterialIndex = InOutMaterials.AddUnique(DominantMaterial);
+						}
+					}
+				}
+
+				Sample.materialIndex0 = MaterialIndex;
+				Sample.materialIndex1 = MaterialIndex;
+			}
+
+			// TODO: edge turning
+		}
+	}
+
+	return Samples;
+}
+#endif // WITH_PHYSX
+
 bool ULandscapeHeightfieldCollisionComponent::CookCollisionData(const FName& Format, bool bUseDefMaterial, bool bCheckDDC, TArray<uint8>& OutCookedData, TArray<UPhysicalMaterial*>& InOutMaterials) const
 {
 #if WITH_PHYSX
@@ -434,69 +535,37 @@
 	UPhysicalMaterial* DefMaterial = Proxy->DefaultPhysMaterial ? Proxy->DefaultPhysMaterial : GEngine->DefaultPhysMaterial;
 
 	// ComponentToWorld might not be initialized at this point, so use landscape transform
-	FVector LandscapeScale = Proxy->GetRootComponent()->RelativeScale3D;
-	bool bIsMirrored = (LandscapeScale.X*LandscapeScale.Y*LandscapeScale.Z) < 0.f;
+	const FVector LandscapeScale = Proxy->GetRootComponent()->RelativeScale3D;
+	const bool bIsMirrored = (LandscapeScale.X*LandscapeScale.Y*LandscapeScale.Z) < 0.f;
 
-	int32 CollisionSizeVerts = CollisionSizeQuads + 1;
+	const bool bGenerateSimpleCollision = SimpleCollisionSizeQuads > 0 && !bUseDefMaterial;
 
+	const int32 CollisionSizeVerts = CollisionSizeQuads + 1;
+	const int32 SimpleCollisionSizeVerts = SimpleCollisionSizeQuads > 0 ? SimpleCollisionSizeQuads + 1 : 0;
+	const int32 NumSamples = FMath::Square(CollisionSizeVerts);
+	const int32 NumSimpleSamples = FMath::Square(SimpleCollisionSizeVerts);
+
 	const uint16* Heights = (const uint16*)CollisionHeightData.LockReadOnly();
-	check(CollisionHeightData.GetElementCount() == FMath::Square(CollisionSizeVerts));
+	check(CollisionHeightData.GetElementCount() == NumSamples + NumSimpleSamples);
 
 	//
 	const uint8* DominantLayers = nullptr;
 	if (DominantLayerData.GetElementCount())
 	{
 		DominantLayers = (const uint8*)DominantLayerData.LockReadOnly();
+		check(DominantLayerData.GetElementCount() == NumSamples + NumSimpleSamples);
 	}
 
 	// List of materials which is actually used by heightfield
 	InOutMaterials.Empty();
 
 	TArray<PxHeightFieldSample> Samples;
-	const int32 NumSamples = FMath::Square(CollisionSizeVerts);
-	Samples.Reserve(NumSamples);
-	Samples.AddZeroed(NumSamples);
+	TArray<PxHeightFieldSample> SimpleSamples;
+	Samples = ConvertHeightfieldDataForPhysx(this, CollisionSizeVerts, bIsMirrored, Heights, bUseDefMaterial, DominantLayers, DefMaterial, InOutMaterials);
 
-	for (int32 RowIndex = 0; RowIndex < CollisionSizeVerts; RowIndex++)
+	if (bGenerateSimpleCollision)
 	{
-		for (int32 ColIndex = 0; ColIndex < CollisionSizeVerts; ColIndex++)
-		{
-			int32 SrcSampleIndex = (ColIndex * CollisionSizeVerts) + (bIsMirrored ? RowIndex : (CollisionSizeVerts - RowIndex - 1));
-			int32 DstSampleIndex = (RowIndex * CollisionSizeVerts) + ColIndex;
-
-			PxHeightFieldSample& Sample = Samples[DstSampleIndex];
-			Sample.height = FMath::Clamp<int32>(((int32)Heights[SrcSampleIndex] - 32768), -32768, 32767);
-
-			// Materials are not relevant on the last row/column because they are per-triangle and the last row/column don't own any
-			if (RowIndex < CollisionSizeVerts - 1 &&
-				ColIndex < CollisionSizeVerts - 1)
-			{
-				int32 MaterialIndex = 0; // Default physical material.
-				if (!bUseDefMaterial && DominantLayers)
-				{
-					uint8 DominantLayerIdx = DominantLayers[SrcSampleIndex];
-					if (ComponentLayerInfos.IsValidIndex(DominantLayerIdx))
-					{
-						ULandscapeLayerInfoObject* Layer = ComponentLayerInfos[DominantLayerIdx];
-						if (Layer == ALandscapeProxy::VisibilityLayer)
-						{
-							// If it's a hole, override with the hole flag.
-							MaterialIndex = PxHeightFieldMaterial::eHOLE;
-						}
-						else
-						{
-							UPhysicalMaterial* DominantMaterial = Layer && Layer->PhysMaterial ? Layer->PhysMaterial : DefMaterial;
-							MaterialIndex = InOutMaterials.AddUnique(DominantMaterial);
-						}
-					}
-				}
-
-				Sample.materialIndex0 = MaterialIndex;
-				Sample.materialIndex1 = MaterialIndex;
-			}
-
-			// TODO: edge turning
-		}
+		SimpleSamples = ConvertHeightfieldDataForPhysx(this, SimpleCollisionSizeVerts, bIsMirrored, Heights + NumSamples, bUseDefMaterial, DominantLayers, DefMaterial, InOutMaterials);
 	}
 
 	CollisionHeightData.Unlock();
@@ -520,6 +589,12 @@
 	const IPhysXFormat* Cooker = TPM->FindPhysXFormat(Format);
 	bool Result = Cooker->CookHeightField(Format, HFSize, HFThickness, Samples.GetData(), Samples.GetTypeSize(), OutData);
 
+	if (Result && bGenerateSimpleCollision)
+	{
+		FIntPoint HFSizeSimple = FIntPoint(SimpleCollisionSizeVerts, SimpleCollisionSizeVerts);
+		Result = Cooker->CookHeightField(Format, HFSizeSimple, HFThickness, SimpleSamples.GetData(), SimpleSamples.GetTypeSize(), OutData);
+	}
+
 	if (Result)
 	{
 		OutCookedData.SetNumUninitialized(OutData.Num());
diff --git a/Engine/Source/Runtime/Landscape/Private/LandscapeEdit.cpp b/Engine/Source/Runtime/Landscape/Private/LandscapeEdit.cpp
--- a/Engine/Source/Runtime/Landscape/Private/LandscapeEdit.cpp
+++ b/Engine/Source/Runtime/Landscape/Private/LandscapeEdit.cpp
@@ -49,6 +49,9 @@ DEFINE_LOG_CATEGORY(LogLandscape);
 
 #if WITH_EDITOR
 
+// Channel remapping
+extern const size_t ChannelOffsets[4];
+
 ULandscapeLayerInfoObject* ALandscapeProxy::VisibilityLayer = NULL;
 
 void ULandscapeComponent::Init(int32 InBaseX, int32 InBaseY, int32 InComponentSizeQuads, int32 InNumSubsections, int32 InSubsectionSizeQuads)
@@ -544,38 +547,31 @@ struct FLandscapeComponentAlphaInfo
 };
 
 
-void ULandscapeComponent::UpdateCollisionHeightData(const FColor* HeightmapTextureMipData, int32 ComponentX1, int32 ComponentY1, int32 ComponentX2, int32 ComponentY2, bool bUpdateBounds, const FColor* XYOffsetmapTextureData, bool bRebuild)
+void ULandscapeComponent::UpdateCollisionHeightData(const FColor* const HeightmapTextureMipData, const FColor* const SimpleCollisionHeightmapTextureData, int32 ComponentX1/*=0*/, int32 ComponentY1/*=0*/, int32 ComponentX2/*=MAX_int32*/, int32 ComponentY2/*=MAX_int32*/, bool bUpdateBounds/*=false*/, const FColor* XYOffsetTextureMipData/*=nullptr*/)
 {
 	ULandscapeInfo* Info = GetLandscapeInfo();
 	ALandscapeProxy* Proxy = GetLandscapeProxy();
 	FIntPoint ComponentKey = GetSectionBase() / ComponentSizeQuads;
 	ULandscapeHeightfieldCollisionComponent* CollisionComp = CollisionComponent.Get();
 	ULandscapeMeshCollisionComponent* MeshCollisionComponent = Cast<ULandscapeMeshCollisionComponent>(CollisionComp);
-
 	ULandscapeHeightfieldCollisionComponent* OldCollisionComponent = CollisionComp;
 
-	ALandscapeProxy* CollisionProxy = NULL;
-	if (CollisionComp && bRebuild)
-	{
-		// Remove existing component
-		CollisionProxy = CollisionComp->GetLandscapeProxy();
-		if (CollisionProxy)
-		{
-			CollisionComp->DestroyComponent();
-			CollisionComp = NULL;
-		}
-	}
+	// Simple collision is not currently supported with mesh collision components
+	const bool bUsingSimpleCollision = (SimpleCollisionMipLevel > CollisionMipLevel && SimpleCollisionHeightmapTextureData && !XYOffsetmapTexture);
 
-	int32 CollisionSubsectionSizeVerts = ((SubsectionSizeQuads + 1) >> CollisionMipLevel);
-	int32 CollisionSubsectionSizeQuads = CollisionSubsectionSizeVerts - 1;
-	int32 CollisionSizeVerts = NumSubsections*CollisionSubsectionSizeQuads + 1;
+	const int32 CollisionSubsectionSizeVerts = ((SubsectionSizeQuads + 1) >> CollisionMipLevel);
+	const int32 CollisionSubsectionSizeQuads = CollisionSubsectionSizeVerts - 1;
+	const int32 CollisionSizeVerts = NumSubsections * CollisionSubsectionSizeQuads + 1;
+	const int32 SimpleCollisionSubsectionSizeVerts = bUsingSimpleCollision ? ((SubsectionSizeQuads + 1) >> SimpleCollisionMipLevel) : 0;
+	const int32 SimpleCollisionSubsectionSizeQuads = bUsingSimpleCollision ? SimpleCollisionSubsectionSizeVerts - 1 : 0;
+	const int32 SimpleCollisionSizeVerts = bUsingSimpleCollision ? NumSubsections * SimpleCollisionSubsectionSizeQuads + 1 : 0;
 
-	uint16* CollisionHeightData = NULL;
-	uint16* CollisionXYOffsetData = NULL;
+	const int32 TotalCollisionSize = FMath::Square(CollisionSizeVerts) + FMath::Square(SimpleCollisionSizeVerts);
+
+	uint16* CollisionHeightData = nullptr;
+	uint16* CollisionXYOffsetData = nullptr;
 	bool CreatedNew = false;
 	bool ChangeType = false;
-	TArray<uint8> DominantLayerData;
-	TArray<ULandscapeLayerInfoObject*> LayerInfos;
 
 	if (CollisionComp)
 	{
@@ -583,7 +579,7 @@ void ULandscapeComponent::UpdateCollisionHeightData(const FColor* HeightmapTextu
 	}
 
 	// Existing collision component is same type with collision
-	if (CollisionComp && ((XYOffsetmapTexture == NULL) == (MeshCollisionComponent == NULL)))
+	if (CollisionComp && ((XYOffsetmapTexture == nullptr) == (MeshCollisionComponent == nullptr)))
 	{
 		ComponentX1 = FMath::Min(ComponentX1, ComponentSizeQuads);
 		ComponentY1 = FMath::Min(ComponentY1, ComponentSizeQuads);
@@ -616,16 +612,19 @@ void ULandscapeComponent::UpdateCollisionHeightData(const FColor* HeightmapTextu
 		ComponentX2 = ComponentSizeQuads;
 		ComponentY2 = ComponentSizeQuads;
 
+		TArray<uint8> DominantLayerData;
+		TArray<ULandscapeLayerInfoObject*> LayerInfos;
 		if (CollisionComp) // remove old component before changing to other type collision...
 		{
 			ChangeType = true;
 
 			if (CollisionComp->DominantLayerData.GetElementCount())
 			{
-				DominantLayerData.AddUninitialized(FMath::Square(CollisionSizeVerts));
+				check(CollisionComp->DominantLayerData.GetElementCount() >= TotalCollisionSize);
+				DominantLayerData.AddUninitialized(TotalCollisionSize);
 
 				const uint8* SrcDominantLayerData = (uint8*)CollisionComp->DominantLayerData.Lock(LOCK_READ_ONLY);
-				FMemory::Memcpy(DominantLayerData.GetData(), SrcDominantLayerData, sizeof(uint8)*FMath::Square(CollisionSizeVerts));
+				FMemory::Memcpy(DominantLayerData.GetData(), SrcDominantLayerData, TotalCollisionSize * sizeof(uint8));
 				CollisionComp->DominantLayerData.Unlock();
 			}
 
@@ -640,12 +639,19 @@ void ULandscapeComponent::UpdateCollisionHeightData(const FColor* HeightmapTextu
 			}
 			Proxy->Modify();
 			CollisionComp->DestroyComponent();
-			CollisionComp = NULL;
+			CollisionComp = nullptr;
 		}
 
-		MeshCollisionComponent = XYOffsetmapTexture ? NewObject<ULandscapeMeshCollisionComponent>(Proxy, NAME_None, RF_Transactional) : NULL;
-		CollisionComp = XYOffsetmapTexture ? Cast<ULandscapeHeightfieldCollisionComponent>(MeshCollisionComponent) :
-			NewObject<ULandscapeHeightfieldCollisionComponent>(Proxy, NAME_None, RF_Transactional);
+		if (XYOffsetmapTexture)
+		{
+			MeshCollisionComponent = NewObject<ULandscapeMeshCollisionComponent>(Proxy, NAME_None, RF_Transactional);
+			CollisionComp = MeshCollisionComponent;
+		}
+		else
+		{
+			MeshCollisionComponent = nullptr;
+			CollisionComp = NewObject<ULandscapeHeightfieldCollisionComponent>(Proxy, NAME_None, RF_Transactional);
+		}
 
 		CollisionComp->SetRelativeLocation(RelativeLocation);
 		CollisionComp->AttachTo(Proxy->GetRootComponent(), NAME_None);
@@ -655,105 +661,97 @@ void ULandscapeComponent::UpdateCollisionHeightData(const FColor* HeightmapTextu
 		CollisionComp->SetSectionBase(GetSectionBase());
 		CollisionComp->CollisionSizeQuads = CollisionSubsectionSizeQuads * NumSubsections;
 		CollisionComp->CollisionScale = (float)(ComponentSizeQuads) / (float)(CollisionComp->CollisionSizeQuads);
+		CollisionComp->SimpleCollisionSizeQuads = bUsingSimpleCollision ? SimpleCollisionSubsectionSizeQuads * NumSubsections : 0;
 		CollisionComp->CachedLocalBox = CachedLocalBox;
 		CreatedNew = true;
 
 		// Reallocate raw collision data
 		CollisionComp->CollisionHeightData.Lock(LOCK_READ_WRITE);
-		CollisionHeightData = (uint16*)CollisionComp->CollisionHeightData.Realloc(FMath::Square(CollisionSizeVerts));
-		FMemory::Memzero(CollisionHeightData, sizeof(uint16)*FMath::Square(CollisionSizeVerts));
+		CollisionHeightData = (uint16*)CollisionComp->CollisionHeightData.Realloc(TotalCollisionSize);
+		FMemory::Memzero(CollisionHeightData, TotalCollisionSize * sizeof(uint16));
 
 		if (XYOffsetmapTexture && MeshCollisionComponent)
 		{
 			// Need XYOffsetData for Collision Component
 			MeshCollisionComponent->CollisionXYOffsetData.Lock(LOCK_READ_WRITE);
-			CollisionXYOffsetData = (uint16*)MeshCollisionComponent->CollisionXYOffsetData.Realloc(FMath::Square(CollisionSizeVerts) * 2);
-			FMemory::Memzero(CollisionXYOffsetData, sizeof(uint16)*FMath::Square(CollisionSizeVerts) * 2);
+			CollisionXYOffsetData = (uint16*)MeshCollisionComponent->CollisionXYOffsetData.Realloc(TotalCollisionSize * 2);
+			FMemory::Memzero(CollisionXYOffsetData, TotalCollisionSize * 2 * sizeof(uint16));
+		}
 
-			if (DominantLayerData.Num())
-			{
-				MeshCollisionComponent->DominantLayerData.Lock(LOCK_READ_WRITE);
-				uint8* DestDominantLayerData = (uint8*)MeshCollisionComponent->DominantLayerData.Realloc(FMath::Square(CollisionSizeVerts));
-				FMemory::Memcpy(DestDominantLayerData, DominantLayerData.GetData(), sizeof(uint8)*FMath::Square(CollisionSizeVerts));
-				MeshCollisionComponent->DominantLayerData.Unlock();
-			}
+		if (DominantLayerData.Num())
+		{
+			CollisionComp->DominantLayerData.Lock(LOCK_READ_WRITE);
+			uint8* DestDominantLayerData = (uint8*)CollisionComp->DominantLayerData.Realloc(TotalCollisionSize);
+			FMemory::Memcpy(DestDominantLayerData, DominantLayerData.GetData(), TotalCollisionSize * sizeof(uint8));
+			CollisionComp->DominantLayerData.Unlock();
+		}
 
-			if (LayerInfos.Num())
-			{
-				MeshCollisionComponent->ComponentLayerInfos = LayerInfos;
-			}
+		if (LayerInfos.Num())
+		{
+			CollisionComp->ComponentLayerInfos = MoveTemp(LayerInfos);
 		}
 	}
 
-	int32 HeightmapSizeU = HeightmapTexture->Source.GetSizeX();
-	int32 HeightmapSizeV = HeightmapTexture->Source.GetSizeY();
-	int32 MipSizeU = HeightmapSizeU >> CollisionMipLevel;
-	int32 MipSizeV = HeightmapSizeV >> CollisionMipLevel;
+	const int32 HeightmapSizeU = HeightmapTexture->Source.GetSizeX();
+	const int32 HeightmapSizeV = HeightmapTexture->Source.GetSizeY();
+	const int32 MipSizeU = HeightmapSizeU >> CollisionMipLevel;
+	const int32 MipSizeV = HeightmapSizeV >> CollisionMipLevel;
 
-	int32 XYMipSizeU = XYOffsetmapTexture ? XYOffsetmapTexture->Source.GetSizeX() >> CollisionMipLevel : 0;
-	int32 XYMipSizeV = XYOffsetmapTexture ? XYOffsetmapTexture->Source.GetSizeY() >> CollisionMipLevel : 0;
+	const int32 XYMipSizeU = XYOffsetmapTexture ? XYOffsetmapTexture->Source.GetSizeX() >> CollisionMipLevel : 0;
+	const int32 XYMipSizeV = XYOffsetmapTexture ? XYOffsetmapTexture->Source.GetSizeY() >> CollisionMipLevel : 0;
 
 	// Ratio to convert update region coordinate to collision mip coordinates
-	float CollisionQuadRatio = (float)CollisionSubsectionSizeQuads / (float)SubsectionSizeQuads;
+	const float CollisionQuadRatio = (float)CollisionSubsectionSizeQuads / (float)SubsectionSizeQuads;
 
 	// XY offset into heightmap mip data
-	int32 HeightmapOffsetX = FMath::RoundToInt(HeightmapScaleBias.Z * (float)HeightmapSizeU) >> CollisionMipLevel;
-	int32 HeightmapOffsetY = FMath::RoundToInt(HeightmapScaleBias.W * (float)HeightmapSizeV) >> CollisionMipLevel;
+	const int32 HeightmapOffsetX = FMath::RoundToInt(HeightmapScaleBias.Z * (float)HeightmapSizeU) >> CollisionMipLevel;
+	const int32 HeightmapOffsetY = FMath::RoundToInt(HeightmapScaleBias.W * (float)HeightmapSizeV) >> CollisionMipLevel;
 
 	//int32 WeightmapOffsetX = FMath::RoundToInt(WeightmapScaleBias.Z * (float)XYMipSizeU) >> CollisionMipLevel;
 	//int32 WeightmapOffsetY = FMath::RoundToInt(WeightmapScaleBias.W * (float)XYMipSizeV) >> CollisionMipLevel;
 
-	// Handle WPO baked into heightfield collision
-	// WPO is not currently supported for mesh collision components
-	uint16* GrassHeights = Proxy->bBakeMaterialPositionOffsetIntoCollision && !MeshCollisionComponent && GrassData->HasData() && !IsGrassMapOutdated() ? GrassData->HeightData.GetData() : nullptr;
+	// Handle Material WPO baked into heightfield collision
+	// Material WPO is not currently supported for mesh collision components
+	const bool bUsingGrassMapHeights = Proxy->bBakeMaterialPositionOffsetIntoCollision && !MeshCollisionComponent && GrassData->HasData() && !IsGrassMapOutdated();
+	const uint16* const GrassHeights = bUsingGrassMapHeights ? (GrassData->CollisionHeightData.Num() > 0 ? GrassData->CollisionHeightData.GetData() : GrassData->HeightData.GetData()) : nullptr;
 
-	for (int32 SubsectionY = 0; SubsectionY < NumSubsections; SubsectionY++)
+	const int32 SubSectionX1 = FMath::Max(0, FMath::DivideAndRoundDown(ComponentX1 - 1, SubsectionSizeQuads));
+	const int32 SubSectionY1 = FMath::Max(0, FMath::DivideAndRoundDown(ComponentY1 - 1, SubsectionSizeQuads));
+	const int32 SubSectionX2 = FMath::Min(FMath::DivideAndRoundUp(ComponentX2 + 1, SubsectionSizeQuads), NumSubsections);
+	const int32 SubSectionY2 = FMath::Min(FMath::DivideAndRoundUp(ComponentY2 + 1, SubsectionSizeQuads), NumSubsections);
+	for (int32 SubsectionY = SubSectionY1; SubsectionY < SubSectionY2; ++SubsectionY)
 	{
-		// Check if subsection is fully above or below the area we are interested in
-		if ((ComponentY2 < SubsectionSizeQuads*SubsectionY) ||		// above
-			(ComponentY1 > SubsectionSizeQuads*(SubsectionY + 1)))	// below
+		for (int32 SubsectionX = SubSectionX1; SubsectionX < SubSectionX2; ++SubsectionX)
 		{
-			continue;
-		}
-
-		for (int32 SubsectionX = 0; SubsectionX < NumSubsections; SubsectionX++)
-		{
-			// Check if subsection is fully to the left or right of the area we are interested in
-			if ((ComponentX2 < SubsectionSizeQuads*SubsectionX) ||		// left
-				(ComponentX1 > SubsectionSizeQuads*(SubsectionX + 1)))	// right
-			{
-				continue;
-			}
-
 			// Area to update in subsection coordinates
-			int32 SubX1 = ComponentX1 - SubsectionSizeQuads*SubsectionX;
-			int32 SubY1 = ComponentY1 - SubsectionSizeQuads*SubsectionY;
-			int32 SubX2 = ComponentX2 - SubsectionSizeQuads*SubsectionX;
-			int32 SubY2 = ComponentY2 - SubsectionSizeQuads*SubsectionY;
+			const int32 SubX1 = ComponentX1 - SubsectionSizeQuads * SubsectionX;
+			const int32 SubY1 = ComponentY1 - SubsectionSizeQuads * SubsectionY;
+			const int32 SubX2 = ComponentX2 - SubsectionSizeQuads * SubsectionX;
+			const int32 SubY2 = ComponentY2 - SubsectionSizeQuads * SubsectionY;
 
 			// Area to update in collision mip level coords
-			int32 CollisionSubX1 = FMath::FloorToInt((float)SubX1 * CollisionQuadRatio);
-			int32 CollisionSubY1 = FMath::FloorToInt((float)SubY1 * CollisionQuadRatio);
-			int32 CollisionSubX2 = FMath::CeilToInt((float)SubX2 * CollisionQuadRatio);
-			int32 CollisionSubY2 = FMath::CeilToInt((float)SubY2 * CollisionQuadRatio);
+			const int32 CollisionSubX1 = FMath::FloorToInt((float)SubX1 * CollisionQuadRatio);
+			const int32 CollisionSubY1 = FMath::FloorToInt((float)SubY1 * CollisionQuadRatio);
+			const int32 CollisionSubX2 = FMath::CeilToInt( (float)SubX2 * CollisionQuadRatio);
+			const int32 CollisionSubY2 = FMath::CeilToInt( (float)SubY2 * CollisionQuadRatio);
 
 			// Clamp area to update
-			int32 VertX1 = FMath::Clamp<int32>(CollisionSubX1, 0, CollisionSubsectionSizeQuads);
-			int32 VertY1 = FMath::Clamp<int32>(CollisionSubY1, 0, CollisionSubsectionSizeQuads);
-			int32 VertX2 = FMath::Clamp<int32>(CollisionSubX2, 0, CollisionSubsectionSizeQuads);
-			int32 VertY2 = FMath::Clamp<int32>(CollisionSubY2, 0, CollisionSubsectionSizeQuads);
+			const int32 VertX1 = FMath::Clamp<int32>(CollisionSubX1, 0, CollisionSubsectionSizeQuads);
+			const int32 VertY1 = FMath::Clamp<int32>(CollisionSubY1, 0, CollisionSubsectionSizeQuads);
+			const int32 VertX2 = FMath::Clamp<int32>(CollisionSubX2, 0, CollisionSubsectionSizeQuads);
+			const int32 VertY2 = FMath::Clamp<int32>(CollisionSubY2, 0, CollisionSubsectionSizeQuads);
 
 			for (int32 VertY = VertY1; VertY <= VertY2; VertY++)
 			{
 				for (int32 VertX = VertX1; VertX <= VertX2; VertX++)
 				{
 					// X/Y of the vertex we're looking indexed into the texture data
-					int32 TexX = HeightmapOffsetX + CollisionSubsectionSizeVerts * SubsectionX + VertX;
-					int32 TexY = HeightmapOffsetY + CollisionSubsectionSizeVerts * SubsectionY + VertY;
+					const int32 TexX = HeightmapOffsetX + CollisionSubsectionSizeVerts * SubsectionX + VertX;
+					const int32 TexY = HeightmapOffsetY + CollisionSubsectionSizeVerts * SubsectionY + VertY;
 
 					// this uses Quads as we don't want the duplicated vertices
-					int32 CompVertX = CollisionSubsectionSizeQuads * SubsectionX + VertX;
-					int32 CompVertY = CollisionSubsectionSizeQuads * SubsectionY + VertY;
+					const int32 CompVertX = CollisionSubsectionSizeQuads * SubsectionX + VertX;
+					const int32 CompVertY = CollisionSubsectionSizeQuads * SubsectionY + VertY;
 
 					if (GrassHeights)
 					{
@@ -766,21 +764,21 @@ void ULandscapeComponent::UpdateCollisionHeightData(const FColor* HeightmapTextu
 						const FColor& TexData = HeightmapTextureMipData[TexX + TexY * MipSizeU];
 
 						// Copy collision data
-						uint16& CollisionHeight = CollisionHeightData[CompVertX + CompVertY*CollisionSizeVerts];
-						uint16 NewHeight = TexData.R << 8 | TexData.G;
+						uint16& CollisionHeight = CollisionHeightData[CompVertX + CompVertY * CollisionSizeVerts];
+						const uint16 NewHeight = TexData.R << 8 | TexData.G;
 
 						CollisionHeight = NewHeight;
 					}
 
-					if (XYOffsetmapTexture && XYOffsetmapTextureData && CollisionXYOffsetData)
+					if (XYOffsetmapTexture && XYOffsetTextureMipData && CollisionXYOffsetData)
 					{
-						const FColor& TexData = XYOffsetmapTextureData[TexX + TexY * XYMipSizeU];
+						const FColor& TexData = XYOffsetTextureMipData[TexX + TexY * XYMipSizeU];
 
 						// Copy collision data
-						uint16 NewXOffset = TexData.R << 8 | TexData.G;
-						uint16 NewYOffset = TexData.B << 8 | TexData.A;
+						const uint16 NewXOffset = TexData.R << 8 | TexData.G;
+						const uint16 NewYOffset = TexData.B << 8 | TexData.A;
 
-						int32 XYIndex = CompVertX + CompVertY*CollisionSizeVerts;
+						const int32 XYIndex = CompVertX + CompVertY * CollisionSizeVerts;
 						CollisionXYOffsetData[XYIndex * 2] = NewXOffset;
 						CollisionXYOffsetData[XYIndex * 2 + 1] = NewYOffset;
 					}
@@ -789,6 +787,81 @@ void ULandscapeComponent::UpdateCollisionHeightData(const FColor* HeightmapTextu
 		}
 	}
 
+	if (bUsingSimpleCollision)
+	{
+		const int32 SimpleMipSizeU = HeightmapSizeU >> SimpleCollisionMipLevel;
+		const int32 SimpleMipSizeV = HeightmapSizeV >> SimpleCollisionMipLevel;
+
+		// Ratio to convert update region coordinate to collision mip coordinates
+		const float SimpleCollisionQuadRatio = (float)SimpleCollisionSubsectionSizeQuads / (float)SubsectionSizeQuads;
+
+		// XY offset into heightmap mip data
+		const int32 SimpleHeightmapOffsetX = FMath::RoundToInt(HeightmapScaleBias.Z * (float)HeightmapSizeU) >> SimpleCollisionMipLevel;
+		const int32 SimpleHeightmapOffsetY = FMath::RoundToInt(HeightmapScaleBias.W * (float)HeightmapSizeV) >> SimpleCollisionMipLevel;
+
+		const uint16* const SimpleCollisionGrassHeights = bUsingGrassMapHeights ? GrassData->SimpleCollisionHeightData.GetData() : nullptr;
+
+		uint16* const SimpleCollisionHeightData = CollisionHeightData + FMath::Square(CollisionSizeVerts);
+
+		for (int32 SubsectionY = SubSectionY1; SubsectionY < SubSectionY2; ++SubsectionY)
+		{
+			for (int32 SubsectionX = SubSectionX1; SubsectionX < SubSectionX2; ++SubsectionX)
+			{
+				// Area to update in subsection coordinates
+				const int32 SubX1 = ComponentX1 - SubsectionSizeQuads * SubsectionX;
+				const int32 SubY1 = ComponentY1 - SubsectionSizeQuads * SubsectionY;
+				const int32 SubX2 = ComponentX2 - SubsectionSizeQuads * SubsectionX;
+				const int32 SubY2 = ComponentY2 - SubsectionSizeQuads * SubsectionY;
+
+				// Area to update in simple collision mip level coords
+				const int32 SimpleCollisionSubX1 = FMath::FloorToInt((float)SubX1 * SimpleCollisionQuadRatio);
+				const int32 SimpleCollisionSubY1 = FMath::FloorToInt((float)SubY1 * SimpleCollisionQuadRatio);
+				const int32 SimpleCollisionSubX2 = FMath::CeilToInt((float)SubX2 * SimpleCollisionQuadRatio);
+				const int32 SimpleCollisionSubY2 = FMath::CeilToInt((float)SubY2 * SimpleCollisionQuadRatio);
+
+				// Clamp area to update
+				const int32 SimpleVertX1 = FMath::Clamp<int32>(SimpleCollisionSubX1, 0, SimpleCollisionSubsectionSizeQuads);
+				const int32 SimpleVertY1 = FMath::Clamp<int32>(SimpleCollisionSubY1, 0, SimpleCollisionSubsectionSizeQuads);
+				const int32 SimpleVertX2 = FMath::Clamp<int32>(SimpleCollisionSubX2, 0, SimpleCollisionSubsectionSizeQuads);
+				const int32 SimpleVertY2 = FMath::Clamp<int32>(SimpleCollisionSubY2, 0, SimpleCollisionSubsectionSizeQuads);
+
+				for (int32 VertY = SimpleVertY1; VertY <= SimpleVertY2; VertY++)
+				{
+					for (int32 VertX = SimpleVertX1; VertX <= SimpleVertX2; VertX++)
+					{
+						// X/Y of the vertex we're looking indexed into the texture data
+						const int32 TexX = SimpleHeightmapOffsetX + SimpleCollisionSubsectionSizeVerts * SubsectionX + VertX;
+						const int32 TexY = SimpleHeightmapOffsetY + SimpleCollisionSubsectionSizeVerts * SubsectionY + VertY;
+
+						// this uses Quads as we don't want the duplicated vertices
+						const int32 CompVertX = SimpleCollisionSubsectionSizeQuads * SubsectionX + VertX;
+						const int32 CompVertY = SimpleCollisionSubsectionSizeQuads * SubsectionY + VertY;
+
+						if (SimpleCollisionGrassHeights)
+						{
+							uint16& CollisionHeight = SimpleCollisionHeightData[CompVertX + CompVertY * SimpleCollisionSizeVerts];
+							const uint16& NewHeight = SimpleCollisionGrassHeights[CompVertX + CompVertY * SimpleMipSizeU];
+							CollisionHeight = NewHeight;
+						}
+						else
+						{
+							const FColor& TexData = SimpleCollisionHeightmapTextureData[TexX + TexY * SimpleMipSizeU];
+
+							// Copy collision data
+							uint16& CollisionHeight = SimpleCollisionHeightData[CompVertX + CompVertY * SimpleCollisionSizeVerts];
+							uint16 NewHeight = TexData.R << 8 | TexData.G;
+
+							CollisionHeight = NewHeight;
+						}
+
+						// not supported
+						// if (XYOffsetmapTexture && XYOffsetTextureMipData && CollisionXYOffsetData)
+					}
+				}
+			}
+		}
+	}
+
 	CollisionComp->CollisionHeightData.Unlock();
 
 	if (XYOffsetmapTexture && MeshCollisionComponent)
@@ -810,10 +883,10 @@ void ULandscapeComponent::UpdateCollisionHeightData(const FColor* HeightmapTextu
 		}
 		else
 		{
-			int32 CollisionCompX1 = FMath::FloorToInt((float)ComponentX1 * CollisionQuadRatio);
-			int32 CollisionCompY1 = FMath::FloorToInt((float)ComponentY1 * CollisionQuadRatio);
-			int32 CollisionCompX2 = FMath::CeilToInt((float)ComponentX2 * CollisionQuadRatio);
-			int32 CollisionCompY2 = FMath::CeilToInt((float)ComponentY2 * CollisionQuadRatio);
+			const int32 CollisionCompX1 = FMath::FloorToInt((float)ComponentX1 * CollisionQuadRatio);
+			const int32 CollisionCompY1 = FMath::FloorToInt((float)ComponentY1 * CollisionQuadRatio);
+			const int32 CollisionCompX2 = FMath::CeilToInt( (float)ComponentX2 * CollisionQuadRatio);
+			const int32 CollisionCompY2 = FMath::CeilToInt( (float)ComponentY2 * CollisionQuadRatio);
 			CollisionComp->UpdateHeightfieldRegion(CollisionCompX1, CollisionCompY1, CollisionCompX2, CollisionCompY2);
 		}
 	}
@@ -833,38 +906,65 @@ void ULandscapeComponent::UpdateCollisionHeightData(const FColor* HeightmapTextu
 	// Set new collision component to pointer
 	CollisionComponent = CollisionComp;
 
-	if (bRebuild)
+	if (CreatedNew && !ChangeType)
 	{
 		UpdateCollisionLayerData();
 	}
 
-	if (bRebuild && CollisionProxy && CollisionProxy->GetRootComponent()->IsRegistered())
+	if (CreatedNew && Proxy->GetRootComponent()->IsRegistered())
 	{
-		CollisionProxy->RegisterAllComponents();
+		CollisionComp->RegisterComponent();
 	}
+}
 
-	if ((ChangeType || CreatedNew) && Proxy->GetRootComponent()->IsRegistered())
+void ULandscapeComponent::UpdateCollisionData(bool bRebuild)
+{
+	ULandscapeHeightfieldCollisionComponent* CollisionComp = CollisionComponent.Get();
+	if (CollisionComp && bRebuild)
 	{
-		Proxy->RegisterAllComponents();
+		CollisionComp->DestroyComponent();
+		CollisionComponent = CollisionComp = nullptr;
+	}
+
+	TArray<uint8> CollisionMipData;
+	TArray<uint8> SimpleCollisionMipData;
+	TArray<uint8> XYOffsetMipData;
+
+	HeightmapTexture->Source.GetMipData(CollisionMipData, CollisionMipLevel);
+	if (SimpleCollisionMipLevel > CollisionMipLevel)
+	{
+		HeightmapTexture->Source.GetMipData(SimpleCollisionMipData, SimpleCollisionMipLevel);
+	}
+	if (XYOffsetmapTexture)
+	{
+		XYOffsetmapTexture->Source.GetMipData(XYOffsetMipData, CollisionMipLevel);
 	}
-}
 
+	UpdateCollisionHeightData(
+		(FColor*)CollisionMipData.GetData(),
+		SimpleCollisionMipLevel > CollisionMipLevel ? (FColor*)SimpleCollisionMipData.GetData() : nullptr,
+		0, 0, MAX_int32, MAX_int32, true,
+		XYOffsetmapTexture ? (FColor*)XYOffsetMipData.GetData() : nullptr);
+}
 
-void ULandscapeComponent::UpdateCollisionLayerData(TArray<FColor*>& WeightmapTextureMipData, int32 ComponentX1, int32 ComponentY1, int32 ComponentX2, int32 ComponentY2)
+void ULandscapeComponent::UpdateCollisionLayerData(const FColor* const* const WeightmapTextureMipData, const FColor* const* const SimpleCollisionWeightmapTextureMipData, int32 ComponentX1, int32 ComponentY1, int32 ComponentX2, int32 ComponentY2)
 {
 	ULandscapeInfo* Info = GetLandscapeInfo();
 	ALandscapeProxy* Proxy = GetLandscapeProxy();
 	FIntPoint ComponentKey = GetSectionBase() / ComponentSizeQuads;
 
-	if (CollisionComponent)
+	ULandscapeHeightfieldCollisionComponent* CollisionComp = CollisionComponent.Get();
+
+	if (CollisionComp)
 	{
-		CollisionComponent->Modify();
+		CollisionComp->Modify();
+
+		// Simple collision is not currently supported with mesh collision components
+		const bool bUsingSimpleCollision = (SimpleCollisionMipLevel > CollisionMipLevel && SimpleCollisionWeightmapTextureMipData && !XYOffsetmapTexture);
 
 		TArray<ULandscapeLayerInfoObject*> CandidateLayers;
 		TArray<uint8*> CandidateDataPtrs;
-
-		// Channel remapping
-		int32 ChannelOffsets[4] = { (int32)STRUCT_OFFSET(FColor, R), (int32)STRUCT_OFFSET(FColor, G), (int32)STRUCT_OFFSET(FColor, B), (int32)STRUCT_OFFSET(FColor, A) };
+		TArray<uint8*> SimpleCollisionDataPtrs;
 
 		bool bExistingLayerMismatch = false;
 		int32 DataLayerIdx = INDEX_NONE;
@@ -879,8 +979,13 @@ void ULandscapeComponent::UpdateCollisionLayerData(TArray<FColor*>& WeightmapTex
 				int32 Idx = CandidateLayers.Add(LayerInfo);
 				CandidateDataPtrs.Add(((uint8*)WeightmapTextureMipData[AllocInfo.WeightmapTextureIndex]) + ChannelOffsets[AllocInfo.WeightmapTextureChannel]);
 
+				if (bUsingSimpleCollision)
+				{
+					SimpleCollisionDataPtrs.Add(((uint8*)SimpleCollisionWeightmapTextureMipData[AllocInfo.WeightmapTextureIndex]) + ChannelOffsets[AllocInfo.WeightmapTextureChannel]);
+				}
+
 				// Check if we still match the collision component.
-				if (!CollisionComponent->ComponentLayerInfos.IsValidIndex(Idx) || CollisionComponent->ComponentLayerInfos[Idx] != LayerInfo)
+				if (!CollisionComp->ComponentLayerInfos.IsValidIndex(Idx) || CollisionComp->ComponentLayerInfos[Idx] != LayerInfo)
 				{
 					bExistingLayerMismatch = true;
 				}
@@ -896,84 +1001,88 @@ void ULandscapeComponent::UpdateCollisionLayerData(TArray<FColor*>& WeightmapTex
 		if (CandidateLayers.Num() == 0)
 		{
 			// No layers, so don't update any weights
-			CollisionComponent->DominantLayerData.RemoveBulkData();
-			CollisionComponent->ComponentLayerInfos.Empty();
+			CollisionComp->DominantLayerData.RemoveBulkData();
+			CollisionComp->ComponentLayerInfos.Empty();
 		}
 		else
 		{
-			int32 CollisionSubsectionSizeVerts = ((SubsectionSizeQuads + 1) >> CollisionMipLevel);
-			int32 CollisionSubsectionSizeQuads = CollisionSubsectionSizeVerts - 1;
-			int32 CollisionSizeVerts = NumSubsections*CollisionSubsectionSizeQuads + 1;
-			uint8* DominantLayerData = NULL;
+			const int32 CollisionSubsectionSizeVerts = ((SubsectionSizeQuads + 1) >> CollisionMipLevel);
+			const int32 CollisionSubsectionSizeQuads = CollisionSubsectionSizeVerts - 1;
+			const int32 CollisionSizeVerts = NumSubsections * CollisionSubsectionSizeQuads + 1;
+			const int32 SimpleCollisionSubsectionSizeVerts = bUsingSimpleCollision ? ((SubsectionSizeQuads + 1) >> SimpleCollisionMipLevel) : 0;
+			const int32 SimpleCollisionSubsectionSizeQuads = bUsingSimpleCollision ? SimpleCollisionSubsectionSizeVerts - 1 : 0;
+			const int32 SimpleCollisionSizeVerts = bUsingSimpleCollision ? NumSubsections * SimpleCollisionSubsectionSizeQuads + 1 : 0;
+
+			const int32 TotalCollisionSize = FMath::Square(CollisionSizeVerts) + FMath::Square(SimpleCollisionSizeVerts);
+
+			uint8* DominantLayerData = nullptr;
 
 			// If there's no existing data, or the layer allocations have changed, we need to update the data for the whole component.
-			if (bExistingLayerMismatch || CollisionComponent->DominantLayerData.GetElementCount() == 0)
+			if (bExistingLayerMismatch || CollisionComp->DominantLayerData.GetElementCount() == 0)
 			{
 				ComponentX1 = 0;
 				ComponentY1 = 0;
-				ComponentX2 = MAX_int32;
-				ComponentY2 = MAX_int32;
+				ComponentX2 = ComponentSizeQuads;
+				ComponentY2 = ComponentSizeQuads;
 
-				CollisionComponent->DominantLayerData.Lock(LOCK_READ_WRITE);
-				DominantLayerData = (uint8*)CollisionComponent->DominantLayerData.Realloc(FMath::Square(CollisionSizeVerts));
-				FMemory::Memzero(DominantLayerData, FMath::Square(CollisionSizeVerts));
+				CollisionComp->DominantLayerData.Lock(LOCK_READ_WRITE);
+				DominantLayerData = (uint8*)CollisionComp->DominantLayerData.Realloc(TotalCollisionSize);
+				FMemory::Memzero(DominantLayerData, TotalCollisionSize);
 
-				CollisionComponent->ComponentLayerInfos = CandidateLayers;
+				CollisionComp->ComponentLayerInfos = CandidateLayers;
 			}
 			else
 			{
-				DominantLayerData = (uint8*)CollisionComponent->DominantLayerData.Lock(LOCK_READ_WRITE);
+				ComponentX1 = FMath::Min(ComponentX1, ComponentSizeQuads);
+				ComponentY1 = FMath::Min(ComponentY1, ComponentSizeQuads);
+				ComponentX2 = FMath::Max(ComponentX2, 0);
+				ComponentY2 = FMath::Max(ComponentY2, 0);
+
+				DominantLayerData = (uint8*)CollisionComp->DominantLayerData.Lock(LOCK_READ_WRITE);
 			}
 
-			int32 MipSizeU = (WeightmapTextures[0]->Source.GetSizeX()) >> CollisionMipLevel;
+			const int32 WeightmapSizeU = WeightmapTextures[0]->Source.GetSizeX();
+			const int32 MipSizeU = WeightmapSizeU >> CollisionMipLevel;
 
 			// Ratio to convert update region coordinate to collision mip coordinates
-			float CollisionQuadRatio = (float)CollisionSubsectionSizeQuads / (float)SubsectionSizeQuads;
+			const float CollisionQuadRatio = (float)CollisionSubsectionSizeQuads / (float)SubsectionSizeQuads;
 
-			for (int32 SubsectionY = 0; SubsectionY < NumSubsections; SubsectionY++)
-			{
-				// Check if subsection is fully above or below the area we are interested in
-				if ((ComponentY2 < SubsectionSizeQuads*SubsectionY) ||		// above
-					(ComponentY1 > SubsectionSizeQuads*(SubsectionY + 1)))	// below
-				{
-					continue;
-				}
+			// gmartin: WeightmapScaleBias not handled?
 
-				for (int32 SubsectionX = 0; SubsectionX < NumSubsections; SubsectionX++)
+			const int32 SubSectionX1 = FMath::Max(0, FMath::DivideAndRoundDown(ComponentX1 - 1, SubsectionSizeQuads));
+			const int32 SubSectionY1 = FMath::Max(0, FMath::DivideAndRoundDown(ComponentY1 - 1, SubsectionSizeQuads));
+			const int32 SubSectionX2 = FMath::Min(FMath::DivideAndRoundUp(ComponentX2 + 1, SubsectionSizeQuads), NumSubsections);
+			const int32 SubSectionY2 = FMath::Min(FMath::DivideAndRoundUp(ComponentY2 + 1, SubsectionSizeQuads), NumSubsections);
+			for (int32 SubsectionY = SubSectionY1; SubsectionY < SubSectionY2; ++SubsectionY)
+			{
+				for (int32 SubsectionX = SubSectionX1; SubsectionX < SubSectionX2; ++SubsectionX)
 				{
-					// Check if subsection is fully to the left or right of the area we are interested in
-					if ((ComponentX2 < SubsectionSizeQuads*SubsectionX) ||		// left
-						(ComponentX1 > SubsectionSizeQuads*(SubsectionX + 1)))	// right
-					{
-						continue;
-					}
-
 					// Area to update in subsection coordinates
-					int32 SubX1 = ComponentX1 - SubsectionSizeQuads*SubsectionX;
-					int32 SubY1 = ComponentY1 - SubsectionSizeQuads*SubsectionY;
-					int32 SubX2 = ComponentX2 - SubsectionSizeQuads*SubsectionX;
-					int32 SubY2 = ComponentY2 - SubsectionSizeQuads*SubsectionY;
+					const int32 SubX1 = ComponentX1 - SubsectionSizeQuads * SubsectionX;
+					const int32 SubY1 = ComponentY1 - SubsectionSizeQuads * SubsectionY;
+					const int32 SubX2 = ComponentX2 - SubsectionSizeQuads * SubsectionX;
+					const int32 SubY2 = ComponentY2 - SubsectionSizeQuads * SubsectionY;
 
 					// Area to update in collision mip level coords
-					int32 CollisionSubX1 = FMath::FloorToInt((float)SubX1 * CollisionQuadRatio);
-					int32 CollisionSubY1 = FMath::FloorToInt((float)SubY1 * CollisionQuadRatio);
-					int32 CollisionSubX2 = FMath::CeilToInt((float)SubX2 * CollisionQuadRatio);
-					int32 CollisionSubY2 = FMath::CeilToInt((float)SubY2 * CollisionQuadRatio);
+					const int32 CollisionSubX1 = FMath::FloorToInt((float)SubX1 * CollisionQuadRatio);
+					const int32 CollisionSubY1 = FMath::FloorToInt((float)SubY1 * CollisionQuadRatio);
+					const int32 CollisionSubX2 = FMath::CeilToInt( (float)SubX2 * CollisionQuadRatio);
+					const int32 CollisionSubY2 = FMath::CeilToInt( (float)SubY2 * CollisionQuadRatio);
 
 					// Clamp area to update
-					int32 VertX1 = FMath::Clamp<int32>(CollisionSubX1, 0, CollisionSubsectionSizeQuads);
-					int32 VertY1 = FMath::Clamp<int32>(CollisionSubY1, 0, CollisionSubsectionSizeQuads);
-					int32 VertX2 = FMath::Clamp<int32>(CollisionSubX2, 0, CollisionSubsectionSizeQuads);
-					int32 VertY2 = FMath::Clamp<int32>(CollisionSubY2, 0, CollisionSubsectionSizeQuads);
+					const int32 VertX1 = FMath::Clamp<int32>(CollisionSubX1, 0, CollisionSubsectionSizeQuads);
+					const int32 VertY1 = FMath::Clamp<int32>(CollisionSubY1, 0, CollisionSubsectionSizeQuads);
+					const int32 VertX2 = FMath::Clamp<int32>(CollisionSubX2, 0, CollisionSubsectionSizeQuads);
+					const int32 VertY2 = FMath::Clamp<int32>(CollisionSubY2, 0, CollisionSubsectionSizeQuads);
 
 					for (int32 VertY = VertY1; VertY <= VertY2; VertY++)
 					{
 						for (int32 VertX = VertX1; VertX <= VertX2; VertX++)
 						{
 							// X/Y of the vertex we're looking indexed into the texture data
-							int32 TexX = CollisionSubsectionSizeVerts * SubsectionX + VertX;
-							int32 TexY = CollisionSubsectionSizeVerts * SubsectionY + VertY;
-							int32 DataOffset = (TexX + TexY * MipSizeU) * sizeof(FColor);
+							const int32 TexX = CollisionSubsectionSizeVerts * SubsectionX + VertX;
+							const int32 TexY = CollisionSubsectionSizeVerts * SubsectionY + VertY;
+							const int32 DataOffset = (TexX + TexY * MipSizeU) * sizeof(FColor);
 
 							uint8 DominantLayer = 255; // 255 as invalid value
 							int32 DominantWeight = 0;
@@ -997,8 +1106,8 @@ void ULandscapeComponent::UpdateCollisionLayerData(TArray<FColor*>& WeightmapTex
 							}
 
 							// this uses Quads as we don't want the duplicated vertices
-							int32 CompVertX = CollisionSubsectionSizeQuads * SubsectionX + VertX;
-							int32 CompVertY = CollisionSubsectionSizeQuads * SubsectionY + VertY;
+							const int32 CompVertX = CollisionSubsectionSizeQuads * SubsectionX + VertX;
+							const int32 CompVertY = CollisionSubsectionSizeQuads * SubsectionY + VertY;
 
 							// Set collision data
 							DominantLayerData[CompVertX + CompVertY*CollisionSizeVerts] = DominantLayer;
@@ -1006,7 +1115,83 @@ void ULandscapeComponent::UpdateCollisionLayerData(TArray<FColor*>& WeightmapTex
 					}
 				}
 			}
-			CollisionComponent->DominantLayerData.Unlock();
+
+			if (bUsingSimpleCollision)
+			{
+				const int32 SimpleMipSizeU = WeightmapSizeU >> SimpleCollisionMipLevel;
+
+				// Ratio to convert update region coordinate to collision mip coordinates
+				const float SimpleCollisionQuadRatio = (float)SimpleCollisionSubsectionSizeQuads / (float)SubsectionSizeQuads;
+
+				// gmartin: WeightmapScaleBias not handled?
+
+				uint8* const SimpleCollisionHeightData = DominantLayerData + FMath::Square(CollisionSizeVerts);
+
+				for (int32 SubsectionY = SubSectionY1; SubsectionY < SubSectionY2; ++SubsectionY)
+				{
+					for (int32 SubsectionX = SubSectionX1; SubsectionX < SubSectionX2; ++SubsectionX)
+					{
+						// Area to update in subsection coordinates
+						const int32 SubX1 = ComponentX1 - SubsectionSizeQuads * SubsectionX;
+						const int32 SubY1 = ComponentY1 - SubsectionSizeQuads * SubsectionY;
+						const int32 SubX2 = ComponentX2 - SubsectionSizeQuads * SubsectionX;
+						const int32 SubY2 = ComponentY2 - SubsectionSizeQuads * SubsectionY;
+
+						// Area to update in simple collision mip level coords
+						const int32 SimpleCollisionSubX1 = FMath::FloorToInt((float)SubX1 * SimpleCollisionQuadRatio);
+						const int32 SimpleCollisionSubY1 = FMath::FloorToInt((float)SubY1 * SimpleCollisionQuadRatio);
+						const int32 SimpleCollisionSubX2 = FMath::CeilToInt( (float)SubX2 * SimpleCollisionQuadRatio);
+						const int32 SimpleCollisionSubY2 = FMath::CeilToInt( (float)SubY2 * SimpleCollisionQuadRatio);
+
+						// Clamp area to update
+						const int32 SimpleVertX1 = FMath::Clamp<int32>(SimpleCollisionSubX1, 0, SimpleCollisionSubsectionSizeQuads);
+						const int32 SimpleVertY1 = FMath::Clamp<int32>(SimpleCollisionSubY1, 0, SimpleCollisionSubsectionSizeQuads);
+						const int32 SimpleVertX2 = FMath::Clamp<int32>(SimpleCollisionSubX2, 0, SimpleCollisionSubsectionSizeQuads);
+						const int32 SimpleVertY2 = FMath::Clamp<int32>(SimpleCollisionSubY2, 0, SimpleCollisionSubsectionSizeQuads);
+
+						for (int32 VertY = SimpleVertY1; VertY <= SimpleVertY2; VertY++)
+						{
+							for (int32 VertX = SimpleVertX1; VertX <= SimpleVertX2; VertX++)
+							{
+								// X/Y of the vertex we're looking indexed into the texture data
+								const int32 TexX = SimpleCollisionSubsectionSizeVerts * SubsectionX + VertX;
+								const int32 TexY = SimpleCollisionSubsectionSizeVerts * SubsectionY + VertY;
+								const int32 DataOffset = (TexX + TexY * SimpleMipSizeU) * sizeof(FColor);
+
+								uint8 DominantLayer = 255; // 255 as invalid value
+								int32 DominantWeight = 0;
+								for (int32 LayerIdx = 0; LayerIdx < CandidateDataPtrs.Num(); LayerIdx++)
+								{
+									const uint8 LayerWeight = SimpleCollisionDataPtrs[LayerIdx][DataOffset];
+
+									if (LayerIdx == DataLayerIdx) // Override value for hole
+									{
+										if (LayerWeight > 170) // 255 * 0.66...
+										{
+											DominantLayer = LayerIdx;
+											DominantWeight = INT_MAX;
+										}
+									}
+									else if (LayerWeight > DominantWeight)
+									{
+										DominantLayer = LayerIdx;
+										DominantWeight = LayerWeight;
+									}
+								}
+
+								// this uses Quads as we don't want the duplicated vertices
+								const int32 CompVertX = SimpleCollisionSubsectionSizeQuads * SubsectionX + VertX;
+								const int32 CompVertY = SimpleCollisionSubsectionSizeQuads * SubsectionY + VertY;
+
+								// Set collision data
+								DominantLayerData[CompVertX + CompVertY * SimpleCollisionSizeVerts] = DominantLayer;
+							}
+						}
+					}
+				}
+			}
+
+			CollisionComp->DominantLayerData.Unlock();
 		}
 
 		// We do not force an update of the physics data here. We don't need the layer information in the editor and it
@@ -1019,18 +1204,32 @@ void ULandscapeComponent::UpdateCollisionLayerData()
 {
 	// Generate the dominant layer data
 	TArray<FColor*> WeightmapTextureMipData;
-	TArray<TArray<uint8> > CachedWeightmapTextureMipData;
+	TArray<TArray<uint8>> CachedWeightmapTextureMipData;
 
 	WeightmapTextureMipData.Empty(WeightmapTextures.Num());
 	CachedWeightmapTextureMipData.Empty(WeightmapTextures.Num());
-	for (int32 Idx = 0; Idx < WeightmapTextures.Num(); Idx++)
+	for (int32 WeightmapIdx = 0; WeightmapIdx < WeightmapTextures.Num(); ++WeightmapIdx)
 	{
-		TArray<uint8>* MipData = new(CachedWeightmapTextureMipData)TArray<uint8>();
-		WeightmapTextures[Idx]->Source.GetMipData(*MipData, CollisionMipLevel);
-		WeightmapTextureMipData.Add((FColor*)MipData->GetData());
+		TArray<uint8> MipData;
+		WeightmapTextures[WeightmapIdx]->Source.GetMipData(MipData, CollisionMipLevel);
+		WeightmapTextureMipData.Add((FColor*)MipData.GetData());
+		CachedWeightmapTextureMipData.Add(MoveTemp(MipData));
 	}
 
-	UpdateCollisionLayerData(WeightmapTextureMipData);
+	TArray<FColor*> SimpleCollisionWeightmapMipData;
+	TArray<TArray<uint8>> SimpleCollisionCachedWeightmapTextureMipData;
+	if (SimpleCollisionMipLevel > CollisionMipLevel)
+	{
+		for (int32 WeightmapIdx = 0; WeightmapIdx < WeightmapTextures.Num(); ++WeightmapIdx)
+		{
+			TArray<uint8> MipData;
+			WeightmapTextures[WeightmapIdx]->Source.GetMipData(MipData, SimpleCollisionMipLevel);
+			SimpleCollisionWeightmapMipData.Add((FColor*)MipData.GetData());
+			SimpleCollisionCachedWeightmapTextureMipData.Add(MoveTemp(MipData));
+		}
+	}
+
+	UpdateCollisionLayerData(WeightmapTextureMipData.GetData(), SimpleCollisionWeightmapMipData.GetData());
 }
 
 
@@ -2275,7 +2474,9 @@ LANDSCAPE_API void ALandscapeProxy::Import(
 
 			ULandscapeComponent* LandscapeComponent = LandscapeComponents[ComponentX + ComponentY*NumComponentsX];
 			LandscapeComponent->GenerateHeightmapMips(HeightmapInfo.HeightmapTextureMipData, ComponentX == NumComponentsX - 1 ? MAX_int32 : 0, ComponentY == NumComponentsY - 1 ? MAX_int32 : 0);
-			LandscapeComponent->UpdateCollisionHeightData(HeightmapInfo.HeightmapTextureMipData[LandscapeComponent->CollisionMipLevel]);
+			LandscapeComponent->UpdateCollisionHeightData(
+				HeightmapInfo.HeightmapTextureMipData[LandscapeComponent->CollisionMipLevel],
+				LandscapeComponent->SimpleCollisionMipLevel > LandscapeComponent->CollisionMipLevel ? HeightmapInfo.HeightmapTextureMipData[LandscapeComponent->SimpleCollisionMipLevel] : nullptr);
 			LandscapeComponent->UpdateCollisionLayerData();
 		}
 	}
@@ -3128,22 +3329,14 @@ void ALandscapeProxy::RecreateCollisionComponents()
 
 	// Recreate collision
 	CollisionMipLevel = FMath::Clamp<int32>(CollisionMipLevel, 0, FMath::CeilLogTwo(SubsectionSizeQuads + 1) - 1);
-	for (int32 i = 0; i < LandscapeComponents.Num(); ++i)
+	SimpleCollisionMipLevel = FMath::Clamp<int32>(SimpleCollisionMipLevel, 0, FMath::CeilLogTwo(SubsectionSizeQuads + 1) - 1);
+	for (ULandscapeComponent* Comp : LandscapeComponents)
 	{
-		ULandscapeComponent* Comp = LandscapeComponents[i];
 		if (Comp)
 		{
 			Comp->CollisionMipLevel = CollisionMipLevel;
-			TArray<uint8> CollisionMipData;
-			Comp->HeightmapTexture->Source.GetMipData(CollisionMipData, CollisionMipLevel);
-			TArray<uint8> XYOffsetMipData;
-			if (Comp->XYOffsetmapTexture)
-			{
-				Comp->XYOffsetmapTexture->Source.GetMipData(XYOffsetMipData, CollisionMipLevel);
-			}
-
-			// Rebuild all collision
-			Comp->UpdateCollisionHeightData((FColor*)CollisionMipData.GetData(), 0, 0, MAX_int32, MAX_int32, true, XYOffsetMipData.Num() ? (FColor*)XYOffsetMipData.GetData() : nullptr, true);
+			Comp->SimpleCollisionMipLevel = SimpleCollisionMipLevel;
+			Comp->UpdateCollisionData(true);
 		}
 	}
 }
@@ -3291,15 +3484,21 @@ void ALandscapeProxy::PostEditChangeProperty(FPropertyChangedEvent& PropertyChan
 	{
 		ChangedPhysMaterial();
 	}
-	else if (GIsEditor && (PropertyName == FName(TEXT("CollisionMipLevel")) || PropertyName == FName(TEXT("CollisionThickness")) || PropertyName == GET_MEMBER_NAME_CHECKED(ALandscapeProxy, bBakeMaterialPositionOffsetIntoCollision)))
+	else if (GIsEditor &&
+		(PropertyName == GET_MEMBER_NAME_CHECKED(ALandscapeProxy, CollisionMipLevel) ||
+		 PropertyName == GET_MEMBER_NAME_CHECKED(ALandscapeProxy, SimpleCollisionMipLevel) ||
+		 PropertyName == GET_MEMBER_NAME_CHECKED(ALandscapeProxy, CollisionThickness) ||
+		 PropertyName == GET_MEMBER_NAME_CHECKED(ALandscapeProxy, bBakeMaterialPositionOffsetIntoCollision)))
 	{
+		if (bBakeMaterialPositionOffsetIntoCollision)
+		{
+			MarkComponentsRenderStateDirty();
+		}
 		RecreateCollisionComponents();
 	}
-	else 
-		if(PropertyName == FName(TEXT("bCastStaticShadow"))
-		|| PropertyName == FName(TEXT("bCastShadowAsTwoSided"))
-		|| PropertyName == FName(TEXT("bCastFarShadow"))
-		)
+	else if(PropertyName == FName(TEXT("bCastStaticShadow")) ||
+		PropertyName == FName(TEXT("bCastShadowAsTwoSided")) ||
+		PropertyName == FName(TEXT("bCastFarShadow")))
 	{
 		// Replicate shared properties to all components.
 		for (int32 ComponentIndex = 0; ComponentIndex < LandscapeComponents.Num(); ComponentIndex++)
@@ -3444,6 +3643,11 @@ void ALandscape::PostEditChangeProperty(FPropertyChangedEvent& PropertyChangedEv
 		CollisionMipLevel = FMath::Clamp<int32>(CollisionMipLevel, 0, FMath::CeilLogTwo(SubsectionSizeQuads + 1) - 1);
 		bPropagateToProxies = true;
 	}
+	else if (PropertyName == GET_MEMBER_NAME_CHECKED(ALandscapeProxy, SimpleCollisionMipLevel))
+	{
+		SimpleCollisionMipLevel = FMath::Clamp<int32>(SimpleCollisionMipLevel, 0, FMath::CeilLogTwo(SubsectionSizeQuads + 1) - 1);
+		bPropagateToProxies = true;
+	}
 	else if (PropertyName == GET_MEMBER_NAME_CHECKED(ALandscapeProxy, bBakeMaterialPositionOffsetIntoCollision))
 	{
 		bPropagateToProxies = true;
@@ -3679,12 +3883,13 @@ void ULandscapeComponent::PostEditChangeProperty(FPropertyChangedEvent& Property
 		LightingLODBias = FMath::Clamp<int32>(LightingLODBias, -1, MaxLOD);
 		InvalidateLightingCache();
 	}
-	else if (GIsEditor && (PropertyName == FName(TEXT("CollisionMipLevel"))))
+	else if (GIsEditor &&
+		(PropertyName == GET_MEMBER_NAME_CHECKED(ULandscapeComponent, CollisionMipLevel) ||
+		 PropertyName == GET_MEMBER_NAME_CHECKED(ULandscapeComponent, SimpleCollisionMipLevel)))
 	{
 		CollisionMipLevel = FMath::Clamp<int32>(CollisionMipLevel, 0, FMath::CeilLogTwo(SubsectionSizeQuads + 1) - 1);
-		TArray<uint8> CollisionMipData;
-		HeightmapTexture->Source.GetMipData(CollisionMipData, CollisionMipLevel);
-		UpdateCollisionHeightData((FColor*)CollisionMipData.GetData(), 0, 0, MAX_int32, MAX_int32, true, NULL, true); // Rebuild for new CollisionMipLevel
+		SimpleCollisionMipLevel = FMath::Clamp<int32>(SimpleCollisionMipLevel, 0, FMath::CeilLogTwo(SubsectionSizeQuads + 1) - 1);
+		UpdateCollisionData(true); // Rebuild for new CollisionMipLevel
 	}
 }
 
@@ -4150,7 +4355,9 @@ void ULandscapeComponent::InitHeightmapData(TArray<FColor>& Heights, bool bUpdat
 
 	if (bUpdateCollision)
 	{
-		ULandscapeComponent::UpdateCollisionHeightData(HeightmapTextureMipData[CollisionMipLevel]);
+		UpdateCollisionHeightData(
+			HeightmapTextureMipData[CollisionMipLevel],
+			SimpleCollisionMipLevel > CollisionMipLevel ? HeightmapTextureMipData[SimpleCollisionMipLevel] : nullptr);
 	}
 
 	for (int32 i = 0; i < HeightmapTextureMipData.Num(); i++)
@@ -4206,9 +4413,6 @@ void ULandscapeComponent::InitWeightmapData(TArray<ULandscapeLayerInfoObject*>&
 	WeightmapScaleBias = FVector4(1.0f / (float)WeightmapSize, 1.0f / (float)WeightmapSize, 0.5f / (float)WeightmapSize, 0.5f / (float)WeightmapSize);
 	WeightmapSubsectionOffset = (float)(SubsectionSizeQuads + 1) / (float)WeightmapSize;
 
-	// Channel remapping
-	int32 ChannelOffsets[4] = { (int32)STRUCT_OFFSET(FColor, R), (int32)STRUCT_OFFSET(FColor, G), (int32)STRUCT_OFFSET(FColor, B), (int32)STRUCT_OFFSET(FColor, A) };
-
 	TArray<void*> WeightmapDataPtrs;
 	WeightmapDataPtrs.AddUninitialized(WeightmapTextures.Num());
 	for (int32 WeightmapIdx = 0; WeightmapIdx < WeightmapTextures.Num(); ++WeightmapIdx)
diff --git a/Engine/Source/Runtime/Landscape/Private/LandscapeEditInterface.cpp b/Engine/Source/Runtime/Landscape/Private/LandscapeEditInterface.cpp
--- a/Engine/Source/Runtime/Landscape/Private/LandscapeEditInterface.cpp
+++ b/Engine/Source/Runtime/Landscape/Private/LandscapeEditInterface.cpp
@@ -18,7 +18,7 @@ LandscapeEditInterface.cpp: Landscape editing interface
 #include "FixedSizeArrayView.h"
 
 // Channel remapping
-static const size_t ChannelOffsets[4] = {STRUCT_OFFSET(FColor,R), STRUCT_OFFSET(FColor,G), STRUCT_OFFSET(FColor,B), STRUCT_OFFSET(FColor,A)};
+extern const size_t ChannelOffsets[4] = {STRUCT_OFFSET(FColor,R), STRUCT_OFFSET(FColor,G), STRUCT_OFFSET(FColor,B), STRUCT_OFFSET(FColor,A)};
 
 //
 // FLandscapeEditDataInterface
@@ -417,7 +417,11 @@ void FLandscapeEditDataInterface::SetHeightData(int32 X1, int32 Y1, int32 X2, in
 			Component->GenerateHeightmapMips(MipData, ComponentX1, ComponentY1, ComponentX2, ComponentY2, TexDataInfo);
 
 			// Update collision
-			Component->UpdateCollisionHeightData(MipData[Component->CollisionMipLevel], ComponentX1, ComponentY1, ComponentX2, ComponentY2, bUpdateBoxSphereBounds, XYOffsetMipData);
+			Component->UpdateCollisionHeightData(
+				MipData[Component->CollisionMipLevel],
+				Component->SimpleCollisionMipLevel > Component->CollisionMipLevel ? MipData[Component->SimpleCollisionMipLevel] : nullptr,
+				ComponentX1, ComponentY1, ComponentX2, ComponentY2, bUpdateBoxSphereBounds,
+				XYOffsetMipData);
 
 			// Update GUID for Platform Data
 			FPlatformMisc::CreateGuid(Component->StateId);
@@ -1726,7 +1730,17 @@ void ULandscapeComponent::DeleteLayer(ULandscapeLayerInfoObject* LayerInfo, FLan
 	{
 		CollisionWeightmapMipData.Add((FColor*)LandscapeEdit->GetTextureDataInfo(Component->WeightmapTextures[WeightmapIdx])->GetMipData(Component->CollisionMipLevel));
 	}
-	Component->UpdateCollisionLayerData(CollisionWeightmapMipData);
+	TArray<FColor*> SimpleCollisionWeightmapMipData;
+	if (Component->SimpleCollisionMipLevel > Component->CollisionMipLevel)
+	{
+		for (int32 WeightmapIdx = 0; WeightmapIdx < Component->WeightmapTextures.Num(); WeightmapIdx++)
+		{
+			SimpleCollisionWeightmapMipData.Add((FColor*)LandscapeEdit->GetTextureDataInfo(Component->WeightmapTextures[WeightmapIdx])->GetMipData(Component->SimpleCollisionMipLevel));
+		}
+	}
+	Component->UpdateCollisionLayerData(
+		CollisionWeightmapMipData.GetData(),
+		Component->SimpleCollisionMipLevel > Component->CollisionMipLevel ? SimpleCollisionWeightmapMipData.GetData() : nullptr);
 }
 
 void FLandscapeEditDataInterface::DeleteLayer(ULandscapeLayerInfoObject* LayerInfo)
@@ -1934,9 +1948,19 @@ void FLandscapeEditDataInterface::ReplaceLayer(ULandscapeLayerInfoObject* FromLa
 		TArray<FColor*> CollisionWeightmapMipData;
 		for( int32 WeightmapIdx=0;WeightmapIdx < Component->WeightmapTextures.Num();WeightmapIdx++ )
 		{
-			CollisionWeightmapMipData.Add( (FColor*)GetTextureDataInfo(Component->WeightmapTextures[WeightmapIdx])->GetMipData(Component->CollisionMipLevel) );
+			CollisionWeightmapMipData.Add((FColor*)GetTextureDataInfo(Component->WeightmapTextures[WeightmapIdx])->GetMipData(Component->CollisionMipLevel));
+		}
+		TArray<FColor*> SimpleCollisionWeightmapMipData;
+		if (Component->SimpleCollisionMipLevel > Component->CollisionMipLevel)
+		{
+			for (int32 WeightmapIdx = 0; WeightmapIdx < Component->WeightmapTextures.Num(); WeightmapIdx++)
+			{
+				SimpleCollisionWeightmapMipData.Add((FColor*)GetTextureDataInfo(Component->WeightmapTextures[WeightmapIdx])->GetMipData(Component->SimpleCollisionMipLevel));
+			}
 		}
-		Component->UpdateCollisionLayerData(CollisionWeightmapMipData);
+		Component->UpdateCollisionLayerData(
+			CollisionWeightmapMipData.GetData(),
+			Component->SimpleCollisionMipLevel > Component->CollisionMipLevel ? SimpleCollisionWeightmapMipData.GetData() : nullptr);
 	}
 }
 
@@ -2386,6 +2410,7 @@ void FLandscapeEditDataInterface::SetAlphaData(ULandscapeLayerInfoObject* const
 	TArray<bool, TInlineAllocator<8>> LayerNoWeightBlends;  // NoWeightBlend flags
 	TArray<bool, TInlineAllocator<8>> LayerEditDataAllZero; // Whether the data we are editing for this layer is all zero
 	TArray<FColor*> CollisionWeightmapMipData;
+	TArray<FColor*> SimpleCollisionWeightmapMipData;
 	TArray<FColor*> WeightmapTextureMipData;
 
 	TMap<FIntPoint, TMap<const ULandscapeLayerInfoObject*, uint32>> LayerInfluenceCache;
@@ -2695,9 +2720,21 @@ void FLandscapeEditDataInterface::SetAlphaData(ULandscapeLayerInfoObject* const
 				WeightmapTextureMipData.Reset();
 			}
 
+			if (Component->SimpleCollisionMipLevel > Component->CollisionMipLevel)
+			{
+				for (int32 WeightmapIdx = 0; WeightmapIdx < Component->WeightmapTextures.Num(); WeightmapIdx++)
+				{
+					SimpleCollisionWeightmapMipData.Add((FColor*)TexDataInfos[WeightmapIdx]->GetMipData(Component->SimpleCollisionMipLevel));
+				}
+			}
+
 			// Update dominant layer info stored in collision component
-			Component->UpdateCollisionLayerData(CollisionWeightmapMipData, ComponentX1, ComponentY1, ComponentX2, ComponentY2);
+			Component->UpdateCollisionLayerData(
+				CollisionWeightmapMipData.GetData(),
+				Component->SimpleCollisionMipLevel > Component->CollisionMipLevel ? SimpleCollisionWeightmapMipData.GetData() : nullptr,
+				ComponentX1, ComponentY1, ComponentX2, ComponentY2);
 			CollisionWeightmapMipData.Reset();
+			SimpleCollisionWeightmapMipData.Reset();
 
 			// Check if we need to remove weightmap allocations for layers that were completely painted away
 			bool bRemovedLayer = false;
@@ -2771,6 +2808,7 @@ void FLandscapeEditDataInterface::SetAlphaData(const TSet<ULandscapeLayerInfoObj
 	TArray<FLayerDataInfo, TInlineAllocator<8>> LayerDataInfos;		// Pointers to all layers' data 
 	TArray<bool, TInlineAllocator<8>> LayerEditDataAllZero; // Whether the data we are editing for this layer is all zero 
 	TArray<FColor*> CollisionWeightmapMipData;
+	TArray<FColor*> SimpleCollisionWeightmapMipData;
 	TArray<FColor*> WeightmapTextureMipData;
 
 	for (int32 ComponentIndexY = ComponentIndexY1; ComponentIndexY <= ComponentIndexY2; ComponentIndexY++)
@@ -2959,9 +2997,21 @@ void FLandscapeEditDataInterface::SetAlphaData(const TSet<ULandscapeLayerInfoObj
 				WeightmapTextureMipData.Reset();
 			}
 
+			if (Component->SimpleCollisionMipLevel > Component->CollisionMipLevel)
+			{
+				for (int32 WeightmapIdx = 0; WeightmapIdx < Component->WeightmapTextures.Num(); WeightmapIdx++)
+				{
+					SimpleCollisionWeightmapMipData.Add((FColor*)TexDataInfos[WeightmapIdx]->GetMipData(Component->SimpleCollisionMipLevel));
+				}
+			}
+
 			// Update dominant layer info stored in collision component
-			Component->UpdateCollisionLayerData(CollisionWeightmapMipData, ComponentX1, ComponentY1, ComponentX2, ComponentY2);
+			Component->UpdateCollisionLayerData(
+				CollisionWeightmapMipData.GetData(),
+				Component->SimpleCollisionMipLevel > Component->CollisionMipLevel ? SimpleCollisionWeightmapMipData.GetData() : nullptr,
+				ComponentX1, ComponentY1, ComponentX2, ComponentY2);
 			CollisionWeightmapMipData.Reset();
+			SimpleCollisionWeightmapMipData.Reset();
 
 			// Check if we need to remove weightmap allocations for layers that were completely painted away
 			bool bRemovedLayer = false;
diff --git a/Engine/Source/Runtime/Landscape/Private/LandscapeGrass.cpp b/Engine/Source/Runtime/Landscape/Private/LandscapeGrass.cpp
--- a/Engine/Source/Runtime/Landscape/Private/LandscapeGrass.cpp
+++ b/Engine/Source/Runtime/Landscape/Private/LandscapeGrass.cpp
@@ -21,6 +21,7 @@
 #include "ContentStreaming.h"
 #include "LandscapeDataAccess.h"
 #include "LandscapeRender.h"
+#include "LandscapeVersion.h"
 
 #define LOCTEXT_NAMESPACE "Landscape"
 
@@ -101,6 +102,7 @@ DECLARE_CYCLE_STAT(TEXT("Grass Update"), STAT_GrassUpdate, STATGROUP_Foliage);
 // Grass weightmap rendering
 //
 
+#if WITH_EDITOR
 static bool ShouldCacheLandscapeGrassShaders(EShaderPlatform Platform, const FMaterial* Material, const FVertexFactoryType* VertexFactoryType)
 {
 	// We only need grass weight shaders for Landscape vertex factories on desktop platforms
@@ -199,7 +201,6 @@ public:
 
 IMPLEMENT_MATERIAL_SHADER_TYPE(, FLandscapeGrassWeightPS, TEXT("LandscapeGrassWeight"), TEXT("PSMain"), SF_Pixel);
 
-
 /**
 * Drawing policy used to write out landscape grass weightmap.
 */
@@ -282,10 +283,26 @@ private:
 // data also accessible by render thread
 class FLandscapeGrassWeightExporter_RenderThread
 {
-	FLandscapeGrassWeightExporter_RenderThread(int32 InNumPasses)
-	: RenderTargetResource(nullptr)
-	, NumPasses(InNumPasses)
-	{}
+	FLandscapeGrassWeightExporter_RenderThread(int32 InNumGrassMaps, int32 InCollisionMip, int32 InSimpleCollisionMip)
+		: RenderTargetResource(nullptr)
+		, NumPasses(FMath::DivideAndRoundUp(2 /* heightmap */ + InNumGrassMaps, 4))
+		, CollisionMipLevel(InCollisionMip)
+		, SimpleCollisionMipLevel(InSimpleCollisionMip)
+		, CollisionPassIndex(INDEX_NONE)
+		, SimpleCollisionPassIndex(INDEX_NONE)
+	{
+		const bool bRenderCollision = InCollisionMip > 0;
+		const bool bRenderSimpleCollision = InSimpleCollisionMip > InCollisionMip;
+
+		if (bRenderCollision)
+		{
+			CollisionPassIndex = NumPasses++;
+		}
+		if (bRenderSimpleCollision)
+		{
+			SimpleCollisionPassIndex = NumPasses++;
+		}
+	}
 
 	friend class FLandscapeGrassWeightExporter;
 
@@ -312,6 +329,10 @@ public:
 	TArray<FComponentInfo> ComponentInfos;
 	FIntPoint TargetSize;
 	int32 NumPasses;
+	int32 CollisionMipLevel;
+	int32 SimpleCollisionMipLevel;
+	int32 CollisionPassIndex;
+	int32 SimpleCollisionPassIndex;
 	float PassOffsetX;
 	FVector ViewOrigin;
 	FMatrix ViewRotationMatrix;
@@ -321,9 +342,7 @@ public:
 	{
 		FSceneViewFamilyContext ViewFamily(FSceneViewFamily::ConstructionValues(RenderTargetResource, NULL, FEngineShowFlags(ESFIM_Game)).SetWorldTimes(FApp::GetCurrentTime() - GStartTime, FApp::GetDeltaTime(), FApp::GetCurrentTime() - GStartTime));
 
-#if WITH_EDITOR
-		ViewFamily.LandscapeLODOverride = 0; 		// Force LOD render
-#endif
+		ViewFamily.LandscapeLODOverride = 0; // Force LOD render
 
 		FSceneViewInitOptions ViewInitOptions;
 		ViewInitOptions.SetViewRectangle(FIntRect(0, 0, TargetSize.X, TargetSize.Y));
@@ -349,11 +368,14 @@ public:
 			{
 				FLandscapeGrassWeightDrawingPolicy DrawingPolicy(Mesh.VertexFactory, Mesh.MaterialRenderProxy, *Mesh.MaterialRenderProxy->GetMaterial(GMaxRHIFeatureLevel));
 				RHICmdList.BuildAndSetLocalBoundShaderState(DrawingPolicy.GetBoundShaderStateInput(GMaxRHIFeatureLevel));
-				DrawingPolicy.SetSharedState(RHICmdList, View, FLandscapeGrassWeightDrawingPolicy::ContextDataType(), PassIdx, ComponentInfo.ViewOffset + FVector2D(PassOffsetX * PassIdx, 0));
+
+				const int32 ShaderPass = (PassIdx == CollisionPassIndex || PassIdx == SimpleCollisionPassIndex) ? 0 : PassIdx;
+				DrawingPolicy.SetSharedState(RHICmdList, View, FLandscapeGrassWeightDrawingPolicy::ContextDataType(), ShaderPass, ComponentInfo.ViewOffset + FVector2D(PassOffsetX * PassIdx, 0));
 
 				// The first batch element contains the grass batch for the entire component
-				DrawingPolicy.SetMeshRenderState(RHICmdList, *View, ComponentInfo.SceneProxy, Mesh, 0, false, FMeshDrawingRenderState(), FMeshDrawingPolicy::ElementDataType(), FLandscapeGrassWeightDrawingPolicy::ContextDataType());
-				DrawingPolicy.DrawMesh(RHICmdList, Mesh, 0);
+				const int32 ElementIndex = PassIdx == SimpleCollisionPassIndex ? 2 : (PassIdx == CollisionPassIndex ? 1 : 0);
+				DrawingPolicy.SetMeshRenderState(RHICmdList, *View, ComponentInfo.SceneProxy, Mesh, ElementIndex, false, FMeshDrawingRenderState(), FMeshDrawingPolicy::ElementDataType(), FLandscapeGrassWeightDrawingPolicy::ContextDataType());
+				DrawingPolicy.DrawMesh(RHICmdList, Mesh, ElementIndex);
 			}
 		}
 	}
@@ -362,23 +384,25 @@ public:
 class FLandscapeGrassWeightExporter : public FLandscapeGrassWeightExporter_RenderThread
 {
 	ALandscapeProxy* LandscapeProxy;
-	int32 ComponentSizeQuads;
 	int32 ComponentSizeVerts;
 	TArray<ULandscapeGrassType*> GrassTypes;
 	UTextureRenderTarget2D* RenderTargetTexture;
 
 public:
 
-	FLandscapeGrassWeightExporter(ALandscapeProxy* InLandscapeProxy, const TArray<ULandscapeComponent*>& InLandscapeComponents, const TArray<ULandscapeGrassType*> InGrassTypes) 
-	:	FLandscapeGrassWeightExporter_RenderThread((InGrassTypes.Num() + 5) / 4)
-	,	LandscapeProxy(InLandscapeProxy)
-	,	ComponentSizeQuads(InLandscapeProxy->ComponentSizeQuads)
-	,	ComponentSizeVerts(ComponentSizeQuads + 1)
-	,	GrassTypes(MoveTemp(InGrassTypes))
-	,	RenderTargetTexture(nullptr)
+	FLandscapeGrassWeightExporter(ALandscapeProxy* InLandscapeProxy, const TArray<ULandscapeComponent*>& InLandscapeComponents, const TArray<ULandscapeGrassType*> InGrassTypes)
+		: FLandscapeGrassWeightExporter_RenderThread(
+			InGrassTypes.Num(),
+			InLandscapeProxy->bBakeMaterialPositionOffsetIntoCollision ? InLandscapeProxy->CollisionMipLevel : INDEX_NONE,
+			InLandscapeProxy->bBakeMaterialPositionOffsetIntoCollision ? InLandscapeProxy->SimpleCollisionMipLevel : INDEX_NONE)
+		, LandscapeProxy(InLandscapeProxy)
+		, ComponentSizeVerts(InLandscapeProxy->ComponentSizeQuads + 1)
+		, GrassTypes(MoveTemp(InGrassTypes))
+		, RenderTargetTexture(nullptr)
 	{
 		check(InLandscapeComponents.Num());
 
+		// todo: use a 2d target?
 		TargetSize = FIntPoint(ComponentSizeVerts * NumPasses * InLandscapeComponents.Num(), ComponentSizeVerts);
 		FIntPoint TargetSizeMinusOne(TargetSize - FIntPoint(1, 1));
 		PassOffsetX = 2.0f * (float)ComponentSizeVerts / (float)TargetSize.X;
@@ -453,6 +477,11 @@ public:
 
 			NewGrassData->HeightData.Empty(FMath::Square(ComponentSizeVerts));
 
+			const int32 CollisionSizeVerts = CollisionPassIndex >= 0 ? ComponentSizeVerts >> CollisionMipLevel : 0;
+			const int32 SimpleCollisionSizeVerts = SimpleCollisionPassIndex >= 0 ? ComponentSizeVerts >> SimpleCollisionMipLevel : 0;
+			NewGrassData->CollisionHeightData.Empty(FMath::Square(CollisionSizeVerts));
+			NewGrassData->SimpleCollisionHeightData.Empty(FMath::Square(SimpleCollisionSizeVerts));
+
 			TArray<TArray<uint8>*> GrassWeightArrays;
 			GrassWeightArrays.Empty(GrassTypes.Num());
 			for (auto GrassType : GrassTypes)
@@ -469,13 +498,23 @@ public:
 				GrassWeightArrays.Add(DataArray);
 			}
 
+			// output debug bitmap
+#if UE_BUILD_DEBUG
+			static bool bOutputGrassBitmap = false;
+			if (bOutputGrassBitmap)
+			{
+				FString TempPath = FPaths::ScreenShotDir();
+				TempPath += TEXT("/GrassDebug");
+				IFileManager::Get().MakeDirectory(*TempPath, true);
+				FFileHelper::CreateBitmap(*(TempPath / "Grass"), TargetSize.X, TargetSize.Y, Samples.GetData(), nullptr, &IFileManager::Get(), nullptr, GrassTypes.Num() >= 2);
+			}
+#endif
+
 			for (int32 PassIdx = 0; PassIdx < NumPasses; PassIdx++)
 			{
 				FColor* SampleData = &Samples[ComponentInfo.PixelOffsetX + PassIdx*ComponentSizeVerts];
 				if (PassIdx == 0)
 				{
-					NewGrassData->HeightData.Empty(FMath::Square(ComponentSizeVerts));
-
 					for (int32 y = 0; y < ComponentSizeVerts; y++)
 					{
 						for (int32 x = 0; x < ComponentSizeVerts; x++)
@@ -483,14 +522,41 @@ public:
 							FColor& Sample = SampleData[x + y * TargetSize.X];
 							uint16 Height = (((uint16)Sample.R) << 8) + (uint16)(Sample.G);
 							NewGrassData->HeightData.Add(Height);
-							GrassWeightArrays[0]->Add(Sample.B);
-							if (GrassTypes.Num() > 1)
+							if (GrassTypes.Num() > 0)
 							{
-								GrassWeightArrays[1]->Add(Sample.A);
+								GrassWeightArrays[0]->Add(Sample.B);
+								if (GrassTypes.Num() > 1)
+								{
+									GrassWeightArrays[1]->Add(Sample.A);
+								}
 							}
 						}
 					}
 				}
+				else if (PassIdx == CollisionPassIndex)
+				{
+					for (int32 y = 0; y < CollisionSizeVerts; y++)
+					{
+						for (int32 x = 0; x < CollisionSizeVerts; x++)
+						{
+							FColor& Sample = SampleData[x + y * TargetSize.X];
+							uint16 Height = (((uint16)Sample.R) << 8) + (uint16)(Sample.G);
+							NewGrassData->CollisionHeightData.Add(Height);
+						}
+					}
+				}
+				else if (PassIdx == SimpleCollisionPassIndex)
+				{
+					for (int32 y = 0; y < SimpleCollisionSizeVerts; y++)
+					{
+						for (int32 x = 0; x < SimpleCollisionSizeVerts; x++)
+						{
+							FColor& Sample = SampleData[x + y * TargetSize.X];
+							uint16 Height = (((uint16)Sample.R) << 8) + (uint16)(Sample.G);
+							NewGrassData->SimpleCollisionHeightData.Add(Height);
+						}
+					}
+				}
 				else
 				{
 					for (int32 y = 0; y < ComponentSizeVerts; y++)
@@ -533,12 +599,10 @@ public:
 			// Assign the new data (thread-safe)
 			Component->GrassData = MakeShareable(NewGrassData);
 
-#if WITH_EDITOR
 			if (Proxy->bBakeMaterialPositionOffsetIntoCollision)
 			{
-				Component->UpdateCollisionHeightData(nullptr, 0, 0, MAX_int32, MAX_int32, true, nullptr, true);
+				Component->UpdateCollisionData(true);
 			}
-#endif
 		}
 	}
 
@@ -572,8 +636,6 @@ public:
 	}
 };
 
-#if WITH_EDITOR
-
 FLandscapeComponentGrassData::FLandscapeComponentGrassData(ULandscapeComponent* Component)
 	: MaterialStateId(Component->MaterialInstance->GetMaterial()->StateId)
 	, RotationForWPO(Component->MaterialInstance->GetMaterial()->WorldPositionOffset.IsConnected() ? Component->ComponentToWorld.GetRotation() : FQuat(0, 0, 0, 0))
@@ -605,13 +667,13 @@ bool ULandscapeComponent::CanRenderGrassMap() const
 	{
 		return false;
 	}
-		
+
 	// Check we can render the material
 	if (!MaterialInstance->GetMaterialResource(GetWorld()->FeatureLevel)->HasValidGameThreadShaderMap())
 	{
 		return false;
 	}
-	
+
 	return true;
 }
 
@@ -855,7 +917,26 @@ FArchive& operator<<(FArchive& Ar, FLandscapeComponentGrassData& Data)
 		Ar << Data.MaterialStateId;
 	}
 
+	Ar.UsingCustomVersion(FLandscapeCustomVersion::GUID);
+
+	if (Ar.CustomVer(FLandscapeCustomVersion::GUID) >= FLandscapeCustomVersion::GrassMaterialWPO)
+	{
+		Ar << Data.RotationForWPO;
+	}
+
 	Data.HeightData.BulkSerialize(Ar);
+
+#if WITH_EDITORONLY_DATA
+	if (Ar.CustomVer(FLandscapeCustomVersion::GUID) >= FLandscapeCustomVersion::CollisionMaterialWPO)
+	{
+		if (!Ar.IsFilterEditorOnly())
+		{
+			Data.CollisionHeightData.BulkSerialize(Ar);
+			Data.SimpleCollisionHeightData.BulkSerialize(Ar);
+		}
+	}
+#endif
+
 	// Each weight data array, being 1 byte will be serialized in bulk.
 	return Ar << Data.WeightData;
 }
@@ -1470,11 +1551,8 @@ void ALandscapeProxy::UpdateGrass(const TArray<FVector>& Cameras, bool bForceSyn
 
 					if (GrassTypes.Num() > 0 || bBakeMaterialPositionOffsetIntoCollision)
 					{
-						if (Component->IsGrassMapOutdated())
-						{
-							ComponentsNeedingGrassMapRender.Add(Component);
-						}
-						else if (!Component->GrassData->HasData())
+						if (Component->IsGrassMapOutdated() ||
+							!Component->GrassData->HasData())
 						{
 							ComponentsNeedingGrassMapRender.Add(Component);
 						}
diff --git a/Engine/Source/Runtime/Landscape/Private/LandscapeRender.cpp b/Engine/Source/Runtime/Landscape/Private/LandscapeRender.cpp
--- a/Engine/Source/Runtime/Landscape/Private/LandscapeRender.cpp
+++ b/Engine/Source/Runtime/Landscape/Private/LandscapeRender.cpp
@@ -562,6 +562,12 @@ FLandscapeComponentSceneProxy::FLandscapeComponentSceneProxy(ULandscapeComponent
 	, ComponentLightInfo(nullptr)
 	, LandscapeComponent(InComponent)
 	, LODFalloff(InComponent->GetLandscapeProxy()->LODFalloff)
+#if WITH_EDITOR || !(UE_BUILD_SHIPPING || UE_BUILD_TEST)
+	, bBakeMaterialPositionOffsetIntoCollision(InComponent->GetLandscapeProxy()->bBakeMaterialPositionOffsetIntoCollision)
+	, CollisionMipLevel(InComponent->CollisionMipLevel)
+	, SimpleCollisionMipLevel(InComponent->SimpleCollisionMipLevel)
+	, CollisionResponse(InComponent->GetLandscapeProxy()->BodyInstance.GetResponseToChannels())
+#endif
 {
 	if (!IsComponentLevelVisible())
 	{
@@ -650,7 +656,17 @@ FLandscapeComponentSceneProxy::FLandscapeComponentSceneProxy(ULandscapeComponent
 #endif
 
 	bRequiresAdjacencyInformation = RequiresAdjacencyInformation(MaterialInterface, XYOffsetmapTexture == nullptr ? &FLandscapeVertexFactory::StaticType : &FLandscapeXYOffsetVertexFactory::StaticType, InComponent->GetWorld()->FeatureLevel);
-	SharedBuffersKey = (SubsectionSizeQuads & 0xffff) | ((NumSubsections & 0xf) << 16) | (FeatureLevel <= ERHIFeatureLevel::ES3_1 ? 0 : 1 << 20) | (XYOffsetmapTexture == nullptr ? 0 : 1 << 31);
+
+	const int8 SubsectionSizeLog2 = FMath::CeilLogTwo(InComponent->SubsectionSizeQuads + 1);
+#if WITH_EDITOR
+	const int8 SharedBufferCollisionMip = bBakeMaterialPositionOffsetIntoCollision ? CollisionMipLevel : 0;
+	const int8 SharedBufferSimpleCollisionMip = bBakeMaterialPositionOffsetIntoCollision ? SimpleCollisionMipLevel : 0;
+#endif
+	SharedBuffersKey = (SubsectionSizeLog2 & 0xf) | ((NumSubsections & 0xf) << 4) |
+#if WITH_EDITOR
+	                   ((SharedBufferCollisionMip & 0xf) << 8) | ((SharedBufferSimpleCollisionMip & 0xf) << 12) |
+#endif
+	                   (FeatureLevel <= ERHIFeatureLevel::ES3_1 ? 0 : 1 << 30) | (XYOffsetmapTexture == nullptr ? 0 : 1 << 31);
 
 	bSupportsHeightfieldRepresentation = true;
 
@@ -681,7 +697,18 @@ void FLandscapeComponentSceneProxy::CreateRenderThreadResources()
 	SharedBuffers = FLandscapeComponentSceneProxy::SharedBuffersMap.FindRef(SharedBuffersKey);
 	if (SharedBuffers == nullptr)
 	{
-		SharedBuffers = new FLandscapeSharedBuffers(SharedBuffersKey, SubsectionSizeQuads, NumSubsections, FeatureLevel, bRequiresAdjacencyInformation);
+#if WITH_EDITOR
+		SharedBuffers = new FLandscapeSharedBuffers(
+			SharedBuffersKey, SubsectionSizeQuads, NumSubsections,
+			FeatureLevel, bRequiresAdjacencyInformation,
+			bBakeMaterialPositionOffsetIntoCollision ? CollisionMipLevel : INDEX_NONE,
+			bBakeMaterialPositionOffsetIntoCollision ? SimpleCollisionMipLevel : INDEX_NONE);
+#else
+		SharedBuffers = new FLandscapeSharedBuffers(
+			SharedBuffersKey, SubsectionSizeQuads, NumSubsections,
+			FeatureLevel, bRequiresAdjacencyInformation);
+#endif
+
 		FLandscapeComponentSceneProxy::SharedBuffersMap.Add(SharedBuffersKey, SharedBuffers);
 
 		if (!XYOffsetmapTexture)
@@ -738,10 +765,12 @@ void FLandscapeComponentSceneProxy::CreateRenderThreadResources()
 	// Assign LandscapeUniformShaderParameters
 	LandscapeUniformShaderParameters.InitResource();
 
+#if WITH_EDITOR
 	// Create MeshBatch for grass rendering
 	if(SharedBuffers->GrassIndexBuffer)
 	{
-		GrassMeshBatch.Elements.Empty(1);
+		GrassMeshBatch.Elements.Empty(3);
+		GrassMeshBatch.Elements.AddDefaulted(3);
 
 		FMaterialRenderProxy* RenderProxy = MaterialInterface->GetRenderProxy(false);
 		GrassMeshBatch.VertexFactory = VertexFactory;
@@ -753,26 +782,54 @@ void FLandscapeComponentSceneProxy::CreateRenderThreadResources()
 		GrassMeshBatch.DepthPriorityGroup = SDPG_World;
 
 		// Combined grass rendering batch element
-		FMeshBatchElement* BatchElement = new(GrassMeshBatch.Elements) FMeshBatchElement;
-		FLandscapeBatchElementParams* BatchElementParams = &GrassBatchParams;
+		FMeshBatchElement* GrassBatchElement = &GrassMeshBatch.Elements[0];
+		FLandscapeBatchElementParams* BatchElementParams = &GrassBatchParams[0];
 		BatchElementParams->LocalToWorldNoScalingPtr = &LocalToWorldNoScaling;
-		BatchElement->UserData = BatchElementParams;
-		if (NeedsUniformBufferUpdate())
-		{
-			UpdateUniformBuffer();
-		}
-		BatchElement->PrimitiveUniformBufferResource = &GetUniformBuffer();
 		BatchElementParams->LandscapeUniformShaderParametersResource = &LandscapeUniformShaderParameters;
 		BatchElementParams->SceneProxy = this;
 		BatchElementParams->SubX = -1;
 		BatchElementParams->SubY = -1;
 		BatchElementParams->CurrentLOD = 0;
-		BatchElement->IndexBuffer = SharedBuffers->GrassIndexBuffer;
-		BatchElement->NumPrimitives = FMath::Square(NumSubsections) * FMath::Square(SubsectionSizeVerts);
-		BatchElement->FirstIndex = 0;
-		BatchElement->MinVertexIndex = 0;
-		BatchElement->MaxVertexIndex = SharedBuffers->NumVertices - 1;
+		GrassBatchElement->UserData = BatchElementParams;
+		if (NeedsUniformBufferUpdate())
+		{
+			UpdateUniformBuffer();
+		}
+		GrassBatchElement->PrimitiveUniformBufferResource = &GetUniformBuffer();
+		GrassBatchElement->IndexBuffer = SharedBuffers->GrassIndexBuffer;
+		GrassBatchElement->NumPrimitives = FMath::Square(NumSubsections) * FMath::Square(SubsectionSizeVerts);
+		GrassBatchElement->FirstIndex = 0;
+		GrassBatchElement->MinVertexIndex = 0;
+		GrassBatchElement->MaxVertexIndex = SharedBuffers->NumVertices - 1;
+
+		const bool bRenderCollision = CollisionMipLevel > 0;
+		const bool bRenderSimpleCollision = SimpleCollisionMipLevel > CollisionMipLevel;
+
+		if (bRenderCollision)
+		{
+			FMeshBatchElement* CollisionBatchElement = &GrassMeshBatch.Elements[1];
+			*CollisionBatchElement = *GrassBatchElement;
+			FLandscapeBatchElementParams* CollisionBatchElementParams = &GrassBatchParams[1];
+			*CollisionBatchElementParams = *BatchElementParams;
+			CollisionBatchElementParams->CurrentLOD = CollisionMipLevel;
+			CollisionBatchElement->UserData = CollisionBatchElementParams;
+			CollisionBatchElement->NumPrimitives = FMath::Square(NumSubsections) * FMath::Square(SubsectionSizeVerts >> CollisionMipLevel);
+			CollisionBatchElement->FirstIndex = SharedBuffers->GrassCollisionOffset;
+		}
+
+		if (bRenderSimpleCollision)
+		{
+			FMeshBatchElement* SimpleCollisionBatchElement = &GrassMeshBatch.Elements[2];
+			*SimpleCollisionBatchElement = *GrassBatchElement;
+			FLandscapeBatchElementParams* SimpleCollisionBatchElementParams = &GrassBatchParams[2];
+			*SimpleCollisionBatchElementParams = *BatchElementParams;
+			SimpleCollisionBatchElementParams->CurrentLOD = SimpleCollisionMipLevel;
+			SimpleCollisionBatchElement->UserData = SimpleCollisionBatchElementParams;
+			SimpleCollisionBatchElement->NumPrimitives = FMath::Square(NumSubsections) * FMath::Square(SubsectionSizeVerts >> SimpleCollisionMipLevel);
+			SimpleCollisionBatchElement->FirstIndex = SharedBuffers->GrassSimpleCollisionOffset;
+		}
 	}
+#endif
 }
 
 void FLandscapeComponentSceneProxy::OnLevelAddedToWorld()
@@ -881,8 +938,15 @@ FPrimitiveViewRelevance FLandscapeComponentSceneProxy::GetViewRelevance(const FS
 	}
 #endif
 
+#if WITH_EDITOR || !(UE_BUILD_SHIPPING || UE_BUILD_TEST)
+	const bool bInCollisionView = View->Family->EngineShowFlags.CollisionVisibility || View->Family->EngineShowFlags.CollisionPawn;
+#endif
+
 	// Use the dynamic path for rendering landscape components pass only for Rich Views or if the static path is disabled for debug.
 	if (IsRichView(*View->Family) ||
+#if WITH_EDITOR || !(UE_BUILD_SHIPPING || UE_BUILD_TEST)
+		bInCollisionView ||
+#endif
 		GLandscapeDebugOptions.bDisableStatic ||
 		View->Family->EngineShowFlags.Wireframe ||
 #if WITH_EDITOR
@@ -1351,6 +1415,12 @@ void FLandscapeComponentSceneProxy::GetDynamicMeshElements(const TArray<const FS
 {
 	QUICK_SCOPE_CYCLE_COUNTER(STAT_FLandscapeComponentSceneProxy_GetMeshElements);
 
+#if WITH_EDITOR || !(UE_BUILD_SHIPPING || UE_BUILD_TEST)
+	const bool bInCollisionView = ViewFamily.EngineShowFlags.CollisionVisibility || ViewFamily.EngineShowFlags.CollisionPawn;
+	const bool bDrawSimpleCollision  = ViewFamily.EngineShowFlags.CollisionPawn       && CollisionResponse.GetResponse(ECC_Pawn) != ECR_Ignore;
+	const bool bDrawComplexCollision = ViewFamily.EngineShowFlags.CollisionVisibility && CollisionResponse.GetResponse(ECC_Visibility) != ECR_Ignore;
+#endif
+
 	int32 NumPasses = 0;
 	int32 NumTriangles = 0;
 	int32 NumDrawCalls = 0;
@@ -1366,7 +1436,7 @@ void FLandscapeComponentSceneProxy::GetDynamicMeshElements(const TArray<const FS
 
 			FLandscapeElementParamArray& ParameterArray = Collector.AllocateOneFrameResource<FLandscapeElementParamArray>();
 			ParameterArray.ElementParams.Empty(NumSubsections * NumSubsections);
-			ParameterArray.ElementParams.AddUninitialized(NumSubsections * NumSubsections);
+			ParameterArray.ElementParams.AddDefaulted(NumSubsections * NumSubsections);
 
 			FMeshBatch& Mesh = Collector.AllocateMesh();
 
@@ -1395,6 +1465,19 @@ void FLandscapeComponentSceneProxy::GetDynamicMeshElements(const TArray<const FS
 				{
 					int32 SubSectionIdx = SubX + SubY*NumSubsections;
 					int32 CurrentLOD = CalcLODForSubsection(*View, SubX, SubY, CameraLocalPos);
+#if WITH_EDITOR || !(UE_BUILD_SHIPPING || UE_BUILD_TEST)
+					if (bInCollisionView)
+					{
+						if (bDrawSimpleCollision)
+						{
+							CurrentLOD = FMath::Max(CollisionMipLevel, SimpleCollisionMipLevel);
+						}
+						else if (bDrawComplexCollision)
+						{
+							CurrentLOD = CollisionMipLevel;
+						}
+					}
+#endif
 
 			#if !(UE_BUILD_SHIPPING || UE_BUILD_TEST)
 					// We simplify this by considering only the biggest LOD index for this mesh element.
@@ -1443,9 +1526,9 @@ void FLandscapeComponentSceneProxy::GetDynamicMeshElements(const TArray<const FS
 #endif
 				}
 			}
+
 			// Render the landscape component
 #if WITH_EDITOR
-
 			switch (GLandscapeViewMode)
 			{
 			case ELandscapeViewMode::DebugLayer:
@@ -1496,7 +1579,7 @@ void FLandscapeComponentSceneProxy::GetDynamicMeshElements(const TArray<const FS
 
 			case ELandscapeViewMode::LayerUsage:
 				if (EditToolRenderData && GLandscapeLayerUsageMaterial)
-				{	
+				{
 					float Rotation = ((SectionBase.X / ComponentSizeQuads) ^ (SectionBase.Y / ComponentSizeQuads)) & 1 ? 0 : 2.f * PI;
 					auto LayerUsageMaterialInstance = new FLandscapeLayerUsageRenderProxy(GLandscapeLayerUsageMaterial->GetRenderProxy(false), ComponentSizeVerts, LayerColors, Rotation);
 					MeshTools.MaterialRenderProxy = LayerUsageMaterialInstance;
@@ -1584,6 +1667,32 @@ void FLandscapeComponentSceneProxy::GetDynamicMeshElements(const TArray<const FS
 #else
 			{
 #endif // WITH_EDITOR
+
+#if WITH_EDITOR || !(UE_BUILD_SHIPPING || UE_BUILD_TEST)
+				if (AllowDebugViewmodes() && bInCollisionView)
+				{
+					if (bDrawSimpleCollision || bDrawComplexCollision)
+					{
+						// Override the mesh's material with our material that draws the collision color
+						auto CollisionMaterialInstance = new FColoredMaterialRenderProxy(
+							GEngine->ShadedLevelColorationUnlitMaterial->GetRenderProxy(IsSelected(), IsHovered()),
+							WireframeColor
+							);
+						Collector.RegisterOneFrameMaterialProxy(CollisionMaterialInstance);
+
+						Mesh.MaterialRenderProxy = CollisionMaterialInstance;
+						Mesh.bCanApplyViewModeOverrides = true;
+						Mesh.bUseWireframeSelectionColoring = IsSelected();
+
+						Collector.AddMesh(ViewIndex, Mesh);
+
+						NumPasses++;
+						NumTriangles += Mesh.GetNumPrimitives();
+						NumDrawCalls += Mesh.Elements.Num();
+					}
+				}
+				else
+#endif
 				// Regular Landscape rendering. Only use the dynamic path if we're rendering a rich view or we've disabled the static path for debugging.
 				if( IsRichView(ViewFamily) || 
 					GLandscapeDebugOptions.bDisableStatic ||
@@ -1607,8 +1716,8 @@ void FLandscapeComponentSceneProxy::GetDynamicMeshElements(const TArray<const FS
 					NumDrawCalls += Mesh.Elements.Num();
 				}
 			}
-#if WITH_EDITOR
 
+#if WITH_EDITOR
 			// Extra render passes for landscape tools
 			if (GLandscapeEditModeActive)
 			{
@@ -1971,11 +2080,20 @@ void FLandscapeSharedBuffers::CreateIndexBuffers(ERHIFeatureLevel::Type InFeatur
 	}
 }
 
+#if WITH_EDITOR
 template <typename INDEX_TYPE>
 void FLandscapeSharedBuffers::CreateGrassIndexBuffer()
 {
 	TArray<INDEX_TYPE> NewIndices;
-	int32 ExpectedNumIndices = (FMath::Square(NumSubsections) * FMath::Square(SubsectionSizeVerts));
+
+	const bool bRenderCollision = GrassCollisionMip > 0;
+	const bool bRenderSimpleCollision = GrassSimpleCollisionMip > GrassCollisionMip;
+	const int32 CollisionSubsectionSizeVerts = bRenderCollision ? SubsectionSizeVerts >> GrassCollisionMip : 0;
+	const int32 SimpleCollisionSubsectionSizeVerts = bRenderSimpleCollision ? SubsectionSizeVerts >> GrassSimpleCollisionMip : 0;
+
+	int32 ExpectedNumIndices = FMath::Square(NumSubsections) * FMath::Square(SubsectionSizeVerts) +
+	                           FMath::Square(NumSubsections) * FMath::Square(CollisionSubsectionSizeVerts) +
+	                           FMath::Square(NumSubsections) * FMath::Square(SimpleCollisionSubsectionSizeVerts);
 	NewIndices.Empty(ExpectedNumIndices);
 
 	int32 SubOffset = 0;
@@ -1994,6 +2112,48 @@ void FLandscapeSharedBuffers::CreateGrassIndexBuffer()
 			SubOffset += FMath::Square(SubsectionSizeVerts);
 		}
 	}
+	if (bRenderCollision)
+	{
+		GrassCollisionOffset = NewIndices.Num();
+
+		SubOffset = 0;
+		for (int32 SubY = 0; SubY < NumSubsections; SubY++)
+		{
+			for (int32 SubX = 0; SubX < NumSubsections; SubX++)
+			{
+				for (int32 y = 0; y < CollisionSubsectionSizeVerts; y++)
+				{
+					for (int32 x = 0; x < CollisionSubsectionSizeVerts; x++)
+					{
+						NewIndices.Add(x + y * SubsectionSizeVerts + SubOffset);
+					}
+				}
+
+				SubOffset += FMath::Square(SubsectionSizeVerts);
+			}
+		}
+	}
+	if (bRenderSimpleCollision)
+	{
+		GrassSimpleCollisionOffset = NewIndices.Num();
+
+		SubOffset = 0;
+		for (int32 SubY = 0; SubY < NumSubsections; SubY++)
+		{
+			for (int32 SubX = 0; SubX < NumSubsections; SubX++)
+			{
+				for (int32 y = 0; y < SimpleCollisionSubsectionSizeVerts; y++)
+				{
+					for (int32 x = 0; x < SimpleCollisionSubsectionSizeVerts; x++)
+					{
+						NewIndices.Add(x + y * SubsectionSizeVerts + SubOffset);
+					}
+				}
+
+				SubOffset += FMath::Square(SubsectionSizeVerts);
+			}
+		}
+	}
 
 	check(NewIndices.Num() == ExpectedNumIndices);
 
@@ -2003,17 +2163,30 @@ void FLandscapeSharedBuffers::CreateGrassIndexBuffer()
 	IndexBuffer->InitResource();
 	GrassIndexBuffer = IndexBuffer;
 }
+#endif
 
-FLandscapeSharedBuffers::FLandscapeSharedBuffers(int32 InSharedBuffersKey, int32 InSubsectionSizeQuads, int32 InNumSubsections, ERHIFeatureLevel::Type InFeatureLevel, bool bRequiresAdjacencyInformation)
+#if WITH_EDITOR
+FLandscapeSharedBuffers::FLandscapeSharedBuffers(const int32 InSharedBuffersKey, const int32 InSubsectionSizeQuads, const int32 InNumSubsections, const ERHIFeatureLevel::Type InFeatureLevel, const bool bRequiresAdjacencyInformation, const int32 InCollisionMipLevel, const int32 InSimpleCollisionMipLevel)
+#else
+FLandscapeSharedBuffers::FLandscapeSharedBuffers(const int32 InSharedBuffersKey, const int32 InSubsectionSizeQuads, const int32 InNumSubsections, const ERHIFeatureLevel::Type InFeatureLevel, const bool bRequiresAdjacencyInformation)
+#endif
 	: SharedBuffersKey(InSharedBuffersKey)
 	, NumIndexBuffers(FMath::CeilLogTwo(InSubsectionSizeQuads + 1))
 	, SubsectionSizeVerts(InSubsectionSizeQuads + 1)
 	, NumSubsections(InNumSubsections)
+#if WITH_EDITOR
+	, GrassCollisionMip(InCollisionMipLevel)
+	, GrassSimpleCollisionMip(InSimpleCollisionMipLevel)
+#endif
 	, VertexFactory(nullptr)
 	, VertexBuffer(nullptr)
 	, AdjacencyIndexBuffers(nullptr)
 	, bUse32BitIndices(false)
+#if WITH_EDITOR
 	, GrassIndexBuffer(nullptr)
+	, GrassCollisionOffset(INDEX_NONE)
+	, GrassSimpleCollisionOffset(INDEX_NONE)
+#endif
 {
 	NumVertices = FMath::Square(SubsectionSizeVerts) * FMath::Square(NumSubsections);
 	if (InFeatureLevel > ERHIFeatureLevel::ES3_1)
@@ -2030,18 +2203,22 @@ FLandscapeSharedBuffers::FLandscapeSharedBuffers(int32 InSharedBuffersKey, int32
 	{
 		bUse32BitIndices = true;
 		CreateIndexBuffers<uint32>(InFeatureLevel, bRequiresAdjacencyInformation);
+#if WITH_EDITOR
 		if (InFeatureLevel > ERHIFeatureLevel::ES3_1)
 		{
 			CreateGrassIndexBuffer<uint32>();
 		}
+#endif
 	}
 	else
 	{
 		CreateIndexBuffers<uint16>(InFeatureLevel, bRequiresAdjacencyInformation);
+#if WITH_EDITOR
 		if (InFeatureLevel > ERHIFeatureLevel::ES3_1)
 		{
 			CreateGrassIndexBuffer<uint16>();
 		}
+#endif
 	}
 }
 
@@ -2057,11 +2234,13 @@ FLandscapeSharedBuffers::~FLandscapeSharedBuffers()
 	delete[] IndexBuffers;
 	delete[] IndexRanges;
 
+#if WITH_EDITOR
 	if (GrassIndexBuffer)
 	{
 		GrassIndexBuffer->ReleaseResource();
 		delete GrassIndexBuffer;
 	}
+#endif
 
 	if (AdjacencyIndexBuffers)
 	{
@@ -2338,14 +2517,14 @@ void FLandscapeVertexFactoryPixelShaderParameters::Bind(const FShaderParameterMa
 void FLandscapeVertexFactoryPixelShaderParameters::Serialize(FArchive& Ar)
 {
 	Ar << NormalmapTextureParameter
-		<< NormalmapTextureParameterSampler
-		<< LocalToWorldNoScalingParameter;
+	   << NormalmapTextureParameterSampler
+	   << LocalToWorldNoScalingParameter;
 }
 
 /**
 * Set any shader data specific to this vertex factory
 */
-void FLandscapeVertexFactoryPixelShaderParameters::SetMesh(FRHICommandList& RHICmdList, FShader* PixelShader, const class FVertexFactory* VertexFactory, const class FSceneView& View, const struct FMeshBatchElement& BatchElement, uint32 DataFlags) const
+void FLandscapeVertexFactoryPixelShaderParameters::SetMesh(FRHICommandList& RHICmdList, FShader* PixelShader, const FVertexFactory* VertexFactory, const FSceneView& View, const FMeshBatchElement& BatchElement, uint32 DataFlags) const
 {
 	SCOPE_CYCLE_COUNTER(STAT_LandscapeVFDrawTime);
 
diff --git a/Engine/Source/Runtime/Landscape/Private/LandscapeRenderMobile.cpp b/Engine/Source/Runtime/Landscape/Private/LandscapeRenderMobile.cpp
--- a/Engine/Source/Runtime/Landscape/Private/LandscapeRenderMobile.cpp
+++ b/Engine/Source/Runtime/Landscape/Private/LandscapeRenderMobile.cpp
@@ -245,9 +245,18 @@ void FLandscapeComponentSceneProxyMobile::CreateRenderThreadResources()
 {
 	// Use only Index buffers
 	SharedBuffers = FLandscapeComponentSceneProxy::SharedBuffersMap.FindRef(SharedBuffersKey);
-	if( SharedBuffers == NULL )
+	if (SharedBuffers == nullptr)
 	{
-		SharedBuffers = new FLandscapeSharedBuffers(SharedBuffersKey, SubsectionSizeQuads, NumSubsections, GetScene().GetFeatureLevel(), false);
+#if WITH_EDITOR
+		SharedBuffers = new FLandscapeSharedBuffers(
+			SharedBuffersKey, SubsectionSizeQuads, NumSubsections,
+			GetScene().GetFeatureLevel(), false,
+			CollisionMipLevel, SimpleCollisionMipLevel);
+#else
+		SharedBuffers = new FLandscapeSharedBuffers(
+			SharedBuffersKey, SubsectionSizeQuads, NumSubsections,
+			GetScene().GetFeatureLevel(), false);
+#endif
 		FLandscapeComponentSceneProxy::SharedBuffersMap.Add(SharedBuffersKey, SharedBuffers);
 	}
 
diff --git a/Engine/Source/Runtime/Landscape/Private/LandscapeVersion.h b/Engine/Source/Runtime/Landscape/Private/LandscapeVersion.h
--- a/Engine/Source/Runtime/Landscape/Private/LandscapeVersion.h
+++ b/Engine/Source/Runtime/Landscape/Private/LandscapeVersion.h
@@ -11,6 +11,10 @@ namespace FLandscapeCustomVersion
 		BeforeCustomVersionWasAdded = 0,
 		// Changed to TMap properties instead of manual serialization and fixed landscape spline control point cross-level mesh components not being serialized
 		NewSplineCrossLevelMeshSerialization,
+		// Support material world-position-offset in the heightmap used for grass placement
+		GrassMaterialWPO,
+		// Support material world-position-offset in landscape simple collision
+		CollisionMaterialWPO,
 
 		// -----<new versions can be added above this line>-------------------------------------------------
 		VersionPlusOne,
diff --git a/Engine/Source/Runtime/Landscape/Public/LandscapeRender.h b/Engine/Source/Runtime/Landscape/Public/LandscapeRender.h
--- a/Engine/Source/Runtime/Landscape/Public/LandscapeRender.h
+++ b/Engine/Source/Runtime/Landscape/Public/LandscapeRender.h
@@ -98,7 +98,7 @@ struct FLandscapeBatchElementParams
 	const FLandscapeComponentSceneProxy* SceneProxy;
 	int32 SubX;
 	int32 SubY;
-	int32	CurrentLOD;
+	int32 CurrentLOD;
 };
 
 class FLandscapeElementParamArray : public FOneFrameResource
@@ -290,6 +290,10 @@ public:
 	int32 NumIndexBuffers;
 	int32 SubsectionSizeVerts;
 	int32 NumSubsections;
+#if WITH_EDITOR
+	int32 GrassCollisionMip;
+	int32 GrassSimpleCollisionMip;
+#endif
 
 	FLandscapeVertexFactory* VertexFactory;
 	FLandscapeVertexBuffer* VertexBuffer;
@@ -297,14 +301,25 @@ public:
 	FLandscapeIndexRanges* IndexRanges;
 	FLandscapeSharedAdjacencyIndexBuffer* AdjacencyIndexBuffers;
 	bool bUse32BitIndices;
+#if WITH_EDITOR
 	FIndexBuffer* GrassIndexBuffer;
+	int32 GrassCollisionOffset;
+	int32 GrassSimpleCollisionOffset;
+#endif
 
-	FLandscapeSharedBuffers(int32 SharedBuffersKey, int32 SubsectionSizeQuads, int32 NumSubsections, ERHIFeatureLevel::Type InFeatureLevel, bool bRequiresAdjacencyInformation);
+#if WITH_EDITOR
+	FLandscapeSharedBuffers(int32 SharedBuffersKey, int32 SubsectionSizeQuads, int32 NumSubsections, ERHIFeatureLevel::Type FeatureLevel, bool bRequiresAdjacencyInformation, int32 CollisionMipLevel, int32 SimpleCollisionMipLevel);
+#else
+	FLandscapeSharedBuffers(int32 SharedBuffersKey, int32 SubsectionSizeQuads, int32 NumSubsections, ERHIFeatureLevel::Type FeatureLevel, bool bRequiresAdjacencyInformation);
+#endif
 
 	template <typename INDEX_TYPE>
 	void CreateIndexBuffers(ERHIFeatureLevel::Type InFeatureLevel, bool bRequiresAdjacencyInformation);
+
+#if WITH_EDITOR
 	template <typename INDEX_TYPE>
 	void CreateGrassIndexBuffer();
+#endif
 
 	virtual ~FLandscapeSharedBuffers();
 };
@@ -476,21 +491,21 @@ protected:
 	int32						LastLOD;	// Last LOD we have batch elements for
 
 	/** 
-	 * Number of subsections within the component in each dimension, this can be 1 or 2.  
-	 * Subsections exist to improve the speed at which LOD transitions can take place over distance. 
+	 * Number of subsections within the component in each dimension, this can be 1 or 2.
+	 * Subsections exist to improve the speed at which LOD transitions can take place over distance.
 	 */
 	int32						NumSubsections;
 	/** Number of unique heights in the subsection. */
 	int32						SubsectionSizeQuads;
-	/** Number of heightmap heights in the subsection.  This includes the duplicate row at the end. */
+	/** Number of heightmap heights in the subsection. This includes the duplicate row at the end. */
 	int32						SubsectionSizeVerts;
 	/** Size of the component in unique heights. */
-	int32						ComponentSizeQuads;	
+	int32						ComponentSizeQuads;
 	/** 
-	 * ComponentSizeQuads + 1.  
+	 * ComponentSizeQuads + 1.
 	 * Note: in the case of multiple subsections, this is not very useful, as there will be an internal duplicate row of heights in addition to the row at the end.
 	 */
-	int32						ComponentSizeVerts; 
+	int32						ComponentSizeVerts;
 	uint8						StaticLightingLOD;
 	float						StaticLightingResolution;
 	/** Address of the component within the parent Landscape in unique height texels. */
@@ -500,9 +515,14 @@ protected:
 	// Storage for static draw list batch params
 	TArray<FLandscapeBatchElementParams> StaticBatchParamArray;
 
-	// Precomputed grass rendering MeshBatch
+#if WITH_EDITOR
+	// Precomputed grass rendering MeshBatch and params
+	// 0 - regular grass rendering at LOD 0
+	// 1 - rendering for collision baking
+	// 2 - rendering for "simple collision" baking
 	FMeshBatch					GrassMeshBatch;
-	FLandscapeBatchElementParams GrassBatchParams;
+	FLandscapeBatchElementParams GrassBatchParams[3];
+#endif
 
 	// Precomputed values
 	float					LODDistance;
@@ -544,6 +564,15 @@ protected:
 
 	ELandscapeLODFalloff::Type LODFalloff;
 
+	// data used in editor or visualisers
+#if WITH_EDITOR || !(UE_BUILD_SHIPPING || UE_BUILD_TEST)
+	bool bBakeMaterialPositionOffsetIntoCollision;
+	int32 CollisionMipLevel;
+	int32 SimpleCollisionMipLevel;
+
+	FCollisionResponseContainer CollisionResponse;
+#endif
+
 	TUniformBuffer<FLandscapeUniformShaderParameters> LandscapeUniformShaderParameters;
 
 	// Cached versions of these
@@ -579,8 +608,9 @@ public:
 	int32 CalcLODForSubsection(const FSceneView& View, int32 SubX, int32 SubY, const FVector2D& CameraLocalPos) const;
 	void CalcLODParamsForSubsection(const FSceneView& View, const FVector2D& CameraLocalPos, int32 SubX, int32 SubY, int32 BatchLOD, float& OutfLOD, FVector4& OutNeighborLODs) const;
 	uint64 GetStaticBatchElementVisibility(const FSceneView& View, const FMeshBatch* Batch) const;
+#if WITH_EDITOR
 	const FMeshBatch& GetGrassMeshBatch() const { return GrassMeshBatch; }
-
+#endif
 
 	// FLandcapeSceneProxy
 	void ChangeLODDistanceFactor_RenderThread(float InLODDistanceFactor);
diff --git a/Engine/Source/Runtime/OpenGLDrv/Private/Android/AndroidOpenGL.cpp b/Engine/Source/Runtime/OpenGLDrv/Private/Android/AndroidOpenGL.cpp
--- a/Engine/Source/Runtime/OpenGLDrv/Private/Android/AndroidOpenGL.cpp
+++ b/Engine/Source/Runtime/OpenGLDrv/Private/Android/AndroidOpenGL.cpp
@@ -317,12 +317,15 @@ void PlatformReleaseRenderQuery( GLuint Query, uint64 QueryContext )
 bool FAndroidOpenGL::bUseHalfFloatTexStorage = false;
 bool FAndroidOpenGL::bUseES30ShadingLanguage = false;
 bool FAndroidOpenGL::bES30Support = false;
+bool FAndroidOpenGL::bSupportsInstancing = false;
 
 void FAndroidOpenGL::ProcessExtensions(const FString& ExtensionsString)
 {
 	FOpenGLES2::ProcessExtensions(ExtensionsString);
 
-	bES30Support = FString(ANSI_TO_TCHAR((const ANSICHAR*)glGetString(GL_VERSION))).Contains(TEXT("OpenGL ES 3."));
+	FString VersionString = FString(ANSI_TO_TCHAR((const ANSICHAR*)glGetString(GL_VERSION)));
+	
+	bES30Support = VersionString.Contains(TEXT("OpenGL ES 3."));
 
 	// Get procedures
 	if (bSupportsOcclusionQueries || bSupportsDisjointTimeQueries)
@@ -389,6 +392,8 @@ void FAndroidOpenGL::ProcessExtensions(const FString& ExtensionsString)
 		glDrawElementsInstanced = (PFNGLDRAWELEMENTSINSTANCEDPROC)((void*)eglGetProcAddress("glDrawElementsInstanced"));
 		glDrawArraysInstanced = (PFNGLDRAWARRAYSINSTANCEDPROC)((void*)eglGetProcAddress("glDrawArraysInstanced"));
 		glVertexAttribDivisor = (PFNGLVERTEXATTRIBDIVISORPROC)((void*)eglGetProcAddress("glVertexAttribDivisor"));
+		
+		bSupportsInstancing = true;
 	}
 
 	if (bES30Support || bIsAdrenoBased)
@@ -425,6 +430,13 @@ void FAndroidOpenGL::ProcessExtensions(const FString& ExtensionsString)
 	{
 		bSupportsShaderTextureCubeLod = false;
 	}
+
+	// Nexus 5 (Android 4.4.2) doesn't like glVertexAttribDivisor(index, 0) called when not using a glDrawElementsInstanced
+	if (bIsAdrenoBased && VersionString.Contains(TEXT("OpenGL ES 3.0 [REDACTED_EMAIL] AU@  (CL@)")))
+	{
+		UE_LOG(LogRHI, Warning, TEXT("Disabling support for hardware instancing on Adreno 330 OpenGL ES 3.0 [REDACTED_EMAIL] AU@  (CL@)"));
+		bSupportsInstancing = false;
+	}
 }
 
 
diff --git a/Engine/Source/Runtime/OpenGLDrv/Private/Android/AndroidOpenGL.h b/Engine/Source/Runtime/OpenGLDrv/Private/Android/AndroidOpenGL.h
--- a/Engine/Source/Runtime/OpenGLDrv/Private/Android/AndroidOpenGL.h
+++ b/Engine/Source/Runtime/OpenGLDrv/Private/Android/AndroidOpenGL.h
@@ -243,7 +243,7 @@ struct FAndroidOpenGL : public FOpenGLES2
 
 	static FORCEINLINE void VertexAttribDivisor(GLuint Index, GLuint Divisor)
 	{
-		if (bES30Support)
+		if (SupportsInstancing())
 		{
 			glVertexAttribDivisor(Index, Divisor);
 		}
@@ -256,7 +256,7 @@ struct FAndroidOpenGL : public FOpenGLES2
 	// 32 bpp HDR encoding mode via 'intrinsic_GetHDR32bppEncodeModeES2()'.
 	static FORCEINLINE bool SupportsHDR32bppEncodeModeIntrinsic()		{ return true; }
 
-	static FORCEINLINE bool SupportsInstancing()						{ return bES30Support; }
+	static FORCEINLINE bool SupportsInstancing()						{ return bSupportsInstancing; }
 
 	static FORCEINLINE bool UseES30ShadingLanguage()
 	{
@@ -273,6 +273,10 @@ struct FAndroidOpenGL : public FOpenGLES2
 	
 	// whether device supports ES 3.0
 	static bool bES30Support;
+
+	// whether device supports hardware instancing
+	static bool bSupportsInstancing;
+
 };
 
 typedef FAndroidOpenGL FOpenGL;
