Copying //UE4/Dev-Landscape to Dev-Main (//UE4/Dev-Main) @ 2892944

==========================
MAJOR FEATURES + CHANGES
==========================

Change 2886096 on 2016/02/29 by [REDACTED_PERSON]

        Lightmass landscape world-position-offset
        (attempt 2)

Change 2886121 on 2016/02/29 by [REDACTED_PERSON]

        Fix missing header and variable shadowing warnings

Change 2887809 on 2016/03/01 by [REDACTED_PERSON]

        Better input filtering for "Landscape Layer Blend" material node
        - it won't now allow "texture object" etc to be connected and the height inputs won't accept material attributes (only the layer inputs)

Change 2892336 on 2016/03/03 by [REDACTED_PERSON]

        Fixed the landscape grass weight data not being saved in editor

diff --git a/Engine/Source/Runtime/Core/Public/Containers/Algo/Accumulate.h b/Engine/Source/Runtime/Core/Public/Containers/Algo/Accumulate.h
new file mode 100644
--- /dev/null
+++ b/Engine/Source/Runtime/Core/Public/Containers/Algo/Accumulate.h
@@ -0,0 +1,91 @@
+// Copyright 1998-2016 Epic Games, Inc. All Rights Reserved.
+
+#pragma once
+
+// TPlus<T> specifically takes const T& and returns T.
+// TPlus<> (empty angle brackets) is late-binding, taking whatever is passed and returning the correct result type for (A+B)
+template<typename T = void>
+struct TPlus
+{
+	FORCEINLINE T operator()(const T& A, const T& B) { return A + B; }
+};
+template<>
+struct TPlus<void>
+{
+	template<typename U, typename V>
+	FORCEINLINE auto operator()(U&& A, V&& B) -> decltype(A + B) { return A + B; }
+};
+
+
+namespace Algo
+{
+	/**
+	* Sums a range by successively applying Op.
+	*
+	* @param  Input  Any iterable type
+	* @param  Init  Initial value for the summation
+	* @param  Op  Summing Operation (the default is TPlus<>)
+	*
+	* @return the result of summing all the elements of Input
+	*/
+	template <typename T, typename A, typename OpT>
+	FORCEINLINE T Accumulate(const A& Input, T init, OpT Op)
+	{
+		T result = init;
+		for (auto&& i : Input)
+		{
+			result = Op(result, i);
+		}
+		return result;
+	}
+
+	/**
+	 * Sums a range.
+	 *
+	 * @param  Input  Any iterable type
+	 * @param  Init  Initial value for the summation
+	 *
+	 * @return the result of summing all the elements of Input
+	 */
+	template <typename T, typename A>
+	FORCEINLINE T Accumulate(const A& Input, T init)
+	{
+		return Accumulate(Input, init, TPlus<>());
+	}
+
+	/**
+	* Sums a range by applying MapOp to each element, and then summing the results.
+	*
+	* @param  Input  Any iterable type
+	* @param  Init  Initial value for the summation
+	* @param  MapOp  Mapping Operation
+	* @param  Op  Summing Operation (the default is TPlus<>)
+	*
+	* @return the result of mapping and then summing all the elements of Input
+	*/
+	template <typename T, typename A, typename MapT, typename OpT>
+	FORCEINLINE T TransformAccumulate(const A& Input, MapT MapOp, T init, OpT Op)
+	{
+		T result = init;
+		for (auto&& i : Input)
+		{
+			result = Op(result, MapOp(i));
+		}
+		return result;
+	}
+
+	/**
+	* Sums a range by applying MapOp to each element, and then summing the results.
+	*
+	* @param  Input  Any iterable type
+	* @param  Init  Initial value for the summation
+	* @param  MapOp  Mapping Operation
+	*
+	* @return the result of mapping and then summing all the elements of Input
+	*/
+	template <typename T, typename A, typename MapT>
+	FORCEINLINE T TransformAccumulate(const A& Input, MapT MapOp, T init)
+	{
+		return TransformAccumulate(Input, MapOp, init, TPlus<>());
+	}
+}
diff --git a/Engine/Source/Runtime/Landscape/Classes/LandscapeComponent.h b/Engine/Source/Runtime/Landscape/Classes/LandscapeComponent.h
--- a/Engine/Source/Runtime/Landscape/Classes/LandscapeComponent.h
+++ b/Engine/Source/Runtime/Landscape/Classes/LandscapeComponent.h
@@ -124,8 +124,8 @@ struct FLandscapeComponentGrassData
 
 	TArray<uint16> HeightData;
 #if WITH_EDITORONLY_DATA
-	TArray<uint16> CollisionHeightData;
-	TArray<uint16> SimpleCollisionHeightData;
+	// Height data for LODs 1+, keyed on LOD index
+	TMap<int32, TArray<uint16>> HeightMipData;
 #endif
 	TMap<ULandscapeGrassType*, TArray<uint8>> WeightData;
 
@@ -135,25 +135,15 @@ struct FLandscapeComponentGrassData
 
 	bool HasData()
 	{
-		return HeightData.Num() > 0;
-	}
-
-	SIZE_T GetAllocatedSize() const
-	{
-		SIZE_T WeightSize = 0; 
-		for (auto It = WeightData.CreateConstIterator(); It; ++It)
-		{
-			WeightSize += It.Value().GetAllocatedSize();
-		}
-		return sizeof(*this)
-			+ HeightData.GetAllocatedSize()
+		return HeightData.Num() > 0 ||
 #if WITH_EDITORONLY_DATA
-			+ CollisionHeightData.GetAllocatedSize()
-			+ SimpleCollisionHeightData.GetAllocatedSize()
+			HeightMipData.Num() > 0 ||
 #endif
-			+ WeightData.GetAllocatedSize() + WeightSize;
+			WeightData.Num() > 0;
 	}
 
+	SIZE_T GetAllocatedSize() const;
+
 	friend FArchive& operator<<(FArchive& Ar, FLandscapeComponentGrassData& Data);
 };
 
@@ -396,6 +386,9 @@ public:
 	/* Is the grassmap data outdated, eg by a material */
 	bool IsGrassMapOutdated() const;
 
+	/** Renders the heightmap of this component (including material world-position-offset) at the specified LOD */
+	TArray<uint16> RenderWPOHeightmap(int32 LOD);
+
 	/* Serialize all hashes/guids that record the current state of this component */
 	void SerializeStateHashes(FArchive& Ar);
 
diff --git a/Engine/Source/Runtime/Landscape/Classes/LandscapeProxy.h b/Engine/Source/Runtime/Landscape/Classes/LandscapeProxy.h
--- a/Engine/Source/Runtime/Landscape/Classes/LandscapeProxy.h
+++ b/Engine/Source/Runtime/Landscape/Classes/LandscapeProxy.h
@@ -421,6 +421,12 @@ public:
 	/** Whether this primitive should cast shadows in the far shadow cascades. */
 	UPROPERTY(EditAnywhere, AdvancedDisplay, Category=Lighting, meta=(DisplayName = "Far Shadow"))
 	uint32 bCastFarShadow:1;
+	
+	/** Whether to use the landscape material's vertical world position offset when calculating static lighting.
+		Note: Only z (vertical) offset is supported. XY offsets are ignored.
+		Does not work correctly with an XY offset map (mesh collision) */
+	UPROPERTY(EditAnywhere, AdvancedDisplay, Category=Lighting)
+	uint32 bUseMaterialPositionOffsetInStaticLighting:1;
 
 	UPROPERTY()
 	uint32 bIsProxy:1;
@@ -455,7 +461,7 @@ public:
 		Note: Only z (vertical) offset is supported. XY offsets are ignored.
 		Does not work with an XY offset map (mesh collision) */
 	UPROPERTY(EditAnywhere, AdvancedDisplay, Category=Collision)
-	bool bBakeMaterialPositionOffsetIntoCollision;
+	uint32 bBakeMaterialPositionOffsetIntoCollision:1;
 
 #if WITH_EDITORONLY_DATA
 	UPROPERTY()
diff --git a/Engine/Source/Runtime/Landscape/Classes/Materials/MaterialExpressionLandscapeLayerBlend.h b/Engine/Source/Runtime/Landscape/Classes/Materials/MaterialExpressionLandscapeLayerBlend.h
--- a/Engine/Source/Runtime/Landscape/Classes/Materials/MaterialExpressionLandscapeLayerBlend.h
+++ b/Engine/Source/Runtime/Landscape/Classes/Materials/MaterialExpressionLandscapeLayerBlend.h
@@ -72,7 +72,7 @@ class UMaterialExpressionLandscapeLayerBlend : public UMaterialExpression
 	//~ End UObject Interface
 
 	//~ Begin UMaterialExpression Interface
-	virtual uint32 GetInputType(int32 InputIndex) override { return MCT_Unknown; }
+	virtual uint32 GetInputType(int32 InputIndex) override;
 	virtual uint32 GetOutputType(int32 InputIndex) override { return MCT_Unknown; }
 	virtual bool IsResultMaterialAttributes(int32 OutputIndex) override;
 	virtual int32 Compile(class FMaterialCompiler* Compiler, int32 OutputIndex, int32 MultiplexIndex) override;
diff --git a/Engine/Source/Runtime/Landscape/Private/LandscapeEdit.cpp b/Engine/Source/Runtime/Landscape/Private/LandscapeEdit.cpp
--- a/Engine/Source/Runtime/Landscape/Private/LandscapeEdit.cpp
+++ b/Engine/Source/Runtime/Landscape/Private/LandscapeEdit.cpp
@@ -697,7 +697,21 @@ void ULandscapeComponent::UpdateCollisionHeightData(const FColor* const Heightma
 	// Handle Material WPO baked into heightfield collision
 	// Material WPO is not currently supported for mesh collision components
 	const bool bUsingGrassMapHeights = Proxy->bBakeMaterialPositionOffsetIntoCollision && !MeshCollisionComponent && GrassData->HasData() && !IsGrassMapOutdated();
-	const uint16* const GrassHeights = bUsingGrassMapHeights ? (GrassData->CollisionHeightData.Num() > 0 ? GrassData->CollisionHeightData.GetData() : GrassData->HeightData.GetData()) : nullptr;
+	const uint16* GrassHeights = nullptr;
+	if (bUsingGrassMapHeights)
+	{
+		if (CollisionMipLevel == 0)
+		{
+			GrassHeights = GrassData->HeightData.GetData();
+		}
+		else
+		{
+			if (GrassData->HeightMipData.Contains(CollisionMipLevel))
+			{
+				GrassHeights = GrassData->HeightMipData[CollisionMipLevel].GetData();
+			}
+		}
+	}
 
 	const int32 SubSectionX1 = FMath::Max(0, FMath::DivideAndRoundDown(ComponentX1 - 1, SubsectionSizeQuads));
 	const int32 SubSectionY1 = FMath::Max(0, FMath::DivideAndRoundDown(ComponentY1 - 1, SubsectionSizeQuads));
@@ -784,8 +798,7 @@ void ULandscapeComponent::UpdateCollisionHeightData(const FColor* const Heightma
 		const int32 SimpleHeightmapOffsetX = FMath::RoundToInt(HeightmapScaleBias.Z * (float)HeightmapSizeU) >> SimpleCollisionMipLevel;
 		const int32 SimpleHeightmapOffsetY = FMath::RoundToInt(HeightmapScaleBias.W * (float)HeightmapSizeV) >> SimpleCollisionMipLevel;
 
-		const uint16* const SimpleCollisionGrassHeights = bUsingGrassMapHeights ? GrassData->SimpleCollisionHeightData.GetData() : nullptr;
-
+		const uint16* SimpleCollisionGrassHeights = bUsingGrassMapHeights && GrassData->HeightMipData.Contains(SimpleCollisionMipLevel) ? GrassData->HeightMipData[SimpleCollisionMipLevel].GetData() : nullptr;
 		uint16* const SimpleCollisionHeightData = CollisionHeightData + FMath::Square(CollisionSizeVerts);
 
 		for (int32 SubsectionY = SubSectionY1; SubsectionY < SubSectionY2; ++SubsectionY)
@@ -3481,6 +3494,10 @@ void ALandscapeProxy::PostEditChangeProperty(FPropertyChangedEvent& PropertyChan
 		}
 		RecreateCollisionComponents();
 	}
+	else if (PropertyName == GET_MEMBER_NAME_CHECKED(ALandscapeProxy, bUseMaterialPositionOffsetInStaticLighting))
+	{
+		InvalidateLightingCache();
+	}
 	else if(PropertyName == FName(TEXT("bCastStaticShadow")) ||
 		PropertyName == FName(TEXT("bCastShadowAsTwoSided")) ||
 		PropertyName == FName(TEXT("bCastFarShadow")))
diff --git a/Engine/Source/Runtime/Landscape/Private/LandscapeGrass.cpp b/Engine/Source/Runtime/Landscape/Private/LandscapeGrass.cpp
--- a/Engine/Source/Runtime/Landscape/Private/LandscapeGrass.cpp
+++ b/Engine/Source/Runtime/Landscape/Private/LandscapeGrass.cpp
@@ -22,6 +22,7 @@
 #include "LandscapeDataAccess.h"
 #include "LandscapeRender.h"
 #include "LandscapeVersion.h"
+#include "Algo/Accumulate.h"
 
 #define LOCTEXT_NAMESPACE "Landscape"
 
@@ -283,24 +284,20 @@ private:
 // data also accessible by render thread
 class FLandscapeGrassWeightExporter_RenderThread
 {
-	FLandscapeGrassWeightExporter_RenderThread(int32 InNumGrassMaps, int32 InCollisionMip, int32 InSimpleCollisionMip)
+	FLandscapeGrassWeightExporter_RenderThread(int32 InNumGrassMaps, bool InbNeedsHeightmap, TArray<int32> InHeightMips)
 		: RenderTargetResource(nullptr)
-		, NumPasses(FMath::DivideAndRoundUp(2 /* heightmap */ + InNumGrassMaps, 4))
-		, CollisionMipLevel(InCollisionMip)
-		, SimpleCollisionMipLevel(InSimpleCollisionMip)
-		, CollisionPassIndex(INDEX_NONE)
-		, SimpleCollisionPassIndex(INDEX_NONE)
+		, NumPasses(0)
+		, HeightMips(MoveTemp(InHeightMips))
+		, FirstHeightMipsPassIndex(MAX_int32)
 	{
-		const bool bRenderCollision = InCollisionMip > 0;
-		const bool bRenderSimpleCollision = InSimpleCollisionMip > InCollisionMip;
-
-		if (bRenderCollision)
+		if (InbNeedsHeightmap || InNumGrassMaps > 0)
 		{
-			CollisionPassIndex = NumPasses++;
+			NumPasses += FMath::DivideAndRoundUp(2 /* heightmap */ + InNumGrassMaps, 4);
 		}
-		if (bRenderSimpleCollision)
+		if (HeightMips.Num() > 0)
 		{
-			SimpleCollisionPassIndex = NumPasses++;
+			FirstHeightMipsPassIndex = NumPasses;
+			NumPasses += HeightMips.Num();
 		}
 	}
 
@@ -326,13 +323,11 @@ public:
 	};
 
 	FTextureRenderTarget2DResource* RenderTargetResource;
-	TArray<FComponentInfo> ComponentInfos;
+	TArray<FComponentInfo, TInlineAllocator<1>> ComponentInfos;
 	FIntPoint TargetSize;
 	int32 NumPasses;
-	int32 CollisionMipLevel;
-	int32 SimpleCollisionMipLevel;
-	int32 CollisionPassIndex;
-	int32 SimpleCollisionPassIndex;
+	TArray<int32> HeightMips;
+	int32 FirstHeightMipsPassIndex;
 	float PassOffsetX;
 	FVector ViewOrigin;
 	FMatrix ViewRotationMatrix;
@@ -369,11 +364,11 @@ public:
 				FLandscapeGrassWeightDrawingPolicy DrawingPolicy(Mesh.VertexFactory, Mesh.MaterialRenderProxy, *Mesh.MaterialRenderProxy->GetMaterial(GMaxRHIFeatureLevel));
 				RHICmdList.BuildAndSetLocalBoundShaderState(DrawingPolicy.GetBoundShaderStateInput(GMaxRHIFeatureLevel));
 
-				const int32 ShaderPass = (PassIdx == CollisionPassIndex || PassIdx == SimpleCollisionPassIndex) ? 0 : PassIdx;
+				const int32 ShaderPass = (PassIdx >= FirstHeightMipsPassIndex) ? 0 : PassIdx;
 				DrawingPolicy.SetSharedState(RHICmdList, View, FLandscapeGrassWeightDrawingPolicy::ContextDataType(), ShaderPass, ComponentInfo.ViewOffset + FVector2D(PassOffsetX * PassIdx, 0));
 
 				// The first batch element contains the grass batch for the entire component
-				const int32 ElementIndex = PassIdx == SimpleCollisionPassIndex ? 2 : (PassIdx == CollisionPassIndex ? 1 : 0);
+				const int32 ElementIndex = (PassIdx >= FirstHeightMipsPassIndex) ? HeightMips[PassIdx - FirstHeightMipsPassIndex] : 0;
 				DrawingPolicy.SetMeshRenderState(RHICmdList, *View, ComponentInfo.SceneProxy, Mesh, ElementIndex, false, FMeshDrawingRenderState(), FMeshDrawingPolicy::ElementDataType(), FLandscapeGrassWeightDrawingPolicy::ContextDataType());
 				DrawingPolicy.DrawMesh(RHICmdList, Mesh, ElementIndex);
 			}
@@ -385,22 +380,25 @@ class FLandscapeGrassWeightExporter : public FLandscapeGrassWeightExporter_Rende
 {
 	ALandscapeProxy* LandscapeProxy;
 	int32 ComponentSizeVerts;
+	int32 SubsectionSizeQuads;
+	int32 NumSubsections;
 	TArray<ULandscapeGrassType*> GrassTypes;
 	UTextureRenderTarget2D* RenderTargetTexture;
 
 public:
-
-	FLandscapeGrassWeightExporter(ALandscapeProxy* InLandscapeProxy, const TArray<ULandscapeComponent*>& InLandscapeComponents, const TArray<ULandscapeGrassType*> InGrassTypes)
+	FLandscapeGrassWeightExporter(ALandscapeProxy* InLandscapeProxy, const TArray<ULandscapeComponent*>& InLandscapeComponents, TArray<ULandscapeGrassType*> InGrassTypes, bool InbNeedsHeightmap = true, TArray<int32> InHeightMips = {})
 		: FLandscapeGrassWeightExporter_RenderThread(
 			InGrassTypes.Num(),
-			InLandscapeProxy->bBakeMaterialPositionOffsetIntoCollision ? InLandscapeProxy->CollisionMipLevel : INDEX_NONE,
-			InLandscapeProxy->bBakeMaterialPositionOffsetIntoCollision ? InLandscapeProxy->SimpleCollisionMipLevel : INDEX_NONE)
+			InbNeedsHeightmap,
+			MoveTemp(InHeightMips))
 		, LandscapeProxy(InLandscapeProxy)
 		, ComponentSizeVerts(InLandscapeProxy->ComponentSizeQuads + 1)
+		, SubsectionSizeQuads(InLandscapeProxy->SubsectionSizeQuads)
+		, NumSubsections(InLandscapeProxy->NumSubsections)
 		, GrassTypes(MoveTemp(InGrassTypes))
 		, RenderTargetTexture(nullptr)
 	{
-		check(InLandscapeComponents.Num());
+		check(InLandscapeComponents.Num() > 0);
 
 		// todo: use a 2d target?
 		TargetSize = FIntPoint(ComponentSizeVerts * NumPasses * InLandscapeComponents.Num(), ComponentSizeVerts);
@@ -418,14 +416,14 @@ public:
 			ViewOffset.X += PixelOffsetX;
 			ViewOffset /= (FVector2D(TargetSize) * 0.5f);
 
-			new(ComponentInfos) FComponentInfo(Component, ViewOffset, PixelOffsetX);
+			ComponentInfos.Add(FComponentInfo(Component, ViewOffset, PixelOffsetX));
 		}
 
 		// center of target area in world
 		FVector TargetCenter = LandscapeProxy->GetTransform().TransformPosition(FVector(TargetSizeMinusOne, 0.f)*0.5f);
 
 		// extent of target in world space
-		FVector TargetExtent = FVector(TargetSize, 0.f)*LandscapeProxy->GetActorScale()*0.5f;
+		FVector TargetExtent = FVector(TargetSize, 0.0f)*LandscapeProxy->GetActorScale()*0.5f;
 
 		ViewOrigin = TargetCenter;
 		ViewRotationMatrix = FInverseRotationMatrix(LandscapeProxy->GetActorRotation());
@@ -444,7 +442,7 @@ public:
 		RenderTargetTexture = NewObject<UTextureRenderTarget2D>();
 		check(RenderTargetTexture);
 		RenderTargetTexture->ClearColor = FLinearColor::White;
-		RenderTargetTexture->TargetGamma = 1.f;
+		RenderTargetTexture->TargetGamma = 1.0f;
 		RenderTargetTexture->InitCustomFormat(TargetSize.X, TargetSize.Y, PF_B8G8R8A8, false);
 		RenderTargetResource = RenderTargetTexture->GameThread_GetRenderTargetResource()->GetTextureRenderTarget2DResource();
 
@@ -458,7 +456,8 @@ public:
 			});
 	}
 
-	void ApplyResults()
+	TMap<ULandscapeComponent*, TUniquePtr<FLandscapeComponentGrassData>, TInlineSetAllocator<1>>
+		FetchResults()
 	{
 		TArray<FColor> Samples;
 		Samples.SetNumUninitialized(TargetSize.X*TargetSize.Y);
@@ -468,19 +467,24 @@ public:
 		ReadSurfaceDataFlags.SetLinearToGamma(false);
 		RenderTargetResource->ReadPixels(Samples, ReadSurfaceDataFlags, FIntRect(0, 0, TargetSize.X, TargetSize.Y));
 
+		TMap<ULandscapeComponent*, TUniquePtr<FLandscapeComponentGrassData>, TInlineSetAllocator<1>> Results;
+		Results.Reserve(ComponentInfos.Num());
 		for (auto& ComponentInfo : ComponentInfos)
 		{
 			ULandscapeComponent* Component = ComponentInfo.Component;
 			ALandscapeProxy* Proxy = Component->GetLandscapeProxy();
 
-			FLandscapeComponentGrassData* NewGrassData = new FLandscapeComponentGrassData(Component);
-
-			NewGrassData->HeightData.Empty(FMath::Square(ComponentSizeVerts));
+			TUniquePtr<FLandscapeComponentGrassData> NewGrassData = MakeUnique<FLandscapeComponentGrassData>(Component);
 
-			const int32 CollisionSizeVerts = CollisionPassIndex >= 0 ? ComponentSizeVerts >> CollisionMipLevel : 0;
-			const int32 SimpleCollisionSizeVerts = SimpleCollisionPassIndex >= 0 ? ComponentSizeVerts >> SimpleCollisionMipLevel : 0;
-			NewGrassData->CollisionHeightData.Empty(FMath::Square(CollisionSizeVerts));
-			NewGrassData->SimpleCollisionHeightData.Empty(FMath::Square(SimpleCollisionSizeVerts));
+			if (FirstHeightMipsPassIndex > 0)
+			{
+				NewGrassData->HeightData.Empty(FMath::Square(ComponentSizeVerts));
+			}
+			else
+			{
+				NewGrassData->HeightData.Empty(0);
+			}
+			NewGrassData->HeightMipData.Empty(HeightMips.Num());
 
 			TArray<TArray<uint8>*> GrassWeightArrays;
 			GrassWeightArrays.Empty(GrassTypes.Num());
@@ -513,75 +517,69 @@ public:
 			for (int32 PassIdx = 0; PassIdx < NumPasses; PassIdx++)
 			{
 				FColor* SampleData = &Samples[ComponentInfo.PixelOffsetX + PassIdx*ComponentSizeVerts];
-				if (PassIdx == 0)
+				if (PassIdx < FirstHeightMipsPassIndex)
 				{
-					for (int32 y = 0; y < ComponentSizeVerts; y++)
+					if (PassIdx == 0)
 					{
-						for (int32 x = 0; x < ComponentSizeVerts; x++)
+						for (int32 y = 0; y < ComponentSizeVerts; y++)
 						{
-							FColor& Sample = SampleData[x + y * TargetSize.X];
-							uint16 Height = (((uint16)Sample.R) << 8) + (uint16)(Sample.G);
-							NewGrassData->HeightData.Add(Height);
-							if (GrassTypes.Num() > 0)
+							for (int32 x = 0; x < ComponentSizeVerts; x++)
 							{
-								GrassWeightArrays[0]->Add(Sample.B);
-								if (GrassTypes.Num() > 1)
+								FColor& Sample = SampleData[x + y * TargetSize.X];
+								uint16 Height = (((uint16)Sample.R) << 8) + (uint16)(Sample.G);
+								NewGrassData->HeightData.Add(Height);
+								if (GrassTypes.Num() > 0)
 								{
-									GrassWeightArrays[1]->Add(Sample.A);
+									GrassWeightArrays[0]->Add(Sample.B);
+									if (GrassTypes.Num() > 1)
+									{
+										GrassWeightArrays[1]->Add(Sample.A);
+									}
 								}
 							}
 						}
 					}
-				}
-				else if (PassIdx == CollisionPassIndex)
-				{
-					for (int32 y = 0; y < CollisionSizeVerts; y++)
-					{
-						for (int32 x = 0; x < CollisionSizeVerts; x++)
-						{
-							FColor& Sample = SampleData[x + y * TargetSize.X];
-							uint16 Height = (((uint16)Sample.R) << 8) + (uint16)(Sample.G);
-							NewGrassData->CollisionHeightData.Add(Height);
-						}
-					}
-				}
-				else if (PassIdx == SimpleCollisionPassIndex)
-				{
-					for (int32 y = 0; y < SimpleCollisionSizeVerts; y++)
-					{
-						for (int32 x = 0; x < SimpleCollisionSizeVerts; x++)
-						{
-							FColor& Sample = SampleData[x + y * TargetSize.X];
-							uint16 Height = (((uint16)Sample.R) << 8) + (uint16)(Sample.G);
-							NewGrassData->SimpleCollisionHeightData.Add(Height);
-						}
-					}
-				}
-				else
-				{
-					for (int32 y = 0; y < ComponentSizeVerts; y++)
+					else
 					{
-						for (int32 x = 0; x < ComponentSizeVerts; x++)
+						for (int32 y = 0; y < ComponentSizeVerts; y++)
 						{
-							FColor& Sample = SampleData[x + y * TargetSize.X];
-
-							int32 TypeIdx = PassIdx * 4 - 2;
-							GrassWeightArrays[TypeIdx++]->Add(Sample.R);
-							if (TypeIdx < GrassTypes.Num())
+							for (int32 x = 0; x < ComponentSizeVerts; x++)
 							{
-								GrassWeightArrays[TypeIdx++]->Add(Sample.G);
+								FColor& Sample = SampleData[x + y * TargetSize.X];
+
+								int32 TypeIdx = PassIdx * 4 - 2;
+								GrassWeightArrays[TypeIdx++]->Add(Sample.R);
 								if (TypeIdx < GrassTypes.Num())
 								{
-									GrassWeightArrays[TypeIdx++]->Add(Sample.B);
+									GrassWeightArrays[TypeIdx++]->Add(Sample.G);
 									if (TypeIdx < GrassTypes.Num())
 									{
-										GrassWeightArrays[TypeIdx++]->Add(Sample.A);
+										GrassWeightArrays[TypeIdx++]->Add(Sample.B);
+										if (TypeIdx < GrassTypes.Num())
+										{
+											GrassWeightArrays[TypeIdx++]->Add(Sample.A);
+										}
 									}
 								}
 							}
 						}
 					}
 				}
+				else // PassIdx >= FirstHeightMipsPassIndex
+				{
+					const int32 Mip = HeightMips[PassIdx - FirstHeightMipsPassIndex];
+					int32 MipSizeVerts = NumSubsections * (SubsectionSizeQuads >> Mip);
+					TArray<uint16>& MipHeightData = NewGrassData->HeightMipData.Add(Mip);
+					for (int32 y = 0; y < MipSizeVerts; y++)
+					{
+						for (int32 x = 0; x < MipSizeVerts; x++)
+						{
+							FColor& Sample = SampleData[x + y * TargetSize.X];
+							uint16 Height = (((uint16)Sample.R) << 8) + (uint16)(Sample.G);
+							MipHeightData.Add(Height);
+						}
+					}
+				}
 			}
 
 			// remove null grass type if we had one (can occur if the node has null entries)
@@ -590,14 +588,30 @@ public:
 			// Remove any grass data that is entirely weight 0
 			for (auto Iter(NewGrassData->WeightData.CreateIterator()); Iter; ++Iter)
 			{
-				if (Iter.Value().IndexOfByPredicate([&](const int8& Weight){ return Weight != 0; }) == INDEX_NONE)
+				if (Iter->Value.IndexOfByPredicate([&](const int8& Weight) { return Weight != 0; }) == INDEX_NONE)
 				{
 					Iter.RemoveCurrent();
 				}
 			}
 
+			Results.Add(Component, MoveTemp(NewGrassData));
+		}
+
+		return Results;
+	}
+
+	void ApplyResults()
+	{
+		TMap<ULandscapeComponent*, TUniquePtr<FLandscapeComponentGrassData>, TInlineSetAllocator<1>> NewGrassData = FetchResults();
+
+		for (auto&& GrassDataPair : NewGrassData)
+		{
+			ULandscapeComponent* Component = GrassDataPair.Key;
+			FLandscapeComponentGrassData* ComponentGrassData = GrassDataPair.Value.Release();
+			ALandscapeProxy* Proxy = Component->GetLandscapeProxy();
+
 			// Assign the new data (thread-safe)
-			Component->GrassData = MakeShareable(NewGrassData);
+			Component->GrassData = MakeShareable(ComponentGrassData);
 
 			if (Proxy->bBakeMaterialPositionOffsetIntoCollision)
 			{
@@ -723,14 +737,61 @@ void ULandscapeComponent::RenderGrassMap()
 				GrassTypes.Add(GrassTypeInput.GrassType);
 			}
 
+			TArray<int32> HeightMips;
+			if (CollisionMipLevel > 0)
+			{
+				HeightMips.Add(CollisionMipLevel);
+			}
+			if (SimpleCollisionMipLevel > CollisionMipLevel)
+			{
+				HeightMips.Add(SimpleCollisionMipLevel);
+			}
+
 			TArray<ULandscapeComponent*> LandscapeComponents;
 			LandscapeComponents.Add(this);
 
-			FLandscapeGrassWeightExporter Exporter(GetLandscapeProxy(), LandscapeComponents, GrassTypes);
+			FLandscapeGrassWeightExporter Exporter(GetLandscapeProxy(), MoveTemp(LandscapeComponents), MoveTemp(GrassTypes), true, MoveTemp(HeightMips));
 			Exporter.ApplyResults();
 		}
+	}
+}
 
+TArray<uint16> ULandscapeComponent::RenderWPOHeightmap(int32 LOD)
+{
+	TArray<uint16> Results;
+
+	if (!CanRenderGrassMap())
+	{
+		MaterialInstance->GetMaterialResource(GetWorld()->FeatureLevel)->FinishCompilation();
+	}
+
+	TArray<const UMaterialExpressionLandscapeGrassOutput*> GrassExpressions;
+	MaterialInstance->GetMaterial()->GetAllExpressionsOfType<UMaterialExpressionLandscapeGrassOutput>(GrassExpressions);
+	if (GrassExpressions.Num() > 0)
+	{
+		TArray<ULandscapeGrassType*> GrassTypes;
+		TArray<ULandscapeComponent*> LandscapeComponents;
+		LandscapeComponents.Add(this);
+
+		if (LOD == 0)
+		{
+			FLandscapeGrassWeightExporter Exporter(GetLandscapeProxy(), MoveTemp(LandscapeComponents), MoveTemp(GrassTypes), true, {});
+			TMap<ULandscapeComponent*, TUniquePtr<FLandscapeComponentGrassData>, TInlineSetAllocator<1>> TempGrassData;
+			TempGrassData = Exporter.FetchResults();
+			Results = MoveTemp(TempGrassData[this]->HeightData);
+		}
+		else
+		{
+			TArray<int32> HeightMips;
+			HeightMips.Add(LOD);
+			FLandscapeGrassWeightExporter Exporter(GetLandscapeProxy(), MoveTemp(LandscapeComponents), MoveTemp(GrassTypes), false, MoveTemp(HeightMips));
+			TMap<ULandscapeComponent*, TUniquePtr<FLandscapeComponentGrassData>, TInlineSetAllocator<1>> TempGrassData;
+			TempGrassData = Exporter.FetchResults();
+			Results = MoveTemp(TempGrassData[this]->HeightMipData[LOD]);
+		}
 	}
+
+	return Results;
 }
 
 void ULandscapeComponent::RemoveGrassMap()
@@ -740,7 +801,17 @@ void ULandscapeComponent::RemoveGrassMap()
 
 void ALandscapeProxy::RenderGrassMaps(const TArray<ULandscapeComponent*>& InLandscapeComponents, const TArray<ULandscapeGrassType*>& GrassTypes)
 {
-	FLandscapeGrassWeightExporter Exporter(this, InLandscapeComponents, GrassTypes);
+	TArray<int32> HeightMips;
+	if (CollisionMipLevel > 0)
+	{
+		HeightMips.Add(CollisionMipLevel);
+	}
+	if (SimpleCollisionMipLevel > CollisionMipLevel)
+	{
+		HeightMips.Add(SimpleCollisionMipLevel);
+	}
+
+	FLandscapeGrassWeightExporter Exporter(this, InLandscapeComponents, GrassTypes, true, MoveTemp(HeightMips));
 	Exporter.ApplyResults();
 }
 
@@ -917,6 +988,22 @@ void ULandscapeGrassType::PostEditChangeProperty(FPropertyChangedEvent& Property
 //
 // FLandscapeComponentGrassData
 //
+SIZE_T FLandscapeComponentGrassData::GetAllocatedSize() const
+{
+	SIZE_T WeightSize = 0; 
+	for (auto It = WeightData.CreateConstIterator(); It; ++It)
+	{
+		WeightSize += It.Value().GetAllocatedSize();
+	}
+	return sizeof(*this)
+		+ HeightData.GetAllocatedSize()
+#if WITH_EDITORONLY_DATA
+		+ HeightMipData.GetAllocatedSize()
+		+ Algo::TransformAccumulate(HeightMipData, [](const TPair<int32, TArray<uint16>>& HeightMipDataPair) { return HeightMipDataPair.Value.GetAllocatedSize(); }, 0)
+#endif
+		+ WeightData.GetAllocatedSize() + WeightSize;
+}
+
 FArchive& operator<<(FArchive& Ar, FLandscapeComponentGrassData& Data)
 {
 	if (Ar.UE4Ver() >= VER_UE4_SERIALIZE_LANDSCAPE_GRASS_DATA_MATERIAL_GUID)
@@ -938,14 +1025,43 @@ FArchive& operator<<(FArchive& Ar, FLandscapeComponentGrassData& Data)
 	{
 		if (!Ar.IsFilterEditorOnly())
 		{
-			Data.CollisionHeightData.BulkSerialize(Ar);
-			Data.SimpleCollisionHeightData.BulkSerialize(Ar);
+			if (Ar.CustomVer(FLandscapeCustomVersion::GUID) >= FLandscapeCustomVersion::LightmassMaterialWPO)
+			{
+				// todo - BulkSerialize each mip?
+				Ar << Data.HeightMipData;
+			}
+			else
+			{
+				checkSlow(Ar.IsLoading());
+
+				TArray<uint16> CollisionHeightData;
+				CollisionHeightData.BulkSerialize(Ar);
+				if (CollisionHeightData.Num())
+				{
+					const int32 ComponentSizeQuads = FMath::Sqrt(Data.HeightData.Num()) - 1;
+					const int32 CollisionSizeQuads = FMath::Sqrt(CollisionHeightData.Num()) - 1;
+					const int32 CollisionMip = FMath::FloorLog2(ComponentSizeQuads / CollisionSizeQuads);
+					Data.HeightMipData.Add(CollisionMip, MoveTemp(CollisionHeightData));
+				}
+
+				TArray<uint16> SimpleCollisionHeightData;
+				SimpleCollisionHeightData.BulkSerialize(Ar);
+				if (SimpleCollisionHeightData.Num())
+				{
+					const int32 ComponentSizeQuads = FMath::Sqrt(Data.HeightData.Num()) - 1;
+					const int32 SimpleCollisionSizeQuads = FMath::Sqrt(SimpleCollisionHeightData.Num()) - 1;
+					const int32 SimpleCollisionMip = FMath::FloorLog2(ComponentSizeQuads / SimpleCollisionSizeQuads);
+					Data.HeightMipData.Add(SimpleCollisionMip, MoveTemp(SimpleCollisionHeightData));
+				}
+			}
 		}
 	}
 #endif
 
 	// Each weight data array, being 1 byte will be serialized in bulk.
-	return Ar << Data.WeightData;
+	Ar << Data.WeightData;
+
+	return Ar;
 }
 
 //
diff --git a/Engine/Source/Runtime/Landscape/Private/LandscapeLight.cpp b/Engine/Source/Runtime/Landscape/Private/LandscapeLight.cpp
--- a/Engine/Source/Runtime/Landscape/Private/LandscapeLight.cpp
+++ b/Engine/Source/Runtime/Landscape/Private/LandscapeLight.cpp
@@ -407,6 +407,9 @@ void FLandscapeStaticLightingMesh::GetHeightmapData(int32 InLOD, int32 GeometryL
 	ULandscapeInfo* const Info = LandscapeComponent->GetLandscapeInfo();
 	check(Info);
 
+	bool bUseRenderedWPO = LandscapeComponent->GetLandscapeProxy()->bUseMaterialPositionOffsetInStaticLighting &&
+	                       LandscapeComponent->MaterialInstance->GetMaterial()->WorldPositionOffset.IsConnected();
+
 	HeightData.Empty(FMath::Square(NumVertices));
 	HeightData.AddUninitialized(FMath::Square(NumVertices));
 
@@ -429,6 +432,13 @@ void FLandscapeStaticLightingMesh::GetHeightmapData(int32 InLOD, int32 GeometryL
 		FLandscapeComponentDataInterface DataInterface(LandscapeComponent, InLOD);
 		::InternalUpscaling(DataInterface, LandscapeComponent, InLOD, GeometryLOD, CompHeightData, CompXYOffsetData);
 
+		TArray<uint16> RenderedWPOData;
+		if (bUseRenderedWPO)
+		{
+			// todo - do we need to invalidate the landscape mesh version?
+			RenderedWPOData = LandscapeComponent->RenderWPOHeightmap(InLOD);
+		}
+
 		for (int32 Y = 0; Y < ComponentSizeQuads + 1; Y++)
 		{
 			const FColor* const Data = DataInterface.GetHeightData(0, Y);
@@ -436,38 +446,49 @@ void FLandscapeStaticLightingMesh::GetHeightmapData(int32 InLOD, int32 GeometryL
 			for (int32 SubsectionX = 0; SubsectionX < NumSubsections; SubsectionX++)
 			{
 				const int32 X = SubsectionSizeQuads * SubsectionX;
-				const int32 CompX = X + FMath::Min(X/SubsectionSizeQuads, NumSubsections - 1);
-				const FColor* const SubsectionData = &Data[CompX];
+				const int32 TexX = X + FMath::Min(X/SubsectionSizeQuads, NumSubsections - 1);
+				const FColor* const SubsectionData = &Data[TexX];
 
 				// Copy the data
-				FMemory::Memcpy( &HeightData[X + ExpandQuadsX + (Y + ExpandQuadsY) * NumVertices], SubsectionData, SubsectionSizeVerts * sizeof(FColor));
+				FMemory::Memcpy(&HeightData[X + ExpandQuadsX + (Y + ExpandQuadsY) * NumVertices], SubsectionData, SubsectionSizeVerts * sizeof(FColor));
+			}
+
+			if (bUseRenderedWPO)
+			{
+				const int32 HeightDataOffset = ExpandQuadsX + (Y + ExpandQuadsY) * NumVertices;
+				const int32 WPODataOffset    = Y * (ComponentSizeQuads + 1);
+				for (int32 X = 0; X < ComponentSizeQuads; ++X)
+				{
+					const uint16 WPOHeight = RenderedWPOData[X + WPODataOffset];
+					FColor& Height = HeightData[X + HeightDataOffset];
+					Height.R = (WPOHeight >> 8);
+					Height.G = (WPOHeight & 0xFF);
+				}
 			}
 		}
 	}
 
 	// copy surrounding heightmaps...
-	for (int32 ComponentY = 0; ComponentY < 3; ComponentY++)
+	for (int32 ComponentY = -1; ComponentY <= 1; ++ComponentY)
 	{
-		for (int32 ComponentX = 0; ComponentX < 3; ComponentX++)
+		for (int32 ComponentX = -1; ComponentX <= 1; ++ComponentX)
 		{
-			if (ComponentX == 1 && ComponentY == 1)
+			if (ComponentX == 0 && ComponentY == 0)
 			{
 				// Ourself
 				continue;
 			}
 
-			const int32 XSource = (ComponentX == 0) ? (ComponentSizeQuads - ExpandQuadsX) : ((ComponentX == 1) ? 0 : 1);
-			const int32 YSource = (ComponentY == 0) ? (ComponentSizeQuads - ExpandQuadsY) : ((ComponentY == 1) ? 0 : 1);
-			const int32 XDest = (ComponentX == 0) ? 0 : ((ComponentX == 1) ? ExpandQuadsX : (ComponentSizeQuads + ExpandQuadsX + 1));
-			const int32 YDest = (ComponentY == 0) ? 0 : ((ComponentY == 1) ? ExpandQuadsY : (ComponentSizeQuads + ExpandQuadsY + 1));
-			const int32 XNum = (ComponentX == 1) ? (ComponentSizeQuads + 1) : ExpandQuadsX;
-			const int32 YNum = (ComponentY == 1) ? (ComponentSizeQuads + 1) : ExpandQuadsY;
-			const int32 XBackup = (ComponentX == 2) ? (ComponentSizeQuads + ExpandQuadsX) : ExpandQuadsX;
-			const int32 YBackup = (ComponentY == 2) ? (ComponentSizeQuads + ExpandQuadsY) : ExpandQuadsY;
-			const int32 XBackupNum = (ComponentX == 1) ? (ComponentSizeQuads + 1) : 1;
-			const int32 YBackupNum = (ComponentY == 1) ? (ComponentSizeQuads + 1) : 1;
-			
-			ULandscapeComponent* Neighbor = Info->XYtoComponentMap.FindRef(ComponentBase + FIntPoint((ComponentX - 1), (ComponentY - 1)));
+			// Coordinates and Num are all in component-space not tex-space
+			// Note: This means they don't include the duped vert when NumSubsections == 2
+			const int32 XSource = (ComponentX == -1) ? (ComponentSizeQuads - ExpandQuadsX) : ((ComponentX == 0) ? 0 : 1);
+			const int32 YSource = (ComponentY == -1) ? (ComponentSizeQuads - ExpandQuadsY) : ((ComponentY == 0) ? 0 : 1);
+			const int32 XDest = (ComponentX == -1) ? 0 : ((ComponentX == 0) ? ExpandQuadsX : (ComponentSizeQuads + ExpandQuadsX + 1));
+			const int32 YDest = (ComponentY == -1) ? 0 : ((ComponentY == 0) ? ExpandQuadsY : (ComponentSizeQuads + ExpandQuadsY + 1));
+			const int32 XNum = (ComponentX == 0) ? (ComponentSizeQuads + 1) : ExpandQuadsX;
+			const int32 YNum = (ComponentY == 0) ? (ComponentSizeQuads + 1) : ExpandQuadsY;
+
+			ULandscapeComponent* Neighbor = Info->XYtoComponentMap.FindRef(ComponentBase + FIntPoint(ComponentX, ComponentY));
 			if (Neighbor)
 			{
 				// Data array for upscaling case
@@ -476,25 +497,53 @@ void FLandscapeStaticLightingMesh::GetHeightmapData(int32 InLOD, int32 GeometryL
 				int32 NeighborGeometricLOD = ::GetLightingLOD(Neighbor);
 				FLandscapeComponentDataInterface DataInterface(Neighbor, InLOD);
 				::InternalUpscaling(DataInterface, Neighbor, InLOD, NeighborGeometricLOD, CompHeightData, CompXYOffsetData);
+
+				TArray<uint16> RenderedWPOData;
+				if (bUseRenderedWPO)
+				{
+					RenderedWPOData = Neighbor->RenderWPOHeightmap(InLOD);
+				}
+
 				for (int32 Y = 0; Y < YNum; Y++)
 				{
 					const FColor* const Data = DataInterface.GetHeightData(0, YSource + Y);
 
+					const int32 HeightDataOffset = XDest - XSource + (YDest + Y) * NumVertices;
+
 					int32 NextX;
 					for (int32 X = XSource; X < XSource + XNum; X = NextX)
 					{
 						NextX = (X / SubsectionSizeQuads + 1) * SubsectionSizeQuads + 1;
 
-						const int32 CompX = X + FMath::Min(X/SubsectionSizeQuads, NumSubsections - 1);
-						const FColor* const SubsectionData = &Data[CompX];
+						// Correct for Subsection texel duplication
+						const int32 TexX = X + FMath::Min(X/SubsectionSizeQuads, NumSubsections - 1);
+						const FColor* const SubsectionData = &Data[TexX];
 
 						// Copy the data
-						FMemory::Memcpy( &HeightData[XDest + (X - XSource) + (YDest + Y) * NumVertices], SubsectionData, FMath::Min(NextX - X, XSource + XNum - X) * sizeof(FColor));
+						FMemory::Memcpy(&HeightData[X + HeightDataOffset], SubsectionData, (FMath::Min(NextX, XSource + XNum) - X) * sizeof(FColor));
+					}
+
+					if (bUseRenderedWPO)
+					{
+						// All in component-space, no need to take texel duplication into account :)
+						const int32 WPODataOffset    = (YSource + Y) * (ComponentSizeQuads + 1);
+						for (int32 X = XSource; X < XSource + XNum; ++X)
+						{
+							const uint16 WPOHeight = RenderedWPOData[X + WPODataOffset];
+							FColor& Height = HeightData[X + HeightDataOffset];
+							Height.R = (WPOHeight >> 8);
+							Height.G = (WPOHeight & 0xFF);
+						}
 					}
 				}
 			}
 			else
 			{
+				const int32 XBackup = (ComponentX == 1) ? (ComponentSizeQuads + ExpandQuadsX) : ExpandQuadsX;
+				const int32 YBackup = (ComponentY == 1) ? (ComponentSizeQuads + ExpandQuadsY) : ExpandQuadsY;
+				const int32 XBackupNum = (ComponentX == 0) ? (ComponentSizeQuads + 1) : 1;
+				const int32 YBackupNum = (ComponentY == 0) ? (ComponentSizeQuads + 1) : 1;
+
 				for (int32 Y = 0; Y < YNum; Y++)
 				{
 					for (int32 X = 0; X < XNum; X += XBackupNum)
diff --git a/Engine/Source/Runtime/Landscape/Private/LandscapeRender.cpp b/Engine/Source/Runtime/Landscape/Private/LandscapeRender.cpp
--- a/Engine/Source/Runtime/Landscape/Private/LandscapeRender.cpp
+++ b/Engine/Source/Runtime/Landscape/Private/LandscapeRender.cpp
@@ -563,7 +563,6 @@ FLandscapeComponentSceneProxy::FLandscapeComponentSceneProxy(ULandscapeComponent
 	, LandscapeComponent(InComponent)
 	, LODFalloff(InComponent->GetLandscapeProxy()->LODFalloff)
 #if WITH_EDITOR || !(UE_BUILD_SHIPPING || UE_BUILD_TEST)
-	, bBakeMaterialPositionOffsetIntoCollision(InComponent->GetLandscapeProxy()->bBakeMaterialPositionOffsetIntoCollision)
 	, CollisionMipLevel(InComponent->CollisionMipLevel)
 	, SimpleCollisionMipLevel(InComponent->SimpleCollisionMipLevel)
 	, CollisionResponse(InComponent->GetLandscapeProxy()->BodyInstance.GetResponseToChannels())
@@ -658,14 +657,7 @@ FLandscapeComponentSceneProxy::FLandscapeComponentSceneProxy(ULandscapeComponent
 	bRequiresAdjacencyInformation = RequiresAdjacencyInformation(MaterialInterface, XYOffsetmapTexture == nullptr ? &FLandscapeVertexFactory::StaticType : &FLandscapeXYOffsetVertexFactory::StaticType, InComponent->GetWorld()->FeatureLevel);
 
 	const int8 SubsectionSizeLog2 = FMath::CeilLogTwo(InComponent->SubsectionSizeQuads + 1);
-#if WITH_EDITOR
-	const int8 SharedBufferCollisionMip = bBakeMaterialPositionOffsetIntoCollision ? CollisionMipLevel : 0;
-	const int8 SharedBufferSimpleCollisionMip = bBakeMaterialPositionOffsetIntoCollision ? SimpleCollisionMipLevel : 0;
-#endif
 	SharedBuffersKey = (SubsectionSizeLog2 & 0xf) | ((NumSubsections & 0xf) << 4) |
-#if WITH_EDITOR
-	                   ((SharedBufferCollisionMip & 0xf) << 8) | ((SharedBufferSimpleCollisionMip & 0xf) << 12) |
-#endif
 	                   (FeatureLevel <= ERHIFeatureLevel::ES3_1 ? 0 : 1 << 30) | (XYOffsetmapTexture == nullptr ? 0 : 1 << 31);
 
 	bSupportsHeightfieldRepresentation = true;
@@ -697,17 +689,9 @@ void FLandscapeComponentSceneProxy::CreateRenderThreadResources()
 	SharedBuffers = FLandscapeComponentSceneProxy::SharedBuffersMap.FindRef(SharedBuffersKey);
 	if (SharedBuffers == nullptr)
 	{
-#if WITH_EDITOR
-		SharedBuffers = new FLandscapeSharedBuffers(
-			SharedBuffersKey, SubsectionSizeQuads, NumSubsections,
-			FeatureLevel, bRequiresAdjacencyInformation,
-			bBakeMaterialPositionOffsetIntoCollision ? CollisionMipLevel : INDEX_NONE,
-			bBakeMaterialPositionOffsetIntoCollision ? SimpleCollisionMipLevel : INDEX_NONE);
-#else
 		SharedBuffers = new FLandscapeSharedBuffers(
 			SharedBuffersKey, SubsectionSizeQuads, NumSubsections,
 			FeatureLevel, bRequiresAdjacencyInformation);
-#endif
 
 		FLandscapeComponentSceneProxy::SharedBuffersMap.Add(SharedBuffersKey, SharedBuffers);
 
@@ -769,8 +753,11 @@ void FLandscapeComponentSceneProxy::CreateRenderThreadResources()
 	// Create MeshBatch for grass rendering
 	if(SharedBuffers->GrassIndexBuffer)
 	{
-		GrassMeshBatch.Elements.Empty(3);
-		GrassMeshBatch.Elements.AddDefaulted(3);
+		const int32 NumMips = FMath::CeilLogTwo(SubsectionSizeVerts);
+		GrassMeshBatch.Elements.Empty(NumMips);
+		GrassMeshBatch.Elements.AddDefaulted(NumMips);
+		GrassBatchParams.Empty(NumMips);
+		GrassBatchParams.AddDefaulted(NumMips);
 
 		FMaterialRenderProxy* RenderProxy = MaterialInterface->GetRenderProxy(false);
 		GrassMeshBatch.VertexFactory = VertexFactory;
@@ -802,31 +789,18 @@ void FLandscapeComponentSceneProxy::CreateRenderThreadResources()
 		GrassBatchElement->MinVertexIndex = 0;
 		GrassBatchElement->MaxVertexIndex = SharedBuffers->NumVertices - 1;
 
-		const bool bRenderCollision = CollisionMipLevel > 0;
-		const bool bRenderSimpleCollision = SimpleCollisionMipLevel > CollisionMipLevel;
-
-		if (bRenderCollision)
+		for (int32 Mip = 1; Mip < NumMips; ++Mip)
 		{
-			FMeshBatchElement* CollisionBatchElement = &GrassMeshBatch.Elements[1];
+			const int32 MipSubsectionSizeVerts = SubsectionSizeVerts >> Mip;
+
+			FMeshBatchElement* CollisionBatchElement = &GrassMeshBatch.Elements[Mip];
 			*CollisionBatchElement = *GrassBatchElement;
-			FLandscapeBatchElementParams* CollisionBatchElementParams = &GrassBatchParams[1];
+			FLandscapeBatchElementParams* CollisionBatchElementParams = &GrassBatchParams[Mip];
 			*CollisionBatchElementParams = *BatchElementParams;
-			CollisionBatchElementParams->CurrentLOD = CollisionMipLevel;
+			CollisionBatchElementParams->CurrentLOD = Mip;
 			CollisionBatchElement->UserData = CollisionBatchElementParams;
-			CollisionBatchElement->NumPrimitives = FMath::Square(NumSubsections) * FMath::Square(SubsectionSizeVerts >> CollisionMipLevel);
-			CollisionBatchElement->FirstIndex = SharedBuffers->GrassCollisionOffset;
-		}
-
-		if (bRenderSimpleCollision)
-		{
-			FMeshBatchElement* SimpleCollisionBatchElement = &GrassMeshBatch.Elements[2];
-			*SimpleCollisionBatchElement = *GrassBatchElement;
-			FLandscapeBatchElementParams* SimpleCollisionBatchElementParams = &GrassBatchParams[2];
-			*SimpleCollisionBatchElementParams = *BatchElementParams;
-			SimpleCollisionBatchElementParams->CurrentLOD = SimpleCollisionMipLevel;
-			SimpleCollisionBatchElement->UserData = SimpleCollisionBatchElementParams;
-			SimpleCollisionBatchElement->NumPrimitives = FMath::Square(NumSubsections) * FMath::Square(SubsectionSizeVerts >> SimpleCollisionMipLevel);
-			SimpleCollisionBatchElement->FirstIndex = SharedBuffers->GrassSimpleCollisionOffset;
+			CollisionBatchElement->NumPrimitives = FMath::Square(NumSubsections) * FMath::Square(MipSubsectionSizeVerts);
+			CollisionBatchElement->FirstIndex = SharedBuffers->GrassIndexMipOffsets[Mip];
 		}
 	}
 #endif
@@ -2086,70 +2060,32 @@ void FLandscapeSharedBuffers::CreateGrassIndexBuffer()
 {
 	TArray<INDEX_TYPE> NewIndices;
 
-	const bool bRenderCollision = GrassCollisionMip > 0;
-	const bool bRenderSimpleCollision = GrassSimpleCollisionMip > GrassCollisionMip;
-	const int32 CollisionSubsectionSizeVerts = bRenderCollision ? SubsectionSizeVerts >> GrassCollisionMip : 0;
-	const int32 SimpleCollisionSubsectionSizeVerts = bRenderSimpleCollision ? SubsectionSizeVerts >> GrassSimpleCollisionMip : 0;
-
-	int32 ExpectedNumIndices = FMath::Square(NumSubsections) * FMath::Square(SubsectionSizeVerts) +
-	                           FMath::Square(NumSubsections) * FMath::Square(CollisionSubsectionSizeVerts) +
-	                           FMath::Square(NumSubsections) * FMath::Square(SimpleCollisionSubsectionSizeVerts);
+	int32 ExpectedNumIndices = FMath::Square(NumSubsections) * (FMath::Square(SubsectionSizeVerts) * 4/3 - 1); // *4/3 is for mips, -1 because we only go down to 2x2 not 1x1
 	NewIndices.Empty(ExpectedNumIndices);
 
-	int32 SubOffset = 0;
-	for (int32 SubY = 0; SubY < NumSubsections; SubY++)
-	{
-		for (int32 SubX = 0; SubX < NumSubsections; SubX++)
-		{
-			for (int32 y = 0; y < SubsectionSizeVerts; y++)
-			{
-				for (int32 x = 0; x < SubsectionSizeVerts; x++)
-				{
-					NewIndices.Add(x + y * SubsectionSizeVerts + SubOffset);
-				}
-			}
-
-			SubOffset += FMath::Square(SubsectionSizeVerts);
-		}
-	}
-	if (bRenderCollision)
-	{
-		GrassCollisionOffset = NewIndices.Num();
-
-		SubOffset = 0;
-		for (int32 SubY = 0; SubY < NumSubsections; SubY++)
-		{
-			for (int32 SubX = 0; SubX < NumSubsections; SubX++)
-			{
-				for (int32 y = 0; y < CollisionSubsectionSizeVerts; y++)
-				{
-					for (int32 x = 0; x < CollisionSubsectionSizeVerts; x++)
-					{
-						NewIndices.Add(x + y * SubsectionSizeVerts + SubOffset);
-					}
-				}
+	int32 NumMips = FMath::CeilLogTwo(SubsectionSizeVerts);
 
-				SubOffset += FMath::Square(SubsectionSizeVerts);
-			}
-		}
-	}
-	if (bRenderSimpleCollision)
+	for (int32 Mip = 0; Mip < NumMips; ++Mip)
 	{
-		GrassSimpleCollisionOffset = NewIndices.Num();
+		// Store offset to the start of this mip in the index buffer
+		GrassIndexMipOffsets.Add(NewIndices.Num());
 
-		SubOffset = 0;
+		int32 MipSubsectionSizeVerts = SubsectionSizeVerts >> Mip;
+		int32 SubOffset = 0;
 		for (int32 SubY = 0; SubY < NumSubsections; SubY++)
 		{
 			for (int32 SubX = 0; SubX < NumSubsections; SubX++)
 			{
-				for (int32 y = 0; y < SimpleCollisionSubsectionSizeVerts; y++)
+				for (int32 y = 0; y < MipSubsectionSizeVerts; y++)
 				{
-					for (int32 x = 0; x < SimpleCollisionSubsectionSizeVerts; x++)
+					for (int32 x = 0; x < MipSubsectionSizeVerts; x++)
 					{
+						// intentionally using SubsectionSizeVerts not MipSubsectionSizeVerts, this is a vert buffer index not a mip vert index
 						NewIndices.Add(x + y * SubsectionSizeVerts + SubOffset);
 					}
 				}
 
+				// intentionally using SubsectionSizeVerts not MipSubsectionSizeVerts (as above)
 				SubOffset += FMath::Square(SubsectionSizeVerts);
 			}
 		}
@@ -2165,27 +2101,17 @@ void FLandscapeSharedBuffers::CreateGrassIndexBuffer()
 }
 #endif
 
-#if WITH_EDITOR
-FLandscapeSharedBuffers::FLandscapeSharedBuffers(const int32 InSharedBuffersKey, const int32 InSubsectionSizeQuads, const int32 InNumSubsections, const ERHIFeatureLevel::Type InFeatureLevel, const bool bRequiresAdjacencyInformation, const int32 InCollisionMipLevel, const int32 InSimpleCollisionMipLevel)
-#else
 FLandscapeSharedBuffers::FLandscapeSharedBuffers(const int32 InSharedBuffersKey, const int32 InSubsectionSizeQuads, const int32 InNumSubsections, const ERHIFeatureLevel::Type InFeatureLevel, const bool bRequiresAdjacencyInformation)
-#endif
 	: SharedBuffersKey(InSharedBuffersKey)
 	, NumIndexBuffers(FMath::CeilLogTwo(InSubsectionSizeQuads + 1))
 	, SubsectionSizeVerts(InSubsectionSizeQuads + 1)
 	, NumSubsections(InNumSubsections)
-#if WITH_EDITOR
-	, GrassCollisionMip(InCollisionMipLevel)
-	, GrassSimpleCollisionMip(InSimpleCollisionMipLevel)
-#endif
 	, VertexFactory(nullptr)
 	, VertexBuffer(nullptr)
 	, AdjacencyIndexBuffers(nullptr)
 	, bUse32BitIndices(false)
 #if WITH_EDITOR
 	, GrassIndexBuffer(nullptr)
-	, GrassCollisionOffset(INDEX_NONE)
-	, GrassSimpleCollisionOffset(INDEX_NONE)
 #endif
 {
 	NumVertices = FMath::Square(SubsectionSizeVerts) * FMath::Square(NumSubsections);
diff --git a/Engine/Source/Runtime/Landscape/Private/LandscapeRenderMobile.cpp b/Engine/Source/Runtime/Landscape/Private/LandscapeRenderMobile.cpp
--- a/Engine/Source/Runtime/Landscape/Private/LandscapeRenderMobile.cpp
+++ b/Engine/Source/Runtime/Landscape/Private/LandscapeRenderMobile.cpp
@@ -247,16 +247,10 @@ void FLandscapeComponentSceneProxyMobile::CreateRenderThreadResources()
 	SharedBuffers = FLandscapeComponentSceneProxy::SharedBuffersMap.FindRef(SharedBuffersKey);
 	if (SharedBuffers == nullptr)
 	{
-#if WITH_EDITOR
-		SharedBuffers = new FLandscapeSharedBuffers(
-			SharedBuffersKey, SubsectionSizeQuads, NumSubsections,
-			GetScene().GetFeatureLevel(), false,
-			CollisionMipLevel, SimpleCollisionMipLevel);
-#else
 		SharedBuffers = new FLandscapeSharedBuffers(
 			SharedBuffersKey, SubsectionSizeQuads, NumSubsections,
 			GetScene().GetFeatureLevel(), false);
-#endif
+
 		FLandscapeComponentSceneProxy::SharedBuffersMap.Add(SharedBuffersKey, SharedBuffers);
 	}
 
diff --git a/Engine/Source/Runtime/Landscape/Private/LandscapeVersion.h b/Engine/Source/Runtime/Landscape/Private/LandscapeVersion.h
--- a/Engine/Source/Runtime/Landscape/Private/LandscapeVersion.h
+++ b/Engine/Source/Runtime/Landscape/Private/LandscapeVersion.h
@@ -15,6 +15,8 @@ namespace FLandscapeCustomVersion
 		GrassMaterialWPO,
 		// Support material world-position-offset in landscape simple collision
 		CollisionMaterialWPO,
+		// Support material world-position-offset in landscape lighting mesh
+		LightmassMaterialWPO,
 
 		// -----<new versions can be added above this line>-------------------------------------------------
 		VersionPlusOne,
diff --git a/Engine/Source/Runtime/Landscape/Private/Materials/MaterialExpressionLandscapeLayerBlend.cpp b/Engine/Source/Runtime/Landscape/Private/Materials/MaterialExpressionLandscapeLayerBlend.cpp
--- a/Engine/Source/Runtime/Landscape/Private/Materials/MaterialExpressionLandscapeLayerBlend.cpp
+++ b/Engine/Source/Runtime/Landscape/Private/Materials/MaterialExpressionLandscapeLayerBlend.cpp
@@ -114,6 +114,27 @@ FString UMaterialExpressionLandscapeLayerBlend::GetInputName(int32 InputIndex) c
 }
 
 #if WITH_EDITOR
+uint32 UMaterialExpressionLandscapeLayerBlend::GetInputType(int32 InputIndex)
+{
+	int32 Idx = 0;
+	for (int32 LayerIdx = 0; LayerIdx<Layers.Num(); LayerIdx++)
+	{
+		if (InputIndex == Idx++)
+		{
+			return MCT_Float | MCT_MaterialAttributes; // can accept pretty much anything including MaterialAttributes
+		}
+		if (Layers[LayerIdx].BlendType == LB_HeightBlend)
+		{
+			if (InputIndex == Idx++)
+			{
+				return MCT_Float1; // the height input must be float1
+			}
+		}
+	}
+
+	return MCT_Unknown;
+}
+
 bool UMaterialExpressionLandscapeLayerBlend::IsResultMaterialAttributes(int32 OutputIndex)
 {
 	if (ContainsInputLoop())
diff --git a/Engine/Source/Runtime/Landscape/Public/LandscapeRender.h b/Engine/Source/Runtime/Landscape/Public/LandscapeRender.h
--- a/Engine/Source/Runtime/Landscape/Public/LandscapeRender.h
+++ b/Engine/Source/Runtime/Landscape/Public/LandscapeRender.h
@@ -290,10 +290,6 @@ public:
 	int32 NumIndexBuffers;
 	int32 SubsectionSizeVerts;
 	int32 NumSubsections;
-#if WITH_EDITOR
-	int32 GrassCollisionMip;
-	int32 GrassSimpleCollisionMip;
-#endif
 
 	FLandscapeVertexFactory* VertexFactory;
 	FLandscapeVertexBuffer* VertexBuffer;
@@ -303,15 +299,10 @@ public:
 	bool bUse32BitIndices;
 #if WITH_EDITOR
 	FIndexBuffer* GrassIndexBuffer;
-	int32 GrassCollisionOffset;
-	int32 GrassSimpleCollisionOffset;
+	TArray<int32, TInlineAllocator<8>> GrassIndexMipOffsets;
 #endif
 
-#if WITH_EDITOR
-	FLandscapeSharedBuffers(int32 SharedBuffersKey, int32 SubsectionSizeQuads, int32 NumSubsections, ERHIFeatureLevel::Type FeatureLevel, bool bRequiresAdjacencyInformation, int32 CollisionMipLevel, int32 SimpleCollisionMipLevel);
-#else
 	FLandscapeSharedBuffers(int32 SharedBuffersKey, int32 SubsectionSizeQuads, int32 NumSubsections, ERHIFeatureLevel::Type FeatureLevel, bool bRequiresAdjacencyInformation);
-#endif
 
 	template <typename INDEX_TYPE>
 	void CreateIndexBuffers(ERHIFeatureLevel::Type InFeatureLevel, bool bRequiresAdjacencyInformation);
@@ -516,12 +507,9 @@ protected:
 	TArray<FLandscapeBatchElementParams> StaticBatchParamArray;
 
 #if WITH_EDITOR
-	// Precomputed grass rendering MeshBatch and params
-	// 0 - regular grass rendering at LOD 0
-	// 1 - rendering for collision baking
-	// 2 - rendering for "simple collision" baking
-	FMeshBatch					GrassMeshBatch;
-	FLandscapeBatchElementParams GrassBatchParams[3];
+	// Precomputed grass rendering MeshBatch and per-LOD params
+	FMeshBatch                           GrassMeshBatch;
+	TArray<FLandscapeBatchElementParams> GrassBatchParams;
 #endif
 
 	// Precomputed values
@@ -566,7 +554,6 @@ protected:
 
 	// data used in editor or visualisers
 #if WITH_EDITOR || !(UE_BUILD_SHIPPING || UE_BUILD_TEST)
-	bool bBakeMaterialPositionOffsetIntoCollision;
 	int32 CollisionMipLevel;
 	int32 SimpleCollisionMipLevel;
 
