Remaking CL 2926270 from 4.11 branch to fix non-unity compile error on Mac due to template function reference between files.

diff --git a/Engine/Source/Runtime/Landscape/Private/LandscapeEditInterface.cpp b/Engine/Source/Runtime/Landscape/Private/LandscapeEditInterface.cpp
--- a/Engine/Source/Runtime/Landscape/Private/LandscapeEditInterface.cpp
+++ b/Engine/Source/Runtime/Landscape/Private/LandscapeEditInterface.cpp
@@ -926,74 +926,6 @@ uint16 FLandscapeEditDataInterface::GetHeightMapData(const ULandscapeComponent*
 	return ((((uint16)TexData.R) << 8) | TexData.G);
 }
 
-template<typename T>
-void FillCornerValues(uint8& CornerSet, T* CornerValues)
-{
-	uint8 OriginalSet = CornerSet;
-
-	if (CornerSet)
-	{
-		// Fill unset values
-		while (CornerSet != 15)
-		{
-			if (CornerSet != 15 && (OriginalSet & 1))
-			{
-				if (!(CornerSet & 1 << 1))
-				{
-					CornerValues[1] = CornerValues[0];
-					CornerSet |= 1 << 1;
-				}
-				if (!(CornerSet & 1 << 2))
-				{
-					CornerValues[2] = CornerValues[0];
-					CornerSet |= 1 << 2;
-				}
-			}
-			if (CornerSet != 15 && (OriginalSet & 1 << 1))
-			{
-				if (!(CornerSet & 1))
-				{
-					CornerValues[0] = CornerValues[1];
-					CornerSet |= 1;
-				}
-				if (!(CornerSet & 1 << 3))
-				{
-					CornerValues[3] = CornerValues[1];
-					CornerSet |= 1 << 3;
-				}
-			}
-			if (CornerSet != 15 && (OriginalSet & 1 << 2))
-			{
-				if (!(CornerSet & 1))
-				{
-					CornerValues[0] = CornerValues[2];
-					CornerSet |= 1;
-				}
-				if (!(CornerSet & 1 << 3))
-				{
-					CornerValues[3] = CornerValues[2];
-					CornerSet |= 1 << 3;
-				}
-			}
-			if (CornerSet != 15 && (OriginalSet & 1 << 3))
-			{
-				if (!(CornerSet & 1 << 1))
-				{
-					CornerValues[1] = CornerValues[3];
-					CornerSet |= 1 << 1;
-				}
-				if (!(CornerSet & 1 << 2))
-				{
-					CornerValues[2] = CornerValues[3];
-					CornerSet |= 1 << 2;
-				}
-			}
-
-			OriginalSet = CornerSet;
-		}
-	}
-}
-
 template<typename TStoreData>
 void FLandscapeEditDataInterface::GetHeightDataTempl(int32& ValidX1, int32& ValidY1, int32& ValidX2, int32& ValidY2, TStoreData& StoreData)
 {
diff --git a/Engine/Source/Runtime/Landscape/Public/LandscapeDataAccess.h b/Engine/Source/Runtime/Landscape/Public/LandscapeDataAccess.h
--- a/Engine/Source/Runtime/Landscape/Public/LandscapeDataAccess.h
+++ b/Engine/Source/Runtime/Landscape/Public/LandscapeDataAccess.h
@@ -357,7 +357,72 @@ public:
 
 // Helper functions
 template<typename T>
-void FillCornerValues(uint8& CornerSet, T* CornerValues);
+void FillCornerValues(uint8& CornerSet, T* CornerValues)
+{
+	uint8 OriginalSet = CornerSet;
+
+	if (CornerSet)
+	{
+		// Fill unset values
+		while (CornerSet != 15)
+		{
+			if (CornerSet != 15 && (OriginalSet & 1))
+			{
+				if (!(CornerSet & 1 << 1))
+				{
+					CornerValues[1] = CornerValues[0];
+					CornerSet |= 1 << 1;
+				}
+				if (!(CornerSet & 1 << 2))
+				{
+					CornerValues[2] = CornerValues[0];
+					CornerSet |= 1 << 2;
+				}
+			}
+			if (CornerSet != 15 && (OriginalSet & 1 << 1))
+			{
+				if (!(CornerSet & 1))
+				{
+					CornerValues[0] = CornerValues[1];
+					CornerSet |= 1;
+				}
+				if (!(CornerSet & 1 << 3))
+				{
+					CornerValues[3] = CornerValues[1];
+					CornerSet |= 1 << 3;
+				}
+			}
+			if (CornerSet != 15 && (OriginalSet & 1 << 2))
+			{
+				if (!(CornerSet & 1))
+				{
+					CornerValues[0] = CornerValues[2];
+					CornerSet |= 1;
+				}
+				if (!(CornerSet & 1 << 3))
+				{
+					CornerValues[3] = CornerValues[2];
+					CornerSet |= 1 << 3;
+				}
+			}
+			if (CornerSet != 15 && (OriginalSet & 1 << 3))
+			{
+				if (!(CornerSet & 1 << 1))
+				{
+					CornerValues[1] = CornerValues[3];
+					CornerSet |= 1 << 1;
+				}
+				if (!(CornerSet & 1 << 2))
+				{
+					CornerValues[2] = CornerValues[3];
+					CornerSet |= 1 << 2;
+				}
+			}
+
+			OriginalSet = CornerSet;
+		}
+	}
+}
 
 
 
