Copying //UE4/Dev-Landscape to //UE4/Dev-Main (Source: //UE4/Dev-Landscape [REDACTED_HANDLE])

diff --git a/Engine/Config/BaseEngine.ini b/Engine/Config/BaseEngine.ini
--- a/Engine/Config/BaseEngine.ini
+++ b/Engine/Config/BaseEngine.ini
@@ -385,6 +385,7 @@ bUseStreamingPause=false
 +ActiveClassRedirects=(OldClassName="LandscapeGizmoActiveActor",NewClassName="/Script/Landscape.LandscapeGizmoActiveActor")
 +ActiveClassRedirects=(OldClassName="LandscapeGizmoActor",NewClassName="/Script/Landscape.LandscapeGizmoActor")
 +ActiveClassRedirects=(OldClassName="LandscapeProxy",NewClassName="/Script/Landscape.LandscapeProxy")
++ActiveClassRedirects=(OldClassName="LandscapeProxy",NewClassName="/Script/Landscape.LandscapeStreamingProxy",InstanceOnly="true")
 +ActiveClassRedirects=(OldClassName="ControlPointMeshComponent",NewClassName="/Script/Landscape.ControlPointMeshComponent")
 +ActiveClassRedirects=(OldClassName="LandscapeComponent",NewClassName="/Script/Landscape.LandscapeComponent")
 +ActiveClassRedirects=(OldClassName="LandscapeGizmoRenderComponent",NewClassName="/Script/Landscape.LandscapeGizmoRenderComponent")
diff --git a/Engine/Source/Developer/MeshUtilities/Public/MeshMergeData.h b/Engine/Source/Developer/MeshUtilities/Public/MeshMergeData.h
--- a/Engine/Source/Developer/MeshUtilities/Public/MeshMergeData.h
+++ b/Engine/Source/Developer/MeshUtilities/Public/MeshMergeData.h
@@ -4,6 +4,7 @@
 
 #include "RawMesh.h"
 #include "PhysicsEngine/AggregateGeom.h"
+#include "Engine/StaticMesh.h"
 
 /** Structure holding intermediate mesh merging data that is used throughout the mesh merging and proxy creation processes */
 struct FMeshMergeData
diff --git a/Engine/Source/Developer/MeshUtilities/Public/MeshUtilities.h b/Engine/Source/Developer/MeshUtilities/Public/MeshUtilities.h
--- a/Engine/Source/Developer/MeshUtilities/Public/MeshUtilities.h
+++ b/Engine/Source/Developer/MeshUtilities/Public/MeshUtilities.h
@@ -4,10 +4,12 @@
 
 #include "ModuleInterface.h"
 #include "SkeletalMeshTypes.h"
+#include "Engine/MeshMerging.h"
 
 typedef FIntPoint FMeshIdAndLOD;
 struct FFlattenMaterial;
 struct FReferenceSkeleton;
+struct FStaticMeshLODResources;
 
 /**
  * Mesh reduction interface.
diff --git a/Engine/Source/Editor/EditorStyle/EditorStyle.Build.cs b/Engine/Source/Editor/EditorStyle/EditorStyle.Build.cs
--- a/Engine/Source/Editor/EditorStyle/EditorStyle.Build.cs
+++ b/Engine/Source/Editor/EditorStyle/EditorStyle.Build.cs
@@ -40,6 +40,8 @@ public class EditorStyle : ModuleRules
 			);
 		}
 
+		// I'm told this is to improve compilation performance of this module
 		OptimizeCode = CodeOptimization.Never;
+		PCHUsage = ModuleRules.PCHUsageMode.NoSharedPCHs;
 	}
 }
diff --git a/Engine/Source/Editor/LandscapeEditor/Private/LandscapeEdMode.cpp b/Engine/Source/Editor/LandscapeEditor/Private/LandscapeEdMode.cpp
--- a/Engine/Source/Editor/LandscapeEditor/Private/LandscapeEdMode.cpp
+++ b/Engine/Source/Editor/LandscapeEditor/Private/LandscapeEdMode.cpp
@@ -147,6 +147,17 @@ void FLandscapeTool::SetEditRenderType()
 	GLandscapeEditRenderMode = ELandscapeEditRenderMode::SelectRegion | (GLandscapeEditRenderMode & ELandscapeEditRenderMode::BitMaskForMask);
 }
 
+namespace LandscapeTool
+{
+	UMaterialInstance* CreateMaterialInstance(UMaterialInterface* BaseMaterial)
+	{
+		ULandscapeMaterialInstanceConstant* MaterialInstance = NewObject<ULandscapeMaterialInstanceConstant>(GetTransientPackage());
+		MaterialInstance->SetParentEditorOnly(BaseMaterial);
+		MaterialInstance->PostEditChange();
+		return MaterialInstance;
+	}
+}
+
 //
 // FEdModeLandscape
 //
@@ -164,12 +175,12 @@ FEdModeLandscape::FEdModeLandscape()
 	, CachedLandscapeMaterial(nullptr)
 	, bToolActive(false)
 {
-	GLayerDebugColorMaterial = LoadObject<UMaterial>(NULL, TEXT("/Engine/EditorLandscapeResources/LayerVisMaterial.LayerVisMaterial"), NULL, LOAD_None, NULL);
-	GSelectionColorMaterial = LoadObject<UMaterialInstanceConstant>(NULL, TEXT("/Engine/EditorLandscapeResources/SelectBrushMaterial_Selected.SelectBrushMaterial_Selected"), NULL, LOAD_None, NULL);
-	GSelectionRegionMaterial = LoadObject<UMaterialInstanceConstant>(NULL, TEXT("/Engine/EditorLandscapeResources/SelectBrushMaterial_SelectedRegion.SelectBrushMaterial_SelectedRegion"), NULL, LOAD_None, NULL);
-	GMaskRegionMaterial = LoadObject<UMaterialInstanceConstant>(NULL, TEXT("/Engine/EditorLandscapeResources/MaskBrushMaterial_MaskedRegion.MaskBrushMaterial_MaskedRegion"), NULL, LOAD_None, NULL);
-	GLandscapeBlackTexture = LoadObject<UTexture2D>(NULL, TEXT("/Engine/EngineResources/Black.Black"), NULL, LOAD_None, NULL);
-	GLandscapeLayerUsageMaterial = LoadObject<UMaterial>(NULL, TEXT("/Engine/EditorLandscapeResources/LandscapeLayerUsageMaterial.LandscapeLayerUsageMaterial"), NULL, LOAD_None, NULL);
+	GLayerDebugColorMaterial = LandscapeTool::CreateMaterialInstance(LoadObject<UMaterial>(nullptr, TEXT("/Engine/EditorLandscapeResources/LayerVisMaterial.LayerVisMaterial")));
+	GSelectionColorMaterial  = LandscapeTool::CreateMaterialInstance(LoadObject<UMaterialInstanceConstant>(nullptr, TEXT("/Engine/EditorLandscapeResources/SelectBrushMaterial_Selected.SelectBrushMaterial_Selected")));
+	GSelectionRegionMaterial = LandscapeTool::CreateMaterialInstance(LoadObject<UMaterialInstanceConstant>(nullptr, TEXT("/Engine/EditorLandscapeResources/SelectBrushMaterial_SelectedRegion.SelectBrushMaterial_SelectedRegion")));
+	GMaskRegionMaterial      = LandscapeTool::CreateMaterialInstance(LoadObject<UMaterialInstanceConstant>(nullptr, TEXT("/Engine/EditorLandscapeResources/MaskBrushMaterial_MaskedRegion.MaskBrushMaterial_MaskedRegion")));
+	GLandscapeBlackTexture   = LoadObject<UTexture2D>(nullptr, TEXT("/Engine/EngineResources/Black.Black"));
+	GLandscapeLayerUsageMaterial = LandscapeTool::CreateMaterialInstance(LoadObject<UMaterial>(nullptr, TEXT("/Engine/EditorLandscapeResources/LandscapeLayerUsageMaterial.LandscapeLayerUsageMaterial")));
 
 	// Initialize modes
 	InitializeToolModes();
@@ -2238,11 +2249,6 @@ bool FEdModeLandscape::Select(AActor* InActor, bool bInSelected)
 		return false;
 	}
 
-	if (NewLandscapePreviewMode != ENewLandscapePreviewMode::None)
-	{
-		return false;
-	}
-
 	if (InActor->IsA<ALandscapeProxy>() && bInSelected)
 	{
 		ALandscapeProxy* Landscape = CastChecked<ALandscapeProxy>(InActor);
@@ -2251,17 +2257,24 @@ bool FEdModeLandscape::Select(AActor* InActor, bool bInSelected)
 		{
 			CurrentToolTarget.LandscapeInfo = Landscape->GetLandscapeInfo();
 			UpdateTargetList();
+
+			// If we were in "New Landscape" mode and we select a landscape then switch to editing mode
+			if (NewLandscapePreviewMode != ENewLandscapePreviewMode::None)
+			{
+				SetCurrentTool("Sculpt");
+			}
 		}
 	}
 
 	if (IsSelectionAllowed(InActor, bInSelected))
 	{
+		// false means "we haven't handled the selection", which allows the editor to perform the selection
+		// so false means "allow"
 		return false;
 	}
-	else if (!bInSelected)
-	{
-		return false;
-	}
+
+	// true means "we have handled the selection", which effectively blocks the selection from happening
+	// so true means "block"
 	return true;
 }
 
@@ -2273,17 +2286,18 @@ bool FEdModeLandscape::IsSelectionAllowed(AActor* InActor, bool bInSelection) co
 		return false;
 	}
 
-	if (NewLandscapePreviewMode != ENewLandscapePreviewMode::None)
-	{
-		return false;
-	}
-
 	// Override Selection for Splines Tool
 	if (CurrentTool && CurrentTool->OverrideSelection())
 	{
 		return CurrentTool->IsSelectionAllowed(InActor, bInSelection);
 	}
 
+	if (!bInSelection)
+	{
+		// always allow de-selection
+		return true;
+	}
+
 	if (InActor->IsA(ALandscapeProxy::StaticClass()))
 	{
 		return true;
@@ -2512,6 +2526,8 @@ void FEdModeLandscape::ImportData(const FLandscapeTargetListInfo& TargetInfo, co
 			{
 				if (RawData->Num() == SizeX * SizeY * sizeof(uint16))
 				{
+					FScopedTransaction Transaction(LOCTEXT("Undo_ImportHeightmap", "Importing Landscape Heightmap"));
+
 					FHeightmapAccessor<false> HeightmapAccessor(LandscapeInfo);
 					HeightmapAccessor.SetData(MinX, MinY, MaxX, MaxY, (const uint16*)RawData->GetData());
 				}
@@ -2525,6 +2541,8 @@ void FEdModeLandscape::ImportData(const FLandscapeTargetListInfo& TargetInfo, co
 			{
 				if (RawData->Num() == SizeX * SizeY)
 				{
+					FScopedTransaction Transaction(LOCTEXT("Undo_ImportWeightmap", "Importing Landscape Layer"));
+
 					FAlphamapAccessor<false, false> AlphamapAccessor(LandscapeInfo, TargetInfo.LayerInfoObj.Get());
 					AlphamapAccessor.SetData(MinX, MinY, MaxX, MaxY, RawData->GetData(), ELandscapeLayerPaintingRestriction::None);
 				}
@@ -2731,10 +2749,9 @@ ALandscape* FEdModeLandscape::ChangeComponentSetting(int32 NumComponentsX, int32
 			}
 
 			// Delete the old Landscape and all its proxies
-			for (TActorIterator<ALandscapeProxy> It(OldLandscapeProxy->GetWorld()); It; ++It)
+			for (ALandscapeStreamingProxy* Proxy : TActorRange<ALandscapeStreamingProxy>(OldLandscapeProxy->GetWorld()))
 			{
-				ALandscapeProxy* Proxy = (*It);
-				if (Proxy && Proxy->LandscapeActor == OldLandscapeActor)
+				if (Proxy->LandscapeActor == OldLandscapeActor)
 				{
 					Proxy->Destroy();
 				}
diff --git a/Engine/Source/Editor/LandscapeEditor/Private/LandscapeEdModeBrushes.cpp b/Engine/Source/Editor/LandscapeEditor/Private/LandscapeEdModeBrushes.cpp
--- a/Engine/Source/Editor/LandscapeEditor/Private/LandscapeEdModeBrushes.cpp
+++ b/Engine/Source/Editor/LandscapeEditor/Private/LandscapeEdModeBrushes.cpp
@@ -9,6 +9,7 @@
 #include "LandscapeRender.h"
 #include "Landscape.h"
 #include "LandscapeHeightfieldCollisionComponent.h"
+#include "LandscapeMaterialInstanceConstant.h"
 
 #include "LevelUtils.h"
 #include "Materials/MaterialInstanceDynamic.h"
@@ -56,10 +57,9 @@ protected:
 	/** Protected so that only subclasses can create instances of this class. */
 	FLandscapeBrushCircle(FEdModeLandscape* InEdMode, UMaterialInterface* InBrushMaterial)
 		: LastMousePosition(0, 0)
-		, BrushMaterial(nullptr)
+		, BrushMaterial(LandscapeTool::CreateMaterialInstance(InBrushMaterial))
 		, EdMode(InEdMode)
 	{
-		BrushMaterial = InBrushMaterial; //UMaterialInstanceDynamic::Create(InBrushMaterial, nullptr);
 	}
 
 public:
@@ -339,7 +339,7 @@ class FLandscapeBrushComponent : public FLandscapeBrush
 
 protected:
 	FVector2D LastMousePosition;
-	UMaterial* BrushMaterial;
+	UMaterialInterface* BrushMaterial;
 public:
 	FEdModeLandscape* EdMode;
 
@@ -347,7 +347,8 @@ public:
 		: BrushMaterial(nullptr)
 		, EdMode(InEdMode)
 	{
-		BrushMaterial = LoadObject<UMaterial>(nullptr, TEXT("/Engine/EditorLandscapeResources/SelectBrushMaterial.SelectBrushMaterial"), nullptr, LOAD_None, nullptr);
+		UMaterial* BaseBrushMaterial = LoadObject<UMaterial>(nullptr, TEXT("/Engine/EditorLandscapeResources/SelectBrushMaterial.SelectBrushMaterial"));
+		BrushMaterial = LandscapeTool::CreateMaterialInstance(BaseBrushMaterial);
 	}
 
 	// FGCObject interface
@@ -541,8 +542,8 @@ public:
 		: BrushMaterial(nullptr)
 		, EdMode(InEdMode)
 	{
-		UMaterialInstanceConstant* GizmoMaterial = LoadObject<UMaterialInstanceConstant>(nullptr, TEXT("/Engine/EditorLandscapeResources/MaskBrushMaterial_Gizmo.MaskBrushMaterial_Gizmo"), nullptr, LOAD_None, nullptr);
-		BrushMaterial = UMaterialInstanceDynamic::Create(GizmoMaterial, nullptr);
+		UMaterialInterface* GizmoMaterial = LoadObject<UMaterialInstanceConstant>(nullptr, TEXT("/Engine/EditorLandscapeResources/MaskBrushMaterial_Gizmo.MaskBrushMaterial_Gizmo"));
+		BrushMaterial = UMaterialInstanceDynamic::Create(LandscapeTool::CreateMaterialInstance(GizmoMaterial), nullptr);
 	}
 
 	// FGCObject interface
diff --git a/Engine/Source/Editor/LandscapeEditor/Private/LandscapeEdModeComponentTools.cpp b/Engine/Source/Editor/LandscapeEditor/Private/LandscapeEdModeComponentTools.cpp
--- a/Engine/Source/Editor/LandscapeEditor/Private/LandscapeEdModeComponentTools.cpp
+++ b/Engine/Source/Editor/LandscapeEditor/Private/LandscapeEdModeComponentTools.cpp
@@ -506,7 +506,7 @@ public:
 				ALandscapeProxy* LandscapeProxy = LandscapeInfo->GetCurrentLevelLandscapeProxy(false);
 				if (!LandscapeProxy)
 				{
-					LandscapeProxy = World->SpawnActor<ALandscapeProxy>();
+					LandscapeProxy = World->SpawnActor<ALandscapeStreamingProxy>();
 					// copy shared properties to this new proxy
 					LandscapeProxy->GetSharedProperties(Landscape);
 
diff --git a/Engine/Source/Editor/LandscapeEditor/Private/LandscapeEdModePaintTools.cpp b/Engine/Source/Editor/LandscapeEditor/Private/LandscapeEdModePaintTools.cpp
--- a/Engine/Source/Editor/LandscapeEditor/Private/LandscapeEdModePaintTools.cpp
+++ b/Engine/Source/Editor/LandscapeEditor/Private/LandscapeEdModePaintTools.cpp
@@ -134,7 +134,7 @@ public:
 				for (ULandscapeComponent* Component : SelectedComponents)
 				{
 					Component->LayerWhitelist.RemoveSingle(Target.LayerInfo.Get());
-					Component->DeleteLayer(Target.LayerInfo.Get(), &LandscapeEdit);
+					Component->DeleteLayer(Target.LayerInfo.Get(), LandscapeEdit);
 				}
 			}
 
diff --git a/Engine/Source/Editor/LandscapeEditor/Private/LandscapeEdModeSplineTools.cpp b/Engine/Source/Editor/LandscapeEditor/Private/LandscapeEdModeSplineTools.cpp
--- a/Engine/Source/Editor/LandscapeEditor/Private/LandscapeEdModeSplineTools.cpp
+++ b/Engine/Source/Editor/LandscapeEditor/Private/LandscapeEdModeSplineTools.cpp
@@ -345,14 +345,14 @@ public:
 		}
 	}
 
-	void AddControlPoint(ALandscapeProxy* Landscape, const FVector& LocalLocation)
+	void AddControlPoint(ULandscapeSplinesComponent* SplinesComponent, const FVector& LocalLocation)
 	{
 		FScopedTransaction Transaction(LOCTEXT("LandscapeSpline_AddControlPoint", "Add Landscape Spline Control Point"));
 
-		Landscape->SplineComponent->Modify();
+		SplinesComponent->Modify();
 
-		ULandscapeSplineControlPoint* NewControlPoint = NewObject<ULandscapeSplineControlPoint>(Landscape->SplineComponent, NAME_None, RF_Transactional);
-		Landscape->SplineComponent->ControlPoints.Add(NewControlPoint);
+		ULandscapeSplineControlPoint* NewControlPoint = NewObject<ULandscapeSplineControlPoint>(SplinesComponent, NAME_None, RF_Transactional);
+		SplinesComponent->ControlPoints.Add(NewControlPoint);
 
 		NewControlPoint->Location = LocalLocation;
 
@@ -388,13 +388,13 @@ public:
 		SelectControlPoint(NewControlPoint);
 		UpdatePropertiesWindows();
 
-		if (!Landscape->SplineComponent->IsRegistered())
+		if (!SplinesComponent->IsRegistered())
 		{
-			Landscape->SplineComponent->RegisterComponent();
+			SplinesComponent->RegisterComponent();
 		}
 		else
 		{
-			Landscape->SplineComponent->MarkRenderStateDirty();
+			SplinesComponent->MarkRenderStateDirty();
 		}
 	}
 
@@ -840,21 +840,31 @@ public:
 		if (ViewportClient->IsCtrlPressed())
 		{
 			LandscapeInfo = InTarget.LandscapeInfo.Get();
-
 			ALandscapeProxy* Landscape = LandscapeInfo->GetCurrentLevelLandscapeProxy(true);
 			if (!Landscape)
 			{
 				return false;
 			}
 
-			if (!Landscape->SplineComponent)
+			ULandscapeSplinesComponent* SplinesComponent = nullptr;
+			if (SelectedSplineControlPoints.Num() > 0)
 			{
-				CreateSplineComponent(Landscape, FVector(1.0f) / Landscape->GetRootComponent()->RelativeScale3D);
+				ULandscapeSplineControlPoint* FirstPoint = *SelectedSplineControlPoints.CreateConstIterator();
+				SplinesComponent = FirstPoint->GetOuterULandscapeSplinesComponent();
+			}
+
+			if (!SplinesComponent)
+			{
+				if (!Landscape->SplineComponent)
+				{
+					CreateSplineComponent(Landscape, FVector(1.0f) / Landscape->GetRootComponent()->RelativeScale3D);
+				}
+				SplinesComponent = Landscape->SplineComponent;
 			}
 
-			const FTransform LandscapeToSpline = Landscape->LandscapeActorToWorld().GetRelativeTransform(Landscape->SplineComponent->ComponentToWorld);
+			const FTransform LandscapeToSpline = Landscape->LandscapeActorToWorld().GetRelativeTransform(SplinesComponent->ComponentToWorld);
 
-			AddControlPoint(Landscape, LandscapeToSpline.TransformPosition(InHitLocation));
+			AddControlPoint(SplinesComponent, LandscapeToSpline.TransformPosition(InHitLocation));
 
 			GUnrealEd->RedrawLevelEditingViewports();
 
@@ -1161,7 +1171,7 @@ public:
 						HLandscapeSplineProxy_Segment* SplineProxy = (HLandscapeSplineProxy_Segment*)HitProxy;
 						ClickedSplineSegment = SplineProxy->SplineSegment;
 
-						LandscapeToSpline = ClickedSplineSegment->GetTypedOuter<AActor>()->ActorToWorld().GetRelativeTransform(ClickedSplineSegment->GetTypedOuter<ULandscapeSplinesComponent>()->ComponentToWorld);
+						LandscapeToSpline = ClickedSplineSegment->GetTypedOuter<AActor>()->ActorToWorld().GetRelativeTransform(ClickedSplineSegment->GetOuterULandscapeSplinesComponent()->ComponentToWorld);
 					}
 					else if (HitProxy->IsA(HActor::StaticGetType()))
 					{
diff --git a/Engine/Source/Editor/LandscapeEditor/Private/LandscapeEditorDetailCustomization_AlphaBrush.cpp b/Engine/Source/Editor/LandscapeEditor/Private/LandscapeEditorDetailCustomization_AlphaBrush.cpp
--- a/Engine/Source/Editor/LandscapeEditor/Private/LandscapeEditorDetailCustomization_AlphaBrush.cpp
+++ b/Engine/Source/Editor/LandscapeEditor/Private/LandscapeEditorDetailCustomization_AlphaBrush.cpp
@@ -64,7 +64,7 @@ public:
 			.OnMouseDoubleClick(this, &STextureMaskThumbnail::OnAssetThumbnailDoubleClick)
 			//.OnClicked_Static(&FLandscapeEditorDetailCustomization_AlphaBrush::OnTextureButtonClicked)
 			[
-				SNew( SBox )
+				SNew(SBox)
 				.ToolTipText(this, &STextureMaskThumbnail::OnGetToolTip)
 				.WidthOverride(64)
 				.HeightOverride(64)
@@ -77,7 +77,7 @@ public:
 					//.ColorAndOpacity(this, &SAssetThumbnail::GetViewportColorAndOpacity)
 					//.Visibility(this, &SAssetThumbnail::GetViewportVisibility)
 					[
-						SAssignNew( ViewportWidget, SViewport )
+						SAssignNew(ViewportWidget, SViewport)
 						.EnableGammaCorrection(false)
 					]
 				]
@@ -94,7 +94,11 @@ public:
 
 	FReply OnAssetThumbnailDoubleClick(const FGeometry& InMyGeometry, const FPointerEvent& InMouseEvent)
 	{
-		GEditor->EditObject(Texture.Get());
+		UTexture2D* Texture2D = Texture.Get();
+		if (Texture2D)
+		{
+			GEditor->EditObject(Texture2D);
+		}
 
 		return FReply::Handled();
 	}
@@ -102,7 +106,8 @@ public:
 	FText OnGetToolTip() const
 	{
 		// Display the package name which is a valid path to the object without redundant information
-		return FText::FromString(Texture.Get()->GetOutermost()->GetName());
+		UTexture2D* Texture2D = Texture.Get();
+		return Texture2D ? FText::FromString(Texture2D->GetOutermost()->GetName()) : FText();
 	}
 
 	virtual int32 OnPaint(const FPaintArgs& Args, const FGeometry& AllottedGeometry, const FSlateRect& MyClippingRect, FSlateWindowElementList& OutDrawElements, int32 LayerId, const FWidgetStyle& InWidgetStyle, bool bParentEnabled) const
diff --git a/Engine/Source/Editor/LandscapeEditor/Private/LandscapeEditorDetailCustomization_NewLandscape.cpp b/Engine/Source/Editor/LandscapeEditor/Private/LandscapeEditorDetailCustomization_NewLandscape.cpp
--- a/Engine/Source/Editor/LandscapeEditor/Private/LandscapeEditorDetailCustomization_NewLandscape.cpp
+++ b/Engine/Source/Editor/LandscapeEditor/Private/LandscapeEditorDetailCustomization_NewLandscape.cpp
@@ -180,6 +180,10 @@ void FLandscapeEditorDetailCustomization_NewLandscape::CustomizeDetails(IDetailL
 		]
 	];
 
+	TSharedRef<IPropertyHandle> PropertyHandle_AlphamapType = DetailBuilder.GetProperty(GET_MEMBER_NAME_CHECKED(ULandscapeEditorObject, ImportLandscape_AlphamapType));
+	NewLandscapeCategory.AddProperty(PropertyHandle_AlphamapType)
+	.Visibility(TAttribute<EVisibility>::Create(TAttribute<EVisibility>::FGetter::CreateStatic(&GetVisibilityOnlyInNewLandscapeMode, ENewLandscapePreviewMode::ImportLandscape)));
+
 	TSharedRef<IPropertyHandle> PropertyHandle_Layers = DetailBuilder.GetProperty(GET_MEMBER_NAME_CHECKED(ULandscapeEditorObject, ImportLandscape_Layers));
 	NewLandscapeCategory.AddProperty(PropertyHandle_Layers);
 
@@ -791,23 +795,50 @@ FReply FLandscapeEditorDetailCustomization_NewLandscape::OnCreateButtonClicked()
 			WordData[i] = 32768;
 		}
 
-		TArray<FLandscapeImportLayerInfo> LayerInfos;
+		TArray<FLandscapeImportLayerInfo> ImportLayers;
 
-		if (LandscapeEdMode->NewLandscapePreviewMode == ENewLandscapePreviewMode::ImportLandscape)
+		if (LandscapeEdMode->NewLandscapePreviewMode == ENewLandscapePreviewMode::NewLandscape)
+		{
+			const auto& ImportLandscapeLayersList = LandscapeEdMode->UISettings->ImportLandscape_Layers;
+			ImportLayers.Reserve(ImportLandscapeLayersList.Num());
+
+			// Fill in LayerInfos array and allocate data
+			for (const FLandscapeImportLayer& UIImportLayer : ImportLandscapeLayersList)
+			{
+				FLandscapeImportLayerInfo ImportLayer = FLandscapeImportLayerInfo(UIImportLayer.LayerName);
+				ImportLayer.LayerInfo = UIImportLayer.LayerInfo;
+				ImportLayer.SourceFilePath = "";
+				ImportLayer.LayerData = TArray<uint8>();
+				ImportLayers.Add(MoveTemp(ImportLayer));
+			}
+
+			// Fill the first weight-blended layer to 100%
+			if (FLandscapeImportLayerInfo* FirstBlendedLayer = ImportLayers.FindByPredicate([](const FLandscapeImportLayerInfo& ImportLayer) { return ImportLayer.LayerInfo && !ImportLayer.LayerInfo->bNoWeightBlend; }))
+			{
+				FirstBlendedLayer->LayerData.AddUninitialized(SizeX * SizeY);
+
+				uint8* ByteData = FirstBlendedLayer->LayerData.GetData();
+				for (int32 i = 0; i < SizeX * SizeY; i++)
+				{
+					ByteData[i] = 255;
+				}
+			}
+		}
+		else if (LandscapeEdMode->NewLandscapePreviewMode == ENewLandscapePreviewMode::ImportLandscape)
 		{
 			const int32 ImportSizeX = LandscapeEdMode->UISettings->ImportLandscape_Width;
 			const int32 ImportSizeY = LandscapeEdMode->UISettings->ImportLandscape_Height;
 
 			const TArray<FLandscapeImportLayer>& ImportLandscapeLayersList = LandscapeEdMode->UISettings->ImportLandscape_Layers;
-			LayerInfos.Reserve(ImportLandscapeLayersList.Num());
+			ImportLayers.Reserve(ImportLandscapeLayersList.Num());
 
 			// Fill in LayerInfos array and allocate data
-			for (int32 LayerIdx = 0; LayerIdx < ImportLandscapeLayersList.Num(); LayerIdx++)
+			for (const FLandscapeImportLayer& UIImportLayer : ImportLandscapeLayersList)
 			{
-				LayerInfos.Add(ImportLandscapeLayersList[LayerIdx]); //slicing is fine here
-				FLandscapeImportLayerInfo& ImportLayer = LayerInfos.Last();
+				ImportLayers.Add((const FLandscapeImportLayer&)UIImportLayer); //slicing is fine here
+				FLandscapeImportLayerInfo& ImportLayer = ImportLayers.Last();
 
-				if (ImportLayer.LayerInfo != NULL && ImportLayer.SourceFilePath != "")
+				if (ImportLayer.LayerInfo != nullptr && ImportLayer.SourceFilePath != "")
 				{
 					FFileHelper::LoadFileToArray(ImportLayer.LayerData, *ImportLayer.SourceFilePath, FILEREAD_Silent);
 
@@ -816,7 +847,7 @@ FReply FLandscapeEditorDetailCustomization_NewLandscape::OnCreateButtonClicked()
 						IImageWrapperModule& ImageWrapperModule = FModuleManager::LoadModuleChecked<IImageWrapperModule>("ImageWrapper");
 						IImageWrapperPtr ImageWrapper = ImageWrapperModule.CreateImageWrapper(EImageFormat::PNG);
 
-						const TArray<uint8>* RawData = NULL;
+						const TArray<uint8>* RawData = nullptr;
 						if (ImageWrapper->SetCompressed(ImportLayer.LayerData.GetData(), ImportLayer.LayerData.Num()) &&
 							ImageWrapper->GetRaw(ERGBFormat::Gray, 8, RawData))
 						{
@@ -848,9 +879,9 @@ FReply FLandscapeEditorDetailCustomization_NewLandscape::OnCreateButtonClicked()
 					-OffsetX, -OffsetY, SizeX - OffsetX - 1, SizeY - OffsetY - 1);
 
 				// Layers
-				for (int32 LayerIdx = 0; LayerIdx < LayerInfos.Num(); LayerIdx++)
+				for (int32 LayerIdx = 0; LayerIdx < ImportLayers.Num(); LayerIdx++)
 				{
-					TArray<uint8>& ImportLayerData = LayerInfos[LayerIdx].LayerData;
+					TArray<uint8>& ImportLayerData = ImportLayers[LayerIdx].LayerData;
 					if (ImportLayerData.Num())
 					{
 						ImportLayerData = LandscapeEditorUtils::ExpandData(ImportLayerData,
@@ -865,14 +896,9 @@ FReply FLandscapeEditorDetailCustomization_NewLandscape::OnCreateButtonClicked()
 
 		const FVector Offset = FTransform(LandscapeEdMode->UISettings->NewLandscape_Rotation, FVector::ZeroVector, LandscapeEdMode->UISettings->NewLandscape_Scale).TransformVector(FVector(-ComponentCountX * QuadsPerComponent / 2, -ComponentCountY * QuadsPerComponent / 2, 0));
 		ALandscape* Landscape = LandscapeEdMode->GetWorld()->SpawnActor<ALandscape>(LandscapeEdMode->UISettings->NewLandscape_Location + Offset, LandscapeEdMode->UISettings->NewLandscape_Rotation);
-
-		//if (LandscapeEdMode->NewLandscapePreviewMode == ENewLandscapePreviewMode::ImportLandscape)
-		{
-			Landscape->LandscapeMaterial = LandscapeEdMode->UISettings->NewLandscape_Material.Get();
-		}
-
+		Landscape->LandscapeMaterial = LandscapeEdMode->UISettings->NewLandscape_Material.Get();
 		Landscape->SetActorRelativeScale3D(LandscapeEdMode->UISettings->NewLandscape_Scale);
-		Landscape->Import(FGuid::NewGuid(), 0, 0, SizeX-1, SizeY-1, LandscapeEdMode->UISettings->NewLandscape_SectionsPerComponent, LandscapeEdMode->UISettings->NewLandscape_QuadsPerSection, Data.GetData(), NULL, LayerInfos, LandscapeEdMode->UISettings->ImportLandscape_AlphamapType);
+		Landscape->Import(FGuid::NewGuid(), 0, 0, SizeX-1, SizeY-1, LandscapeEdMode->UISettings->NewLandscape_SectionsPerComponent, LandscapeEdMode->UISettings->NewLandscape_QuadsPerSection, Data.GetData(), nullptr, ImportLayers, LandscapeEdMode->UISettings->ImportLandscape_AlphamapType);
 
 		// automatically calculate a lighting LOD that won't crash lightmass (hopefully)
 		// < 2048x2048 -> LOD0
@@ -889,10 +915,11 @@ FReply FLandscapeEditorDetailCustomization_NewLandscape::OnCreateButtonClicked()
 		ULandscapeInfo* LandscapeInfo = Landscape->GetLandscapeInfo(true);
 		LandscapeInfo->UpdateLayerInfoMap(Landscape);
 
+		// Import doesn't fill in the LayerInfo for layers with no data, do that now
 		const TArray<FLandscapeImportLayer>& ImportLandscapeLayersList = LandscapeEdMode->UISettings->ImportLandscape_Layers;
 		for (int32 i = 0; i < ImportLandscapeLayersList.Num(); i++)
 		{
-			if (ImportLandscapeLayersList[i].LayerInfo != NULL)
+			if (ImportLandscapeLayersList[i].LayerInfo != nullptr)
 			{
 				if (LandscapeEdMode->NewLandscapePreviewMode == ENewLandscapePreviewMode::ImportLandscape)
 				{
@@ -915,7 +942,7 @@ FReply FLandscapeEditorDetailCustomization_NewLandscape::OnCreateButtonClicked()
 		LandscapeEdMode->UpdateLandscapeList();
 		LandscapeEdMode->CurrentToolTarget.LandscapeInfo = Landscape->GetLandscapeInfo();
 		LandscapeEdMode->CurrentToolTarget.TargetType = ELandscapeToolTargetType::Heightmap;
-		LandscapeEdMode->CurrentToolTarget.LayerInfo = NULL;
+		LandscapeEdMode->CurrentToolTarget.LayerInfo = nullptr;
 		LandscapeEdMode->CurrentToolTarget.LayerName = NAME_None;
 		LandscapeEdMode->UpdateTargetList();
 
diff --git a/Engine/Source/Editor/LandscapeEditor/Private/LandscapeEditorDetailCustomization_TargetLayers.cpp b/Engine/Source/Editor/LandscapeEditor/Private/LandscapeEditorDetailCustomization_TargetLayers.cpp
--- a/Engine/Source/Editor/LandscapeEditor/Private/LandscapeEditorDetailCustomization_TargetLayers.cpp
+++ b/Engine/Source/Editor/LandscapeEditor/Private/LandscapeEditorDetailCustomization_TargetLayers.cpp
@@ -26,6 +26,7 @@
 #include "LandscapeRender.h"
 #include "LandscapeLayerInfoObject.h"
 #include "Materials/MaterialExpressionLandscapeVisibilityMask.h"
+#include "LandscapeEdit.h"
 
 #define LOCTEXT_NAMESPACE "LandscapeEditor.TargetLayers"
 
@@ -498,6 +499,27 @@ TSharedPtr<SWidget> FLandscapeEditorCustomNodeBuilder_TargetLayers::OnTargetLaye
 				FUIAction ReImportAction = FUIAction(FExecuteAction::CreateStatic(&FLandscapeEditorCustomNodeBuilder_TargetLayers::OnReimportLayer, Target));
 				MenuBuilder.AddMenuEntry(FText::Format(LOCTEXT("LayerContextMenu.ReImport", "Reimport from {0}"), FText::FromString(ReimportPath)), FText(), FSlateIcon(), ReImportAction);
 			}
+
+			if (Target->TargetType == ELandscapeToolTargetType::Weightmap)
+			{
+				MenuBuilder.AddMenuSeparator();
+
+				// Fill
+				FUIAction FillAction = FUIAction(FExecuteAction::CreateStatic(&FLandscapeEditorCustomNodeBuilder_TargetLayers::OnFillLayer, Target));
+				MenuBuilder.AddMenuEntry(LOCTEXT("LayerContextMenu.Fill", "Fill Layer"), LOCTEXT("LayerContextMenu.Fill_Tooltip", "Fills this layer to 100% across the entire landscape. If this is a weight-blended (normal) layer, all other weight-blended layers will be cleared."), FSlateIcon(), FillAction);
+
+				// Clear
+				FUIAction ClearAction = FUIAction(FExecuteAction::CreateStatic(&FLandscapeEditorCustomNodeBuilder_TargetLayers::OnClearLayer, Target));
+				MenuBuilder.AddMenuEntry(LOCTEXT("LayerContextMenu.Clear", "Clear Layer"), LOCTEXT("LayerContextMenu.Clear_Tooltip", "Clears this layer to 0% across the entire landscape. If this is a weight-blended (normal) layer, other weight-blended layers will be adjusted to compensate."), FSlateIcon(), ClearAction);
+			}
+			else if (Target->TargetType == ELandscapeToolTargetType::Visibility)
+			{
+				MenuBuilder.AddMenuSeparator();
+
+				// Clear
+				FUIAction ClearAction = FUIAction(FExecuteAction::CreateStatic(&FLandscapeEditorCustomNodeBuilder_TargetLayers::OnClearLayer, Target));
+				MenuBuilder.AddMenuEntry(LOCTEXT("LayerContextMenu.ClearHoles", "Remove all Holes"), FText(), FSlateIcon(), ClearAction);
+			}
 		}
 		MenuBuilder.EndSection();
 
@@ -637,6 +659,27 @@ void FLandscapeEditorCustomNodeBuilder_TargetLayers::OnReimportLayer(const TShar
 	}
 }
 
+void FLandscapeEditorCustomNodeBuilder_TargetLayers::OnFillLayer(const TSharedRef<FLandscapeTargetListInfo> Target)
+{
+	FScopedTransaction Transaction(LOCTEXT("Undo_FillLayer", "Filling Landscape Layer"));
+	if (Target->LandscapeInfo.IsValid() && Target->LayerInfoObj.IsValid())
+	{
+		FLandscapeEditDataInterface LandscapeEdit(Target->LandscapeInfo.Get());
+		LandscapeEdit.FillLayer(Target->LayerInfoObj.Get());
+	}
+}
+
+
+void FLandscapeEditorCustomNodeBuilder_TargetLayers::OnClearLayer(const TSharedRef<FLandscapeTargetListInfo> Target)
+{
+	FScopedTransaction Transaction(LOCTEXT("Undo_ClearLayer", "Clearing Landscape Layer"));
+	if (Target->LandscapeInfo.IsValid() && Target->LayerInfoObj.IsValid())
+	{
+		FLandscapeEditDataInterface LandscapeEdit(Target->LandscapeInfo.Get());
+		LandscapeEdit.DeleteLayer(Target->LayerInfoObj.Get());
+	}
+}
+
 bool FLandscapeEditorCustomNodeBuilder_TargetLayers::ShouldFilterLayerInfo(const FAssetData& AssetData, FName LayerName)
 {
 	const FName LayerNameMetaData = AssetData.GetTagValueRef<FName>("LayerName");
diff --git a/Engine/Source/Editor/LandscapeEditor/Private/LandscapeEditorDetailCustomization_TargetLayers.h b/Engine/Source/Editor/LandscapeEditor/Private/LandscapeEditorDetailCustomization_TargetLayers.h
--- a/Engine/Source/Editor/LandscapeEditor/Private/LandscapeEditorDetailCustomization_TargetLayers.h
+++ b/Engine/Source/Editor/LandscapeEditor/Private/LandscapeEditorDetailCustomization_TargetLayers.h
@@ -53,6 +53,8 @@ protected:
 	static void OnExportLayer(const TSharedRef<FLandscapeTargetListInfo> Target);
 	static void OnImportLayer(const TSharedRef<FLandscapeTargetListInfo> Target);
 	static void OnReimportLayer(const TSharedRef<FLandscapeTargetListInfo> Target);
+	static void OnFillLayer(const TSharedRef<FLandscapeTargetListInfo> Target);
+	static void OnClearLayer(const TSharedRef<FLandscapeTargetListInfo> Target);
 	static bool ShouldFilterLayerInfo(const class FAssetData& AssetData, FName LayerName);
 	static void OnTargetLayerSetObject(const FAssetData& AssetData, const TSharedRef<FLandscapeTargetListInfo> Target);
 	static EVisibility GetTargetLayerInfoSelectorVisibility(const TSharedRef<FLandscapeTargetListInfo> Target);
diff --git a/Engine/Source/Editor/LandscapeEditor/Private/LandscapeEditorObject.cpp b/Engine/Source/Editor/LandscapeEditor/Private/LandscapeEditorObject.cpp
--- a/Engine/Source/Editor/LandscapeEditor/Private/LandscapeEditorObject.cpp
+++ b/Engine/Source/Editor/LandscapeEditor/Private/LandscapeEditorObject.cpp
@@ -630,14 +630,11 @@ void ULandscapeEditorObject::RefreshImportLayersList()
 
 		if (bFound)
 		{
-			UMaterialInstanceConstant* CombinationMaterialInstance = CastChecked<UMaterialInstanceConstant>(NewImportLayer.ThumbnailMIC->Parent);
-			if (CombinationMaterialInstance->Parent != Material)
+			if (NewImportLayer.ThumbnailMIC->Parent != Material)
 			{
 				FMaterialUpdateContext Context;
-				CombinationMaterialInstance->SetParentEditorOnly(Material);	
-				Context.AddMaterialInterface(Material);
-
-				//NewImportLayer.ThumbnailMIC = ALandscapeProxy::GetLayerThumbnailMIC(Material, LayerName, ThumbnailWeightmap, ThumbnailHeightmap, NULL);
+				NewImportLayer.ThumbnailMIC->SetParentEditorOnly(Material);
+				Context.AddMaterialInterface(NewImportLayer.ThumbnailMIC);
 			}
 
 			NewImportLayer.ImportError = ELandscapeImportLayerError::None;
@@ -694,7 +691,7 @@ void ULandscapeEditorObject::RefreshImportLayersList()
 		else
 		{
 			NewImportLayer.LayerName = LayerName;
-			NewImportLayer.ThumbnailMIC = ALandscapeProxy::GetLayerThumbnailMIC(Material, LayerName, ThumbnailWeightmap, ThumbnailHeightmap, NULL);
+			NewImportLayer.ThumbnailMIC = ALandscapeProxy::GetLayerThumbnailMIC(Material, LayerName, ThumbnailWeightmap, ThumbnailHeightmap, nullptr);
 		}
 
 		ImportLandscape_Layers.Add(MoveTemp(NewImportLayer));
diff --git a/Engine/Source/Editor/LandscapeEditor/Private/LandscapeEditorPrivatePCH.h b/Engine/Source/Editor/LandscapeEditor/Private/LandscapeEditorPrivatePCH.h
--- a/Engine/Source/Editor/LandscapeEditor/Private/LandscapeEditorPrivatePCH.h
+++ b/Engine/Source/Editor/LandscapeEditor/Private/LandscapeEditorPrivatePCH.h
@@ -4,3 +4,5 @@
 #include "LandscapeEditorModule.h"
 #include "LandscapeEdMode.h"
 #include "LandscapeEditorObject.h"
+#include "Landscape.h"
+#include "LandscapeStreamingProxy.h"
diff --git a/Engine/Source/Editor/LandscapeEditor/Public/LandscapeEditorObject.h b/Engine/Source/Editor/LandscapeEditor/Public/LandscapeEditorObject.h
--- a/Engine/Source/Editor/LandscapeEditor/Public/LandscapeEditorObject.h
+++ b/Engine/Source/Editor/LandscapeEditor/Public/LandscapeEditorObject.h
@@ -186,11 +186,15 @@ struct FLandscapeImportLayer : public FLandscapeImportLayerInfo
 {
 	GENERATED_USTRUCT_BODY()
 
+	UPROPERTY(Category="Import", VisibleAnywhere)
+	ULandscapeMaterialInstanceConstant* ThumbnailMIC;
+
 	UPROPERTY(Category="Import", VisibleAnywhere)
 	ELandscapeImportLayerError ImportError;
 
 	FLandscapeImportLayer()
 		: FLandscapeImportLayerInfo()
+		, ThumbnailMIC(nullptr)
 		, ImportError(ELandscapeImportLayerError::None)
 	{
 	}
@@ -473,7 +477,7 @@ class ULandscapeEditorObject : public UObject
 	TArray<uint16> ImportLandscape_Data;
 
 	// Whether the imported alpha maps are to be interpreted as "layered" or "additive" (UE4 uses additive internally)
-	UPROPERTY(Category="New Landscape", EditAnywhere, NonTransactional, EditFixedSize, meta=(DisplayName="Layer Alphamap Type", ShowForTools="NewLandscape"))
+	UPROPERTY(Category="New Landscape", EditAnywhere, NonTransactional, meta=(DisplayName="Layer Alphamap Type", ShowForTools="NewLandscape"))
 	ELandscapeImportAlphamapType ImportLandscape_AlphamapType;
 
 	// The landscape layers that will be created. Only layer names referenced in the material assigned above are shown here. Modify the material to add more layers.
diff --git a/Engine/Source/Editor/LandscapeEditor/Public/LandscapeToolInterface.h b/Engine/Source/Editor/LandscapeEditor/Public/LandscapeToolInterface.h
--- a/Engine/Source/Editor/LandscapeEditor/Public/LandscapeToolInterface.h
+++ b/Engine/Source/Editor/LandscapeEditor/Public/LandscapeToolInterface.h
@@ -260,3 +260,8 @@ public:
 	int32					PreviousBrushIndex;
 	TArray<FName>			ValidBrushes;
 };
+
+namespace LandscapeTool
+{
+	UMaterialInstance* CreateMaterialInstance(UMaterialInterface* BaseMaterial);
+}
diff --git a/Engine/Source/Editor/UnrealEd/Private/EditorHook.cpp b/Engine/Source/Editor/UnrealEd/Private/EditorHook.cpp
--- a/Engine/Source/Editor/UnrealEd/Private/EditorHook.cpp
+++ b/Engine/Source/Editor/UnrealEd/Private/EditorHook.cpp
@@ -9,6 +9,7 @@
 
 #include "LandscapeProxy.h"
 #include "LandscapeComponent.h"
+#include "LandscapeInfo.h"
 
 uint32			EngineThreadId;
 const TCHAR*	GItem;
diff --git a/Engine/Source/Editor/UnrealEd/Private/Lightmass/LightmassLandscapeRender.cpp b/Engine/Source/Editor/UnrealEd/Private/Lightmass/LightmassLandscapeRender.cpp
--- a/Engine/Source/Editor/UnrealEd/Private/Lightmass/LightmassLandscapeRender.cpp
+++ b/Engine/Source/Editor/UnrealEd/Private/Lightmass/LightmassLandscapeRender.cpp
@@ -10,6 +10,7 @@
 #include "EngineModule.h"
 #include "RHIStaticStates.h"
 #include "LandscapeEdit.h"
+#include "LandscapeInfo.h"
 
 struct FLightmassLandscapeVertex
 {
diff --git a/Engine/Source/Editor/WorldBrowser/Private/Tiles/WorldTileCollectionModel.cpp b/Engine/Source/Editor/WorldBrowser/Private/Tiles/WorldTileCollectionModel.cpp
--- a/Engine/Source/Editor/WorldBrowser/Private/Tiles/WorldTileCollectionModel.cpp
+++ b/Engine/Source/Editor/WorldBrowser/Private/Tiles/WorldTileCollectionModel.cpp
@@ -26,6 +26,8 @@
 #include "InstancedFoliageActor.h"
 #include "LandscapeMeshProxyActor.h"
 #include "LandscapeMeshProxyComponent.h"
+#include "LandscapeLayerInfoObject.h"
+#include "Components/HierarchicalInstancedStaticMeshComponent.h"
 
 #define LOCTEXT_NAMESPACE "WorldBrowser"
 
diff --git a/Engine/Source/Editor/WorldBrowser/Private/Tiles/WorldTileModel.cpp b/Engine/Source/Editor/WorldBrowser/Private/Tiles/WorldTileModel.cpp
--- a/Engine/Source/Editor/WorldBrowser/Private/Tiles/WorldTileModel.cpp
+++ b/Engine/Source/Editor/WorldBrowser/Private/Tiles/WorldTileModel.cpp
@@ -13,6 +13,7 @@
 #include "Engine/LevelStreaming.h"
 #include "Engine/LevelStreamingKismet.h"
 #include "Landscape.h"
+#include "LandscapeStreamingProxy.h"
 #include "AssetData.h"
 
 
@@ -920,7 +921,7 @@ ALandscapeProxy* FWorldTileModel::ImportLandscapeTile(const FLandscapeImportSett
 	
 	check(Settings.LandscapeGuid.IsValid())
 	
-	ALandscapeProxy* LandscapeProxy = Cast<UWorld>(LoadedLevel->GetOuter())->SpawnActor<ALandscapeProxy>();
+	ALandscapeProxy* LandscapeProxy = Cast<UWorld>(LoadedLevel->GetOuter())->SpawnActor<ALandscapeStreamingProxy>();
 	LandscapeProxy->SetActorTransform(Settings.LandscapeTransform);
 		
 	if (Settings.LandscapeMaterial)
diff --git a/Engine/Source/Runtime/Engine/Classes/Materials/Material.h b/Engine/Source/Runtime/Engine/Classes/Materials/Material.h
--- a/Engine/Source/Runtime/Engine/Classes/Materials/Material.h
+++ b/Engine/Source/Runtime/Engine/Classes/Materials/Material.h
@@ -550,13 +550,6 @@ public:
 	UPROPERTY(EditAnywhere, BlueprintReadOnly, Category=Usage)
 	uint32 bUsedWithEditorCompositing:1;
 
-	/** 
-	 * Indicates that the material and its instances can be use with landscapes 
-	 * This will result in the shaders required to support landscapes being compiled which will increase shader compile time and memory usage.
-	 */
-	UPROPERTY(EditAnywhere, BlueprintReadOnly, Category=Usage)
-	uint32 bUsedWithLandscape:1;
-
 	/** 
 	 * Indicates that the material and its instances can be use with particle sprites 
 	 * This will result in the shaders required to support particle sprites being compiled which will increase shader compile time and memory usage.
diff --git a/Engine/Source/Runtime/Engine/Classes/Materials/MaterialInstance.h b/Engine/Source/Runtime/Engine/Classes/Materials/MaterialInstance.h
--- a/Engine/Source/Runtime/Engine/Classes/Materials/MaterialInstance.h
+++ b/Engine/Source/Runtime/Engine/Classes/Materials/MaterialInstance.h
@@ -5,6 +5,7 @@
 #include "Materials/MaterialInterface.h"
 #include "Materials/MaterialInstanceBasePropertyOverrides.h"
 #include "StaticParameterSet.h"
+#include "Materials/Material.h"
 #include "MaterialInstance.generated.h"
 
 //
@@ -253,6 +254,7 @@ public:
 	virtual ENGINE_API UMaterial* GetMaterial() override;
 	virtual ENGINE_API const UMaterial* GetMaterial() const override;
 	virtual ENGINE_API const UMaterial* GetMaterial_Concurrent(TMicRecursionGuard& RecursionGuard) const override;
+	virtual ENGINE_API FMaterialResource* AllocatePermutationResource();
 	virtual ENGINE_API FMaterialResource* GetMaterialResource(ERHIFeatureLevel::Type InFeatureLevel, EMaterialQualityLevel::Type QualityLevel = EMaterialQualityLevel::Num) override;
 	virtual ENGINE_API const FMaterialResource* GetMaterialResource(ERHIFeatureLevel::Type InFeatureLevel, EMaterialQualityLevel::Type QualityLevel = EMaterialQualityLevel::Num) const override;
 	virtual ENGINE_API bool GetFontParameterValue(FName ParameterName, class UFont*& OutFontValue, int32& OutFontPage) const override;
diff --git a/Engine/Source/Runtime/Engine/Classes/Materials/MaterialInterface.h b/Engine/Source/Runtime/Engine/Classes/Materials/MaterialInterface.h
--- a/Engine/Source/Runtime/Engine/Classes/Materials/MaterialInterface.h
+++ b/Engine/Source/Runtime/Engine/Classes/Materials/MaterialInterface.h
@@ -14,7 +14,6 @@ class FMaterialCompiler;
 struct FPrimitiveViewRelevance;
 class UTexture;
 
-UENUM()
 enum EMaterialUsage
 {
 	MATUSAGE_SkeletalMesh,
@@ -24,7 +23,6 @@ enum EMaterialUsage
 	MATUSAGE_StaticLighting,
 	MATUSAGE_MorphTargets,
 	MATUSAGE_SplineMesh,
-	MATUSAGE_Landscape,
 	MATUSAGE_InstancedStaticMeshes,
 	MATUSAGE_Clothing,
 	MATUSAGE_MAX,
@@ -637,4 +635,4 @@ private:
 /** Helper function to serialize inline shader maps for the given material resources. */
 extern void SerializeInlineShaderMaps(const TMap<const class ITargetPlatform*, TArray<FMaterialResource*>>* PlatformMaterialResourcesToSave, FArchive& Ar, TArray<FMaterialResource>& OutLoadedResources);
 /** Helper function to process (register) serialized inline shader maps for the given material resources. */
-extern void ProcessSerializedInlineShaderMaps(TArray<FMaterialResource>& LoadedResources, FMaterialResource* (&OutMaterialResourcesLoaded)[EMaterialQualityLevel::Num][ERHIFeatureLevel::Num]);
\ No newline at end of file
+extern void ProcessSerializedInlineShaderMaps(UMaterialInterface* Owner, TArray<FMaterialResource>& LoadedResources, FMaterialResource* (&OutMaterialResourcesLoaded)[EMaterialQualityLevel::Num][ERHIFeatureLevel::Num]);
\ No newline at end of file
diff --git a/Engine/Source/Runtime/Engine/Private/AI/Navigation/RecastNavMeshGenerator.cpp b/Engine/Source/Runtime/Engine/Private/AI/Navigation/RecastNavMeshGenerator.cpp
--- a/Engine/Source/Runtime/Engine/Private/AI/Navigation/RecastNavMeshGenerator.cpp
+++ b/Engine/Source/Runtime/Engine/Private/AI/Navigation/RecastNavMeshGenerator.cpp
@@ -504,13 +504,16 @@ void ExportPxHeightField(PxHeightField const * const HeightField, const FTransfo
 			const PxHeightFieldSample& Sample = HFSamples[SampleIdx];
 			const bool HoleQuad = (Sample.materialIndex0 == PxHeightFieldMaterial::eHOLE);
 
-			IndexBuffer.Add(VertOffset + I00);
-			IndexBuffer.Add(VertOffset + (HoleQuad ? I00 : I11));
-			IndexBuffer.Add(VertOffset + (HoleQuad ? I00 : I10));
+			if (!HoleQuad)
+			{
+				IndexBuffer.Add(VertOffset + I00);
+				IndexBuffer.Add(VertOffset + I11);
+				IndexBuffer.Add(VertOffset + I10);
 
-			IndexBuffer.Add(VertOffset + I00);
-			IndexBuffer.Add(VertOffset + (HoleQuad ? I00 : I01));
-			IndexBuffer.Add(VertOffset + (HoleQuad ? I00 : I11));
+				IndexBuffer.Add(VertOffset + I00);
+				IndexBuffer.Add(VertOffset + I01);
+				IndexBuffer.Add(VertOffset + I11);
+			}
 		}
 	}
 }
diff --git a/Engine/Source/Runtime/Engine/Private/GameplayStatics.cpp b/Engine/Source/Runtime/Engine/Private/GameplayStatics.cpp
--- a/Engine/Source/Runtime/Engine/Private/GameplayStatics.cpp
+++ b/Engine/Source/Runtime/Engine/Private/GameplayStatics.cpp
@@ -20,6 +20,7 @@
 #include "Components/DecalComponent.h"
 #include "LandscapeProxy.h"
 #include "MessageLog.h"
+#include "Components/HierarchicalInstancedStaticMeshComponent.h"
 
 #define LOCTEXT_NAMESPACE "GameplayStatics"
 
diff --git a/Engine/Source/Runtime/Engine/Private/HierarchicalInstancedStaticMesh.cpp b/Engine/Source/Runtime/Engine/Private/HierarchicalInstancedStaticMesh.cpp
--- a/Engine/Source/Runtime/Engine/Private/HierarchicalInstancedStaticMesh.cpp
+++ b/Engine/Source/Runtime/Engine/Private/HierarchicalInstancedStaticMesh.cpp
@@ -2195,6 +2195,12 @@ void UHierarchicalInstancedStaticMeshComponent::BuildTree()
 			bDiscardAsyncBuildResults = true;
 		}
 	}
+	else
+	{
+		InstanceReorderTable.Empty();
+		SortedInstances.Empty();
+		RemovedInstances.Empty();
+	}
 
 	ReleasePerInstanceRenderData();
 }
@@ -2387,6 +2393,12 @@ void UHierarchicalInstancedStaticMeshComponent::BuildTreeAsync()
 			)
 			);
 	}
+	else
+	{
+		InstanceReorderTable.Empty();
+		SortedInstances.Empty();
+		RemovedInstances.Empty();
+	}
 }
 
 FPrimitiveSceneProxy* UHierarchicalInstancedStaticMeshComponent::CreateSceneProxy()
@@ -2514,19 +2526,30 @@ void UHierarchicalInstancedStaticMeshComponent::PostLoad()
 			}
 
 			// prune tree
-			StaticMesh->ConditionalPostLoad();
+			if (StaticMesh)
+			{
+				StaticMesh->ConditionalPostLoad();
+			}
 			BuildTree();
 		}
 	}
 
-	// For some reason we don't have a tree, or it is out of date. Build one now!
-	if (StaticMesh && PerInstanceSMData.Num() > 0 && (!ClusterTreePtr.IsValid() || ClusterTreePtr->Num() == 0 || (NumBuiltInstances != PerInstanceSMData.Num()) || UnbuiltInstanceBoundsList.Num() > 0 || GetLinkerUE4Version() < VER_UE4_REBUILD_HIERARCHICAL_INSTANCE_TREES))
+#if WITH_EDITOR
+	// If any of the data is out of sync, build the tree now!
+	if (InstanceReorderTable.Num() != PerInstanceSMData.Num() ||
+		NumBuiltInstances != PerInstanceSMData.Num() ||
+		UnbuiltInstanceBoundsList.Num() > 0 ||
+		GetLinkerUE4Version() < VER_UE4_REBUILD_HIERARCHICAL_INSTANCE_TREES)
 	{
-		UE_LOG(LogStaticMesh, Warning, TEXT("Rebuilding foliage, please resave map %s."), *GetFullName());
+		UE_LOG(LogStaticMesh, Warning, TEXT("Rebuilding hierarchical instanced mesh component, please resave map %s."), *GetFullName());
 		check(!IsAsyncBuilding());
-		StaticMesh->ConditionalPostLoad();
+		if (StaticMesh)
+		{
+			StaticMesh->ConditionalPostLoad();
+		}
 		BuildTree();
 	}
+#endif
 
 	if (CVarASyncInstaneBufferConversion.GetValueOnGameThread() > 0)
 	{
diff --git a/Engine/Source/Runtime/Engine/Private/InstancedStaticMesh.cpp b/Engine/Source/Runtime/Engine/Private/InstancedStaticMesh.cpp
--- a/Engine/Source/Runtime/Engine/Private/InstancedStaticMesh.cpp
+++ b/Engine/Source/Runtime/Engine/Private/InstancedStaticMesh.cpp
@@ -1852,26 +1852,28 @@ bool UInstancedStaticMeshComponent::IsInstanceSelected(int32 InInstanceIndex) co
 void UInstancedStaticMeshComponent::SelectInstance(bool bInSelected, int32 InInstanceIndex, int32 InInstanceCount)
 {
 #if WITH_EDITOR
-	if(PerInstanceSMData.Num() != SelectedInstances.Num())
+	if (InInstanceCount > 0)
 	{
-		SelectedInstances.Init(false, PerInstanceSMData.Num());
-	}
+		if (PerInstanceSMData.Num() != SelectedInstances.Num())
+		{
+			SelectedInstances.Init(false, PerInstanceSMData.Num());
+		}
 
-	check(SelectedInstances.IsValidIndex(InInstanceIndex));
-	check(SelectedInstances.IsValidIndex(InInstanceIndex + (InInstanceCount - 1)));
+		check(SelectedInstances.IsValidIndex(InInstanceIndex));
+		check(SelectedInstances.IsValidIndex(InInstanceIndex + (InInstanceCount - 1)));
 
-	for (int32 InstanceIndex = InInstanceIndex; InstanceIndex < InInstanceIndex + InInstanceCount; InstanceIndex++)
-	{
-		SelectedInstances[InstanceIndex] = bInSelected;
+		for (int32 InstanceIndex = InInstanceIndex; InstanceIndex < InInstanceIndex + InInstanceCount; InstanceIndex++)
+		{
+			SelectedInstances[InstanceIndex] = bInSelected;
+		}
 	}
-
 	SelectionStamp++;
 #endif
 }
 
 void UInstancedStaticMeshComponent::ClearInstanceSelection()
 {
-#if WITH_EDITOR	
+#if WITH_EDITOR
 	SelectedInstances.Empty();
 	SelectionStamp++;
 #endif
diff --git a/Engine/Source/Runtime/Engine/Private/Materials/Material.cpp b/Engine/Source/Runtime/Engine/Private/Materials/Material.cpp
--- a/Engine/Source/Runtime/Engine/Private/Materials/Material.cpp
+++ b/Engine/Source/Runtime/Engine/Private/Materials/Material.cpp
@@ -463,10 +463,13 @@ void SerializeInlineShaderMaps(const TMap<const ITargetPlatform*, TArray<FMateri
 	}
 }
 
-void ProcessSerializedInlineShaderMaps(TArray<FMaterialResource>& LoadedResources, FMaterialResource* (&OutMaterialResourcesLoaded)[EMaterialQualityLevel::Num][ERHIFeatureLevel::Num])
+void ProcessSerializedInlineShaderMaps(UMaterialInterface* Owner, TArray<FMaterialResource>& LoadedResources, FMaterialResource* (&OutMaterialResourcesLoaded)[EMaterialQualityLevel::Num][ERHIFeatureLevel::Num])
 {
 	check(IsInGameThread());
 
+	UMaterial* OwnerMaterial = Cast<UMaterial>(Owner);
+	UMaterialInstance* OwnerMaterialInstance = Cast<UMaterialInstance>(Owner);
+
 	for (FMaterialResource& Resource : LoadedResources)
 	{
 		Resource.RegisterInlineShaderMap();
@@ -515,7 +518,8 @@ void ProcessSerializedInlineShaderMaps(TArray<FMaterialResource>& LoadedResource
 					{
 						if (!OutMaterialResourcesLoaded[QualityLevelIndex][LoadedFeatureLevel])
 						{
-							OutMaterialResourcesLoaded[QualityLevelIndex][LoadedFeatureLevel] = new FMaterialResource();
+							OutMaterialResourcesLoaded[QualityLevelIndex][LoadedFeatureLevel] =
+								OwnerMaterialInstance ? OwnerMaterialInstance->AllocatePermutationResource() : OwnerMaterial->AllocateResource();
 						}
 						OutMaterialResourcesLoaded[QualityLevelIndex][LoadedFeatureLevel]->ReleaseShaderMap();
 						OutMaterialResourcesLoaded[QualityLevelIndex][LoadedFeatureLevel]->SetInlineShaderMap(LoadedShaderMap);
@@ -548,7 +552,8 @@ void ProcessSerializedInlineShaderMaps(TArray<FMaterialResource>& LoadedResource
 						{
 							if (!OutMaterialResourcesLoaded[QualityLevelIndex][LoadedFeatureLevel])
 							{
-								OutMaterialResourcesLoaded[QualityLevelIndex][LoadedFeatureLevel] = new FMaterialResource();
+								OutMaterialResourcesLoaded[QualityLevelIndex][LoadedFeatureLevel] =
+									OwnerMaterialInstance ? OwnerMaterialInstance->AllocatePermutationResource() : OwnerMaterial->AllocateResource();
 							}
 
 							OutMaterialResourcesLoaded[QualityLevelIndex][LoadedFeatureLevel]->SetInlineShaderMap(LoadedShaderMap);
@@ -943,7 +948,6 @@ bool UMaterial::GetUsageByFlag(EMaterialUsage Usage) const
 	switch(Usage)
 	{
 		case MATUSAGE_SkeletalMesh: UsageValue = bUsedWithSkeletalMesh; break;
-		case MATUSAGE_Landscape: UsageValue = bUsedWithLandscape; break;
 		case MATUSAGE_ParticleSprites: UsageValue = bUsedWithParticleSprites; break;
 		case MATUSAGE_BeamTrails: UsageValue = bUsedWithBeamTrails; break;
 		case MATUSAGE_MeshParticles: UsageValue = bUsedWithMeshParticles; break;
@@ -998,10 +1002,6 @@ void UMaterial::SetUsageByFlag(EMaterialUsage Usage, bool NewValue)
 		{
 			bUsedWithSkeletalMesh = NewValue; break;
 		}
-		case MATUSAGE_Landscape: 
-		{
-			bUsedWithLandscape = NewValue; break;
-		}
 		case MATUSAGE_ParticleSprites:
 		{
 			bUsedWithParticleSprites = NewValue; break;
@@ -1011,7 +1011,7 @@ void UMaterial::SetUsageByFlag(EMaterialUsage Usage, bool NewValue)
 			bUsedWithBeamTrails = NewValue; break;
 		}
 		case MATUSAGE_MeshParticles:
-		{	
+		{
 			bUsedWithMeshParticles = NewValue; break;
 		}
 		case MATUSAGE_StaticLighting:
@@ -1048,7 +1048,6 @@ FString UMaterial::GetUsageName(EMaterialUsage Usage) const
 	switch(Usage)
 	{
 		case MATUSAGE_SkeletalMesh: UsageName = TEXT("bUsedWithSkeletalMesh"); break;
-		case MATUSAGE_Landscape: UsageName = TEXT("bUsedWithLandscape"); break;
 		case MATUSAGE_ParticleSprites: UsageName = TEXT("bUsedWithParticleSprites"); break;
 		case MATUSAGE_BeamTrails: UsageName = TEXT("bUsedWithBeamTrails"); break;
 		case MATUSAGE_MeshParticles: UsageName = TEXT("bUsedWithMeshParticles"); break;
@@ -2392,7 +2391,7 @@ void UMaterial::PostLoad()
 	Super::PostLoad();
 
 	// Resources can be processed / registered now that we're back on the main thread
-	ProcessSerializedInlineShaderMaps(LoadedMaterialResources, MaterialResources);
+	ProcessSerializedInlineShaderMaps(this, LoadedMaterialResources, MaterialResources);
 	// Empty the lsit of loaded resources, we don't need it anymore
 	LoadedMaterialResources.Empty();
 
diff --git a/Engine/Source/Runtime/Engine/Private/Materials/MaterialInstance.cpp b/Engine/Source/Runtime/Engine/Private/Materials/MaterialInstance.cpp
--- a/Engine/Source/Runtime/Engine/Private/Materials/MaterialInstance.cpp
+++ b/Engine/Source/Runtime/Engine/Private/Materials/MaterialInstance.cpp
@@ -1573,6 +1573,11 @@ void UMaterialInstance::GetAllShaderMaps(TArray<FMaterialShaderMap*>& OutShaderM
 	}
 }
 
+FMaterialResource* UMaterialInstance::AllocatePermutationResource()
+{
+	return new FMaterialResource();
+}
+
 void UMaterialInstance::UpdatePermutationAllocations()
 {
 	if (bHasStaticPermutationResource)
@@ -1591,7 +1596,7 @@ void UMaterialInstance::UpdatePermutationAllocations()
 
 				if (!CurrentResource)
 				{
-					CurrentResource = new FMaterialResource();
+					CurrentResource = AllocatePermutationResource();
 				}
 
 				const bool bQualityLevelHasDifferentNodes = QualityLevelsUsed[QualityLevelIndex];
@@ -1662,7 +1667,7 @@ void UMaterialInstance::CacheResourceShadersForCooking(EShaderPlatform ShaderPla
 			// Cache all quality levels, unless they are all the same (due to using the same nodes), then just cache the high quality
 			if (bAnyQualityLevelUsed || QualityLevelIndex == EMaterialQualityLevel::High)
 			{
-				FMaterialResource* NewResource = new FMaterialResource();
+				FMaterialResource* NewResource = AllocatePermutationResource();
 				NewResource->SetMaterial(BaseMaterial, (EMaterialQualityLevel::Type)QualityLevelIndex, QualityLevelsUsed[QualityLevelIndex], (ERHIFeatureLevel::Type)TargetFeatureLevel, this);
 				ResourcesToCache.Add(NewResource);
 			}
@@ -1977,7 +1982,7 @@ void UMaterialInstance::PostLoad()
 	Super::PostLoad();
 
 	// Resources can be processed / registered now that we're back on the main thread
-	ProcessSerializedInlineShaderMaps(LoadedMaterialResources, StaticPermutationMaterialResources);
+	ProcessSerializedInlineShaderMaps(this, LoadedMaterialResources, StaticPermutationMaterialResources);
 	// Empty the lsit of loaded resources, we don't need it anymore
 	LoadedMaterialResources.Empty();
 
diff --git a/Engine/Source/Runtime/Engine/Private/Materials/MaterialShared.cpp b/Engine/Source/Runtime/Engine/Private/Materials/MaterialShared.cpp
--- a/Engine/Source/Runtime/Engine/Private/Materials/MaterialShared.cpp
+++ b/Engine/Source/Runtime/Engine/Private/Materials/MaterialShared.cpp
@@ -835,7 +835,7 @@ bool FMaterialResource::IsUsedWithSkeletalMesh() const
 
 bool FMaterialResource::IsUsedWithLandscape() const
 {
-	return Material->bUsedWithLandscape;
+	return false;
 }
 
 bool FMaterialResource::IsUsedWithParticleSystem() const
diff --git a/Engine/Source/Runtime/Foliage/Private/InstancedFoliage.cpp b/Engine/Source/Runtime/Foliage/Private/InstancedFoliage.cpp
--- a/Engine/Source/Runtime/Foliage/Private/InstancedFoliage.cpp
+++ b/Engine/Source/Runtime/Foliage/Private/InstancedFoliage.cpp
@@ -2641,14 +2641,16 @@ bool AInstancedFoliageActor::FoliageTrace(const UWorld* InWorld, FHitResult& Out
 	const FVector StartTrace = DesiredInstance.StartTrace - (Dir * DesiredInstance.TraceRadius);
 
 	TArray<FHitResult> Hits;
-
-	bool bInsideProceduralVolumeOrArentUsingOne = false;
 	FCollisionShape SphereShape;
 	SphereShape.SetSphere(DesiredInstance.TraceRadius);
 	InWorld->SweepMultiByObjectType(Hits, StartTrace, DesiredInstance.EndTrace, FQuat::Identity, FCollisionObjectQueryParams(ECC_WorldStatic), SphereShape, QueryParams);
 
 	for (const FHitResult& Hit : Hits)
 	{
+		const UPrimitiveComponent* HitComponent = Hit.GetComponent();
+
+		// don't place procedural foliage inside an AProceduralFoliageBlockingVolume
+		// this test is first because two of the tests below would otherwise cause the trace to ignore AProceduralFoliageBlockingVolume
 		if (DesiredInstance.PlacementMode == EFoliagePlacementMode::Procedural)
 		{
 			if (const AProceduralFoliageBlockingVolume* ProceduralFoliageBlockingVolume = Cast<AProceduralFoliageBlockingVolume>(Hit.Actor.Get()))
@@ -2659,52 +2661,46 @@ bool AInstancedFoliageActor::FoliageTrace(const UWorld* InWorld, FHitResult& Out
 					return false;
 				}
 			}
-			else if(Cast<AInstancedFoliageActor>(Hit.Actor.Get()))
-			{
-				return false;
-			}
-			else if (Cast<AProceduralFoliageVolume>(Hit.Actor.Get()))	//we never want to collide with our spawning volume
+			else if (Hit.GetActor()->IsA<AProceduralFoliageVolume>()) //we never want to collide with our spawning volume
 			{
 				continue;
 			}
-
-			if (DesiredInstance.ProceduralVolumeBodyInstance)
-			{
-				// We have a procedural volume, so lets make sure we are inside it.
-				bInsideProceduralVolumeOrArentUsingOne = DesiredInstance.ProceduralVolumeBodyInstance->OverlapTest(Hit.ImpactPoint, FQuat::Identity, FCollisionShape::MakeSphere(1.f));	//make sphere of 1cm radius to test if we're in the procedural volume
-			}
-			else
-			{
-				// We have no procedural volume, so we aren't using one.
-				bInsideProceduralVolumeOrArentUsingOne = true;
-			}
 		}
-		else
+
+		// Don't place foliage on foliage
+		if (Hit.GetActor()->IsA<AInstancedFoliageActor>())
 		{
-			// Not procedural, so we aren't using a procedural volume.
-			bInsideProceduralVolumeOrArentUsingOne = true;
+			return false;
 		}
-			
-		// In the editor traces can hit "No Collision" type actors, so ugh.
-		const FBodyInstance* BodyInstance = Hit.Component->GetBodyInstance();
+
+		// In the editor traces can hit "No Collision" type actors, so ugh. (ignore these)
+		const FBodyInstance* BodyInstance = HitComponent->GetBodyInstance();
 		if (BodyInstance->GetCollisionEnabled() != ECollisionEnabled::QueryAndPhysics || BodyInstance->GetResponseToChannel(ECC_WorldStatic) != ECR_Block)
 		{
 			continue;
 		}
 
-		const UPrimitiveComponent* HitComponent = Hit.Component.Get();
+		// Don't place foliage on invisible walls / triggers / volumes
+		if (HitComponent->IsA<UBrushComponent>())
+		{
+			continue;
+		}
 
-		if (HitComponent && HitComponent->GetComponentLevel())
+		if (FilterFunc && FilterFunc(HitComponent) == false)
 		{
-			if (FilterFunc && FilterFunc(HitComponent) == false)
-			{
-				// supplied filter does not like this component, so keep iterating
-				continue;
-			}
-			
-			OutHit = Hit;
-			return bInsideProceduralVolumeOrArentUsingOne;
+			// supplied filter does not like this component, so keep iterating
+			continue;
+		}
+
+		bool bInsideProceduralVolumeOrArentUsingOne = true;
+		if (DesiredInstance.PlacementMode == EFoliagePlacementMode::Procedural && DesiredInstance.ProceduralVolumeBodyInstance)
+		{
+			// We have a procedural volume, so lets make sure we are inside it.
+			bInsideProceduralVolumeOrArentUsingOne = DesiredInstance.ProceduralVolumeBodyInstance->OverlapTest(Hit.ImpactPoint, FQuat::Identity, FCollisionShape::MakeSphere(1.f));	//make sphere of 1cm radius to test if we're in the procedural volume
 		}
+
+		OutHit = Hit;
+		return bInsideProceduralVolumeOrArentUsingOne;
 	}
 
 	return false;
@@ -2763,12 +2759,13 @@ bool AInstancedFoliageActor::CheckCollisionWithWorld(const UWorld* InWorld, cons
 }
 
 FPotentialInstance::FPotentialInstance(FVector InHitLocation, FVector InHitNormal, UPrimitiveComponent* InHitComponent, float InHitWeight, const FDesiredFoliageInstance& InDesiredInstance)
-: HitLocation(InHitLocation)
-, HitNormal(InHitNormal)
-, HitComponent(InHitComponent)
-, HitWeight(InHitWeight)
-, DesiredInstance(InDesiredInstance)
-{}
+	: HitLocation(InHitLocation)
+	, HitNormal(InHitNormal)
+	, HitComponent(InHitComponent)
+	, HitWeight(InHitWeight)
+	, DesiredInstance(InDesiredInstance)
+{
+}
 
 bool FPotentialInstance::PlaceInstance(const UWorld* InWorld, const UFoliageType* Settings, FFoliageInstance& Inst, bool bSkipCollision)
 {
diff --git a/Engine/Source/Runtime/Landscape/Classes/Landscape.h b/Engine/Source/Runtime/Landscape/Classes/Landscape.h
--- a/Engine/Source/Runtime/Landscape/Classes/Landscape.h
+++ b/Engine/Source/Runtime/Landscape/Classes/Landscape.h
@@ -6,10 +6,11 @@
 #include "CoreUObject.h"
 #include "Casts.h"
 #include "LazyObjectPtr.h"
+#include "Engine/EngineBaseTypes.h"
+#include "Engine/World.h"
 
 #include "LandscapeProxy.h"
 #include "LandscapeInfoMap.h"
-#include "Engine/World.h"
 
 #include "Landscape.generated.h"
 
@@ -26,10 +27,13 @@ enum ELandscapeSetupErrors
 	LSE_MAX,
 };
 
-UCLASS(Placeable, hidecategories = LandscapeProxy, showcategories = (Display, Movement, Collision, Lighting, LOD, Input), MinimalAPI)
+UCLASS(MinimalAPI, showcategories=(Display, Movement, Collision, Lighting, LOD, Input))
 class ALandscape : public ALandscapeProxy
 {
-	GENERATED_UCLASS_BODY()
+	GENERATED_BODY()
+
+public:
+	ALandscape(const FObjectInitializer& ObjectInitializer);
 
 	// Make a key for XYtoComponentMap
 	static FIntPoint MakeKey( int32 X, int32 Y ) { return FIntPoint(X, Y); }
@@ -44,8 +48,6 @@ class ALandscape : public ALandscapeProxy
 	//~ Begin ALandscapeProxy Interface
 	LANDSCAPE_API virtual ALandscape* GetLandscapeActor() override;
 #if WITH_EDITOR
-	virtual UMaterialInterface* GetLandscapeMaterial() const override;
-	virtual UMaterialInterface* GetLandscapeHoleMaterial() const override;
 	//~ End ALandscapeProxy Interface
 
 	LANDSCAPE_API bool HasAllComponent(); // determine all component is in this actor
@@ -66,6 +68,7 @@ class ALandscape : public ALandscapeProxy
 	virtual void PreSave(const class ITargetPlatform* TargetPlatform) override;
 	virtual void PostEditChangeProperty(FPropertyChangedEvent& PropertyChangedEvent) override;
 	virtual void PostEditMove(bool bFinished) override;
+	virtual void PostEditImport() override;
 #endif
 	virtual void PostLoad() override;
 	//~ End UObject Interface
diff --git a/Engine/Source/Runtime/Landscape/Classes/LandscapeComponent.h b/Engine/Source/Runtime/Landscape/Classes/LandscapeComponent.h
--- a/Engine/Source/Runtime/Landscape/Classes/LandscapeComponent.h
+++ b/Engine/Source/Runtime/Landscape/Classes/LandscapeComponent.h
@@ -8,7 +8,6 @@
 #include "SceneTypes.h"
 #include "StaticLighting.h"
 #include "Components/PrimitiveComponent.h"
-#include "LandscapeGrassType.h"
 
 #include "LandscapeComponent.generated.h"
 
@@ -19,6 +18,7 @@ class ALandscapeProxy;
 class ALandscape;
 class ULandscapeHeightfieldCollisionComponent;
 class ULandscapeComponent;
+class ULandscapeGrassType;
 
 struct FEngineShowFlags;
 struct FConvexVolume;
@@ -26,6 +26,7 @@ struct FLandscapeEditDataInterface;
 struct FLandscapeEditToolRenderData;
 struct FLandscapeWeightmapUsage;
 struct FLandscapeTextureDataInfo;
+struct FLandscapeComponentGrassData;
 
 class FLandscapeComponentDerivedData
 {
@@ -100,14 +101,9 @@ struct FWeightmapLayerAllocationInfo
 	UPROPERTY()
 	uint8 WeightmapTextureChannel;
 
-	/** Only relevant in non-editor builds, this indicates which channel in the data array is this layer...must be > 1 to be valid, the first two are height **/
-	UPROPERTY()
-	uint8 GrassMapChannelIndex;
-
 	FWeightmapLayerAllocationInfo()
 		: WeightmapTextureIndex(0)
 		, WeightmapTextureChannel(0)
-		, GrassMapChannelIndex(0)
 	{
 	}
 
@@ -116,7 +112,6 @@ struct FWeightmapLayerAllocationInfo
 		:	LayerInfo(InLayerInfo)
 		,	WeightmapTextureIndex(255)	// Indicates an invalid allocation
 		,	WeightmapTextureChannel(255)
-		,	GrassMapChannelIndex(0) // Indicates an invalid allocation
 	{
 	}
 	
@@ -288,7 +283,7 @@ public:
 	FGuid BakedTextureMaterialGuid;
 
 	/** Pre-baked Base Color texture for use by distance field GI */
-	UPROPERTY(EditAnywhere, BlueprintReadOnly, Category = BakedTextures)
+	UPROPERTY(VisibleAnywhere, BlueprintReadOnly, Category = BakedTextures)
 	UTexture2D* GIBakedBaseColorTexture;
 
 #if WITH_EDITORONLY_DATA
@@ -332,8 +327,6 @@ public:
 	/** Grass data for generation **/
 	TSharedRef<FLandscapeComponentGrassData, ESPMode::ThreadSafe> GrassData;
 
-	virtual ~ULandscapeComponent();
-
 	//~ Begin UObject Interface.	
 	virtual void PostInitProperties() override;	
 	virtual void Serialize(FArchive& Ar) override;
@@ -393,10 +386,10 @@ public:
 	LANDSCAPE_API ULandscapeInfo* GetLandscapeInfo(bool bSpawnNewActor = true) const;
 
 	/** @todo document */
-	LANDSCAPE_API void DeleteLayer(ULandscapeLayerInfoObject* LayerInfo, FLandscapeEditDataInterface* LandscapeEdit);
+	LANDSCAPE_API void DeleteLayer(ULandscapeLayerInfoObject* LayerInfo, FLandscapeEditDataInterface& LandscapeEdit);
+	LANDSCAPE_API void FillLayer(ULandscapeLayerInfoObject* LayerInfo, FLandscapeEditDataInterface& LandscapeEdit);
+	LANDSCAPE_API void ReplaceLayer(ULandscapeLayerInfoObject* FromLayerInfo, ULandscapeLayerInfoObject* ToLayerInfo, FLandscapeEditDataInterface& LandscapeEdit);
 
-	LANDSCAPE_API void ReplaceLayer(ULandscapeLayerInfoObject* FromLayerInfo, ULandscapeLayerInfoObject* ToLayerInfo, FLandscapeEditDataInterface* LandscapeEdit);
-	
 	void GeneratePlatformVertexData();
 	void GeneratePlatformPixelData();
 
diff --git a/Engine/Source/Runtime/Landscape/Classes/LandscapeInfo.h b/Engine/Source/Runtime/Landscape/Classes/LandscapeInfo.h
--- a/Engine/Source/Runtime/Landscape/Classes/LandscapeInfo.h
+++ b/Engine/Source/Runtime/Landscape/Classes/LandscapeInfo.h
@@ -6,6 +6,7 @@
 // Forward declarations
 class ALandscapeProxy;
 class ALandscape;
+class ALandscapeStreamingProxy;
 class ULandscapeComponent;
 class ULandscapeLayerInfoObject;
 class UMaterialInstanceConstant;
@@ -126,7 +127,7 @@ public:
 	TMap<FIntPoint, FLandscapeAddCollision> XYtoAddCollisionMap;
 #endif
 
-	TSet<ALandscapeProxy*> Proxies;
+	TSet<ALandscapeStreamingProxy*> Proxies;
 
 private:
 	TSet<ULandscapeComponent*> SelectedComponents;
diff --git a/Engine/Source/Runtime/Landscape/Classes/LandscapeMaterialInstanceConstant.h b/Engine/Source/Runtime/Landscape/Classes/LandscapeMaterialInstanceConstant.h
--- a/Engine/Source/Runtime/Landscape/Classes/LandscapeMaterialInstanceConstant.h
+++ b/Engine/Source/Runtime/Landscape/Classes/LandscapeMaterialInstanceConstant.h
@@ -1,6 +1,8 @@
 // Copyright 1998-2016 Epic Games, Inc. All Rights Reserved.
 
 #pragma once
+
+#include "Materials/MaterialInstanceConstant.h"
 #include "LandscapeMaterialInstanceConstant.generated.h"
 
 UCLASS(MinimalAPI)
@@ -10,6 +12,9 @@ class ULandscapeMaterialInstanceConstant : public UMaterialInstanceConstant
 
 	UPROPERTY()
 	uint32 bIsLayerThumbnail:1;
+
+	virtual FMaterialResource* AllocatePermutationResource() override;
+	virtual bool HasOverridenBaseProperties() const override;
 };
 
 
diff --git a/Engine/Source/Runtime/Landscape/Classes/LandscapeMeshProxyComponent.h b/Engine/Source/Runtime/Landscape/Classes/LandscapeMeshProxyComponent.h
--- a/Engine/Source/Runtime/Landscape/Classes/LandscapeMeshProxyComponent.h
+++ b/Engine/Source/Runtime/Landscape/Classes/LandscapeMeshProxyComponent.h
@@ -2,17 +2,19 @@
 
 #pragma once
 
-#include "SceneTypes.h"
-#include "StaticLighting.h"
-#include "Components/PrimitiveComponent.h"
-#include "LandscapeProxy.h"
+#include "Components/StaticMeshComponent.h"
 
 #include "LandscapeMeshProxyComponent.generated.h"
 
+class ALandscapeProxy;
+
 UCLASS(MinimalAPI)
 class ULandscapeMeshProxyComponent : public UStaticMeshComponent
 {
-	GENERATED_UCLASS_BODY()
+	GENERATED_BODY()
+
+public:
+	ULandscapeMeshProxyComponent(const FObjectInitializer& ObjectInitializer);
 
 private:
 	/* The landscape this proxy was generated for */
@@ -26,24 +28,9 @@ private:
 	/* LOD level proxy was generated for */
 	UPROPERTY()
 	int8 ProxyLOD;
+
 public:
-	void InitializeForLandscape(ALandscapeProxy* Landscape, int8 InProxyLOD)
-	{
-		LandscapeGuid = Landscape->GetLandscapeGuid();
-
-		for (ULandscapeComponent* Component : Landscape->LandscapeComponents)
-		{
-			if (Component)
-			{
-				ProxyComponentBases.Add(Component->GetSectionBase() / Component->ComponentSizeQuads);
-			}
-		}
-
-		if (InProxyLOD != INDEX_NONE)
-		{
-			ProxyLOD = FMath::Clamp<int32>(InProxyLOD, 0, FMath::CeilLogTwo(Landscape->SubsectionSizeQuads + 1) - 1);
-		}
-	}
+	LANDSCAPE_API void InitializeForLandscape(ALandscapeProxy* Landscape, int8 InProxyLOD);
 	
 	virtual FPrimitiveSceneProxy* CreateSceneProxy() override;
 };
diff --git a/Engine/Source/Runtime/Landscape/Classes/LandscapeProxy.h b/Engine/Source/Runtime/Landscape/Classes/LandscapeProxy.h
--- a/Engine/Source/Runtime/Landscape/Classes/LandscapeProxy.h
+++ b/Engine/Source/Runtime/Landscape/Classes/LandscapeProxy.h
@@ -2,34 +2,22 @@
 
 #pragma once
 
-#include "LandscapeInfo.h"
-
-#include "Engine/EngineTypes.h"
-#include "Engine/Texture.h"
-#include "Engine/TextureDefines.h"
-
-#include "GameFramework/Actor.h"
 #include "PhysicsEngine/BodyInstance.h"
-#include "Materials/Material.h"
-#include "Materials/MaterialInstanceConstant.h"
-#include "Components/HierarchicalInstancedStaticMeshComponent.h"
-#include "LandscapeComponent.h"
-#include "LandscapeLayerInfoObject.h"
-#include "LandscapeGrassType.h"
-#include "Tickable.h"
 #include "AI/Navigation/NavigationTypes.h"
-
+#include "Engine/Texture.h"
 #include "LandscapeProxy.generated.h"
 
-class ULandscapeMaterialInstanceConstant;
+class ULandscapeInfo;
+class ULandscapeComponent;
 class ULandscapeSplinesComponent;
+class ULandscapeGrassType;
+class UHierarchicalInstancedStaticMeshComponent;
 class ULandscapeHeightfieldCollisionComponent;
-class UMaterialInterface;
-class UTexture2D;
-class ALandscape;
+class UMaterialInstanceConstant;
+class ULandscapeMaterialInstanceConstant;
 class ALandscapeProxy;
-class ULandscapeComponent;
-class USplineComponent;
+class ALandscape;
+class ALandscapeStreamingProxy;
 struct FLandscapeInfoLayerSettings;
 struct FAsyncGrassBuilder;
 
@@ -44,18 +32,18 @@ struct FLandscapeWeightmapUsage
 
 	FLandscapeWeightmapUsage()
 	{
-		ChannelUsage[0] = NULL;
-		ChannelUsage[1] = NULL;
-		ChannelUsage[2] = NULL;
-		ChannelUsage[3] = NULL;
+		ChannelUsage[0] = nullptr;
+		ChannelUsage[1] = nullptr;
+		ChannelUsage[2] = nullptr;
+		ChannelUsage[3] = nullptr;
 	}
 	friend FArchive& operator<<( FArchive& Ar, FLandscapeWeightmapUsage& U );
 	int32 FreeChannelCount() const
 	{
-		return	((ChannelUsage[0] == NULL) ? 1 : 0) + 
-				((ChannelUsage[1] == NULL) ? 1 : 0) + 
-				((ChannelUsage[2] == NULL) ? 1 : 0) + 
-				((ChannelUsage[3] == NULL) ? 1 : 0);
+		return	((ChannelUsage[0] == nullptr) ? 1 : 0) +
+				((ChannelUsage[1] == nullptr) ? 1 : 0) +
+				((ChannelUsage[2] == nullptr) ? 1 : 0) +
+				((ChannelUsage[3] == nullptr) ? 1 : 0);
 	}
 };
 
@@ -72,7 +60,7 @@ struct FLandscapeEditorLayerSettings
 	FString ReimportLayerFilePath;
 
 	FLandscapeEditorLayerSettings()
-		: LayerInfoObj(NULL)
+		: LayerInfoObj(nullptr)
 		, ReimportLayerFilePath()
 	{
 	}
@@ -100,7 +88,7 @@ struct FLandscapeLayerStruct
 
 #if WITH_EDITORONLY_DATA
 	UPROPERTY(transient)
-	UMaterialInstanceConstant* ThumbnailMIC;
+	ULandscapeMaterialInstanceConstant* ThumbnailMIC;
 
 	UPROPERTY()
 	ALandscapeProxy* Owner;
@@ -116,10 +104,10 @@ struct FLandscapeLayerStruct
 #endif // WITH_EDITORONLY_DATA
 
 	FLandscapeLayerStruct()
-		: LayerInfoObj(NULL)
+		: LayerInfoObj(nullptr)
 #if WITH_EDITORONLY_DATA
-		, ThumbnailMIC(NULL)
-		, Owner(NULL)
+		, ThumbnailMIC(nullptr)
+		, Owner(nullptr)
 		, DebugColorChannel(0)
 		, bSelected(false)
 		, SourceFilePath()
@@ -155,21 +143,17 @@ struct FLandscapeImportLayerInfo
 	UPROPERTY(Category="Import", EditAnywhere)
 	ULandscapeLayerInfoObject* LayerInfo;
 
-	UPROPERTY(Category="Import", VisibleAnywhere)
-	UMaterialInstanceConstant* ThumbnailMIC; // Optional
-
 	UPROPERTY(Category="Import", EditAnywhere)
 	FString SourceFilePath; // Optional
 	
 	// Raw weightmap data
-	TArray<uint8> LayerData;		
+	TArray<uint8> LayerData;
 #endif
 
 #if WITH_EDITOR
 	FLandscapeImportLayerInfo(FName InLayerName = NAME_None)
 	:	LayerName(InLayerName)
-	,	LayerInfo(NULL)
-	,	ThumbnailMIC(NULL)
+	,	LayerInfo(nullptr)
 	,	SourceFilePath("")
 	{
 	}
@@ -305,12 +289,7 @@ public:
 	FCachedLandscapeFoliage::FGrassCompKey Key;
 	TWeakObjectPtr<UHierarchicalInstancedStaticMeshComponent> Foliage;
 
-	FAsyncGrassTask(FAsyncGrassBuilder* InBuilder, const FCachedLandscapeFoliage::FGrassCompKey& InKey, UHierarchicalInstancedStaticMeshComponent* InFoliage)
-		: Builder(InBuilder)
-		, Key(InKey)
-		, Foliage(InFoliage)
-	{
-	}
+	FAsyncGrassTask(FAsyncGrassBuilder* InBuilder, const FCachedLandscapeFoliage::FGrassCompKey& InKey, UHierarchicalInstancedStaticMeshComponent* InFoliage);
 	void DoWork();
 
 	FORCEINLINE TStatId GetStatId() const
@@ -321,10 +300,13 @@ public:
 	~FAsyncGrassTask();
 };
 
-UCLASS(NotPlaceable, NotBlueprintable, hidecategories=(Display, Attachment, Physics, Debug, Lighting, LOD), showcategories=(Lighting, Rendering, "Utilities|Transformation"), MinimalAPI)
+UCLASS(Abstract, MinimalAPI, hidecategories=(Display, Attachment, Physics, Debug, Lighting, LOD), showcategories=(Lighting, Rendering, "Utilities|Transformation"))
 class ALandscapeProxy : public AActor
 {
-	GENERATED_UCLASS_BODY()
+	GENERATED_BODY()
+
+public:
+	ALandscapeProxy(const FObjectInitializer& ObjectInitializer);
 
 	virtual ~ALandscapeProxy();
 
@@ -428,9 +410,6 @@ public:
 	UPROPERTY(EditAnywhere, Category=Lighting)
 	float StaticLightingResolution;
 
-	UPROPERTY(EditAnywhere, Category=LandscapeProxy)
-	TLazyObjectPtr<ALandscape> LandscapeActor;
-
 	UPROPERTY(EditAnywhere, Category=Lighting, meta=(DisplayName = "Static Shadow"))
 	uint32 bCastStaticShadow:1;
 
@@ -448,9 +427,6 @@ public:
 	UPROPERTY(EditAnywhere, AdvancedDisplay, Category=Lighting)
 	uint32 bUseMaterialPositionOffsetInStaticLighting:1;
 
-	UPROPERTY()
-	uint32 bIsProxy:1;
-
 #if WITH_EDITORONLY_DATA
 	UPROPERTY(transient)
 	uint32 bIsMovingToLevel:1;    // Check for the Move to Current Level case
@@ -575,7 +551,7 @@ public:
 
 	FGuid GetLandscapeGuid() const { return LandscapeGuid; }
 	void SetLandscapeGuid(const FGuid& Guid) { LandscapeGuid = Guid; }
-	virtual ALandscape* GetLandscapeActor();
+	virtual ALandscape* GetLandscapeActor() PURE_VIRTUAL(GetLandscapeActor, return nullptr;)
 
 	/* Per-frame call to update dynamic grass placement and render grassmaps */
 	void TickGrass();
@@ -619,7 +595,6 @@ public:
 
 #if WITH_EDITOR
 	virtual void PostEditChangeProperty(FPropertyChangedEvent& PropertyChangedEvent) override;
-	virtual void PostEditChangeChainProperty(FPropertyChangedChainEvent& PropertyChangedEvent) override;
 	virtual void PreEditUndo() override;
 	virtual void PostEditUndo() override;
 	virtual void PostEditImport() override;
@@ -644,9 +619,6 @@ public:
 	// Changed Physical Material
 	LANDSCAPE_API void ChangedPhysMaterial();
 
-	// Check input Landscape actor is match for this LandscapeProxy (by GUID)
-	LANDSCAPE_API bool IsValidLandscapeActor(ALandscape* Landscape);
-
 	// Copy properties from parent Landscape actor
 	LANDSCAPE_API void GetSharedProperties(ALandscapeProxy* Landscape);
 	// Assign only mismatched properties and mark proxy package dirty
@@ -709,4 +681,3 @@ public:
 	LANDSCAPE_API void RemoveOverlappingComponent(ULandscapeComponent* Component);
 #endif
 };
-
diff --git a/Engine/Source/Runtime/Landscape/Classes/LandscapeStreamingProxy.h b/Engine/Source/Runtime/Landscape/Classes/LandscapeStreamingProxy.h
new file mode 100644
--- /dev/null
+++ b/Engine/Source/Runtime/Landscape/Classes/LandscapeStreamingProxy.h
@@ -0,0 +1,34 @@
+// Copyright 1998-2016 Epic Games, Inc. All Rights Reserved.
+
+#pragma once
+
+#include "LandscapeStreamingProxy.generated.h"
+
+UCLASS(MinimalAPI, notplaceable, NotBlueprintable)
+class ALandscapeStreamingProxy : public ALandscapeProxy
+{
+	GENERATED_BODY()
+
+public:
+	ALandscapeStreamingProxy(const FObjectInitializer& ObjectInitializer);
+
+	UPROPERTY(EditAnywhere, Category=LandscapeProxy)
+	TLazyObjectPtr<ALandscape> LandscapeActor;
+
+	//~ Begin UObject Interface
+#if WITH_EDITOR
+	void PostEditChangeProperty(FPropertyChangedEvent& PropertyChangedEvent) override;
+#endif
+	//~ End UObject Interface
+
+	//~ Begin ALandscapeBase Interface
+	virtual ALandscape* GetLandscapeActor() override;
+#if WITH_EDITOR
+	virtual UMaterialInterface* GetLandscapeMaterial() const override;
+	virtual UMaterialInterface* GetLandscapeHoleMaterial() const override;
+#endif
+	//~ End ALandscapeBase Interface
+
+	// Check input Landscape actor is match for this LandscapeProxy (by GUID)
+	bool IsValidLandscapeActor(ALandscape* Landscape);
+};
diff --git a/Engine/Source/Runtime/Landscape/Private/Landscape.cpp b/Engine/Source/Runtime/Landscape/Private/Landscape.cpp
--- a/Engine/Source/Runtime/Landscape/Private/Landscape.cpp
+++ b/Engine/Source/Runtime/Landscape/Private/Landscape.cpp
@@ -5,6 +5,7 @@ Landscape.cpp: Terrain rendering
 =============================================================================*/
 
 #include "Landscape.h"
+#include "LandscapeStreamingProxy.h"
 #include "EditorSupportDelegates.h"
 #include "LandscapeDataAccess.h"
 #include "LandscapeRender.h"
@@ -124,10 +125,6 @@ ULandscapeComponent::ULandscapeComponent(const FObjectInitializer& ObjectInitial
 	LpvBiasMultiplier = 0.0f; // Bias is 0 for landscape, since it's single sided
 }
 
-ULandscapeComponent::~ULandscapeComponent()
-{
-}
-
 void ULandscapeComponent::AddReferencedObjects(UObject* InThis, FReferenceCollector& Collector)
 {
 	ULandscapeComponent* This = CastChecked<ULandscapeComponent>(InThis);
@@ -529,7 +526,7 @@ void ULandscapeComponent::PostLoad()
 			RelativeLocation.X = CheckRelativeLocationX;
 			RelativeLocation.Y = CheckRelativeLocationY;
 		}
-						
+
 		// Remove standalone flags from data textures to ensure data is unloaded in the editor when reverting an unsaved level.
 		// Previous version of landscape set these flags on creation.
 		if (HeightmapTexture && HeightmapTexture->HasAnyFlags(RF_Standalone))
@@ -543,6 +540,13 @@ void ULandscapeComponent::PostLoad()
 				WeightmapTextures[Idx]->ClearFlags(RF_Standalone);
 			}
 		}
+
+		// Remove public flag from GI textures to stop them being visible in the content browser.
+		// Previous version of landscape set these flags on creation.
+		if (GIBakedBaseColorTexture && GIBakedBaseColorTexture->HasAnyFlags(RF_Public))
+		{
+			GIBakedBaseColorTexture->ClearFlags(RF_Public);
+		}
 	}
 #endif
 
@@ -585,21 +589,6 @@ void ULandscapeComponent::PostLoad()
 
 #endif // WITH_EDITOR
 
-ALandscape::ALandscape(const FObjectInitializer& ObjectInitializer)
-	: Super(ObjectInitializer)
-{
-	bIsProxy = false;
-
-#if WITH_EDITORONLY_DATA
-	bLockLocation = false;
-#endif // WITH_EDITORONLY_DATA
-}
-
-ALandscape* ALandscape::GetLandscapeActor()
-{
-	return this;
-}
-
 ALandscapeProxy::ALandscapeProxy(const FObjectInitializer& ObjectInitializer)
 	: Super(ObjectInitializer)
 	, bHasLandscapeGrass(true)
@@ -625,7 +614,6 @@ ALandscapeProxy::ALandscapeProxy(const FObjectInitializer& ObjectInitializer)
 
 	StaticLightingResolution = 1.0f;
 	StreamingDistanceMultiplier = 1.0f;
-	bIsProxy = true;
 	MaxLODLevel = -1;
 #if WITH_EDITORONLY_DATA
 	bLockLocation = true;
@@ -668,15 +656,32 @@ ALandscapeProxy::ALandscapeProxy(const FObjectInitializer& ObjectInitializer)
 #endif
 }
 
-ALandscape* ALandscapeProxy::GetLandscapeActor()
+ALandscape::ALandscape(const FObjectInitializer& ObjectInitializer)
+	: Super(ObjectInitializer)
 {
 #if WITH_EDITORONLY_DATA
-	return LandscapeActor.Get();
-#else
-	return 0;
+	bLockLocation = false;
 #endif // WITH_EDITORONLY_DATA
 }
 
+ALandscapeStreamingProxy::ALandscapeStreamingProxy(const FObjectInitializer& ObjectInitializer)
+	: Super(ObjectInitializer)
+{
+#if WITH_EDITORONLY_DATA
+	bLockLocation = true;
+#endif // WITH_EDITORONLY_DATA
+}
+
+ALandscape* ALandscape::GetLandscapeActor()
+{
+	return this;
+}
+
+ALandscape* ALandscapeStreamingProxy::GetLandscapeActor()
+{
+	return LandscapeActor.Get();
+}
+
 #if WITH_EDITOR
 ULandscapeInfo* ALandscapeProxy::GetLandscapeInfo(bool bSpawnNewActor /*= true*/)
 {
@@ -1282,7 +1287,7 @@ bool ULandscapeInfo::UpdateLayerInfoMap(ALandscapeProxy* Proxy /*= NULL*/, bool
 								// Valid layer infos take precedence over invalid ones
 								// Landscape Actors take precedence over Proxies
 								if ((bValid && !LayerSettings.bValid)
-									|| (bValid == LayerSettings.bValid && !Proxy->bIsProxy))
+									|| (bValid == LayerSettings.bValid && Proxy->IsA<ALandscape>()))
 								{
 									LayerSettings.Owner = Proxy;
 									LayerSettings.bValid = bValid;
@@ -1605,10 +1610,6 @@ UMaterialInterface* ALandscapeProxy::GetLandscapeMaterial() const
 	{
 		return LandscapeMaterial;
 	}
-	else if (LandscapeActor)
-	{
-		return LandscapeActor->GetLandscapeMaterial();
-	}
 	return UMaterial::GetDefaultMaterial(MD_Surface);
 }
 
@@ -1618,28 +1619,32 @@ UMaterialInterface* ALandscapeProxy::GetLandscapeHoleMaterial() const
 	{
 		return LandscapeHoleMaterial;
 	}
-	else if (ALandscape* Landscape = LandscapeActor.Get())
-	{
-		return Landscape->GetLandscapeHoleMaterial();
-	}
 	return nullptr;
 }
 
-UMaterialInterface* ALandscape::GetLandscapeMaterial() const
+UMaterialInterface* ALandscapeStreamingProxy::GetLandscapeMaterial() const
 {
 	if (LandscapeMaterial)
 	{
 		return LandscapeMaterial;
 	}
+	else if (LandscapeActor)
+	{
+		return LandscapeActor->GetLandscapeMaterial();
+	}
 	return UMaterial::GetDefaultMaterial(MD_Surface);
 }
 
-UMaterialInterface* ALandscape::GetLandscapeHoleMaterial() const
+UMaterialInterface* ALandscapeStreamingProxy::GetLandscapeHoleMaterial() const
 {
 	if (LandscapeHoleMaterial)
 	{
 		return LandscapeHoleMaterial;
 	}
+	else if (ALandscape* Landscape = LandscapeActor.Get())
+	{
+		return Landscape->GetLandscapeHoleMaterial();
+	}
 	return nullptr;
 }
 
@@ -1779,25 +1784,26 @@ void ULandscapeInfo::RegisterActor(ALandscapeProxy* Proxy, bool bMapCheck)
 	}
 
 	// register
-	ALandscape* Landscape = Cast<ALandscape>(Proxy);
-	if (Landscape)
+	if (ALandscape* Landscape = Cast<ALandscape>(Proxy))
 	{
 		LandscapeActor = Landscape;
 		// In world composition user is not allowed to move landscape in editor, only through WorldBrowser 
-		LandscapeActor->bLockLocation = (OwningWorld->WorldComposition != NULL);
+		LandscapeActor->bLockLocation = (OwningWorld->WorldComposition != nullptr);
 
 		// update proxies reference actor
-		for (auto It = Proxies.CreateConstIterator(); It; ++It)
+		for (ALandscapeStreamingProxy* StreamingProxy : Proxies)
 		{
-			(*It)->LandscapeActor = LandscapeActor;
-			(*It)->ConditionalAssignCommonProperties(Landscape);
+			StreamingProxy->LandscapeActor = LandscapeActor;
+			StreamingProxy->ConditionalAssignCommonProperties(Landscape);
 		}
 	}
 	else
 	{
-		Proxies.Add(Proxy);
-		Proxy->LandscapeActor = LandscapeActor;
-		Proxy->ConditionalAssignCommonProperties(LandscapeActor.Get());
+		ALandscapeStreamingProxy* StreamingProxy = CastChecked<ALandscapeStreamingProxy>(Proxy);
+
+		Proxies.Add(StreamingProxy);
+		StreamingProxy->LandscapeActor = LandscapeActor;
+		StreamingProxy->ConditionalAssignCommonProperties(LandscapeActor.Get());
 	}
 
 	//
@@ -1811,22 +1817,22 @@ void ULandscapeInfo::RegisterActor(ALandscapeProxy* Proxy, bool bMapCheck)
 
 void ULandscapeInfo::UnregisterActor(ALandscapeProxy* Proxy)
 {
-	ALandscape* Landscape = Cast<ALandscape>(Proxy);
-	if (Landscape)
+	if (ALandscape* Landscape = Cast<ALandscape>(Proxy))
 	{
 		check(LandscapeActor.Get() == Landscape);
 		LandscapeActor = nullptr;
 
 		// update proxies reference to landscape actor
-		for (auto It = Proxies.CreateConstIterator(); It; ++It)
+		for (ALandscapeStreamingProxy* StreamingProxy : Proxies)
 		{
-			(*It)->LandscapeActor = nullptr;
+			StreamingProxy->LandscapeActor = nullptr;
 		}
 	}
 	else
 	{
-		Proxies.Remove(Proxy);
-		Proxy->LandscapeActor = nullptr;
+		ALandscapeStreamingProxy* StreamingProxy = CastChecked<ALandscapeStreamingProxy>(Proxy);
+		Proxies.Remove(StreamingProxy);
+		StreamingProxy->LandscapeActor = nullptr;
 	}
 
 	// remove proxy components from the XY map
@@ -2396,12 +2402,29 @@ ALandscapeMeshProxyActor::ALandscapeMeshProxyActor(const FObjectInitializer& Obj
 // ULandscapeMeshProxyComponent
 //
 ULandscapeMeshProxyComponent::ULandscapeMeshProxyComponent(const FObjectInitializer& ObjectInitializer)
-: Super(ObjectInitializer)
+	: Super(ObjectInitializer)
 {
 }
 
-#if WITH_EDITOR
+void ULandscapeMeshProxyComponent::InitializeForLandscape(ALandscapeProxy* Landscape, int8 InProxyLOD)
+{
+	LandscapeGuid = Landscape->GetLandscapeGuid();
 
+	for (ULandscapeComponent* Component : Landscape->LandscapeComponents)
+	{
+		if (Component)
+		{
+			ProxyComponentBases.Add(Component->GetSectionBase() / Component->ComponentSizeQuads);
+		}
+	}
+
+	if (InProxyLOD != INDEX_NONE)
+	{
+		ProxyLOD = FMath::Clamp<int32>(InProxyLOD, 0, FMath::CeilLogTwo(Landscape->SubsectionSizeQuads + 1) - 1);
+	}
+}
+
+#if WITH_EDITOR
 void ULandscapeComponent::SerializeStateHashes(FArchive& Ar)
 {
 	if (MaterialInstance)
@@ -2572,7 +2595,7 @@ void ALandscapeProxy::UpdateBakedTextures()
 							NumGenerated++;
 						}
 					}
-					UTexture2D* AtlasTexture = FMaterialUtilities::CreateTexture(GetOutermost(), HeightmapTexture->GetName() + TEXT("_BaseColor"), AtlasSize, AtlasSamples, TC_Default, TEXTUREGROUP_World, RF_Public, true, CombinedStateId);
+					UTexture2D* AtlasTexture = FMaterialUtilities::CreateTexture(GetOutermost(), HeightmapTexture->GetName() + TEXT("_BaseColor"), AtlasSize, AtlasSamples, TC_Default, TEXTUREGROUP_World, RF_NoFlags, true, CombinedStateId);
 					AtlasTexture->MarkPackageDirty();
 
 					for (ULandscapeComponent* Component : Info.Components)
diff --git a/Engine/Source/Runtime/Landscape/Private/LandscapeCollision.cpp b/Engine/Source/Runtime/Landscape/Private/LandscapeCollision.cpp
--- a/Engine/Source/Runtime/Landscape/Private/LandscapeCollision.cpp
+++ b/Engine/Source/Runtime/Landscape/Private/LandscapeCollision.cpp
@@ -1485,11 +1485,20 @@ bool ULandscapeHeightfieldCollisionComponent::DoCustomNavigableGeometryExport(FN
 {
 	check(IsInGameThread());
 #if WITH_PHYSX
-	if (IsValidRef(HeightfieldRef) && HeightfieldRef->RBHeightfield != nullptr)
+	if (IsValidRef(HeightfieldRef) && HeightfieldRef->RBHeightfield)
 	{
 		FTransform HFToW = ComponentToWorld;
-		HFToW.MultiplyScale3D(FVector(CollisionScale, CollisionScale, LANDSCAPE_ZSCALE));
-		GeomExport.ExportPxHeightField(HeightfieldRef->RBHeightfield, HFToW);
+		if (HeightfieldRef->RBHeightfieldSimple)
+		{
+			const float SimpleCollisionScale = CollisionScale * CollisionSizeQuads / SimpleCollisionSizeQuads;
+			HFToW.MultiplyScale3D(FVector(SimpleCollisionScale, SimpleCollisionScale, LANDSCAPE_ZSCALE));
+			GeomExport.ExportPxHeightField(HeightfieldRef->RBHeightfieldSimple, HFToW);
+		}
+		else
+		{
+			HFToW.MultiplyScale3D(FVector(CollisionScale, CollisionScale, LANDSCAPE_ZSCALE));
+			GeomExport.ExportPxHeightField(HeightfieldRef->RBHeightfield, HFToW);
+		}
 	}
 #endif// WITH_PHYSX
 	return false;
diff --git a/Engine/Source/Runtime/Landscape/Private/LandscapeEdit.cpp b/Engine/Source/Runtime/Landscape/Private/LandscapeEdit.cpp
--- a/Engine/Source/Runtime/Landscape/Private/LandscapeEdit.cpp
+++ b/Engine/Source/Runtime/Landscape/Private/LandscapeEdit.cpp
@@ -5,7 +5,7 @@ LandscapeEdit.cpp: Landscape editing
 =============================================================================*/
 
 #include "Landscape.h"
-#include "LandscapeProxy.h"
+#include "LandscapeStreamingProxy.h"
 #include "Materials/MaterialInstanceDynamic.h"
 #include "Materials/MaterialExpressionLandscapeVisibilityMask.h"
 #include "Materials/MaterialExpressionLandscapeLayerWeight.h"
@@ -52,7 +52,7 @@ DEFINE_LOG_CATEGORY(LogLandscape);
 // Channel remapping
 extern const size_t ChannelOffsets[4];
 
-ULandscapeLayerInfoObject* ALandscapeProxy::VisibilityLayer = NULL;
+ULandscapeLayerInfoObject* ALandscapeProxy::VisibilityLayer = nullptr;
 
 void ULandscapeComponent::Init(int32 InBaseX, int32 InBaseY, int32 InComponentSizeQuads, int32 InNumSubsections, int32 InSubsectionSizeQuads)
 {
@@ -111,60 +111,35 @@ ULandscapeMaterialInstanceConstant* ALandscapeProxy::GetLayerThumbnailMIC(UMater
 		LandscapeMaterial = Proxy ? Proxy->GetLandscapeMaterial() : UMaterial::GetDefaultMaterial(MD_Surface);
 	}
 
-	UMaterialInstanceConstant* CombinationMaterialInstance = NULL;
-	FString LayerKey = LandscapeMaterial->GetPathName() + FString::Printf(TEXT("_%s_0"), *LayerName.ToString());
-	if (Proxy)
-	{
-		CombinationMaterialInstance = Proxy->MaterialInstanceConstantMap.FindRef(*LayerKey);
-	}
+	FlushRenderingCommands();
+
+	ULandscapeMaterialInstanceConstant* MaterialInstance = NewObject<ULandscapeMaterialInstanceConstant>(GetTransientPackage());
+	MaterialInstance->bIsLayerThumbnail = true;
+	MaterialInstance->SetParentEditorOnly(LandscapeMaterial);
+
+	FStaticParameterSet StaticParameters;
+	MaterialInstance->GetStaticParameterValues(StaticParameters);
 
-	if (CombinationMaterialInstance == NULL || CombinationMaterialInstance->Parent != LandscapeMaterial || (Proxy && Proxy->GetOutermost() != CombinationMaterialInstance->GetOutermost()))
+	for (int32 LayerParameterIdx = 0; LayerParameterIdx < StaticParameters.TerrainLayerWeightParameters.Num(); ++LayerParameterIdx)
 	{
-		FlushRenderingCommands();
-		UObject* MICOuter = GetTransientPackage();
-		if (Proxy)
+		FStaticTerrainLayerWeightParameter& LayerParameter = StaticParameters.TerrainLayerWeightParameters[LayerParameterIdx];
+		if (LayerParameter.ParameterName == LayerName)
 		{
-			MICOuter = Proxy->GetOutermost();
+			LayerParameter.WeightmapIndex = 0;
+			LayerParameter.bOverride = true;
 		}
-		CombinationMaterialInstance = NewObject<ULandscapeMaterialInstanceConstant>(MICOuter);
-		if (Proxy)
-		{
-			UE_LOG(LogLandscape, Log, TEXT("Looking for key %s, making new combination %s"), *LayerKey, *CombinationMaterialInstance->GetName());
-			Proxy->MaterialInstanceConstantMap.Add(*LayerKey, CombinationMaterialInstance);
-		}
-		CombinationMaterialInstance->SetParentEditorOnly(LandscapeMaterial);
-
-		FStaticParameterSet StaticParameters;
-		CombinationMaterialInstance->GetStaticParameterValues(StaticParameters);
-
-		for (int32 LayerParameterIdx = 0; LayerParameterIdx < StaticParameters.TerrainLayerWeightParameters.Num(); LayerParameterIdx++)
+		else
 		{
-			FStaticTerrainLayerWeightParameter& LayerParameter = StaticParameters.TerrainLayerWeightParameters[LayerParameterIdx];
-			if (LayerParameter.ParameterName == LayerName)
-			{
-				LayerParameter.WeightmapIndex = 0;
-				LayerParameter.bOverride = true;
-			}
-			else
-			{
-				LayerParameter.WeightmapIndex = INDEX_NONE;
-			}
+			LayerParameter.WeightmapIndex = INDEX_NONE;
 		}
-
-		CombinationMaterialInstance->UpdateStaticPermutation(StaticParameters);
-
-		CombinationMaterialInstance->PostEditChange();
 	}
-
-	// Create the instance for this component, that will use the layer combination instance.
-	ULandscapeMaterialInstanceConstant* MaterialInstance = NewObject<ULandscapeMaterialInstanceConstant>();
-	MaterialInstance->SetParentEditorOnly(CombinationMaterialInstance);
-	MaterialInstance->bIsLayerThumbnail = true;
+	MaterialInstance->UpdateStaticPermutation(StaticParameters);
 
 	FLinearColor Mask(1.0f, 0.0f, 0.0f, 0.0f);
 	MaterialInstance->SetVectorParameterValueEditorOnly(FName(*FString::Printf(TEXT("LayerMask_%s"), *LayerName.ToString())), Mask);
 	MaterialInstance->SetTextureParameterValueEditorOnly(FName(TEXT("Weightmap0")), ThumbnailWeightmap);
 	MaterialInstance->SetTextureParameterValueEditorOnly(FName(TEXT("Heightmap")), ThumbnailHeightmap);
+
 	MaterialInstance->PostEditChange();
 
 	return MaterialInstance;
@@ -205,15 +180,6 @@ UMaterialInstanceConstant* ULandscapeComponent::GetCombinationMaterial(bool bMob
 
 	if (ensure(MaterialToUse != nullptr))
 	{
-		// Ensure top level UMaterial has appropriate usage flags set.
-		bool bNeedsRecompile;
-		UMaterial* ParentUMaterial = MaterialToUse->GetMaterial();
-		if (ParentUMaterial && ParentUMaterial != UMaterial::GetDefaultMaterial(MD_Surface))
-		{
-			ParentUMaterial->SetMaterialUsage(bNeedsRecompile, MATUSAGE_Landscape);
-			ParentUMaterial->SetMaterialUsage(bNeedsRecompile, MATUSAGE_StaticLighting);
-		}
-
 		FString LayerKey = GetLayerAllocationKey(MaterialToUse, bMobile);
 		//UE_LOG(LogLandscape, Log, TEXT("Looking for key %s"), *LayerKey);
 
@@ -260,7 +226,7 @@ void ULandscapeComponent::UpdateMaterialInstances()
 	// Find or set a matching MIC in the Landscape's map.
 	UMaterialInstanceConstant* CombinationMaterialInstance = GetCombinationMaterial(false);
 
-	if (CombinationMaterialInstance != NULL)
+	if (CombinationMaterialInstance != nullptr)
 	{
 		// not having the context recreate the render state because we will manually do it for only this component
 		FMaterialUpdateContext Context(FMaterialUpdateContext::EOptions::Default & ~FMaterialUpdateContext::EOptions::RecreateRenderStates);
@@ -272,7 +238,7 @@ void ULandscapeComponent::UpdateMaterialInstances()
 		}
 
 		// Create the instance for this component, that will use the layer combination instance.
-		if (MaterialInstance == NULL || GetOutermost() != MaterialInstance->GetOutermost())
+		if (MaterialInstance == nullptr || GetOutermost() != MaterialInstance->GetOutermost())
 		{
 			MaterialInstance = NewObject<ULandscapeMaterialInstanceConstant>(GetOutermost());
 		}
@@ -336,7 +302,7 @@ class UMaterialInterface* ULandscapeComponent::GetMaterial(int32 ElementIndex) c
 	}
 	else
 	{
-		return NULL;
+		return nullptr;
 	}
 }
 
@@ -433,7 +399,7 @@ void ULandscapeComponent::FixupWeightmaps()
 				FLandscapeEditDataInterface LandscapeEdit(Info);
 				for (int32 Idx = 0; Idx < LayersToDelete.Num(); ++Idx)
 				{
-					DeleteLayer(LayersToDelete[Idx], &LandscapeEdit);
+					DeleteLayer(LayersToDelete[Idx], LandscapeEdit);
 				}
 			}
 
@@ -463,7 +429,7 @@ void ULandscapeComponent::FixupWeightmaps()
 				FLandscapeWeightmapUsage& Usage = Proxy->WeightmapUsageMap.FindOrAdd(WeightmapTexture);
 
 				// Detect a shared layer allocation, caused by a previous undo or layer deletion bugs
-				if (Usage.ChannelUsage[Allocation.WeightmapTextureChannel] != NULL &&
+				if (Usage.ChannelUsage[Allocation.WeightmapTextureChannel] != nullptr &&
 					Usage.ChannelUsage[Allocation.WeightmapTextureChannel] != this)
 				{
 					FFormatNamedArguments Arguments;
@@ -486,7 +452,7 @@ void ULandscapeComponent::FixupWeightmaps()
 			RemoveInvalidWeightmaps();
 
 			// Store the layer combination in the MaterialInstanceConstantMap
-			if (MaterialInstance != NULL)
+			if (MaterialInstance != nullptr)
 			{
 				UMaterialInstanceConstant* CombinationMaterialInstance = Cast<UMaterialInstanceConstant>(MaterialInstance->Parent);
 				if (CombinationMaterialInstance)
@@ -1237,7 +1203,7 @@ void ULandscapeComponent::UpdateCollisionLayerData()
 
 
 
-void ULandscapeComponent::GenerateHeightmapMips(TArray<FColor*>& HeightmapTextureMipData, int32 ComponentX1/*=0*/, int32 ComponentY1/*=0*/, int32 ComponentX2/*=MAX_int32*/, int32 ComponentY2/*=MAX_int32*/, FLandscapeTextureDataInfo* TextureDataInfo/*=NULL*/)
+void ULandscapeComponent::GenerateHeightmapMips(TArray<FColor*>& HeightmapTextureMipData, int32 ComponentX1/*=0*/, int32 ComponentY1/*=0*/, int32 ComponentX2/*=MAX_int32*/, int32 ComponentY2/*=MAX_int32*/, FLandscapeTextureDataInfo* TextureDataInfo/*=nullptr*/)
 {
 	bool EndX = false;
 	bool EndY = false;
@@ -1589,7 +1555,7 @@ namespace
 };
 
 template<typename DataType>
-void ULandscapeComponent::UpdateMipsTempl(int32 InNumSubsections, int32 InSubsectionSizeQuads, UTexture2D* Texture, TArray<DataType*>& TextureMipData, int32 ComponentX1/*=0*/, int32 ComponentY1/*=0*/, int32 ComponentX2/*=MAX_int32*/, int32 ComponentY2/*=MAX_int32*/, struct FLandscapeTextureDataInfo* TextureDataInfo/*=NULL*/)
+void ULandscapeComponent::UpdateMipsTempl(int32 InNumSubsections, int32 InSubsectionSizeQuads, UTexture2D* Texture, TArray<DataType*>& TextureMipData, int32 ComponentX1/*=0*/, int32 ComponentY1/*=0*/, int32 ComponentX2/*=MAX_int32*/, int32 ComponentY2/*=MAX_int32*/, struct FLandscapeTextureDataInfo* TextureDataInfo/*=nullptr*/)
 {
 	int32 WeightmapSizeU = Texture->Source.GetSizeX();
 	int32 WeightmapSizeV = Texture->Source.GetSizeY();
@@ -1754,12 +1720,12 @@ void ULandscapeComponent::UpdateMipsTempl(int32 InNumSubsections, int32 InSubsec
 	}
 }
 
-void ULandscapeComponent::UpdateWeightmapMips(int32 InNumSubsections, int32 InSubsectionSizeQuads, UTexture2D* WeightmapTexture, TArray<FColor*>& WeightmapTextureMipData, int32 ComponentX1/*=0*/, int32 ComponentY1/*=0*/, int32 ComponentX2/*=MAX_int32*/, int32 ComponentY2/*=MAX_int32*/, struct FLandscapeTextureDataInfo* TextureDataInfo/*=NULL*/)
+void ULandscapeComponent::UpdateWeightmapMips(int32 InNumSubsections, int32 InSubsectionSizeQuads, UTexture2D* WeightmapTexture, TArray<FColor*>& WeightmapTextureMipData, int32 ComponentX1/*=0*/, int32 ComponentY1/*=0*/, int32 ComponentX2/*=MAX_int32*/, int32 ComponentY2/*=MAX_int32*/, struct FLandscapeTextureDataInfo* TextureDataInfo/*=nullptr*/)
 {
 	UpdateMipsTempl<FColor>(InNumSubsections, InSubsectionSizeQuads, WeightmapTexture, WeightmapTextureMipData, ComponentX1, ComponentY1, ComponentX2, ComponentY2, TextureDataInfo);
 }
 
-void ULandscapeComponent::UpdateDataMips(int32 InNumSubsections, int32 InSubsectionSizeQuads, UTexture2D* Texture, TArray<uint8*>& TextureMipData, int32 ComponentX1/*=0*/, int32 ComponentY1/*=0*/, int32 ComponentX2/*=MAX_int32*/, int32 ComponentY2/*=MAX_int32*/, struct FLandscapeTextureDataInfo* TextureDataInfo/*=NULL*/)
+void ULandscapeComponent::UpdateDataMips(int32 InNumSubsections, int32 InSubsectionSizeQuads, UTexture2D* Texture, TArray<uint8*>& TextureMipData, int32 ComponentX1/*=0*/, int32 ComponentY1/*=0*/, int32 ComponentX2/*=MAX_int32*/, int32 ComponentY2/*=MAX_int32*/, struct FLandscapeTextureDataInfo* TextureDataInfo/*=nullptr*/)
 {
 	UpdateMipsTempl<uint8>(InNumSubsections, InSubsectionSizeQuads, Texture, TextureMipData, ComponentX1, ComponentY1, ComponentX2, ComponentY2, TextureDataInfo);
 }
@@ -1768,7 +1734,7 @@ float ULandscapeComponent::GetLayerWeightAtLocation(const FVector& InLocation, U
 {
 	// Allocate and discard locally if no external cache is passed in.
 	TArray<uint8> LocalCache;
-	if (LayerCache == NULL)
+	if (LayerCache == nullptr)
 	{
 		LayerCache = &LocalCache;
 	}
@@ -1901,10 +1867,10 @@ TArray<FName> ALandscapeProxy::GetLayersFromMaterial(UMaterialInterface* Materia
 		UMaterial* Material = MaterialInterface->GetMaterial();
 		TArray<FName> ParameterNames;
 		TArray<FGuid> Guids;
+		Material->GetAllParameterNames<UMaterialExpressionLandscapeLayerBlend>(ParameterNames, Guids);
 		Material->GetAllParameterNames<UMaterialExpressionLandscapeLayerWeight>(ParameterNames, Guids);
 		Material->GetAllParameterNames<UMaterialExpressionLandscapeLayerSwitch>(ParameterNames, Guids);
 		Material->GetAllParameterNames<UMaterialExpressionLandscapeLayerSample>(ParameterNames, Guids);
-		Material->GetAllParameterNames<UMaterialExpressionLandscapeLayerBlend>(ParameterNames, Guids);
 
 		for (const FName& Name : ParameterNames)
 		{
@@ -1931,13 +1897,13 @@ ULandscapeLayerInfoObject* ALandscapeProxy::CreateLayerInfo(const TCHAR* LayerNa
 	FString PackageName = Path + LayerObjectName.ToString();
 	FString PackageFilename;
 	int32 Suffix = 1;
-	while (FPackageName::DoesPackageExist(PackageName, NULL, &PackageFilename))
+	while (FPackageName::DoesPackageExist(PackageName, nullptr, &PackageFilename))
 	{
 		LayerObjectName = FName(*FString::Printf(TEXT("LayerInfoObject_%s_%d"), LayerName, Suffix));
 		PackageName = Path + LayerObjectName.ToString();
 		Suffix++;
 	}
-	UPackage* Package = CreatePackage(NULL, *PackageName);
+	UPackage* Package = CreatePackage(nullptr, *PackageName);
 	ULandscapeLayerInfoObject* LayerInfo = NewObject<ULandscapeLayerInfoObject>(Package, LayerObjectName, RF_Public | RF_Standalone | RF_Transactional);
 	LayerInfo->LayerName = LayerName;
 
@@ -2901,7 +2867,7 @@ void ULandscapeInfo::ExportHeightmap(const FString& Filename)
 		IImageWrapperModule& ImageWrapperModule = FModuleManager::LoadModuleChecked<IImageWrapperModule>("ImageWrapper");
 		IImageWrapperPtr ImageWrapper = ImageWrapperModule.CreateImageWrapper(EImageFormat::PNG);
 
-		const TArray<uint8>* RawData = NULL;
+		const TArray<uint8>* RawData = nullptr;
 		if (ImageWrapper->SetRaw(HeightData.GetData(), HeightData.Num(), (MaxX - MinX + 1), (MaxY - MinY + 1), ERGBFormat::Gray, 16))
 		{
 			HeightData = ImageWrapper->GetCompressed();
@@ -2940,7 +2906,7 @@ void ULandscapeInfo::ExportLayer(ULandscapeLayerInfoObject* LayerInfo, const FSt
 		IImageWrapperModule& ImageWrapperModule = FModuleManager::LoadModuleChecked<IImageWrapperModule>("ImageWrapper");
 		IImageWrapperPtr ImageWrapper = ImageWrapperModule.CreateImageWrapper(EImageFormat::PNG);
 
-		const TArray<uint8>* RawData = NULL;
+		const TArray<uint8>* RawData = nullptr;
 		if (ImageWrapper->SetRaw(WeightData.GetData(), WeightData.Num(), (MaxX - MinX + 1), (MaxY - MinY + 1), ERGBFormat::Gray, 8))
 		{
 			WeightData = ImageWrapper->GetCompressed();
@@ -2971,7 +2937,7 @@ void ULandscapeInfo::DeleteLayer(ULandscapeLayerInfoObject* LayerInfo)
 	}
 
 	ALandscape* Landscape = LandscapeActor.Get();
-	if (Landscape != NULL)
+	if (Landscape != nullptr)
 	{
 		Landscape->Modify();
 		Landscape->EditorLayerSettings.Remove(LayerInfo);
@@ -3009,11 +2975,11 @@ void ULandscapeInfo::ReplaceLayer(ULandscapeLayerInfoObject* FromLayerInfo, ULan
 		}
 
 		ALandscape* Landscape = LandscapeActor.Get();
-		if (Landscape != NULL)
+		if (Landscape != nullptr)
 		{
 			Landscape->Modify();
 			FLandscapeEditorLayerSettings* ToEditorLayerSettings = Landscape->EditorLayerSettings.FindByKey(ToLayerInfo);
-			if (ToEditorLayerSettings != NULL)
+			if (ToEditorLayerSettings != nullptr)
 			{
 				// If the new layer already exists, simple remove the old layer
 				Landscape->EditorLayerSettings.Remove(FromLayerInfo);
@@ -3021,7 +2987,7 @@ void ULandscapeInfo::ReplaceLayer(ULandscapeLayerInfoObject* FromLayerInfo, ULan
 			else
 			{
 				FLandscapeEditorLayerSettings* FromEditorLayerSettings = Landscape->EditorLayerSettings.FindByKey(FromLayerInfo);
-				if (FromEditorLayerSettings != NULL)
+				if (FromEditorLayerSettings != nullptr)
 				{
 					// If only the old layer exists (most common case), change it to point to the new layer info
 					FromEditorLayerSettings->LayerInfoObj = ToLayerInfo;
@@ -3039,7 +3005,7 @@ void ULandscapeInfo::ReplaceLayer(ULandscapeLayerInfoObject* FromLayerInfo, ULan
 			ALandscapeProxy* Proxy = *It;
 			Proxy->Modify();
 			FLandscapeEditorLayerSettings* ToEditorLayerSettings = Proxy->EditorLayerSettings.FindByKey(ToLayerInfo);
-			if (ToEditorLayerSettings != NULL)
+			if (ToEditorLayerSettings != nullptr)
 			{
 				// If the new layer already exists, simple remove the old layer
 				Proxy->EditorLayerSettings.Remove(FromLayerInfo);
@@ -3047,7 +3013,7 @@ void ULandscapeInfo::ReplaceLayer(ULandscapeLayerInfoObject* FromLayerInfo, ULan
 			else
 			{
 				FLandscapeEditorLayerSettings* FromEditorLayerSettings = Proxy->EditorLayerSettings.FindByKey(FromLayerInfo);
-				if (FromEditorLayerSettings != NULL)
+				if (FromEditorLayerSettings != nullptr)
 				{
 					// If only the old layer exists (most common case), change it to point to the new layer info
 					FromEditorLayerSettings->LayerInfoObj = ToLayerInfo;
@@ -3145,18 +3111,6 @@ void ALandscapeProxy::PostEditMove(bool bFinished)
 void ALandscapeProxy::PostEditImport()
 {
 	Super::PostEditImport();
-	if (!bIsProxy && GetWorld()) // For Landscape
-	{
-		for (TActorIterator<ALandscape> It(GetWorld()); It; ++It)
-		{
-			ALandscape* Landscape = *It;
-			if (Landscape != this && !Landscape->HasAnyFlags(RF_BeginDestroyed) && Landscape->LandscapeGuid == LandscapeGuid)
-			{
-				// Copy/Paste case, need to generate new GUID
-				LandscapeGuid = FGuid::NewGuid();
-			}
-		}
-	}
 
 	for (int32 ComponentIndex = 0; ComponentIndex < LandscapeComponents.Num(); ++ComponentIndex)
 	{
@@ -3181,6 +3135,24 @@ void ALandscape::PostEditMove(bool bFinished)
 
 	Super::PostEditMove(bFinished);
 }
+
+void ALandscape::PostEditImport()
+{
+	Super::PostEditImport();
+
+	if (GetWorld())
+	{
+		for (ALandscape* Landscape : TActorRange<ALandscape>(GetWorld()))
+		{
+			if (Landscape != this && !Landscape->HasAnyFlags(RF_BeginDestroyed) && Landscape->LandscapeGuid == LandscapeGuid)
+			{
+				// Copy/Paste case, need to generate new GUID
+				LandscapeGuid = FGuid::NewGuid();
+				break;
+			}
+		}
+	}
+}
 #endif	//WITH_EDITOR
 
 ULandscapeLayerInfoObject::ULandscapeLayerInfoObject(const FObjectInitializer& ObjectInitializer)
@@ -3287,7 +3259,7 @@ void ALandscapeProxy::RemoveXYOffsets()
 			Comp->XYOffsetmapTexture->ClearFlags(RF_Standalone);
 			Comp->Modify();
 			Comp->MarkPackageDirty();
-			Comp->XYOffsetmapTexture = NULL;
+			Comp->XYOffsetmapTexture = nullptr;
 			Comp->MarkRenderStateDirty();
 			bFoundXYOffset = true;
 		}
@@ -3426,57 +3398,76 @@ namespace
 
 void ALandscapeProxy::PostEditChangeProperty(FPropertyChangedEvent& PropertyChangedEvent)
 {
-	Super::PostEditChangeProperty(PropertyChangedEvent);
 	const FName PropertyName = PropertyChangedEvent.MemberProperty ? PropertyChangedEvent.MemberProperty->GetFName() : NAME_None;
-	if (bIsProxy)
+	const FName SubPropertyName = PropertyChangedEvent.Property ? PropertyChangedEvent.Property->GetFName() : NAME_None;
+
+	if (PropertyName == FName(TEXT("RelativeScale3D")))
 	{
-		if (PropertyName == FName(TEXT("LandscapeActor")))
+		// RelativeScale3D isn't even a property of ALandscapeProxy, it's a property of the root component
+		if (RootComponent)
 		{
-			if (LandscapeActor && IsValidLandscapeActor(LandscapeActor.Get()))
+			FVector ModifiedScale = RootComponent->RelativeScale3D;
+
+			// Lock X and Y scaling to the same value
+			if (SubPropertyName == FName("Y"))
 			{
-				LandscapeGuid = LandscapeActor->LandscapeGuid;
-				// defer LandscapeInfo setup
-				if (GIsEditor && GetWorld() && !GetWorld()->IsPlayInEditor())
-				{
-					GEngine->DeferredCommands.AddUnique(TEXT("UpdateLandscapeEditorData"));
-				}
+				ModifiedScale.X = FMath::Abs(RootComponent->RelativeScale3D.Y)*FMath::Sign(ModifiedScale.X);
 			}
 			else
 			{
-				LandscapeActor = nullptr;
+				// There's no "if name == X" here so that if we can't tell which has changed out of X and Y, we just use X
+				ModifiedScale.Y = FMath::Abs(RootComponent->RelativeScale3D.X)*FMath::Sign(ModifiedScale.Y);
 			}
-		}
-		else if (PropertyName == FName(TEXT("LandscapeMaterial")) || PropertyName == FName(TEXT("LandscapeHoleMaterial")))
-		{
-			{
-				FMaterialUpdateContext MaterialUpdateContext;
-				GetLandscapeInfo()->UpdateLayerInfoMap(/*this*/);
 
-				// Clear the parents out of combination material instances
-				for (const auto& MICPair : MaterialInstanceConstantMap)
+			ULandscapeInfo* Info = GetLandscapeInfo(false);
+
+			// Correct for attempts to scale to 0 on any axis
+			if (ModifiedScale.X == 0)
+			{
+				if (Info && Info->DrawScale.X < 0)
+				{
+					ModifiedScale.Y = ModifiedScale.X = -1;
+				}
+				else
+				{
+					ModifiedScale.Y = ModifiedScale.X = 1;
+				}
+			}
+			if (ModifiedScale.Z == 0)
+			{
+				if (Info && Info->DrawScale.Z < 0)
+				{
+					ModifiedScale.Z = -1;
+				}
+				else
 				{
-					UMaterialInstanceConstant* MaterialInstance = MICPair.Value;
-					MaterialInstance->BasePropertyOverrides.bOverride_BlendMode = false;
-					MaterialInstance->SetParentEditorOnly(nullptr);
-					MaterialUpdateContext.AddMaterialInstance(MaterialInstance);
+					ModifiedScale.Z = 1;
 				}
+			}
+
+			RootComponent->SetRelativeScale3D(ModifiedScale);
 
-				// Remove our references to any material instances
-				MaterialInstanceConstantMap.Empty();
+			// Update ULandscapeInfo cached DrawScale
+			if (Info)
+			{
+				Info->DrawScale = ModifiedScale;
 			}
 
-			for (int32 ComponentIndex = 0; ComponentIndex < LandscapeComponents.Num(); ComponentIndex++)
+			// We need to regenerate collision objects, they depend on scale value 
+			for (int32 ComponentIndex = 0; ComponentIndex < CollisionComponents.Num(); ComponentIndex++)
 			{
-				ULandscapeComponent* Comp = LandscapeComponents[ComponentIndex];
+				ULandscapeHeightfieldCollisionComponent* Comp = CollisionComponents[ComponentIndex];
 				if (Comp)
 				{
-					// Update the MIC
-					Comp->UpdateMaterialInstances();
+					Comp->RecreateCollision();
 				}
 			}
 		}
 	}
 
+	// Must do this *after* correcting the scale or reattaching the landscape components will crash!
+	Super::PostEditChangeProperty(PropertyChangedEvent);
+
 	if (GIsEditor && PropertyName == FName(TEXT("StreamingDistanceMultiplier")))
 	{
 		// Recalculate in a few seconds.
@@ -3520,77 +3511,56 @@ void ALandscapeProxy::PostEditChangeProperty(FPropertyChangedEvent& PropertyChan
 	}
 }
 
-void ALandscapeProxy::PostEditChangeChainProperty(FPropertyChangedChainEvent& PropertyChangedEvent)
+void ALandscapeStreamingProxy::PostEditChangeProperty(FPropertyChangedEvent& PropertyChangedEvent)
 {
-	const FName MemberPropertyName = PropertyChangedEvent.PropertyChain.GetActiveMemberNode()->GetValue()->GetFName();
-	const FName PropertyName = PropertyChangedEvent.PropertyChain.GetActiveNode()->GetValue()->GetFName();
+	Super::PostEditChangeProperty(PropertyChangedEvent);
 
-	if (MemberPropertyName == FName(TEXT("RelativeScale3D")))
+	const FName PropertyName = PropertyChangedEvent.MemberProperty ? PropertyChangedEvent.MemberProperty->GetFName() : NAME_None;
+	if (PropertyName == FName(TEXT("LandscapeActor")))
 	{
-		// RelativeScale3D isn't even a property of ALandscapeProxy, it's a property of the root component
-		if (RootComponent)
+		if (LandscapeActor && IsValidLandscapeActor(LandscapeActor.Get()))
 		{
-			FVector ModifiedScale = RootComponent->RelativeScale3D;
-
-			// Lock X and Y scaling to the same value
-			if (PropertyName == FName("Y"))
-			{
-				ModifiedScale.X = FMath::Abs(RootComponent->RelativeScale3D.Y)*FMath::Sign(ModifiedScale.X);
-			}
-			else
+			LandscapeGuid = LandscapeActor->GetLandscapeGuid();
+			// defer LandscapeInfo setup
+			if (GIsEditor && GetWorld() && !GetWorld()->IsPlayInEditor())
 			{
-				// There's no "if name == X" here so that if we can't tell which has changed out of X and Y, we just use X
-				ModifiedScale.Y = FMath::Abs(RootComponent->RelativeScale3D.X)*FMath::Sign(ModifiedScale.Y);
+				GEngine->DeferredCommands.AddUnique(TEXT("UpdateLandscapeEditorData"));
 			}
+		}
+		else
+		{
+			LandscapeActor = nullptr;
+		}
+	}
+	else if (PropertyName == FName(TEXT("LandscapeMaterial")) || PropertyName == FName(TEXT("LandscapeHoleMaterial")))
+	{
+		{
+			FMaterialUpdateContext MaterialUpdateContext;
+			GetLandscapeInfo()->UpdateLayerInfoMap(/*this*/);
 
-			ULandscapeInfo* Info = GetLandscapeInfo(false);
-
-			// Correct for attempts to scale to 0 on any axis
-			if (ModifiedScale.X == 0)
-			{
-				if (Info && Info->DrawScale.X < 0)
-				{
-					ModifiedScale.Y = ModifiedScale.X = -1;
-				}
-				else
-				{
-					ModifiedScale.Y = ModifiedScale.X = 1;
-				}
-			}
-			if (ModifiedScale.Z == 0)
+			// Clear the parents out of combination material instances
+			for (const auto& MICPair : MaterialInstanceConstantMap)
 			{
-				if (Info && Info->DrawScale.Z < 0)
-				{
-					ModifiedScale.Z = -1;
-				}
-				else
-				{
-					ModifiedScale.Z = 1;
-				}
+				UMaterialInstanceConstant* MaterialInstance = MICPair.Value;
+				MaterialInstance->BasePropertyOverrides.bOverride_BlendMode = false;
+				MaterialInstance->SetParentEditorOnly(nullptr);
+				MaterialUpdateContext.AddMaterialInstance(MaterialInstance);
 			}
 
-			RootComponent->SetRelativeScale3D(ModifiedScale);
-
-			// Update ULandscapeInfo cached DrawScale
-			if (Info)
-			{
-				Info->DrawScale = ModifiedScale;
-			}
+			// Remove our references to any material instances
+			MaterialInstanceConstantMap.Empty();
+		}
 
-			// We need to regenerate collision objects, they depend on scale value 
-			for (int32 ComponentIndex = 0; ComponentIndex < CollisionComponents.Num(); ComponentIndex++)
+		for (int32 ComponentIndex = 0; ComponentIndex < LandscapeComponents.Num(); ComponentIndex++)
+		{
+			ULandscapeComponent* Comp = LandscapeComponents[ComponentIndex];
+			if (Comp)
 			{
-				ULandscapeHeightfieldCollisionComponent* Comp = CollisionComponents[ComponentIndex];
-				if (Comp)
-				{
-					Comp->RecreateCollision();
-				}
+				// Update the MIC
+				Comp->UpdateMaterialInstances();
 			}
 		}
 	}
-
-	// Must do this *after* correcting the scale or reattaching the landscape components will crash!
-	Super::PostEditChangeChainProperty(PropertyChangedEvent);
 }
 
 void ALandscape::PostEditChangeProperty(FPropertyChangedEvent& PropertyChangedEvent)
@@ -3688,10 +3658,10 @@ void ALandscape::PostEditChangeProperty(FPropertyChangedEvent& PropertyChangedEv
 		if (bPropagateToProxies)
 		{
 			// Propagate Event to Proxies...
-			for (TSet<ALandscapeProxy*>::TIterator It(Info->Proxies); It; ++It)
+			for (ALandscapeProxy* Proxy : Info->Proxies)
 			{
-				(*It)->GetSharedProperties(this);
-				(*It)->PostEditChangeProperty(PropertyChangedEvent);
+				Proxy->GetSharedProperties(this);
+				Proxy->PostEditChangeProperty(PropertyChangedEvent);
 			}
 		}
 
@@ -3705,9 +3675,8 @@ void ALandscape::PostEditChangeProperty(FPropertyChangedEvent& PropertyChangedEv
 		// We cannot iterate the XYtoComponentMap directly because reregistering components modifies the array.
 		TArray<ULandscapeComponent*> AllComponents;
 		Info->XYtoComponentMap.GenerateValueArray(AllComponents);
-		for (auto It = AllComponents.CreateIterator(); It; ++It)
+		for (ULandscapeComponent* Comp : AllComponents)
 		{
-			ULandscapeComponent* Comp = *It;
 			if (Comp)
 			{
 				if (bNeedsRecalcBoundingBox)
@@ -3740,9 +3709,9 @@ void ALandscape::PostEditChangeProperty(FPropertyChangedEvent& PropertyChangedEv
 		// Need to update Gizmo scene proxy
 		if (bNeedsRecalcBoundingBox && GetWorld())
 		{
-			for (TActorIterator<ALandscapeGizmoActiveActor> It(GetWorld()); It; ++It)
+			for (ALandscapeGizmoActiveActor* Gizmo : TActorRange<ALandscapeGizmoActiveActor>(GetWorld()))
 			{
-				It->ReregisterAllComponents();
+				Gizmo->ReregisterAllComponents();
 			}
 		}
 
@@ -3754,9 +3723,9 @@ void ALandscape::PostEditChangeProperty(FPropertyChangedEvent& PropertyChangedEv
 			}
 
 			// Update all the proxies...
-			for (TSet<ALandscapeProxy*>::TIterator It(Info->Proxies); It; ++It)
+			for (ALandscapeProxy* Proxy : Info->Proxies)
 			{
-				(*It)->MarkComponentsRenderStateDirty();
+				Proxy->MarkComponentsRenderStateDirty();
 			}
 		}
 	}
@@ -3847,7 +3816,7 @@ void ULandscapeComponent::PreEditChange(UProperty* PropertyThatWillChange)
 			FIntPoint ComponentKey = GetSectionBase() / ComponentSizeQuads;
 			auto RegisteredComponent = Info->XYtoComponentMap.FindRef(ComponentKey);
 
-			if (RegisteredComponent == NULL)
+			if (RegisteredComponent == nullptr)
 			{
 				Info->XYtoComponentMap.Add(ComponentKey, this);
 			}
@@ -3918,7 +3887,7 @@ void ULandscapeInfo::UpdateSelectedComponents(TSet<ULandscapeComponent*>& NewCom
 		for (TSet<ULandscapeComponent*>::TIterator It(NewComponents); It; ++It)
 		{
 			ULandscapeComponent* Comp = *It;
-			if (Comp->EditToolRenderData != NULL && (Comp->EditToolRenderData->SelectedType & InSelectType) == 0)
+			if (Comp->EditToolRenderData != nullptr && (Comp->EditToolRenderData->SelectedType & InSelectType) == 0)
 			{
 				Comp->Modify();
 				int32 SelectedType = Comp->EditToolRenderData->SelectedType;
@@ -3932,7 +3901,7 @@ void ULandscapeInfo::UpdateSelectedComponents(TSet<ULandscapeComponent*>& NewCom
 		for (TSet<ULandscapeComponent*>::TIterator It(RemovedComponents); It; ++It)
 		{
 			ULandscapeComponent* Comp = *It;
-			if (Comp->EditToolRenderData != NULL)
+			if (Comp->EditToolRenderData != nullptr)
 			{
 				Comp->Modify();
 				int32 SelectedType = Comp->EditToolRenderData->SelectedType;
@@ -3950,7 +3919,7 @@ void ULandscapeInfo::UpdateSelectedComponents(TSet<ULandscapeComponent*>& NewCom
 			for (TSet<ULandscapeComponent*>::TIterator It(NewComponents); It; ++It)
 			{
 				ULandscapeComponent* Comp = *It;
-				if (Comp->EditToolRenderData != NULL && (Comp->EditToolRenderData->SelectedType & InSelectType) == 0)
+				if (Comp->EditToolRenderData != nullptr && (Comp->EditToolRenderData->SelectedType & InSelectType) == 0)
 				{
 					Comp->Modify();
 					int32 SelectedType = Comp->EditToolRenderData->SelectedType;
@@ -3967,7 +3936,7 @@ void ULandscapeInfo::UpdateSelectedComponents(TSet<ULandscapeComponent*>& NewCom
 			for (TSet<ULandscapeComponent*>::TIterator It(SelectedRegionComponents); It; ++It)
 			{
 				ULandscapeComponent* Comp = *It;
-				if (Comp->EditToolRenderData != NULL)
+				if (Comp->EditToolRenderData != nullptr)
 				{
 					Comp->Modify();
 					int32 SelectedType = Comp->EditToolRenderData->SelectedType;
@@ -4004,7 +3973,7 @@ void ULandscapeInfo::ClearSelectedRegion(bool bIsComponentwise /*= true*/)
 
 struct FLandscapeDataInterface* ULandscapeInfo::GetDataInterface()
 {
-	if (DataInterface == NULL)
+	if (DataInterface == nullptr)
 	{
 		DataInterface = new FLandscapeDataInterface();
 	}
@@ -4064,7 +4033,7 @@ void ULandscapeComponent::ReallocateWeightmaps(FLandscapeEditDataInterface* Data
 
 			for (int32 ChanIdx = 0; ChanIdx < 4; ChanIdx++)
 			{
-				if (Usage->ChannelUsage[ChanIdx] == NULL)
+				if (Usage->ChannelUsage[ChanIdx] == nullptr)
 				{
 					for (int32 LayerIdx = 0; LayerIdx < WeightmapLayerAllocations.Num(); LayerIdx++)
 					{
@@ -4105,8 +4074,8 @@ void ULandscapeComponent::ReallocateWeightmaps(FLandscapeEditDataInterface* Data
 	{
 		// UE_LOG(LogLandscape, Log, TEXT("Still need %d channels"), TotalNeededChannels);
 
-		UTexture2D* CurrentWeightmapTexture = NULL;
-		FLandscapeWeightmapUsage* CurrentWeightmapUsage = NULL;
+		UTexture2D* CurrentWeightmapTexture = nullptr;
+		FLandscapeWeightmapUsage* CurrentWeightmapUsage = nullptr;
 
 		if (TotalNeededChannels < 4)
 		{
@@ -4122,7 +4091,7 @@ void ULandscapeComponent::ReallocateWeightmaps(FLandscapeEditDataInterface* Data
 					// See if this candidate is closer than any others we've found
 					for (int32 ChanIdx = 0; ChanIdx < 4; ChanIdx++)
 					{
-						if (TryWeightmapUsage->ChannelUsage[ChanIdx] != NULL)
+						if (TryWeightmapUsage->ChannelUsage[ChanIdx] != nullptr)
 						{
 							int32 TryDistanceSquared = (TryWeightmapUsage->ChannelUsage[ChanIdx]->GetSectionBase() - GetSectionBase()).SizeSquared();
 							if (TryDistanceSquared < BestDistanceSquared)
@@ -4139,7 +4108,7 @@ void ULandscapeComponent::ReallocateWeightmaps(FLandscapeEditDataInterface* Data
 
 		bool NeedsUpdateResource = false;
 		// No suitable weightmap texture
-		if (CurrentWeightmapTexture == NULL)
+		if (CurrentWeightmapTexture == nullptr)
 		{
 			MarkPackageDirty();
 
@@ -4164,7 +4133,7 @@ void ULandscapeComponent::ReallocateWeightmaps(FLandscapeEditDataInterface* Data
 		{
 			// UE_LOG(LogLandscape, Log, TEXT("Finding allocation for layer %d"), CurrentLayer);
 
-			if (CurrentWeightmapUsage->ChannelUsage[ChanIdx] == NULL)
+			if (CurrentWeightmapUsage->ChannelUsage[ChanIdx] == nullptr)
 			{
 				// Use this allocation
 				FWeightmapLayerAllocationInfo& AllocInfo = WeightmapLayerAllocations[CurrentLayer];
@@ -4183,7 +4152,7 @@ void ULandscapeComponent::ReallocateWeightmaps(FLandscapeEditDataInterface* Data
 					UTexture2D* OldWeightmapTexture = WeightmapTextures[AllocInfo.WeightmapTextureIndex];
 
 					// Copy the data
-					if (ensure(DataInterface != NULL)) // it's not safe to skip the copy
+					if (ensure(DataInterface != nullptr)) // it's not safe to skip the copy
 					{
 						DataInterface->CopyTextureChannel(CurrentWeightmapTexture, ChanIdx, OldWeightmapTexture, AllocInfo.WeightmapTextureChannel);
 						DataInterface->ZeroTextureChannel(OldWeightmapTexture, AllocInfo.WeightmapTextureChannel);
@@ -4192,7 +4161,7 @@ void ULandscapeComponent::ReallocateWeightmaps(FLandscapeEditDataInterface* Data
 
 					// Remove the old allocation
 					FLandscapeWeightmapUsage* OldWeightmapUsage = Proxy->WeightmapUsageMap.Find(OldWeightmapTexture);
-					OldWeightmapUsage->ChannelUsage[AllocInfo.WeightmapTextureChannel] = NULL;
+					OldWeightmapUsage->ChannelUsage[AllocInfo.WeightmapTextureChannel] = nullptr;
 				}
 
 				// Assign the new allocation
@@ -4411,7 +4380,7 @@ void ULandscapeComponent::InitWeightmapData(TArray<ULandscapeLayerInfoObject*>&
 		new (WeightmapLayerAllocations)FWeightmapLayerAllocationInfo(LayerInfos[Idx]);
 	}
 
-	ReallocateWeightmaps(NULL);
+	ReallocateWeightmaps(nullptr);
 
 	check(WeightmapLayerAllocations.Num() > 0 && WeightmapTextures.Num() > 0);
 
@@ -4466,7 +4435,7 @@ void ULandscapeComponent::InitWeightmapData(TArray<ULandscapeLayerInfoObject*>&
 
 	FlushRenderingCommands();
 
-	MaterialInstance = NULL;
+	MaterialInstance = nullptr;
 
 }
 
@@ -4619,7 +4588,7 @@ void ULandscapeComponent::ImportCustomProperties(const TCHAR* SourceText, FFeedb
 		{
 			if (FParse::Value(SourceText, TEXT("LayerInfo="), LayerInfoPath))
 			{
-				LayerInfos.Add(LoadObject<ULandscapeLayerInfoObject>(NULL, *LayerInfoPath));
+				LayerInfos.Add(LoadObject<ULandscapeLayerInfoObject>(nullptr, *LayerInfoPath));
 
 				while (*SourceText && (!FChar::IsWhitespace(*SourceText)))
 				{
@@ -4658,9 +4627,9 @@ void ULandscapeComponent::ImportCustomProperties(const TCHAR* SourceText, FFeedb
 	}
 }
 
-bool ALandscapeProxy::IsValidLandscapeActor(ALandscape* Landscape)
+bool ALandscapeStreamingProxy::IsValidLandscapeActor(ALandscape* Landscape)
 {
-	if (bIsProxy && Landscape)
+	if (Landscape)
 	{
 		if (!Landscape->HasAnyFlags(RF_BeginDestroyed))
 		{
@@ -4669,7 +4638,7 @@ bool ALandscapeProxy::IsValidLandscapeActor(ALandscape* Landscape)
 				return true; // always valid for newly created Proxy
 			}
 			if (((LandscapeActor && LandscapeActor == Landscape)
-				|| (LandscapeActor.IsNull() && LandscapeGuid.IsValid() && LandscapeGuid == Landscape->LandscapeGuid))
+				|| (LandscapeActor.IsNull() && LandscapeGuid.IsValid() && LandscapeGuid == Landscape->GetLandscapeGuid()))
 				&& ComponentSizeQuads == Landscape->ComponentSizeQuads
 				&& NumSubsections == Landscape->NumSubsections
 				&& SubsectionSizeQuads == Landscape->SubsectionSizeQuads)
@@ -4920,25 +4889,25 @@ void ULandscapeComponent::GeneratePlatformVertexData()
 						FLandscapeVertexRef V4(x0, y1, SubX, SubY);
 
 						uint64 Key1 = V1.MakeKey();
-						if (VertexMap.Find(Key1) == NULL)
+						if (VertexMap.Find(Key1) == nullptr)
 						{
 							VertexMap.Add(Key1, VertexOrder.Num());
 							VertexOrder.Add(V1);
 						}
 						uint64 Key2 = V2.MakeKey();
-						if (VertexMap.Find(Key2) == NULL)
+						if (VertexMap.Find(Key2) == nullptr)
 						{
 							VertexMap.Add(Key2, VertexOrder.Num());
 							VertexOrder.Add(V2);
 						}
 						uint64 Key3 = V3.MakeKey();
-						if (VertexMap.Find(Key3) == NULL)
+						if (VertexMap.Find(Key3) == nullptr)
 						{
 							VertexMap.Add(Key3, VertexOrder.Num());
 							VertexOrder.Add(V3);
 						}
 						uint64 Key4 = V4.MakeKey();
-						if (VertexMap.Find(Key4) == NULL)
+						if (VertexMap.Find(Key4) == nullptr)
 						{
 							VertexMap.Add(Key4, VertexOrder.Num());
 							VertexOrder.Add(V4);
diff --git a/Engine/Source/Runtime/Landscape/Private/LandscapeEditInterface.cpp b/Engine/Source/Runtime/Landscape/Private/LandscapeEditInterface.cpp
--- a/Engine/Source/Runtime/Landscape/Private/LandscapeEditInterface.cpp
+++ b/Engine/Source/Runtime/Landscape/Private/LandscapeEditInterface.cpp
@@ -175,7 +175,7 @@ namespace
 
 };
 
-bool FLandscapeEditDataInterface::GetComponentsInRegion(int32 X1, int32 Y1, int32 X2, int32 Y2, TSet<ULandscapeComponent*>* OutComponents /*= NULL*/)
+bool FLandscapeEditDataInterface::GetComponentsInRegion(int32 X1, int32 Y1, int32 X2, int32 Y2, TSet<ULandscapeComponent*>* OutComponents /*= nullptr*/)
 {
 	if (ComponentSizeQuads <= 0 || !LandscapeInfo)
 	{
@@ -186,14 +186,14 @@ bool FLandscapeEditDataInterface::GetComponentsInRegion(int32 X1, int32 Y1, int3
 	ALandscape::CalcComponentIndicesOverlap(X1, Y1, X2, Y2, ComponentSizeQuads, ComponentIndexX1, ComponentIndexY1, ComponentIndexX2, ComponentIndexY2);
 
 	bool bNotLocked = true;
-	for( int32 ComponentIndexY=ComponentIndexY1;ComponentIndexY<=ComponentIndexY2;ComponentIndexY++ )
+	for (int32 ComponentIndexY = ComponentIndexY1; ComponentIndexY <= ComponentIndexY2; ComponentIndexY++)
 	{
-		for( int32 ComponentIndexX=ComponentIndexX1;ComponentIndexX<=ComponentIndexX2;ComponentIndexX++ )
+		for (int32 ComponentIndexX = ComponentIndexX1; ComponentIndexX <= ComponentIndexX2; ComponentIndexX++)
 		{
-			ULandscapeComponent* Component = LandscapeInfo->XYtoComponentMap.FindRef(FIntPoint(ComponentIndexX,ComponentIndexY));
-			if( Component )
+			ULandscapeComponent* Component = LandscapeInfo->XYtoComponentMap.FindRef(FIntPoint(ComponentIndexX, ComponentIndexY));
+			if (Component)
 			{
-				bNotLocked = bNotLocked && ( !FLevelUtils::IsLevelLocked(Component->GetLandscapeProxy()->GetLevel()) ) && FLevelUtils::IsLevelVisible(Component->GetLandscapeProxy()->GetLevel());
+				bNotLocked = bNotLocked && (!FLevelUtils::IsLevelLocked(Component->GetLandscapeProxy()->GetLevel())) && FLevelUtils::IsLevelVisible(Component->GetLandscapeProxy()->GetLevel());
 				if (OutComponents)
 				{
 					OutComponents->Add(Component);
@@ -1478,7 +1478,7 @@ void FLandscapeEditDataInterface::GetHeightDataFast(const int32 X1, const int32
 	}
 }
 
-void ULandscapeComponent::DeleteLayer(ULandscapeLayerInfoObject* LayerInfo, FLandscapeEditDataInterface* LandscapeEdit)
+void ULandscapeComponent::DeleteLayer(ULandscapeLayerInfoObject* LayerInfo, FLandscapeEditDataInterface& LandscapeEdit)
 {
 	ULandscapeComponent* Component = this;
 
@@ -1516,7 +1516,7 @@ void ULandscapeComponent::DeleteLayer(ULandscapeLayerInfoObject* LayerInfo, FLan
 		TArray<FLandscapeTextureDataInfo*> TexDataInfos;
 		for (int32 WeightmapIdx = 0; WeightmapIdx < Component->WeightmapTextures.Num(); WeightmapIdx++)
 		{
-			TexDataInfos.Add(LandscapeEdit->GetTextureDataInfo(Component->WeightmapTextures[WeightmapIdx]));
+			TexDataInfos.Add(LandscapeEdit.GetTextureDataInfo(Component->WeightmapTextures[WeightmapIdx]));
 		}
 
 		TArray<bool> LayerNoWeightBlends;	// Array of NoWeightBlend flags
@@ -1663,14 +1663,14 @@ void ULandscapeComponent::DeleteLayer(ULandscapeLayerInfoObject* LayerInfo, FLan
 	TArray<FColor*> CollisionWeightmapMipData;
 	for (int32 WeightmapIdx = 0; WeightmapIdx < Component->WeightmapTextures.Num(); WeightmapIdx++)
 	{
-		CollisionWeightmapMipData.Add((FColor*)LandscapeEdit->GetTextureDataInfo(Component->WeightmapTextures[WeightmapIdx])->GetMipData(Component->CollisionMipLevel));
+		CollisionWeightmapMipData.Add((FColor*)LandscapeEdit.GetTextureDataInfo(Component->WeightmapTextures[WeightmapIdx])->GetMipData(Component->CollisionMipLevel));
 	}
 	TArray<FColor*> SimpleCollisionWeightmapMipData;
 	if (Component->SimpleCollisionMipLevel > Component->CollisionMipLevel)
 	{
 		for (int32 WeightmapIdx = 0; WeightmapIdx < Component->WeightmapTextures.Num(); WeightmapIdx++)
 		{
-			SimpleCollisionWeightmapMipData.Add((FColor*)LandscapeEdit->GetTextureDataInfo(Component->WeightmapTextures[WeightmapIdx])->GetMipData(Component->SimpleCollisionMipLevel));
+			SimpleCollisionWeightmapMipData.Add((FColor*)LandscapeEdit.GetTextureDataInfo(Component->WeightmapTextures[WeightmapIdx])->GetMipData(Component->SimpleCollisionMipLevel));
 		}
 	}
 	Component->UpdateCollisionLayerData(
@@ -1680,17 +1680,214 @@ void ULandscapeComponent::DeleteLayer(ULandscapeLayerInfoObject* LayerInfo, FLan
 
 void FLandscapeEditDataInterface::DeleteLayer(ULandscapeLayerInfoObject* LayerInfo)
 {
-	if (!LandscapeInfo) return;
+	if (!LandscapeInfo)
+	{
+		return;
+	}
+
+	for (auto& XYComponentPair : LandscapeInfo->XYtoComponentMap)
+	{
+		ULandscapeComponent* Component = XYComponentPair.Value;
+		Component->DeleteLayer(LayerInfo, *this);
+	}
+
+	// Flush dynamic data (e.g. grass)
+	TSet<ULandscapeComponent*> Components;
+	for (auto& XYComponentPair : LandscapeInfo->XYtoComponentMap)
+	{
+		Components.Add(XYComponentPair.Value);
+	}
+	ALandscapeProxy::InvalidateGeneratedComponentData(Components);
+}
+
+
+void ULandscapeComponent::FillLayer(ULandscapeLayerInfoObject* LayerInfo, FLandscapeEditDataInterface& LandscapeEdit)
+{
+	check(LayerInfo);
+
+	ULandscapeComponent* Component = this;
+
+	ALandscapeProxy* Proxy = Component->GetLandscapeProxy();
+	Component->Modify();
+	Proxy->Modify();
+
+	const bool bFillLayerIsNoWeightBlend = LayerInfo->bNoWeightBlend;
+	bool bClearOtherWeightBlendLayers = !bFillLayerIsNoWeightBlend;
+
+	// Find the index for this layer in this component.
+	int32 FillLayerIdx = Component->WeightmapLayerAllocations.IndexOfByPredicate(
+		[LayerInfo](const FWeightmapLayerAllocationInfo& Allocation) { return Allocation.LayerInfo == LayerInfo; });
+
+	// if the layer isn't used on this component yet but is a weight-blend layer, then simply steal the allocation of another weight-blend layer!
+	if (FillLayerIdx == INDEX_NONE && !bFillLayerIsNoWeightBlend)
+	{
+		FillLayerIdx = Component->WeightmapLayerAllocations.IndexOfByPredicate(
+			[](const FWeightmapLayerAllocationInfo& Allocation) { return !Allocation.LayerInfo || !Allocation.LayerInfo->bNoWeightBlend; });
+
+		if (FillLayerIdx != INDEX_NONE)
+		{
+			Component->WeightmapLayerAllocations[FillLayerIdx].LayerInfo = LayerInfo;
+		}
+		else
+		{
+			// no other weight-blend layers exist
+			bClearOtherWeightBlendLayers = false;
+		}
+	}
+
+	// if the layer is still not found then we are forced to make a new allocation
+	if (FillLayerIdx == INDEX_NONE)
+	{
+		FillLayerIdx = Component->WeightmapLayerAllocations.Num();
+		Component->WeightmapLayerAllocations.Add(FWeightmapLayerAllocationInfo(LayerInfo));
+		Component->ReallocateWeightmaps(&LandscapeEdit);
+	}
+
+	check(FillLayerIdx != INDEX_NONE);
+
+	// fill the layer
+	{
+		// Find the texture data corresponding to this vertex
+		const int32 SizeU = (SubsectionSizeQuads + 1) * NumSubsections;
+		const int32 SizeV = (SubsectionSizeQuads + 1) * NumSubsections;
+		const int32 WeightmapOffsetX = Component->WeightmapScaleBias.Z * (float)SizeU;
+		const int32 WeightmapOffsetY = Component->WeightmapScaleBias.W * (float)SizeV;
+
+		FWeightmapLayerAllocationInfo& FillLayerAllocation = Component->WeightmapLayerAllocations[FillLayerIdx];
+		uint8* LayerData = (uint8*)LandscapeEdit.GetTextureDataInfo(Component->WeightmapTextures[FillLayerAllocation.WeightmapTextureIndex])->GetMipData(0);
+
+		for (int32 Y = 0; Y < SizeV; ++Y)
+		{
+			const int32 TexY = WeightmapOffsetY + Y;
+			uint8* RowData = LayerData + ((WeightmapOffsetY + Y) * SizeU + WeightmapOffsetX) * 4 + ChannelOffsets[FillLayerAllocation.WeightmapTextureChannel];
+			for (int32 X = 0; X < SizeU; ++X)
+			{
+				RowData[X * 4] = 255;
+			}
+		}
+	}
+
+	// clear other layers
+	if (bClearOtherWeightBlendLayers)
+	{
+		for (int32 LayerIdx = 0; LayerIdx < Component->WeightmapLayerAllocations.Num(); ++LayerIdx)
+		{
+			if (LayerIdx == FillLayerIdx)
+			{
+				continue;
+			}
+			FWeightmapLayerAllocationInfo& Allocation = Component->WeightmapLayerAllocations[LayerIdx];
+			if (Allocation.LayerInfo->bNoWeightBlend)
+			{
+				continue;
+			}
+
+			FLandscapeWeightmapUsage* Usage = Proxy->WeightmapUsageMap.Find(Component->WeightmapTextures[Allocation.WeightmapTextureIndex]);
+			if (Usage) // can be null if WeightmapUsageMap hasn't been built yet
+			{
+				Usage->ChannelUsage[Allocation.WeightmapTextureChannel] = nullptr;
+			}
+
+			Allocation.WeightmapTextureIndex = 255;
+		}
+
+		Component->WeightmapLayerAllocations.RemoveAll(
+			[](const FWeightmapLayerAllocationInfo& Allocation) { return Allocation.WeightmapTextureIndex == 255; });
+
+		// remove any textures we're no longer using
+		for (int32 TextureIdx = 0; TextureIdx < Component->WeightmapTextures.Num(); ++TextureIdx)
+		{
+			if (!Component->WeightmapLayerAllocations.ContainsByPredicate(
+				[TextureIdx](const FWeightmapLayerAllocationInfo& Allocation) { return Allocation.WeightmapTextureIndex == TextureIdx; }))
+			{
+				Component->WeightmapTextures[TextureIdx]->Modify();
+
+				for (int32 LayerIdx = 0; LayerIdx < Component->WeightmapLayerAllocations.Num(); ++LayerIdx)
+				{
+					FWeightmapLayerAllocationInfo& Allocation = Component->WeightmapLayerAllocations[LayerIdx];
+					if (Allocation.WeightmapTextureIndex > TextureIdx)
+					{
+						--Allocation.WeightmapTextureIndex;
+					}
+				}
+
+				Component->WeightmapTextures.RemoveAt(TextureIdx--);
+			}
+		}
+	}
+
+	// todo - normalize texture usage: it's possible to end up with two textures each using one channel at this point
+	// e.g. if you start with 4 blended layers (in one texture) and a non-blended layer (in a 2nd), and fill one weight-blended layer (deleting the other three)
+	// this can also happen with normal painting I believe
+
+	// update mips
+	for (int32 TextureIdx = 0; TextureIdx < Component->WeightmapTextures.Num(); ++TextureIdx)
+	{
+		UTexture2D* WeightmapTexture = WeightmapTextures[TextureIdx];
+		FLandscapeTextureDataInfo* WeightmapDataInfo = LandscapeEdit.GetTextureDataInfo(WeightmapTexture);
+
+		const int32 NumMips = WeightmapTexture->Source.GetNumMips();
+		TArray<FColor*> WeightmapTextureMipData;
+		WeightmapTextureMipData.AddUninitialized(NumMips);
+		for (int32 MipIdx = 0; MipIdx < NumMips; MipIdx++)
+		{
+			WeightmapTextureMipData[MipIdx] = (FColor*)WeightmapDataInfo->GetMipData(MipIdx);
+		}
+
+		ULandscapeComponent::UpdateWeightmapMips(NumSubsections, SubsectionSizeQuads, WeightmapTexture, WeightmapTextureMipData, 0, 0, MAX_int32, MAX_int32, WeightmapDataInfo);
+
+		WeightmapDataInfo->AddMipUpdateRegion(0, 0, 0, WeightmapTexture->Source.GetSizeX() - 1, WeightmapTexture->Source.GetSizeY() - 1);
+	}
+
+	// Update the shaders for this component
+	Component->UpdateMaterialInstances();
+
+	Component->InvalidateLightingCache();
+
+	// Update dominant layer info stored in collision component
+	TArray<FColor*> CollisionWeightmapMipData;
+	for (int32 WeightmapIdx = 0; WeightmapIdx < Component->WeightmapTextures.Num(); WeightmapIdx++)
+	{
+		CollisionWeightmapMipData.Add((FColor*)LandscapeEdit.GetTextureDataInfo(Component->WeightmapTextures[WeightmapIdx])->GetMipData(Component->CollisionMipLevel));
+	}
+	TArray<FColor*> SimpleCollisionWeightmapMipData;
+	if (Component->SimpleCollisionMipLevel > Component->CollisionMipLevel)
+	{
+		for (int32 WeightmapIdx = 0; WeightmapIdx < Component->WeightmapTextures.Num(); WeightmapIdx++)
+		{
+			SimpleCollisionWeightmapMipData.Add((FColor*)LandscapeEdit.GetTextureDataInfo(Component->WeightmapTextures[WeightmapIdx])->GetMipData(Component->SimpleCollisionMipLevel));
+		}
+	}
+	Component->UpdateCollisionLayerData(
+		CollisionWeightmapMipData.GetData(),
+		Component->SimpleCollisionMipLevel > Component->CollisionMipLevel ? SimpleCollisionWeightmapMipData.GetData() : nullptr);
+}
+
+void FLandscapeEditDataInterface::FillLayer(ULandscapeLayerInfoObject* LayerInfo)
+{
+	if (!LandscapeInfo)
+	{
+		return;
+	}
+
 	for (auto& XYComponentPair : LandscapeInfo->XYtoComponentMap)
 	{
 		ULandscapeComponent* Component = XYComponentPair.Value;
-		Component->DeleteLayer(LayerInfo, this);
+		Component->FillLayer(LayerInfo, *this);
+	}
+
+	// Flush dynamic data (e.g. grass)
+	TSet<ULandscapeComponent*> Components;
+	for (auto& XYComponentPair : LandscapeInfo->XYtoComponentMap)
+	{
+		Components.Add(XYComponentPair.Value);
 	}
+	ALandscapeProxy::InvalidateGeneratedComponentData(Components);
 }
 
-void ULandscapeComponent::ReplaceLayer(ULandscapeLayerInfoObject* FromLayerInfo, ULandscapeLayerInfoObject* ToLayerInfo, struct FLandscapeEditDataInterface* LandscapeEdit)
+void ULandscapeComponent::ReplaceLayer(ULandscapeLayerInfoObject* FromLayerInfo, ULandscapeLayerInfoObject* ToLayerInfo, FLandscapeEditDataInterface& LandscapeEdit)
 {
-	check(FromLayerInfo && ToLayerInfo && LandscapeEdit);
+	check(FromLayerInfo && ToLayerInfo);
 	if (FromLayerInfo == ToLayerInfo)
 	{
 		return;
@@ -1763,8 +1960,8 @@ void ULandscapeComponent::ReplaceLayer(ULandscapeLayerInfoObject* FromLayerInfo,
 		FWeightmapLayerAllocationInfo& ToLayerAllocation = WeightmapLayerAllocations[ToLayerIdx];
 
 		// Lock data for all the weightmaps
-		FLandscapeTextureDataInfo* FromTexDataInfo = LandscapeEdit->GetTextureDataInfo(WeightmapTextures[FromLayerAllocation.WeightmapTextureIndex]);
-		FLandscapeTextureDataInfo* ToTexDataInfo = LandscapeEdit->GetTextureDataInfo(WeightmapTextures[ToLayerAllocation.WeightmapTextureIndex]);
+		FLandscapeTextureDataInfo* FromTexDataInfo = LandscapeEdit.GetTextureDataInfo(WeightmapTextures[FromLayerAllocation.WeightmapTextureIndex]);
+		FLandscapeTextureDataInfo* ToTexDataInfo = LandscapeEdit.GetTextureDataInfo(WeightmapTextures[ToLayerAllocation.WeightmapTextureIndex]);
 
 		check(FromTexDataInfo->GetMipSizeX(0) == FromTexDataInfo->GetMipSizeY(0));
 		check(ToTexDataInfo->GetMipSizeX(0) == ToTexDataInfo->GetMipSizeY(0));
@@ -1877,7 +2074,7 @@ void FLandscapeEditDataInterface::ReplaceLayer(ULandscapeLayerInfoObject* FromLa
 	for( auto It = LandscapeInfo->XYtoComponentMap.CreateIterator(); It; ++It )
 	{
 		ULandscapeComponent* Component = It.Value();
-		Component->ReplaceLayer(FromLayerInfo, ToLayerInfo, this);
+		Component->ReplaceLayer(FromLayerInfo, ToLayerInfo, *this);
 
 		// Update dominant layer info stored in collision component
 		TArray<FColor*> CollisionWeightmapMipData;
diff --git a/Engine/Source/Runtime/Landscape/Private/LandscapeGizmoActor.cpp b/Engine/Source/Runtime/Landscape/Private/LandscapeGizmoActor.cpp
--- a/Engine/Source/Runtime/Landscape/Private/LandscapeGizmoActor.cpp
+++ b/Engine/Source/Runtime/Landscape/Private/LandscapeGizmoActor.cpp
@@ -12,6 +12,7 @@
 #include "DynamicMeshBuilder.h"
 #include "Engine/CollisionProfile.h"
 #include "EngineUtils.h"
+#include "Materials/MaterialInstanceConstant.h"
 
 namespace
 {
diff --git a/Engine/Source/Runtime/Landscape/Private/LandscapeGrass.cpp b/Engine/Source/Runtime/Landscape/Private/LandscapeGrass.cpp
--- a/Engine/Source/Runtime/Landscape/Private/LandscapeGrass.cpp
+++ b/Engine/Source/Runtime/Landscape/Private/LandscapeGrass.cpp
@@ -24,6 +24,8 @@
 #include "LandscapeVersion.h"
 #include "Algo/Accumulate.h"
 #include "LandscapeLight.h"
+#include "Components/HierarchicalInstancedStaticMeshComponent.h"
+#include "Materials/MaterialInstanceConstant.h"
 
 #define LOCTEXT_NAMESPACE "Landscape"
 
@@ -2352,6 +2354,13 @@ void ALandscapeProxy::UpdateGrass(const TArray<FVector>& Cameras, bool bForceSyn
 	}
 }
 
+FAsyncGrassTask::FAsyncGrassTask(FAsyncGrassBuilder* InBuilder, const FCachedLandscapeFoliage::FGrassCompKey& InKey, UHierarchicalInstancedStaticMeshComponent* InFoliage)
+	: Builder(InBuilder)
+	, Key(InKey)
+	, Foliage(InFoliage)
+{
+}
+
 void FAsyncGrassTask::DoWork()
 {
 	Builder->Build();
@@ -2364,25 +2373,17 @@ FAsyncGrassTask::~FAsyncGrassTask()
 
 static void FlushGrass(const TArray<FString>& Args)
 {
-	for (TObjectIterator<ALandscapeProxy> It; It; ++It)
+	for (ALandscapeProxy* Landscape : TObjectRange<ALandscapeProxy>(RF_ClassDefaultObject | RF_ArchetypeObject, true, EInternalObjectFlags::PendingKill))
 	{
-		ALandscapeProxy* Landscape = *It;
-		if (Landscape && !Landscape->IsTemplate() && !Landscape->IsPendingKill())
-		{
-			Landscape->FlushGrassComponents();
-		}
+		Landscape->FlushGrassComponents();
 	}
 }
 
 static void FlushGrassPIE(const TArray<FString>& Args)
 {
-	for (TObjectIterator<ALandscapeProxy> It; It; ++It)
+	for (ALandscapeProxy* Landscape : TObjectRange<ALandscapeProxy>(RF_ClassDefaultObject | RF_ArchetypeObject, true, EInternalObjectFlags::PendingKill))
 	{
-		ALandscapeProxy* Landscape = *It;
-		if (Landscape && !Landscape->IsTemplate() && !Landscape->IsPendingKill())
-		{
-			Landscape->FlushGrassComponents(nullptr, false);
-		}
+		Landscape->FlushGrassComponents(nullptr, false);
 	}
 }
 
diff --git a/Engine/Source/Runtime/Landscape/Private/LandscapeLight.cpp b/Engine/Source/Runtime/Landscape/Private/LandscapeLight.cpp
--- a/Engine/Source/Runtime/Landscape/Private/LandscapeLight.cpp
+++ b/Engine/Source/Runtime/Landscape/Private/LandscapeLight.cpp
@@ -10,9 +10,9 @@ LandscapeLight.cpp: Static lighting for LandscapeComponents
 #include "LandscapeRender.h"
 #include "LandscapeDataAccess.h"
 
-#include "ComponentReregisterContext.h"
-
 #include "UnrealEngine.h"
+#include "ComponentReregisterContext.h"
+#include "Materials/MaterialInstanceConstant.h"
 
 #if WITH_EDITOR
 
diff --git a/Engine/Source/Runtime/Landscape/Private/LandscapeModule.cpp b/Engine/Source/Runtime/Landscape/Private/LandscapeModule.cpp
--- a/Engine/Source/Runtime/Landscape/Private/LandscapeModule.cpp
+++ b/Engine/Source/Runtime/Landscape/Private/LandscapeModule.cpp
@@ -7,6 +7,8 @@
 #include "Engine/World.h"
 #include "Engine/Level.h"
 #include "LandscapeVersion.h"
+#include "Materials/Material.h"
+#include "Materials/MaterialInstance.h"
 
 // Register the custom version with core
 FCustomVersionRegistration GRegisterLandscapeCustomVersion(FLandscapeCustomVersion::GUID, FLandscapeCustomVersion::LatestVersion, TEXT("Landscape"));
diff --git a/Engine/Source/Runtime/Landscape/Private/LandscapeProxy.cpp b/Engine/Source/Runtime/Landscape/Private/LandscapeProxy.cpp
--- a/Engine/Source/Runtime/Landscape/Private/LandscapeProxy.cpp
+++ b/Engine/Source/Runtime/Landscape/Private/LandscapeProxy.cpp
@@ -6,11 +6,11 @@
 
 #if WITH_EDITOR
 
-LANDSCAPE_API FLandscapeImportLayerInfo::FLandscapeImportLayerInfo(const struct FLandscapeInfoLayerSettings& InLayerSettings)
-:	LayerName(InLayerSettings.GetLayerName())
-,	LayerInfo(InLayerSettings.LayerInfoObj)
-,	ThumbnailMIC(NULL)
-,	SourceFilePath(InLayerSettings.GetEditorSettings().ReimportLayerFilePath)
-{ }
+LANDSCAPE_API FLandscapeImportLayerInfo::FLandscapeImportLayerInfo(const FLandscapeInfoLayerSettings& InLayerSettings)
+	: LayerName(InLayerSettings.GetLayerName())
+	, LayerInfo(InLayerSettings.LayerInfoObj)
+	, SourceFilePath(InLayerSettings.GetEditorSettings().ReimportLayerFilePath)
+{
+}
 
 #endif // WITH_EDITOR
\ No newline at end of file
diff --git a/Engine/Source/Runtime/Landscape/Private/LandscapeRender.cpp b/Engine/Source/Runtime/Landscape/Private/LandscapeRender.cpp
--- a/Engine/Source/Runtime/Landscape/Private/LandscapeRender.cpp
+++ b/Engine/Source/Runtime/Landscape/Private/LandscapeRender.cpp
@@ -23,6 +23,7 @@ LandscapeRender.cpp: New terrain rendering
 #include "EngineGlobals.h"
 #include "UnrealEngine.h"
 #include "LandscapeLight.h"
+#include "LandscapeLayerInfoObject.h"
 
 IMPLEMENT_UNIFORM_BUFFER_STRUCT(FLandscapeUniformShaderParameters, TEXT("LandscapeParameters"));
 
@@ -450,12 +451,12 @@ LANDSCAPE_API bool GLandscapeEditModeActive = false;
 LANDSCAPE_API ELandscapeViewMode::Type GLandscapeViewMode = ELandscapeViewMode::Normal;
 LANDSCAPE_API int32 GLandscapeEditRenderMode = ELandscapeEditRenderMode::None;
 LANDSCAPE_API int32 GLandscapePreviewMeshRenderMode = 0;
-UMaterial* GLayerDebugColorMaterial = nullptr;
-UMaterialInstanceConstant* GSelectionColorMaterial = nullptr;
-UMaterialInstanceConstant* GSelectionRegionMaterial = nullptr;
-UMaterialInstanceConstant* GMaskRegionMaterial = nullptr;
+UMaterialInterface* GLayerDebugColorMaterial = nullptr;
+UMaterialInterface* GSelectionColorMaterial = nullptr;
+UMaterialInterface* GSelectionRegionMaterial = nullptr;
+UMaterialInterface* GMaskRegionMaterial = nullptr;
 UTexture2D* GLandscapeBlackTexture = nullptr;
-UMaterial* GLandscapeLayerUsageMaterial = nullptr;
+UMaterialInterface* GLandscapeLayerUsageMaterial = nullptr;
 
 // Game thread update
 void FLandscapeEditToolRenderData::Update(UMaterialInterface* InToolMaterial)
@@ -642,7 +643,6 @@ FLandscapeComponentSceneProxy::FLandscapeComponentSceneProxy(ULandscapeComponent
 
 	// Check material usage
 	if (MaterialInterface == nullptr ||
-		!MaterialInterface->CheckMaterialUsage(MATUSAGE_Landscape) ||
 		(bHasStaticLighting && !MaterialInterface->CheckMaterialUsage(MATUSAGE_StaticLighting)))
 	{
 		MaterialInterface = UMaterial::GetDefaultMaterial(MD_Surface);
@@ -1631,8 +1631,6 @@ void FLandscapeComponentSceneProxy::GetDynamicMeshElements(const TArray<const FS
 
 			default:
 
-#else
-			{
 #endif // WITH_EDITOR
 
 #if WITH_EDITOR || !(UE_BUILD_SHIPPING || UE_BUILD_TEST)
@@ -1682,7 +1680,10 @@ void FLandscapeComponentSceneProxy::GetDynamicMeshElements(const TArray<const FS
 					NumTriangles += Mesh.GetNumPrimitives();
 					NumDrawCalls += Mesh.Elements.Num();
 				}
-			}
+
+#if WITH_EDITOR
+			} // switch
+#endif
 
 #if WITH_EDITOR
 			// Extra render passes for landscape tools
@@ -2541,6 +2542,104 @@ ULandscapeMaterialInstanceConstant::ULandscapeMaterialInstanceConstant(const FOb
 	bIsLayerThumbnail = false;
 }
 
+class FLandscapeMaterialResource : public FMaterialResource
+{
+	bool bIsLayerThumbnail;
+
+public:
+	FLandscapeMaterialResource(bool InbIsLayerThumbnail)
+		: bIsLayerThumbnail(InbIsLayerThumbnail)
+	{
+	}
+
+	bool IsUsedWithLandscape() const override
+	{
+		return !bIsLayerThumbnail;
+	}
+
+	bool IsUsedWithStaticLighting() const override
+	{
+		if (bIsLayerThumbnail)
+		{
+			return false;
+		}
+		return FMaterialResource::IsUsedWithStaticLighting();
+	}
+
+	bool IsUsedWithSkeletalMesh()          const override { return false; }
+	bool IsUsedWithParticleSystem()        const override { return false; }
+	bool IsUsedWithParticleSprites()       const override { return false; }
+	bool IsUsedWithBeamTrails()            const override { return false; }
+	bool IsUsedWithMeshParticles()         const override { return false; }
+	bool IsUsedWithMorphTargets()          const override { return false; }
+	bool IsUsedWithSplineMeshes()          const override { return false; }
+	bool IsUsedWithInstancedStaticMeshes() const override { return false; }
+	bool IsUsedWithAPEXCloth()             const override { return false; }
+
+	bool ShouldCache(EShaderPlatform Platform, const FShaderType* ShaderType, const FVertexFactoryType* VertexFactoryType) const override
+	{
+		//return FMaterialResource::ShouldCache(Platform, ShaderType, VertexFactoryType);
+
+		static const FName LocalVertexFactory = FName(TEXT("FLocalVertexFactory"));
+		static const FName LandscapeVertexFactory = FName(TEXT("FLandscapeVertexFactory"));
+		static const FName LandscapeVertexFactoryMobile = FName(TEXT("FLandscapeVertexFactoryMobile"));
+
+		static const FName TBasePassVSFNoLightMapPolicy = FName(TEXT("TBasePassVSFNoLightMapPolicy"));
+		static const FName TBasePassPSFNoLightMapPolicy = FName(TEXT("TBasePassPSFNoLightMapPolicy"));
+		static const FName TBasePassVSFCachedPointIndirectLightingPolicy = FName(TEXT("TBasePassVSFCachedPointIndirectLightingPolicy"));
+		static const FName TBasePassPSFCachedPointIndirectLightingPolicy = FName(TEXT("TBasePassPSFCachedPointIndirectLightingPolicy"));
+		static const FName TShadowDepthVSVertexShadowDepth_OutputDepthfalse = FName(TEXT("TShadowDepthVSVertexShadowDepth_OutputDepthfalse"));
+
+		if (VertexFactoryType)
+		{
+			if (bIsLayerThumbnail)
+			{
+				// Thumbnail MICs are only rendered in the preview scene using a simple LocalVertexFactory
+				if (VertexFactoryType->GetFName() == LocalVertexFactory)
+				{
+					if (ShaderType->GetFName() == TBasePassVSFNoLightMapPolicy ||
+						ShaderType->GetFName() == TBasePassPSFNoLightMapPolicy ||
+						ShaderType->GetFName() == TBasePassVSFCachedPointIndirectLightingPolicy ||
+						ShaderType->GetFName() == TBasePassPSFCachedPointIndirectLightingPolicy ||
+						ShaderType->GetFName() == TShadowDepthVSVertexShadowDepth_OutputDepthfalse)
+					{
+						return true;
+					}
+				}
+			}
+			else
+			{
+				// Landscape MICs are only for use with the Landscape vertex factories
+				if (VertexFactoryType->GetFName() == LandscapeVertexFactory ||
+					VertexFactoryType->GetFName() == LandscapeVertexFactoryMobile)
+				{
+					return true;
+				}
+			}
+		}
+
+		return false;
+	}
+};
+
+FMaterialResource* ULandscapeMaterialInstanceConstant::AllocatePermutationResource()
+{
+	return new FLandscapeMaterialResource(bIsLayerThumbnail);
+}
+
+bool ULandscapeMaterialInstanceConstant::HasOverridenBaseProperties() const
+{
+	// force a static permutation for ULandscapeMaterialInstanceConstants
+	if (Parent && !Parent->IsA<ULandscapeMaterialInstanceConstant>())
+	{
+		return true;
+	}
+
+	return Super::HasOverridenBaseProperties();
+}
+
+//////////////////////////////////////////////////////////////////////////
+
 void ULandscapeComponent::GetStreamingTextureInfo(FStreamingTextureLevelContext& LevelContext, TArray<FStreamingTexturePrimitiveInfo>& OutStreamingTextures) const
 {
 	ALandscapeProxy* Proxy = Cast<ALandscapeProxy>(GetOuter());
diff --git a/Engine/Source/Runtime/Landscape/Private/LandscapeSplineRaster.cpp b/Engine/Source/Runtime/Landscape/Private/LandscapeSplineRaster.cpp
--- a/Engine/Source/Runtime/Landscape/Private/LandscapeSplineRaster.cpp
+++ b/Engine/Source/Runtime/Landscape/Private/LandscapeSplineRaster.cpp
@@ -7,13 +7,15 @@
 
 #include "Landscape.h"
 #include "LandscapeSplineRaster.h"
-#include "LandscapeProxy.h"
+#include "LandscapeStreamingProxy.h"
 #include "LandscapeDataAccess.h"
 #include "LandscapeEdit.h"
 #include "LandscapeHeightfieldCollisionComponent.h"
 #include "LandscapeSplinesComponent.h"
 #include "LandscapeSplineControlPoint.h"
 #include "LandscapeSplineSegment.h"
+#include "LandscapeLayerInfoObject.h"
+#include "LandscapeComponent.h"
 #if WITH_EDITOR
 #include "AI/Navigation/NavigationSystem.h"
 #include "ScopedTransaction.h"
@@ -424,9 +426,8 @@ bool ULandscapeInfo::ApplySplines(bool bOnlySelected)
 
 	bResult |= ApplySplinesInternal(bOnlySelected, Landscape);
 
-	for (auto It = Proxies.CreateIterator(); It; ++It)
+	for (ALandscapeProxy* LandscapeProxy : Proxies)
 	{
-		ALandscapeProxy* LandscapeProxy = (*It);
 		bResult |= ApplySplinesInternal(bOnlySelected, LandscapeProxy);
 	}
 
diff --git a/Engine/Source/Runtime/Landscape/Private/LandscapeSplines.cpp b/Engine/Source/Runtime/Landscape/Private/LandscapeSplines.cpp
--- a/Engine/Source/Runtime/Landscape/Private/LandscapeSplines.cpp
+++ b/Engine/Source/Runtime/Landscape/Private/LandscapeSplines.cpp
@@ -6,6 +6,8 @@
 
 #include "Landscape.h"
 #include "Components/SplineMeshComponent.h"
+#include "LandscapeStreamingProxy.h"
+#include "LandscapeComponent.h"
 #include "ShaderParameterUtils.h"
 #include "StaticMeshResources.h"
 #include "LandscapeSplineProxies.h"
@@ -1441,12 +1443,20 @@ void ULandscapeSplineControlPoint::UpdateSplinePoints(bool bUpdateCollision, boo
 			}
 		}
 
-		if (MeshComponent->RelativeLocation != Location ||
-			MeshComponent->RelativeRotation != Rotation ||
+		FVector MeshLocation = Location;
+		FRotator MeshRotation = Rotation;
+		if (MeshComponentOuterSplines != OuterSplines)
+		{
+			const FTransform RelativeTransform = OuterSplines->ComponentToWorld.GetRelativeTransform(MeshComponentOuterSplines->ComponentToWorld);
+			MeshLocation = RelativeTransform.TransformPosition(MeshLocation);
+		}
+
+		if (MeshComponent->RelativeLocation != MeshLocation ||
+			MeshComponent->RelativeRotation != MeshRotation ||
 			MeshComponent->RelativeScale3D != MeshScale)
 		{
 			MeshComponent->Modify();
-			MeshComponent->SetRelativeTransform(FTransform(Rotation, Location, MeshScale));
+			MeshComponent->SetRelativeTransform(FTransform(MeshRotation, MeshLocation, MeshScale));
 			MeshComponent->InvalidateLightingCache();
 		}
 
@@ -2296,9 +2306,10 @@ void ULandscapeSplineSegment::UpdateSplinePoints(bool bUpdateCollision)
 			MeshComponent->SplineUpDir = FVector(0,0,1); // Up, to be consistent between joined meshes. We rotate it to horizontal using roll if using Z Forward X Up or X Forward Y Up
 			MeshComponent->ForwardAxis = MeshEntry->ForwardAxis;
 
-			if (MeshComponent->GetAttachParent() != OuterSplines)
+			auto* const MeshComponentOuterSplines = MeshComponent->GetAttachParent();
+			if (MeshComponentOuterSplines != OuterSplines)
 			{
-				FTransform RelativeTransform = OuterSplines->ComponentToWorld.GetRelativeTransform(MeshComponent->GetAttachParent()->ComponentToWorld);
+				const FTransform RelativeTransform = OuterSplines->ComponentToWorld.GetRelativeTransform(MeshComponentOuterSplines->ComponentToWorld);
 				MeshComponent->SplineParams.StartPos = RelativeTransform.TransformPosition(MeshComponent->SplineParams.StartPos);
 				MeshComponent->SplineParams.EndPos   = RelativeTransform.TransformPosition(MeshComponent->SplineParams.EndPos);
 			}
diff --git a/Engine/Source/Runtime/Landscape/Public/LandscapeEdit.h b/Engine/Source/Runtime/Landscape/Public/LandscapeEdit.h
--- a/Engine/Source/Runtime/Landscape/Public/LandscapeEdit.h
+++ b/Engine/Source/Runtime/Landscape/Public/LandscapeEdit.h
@@ -119,6 +119,8 @@ struct LANDSCAPE_API FLandscapeEditDataInterface
 	void SetAlphaData(const TSet<ULandscapeLayerInfoObject*>& DirtyLayerInfos, const int32 X1, const int32 Y1, const int32 X2, const int32 Y2, const uint8* Data, int32 Stride, ELandscapeLayerPaintingRestriction PaintingRestriction = ELandscapeLayerPaintingRestriction::None);
 	// Delete a layer and re-normalize other layers
 	void DeleteLayer(ULandscapeLayerInfoObject* LayerInfo);
+	// Fill a layer and re-normalize other layers
+	void FillLayer(ULandscapeLayerInfoObject* LayerInfo);
 	// Replace/merge a layer
 	void ReplaceLayer(ULandscapeLayerInfoObject* FromLayerInfo, ULandscapeLayerInfoObject* ToLayerInfo);
 
diff --git a/Engine/Source/Runtime/Landscape/Public/LandscapeLight.h b/Engine/Source/Runtime/Landscape/Public/LandscapeLight.h
--- a/Engine/Source/Runtime/Landscape/Public/LandscapeLight.h
+++ b/Engine/Source/Runtime/Landscape/Public/LandscapeLight.h
@@ -7,6 +7,7 @@
 #pragma once
 
 #include "StaticLighting.h"
+#include "RenderUtils.h"
 
 class ULandscapeComponent;
 
diff --git a/Engine/Source/Runtime/Landscape/Public/LandscapeRender.h b/Engine/Source/Runtime/Landscape/Public/LandscapeRender.h
--- a/Engine/Source/Runtime/Landscape/Public/LandscapeRender.h
+++ b/Engine/Source/Runtime/Landscape/Public/LandscapeRender.h
@@ -12,6 +12,7 @@ LandscapeRender.h: New terrain rendering
 #include "LandscapeProxy.h"
 #include "LandscapeMeshProxyComponent.h"
 
+#include "Materials/MaterialInterface.h"
 #include "PrimitiveSceneProxy.h"
 #include "StaticMeshResources.h"
 
@@ -67,12 +68,12 @@ namespace ELandscapeEditRenderMode
 LANDSCAPE_API extern bool GLandscapeEditModeActive;
 LANDSCAPE_API extern int32 GLandscapeEditRenderMode;
 LANDSCAPE_API extern int32 GLandscapePreviewMeshRenderMode;
-LANDSCAPE_API extern UMaterial* GLayerDebugColorMaterial;
-LANDSCAPE_API extern UMaterialInstanceConstant* GSelectionColorMaterial;
-LANDSCAPE_API extern UMaterialInstanceConstant* GSelectionRegionMaterial;
-LANDSCAPE_API extern UMaterialInstanceConstant* GMaskRegionMaterial;
+LANDSCAPE_API extern UMaterialInterface* GLayerDebugColorMaterial;
+LANDSCAPE_API extern UMaterialInterface* GSelectionColorMaterial;
+LANDSCAPE_API extern UMaterialInterface* GSelectionRegionMaterial;
+LANDSCAPE_API extern UMaterialInterface* GMaskRegionMaterial;
 LANDSCAPE_API extern UTexture2D* GLandscapeBlackTexture;
-LANDSCAPE_API extern UMaterial* GLandscapeLayerUsageMaterial;
+LANDSCAPE_API extern UMaterialInterface* GLandscapeLayerUsageMaterial;
 #endif
 
 
