Fix KiteDemo and NonUnity build failures

diff --git a/Engine/Source/Runtime/Landscape/Classes/LandscapeHeightfieldCollisionComponent.h b/Engine/Source/Runtime/Landscape/Classes/LandscapeHeightfieldCollisionComponent.h
--- a/Engine/Source/Runtime/Landscape/Classes/LandscapeHeightfieldCollisionComponent.h
+++ b/Engine/Source/Runtime/Landscape/Classes/LandscapeHeightfieldCollisionComponent.h
@@ -12,6 +12,7 @@ class ALandscapeProxy;
 class AInstancedFoliageActor;
 
 #include "AI/Navigation/NavigationTypes.h"
+#include "Components/PrimitiveComponent.h"
 #include "LandscapeHeightfieldCollisionComponent.generated.h"
 
 #if WITH_PHYSX
diff --git a/Engine/Source/Runtime/Windows/D3D12RHI/Private/D3D12DiskCache.cpp b/Engine/Source/Runtime/Windows/D3D12RHI/Private/D3D12DiskCache.cpp
--- a/Engine/Source/Runtime/Windows/D3D12RHI/Private/D3D12DiskCache.cpp
+++ b/Engine/Source/Runtime/Windows/D3D12RHI/Private/D3D12DiskCache.cpp
@@ -16,7 +16,7 @@ static FAutoConsoleVariableRef CVarEnablePSOCache(
 	ECVF_RenderThreadSafe | ECVF_ReadOnly
 	);
 
-inline bool FDiskCacheInterface::IsInErrorState() const
+bool FDiskCacheInterface::IsInErrorState() const
 {
 	return !GEnablePSOCache || mInErrorState;
 }
