Fixing some more build failures

diff --git a/Engine/Source/Runtime/Landscape/Classes/LandscapeSplinesComponent.h b/Engine/Source/Runtime/Landscape/Classes/LandscapeSplinesComponent.h
--- a/Engine/Source/Runtime/Landscape/Classes/LandscapeSplinesComponent.h
+++ b/Engine/Source/Runtime/Landscape/Classes/LandscapeSplinesComponent.h
@@ -2,11 +2,13 @@
 
 #pragma once
 
+#include "Components/PrimitiveComponent.h"
 #include "LandscapeSplinesComponent.generated.h"
 
 // Forward declarations
 class ULandscapeSplineControlPoint;
 class ULandscapeSplineSegment;
+class UMeshComponent;
 class USplineMeshComponent;
 class UControlPointMeshComponent;
 
