_Another_ build fix

diff --git a/Engine/Source/Runtime/Landscape/Private/LandscapeCollision.cpp b/Engine/Source/Runtime/Landscape/Private/LandscapeCollision.cpp
--- a/Engine/Source/Runtime/Landscape/Private/LandscapeCollision.cpp
+++ b/Engine/Source/Runtime/Landscape/Private/LandscapeCollision.cpp
@@ -27,7 +27,7 @@
 #include "InstancedFoliageActor.h"
 #include "EngineUtils.h"
 #include "CookStats.h"
-#include "Materials/MaterialInstance.h"
+#include "Materials/MaterialInstanceConstant.h"
 
 #if ENABLE_COOK_STATS
 namespace LandscapeCollisionCookStats
