Copying //UE4/Dev-Landscape to //UE4/Dev-Main (Source: //UE4/Dev-Landscape @ 3053591)

diff --git a/Engine/Source/Runtime/Core/Public/Containers/ArrayView.h b/Engine/Source/Runtime/Core/Public/Containers/ArrayView.h
new file mode 100644
--- /dev/null
+++ b/Engine/Source/Runtime/Core/Public/Containers/ArrayView.h
@@ -0,0 +1,604 @@
+// Copyright 1998-2016 Epic Games, Inc. All Rights Reserved.
+
+#pragma once
+
+#include <initializer_list>
+
+namespace ArrayViewPrivate
+{
+	/**
+	 * Trait testing whether a type is compatible with the view type
+	 */
+	template <typename T, typename ElementType>
+	struct TIsCompatibleElementType
+	{
+	public:
+		/** NOTE:
+		 * The stars in the TPointerIsConvertibleFromTo test are *IMPORTANT*
+		 * They prevent TArrayView<Base>(TArray<Derived>&) from compiling!
+		 */
+		enum { Value = TPointerIsConvertibleFromTo<T*, ElementType* const>::Value };
+	};
+}
+
+/**
+ * Templated fixed-size view of another array
+ *
+ * A statically sized view of an array of typed elements.  Designed to allow functions to take either a fixed C array
+ * or a TArray with an arbitrary allocator as an argument when the function neither adds nor removes elements
+ *
+ * e.g.:
+ * int32 SumAll(TArrayView<const int32> array)
+ * {
+ *     return Algo::Accumulate(array);
+ * }
+ * 
+ * could be called as:
+ *     SumAll(MyTArray);
+ *     SumAll(MyCArray);
+ *     SumAll({1, 2, 3});
+ *     SumAll(MakeArrayView(Ptr, Num));
+ * 
+ * Note:
+ *   View classes are not const-propagating! If you want a view where the elements are const, you need "TArrayView<const T>" not "const TArrayView<T>"!
+ *
+ * Caution:
+ *   Treat a view like a *reference* to the elements in the array. DO NOT free or reallocate the array while the view exists!
+ */
+template<typename InElementType>
+class TArrayView
+{
+public:
+	using ElementType = InElementType;
+
+	/**
+	 * Constructor.
+	 */
+	TArrayView()
+		: DataPtr(nullptr)
+		, ArrayNum(0)
+	{
+	}
+
+private:
+	template <typename T>
+	using TIsCompatibleElementType = ArrayViewPrivate::TIsCompatibleElementType<T, ElementType>;
+
+public:
+	/**
+	 * Copy constructor from another view
+	 *
+	 * @param Other The source array view to copy
+	 */
+	template <typename OtherElementType,
+		typename = typename TEnableIf<TIsCompatibleElementType<OtherElementType>::Value>::Type>
+	FORCEINLINE TArrayView(const TArrayView<OtherElementType>& Other)
+		: DataPtr(Other.GetData())
+		, ArrayNum(Other.Num())
+	{
+	}
+
+	/**
+	 * Construct a view of a C array with a compatible element type
+	 *
+	 * @param Other The source array to view.
+	 */
+	template <typename OtherElementType, size_t Size,
+		typename = typename TEnableIf<TIsCompatibleElementType<OtherElementType>::Value>::Type>
+	FORCEINLINE TArrayView(OtherElementType (&Other)[Size])
+		: DataPtr(Other)
+		, ArrayNum((int32)Size)
+	{
+		static_assert(Size <= MAX_int32, "Array Size too large! Size is only int32 for compatibility with TArray!");
+	}
+
+	/**
+	 * Construct a view of a C array with a compatible element type
+	 *
+	 * @param Other The source array to view.
+	 */
+	template <typename OtherElementType, size_t Size,
+		typename = typename TEnableIf<TIsCompatibleElementType<const OtherElementType>::Value>::Type>
+		FORCEINLINE TArrayView(const OtherElementType (&Other)[Size])
+		: DataPtr(Other)
+		, ArrayNum((int32)Size)
+	{
+		static_assert(Size <= MAX_int32, "Array Size too large! Size is only int32 for compatibility with TArray!");
+	}
+
+	/**
+	 * Construct a view of a TArray with a compatible element type
+	 *
+	 * @param Other The source array to view.
+	 */
+	template <typename OtherElementType, typename OtherAllocator,
+		typename = typename TEnableIf<TIsCompatibleElementType<OtherElementType>::Value>::Type>
+	FORCEINLINE TArrayView(TArray<OtherElementType, OtherAllocator>& Other)
+		: DataPtr(Other.GetData())
+		, ArrayNum(Other.Num())
+	{
+	}
+
+	template <typename OtherElementType, typename OtherAllocator,
+		typename = typename TEnableIf<TIsCompatibleElementType<const OtherElementType>::Value>::Type>
+	FORCEINLINE TArrayView(const TArray<OtherElementType, OtherAllocator>& Other)
+		: DataPtr(Other.GetData())
+		, ArrayNum(Other.Num())
+	{
+	}
+
+	/**
+	 * Construct a view from an initializer list with a compatible element type
+	 *
+	 * @param List The initializer list to view.
+	 */
+	template <typename OtherElementType,
+		typename = typename TEnableIf<TIsCompatibleElementType<OtherElementType>::Value>::Type>
+	FORCEINLINE TArrayView(std::initializer_list<OtherElementType> List)
+		: DataPtr(&*List.begin())
+		, ArrayNum(List.size())
+	{
+	}
+
+	/**
+	 * Construct a view of an arbitrary pointer
+	 *
+	 * @param InData	The data to view
+	 * @param InCount	The number of elements
+	 */
+	template <typename OtherElementType,
+		typename = typename TEnableIf<TIsCompatibleElementType<OtherElementType>::Value>::Type>
+	FORCEINLINE TArrayView(OtherElementType* InData, int32 InCount)
+		: DataPtr(InData)
+		, ArrayNum(InCount)
+	{
+		check(ArrayNum >= 0);
+	}
+
+	/**
+	 * Assignment operator
+	 *
+	 * @param Other The source array view to assign from
+	 */
+	template <typename OtherElementType,
+		typename = typename TEnableIf<TIsCompatibleElementType<OtherElementType>::Value>::Type>
+	FORCEINLINE TArrayView& operator=(const TArrayView<OtherElementType>& Other)
+	{
+		if (this != &Other)
+		{
+			DataPtr = Other.DataPtr;
+			ArrayNum = Other.ArrayNum;
+		}
+		return *this;
+	}
+
+public:
+
+	/**
+	 * Helper function for returning a typed pointer to the first array entry.
+	 *
+	 * @returns Pointer to first array entry or nullptr if ArrayMax == 0.
+	 */
+	FORCEINLINE ElementType* GetData() const
+	{
+		return DataPtr;
+	}
+
+	/** 
+	 * Helper function returning the size of the inner type.
+	 *
+	 * @returns Size in bytes of array type.
+	 */
+	FORCEINLINE size_t GetTypeSize() const
+	{
+		return sizeof(ElementType);
+	}
+
+	/**
+	 * Checks array invariants: if array size is greater than zero and less
+	 * than maximum.
+	 */
+	FORCEINLINE void CheckInvariants() const
+	{
+		checkSlow(ArrayNum >= 0);
+	}
+
+	/**
+	 * Checks if index is in array range.
+	 *
+	 * @param Index Index to check.
+	 */
+	FORCEINLINE void RangeCheck(int32 Index) const
+	{
+		CheckInvariants();
+
+		checkf((Index >= 0) & (Index < ArrayNum),TEXT("Array index out of bounds: %i from an array of size %i"),Index,ArrayNum); // & for one branch
+	}
+
+	/**
+	 * Tests if index is valid, i.e. than or equal to zero, and less than the number of elements in the array.
+	 *
+	 * @param Index Index to test.
+	 *
+	 * @returns True if index is valid. False otherwise.
+	 */
+	FORCEINLINE bool IsValidIndex(int32 Index) const
+	{
+		return (Index >= 0) && (Index < ArrayNum);
+	}
+
+	/**
+	 * Returns number of elements in array.
+	 *
+	 * @returns Number of elements in array.
+	 */
+	FORCEINLINE int32 Num() const
+	{
+		return ArrayNum;
+	}
+
+	/**
+	 * Array bracket operator. Returns reference to element at give index.
+	 *
+	 * @returns Reference to indexed element.
+	 */
+	FORCEINLINE ElementType& operator[](int32 Index) const
+	{
+		RangeCheck(Index);
+		return GetData()[Index];
+	}
+
+	/**
+	 * Returns n-th last element from the array.
+	 *
+	 * @param IndexFromTheEnd (Optional) Index from the end of array.
+	 *                        Default is 0.
+	 *
+	 * @returns Reference to n-th last element from the array.
+	 */
+	FORCEINLINE ElementType& Last(int32 IndexFromTheEnd = 0) const
+	{
+		RangeCheck(ArrayNum - IndexFromTheEnd - 1);
+		return GetData()[ArrayNum - IndexFromTheEnd - 1];
+	}
+
+	/**
+	 * Returns a sliced view
+	 *
+	 * @param Index starting index of the new view
+	 * @param InNum number of elements in the new view
+	 * @returns Sliced view
+	 */
+	FORCEINLINE TArrayView Slice(int32 Index, int32 InNum)
+	{
+		check(InNum > 0);
+		check(IsValidIndex(Index));
+		check(IsValidIndex(Index + InNum - 1));
+		return TArrayView(DataPtr + Index, InNum);
+	}
+
+	/**
+	 * Finds element within the array.
+	 *
+	 * @param Item Item to look for.
+	 * @param Index Output parameter. Found index.
+	 *
+	 * @returns True if found. False otherwise.
+	 */
+	FORCEINLINE bool Find(const ElementType& Item, int32& Index) const
+	{
+		Index = this->Find(Item);
+		return Index != INDEX_NONE;
+	}
+
+	/**
+	 * Finds element within the array.
+	 *
+	 * @param Item Item to look for.
+	 *
+	 * @returns Index of the found element. INDEX_NONE otherwise.
+	 */
+	int32 Find(const ElementType& Item) const
+	{
+		const ElementType* RESTRICT Start = GetData();
+		for (const ElementType* RESTRICT Data = Start, *RESTRICT DataEnd = Data + ArrayNum; Data != DataEnd; ++Data)
+		{
+			if (*Data == Item)
+			{
+				return static_cast<int32>(Data - Start);
+			}
+		}
+		return INDEX_NONE;
+	}
+
+	/**
+	 * Finds element within the array starting from the end.
+	 *
+	 * @param Item Item to look for.
+	 * @param Index Output parameter. Found index.
+	 *
+	 * @returns True if found. False otherwise.
+	 */
+	FORCEINLINE bool FindLast(const ElementType& Item, int32& Index) const
+	{
+		Index = this->FindLast(Item);
+		return Index != INDEX_NONE;
+	}
+
+	/**
+	 * Finds element within the array starting from StartIndex and going backwards. Uses predicate to match element.
+	 *
+	 * @param Pred Predicate taking array element and returns true if element matches search criteria, false otherwise.
+	 * @param StartIndex Index of element from which to start searching.
+	 *
+	 * @returns Index of the found element. INDEX_NONE otherwise.
+	 */
+	template <typename Predicate>
+	int32 FindLastByPredicate(Predicate Pred, int32 StartIndex) const
+	{
+		check(StartIndex >= 0 && StartIndex <= this->Num());
+		for (const ElementType* RESTRICT Start = GetData(), *RESTRICT Data = Start + StartIndex; Data != Start; )
+		{
+			--Data;
+			if (Pred(*Data))
+			{
+				return static_cast<int32>(Data - Start);
+			}
+		}
+		return INDEX_NONE;
+	}
+
+	/**
+	* Finds element within the array starting from the end. Uses predicate to match element.
+	*
+	* @param Pred Predicate taking array element and returns true if element matches search criteria, false otherwise.
+	*
+	* @returns Index of the found element. INDEX_NONE otherwise.
+	*/
+	template <typename Predicate>
+	FORCEINLINE int32 FindLastByPredicate(Predicate Pred) const
+	{
+		return FindLastByPredicate(Pred, ArrayNum);
+	}
+
+	/**
+	 * Finds an item by key (assuming the ElementType overloads operator== for
+	 * the comparison).
+	 *
+	 * @param Key The key to search by.
+	 *
+	 * @returns Index to the first matching element, or INDEX_NONE if none is
+	 *          found.
+	 */
+	template <typename KeyType>
+	int32 IndexOfByKey(const KeyType& Key) const
+	{
+		const ElementType* RESTRICT Start = GetData();
+		for (const ElementType* RESTRICT Data = Start, *RESTRICT DataEnd = Start + ArrayNum; Data != DataEnd; ++Data)
+		{
+			if (*Data == Key)
+			{
+				return static_cast<int32>(Data - Start);
+			}
+		}
+		return INDEX_NONE;
+	}
+
+	/**
+	 * Finds an item by predicate.
+	 *
+	 * @param Pred The predicate to match.
+	 *
+	 * @returns Index to the first matching element, or INDEX_NONE if none is
+	 *          found.
+	 */
+	template <typename Predicate>
+	int32 IndexOfByPredicate(Predicate Pred) const
+	{
+		const ElementType* RESTRICT Start = GetData();
+		for (const ElementType* RESTRICT Data = Start, *RESTRICT DataEnd = Start + ArrayNum; Data != DataEnd; ++Data)
+		{
+			if (Pred(*Data))
+			{
+				return static_cast<int32>(Data - Start);
+			}
+		}
+		return INDEX_NONE;
+	}
+
+	/**
+	 * Finds an item by key (assuming the ElementType overloads operator== for
+	 * the comparison).
+	 *
+	 * @param Key The key to search by.
+	 *
+	 * @returns Pointer to the first matching element, or nullptr if none is found.
+	 */
+	template <typename KeyType>
+	ElementType* FindByKey(const KeyType& Key) const
+	{
+		for (ElementType* RESTRICT Data = GetData(), *RESTRICT DataEnd = Data + ArrayNum; Data != DataEnd; ++Data)
+		{
+			if (*Data == Key)
+			{
+				return Data;
+			}
+		}
+
+		return nullptr;
+	}
+
+	/**
+	 * Finds an element which matches a predicate functor.
+	 *
+	 * @param Pred The functor to apply to each element.
+	 *
+	 * @return Pointer to the first element for which the predicate returns
+	 *         true, or nullptr if none is found.
+	 */
+	template <typename Predicate>
+	ElementType* FindByPredicate(Predicate Pred) const
+	{
+		for (ElementType* RESTRICT Data = GetData(), *RESTRICT DataEnd = Data + ArrayNum; Data != DataEnd; ++Data)
+		{
+			if (Pred(*Data))
+			{
+				return Data;
+			}
+		}
+
+		return nullptr;
+	}
+
+	/**
+	 * Filters the elements in the array based on a predicate functor.
+	 *
+	 * @param Pred The functor to apply to each element.
+	 *
+	 * @returns TArray with the same type as this object which contains
+	 *          the subset of elements for which the functor returns true.
+	 */
+	template <typename Predicate>
+	TArray<typename TRemoveConst<ElementType>::Type> FilterByPredicate(Predicate Pred) const
+	{
+		TArray<typename TRemoveConst<ElementType>::Type> FilterResults;
+		for (const ElementType* RESTRICT Data = GetData(), *RESTRICT DataEnd = Data + ArrayNum; Data != DataEnd; ++Data)
+		{
+			if (Pred(*Data))
+			{
+				FilterResults.Add(*Data);
+			}
+		}
+		return FilterResults;
+	}
+
+	/**
+	 * Checks if this array contains the element.
+	 *
+	 * @returns	True if found. False otherwise.
+	 */
+	template <typename ComparisonType>
+	bool Contains(const ComparisonType& Item) const
+	{
+		for (const ElementType* RESTRICT Data = GetData(), *RESTRICT DataEnd = Data + ArrayNum; Data != DataEnd; ++Data)
+		{
+			if (*Data == Item)
+			{
+				return true;
+			}
+		}
+		return false;
+	}
+
+	/**
+	 * Checks if this array contains element for which the predicate is true.
+	 *
+	 * @param Predicate to use
+	 *
+	 * @returns	True if found. False otherwise.
+	 */
+	template <typename Predicate>
+	FORCEINLINE bool ContainsByPredicate(Predicate Pred) const
+	{
+		return FindByPredicate(Pred) != nullptr;
+	}
+
+private:
+	/**
+	 * DO NOT USE DIRECTLY
+	 * STL-like iterators to enable range-based for loop support.
+	 */
+	FORCEINLINE friend ElementType* begin(const TArrayView& Array) { return Array.GetData(); }
+	FORCEINLINE friend ElementType* end  (const TArrayView& Array) { return Array.GetData() + Array.Num(); }
+
+public:
+	/**
+	 * Sorts the array assuming < operator is defined for the item type.
+	 */
+	void Sort()
+	{
+		::Sort(GetData(), Num());
+	}
+
+	/**
+	 * Sorts the array using user define predicate class.
+	 *
+	 * @param Predicate Predicate class instance.
+	 */
+	template <class PREDICATE_CLASS>
+	void Sort(const PREDICATE_CLASS& Predicate)
+	{
+		::Sort(GetData(), Num(), Predicate);
+	}
+
+	/**
+	 * Stable sorts the array assuming < operator is defined for the item type.
+	 *
+	 * Stable sort is slower than non-stable algorithm.
+	 */
+	void StableSort()
+	{
+		::StableSort(GetData(), Num());
+	}
+
+	/**
+	 * Stable sorts the array using user defined predicate class.
+	 *
+	 * Stable sort is slower than non-stable algorithm.
+	 *
+	 * @param Predicate Predicate class instance
+	 */
+	template <class PREDICATE_CLASS>
+	void StableSort(const PREDICATE_CLASS& Predicate)
+	{
+		::StableSort(GetData(), Num(), Predicate);
+	}
+
+private:
+	ElementType* DataPtr;
+	int32 ArrayNum;
+};
+
+template <typename InElementType>
+struct TIsZeroConstructType<TArrayView<InElementType>>
+{
+	enum { Value = true };
+};
+
+//////////////////////////////////////////////////////////////////////////
+
+template <typename ElementType, size_t Size>
+TArrayView<ElementType> MakeArrayView(ElementType(&Other)[Size])
+{
+	return TArrayView<ElementType>(Other);
+}
+
+template <typename ElementType, size_t Size>
+TArrayView<const ElementType> MakeArrayView(const ElementType(&Other)[Size])
+{
+	return TArrayView<ElementType>(Other);
+}
+
+template <typename ElementType, typename Allocator>
+TArrayView<ElementType> MakeArrayView(TArray<ElementType, Allocator>& Other)
+{
+	return TArrayView<ElementType>(Other);
+}
+
+template <typename ElementType, typename Allocator>
+TArrayView<const ElementType> MakeArrayView(const TArray<ElementType, Allocator>& Other)
+{
+	return TArrayView<const ElementType>(Other);
+}
+
+template<typename ElementType>
+TArrayView<ElementType> MakeArrayView(std::initializer_list<ElementType> List)
+{
+	return TArrayView<ElementType>(List);
+}
+
+template<typename ElementType>
+TArrayView<ElementType> MakeArrayView(ElementType* Pointer, int32 Size)
+{
+	return TArrayView<ElementType>(Pointer, Size);
+}
diff --git a/Engine/Source/Runtime/Core/Public/Misc/CoreMisc.h b/Engine/Source/Runtime/Core/Public/Misc/CoreMisc.h
--- a/Engine/Source/Runtime/Core/Public/Misc/CoreMisc.h
+++ b/Engine/Source/Runtime/Core/Public/Misc/CoreMisc.h
@@ -5,6 +5,7 @@
 #include "IntPoint.h"
 #include "Map.h"
 #include "ThreadingBase.h"
+#include "Containers/ArrayView.h"
 
 /**
  * Exec handler that registers itself and is being routed via StaticExec.
@@ -212,7 +213,7 @@ struct CORE_API FFileHelper
 	/**
 	 * Save a binary array to a file.
 	 */
-	static bool SaveArrayToFile( const TArray<uint8>& Array, const TCHAR* Filename, IFileManager* FileManager=&IFileManager::Get(), uint32 WriteFlags = 0 );
+	static bool SaveArrayToFile(TArrayView<const uint8> Array, const TCHAR* Filename, IFileManager* FileManager=&IFileManager::Get(), uint32 WriteFlags = 0);
 
 	/**
 	 * Write the FString to a file.
diff --git a/Engine/Source/Runtime/CoreUObject/Private/Serialization/ArchiveUObject.cpp b/Engine/Source/Runtime/CoreUObject/Private/Serialization/ArchiveUObject.cpp
--- a/Engine/Source/Runtime/CoreUObject/Private/Serialization/ArchiveUObject.cpp
+++ b/Engine/Source/Runtime/CoreUObject/Private/Serialization/ArchiveUObject.cpp
@@ -13,7 +13,7 @@ DEFINE_LOG_CATEGORY_STATIC(LogArchiveUObject, Log, All);
  * don't need to bother serializing lazy object pointers.  However, serialization is required if you
  * want to load and save your object.
  */
-FArchive& FArchiveUObject::operator<<( class FLazyObjectPtr& LazyObjectPtr )
+FArchive& FArchiveUObject::operator<<(class FLazyObjectPtr& LazyObjectPtr)
 {
 	FArchive& Ar = *this;
 	// We never serialize our reference while the garbage collector is harvesting references
@@ -21,14 +21,17 @@ FArchive& FArchiveUObject::operator<<( class FLazyObjectPtr& LazyObjectPtr )
 	// collected.  That would defeat the whole purpose of a weak object pointer!
 	// However, when modifying both kinds of references we want to serialize and writeback the updated value.
 	// We only want to write the modified value during reference fixup if the data is loaded
-	if( !IsObjectReferenceCollector() || IsModifyingWeakAndStrongReferences() )
+	if (!IsObjectReferenceCollector() || IsModifyingWeakAndStrongReferences())
 	{
+		// when transacting objects are serialized and restored out-of-order, resulting in *temporary* bad references
+		const bool bEvenIfPendingKill = Ar.IsTransacting();
+
 		// Downcast from UObjectBase to UObject
-		UObject* Object = static_cast< UObject* >( LazyObjectPtr.Get() );
+		UObject* Object = static_cast<UObject*>(LazyObjectPtr.Get(bEvenIfPendingKill));
 
 		Ar << Object;
 
-		if( IsLoading() || (Object && IsModifyingWeakAndStrongReferences()) )
+		if (IsLoading() || (Object && IsModifyingWeakAndStrongReferences()))
 		{
 			LazyObjectPtr = Object;
 		}
diff --git a/Engine/Source/Runtime/CoreUObject/Private/UObject/UObjectGlobals.cpp b/Engine/Source/Runtime/CoreUObject/Private/UObject/UObjectGlobals.cpp
--- a/Engine/Source/Runtime/CoreUObject/Private/UObject/UObjectGlobals.cpp
+++ b/Engine/Source/Runtime/CoreUObject/Private/UObject/UObjectGlobals.cpp
@@ -3117,6 +3117,10 @@ UObject* StaticConstructObject_Internal
 	SCOPE_CYCLE_COUNTER(STAT_ConstructObject);
 	UObject* Result = NULL;
 
+#if WITH_EDITORONLY_DATA
+	UE_CLOG(GIsSavingPackage && InOuter != GetTransientPackage(), LogUObjectGlobals, Fatal, TEXT("Illegal call to StaticConstructObject() while serializing object data! (Object will not be saved!)"));
+#endif
+
 	checkf(!InTemplate || InTemplate->IsA(InClass) || (InFlags & RF_ClassDefaultObject), TEXT("StaticConstructObject %s is not an instance of class %s and it is not a CDO."), *GetFullNameSafe(InTemplate), *GetFullNameSafe(InClass)); // template must be an instance of the class we are creating, except CDOs
 
 	// Subobjects are always created in the constructor, no need to re-create them unless their archetype != CDO or they're blueprint generated.
@@ -3126,7 +3130,7 @@ UObject* StaticConstructObject_Internal
 		(
 			!InTemplate || 
 			(InName != NAME_None && InTemplate == UObject::GetArchetypeFromRequiredInfo(InClass, InOuter, InName, InFlags))
-		);	
+		);
 #if WITH_HOT_RELOAD
 	// Do not recycle subobjects when performing hot-reload as they may contain old property values.
 	const bool bCanRecycleSubobjects = bIsNativeFromCDO && !GIsHotReload;
diff --git a/Engine/Source/Runtime/CoreUObject/Public/UObject/PersistentObjectPtr.h b/Engine/Source/Runtime/CoreUObject/Public/UObject/PersistentObjectPtr.h
--- a/Engine/Source/Runtime/CoreUObject/Public/UObject/PersistentObjectPtr.h
+++ b/Engine/Source/Runtime/CoreUObject/Public/UObject/PersistentObjectPtr.h
@@ -119,11 +119,11 @@ public:
 		return ObjectID;
 	}
 
-	/**  
+	/**
 	 * Dereference the lazy pointer, which may cause it to become valid again. Will not try to load pending outside of game thread
 	 * @return NULL if this object is gone or the lazy pointer was NULL, otherwise a valid UObject pointer
 	 */
-	FORCEINLINE class UObject* Get() const
+	FORCEINLINE UObject* Get() const
 	{
 		UObject *Object = WeakPtr.Get();
 		if (!Object && TObjectID::GetCurrentTag() != TagAtLastTest && ObjectID.IsValid())
@@ -138,6 +138,26 @@ public:
 		return Object;
 	}
 
+	/**
+	 * Dereference the lazy pointer, which may cause it to become valid again. Will not try to load pending outside of game thread
+	 * @param bEvenIfPendingKill, if this is true, pendingkill objects are considered valid
+	 * @return NULL if this object is gone or the lazy pointer was NULL, otherwise a valid UObject pointer
+	 */
+	FORCEINLINE UObject* Get(bool bEvenIfPendingKill) const
+	{
+		UObject *Object = WeakPtr.Get(bEvenIfPendingKill);
+		if (!Object && TObjectID::GetCurrentTag() != TagAtLastTest && ObjectID.IsValid())
+		{
+			Object = ObjectID.ResolveObject();
+			WeakPtr = Object;
+			TagAtLastTest = TObjectID::GetCurrentTag();
+
+			// If this object is pending kill or otherwise invalid, this will return NULL as expected
+			Object = WeakPtr.Get(bEvenIfPendingKill);
+		}
+		return Object;
+	}
+
 	/**  
 	 * Dereference the lazy pointer.
 	 */
diff --git a/Engine/Source/Runtime/Engine/Private/Animation/AnimationRuntime.cpp b/Engine/Source/Runtime/Engine/Private/Animation/AnimationRuntime.cpp
--- a/Engine/Source/Runtime/Engine/Private/Animation/AnimationRuntime.cpp
+++ b/Engine/Source/Runtime/Engine/Private/Animation/AnimationRuntime.cpp
@@ -78,7 +78,7 @@ FORCEINLINE void BlendPose(const FCompactPose& SourcePose, FCompactPose& ResultP
 	}
 }
 
-FORCEINLINE void BlendCurves(const TFixedSizeArrayView<FBlendedCurve>& SourceCurves, const TFixedSizeArrayView<float>& SourceWeights, const TFixedSizeArrayView<int32>& SourceWeightsIndices, FBlendedCurve& OutCurve)
+FORCEINLINE void BlendCurves(const TArrayView<const FBlendedCurve> SourceCurves, const TArrayView<const float> SourceWeights, const TArrayView<const int32> SourceWeightsIndices, FBlendedCurve& OutCurve)
 {
 	if (SourceCurves.Num() > 0)
 	{
@@ -91,7 +91,7 @@ FORCEINLINE void BlendCurves(const TFixedSizeArrayView<FBlendedCurve>& SourceCur
 	}
 }
 
-FORCEINLINE void BlendCurves(const TFixedSizeArrayView<const FBlendedCurve*>& SourceCurves, const TFixedSizeArrayView<float>& SourceWeights, FBlendedCurve& OutCurve)
+FORCEINLINE void BlendCurves(const TArrayView<const FBlendedCurve* const> SourceCurves, const TArrayView<const float> SourceWeights, FBlendedCurve& OutCurve)
 {
 	if(SourceCurves.Num() > 0)
 	{
@@ -104,7 +104,7 @@ FORCEINLINE void BlendCurves(const TFixedSizeArrayView<const FBlendedCurve*>& So
 	}
 }
 
-FORCEINLINE void BlendCurves(const TFixedSizeArrayView<const FBlendedCurve*>& SourceCurves, const TFixedSizeArrayView<float>& SourceWeights, FBlendedCurve& OutCurve, ECurveBlendOption::Type BlendOption)
+FORCEINLINE void BlendCurves(const TArrayView<const FBlendedCurve* const> SourceCurves, const TArrayView<const float> SourceWeights, FBlendedCurve& OutCurve, ECurveBlendOption::Type BlendOption)
 {
 	if(SourceCurves.Num() > 0)
 	{
@@ -149,9 +149,9 @@ FORCEINLINE void BlendCurves(const TFixedSizeArrayView<const FBlendedCurve*>& So
 }
 
 void FAnimationRuntime::BlendPosesTogether(
-	const TFixedSizeArrayView<FCompactPose>& SourcePoses,
-	const TFixedSizeArrayView<FBlendedCurve>& SourceCurves,
-	const TFixedSizeArrayView<float>& SourceWeights,
+	const TArrayView<const FCompactPose> SourcePoses,
+	const TArrayView<const FBlendedCurve> SourceCurves,
+	const TArrayView<const float> SourceWeights,
 	/*out*/ FCompactPose& ResultPose, 
 	/*out*/ FBlendedCurve& ResultCurve)
 {
@@ -177,7 +177,13 @@ void FAnimationRuntime::BlendPosesTogether(
 	}
 }
 
-void FAnimationRuntime::BlendPosesTogether(const TFixedSizeArrayView<FCompactPose>& SourcePoses, const TFixedSizeArrayView<FBlendedCurve>& SourceCurves, const TFixedSizeArrayView<float>& SourceWeights, const TFixedSizeArrayView<int32>& SourceWeightsIndices, /*out*/ FCompactPose& ResultPose, /*out*/ FBlendedCurve& ResultCurve)
+void FAnimationRuntime::BlendPosesTogether(
+	const TArrayView<const FCompactPose> SourcePoses,
+	const TArrayView<const FBlendedCurve> SourceCurves,
+	const TArrayView<const float> SourceWeights,
+	const TArrayView<const int32> SourceWeightsIndices,
+	/*out*/ FCompactPose& ResultPose,
+	/*out*/ FBlendedCurve& ResultCurve)
 {
 	check(SourcePoses.Num() > 0);
 
@@ -202,10 +208,10 @@ void FAnimationRuntime::BlendPosesTogether(const TFixedSizeArrayView<FCompactPos
 }
 
 void FAnimationRuntime::BlendPosesTogetherIndirect(
-	const TFixedSizeArrayView<const FCompactPose*>& SourcePoses,
-	const TFixedSizeArrayView<const FBlendedCurve*>& SourceCurves,
-	const TFixedSizeArrayView<float>& SourceWeights,
-	/*out*/ FCompactPose& ResultPose, 
+	const TArrayView<const FCompactPose* const> SourcePoses,
+	const TArrayView<const FBlendedCurve* const> SourceCurves,
+	const TArrayView<const float> SourceWeights,
+	/*out*/ FCompactPose& ResultPose,
 	/*out*/ FBlendedCurve& ResultCurve)
 {
 	check(SourcePoses.Num() > 0);
@@ -318,7 +324,13 @@ void BlendPosePerBone(const TArray<int32>& PerBoneIndices, const FBlendSampleDat
 	}
 }
 
-void FAnimationRuntime::BlendPosesTogetherPerBone(const TFixedSizeArrayView<FCompactPose>& SourcePoses, const TFixedSizeArrayView<FBlendedCurve>& SourceCurves, const IInterpolationIndexProvider* InterpolationIndexProvider, const TFixedSizeArrayView<FBlendSampleData>& BlendSampleDataCache, /*out*/ FCompactPose& ResultPose, /*out*/ FBlendedCurve& ResultCurve)
+void FAnimationRuntime::BlendPosesTogetherPerBone(
+	const TArrayView<const FCompactPose> SourcePoses,
+	const TArrayView<const FBlendedCurve> SourceCurves,
+	const IInterpolationIndexProvider* InterpolationIndexProvider,
+	const TArrayView<const FBlendSampleData> BlendSampleDataCache,
+	/*out*/ FCompactPose& ResultPose,
+	/*out*/ FBlendedCurve& ResultCurve)
 {
 	check(SourcePoses.Num() > 0);
 
@@ -354,7 +366,14 @@ void FAnimationRuntime::BlendPosesTogetherPerBone(const TFixedSizeArrayView<FCom
 	}
 }
 
-void FAnimationRuntime::BlendPosesTogetherPerBone(const TFixedSizeArrayView<FCompactPose>& SourcePoses, const TFixedSizeArrayView<FBlendedCurve>& SourceCurves, const IInterpolationIndexProvider* InterpolationIndexProvider, const TFixedSizeArrayView<FBlendSampleData>& BlendSampleDataCache, const TFixedSizeArrayView<int32>& BlendSampleDataCacheIndices, /*out*/ FCompactPose& ResultPose, /*out*/ FBlendedCurve& ResultCurve)
+void FAnimationRuntime::BlendPosesTogetherPerBone(
+	const TArrayView<const FCompactPose> SourcePoses,
+	const TArrayView<const FBlendedCurve> SourceCurves,
+	const IInterpolationIndexProvider* InterpolationIndexProvider,
+	const TArrayView<const FBlendSampleData> BlendSampleDataCache,
+	const TArrayView<const int32> BlendSampleDataCacheIndices,
+	/*out*/ FCompactPose& ResultPose,
+	/*out*/ FBlendedCurve& ResultCurve)
 {
 	check(SourcePoses.Num() > 0);
 
@@ -390,7 +409,13 @@ void FAnimationRuntime::BlendPosesTogetherPerBone(const TFixedSizeArrayView<FCom
 	}
 }
 
-void FAnimationRuntime::BlendPosesTogetherPerBoneInMeshSpace(TFixedSizeArrayView<FCompactPose>& SourcePoses, const TFixedSizeArrayView<FBlendedCurve>& SourceCurves, const UBlendSpaceBase* BlendSpace, const TFixedSizeArrayView<FBlendSampleData>& BlendSampleDataCache, FCompactPose& ResultPose, FBlendedCurve& ResultCurve)
+void FAnimationRuntime::BlendPosesTogetherPerBoneInMeshSpace(
+	const TArrayView<FCompactPose> SourcePoses,
+	const TArrayView<const FBlendedCurve> SourceCurves,
+	const UBlendSpaceBase* BlendSpace,
+	const TArrayView<const FBlendSampleData> BlendSampleDataCache,
+	/*out*/ FCompactPose& ResultPose,
+	/*out*/ FBlendedCurve& ResultCurve)
 {
 	FQuat NewRotation;
 	USkeleton* Skeleton = BlendSpace->GetSkeleton();
diff --git a/Engine/Source/Runtime/Engine/Private/Animation/BlendSpaceBase.cpp b/Engine/Source/Runtime/Engine/Private/Animation/BlendSpaceBase.cpp
--- a/Engine/Source/Runtime/Engine/Private/Animation/BlendSpaceBase.cpp
+++ b/Engine/Source/Runtime/Engine/Private/Animation/BlendSpaceBase.cpp
@@ -1274,7 +1274,7 @@ void UBlendSpaceBase::GetAnimationPose(TArray<FBlendSampleData>& BlendSampleData
 		}
 	}
 
-	TFixedSizeArrayView<FCompactPose> ChildrenPosesView(ChildrenPoses);
+	TArrayView<FCompactPose> ChildrenPosesView(ChildrenPoses);
 
 	if (PerBoneBlend.Num() > 0)
 	{
diff --git a/Engine/Source/Runtime/Engine/Public/AnimationRuntime.h b/Engine/Source/Runtime/Engine/Public/AnimationRuntime.h
--- a/Engine/Source/Runtime/Engine/Public/AnimationRuntime.h
+++ b/Engine/Source/Runtime/Engine/Public/AnimationRuntime.h
@@ -10,7 +10,7 @@
 #include "BonePose.h"
 #include "AnimTypes.h"
 #include "AnimCurveTypes.h"
-#include "FixedSizeArrayView.h"
+#include "Containers/ArrayView.h"
 
 struct FInputBlendPose;
 struct FA2CSPose;
@@ -53,7 +53,7 @@ FORCEINLINE void BlendTransform<ETransformBlendMode::Accumulate>(const FTransfor
 	Dest.AccumulateWithShortestRotation(Source, VBlendWeight);
 }
 
-FORCEINLINE void BlendCurves(const TFixedSizeArrayView<FBlendedCurve>& SourceCurves, const TFixedSizeArrayView<float>& SourceWeights, FBlendedCurve& OutCurve)
+FORCEINLINE void BlendCurves(const TArrayView<const FBlendedCurve> SourceCurves, const TArrayView<const float> SourceWeights, FBlendedCurve& OutCurve)
 {
 	if (SourceCurves.Num() > 0)
 	{
@@ -99,9 +99,9 @@ public:
 	* @param	ResultPose		Output pose of relative bone transforms.
 	*/
 	static void BlendPosesTogether(
-		const TFixedSizeArrayView<FCompactPose>& SourcePoses,
-		const TFixedSizeArrayView<FBlendedCurve>& SourceCurves,
-		const TFixedSizeArrayView<float>& SourceWeights,
+		TArrayView<const FCompactPose> SourcePoses,
+		TArrayView<const FBlendedCurve> SourceCurves,
+		TArrayView<const float> SourceWeights,
 		/*out*/ FCompactPose& ResultPose, 
 		/*out*/ FBlendedCurve& ResultCurve);
 
@@ -117,10 +117,10 @@ public:
 	* @param	ResultPose		Output pose of relative bone transforms.
 	*/
 	static void BlendPosesTogether(
-		const TFixedSizeArrayView<FCompactPose>& SourcePoses,
-		const TFixedSizeArrayView<FBlendedCurve>& SourceCurves,
-		const TFixedSizeArrayView<float>& SourceWeights,
-		const TFixedSizeArrayView<int32>& SourceWeightsIndices,
+		TArrayView<const FCompactPose> SourcePoses,
+		TArrayView<const FBlendedCurve> SourceCurves,
+		TArrayView<const float> SourceWeights,
+		TArrayView<const int32> SourceWeightsIndices,
 		/*out*/ FCompactPose& ResultPose,
 		/*out*/ FBlendedCurve& ResultCurve);
 
@@ -134,10 +134,10 @@ public:
 	* @param	ResultPose		Output pose of relative bone transforms.
 	*/
 	static void BlendPosesTogetherIndirect(
-		const TFixedSizeArrayView<const FCompactPose*>& SourcePoses,
-		const TFixedSizeArrayView<const FBlendedCurve*>& SourceCurves,
-		const TFixedSizeArrayView<float>& SourceWeights,
-		/*out*/ FCompactPose& ResultPose, 
+		TArrayView<const FCompactPose* const> SourcePoses,
+		TArrayView<const FBlendedCurve* const> SourceCurves,
+		TArrayView<const float> SourceWeights,
+		/*out*/ FCompactPose& ResultPose,
 		/*out*/ FBlendedCurve& ResultCurve);
 
 	/**
@@ -184,11 +184,11 @@ public:
 	* @param	ResultPose		Output pose of relative bone transforms.
 	*/
 	static void BlendPosesTogetherPerBone(
-		const TFixedSizeArrayView<FCompactPose>& SourcePoses,
-		const TFixedSizeArrayView<FBlendedCurve>& SourceCurves,
-		const IInterpolationIndexProvider* IndexProvider,
-		const TFixedSizeArrayView<FBlendSampleData>& BlendSampleDataCache,
-		/*out*/ FCompactPose& ResultPose, 
+		TArrayView<const FCompactPose> SourcePoses,
+		TArrayView<const FBlendedCurve> SourceCurves,
+		const IInterpolationIndexProvider* InterpolationIndexProvider,
+		TArrayView<const FBlendSampleData> BlendSampleDataCache,
+		/*out*/ FCompactPose& ResultPose,
 		/*out*/ FBlendedCurve& ResultCurve);
 
 	/**
@@ -200,11 +200,11 @@ public:
 	* @param	ResultPose		Output pose of relative bone transforms.
 	*/
 	static void BlendPosesTogetherPerBone(
-		const TFixedSizeArrayView<FCompactPose>& SourcePoses,
-		const TFixedSizeArrayView<FBlendedCurve>& SourceCurves,
-		const IInterpolationIndexProvider* IndexProvider,
-		const TFixedSizeArrayView<FBlendSampleData>& BlendSampleDataCache,
-		const TFixedSizeArrayView<int32>& BlendSampleDataCacheIndices,
+		TArrayView<const FCompactPose> SourcePoses,
+		TArrayView<const FBlendedCurve> SourceCurves,
+		const IInterpolationIndexProvider* InterpolationIndexProvider,
+		TArrayView<const FBlendSampleData> BlendSampleDataCache,
+		TArrayView<const int32> BlendSampleDataCacheIndices,
 		/*out*/ FCompactPose& ResultPose,
 		/*out*/ FBlendedCurve& ResultCurve);
 
@@ -217,10 +217,10 @@ public:
 	* @param	ResultPose		Output pose of relative bone transforms.
 	*/
 	static void BlendPosesTogetherPerBoneInMeshSpace(
-		TFixedSizeArrayView<FCompactPose>& SourcePoses,
-		const TFixedSizeArrayView<FBlendedCurve>& SourceCurves,
+		TArrayView<FCompactPose> SourcePoses,
+		TArrayView<const FBlendedCurve> SourceCurves,
 		const UBlendSpaceBase* BlendSpace,
-		const TFixedSizeArrayView<FBlendSampleData>& BlendSampleDataCache,
+		TArrayView<const FBlendSampleData> BlendSampleDataCache,
 		/*out*/ FCompactPose& ResultPose,
 		/*out*/ FBlendedCurve& ResultCurve);
 
diff --git a/Engine/Source/Runtime/Engine/Public/MeshBatch.h b/Engine/Source/Runtime/Engine/Public/MeshBatch.h
--- a/Engine/Source/Runtime/Engine/Public/MeshBatch.h
+++ b/Engine/Source/Runtime/Engine/Public/MeshBatch.h
@@ -133,6 +133,9 @@ struct FMeshBatch
 	/** Whether the mesh batch can be selected through editor selection, aka hit proxies. */
 	uint32 bSelectable : 1;
 
+	/** Whether the mesh batch needs VertexFactory->GetStaticBatchElementVisibility to be called each frame to determine which elements of the batch are visible. */
+	uint32 bRequiresPerElementVisibility : 1;
+	
 	/** Whether the mesh batch should apply dithered LOD. */
 	uint32 bDitheredLODTransition : 1;
 
@@ -252,6 +255,7 @@ struct FMeshBatch
 	,	bUseWireframeSelectionColoring(false)
 	,	bUseSelectionOutline(true)
 	,	bSelectable(true)
+	,	bRequiresPerElementVisibility(false)
 	,	bDitheredLODTransition(false)
 	,   DitheredLODTransitionAlpha(0.0f)
 	,	LCI(NULL)
diff --git a/Engine/Source/Runtime/Landscape/Classes/Landscape.h b/Engine/Source/Runtime/Landscape/Classes/Landscape.h
--- a/Engine/Source/Runtime/Landscape/Classes/Landscape.h
+++ b/Engine/Source/Runtime/Landscape/Classes/Landscape.h
@@ -39,12 +39,6 @@ public:
 	static FIntPoint MakeKey( int32 X, int32 Y ) { return FIntPoint(X, Y); }
 	static void UnpackKey( FIntPoint Key, int32& OutX, int32& OutY ) { OutX = Key.X; OutY = Key.Y; }
 
-	//~ Begin AActor Interface
-#if WITH_EDITOR
-	virtual void CheckForErrors() override;
-#endif
-	//~ End AActor Interface
-
 	//~ Begin ALandscapeProxy Interface
 	LANDSCAPE_API virtual ALandscape* GetLandscapeActor() override;
 #if WITH_EDITOR
diff --git a/Engine/Source/Runtime/Landscape/Classes/LandscapeComponent.h b/Engine/Source/Runtime/Landscape/Classes/LandscapeComponent.h
--- a/Engine/Source/Runtime/Landscape/Classes/LandscapeComponent.h
+++ b/Engine/Source/Runtime/Landscape/Classes/LandscapeComponent.h
@@ -24,7 +24,6 @@ struct FEngineShowFlags;
 struct FConvexVolume;
 struct FLandscapeEditDataInterface;
 struct FLandscapeEditToolRenderData;
-struct FLandscapeWeightmapUsage;
 struct FLandscapeTextureDataInfo;
 struct FLandscapeComponentGrassData;
 
@@ -335,7 +334,6 @@ public:
 	static void AddReferencedObjects(UObject* InThis, FReferenceCollector& Collector);
 	virtual void PostDuplicate(bool bDuplicateForPIE) override;
 #if WITH_EDITOR
-	virtual void BeginCacheForCookedPlatformData(const ITargetPlatform* TargetPlatform) override;
 	virtual void PostLoad() override;
 	virtual void PostEditUndo() override;
 	virtual void PreEditChange(UProperty* PropertyThatWillChange) override;
@@ -378,20 +376,24 @@ public:
 	//~ Begin UActorComponent Interface.
 	virtual void OnRegister() override;
 	virtual void OnUnregister() override;
+#if WITH_EDITOR
+	virtual void InvalidateLightingCacheDetailed(bool bInvalidateBuildEnqueuedLighting, bool bTranslationOnly) override;
+#endif
 	//~ End UActorComponent Interface.
 
 
 #if WITH_EDITOR
-	/** @todo document */
-	LANDSCAPE_API ULandscapeInfo* GetLandscapeInfo(bool bSpawnNewActor = true) const;
+	/** Gets the landscape info object for this landscape */
+	LANDSCAPE_API ULandscapeInfo* GetLandscapeInfo() const;
 
-	/** @todo document */
+	/** Deletes a layer from this component, removing all its data */
 	LANDSCAPE_API void DeleteLayer(ULandscapeLayerInfoObject* LayerInfo, FLandscapeEditDataInterface& LandscapeEdit);
+
+	/** Fills a layer to 100% on this component, adding it if needed and removing other layers that get painted away */
 	LANDSCAPE_API void FillLayer(ULandscapeLayerInfoObject* LayerInfo, FLandscapeEditDataInterface& LandscapeEdit);
-	LANDSCAPE_API void ReplaceLayer(ULandscapeLayerInfoObject* FromLayerInfo, ULandscapeLayerInfoObject* ToLayerInfo, FLandscapeEditDataInterface& LandscapeEdit);
 
-	void GeneratePlatformVertexData();
-	void GeneratePlatformPixelData();
+	/** Replaces one layerinfo on this component with another */
+	LANDSCAPE_API void ReplaceLayer(ULandscapeLayerInfoObject* FromLayerInfo, ULandscapeLayerInfoObject* ToLayerInfo, FLandscapeEditDataInterface& LandscapeEdit);
 
 	// true if the component's landscape material supports grass
 	bool MaterialHasGrass() const;
@@ -415,14 +417,14 @@ public:
 	/* Serialize all hashes/guids that record the current state of this component */
 	void SerializeStateHashes(FArchive& Ar);
 
+	// Generates mobile platform data for this component
+	void GeneratePlatformVertexData();
+	void GeneratePlatformPixelData();
+
 	/** Generate mobile data if it's missing or outdated */
 	void CheckGenerateLandscapePlatformData(bool bIsCooking);
 #endif
 
-	/** @todo document */
-	virtual void InvalidateLightingCacheDetailed(bool bInvalidateBuildEnqueuedLighting, bool bTranslationOnly) override;
-
-
 	/** Get the landscape actor associated with this component. */
 	ALandscape* GetLandscapeActor() const;
 
@@ -432,7 +434,7 @@ public:
 	/** Returns all generated textures and material instances used by this component. */
 	LANDSCAPE_API void GetGeneratedTexturesAndMaterialInstances(TArray<UObject*>& OutTexturesAndMaterials) const;
 
-	/** @todo document */
+	/** Gets the landscape proxy actor which owns this component */
 	LANDSCAPE_API ALandscapeProxy* GetLandscapeProxy() const;
 
 	/** @return Component section base as FIntPoint */
@@ -441,9 +443,6 @@ public:
 	/** @param InSectionBase new section base for a component */
 	LANDSCAPE_API void SetSectionBase(FIntPoint InSectionBase);
 
-	/** @todo document */
-	TMap<UTexture2D*, FLandscapeWeightmapUsage>& GetWeightmapUsageMap();
-
 	/** @todo document */
 	const FGuid& GetLightingGuid() const
 	{
@@ -463,9 +462,8 @@ public:
 #endif // WITH_EDITORONLY_DATA
 	}
 
-#if WITH_EDITOR
-
 
+#if WITH_EDITOR
 	/** Initialize the landscape component */
 	LANDSCAPE_API void Init(int32 InBaseX,int32 InBaseY,int32 InComponentSizeQuads, int32 InNumSubsections,int32 InSubsectionSizeQuads);
 
@@ -614,5 +612,3 @@ protected:
 	}
 };
 
-
-
diff --git a/Engine/Source/Runtime/Landscape/Classes/LandscapeHeightfieldCollisionComponent.h b/Engine/Source/Runtime/Landscape/Classes/LandscapeHeightfieldCollisionComponent.h
--- a/Engine/Source/Runtime/Landscape/Classes/LandscapeHeightfieldCollisionComponent.h
+++ b/Engine/Source/Runtime/Landscape/Classes/LandscapeHeightfieldCollisionComponent.h
@@ -200,8 +200,8 @@ class ULandscapeHeightfieldCollisionComponent : public UPrimitiveComponent
 	virtual void PostEditUndo() override;
 	//~ End UObject Interface.
 
-	// @todo document
-	ULandscapeInfo* GetLandscapeInfo(bool bSpawnNewActor = true) const;
+	/** Gets the landscape info object for this landscape */
+	ULandscapeInfo* GetLandscapeInfo() const;
 
 	/**  We speculatively async load collision object from DDC to remove hitch when streaming */
 	void SpeculativelyLoadAsyncDDCCollsionData();
@@ -213,13 +213,12 @@ class ULandscapeHeightfieldCollisionComponent : public UPrimitiveComponent
 
 	/** Modify a sub-region of the PhysX heightfield. Note that this does not update the physical material */
 	void UpdateHeightfieldRegion(int32 ComponentX1, int32 ComponentY1, int32 ComponentX2, int32 ComponentY2);
-
 #endif
+
 	/** Creates collision object from a cooked collision data */
 	virtual void CreateCollisionObject();
 
 	/** Return the landscape actor associated with this component. */
-	ALandscape* GetLandscapeActor() const;
 	LANDSCAPE_API ALandscapeProxy* GetLandscapeProxy() const;
 
 	/** @return Component section base as FIntPoint */
diff --git a/Engine/Source/Runtime/Landscape/Classes/LandscapeInfo.h b/Engine/Source/Runtime/Landscape/Classes/LandscapeInfo.h
--- a/Engine/Source/Runtime/Landscape/Classes/LandscapeInfo.h
+++ b/Engine/Source/Runtime/Landscape/Classes/LandscapeInfo.h
@@ -11,7 +11,6 @@ class ULandscapeComponent;
 class ULandscapeLayerInfoObject;
 class UMaterialInstanceConstant;
 struct FLandscapeEditorLayerSettings;
-struct FLandscapeDataInterface;
 
 /** Structure storing Collision for LandscapeComponent Add */
 #if WITH_EDITORONLY_DATA
@@ -84,7 +83,7 @@ struct FLandscapeInfoLayerSettings
 #endif
 };
 
-UCLASS()
+UCLASS(Transient)
 class ULandscapeInfo : public UObject
 {
 	GENERATED_UCLASS_BODY()
@@ -108,16 +107,11 @@ class ULandscapeInfo : public UObject
 	FVector DrawScale;
 	
 #if WITH_EDITORONLY_DATA
-	UPROPERTY(transient)
-	uint32 bIsValid:1;
-
 	UPROPERTY()
 	TArray<FLandscapeInfoLayerSettings> Layers;
 #endif // WITH_EDITORONLY_DATA
 
 public:
-	FLandscapeDataInterface* DataInterface;
-
 	/** Map of the offsets (in component space) to the component. Valid in editor only. */
 	TMap<FIntPoint, ULandscapeComponent*> XYtoComponentMap;
 
@@ -127,33 +121,27 @@ public:
 	TMap<FIntPoint, FLandscapeAddCollision> XYtoAddCollisionMap;
 #endif
 
+	UPROPERTY()
 	TSet<ALandscapeStreamingProxy*> Proxies;
 
 private:
 	TSet<ULandscapeComponent*> SelectedComponents;
+
 	TSet<ULandscapeComponent*> SelectedRegionComponents;
 
 public:
-
-	/** True if we're currently editing this landscape */
-	bool bCurrentlyEditing;
-
 	TMap<FIntPoint,float> SelectedRegion;
 
 	//~ Begin UObject Interface.
 	virtual void Serialize(FArchive& Ar) override;
 	virtual void BeginDestroy() override;
-#if WITH_EDITOR
-	virtual void PostEditUndo() override;
 	//~ End UObject Interface
 
+#if WITH_EDITOR
 	// @todo document 
 	// all below.
-	LANDSCAPE_API FLandscapeDataInterface* GetDataInterface();
-
 	LANDSCAPE_API void GetComponentsInRegion(int32 X1, int32 Y1, int32 X2, int32 Y2, TSet<ULandscapeComponent*>& OutComponents) const;
 	LANDSCAPE_API bool GetLandscapeExtent(int32& MinX, int32& MinY, int32& MaxX, int32& MaxY) const;
-	LANDSCAPE_API void Export(const TArray<ULandscapeLayerInfoObject*>& LayerInfos, const TArray<FString>& Filenames);
 	LANDSCAPE_API void ExportHeightmap(const FString& Filename);
 	LANDSCAPE_API void ExportLayer(ULandscapeLayerInfoObject* LayerInfo, const FString& Filename);
 	LANDSCAPE_API bool ApplySplines(bool bOnlySelected);
@@ -239,10 +227,6 @@ public:
 	 */
 	LANDSCAPE_API void FixupProxiesTransform();
 	
-	/** Fix up component layers, weightmaps for all registered proxies
-	 */
-	LANDSCAPE_API void FixupProxiesWeightmaps();
-
 	// Update per-component layer whitelists to include the currently painted layers
 	LANDSCAPE_API void UpdateComponentLayerWhitelist();
 
diff --git a/Engine/Source/Runtime/Landscape/Classes/LandscapeProxy.h b/Engine/Source/Runtime/Landscape/Classes/LandscapeProxy.h
--- a/Engine/Source/Runtime/Landscape/Classes/LandscapeProxy.h
+++ b/Engine/Source/Runtime/Landscape/Classes/LandscapeProxy.h
@@ -324,12 +324,6 @@ public:
 	UPROPERTY()
 	FIntPoint LandscapeSectionOffset;
 
-#if WITH_EDITORONLY_DATA
-	/** To support legacy landscape section offset modification under world composition mode */
-	UPROPERTY()
-	bool bStaticSectionOffset;
-#endif
-
 	/** Max LOD level to use when rendering, -1 means the max available */
 	UPROPERTY(EditAnywhere, Category=LOD)
 	int32 MaxLODLevel;
@@ -535,8 +529,8 @@ public:
 	// End blueprint functions
 
 	//~ Begin AActor Interface
+	virtual void PostRegisterAllComponents() override;
 	virtual void UnregisterAllComponents() override;
-	virtual void RegisterAllComponents() override;
 	virtual void RerunConstructionScripts() override {}
 	virtual bool IsLevelBoundsRelevant() const override { return true; }
 
@@ -596,8 +590,6 @@ public:
 
 #if WITH_EDITOR
 	virtual void PostEditChangeProperty(FPropertyChangedEvent& PropertyChangedEvent) override;
-	virtual void PreEditUndo() override;
-	virtual void PostEditUndo() override;
 	virtual void PostEditImport() override;
 	//~ End UObject Interface
 
@@ -606,7 +598,8 @@ public:
 	LANDSCAPE_API static ULandscapeLayerInfoObject* CreateLayerInfo(const TCHAR* LayerName, ULevel* Level);
 	LANDSCAPE_API ULandscapeLayerInfoObject* CreateLayerInfo(const TCHAR* LayerName);
 
-	LANDSCAPE_API ULandscapeInfo* GetLandscapeInfo(bool bSpawnNewActor = true);
+	LANDSCAPE_API ULandscapeInfo* CreateLandscapeInfo();
+	LANDSCAPE_API ULandscapeInfo* GetLandscapeInfo() const;
 
 	// Get Landscape Material assigned to this Landscape
 	virtual UMaterialInterface* GetLandscapeMaterial() const;
@@ -614,6 +607,9 @@ public:
 	// Get Hole Landscape Material assigned to this Landscape
 	virtual UMaterialInterface* GetLandscapeHoleMaterial() const;
 
+	// 
+	void FixupWeightmaps();
+
 	// Remove Invalid weightmaps
 	LANDSCAPE_API void RemoveInvalidWeightmaps();
 
diff --git a/Engine/Source/Runtime/Landscape/Classes/Materials/MaterialExpressionLandscapeGrassOutput.h b/Engine/Source/Runtime/Landscape/Classes/Materials/MaterialExpressionLandscapeGrassOutput.h
--- a/Engine/Source/Runtime/Landscape/Classes/Materials/MaterialExpressionLandscapeGrassOutput.h
+++ b/Engine/Source/Runtime/Landscape/Classes/Materials/MaterialExpressionLandscapeGrassOutput.h
@@ -31,8 +31,8 @@ struct FGrassInput
 	{}
 };
 
-UCLASS(collapsecategories, hidecategories=Object, MinimalAPI)
-class UMaterialExpressionLandscapeGrassOutput : public UMaterialExpressionCustomOutput
+UCLASS(collapsecategories, hidecategories=Object)
+class LANDSCAPE_API UMaterialExpressionLandscapeGrassOutput : public UMaterialExpressionCustomOutput
 {
 	GENERATED_UCLASS_BODY()
 
diff --git a/Engine/Source/Runtime/Landscape/Classes/Materials/MaterialExpressionLandscapeLayerBlend.h b/Engine/Source/Runtime/Landscape/Classes/Materials/MaterialExpressionLandscapeLayerBlend.h
--- a/Engine/Source/Runtime/Landscape/Classes/Materials/MaterialExpressionLandscapeLayerBlend.h
+++ b/Engine/Source/Runtime/Landscape/Classes/Materials/MaterialExpressionLandscapeLayerBlend.h
@@ -51,7 +51,7 @@ struct FLayerBlendInput
 };
 
 UCLASS(collapsecategories, hidecategories=Object)
-class UMaterialExpressionLandscapeLayerBlend : public UMaterialExpression
+class LANDSCAPE_API UMaterialExpressionLandscapeLayerBlend : public UMaterialExpression
 {
 	GENERATED_UCLASS_BODY()
 
@@ -83,7 +83,7 @@ class UMaterialExpressionLandscapeLayerBlend : public UMaterialExpression
 	virtual UTexture* GetReferencedTexture() override;
 	//~ End UMaterialExpression Interface
 
-	LANDSCAPE_API virtual FGuid& GetParameterExpressionId() override;
+	virtual FGuid& GetParameterExpressionId() override;
 
 	/**
 	 * Get list of parameter names for static parameter sets
diff --git a/Engine/Source/Runtime/Landscape/Classes/Materials/MaterialExpressionLandscapeLayerCoords.h b/Engine/Source/Runtime/Landscape/Classes/Materials/MaterialExpressionLandscapeLayerCoords.h
--- a/Engine/Source/Runtime/Landscape/Classes/Materials/MaterialExpressionLandscapeLayerCoords.h
+++ b/Engine/Source/Runtime/Landscape/Classes/Materials/MaterialExpressionLandscapeLayerCoords.h
@@ -29,7 +29,7 @@ enum ELandscapeCustomizedCoordType
 };
 
 UCLASS(collapsecategories, hidecategories=Object)
-class UMaterialExpressionLandscapeLayerCoords : public UMaterialExpression
+class LANDSCAPE_API UMaterialExpressionLandscapeLayerCoords : public UMaterialExpression
 {
 	GENERATED_UCLASS_BODY()
 
diff --git a/Engine/Source/Runtime/Landscape/Classes/Materials/MaterialExpressionLandscapeLayerSample.h b/Engine/Source/Runtime/Landscape/Classes/Materials/MaterialExpressionLandscapeLayerSample.h
--- a/Engine/Source/Runtime/Landscape/Classes/Materials/MaterialExpressionLandscapeLayerSample.h
+++ b/Engine/Source/Runtime/Landscape/Classes/Materials/MaterialExpressionLandscapeLayerSample.h
@@ -7,7 +7,7 @@
 #include "MaterialExpressionLandscapeLayerSample.generated.h"
 
 UCLASS(collapsecategories, hidecategories=Object)
-class UMaterialExpressionLandscapeLayerSample : public UMaterialExpression
+class LANDSCAPE_API UMaterialExpressionLandscapeLayerSample : public UMaterialExpression
 {
 	GENERATED_UCLASS_BODY()
 
@@ -31,7 +31,7 @@ public:
 	virtual UTexture* GetReferencedTexture() override;
 	//~ End UMaterialExpression Interface
 
-	LANDSCAPE_API virtual FGuid& GetParameterExpressionId() override;
+	virtual FGuid& GetParameterExpressionId() override;
 
 	/**
 	 * Called to get list of parameter names for static parameter sets
diff --git a/Engine/Source/Runtime/Landscape/Classes/Materials/MaterialExpressionLandscapeLayerSwitch.h b/Engine/Source/Runtime/Landscape/Classes/Materials/MaterialExpressionLandscapeLayerSwitch.h
--- a/Engine/Source/Runtime/Landscape/Classes/Materials/MaterialExpressionLandscapeLayerSwitch.h
+++ b/Engine/Source/Runtime/Landscape/Classes/Materials/MaterialExpressionLandscapeLayerSwitch.h
@@ -6,7 +6,7 @@
 #include "MaterialExpressionLandscapeLayerSwitch.generated.h"
 
 UCLASS(collapsecategories, hidecategories=Object)
-class UMaterialExpressionLandscapeLayerSwitch : public UMaterialExpression
+class LANDSCAPE_API UMaterialExpressionLandscapeLayerSwitch : public UMaterialExpression
 {
 	GENERATED_UCLASS_BODY()
 
@@ -46,7 +46,7 @@ public:
 #endif
 	//~ End UMaterialExpression Interface
 
-	LANDSCAPE_API virtual FGuid& GetParameterExpressionId() override;
+	virtual FGuid& GetParameterExpressionId() override;
 
 	void GetAllParameterNames(TArray<FName> &OutParameterNames, TArray<FGuid> &OutParameterIds) const;
 
diff --git a/Engine/Source/Runtime/Landscape/Classes/Materials/MaterialExpressionLandscapeLayerWeight.h b/Engine/Source/Runtime/Landscape/Classes/Materials/MaterialExpressionLandscapeLayerWeight.h
--- a/Engine/Source/Runtime/Landscape/Classes/Materials/MaterialExpressionLandscapeLayerWeight.h
+++ b/Engine/Source/Runtime/Landscape/Classes/Materials/MaterialExpressionLandscapeLayerWeight.h
@@ -6,7 +6,7 @@
 #include "MaterialExpressionLandscapeLayerWeight.generated.h"
 
 UCLASS(collapsecategories, hidecategories=Object)
-class UMaterialExpressionLandscapeLayerWeight : public UMaterialExpression
+class LANDSCAPE_API UMaterialExpressionLandscapeLayerWeight : public UMaterialExpression
 {
 	GENERATED_UCLASS_BODY()
 
@@ -44,7 +44,7 @@ public:
 #endif //WITH_EDITOR
 	//~ End UMaterialExpression Interface
 
-	LANDSCAPE_API virtual FGuid& GetParameterExpressionId() override;
+	virtual FGuid& GetParameterExpressionId() override;
 
 	/**
 	 * Called to get list of parameter names for static parameter sets
diff --git a/Engine/Source/Runtime/Landscape/Classes/Materials/MaterialExpressionLandscapeVisibilityMask.h b/Engine/Source/Runtime/Landscape/Classes/Materials/MaterialExpressionLandscapeVisibilityMask.h
--- a/Engine/Source/Runtime/Landscape/Classes/Materials/MaterialExpressionLandscapeVisibilityMask.h
+++ b/Engine/Source/Runtime/Landscape/Classes/Materials/MaterialExpressionLandscapeVisibilityMask.h
@@ -5,8 +5,8 @@
 #include "Materials/MaterialExpression.h"
 #include "MaterialExpressionLandscapeVisibilityMask.generated.h"
 
-UCLASS(MinimalAPI, collapseCategories, hideCategories=Object)
-class UMaterialExpressionLandscapeVisibilityMask : public UMaterialExpression
+UCLASS(collapseCategories, hideCategories=Object)
+class LANDSCAPE_API UMaterialExpressionLandscapeVisibilityMask : public UMaterialExpression
 {
 	GENERATED_UCLASS_BODY()
 
@@ -16,7 +16,7 @@ class UMaterialExpressionLandscapeVisibilityMask : public UMaterialExpression
 
 public:
 
-	LANDSCAPE_API static FName ParameterName;
+	static FName ParameterName;
 
 	//~ Begin UMaterialExpression Interface
 #if WITH_EDITOR
@@ -26,7 +26,7 @@ public:
 	virtual UTexture* GetReferencedTexture() override;
 	//~ End UMaterialExpression Interface
 
-	LANDSCAPE_API virtual FGuid& GetParameterExpressionId() override;
+	virtual FGuid& GetParameterExpressionId() override;
 
 	/**
 	 * Called to get list of parameter names for static parameter sets
diff --git a/Engine/Source/Runtime/Landscape/Landscape.Build.cs b/Engine/Source/Runtime/Landscape/Landscape.Build.cs
--- a/Engine/Source/Runtime/Landscape/Landscape.Build.cs
+++ b/Engine/Source/Runtime/Landscape/Landscape.Build.cs
@@ -9,7 +9,7 @@ public class Landscape : ModuleRules
 	{
 		PrivateIncludePaths.AddRange(
 			new string[] {
-                "Runtime/Engine/Private",
+				"Runtime/Engine/Private", // for Engine/Private/Collision/PhysXCollision.h
 			}
 		);
 
@@ -17,8 +17,7 @@ public class Landscape : ModuleRules
 			new string[] {
 				"TargetPlatform",
 				"DerivedDataCache",
-				"ImageWrapper",
-                "Foliage",
+				"Foliage",
 			}
 		);
 
@@ -30,8 +29,8 @@ public class Landscape : ModuleRules
 				"RenderCore", 
 				"RHI",
 				"ShaderCore",
-                "Renderer",
-                "Foliage",
+				"Renderer",
+				"Foliage",
 			}
 		);
 
@@ -53,20 +52,23 @@ public class Landscape : ModuleRules
 
 		if (UEBuildConfiguration.bBuildEditor == true)
 		{
+			// TODO: Remove all landscape editing code from the Landscape module!!!
+			PrivateIncludePathModuleNames.Add("LandscapeEditor");
+
 			PrivateDependencyModuleNames.AddRange(
 				new string[] {
-                    "UnrealEd",
-                    "MaterialUtilities", 
-    				"SlateCore",
-    				"Slate",
-                }
+					"UnrealEd",
+					"MaterialUtilities", 
+					"SlateCore",
+					"Slate",
+				}
 			);
 
 			CircularlyReferencedDependentModules.AddRange(
 				new string[] {
-                    "UnrealEd",
-                    "MaterialUtilities",
-                }
+					"UnrealEd",
+					"MaterialUtilities",
+				}
 			);
 		}
 	}
diff --git a/Engine/Source/Runtime/Landscape/Private/Landscape.cpp b/Engine/Source/Runtime/Landscape/Private/Landscape.cpp
--- a/Engine/Source/Runtime/Landscape/Private/Landscape.cpp
+++ b/Engine/Source/Runtime/Landscape/Private/Landscape.cpp
@@ -140,15 +140,6 @@ void ULandscapeComponent::AddReferencedObjects(UObject* InThis, FReferenceCollec
 }
 
 #if WITH_EDITOR
-
-void ULandscapeComponent::BeginCacheForCookedPlatformData(const ITargetPlatform* TargetPlatform)
-{
-	if (TargetPlatform->SupportsFeature(ETargetPlatformFeatures::MobileRendering))
-	{
-		CheckGenerateLandscapePlatformData(true);
-	}
-}
-
 void ULandscapeComponent::CheckGenerateLandscapePlatformData(bool bIsCooking)
 {
 #if ENABLE_LANDSCAPE_COOKING
@@ -218,7 +209,13 @@ void ULandscapeComponent::CheckGenerateLandscapePlatformData(bool bIsCooking)
 
 void ULandscapeComponent::Serialize(FArchive& Ar)
 {
-	if (Ar.IsCooking() && !Ar.CookingTarget()->SupportsFeature(ETargetPlatformFeatures::DeferredRendering))
+#if WITH_EDITOR
+	if (Ar.IsCooking() && !HasAnyFlags(RF_ClassDefaultObject) && Ar.CookingTarget()->SupportsFeature(ETargetPlatformFeatures::MobileRendering))
+	{
+		CheckGenerateLandscapePlatformData(true);
+	}
+
+	if (Ar.IsCooking() && !HasAnyFlags(RF_ClassDefaultObject) && !Ar.CookingTarget()->SupportsFeature(ETargetPlatformFeatures::DeferredRendering))
 	{
 		// These properties are only used for SM4+ so we back them up and clear them before serializing them.
 		UTexture2D* BackupHeightmapTexture = nullptr;
@@ -239,6 +236,7 @@ void ULandscapeComponent::Serialize(FArchive& Ar)
 		Exchange(BackupWeightmapTextures, WeightmapTextures);
 	}
 	else
+#endif
 	{
 		Super::Serialize(Ar);
 	}
@@ -419,8 +417,8 @@ FString ULandscapeComponent::GetLayerAllocationKey(UMaterialInterface* Landscape
 
 void ULandscapeComponent::GetLayerDebugColorKey(int32& R, int32& G, int32& B) const
 {
-	ULandscapeInfo* Info = GetLandscapeInfo(false);
-	if (Info)
+	ULandscapeInfo* Info = GetLandscapeInfo();
+	if (ensure(Info))
 	{
 		R = INDEX_NONE, G = INDEX_NONE, B = INDEX_NONE;
 
@@ -465,9 +463,6 @@ ULandscapeMeshCollisionComponent::ULandscapeMeshCollisionComponent(const FObject
 ULandscapeInfo::ULandscapeInfo(const FObjectInitializer& ObjectInitializer)
 : Super(ObjectInitializer)
 {
-#if WITH_EDITOR
-	bIsValid = false;
-#endif
 }
 
 #if WITH_EDITOR
@@ -618,7 +613,6 @@ ALandscapeProxy::ALandscapeProxy(const FObjectInitializer& ObjectInitializer)
 #if WITH_EDITORONLY_DATA
 	bLockLocation = true;
 	bIsMovingToLevel = false;
-	bStaticSectionOffset = true;
 #endif // WITH_EDITORONLY_DATA
 	LODDistanceFactor = 1.0f;
 	LODFalloff = ELandscapeLODFalloff::Linear;
@@ -683,29 +677,46 @@ ALandscape* ALandscapeStreamingProxy::GetLandscapeActor()
 }
 
 #if WITH_EDITOR
-ULandscapeInfo* ALandscapeProxy::GetLandscapeInfo(bool bSpawnNewActor /*= true*/)
+ULandscapeInfo* ALandscapeProxy::CreateLandscapeInfo()
 {
-	// LandscapeInfo generate
-	if (GIsEditor)
-	{
-		UWorld* OwningWorld = GetWorld();
-		if (OwningWorld && !OwningWorld->IsGameWorld())
-		{
-			auto &LandscapeInfoMap = GetLandscapeInfoMap(OwningWorld);
+	ULandscapeInfo* LandscapeInfo = nullptr;
 
-			ULandscapeInfo* LandscapeInfo = LandscapeInfoMap.Map.FindRef(LandscapeGuid);
-			if (!LandscapeInfo && bSpawnNewActor && !HasAnyFlags(RF_BeginDestroyed))
-			{
-				LandscapeInfo = NewObject<ULandscapeInfo>(OwningWorld, NAME_None, RF_Transactional | RF_Transient);
-				LandscapeInfoMap.Modify(false);
-				LandscapeInfoMap.Map.Add(LandscapeGuid, LandscapeInfo);
-			}
+	check(GIsEditor);
+	check(LandscapeGuid.IsValid());
+	UWorld* OwningWorld = GetWorld();
+	check(OwningWorld);
+	check(!OwningWorld->IsGameWorld());
+	
+	auto& LandscapeInfoMap = GetLandscapeInfoMap(OwningWorld);
+	LandscapeInfo = LandscapeInfoMap.Map.FindRef(LandscapeGuid);
 
-			return LandscapeInfo;
-		}
+	if (!LandscapeInfo)
+	{
+		check(!HasAnyFlags(RF_BeginDestroyed));
+		LandscapeInfo = NewObject<ULandscapeInfo>(GetTransientPackage(), NAME_None, RF_Transactional | RF_Transient);
+		LandscapeInfoMap.Modify(false);
+		LandscapeInfoMap.Map.Add(LandscapeGuid, LandscapeInfo);
 	}
+	check(LandscapeInfo);
+	LandscapeInfo->RegisterActor(this);
 
-	return NULL;
+	return LandscapeInfo;
+}
+
+ULandscapeInfo* ALandscapeProxy::GetLandscapeInfo() const
+{
+	ULandscapeInfo* LandscapeInfo = nullptr;
+
+	check(GIsEditor);
+	check(LandscapeGuid.IsValid());
+	UWorld* OwningWorld = GetWorld();
+	check(OwningWorld);
+	check(!OwningWorld->IsGameWorld());
+
+	auto& LandscapeInfoMap = GetLandscapeInfoMap(OwningWorld);
+	LandscapeInfo = LandscapeInfoMap.Map.FindRef(LandscapeGuid);
+
+	return LandscapeInfo;
 }
 #endif
 
@@ -786,13 +797,9 @@ void ULandscapeComponent::SetSectionBase(FIntPoint InSectionBase)
 }
 
 #if WITH_EDITOR
-ULandscapeInfo* ULandscapeComponent::GetLandscapeInfo(bool bSpawnNewActor /*= true*/) const
+ULandscapeInfo* ULandscapeComponent::GetLandscapeInfo() const
 {
-	if (GetLandscapeProxy())
-	{
-		return GetLandscapeProxy()->GetLandscapeInfo(bSpawnNewActor);
-	}
-	return NULL;
+	return GetLandscapeProxy()->GetLandscapeInfo();
 }
 #endif
 
@@ -903,10 +910,13 @@ void ULandscapeComponent::OnRegister()
 	Super::OnRegister();
 
 #if WITH_EDITOR
-	ULandscapeInfo* Info = GetLandscapeInfo(false);
-	if (Info)
+	if (!GetWorld()->IsGameWorld())
 	{
-		Info->RegisterActorComponent(this);
+		ULandscapeInfo* Info = GetLandscapeInfo();
+		if (Info)
+		{
+			Info->RegisterActorComponent(this);
+		}
 	}
 #endif
 }
@@ -916,10 +926,16 @@ void ULandscapeComponent::OnUnregister()
 	Super::OnUnregister();
 
 #if WITH_EDITOR
-	ULandscapeInfo* Info = GetLandscapeInfo(false);
-	if (Info)
+	// Game worlds don't have landscape infos
+	if (GetWorld() && !GetWorld()->IsGameWorld()
+		// On shutdown the world will be unreachable
+		&& !GetWorld()->IsPendingKillOrUnreachable())
 	{
-		Info->UnregisterActorComponent(this);
+		ULandscapeInfo* Info = GetLandscapeInfo();
+		if (Info)
+		{
+			Info->UnregisterActorComponent(this);
+		}
 	}
 #endif
 }
@@ -932,28 +948,45 @@ void ULandscapeComponent::GetUsedMaterials(TArray<UMaterialInterface*>& OutMater
 	}
 }
 
-void ALandscapeProxy::RegisterAllComponents()
+void ALandscapeProxy::PostRegisterAllComponents()
 {
-	Super::RegisterAllComponents();
+	Super::PostRegisterAllComponents();
 
-	// Landscape was added to world 
-	// We might need to update shared landscape data
-	if (GIsEditor && GetWorld() && !GetWorld()->IsPlayInEditor())
+#if WITH_EDITOR
+	// Game worlds don't have landscape infos
+	if (!GetWorld()->IsGameWorld())
 	{
-		GEngine->DeferredCommands.AddUnique(TEXT("UpdateLandscapeEditorData"));
+		// Duplicated Landscapes don't have a valid guid until PostEditImport is called, we'll register then
+		if (LandscapeGuid.IsValid())
+		{
+			ULandscapeInfo* LandscapeInfo = CreateLandscapeInfo();
+
+			LandscapeInfo->FixupProxiesTransform();
+		}
 	}
+#endif
 }
 
 void ALandscapeProxy::UnregisterAllComponents()
 {
-	Super::UnregisterAllComponents();
-
-	// Landscape was removed from world 
-	// We might need to update shared landscape data
-	if (GEngine && GIsEditor && GetWorld() && !GetWorld()->IsPlayInEditor())
-	{
-		GEngine->DeferredCommands.AddUnique(TEXT("UpdateLandscapeEditorData"));
+#if WITH_EDITOR
+	// Game worlds don't have landscape infos
+	if (GetWorld() && !GetWorld()->IsGameWorld()
+		// On shutdown the world will be unreachable
+		&& !GetWorld()->IsPendingKillOrUnreachable() &&
+		// When redoing the creation of a landscape we may get UnregisterAllComponents called when
+		// we are in a "pre-initialized" state (empty guid, etc)
+		LandscapeGuid.IsValid())
+	{
+		ULandscapeInfo* LandscapeInfo = GetLandscapeInfo();
+		if (LandscapeInfo)
+		{
+			LandscapeInfo->UnregisterActor(this);
+		}
 	}
+#endif
+
+	Super::UnregisterAllComponents();
 }
 
 // FLandscapeWeightmapUsage serializer
@@ -987,8 +1020,6 @@ void ULandscapeInfo::Serialize(FArchive& Ar)
 {
 	Super::Serialize(Ar);
 
-	// We do not serialize XYtoComponentMap as we don't want these references to hold a component in memory.
-	// The references are automatically cleaned up by the components' BeginDestroy method.
 	if (Ar.IsTransacting())
 	{
 		Ar << XYtoComponentMap;
@@ -1014,21 +1045,7 @@ void ALandscape::PostLoad()
 void ULandscapeInfo::BeginDestroy()
 {
 	Super::BeginDestroy();
-
-#if WITH_EDITOR
-	if (DataInterface)
-	{
-		delete DataInterface;
-		DataInterface = NULL;
-	}
-#endif
-}
-
-#if WITH_EDITOR
-void ALandscape::CheckForErrors()
-{
 }
-#endif
 
 void ALandscapeProxy::PreSave(const class ITargetPlatform* TargetPlatform)
 {
@@ -1074,39 +1091,6 @@ void ALandscapeProxy::AddReferencedObjects(UObject* InThis, FReferenceCollector&
 }
 
 #if WITH_EDITOR
-void ALandscapeProxy::PreEditUndo()
-{
-	Super::PreEditUndo();
-
-	// PostEditUndo doesn't get called when undoing a create action
-	// and PreEditUndo doesn't get called when undoing a delete action
-	// so this code needs to be in both!
-	if (GIsEditor && GetWorld() && !GetWorld()->IsPlayInEditor())
-	{
-		GEngine->DeferredCommands.AddUnique(TEXT("UpdateLandscapeEditorData"));
-	}
-}
-
-void ALandscapeProxy::PostEditUndo()
-{
-	Super::PostEditUndo();
-
-	if (GIsEditor && GetWorld() && !GetWorld()->IsPlayInEditor())
-	{
-		GEngine->DeferredCommands.AddUnique(TEXT("UpdateLandscapeEditorData"));
-	}
-}
-
-void ULandscapeInfo::PostEditUndo()
-{
-	Super::PostEditUndo();
-
-	if (GIsEditor && GetWorld() && !GetWorld()->IsPlayInEditor())
-	{
-		GEngine->DeferredCommands.AddUnique(TEXT("UpdateLandscapeEditorData"));
-	}
-}
-
 FName FLandscapeInfoLayerSettings::GetLayerName() const
 {
 	checkSlow(LayerInfoObj == NULL || LayerInfoObj->LayerName == LayerName);
@@ -1210,8 +1194,6 @@ bool ULandscapeInfo::UpdateLayerInfoMap(ALandscapeProxy* Proxy /*= NULL*/, bool
 	bool bHasCollision = false;
 	if (GIsEditor)
 	{
-		bIsValid = !bInvalidate;
-
 		if (Proxy)
 		{
 			if (bInvalidate)
@@ -1413,14 +1395,10 @@ void ALandscapeProxy::PostLoad()
 		LandscapeComponents.ContainsByPredicate([](ULandscapeComponent* Comp) { return !Comp->CollisionComponent.IsValid(); }))
 	{
 		// Need to clean up invalid collision components
+		CreateLandscapeInfo();
 		RecreateCollisionComponents();
 	}
 
-	if (GetLinker() && GetLinker()->UE4Ver() < VER_UE4_LANDSCAPE_STATIC_SECTION_OFFSET)
-	{
-		bStaticSectionOffset = false;
-	}
-
 	EditorLayerSettings.Remove(nullptr);
 
 	if (EditorCachedLayerInfos_DEPRECATED.Num() > 0)
@@ -1432,10 +1410,12 @@ void ALandscapeProxy::PostLoad()
 		EditorCachedLayerInfos_DEPRECATED.Empty();
 	}
 
-	if (GIsEditor && !GetWorld()->IsPlayInEditor())
+	if (GIsEditor && !GetWorld()->IsGameWorld())
 	{
-		// defer LandscapeInfo setup
-		GEngine->DeferredCommands.AddUnique(TEXT("UpdateLandscapeEditorData -warnings"));
+		ULandscapeInfo* LandscapeInfo = CreateLandscapeInfo();
+		LandscapeInfo->RegisterActor(this, true);
+
+		FixupWeightmaps();
 	}
 #endif
 }
@@ -1756,13 +1736,14 @@ ALandscapeProxy* ULandscapeInfo::GetLandscapeProxy() const
 void ULandscapeInfo::RegisterActor(ALandscapeProxy* Proxy, bool bMapCheck)
 {
 	// do not pass here invalid actors
+	checkSlow(Proxy);
 	check(Proxy->GetLandscapeGuid().IsValid());
-	UWorld* OwningWorld = CastChecked<UWorld>(GetOuter());
+	UWorld* OwningWorld = Proxy->GetWorld();
 
 	// in case this Info object is not initialized yet
 	// initialized it with properties from passed actor
 	if (LandscapeGuid.IsValid() == false ||
-		(GetLandscapeProxy() == NULL && ensure(LandscapeGuid == Proxy->GetLandscapeGuid())))
+		(GetLandscapeProxy() == nullptr && ensure(LandscapeGuid == Proxy->GetLandscapeGuid())))
 	{
 		LandscapeGuid = Proxy->GetLandscapeGuid();
 		ComponentSizeQuads = Proxy->ComponentSizeQuads;
@@ -1786,6 +1767,7 @@ void ULandscapeInfo::RegisterActor(ALandscapeProxy* Proxy, bool bMapCheck)
 	// register
 	if (ALandscape* Landscape = Cast<ALandscape>(Proxy))
 	{
+		check(!LandscapeActor || LandscapeActor == Landscape)
 		LandscapeActor = Landscape;
 		// In world composition user is not allowed to move landscape in editor, only through WorldBrowser 
 		LandscapeActor->bLockLocation = (OwningWorld->WorldComposition != nullptr);
@@ -1806,6 +1788,9 @@ void ULandscapeInfo::RegisterActor(ALandscapeProxy* Proxy, bool bMapCheck)
 		StreamingProxy->ConditionalAssignCommonProperties(LandscapeActor.Get());
 	}
 
+	UpdateLayerInfoMap(Proxy);
+	UpdateAllAddCollisions();
+
 	//
 	// add proxy components to the XY map
 	//
@@ -1819,7 +1804,8 @@ void ULandscapeInfo::UnregisterActor(ALandscapeProxy* Proxy)
 {
 	if (ALandscape* Landscape = Cast<ALandscape>(Proxy))
 	{
-		check(LandscapeActor.Get() == Landscape);
+		// LandscapeActor can be null if Unregister gets triggered twice, which has been observed to happen during redo
+		check(LandscapeActor.IsNull() || LandscapeActor.Get() == Landscape);
 		LandscapeActor = nullptr;
 
 		// update proxies reference to landscape actor
@@ -1840,8 +1826,10 @@ void ULandscapeInfo::UnregisterActor(ALandscapeProxy* Proxy)
 	{
 		UnregisterActorComponent(Proxy->LandscapeComponents[CompIdx]);
 	}
-
 	XYtoComponentMap.Compact();
+
+	UpdateLayerInfoMap();
+	UpdateAllAddCollisions();
 }
 
 void ULandscapeInfo::RegisterActorComponent(ULandscapeComponent* Component, bool bMapCheck)
@@ -1857,36 +1845,36 @@ void ULandscapeInfo::RegisterActorComponent(ULandscapeComponent* Component, bool
 	FIntPoint ComponentKey = Component->GetSectionBase() / Component->ComponentSizeQuads;
 	auto RegisteredComponent = XYtoComponentMap.FindRef(ComponentKey);
 
-	if (RegisteredComponent == NULL)
+	if (RegisteredComponent != Component)
 	{
-		XYtoComponentMap.Add(ComponentKey, Component);
-	}
-	else if (bMapCheck)
-	{
-		ALandscapeProxy* OurProxy = Component->GetLandscapeProxy();
-		ALandscapeProxy* ExistingProxy = RegisteredComponent->GetLandscapeProxy();
-		FFormatNamedArguments Arguments;
-		Arguments.Add(TEXT("ProxyName1"), FText::FromString(OurProxy->GetName()));
-		Arguments.Add(TEXT("LevelName1"), FText::FromString(OurProxy->GetLevel()->GetOutermost()->GetName()));
-		Arguments.Add(TEXT("ProxyName2"), FText::FromString(ExistingProxy->GetName()));
-		Arguments.Add(TEXT("LevelName2"), FText::FromString(ExistingProxy->GetLevel()->GetOutermost()->GetName()));
-		Arguments.Add(TEXT("XLocation"), Component->GetSectionBase().X);
-		Arguments.Add(TEXT("YLocation"), Component->GetSectionBase().Y);
-		FMessageLog("MapCheck").Warning()
-			->AddToken(FUObjectToken::Create(OurProxy))
-			->AddToken(FTextToken::Create(FText::Format(LOCTEXT("MapCheck_Message_LandscapeComponentPostLoad_Warning", "Landscape {ProxyName1} of {LevelName1} has overlapping render components with {ProxyName2} of {LevelName2} at location ({XLocation}, {YLocation})."), Arguments)))
-#if WITH_EDITOR
-			->AddToken(FActionToken::Create(LOCTEXT("MapCheck_RemoveDuplicateLandscapeComponent", "Delete Duplicate"), LOCTEXT("MapCheck_RemoveDuplicateLandscapeComponentDesc", "Deletes the duplicate landscape component."), FOnActionTokenExecuted::CreateUObject(OurProxy, &ALandscapeProxy::RemoveOverlappingComponent, Component), true))
-#endif
-			->AddToken(FMapErrorToken::Create(FMapErrors::LandscapeComponentPostLoad_Warning));
+		if (RegisteredComponent == nullptr)
+		{
+			XYtoComponentMap.Add(ComponentKey, Component);
+		}
+		else if (bMapCheck)
+		{
+			ALandscapeProxy* OurProxy = Component->GetLandscapeProxy();
+			ALandscapeProxy* ExistingProxy = RegisteredComponent->GetLandscapeProxy();
+			FFormatNamedArguments Arguments;
+			Arguments.Add(TEXT("ProxyName1"), FText::FromString(OurProxy->GetName()));
+			Arguments.Add(TEXT("LevelName1"), FText::FromString(OurProxy->GetLevel()->GetOutermost()->GetName()));
+			Arguments.Add(TEXT("ProxyName2"), FText::FromString(ExistingProxy->GetName()));
+			Arguments.Add(TEXT("LevelName2"), FText::FromString(ExistingProxy->GetLevel()->GetOutermost()->GetName()));
+			Arguments.Add(TEXT("XLocation"), Component->GetSectionBase().X);
+			Arguments.Add(TEXT("YLocation"), Component->GetSectionBase().Y);
+			FMessageLog("MapCheck").Warning()
+				->AddToken(FUObjectToken::Create(OurProxy))
+				->AddToken(FTextToken::Create(FText::Format(LOCTEXT("MapCheck_Message_LandscapeComponentPostLoad_Warning", "Landscape {ProxyName1} of {LevelName1} has overlapping render components with {ProxyName2} of {LevelName2} at location ({XLocation}, {YLocation})."), Arguments)))
+				->AddToken(FActionToken::Create(LOCTEXT("MapCheck_RemoveDuplicateLandscapeComponent", "Delete Duplicate"), LOCTEXT("MapCheck_RemoveDuplicateLandscapeComponentDesc", "Deletes the duplicate landscape component."), FOnActionTokenExecuted::CreateUObject(OurProxy, &ALandscapeProxy::RemoveOverlappingComponent, Component), true))
+				->AddToken(FMapErrorToken::Create(FMapErrors::LandscapeComponentPostLoad_Warning));
 
-		// Show MapCheck window
-		FMessageLog("MapCheck").Open(EMessageSeverity::Warning);
+			// Show MapCheck window
+			FMessageLog("MapCheck").Open(EMessageSeverity::Warning);
+		}
 	}
 
 	// Update Selected Components/Regions
-#if WITH_EDITOR
-	if (Component->EditToolRenderData != NULL && Component->EditToolRenderData->SelectedType)
+	if (Component->EditToolRenderData != nullptr && Component->EditToolRenderData->SelectedType)
 	{
 		if (Component->EditToolRenderData->SelectedType & FLandscapeEditToolRenderData::ST_COMPONENT)
 		{
@@ -1897,12 +1885,11 @@ void ULandscapeInfo::RegisterActorComponent(ULandscapeComponent* Component, bool
 			SelectedRegionComponents.Add(Component);
 		}
 	}
-#endif
 }
 
 void ULandscapeInfo::UnregisterActorComponent(ULandscapeComponent* Component)
 {
-	check(Component != NULL);
+	check(Component != nullptr);
 
 	FIntPoint ComponentKey = Component->GetSectionBase() / Component->ComponentSizeQuads;
 	auto RegisteredComponent = XYtoComponentMap.FindRef(ComponentKey);
@@ -1970,7 +1957,6 @@ void ULandscapeInfo::FixupProxiesTransform()
 		if (!Proxy->GetTransform().Equals(ProxyTransform))
 		{
 			Proxy->SetActorTransform(ProxyTransform);
-			Proxy->RecreateComponentsState();
 
 			// Let other systems know that an actor was moved
 			GEngine->BroadcastOnActorMoved(Proxy);
@@ -1978,33 +1964,6 @@ void ULandscapeInfo::FixupProxiesTransform()
 	}
 }
 
-void ULandscapeInfo::FixupProxiesWeightmaps()
-{
-	if (LandscapeActor.IsValid())
-	{
-		LandscapeActor->WeightmapUsageMap.Empty();
-		for (ULandscapeComponent* Comp : LandscapeActor->LandscapeComponents)
-		{
-			if (Comp)
-			{
-				Comp->FixupWeightmaps();
-			}
-		}
-	}
-
-	for (ALandscapeProxy* Proxy : Proxies)
-	{
-		Proxy->WeightmapUsageMap.Empty();
-		for (ULandscapeComponent* Comp : Proxy->LandscapeComponents)
-		{
-			if (Comp)
-			{
-				Comp->FixupWeightmaps();
-			}
-		}
-	}
-}
-
 void ULandscapeInfo::UpdateComponentLayerWhitelist()
 {
 	if (LandscapeActor.IsValid())
@@ -2024,143 +1983,55 @@ void ULandscapeInfo::UpdateComponentLayerWhitelist()
 	}
 }
 
-//
-// This handles legacy behavior of landscapes created under world composition
-// We adjust landscape section offsets and set a flag inside landscape that sections offset are static and should never be touched again
-//
-void AdjustLandscapeSectionOffsets(UWorld* InWorld, const TArray<ALandscapeProxy*>& InLandscapeList)
-{
-	// We interested only in registered actors
-	TArray<ALandscapeProxy*> RegisteredLandscapeList = InLandscapeList.FilterByPredicate([](ALandscapeProxy* Proxy) {
-		return Proxy->GetRootComponent()->IsRegistered();
-	});
-
-	// Main Landscape actor should act as origin of global components grid
-	int32 LandcapeIndex = RegisteredLandscapeList.IndexOfByPredicate([](ALandscapeProxy* Proxy) {
-		return Proxy->IsA<ALandscape>();
-	});
-	ALandscapeProxy* StaticLandscape = RegisteredLandscapeList.IsValidIndex(LandcapeIndex) ? RegisteredLandscapeList[LandcapeIndex] : nullptr;
-
-	if (StaticLandscape && !StaticLandscape->bStaticSectionOffset)
-	{
-		StaticLandscape->SetAbsoluteSectionBase(FIntPoint::ZeroValue);
-	}
-
-	// In case there is no main landscape actor loaded try use any landscape that already has static offsets
-	if (StaticLandscape == nullptr)
-	{
-		LandcapeIndex = RegisteredLandscapeList.IndexOfByPredicate([](ALandscapeProxy* Proxy) {
-			return Proxy->bStaticSectionOffset;
-		});
-		StaticLandscape = RegisteredLandscapeList.IsValidIndex(LandcapeIndex) ? RegisteredLandscapeList[LandcapeIndex] : nullptr;
-		// Otherwise offsets will stay variable
-	}
-
-	ALandscapeProxy* OriginLandscape = StaticLandscape;
-
-	for (ALandscapeProxy* Proxy : RegisteredLandscapeList)
-	{
-		if (OriginLandscape == nullptr)
-		{
-			OriginLandscape = Proxy;
-		}
-		else if (!Proxy->bStaticSectionOffset)
-		{
-			// Calculate section offset based on relative position from "origin" landscape
-			FVector Offset = Proxy->GetActorLocation() - OriginLandscape->GetActorLocation();
-			FVector	DrawScale = OriginLandscape->GetRootComponent()->RelativeScale3D;
-
-			FIntPoint QuadsSpaceOffset;
-			QuadsSpaceOffset.X = FMath::RoundToInt(Offset.X / DrawScale.X);
-			QuadsSpaceOffset.Y = FMath::RoundToInt(Offset.Y / DrawScale.Y);
-			Proxy->SetAbsoluteSectionBase(QuadsSpaceOffset + OriginLandscape->LandscapeSectionOffset);
-
-			if (StaticLandscape)
-			{
-				Proxy->bStaticSectionOffset = true;
-			}
-		}
-	}
-}
-
 void ULandscapeInfo::RecreateLandscapeInfo(UWorld* InWorld, bool bMapCheck)
 {
 	check(InWorld);
 
-	auto& LandscapeInfoMap = GetLandscapeInfoMap(InWorld);
+	ULandscapeInfoMap& LandscapeInfoMap = GetLandscapeInfoMap(InWorld);
+	LandscapeInfoMap.Modify();
 
 	// reset all LandscapeInfo objects
-	for (auto It = LandscapeInfoMap.Map.CreateIterator(); It; ++It)
+	for (auto& LandscapeInfoPair : LandscapeInfoMap.Map)
 	{
-		It.Value()->Modify();
-		It.Value()->Reset();
+		ULandscapeInfo* LandscapeInfo = LandscapeInfoPair.Value;
+		LandscapeInfo->Modify();
+		LandscapeInfo->Reset();
 	}
 
 	TMap<FGuid, TArray<ALandscapeProxy*>> ValidLandscapesMap;
 	// Gather all valid landscapes in the world
-	for (TActorIterator<ALandscapeProxy> It(InWorld); It; ++It)
+	for (ALandscapeProxy* Proxy : TActorRange<ALandscapeProxy>(InWorld))
 	{
-		ALandscapeProxy* Proxy = *It;
-		if (Proxy
-			&& Proxy->GetLevel()
-			&& Proxy->GetLevel()->bIsVisible
-			&& Proxy->HasAnyFlags(RF_BeginDestroyed) == false
-			&& Proxy->IsPendingKill() == false
-			&& Proxy->IsPendingKillPending() == false)
+		if (Proxy->GetLevel() &&
+			Proxy->GetLevel()->bIsVisible &&
+			!Proxy->HasAnyFlags(RF_BeginDestroyed) &&
+			!Proxy->IsPendingKill() &&
+			!Proxy->IsPendingKillPending())
 		{
 			ValidLandscapesMap.FindOrAdd(Proxy->GetLandscapeGuid()).Add(Proxy);
 		}
 	}
 
-	// Handle legacy landscape data under world composition
-	if (InWorld->WorldComposition)
-	{
-		for (auto It = ValidLandscapesMap.CreateIterator(); It; ++It)
-		{
-			AdjustLandscapeSectionOffsets(InWorld, It.Value());
-		}
-	}
-
 	// Register landscapes in global landscape map
-	for (auto It = ValidLandscapesMap.CreateIterator(); It; ++It)
+	for (auto& ValidLandscapesPair : ValidLandscapesMap)
 	{
-		auto& LandscapeList = It.Value();
+		auto& LandscapeList = ValidLandscapesPair.Value;
 		for (ALandscapeProxy* Proxy : LandscapeList)
 		{
-			Proxy->GetLandscapeInfo(true)->RegisterActor(Proxy, bMapCheck);
+			Proxy->CreateLandscapeInfo()->RegisterActor(Proxy, bMapCheck);
 		}
 	}
 
 	// Remove empty entries from global LandscapeInfo map
 	for (auto It = LandscapeInfoMap.Map.CreateIterator(); It; ++It)
 	{
-		if (It.Value()->GetLandscapeProxy() == NULL)
+		if (It.Value()->GetLandscapeProxy() == nullptr)
 		{
 			It.Value()->MarkPendingKill();
 			It.RemoveCurrent();
 		}
 	}
 
-	// Update layer info maps
-	for (auto It = LandscapeInfoMap.Map.CreateIterator(); It; ++It)
-	{
-		ULandscapeInfo* LandscapeInfo = It.Value();
-		if (LandscapeInfo)
-		{
-			LandscapeInfo->UpdateLayerInfoMap();
-		}
-	}
-
-	// Update add collision
-	for (auto It = LandscapeInfoMap.Map.CreateIterator(); It; ++It)
-	{
-		ULandscapeInfo* LandscapeInfo = It.Value();
-		if (LandscapeInfo)
-		{
-			LandscapeInfo->UpdateAllAddCollisions();
-		}
-	}
-
 	// We need to inform Landscape editor tools about LandscapeInfo updates
 	FEditorSupportDelegates::WorldChange.Broadcast();
 }
diff --git a/Engine/Source/Runtime/Landscape/Private/LandscapeCollision.cpp b/Engine/Source/Runtime/Landscape/Private/LandscapeCollision.cpp
--- a/Engine/Source/Runtime/Landscape/Private/LandscapeCollision.cpp
+++ b/Engine/Source/Runtime/Landscape/Private/LandscapeCollision.cpp
@@ -2035,27 +2035,13 @@ void ULandscapeMeshCollisionComponent::ImportCustomProperties(const TCHAR* Sourc
 	}
 }
 
-ULandscapeInfo* ULandscapeHeightfieldCollisionComponent::GetLandscapeInfo(bool bSpawnNewActor /*= true*/) const
+ULandscapeInfo* ULandscapeHeightfieldCollisionComponent::GetLandscapeInfo() const
 {
-	if (GetLandscapeProxy())
-	{
-		return GetLandscapeProxy()->GetLandscapeInfo(bSpawnNewActor);
-	}
-	return NULL;
+	return GetLandscapeProxy()->GetLandscapeInfo();
 }
 
 #endif // WITH_EDITOR
 
-ALandscape* ULandscapeHeightfieldCollisionComponent::GetLandscapeActor() const
-{
-	ALandscapeProxy* Landscape = GetLandscapeProxy();
-	if (Landscape)
-	{
-		return Landscape->GetLandscapeActor();
-	}
-	return NULL;
-}
-
 ALandscapeProxy* ULandscapeHeightfieldCollisionComponent::GetLandscapeProxy() const
 {
 	return CastChecked<ALandscapeProxy>(GetOuter());
diff --git a/Engine/Source/Runtime/Landscape/Private/LandscapeDataAccess.cpp b/Engine/Source/Runtime/Landscape/Private/LandscapeDataAccess.cpp
--- a/Engine/Source/Runtime/Landscape/Private/LandscapeDataAccess.cpp
+++ b/Engine/Source/Runtime/Landscape/Private/LandscapeDataAccess.cpp
@@ -12,19 +12,8 @@ LANDSCAPE_API FLandscapeComponentDataInterface::FLandscapeComponentDataInterface
 	Component(InComponent),
 	HeightMipData(NULL),
 	XYOffsetMipData(NULL),
-	bNeedToDeleteDataInterface(false),
 	MipLevel(InMipLevel)
 {
-	if (Component->GetLandscapeInfo())
-	{
-		DataInterface = Component->GetLandscapeInfo()->GetDataInterface();
-	}
-	else
-	{
-		DataInterface = new FLandscapeDataInterface;
-		bNeedToDeleteDataInterface = true;
-	}
-
 	// Offset and stride for this component's data in heightmap texture
 	HeightmapStride = Component->HeightmapTexture->Source.GetSizeX() >> MipLevel;
 	HeightmapComponentOffsetX = FMath::RoundToInt((float)(Component->HeightmapTexture->Source.GetSizeX() >> MipLevel) * Component->HeightmapScaleBias.Z);
@@ -37,10 +26,10 @@ LANDSCAPE_API FLandscapeComponentDataInterface::FLandscapeComponentDataInterface
 
 	if (MipLevel < Component->HeightmapTexture->Source.GetNumMips())
 	{
-		HeightMipData = (FColor*)DataInterface->LockMip(Component->HeightmapTexture, MipLevel);
+		HeightMipData = (FColor*)DataInterface.LockMip(Component->HeightmapTexture, MipLevel);
 		if (Component->XYOffsetmapTexture)
 		{
-			XYOffsetMipData = (FColor*)DataInterface->LockMip(Component->XYOffsetmapTexture, MipLevel);
+			XYOffsetMipData = (FColor*)DataInterface.LockMip(Component->XYOffsetmapTexture, MipLevel);
 		}
 	}
 }
@@ -49,23 +38,12 @@ LANDSCAPE_API FLandscapeComponentDataInterface::~FLandscapeComponentDataInterfac
 {
 	if (HeightMipData)
 	{
-		DataInterface->UnlockMip(Component->HeightmapTexture, MipLevel);
+		DataInterface.UnlockMip(Component->HeightmapTexture, MipLevel);
 		if (Component->XYOffsetmapTexture)
 		{
-			DataInterface->UnlockMip(Component->XYOffsetmapTexture, MipLevel);
+			DataInterface.UnlockMip(Component->XYOffsetmapTexture, MipLevel);
 		}
 	}
-
-	if (bNeedToDeleteDataInterface)
-	{
-		delete DataInterface;
-		DataInterface = NULL;
-	}
-}
-
-LANDSCAPE_API void FLandscapeComponentDataInterface::UnlockRawHeightData() const
-{
-	DataInterface->UnlockMip(Component->HeightmapTexture, MipLevel);
 }
 
 LANDSCAPE_API void FLandscapeComponentDataInterface::GetHeightmapTextureData(TArray<FColor>& OutData, bool bOkToFail)
@@ -123,7 +101,7 @@ LANDSCAPE_API bool FLandscapeComponentDataInterface::GetWeightmapTextureData(ULa
 	OutData.Empty(FMath::Square(WeightmapSize));
 	OutData.AddUninitialized(FMath::Square(WeightmapSize));
 
-	FColor* WeightMipData = (FColor*)DataInterface->LockMip(Component->WeightmapTextures[Component->WeightmapLayerAllocations[LayerIdx].WeightmapTextureIndex], MipLevel);
+	FColor* WeightMipData = (FColor*)DataInterface.LockMip(Component->WeightmapTextures[Component->WeightmapLayerAllocations[LayerIdx].WeightmapTextureIndex], MipLevel);
 
 	// Channel remapping
 	int32 ChannelOffsets[4] = { (int32)STRUCT_OFFSET(FColor, R), (int32)STRUCT_OFFSET(FColor, G), (int32)STRUCT_OFFSET(FColor, B), (int32)STRUCT_OFFSET(FColor, A) };
@@ -135,7 +113,7 @@ LANDSCAPE_API bool FLandscapeComponentDataInterface::GetWeightmapTextureData(ULa
 		OutData[i] = SrcTextureData[i * 4];
 	}
 
-	DataInterface->UnlockMip(Component->WeightmapTextures[Component->WeightmapLayerAllocations[LayerIdx].WeightmapTextureIndex], MipLevel);
+	DataInterface.UnlockMip(Component->WeightmapTextures[Component->WeightmapLayerAllocations[LayerIdx].WeightmapTextureIndex], MipLevel);
 	return true;
 }
 
