Copying //UE4/Dev-Landscape to //UE4/Dev-Main (Source: //UE4/Dev-Landscape @ 3053591)

diff --git a/Engine/Source/Runtime/Landscape/Private/LandscapeEdit.cpp b/Engine/Source/Runtime/Landscape/Private/LandscapeEdit.cpp
--- a/Engine/Source/Runtime/Landscape/Private/LandscapeEdit.cpp
+++ b/Engine/Source/Runtime/Landscape/Private/LandscapeEdit.cpp
@@ -30,7 +30,6 @@ LandscapeEdit.cpp: Landscape editing
 #if WITH_EDITOR
 #include "RawMesh.h"
 #include "ScopedTransaction.h"
-#include "ImageWrapper.h"
 #include "Engine/Level.h"
 #include "EngineUtils.h"
 #include "Engine/Engine.h"
@@ -40,8 +39,8 @@ LandscapeEdit.cpp: Landscape editing
 #include "SlateBasics.h"  // For AddNotification
 #include "SNotificationList.h"
 #include "NotificationManager.h"
+#include "LandscapeEditorModule.h"
 #endif
-#include "ComponentReregisterContext.h"
 
 DEFINE_LOG_CATEGORY(LogLandscape);
 
@@ -351,7 +350,7 @@ void ULandscapeComponent::PostEditUndo()
 
 	Super::PostEditUndo();
 
-	if (EditToolRenderData)
+	if (EditToolRenderData && !IsPendingKill())
 	{
 		EditToolRenderData->UpdateDebugColorMaterial();
 		EditToolRenderData->UpdateSelectionMaterial(EditToolRenderData->SelectedType);
@@ -362,6 +361,15 @@ void ULandscapeComponent::PostEditUndo()
 	GetLandscapeProxy()->FlushGrassComponents(&Components);
 }
 
+void ALandscapeProxy::FixupWeightmaps()
+{
+	WeightmapUsageMap.Empty();
+	for (ULandscapeComponent* Component : LandscapeComponents)
+	{
+		Component->FixupWeightmaps();
+	}
+}
+
 void ULandscapeComponent::FixupWeightmaps()
 {
 	if (GIsEditor && !HasAnyFlags(RF_ClassDefaultObject))
@@ -1954,7 +1962,7 @@ LANDSCAPE_API void ALandscapeProxy::Import(
 	Modify();
 
 	// Create and initialize landscape info object
-	GetLandscapeInfo(true)->RegisterActor(this);
+	ULandscapeInfo* LandscapeInfo = CreateLandscapeInfo();
 
 	const int32 NumPatchesX = (VertsX - 1);
 	const int32 NumPatchesY = (VertsY - 1);
@@ -2470,10 +2478,10 @@ LANDSCAPE_API void ALandscapeProxy::Import(
 				{
 					FColor* const TexData = &(HeightmapInfo.HeightmapTextureMipData[Mip])[X + Y * MipSizeU];
 
-					const FColor* const PreMipTexData00 = &(HeightmapInfo.HeightmapTextureMipData[Mip - 1])[(X * 2 + 0) + (Y * 2 + 0)  * PrevMipSizeU];
-					const FColor* const PreMipTexData01 = &(HeightmapInfo.HeightmapTextureMipData[Mip - 1])[(X * 2 + 0) + (Y * 2 + 1)  * PrevMipSizeU];
-					const FColor* const PreMipTexData10 = &(HeightmapInfo.HeightmapTextureMipData[Mip - 1])[(X * 2 + 1) + (Y * 2 + 0)  * PrevMipSizeU];
-					const FColor* const PreMipTexData11 = &(HeightmapInfo.HeightmapTextureMipData[Mip - 1])[(X * 2 + 1) + (Y * 2 + 1)  * PrevMipSizeU];
+					const FColor* const PreMipTexData00 = &(HeightmapInfo.HeightmapTextureMipData[Mip - 1])[(X * 2 + 0) + (Y * 2 + 0) * PrevMipSizeU];
+					const FColor* const PreMipTexData01 = &(HeightmapInfo.HeightmapTextureMipData[Mip - 1])[(X * 2 + 0) + (Y * 2 + 1) * PrevMipSizeU];
+					const FColor* const PreMipTexData10 = &(HeightmapInfo.HeightmapTextureMipData[Mip - 1])[(X * 2 + 1) + (Y * 2 + 0) * PrevMipSizeU];
+					const FColor* const PreMipTexData11 = &(HeightmapInfo.HeightmapTextureMipData[Mip - 1])[(X * 2 + 1) + (Y * 2 + 1) * PrevMipSizeU];
 
 					TexData->R = (((int32)PreMipTexData00->R + (int32)PreMipTexData01->R + (int32)PreMipTexData10->R + (int32)PreMipTexData11->R) >> 2);
 					TexData->G = (((int32)PreMipTexData00->G + (int32)PreMipTexData01->G + (int32)PreMipTexData10->G + (int32)PreMipTexData11->G) >> 2);
@@ -2501,7 +2509,9 @@ LANDSCAPE_API void ALandscapeProxy::Import(
 
 	ReimportHeightmapFilePath = HeightmapFileName;
 
-	ULandscapeInfo::RecreateLandscapeInfo(GetWorld(), false);
+	LandscapeInfo->UpdateLayerInfoMap();
+	LandscapeInfo->RecreateCollisionComponents();
+	LandscapeInfo->UpdateAllAddCollisions();
 
 	GWarn->EndSlowTask();
 }
@@ -2818,7 +2828,7 @@ FIntRect ALandscapeProxy::GetBoundingRect() const
 
 bool ALandscape::HasAllComponent()
 {
-	ULandscapeInfo* Info = GetLandscapeInfo(false);
+	ULandscapeInfo* Info = GetLandscapeInfo();
 	if (Info && Info->XYtoComponentMap.Num() == LandscapeComponents.Num())
 	{
 		// all components are owned by this Landscape actor (no Landscape Proxies)
@@ -2959,47 +2969,6 @@ bool ULandscapeInfo::IsValidPosition(int32 X, int32 Y)
 	return false;
 }
 
-void ULandscapeInfo::Export(const TArray<ULandscapeLayerInfoObject*>& LayerInfos, const TArray<FString>& Filenames)
-{
-	check(Filenames.Num() > 0);
-
-	int32 MinX = MAX_int32;
-	int32 MinY = MAX_int32;
-	int32 MaxX = -MAX_int32;
-	int32 MaxY = -MAX_int32;
-
-	if (!GetLandscapeExtent(MinX, MinY, MaxX, MaxY))
-	{
-		return;
-	}
-
-	GWarn->BeginSlowTask(LOCTEXT("BeginExportingLandscapeTask", "Exporting Landscape"), true);
-
-	FLandscapeEditDataInterface LandscapeEdit(this);
-
-	TArray<uint8> HeightData;
-	HeightData.AddZeroed((1 + MaxX - MinX)*(1 + MaxY - MinY)*sizeof(uint16));
-	LandscapeEdit.GetHeightDataFast(MinX, MinY, MaxX, MaxY, (uint16*)HeightData.GetData(), 0);
-	FFileHelper::SaveArrayToFile(HeightData, *Filenames[0]);
-
-	for (int32 i = 1; i < Filenames.Num(); i++)
-	{
-		if (i <= LayerInfos.Num())
-		{
-			TArray<uint8> WeightData;
-			WeightData.AddZeroed((1 + MaxX - MinX)*(1 + MaxY - MinY));
-			ULandscapeLayerInfoObject* LayerInfo = LayerInfos[i - 1];
-			if (LayerInfo)
-			{
-				LandscapeEdit.GetWeightDataFast(LayerInfo, MinX, MinY, MaxX, MaxY, WeightData.GetData(), 0);
-			}
-			FFileHelper::SaveArrayToFile(WeightData, *Filenames[i]);
-		}
-	}
-
-	GWarn->EndSlowTask();
-}
-
 void ULandscapeInfo::ExportHeightmap(const FString& Filename)
 {
 	int32 MinX = MAX_int32;
@@ -3014,26 +2983,19 @@ void ULandscapeInfo::ExportHeightmap(const FString& Filename)
 
 	GWarn->BeginSlowTask(LOCTEXT("BeginExportingLandscapeHeightmapTask", "Exporting Landscape Heightmap"), true);
 
+	ILandscapeEditorModule& LandscapeEditorModule = FModuleManager::GetModuleChecked<ILandscapeEditorModule>("LandscapeEditor");
 	FLandscapeEditDataInterface LandscapeEdit(this);
 
-	TArray<uint8> HeightData;
-	HeightData.AddZeroed((MaxX - MinX + 1) * (MaxY - MinY + 1) * sizeof(uint16));
-	LandscapeEdit.GetHeightDataFast(MinX, MinY, MaxX, MaxY, (uint16*)HeightData.GetData(), 0);
+	TArray<uint16> HeightData;
+	HeightData.AddZeroed((MaxX - MinX + 1) * (MaxY - MinY + 1));
+	LandscapeEdit.GetHeightDataFast(MinX, MinY, MaxX, MaxY, HeightData.GetData(), 0);
 
-	if (Filename.EndsWith(".png"))
+	const ILandscapeHeightmapFileFormat* HeightmapFormat = LandscapeEditorModule.GetHeightmapFormatByExtension(*FPaths::GetExtension(Filename, true));
+	if (HeightmapFormat)
 	{
-		IImageWrapperModule& ImageWrapperModule = FModuleManager::LoadModuleChecked<IImageWrapperModule>("ImageWrapper");
-		IImageWrapperPtr ImageWrapper = ImageWrapperModule.CreateImageWrapper(EImageFormat::PNG);
-
-		const TArray<uint8>* RawData = nullptr;
-		if (ImageWrapper->SetRaw(HeightData.GetData(), HeightData.Num(), (MaxX - MinX + 1), (MaxY - MinY + 1), ERGBFormat::Gray, 16))
-		{
-			HeightData = ImageWrapper->GetCompressed();
-		}
+		HeightmapFormat->Export(*Filename, HeightData, {(uint32)(MaxX - MinX + 1), (uint32)(MaxY - MinY + 1)}, DrawScale * FVector(1, 1, LANDSCAPE_ZSCALE));
 	}
 
-	FFileHelper::SaveArrayToFile(HeightData, *Filename);
-
 	GWarn->EndSlowTask();
 }
 
@@ -3051,6 +3013,8 @@ void ULandscapeInfo::ExportLayer(ULandscapeLayerInfoObject* LayerInfo, const FSt
 
 	GWarn->BeginSlowTask(LOCTEXT("BeginExportingLandscapeWeightmapTask", "Exporting Landscape Layer Weightmap"), true);
 
+	ILandscapeEditorModule& LandscapeEditorModule = FModuleManager::GetModuleChecked<ILandscapeEditorModule>("LandscapeEditor");
+
 	TArray<uint8> WeightData;
 	WeightData.AddZeroed((MaxX - MinX + 1) * (MaxY - MinY + 1));
 	if (LayerInfo)
@@ -3059,20 +3023,12 @@ void ULandscapeInfo::ExportLayer(ULandscapeLayerInfoObject* LayerInfo, const FSt
 		LandscapeEdit.GetWeightDataFast(LayerInfo, MinX, MinY, MaxX, MaxY, WeightData.GetData(), 0);
 	}
 
-	if (Filename.EndsWith(".png"))
+	const ILandscapeWeightmapFileFormat* WeightmapFormat = LandscapeEditorModule.GetWeightmapFormatByExtension(*FPaths::GetExtension(Filename, true));
+	if (WeightmapFormat)
 	{
-		IImageWrapperModule& ImageWrapperModule = FModuleManager::LoadModuleChecked<IImageWrapperModule>("ImageWrapper");
-		IImageWrapperPtr ImageWrapper = ImageWrapperModule.CreateImageWrapper(EImageFormat::PNG);
-
-		const TArray<uint8>* RawData = nullptr;
-		if (ImageWrapper->SetRaw(WeightData.GetData(), WeightData.Num(), (MaxX - MinX + 1), (MaxY - MinY + 1), ERGBFormat::Gray, 8))
-		{
-			WeightData = ImageWrapper->GetCompressed();
-		}
+		WeightmapFormat->Export(*Filename, LayerInfo->LayerName, WeightData, {(uint32)(MaxX - MinX + 1), (uint32)(MaxY - MinY + 1)});
 	}
 
-	FFileHelper::SaveArrayToFile(WeightData, *Filename);
-
 	GWarn->EndSlowTask();
 }
 
@@ -3270,6 +3226,12 @@ void ALandscapeProxy::PostEditImport()
 {
 	Super::PostEditImport();
 
+	// during import this gets called multiple times, without a valid guid the first time
+	if (LandscapeGuid.IsValid())
+	{
+		CreateLandscapeInfo();
+	}
+
 	for (int32 ComponentIndex = 0; ComponentIndex < LandscapeComponents.Num(); ++ComponentIndex)
 	{
 		ULandscapeComponent* Comp = LandscapeComponents[ComponentIndex];
@@ -3279,8 +3241,6 @@ void ALandscapeProxy::PostEditImport()
 			Comp->UpdateMaterialInstances();
 		}
 	}
-
-	GEngine->DeferredCommands.AddUnique(TEXT("UpdateLandscapeEditorData"));
 }
 
 void ALandscape::PostEditMove(bool bFinished)
@@ -3296,8 +3256,6 @@ void ALandscape::PostEditMove(bool bFinished)
 
 void ALandscape::PostEditImport()
 {
-	Super::PostEditImport();
-
 	if (GetWorld())
 	{
 		for (ALandscape* Landscape : TActorRange<ALandscape>(GetWorld()))
@@ -3310,6 +3268,8 @@ void ALandscape::PostEditImport()
 			}
 		}
 	}
+
+	Super::PostEditImport();
 }
 #endif	//WITH_EDITOR
 
@@ -3359,7 +3319,7 @@ void ULandscapeLayerInfoObject::PostEditChangeProperty(FPropertyChangedEvent& Pr
 				ALandscapeProxy* Proxy = *It;
 				if (Proxy->GetWorld() && !Proxy->GetWorld()->IsPlayInEditor())
 				{
-					ULandscapeInfo* Info = Proxy->GetLandscapeInfo(false);
+					ULandscapeInfo* Info = Proxy->GetLandscapeInfo();
 					if (Info)
 					{
 						for (int32 i = 0; i < Info->Layers.Num(); ++i)
@@ -3577,7 +3537,7 @@ void ALandscapeProxy::PostEditChangeProperty(FPropertyChangedEvent& PropertyChan
 				ModifiedScale.Y = FMath::Abs(RootComponent->RelativeScale3D.X)*FMath::Sign(ModifiedScale.Y);
 			}
 
-			ULandscapeInfo* Info = GetLandscapeInfo(false);
+			ULandscapeInfo* Info = GetLandscapeInfo();
 
 			// Correct for attempts to scale to 0 on any axis
 			if (ModifiedScale.X == 0)
@@ -3611,21 +3571,21 @@ void ALandscapeProxy::PostEditChangeProperty(FPropertyChangedEvent& PropertyChan
 				Info->DrawScale = ModifiedScale;
 			}
 
-			// We need to regenerate collision objects, they depend on scale value 
-			for (int32 ComponentIndex = 0; ComponentIndex < CollisionComponents.Num(); ComponentIndex++)
+			// We need to regenerate collision objects, they depend on scale value
+			if (PropertyChangedEvent.ChangeType != EPropertyChangeType::Interactive)
 			{
-				ULandscapeHeightfieldCollisionComponent* Comp = CollisionComponents[ComponentIndex];
-				if (Comp)
+				for (int32 ComponentIndex = 0; ComponentIndex < CollisionComponents.Num(); ComponentIndex++)
 				{
-					Comp->RecreateCollision();
+					ULandscapeHeightfieldCollisionComponent* Comp = CollisionComponents[ComponentIndex];
+					if (Comp)
+					{
+						Comp->RecreateCollision();
+					}
 				}
 			}
 		}
 	}
 
-	// Must do this *after* correcting the scale or reattaching the landscape components will crash!
-	Super::PostEditChangeProperty(PropertyChangedEvent);
-
 	if (GIsEditor && PropertyName == FName(TEXT("StreamingDistanceMultiplier")))
 	{
 		// Recalculate in a few seconds.
@@ -3645,7 +3605,10 @@ void ALandscapeProxy::PostEditChangeProperty(FPropertyChangedEvent& PropertyChan
 		{
 			MarkComponentsRenderStateDirty();
 		}
-		RecreateCollisionComponents();
+		if (PropertyChangedEvent.ChangeType != EPropertyChangeType::Interactive)
+		{
+			RecreateCollisionComponents();
+		}
 	}
 	else if (PropertyName == GET_MEMBER_NAME_CHECKED(ALandscapeProxy, bUseMaterialPositionOffsetInStaticLighting))
 	{
@@ -3667,22 +3630,24 @@ void ALandscapeProxy::PostEditChangeProperty(FPropertyChangedEvent& PropertyChan
 			}
 		}
 	}
+
+	// Must do this *after* correcting the scale or reattaching the landscape components will crash!
+	// Must do this *after* clamping values / propogating values to components
+	Super::PostEditChangeProperty(PropertyChangedEvent);
 }
 
 void ALandscapeStreamingProxy::PostEditChangeProperty(FPropertyChangedEvent& PropertyChangedEvent)
 {
-	Super::PostEditChangeProperty(PropertyChangedEvent);
-
 	const FName PropertyName = PropertyChangedEvent.MemberProperty ? PropertyChangedEvent.MemberProperty->GetFName() : NAME_None;
 	if (PropertyName == FName(TEXT("LandscapeActor")))
 	{
 		if (LandscapeActor && IsValidLandscapeActor(LandscapeActor.Get()))
 		{
 			LandscapeGuid = LandscapeActor->GetLandscapeGuid();
-			// defer LandscapeInfo setup
 			if (GIsEditor && GetWorld() && !GetWorld()->IsPlayInEditor())
 			{
-				GEngine->DeferredCommands.AddUnique(TEXT("UpdateLandscapeEditorData"));
+				// TODO - only need to refresh the old and new landscape info
+				ULandscapeInfo::RecreateLandscapeInfo(GetWorld(), false);
 			}
 		}
 		else
@@ -3719,12 +3684,13 @@ void ALandscapeStreamingProxy::PostEditChangeProperty(FPropertyChangedEvent& Pro
 			}
 		}
 	}
+
+	// Must do this *after* clamping values
+	Super::PostEditChangeProperty(PropertyChangedEvent);
 }
 
 void ALandscape::PostEditChangeProperty(FPropertyChangedEvent& PropertyChangedEvent)
 {
-	Super::PostEditChangeProperty(PropertyChangedEvent);
-
 	const FName PropertyName = PropertyChangedEvent.Property ? PropertyChangedEvent.Property->GetFName() : NAME_None;
 
 	bool ChangedMaterial = false;
@@ -3738,7 +3704,7 @@ void ALandscape::PostEditChangeProperty(FPropertyChangedEvent& PropertyChangedEv
 	if (PropertyName == FName(TEXT("LandscapeMaterial")) || PropertyName == FName(TEXT("LandscapeHoleMaterial")))
 	{
 		FMaterialUpdateContext MaterialUpdateContext;
-		GetLandscapeInfo()->UpdateLayerInfoMap(/*this*/);
+		Info->UpdateLayerInfoMap(/*this*/);
 
 		ChangedMaterial = true;
 
@@ -3809,9 +3775,12 @@ void ALandscape::PostEditChangeProperty(FPropertyChangedEvent& PropertyChangedEv
 		bChangedNavRelevance = true;
 	}
 
+	// Must do this *after* clamping values
+	Super::PostEditChangeProperty(PropertyChangedEvent);
+
 	bPropagateToProxies = bPropagateToProxies || bNeedsRecalcBoundingBox || bChangedLighting;
 
-	if (Info)
+	if (ensure(Info))
 	{
 		if (bPropagateToProxies)
 		{
@@ -3830,37 +3799,39 @@ void ALandscape::PostEditChangeProperty(FPropertyChangedEvent& PropertyChangedEv
 			LandscapeEdit.RecalculateNormals();
 		}
 
-		// We cannot iterate the XYtoComponentMap directly because reregistering components modifies the array.
-		TArray<ULandscapeComponent*> AllComponents;
-		Info->XYtoComponentMap.GenerateValueArray(AllComponents);
-		for (ULandscapeComponent* Comp : AllComponents)
+		if (bNeedsRecalcBoundingBox || ChangedMaterial || bChangedLighting || bChangedNavRelevance)
 		{
-			if (Comp)
+			// We cannot iterate the XYtoComponentMap directly because reregistering components modifies the array.
+			TArray<ULandscapeComponent*> AllComponents;
+			Info->XYtoComponentMap.GenerateValueArray(AllComponents);
+			for (ULandscapeComponent* Comp : AllComponents)
 			{
-				if (bNeedsRecalcBoundingBox)
+				if (ensure(Comp))
 				{
-					Comp->UpdateCachedBounds();
-					Comp->UpdateBounds();
-				}
+					Comp->Modify();
 
-				if (ChangedMaterial)
-				{
-					// Update the MIC
-					Comp->UpdateMaterialInstances();
-				}
+					if (bNeedsRecalcBoundingBox)
+					{
+						Comp->UpdateCachedBounds();
+						Comp->UpdateBounds();
+					}
 
-				if (bChangedLighting)
-				{
-					Comp->InvalidateLightingCache();
-				}
+					if (ChangedMaterial)
+					{
+						// Update the MIC
+						Comp->UpdateMaterialInstances();
+					}
 
-				if (bChangedNavRelevance)
-				{
-					Comp->UpdateNavigationRelevance();
-				}
+					if (bChangedLighting)
+					{
+						Comp->InvalidateLightingCache();
+					}
 
-				// Reattach all components
-				FComponentReregisterContext ReregisterContext(Comp);
+					if (bChangedNavRelevance)
+					{
+						Comp->UpdateNavigationRelevance();
+					}
+				}
 			}
 		}
 
@@ -3869,21 +3840,7 @@ void ALandscape::PostEditChangeProperty(FPropertyChangedEvent& PropertyChangedEv
 		{
 			for (ALandscapeGizmoActiveActor* Gizmo : TActorRange<ALandscapeGizmoActiveActor>(GetWorld()))
 			{
-				Gizmo->ReregisterAllComponents();
-			}
-		}
-
-		if (ChangedMaterial)
-		{
-			if (GIsEditor && GetWorld() && !GetWorld()->IsPlayInEditor())
-			{
-				GEngine->DeferredCommands.AddUnique(TEXT("UpdateLandscapeMIC"));
-			}
-
-			// Update all the proxies...
-			for (ALandscapeProxy* Proxy : Info->Proxies)
-			{
-				Proxy->MarkComponentsRenderStateDirty();
+				Gizmo->MarkComponentsRenderStateDirty();
 			}
 		}
 	}
@@ -3930,9 +3887,10 @@ void ULandscapeComponent::SetLOD(bool bForcedLODChanged, int32 InLODValue)
 	}
 
 	InvalidateLightingCache();
+	MarkRenderStateDirty();
 
 	// Update neighbor components
-	ULandscapeInfo* Info = GetLandscapeInfo(false);
+	ULandscapeInfo* Info = GetLandscapeInfo();
 	if (Info)
 	{
 		FIntPoint ComponentBase = GetSectionBase() / ComponentSizeQuads;
@@ -3955,11 +3913,10 @@ void ULandscapeComponent::SetLOD(bool bForcedLODChanged, int32 InLODValue)
 			{
 				Comp->Modify();
 				Comp->InvalidateLightingCache();
-				FComponentReregisterContext ReregisterContext(Comp);
+				Comp->MarkRenderStateDirty();
 			}
 		}
 	}
-	FComponentReregisterContext ReregisterContext(this);
 }
 
 void ULandscapeComponent::PreEditChange(UProperty* PropertyThatWillChange)
@@ -3968,7 +3925,7 @@ void ULandscapeComponent::PreEditChange(UProperty* PropertyThatWillChange)
 	if (GIsEditor && PropertyThatWillChange && (PropertyThatWillChange->GetFName() == GET_MEMBER_NAME_CHECKED(ULandscapeComponent, ForcedLOD) || PropertyThatWillChange->GetFName() == GET_MEMBER_NAME_CHECKED(ULandscapeComponent, LODBias)))
 	{
 		// PreEdit unregister component and re-register after PostEdit so we will lose XYtoComponentMap for this component
-		ULandscapeInfo* Info = GetLandscapeInfo(false);
+		ULandscapeInfo* Info = GetLandscapeInfo();
 		if (Info)
 		{
 			FIntPoint ComponentKey = GetSectionBase() / ComponentSizeQuads;
@@ -3984,14 +3941,10 @@ void ULandscapeComponent::PreEditChange(UProperty* PropertyThatWillChange)
 
 void ULandscapeComponent::PostEditChangeProperty(FPropertyChangedEvent& PropertyChangedEvent)
 {
-	Super::PostEditChangeProperty(PropertyChangedEvent);
-
 	const FName PropertyName = PropertyChangedEvent.Property ? PropertyChangedEvent.Property->GetFName() : NAME_None;
 	if (PropertyName == FName(TEXT("OverrideMaterial")))
 	{
 		UpdateMaterialInstances();
-		// Reregister all components
-		FComponentReregisterContext ReregisterContext(this);
 	}
 	else if (GIsEditor && (PropertyName == FName(TEXT("ForcedLOD")) || PropertyName == FName(TEXT("LODBias"))))
 	{
@@ -4022,8 +3975,14 @@ void ULandscapeComponent::PostEditChangeProperty(FPropertyChangedEvent& Property
 	{
 		CollisionMipLevel = FMath::Clamp<int32>(CollisionMipLevel, 0, FMath::CeilLogTwo(SubsectionSizeQuads + 1) - 1);
 		SimpleCollisionMipLevel = FMath::Clamp<int32>(SimpleCollisionMipLevel, 0, FMath::CeilLogTwo(SubsectionSizeQuads + 1) - 1);
-		UpdateCollisionData(true); // Rebuild for new CollisionMipLevel
+		if (PropertyChangedEvent.ChangeType != EPropertyChangeType::Interactive)
+		{
+			UpdateCollisionData(true); // Rebuild for new CollisionMipLevel
+		}
 	}
+
+	// Must do this *after* clamping values
+	Super::PostEditChangeProperty(PropertyChangedEvent);
 }
 
 TSet<class ULandscapeComponent*> ULandscapeInfo::GetSelectedComponents() const
@@ -4129,16 +4088,6 @@ void ULandscapeInfo::ClearSelectedRegion(bool bIsComponentwise /*= true*/)
 	}
 }
 
-struct FLandscapeDataInterface* ULandscapeInfo::GetDataInterface()
-{
-	if (DataInterface == nullptr)
-	{
-		DataInterface = new FLandscapeDataInterface();
-	}
-
-	return DataInterface;
-}
-
 void ULandscapeComponent::ReallocateWeightmaps(FLandscapeEditDataInterface* DataInterface)
 {
 	ALandscapeProxy* Proxy = GetLandscapeProxy();
@@ -4389,43 +4338,46 @@ void ALandscapeProxy::RemoveInvalidWeightmaps()
 void ULandscapeComponent::RemoveInvalidWeightmaps()
 {
 	// Adjust WeightmapTextureIndex index for other layers
-	TSet<int32> UsedTextureIndices;
-	TSet<int32> AllTextureIndices;
-	for (int32 LayerIdx = 0; LayerIdx < WeightmapLayerAllocations.Num(); LayerIdx++)
+	TSet<int32> UnUsedTextureIndices;
 	{
-		UsedTextureIndices.Add(WeightmapLayerAllocations[LayerIdx].WeightmapTextureIndex);
-	}
+		TSet<int32> UsedTextureIndices;
+		for (int32 LayerIdx = 0; LayerIdx < WeightmapLayerAllocations.Num(); LayerIdx++)
+		{
+			UsedTextureIndices.Add(WeightmapLayerAllocations[LayerIdx].WeightmapTextureIndex);
+		}
 
-	for (int32 WeightIdx = 0; WeightIdx < WeightmapTextures.Num(); ++WeightIdx)
-	{
-		AllTextureIndices.Add(WeightIdx);
+		for (int32 WeightIdx = 0; WeightIdx < WeightmapTextures.Num(); ++WeightIdx)
+		{
+			if (!UsedTextureIndices.Contains(WeightIdx))
+			{
+				UnUsedTextureIndices.Add(WeightIdx);
+			}
+		}
 	}
 
-	TSet<int32> UnUsedTextureIndices = AllTextureIndices.Difference(UsedTextureIndices);
-
-	int32 DeletedLayers = 0;
-	for (TSet<int32>::TIterator It(UnUsedTextureIndices); It; ++It)
+	int32 RemovedTextures = 0;
+	for (int32 UnusedIndex : UnUsedTextureIndices)
 	{
-		int32 DeleteLayerWeightmapTextureIndex = *It - DeletedLayers;
-		WeightmapTextures[DeleteLayerWeightmapTextureIndex]->SetFlags(RF_Transactional);
-		WeightmapTextures[DeleteLayerWeightmapTextureIndex]->Modify();
-		WeightmapTextures[DeleteLayerWeightmapTextureIndex]->MarkPackageDirty();
-		WeightmapTextures[DeleteLayerWeightmapTextureIndex]->ClearFlags(RF_Standalone);
-		WeightmapTextures.RemoveAt(DeleteLayerWeightmapTextureIndex);
+		int32 WeightmapTextureIndexToRemove = UnusedIndex - RemovedTextures;
+		WeightmapTextures[WeightmapTextureIndexToRemove]->SetFlags(RF_Transactional);
+		WeightmapTextures[WeightmapTextureIndexToRemove]->Modify();
+		WeightmapTextures[WeightmapTextureIndexToRemove]->MarkPackageDirty();
+		WeightmapTextures[WeightmapTextureIndexToRemove]->ClearFlags(RF_Standalone);
+		WeightmapTextures.RemoveAt(WeightmapTextureIndexToRemove);
 
 		// Adjust WeightmapTextureIndex index for other layers
 		for (int32 LayerIdx = 0; LayerIdx < WeightmapLayerAllocations.Num(); LayerIdx++)
 		{
 			FWeightmapLayerAllocationInfo& Allocation = WeightmapLayerAllocations[LayerIdx];
 
-			if (Allocation.WeightmapTextureIndex > DeleteLayerWeightmapTextureIndex)
+			if (Allocation.WeightmapTextureIndex > WeightmapTextureIndexToRemove)
 			{
 				Allocation.WeightmapTextureIndex--;
 			}
 
-			check(Allocation.WeightmapTextureIndex < WeightmapTextures.Num());
+			checkSlow(Allocation.WeightmapTextureIndex < WeightmapTextures.Num());
 		}
-		DeletedLayers++;
+		RemovedTextures++;
 	}
 }
 
@@ -4892,7 +4844,7 @@ void ULandscapeComponent::GeneratePlatformPixelData()
 	CreateEmptyTextureMips(NewWeightNormalmapTexture);
 
 	{
-		FLandscapeEditDataInterface LandscapeEdit(GetLandscapeInfo(false));
+		FLandscapeEditDataInterface LandscapeEdit(GetLandscapeInfo());
 
 		if (WeightmapTextures.Num() > 0)
 		{
diff --git a/Engine/Source/Runtime/Landscape/Private/LandscapeEditInterface.cpp b/Engine/Source/Runtime/Landscape/Private/LandscapeEditInterface.cpp
--- a/Engine/Source/Runtime/Landscape/Private/LandscapeEditInterface.cpp
+++ b/Engine/Source/Runtime/Landscape/Private/LandscapeEditInterface.cpp
@@ -15,7 +15,7 @@ LandscapeEditInterface.cpp: Landscape editing interface
 #include "LandscapeInfo.h"
 #include "LandscapeLayerInfoObject.h"
 #include "ComponentReregisterContext.h"
-#include "FixedSizeArrayView.h"
+#include "Containers/ArrayView.h"
 
 // Channel remapping
 extern const size_t ChannelOffsets[4] = {STRUCT_OFFSET(FColor,R), STRUCT_OFFSET(FColor,G), STRUCT_OFFSET(FColor,B), STRUCT_OFFSET(FColor,A)};
@@ -2369,7 +2369,7 @@ inline bool FLandscapeEditDataInterface::IsWhitelisted(const ULandscapeLayerInfo
 	return true;
 }
 
-inline TMap<const ULandscapeLayerInfoObject*, uint32> FLandscapeEditDataInterface::CountWeightBlendedLayerInfluence(const int32 ComponentIndexX, const int32 ComponentIndexY, TOptional<TFixedSizeArrayView<const uint8* const>> InOptionalLayerDataPtrs)
+inline TMap<const ULandscapeLayerInfoObject*, uint32> FLandscapeEditDataInterface::CountWeightBlendedLayerInfluence(const int32 ComponentIndexX, const int32 ComponentIndexY, TOptional<TArrayView<const uint8* const>> InOptionalLayerDataPtrs)
 {
 	// the counts should easily fit in a uint32, a 255x255 x2x2 Component with weights of all 255 only totals 26 bits
 	checkSlow(FMath::CeilLogTwo(ComponentSizeQuads + 1) * 2 + 8 /*ceillog2(255)*/ <= 32);
@@ -2380,7 +2380,7 @@ inline TMap<const ULandscapeLayerInfoObject*, uint32> FLandscapeEditDataInterfac
 	// used if InOptionalLayerDataPtrs is null
 	TArray<FLandscapeTextureDataInfo*, TInlineAllocator<2>> InternalTexDataInfos;
 	TArray<const uint8*, TInlineAllocator<8>> InternalLayerDataPtrs;
-	TFixedSizeArrayView<const uint8* const> LayerDataPtrs;
+	TArrayView<const uint8* const> LayerDataPtrs;
 	if (InOptionalLayerDataPtrs)
 	{
 		check(InOptionalLayerDataPtrs->Num() == Component->WeightmapLayerAllocations.Num());
@@ -2440,7 +2440,7 @@ inline TMap<const ULandscapeLayerInfoObject*, uint32> FLandscapeEditDataInterfac
 	return LayerInfluenceMap;
 }
 
-const ULandscapeLayerInfoObject* FLandscapeEditDataInterface::ChooseReplacementLayer(const ULandscapeLayerInfoObject* const LayerInfo, const int32 ComponentIndexX, const int32 SubIndexX, const int32 SubX, const int32 ComponentIndexY, const int32 SubIndexY, const int32 SubY, TMap<FIntPoint, TMap<const ULandscapeLayerInfoObject*, uint32>>& LayerInfluenceCache, TFixedSizeArrayView<const uint8* const> LayerDataPtrs)
+const ULandscapeLayerInfoObject* FLandscapeEditDataInterface::ChooseReplacementLayer(const ULandscapeLayerInfoObject* const LayerInfo, const int32 ComponentIndexX, const int32 SubIndexX, const int32 SubX, const int32 ComponentIndexY, const int32 SubIndexY, const int32 SubY, TMap<FIntPoint, TMap<const ULandscapeLayerInfoObject*, uint32>>& LayerInfluenceCache, TArrayView<const uint8* const> LayerDataPtrs)
 {
 	const TMap<const ULandscapeLayerInfoObject*, uint32>* LayerInfluenceMapCacheEntry = LayerInfluenceCache.Find(FIntPoint(ComponentIndexX, ComponentIndexY));
 	if (!LayerInfluenceMapCacheEntry)
diff --git a/Engine/Source/Runtime/Landscape/Private/LandscapeInfoMap.cpp b/Engine/Source/Runtime/Landscape/Private/LandscapeInfoMap.cpp
--- a/Engine/Source/Runtime/Landscape/Private/LandscapeInfoMap.cpp
+++ b/Engine/Source/Runtime/Landscape/Private/LandscapeInfoMap.cpp
@@ -19,7 +19,7 @@ void ULandscapeInfoMap::Serialize(FArchive& Ar)
 {
 	Super::Serialize(Ar);
 
-	if (!Ar.IsLoading() && !Ar.IsSaving())
+	if (Ar.IsTransacting() || Ar.IsObjectReferenceCollector())
 	{
 		Ar << Map;
 	}
diff --git a/Engine/Source/Runtime/Landscape/Private/LandscapeLight.cpp b/Engine/Source/Runtime/Landscape/Private/LandscapeLight.cpp
--- a/Engine/Source/Runtime/Landscape/Private/LandscapeLight.cpp
+++ b/Engine/Source/Runtime/Landscape/Private/LandscapeLight.cpp
@@ -9,6 +9,7 @@ LandscapeLight.cpp: Static lighting for LandscapeComponents
 #include "LandscapeInfo.h"
 #include "LandscapeRender.h"
 #include "LandscapeDataAccess.h"
+#include "LandscapeGrassType.h"
 
 #include "UnrealEngine.h"
 #include "ComponentReregisterContext.h"
@@ -724,7 +725,6 @@ void ULandscapeComponent::GetLightAndShadowMapMemoryUsage( int32& LightMapMemory
 	ShadowMapMemoryUsage = (Width * Height * 4 / 3); // assuming G8
 	return;
 }
-#endif
 
 void ULandscapeComponent::InvalidateLightingCacheDetailed(bool bInvalidateBuildEnqueuedLighting, bool bTranslationOnly)
 {
@@ -739,15 +739,31 @@ void ULandscapeComponent::InvalidateLightingCacheDetailed(bool bInvalidateBuildE
 
 		Super::InvalidateLightingCacheDetailed(bInvalidateBuildEnqueuedLighting, bTranslationOnly);
 
+		// invalidate grass that has bUseLandscapeLightmap so the new lightmap is applied to the grass
+		for (auto Iter = GetLandscapeProxy()->FoliageCache.CachedGrassComps.CreateIterator(); Iter; ++Iter)
+		{
+			const auto& GrassKey = Iter->Key;
+			const ULandscapeGrassType* GrassType = GrassKey.GrassType.Get();
+			const ULandscapeComponent* BasedOn = GrassKey.BasedOn.Get();
+			UHierarchicalInstancedStaticMeshComponent* GrassComponent = Iter->Foliage.Get();
+
+			if (BasedOn == this && GrassType && GrassComponent &&
+				GrassType->GrassVarieties.IsValidIndex(GrassKey.VarietyIndex) &&
+				GrassType->GrassVarieties[GrassKey.VarietyIndex].bUseLandscapeLightmap)
+			{
+				// Remove this grass component from the cache, which will cause it to be replaced
+				Iter.RemoveCurrent();
+			}
+		}
+
 		// Discard all cached lighting.
 		IrrelevantLights.Empty();
-		LightMap = NULL;
-		ShadowMap = NULL;
+		LightMap = nullptr;
+		ShadowMap = nullptr;
+	}
+	else
+	{
+		Super::InvalidateLightingCacheDetailed(bInvalidateBuildEnqueuedLighting, bTranslationOnly);
 	}
-
-	// invalidate grass in case bUseLandscapeLightmap is being used
-	// we don't need to invalidate the textures used to place grass, only the instances
-	TSet<ULandscapeComponent*> Components;
-	Components.Add(this);
-	GetLandscapeProxy()->FlushGrassComponents(&Components, false);
 }
+#endif
diff --git a/Engine/Source/Runtime/Landscape/Private/LandscapeModule.cpp b/Engine/Source/Runtime/Landscape/Private/LandscapeModule.cpp
--- a/Engine/Source/Runtime/Landscape/Private/LandscapeModule.cpp
+++ b/Engine/Source/Runtime/Landscape/Private/LandscapeModule.cpp
@@ -30,7 +30,7 @@ void AddPerWorldLandscapeData(UWorld* World)
 {
 	if (!World->PerModuleDataObjects.FindItemByClass<ULandscapeInfoMap>())
 	{
-		World->PerModuleDataObjects.Add(NewObject<ULandscapeInfoMap>(GetTransientPackage()));
+		World->PerModuleDataObjects.Add(NewObject<ULandscapeInfoMap>(GetTransientPackage(), NAME_None, RF_Transactional));
 	}
 }
 
diff --git a/Engine/Source/Runtime/Landscape/Private/LandscapeSplines.cpp b/Engine/Source/Runtime/Landscape/Private/LandscapeSplines.cpp
--- a/Engine/Source/Runtime/Landscape/Private/LandscapeSplines.cpp
+++ b/Engine/Source/Runtime/Landscape/Private/LandscapeSplines.cpp
@@ -573,6 +573,7 @@ void ULandscapeSplinesComponent::CheckForErrors()
 	Super::CheckForErrors();
 
 	UWorld* ThisOuterWorld = GetTypedOuter<UWorld>();
+	check(ThisOuterWorld->WorldType == EWorldType::Editor);
 
 	TSet<UWorld*> OutdatedWorlds;
 	TMap<UWorld*, FForeignWorldSplineData*> ForeignWorldSplineDataMapCache;
@@ -682,7 +683,7 @@ void ULandscapeSplinesComponent::PostLoad()
 	Super::PostLoad();
 
 #if WITH_EDITOR
-	if (GIsEditor)
+	if (GIsEditor && GetWorld()->WorldType == EWorldType::Editor)
 	{
 		// Build MeshComponentForeignOwnersMap (Component->Spline) from ForeignWorldSplineDataMap (World->Spline->Component)
 		for (auto& ForeignWorldSplineDataPair : ForeignWorldSplineDataMap)
@@ -714,7 +715,7 @@ void ULandscapeSplinesComponent::PostLoad()
 	CheckSplinesValid();
 
 #if WITH_EDITOR
-	if (GIsEditor)
+	if (GIsEditor && GetWorld()->WorldType == EWorldType::Editor)
 	{
 		CheckForErrors();
 	}
@@ -778,7 +779,10 @@ bool FForeignWorldSplineData::IsEmpty()
 ULandscapeSplinesComponent* ULandscapeSplinesComponent::GetStreamingSplinesComponentByLocation(const FVector& LocalLocation, bool bCreate /* = true*/)
 {
 	ALandscapeProxy* OuterLandscape = Cast<ALandscapeProxy>(GetOwner());
-	if (OuterLandscape)
+	if (OuterLandscape &&
+		// when copy/pasting this can get called with a null guid on the parent landscape
+		// this is fine, we won't have any cross-level meshes in this case anyway
+		OuterLandscape->GetLandscapeGuid().IsValid())
 	{
 		FVector LandscapeLocalLocation = ComponentToWorld.GetRelativeTransform(OuterLandscape->LandscapeActorToWorld()).TransformPosition(LocalLocation);
 		const int32 ComponentIndexX = (LandscapeLocalLocation.X >= 0.0f) ? FMath::FloorToInt(LandscapeLocalLocation.X / OuterLandscape->ComponentSizeQuads) : FMath::CeilToInt(LandscapeLocalLocation.X / OuterLandscape->ComponentSizeQuads);
@@ -832,28 +836,33 @@ ULandscapeSplinesComponent* ULandscapeSplinesComponent::GetStreamingSplinesCompo
 TArray<ULandscapeSplinesComponent*> ULandscapeSplinesComponent::GetAllStreamingSplinesComponents()
 {
 	ALandscapeProxy* OuterLandscape = Cast<ALandscapeProxy>(GetOwner());
-	if (OuterLandscape)
+	if (OuterLandscape &&
+		// when copy/pasting this can get called with a null guid on the parent landscape
+		// this is fine, we won't have any cross-level meshes in this case anyway
+		OuterLandscape->GetLandscapeGuid().IsValid())
 	{
 		ULandscapeInfo* LandscapeInfo = OuterLandscape->GetLandscapeInfo();
-		check(LandscapeInfo);
-
-		TArray<ULandscapeSplinesComponent*> SplinesComponents;
-		SplinesComponents.Reserve(LandscapeInfo->Proxies.Num() + 1);
 
-		ALandscape* RootLandscape = LandscapeInfo->LandscapeActor.Get();
-		if (RootLandscape && RootLandscape->SplineComponent)
-		{
-			SplinesComponents.Add(RootLandscape->SplineComponent);
-		}
-		for (ALandscapeProxy* LandscapeProxy : LandscapeInfo->Proxies)
+		if (LandscapeInfo)
 		{
-			if (LandscapeProxy && LandscapeProxy->SplineComponent)
+			TArray<ULandscapeSplinesComponent*> SplinesComponents;
+			SplinesComponents.Reserve(LandscapeInfo->Proxies.Num() + 1);
+
+			ALandscape* RootLandscape = LandscapeInfo->LandscapeActor.Get();
+			if (RootLandscape && RootLandscape->SplineComponent)
 			{
-				SplinesComponents.Add(LandscapeProxy->SplineComponent);
+				SplinesComponents.Add(RootLandscape->SplineComponent);
+			}
+			for (ALandscapeProxy* LandscapeProxy : LandscapeInfo->Proxies)
+			{
+				if (LandscapeProxy && LandscapeProxy->SplineComponent)
+				{
+					SplinesComponents.Add(LandscapeProxy->SplineComponent);
+				}
 			}
-		}
 
-		return SplinesComponents;
+			return SplinesComponents;
+		}
 	}
 
 	return {};
diff --git a/Engine/Source/Runtime/Landscape/Public/LandscapeDataAccess.h b/Engine/Source/Runtime/Landscape/Public/LandscapeDataAccess.h
--- a/Engine/Source/Runtime/Landscape/Public/LandscapeDataAccess.h
+++ b/Engine/Source/Runtime/Landscape/Public/LandscapeDataAccess.h
@@ -4,8 +4,7 @@
 LandscapeDataAccess.h: Classes for the editor to access to Landscape data
 =============================================================================*/
 
-#ifndef _LANDSCAPEDATAACCESS_H
-#define _LANDSCAPEDATAACCESS_H
+#pragma once
 
 #include "Engine/Texture2D.h"
 
@@ -42,12 +41,6 @@ class ULandscapeLayerInfoObject;
 //
 // FLandscapeDataInterface
 //
-
[REDACTED_EMAIL]: Apparent VC10 compiler bug here causes an access violation in UnlockMip in Shipping builds
-#if _MSC_VER
-PRAGMA_DISABLE_OPTIMIZATION
-#endif
-
 struct FLandscapeDataInterface
 {
 private:
@@ -219,8 +212,6 @@ struct FLandscapeComponentDataInterface
 		XYOffsetMipData = NewXYOffsetData;
 	}
 
-	LANDSCAPE_API void UnlockRawHeightData() const;
-
 	/* Return the raw heightmap data exactly same size for Heightmap texture which belong to only this component */
 	LANDSCAPE_API void GetHeightmapTextureData(TArray<FColor>& OutData, bool bOkToFail = false);
 
@@ -332,7 +323,7 @@ struct FLandscapeComponentDataInterface
 	}
 
 private:
-	struct FLandscapeDataInterface* DataInterface;
+	FLandscapeDataInterface DataInterface;
 	ULandscapeComponent* Component;
 
 public:
@@ -349,8 +340,6 @@ private:
 	int32 ComponentSizeVerts;
 	int32 SubsectionSizeVerts;
 	int32 ComponentNumSubsections;
-
-	bool bNeedToDeleteDataInterface;
 public:
 	const int32 MipLevel;
 };
@@ -424,8 +413,4 @@ void FillCornerValues(uint8& CornerSet, T* CornerValues)
 	}
 }
 
-
-
 #endif // WITH_EDITOR
-
-#endif // _LANDSCAPEDATAACCESS_H
diff --git a/Engine/Source/Runtime/Landscape/Public/LandscapeEdit.h b/Engine/Source/Runtime/Landscape/Public/LandscapeEdit.h
--- a/Engine/Source/Runtime/Landscape/Public/LandscapeEdit.h
+++ b/Engine/Source/Runtime/Landscape/Public/LandscapeEdit.h
@@ -12,7 +12,7 @@ LandscapeEdit.h: Classes for the editor to access to Landscape data
 #if WITH_EDITOR
 
 #include "LandscapeProxy.h" // for ELandscapeLayerPaintingRestriction
-#include "FixedSizeArrayView.h"
+#include "Containers/ArrayView.h"
 
 struct FLandscapeTextureDataInfo
 {
@@ -202,10 +202,10 @@ private:
 	                          int32 ComponentIndexY, int32 SubIndexY, int32 SubY);
 
 	// counts the total influence of each weight-blended layer on this component
-	inline TMap<const ULandscapeLayerInfoObject*, uint32> CountWeightBlendedLayerInfluence(int32 ComponentIndexX, int32 ComponentIndexY, TOptional<TFixedSizeArrayView<const uint8* const>> LayerDataPtrs);
+	inline TMap<const ULandscapeLayerInfoObject*, uint32> CountWeightBlendedLayerInfluence(int32 ComponentIndexX, int32 ComponentIndexY, TOptional<TArrayView<const uint8* const>> LayerDataPtrs);
 
 	// chooses a replacement layer to use when erasing from 100% influence on a texel
-	const ULandscapeLayerInfoObject* ChooseReplacementLayer(const ULandscapeLayerInfoObject* LayerInfo, int32 ComponentIndexX, int32 SubIndexX, int32 SubX, int32 ComponentIndexY, int32 SubIndexY, int32 SubY, TMap<FIntPoint, TMap<const ULandscapeLayerInfoObject*, uint32>>& LayerInfluenceCache, TFixedSizeArrayView<const uint8* const> LayerDataPtrs);
+	const ULandscapeLayerInfoObject* ChooseReplacementLayer(const ULandscapeLayerInfoObject* LayerInfo, int32 ComponentIndexX, int32 SubIndexX, int32 SubX, int32 ComponentIndexY, int32 SubIndexY, int32 SubY, TMap<FIntPoint, TMap<const ULandscapeLayerInfoObject*, uint32>>& LayerInfluenceCache, TArrayView<const uint8* const> LayerDataPtrs);
 };
 
 template<typename T>
diff --git a/Engine/Source/Runtime/Renderer/Private/CustomDepthRendering.cpp b/Engine/Source/Runtime/Renderer/Private/CustomDepthRendering.cpp
--- a/Engine/Source/Runtime/Renderer/Private/CustomDepthRendering.cpp
+++ b/Engine/Source/Runtime/Renderer/Private/CustomDepthRendering.cpp
@@ -72,7 +72,7 @@ bool FCustomDepthPrimSet::DrawPrims(FRHICommandListImmediate& RHICmdList, const
 								View,
 								Context,
 								StaticMesh,
-								StaticMesh.Elements.Num() == 1 ? 1 : View.StaticMeshBatchVisibility[StaticMesh.Id],
+								StaticMesh.bRequiresPerElementVisibility ? View.StaticMeshBatchVisibility[StaticMesh.Id] : ((1ull << StaticMesh.Elements.Num()) - 1),
 								true,
 								DrawRenderState,
 								PrimitiveSceneProxy,
diff --git a/Engine/Source/Runtime/Renderer/Private/DistortionRendering.cpp b/Engine/Source/Runtime/Renderer/Private/DistortionRendering.cpp
--- a/Engine/Source/Runtime/Renderer/Private/DistortionRendering.cpp
+++ b/Engine/Source/Runtime/Renderer/Private/DistortionRendering.cpp
@@ -798,7 +798,7 @@ bool FDistortionPrimSet::DrawAccumulatedOffsets(FRHICommandListImmediate& RHICmd
 							&View,
 							bInitializeOffsets,
 							StaticMesh,
-							StaticMesh.Elements.Num() == 1 ? 1 : View.StaticMeshBatchVisibility[StaticMesh.Id],
+							StaticMesh.bRequiresPerElementVisibility ? View.StaticMeshBatchVisibility[StaticMesh.Id] : ((1ull << StaticMesh.Elements.Num()) - 1),
 							false,
 							PrimitiveSceneProxy,
 							StaticMesh.BatchHitProxyId
diff --git a/Engine/Source/Runtime/Renderer/Private/SceneVisibility.cpp b/Engine/Source/Runtime/Renderer/Private/SceneVisibility.cpp
--- a/Engine/Source/Runtime/Renderer/Private/SceneVisibility.cpp
+++ b/Engine/Source/Runtime/Renderer/Private/SceneVisibility.cpp
@@ -1663,8 +1663,8 @@ struct FRelevancePacket
 						MarkMasks[StaticMesh.Id] = MarkMask;
 					}
 
-					// Static meshes with a single element always draw, as if the mask were 0x1.
-					if(bNeedsBatchVisibility && StaticMesh.Elements.Num() > 1)
+					// Static meshes which don't need per-element visibility always draw all elements
+					if (bNeedsBatchVisibility && StaticMesh.bRequiresPerElementVisibility)
 					{
 						WriteView.StaticMeshBatchVisibility[StaticMesh.Id] = StaticMesh.VertexFactory->GetStaticBatchElementVisibility(View, &StaticMesh);
 					}
diff --git a/Engine/Source/Runtime/Renderer/Private/ShadowDepthRendering.cpp b/Engine/Source/Runtime/Renderer/Private/ShadowDepthRendering.cpp
--- a/Engine/Source/Runtime/Renderer/Private/ShadowDepthRendering.cpp
+++ b/Engine/Source/Runtime/Renderer/Private/ShadowDepthRendering.cpp
@@ -1300,7 +1300,7 @@ void DrawMeshElements(FRHICommandList& RHICmdList, FShadowDepthDrawingPolicy<bRe
 
 	const FMeshDrawingRenderState DrawRenderState(View.GetDitheredLODTransitionState(*Mesh));
 	// Render only those batch elements that match the current LOD
-	uint64 BatchElementMask = Mesh->Elements.Num() == 1 ? 1 : View.StaticMeshBatchVisibility[Mesh->Id];
+	uint64 BatchElementMask = Mesh->bRequiresPerElementVisibility ? View.StaticMeshBatchVisibility[Mesh->Id] : ((1ull << Mesh->Elements.Num()) - 1);
 	int32 BatchElementIndex = 0;
 	do
 	{
diff --git a/Engine/Source/Runtime/Renderer/Private/ShadowRendering.cpp b/Engine/Source/Runtime/Renderer/Private/ShadowRendering.cpp
--- a/Engine/Source/Runtime/Renderer/Private/ShadowRendering.cpp
+++ b/Engine/Source/Runtime/Renderer/Private/ShadowRendering.cpp
@@ -564,7 +564,7 @@ void FProjectedShadowInfo::RenderProjection(FRHICommandListImmediate& RHICmdList
 								*View,
 								FDepthDrawingPolicyFactory::ContextType(DDM_AllOccluders, false),
 								StaticMesh,
-								StaticMesh.Elements.Num() == 1 ? 1 : View->StaticMeshBatchVisibility[StaticMesh.Id],
+								StaticMesh.bRequiresPerElementVisibility ? View->StaticMeshBatchVisibility[StaticMesh.Id] : ((1ull << StaticMesh.Elements.Num() )- 1),
 								true,
 								DrawRenderState,
 								ReceiverPrimitiveSceneInfo->Proxy,
@@ -589,7 +589,7 @@ void FProjectedShadowInfo::RenderProjection(FRHICommandListImmediate& RHICmdList
 					*View,
 					FDepthDrawingPolicyFactory::ContextType(DDM_AllOccluders, false),
 					StaticMesh,
-					StaticMesh.Elements.Num() == 1 ? 1 : View->StaticMeshBatchVisibility[StaticMesh.Id],
+					StaticMesh.bRequiresPerElementVisibility ? View->StaticMeshBatchVisibility[StaticMesh.Id] : ((1ull << StaticMesh.Elements.Num() )- 1),
 					true,
 					DrawRenderState,
 					StaticMesh.PrimitiveSceneInfo->Proxy,
diff --git a/Engine/Source/Runtime/Renderer/Private/ShadowSetup.cpp b/Engine/Source/Runtime/Renderer/Private/ShadowSetup.cpp
--- a/Engine/Source/Runtime/Renderer/Private/ShadowSetup.cpp
+++ b/Engine/Source/Runtime/Renderer/Private/ShadowSetup.cpp
@@ -844,7 +844,9 @@ void FProjectedShadowInfo::AddSubjectPrimitive(FPrimitiveSceneInfo* PrimitiveSce
 								if (bMeshIsVisible && bWholeSceneDirectionalShadow)
 								{
 									StaticMeshWholeSceneShadowDepthMap[StaticMesh.Id] = true;
-									StaticMeshWholeSceneShadowBatchVisibility[StaticMesh.Id] = CurrentView.StaticMeshBatchVisibility[StaticMesh.Id];
+									StaticMeshWholeSceneShadowBatchVisibility[StaticMesh.Id] = StaticMesh.bRequiresPerElementVisibility
+										? StaticMesh.VertexFactory->GetStaticBatchElementShadowVisibility(CurrentView, LightSceneInfo->Proxy, &StaticMesh)
+										: ((1ull << StaticMesh.Elements.Num()) - 1);
 								}
 							}
 						}
@@ -891,7 +893,9 @@ void FProjectedShadowInfo::AddSubjectPrimitive(FPrimitiveSceneInfo* PrimitiveSce
 									if (StaticMesh.CastShadow && LODToRender.ContainsLOD(StaticMesh.LODIndex))
 									{
 										StaticMeshWholeSceneShadowDepthMap[StaticMesh.Id] = true;
-										StaticMeshWholeSceneShadowBatchVisibility[StaticMesh.Id] = StaticMesh.Elements.Num() == 1 ? 1 : StaticMesh.VertexFactory->GetStaticBatchElementVisibility(*DependentView, &StaticMesh);
+										StaticMeshWholeSceneShadowBatchVisibility[StaticMesh.Id] = StaticMesh.bRequiresPerElementVisibility
+											? StaticMesh.VertexFactory->GetStaticBatchElementShadowVisibility(CurrentView, LightSceneInfo->Proxy, &StaticMesh)
+											: ((1ull << StaticMesh.Elements.Num()) - 1);
 
 										bDrawingStaticMeshes = true;
 									}
@@ -905,7 +909,9 @@ void FProjectedShadowInfo::AddSubjectPrimitive(FPrimitiveSceneInfo* PrimitiveSce
 									if (StaticMesh.CastShadow && LODToRender.ContainsLOD(StaticMesh.LODIndex))
 									{
 										CurrentView.StaticMeshShadowDepthMap[StaticMesh.Id] = true;
-										CurrentView.StaticMeshBatchVisibility[StaticMesh.Id] = StaticMesh.Elements.Num() == 1 ? 1 : StaticMesh.VertexFactory->GetStaticBatchElementVisibility(CurrentView, &StaticMesh);
+										CurrentView.StaticMeshBatchVisibility[StaticMesh.Id] = StaticMesh.bRequiresPerElementVisibility
+											? StaticMesh.VertexFactory->GetStaticBatchElementShadowVisibility(CurrentView, LightSceneInfo->Proxy, &StaticMesh)
+											: ((1ull << StaticMesh.Elements.Num()) - 1);
 
 										bDrawingStaticMeshes = true;
 									}
diff --git a/Engine/Source/Runtime/Renderer/Private/StaticMeshDrawList.inl b/Engine/Source/Runtime/Renderer/Private/StaticMeshDrawList.inl
--- a/Engine/Source/Runtime/Renderer/Private/StaticMeshDrawList.inl
+++ b/Engine/Source/Runtime/Renderer/Private/StaticMeshDrawList.inl
@@ -293,7 +293,7 @@ bool TStaticMeshDrawList<DrawingPolicyType>::DrawVisible(
 				const FElement& Element = DrawingPolicyLink->Elements[ElementIndex];
 				STAT(StatInc +=  Element.Mesh->GetNumPrimitives();)
 				// Avoid the virtual call looking up batch visibility if there is only one element.
-				uint32 BatchElementMask = Element.Mesh->Elements.Num() == 1 ? 1 : Element.Mesh->VertexFactory->GetStaticBatchElementVisibility(View,Element.Mesh);
+				uint64 BatchElementMask = Element.Mesh->bRequiresPerElementVisibility ? Element.Mesh->VertexFactory->GetStaticBatchElementVisibility(View, Element.Mesh) : ((1ull << Element.Mesh->Elements.Num()) - 1);
 				DrawElement<InstancedStereoPolicy::Disabled>(View, PolicyContext, Element, BatchElementMask, DrawingPolicyLink, bDrawnShared);
 				bDirty = true;
 			}
@@ -342,7 +342,7 @@ bool TStaticMeshDrawList<DrawingPolicyType>::DrawVisibleInner(
 					STAT(StatInc += Element.Mesh->GetNumPrimitives();)
 					int32 SubCount = Element.Mesh->Elements.Num();
 					// Avoid the cache miss looking up batch visibility if there is only one element.
-					uint64 BatchElementMask = SubCount == 1 ? 1 : (*BatchVisibilityArray)[Element.Mesh->Id];
+					uint64 BatchElementMask = Element.Mesh->bRequiresPerElementVisibility ? (*BatchVisibilityArray)[Element.Mesh->Id] : ((1ull << SubCount) - 1);
 					Count += DrawElement<InstancedStereoPolicy::Disabled>(RHICmdList, View, PolicyContext, Element, BatchElementMask, DrawingPolicyLink, bDrawnShared);
 				}
 			}
@@ -367,7 +367,7 @@ bool TStaticMeshDrawList<DrawingPolicyType>::DrawVisibleInner(
 					    STAT(StatInc += Element.Mesh->GetNumPrimitives();)
 				    int32 SubCount = Element.Mesh->Elements.Num();
 				    // Avoid the cache miss looking up batch visibility if there is only one element.
-					uint64 BatchElementMask = SubCount == 1 ? 1 : (*ResolvedVisiblityArray)[Element.Mesh->Id];
+					uint64 BatchElementMask = Element.Mesh->bRequiresPerElementVisibility ? (*ResolvedVisiblityArray)[Element.Mesh->Id] : ((1ull << SubCount) - 1);
 					Count += DrawElement<InstancedStereoPolicy::Enabled>(RHICmdList, View, PolicyContext, Element, BatchElementMask, DrawingPolicyLink, bDrawnShared);
 				}
 			}
@@ -822,7 +822,7 @@ int32 TStaticMeshDrawList<DrawingPolicyType>::DrawVisibleFrontToBack(
 		const FElement& Element = DrawingPolicyLink->Elements[ElementIndex];
 		STAT(StatInc +=  Element.Mesh->GetNumPrimitives();)
 		// Avoid the cache miss looking up batch visibility if there is only one element.
-		uint64 BatchElementMask = Element.Mesh->Elements.Num() == 1 ? 1 : BatchVisibilityArray[Element.Mesh->Id];
+		uint64 BatchElementMask = Element.Mesh->bRequiresPerElementVisibility ? BatchVisibilityArray[Element.Mesh->Id] : ((1ull << Element.Mesh->Elements.Num()) - 1);
 		DrawElement<InstancedStereoPolicy::Disabled>(RHICmdList, View, PolicyContext, Element, BatchElementMask, DrawingPolicyLink, bDrawnShared);
 		NumDraws++;
 	}
diff --git a/Engine/Source/Runtime/ShaderCore/Public/VertexFactory.h b/Engine/Source/Runtime/ShaderCore/Public/VertexFactory.h
--- a/Engine/Source/Runtime/ShaderCore/Public/VertexFactory.h
+++ b/Engine/Source/Runtime/ShaderCore/Public/VertexFactory.h
@@ -407,8 +407,15 @@ public:
 
 	/**
 	 * Get a bitmask representing the visibility of each FMeshBatch element.
+	 * FMeshBatch.bRequiresPerElementVisibility must be set for this to be called.
 	 */
-	virtual uint64 GetStaticBatchElementVisibility( const class FSceneView& View, const struct FMeshBatch* Batch ) const { return 1; }
+	virtual uint64 GetStaticBatchElementVisibility(const class FSceneView& View, const struct FMeshBatch* Batch) const { return 1; }
+
+	/**
+	  * Get a bitmask representing the visibility of each FMeshBatch element for the shadow depths of a given light.
+	  * FMeshBatch.bRequiresPerElementVisibility must be set for this to be called.
+	  */
+	virtual uint64 GetStaticBatchElementShadowVisibility(const class FSceneView& View, const class FLightSceneProxy* LightSceneProxy, const struct FMeshBatch* Batch) const { return GetStaticBatchElementVisibility(View, Batch); }
 
 protected:
 
