Fixed non-unity build

diff --git a/Engine/Source/Runtime/Landscape/Private/LandscapeSplineRaster.cpp b/Engine/Source/Runtime/Landscape/Private/LandscapeSplineRaster.cpp
--- a/Engine/Source/Runtime/Landscape/Private/LandscapeSplineRaster.cpp
+++ b/Engine/Source/Runtime/Landscape/Private/LandscapeSplineRaster.cpp
@@ -7,6 +7,7 @@
 
 #include "LandscapePrivatePCH.h"
 #include "LandscapeSplineRaster.h"
+#include "Landscape.h"
 #include "LandscapeStreamingProxy.h"
 #include "LandscapeDataAccess.h"
 #include "LandscapeEdit.h"
diff --git a/Engine/Source/Runtime/Landscape/Private/LandscapeSplines.cpp b/Engine/Source/Runtime/Landscape/Private/LandscapeSplines.cpp
--- a/Engine/Source/Runtime/Landscape/Private/LandscapeSplines.cpp
+++ b/Engine/Source/Runtime/Landscape/Private/LandscapeSplines.cpp
@@ -5,6 +5,8 @@
 =============================================================================*/
 
 #include "LandscapePrivatePCH.h"
+#include "Landscape.h"
+#include "LandscapeStreamingProxy.h"
 #include "Components/SplineMeshComponent.h"
 #include "ShaderParameterUtils.h"
 #include "StaticMeshResources.h"
