Copying //UE4/Dev-Editor to //UE4/Dev-Main (Source: //UE4/Dev-Editor @ 3152045)

diff --git a/Engine/Source/Runtime/Engine/Private/UnrealEngine.cpp b/Engine/Source/Runtime/Engine/Private/UnrealEngine.cpp
--- a/Engine/Source/Runtime/Engine/Private/UnrealEngine.cpp
+++ b/Engine/Source/Runtime/Engine/Private/UnrealEngine.cpp
@@ -2355,11 +2355,11 @@ struct FSortedParticleSet
 	int32		ModuleSize;
 	int32		ComponentSize;
 	int32		ComponentCount;
-	int32		ComponentResourceSize;
-	int32		ComponentTrueResourceSize;
+	FResourceSizeEx ComponentResourceSize;
+	FResourceSizeEx ComponentTrueResourceSize;
 
 	FSortedParticleSet( const FString& InName, int32 InSize, int32 InPSysSize, int32 InModuleSize, 
-		int32 InComponentSize, int32 InComponentCount, int32 InComponentResourceSize, int32 InComponentTrueResourceSize) :
+		int32 InComponentSize, int32 InComponentCount, FResourceSizeEx InComponentResourceSize, FResourceSizeEx InComponentTrueResourceSize) :
 		  Name(InName)
 		, Size(InSize)
 		, PSysSize(InPSysSize)
@@ -2377,8 +2377,8 @@ struct FSortedParticleSet
 		, ModuleSize(0)
 		, ComponentSize(0)
 		, ComponentCount(0)
-		, ComponentResourceSize(0)
-		, ComponentTrueResourceSize(0)
+		, ComponentResourceSize(EResourceSizeMode::Inclusive)
+		, ComponentTrueResourceSize(EResourceSizeMode::Exclusive)
 	{
 	}
 
@@ -2398,7 +2398,7 @@ struct FSortedParticleSet
 	{
 		InArchive.Logf(TEXT("%10d,%s,%d,%d,%d,%d,%d,%d"),
 			Size, *Name, PSysSize, ModuleSize, ComponentSize,
-			ComponentCount, ComponentResourceSize, ComponentTrueResourceSize);
+			ComponentCount, ComponentResourceSize.GetTotalMemoryBytes(), ComponentTrueResourceSize.GetTotalMemoryBytes());
 	}
 };
 
@@ -2449,13 +2449,13 @@ struct FItem
 	SIZE_T Num;
 	SIZE_T Max;
 	/** Only exclusive resource size, the truer resource size. */
-	SIZE_T TrueResourceSize;
+	FResourceSizeEx TrueResourceSize;
 
 	FItem( UClass* InClass=NULL )
-	: Class(InClass), Count(0), Num(0), Max(0), TrueResourceSize(0)
+	: Class(InClass), Count(0), Num(0), Max(0), TrueResourceSize()
 	{}
 
-	FItem( UClass* InClass, int32 InCount, SIZE_T InNum, SIZE_T InMax, SIZE_T InTrueResourceSize ) :
+	FItem( UClass* InClass, int32 InCount, SIZE_T InNum, SIZE_T InMax, FResourceSizeEx InTrueResourceSize ) :
 		Class( InClass ),
 		Count( InCount ),
 		Num( InNum ), 
@@ -2463,7 +2463,7 @@ struct FItem
 		TrueResourceSize( InTrueResourceSize )
 	{}
 
-	void Add( FArchiveCountMem& Ar, SIZE_T InTrueResourceSize )
+	void Add( FArchiveCountMem& Ar, FResourceSizeEx InTrueResourceSize )
 	{
 		Count++;
 		Num += Ar.GetNum();
@@ -2482,9 +2482,9 @@ struct FSubItem
 	/** Resource size of the object and all of its references, the 'old-style'. */
 	SIZE_T ResourceSize;
 	/** Only exclusive resource size, the truer resource size. */
-	SIZE_T TrueResourceSize;
+	FResourceSizeEx TrueResourceSize;
 
-	FSubItem( UObject* InObject, SIZE_T InNum, SIZE_T InMax, SIZE_T InTrueResourceSize )
+	FSubItem( UObject* InObject, SIZE_T InNum, SIZE_T InMax, FResourceSizeEx InTrueResourceSize )
 	: Object( InObject ), Num( InNum ), Max( InMax ), TrueResourceSize( InTrueResourceSize )
 	{}
 };
@@ -3951,8 +3951,7 @@ bool UEngine::HandleListParticleSystemsCommand( const TCHAR* Cmd, FOutputDevice&
 		FArchiveCountMem Count( Tree );
 		int32 RootSize = Count.GetMax();
 
-		FSortedParticleSet *pSet = new(SortedSets) FSortedParticleSet(Description, RootSize, RootSize, 0, 0, 0, 0, 0);
-
+		SortedSets.Add(FSortedParticleSet(Description, RootSize, RootSize, 0, 0, 0, FResourceSizeEx(), FResourceSizeEx()));
 		SortMap.Add(Tree,SortedSets.Num() - 1);
 	}
 
@@ -3982,12 +3981,14 @@ bool UEngine::HandleListParticleSystemsCommand( const TCHAR* Cmd, FOutputDevice&
 			Set.ComponentSize += ComponentCount.GetMax();
 
 			// Save this for adding to the total
-			SIZE_T CompResSize = Comp->GetResourceSize(EResourceSizeMode::Inclusive);
+			FResourceSizeEx CompResSize = FResourceSizeEx(EResourceSizeMode::Inclusive);
+			Comp->GetResourceSizeEx(CompResSize);
+
 			Set.ComponentResourceSize += CompResSize;
-			Set.ComponentTrueResourceSize += Comp->GetResourceSize(EResourceSizeMode::Exclusive);
+			Comp->GetResourceSizeEx(Set.ComponentTrueResourceSize);
 
 			Set.Size += ComponentCount.GetMax();
-			Set.Size += CompResSize;
+			Set.Size += CompResSize.GetTotalMemoryBytes();
 			Set.ComponentCount++;
 
 			UParticleSystem* Tree = Comp->Template;
@@ -4239,7 +4240,9 @@ bool UEngine::HandleParticleMeshUsageCommand( const TCHAR* Cmd, FOutputDevice& A
 	{
 		FORCEINLINE bool operator()( UStaticMesh& A, UStaticMesh& B ) const
 		{
-			return B.GetResourceSize(EResourceSizeMode::Inclusive) < A.GetResourceSize(EResourceSizeMode::Inclusive);
+			const SIZE_T ResourceSizeA = A.GetResourceSizeBytes(EResourceSizeMode::Inclusive);
+			const SIZE_T ResourceSizeB = B.GetResourceSizeBytes(EResourceSizeMode::Inclusive);
+			return ResourceSizeB < ResourceSizeA;
 		}
 	};
 
@@ -4251,7 +4254,8 @@ bool UEngine::HandleParticleMeshUsageCommand( const TCHAR* Cmd, FOutputDevice& A
 	for( int32 StaticMeshIndex=0; StaticMeshIndex<UniqueReferencedMeshes.Num(); StaticMeshIndex++ )
 	{
 		UStaticMesh* StaticMesh	= UniqueReferencedMeshes[StaticMeshIndex];
-		TotalSize += StaticMesh->GetResourceSize(EResourceSizeMode::Inclusive);
+		const SIZE_T StaticMeshResourceSize = StaticMesh->GetResourceSizeBytes(EResourceSizeMode::Inclusive);
+		TotalSize += StaticMeshResourceSize;
 	}
 
 	// Log sorted summary.
@@ -4264,8 +4268,10 @@ bool UEngine::HandleParticleMeshUsageCommand( const TCHAR* Cmd, FOutputDevice& A
 		TArray<UParticleSystem*> ParticleSystems;
 		StaticMeshToParticleSystemMap.MultiFind( StaticMesh, ParticleSystems );
 
+		const SIZE_T StaticMeshResourceSize = StaticMesh->GetResourceSizeBytes(EResourceSizeMode::Inclusive);
+
 		// Log meshes including resource size and referencing particle systems.
-		Ar.Logf(TEXT("%5i KByte  %s"), StaticMesh->GetResourceSize(EResourceSizeMode::Inclusive) / 1024, *StaticMesh->GetFullName());
+		Ar.Logf(TEXT("%5i KByte  %s"), StaticMeshResourceSize / 1024, *StaticMesh->GetFullName());
 		for( int32 ParticleSystemIndex=0; ParticleSystemIndex<ParticleSystems.Num(); ParticleSystemIndex++ )
 		{
 			UParticleSystem* ParticleSystem = ParticleSystems[ParticleSystemIndex];
@@ -4880,7 +4886,7 @@ struct FHierarchy
 			{
 				FSubItem const& Item = Objects.FindChecked(This);
 				Node.Exc += Item.Max;
-				Node.Exc += Item.TrueResourceSize;
+				Node.Exc += Item.TrueResourceSize.GetTotalMemoryBytes();
 				if (bCountItems)
 				{
 					Node.ExcCount += Node.Items.Num();
@@ -5094,9 +5100,9 @@ bool UEngine::HandleObjCommand( const TCHAR* Cmd, FOutputDevice& Ar )
 
 		TMap<UClass*,FItem> ObjectsByClass;
 
-		UE_LOG( LogEngine, Log, TEXT("**********************************************") );
-		UE_LOG( LogEngine, Log, TEXT("Obj MemSub for class '%s'"), *ClassToCheck->GetName() );
-		UE_LOG( LogEngine, Log, TEXT("") );
+		Ar.Logf( TEXT("**********************************************") );
+		Ar.Logf( TEXT("Obj MemSub for class '%s'"), *ClassToCheck->GetName() );
+		Ar.Logf( TEXT("") );
 
 		for( FObjectIterator It(ClassToCheck); It; ++It )
 		{
@@ -5115,7 +5121,8 @@ bool UEngine::HandleObjCommand( const TCHAR* Cmd, FOutputDevice& Ar )
 			for( UObject*& RefObj : ReferencedObjects )
 			{
 				FArchiveCountMem Count( RefObj );
-				const SIZE_T TrueResourceSize = It->GetResourceSize( EResourceSizeMode::Exclusive );
+				FResourceSizeEx TrueResourceSize = FResourceSizeEx(EResourceSizeMode::Exclusive);
+				It->GetResourceSizeEx( TrueResourceSize );
 				ThisObject.Add( Count, TrueResourceSize );
 			}
 
@@ -5128,12 +5135,19 @@ bool UEngine::HandleObjCommand( const TCHAR* Cmd, FOutputDevice& Ar )
 
 		ObjectsByClass.ValueSort( FCompareByInclusiveSize() );
 
-		UE_LOG( LogEngine, Log, TEXT("%32s, %6s %6s %6s, %6s"), 
+		Ar.Logf( 
+			TEXT("%32s %12s %12s %12s %12s %12s %12s %12s %12s %12s"), 
 			TEXT("Class"),
 			TEXT("IncMax"),
 			TEXT("IncNum"),
 			TEXT("ResExc"),
-			TEXT("Count") );
+			TEXT("ResExcDedSys"),
+			TEXT("ResExcShrSys"),
+			TEXT("ResExcDedVid"),
+			TEXT("ResExcShrVid"),
+			TEXT("ResExcUnk"),
+			TEXT("Count") 
+			);
 
 		FItem Total;
 		FItem Culled;
@@ -5151,10 +5165,16 @@ bool UEngine::HandleObjCommand( const TCHAR* Cmd, FOutputDevice& Ar )
 			}
 			else
 			{
-				UE_LOG( LogEngine, Log, TEXT("%32s, %6s %6s %6s, %6i"), 
+				Ar.Logf( TEXT("%32s %12s %12s %12s %12s %12s %12s %12s %12s %12i"), 
 					*Class->GetName(), 
-					*FHierarchy::Size(ClassObjects.Max), *FHierarchy::Size(ClassObjects.Num), 
-					*FHierarchy::Size(ClassObjects.TrueResourceSize),
+					*FHierarchy::Size(ClassObjects.Max), 
+					*FHierarchy::Size(ClassObjects.Num), 
+					*FHierarchy::Size(ClassObjects.TrueResourceSize.GetTotalMemoryBytes()),
+					*FHierarchy::Size(ClassObjects.TrueResourceSize.GetDedicatedSystemMemoryBytes()),
+					*FHierarchy::Size(ClassObjects.TrueResourceSize.GetSharedSystemMemoryBytes()),
+					*FHierarchy::Size(ClassObjects.TrueResourceSize.GetDedicatedVideoMemoryBytes()),
+					*FHierarchy::Size(ClassObjects.TrueResourceSize.GetSharedVideoMemoryBytes()),
+					*FHierarchy::Size(ClassObjects.TrueResourceSize.GetUnknownMemoryBytes()),
 					ClassObjects.Count 
 					);
 
@@ -5168,23 +5188,35 @@ bool UEngine::HandleObjCommand( const TCHAR* Cmd, FOutputDevice& Ar )
 
 		if( Culled.Count > 0 )
 		{
-			UE_LOG( LogEngine, Log, TEXT("") );
-			UE_LOG( LogEngine, Log, TEXT("%32s, %6s %6s %6s, %6i"), 
+			Ar.Logf( TEXT("") );
+			Ar.Logf( TEXT("%32s %12s %12s %12s %12s %12s %12s %12s %12s %12i"), 
 				TEXT("(Culled)"), 
-				*FHierarchy::Size(Culled.Max), *FHierarchy::Size(Culled.Num), 
-				*FHierarchy::Size(Culled.TrueResourceSize),
+				*FHierarchy::Size(Culled.Max), 
+				*FHierarchy::Size(Culled.Num), 
+				*FHierarchy::Size(Culled.TrueResourceSize.GetTotalMemoryBytes()),
+				*FHierarchy::Size(Culled.TrueResourceSize.GetDedicatedSystemMemoryBytes()),
+				*FHierarchy::Size(Culled.TrueResourceSize.GetSharedSystemMemoryBytes()),
+				*FHierarchy::Size(Culled.TrueResourceSize.GetDedicatedVideoMemoryBytes()),
+				*FHierarchy::Size(Culled.TrueResourceSize.GetSharedVideoMemoryBytes()),
+				*FHierarchy::Size(Culled.TrueResourceSize.GetUnknownMemoryBytes()),
 				Culled.Count 
 				);
 		}
 
-		UE_LOG( LogEngine, Log, TEXT("") );
-		UE_LOG( LogEngine, Log, TEXT("%32s, %6s %6s %6s, %6i"), 
+		Ar.Logf( TEXT("") );
+		Ar.Logf( TEXT("%32s %12s %12s %12s %12s %12s %12s %12s %12s %12i"), 
 			TEXT("Total"), 
-			*FHierarchy::Size(Total.Max), *FHierarchy::Size(Total.Num), 
-			*FHierarchy::Size(Total.TrueResourceSize),
+			*FHierarchy::Size(Total.Max), 
+			*FHierarchy::Size(Total.Num), 
+			*FHierarchy::Size(Total.TrueResourceSize.GetTotalMemoryBytes()),
+			*FHierarchy::Size(Total.TrueResourceSize.GetDedicatedSystemMemoryBytes()),
+			*FHierarchy::Size(Total.TrueResourceSize.GetSharedSystemMemoryBytes()),
+			*FHierarchy::Size(Total.TrueResourceSize.GetDedicatedVideoMemoryBytes()),
+			*FHierarchy::Size(Total.TrueResourceSize.GetSharedVideoMemoryBytes()),
+			*FHierarchy::Size(Total.TrueResourceSize.GetUnknownMemoryBytes()),
 			Total.Count 
 			);
-		UE_LOG( LogEngine, Log, TEXT("**********************************************") );
+		Ar.Logf( TEXT("**********************************************") );
 		return true;
 	}
 	else if (FParse::Command(&Cmd,TEXT("Mem")))
@@ -5214,7 +5246,8 @@ bool UEngine::HandleObjCommand( const TCHAR* Cmd, FOutputDevice& Ar )
 		for( FObjectIterator It; It; ++It )
 		{
 			FArchiveCountMem Count( *It );
-			const SIZE_T TrueResourceSize = It->GetResourceSize(EResourceSizeMode::Exclusive);
+			FResourceSizeEx TrueResourceSize = FResourceSizeEx(EResourceSizeMode::Exclusive);
+			It->GetResourceSizeEx(TrueResourceSize);
 			Objects.Add(*It, FSubItem(*It, Count.GetNum(), Count.GetMax(), TrueResourceSize));
 			Classes.AddClassInstance(*It);
 			Outers.AddOuter(*It);
@@ -5366,8 +5399,8 @@ bool UEngine::HandleObjCommand( const TCHAR* Cmd, FOutputDevice& Ar )
 				}
 
 				FArchiveCountMem Count( *It );
-
-				const SIZE_T TrueResourceSize = It->GetResourceSize(EResourceSizeMode::Exclusive);
+				FResourceSizeEx TrueResourceSize = FResourceSizeEx(EResourceSizeMode::Exclusive);
+				It->GetResourceSizeEx(TrueResourceSize);
 
 				int32 i;
 
@@ -5404,7 +5437,7 @@ bool UEngine::HandleObjCommand( const TCHAR* Cmd, FOutputDevice& Ar )
 
 				if( bShowDetailedObjectInfo )
 				{
-					new(Objects)FSubItem( *It, Count.GetNum(), Count.GetMax(), TrueResourceSize );
+					Objects.Add( FSubItem( *It, Count.GetNum(), Count.GetMax(), TrueResourceSize ) );
 				}
 				List[i].Add( Count, TrueResourceSize );
 				Total.Add( Count, TrueResourceSize );
@@ -5430,11 +5463,33 @@ bool UEngine::HandleObjCommand( const TCHAR* Cmd, FOutputDevice& Ar )
 			};
 			Objects.Sort( FCompareFSubItem( bAlphaSort ) );
 
-			Ar.Logf( TEXT("%140s, % 10s, % 10s, % 10s"), TEXT("Object"), TEXT("NumKBytes"), TEXT("MaxKBytes"), TEXT("ExclusiveResKBytes") );
+			Ar.Logf( 
+				TEXT("%140s %10s %10s %10s %15s %15s %15s %15s %15s"), 
+				TEXT("Object"), 
+				TEXT("NumKB"), 
+				TEXT("MaxKB"), 
+				TEXT("ResExcKB"),
+				TEXT("ResExcDedSysKB"),
+				TEXT("ResExcShrSysKB"),
+				TEXT("ResExcDedVidKB"),
+				TEXT("ResExcShrVidKB"),
+				TEXT("ResExcUnkKB")
+				);
 
 			for (const FSubItem& ObjItem : Objects)
 			{
-				Ar.Logf(TEXT("%140s, % 10.2f, % 10.2f, % 10.2f"), *ObjItem.Object->GetFullName(), ObjItem.Num / 1024.0f, ObjItem.Max / 1024.0f, ObjItem.TrueResourceSize / 1024.0f);
+				Ar.Logf(
+					TEXT("%140s %10.2f %10.2f %10.2f %15.2f %15.2f %15.2f %15.2f %15.2f"), 
+					*ObjItem.Object->GetFullName(), 
+					ObjItem.Num / 1024.0f, 
+					ObjItem.Max / 1024.0f, 
+					ObjItem.TrueResourceSize.GetTotalMemoryBytes() / 1024.0f, 
+					ObjItem.TrueResourceSize.GetDedicatedSystemMemoryBytes() / 1024.0f, 
+					ObjItem.TrueResourceSize.GetSharedSystemMemoryBytes() / 1024.0f, 
+					ObjItem.TrueResourceSize.GetDedicatedVideoMemoryBytes() / 1024.0f, 
+					ObjItem.TrueResourceSize.GetSharedVideoMemoryBytes() / 1024.0f, 
+					ObjItem.TrueResourceSize.GetUnknownMemoryBytes() / 1024.0f
+					);
 			}
 			Ar.Log(TEXT(""));
 		}
@@ -5454,15 +5509,50 @@ bool UEngine::HandleObjCommand( const TCHAR* Cmd, FOutputDevice& Ar )
 				}
 			};
 			List.Sort( FCompareFItem( bAlphaSort, bCountSort ) );
-			Ar.Logf(TEXT(" %100s, % 6s, % 10s, % 10s, % 10s"), TEXT("Class"), TEXT("Count"), TEXT("NumKBytes"), TEXT("MaxKBytes"), TEXT("ExclusiveResKBytes") );
+			Ar.Logf(
+				TEXT(" %100s %8s %10s %10s %10s %15s %15s %15s %15s %15s"), 
+				TEXT("Class"), 
+				TEXT("Count"), 
+				TEXT("NumKB"), 
+				TEXT("MaxKB"), 
+				TEXT("ResExcKB"),
+				TEXT("ResExcDedSysKB"),
+				TEXT("ResExcShrSysKB"),
+				TEXT("ResExcDedVidKB"),
+				TEXT("ResExcShrVidKB"),
+				TEXT("ResExcUnkKB")
+				);
 
 			for( int32 i=0; i<List.Num(); i++ )
 			{
-				Ar.Logf(TEXT(" %100s, % 6i, % 10.2f, % 10.2f, % 10.2f"), *List[i].Class->GetName(), (int32)List[i].Count, List[i].Num / 1024.0f, List[i].Max / 1024.0f, List[i].TrueResourceSize / 1024.0f);
+				Ar.Logf(
+					TEXT(" %100s %8i %10.2f %10.2f %10.2f %15.2f %15.2f %15.2f %15.2f %15.2f"), 
+					*List[i].Class->GetName(), 
+					(int32)List[i].Count, 
+					List[i].Num / 1024.0f, 
+					List[i].Max / 1024.0f, 
+					List[i].TrueResourceSize.GetTotalMemoryBytes() / 1024.0f, 
+					List[i].TrueResourceSize.GetDedicatedSystemMemoryBytes() / 1024.0f, 
+					List[i].TrueResourceSize.GetSharedSystemMemoryBytes() / 1024.0f, 
+					List[i].TrueResourceSize.GetDedicatedVideoMemoryBytes() / 1024.0f, 
+					List[i].TrueResourceSize.GetSharedVideoMemoryBytes() / 1024.0f, 
+					List[i].TrueResourceSize.GetUnknownMemoryBytes() / 1024.0f
+					);
 			}
 			Ar.Log( TEXT("") );
 		}
-		Ar.Logf( TEXT("%i Objects (%.3fM / %.3fM / %.3fM)"), Total.Count, (float)Total.Num/1024.0/1024.0, (float)Total.Max/1024.0/1024.0, (float)Total.TrueResourceSize/1024.0/1024.0 );
+		Ar.Logf( 
+			TEXT("%i Objects (Total: %.3fM / Max: %.3fM / Res: %.3fM | ResDedSys: %.3fM / ResShrSys: %.3fM / ResDedVid: %.3fM / ResShrVid: %.3fM / ResUnknown: %.3fM)"),
+			Total.Count, 
+			(float)Total.Num/1024.0/1024.0, 
+			(float)Total.Max/1024.0/1024.0, 
+			(float)Total.TrueResourceSize.GetTotalMemoryBytes()/1024.0/1024.0,
+			(float)Total.TrueResourceSize.GetDedicatedSystemMemoryBytes()/1024.0/1024.0, 
+			(float)Total.TrueResourceSize.GetSharedSystemMemoryBytes()/1024.0/1024.0, 
+			(float)Total.TrueResourceSize.GetDedicatedVideoMemoryBytes()/1024.0/1024.0, 
+			(float)Total.TrueResourceSize.GetSharedVideoMemoryBytes()/1024.0/1024.0, 
+			(float)Total.TrueResourceSize.GetUnknownMemoryBytes()/1024.0/1024.0
+			);
 		return true;
 
 	}
diff --git a/Engine/Source/Runtime/Engine/Private/UserInterface/Canvas.cpp b/Engine/Source/Runtime/Engine/Private/UserInterface/Canvas.cpp
--- a/Engine/Source/Runtime/Engine/Private/UserInterface/Canvas.cpp
+++ b/Engine/Source/Runtime/Engine/Private/UserInterface/Canvas.cpp
@@ -387,9 +387,7 @@ bool FCanvasBatchedElementRenderItem::Render_RenderThread(FRHICommandListImmedia
 			RHICmdList,
 			Canvas->GetFeatureLevel(),
 			bNeedsToSwitchVerticalAxis,
-			Data->Transform.GetMatrix(),
-			CanvasRenderTarget->GetSizeXY().X,
-			CanvasRenderTarget->GetSizeXY().Y,
+			FBatchedElements::CreateProxySceneView(Data->Transform.GetMatrix(), FIntRect(0, 0, CanvasRenderTarget->GetSizeXY().X, CanvasRenderTarget->GetSizeXY().Y)),
 			Canvas->IsHitTesting(),
 			Gamma
 			);
@@ -460,9 +458,7 @@ bool FCanvasBatchedElementRenderItem::Render_GameThread(const FCanvas* Canvas)
 				RHICmdList,
 				Parameters.FeatureLevel,
 				Parameters.bNeedsToSwitchVerticalAxis,
-				Parameters.RenderData->Transform.GetMatrix(),
-				Parameters.ViewportSizeX,
-				Parameters.ViewportSizeY,
+				FBatchedElements::CreateProxySceneView(Parameters.RenderData->Transform.GetMatrix(),FIntRect(0, 0, Parameters.ViewportSizeX, Parameters.ViewportSizeY)),
 				Parameters.bHitTesting,
 				Parameters.DisplayGamma
 				);
diff --git a/Engine/Source/Runtime/Engine/Private/UserInterfaceSettings.cpp b/Engine/Source/Runtime/Engine/Private/UserInterfaceSettings.cpp
--- a/Engine/Source/Runtime/Engine/Private/UserInterfaceSettings.cpp
+++ b/Engine/Source/Runtime/Engine/Private/UserInterfaceSettings.cpp
@@ -106,17 +106,4 @@ void UUserInterfaceSettings::ForceLoadResources()
 	CustomScalingRuleClassInstance = CustomScalingRuleClass.TryLoadClass<UDPICustomScalingRule>();
 }
 
-#if WITH_EDITOR
-void UUserInterfaceSettings::PostEditChangeProperty(FPropertyChangedEvent& PropertyChangedEvent)
-{
-	Super::PostEditChangeProperty(PropertyChangedEvent);
-
-	if ( PropertyChangedEvent.Property )
-	{
-		URendererSettings* RendererSettings = GetMutableDefault<URendererSettings>(URendererSettings::StaticClass());
-		RendererSettings->UpdateDefaultConfigFile();
-	}
-}
-#endif // #if WITH_EDITOR
-
 #undef LOCTEXT_NAMESPACE
diff --git a/Engine/Source/Runtime/Engine/Public/BatchedElements.h b/Engine/Source/Runtime/Engine/Public/BatchedElements.h
--- a/Engine/Source/Runtime/Engine/Public/BatchedElements.h
+++ b/Engine/Source/Runtime/Engine/Public/BatchedElements.h
@@ -195,8 +195,24 @@ public:
 	 * @param View			Optional FSceneView for shaders that need access to view constants
 	 * @param DepthTexture	DepthTexture for manual depth testing with editor compositing in the pixel shader
 	 */
+	DEPRECATED(4.14, "Deprecated. Use the FBatchedElements::Draw method that takes a non-optional FSceneView parameter instead")
 	bool Draw(FRHICommandList& RHICmdList, ERHIFeatureLevel::Type FeatureLevel, bool bNeedToSwitchVerticalAxis, const FMatrix& Transform, uint32 ViewportSizeX, uint32 ViewportSizeY, bool bHitTesting, float Gamma = 1.0f, const FSceneView* View = NULL, FTexture2DRHIRef DepthTexture = FTexture2DRHIRef(), EBlendModeFilter::Type Filter = EBlendModeFilter::All) const;
 	
+	/**
+	 * Draws the batch
+	 *
+	 * @param View			FSceneView for shaders that need access to view constants. Non-optional to also reference its ViewProjectionMatrix and size of the ViewRect
+	 * @param bHitTesting	Whether or not we are hit testing
+	 * @param Gamma			Optional gamma override
+	 * @param DepthTexture	DepthTexture for manual depth testing with editor compositing in the pixel shader
+	 */
+	bool Draw(FRHICommandList& RHICmdList, ERHIFeatureLevel::Type FeatureLevel, bool bNeedToSwitchVerticalAxis, const FSceneView& View, bool bHitTesting, float Gamma = 1.0f, FTexture2DRHIRef DepthTexture = FTexture2DRHIRef(), EBlendModeFilter::Type Filter = EBlendModeFilter::All) const;
+
+	/**
+	 * Creates a proxy FSceneView for operations that are not tied directly to a scene but still require batched elements to be drawn.
+	 */
+	static FSceneView CreateProxySceneView(const FMatrix& ProjectionMatrix, const FIntRect& ViewRect);
+
 	FORCEINLINE bool HasPrimsToDraw() const
 	{
 		return( LineVertices.Num() || Points.Num() || Sprites.Num() || MeshElements.Num() || ThickLines.Num() || WireTris.Num() > 0 );
diff --git a/Engine/Source/Runtime/Engine/Public/DistanceFieldAtlas.h b/Engine/Source/Runtime/Engine/Public/DistanceFieldAtlas.h
--- a/Engine/Source/Runtime/Engine/Public/DistanceFieldAtlas.h
+++ b/Engine/Source/Runtime/Engine/Public/DistanceFieldAtlas.h
@@ -143,9 +143,23 @@ public:
 		delete this;
 	}
 
+	DEPRECATED(4.14, "GetResourceSize is deprecated. Please use GetResourceSizeEx or GetResourceSizeBytes instead.")
 	SIZE_T GetResourceSize() const
 	{
-		return sizeof(*this) + DistanceFieldVolume.GetAllocatedSize();
+		return GetResourceSizeBytes();
+	}
+
+	void GetResourceSizeEx(FResourceSizeEx& CumulativeResourceSize) const
+	{
+		CumulativeResourceSize.AddDedicatedSystemMemoryBytes(sizeof(*this));
+		CumulativeResourceSize.AddUnknownMemoryBytes(DistanceFieldVolume.GetAllocatedSize());
+	}
+
+	SIZE_T GetResourceSizeBytes() const
+	{
+		FResourceSizeEx ResSize;
+		GetResourceSizeEx(ResSize);
+		return ResSize.GetTotalMemoryBytes();
 	}
 
 #if WITH_EDITORONLY_DATA
diff --git a/Engine/Source/Runtime/Engine/Public/MaterialShared.h b/Engine/Source/Runtime/Engine/Public/MaterialShared.h
--- a/Engine/Source/Runtime/Engine/Public/MaterialShared.h
+++ b/Engine/Source/Runtime/Engine/Public/MaterialShared.h
@@ -1682,8 +1682,11 @@ public:
 
 	ENGINE_API void GetRepresentativeShaderTypesAndDescriptions(TMap<FName, FString>& OutShaderTypeNameAndDescriptions) const;
 
+	DEPRECATED(4.14, "GetResourceSizeInclusive is deprecated. Please use GetResourceSizeEx instead.")
 	ENGINE_API SIZE_T GetResourceSizeInclusive();
 
+	ENGINE_API void GetResourceSizeEx(FResourceSizeEx& CumulativeResourceSize);
+
 	ENGINE_API virtual void LegacySerialize(FArchive& Ar) override;
 
 	ENGINE_API virtual const TArray<UTexture*>& GetReferencedTextures() const override;
diff --git a/Engine/Source/Runtime/Engine/Public/ParticleEmitterInstances.h b/Engine/Source/Runtime/Engine/Public/ParticleEmitterInstances.h
--- a/Engine/Source/Runtime/Engine/Public/ParticleEmitterInstances.h
+++ b/Engine/Source/Runtime/Engine/Public/ParticleEmitterInstances.h
@@ -593,9 +593,21 @@ public:
 	 * @param	Mode	Specifies which resource size should be displayed. ( see EResourceSizeMode::Type )
 	 * @return  Size of resource as to be displayed to artists/ LDs in the Editor.
 	 */
+	DEPRECATED(4.14, "GetResourceSize is deprecated. Please use GetResourceSizeEx or GetResourceSizeBytes instead.")
 	virtual SIZE_T GetResourceSize(EResourceSizeMode::Type Mode)
 	{
-		return 0;
+		return GetResourceSizeBytes(Mode);
+	}
+
+	virtual void GetResourceSizeEx(FResourceSizeEx& CumulativeResourceSize)
+	{
+	}
+
+	SIZE_T GetResourceSizeBytes(EResourceSizeMode::Type Mode)
+	{
+		FResourceSizeEx ResSize = FResourceSizeEx(Mode);
+		GetResourceSizeEx(ResSize);
+		return ResSize.GetTotalMemoryBytes();
 	}
 
 	/**
@@ -826,7 +838,7 @@ struct FParticleSpriteEmitterInstance : public FParticleEmitterInstance
 	 * @param	Mode	Specifies which resource size should be displayed. ( see EResourceSizeMode::Type )
 	 * @return  Size of resource as to be displayed to artists/ LDs in the Editor.
 	 */
-	virtual SIZE_T GetResourceSize(EResourceSizeMode::Type Mode) override;
+	virtual void GetResourceSizeEx(FResourceSizeEx& CumulativeResourceSize) override;
 
 protected:
 
@@ -890,7 +902,7 @@ struct ENGINE_API FParticleMeshEmitterInstance : public FParticleEmitterInstance
 	 * @param	Mode	Specifies which resource size should be displayed. ( see EResourceSizeMode::Type )
 	 * @return  Size of resource as to be displayed to artists/ LDs in the Editor.
 	 */
-	virtual SIZE_T GetResourceSize(EResourceSizeMode::Type Mode) override;
+	virtual void GetResourceSizeEx(FResourceSizeEx& CumulativeResourceSize) override;
 
 	/**
 	 * Returns the offset to the mesh rotation payload, if any.
@@ -1088,7 +1100,7 @@ struct FParticleBeam2EmitterInstance : public FParticleEmitterInstance
 	 * @param	Mode	Specifies which resource size should be displayed. ( see EResourceSizeMode::Type )
 	 * @return	Size of resource as to be displayed to artists/ LDs in the Editor.
 	 */
-	virtual SIZE_T GetResourceSize(EResourceSizeMode::Type Mode) override;
+	virtual void GetResourceSizeEx(FResourceSizeEx& CumulativeResourceSize) override;
 
 	/**
 	 * When an emitter is killed, this will check other emitters and clean up anything pointing to this one
@@ -1600,7 +1612,7 @@ struct FParticleRibbonEmitterInstance : public FParticleTrailsEmitterInstance_Ba
 	 * @param	Mode	Specifies which resource size should be displayed. ( see EResourceSizeMode::Type )
 	 * @return  Size of resource as to be displayed to artists/ LDs in the Editor.
 	 */
-	virtual SIZE_T GetResourceSize(EResourceSizeMode::Type Mode) override;
+	virtual void GetResourceSizeEx(FResourceSizeEx& CumulativeResourceSize) override;
 
 	/**
 	 * When an emitter is killed, this will check other emitters and clean up anything pointing to this one
@@ -1763,7 +1775,7 @@ struct FParticleAnimTrailEmitterInstance : public FParticleTrailsEmitterInstance
 	 * @param	Mode	Specifies which resource size should be displayed. ( see EResourceSizeMode::Type )
 	 * @return  Size of resource as to be displayed to artists/ LDs in the Editor.
 	 */
-	virtual SIZE_T GetResourceSize(EResourceSizeMode::Type Mode) override;
+	virtual void GetResourceSizeEx(FResourceSizeEx& CumulativeResourceSize) override;
 
 	virtual bool IsTrailEmitter() const override { return true; }
 
diff --git a/Engine/Source/Runtime/Engine/Public/SkeletalMeshTypes.h b/Engine/Source/Runtime/Engine/Public/SkeletalMeshTypes.h
--- a/Engine/Source/Runtime/Engine/Public/SkeletalMeshTypes.h
+++ b/Engine/Source/Runtime/Engine/Public/SkeletalMeshTypes.h
@@ -1707,7 +1707,10 @@ public:
 	/**
 	 * Get Resource Size
 	 */
+	DEPRECATED(4.14, "GetResourceSize is deprecated. Please use GetResourceSizeEx or GetResourceSizeBytes instead.")
 	SIZE_T GetResourceSize() const;
+	void GetResourceSizeEx(FResourceSizeEx& CumulativeResourceSize) const;
+	SIZE_T GetResourceSizeBytes() const;
 
 	/** Rebuild index buffer for everything **/
 #if WITH_EDITOR
@@ -1751,7 +1754,10 @@ public:
 	/** 
 	 *	Return the resource size
 	 */
+	DEPRECATED(4.14, "GetResourceSize is deprecated. Please use GetResourceSizeEx or GetResourceSizeBytes instead.")
 	SIZE_T GetResourceSize();
+	void GetResourceSizeEx(FResourceSizeEx& CumulativeResourceSize);
+	SIZE_T GetResourceSizeBytes();
 
 private:
 	/** True if the resource has been initialized. */
diff --git a/Engine/Source/Runtime/Engine/Public/SkeletalRenderPublic.h b/Engine/Source/Runtime/Engine/Public/SkeletalRenderPublic.h
--- a/Engine/Source/Runtime/Engine/Public/SkeletalRenderPublic.h
+++ b/Engine/Source/Runtime/Engine/Public/SkeletalRenderPublic.h
@@ -155,7 +155,20 @@ public:
 	/**
 	* Returns the size of memory allocated by render data
 	*/
-	virtual SIZE_T GetResourceSize() = 0;
+	DEPRECATED(4.14, "GetResourceSize is deprecated. Please use GetResourceSizeEx or GetResourceSizeBytes instead.")
+	SIZE_T GetResourceSize()
+	{
+		return GetResourceSizeBytes();
+	}
+
+	virtual void GetResourceSizeEx(FResourceSizeEx& CumulativeResourceSize) = 0;
+
+	SIZE_T GetResourceSizeBytes()
+	{
+		FResourceSizeEx ResSize;
+		GetResourceSizeEx(ResSize);
+		return ResSize.GetTotalMemoryBytes();
+	}
 
 	/**
 	 * List of sections to be rendered based on instance weight usage. Full swap of weights will render with its own sections.
diff --git a/Engine/Source/Runtime/Engine/Public/Slate/SceneViewport.h b/Engine/Source/Runtime/Engine/Public/Slate/SceneViewport.h
--- a/Engine/Source/Runtime/Engine/Public/Slate/SceneViewport.h
+++ b/Engine/Source/Runtime/Engine/Public/Slate/SceneViewport.h
@@ -321,6 +321,17 @@ private:
 		return bUseSeparateRenderTarget || bForceSeparateRenderTarget;
 	}
 
+	/** 
+	 * Called right before a backbuffer is resized. If this viewport is using this backbuffer 
+	 * it will release its resource here
+	 */
+	void OnPreResizeWindowBackbuffer(void* Backbuffer);
+
+	/** 
+	 * Called right after a backbuffer is resized. This viewport will reaquire its backbuffer handle if needed
+	 */
+	void OnPostResizeWindowBackbuffer(void* Backbuffer);
+
 private:
 	/** An intermediate reply state that is reset whenever an input event is generated */
 	FReply CurrentReplyState;
diff --git a/Engine/Source/Runtime/Engine/Public/StaticMeshResources.h b/Engine/Source/Runtime/Engine/Public/StaticMeshResources.h
--- a/Engine/Source/Runtime/Engine/Public/StaticMeshResources.h
+++ b/Engine/Source/Runtime/Engine/Public/StaticMeshResources.h
@@ -857,7 +857,10 @@ public:
 	ENGINE_API void ReleaseResources();
 
 	/** Compute the size of this resource. */
+	DEPRECATED(4.14, "GetResourceSize is deprecated. Please use GetResourceSizeEx or GetResourceSizeBytes instead.")
 	SIZE_T GetResourceSize() const;
+	void GetResourceSizeEx(FResourceSizeEx& CumulativeResourceSize) const;
+	SIZE_T GetResourceSizeBytes() const;
 
 	/** Allocate LOD resources. */
 	ENGINE_API void AllocateLODResources(int32 NumLODs);
diff --git a/Engine/Source/Runtime/Foliage/Private/ProceduralFoliageTile.cpp b/Engine/Source/Runtime/Foliage/Private/ProceduralFoliageTile.cpp
--- a/Engine/Source/Runtime/Foliage/Private/ProceduralFoliageTile.cpp
+++ b/Engine/Source/Runtime/Foliage/Private/ProceduralFoliageTile.cpp
@@ -518,16 +518,16 @@ void UProceduralFoliageTile::Empty()
 	PendingRemovals.Empty();
 }
 
-SIZE_T UProceduralFoliageTile::GetResourceSize(EResourceSizeMode::Type Mode)
+void UProceduralFoliageTile::GetResourceSizeEx(FResourceSizeEx& CumulativeResourceSize)
 {
-	SIZE_T TotalSize = 0;
+	Super::GetResourceSizeEx(CumulativeResourceSize);
+
 	for (FProceduralFoliageInstance* Inst : InstancesSet)
 	{
-		TotalSize += sizeof(FProceduralFoliageInstance);
+		CumulativeResourceSize.AddDedicatedSystemMemoryBytes(sizeof(FProceduralFoliageInstance));
 	}
 	
 	//@TODO: account for broadphase
-	return TotalSize;
 }
 
 
diff --git a/Engine/Source/Runtime/GeometryCache/Classes/GeometryCacheMeshData.h b/Engine/Source/Runtime/GeometryCache/Classes/GeometryCacheMeshData.h
--- a/Engine/Source/Runtime/GeometryCache/Classes/GeometryCacheMeshData.h
+++ b/Engine/Source/Runtime/GeometryCache/Classes/GeometryCacheMeshData.h
@@ -85,18 +85,28 @@ struct FGeometryCacheMeshData
 		return Ar;
 	}
 
+	DEPRECATED(4.14, "GetResourceSize is deprecated. Please use GetResourceSizeEx or GetResourceSizeBytes instead.")
 	SIZE_T GetResourceSize() const
+	{
+		return GetResourceSizeBytes();
+	}
+
+	void GetResourceSizeEx(FResourceSizeEx& CumulativeResourceSize) const
 	{
 		// Calculate resource size according to what is actually serialized
-		SIZE_T Size = 0;
-		Size += Vertices.Num() * sizeof(FDynamicMeshVertex);
-		Size += BatchesInfo.Num() * sizeof(FGeometryCacheMeshBatchInfo);
-		Size += sizeof(Vertices);
-		Size += sizeof(BatchesInfo);
-		Size += sizeof(BoundingBox);
-		Size += Indices.Num() * sizeof(uint32);
-		Size += sizeof(Indices);
+		CumulativeResourceSize.AddUnknownMemoryBytes(Vertices.Num() * sizeof(FDynamicMeshVertex));
+		CumulativeResourceSize.AddUnknownMemoryBytes(BatchesInfo.Num() * sizeof(FGeometryCacheMeshBatchInfo));
+		CumulativeResourceSize.AddUnknownMemoryBytes(sizeof(Vertices));
+		CumulativeResourceSize.AddUnknownMemoryBytes(sizeof(BatchesInfo));
+		CumulativeResourceSize.AddUnknownMemoryBytes(sizeof(BoundingBox));
+		CumulativeResourceSize.AddUnknownMemoryBytes(Indices.Num() * sizeof(uint32));
+		CumulativeResourceSize.AddUnknownMemoryBytes(sizeof(Indices));
+	}
 
-		return Size;
+	SIZE_T GetResourceSizeBytes() const
+	{
+		FResourceSizeEx ResSize;
+		GetResourceSizeEx(ResSize);
+		return ResSize.GetTotalMemoryBytes();
 	}
 };
diff --git a/Engine/Source/Runtime/GeometryCache/Private/GeometryCache.cpp b/Engine/Source/Runtime/GeometryCache/Private/GeometryCache.cpp
--- a/Engine/Source/Runtime/GeometryCache/Private/GeometryCache.cpp
+++ b/Engine/Source/Runtime/GeometryCache/Private/GeometryCache.cpp
@@ -39,24 +39,23 @@ FString UGeometryCache::GetDesc()
 	return FString("%d Tracks", NumTracks);
 }
 
-SIZE_T UGeometryCache::GetResourceSize(EResourceSizeMode::Type Mode)
+void UGeometryCache::GetResourceSizeEx(FResourceSizeEx& CumulativeResourceSize)
 {
-	SIZE_T ResourceSize = 0;
+	Super::GetResourceSizeEx(CumulativeResourceSize);
+
 #if WITH_EDITORONLY_DATA
-	ResourceSize += sizeof( AssetImportData );
+	CumulativeResourceSize.AddDedicatedSystemMemoryBytes(sizeof(AssetImportData));
 #endif
 	// Calculate Resource Size according to what is serialized
 	const int32 NumTracks = Tracks.Num();
 	for (int32 TrackIndex = 0; TrackIndex < NumTracks; ++TrackIndex)
 	{
-		ResourceSize += Tracks[TrackIndex]->GetResourceSize(Mode);
+		Tracks[TrackIndex]->GetResourceSizeEx(CumulativeResourceSize);
 	}
 
-	ResourceSize += sizeof(Tracks);
-	ResourceSize += sizeof(NumVertexAnimationTracks);
-	ResourceSize += sizeof(NumTransformAnimationTracks);
-
-	return ResourceSize;
+	CumulativeResourceSize.AddDedicatedSystemMemoryBytes(sizeof(Tracks));
+	CumulativeResourceSize.AddUnknownMemoryBytes(sizeof(NumVertexAnimationTracks));
+	CumulativeResourceSize.AddUnknownMemoryBytes(sizeof(NumTransformAnimationTracks));
 }
 
 void UGeometryCache::GetAssetRegistryTags(TArray<FAssetRegistryTag>& OutTags) const
diff --git a/Engine/Source/Runtime/GeometryCache/Private/GeometryCacheTrack.cpp b/Engine/Source/Runtime/GeometryCache/Private/GeometryCacheTrack.cpp
--- a/Engine/Source/Runtime/GeometryCache/Private/GeometryCacheTrack.cpp
+++ b/Engine/Source/Runtime/GeometryCache/Private/GeometryCacheTrack.cpp
@@ -104,13 +104,11 @@ const uint32 UGeometryCacheTrack::FindSampleIndexFromTime(const TArray<float>& S
 	return MinIndex;
 }
 
-SIZE_T UGeometryCacheTrack::GetResourceSize(EResourceSizeMode::Type Mode)
+void UGeometryCacheTrack::GetResourceSizeEx(FResourceSizeEx& CumulativeResourceSize)
 {
-	SIZE_T ResourceSize = 0;
+	Super::GetResourceSizeEx(CumulativeResourceSize);
 
 	// Determine resource size from data that is serialized
-	ResourceSize += MatrixSamples.Num() * sizeof(FMatrix);
-	ResourceSize += MatrixSampleTimes.Num() * sizeof(float);
-
-	return ResourceSize;
+	CumulativeResourceSize.AddDedicatedSystemMemoryBytes(MatrixSamples.Num() * sizeof(FMatrix));
+	CumulativeResourceSize.AddDedicatedSystemMemoryBytes(MatrixSampleTimes.Num() * sizeof(float));
 }
diff --git a/Engine/Source/Runtime/GeometryCache/Private/GeometryCacheTrackFlipbookAnimation.cpp b/Engine/Source/Runtime/GeometryCache/Private/GeometryCacheTrackFlipbookAnimation.cpp
--- a/Engine/Source/Runtime/GeometryCache/Private/GeometryCacheTrackFlipbookAnimation.cpp
+++ b/Engine/Source/Runtime/GeometryCache/Private/GeometryCacheTrackFlipbookAnimation.cpp
@@ -17,21 +17,19 @@ UGeometryCacheTrack_FlipbookAnimation::~UGeometryCacheTrack_FlipbookAnimation()
 	MeshSampleTimes.Empty();
 }
 
-SIZE_T UGeometryCacheTrack_FlipbookAnimation::GetResourceSize(EResourceSizeMode::Type Mode)
+void UGeometryCacheTrack_FlipbookAnimation::GetResourceSizeEx(FResourceSizeEx& CumulativeResourceSize)
 {
+	Super::GetResourceSizeEx(CumulativeResourceSize);
+
 	// Determine resource size according to what is actually serialized
-	SIZE_T ResourceSize = 0;
-	ResourceSize += UGeometryCacheTrack::GetResourceSize(Mode);
 	for (int32 SampleIndex = 0; SampleIndex < MeshSamples.Num(); ++SampleIndex )
 	{
-		ResourceSize += MeshSamples[SampleIndex].GetResourceSize();
+		MeshSamples[SampleIndex].GetResourceSizeEx(CumulativeResourceSize);
 	}
-	ResourceSize += sizeof(MeshSamples);
-	ResourceSize += MeshSampleTimes.Num() * sizeof(float);
-	ResourceSize += sizeof(MeshSampleTimes);
-	ResourceSize += sizeof(NumMeshSamples);
-
-	return ResourceSize;
+	CumulativeResourceSize.AddDedicatedSystemMemoryBytes(sizeof(MeshSamples));
+	CumulativeResourceSize.AddDedicatedSystemMemoryBytes(MeshSampleTimes.Num() * sizeof(float));
+	CumulativeResourceSize.AddDedicatedSystemMemoryBytes(sizeof(MeshSampleTimes));
+	CumulativeResourceSize.AddDedicatedSystemMemoryBytes(sizeof(NumMeshSamples));
 }
 
 void UGeometryCacheTrack_FlipbookAnimation::Serialize(FArchive& Ar)
diff --git a/Engine/Source/Runtime/GeometryCache/Private/GeometryCacheTrackTransformAnimation.cpp b/Engine/Source/Runtime/GeometryCache/Private/GeometryCacheTrackTransformAnimation.cpp
--- a/Engine/Source/Runtime/GeometryCache/Private/GeometryCacheTrackTransformAnimation.cpp
+++ b/Engine/Source/Runtime/GeometryCache/Private/GeometryCacheTrackTransformAnimation.cpp
@@ -8,13 +8,11 @@ GEOMETRYCACHE_API UGeometryCacheTrack_TransformAnimation::UGeometryCacheTrack_Tr
 {
 }
 
-SIZE_T UGeometryCacheTrack_TransformAnimation::GetResourceSize(EResourceSizeMode::Type Mode)
+void UGeometryCacheTrack_TransformAnimation::GetResourceSizeEx(FResourceSizeEx& CumulativeResourceSize)
 {
-	SIZE_T ResourceSize = 0;
-	ResourceSize += UGeometryCacheTrack::GetResourceSize(Mode);
-	ResourceSize += MeshData.GetResourceSize();
+	Super::GetResourceSizeEx(CumulativeResourceSize);
 
-	return ResourceSize;
+	MeshData.GetResourceSizeEx(CumulativeResourceSize);
 }
 
 void UGeometryCacheTrack_TransformAnimation::Serialize(FArchive& Ar)
diff --git a/Engine/Source/Runtime/GeometryCache/Private/GeometryCacheTrackTransformGroupAnimation.cpp b/Engine/Source/Runtime/GeometryCache/Private/GeometryCacheTrackTransformGroupAnimation.cpp
--- a/Engine/Source/Runtime/GeometryCache/Private/GeometryCacheTrackTransformGroupAnimation.cpp
+++ b/Engine/Source/Runtime/GeometryCache/Private/GeometryCacheTrackTransformGroupAnimation.cpp
@@ -8,13 +8,11 @@ GEOMETRYCACHE_API UGeometryCacheTrack_TransformGroupAnimation::UGeometryCacheTra
 {
 }
 
-SIZE_T UGeometryCacheTrack_TransformGroupAnimation::GetResourceSize(EResourceSizeMode::Type Mode)
+void UGeometryCacheTrack_TransformGroupAnimation::GetResourceSizeEx(FResourceSizeEx& CumulativeResourceSize)
 {
-	SIZE_T ResourceSize = 0;
-	ResourceSize += UGeometryCacheTrack::GetResourceSize(Mode);
-	ResourceSize += MeshData.GetResourceSize();
+	Super::GetResourceSizeEx(CumulativeResourceSize);
 
-	return ResourceSize;
+	MeshData.GetResourceSizeEx(CumulativeResourceSize);
 }
 
 void UGeometryCacheTrack_TransformGroupAnimation::Serialize(FArchive& Ar)
diff --git a/Engine/Source/Runtime/Landscape/Classes/Landscape.h b/Engine/Source/Runtime/Landscape/Classes/Landscape.h
--- a/Engine/Source/Runtime/Landscape/Classes/Landscape.h
+++ b/Engine/Source/Runtime/Landscape/Classes/Landscape.h
@@ -52,6 +52,7 @@ public:
 	virtual void PostEditChangeProperty(FPropertyChangedEvent& PropertyChangedEvent) override;
 	virtual void PostEditMove(bool bFinished) override;
 	virtual void PostEditImport() override;
+	virtual void PostDuplicate(bool bDuplicateForPIE) override;
 #endif
 	virtual void PostLoad() override;
 	//~ End UObject Interface
diff --git a/Engine/Source/Runtime/Landscape/Classes/LandscapeComponent.h b/Engine/Source/Runtime/Landscape/Classes/LandscapeComponent.h
--- a/Engine/Source/Runtime/Landscape/Classes/LandscapeComponent.h
+++ b/Engine/Source/Runtime/Landscape/Classes/LandscapeComponent.h
@@ -334,7 +334,7 @@ public:
 	//~ Begin UObject Interface.	
 	virtual void PostInitProperties() override;	
 	virtual void Serialize(FArchive& Ar) override;
-	virtual SIZE_T GetResourceSize(EResourceSizeMode::Type Mode) override;
+	virtual void GetResourceSizeEx(FResourceSizeEx& CumulativeResourceSize) override;
 	virtual void BeginDestroy() override;
 	static void AddReferencedObjects(UObject* InThis, FReferenceCollector& Collector);
 	virtual void PostDuplicate(bool bDuplicateForPIE) override;
@@ -486,6 +486,9 @@ public:
 	 */
 	LANDSCAPE_API void UpdateMaterialInstances();
 
+	// Internal implementation of UpdateMaterialInstances, not safe to call directly
+	void UpdateMaterialInstances_Internal(FMaterialUpdateContext& Context);
+
 	/** Helper function for UpdateMaterialInstance to get Material without set parameters */
 	UMaterialInstanceConstant* GetCombinationMaterial(bool bMobile = false);
 	/**
diff --git a/Engine/Source/Runtime/Landscape/Classes/LandscapeInfo.h b/Engine/Source/Runtime/Landscape/Classes/LandscapeInfo.h
--- a/Engine/Source/Runtime/Landscape/Classes/LandscapeInfo.h
+++ b/Engine/Source/Runtime/Landscape/Classes/LandscapeInfo.h
@@ -51,11 +51,11 @@ struct FLandscapeInfoLayerSettings
 #endif // WITH_EDITORONLY_DATA
 
 	FLandscapeInfoLayerSettings()
-		: LayerInfoObj(NULL)
+		: LayerInfoObj(nullptr)
 		, LayerName(NAME_None)
 #if WITH_EDITORONLY_DATA
-		, ThumbnailMIC(NULL)
-		, Owner(NULL)
+		, ThumbnailMIC(nullptr)
+		, Owner(nullptr)
 		, DebugColorChannel(0)
 		, bValid(false)
 #endif // WITH_EDITORONLY_DATA
@@ -65,10 +65,10 @@ struct FLandscapeInfoLayerSettings
 	LANDSCAPE_API FLandscapeInfoLayerSettings(ULandscapeLayerInfoObject* InLayerInfo, ALandscapeProxy* InProxy);
 
 	FLandscapeInfoLayerSettings(FName InPlaceholderLayerName, ALandscapeProxy* InProxy)
-		: LayerInfoObj(NULL)
+		: LayerInfoObj(nullptr)
 		, LayerName(InPlaceholderLayerName)
 #if WITH_EDITORONLY_DATA
-		, ThumbnailMIC(NULL)
+		, ThumbnailMIC(nullptr)
 		, Owner(InProxy)
 		, DebugColorChannel(0)
 		, bValid(false)
@@ -134,7 +134,6 @@ public:
 
 	//~ Begin UObject Interface.
 	virtual void Serialize(FArchive& Ar) override;
-	virtual void BeginDestroy() override;
 	//~ End UObject Interface
 
 #if WITH_EDITOR
@@ -168,10 +167,10 @@ public:
 	LANDSCAPE_API FLandscapeEditorLayerSettings& GetLayerEditorSettings(ULandscapeLayerInfoObject* LayerInfo) const;
 	LANDSCAPE_API void CreateLayerEditorSettingsFor(ULandscapeLayerInfoObject* LayerInfo);
 
-	LANDSCAPE_API ULandscapeLayerInfoObject* GetLayerInfoByName(FName LayerName, ALandscapeProxy* Owner = NULL) const;
-	LANDSCAPE_API int32 GetLayerInfoIndex(FName LayerName, ALandscapeProxy* Owner = NULL) const;
-	LANDSCAPE_API int32 GetLayerInfoIndex(ULandscapeLayerInfoObject* LayerInfo, ALandscapeProxy* Owner = NULL) const;
-	LANDSCAPE_API bool UpdateLayerInfoMap(ALandscapeProxy* Proxy = NULL, bool bInvalidate = false);
+	LANDSCAPE_API ULandscapeLayerInfoObject* GetLayerInfoByName(FName LayerName, ALandscapeProxy* Owner = nullptr) const;
+	LANDSCAPE_API int32 GetLayerInfoIndex(FName LayerName, ALandscapeProxy* Owner = nullptr) const;
+	LANDSCAPE_API int32 GetLayerInfoIndex(ULandscapeLayerInfoObject* LayerInfo, ALandscapeProxy* Owner = nullptr) const;
+	LANDSCAPE_API bool UpdateLayerInfoMap(ALandscapeProxy* Proxy = nullptr, bool bInvalidate = false);
 
 	/**
 	 *  Returns the landscape proxy of this landscape info in the given level (if it exists)
@@ -193,6 +192,16 @@ public:
 	 */
 	LANDSCAPE_API ALandscapeProxy* GetLandscapeProxy() const;
 
+	/**
+	 * Runs the given function on the root landscape actor and all streaming proxies
+	 * Most easily used with a lambda as follows:
+	 * ForAllLandscapeProxies([](ALandscapeProxy* Proxy)
+	 * {
+	 *     // Code
+	 * });
+	 */
+	LANDSCAPE_API void ForAllLandscapeProxies(TFunctionRef<void(ALandscapeProxy*)> Fn) const;
+
 	/** Associates passed actor with this info object
  	 *  @param	Proxy		Landscape actor to register
 	 *  @param  bMapCheck	Whether to warn about landscape errors
diff --git a/Engine/Source/Runtime/Landscape/Classes/LandscapeProxy.h b/Engine/Source/Runtime/Landscape/Classes/LandscapeProxy.h
--- a/Engine/Source/Runtime/Landscape/Classes/LandscapeProxy.h
+++ b/Engine/Source/Runtime/Landscape/Classes/LandscapeProxy.h
@@ -66,15 +66,16 @@ struct FLandscapeEditorLayerSettings
 	{
 	}
 
-	FLandscapeEditorLayerSettings(ULandscapeLayerInfoObject* InLayerInfo, const FString& InFilePath = FString())
+	explicit FLandscapeEditorLayerSettings(ULandscapeLayerInfoObject* InLayerInfo, const FString& InFilePath = FString())
 		: LayerInfoObj(InLayerInfo)
 		, ReimportLayerFilePath(InFilePath)
 	{
 	}
 
-	bool operator==(const FLandscapeEditorLayerSettings& rhs) const
+	// To allow FindByKey etc
+	bool operator==(const ULandscapeLayerInfoObject* LayerInfo) const
 	{
-		return LayerInfoObj == rhs.LayerInfoObj;
+		return LayerInfoObj == LayerInfo;
 	}
 #endif // WITH_EDITORONLY_DATA
 };
@@ -649,9 +650,13 @@ public:
 	/** Remove all XYOffset values */
 	LANDSCAPE_API void RemoveXYOffsets();
 
+	/** Update the material instances for all the landscape components */
+	LANDSCAPE_API void UpdateAllComponentMaterialInstances();
 
+	/** Create a thumbnail material for a given layer */
 	LANDSCAPE_API static ULandscapeMaterialInstanceConstant* GetLayerThumbnailMIC(UMaterialInterface* LandscapeMaterial, FName LayerName, UTexture2D* ThumbnailWeightmap, UTexture2D* ThumbnailHeightmap, ALandscapeProxy* Proxy);
 
+	/** Import the given Height/Weight data into this landscape */
 	LANDSCAPE_API void Import(FGuid Guid, int32 MinX, int32 MinY, int32 MaxX, int32 MaxY, int32 NumSubsections, int32 SubsectionSizeQuads,
 							const uint16* HeightData, const TCHAR* HeightmapFileName,
 							const TArray<FLandscapeImportLayerInfo>& ImportLayerInfos, ELandscapeImportAlphamapType ImportLayerType);
diff --git a/Engine/Source/Runtime/Landscape/Private/Landscape.cpp b/Engine/Source/Runtime/Landscape/Private/Landscape.cpp
--- a/Engine/Source/Runtime/Landscape/Private/Landscape.cpp
+++ b/Engine/Source/Runtime/Landscape/Private/Landscape.cpp
@@ -120,7 +120,7 @@ ULandscapeComponent::ULandscapeComponent(const FObjectInitializer& ObjectInitial
 	Mobility = EComponentMobility::Static;
 
 #if WITH_EDITORONLY_DATA
-	EditToolRenderData = NULL;
+	EditToolRenderData = nullptr;
 #endif
 
 	LpvBiasMultiplier = 0.0f; // Bias is 0 for landscape, since it's single sided
@@ -132,11 +132,11 @@ ULandscapeComponent::ULandscapeComponent(const FObjectInitializer& ObjectInitial
 void ULandscapeComponent::AddReferencedObjects(UObject* InThis, FReferenceCollector& Collector)
 {
 	ULandscapeComponent* This = CastChecked<ULandscapeComponent>(InThis);
-	if (This->LightMap != NULL)
+	if (This->LightMap != nullptr)
 	{
 		This->LightMap->AddReferencedObjects(Collector);
 	}
-	if (This->ShadowMap != NULL)
+	if (This->ShadowMap != nullptr)
 	{
 		This->ShadowMap->AddReferencedObjects(Collector);
 	}
@@ -358,9 +358,11 @@ void ULandscapeComponent::Serialize(FArchive& Ar)
 #endif
 }
 
-SIZE_T ULandscapeComponent::GetResourceSize(EResourceSizeMode::Type Mode)
+void ULandscapeComponent::GetResourceSizeEx(FResourceSizeEx& CumulativeResourceSize)
 {
-	return Super::GetResourceSize(Mode) + GrassData->GetAllocatedSize();
+	Super::GetResourceSizeEx(CumulativeResourceSize);
+
+	CumulativeResourceSize.AddUnknownMemoryBytes(GrassData->GetAllocatedSize());
 }
 
 #if WITH_EDITOR
@@ -620,7 +622,7 @@ void ULandscapeComponent::PostLoad()
 				Obj->ClearFlags(RF_Public);
 				if (Obj->GetOuter() == Level)
 				{
-					Obj->Rename(NULL, MyPackage, REN_DoNotDirty | REN_DontCreateRedirectors | REN_ForceNoResetLoaders | REN_NonTransactional);
+					Obj->Rename(nullptr, MyPackage, REN_DoNotDirty | REN_DontCreateRedirectors | REN_ForceNoResetLoaders | REN_NonTransactional);
 				}
 			}
 		}
@@ -673,7 +675,7 @@ ALandscapeProxy::ALandscapeProxy(const FObjectInitializer& ObjectInitializer)
 #endif
 
 #if WITH_EDITOR
-	if (VisibilityLayer == NULL)
+	if (VisibilityLayer == nullptr)
 	{
 		// Structure to hold one-time initialization
 		struct FConstructorStatics
@@ -776,13 +778,13 @@ ALandscape* ULandscapeComponent::GetLandscapeActor() const
 	{
 		return Landscape->GetLandscapeActor();
 	}
-	return NULL;
+	return nullptr;
 }
 
 ULevel* ULandscapeComponent::GetLevel() const
 {
 	AActor* MyOwner = GetOwner();
-	return MyOwner ? MyOwner->GetLevel() : NULL;
+	return MyOwner ? MyOwner->GetLevel() : nullptr;
 }
 
 #if WITH_EDITOR
@@ -856,11 +858,11 @@ void ULandscapeComponent::BeginDestroy()
 	Super::BeginDestroy();
 
 #if WITH_EDITOR
-	if (EditToolRenderData != NULL)
+	if (EditToolRenderData != nullptr)
 	{
 		// Ask render thread to destroy EditToolRenderData
 		EditToolRenderData->Cleanup();
-		EditToolRenderData = NULL;
+		EditToolRenderData = nullptr;
 	}
 
 	if (GIsEditor && !HasAnyFlags(RF_ClassDefaultObject))
@@ -875,9 +877,9 @@ void ULandscapeComponent::BeginDestroy()
 			{
 				UTexture2D* WeightmapTexture = WeightmapTextures[WeightmapIndex];
 				FLandscapeWeightmapUsage* Usage = Proxy->WeightmapUsageMap.Find(WeightmapTexture);
-				if (Usage != NULL)
+				if (Usage != nullptr)
 				{
-					Usage->ChannelUsage[WeightmapLayerAllocations[LayerIdx].WeightmapTextureChannel] = NULL;
+					Usage->ChannelUsage[WeightmapLayerAllocations[LayerIdx].WeightmapTextureChannel] = nullptr;
 
 					if (Usage->FreeChannelCount() == 4)
 					{
@@ -1084,15 +1086,37 @@ void ALandscape::PostLoad()
 	{
 		LandscapeGuid = FGuid::NewGuid();
 	}
+	else
+	{
+#if WITH_EDITOR
+		for (ALandscape* Landscape : TActorRange<ALandscape>(GetWorld()))
+		{
+			if (Landscape && Landscape != this && !Landscape->HasAnyFlags(RF_BeginDestroyed) && Landscape->LandscapeGuid == LandscapeGuid)
+			{
+				// Duplicated landscape level, need to generate new GUID
+				Modify();
+				LandscapeGuid = FGuid::NewGuid();
+
+				FFormatNamedArguments Arguments;
+				Arguments.Add(TEXT("ProxyName1"), FText::FromString(Landscape->GetName()));
+				Arguments.Add(TEXT("LevelName1"), FText::FromString(Landscape->GetLevel()->GetOutermost()->GetName()));
+				Arguments.Add(TEXT("ProxyName2"), FText::FromString(this->GetName()));
+				Arguments.Add(TEXT("LevelName2"), FText::FromString(this->GetLevel()->GetOutermost()->GetName()));
+				FMessageLog("LoadErrors").Warning()
+					->AddToken(FUObjectToken::Create(this))
+					->AddToken(FTextToken::Create(FText::Format(LOCTEXT("LoadError_DuplicateLandscapeGuid", "Landscape {ProxyName1} of {LevelName1} has the same guid as {ProxyName2} of {LevelName2}. {LevelName2}.{ProxyName2} has had its guid automatically changed, please save {LevelName2}!"), Arguments)));
+
+				// Show MapCheck window
+				FMessageLog("LoadErrors").Open();
+				break;
+			}
+		}
+#endif
+	}
 
 	Super::PostLoad();
 }
 
-void ULandscapeInfo::BeginDestroy()
-{
-	Super::BeginDestroy();
-}
-
 void ALandscapeProxy::PreSave(const class ITargetPlatform* TargetPlatform)
 {
 	Super::PreSave(TargetPlatform);
@@ -1139,7 +1163,7 @@ void ALandscapeProxy::AddReferencedObjects(UObject* InThis, FReferenceCollector&
 #if WITH_EDITOR
 FName FLandscapeInfoLayerSettings::GetLayerName() const
 {
-	checkSlow(LayerInfoObj == NULL || LayerInfoObj->LayerName == LayerName);
+	checkSlow(LayerInfoObj == nullptr || LayerInfoObj->LayerName == LayerName);
 
 	return LayerName;
 }
@@ -1154,51 +1178,39 @@ FLandscapeEditorLayerSettings& FLandscapeInfoLayerSettings::GetEditorSettings()
 
 FLandscapeEditorLayerSettings& ULandscapeInfo::GetLayerEditorSettings(ULandscapeLayerInfoObject* LayerInfo) const
 {
-	auto& LayerSettingsList = GetLandscapeProxy()->EditorLayerSettings;
-	FLandscapeEditorLayerSettings* EditorLayerSettings = LayerSettingsList.FindByKey(LayerInfo);
+	ALandscapeProxy* Proxy = GetLandscapeProxy();
+	FLandscapeEditorLayerSettings* EditorLayerSettings = Proxy->EditorLayerSettings.FindByKey(LayerInfo);
 	if (EditorLayerSettings)
 	{
 		return *EditorLayerSettings;
 	}
 	else
 	{
-		int32 Index = LayerSettingsList.Add(LayerInfo);
-		return LayerSettingsList[Index];
+		int32 Index = Proxy->EditorLayerSettings.Add(FLandscapeEditorLayerSettings(LayerInfo));
+		return Proxy->EditorLayerSettings[Index];
 	}
 }
 
 void ULandscapeInfo::CreateLayerEditorSettingsFor(ULandscapeLayerInfoObject* LayerInfo)
 {
-	ALandscape* Landscape = LandscapeActor.Get();
-	if (Landscape != NULL)
+	ForAllLandscapeProxies([LayerInfo](ALandscapeProxy* Proxy)
 	{
-		FLandscapeEditorLayerSettings* EditorLayerSettings = Landscape->EditorLayerSettings.FindByKey(LayerInfo);
-		if (EditorLayerSettings == NULL)
-		{
-			Landscape->Modify();
-			Landscape->EditorLayerSettings.Add(LayerInfo);
-		}
-	}
-
-	for (auto It = Proxies.CreateConstIterator(); It; ++It)
-	{
-		ALandscapeProxy* Proxy = *It;
 		FLandscapeEditorLayerSettings* EditorLayerSettings = Proxy->EditorLayerSettings.FindByKey(LayerInfo);
-		if (EditorLayerSettings == NULL)
+		if (!EditorLayerSettings)
 		{
 			Proxy->Modify();
-			Proxy->EditorLayerSettings.Add(LayerInfo);
+			Proxy->EditorLayerSettings.Add(FLandscapeEditorLayerSettings(LayerInfo));
 		}
-	}
+	});
 }
 
-ULandscapeLayerInfoObject* ULandscapeInfo::GetLayerInfoByName(FName LayerName, ALandscapeProxy* Owner /*= NULL*/) const
+ULandscapeLayerInfoObject* ULandscapeInfo::GetLayerInfoByName(FName LayerName, ALandscapeProxy* Owner /*= nullptr*/) const
 {
-	ULandscapeLayerInfoObject* LayerInfo = NULL;
+	ULandscapeLayerInfoObject* LayerInfo = nullptr;
 	for (int32 j = 0; j < Layers.Num(); j++)
 	{
 		if (Layers[j].LayerInfoObj && Layers[j].LayerInfoObj->LayerName == LayerName
-			&& (Owner == NULL || Layers[j].Owner == Owner))
+			&& (Owner == nullptr || Layers[j].Owner == Owner))
 		{
 			LayerInfo = Layers[j].LayerInfoObj;
 		}
@@ -1206,12 +1218,12 @@ ULandscapeLayerInfoObject* ULandscapeInfo::GetLayerInfoByName(FName LayerName, A
 	return LayerInfo;
 }
 
-int32 ULandscapeInfo::GetLayerInfoIndex(ULandscapeLayerInfoObject* LayerInfo, ALandscapeProxy* Owner /*= NULL*/) const
+int32 ULandscapeInfo::GetLayerInfoIndex(ULandscapeLayerInfoObject* LayerInfo, ALandscapeProxy* Owner /*= nullptr*/) const
 {
 	for (int32 j = 0; j < Layers.Num(); j++)
 	{
 		if (Layers[j].LayerInfoObj && Layers[j].LayerInfoObj == LayerInfo
-			&& (Owner == NULL || Layers[j].Owner == Owner))
+			&& (Owner == nullptr || Layers[j].Owner == Owner))
 		{
 			return j;
 		}
@@ -1220,12 +1232,12 @@ int32 ULandscapeInfo::GetLayerInfoIndex(ULandscapeLayerInfoObject* LayerInfo, AL
 	return INDEX_NONE;
 }
 
-int32 ULandscapeInfo::GetLayerInfoIndex(FName LayerName, ALandscapeProxy* Owner /*= NULL*/) const
+int32 ULandscapeInfo::GetLayerInfoIndex(FName LayerName, ALandscapeProxy* Owner /*= nullptr*/) const
 {
 	for (int32 j = 0; j < Layers.Num(); j++)
 	{
 		if (Layers[j].GetLayerName() == LayerName
-			&& (Owner == NULL || Layers[j].Owner == Owner))
+			&& (Owner == nullptr || Layers[j].Owner == Owner))
 		{
 			return j;
 		}
@@ -1235,7 +1247,7 @@ int32 ULandscapeInfo::GetLayerInfoIndex(FName LayerName, ALandscapeProxy* Owner
 }
 
 
-bool ULandscapeInfo::UpdateLayerInfoMap(ALandscapeProxy* Proxy /*= NULL*/, bool bInvalidate /*= false*/)
+bool ULandscapeInfo::UpdateLayerInfoMap(ALandscapeProxy* Proxy /*= nullptr*/, bool bInvalidate /*= false*/)
 {
 	bool bHasCollision = false;
 	if (GIsEditor)
@@ -1285,7 +1297,7 @@ bool ULandscapeInfo::UpdateLayerInfoMap(ALandscapeProxy* Proxy /*= NULL*/, bool
 					ULandscapeComponent* Component = Proxy->LandscapeComponents[ComponentIndex];
 
 					// Add layers from per-component override materials
-					if (Component->OverrideMaterial != NULL)
+					if (Component->OverrideMaterial != nullptr)
 					{
 						TArray<FName> ComponentLayerNames = Proxy->GetLayersFromMaterial(Component->OverrideMaterial);
 						for (int32 i = 0; i < ComponentLayerNames.Num(); i++)
@@ -1319,7 +1331,7 @@ bool ULandscapeInfo::UpdateLayerInfoMap(ALandscapeProxy* Proxy /*= NULL*/, bool
 								{
 									LayerSettings.Owner = Proxy;
 									LayerSettings.bValid = bValid;
-									LayerSettings.ThumbnailMIC = NULL;
+									LayerSettings.ThumbnailMIC = nullptr;
 								}
 							}
 							else
@@ -1335,7 +1347,7 @@ bool ULandscapeInfo::UpdateLayerInfoMap(ALandscapeProxy* Proxy /*= NULL*/, bool
 										LayerSettings.Owner = Proxy;
 										LayerSettings.LayerInfoObj = LayerInfo;
 										LayerSettings.bValid = bValid;
-										LayerSettings.ThumbnailMIC = NULL;
+										LayerSettings.ThumbnailMIC = nullptr;
 									}
 								}
 								else
@@ -1350,22 +1362,22 @@ bool ULandscapeInfo::UpdateLayerInfoMap(ALandscapeProxy* Proxy /*= NULL*/, bool
 				}
 
 				// Add any layer infos cached in the actor
-				Proxy->EditorLayerSettings.Remove(NULL);
+				Proxy->EditorLayerSettings.RemoveAll([](const FLandscapeEditorLayerSettings& Settings) { return Settings.LayerInfoObj == nullptr; });
 				for (int32 i = 0; i < Proxy->EditorLayerSettings.Num(); i++)
 				{
-					FLandscapeEditorLayerSettings& LayerInfo = Proxy->EditorLayerSettings[i];
-					if (LayerNames.Contains(LayerInfo.LayerInfoObj->LayerName))
+					FLandscapeEditorLayerSettings& EditorLayerSettings = Proxy->EditorLayerSettings[i];
+					if (LayerNames.Contains(EditorLayerSettings.LayerInfoObj->LayerName))
 					{
 						// intentionally using the layer name here so we don't add layer infos from
 						// the cache that have the same name as an actual assignment from a component above
-						int32 LayerInfoIndex = GetLayerInfoIndex(LayerInfo.LayerInfoObj->LayerName);
+						int32 LayerInfoIndex = GetLayerInfoIndex(EditorLayerSettings.LayerInfoObj->LayerName);
 						if (LayerInfoIndex != INDEX_NONE)
 						{
 							FLandscapeInfoLayerSettings& LayerSettings = Layers[LayerInfoIndex];
-							if (LayerSettings.LayerInfoObj == NULL)
+							if (LayerSettings.LayerInfoObj == nullptr)
 							{
 								LayerSettings.Owner = Proxy;
-								LayerSettings.LayerInfoObj = LayerInfo.LayerInfoObj;
+								LayerSettings.LayerInfoObj = EditorLayerSettings.LayerInfoObj;
 								LayerSettings.bValid = true;
 							}
 						}
@@ -1384,17 +1396,14 @@ bool ULandscapeInfo::UpdateLayerInfoMap(ALandscapeProxy* Proxy /*= NULL*/, bool
 
 			if (!bInvalidate)
 			{
-				if (LandscapeActor.IsValid() && !LandscapeActor->IsPendingKillPending())
-				{
-					checkSlow(LandscapeActor->GetLandscapeInfo() == this);
-					UpdateLayerInfoMap(LandscapeActor.Get(), false);
-				}
-				for (auto It = Proxies.CreateIterator(); It; ++It)
+				ForAllLandscapeProxies([this](ALandscapeProxy* EachProxy)
 				{
-					ALandscapeProxy* LandscapeProxy = *It;
-					checkSlow(LandscapeProxy && LandscapeProxy->GetLandscapeInfo() == this);
-					UpdateLayerInfoMap(LandscapeProxy, false);
-				}
+					if (!EachProxy->IsPendingKillPending())
+					{
+						checkSlow(EachProxy->GetLandscapeInfo() == this);
+						UpdateLayerInfoMap(EachProxy, false);
+					}
+				});
 			}
 		}
 
@@ -1448,13 +1457,13 @@ void ALandscapeProxy::PostLoad()
 		}
 	}
 
-	EditorLayerSettings.Remove(nullptr);
+	EditorLayerSettings.RemoveAll([](const FLandscapeEditorLayerSettings& Settings) { return Settings.LayerInfoObj == nullptr; });
 
 	if (EditorCachedLayerInfos_DEPRECATED.Num() > 0)
 	{
 		for (int32 i = 0; i < EditorCachedLayerInfos_DEPRECATED.Num(); i++)
 		{
-			EditorLayerSettings.Add(EditorCachedLayerInfos_DEPRECATED[i]);
+			EditorLayerSettings.Add(FLandscapeEditorLayerSettings(EditorCachedLayerInfos_DEPRECATED[i]));
 		}
 		EditorCachedLayerInfos_DEPRECATED.Empty();
 	}
@@ -1697,54 +1706,33 @@ void ALandscape::PreSave(const class ITargetPlatform* TargetPlatform)
 
 ALandscapeProxy* ULandscapeInfo::GetLandscapeProxyForLevel(ULevel* Level) const
 {
-	ALandscape* Landscape = LandscapeActor.Get();
-
-	if (Landscape && Landscape->GetLevel() == Level)
+	ALandscapeProxy* LandscapeProxy = nullptr;
+	ForAllLandscapeProxies([&LandscapeProxy, Level](ALandscapeProxy* Proxy)
 	{
-		return Landscape;
-	}
-
-	for (ALandscapeProxy* LandscapeProxy : Proxies)
-	{
-		if (LandscapeProxy && LandscapeProxy->GetLevel() == Level)
+		if (Proxy->GetLevel() == Level)
 		{
-			return LandscapeProxy;
+			LandscapeProxy = Proxy;
 		}
-	}
-
-	return nullptr;
+	});
+	return LandscapeProxy;
 }
 
 ALandscapeProxy* ULandscapeInfo::GetCurrentLevelLandscapeProxy(bool bRegistered) const
 {
-	ALandscape* Landscape = LandscapeActor.Get();
-
-	if (Landscape && (!bRegistered || Landscape->GetRootComponent()->IsRegistered()))
+	ALandscapeProxy* LandscapeProxy = nullptr;
+	ForAllLandscapeProxies([&LandscapeProxy, bRegistered](ALandscapeProxy* Proxy)
 	{
-		UWorld* LandscapeWorld = Landscape->GetWorld();
-		if (LandscapeWorld &&
-			LandscapeWorld->GetCurrentLevel() == Landscape->GetOuter())
+		if (!bRegistered || Proxy->GetRootComponent()->IsRegistered())
 		{
-			return Landscape;
-		}
-	}
-
-	for (auto It = Proxies.CreateConstIterator(); It; ++It)
-	{
-		ALandscapeProxy* LandscapeProxy = (*It);
-
-		if (LandscapeProxy && (!bRegistered || LandscapeProxy->GetRootComponent()->IsRegistered()))
-		{
-			UWorld* ProxyWorld = LandscapeProxy->GetWorld();
+			UWorld* ProxyWorld = Proxy->GetWorld();
 			if (ProxyWorld &&
-				ProxyWorld->GetCurrentLevel() == LandscapeProxy->GetOuter())
+				ProxyWorld->GetCurrentLevel() == Proxy->GetOuter())
 			{
-				return LandscapeProxy;
+				LandscapeProxy = Proxy;
 			}
 		}
-	}
-
-	return NULL;
+	});
+	return LandscapeProxy;
 }
 
 ALandscapeProxy* ULandscapeInfo::GetLandscapeProxy() const
@@ -1755,7 +1743,7 @@ ALandscapeProxy* ULandscapeInfo::GetLandscapeProxy() const
 
 	// prefer LandscapeActor in case it is loaded
 	ALandscape* Landscape = LandscapeActor.Get();
-	if (Landscape != NULL &&
+	if (Landscape != nullptr &&
 		Landscape->GetRootComponent()->IsRegistered())
 	{
 		return Landscape;
@@ -1763,7 +1751,7 @@ ALandscapeProxy* ULandscapeInfo::GetLandscapeProxy() const
 
 	// prefer current level proxy 
 	ALandscapeProxy* Proxy = GetCurrentLevelLandscapeProxy(true);
-	if (Proxy != NULL)
+	if (Proxy != nullptr)
 	{
 		return Proxy;
 	}
@@ -1772,14 +1760,28 @@ ALandscapeProxy* ULandscapeInfo::GetLandscapeProxy() const
 	for (auto It = Proxies.CreateConstIterator(); It; ++It)
 	{
 		Proxy = (*It);
-		if (Proxy != NULL &&
+		if (Proxy != nullptr &&
 			Proxy->GetRootComponent()->IsRegistered())
 		{
 			return Proxy;
 		}
 	}
 
-	return NULL;
+	return nullptr;
+}
+
+void ULandscapeInfo::ForAllLandscapeProxies(TFunctionRef<void(ALandscapeProxy*)> Fn) const
+{
+	ALandscape* Landscape = LandscapeActor.Get();
+	if (Landscape)
+	{
+		Fn(Landscape);
+	}
+
+	for (ALandscapeProxy* LandscapeProxy : Proxies)
+	{
+		Fn(LandscapeProxy);
+	}
 }
 
 void ULandscapeInfo::RegisterActor(ALandscapeProxy* Proxy, bool bMapCheck)
@@ -1853,14 +1855,19 @@ void ULandscapeInfo::UnregisterActor(ALandscapeProxy* Proxy)
 {
 	if (ALandscape* Landscape = Cast<ALandscape>(Proxy))
 	{
-		// LandscapeActor can be null if Unregister gets triggered twice, which has been observed to happen during redo
-		check(LandscapeActor.IsNull() || LandscapeActor.Get() == Landscape);
-		LandscapeActor = nullptr;
+		// Note: UnregisterActor sometimes gets triggered twice, e.g. it has been observed to happen during redo
+		// Note: In some cases LandscapeActor could be updated to a new landscape actor before the old landscape is unregistered/destroyed
+		// e.g. this has been observed when merging levels in the editor
+
+		if (LandscapeActor.Get() == Landscape)
+		{
+			LandscapeActor = nullptr;
+		}
 
 		// update proxies reference to landscape actor
 		for (ALandscapeStreamingProxy* StreamingProxy : Proxies)
 		{
-			StreamingProxy->LandscapeActor = nullptr;
+			StreamingProxy->LandscapeActor = LandscapeActor;
 		}
 	}
 	else
@@ -1873,7 +1880,11 @@ void ULandscapeInfo::UnregisterActor(ALandscapeProxy* Proxy)
 	// remove proxy components from the XY map
 	for (int32 CompIdx = 0; CompIdx < Proxy->LandscapeComponents.Num(); ++CompIdx)
 	{
-		UnregisterActorComponent(Proxy->LandscapeComponents[CompIdx]);
+		ULandscapeComponent* Component = Proxy->LandscapeComponents[CompIdx];
+		if (Component) // When a landscape actor is being GC'd it's possible the components were already GC'd and are null
+		{
+			UnregisterActorComponent(Component);
+		}
 	}
 	XYtoComponentMap.Compact();
 
@@ -1884,13 +1895,14 @@ void ULandscapeInfo::UnregisterActor(ALandscapeProxy* Proxy)
 void ULandscapeInfo::RegisterActorComponent(ULandscapeComponent* Component, bool bMapCheck)
 {
 	// Do not register components which are not part of the world
-	if (Component == NULL ||
+	if (Component == nullptr ||
 		Component->IsRegistered() == false)
 	{
 		return;
 	}
 
-	check(Component != NULL);
+	check(Component);
+
 	FIntPoint ComponentKey = Component->GetSectionBase() / Component->ComponentSizeQuads;
 	auto RegisteredComponent = XYtoComponentMap.FindRef(ComponentKey);
 
@@ -1938,18 +1950,19 @@ void ULandscapeInfo::RegisterActorComponent(ULandscapeComponent* Component, bool
 
 void ULandscapeInfo::UnregisterActorComponent(ULandscapeComponent* Component)
 {
-	check(Component != nullptr);
+	if (ensure(Component))
+	{
+		FIntPoint ComponentKey = Component->GetSectionBase() / Component->ComponentSizeQuads;
+		auto RegisteredComponent = XYtoComponentMap.FindRef(ComponentKey);
 
-	FIntPoint ComponentKey = Component->GetSectionBase() / Component->ComponentSizeQuads;
-	auto RegisteredComponent = XYtoComponentMap.FindRef(ComponentKey);
+		if (RegisteredComponent == Component)
+		{
+			XYtoComponentMap.Remove(ComponentKey);
+		}
 
-	if (RegisteredComponent == Component)
-	{
-		XYtoComponentMap.Remove(ComponentKey);
+		SelectedComponents.Remove(Component);
+		SelectedRegionComponents.Remove(Component);
 	}
-
-	SelectedComponents.Remove(Component);
-	SelectedRegionComponents.Remove(Component);
 }
 
 void ULandscapeInfo::Reset()
@@ -1969,7 +1982,7 @@ void ULandscapeInfo::FixupProxiesTransform()
 {
 	ALandscape* Landscape = LandscapeActor.Get();
 
-	if (Landscape == NULL ||
+	if (Landscape == nullptr ||
 		Landscape->GetRootComponent()->IsRegistered() == false)
 	{
 		return;
@@ -2015,21 +2028,13 @@ void ULandscapeInfo::FixupProxiesTransform()
 
 void ULandscapeInfo::UpdateComponentLayerWhitelist()
 {
-	if (LandscapeActor.IsValid())
-	{
-		for (ULandscapeComponent* Comp : LandscapeActor->LandscapeComponents)
-		{
-			Comp->UpdateLayerWhitelistFromPaintedLayers();
-		}
-	}
-
-	for (ALandscapeProxy* Proxy : Proxies)
+	ForAllLandscapeProxies([](ALandscapeProxy* Proxy)
 	{
 		for (ULandscapeComponent* Comp : Proxy->LandscapeComponents)
 		{
 			Comp->UpdateLayerWhitelistFromPaintedLayers();
 		}
-	}
+	});
 }
 
 void ULandscapeInfo::RecreateLandscapeInfo(UWorld* InWorld, bool bMapCheck)
diff --git a/Engine/Source/Runtime/Landscape/Private/LandscapeEdit.cpp b/Engine/Source/Runtime/Landscape/Private/LandscapeEdit.cpp
--- a/Engine/Source/Runtime/Landscape/Private/LandscapeEdit.cpp
+++ b/Engine/Source/Runtime/Landscape/Private/LandscapeEdit.cpp
@@ -39,7 +39,9 @@ LandscapeEdit.cpp: Landscape editing
 #include "SNotificationList.h"
 #include "NotificationManager.h"
 #include "LandscapeEditorModule.h"
+#include "ComponentRecreateRenderStateContext.h"
 #endif
+#include "Algo/Count.h"
 
 DEFINE_LOG_CATEGORY(LogLandscape);
 
@@ -217,7 +219,7 @@ UMaterialInstanceConstant* ULandscapeComponent::GetCombinationMaterial(bool bMob
 	return nullptr;
 }
 
-void ULandscapeComponent::UpdateMaterialInstances()
+void ULandscapeComponent::UpdateMaterialInstances_Internal(FMaterialUpdateContext& Context)
 {
 	check(GIsEditor);
 
@@ -226,15 +228,6 @@ void ULandscapeComponent::UpdateMaterialInstances()
 
 	if (CombinationMaterialInstance != nullptr)
 	{
-		// not having the context recreate the render state because we will manually do it for only this component
-		FMaterialUpdateContext Context(FMaterialUpdateContext::EOptions::Default & ~FMaterialUpdateContext::EOptions::RecreateRenderStates);
-
-		if (bRenderStateCreated)
-		{
-			DestroyRenderState_Concurrent();
-			FlushRenderingCommands();
-		}
-
 		UMaterialInstanceConstant* MaterialInstance = MaterialInstances[0];
 
 		// Create the instance for this component, that will use the layer combination instance.
@@ -270,14 +263,10 @@ void ULandscapeComponent::UpdateMaterialInstances()
 		// Set the weightmaps
 		for (int32 i = 0; i < WeightmapTextures.Num(); i++)
 		{
-			// gmartin: Trying to locate UE-23902
-			checkf(!WeightmapTextures[i] || WeightmapTextures[i]->IsValidLowLevel(), TEXT("Texture not valid! UE-23902! Parameter (Weightmap%d)"), i);
-
-			// UE_LOG(LogLandscape, Log, TEXT("Setting Weightmap%d = %s"), i, *WeightmapTextures(i)->GetName());
 			MaterialInstance->SetTextureParameterValueEditorOnly(FName(*FString::Printf(TEXT("Weightmap%d"), i)), WeightmapTextures[i]);
 		}
-		// Set the heightmap, if needed.
 
+		// Set the heightmap, if needed.
 		if (HeightmapTexture)
 		{
 			MaterialInstance->SetTextureParameterValueEditorOnly(FName(TEXT("Heightmap")), HeightmapTexture);
@@ -309,9 +298,48 @@ void ULandscapeComponent::UpdateMaterialInstances()
 		MaterialInstances.Empty(1);
 		MaterialInstances.Add(nullptr);
 	}
+}
+
+void ULandscapeComponent::UpdateMaterialInstances()
+{
+	// we're not having the material update context recreate the render state because we will manually do it for only this component
+	TOptional<FComponentRecreateRenderStateContext> RecreateRenderStateContext;
+	RecreateRenderStateContext.Emplace(this);
+	TOptional<FMaterialUpdateContext> MaterialUpdateContext;
+	MaterialUpdateContext.Emplace(FMaterialUpdateContext::EOptions::Default & ~FMaterialUpdateContext::EOptions::RecreateRenderStates);
+
+	UpdateMaterialInstances_Internal(MaterialUpdateContext.GetValue());
+
+	// End material update
+	MaterialUpdateContext.Reset();
+
+	// Recreate the render state for this component, needed to update the static drawlist which has cached the MaterialRenderProxies
+	// Must be after the FMaterialUpdateContext is destroyed
+	RecreateRenderStateContext.Reset();
+}
+
+void ALandscapeProxy::UpdateAllComponentMaterialInstances()
+{
+	// we're not having the material update context recreate render states because we will manually do it for only our components
+	TArray<FComponentRecreateRenderStateContext> RecreateRenderStateContexts;
+	for (ULandscapeComponent* Component : LandscapeComponents)
+	{
+		RecreateRenderStateContexts.Emplace(Component);
+	}
+	TOptional<FMaterialUpdateContext> MaterialUpdateContext;
+	MaterialUpdateContext.Emplace(FMaterialUpdateContext::EOptions::Default & ~FMaterialUpdateContext::EOptions::RecreateRenderStates);
 
-	// Recreate the render state, needed to update the static drawlist which has cached the MaterialRenderProxy.
-	RecreateRenderState_Concurrent();
+	for (ULandscapeComponent* Component : LandscapeComponents)
+	{
+		Component->UpdateMaterialInstances_Internal(MaterialUpdateContext.GetValue());
+	}
+
+	// End material update
+	MaterialUpdateContext.Reset();
+
+	// Recreate the render state for our components, needed to update the static drawlist which has cached the MaterialRenderProxies
+	// Must be after the FMaterialUpdateContext is destroyed
+	RecreateRenderStateContexts.Empty();
 }
 
 int32 ULandscapeComponent::GetNumMaterials() const
@@ -372,7 +400,10 @@ void ULandscapeComponent::PreFeatureLevelChange(ERHIFeatureLevel::Type PendingFe
 
 void ULandscapeComponent::PostEditUndo()
 {
-	UpdateMaterialInstances();
+	if (!IsPendingKill())
+	{
+		UpdateMaterialInstances();
+	}
 
 	Super::PostEditUndo();
 
@@ -2460,9 +2491,6 @@ LANDSCAPE_API void ALandscapeProxy::Import(
 			}
 
 			LandscapeComponent->CachedLocalBox = LocalBox;
-
-			// Update MaterialInstance
-			LandscapeComponent->UpdateMaterialInstances();
 		}
 	}
 
@@ -2541,6 +2569,9 @@ LANDSCAPE_API void ALandscapeProxy::Import(
 		HeightmapInfo.HeightmapTexture->PostEditChange();
 	}
 
+	// Update MaterialInstances (must be done after textures are fully initialized)
+	UpdateAllComponentMaterialInstances();
+
 	if (GetLevel()->bIsVisible)
 	{
 		// Update our new components
@@ -3079,29 +3110,24 @@ void ULandscapeInfo::DeleteLayer(ULandscapeLayerInfoObject* LayerInfo)
 	FLandscapeEditDataInterface LandscapeEdit(this);
 	LandscapeEdit.DeleteLayer(LayerInfo);
 
-	// Remove from array
-	for (int32 j = 0; j < Layers.Num(); j++)
+	// Remove from layer settings array
 	{
-		if (Layers[j].LayerInfoObj && Layers[j].LayerInfoObj == LayerInfo)
+		int32 LayerIndex = Layers.IndexOfByPredicate([LayerInfo](const FLandscapeInfoLayerSettings& LayerSettings) { return LayerSettings.LayerInfoObj == LayerInfo; });
+		if (LayerIndex != INDEX_NONE)
 		{
-			Layers.RemoveAt(j);
-			break;
+			Layers.RemoveAt(LayerIndex);
 		}
 	}
 
-	ALandscape* Landscape = LandscapeActor.Get();
-	if (Landscape != nullptr)
-	{
-		Landscape->Modify();
-		Landscape->EditorLayerSettings.Remove(LayerInfo);
-	}
-
-	for (auto It = Proxies.CreateConstIterator(); It; ++It)
+	ForAllLandscapeProxies([LayerInfo](ALandscapeProxy* Proxy)
 	{
-		ALandscapeProxy* Proxy = *It;
 		Proxy->Modify();
-		Proxy->EditorLayerSettings.Remove(LayerInfo);
-	}
+		int32 Index = Proxy->EditorLayerSettings.IndexOfByKey(LayerInfo);
+		if (Index != INDEX_NONE)
+		{
+			Proxy->EditorLayerSettings.RemoveAt(Index);
+		}
+	});
 
 	//UpdateLayerInfoMap();
 
@@ -3121,47 +3147,24 @@ void ULandscapeInfo::ReplaceLayer(ULandscapeLayerInfoObject* FromLayerInfo, ULan
 		// Convert array
 		for (int32 j = 0; j < Layers.Num(); j++)
 		{
-			if (Layers[j].LayerInfoObj && Layers[j].LayerInfoObj == FromLayerInfo)
+			if (Layers[j].LayerInfoObj == FromLayerInfo)
 			{
 				Layers[j].LayerInfoObj = ToLayerInfo;
 			}
 		}
 
-		ALandscape* Landscape = LandscapeActor.Get();
-		if (Landscape != nullptr)
-		{
-			Landscape->Modify();
-			FLandscapeEditorLayerSettings* ToEditorLayerSettings = Landscape->EditorLayerSettings.FindByKey(ToLayerInfo);
-			if (ToEditorLayerSettings != nullptr)
-			{
-				// If the new layer already exists, simple remove the old layer
-				Landscape->EditorLayerSettings.Remove(FromLayerInfo);
-			}
-			else
-			{
-				FLandscapeEditorLayerSettings* FromEditorLayerSettings = Landscape->EditorLayerSettings.FindByKey(FromLayerInfo);
-				if (FromEditorLayerSettings != nullptr)
-				{
-					// If only the old layer exists (most common case), change it to point to the new layer info
-					FromEditorLayerSettings->LayerInfoObj = ToLayerInfo;
-				}
-				else
-				{
-					// If neither exists in the EditorLayerSettings cache, add it
-					Landscape->EditorLayerSettings.Add(ToLayerInfo);
-				}
-			}
-		}
-
-		for (auto It = Proxies.CreateConstIterator(); It; ++It)
+		ForAllLandscapeProxies([FromLayerInfo, ToLayerInfo](ALandscapeProxy* Proxy)
 		{
-			ALandscapeProxy* Proxy = *It;
 			Proxy->Modify();
 			FLandscapeEditorLayerSettings* ToEditorLayerSettings = Proxy->EditorLayerSettings.FindByKey(ToLayerInfo);
 			if (ToEditorLayerSettings != nullptr)
 			{
 				// If the new layer already exists, simple remove the old layer
-				Proxy->EditorLayerSettings.Remove(FromLayerInfo);
+				int32 Index = Proxy->EditorLayerSettings.IndexOfByKey(FromLayerInfo);
+				if (Index != INDEX_NONE)
+				{
+					Proxy->EditorLayerSettings.RemoveAt(Index);
+				}
 			}
 			else
 			{
@@ -3174,10 +3177,10 @@ void ULandscapeInfo::ReplaceLayer(ULandscapeLayerInfoObject* FromLayerInfo, ULan
 				else
 				{
 					// If neither exists in the EditorLayerSettings cache, add it
-					Proxy->EditorLayerSettings.Add(ToLayerInfo);
+					Proxy->EditorLayerSettings.Add(FLandscapeEditorLayerSettings(ToLayerInfo));
 				}
 			}
-		}
+		});
 
 		//UpdateLayerInfoMap();
 
@@ -3314,6 +3317,17 @@ void ALandscape::PostEditImport()
 
 	Super::PostEditImport();
 }
+
+void ALandscape::PostDuplicate(bool bDuplicateForPIE)
+{
+	if (!bDuplicateForPIE)
+	{
+		// Need to generate new GUID when duplicating
+		LandscapeGuid = FGuid::NewGuid();
+	}
+
+	Super::PostDuplicate(bDuplicateForPIE);
+}
 #endif	//WITH_EDITOR
 
 ULandscapeLayerInfoObject::ULandscapeLayerInfoObject(const FObjectInitializer& ObjectInitializer)
@@ -3478,46 +3492,28 @@ void ALandscapeProxy::RecreateCollisionComponents()
 
 void ULandscapeInfo::RecreateCollisionComponents()
 {
-	if (LandscapeActor.IsValid())
-	{
-		LandscapeActor->RecreateCollisionComponents();
-	}
-
-	for (auto It = Proxies.CreateConstIterator(); It; ++It)
+	ForAllLandscapeProxies([](ALandscapeProxy* Proxy)
 	{
-		ALandscapeProxy* Proxy = (*It);
 		Proxy->RecreateCollisionComponents();
-	}
+	});
 }
 
 void ULandscapeInfo::RemoveXYOffsets()
 {
-	if (LandscapeActor.IsValid())
-	{
-		LandscapeActor->RemoveXYOffsets();
-	}
-
-	for (auto It = Proxies.CreateConstIterator(); It; ++It)
+	ForAllLandscapeProxies([](ALandscapeProxy* Proxy)
 	{
-		ALandscapeProxy* Proxy = (*It);
 		Proxy->RemoveXYOffsets();
-	}
+	});
 }
 
 void ULandscapeInfo::PostponeTextureBaking()
 {
-	const int32 PostponeValue = 60; //frames
+	static const int32 PostponeValue = 60; //frames
 	
-	ALandscape* Landscape = LandscapeActor.Get();
-	if (Landscape)
-	{
-		Landscape->UpdateBakedTexturesCountdown = PostponeValue;
-	}
-
-	for (ALandscapeProxy* Proxy : Proxies)
+	ForAllLandscapeProxies([](ALandscapeProxy* Proxy)
 	{
 		Proxy->UpdateBakedTexturesCountdown = PostponeValue;
-	}
+	});
 }
 
 namespace
@@ -3675,6 +3671,33 @@ void ALandscapeProxy::PostEditChangeProperty(FPropertyChangedEvent& PropertyChan
 		}
 	}
 
+	// Remove null layer infos
+	EditorLayerSettings.RemoveAll([](const FLandscapeEditorLayerSettings& Entry) { return Entry.LayerInfoObj == nullptr; });
+
+	ULandscapeInfo* Info = GetLandscapeInfo();
+	bool bRemovedAnyLayers = false;
+	for (ULandscapeComponent* Component : LandscapeComponents)
+	{
+		int32 NumNullLayers = Algo::CountIf(Component->WeightmapLayerAllocations, [](const FWeightmapLayerAllocationInfo& Allocation) { return Allocation.LayerInfo == nullptr; });
+		if (NumNullLayers > 0)
+		{
+			FLandscapeEditDataInterface LandscapeEdit(Info);
+			for (int32 i = 0; i < NumNullLayers; ++i)
+			{
+				// DeleteLayer doesn't expect duplicates, so we need to call it once for each null
+				Component->DeleteLayer(nullptr, LandscapeEdit);
+			}
+			bRemovedAnyLayers = true;
+		}
+	}
+	if (bRemovedAnyLayers)
+	{
+		// Flush dynamic data (e.g. grass)
+		TSet<ULandscapeComponent*> Components;
+		Components.Append(LandscapeComponents);
+		ALandscapeProxy::InvalidateGeneratedComponentData(Components);
+	}
+
 	// Must do this *after* correcting the scale or reattaching the landscape components will crash!
 	// Must do this *after* clamping values / propogating values to components
 	Super::PostEditChangeProperty(PropertyChangedEvent);
@@ -3718,15 +3741,7 @@ void ALandscapeStreamingProxy::PostEditChangeProperty(FPropertyChangedEvent& Pro
 			MaterialInstanceConstantMap.Empty();
 		}
 
-		for (int32 ComponentIndex = 0; ComponentIndex < LandscapeComponents.Num(); ComponentIndex++)
-		{
-			ULandscapeComponent* Comp = LandscapeComponents[ComponentIndex];
-			if (Comp)
-			{
-				// Update the MIC
-				Comp->UpdateMaterialInstances();
-			}
-		}
+		UpdateAllComponentMaterialInstances();
 	}
 
 	// Must do this *after* clamping values
@@ -3860,12 +3875,6 @@ void ALandscape::PostEditChangeProperty(FPropertyChangedEvent& PropertyChangedEv
 						Comp->UpdateBounds();
 					}
 
-					if (ChangedMaterial)
-					{
-						// Update the MIC
-						Comp->UpdateMaterialInstances();
-					}
-
 					if (bChangedLighting)
 					{
 						Comp->InvalidateLightingCache();
@@ -3877,6 +3886,11 @@ void ALandscape::PostEditChangeProperty(FPropertyChangedEvent& PropertyChangedEv
 					}
 				}
 			}
+
+			if (ChangedMaterial)
+			{
+				UpdateAllComponentMaterialInstances();
+			}
 		}
 
 		// Need to update Gizmo scene proxy
diff --git a/Engine/Source/Runtime/Landscape/Private/LandscapeRender.cpp b/Engine/Source/Runtime/Landscape/Private/LandscapeRender.cpp
--- a/Engine/Source/Runtime/Landscape/Private/LandscapeRender.cpp
+++ b/Engine/Source/Runtime/Landscape/Private/LandscapeRender.cpp
@@ -2594,8 +2594,8 @@ ULandscapeMaterialInstanceConstant::ULandscapeMaterialInstanceConstant(const FOb
 
 class FLandscapeMaterialResource : public FMaterialResource
 {
-	bool bIsLayerThumbnail;
-	bool bDisableTessellation;
+	const bool bIsLayerThumbnail;
+	const bool bDisableTessellation;
 
 public:
 	FLandscapeMaterialResource(ULandscapeMaterialInstanceConstant* Parent)
@@ -2604,6 +2604,23 @@ public:
 	{
 	}
 
+	void GetShaderMapId(EShaderPlatform Platform, FMaterialShaderMapId& OutId) const override
+	{
+		FMaterialResource::GetShaderMapId(Platform, OutId);
+
+		if (bIsLayerThumbnail || bDisableTessellation)
+		{
+			FSHA1 Hash;
+			Hash.Update(OutId.BasePropertyOverridesHash.Hash, ARRAY_COUNT(OutId.BasePropertyOverridesHash.Hash));
+
+			const FString HashString = TEXT("bOverride_TessellationMode");
+			Hash.UpdateWithString(*HashString, HashString.Len());
+
+			Hash.Final();
+			Hash.GetHash(OutId.BasePropertyOverridesHash.Hash);
+		}
+	}
+
 	bool IsUsedWithLandscape() const override
 	{
 		return !bIsLayerThumbnail;
diff --git a/Engine/Source/Runtime/Landscape/Private/LandscapeSplineRaster.cpp b/Engine/Source/Runtime/Landscape/Private/LandscapeSplineRaster.cpp
--- a/Engine/Source/Runtime/Landscape/Private/LandscapeSplineRaster.cpp
+++ b/Engine/Source/Runtime/Landscape/Private/LandscapeSplineRaster.cpp
@@ -7,8 +7,6 @@
 
 #include "LandscapePrivatePCH.h"
 #include "LandscapeSplineRaster.h"
-#include "Landscape.h"
-#include "LandscapeStreamingProxy.h"
 #include "LandscapeDataAccess.h"
 #include "LandscapeEdit.h"
 #include "LandscapeHeightfieldCollisionComponent.h"
@@ -421,14 +419,10 @@ bool ULandscapeInfo::ApplySplines(bool bOnlySelected)
 {
 	bool bResult = false;
 
-	ALandscape* Landscape = LandscapeActor.Get();
-
-	bResult |= ApplySplinesInternal(bOnlySelected, Landscape);
-
-	for (ALandscapeProxy* LandscapeProxy : Proxies)
+	ForAllLandscapeProxies([&bResult, bOnlySelected, this](ALandscapeProxy* Proxy)
 	{
-		bResult |= ApplySplinesInternal(bOnlySelected, LandscapeProxy);
-	}
+		bResult |= ApplySplinesInternal(bOnlySelected, Proxy);
+	});
 
 	return bResult;
 }
diff --git a/Engine/Source/Runtime/Landscape/Private/LandscapeSplines.cpp b/Engine/Source/Runtime/Landscape/Private/LandscapeSplines.cpp
--- a/Engine/Source/Runtime/Landscape/Private/LandscapeSplines.cpp
+++ b/Engine/Source/Runtime/Landscape/Private/LandscapeSplines.cpp
@@ -5,8 +5,6 @@
 =============================================================================*/
 
 #include "LandscapePrivatePCH.h"
-#include "Landscape.h"
-#include "LandscapeStreamingProxy.h"
 #include "Components/SplineMeshComponent.h"
 #include "ShaderParameterUtils.h"
 #include "StaticMeshResources.h"
@@ -846,21 +844,13 @@ TArray<ULandscapeSplinesComponent*> ULandscapeSplinesComponent::GetAllStreamingS
 		if (LandscapeInfo)
 		{
 			TArray<ULandscapeSplinesComponent*> SplinesComponents;
-			SplinesComponents.Reserve(LandscapeInfo->Proxies.Num() + 1);
-
-			ALandscape* RootLandscape = LandscapeInfo->LandscapeActor.Get();
-			if (RootLandscape && RootLandscape->SplineComponent)
-			{
-				SplinesComponents.Add(RootLandscape->SplineComponent);
-			}
-			for (ALandscapeProxy* LandscapeProxy : LandscapeInfo->Proxies)
+			LandscapeInfo->ForAllLandscapeProxies([&SplinesComponents](ALandscapeProxy* Proxy)
 			{
-				if (LandscapeProxy && LandscapeProxy->SplineComponent)
+				if (Proxy->SplineComponent)
 				{
-					SplinesComponents.Add(LandscapeProxy->SplineComponent);
+					SplinesComponents.Add(Proxy->SplineComponent);
 				}
-			}
-
+			});
 			return SplinesComponents;
 		}
 	}
diff --git a/Engine/Source/Runtime/MediaAssets/Private/Assets/MediaSoundWave.cpp b/Engine/Source/Runtime/MediaAssets/Private/Assets/MediaSoundWave.cpp
--- a/Engine/Source/Runtime/MediaAssets/Private/Assets/MediaSoundWave.cpp
+++ b/Engine/Source/Runtime/MediaAssets/Private/Assets/MediaSoundWave.cpp
@@ -168,9 +168,9 @@ void UMediaSoundWave::GetAssetRegistryTags(TArray<FAssetRegistryTag>& OutTags) c
 }
 
 
-SIZE_T UMediaSoundWave::GetResourceSize(EResourceSizeMode::Type Mode)
+void UMediaSoundWave::GetResourceSizeEx(FResourceSizeEx& CumulativeResourceSize)
 {
-	return 0; // procedural sound
+	Super::GetResourceSizeEx(CumulativeResourceSize);
 }
 
 
diff --git a/Engine/Source/Runtime/MediaAssets/Private/Assets/MediaTexture.cpp b/Engine/Source/Runtime/MediaAssets/Private/Assets/MediaTexture.cpp
--- a/Engine/Source/Runtime/MediaAssets/Private/Assets/MediaTexture.cpp
+++ b/Engine/Source/Runtime/MediaAssets/Private/Assets/MediaTexture.cpp
@@ -206,9 +206,14 @@ FString UMediaTexture::GetDesc()
 }
 
 
-SIZE_T UMediaTexture::GetResourceSize(EResourceSizeMode::Type Mode)
+void UMediaTexture::GetResourceSizeEx(FResourceSizeEx& CumulativeResourceSize)
 {
-	return (Resource != nullptr) ? ((FMediaTextureResource*)Resource)->GetResourceSize() : 0;
+	Super::GetResourceSizeEx(CumulativeResourceSize);
+
+	if (Resource)
+	{
+		((FMediaTextureResource*)Resource)->GetResourceSizeEx(CumulativeResourceSize);
+	}
 }
 
 
diff --git a/Engine/Source/Runtime/MediaAssets/Private/Misc/MediaTextureResource.cpp b/Engine/Source/Runtime/MediaAssets/Private/Misc/MediaTextureResource.cpp
--- a/Engine/Source/Runtime/MediaAssets/Private/Misc/MediaTextureResource.cpp
+++ b/Engine/Source/Runtime/MediaAssets/Private/Misc/MediaTextureResource.cpp
@@ -14,7 +14,7 @@ FMediaTextureResource::FMediaTextureResource(UMediaTexture& InOwner, const FLine
 	, BufferBytesPerPixel(0)
 	, BufferClearColor(InClearColor)
 	, BufferDimensions(FIntPoint::ZeroValue)
-	, CachedResourceSize(0)
+	, CachedResourceSize()
 	, OutputDimensions(InOutputDimensions)
 	, RequiresConversion(false)
 	, SinkFormat(InSinkFormat)
@@ -530,7 +530,7 @@ bool FMediaTextureResource::Tick(float DeltaTime)
 
 void FMediaTextureResource::CacheResourceSize()
 {
-	SIZE_T ResourceSize = 0;
+	FResourceSizeEx ResourceSize;
 
 	if (OutputDimensions.GetMin() > 0)
 	{
@@ -538,16 +538,16 @@ void FMediaTextureResource::CacheResourceSize()
 
 		if (SinkMode == EMediaTextureSinkMode::Buffered)
 		{
-			ResourceSize += 3 * BufferSize;
+			ResourceSize.AddUnknownMemoryBytes(3 * BufferSize);
 		}
 		else
 		{
-			ResourceSize += BufferSize;
+			ResourceSize.AddUnknownMemoryBytes(BufferSize);
 		}
 
 		if (RequiresConversion)
 		{
-			ResourceSize += OutputDimensions.X * OutputDimensions.Y * 4;
+			ResourceSize.AddUnknownMemoryBytes(OutputDimensions.X * OutputDimensions.Y * 4);
 		}
 	}
 
diff --git a/Engine/Source/Runtime/MediaAssets/Private/Misc/MediaTextureResource.h b/Engine/Source/Runtime/MediaAssets/Private/Misc/MediaTextureResource.h
--- a/Engine/Source/Runtime/MediaAssets/Private/Misc/MediaTextureResource.h
+++ b/Engine/Source/Runtime/MediaAssets/Private/Misc/MediaTextureResource.h
@@ -60,9 +60,20 @@ public:
 	 *
 	 * @return Resource size (in bytes).
 	 */
+	DEPRECATED(4.14, "GetResourceSize is deprecated. Please use GetResourceSizeEx or GetResourceSizeBytes instead.")
 	SIZE_T GetResourceSize() const
 	{
-		return CachedResourceSize;
+		return GetResourceSizeBytes();
+	}
+
+	void GetResourceSizeEx(FResourceSizeEx& CumulativeResourceSize) const
+	{
+		CumulativeResourceSize += CachedResourceSize;
+	}
+
+	SIZE_T GetResourceSizeBytes() const
+	{
+		return CachedResourceSize.GetTotalMemoryBytes();
 	}
 
 	/**
@@ -208,8 +219,8 @@ private:
 	 */
 	FResource BufferResources[3];
 
-	/** Total size of this resource (in bytes) .*/
-	SIZE_T CachedResourceSize;
+	/** Total size of this resource.*/
+	FResourceSizeEx CachedResourceSize;
 
 	/** Width and height of the output resource (in pixels). */
 	FIntPoint OutputDimensions;
diff --git a/Engine/Source/Runtime/Renderer/Private/BasePassRendering.cpp b/Engine/Source/Runtime/Renderer/Private/BasePassRendering.cpp
--- a/Engine/Source/Runtime/Renderer/Private/BasePassRendering.cpp
+++ b/Engine/Source/Runtime/Renderer/Private/BasePassRendering.cpp
@@ -1097,13 +1097,13 @@ void FDeferredShadingSceneRenderer::RenderEditorPrimitives(FRHICommandList& RHIC
 		bDirty |= DrawViewElements<FBasePassOpaqueDrawingPolicyFactory>(RHICmdList, View, FBasePassOpaqueDrawingPolicyFactory::ContextType(false, ESceneRenderTargetsMode::DontSet), SDPG_World, true) || bDirty;
 
 		// Draw the view's batched simple elements(lines, sprites, etc).
-		bDirty |= View.BatchedViewElements.Draw(RHICmdList, FeatureLevel, bNeedToSwitchVerticalAxis, View.ViewProjectionMatrix, View.ViewRect.Width(), View.ViewRect.Height(), false) || bDirty;
+		bDirty |= View.BatchedViewElements.Draw(RHICmdList, FeatureLevel, bNeedToSwitchVerticalAxis, View, false) || bDirty;
 
 		// Draw foreground objects last
 		bDirty |= DrawViewElements<FBasePassOpaqueDrawingPolicyFactory>(RHICmdList, View, FBasePassOpaqueDrawingPolicyFactory::ContextType(false, ESceneRenderTargetsMode::DontSet), SDPG_Foreground, true) || bDirty;
 
 		// Draw the view's batched simple elements(lines, sprites, etc).
-		bDirty |= View.TopBatchedViewElements.Draw(RHICmdList, FeatureLevel, bNeedToSwitchVerticalAxis, View.ViewProjectionMatrix, View.ViewRect.Width(), View.ViewRect.Height(), false) || bDirty;
+		bDirty |= View.TopBatchedViewElements.Draw(RHICmdList, FeatureLevel, bNeedToSwitchVerticalAxis, View, false) || bDirty;
 
 	}
 
diff --git a/Engine/Source/Runtime/Renderer/Private/MobileBasePassRendering.cpp b/Engine/Source/Runtime/Renderer/Private/MobileBasePassRendering.cpp
--- a/Engine/Source/Runtime/Renderer/Private/MobileBasePassRendering.cpp
+++ b/Engine/Source/Runtime/Renderer/Private/MobileBasePassRendering.cpp
@@ -711,13 +711,13 @@ void FMobileSceneRenderer::RenderMobileBasePass(FRHICommandListImmediate& RHICmd
 				DrawViewElements<FMobileBasePassOpaqueDrawingPolicyFactory>(RHICmdList, View, FMobileBasePassOpaqueDrawingPolicyFactory::ContextType(false, ESceneRenderTargetsMode::DontSet), SDPG_World, true);
 
 				// Draw the view's batched simple elements(lines, sprites, etc).
-				View.BatchedViewElements.Draw(RHICmdList, FeatureLevel, bNeedToSwitchVerticalAxis, View.ViewProjectionMatrix, View.ViewRect.Width(), View.ViewRect.Height(), false);
+				View.BatchedViewElements.Draw(RHICmdList, FeatureLevel, bNeedToSwitchVerticalAxis, View, false);
 
 				// Draw foreground objects last
 				DrawViewElements<FMobileBasePassOpaqueDrawingPolicyFactory>(RHICmdList, View, FMobileBasePassOpaqueDrawingPolicyFactory::ContextType(false, ESceneRenderTargetsMode::DontSet), SDPG_Foreground, true);
 
 				// Draw the view's batched simple elements(lines, sprites, etc).
-				View.TopBatchedViewElements.Draw(RHICmdList, FeatureLevel, bNeedToSwitchVerticalAxis, View.ViewProjectionMatrix, View.ViewRect.Width(), View.ViewRect.Height(), false);
+				View.TopBatchedViewElements.Draw(RHICmdList, FeatureLevel, bNeedToSwitchVerticalAxis, View, false);
 			}
 		}
 
diff --git a/Engine/Source/Runtime/Renderer/Private/PostProcess/PostProcessCompositeEditorPrimitives.cpp b/Engine/Source/Runtime/Renderer/Private/PostProcess/PostProcessCompositeEditorPrimitives.cpp
--- a/Engine/Source/Runtime/Renderer/Private/PostProcess/PostProcessCompositeEditorPrimitives.cpp
+++ b/Engine/Source/Runtime/Renderer/Private/PostProcess/PostProcessCompositeEditorPrimitives.cpp
@@ -347,7 +347,7 @@ void FRCPassPostProcessCompositeEditorPrimitives::RenderPrimitivesToComposite(FR
 	DrawViewElements<TBasePass>(RHICmdList, View, typename TBasePass::ContextType(bDepthTest, ESceneRenderTargetsMode::SetTextures), SDPG_World, false);
 
 	// Draw the view's batched simple elements(lines, sprites, etc).
-	View.BatchedViewElements.Draw(RHICmdList, FeatureLevel, bNeedToSwitchVerticalAxis, View.ViewProjectionMatrix, View.ViewRect.Width(), View.ViewRect.Height(), false, 1.0f, &View, SceneDepth);
+	View.BatchedViewElements.Draw(RHICmdList, FeatureLevel, bNeedToSwitchVerticalAxis, View, false, 1.0f, SceneDepth);
 
 	// Draw foreground objects. Draw twice, once without depth testing to bring them into the foreground and again to depth test against themselves
 	{
@@ -358,13 +358,13 @@ void FRCPassPostProcessCompositeEditorPrimitives::RenderPrimitivesToComposite(FR
 
 		DrawViewElements<TBasePass>(RHICmdList, View, typename TBasePass::ContextType(bDepthTest, ESceneRenderTargetsMode::SetTextures), SDPG_Foreground, false);
 
-		View.TopBatchedViewElements.Draw(RHICmdList, FeatureLevel, bNeedToSwitchVerticalAxis, View.ViewProjectionMatrix, View.ViewRect.Width(), View.ViewRect.Height(), false);
+		View.TopBatchedViewElements.Draw(RHICmdList, FeatureLevel, bNeedToSwitchVerticalAxis, View, false);
 
 		RHICmdList.SetDepthStencilState(TStaticDepthStencilState<true, CF_DepthNearOrEqual>::GetRHI());
 
 		DrawViewElements<TBasePass>(RHICmdList, View, typename TBasePass::ContextType(bDepthTest, ESceneRenderTargetsMode::SetTextures), SDPG_Foreground, false);
 
-		View.TopBatchedViewElements.Draw(RHICmdList, FeatureLevel, bNeedToSwitchVerticalAxis, View.ViewProjectionMatrix, View.ViewRect.Width(), View.ViewRect.Height(), false);
+		View.TopBatchedViewElements.Draw(RHICmdList, FeatureLevel, bNeedToSwitchVerticalAxis, View, false);
 	}
 }
 
diff --git a/Engine/Source/Runtime/Renderer/Private/SceneHitProxyRendering.cpp b/Engine/Source/Runtime/Renderer/Private/SceneHitProxyRendering.cpp
--- a/Engine/Source/Runtime/Renderer/Private/SceneHitProxyRendering.cpp
+++ b/Engine/Source/Runtime/Renderer/Private/SceneHitProxyRendering.cpp
@@ -509,7 +509,7 @@ void RenderHitProxies(FRHICommandListImmediate& RHICmdList, const FSceneRenderer
 		DrawViewElements<FHitProxyDrawingPolicyFactory>(RHICmdList, View, FHitProxyDrawingPolicyFactory::ContextType(), SDPG_World, bPreFog);
 
 		// Draw the view's batched simple elements(lines, sprites, etc).
-		View.BatchedViewElements.Draw(RHICmdList, FeatureLevel, bNeedToSwitchVerticalAxis, View.ViewProjectionMatrix, View.ViewRect.Width(), View.ViewRect.Height(), true);
+		View.BatchedViewElements.Draw(RHICmdList, FeatureLevel, bNeedToSwitchVerticalAxis, View, true);
 
 		// Some elements should never be occluded (e.g. gizmos).
 		// So we render those twice, first to overwrite potentially nearer objects,
@@ -518,13 +518,13 @@ void RenderHitProxies(FRHICommandListImmediate& RHICmdList, const FSceneRenderer
 		// Draw the view's foreground elements last.
 		DrawViewElements<FHitProxyDrawingPolicyFactory>(RHICmdList, View, FHitProxyDrawingPolicyFactory::ContextType(), SDPG_Foreground, bPreFog);
 
-		View.TopBatchedViewElements.Draw(RHICmdList, FeatureLevel, bNeedToSwitchVerticalAxis, View.ViewProjectionMatrix, View.ViewRect.Width(), View.ViewRect.Height(), true);
+		View.TopBatchedViewElements.Draw(RHICmdList, FeatureLevel, bNeedToSwitchVerticalAxis, View, true);
 
 		RHICmdList.SetDepthStencilState(TStaticDepthStencilState<true, CF_DepthNearOrEqual>::GetRHI());
 		// Draw the view's foreground elements last.
 		DrawViewElements<FHitProxyDrawingPolicyFactory>(RHICmdList, View, FHitProxyDrawingPolicyFactory::ContextType(), SDPG_Foreground, bPreFog);
 
-		View.TopBatchedViewElements.Draw(RHICmdList, FeatureLevel, bNeedToSwitchVerticalAxis, View.ViewProjectionMatrix, View.ViewRect.Width(), View.ViewRect.Height(), true);
+		View.TopBatchedViewElements.Draw(RHICmdList, FeatureLevel, bNeedToSwitchVerticalAxis, View, true);
 	}
 
 	// Finish drawing to the hit proxy render target.
@@ -588,9 +588,7 @@ void RenderHitProxies(FRHICommandListImmediate& RHICmdList, const FSceneRenderer
 				RHICmdList,
 				FeatureLevel,
 				bNeedToSwitchVerticalAxis,
-				PixelToView,
-				RenderTargetSize.X,
-				RenderTargetSize.Y,
+				FBatchedElements::CreateProxySceneView(PixelToView, FIntRect(0, 0, RenderTargetSize.X, RenderTargetSize.Y)),
 				false,
 				1.0f
 				);
diff --git a/Engine/Source/Runtime/RuntimeAssetCache/Private/RuntimeAssetCacheBuilders.cpp b/Engine/Source/Runtime/RuntimeAssetCache/Private/RuntimeAssetCacheBuilders.cpp
--- a/Engine/Source/Runtime/RuntimeAssetCache/Private/RuntimeAssetCacheBuilders.cpp
+++ b/Engine/Source/Runtime/RuntimeAssetCache/Private/RuntimeAssetCacheBuilders.cpp
@@ -232,8 +232,8 @@ void UExampleTextureCacheBuilder::OnAssetPostLoad()
 int64 UExampleTextureCacheBuilder::GetSerializedDataSizeEstimate()
 {
 	int64 DataSize = sizeof(FTexturePlatformData);
-	DataSize += sizeof(FString) + (sizeof(TCHAR) * 12);						// Guess the size of the pixel format string (most are less than 12 characters, but we don't need to be exact)
-	DataSize += Texture->GetResourceSize(EResourceSizeMode::Exclusive);		// Size of all the mips
+	DataSize += sizeof(FString) + (sizeof(TCHAR) * 12);							// Guess the size of the pixel format string (most are less than 12 characters, but we don't need to be exact)
+	DataSize += Texture->GetResourceSizeBytes(EResourceSizeMode::Exclusive);	// Size of all the mips
 	DataSize += (sizeof(int32) * 3) * Texture->GetNumMips();					// Each mip stores its X and Y size, and its BulkDataSize
 	return DataSize;
 }
diff --git a/Engine/Source/Runtime/Slate/Private/Widgets/Input/SButton.cpp b/Engine/Source/Runtime/Slate/Private/Widgets/Input/SButton.cpp
--- a/Engine/Source/Runtime/Slate/Private/Widgets/Input/SButton.cpp
+++ b/Engine/Source/Runtime/Slate/Private/Widgets/Input/SButton.cpp
@@ -90,11 +90,7 @@ int32 SButton::OnPaint(const FPaintArgs& Args, const FGeometry& AllottedGeometry
 			);
 	}
 
-	FWidgetStyle CompoundedWidgetStyle = FWidgetStyle(InWidgetStyle)
-		.BlendColorAndOpacityTint(ColorAndOpacity.Get())
-		.SetForegroundColor(ForegroundColor.Get());
-
-	return SCompoundWidget::OnPaint(Args, AllottedGeometry, MyClippingRect.IntersectionWith(AllottedGeometry.GetClippingRect()), OutDrawElements, LayerId, CompoundedWidgetStyle, bEnabled);
+	return SCompoundWidget::OnPaint(Args, AllottedGeometry, MyClippingRect.IntersectionWith(AllottedGeometry.GetClippingRect()), OutDrawElements, LayerId, InWidgetStyle, bEnabled);
 }
 
 FMargin SButton::GetCombinedPadding() const
diff --git a/Engine/Source/Runtime/Slate/Private/Widgets/Layout/SBorder.cpp b/Engine/Source/Runtime/Slate/Private/Widgets/Layout/SBorder.cpp
--- a/Engine/Source/Runtime/Slate/Private/Widgets/Layout/SBorder.cpp
+++ b/Engine/Source/Runtime/Slate/Private/Widgets/Layout/SBorder.cpp
@@ -78,15 +78,11 @@ int32 SBorder::OnPaint( const FPaintArgs& Args, const FGeometry& AllottedGeometr
 			BrushResource,
 			MyClippingRect,
 			DrawEffects,
-			BrushResource->GetTint( InWidgetStyle ) * InWidgetStyle.GetColorAndOpacityTint() * BorderBackgroundColor.Get().GetColor( InWidgetStyle )
+			BrushResource->GetTint( InWidgetStyle ) *InWidgetStyle.GetColorAndOpacityTint() * BorderBackgroundColor.Get().GetColor( InWidgetStyle )
 		);
 	}
 
-	FWidgetStyle CompoundedWidgetStyle = FWidgetStyle(InWidgetStyle)
-		.BlendColorAndOpacityTint(ColorAndOpacity.Get())
-		.SetForegroundColor( ForegroundColor.Get() );
-
-	return SCompoundWidget::OnPaint(Args, AllottedGeometry, MyClippingRect.IntersectionWith( AllottedGeometry.GetClippingRect() ), OutDrawElements, LayerId, CompoundedWidgetStyle, bEnabled );
+	return SCompoundWidget::OnPaint(Args, AllottedGeometry, MyClippingRect.IntersectionWith( AllottedGeometry.GetClippingRect() ), OutDrawElements, LayerId, InWidgetStyle, bEnabled );
 }
 
 FVector2D SBorder::ComputeDesiredSize(float LayoutScaleMultiplier) const
