[Copy] Fixed ULandscapeComponent::GetUsedMaterials which was causing ensure

diff --git a/Engine/Source/Runtime/Engine/Public/PrimitiveSceneProxy.h b/Engine/Source/Runtime/Engine/Public/PrimitiveSceneProxy.h
--- a/Engine/Source/Runtime/Engine/Public/PrimitiveSceneProxy.h
+++ b/Engine/Source/Runtime/Engine/Public/PrimitiveSceneProxy.h
@@ -458,6 +458,11 @@ public:
 
 #if WITH_EDITOR
 	inline int32 GetNumUncachedStaticLightingInteractions() { return NumUncachedStaticLightingInteractions; }
+
+	inline void AddUsedMaterialForVerification(UMaterialInterface* Material)
+	{
+		UsedMaterialsForVerification.AddUnique(Material);
+	}
 #endif
 
 	inline FLinearColor GetWireframeColor() const { return WireframeColor; }
diff --git a/Engine/Source/Runtime/Landscape/Private/Landscape.cpp b/Engine/Source/Runtime/Landscape/Private/Landscape.cpp
--- a/Engine/Source/Runtime/Landscape/Private/Landscape.cpp
+++ b/Engine/Source/Runtime/Landscape/Private/Landscape.cpp
@@ -1054,12 +1054,6 @@ void ULandscapeComponent::OnUnregister()
 #endif
 }
 
-void ULandscapeComponent::GetUsedMaterials(TArray<UMaterialInterface*>& OutMaterials) const
-{
-	// TODO - investigate whether this is correct
-	OutMaterials.Append(MaterialInstances.FilterByPredicate([](UMaterialInstance* MaterialInstance) { return MaterialInstance != nullptr; }));
-}
-
 void ALandscapeProxy::PostRegisterAllComponents()
 {
 	Super::PostRegisterAllComponents();
diff --git a/Engine/Source/Runtime/Landscape/Private/LandscapeRender.cpp b/Engine/Source/Runtime/Landscape/Private/LandscapeRender.cpp
--- a/Engine/Source/Runtime/Landscape/Private/LandscapeRender.cpp
+++ b/Engine/Source/Runtime/Landscape/Private/LandscapeRender.cpp
@@ -463,23 +463,35 @@ UMaterialInterface* GLandscapeLayerUsageMaterial = nullptr;
 // Game thread update
 void FLandscapeEditToolRenderData::Update(UMaterialInterface* InToolMaterial)
 {
-	ENQUEUE_UNIQUE_RENDER_COMMAND_TWOPARAMETER(
+	ENQUEUE_UNIQUE_RENDER_COMMAND_THREEPARAMETER(
 		UpdateEditToolRenderData,
 		FLandscapeEditToolRenderData*, LandscapeEditToolRenderData, this,
+		FPrimitiveSceneProxy*, SceneProxy, LandscapeComponent->SceneProxy,
 		UMaterialInterface*, NewToolMaterial, InToolMaterial,
 		{
 			LandscapeEditToolRenderData->ToolMaterial = NewToolMaterial;
+
+			if (SceneProxy)
+			{
+				SceneProxy->AddUsedMaterialForVerification(NewToolMaterial);
+			}
 		});
 }
 
 void FLandscapeEditToolRenderData::UpdateGizmo(UMaterialInterface* InGizmoMaterial)
 {
-	ENQUEUE_UNIQUE_RENDER_COMMAND_TWOPARAMETER(
+	ENQUEUE_UNIQUE_RENDER_COMMAND_THREEPARAMETER(
 		UpdateEditToolRenderData,
 		FLandscapeEditToolRenderData*, LandscapeEditToolRenderData, this,
+		FPrimitiveSceneProxy*, SceneProxy, LandscapeComponent->SceneProxy,
 		UMaterialInterface*, NewGizmoMaterial, InGizmoMaterial,
 		{
 			LandscapeEditToolRenderData->GizmoMaterial = NewGizmoMaterial;
+
+			if (SceneProxy)
+			{
+				SceneProxy->AddUsedMaterialForVerification(NewGizmoMaterial);
+			}
 		});
 }
 
@@ -536,6 +548,41 @@ void FLandscapeEditToolRenderData::UpdateSelectionMaterial(int32 InSelectedType)
 }
 #endif
 
+void ULandscapeComponent::GetUsedMaterials(TArray<UMaterialInterface*>& OutMaterials) const
+{
+	OutMaterials.Append(MaterialInstances.FilterByPredicate([](UMaterialInstance* MaterialInstance) { return MaterialInstance != nullptr; }));
+
+	if (OverrideMaterial)
+	{
+		OutMaterials.Add(OverrideMaterial);
+	}
+
+	if (OverrideHoleMaterial)
+	{
+		OutMaterials.Add(OverrideHoleMaterial);
+	}
+
+	if (EditToolRenderData)
+	{
+		if (EditToolRenderData->ToolMaterial)
+		{
+			OutMaterials.Add(EditToolRenderData->ToolMaterial);
+		}
+
+		if (EditToolRenderData->GizmoMaterial)
+		{
+			OutMaterials.Add(EditToolRenderData->GizmoMaterial);
+		}
+	}
+
+#if WITH_EDITOR
+	OutMaterials.Add(GLayerDebugColorMaterial);
+	OutMaterials.Add(GSelectionColorMaterial);
+	OutMaterials.Add(GSelectionRegionMaterial);
+	OutMaterials.Add(GMaskRegionMaterial);
+	OutMaterials.Add(GLandscapeLayerUsageMaterial);
+#endif
+}
 
 //
 // FLandscapeComponentSceneProxy
