Non-editor compile fix

diff --git a/Engine/Source/Runtime/Landscape/Private/LandscapeRender.cpp b/Engine/Source/Runtime/Landscape/Private/LandscapeRender.cpp
--- a/Engine/Source/Runtime/Landscape/Private/LandscapeRender.cpp
+++ b/Engine/Source/Runtime/Landscape/Private/LandscapeRender.cpp
@@ -562,6 +562,7 @@ void ULandscapeComponent::GetUsedMaterials(TArray<UMaterialInterface*>& OutMater
 		OutMaterials.Add(OverrideHoleMaterial);
 	}
 
+#if WITH_EDITORONLY_DATA
 	if (EditToolRenderData)
 	{
 		if (EditToolRenderData->ToolMaterial)
@@ -574,6 +575,7 @@ void ULandscapeComponent::GetUsedMaterials(TArray<UMaterialInterface*>& OutMater
 			OutMaterials.Add(EditToolRenderData->GizmoMaterial);
 		}
 	}
+#endif
 
 #if WITH_EDITOR
 	OutMaterials.Add(GLayerDebugColorMaterial);
