[NO NON-PII COMMIT MESSAGE]

diff --git a/Engine/Source/Runtime/Landscape/Private/LandscapeLight.cpp b/Engine/Source/Runtime/Landscape/Private/LandscapeLight.cpp
--- a/Engine/Source/Runtime/Landscape/Private/LandscapeLight.cpp
+++ b/Engine/Source/Runtime/Landscape/Private/LandscapeLight.cpp
@@ -163,11 +163,12 @@ namespace
 	void GetLODData(ULandscapeComponent* LandscapeComponent, int32 X, int32 Y, int32 HeightmapOffsetX, int32 HeightmapOffsetY, int32 LODValue, int32 HeightmapStride, FColor& OutHeight, FColor& OutXYOffset)
 	{
 		int32 ComponentSize = ((LandscapeComponent->SubsectionSizeQuads + 1) * LandscapeComponent->NumSubsections) >> LODValue;
-		int32 LODHeightmapSize = LandscapeComponent->HeightmapTexture->Source.GetSizeX() >> LODValue;
-		float Ratio = (float)(LODHeightmapSize) / (HeightmapStride);
+		int32 LODHeightmapSizeX = LandscapeComponent->HeightmapTexture->Source.GetSizeX() >> LODValue;
+		int32 LODHeightmapSizeY = LandscapeComponent->HeightmapTexture->Source.GetSizeY() >> LODValue;
+		float Ratio = (float)(LODHeightmapSizeX) / (HeightmapStride);
 
-		int32 CurrentHeightmapOffsetX = FMath::RoundToInt((float)(LODHeightmapSize)* LandscapeComponent->HeightmapScaleBias.Z);
-		int32 CurrentHeightmapOffsetY = FMath::RoundToInt((float)(LODHeightmapSize)* LandscapeComponent->HeightmapScaleBias.W);
+		int32 CurrentHeightmapOffsetX = FMath::RoundToInt((float)(LODHeightmapSizeX) * LandscapeComponent->HeightmapScaleBias.Z);
+		int32 CurrentHeightmapOffsetY = FMath::RoundToInt((float)(LODHeightmapSizeY) * LandscapeComponent->HeightmapScaleBias.W);
 
 		float XX = FMath::Clamp<float>((X - HeightmapOffsetX) * Ratio, 0.f, ComponentSize - 1.f) + CurrentHeightmapOffsetX;
 		int32 XI = (int32)XX;
@@ -181,10 +182,10 @@ namespace
 		FColor* HeightMipData = DataInterface.GetRawHeightData();
 		FColor* XYOffsetMipData = DataInterface.GetRawXYOffsetData();
 
-		FColor H1 = HeightMipData[XI + YI * LODHeightmapSize];
-		FColor H2 = HeightMipData[FMath::Min(XI + 1, LODHeightmapSize - 1) + YI * LODHeightmapSize];
-		FColor H3 = HeightMipData[XI + FMath::Min(YI + 1, LODHeightmapSize - 1) * LODHeightmapSize];
-		FColor H4 = HeightMipData[FMath::Min(XI + 1, LODHeightmapSize - 1) + FMath::Min(YI + 1, LODHeightmapSize - 1) * LODHeightmapSize];
+		FColor H1 = HeightMipData[XI + YI * LODHeightmapSizeX];
+		FColor H2 = HeightMipData[FMath::Min(XI + 1, LODHeightmapSizeX - 1) + YI * LODHeightmapSizeX];
+		FColor H3 = HeightMipData[XI + FMath::Min(YI + 1, LODHeightmapSizeY - 1) * LODHeightmapSizeX];
+		FColor H4 = HeightMipData[FMath::Min(XI + 1, LODHeightmapSizeX - 1) + FMath::Min(YI + 1, LODHeightmapSizeY - 1) * LODHeightmapSizeX];
 
 		uint16 Height = FMath::RoundToInt(FMath::Lerp(FMath::Lerp<float>(((H1.R << 8) + H1.G), ((H2.R << 8) + H2.G), XF),
 			FMath::Lerp<float>(((H3.R << 8) + H3.G), ((H4.R << 8) + H4.G), XF), YF));
@@ -197,10 +198,10 @@ namespace
 
 		if (LandscapeComponent->XYOffsetmapTexture)
 		{
-			FColor X1 = XYOffsetMipData[XI + YI * LODHeightmapSize];
-			FColor X2 = XYOffsetMipData[FMath::Min(XI + 1, LODHeightmapSize - 1) + YI * LODHeightmapSize];
-			FColor X3 = XYOffsetMipData[XI + FMath::Min(YI + 1, LODHeightmapSize - 1) * LODHeightmapSize];
-			FColor X4 = XYOffsetMipData[FMath::Min(XI + 1, LODHeightmapSize - 1) + FMath::Min(YI + 1, LODHeightmapSize - 1) * LODHeightmapSize];
+			FColor X1 = XYOffsetMipData[XI + YI * LODHeightmapSizeX];
+			FColor X2 = XYOffsetMipData[FMath::Min(XI + 1, LODHeightmapSizeX - 1) + YI * LODHeightmapSizeX];
+			FColor X3 = XYOffsetMipData[XI + FMath::Min(YI + 1, LODHeightmapSizeY - 1) * LODHeightmapSizeX];
+			FColor X4 = XYOffsetMipData[FMath::Min(XI + 1, LODHeightmapSizeX - 1) + FMath::Min(YI + 1, LODHeightmapSizeY - 1) * LODHeightmapSizeX];
 
 			uint16 XComp = FMath::RoundToInt(FMath::Lerp(FMath::Lerp<float>(((X1.R << 8) + X1.G), ((X2.R << 8) + X2.G), XF),
 				FMath::Lerp<float>(((X3.R << 8) + X3.G), ((X4.R << 8) + X4.G), XF), YF));
@@ -272,8 +273,8 @@ namespace
 			check(LandscapeComponent);
 			// Need Upscaling
 			int32 HeightmapStride = LandscapeComponent->HeightmapTexture->Source.GetSizeX() >> InLOD;
+			int32 HeightDataSize = HeightmapStride * LandscapeComponent->HeightmapTexture->Source.GetSizeY() >> InLOD;
 
-			int32 HeightDataSize = HeightmapStride * HeightmapStride;
 			CompHeightData.Empty(HeightDataSize);
 			CompXYOffsetData.Empty(HeightDataSize);
 			CompHeightData.AddZeroed(HeightDataSize);
