Fix Landscape trying to dereference a null UObject (either because of undo nulling it, or GC cleaning it up).

diff --git a/Engine/Source/Runtime/Landscape/Private/LandscapeModule.cpp b/Engine/Source/Runtime/Landscape/Private/LandscapeModule.cpp
--- a/Engine/Source/Runtime/Landscape/Private/LandscapeModule.cpp
+++ b/Engine/Source/Runtime/Landscape/Private/LandscapeModule.cpp
@@ -80,7 +80,7 @@ void WorldDestroyEventFunction(UWorld* World)
 	World->PerModuleDataObjects.RemoveAll(
 		[](UObject* Object)
 		{
-			return Object->IsA(ULandscapeInfoMap::StaticClass());
+			return Object != nullptr && Object->IsA(ULandscapeInfoMap::StaticClass());
 		}
 	);
 }
