Added missing shader cache

diff --git a/Engine/Source/Runtime/Landscape/Private/LandscapeRender.cpp b/Engine/Source/Runtime/Landscape/Private/LandscapeRender.cpp
--- a/Engine/Source/Runtime/Landscape/Private/LandscapeRender.cpp
+++ b/Engine/Source/Runtime/Landscape/Private/LandscapeRender.cpp
@@ -2681,6 +2681,12 @@ public:
 			FName(TEXT("TShadowDepthVSVertexShadowDepth_OutputDepthtrue")), // used by LPV
 			FName(TEXT("TShadowDepthPSPixelShadowDepth_NonPerspectiveCorrectfalse")),
 			FName(TEXT("TShadowDepthPSPixelShadowDepth_NonPerspectiveCorrecttrue")), // used by LPV
+			FName(TEXT("TBasePassPSFSimpleDirectionalLightLightingPolicy"))
+			FName(TEXT("TBasePassPSFSimpleDirectionalLightLightingPolicySkylight"))
+			FName(TEXT("TBasePassVSFSimpleDirectionalLightLightingPolicy"))
+			FName(TEXT("TBasePassPSFSimpleNoLightmapLightingPolicy"))
+			FName(TEXT("TBasePassPSFSimpleNoLightmapLightingPolicySkylight"))
+			FName(TEXT("TBasePassVSFSimpleNoLightmapLightingPolicy"))
 			FName(TEXT("TDepthOnlyVS<false>")),
 			FName(TEXT("TDepthOnlyVS<true>")),
 			FName(TEXT("FDepthOnlyPS")),
