Added missing shader from cache

diff --git a/Engine/Source/Runtime/Landscape/Private/LandscapeRender.cpp b/Engine/Source/Runtime/Landscape/Private/LandscapeRender.cpp
--- a/Engine/Source/Runtime/Landscape/Private/LandscapeRender.cpp
+++ b/Engine/Source/Runtime/Landscape/Private/LandscapeRender.cpp
@@ -2739,6 +2739,8 @@ public:
 			FName(TEXT("TBasePassPSFPrecomputedVolumetricLightmapLightingPolicy")),
 			FName(TEXT("TBasePassPSFPrecomputedVolumetricLightmapLightingPolicySkylight")),
 
+			FName(TEXT("TBasePassPSFSimpleStationaryLightVolumetricLightmapShadowsLightingPolicy")),			
+
 			FName(TEXT("TBasePassVSFNoLightMapPolicyAtmosphericFog")),
 			FName(TEXT("TBasePassVSFCachedPointIndirectLightingPolicyAtmosphericFog")),
 			FName(TEXT("TBasePassVSFSelfShadowedCachedPointIndirectLightingPolicy")),
