Fix yet more warnings due to landscape and shaders

diff --git a/Engine/Source/Runtime/Landscape/Private/LandscapeRender.cpp b/Engine/Source/Runtime/Landscape/Private/LandscapeRender.cpp
--- a/Engine/Source/Runtime/Landscape/Private/LandscapeRender.cpp
+++ b/Engine/Source/Runtime/Landscape/Private/LandscapeRender.cpp
@@ -2693,6 +2693,17 @@ public:
 			// UE-44519, masked material with landscape layers requires FHitProxy shaders.
 			FName(TEXT("FHitProxyVS")),
 			FName(TEXT("FHitProxyPS")),
+
+			FName(TEXT("TBasePassVSFSimpleStationaryLightVolumetricLightmapShadowsLightingPolicy")),
+			FName(TEXT("TBasePassPSFSimpleStationaryLightSingleSampleShadowsLightingPolicy")),
+			FName(TEXT("TBasePassPSFSimpleStationaryLightSingleSampleShadowsLightingPolicySkylight")),
+			FName(TEXT("TBasePassVSFSimpleStationaryLightSingleSampleShadowsLightingPolicy")),
+			FName(TEXT("TBasePassPSFSimpleStationaryLightPrecomputedShadowsLightingPolicy")),
+			FName(TEXT("TBasePassPSFSimpleStationaryLightPrecomputedShadowsLightingPolicySkylight")),
+			FName(TEXT("TBasePassVSFSimpleStationaryLightPrecomputedShadowsLightingPolicy")),
+			FName(TEXT("TBasePassPSFSimpleLightmapOnlyLightingPolicy")),
+			FName(TEXT("TBasePassPSFSimpleLightmapOnlyLightingPolicySkylight")),
+			FName(TEXT("TBasePassVSFSimpleLightmapOnlyLightingPolicy")),
 		};
 		return AllowedShaderTypes;
 	}
