[NO NON-PII COMMIT MESSAGE]

diff --git a/Engine/Source/Runtime/Landscape/Private/LandscapeRender.cpp b/Engine/Source/Runtime/Landscape/Private/LandscapeRender.cpp
--- a/Engine/Source/Runtime/Landscape/Private/LandscapeRender.cpp
+++ b/Engine/Source/Runtime/Landscape/Private/LandscapeRender.cpp
@@ -1570,7 +1570,7 @@ void FLandscapeComponentSceneProxy::CalculateBatchElementLOD(const FSceneView& I
 	check(InMeshBatchScreenRadiusSquared >= 0.0f && InMeshBatchScreenRadiusSquared <= 1.0f);
 	float ComponentScreenSize = InMeshBatchScreenRadiusSquared;
 
-	if (ForcedLOD == INDEX_NONE && NumSubsections > 1)
+	if (NumSubsections > 1)
 	{
 		InOutLODData.UseCombinedMeshBatch = false; // default to individual batch render
 
@@ -1621,7 +1621,7 @@ void FLandscapeComponentSceneProxy::CalculateBatchElementLOD(const FSceneView& I
 			}
 		}
 
-		if (!GLandscapeDebugOptions.IsCombinedDisabled() && (AllSubSectionHaveSameScreenSize || GLandscapeDebugOptions.IsCombinedAll()))
+		if (!GLandscapeDebugOptions.IsCombinedDisabled() && (AllSubSectionHaveSameScreenSize || GLandscapeDebugOptions.IsCombinedAll() || ForcedLOD != INDEX_NONE))
 		{
 			InOutLODData.UseCombinedMeshBatch = true;
 
@@ -1907,7 +1907,7 @@ float FLandscapeComponentSceneProxy::GetNeighborLOD(const FSceneView& InView, fl
 			}
 		}
 
-		if (ForcedLOD == INDEX_NONE && NumSubsections > 1)
+		if (NumSubsections > 1)
 		{
 			float SubSectionMaxExtend = LandscapeComponentMaxExtends / 2.0f;
 			FVector ComponentTopLeftCorner = LandscapeComponentOrigin - FVector(SubSectionMaxExtend, SubSectionMaxExtend, 0.0f);
