[NO NON-PII COMMIT MESSAGE]

diff --git a/Engine/Source/Runtime/Landscape/Private/LandscapeRender.cpp b/Engine/Source/Runtime/Landscape/Private/LandscapeRender.cpp
--- a/Engine/Source/Runtime/Landscape/Private/LandscapeRender.cpp
+++ b/Engine/Source/Runtime/Landscape/Private/LandscapeRender.cpp
@@ -1998,7 +1998,7 @@ float FLandscapeComponentSceneProxy::GetNeighborLOD(const FSceneView& InView, fl
 			if (Neighbor->GetLandscapeComponent() != nullptr)
 			{
 				LandscapeComponentOrigin = NeighborComponent->Bounds.Origin;
-				LandscapeComponentMaxExtends = NeighborComponent->Bounds.BoxExtent[0];
+				LandscapeComponentMaxExtends = NeighborComponent->SubsectionSizeQuads * FMath::Max(NeighborComponent->GetComponentTransform().GetScale3D().X, NeighborComponent->GetComponentTransform().GetScale3D().Y);
 			}
 		}
 
