[NO NON-PII COMMIT MESSAGE]

diff --git a/Engine/Source/Runtime/Landscape/Private/Landscape.cpp b/Engine/Source/Runtime/Landscape/Private/Landscape.cpp
--- a/Engine/Source/Runtime/Landscape/Private/Landscape.cpp
+++ b/Engine/Source/Runtime/Landscape/Private/Landscape.cpp
@@ -2581,6 +2581,10 @@ void ULandscapeComponent::SerializeStateHashes(FArchive& Ar)
 
 	int32 OccluderGeometryLOD = GetLandscapeProxy()->OccluderGeometryLOD;
 	Ar << OccluderGeometryLOD;
+
+	// Take into account the Heightmap offset per component
+	Ar << HeightmapScaleBias.Z;
+	Ar << HeightmapScaleBias.W;
 }
 
 void ALandscapeProxy::UpdateBakedTextures()
