[NO NON-PII COMMIT MESSAGE]

diff --git a/Engine/Source/Runtime/Landscape/Private/LandscapeRender.cpp b/Engine/Source/Runtime/Landscape/Private/LandscapeRender.cpp
--- a/Engine/Source/Runtime/Landscape/Private/LandscapeRender.cpp
+++ b/Engine/Source/Runtime/Landscape/Private/LandscapeRender.cpp
@@ -1486,12 +1486,16 @@ bool FLandscapeComponentSceneProxy::GetMeshElement(bool UseSeperateBatchForShado
 void FLandscapeComponentSceneProxy::ApplyWorldOffset(FVector InOffset)
 {
 	FPrimitiveSceneProxy::ApplyWorldOffset(InOffset);
-	for (int32 SubY = 0; SubY < NumSubsections; ++SubY)
+
+	if (NumSubsections > 1)
 	{
-		for (int32 SubX = 0; SubX < NumSubsections; ++SubX)
+		for (int32 SubY = 0; SubY < NumSubsections; ++SubY)
 		{
-			int32 SubSectionIndex = SubX + SubY * NumSubsections;
-			SubSectionScreenSizeTestingPosition[SubSectionIndex] += InOffset;
+			for (int32 SubX = 0; SubX < NumSubsections; ++SubX)
+			{
+				int32 SubSectionIndex = SubX + SubY * NumSubsections;
+				SubSectionScreenSizeTestingPosition[SubSectionIndex] += InOffset;
+			}
 		}
 	}
 }
