[NO NON-PII COMMIT MESSAGE]

diff --git a/Engine/Source/Runtime/Landscape/Classes/LandscapeProxy.h b/Engine/Source/Runtime/Landscape/Classes/LandscapeProxy.h
--- a/Engine/Source/Runtime/Landscape/Classes/LandscapeProxy.h
+++ b/Engine/Source/Runtime/Landscape/Classes/LandscapeProxy.h
@@ -364,7 +364,7 @@ public:
 	float ComponentScreenSizeToUseSubSections;
 
 	/** The distribution setting used to change the LOD 0 generation, 1.75 is the normal distribution, numbers influence directly the LOD0 proportion on screen. */
-	UPROPERTY(EditAnywhere, Category = "LOD Distribution", meta = (DisplayName = "LOD 0", ClampMin = "1.0", ClampMax = "5.0", UIMin = "1.0", UIMax = "5.0"))
+	UPROPERTY(EditAnywhere, Category = "LOD Distribution", meta = (DisplayName = "LOD 0", ClampMin = "1.0", ClampMax = "10.0", UIMin = "1.0", UIMax = "10.0"))
 	float LOD0DistributionSetting;
 
 	/** The distribution setting used to change the LOD generation, 2 is the normal distribution, small number mean you want your last LODs to take more screen space and big number mean you want your first LODs to take more screen space. */
diff --git a/Engine/Source/Runtime/Landscape/Private/LandscapeEdit.cpp b/Engine/Source/Runtime/Landscape/Private/LandscapeEdit.cpp
--- a/Engine/Source/Runtime/Landscape/Private/LandscapeEdit.cpp
+++ b/Engine/Source/Runtime/Landscape/Private/LandscapeEdit.cpp
@@ -3890,7 +3890,7 @@ void ALandscape::PostEditChangeProperty(FPropertyChangedEvent& PropertyChangedEv
 	}
 	else if (PropertyName == FName(TEXT("LOD0DistributionSetting")))
 	{
-		LOD0DistributionSetting = FMath::Clamp<float>(LOD0DistributionSetting, 1.0f, 5.0f);
+		LOD0DistributionSetting = FMath::Clamp<float>(LOD0DistributionSetting, 1.0f, 10.0f);
 		bPropagateToProxies = true;
 	}
 	else if (PropertyName == FName(TEXT("CollisionMipLevel")))
