[NO NON-PII COMMIT MESSAGE]

diff --git a/Engine/Source/Runtime/Landscape/Private/Landscape.cpp b/Engine/Source/Runtime/Landscape/Private/Landscape.cpp
--- a/Engine/Source/Runtime/Landscape/Private/Landscape.cpp
+++ b/Engine/Source/Runtime/Landscape/Private/Landscape.cpp
@@ -2609,7 +2609,7 @@ void ALandscapeProxy::UpdateBakedTextures()
 		// Clear out any existing GI textures
 		for (ULandscapeComponent* Component : LandscapeComponents)
 		{
-			if (Component->GIBakedBaseColorTexture != nullptr)
+			if (Component != nullptr && Component->GIBakedBaseColorTexture != nullptr)
 			{
 				Component->BakedTextureMaterialGuid.Invalidate();
 				Component->GIBakedBaseColorTexture = nullptr;
@@ -2640,6 +2640,11 @@ void ALandscapeProxy::UpdateBakedTextures()
 	TMap<UTexture2D*, FBakedTextureSourceInfo> ComponentsByHeightmap;
 	for (ULandscapeComponent* Component : LandscapeComponents)
 	{
+		if (Component == nullptr)
+		{
+			continue;
+		}
+
 		FBakedTextureSourceInfo& Info = ComponentsByHeightmap.FindOrAdd(Component->HeightmapTexture);
 		Info.Components.Add(Component);
 		Component->SerializeStateHashes(*Info.ComponentStateAr);
