Add null check to ALandscapeProxy::FlushGrassComponents

diff --git a/Engine/Source/Runtime/Landscape/Private/LandscapeGrass.cpp b/Engine/Source/Runtime/Landscape/Private/LandscapeGrass.cpp
--- a/Engine/Source/Runtime/Landscape/Private/LandscapeGrass.cpp
+++ b/Engine/Source/Runtime/Landscape/Private/LandscapeGrass.cpp
@@ -1900,7 +1900,7 @@ void ALandscapeProxy::FlushGrassComponents(const TSet<ULandscapeComponent*>* Onl
 			}
 		}
 #if WITH_EDITOR
-		if (GIsEditor && bFlushGrassMaps && GetWorld()->Scene->GetFeatureLevel() >= ERHIFeatureLevel::SM4)
+		if (GIsEditor && bFlushGrassMaps && GetWorld() && GetWorld()->Scene->GetFeatureLevel() >= ERHIFeatureLevel::SM4)
 		{
 			for (ULandscapeComponent* Component : *OnlyForComponents)
 			{
@@ -1943,7 +1943,7 @@ void ALandscapeProxy::FlushGrassComponents(const TSet<ULandscapeComponent*>* Onl
 		}
 
 #if WITH_EDITOR
-		if (GIsEditor && bFlushGrassMaps && GetWorld()->Scene->GetFeatureLevel() >= ERHIFeatureLevel::SM4)
+		if (GIsEditor && bFlushGrassMaps && GetWorld() && GetWorld()->Scene->GetFeatureLevel() >= ERHIFeatureLevel::SM4)
 		{
 			// Clear GrassMaps
 			TInlineComponentArray<ULandscapeComponent*> LandComps;
