[NO NON-PII COMMIT MESSAGE]

diff --git a/Engine/Source/Runtime/Landscape/Private/Landscape.cpp b/Engine/Source/Runtime/Landscape/Private/Landscape.cpp
--- a/Engine/Source/Runtime/Landscape/Private/Landscape.cpp
+++ b/Engine/Source/Runtime/Landscape/Private/Landscape.cpp
@@ -2339,9 +2339,15 @@ void ULandscapeInfo::RecreateLandscapeInfo(UWorld* InWorld, bool bMapCheck)
 	// Remove empty entries from global LandscapeInfo map
 	for (auto It = LandscapeInfoMap.Map.CreateIterator(); It; ++It)
 	{
-		if (It.Value()->GetLandscapeProxy() == nullptr)
+		ULandscapeInfo* Info = It.Value();
+
+		if (Info != nullptr && Info->GetLandscapeProxy() == nullptr)
+		{
+			Info->MarkPendingKill();
+			It.RemoveCurrent();
+		}
+		else if (Info == nullptr) // remove invalid entry
 		{
-			It.Value()->MarkPendingKill();
 			It.RemoveCurrent();
 		}
 	}
