[NO NON-PII COMMIT MESSAGE]

diff --git a/Engine/Source/Runtime/Landscape/Private/Landscape.cpp b/Engine/Source/Runtime/Landscape/Private/Landscape.cpp
--- a/Engine/Source/Runtime/Landscape/Private/Landscape.cpp
+++ b/Engine/Source/Runtime/Landscape/Private/Landscape.cpp
@@ -2050,7 +2050,7 @@ void ULandscapeInfo::RegisterActor(ALandscapeProxy* Proxy, bool bMapCheck)
 	// in case this Info object is not initialized yet
 	// initialized it with properties from passed actor
 	if (LandscapeGuid.IsValid() == false ||
-		(GetLandscapeProxy() == nullptr && ensure(LandscapeGuid == Proxy->GetLandscapeGuid())))
+		(GetLandscapeProxy() == nullptr && ensure(LandscapeGuid == Proxy->GetLandscapeGuid()) && Proxy->HasValidRootComponent()))
 	{
 		LandscapeGuid = Proxy->GetLandscapeGuid();
 		ComponentSizeQuads = Proxy->ComponentSizeQuads;
@@ -2065,7 +2065,7 @@ void ULandscapeInfo::RegisterActor(ALandscapeProxy* Proxy, bool bMapCheck)
 	check(ComponentNumSubsections == Proxy->NumSubsections);
 	check(SubsectionSizeQuads == Proxy->SubsectionSizeQuads);
 
-	if (!DrawScale.Equals(Proxy->GetRootComponent()->RelativeScale3D))
+	if (Proxy->HasValidRootComponent() && !DrawScale.Equals(Proxy->GetRootComponent()->RelativeScale3D))
 	{
 		UE_LOG(LogLandscape, Warning, TEXT("Landscape proxy (%s) scale (%s) does not match to main actor scale (%s)."),
 			*Proxy->GetName(), *Proxy->GetRootComponent()->RelativeScale3D.ToCompactString(), *DrawScale.ToCompactString());
