[NO NON-PII COMMIT MESSAGE]

diff --git a/Engine/Source/Runtime/Landscape/Private/LandscapeEdit.cpp b/Engine/Source/Runtime/Landscape/Private/LandscapeEdit.cpp
--- a/Engine/Source/Runtime/Landscape/Private/LandscapeEdit.cpp
+++ b/Engine/Source/Runtime/Landscape/Private/LandscapeEdit.cpp
@@ -3791,20 +3791,25 @@ void ALandscapeStreamingProxy::PostEditChangeProperty(FPropertyChangedEvent& Pro
 	else if (PropertyName == FName(TEXT("LandscapeMaterial")) || PropertyName == FName(TEXT("LandscapeHoleMaterial")))
 	{
 		{
-			FMaterialUpdateContext MaterialUpdateContext;
-			GetLandscapeInfo()->UpdateLayerInfoMap(/*this*/);
+			ULandscapeInfo* Info = GetLandscapeInfo();
 
-			// Clear the parents out of combination material instances
-			for (const auto& MICPair : MaterialInstanceConstantMap)
+			if (Info != nullptr)
 			{
-				UMaterialInstanceConstant* MaterialInstance = MICPair.Value;
-				MaterialInstance->BasePropertyOverrides.bOverride_BlendMode = false;
-				MaterialInstance->SetParentEditorOnly(nullptr);
-				MaterialUpdateContext.AddMaterialInstance(MaterialInstance);
-			}
+				FMaterialUpdateContext MaterialUpdateContext;
+				Info->UpdateLayerInfoMap(/*this*/);
 
-			// Remove our references to any material instances
-			MaterialInstanceConstantMap.Empty();
+				// Clear the parents out of combination material instances
+				for (const auto& MICPair : MaterialInstanceConstantMap)
+				{
+					UMaterialInstanceConstant* MaterialInstance = MICPair.Value;
+					MaterialInstance->BasePropertyOverrides.bOverride_BlendMode = false;
+					MaterialInstance->SetParentEditorOnly(nullptr);
+					MaterialUpdateContext.AddMaterialInstance(MaterialInstance);
+				}
+
+				// Remove our references to any material instances
+				MaterialInstanceConstantMap.Empty();
+			}
 		}
 
 		UpdateAllComponentMaterialInstances();
