[NO NON-PII COMMIT MESSAGE]

diff --git a/Engine/Source/Runtime/Landscape/Private/LandscapeDataAccess.cpp b/Engine/Source/Runtime/Landscape/Private/LandscapeDataAccess.cpp
--- a/Engine/Source/Runtime/Landscape/Private/LandscapeDataAccess.cpp
+++ b/Engine/Source/Runtime/Landscape/Private/LandscapeDataAccess.cpp
@@ -18,12 +18,10 @@ LANDSCAPE_API FLandscapeComponentDataInterface::FLandscapeComponentDataInterface
 	HeightmapComponentOffsetY = FMath::RoundToInt((float)(Component->HeightmapTexture->Source.GetSizeY() >> MipLevel) * Component->HeightmapScaleBias.W);
 	HeightmapSubsectionOffset = (Component->SubsectionSizeQuads + 1) >> MipLevel;
 
+	ComponentSizeVerts = (Component->ComponentSizeQuads + 1) >> MipLevel;
 	SubsectionSizeVerts = (Component->SubsectionSizeQuads + 1) >> MipLevel;
-	SubsectionSizeQuads = SubsectionSizeVerts - 1;
-	ComponentSizeQuads = Component->NumSubsections * SubsectionSizeQuads;
-	ComponentSizeVerts = Component->NumSubsections * (SubsectionSizeQuads + 1);
 	ComponentNumSubsections = Component->NumSubsections;
-	
+
 	if (MipLevel < Component->HeightmapTexture->Source.GetNumMips())
 	{
 		HeightMipData = (FColor*)DataInterface.LockMip(Component->HeightmapTexture, MipLevel);
@@ -136,7 +134,7 @@ LANDSCAPE_API FColor* FLandscapeComponentDataInterface::GetXYOffsetData(int32 Lo
 
 LANDSCAPE_API FVector FLandscapeComponentDataInterface::GetLocalVertex(int32 LocalX, int32 LocalY) const
 {
-	const float ScaleFactor = (float)Component->ComponentSizeQuads / (float)ComponentSizeQuads;
+	const float ScaleFactor = (float)Component->ComponentSizeQuads / (float)(ComponentSizeVerts - 1);
 	float XOffset, YOffset;
 	GetXYOffset(LocalX, LocalY, XOffset, YOffset);
 	return FVector(LocalX * ScaleFactor + XOffset, LocalY * ScaleFactor + YOffset, LandscapeDataAccess::GetLocalHeight(GetHeight(LocalX, LocalY)));
@@ -173,7 +171,7 @@ LANDSCAPE_API void FLandscapeComponentDataInterface::GetWorldPositionTangents(in
 
 	uint16 Height = (Data->R << 8) + Data->G;
 
-	const float ScaleFactor = (float)Component->ComponentSizeQuads / (float)ComponentSizeQuads;
+	const float ScaleFactor = (float)Component->ComponentSizeQuads / (float)(ComponentSizeVerts - 1);
 	float XOffset, YOffset;
 	GetXYOffset(LocalX, LocalY, XOffset, YOffset);
 	WorldPos = Component->GetComponentTransform().TransformPosition(FVector(LocalX * ScaleFactor + XOffset, LocalY * ScaleFactor + YOffset, LandscapeDataAccess::GetLocalHeight(Height)));
diff --git a/Engine/Source/Runtime/Landscape/Public/LandscapeDataAccess.h b/Engine/Source/Runtime/Landscape/Public/LandscapeDataAccess.h
--- a/Engine/Source/Runtime/Landscape/Public/LandscapeDataAccess.h
+++ b/Engine/Source/Runtime/Landscape/Public/LandscapeDataAccess.h
@@ -341,10 +341,8 @@ private:
 	FColor* HeightMipData;
 	FColor* XYOffsetMipData;
 
-	int32 ComponentSizeQuads;
 	int32 ComponentSizeVerts;
 	int32 SubsectionSizeVerts;
-	int32 SubsectionSizeQuads;
 	int32 ComponentNumSubsections;
 public:
 	const int32 MipLevel;
