Improvements to mobile landscape cooking
Fixed texture streaming for extra weightmaps
Save MobileDataSourceHash to disk so we can invalidate mobile pixel data in DDC if the format changes

diff --git a/Engine/Source/Editor/LandscapeEditor/Private/LandscapeEdModeComponentTools.cpp b/Engine/Source/Editor/LandscapeEditor/Private/LandscapeEdModeComponentTools.cpp
--- a/Engine/Source/Editor/LandscapeEditor/Private/LandscapeEdModeComponentTools.cpp
+++ b/Engine/Source/Editor/LandscapeEditor/Private/LandscapeEdModeComponentTools.cpp
@@ -720,7 +720,7 @@ public:
 					// clear transient mobile data
 					Component->MobileDataSourceHash.Invalidate();
 					Component->MobileMaterialInterface = nullptr;
-					Component->MobileWeightNormalmapTexture = nullptr;
+					Component->MobileWeightmapTextures.Reset();
 					
 					Component->UpdateMaterialInstances();
 
diff --git a/Engine/Source/Runtime/Landscape/Classes/LandscapeComponent.h b/Engine/Source/Runtime/Landscape/Classes/LandscapeComponent.h
--- a/Engine/Source/Runtime/Landscape/Classes/LandscapeComponent.h
+++ b/Engine/Source/Runtime/Landscape/Classes/LandscapeComponent.h
@@ -385,9 +385,8 @@ public:
 	UPROPERTY(Transient, DuplicateTransient, NonTransactional)
 	FLandscapeEditToolRenderData EditToolRenderData;
 
-	/** Hash of source for ES2 generated data. Used for mobile preview and cook-in-editor
-	 * to determine if we need to re-generate ES2 pixel data. */
-	UPROPERTY(Transient, DuplicateTransient)
+	/** Hash of source for ES2 generated data. Used determine if we need to re-generate ES2 pixel data. */
+	UPROPERTY(DuplicateTransient)
 	FGuid MobileDataSourceHash;
 #endif
 
@@ -399,9 +398,10 @@ public:
 	UPROPERTY(NonPIEDuplicateTransient)
 	UMaterialInterface* MobileMaterialInterface;
 
-	/** Generated weight/normal map texture used for ES2. Serialized only when cooking or loading cooked builds. */
+	/** Generated weightmap textures used for ES2. The first entry is also used for the normal map. 
+	  * Serialized only when cooking or loading cooked builds. */
 	UPROPERTY(NonPIEDuplicateTransient)
-	UTexture2D* MobileWeightNormalmapTexture;
+	TArray<UTexture2D*> MobileWeightmapTextures;
 
 #if WITH_EDITORONLY_DATA
 	/** Layer allocations used by mobile. Cached value here used only in the editor for usage visualization. */
diff --git a/Engine/Source/Runtime/Landscape/Private/Landscape.cpp b/Engine/Source/Runtime/Landscape/Private/Landscape.cpp
--- a/Engine/Source/Runtime/Landscape/Private/Landscape.cpp
+++ b/Engine/Source/Runtime/Landscape/Private/Landscape.cpp
@@ -91,7 +91,7 @@ namespace LandscapeCookStats
 #define ENABLE_LANDSCAPE_COOKING 1
 
 // Increment this to regenerate mobile landscape data
-#define LANDSCAPE_MOBILE_COOK_VERSION 0
+#define LANDSCAPE_MOBILE_COOK_VERSION 2
 
 #define LOCTEXT_NAMESPACE "Landscape"
 
@@ -227,9 +227,9 @@ void ULandscapeComponent::CheckGenerateLandscapePlatformData(bool bIsCooking, co
 	FSHA1::HashBuffer(ComponentStateAr.GetData(), ComponentStateAr.Num(), (uint8*)Hash);
 	FGuid NewSourceHash = FGuid(Hash[0] ^ Hash[4], Hash[1], Hash[2], Hash[3]);
 
-	bool bHashMismatch = MobileDataSourceHash.IsValid() && MobileDataSourceHash != NewSourceHash;
+	bool bHashMismatch = MobileDataSourceHash != NewSourceHash;
 	bool bMissingVertexData = !PlatformData.HasValidPlatformData();
-	bool bMissingPixelData = !MobileMaterialInterface || !MobileWeightNormalmapTexture;
+	bool bMissingPixelData = MobileMaterialInterface == nullptr || MobileWeightmapTextures.Num() == 0;
 	
 	bool bRegenerateVertexData = bMissingVertexData || bMissingPixelData || bHashMismatch;
 	
@@ -307,6 +307,21 @@ void ULandscapeComponent::Serialize(FArchive& Ar)
 		Exchange(BackupWeightmapTextures, WeightmapTextures);
 	}
 	else
+	if (Ar.IsCooking() && !HasAnyFlags(RF_ClassDefaultObject) && !Ar.CookingTarget()->SupportsFeature(ETargetPlatformFeatures::MobileRendering))
+	{
+		// These properties are only used for mobile so we back them up and clear them before serializing them.
+		UMaterialInterface* BackupMobileMaterialInterface = nullptr;
+		TArray<UTexture2D*> BackupMobileWeightmapTextures;
+
+		Exchange(MobileMaterialInterface, BackupMobileMaterialInterface);
+		Exchange(MobileWeightmapTextures, BackupMobileWeightmapTextures);
+
+		Super::Serialize(Ar);
+
+		Exchange(MobileMaterialInterface, BackupMobileMaterialInterface);
+		Exchange(MobileWeightmapTextures, BackupMobileWeightmapTextures);
+	}
+	else
 #endif
 	{
 		Super::Serialize(Ar);
@@ -383,23 +398,6 @@ void ULandscapeComponent::Serialize(FArchive& Ar)
 				check(PlatformData.HasValidPlatformData());
 			}
 			Ar << PlatformData;
-			if (Ar.UE4Ver() >= VER_UE4_SERIALIZE_LANDSCAPE_ES2_TEXTURES)
-			{
-				Ar << MobileMaterialInterface;
-				Ar << MobileWeightNormalmapTexture;
-			}
-		}
-
-		if (Ar.UE4Ver() >= VER_UE4_LANDSCAPE_GRASS_COOKING && Ar.UE4Ver() < VER_UE4_SERIALIZE_LANDSCAPE_GRASS_DATA)
-		{
-			// deal with previous cooked FGrassMap data
-			int32 NumChannels = 0;
-			Ar << NumChannels;
-			if (NumChannels)
-			{
-				TArray<uint8> OldData;
-				OldData.BulkSerialize(Ar);
-			}
 		}
 	}
 #endif
@@ -421,7 +419,7 @@ void ULandscapeComponent::GetResourceSizeEx(FResourceSizeEx& CumulativeResourceS
 
 #if WITH_EDITOR
 UMaterialInterface* ULandscapeComponent::GetLandscapeMaterial() const
-{
+	{		
 	if (OverrideMaterial)
 	{
 		return OverrideMaterial;
@@ -810,10 +808,10 @@ ALandscapeProxy::ALandscapeProxy(const FObjectInitializer& ObjectInitializer)
 
 		VisibilityLayer = ConstructorStatics.DataLayer.Get();
 		check(VisibilityLayer);
-#if WITH_EDITORONLY_DATA
+	#if WITH_EDITORONLY_DATA
 		// This layer should be no weight blending
 		VisibilityLayer->bNoWeightBlend = true;
-#endif
+	#endif
 		VisibilityLayer->LayerUsageDebugColor = FLinearColor(0, 0, 0, 0);
 		VisibilityLayer->AddToRoot();
 	}
@@ -1907,7 +1905,7 @@ UMaterialInterface* ALandscapeProxy::GetLandscapeMaterial() const
 	if (LandscapeMaterial)
 	{
 		return LandscapeMaterial;
-	}
+					}
 	return UMaterial::GetDefaultMaterial(MD_Surface);
 }
 
@@ -1921,7 +1919,7 @@ UMaterialInterface* ALandscapeProxy::GetLandscapeHoleMaterial() const
 }
 
 UMaterialInterface* ALandscapeStreamingProxy::GetLandscapeMaterial() const
-{
+			{
 	if (LandscapeMaterial)
 	{
 		return LandscapeMaterial;
diff --git a/Engine/Source/Runtime/Landscape/Private/LandscapeEdit.cpp b/Engine/Source/Runtime/Landscape/Private/LandscapeEdit.cpp
--- a/Engine/Source/Runtime/Landscape/Private/LandscapeEdit.cpp
+++ b/Engine/Source/Runtime/Landscape/Private/LandscapeEdit.cpp
@@ -5018,9 +5018,11 @@ void ULandscapeComponent::GeneratePlatformPixelData()
 
 	int32 WeightmapSize = (SubsectionSizeQuads + 1) * NumSubsections;
 
-	MobileWeightNormalmapTexture = GetLandscapeProxy()->CreateLandscapeTexture(WeightmapSize, WeightmapSize, TEXTUREGROUP_Terrain_Weightmap, TSF_BGRA8);
+	UTexture2D* MobileWeightNormalmapTexture = GetLandscapeProxy()->CreateLandscapeTexture(WeightmapSize, WeightmapSize, TEXTUREGROUP_Terrain_Weightmap, TSF_BGRA8);
 	CreateEmptyTextureMips(MobileWeightNormalmapTexture);
-	TArray<UTexture2D*> MobileWeightmapTextures;
+
+	MobileWeightmapTextures.Empty();
+	MobileWeightmapTextures.Add(MobileWeightNormalmapTexture);
 
 	{
 		FLandscapeTextureDataInterface LandscapeData;
diff --git a/Engine/Source/Runtime/Landscape/Private/LandscapeRenderMobile.cpp b/Engine/Source/Runtime/Landscape/Private/LandscapeRenderMobile.cpp
--- a/Engine/Source/Runtime/Landscape/Private/LandscapeRenderMobile.cpp
+++ b/Engine/Source/Runtime/Landscape/Private/LandscapeRenderMobile.cpp
@@ -311,11 +311,10 @@ FLandscapeComponentSceneProxyMobile::FLandscapeComponentSceneProxyMobile(ULandsc
 	check(InComponent);
 	
 	check(InComponent->MobileMaterialInterface);
-	check(InComponent->MobileWeightNormalmapTexture);
+	check(InComponent->MobileWeightmapTextures.Num() > 0);
 
-	WeightmapTextures.Empty(1);
-	WeightmapTextures.Add(InComponent->MobileWeightNormalmapTexture);
-	NormalmapTexture = InComponent->MobileWeightNormalmapTexture;
+	WeightmapTextures = InComponent->MobileWeightmapTextures;
+	NormalmapTexture = InComponent->MobileWeightmapTextures[0];
 
 	BlendableLayerMask = InComponent->MobileBlendableLayerMask;
 
