Cleanup blueprint callable categories

Landscape Editor -> Landscape|Editor
Landscape Runtime -> Landscape|Runtime
Cloth -> Clothing Simulation
Cinematics -> Cinematic
Utility -> Utilities

diff --git a/Engine/Source/Developer/FunctionalTesting/Classes/FunctionalTestUtilityLibrary.h b/Engine/Source/Developer/FunctionalTesting/Classes/FunctionalTestUtilityLibrary.h
--- a/Engine/Source/Developer/FunctionalTesting/Classes/FunctionalTestUtilityLibrary.h
+++ b/Engine/Source/Developer/FunctionalTesting/Classes/FunctionalTestUtilityLibrary.h
@@ -19,6 +19,6 @@ class FUNCTIONALTESTING_API UFunctionalTestUtilityLibrary : public UBlueprintFun
 public:
 
 	/** Helper function to trace and permute many options at once */
-	UFUNCTION(BlueprintCallable, Category = "Utility|Collision", meta = (bIgnoreSelf = "true", WorldContext = "WorldContextObject", AutoCreateRefTerm = "ActorsToIgnore", AdvancedDisplay = "TraceColor,TraceHitColor,DrawTime", Keywords = "sweep"))
+	UFUNCTION(BlueprintCallable, Category = "Utilities|Collision", meta = (bIgnoreSelf = "true", WorldContext = "WorldContextObject", AutoCreateRefTerm = "ActorsToIgnore", AdvancedDisplay = "TraceColor,TraceHitColor,DrawTime", Keywords = "sweep"))
 	static UTraceQueryTestResults* TraceChannelTestUtil(UObject* WorldContextObject, const FTraceChannelTestBatchOptions& BatchOptions, const FVector Start, const FVector End, float SphereCapsuleRadius, float CapsuleHalfHeight, FVector BoxHalfSize, const FRotator Orientation, ETraceTypeQuery TraceChannel, TArray<TEnumAsByte<EObjectTypeQuery> > ObjectTypes, FName ProfileName, bool bTraceComplex, const TArray<AActor*>& ActorsToIgnore, bool bIgnoreSelf, EDrawDebugTrace::Type DrawDebugType, FLinearColor TraceColor = FLinearColor::Red, FLinearColor TraceHitColor = FLinearColor::Green, float DrawTime = 5.0f);
 };
diff --git a/Engine/Source/Developer/FunctionalTesting/Classes/TraceQueryTestResults.h b/Engine/Source/Developer/FunctionalTesting/Classes/TraceQueryTestResults.h
--- a/Engine/Source/Developer/FunctionalTesting/Classes/TraceQueryTestResults.h
+++ b/Engine/Source/Developer/FunctionalTesting/Classes/TraceQueryTestResults.h
@@ -170,7 +170,7 @@ public:
 	FTraceChannelTestBatchOptions BatchOptions;
 
 	/** Output string value */
-	UFUNCTION(BlueprintCallable, Category = "Utility|Collision")
+	UFUNCTION(BlueprintCallable, Category = "Utilities|Collision")
 	FString ToString();
 
 	bool AssertEqual(const UTraceQueryTestResults* Expected, const FString& What, const UObject* ContextObject, AFunctionalTest& FunctionalTest) const;
diff --git a/Engine/Source/Runtime/Engine/Classes/Engine/LevelScriptActor.h b/Engine/Source/Runtime/Engine/Classes/Engine/LevelScriptActor.h
--- a/Engine/Source/Runtime/Engine/Classes/Engine/LevelScriptActor.h
+++ b/Engine/Source/Runtime/Engine/Classes/Engine/LevelScriptActor.h
@@ -29,7 +29,7 @@ class ENGINE_API ALevelScriptActor : public AActor
 	// --- Utility Functions ----------------------------
 	
 	/** Tries to find an event named "EventName" on all other levels, and calls it */
-	UFUNCTION(BlueprintCallable, meta=(BlueprintProtected = "true"), Category="Misc")
+	UFUNCTION(BlueprintCallable, meta=(BlueprintProtected = "true"), Category="Miscellaneous")
 	virtual bool RemoteEvent(FName EventName);
 
 	/**
diff --git a/Engine/Source/Runtime/Engine/Classes/Engine/SkeletalMesh.h b/Engine/Source/Runtime/Engine/Classes/Engine/SkeletalMesh.h
--- a/Engine/Source/Runtime/Engine/Classes/Engine/SkeletalMesh.h
+++ b/Engine/Source/Runtime/Engine/Classes/Engine/SkeletalMesh.h
@@ -940,7 +940,7 @@ public:
 	 * @param InSectionIndex Index to check
 	 * @param bCheckCorrespondingSections Whether to check corresponding sections for disabled sections
 	 */
-	UFUNCTION(BlueprintCallable, Category="Cloth")
+	UFUNCTION(BlueprintCallable, Category="Clothing Simulation")
 	bool IsSectionUsingCloth(int32 InSectionIndex, bool bCheckCorrespondingSections = true) const;
 
 	void CreateBodySetup();
diff --git a/Engine/Source/Runtime/Engine/Classes/GameFramework/PlayerController.h b/Engine/Source/Runtime/Engine/Classes/GameFramework/PlayerController.h
--- a/Engine/Source/Runtime/Engine/Classes/GameFramework/PlayerController.h
+++ b/Engine/Source/Runtime/Engine/Classes/GameFramework/PlayerController.h
@@ -1000,7 +1000,7 @@ public:
 	 * @param Cursor - the cursor to set the widget for
 	 * @param CursorWidget - the widget to set the cursor to
 	 */
-	UFUNCTION(BlueprintCallable, Category="UI")
+	UFUNCTION(BlueprintCallable, Category="User Interface")
 	void SetMouseCursorWidget(EMouseCursor::Type Cursor, class UUserWidget* CursorWidget);
 
 	/** Set the view target
diff --git a/Engine/Source/Runtime/Landscape/Classes/LandscapeComponent.h b/Engine/Source/Runtime/Landscape/Classes/LandscapeComponent.h
--- a/Engine/Source/Runtime/Landscape/Classes/LandscapeComponent.h
+++ b/Engine/Source/Runtime/Landscape/Classes/LandscapeComponent.h
@@ -529,7 +529,7 @@ public:
 	LANDSCAPE_API class UMaterialInstance* GetMaterialInstance(int32 InIndex, bool InDynamic = true) const;
 
 	/** Gets the landscape material instance dynamic for this component */
-	UFUNCTION(BlueprintCallable, Category = "Landscape Runtime | Material")
+	UFUNCTION(BlueprintCallable, Category = "Landscape|Runtime|Material")
 	class UMaterialInstanceDynamic* GetMaterialInstanceDynamic(int32 InIndex) const;
 
 	/** Get the landscape actor associated with this component. */
diff --git a/Engine/Source/Runtime/Landscape/Classes/LandscapeProxy.h b/Engine/Source/Runtime/Landscape/Classes/LandscapeProxy.h
--- a/Engine/Source/Runtime/Landscape/Classes/LandscapeProxy.h
+++ b/Engine/Source/Runtime/Landscape/Classes/LandscapeProxy.h
@@ -645,19 +645,19 @@ public:
 	 * @param bLowerHeights - Allow the landscape to be lowered down to the level of the spline. If both bRaiseHeights and bLowerHeights are false, no height modification of the landscape will be performed
 	 * @param PaintLayer - LayerInfo to paint, or none to skip painting. The landscape must be configured with the same layer info in one of its layers or this will do nothing!
 	 */
-	UFUNCTION(BlueprintCallable, Category = "Landscape Editor")
+	UFUNCTION(BlueprintCallable, Category = "Landscape|Editor")
 	void EditorApplySpline(USplineComponent* InSplineComponent, float StartWidth = 200, float EndWidth = 200, float StartSideFalloff = 200, float EndSideFalloff = 200, float StartRoll = 0, float EndRoll = 0, int32 NumSubdivisions = 20, bool bRaiseHeights = true, bool bLowerHeights = true, ULandscapeLayerInfoObject* PaintLayer = nullptr);
 
 	/** Set an MID texture parameter value for all landscape components. */
-	UFUNCTION(BlueprintCallable, Category = "Landscape Runtime|Material")
+	UFUNCTION(BlueprintCallable, Category = "Landscape|Runtime|Material")
 	void SetLandscapeMaterialTextureParameterValue(FName ParameterName, class UTexture* Value);
 
 	/** Set an MID vector parameter value for all landscape components. */
-	UFUNCTION(BlueprintCallable, meta = (Keywords = "SetColorParameterValue"), Category = "Landscape Runtime|Material")
+	UFUNCTION(BlueprintCallable, meta = (Keywords = "SetColorParameterValue"), Category = "Landscape|Runtime|Material")
 	void SetLandscapeMaterialVectorParameterValue(FName ParameterName, FLinearColor Value);
 
 	/** Set a MID scalar (float) parameter value for all landscape components. */
-	UFUNCTION(BlueprintCallable, meta = (Keywords = "SetFloatParameterValue"), Category = "Landscape Runtime|Material")
+	UFUNCTION(BlueprintCallable, meta = (Keywords = "SetFloatParameterValue"), Category = "Landscape|Runtime|Material")
 	void SetLandscapeMaterialScalarParameterValue(FName ParameterName, float Value);
 
 	// End blueprint functions
diff --git a/Engine/Source/Runtime/MovieSceneCapture/Public/MovieSceneCaptureEnvironment.h b/Engine/Source/Runtime/MovieSceneCapture/Public/MovieSceneCaptureEnvironment.h
--- a/Engine/Source/Runtime/MovieSceneCapture/Public/MovieSceneCaptureEnvironment.h
+++ b/Engine/Source/Runtime/MovieSceneCapture/Public/MovieSceneCaptureEnvironment.h
@@ -15,10 +15,10 @@ public:
 	GENERATED_BODY()
 
 	/** Get the frame number of the current capture */
-	UFUNCTION(BlueprintPure, Category="Cinematics|Capture")
+	UFUNCTION(BlueprintPure, Category="Cinematic|Capture")
 	static int32 GetCaptureFrameNumber();
 
 	/** Get the total elapsed time of the current capture in seconds */
-	UFUNCTION(BlueprintPure, Category="Cinematics|Capture")
+	UFUNCTION(BlueprintPure, Category="Cinematic|Capture")
 	static float GetCaptureElapsedTime();
 };
