[NO NON-PII COMMIT MESSAGE]

diff --git a/Engine/Source/Runtime/Landscape/Private/Landscape.cpp b/Engine/Source/Runtime/Landscape/Private/Landscape.cpp
--- a/Engine/Source/Runtime/Landscape/Private/Landscape.cpp
+++ b/Engine/Source/Runtime/Landscape/Private/Landscape.cpp
@@ -2055,13 +2055,13 @@ void ULandscapeInfo::RegisterActor(ALandscapeProxy* Proxy, bool bMapCheck)
 	// in case this Info object is not initialized yet
 	// initialized it with properties from passed actor
 	if (LandscapeGuid.IsValid() == false ||
-		(GetLandscapeProxy() == nullptr && ensure(LandscapeGuid == Proxy->GetLandscapeGuid()) && Proxy->HasValidRootComponent()))
+		(GetLandscapeProxy() == nullptr && ensure(LandscapeGuid == Proxy->GetLandscapeGuid())))
 	{
 		LandscapeGuid = Proxy->GetLandscapeGuid();
 		ComponentSizeQuads = Proxy->ComponentSizeQuads;
 		ComponentNumSubsections = Proxy->NumSubsections;
 		SubsectionSizeQuads = Proxy->SubsectionSizeQuads;
-		DrawScale = Proxy->GetRootComponent()->RelativeScale3D;
+		DrawScale = Proxy->HasValidRootComponent() ? Proxy->GetRootComponent()->RelativeScale3D : FVector(1.0f);
 	}
 
 	// check that passed actor matches all shared parameters
