[NO NON-PII COMMIT MESSAGE]

diff --git a/Engine/Source/Runtime/Landscape/Private/Landscape.cpp b/Engine/Source/Runtime/Landscape/Private/Landscape.cpp
--- a/Engine/Source/Runtime/Landscape/Private/Landscape.cpp
+++ b/Engine/Source/Runtime/Landscape/Private/Landscape.cpp
@@ -2061,7 +2061,7 @@ void ULandscapeInfo::RegisterActor(ALandscapeProxy* Proxy, bool bMapCheck)
 		ComponentSizeQuads = Proxy->ComponentSizeQuads;
 		ComponentNumSubsections = Proxy->NumSubsections;
 		SubsectionSizeQuads = Proxy->SubsectionSizeQuads;
-		DrawScale = Proxy->HasValidRootComponent() ? Proxy->GetRootComponent()->RelativeScale3D : FVector(1.0f);
+		DrawScale = Proxy->HasValidRootComponent() ? Proxy->GetRootComponent()->RelativeScale3D : FVector(100.0f);
 	}
 
 	// check that passed actor matches all shared parameters
