[NO NON-PII COMMIT MESSAGE]

diff --git a/Engine/Source/Runtime/Landscape/Private/Landscape.cpp b/Engine/Source/Runtime/Landscape/Private/Landscape.cpp
--- a/Engine/Source/Runtime/Landscape/Private/Landscape.cpp
+++ b/Engine/Source/Runtime/Landscape/Private/Landscape.cpp
@@ -2061,7 +2061,7 @@ void ULandscapeInfo::RegisterActor(ALandscapeProxy* Proxy, bool bMapCheck)
 		ComponentSizeQuads = Proxy->ComponentSizeQuads;
 		ComponentNumSubsections = Proxy->NumSubsections;
 		SubsectionSizeQuads = Proxy->SubsectionSizeQuads;
-		DrawScale = Proxy->HasValidRootComponent() ? Proxy->GetRootComponent()->RelativeScale3D : FVector(100.0f);
+		DrawScale = Proxy->GetRootComponent() != nullptr ? Proxy->GetRootComponent()->RelativeScale3D : FVector(100.0f);
 	}
 
 	// check that passed actor matches all shared parameters
@@ -2070,7 +2070,7 @@ void ULandscapeInfo::RegisterActor(ALandscapeProxy* Proxy, bool bMapCheck)
 	check(ComponentNumSubsections == Proxy->NumSubsections);
 	check(SubsectionSizeQuads == Proxy->SubsectionSizeQuads);
 
-	if (Proxy->HasValidRootComponent() && !DrawScale.Equals(Proxy->GetRootComponent()->RelativeScale3D))
+	if (Proxy->GetRootComponent() != nullptr && !DrawScale.Equals(Proxy->GetRootComponent()->RelativeScale3D))
 	{
 		UE_LOG(LogLandscape, Warning, TEXT("Landscape proxy (%s) scale (%s) does not match to main actor scale (%s)."),
 			*Proxy->GetName(), *Proxy->GetRootComponent()->RelativeScale3D.ToCompactString(), *DrawScale.ToCompactString());
