[NO NON-PII COMMIT MESSAGE]

diff --git a/Engine/Source/Runtime/Landscape/Private/LandscapeGrass.cpp b/Engine/Source/Runtime/Landscape/Private/LandscapeGrass.cpp
--- a/Engine/Source/Runtime/Landscape/Private/LandscapeGrass.cpp
+++ b/Engine/Source/Runtime/Landscape/Private/LandscapeGrass.cpp
@@ -1943,7 +1943,9 @@ void ALandscapeProxy::FlushGrassComponents(const TSet<ULandscapeComponent*>* Onl
 		}
 
 #if WITH_EDITOR
-		if (GIsEditor && bFlushGrassMaps && GetWorld() && GetWorld()->Scene->GetFeatureLevel() >= ERHIFeatureLevel::SM4)
+		UWorld* World = GetWorld();
+
+		if (GIsEditor && bFlushGrassMaps && World != nullptr && World->Scene != nullptr && World->Scene->GetFeatureLevel() >= ERHIFeatureLevel::SM4)
 		{
 			// Clear GrassMaps
 			TInlineComponentArray<ULandscapeComponent*> LandComps;
