Fixed crash during cooking for certain landscapes with collision data out of sync with render data.

diff --git a/Engine/Source/Runtime/Landscape/Private/LandscapeGrass.cpp b/Engine/Source/Runtime/Landscape/Private/LandscapeGrass.cpp
--- a/Engine/Source/Runtime/Landscape/Private/LandscapeGrass.cpp
+++ b/Engine/Source/Runtime/Landscape/Private/LandscapeGrass.cpp
@@ -1900,7 +1900,7 @@ void ALandscapeProxy::FlushGrassComponents(const TSet<ULandscapeComponent*>* Onl
 			}
 		}
 #if WITH_EDITOR
-		if (GIsEditor && bFlushGrassMaps && GetWorld() && GetWorld()->Scene->GetFeatureLevel() >= ERHIFeatureLevel::SM4)
+		if (GIsEditor && bFlushGrassMaps && GetWorld() && GetWorld()->Scene && GetWorld()->Scene->GetFeatureLevel() >= ERHIFeatureLevel::SM4)
 		{
 			for (ULandscapeComponent* Component : *OnlyForComponents)
 			{
@@ -1943,7 +1943,7 @@ void ALandscapeProxy::FlushGrassComponents(const TSet<ULandscapeComponent*>* Onl
 		}
 
 #if WITH_EDITOR
-		if (GIsEditor && bFlushGrassMaps && GetWorld() && GetWorld()->Scene->GetFeatureLevel() >= ERHIFeatureLevel::SM4)
+		if (GIsEditor && bFlushGrassMaps && GetWorld() && GetWorld()->Scene && GetWorld()->Scene->GetFeatureLevel() >= ERHIFeatureLevel::SM4)
 		{
 			// Clear GrassMaps
 			for (UActorComponent* Component : GetComponents())
