Was missing UpdateUniformExpressionCacheIfNeeded in the Grassmap generation, to fit with the new rendering change.

diff --git a/Engine/Source/Runtime/Landscape/Private/LandscapeGrass.cpp b/Engine/Source/Runtime/Landscape/Private/LandscapeGrass.cpp
--- a/Engine/Source/Runtime/Landscape/Private/LandscapeGrass.cpp
+++ b/Engine/Source/Runtime/Landscape/Private/LandscapeGrass.cpp
@@ -507,6 +507,8 @@ public:
 			for (auto& ComponentInfo : ComponentInfos)
 			{
 				const FMeshBatch& Mesh = ComponentInfo.SceneProxy->GetGrassMeshBatch();
+				Mesh.MaterialRenderProxy->UpdateUniformExpressionCacheIfNeeded(View->GetFeatureLevel());
+
 				PassMeshProcessor.AddMeshBatch(Mesh, DefaultBatchElementMask, NumPasses, ComponentInfo.ViewOffset, PassOffsetX, ComponentInfo.SceneProxy);
 			}
 		});
