Adapted pull request #5564: Fixed crash while building grass in editor.

diff --git a/Engine/Source/Runtime/Landscape/Private/LandscapeGrass.cpp b/Engine/Source/Runtime/Landscape/Private/LandscapeGrass.cpp
--- a/Engine/Source/Runtime/Landscape/Private/LandscapeGrass.cpp
+++ b/Engine/Source/Runtime/Landscape/Private/LandscapeGrass.cpp
@@ -1894,11 +1894,11 @@ struct FAsyncGrassBuilder : public FGrassBuilderBase
 		int32 X2 = FMath::CeilToInt(TestX);
 		int32 Y2 = FMath::CeilToInt(TestY);
 
-		// Min is to prevent the sampling of the final column from overflowing
-		int32 IdxX1 = FMath::Min<int32>(X1, GrassData.GetStride() - 1);
-		int32 IdxY1 = FMath::Min<int32>(Y1, GrassData.GetStride() - 1);
-		int32 IdxX2 = FMath::Min<int32>(X2, GrassData.GetStride() - 1);
-		int32 IdxY2 = FMath::Min<int32>(Y2, GrassData.GetStride() - 1);
+		// Clamp to prevent the sampling of the final columns from overflowing
+		int32 IdxX1 = FMath::Clamp<int32>(X1, 0, GrassData.GetStride() - 1);
+		int32 IdxY1 = FMath::Clamp<int32>(Y1, 0, GrassData.GetStride() - 1);
+		int32 IdxX2 = FMath::Clamp<int32>(X2, 0, GrassData.GetStride() - 1);
+		int32 IdxY2 = FMath::Clamp<int32>(Y2, 0, GrassData.GetStride() - 1);
 
 		float LerpX = FMath::Fractional(TestX);
 		float LerpY = FMath::Fractional(TestY);
