Change SetRenderTarget by Begin/End RenderPass.

diff --git a/Engine/Source/Runtime/Landscape/Private/LandscapeEditProcedural.cpp b/Engine/Source/Runtime/Landscape/Private/LandscapeEditProcedural.cpp
--- a/Engine/Source/Runtime/Landscape/Private/LandscapeEditProcedural.cpp
+++ b/Engine/Source/Runtime/Landscape/Private/LandscapeEditProcedural.cpp
@@ -536,6 +536,9 @@ public:
 		GraphicsPSOInit.RasterizerState = TStaticRasterizerState<>::GetRHI();
 		GraphicsPSOInit.DepthStencilState = TStaticDepthStencilState<false, CF_Always>::GetRHI();
 
+		FRHIRenderPassInfo RenderPassInfo(ViewFamily.RenderTarget->GetRenderTargetTexture(), CurrentMip == 0 ? ERenderTargetActions::Clear_Store : ERenderTargetActions::Load_Store, nullptr, 0, 0);
+		InRHICmdList.BeginRenderPass(RenderPassInfo, TEXT("DrawProceduralHeightmaps"));
+
 		if (CurrentMip == 0)
 		{
 			// Setup Shaders
@@ -546,7 +549,6 @@ public:
 			GraphicsPSOInit.BoundShaderState.PixelShaderRHI = GETSAFERHISHADER_PIXEL(*PixelShader);
 
 			InRHICmdList.SetViewport(View->UnscaledViewRect.Min.X, View->UnscaledViewRect.Min.Y, 0.0f, View->UnscaledViewRect.Max.X, View->UnscaledViewRect.Max.Y, 1.0f);
-			SetRenderTarget(InRHICmdList, ViewFamily.RenderTarget->GetRenderTargetTexture(), nullptr, InClearRT ? ESimpleRenderTargetMode::EClearColorAndDepth : ESimpleRenderTargetMode::EExistingColorAndDepth, FExclusiveDepthStencil::DepthWrite_StencilWrite, true);
 
 			InRHICmdList.ApplyCachedRenderTargets(GraphicsPSOInit);
 			SetGraphicsPipelineState(InRHICmdList, GraphicsPSOInit);
@@ -565,7 +567,6 @@ public:
 			GraphicsPSOInit.BoundShaderState.PixelShaderRHI = GETSAFERHISHADER_PIXEL(*PixelShader);
 
 			InRHICmdList.SetViewport(0.0f, 0.0f, 0.0f, WriteRenderTargetSize.X, WriteRenderTargetSize.Y, 1.0f);
-			SetRenderTarget(InRHICmdList, ViewFamily.RenderTarget->GetRenderTargetTexture(), nullptr, ESimpleRenderTargetMode::EClearColorAndDepth, FExclusiveDepthStencil::DepthWrite_StencilWrite, true);
 
 			InRHICmdList.ApplyCachedRenderTargets(GraphicsPSOInit);
 			SetGraphicsPipelineState(InRHICmdList, GraphicsPSOInit);
@@ -581,6 +582,8 @@ public:
 
 		InRHICmdList.DrawPrimitive(0, PrimitiveCount, 1);
 
+		InRHICmdList.EndRenderPass();
+
 		VertexDeclaration.ReleaseResource();
 		VertexBufferResource.ReleaseResource();
 	}
