When SetParentEditorOnly is called we no longer recompute shader in most case as we will call UpdateStaticPermutation that will do so.

diff --git a/Engine/Source/Runtime/Landscape/Private/Landscape.cpp b/Engine/Source/Runtime/Landscape/Private/Landscape.cpp
--- a/Engine/Source/Runtime/Landscape/Private/Landscape.cpp
+++ b/Engine/Source/Runtime/Landscape/Private/Landscape.cpp
@@ -75,10 +75,17 @@ DEFINE_STAT(STAT_LandscapeComponentUsingSubSectionDrawCalls);
 DEFINE_STAT(STAT_LandscapeDrawCalls);
 DEFINE_STAT(STAT_LandscapeTriangles);
 
+DEFINE_STAT(STAT_LandscapeRegenerateProcedural_RenderThread);
+DEFINE_STAT(STAT_LandscapeRegenerateProceduralDrawCalls);
+
 DEFINE_STAT(STAT_LandscapeRegenerateProceduralHeightmaps);
-DEFINE_STAT(STAT_LandscapeRegenerateProceduralHeightmaps_RenderThread);
 DEFINE_STAT(STAT_LandscapeResolveProceduralHeightmap);
-DEFINE_STAT(STAT_LandscapeRegenerateProceduralHeightmapsDrawCalls);
+
+DEFINE_STAT(STAT_LandscapeProceduralUpdateMaterialInstance);
+DEFINE_STAT(STAT_LandscapeReallocateProceduralWeightmaps);
+
+DEFINE_STAT(STAT_LandscapeResolveProceduralWeightmap);
+DEFINE_STAT(STAT_LandscapeRegenerateProceduralWeightmaps);
 
 DEFINE_STAT(STAT_LandscapeVertexMem);
 DEFINE_STAT(STAT_LandscapeOccluderMem);
@@ -161,6 +168,7 @@ ULandscapeComponent::ULandscapeComponent(const FObjectInitializer& ObjectInitial
 
 #if WITH_EDITORONLY_DATA
 	EditToolRenderData = FLandscapeEditToolRenderData();
+	CurrentProceduralLayerName = NAME_None;
 #endif
 
 	LpvBiasMultiplier = 0.0f; // Bias is 0 for landscape, since it's single sided
@@ -213,6 +221,41 @@ UMaterialInstanceDynamic* ULandscapeComponent::GetMaterialInstanceDynamic(int32
 void ULandscapeComponent::AddReferencedObjects(UObject* InThis, FReferenceCollector& Collector)
 {
 	ULandscapeComponent* This = CastChecked<ULandscapeComponent>(InThis);
+
+#if WITH_EDITORONLY_DATA
+	Collector.AddReferencedObject(This->CurrentEditingHeightmapTexture);
+
+	if (This->CurrentEditingWeightmapLayerAllocations != nullptr)
+	{
+		for (FWeightmapLayerAllocationInfo& WeightmapLayerAllocation : *This->CurrentEditingWeightmapLayerAllocations)
+		{
+			Collector.AddReferencedObject(WeightmapLayerAllocation.LayerInfo);
+		}
+	}
+
+	if (This->CurrentEditingWeightmapTexturesUsage != nullptr)
+	{
+		for (FLandscapeWeightmapUsage* TextureUsage : *This->CurrentEditingWeightmapTexturesUsage)
+		{
+			if (TextureUsage != nullptr)
+			{
+				for (int32 i = 0; i < 4; ++i)
+				{
+					Collector.AddReferencedObject(TextureUsage->ChannelUsage[i]);
+				}
+			}
+		}
+	}	
+
+	if (This->CurrentEditingWeightmapTextures != nullptr)
+	{
+		for (UTexture2D* Weightmap : *This->CurrentEditingWeightmapTextures)
+		{
+			Collector.AddReferencedObject(Weightmap);
+		}
+	}
+#endif
+
 	Super::AddReferencedObjects(This, Collector);
 }
 
@@ -537,21 +580,23 @@ void ULandscapeComponent::GetLayerDebugColorKey(int32& R, int32& G, int32& B) co
 			if (LayerStruct.DebugColorChannel > 0
 				&& LayerStruct.LayerInfoObj)
 			{
-				for (int32 LayerIdx = 0; LayerIdx < WeightmapLayerAllocations.Num(); LayerIdx++)
+				const TArray<FWeightmapLayerAllocationInfo>& ComponentWeightmapLayerAllocations = GetWeightmapLayerAllocations(true);
+
+				for (int32 LayerIdx = 0; LayerIdx < ComponentWeightmapLayerAllocations.Num(); LayerIdx++)
 				{
-					if (WeightmapLayerAllocations[LayerIdx].LayerInfo == LayerStruct.LayerInfoObj)
+					if (ComponentWeightmapLayerAllocations[LayerIdx].LayerInfo == LayerStruct.LayerInfoObj)
 					{
 						if (LayerStruct.DebugColorChannel & 1) // R
 						{
-							R = (WeightmapLayerAllocations[LayerIdx].WeightmapTextureIndex * 4 + WeightmapLayerAllocations[LayerIdx].WeightmapTextureChannel);
+							R = (ComponentWeightmapLayerAllocations[LayerIdx].WeightmapTextureIndex * 4 + ComponentWeightmapLayerAllocations[LayerIdx].WeightmapTextureChannel);
 						}
 						if (LayerStruct.DebugColorChannel & 2) // G
 						{
-							G = (WeightmapLayerAllocations[LayerIdx].WeightmapTextureIndex * 4 + WeightmapLayerAllocations[LayerIdx].WeightmapTextureChannel);
+							G = (ComponentWeightmapLayerAllocations[LayerIdx].WeightmapTextureIndex * 4 + ComponentWeightmapLayerAllocations[LayerIdx].WeightmapTextureChannel);
 						}
 						if (LayerStruct.DebugColorChannel & 4) // B
 						{
-							B = (WeightmapLayerAllocations[LayerIdx].WeightmapTextureIndex * 4 + WeightmapLayerAllocations[LayerIdx].WeightmapTextureChannel);
+							B = (ComponentWeightmapLayerAllocations[LayerIdx].WeightmapTextureIndex * 4 + ComponentWeightmapLayerAllocations[LayerIdx].WeightmapTextureChannel);
 						}
 						break;
 					}
@@ -942,7 +987,12 @@ ALandscape::ALandscape(const FObjectInitializer& ObjectInitializer)
 #if WITH_EDITORONLY_DATA
 	bLockLocation = false;
 	PreviousExperimentalLandscapeProcedural = false;
+	WasCompilingShaders = false;
 	ProceduralContentUpdateFlags = 0;
+	CombinedProcLayerWeightmapAllLayersResource = nullptr;
+	CurrentProcLayerWeightmapAllLayersResource = nullptr;
+	WeightmapScratchExtractLayerTextureResource = nullptr;
+	WeightmapScratchPackLayerTextureResource = nullptr;
 #endif // WITH_EDITORONLY_DATA
 }
 
@@ -1037,9 +1087,9 @@ void ULandscapeComponent::GetGeneratedTexturesAndMaterialInstances(TArray<UObjec
 		OutTexturesAndMaterials.Add(CurrentEditingHeightmapTexture);
 	}
 
-	for (auto* Tex : WeightmapTextures)
+	for (UTexture2D* Weightmap : WeightmapTextures)
 	{
-		OutTexturesAndMaterials.Add(Tex);
+		OutTexturesAndMaterials.Add(Weightmap);
 	}
 
 	if (XYOffsetmapTexture)
@@ -1165,18 +1215,21 @@ void ULandscapeComponent::BeginDestroy()
 			if (WeightmapTextures.IsValidIndex(WeightmapIndex))
 			{
 				UTexture2D* WeightmapTexture = WeightmapTextures[WeightmapIndex];
-				FLandscapeWeightmapUsage* Usage = Proxy->WeightmapUsageMap.Find(WeightmapTexture);
-				if (Usage != nullptr)
+				FLandscapeWeightmapUsage** Usage = Proxy->WeightmapUsageMap.Find(WeightmapTexture);
+				if (Usage != nullptr && (*Usage) != nullptr)
 				{
-					Usage->ChannelUsage[WeightmapLayerAllocations[LayerIdx].WeightmapTextureChannel] = nullptr;
+					(*Usage)->ChannelUsage[WeightmapLayerAllocations[LayerIdx].WeightmapTextureChannel] = nullptr;
 
-					if (Usage->FreeChannelCount() == 4)
+					if ((*Usage)->FreeChannelCount() == 4)
 					{
 						Proxy->WeightmapUsageMap.Remove(WeightmapTexture);
+						delete (*Usage);
 					}
 				}
 			}
 		}
+
+		WeightmapTexturesUsage.Reset();
 	}
 #endif
 }
@@ -1301,14 +1354,74 @@ UTexture2D* ULandscapeComponent::GetHeightmap(bool InReturnCurrentEditingHeightm
 	return HeightmapTexture;
 }
 
+const TArray<UTexture2D*>& ULandscapeComponent::GetWeightmapTextures(bool InReturnCurrentEditingWeightmap) const
+{
+#if WITH_EDITORONLY_DATA
+	if (InReturnCurrentEditingWeightmap && GetMutableDefault<UEditorExperimentalSettings>()->bProceduralLandscape)
+	{
+		if (CurrentEditingWeightmapTextures != nullptr)
+		{
+			return *CurrentEditingWeightmapTextures;
+		}
+	}
+#endif
+
+	return WeightmapTextures;
+}
+
+TArray<UTexture2D*>& ULandscapeComponent::GetWeightmapTextures(bool InReturnCurrentEditingWeightmap)
+{
+#if WITH_EDITORONLY_DATA
+	if (InReturnCurrentEditingWeightmap && GetMutableDefault<UEditorExperimentalSettings>()->bProceduralLandscape)
+	{
+		if (CurrentEditingWeightmapTextures != nullptr)
+		{
+			return *CurrentEditingWeightmapTextures;
+		}
+	}
+#endif
+
+	return WeightmapTextures;
+}
+
+const TArray<FWeightmapLayerAllocationInfo>& ULandscapeComponent::GetWeightmapLayerAllocations(bool InReturnCurrentEditingWeightmap) const
+{
+#if WITH_EDITORONLY_DATA
+	if (InReturnCurrentEditingWeightmap && GetMutableDefault<UEditorExperimentalSettings>()->bProceduralLandscape)
+	{
+		if (CurrentEditingWeightmapLayerAllocations != nullptr)
+		{
+			return *CurrentEditingWeightmapLayerAllocations;
+		}
+	}
+#endif
+
+	return WeightmapLayerAllocations;
+}
+
+TArray<FWeightmapLayerAllocationInfo>& ULandscapeComponent::GetWeightmapLayerAllocations(bool InReturnCurrentEditingWeightmap)
+{
+#if WITH_EDITORONLY_DATA
+	if (InReturnCurrentEditingWeightmap && GetMutableDefault<UEditorExperimentalSettings>()->bProceduralLandscape)
+	{
+		if (CurrentEditingWeightmapLayerAllocations != nullptr)
+		{
+			return *CurrentEditingWeightmapLayerAllocations;
+		}
+	}
+#endif
+
+	return WeightmapLayerAllocations;
+}
+
 #if WITH_EDITOR
+
 void ULandscapeComponent::SetCurrentEditingHeightmap(UTexture2D* InNewHeightmap)
 {
 #if WITH_EDITORONLY_DATA
 	CurrentEditingHeightmapTexture = InNewHeightmap;
 #endif
 }
-#endif
 
 void ULandscapeComponent::SetHeightmap(UTexture2D* NewHeightmap)
 {
@@ -1316,6 +1429,104 @@ void ULandscapeComponent::SetHeightmap(UTexture2D* NewHeightmap)
 	HeightmapTexture = NewHeightmap;
 }
 
+void ULandscapeComponent::SetWeightmapTextures(const TArray<UTexture2D*>& InNewWeightmapTextures, bool InApplyToCurrentEditingWeightmap)
+{
+#if WITH_EDITORONLY_DATA
+	if (InApplyToCurrentEditingWeightmap && CurrentEditingWeightmapTextures != nullptr)
+	{
+		CurrentEditingWeightmapTextures->Reset(InNewWeightmapTextures.Num());
+		CurrentEditingWeightmapTextures->Append(InNewWeightmapTextures);
+	}
+	else
+#endif
+	{
+		WeightmapTextures = InNewWeightmapTextures;
+	}
+}
+
+void ULandscapeComponent::SetCurrentEditingWeightmaps(TArray<UTexture2D*>* InNewWeightmapTextures)
+{
+#if WITH_EDITORONLY_DATA
+	CurrentEditingWeightmapTextures = InNewWeightmapTextures;
+#endif
+}
+
+void ULandscapeComponent::SetWeightmapLayerAllocations(const TArray<FWeightmapLayerAllocationInfo>& InNewWeightmapLayerAllocations)
+{
+	WeightmapLayerAllocations = InNewWeightmapLayerAllocations;
+}
+
+
+void ULandscapeComponent::SetCurrentEditingWeightmapLayerAllocations(TArray<FWeightmapLayerAllocationInfo>* InNewWeightmapLayerAllocations)
+{
+#if WITH_EDITORONLY_DATA
+	CurrentEditingWeightmapLayerAllocations = InNewWeightmapLayerAllocations;
+#endif
+}
+
+TArray<FLandscapeWeightmapUsage*>& ULandscapeComponent::GetWeightmapTexturesUsage(bool InReturnCurrentEditingWeightmap)
+{
+#if WITH_EDITORONLY_DATA
+	if (InReturnCurrentEditingWeightmap && GetMutableDefault<UEditorExperimentalSettings>()->bProceduralLandscape)
+	{
+		if (CurrentEditingWeightmapTexturesUsage != nullptr)
+		{
+			return *CurrentEditingWeightmapTexturesUsage;
+		}
+	}
+#endif
+
+	return WeightmapTexturesUsage;
+}
+
+const TArray<FLandscapeWeightmapUsage*>& ULandscapeComponent::GetWeightmapTexturesUsage(bool InReturnCurrentEditingWeightmap) const
+{
+#if WITH_EDITORONLY_DATA
+	if (InReturnCurrentEditingWeightmap && GetMutableDefault<UEditorExperimentalSettings>()->bProceduralLandscape)
+	{
+		if (CurrentEditingWeightmapTexturesUsage != nullptr)
+		{
+			return *CurrentEditingWeightmapTexturesUsage;
+		}
+	}
+#endif
+
+	return WeightmapTexturesUsage;
+}
+
+void ULandscapeComponent::SetCurrentEditingWeightmapTexturesUsage(TArray<FLandscapeWeightmapUsage*>* InNewWeightmapTexturesUsage)
+{
+#if WITH_EDITORONLY_DATA
+	CurrentEditingWeightmapTexturesUsage = InNewWeightmapTexturesUsage;
+#endif
+}
+
+void ULandscapeComponent::SetWeightmapTexturesUsage(const TArray<FLandscapeWeightmapUsage*>& InNewWeightmapTexturesUsage, bool InApplyToCurrentEditingWeightmap )
+{
+#if WITH_EDITORONLY_DATA
+	if (InApplyToCurrentEditingWeightmap && CurrentEditingWeightmapTexturesUsage != nullptr)
+	{
+		CurrentEditingWeightmapTexturesUsage->Reset(InNewWeightmapTexturesUsage.Num());
+		CurrentEditingWeightmapTexturesUsage->Append(InNewWeightmapTexturesUsage);
+	}
+	else
+#endif
+	{
+		WeightmapTexturesUsage = InNewWeightmapTexturesUsage;
+	}
+}
+
+void ULandscapeComponent::SetCurrentProceduralLayerName(const FName& InName)
+{
+	CurrentProceduralLayerName = InName;
+}
+
+const FName& ULandscapeComponent::GetCurrentProceduralLayerName() const
+{
+	return CurrentProceduralLayerName;
+}
+#endif
+
 void ALandscapeProxy::PostRegisterAllComponents()
 {
 	Super::PostRegisterAllComponents();
@@ -1363,6 +1574,21 @@ FArchive& operator<<(FArchive& Ar, FLandscapeWeightmapUsage& U)
 	return Ar << U.ChannelUsage[0] << U.ChannelUsage[1] << U.ChannelUsage[2] << U.ChannelUsage[3];
 }
 
+FArchive& operator<<(FArchive& Ar, FLandscapeWeightmapUsage* U)
+{
+	if (U != nullptr)
+	{
+		return Ar << U->ChannelUsage[0] << U->ChannelUsage[1] << U->ChannelUsage[2] << U->ChannelUsage[3];
+	}
+
+	return Ar;
+}
+
+FArchive& operator<<(FArchive& Ar, FWeightmapLayerAllocationInfo& U)
+{
+	return Ar << U.LayerInfo << U.WeightmapTextureChannel << U.WeightmapTextureIndex;
+}
+
 #if WITH_EDITORONLY_DATA
 FArchive& operator<<(FArchive& Ar, FLandscapeAddCollision& U)
 {
@@ -1435,31 +1661,12 @@ void ALandscape::PostLoad()
 				break;
 			}
 		}
-
-		if (GetMutableDefault<UEditorExperimentalSettings>()->bProceduralLandscape)
-		{
-			FEditorDelegates::PreSaveWorld.AddUObject(this, &ALandscape::OnPreSaveWorld);
-			FEditorDelegates::PostSaveWorld.AddUObject(this, &ALandscape::OnPostSaveWorld);
-		}
 #endif
 	}
 
 	Super::PostLoad();
 }
 
-void ALandscape::BeginDestroy()
-{
-#if WITH_EDITOR
-	if (GetMutableDefault<UEditorExperimentalSettings>()->bProceduralLandscape)
-	{
-		FEditorDelegates::PreSaveWorld.RemoveAll(this);
-		FEditorDelegates::PostSaveWorld.RemoveAll(this);
-	}
-#endif
-
-	Super::BeginDestroy();
-}
-
 #if WITH_EDITOR
 void ALandscapeProxy::OnFeatureLevelChanged(ERHIFeatureLevel::Type NewFeatureLevel)
 {
@@ -1547,10 +1754,10 @@ void ALandscapeProxy::AddReferencedObjects(UObject* InThis, FReferenceCollector&
 	for (auto It = This->WeightmapUsageMap.CreateIterator(); It; ++It)
 	{
 		Collector.AddReferencedObject(It.Key(), This);
-		Collector.AddReferencedObject(It.Value().ChannelUsage[0], This);
-		Collector.AddReferencedObject(It.Value().ChannelUsage[1], This);
-		Collector.AddReferencedObject(It.Value().ChannelUsage[2], This);
-		Collector.AddReferencedObject(It.Value().ChannelUsage[3], This);
+		Collector.AddReferencedObject(It.Value()->ChannelUsage[0], This);
+		Collector.AddReferencedObject(It.Value()->ChannelUsage[1], This);
+		Collector.AddReferencedObject(It.Value()->ChannelUsage[2], This);
+		Collector.AddReferencedObject(It.Value()->ChannelUsage[3], This);
 	}
 }
 
@@ -1706,9 +1913,11 @@ bool ULandscapeInfo::UpdateLayerInfoMap(ALandscapeProxy* Proxy /*= nullptr*/, bo
 						}
 					}
 
-					for (int32 AllocationIndex = 0; AllocationIndex < Component->WeightmapLayerAllocations.Num(); AllocationIndex++)
+					TArray<FWeightmapLayerAllocationInfo>& ComponentWeightmapLayerAllocations = Component->GetWeightmapLayerAllocations();
+
+					for (int32 AllocationIndex = 0; AllocationIndex < ComponentWeightmapLayerAllocations.Num(); AllocationIndex++)
 					{
-						ULandscapeLayerInfoObject* LayerInfo = Component->WeightmapLayerAllocations[AllocationIndex].LayerInfo;
+						ULandscapeLayerInfoObject* LayerInfo = ComponentWeightmapLayerAllocations[AllocationIndex].LayerInfo;
 						if (LayerInfo)
 						{
 							int32 LayerInfoIndex = GetLayerInfoIndex(LayerInfo);
@@ -2757,47 +2966,9 @@ void ALandscape::TickActor(float DeltaTime, ELevelTick TickType, FActorTickFunct
 	Super::TickActor(DeltaTime, TickType, ThisTickFunction);
 
 #if WITH_EDITOR
-	UWorld* World = GetWorld();
-	if (GIsEditor && World && !World->IsPlayInEditor())
+	if (GIsEditor)
 	{
-		if (GetMutableDefault<UEditorExperimentalSettings>()->bProceduralLandscape)
-		{
-			if (PreviousExperimentalLandscapeProcedural != GetMutableDefault<UEditorExperimentalSettings>()->bProceduralLandscape)
-			{
-				PreviousExperimentalLandscapeProcedural = GetMutableDefault<UEditorExperimentalSettings>()->bProceduralLandscape;
-
-				RequestProceduralContentUpdate(EProceduralContentUpdateFlag::All_Setup);
-			}
-
-			RegenerateProceduralContent();
-		}
-		else
-		{
-			if (PreviousExperimentalLandscapeProcedural != GetMutableDefault<UEditorExperimentalSettings>()->bProceduralLandscape)
-			{
-				PreviousExperimentalLandscapeProcedural = GetMutableDefault<UEditorExperimentalSettings>()->bProceduralLandscape;
-
-				for (auto& ItPair : RenderDataPerHeightmap)
-				{
-					FRenderDataPerHeightmap& HeightmapRenderData = ItPair.Value;
-
-					if (HeightmapRenderData.HeightmapsCPUReadBack != nullptr)
-					{
-						BeginReleaseResource(HeightmapRenderData.HeightmapsCPUReadBack);
-					}
-				}
-
-				FlushRenderingCommands();
-
-				for (auto& ItPair : RenderDataPerHeightmap)
-				{
-					FRenderDataPerHeightmap& HeightmapRenderData = ItPair.Value;
-
-					delete HeightmapRenderData.HeightmapsCPUReadBack;
-					HeightmapRenderData.HeightmapsCPUReadBack = nullptr;
-				}
-			}
-		}
+		TickProcedural(DeltaTime, TickType, ThisTickFunction);
 	}
 #endif
 }
diff --git a/Engine/Source/Runtime/Landscape/Private/LandscapeBPCustomBrush.cpp b/Engine/Source/Runtime/Landscape/Private/LandscapeBPCustomBrush.cpp
--- a/Engine/Source/Runtime/Landscape/Private/LandscapeBPCustomBrush.cpp
+++ b/Engine/Source/Runtime/Landscape/Private/LandscapeBPCustomBrush.cpp
@@ -67,6 +67,11 @@ void ALandscapeBlueprintCustomBrush::SetIsInitialized(bool InIsInitialized)
 	bIsInitialized = InIsInitialized;
 }
 
+bool ALandscapeBlueprintCustomBrush::IsAffectingWeightmapLayer(const FName& InLayerName) const
+{
+	return AffectedWeightmapLayers.Contains(InLayerName);
+}
+
 void ALandscapeBlueprintCustomBrush::PostEditMove(bool bFinished)
 {
 	Super::PostEditMove(bFinished);
@@ -152,7 +157,7 @@ void ALandscapeBlueprintCustomBrush::PostEditChangeProperty(FPropertyChangedEven
 
 	if (OwningLandscape != nullptr)
 	{
-		OwningLandscape->RequestProceduralContentUpdate(EProceduralContentUpdateFlag::All);
+		OwningLandscape->RequestProceduralContentUpdate(EProceduralContentUpdateFlag::All, true); // For now since we dont have brush bounds, we have to force an update of all components, otherwise we will have not right layer being rendered as materials are not updated
 	}
 }
 #endif
diff --git a/Engine/Source/Runtime/Landscape/Private/LandscapeDataAccess.cpp b/Engine/Source/Runtime/Landscape/Private/LandscapeDataAccess.cpp
--- a/Engine/Source/Runtime/Landscape/Private/LandscapeDataAccess.cpp
+++ b/Engine/Source/Runtime/Landscape/Private/LandscapeDataAccess.cpp
@@ -74,9 +74,12 @@ LANDSCAPE_API void FLandscapeComponentDataInterface::GetHeightmapTextureData(TAr
 LANDSCAPE_API bool FLandscapeComponentDataInterface::GetWeightmapTextureData(ULandscapeLayerInfoObject* LayerInfo, TArray<uint8>& OutData)
 {
 	int32 LayerIdx = INDEX_NONE;
-	for (int32 Idx = 0; Idx < Component->WeightmapLayerAllocations.Num(); Idx++)
+	TArray<FWeightmapLayerAllocationInfo>& ComponentWeightmapLayerAllocations = Component->GetWeightmapLayerAllocations();
+	TArray<UTexture2D*>& ComponentWeightmapTextures = Component->GetWeightmapTextures();
+
+	for (int32 Idx = 0; Idx < ComponentWeightmapLayerAllocations.Num(); Idx++)
 	{
-		if (Component->WeightmapLayerAllocations[Idx].LayerInfo == LayerInfo)
+		if (ComponentWeightmapLayerAllocations[Idx].LayerInfo == LayerInfo)
 		{
 			LayerIdx = Idx;
 			break;
@@ -86,11 +89,11 @@ LANDSCAPE_API bool FLandscapeComponentDataInterface::GetWeightmapTextureData(ULa
 	{
 		return false;
 	}
-	if (Component->WeightmapLayerAllocations[LayerIdx].WeightmapTextureIndex >= Component->WeightmapTextures.Num())
+	if (ComponentWeightmapLayerAllocations[LayerIdx].WeightmapTextureIndex >= ComponentWeightmapTextures.Num())
 	{
 		return false;
 	}
-	if (Component->WeightmapLayerAllocations[LayerIdx].WeightmapTextureChannel >= 4)
+	if (ComponentWeightmapLayerAllocations[LayerIdx].WeightmapTextureChannel >= 4)
 	{
 		return false;
 	}
@@ -99,19 +102,19 @@ LANDSCAPE_API bool FLandscapeComponentDataInterface::GetWeightmapTextureData(ULa
 	OutData.Empty(FMath::Square(WeightmapSize));
 	OutData.AddUninitialized(FMath::Square(WeightmapSize));
 
-	FColor* WeightMipData = (FColor*)DataInterface.LockMip(Component->WeightmapTextures[Component->WeightmapLayerAllocations[LayerIdx].WeightmapTextureIndex], MipLevel);
+	FColor* WeightMipData = (FColor*)DataInterface.LockMip(ComponentWeightmapTextures[ComponentWeightmapLayerAllocations[LayerIdx].WeightmapTextureIndex], MipLevel);
 
 	// Channel remapping
 	int32 ChannelOffsets[4] = { (int32)STRUCT_OFFSET(FColor, R), (int32)STRUCT_OFFSET(FColor, G), (int32)STRUCT_OFFSET(FColor, B), (int32)STRUCT_OFFSET(FColor, A) };
 
-	uint8* SrcTextureData = (uint8*)WeightMipData + ChannelOffsets[Component->WeightmapLayerAllocations[LayerIdx].WeightmapTextureChannel];
+	uint8* SrcTextureData = (uint8*)WeightMipData + ChannelOffsets[ComponentWeightmapLayerAllocations[LayerIdx].WeightmapTextureChannel];
 
 	for (int32 i = 0; i < FMath::Square(WeightmapSize); i++)
 	{
 		OutData[i] = SrcTextureData[i * 4];
 	}
 
-	DataInterface.UnlockMip(Component->WeightmapTextures[Component->WeightmapLayerAllocations[LayerIdx].WeightmapTextureIndex], MipLevel);
+	DataInterface.UnlockMip(ComponentWeightmapTextures[ComponentWeightmapLayerAllocations[LayerIdx].WeightmapTextureIndex], MipLevel);
 	return true;
 }
 
diff --git a/Engine/Source/Runtime/Landscape/Private/LandscapeEdit.cpp b/Engine/Source/Runtime/Landscape/Private/LandscapeEdit.cpp
--- a/Engine/Source/Runtime/Landscape/Private/LandscapeEdit.cpp
+++ b/Engine/Source/Runtime/Landscape/Private/LandscapeEdit.cpp
@@ -143,7 +143,7 @@ ULandscapeMaterialInstanceConstant* ALandscapeProxy::GetLayerThumbnailMIC(UMater
 	ULandscapeMaterialInstanceConstant* MaterialInstance = NewObject<ULandscapeMaterialInstanceConstant>(GetTransientPackage());
 	MaterialInstance->bIsLayerThumbnail = true;
 	MaterialInstance->bMobile = false;
-	MaterialInstance->SetParentEditorOnly(LandscapeMaterial);
+	MaterialInstance->SetParentEditorOnly(LandscapeMaterial, false);
 
 	FStaticParameterSet StaticParameters;
 	MaterialInstance->GetStaticParameterValues(StaticParameters);
@@ -260,7 +260,7 @@ UMaterialInstanceConstant* ULandscapeComponent::GetCombinationMaterial(FMaterial
 			CombinationMaterialInstance = LandscapeCombinationMaterialInstance;
 			UE_LOG(LogLandscape, Log, TEXT("Looking for key %s, making new combination %s"), *LayerKey, *CombinationMaterialInstance->GetName());
 			Proxy->MaterialInstanceConstantMap.Add(*LayerKey, CombinationMaterialInstance);
-			CombinationMaterialInstance->SetParentEditorOnly(MaterialToUse);
+			CombinationMaterialInstance->SetParentEditorOnly(MaterialToUse, false);
 
 			CombinationMaterialInstance->BasePropertyOverrides.bOverride_BlendMode = bOverrideBlendMode;
 			if (bOverrideBlendMode)
@@ -331,12 +331,16 @@ void ULandscapeComponent::UpdateMaterialInstances_Internal(FMaterialUpdateContex
 	int8 TessellatedMaterialCount = 0;
 	int8 MaterialIndex = 0;
 
+	TArray<FWeightmapLayerAllocationInfo>& WeightmapBaseLayerAllocation = GetWeightmapLayerAllocations();
+	TArray<UTexture2D*>& WeightmapBaseTexture = GetWeightmapTextures();
+	UTexture2D* BaseHeightmap = GetHeightmap();
+
 	for (auto It = MaterialPerLOD.CreateConstIterator(); It; ++It)
 	{
 		const int8 MaterialLOD = It.Value();
 
 		// Find or set a matching MIC in the Landscape's map.
-		UMaterialInstanceConstant* CombinationMaterialInstance = GetCombinationMaterial(&Context, WeightmapLayerAllocations, MaterialLOD, false);
+		UMaterialInstanceConstant* CombinationMaterialInstance = GetCombinationMaterial(&Context, WeightmapBaseLayerAllocation, MaterialLOD, false);
 
 		if (CombinationMaterialInstance != nullptr)
 		{
@@ -356,24 +360,24 @@ void ULandscapeComponent::UpdateMaterialInstances_Internal(FMaterialUpdateContex
 			FLinearColor Masks[4] = { FLinearColor(1.0f, 0.0f, 0.0f, 0.0f), FLinearColor(0.0f, 1.0f, 0.0f, 0.0f), FLinearColor(0.0f, 0.0f, 1.0f, 0.0f), FLinearColor(0.0f, 0.0f, 0.0f, 1.0f) };
 
 			// Set the layer mask
-			for (int32 AllocIdx = 0; AllocIdx < WeightmapLayerAllocations.Num(); AllocIdx++)
+			for (int32 AllocIdx = 0; AllocIdx < WeightmapBaseLayerAllocation.Num(); AllocIdx++)
 			{
-				FWeightmapLayerAllocationInfo& Allocation = WeightmapLayerAllocations[AllocIdx];
+				FWeightmapLayerAllocationInfo& Allocation = WeightmapBaseLayerAllocation[AllocIdx];
 
 				FName LayerName = Allocation.LayerInfo == ALandscapeProxy::VisibilityLayer ? UMaterialExpressionLandscapeVisibilityMask::ParameterName : Allocation.LayerInfo ? Allocation.LayerInfo->LayerName : NAME_None;
 				MaterialInstance->SetVectorParameterValueEditorOnly(FName(*FString::Printf(TEXT("LayerMask_%s"), *LayerName.ToString())), Masks[Allocation.WeightmapTextureChannel]);
 			}
 
 			// Set the weightmaps
-			for (int32 i = 0; i < WeightmapTextures.Num(); i++)
+			for (int32 i = 0; i < WeightmapBaseTexture.Num(); i++)
 			{
-				MaterialInstance->SetTextureParameterValueEditorOnly(FName(*FString::Printf(TEXT("Weightmap%d"), i)), WeightmapTextures[i]);
+				MaterialInstance->SetTextureParameterValueEditorOnly(FName(*FString::Printf(TEXT("Weightmap%d"), i)), WeightmapBaseTexture[i]);
 			}
 
 			// Set the heightmap, if needed.
-			if (HeightmapTexture)
+			if (BaseHeightmap)
 			{
-				MaterialInstance->SetTextureParameterValueEditorOnly(FName(TEXT("Heightmap")), HeightmapTexture);
+				MaterialInstance->SetTextureParameterValueEditorOnly(FName(TEXT("Heightmap")), BaseHeightmap);
 			}
 			MaterialInstance->PostEditChange();
 
@@ -460,6 +464,8 @@ void ALandscapeProxy::UpdateAllComponentMaterialInstances()
 {
 	// we're not having the material update context recreate render states because we will manually do it for only our components
 	TArray<FComponentRecreateRenderStateContext> RecreateRenderStateContexts;
+	RecreateRenderStateContexts.Reserve(LandscapeComponents.Num());
+
 	for (ULandscapeComponent* Component : LandscapeComponents)
 	{
 		RecreateRenderStateContexts.Emplace(Component);
@@ -558,7 +564,13 @@ void ULandscapeComponent::PostEditUndo()
 
 void ALandscapeProxy::FixupWeightmaps()
 {
+	for (auto& ItPair : WeightmapUsageMap)
+	{
+		delete ItPair.Value;
+	}
+
 	WeightmapUsageMap.Empty();
+
 	for (ULandscapeComponent* Component : LandscapeComponents)
 	{
 		Component->FixupWeightmaps();
@@ -574,6 +586,9 @@ void ULandscapeComponent::FixupWeightmaps()
 
 		if (Info)
 		{
+			WeightmapTexturesUsage.Empty();
+			WeightmapTexturesUsage.AddDefaulted(WeightmapTextures.Num());
+
 			TArray<ULandscapeLayerInfoObject*> LayersToDelete;
 			bool bFixedLayerDeletion = false;
 
@@ -635,16 +650,26 @@ void ULandscapeComponent::FixupWeightmaps()
 				}
 
 				UTexture2D* WeightmapTexture = WeightmapTextures[Allocation.WeightmapTextureIndex];
-				FLandscapeWeightmapUsage& Usage = Proxy->WeightmapUsageMap.FindOrAdd(WeightmapTexture);
+
+				FLandscapeWeightmapUsage** TempUsage = Proxy->WeightmapUsageMap.Find(WeightmapTexture);
+
+				if (TempUsage == nullptr)
+				{
+					TempUsage = &Proxy->WeightmapUsageMap.Add(WeightmapTexture, new FLandscapeWeightmapUsage());
+					(*TempUsage)->ProceduralLayerName = NAME_None;
+				}
+
+				FLandscapeWeightmapUsage* Usage = *TempUsage;
+				WeightmapTexturesUsage[Allocation.WeightmapTextureIndex] = Usage; // Keep a ref to it for faster access
 
 				// Detect a shared layer allocation, caused by a previous undo or layer deletion bugs
-				if (Usage.ChannelUsage[Allocation.WeightmapTextureChannel] != nullptr &&
-					Usage.ChannelUsage[Allocation.WeightmapTextureChannel] != this)
+				if (Usage->ChannelUsage[Allocation.WeightmapTextureChannel] != nullptr &&
+					Usage->ChannelUsage[Allocation.WeightmapTextureChannel] != this)
 				{
 					FFormatNamedArguments Arguments;
 					Arguments.Add(TEXT("LayerName"), FText::FromString(Allocation.GetLayerName().ToString()));
 					Arguments.Add(TEXT("LandscapeName"), FText::FromString(GetName()));
-					Arguments.Add(TEXT("ChannelName"), FText::FromString(Usage.ChannelUsage[Allocation.WeightmapTextureChannel]->GetName()));
+					Arguments.Add(TEXT("ChannelName"), FText::FromString(Usage->ChannelUsage[Allocation.WeightmapTextureChannel]->GetName()));
 					FMessageLog("MapCheck").Warning()
 						->AddToken(FTextToken::Create(FText::Format(LOCTEXT("MapCheck_Message_FixedUpSharedLayerWeightmap", "Fixed up shared weightmap texture for layer {LayerName} in component '{LandscapeName}' (shares with '{ChannelName}')"), Arguments)))
 						->AddToken(FMapErrorToken::Create(FMapErrors::FixedUpSharedLayerWeightmap));
@@ -654,7 +679,7 @@ void ULandscapeComponent::FixupWeightmaps()
 				}
 				else
 				{
-					Usage.ChannelUsage[Allocation.WeightmapTextureChannel] = this;
+					Usage->ChannelUsage[Allocation.WeightmapTextureChannel] = this;
 				}
 			}
 
@@ -675,7 +700,9 @@ void ULandscapeComponent::FixupWeightmaps()
 
 void ULandscapeComponent::UpdateLayerWhitelistFromPaintedLayers()
 {
-	for (const auto& Allocation : WeightmapLayerAllocations)
+	TArray<FWeightmapLayerAllocationInfo>& ComponentWeightmapLayerAllocations = GetWeightmapLayerAllocations();
+
+	for (const auto& Allocation : ComponentWeightmapLayerAllocations)
 	{
 		LayerWhitelist.AddUnique(Allocation.LayerInfo);
 	}
@@ -1147,10 +1174,13 @@ void ULandscapeComponent::UpdateCollisionLayerData(const FColor* const* const We
 		bool bExistingLayerMismatch = false;
 		int32 DataLayerIdx = INDEX_NONE;
 
+		TArray<FWeightmapLayerAllocationInfo>& ComponentWeightmapLayerAllocations = GetWeightmapLayerAllocations(true);
+		TArray<UTexture2D*>& ComponentWeightmapsTexture = GetWeightmapTextures(true);
+
 		// Find the layers we're interested in
-		for (int32 AllocIdx = 0; AllocIdx < WeightmapLayerAllocations.Num(); AllocIdx++)
+		for (int32 AllocIdx = 0; AllocIdx < ComponentWeightmapLayerAllocations.Num(); AllocIdx++)
 		{
-			FWeightmapLayerAllocationInfo& AllocInfo = WeightmapLayerAllocations[AllocIdx];
+			FWeightmapLayerAllocationInfo& AllocInfo = ComponentWeightmapLayerAllocations[AllocIdx];
 			ULandscapeLayerInfoObject* LayerInfo = AllocInfo.LayerInfo;
 			if (LayerInfo == ALandscapeProxy::VisibilityLayer || LayerInfo != nullptr)
 			{
@@ -1219,7 +1249,7 @@ void ULandscapeComponent::UpdateCollisionLayerData(const FColor* const* const We
 				DominantLayerData = (uint8*)CollisionComp->DominantLayerData.Lock(LOCK_READ_WRITE);
 			}
 
-			const int32 WeightmapSizeU = WeightmapTextures[0]->Source.GetSizeX();
+			const int32 WeightmapSizeU = ComponentWeightmapsTexture[0]->Source.GetSizeX();
 			const int32 MipSizeU = WeightmapSizeU >> CollisionMipLevel;
 
 			// Ratio to convert update region coordinate to collision mip coordinates
@@ -1380,16 +1410,18 @@ void ULandscapeComponent::UpdateCollisionLayerData(const FColor* const* const We
 
 void ULandscapeComponent::UpdateCollisionLayerData()
 {
+	TArray<UTexture2D*>& ComponentWeightmapsTexture = GetWeightmapTextures(true);
+
 	// Generate the dominant layer data
 	TArray<FColor*> WeightmapTextureMipData;
 	TArray<TArray<uint8>> CachedWeightmapTextureMipData;
 
-	WeightmapTextureMipData.Empty(WeightmapTextures.Num());
-	CachedWeightmapTextureMipData.Empty(WeightmapTextures.Num());
-	for (int32 WeightmapIdx = 0; WeightmapIdx < WeightmapTextures.Num(); ++WeightmapIdx)
+	WeightmapTextureMipData.Empty(ComponentWeightmapsTexture.Num());
+	CachedWeightmapTextureMipData.Empty(ComponentWeightmapsTexture.Num());
+	for (int32 WeightmapIdx = 0; WeightmapIdx < ComponentWeightmapsTexture.Num(); ++WeightmapIdx)
 	{
 		TArray<uint8> MipData;
-		WeightmapTextures[WeightmapIdx]->Source.GetMipData(MipData, CollisionMipLevel);
+		ComponentWeightmapsTexture[WeightmapIdx]->Source.GetMipData(MipData, CollisionMipLevel);
 		WeightmapTextureMipData.Add((FColor*)MipData.GetData());
 		CachedWeightmapTextureMipData.Add(MoveTemp(MipData));
 	}
@@ -1398,10 +1430,10 @@ void ULandscapeComponent::UpdateCollisionLayerData()
 	TArray<TArray<uint8>> SimpleCollisionCachedWeightmapTextureMipData;
 	if (SimpleCollisionMipLevel > CollisionMipLevel)
 	{
-		for (int32 WeightmapIdx = 0; WeightmapIdx < WeightmapTextures.Num(); ++WeightmapIdx)
+		for (int32 WeightmapIdx = 0; WeightmapIdx < ComponentWeightmapsTexture.Num(); ++WeightmapIdx)
 		{
 			TArray<uint8> MipData;
-			WeightmapTextures[WeightmapIdx]->Source.GetMipData(MipData, SimpleCollisionMipLevel);
+			ComponentWeightmapsTexture[WeightmapIdx]->Source.GetMipData(MipData, SimpleCollisionMipLevel);
 			SimpleCollisionWeightmapMipData.Add((FColor*)MipData.GetData());
 			SimpleCollisionCachedWeightmapTextureMipData.Add(MoveTemp(MipData));
 		}
@@ -2338,10 +2370,12 @@ LANDSCAPE_API void ALandscapeProxy::Import(
 
 			UE_LOG(LogLandscape, Log, TEXT("%s needs %d alphamaps"), *LandscapeComponent->GetName(), EditingAlphaLayerData.Num());
 
+			TArray<FWeightmapLayerAllocationInfo>& ComponentWeightmapLayerAllocations = LandscapeComponent->GetWeightmapLayerAllocations();
+
 			// Calculate weightmap weights for this component
 			WeightValues.Empty(EditingAlphaLayerData.Num());
 			WeightValues.AddZeroed(EditingAlphaLayerData.Num());
-			LandscapeComponent->WeightmapLayerAllocations.Empty(EditingAlphaLayerData.Num());
+			ComponentWeightmapLayerAllocations.Empty(EditingAlphaLayerData.Num());
 
 			TArray<bool, TInlineAllocator<16>> IsNoBlendArray;
 			IsNoBlendArray.Empty(EditingAlphaLayerData.Num());
@@ -2351,7 +2385,7 @@ LANDSCAPE_API void ALandscapeProxy::Import(
 			{
 				// Lookup the original layer name
 				WeightValues[WeightLayerIndex] = EditingAlphaLayerData[WeightLayerIndex].AlphaValues;
-				new(LandscapeComponent->WeightmapLayerAllocations) FWeightmapLayerAllocationInfo(ImportLayerInfos[EditingAlphaLayerData[WeightLayerIndex].LayerIndex].LayerInfo);
+				new(ComponentWeightmapLayerAllocations) FWeightmapLayerAllocationInfo(ImportLayerInfos[EditingAlphaLayerData[WeightLayerIndex].LayerIndex].LayerInfo);
 				IsNoBlendArray[WeightLayerIndex] = ImportLayerInfos[EditingAlphaLayerData[WeightLayerIndex].LayerIndex].LayerInfo->bNoWeightBlend;
 			}
 
@@ -2386,7 +2420,7 @@ LANDSCAPE_API void ALandscapeProxy::Import(
 						{
 							// Remove the layer as it has no contribution
 							WeightValues.RemoveAt(BelowWeightLayerIndex);
-							LandscapeComponent->WeightmapLayerAllocations.RemoveAt(BelowWeightLayerIndex);
+							ComponentWeightmapLayerAllocations.RemoveAt(BelowWeightLayerIndex);
 							IsNoBlendArray.RemoveAt(BelowWeightLayerIndex);
 
 							// The current layer has been re-numbered
@@ -2507,18 +2541,23 @@ LANDSCAPE_API void ALandscapeProxy::Import(
 					}
 				}
 
+				TArray<FWeightmapLayerAllocationInfo>& ComponentWeightmapLayerAllocations = LandscapeComponent->GetWeightmapLayerAllocations();
+				TArray<UTexture2D*>& ComponentWeightmapTextures = LandscapeComponent->GetWeightmapTextures();
+				TArray<FLandscapeWeightmapUsage*>& ComponentWeightmapTexturesUsage = LandscapeComponent->GetWeightmapTexturesUsage();
+
 				if (BestAllocationIndex != -1)
 				{
 					FWeightmapTextureAllocation& Allocation = TextureAllocations[BestAllocationIndex];
-					FLandscapeWeightmapUsage& WeightmapUsage = WeightmapUsageMap.FindChecked(Allocation.Texture);
+					FLandscapeWeightmapUsage* WeightmapUsage = WeightmapUsageMap.FindChecked(Allocation.Texture);
+					ComponentWeightmapTexturesUsage.Add(WeightmapUsage);
 
 					UE_LOG(LogLandscape, Log, TEXT("  ==> Storing %d channels starting at %s[%d]"), RemainingLayers, *Allocation.Texture->GetName(), Allocation.ChannelsInUse);
 
 					for (int32 i = 0; i < RemainingLayers; i++)
 					{
-						LandscapeComponent->WeightmapLayerAllocations[LayerIndex + i].WeightmapTextureIndex = LandscapeComponent->WeightmapTextures.Num();
-						LandscapeComponent->WeightmapLayerAllocations[LayerIndex + i].WeightmapTextureChannel = Allocation.ChannelsInUse;
-						WeightmapUsage.ChannelUsage[Allocation.ChannelsInUse] = LandscapeComponent;
+						ComponentWeightmapLayerAllocations[LayerIndex + i].WeightmapTextureIndex = ComponentWeightmapTextures.Num();
+						ComponentWeightmapLayerAllocations[LayerIndex + i].WeightmapTextureChannel = Allocation.ChannelsInUse;
+						WeightmapUsage->ChannelUsage[Allocation.ChannelsInUse] = LandscapeComponent;
 						switch (Allocation.ChannelsInUse)
 						{
 						case 1:
@@ -2539,7 +2578,7 @@ LANDSCAPE_API void ALandscapeProxy::Import(
 					}
 
 					LayerIndex += RemainingLayers;
-					LandscapeComponent->WeightmapTextures.Add(Allocation.Texture);
+					ComponentWeightmapTextures.Add(Allocation.Texture);
 				}
 				else
 				{
@@ -2549,39 +2588,40 @@ LANDSCAPE_API void ALandscapeProxy::Import(
 
 					const int32 ThisAllocationLayers = FMath::Min<int32>(RemainingLayers, 4);
 					new(TextureAllocations) FWeightmapTextureAllocation(ComponentX, ComponentY, ThisAllocationLayers, WeightmapTexture, MipData);
-					FLandscapeWeightmapUsage& WeightmapUsage = WeightmapUsageMap.Add(WeightmapTexture, FLandscapeWeightmapUsage());
+					FLandscapeWeightmapUsage* WeightmapUsage = WeightmapUsageMap.Add(WeightmapTexture, new FLandscapeWeightmapUsage());
+					ComponentWeightmapTexturesUsage.Add(WeightmapUsage);
 
 					UE_LOG(LogLandscape, Log, TEXT("  ==> Storing %d channels in new texture %s"), ThisAllocationLayers, *WeightmapTexture->GetName());
 
 					WeightmapTextureDataPointers.Add((uint8*)&MipData->R);
-					LandscapeComponent->WeightmapLayerAllocations[LayerIndex + 0].WeightmapTextureIndex = LandscapeComponent->WeightmapTextures.Num();
-					LandscapeComponent->WeightmapLayerAllocations[LayerIndex + 0].WeightmapTextureChannel = 0;
-					WeightmapUsage.ChannelUsage[0] = LandscapeComponent;
+					ComponentWeightmapLayerAllocations[LayerIndex + 0].WeightmapTextureIndex = ComponentWeightmapTextures.Num();
+					ComponentWeightmapLayerAllocations[LayerIndex + 0].WeightmapTextureChannel = 0;
+					WeightmapUsage->ChannelUsage[0] = LandscapeComponent;
 
 					if (ThisAllocationLayers > 1)
 					{
 						WeightmapTextureDataPointers.Add((uint8*)&MipData->G);
-						LandscapeComponent->WeightmapLayerAllocations[LayerIndex + 1].WeightmapTextureIndex = LandscapeComponent->WeightmapTextures.Num();
-						LandscapeComponent->WeightmapLayerAllocations[LayerIndex + 1].WeightmapTextureChannel = 1;
-						WeightmapUsage.ChannelUsage[1] = LandscapeComponent;
+						ComponentWeightmapLayerAllocations[LayerIndex + 1].WeightmapTextureIndex = ComponentWeightmapTextures.Num();
+						ComponentWeightmapLayerAllocations[LayerIndex + 1].WeightmapTextureChannel = 1;
+						WeightmapUsage->ChannelUsage[1] = LandscapeComponent;
 
 						if (ThisAllocationLayers > 2)
 						{
 							WeightmapTextureDataPointers.Add((uint8*)&MipData->B);
-							LandscapeComponent->WeightmapLayerAllocations[LayerIndex + 2].WeightmapTextureIndex = LandscapeComponent->WeightmapTextures.Num();
-							LandscapeComponent->WeightmapLayerAllocations[LayerIndex + 2].WeightmapTextureChannel = 2;
-							WeightmapUsage.ChannelUsage[2] = LandscapeComponent;
+							ComponentWeightmapLayerAllocations[LayerIndex + 2].WeightmapTextureIndex = ComponentWeightmapTextures.Num();
+							ComponentWeightmapLayerAllocations[LayerIndex + 2].WeightmapTextureChannel = 2;
+							WeightmapUsage->ChannelUsage[2] = LandscapeComponent;
 
 							if (ThisAllocationLayers > 3)
 							{
 								WeightmapTextureDataPointers.Add((uint8*)&MipData->A);
-								LandscapeComponent->WeightmapLayerAllocations[LayerIndex + 3].WeightmapTextureIndex = LandscapeComponent->WeightmapTextures.Num();
-								LandscapeComponent->WeightmapLayerAllocations[LayerIndex + 3].WeightmapTextureChannel = 3;
-								WeightmapUsage.ChannelUsage[3] = LandscapeComponent;
+								ComponentWeightmapLayerAllocations[LayerIndex + 3].WeightmapTextureIndex = ComponentWeightmapTextures.Num();
+								ComponentWeightmapLayerAllocations[LayerIndex + 3].WeightmapTextureChannel = 3;
+								WeightmapUsage->ChannelUsage[3] = LandscapeComponent;
 							}
 						}
 					}
-					LandscapeComponent->WeightmapTextures.Add(WeightmapTexture);
+					ComponentWeightmapTextures.Add(WeightmapTexture);
 
 					LayerIndex += ThisAllocationLayers;
 				}
@@ -2782,9 +2822,6 @@ LANDSCAPE_API void ALandscapeProxy::Import(
 	if (GetMutableDefault<UEditorExperimentalSettings>()->bProceduralLandscape)
 	{
 		SetupProceduralLayers(NumComponentsX, NumComponentsY);
-		
-		FEditorDelegates::PreSaveWorld.AddUObject(GetLandscapeActor(), &ALandscape::OnPreSaveWorld);
-		FEditorDelegates::PostSaveWorld.AddUObject(GetLandscapeActor(), &ALandscape::OnPostSaveWorld);
 	}
 
 	if (GetLevel()->bIsVisible)
@@ -2878,10 +2915,11 @@ bool ALandscapeProxy::ExportToRawMesh(int32 InExportLOD, FMeshDescription& OutRa
 		// Check if there are any holes
 		const int32 VisThreshold = 170;
 		TArray<uint8> VisDataMap;
+		TArray<FWeightmapLayerAllocationInfo>& ComponentWeightmapLayerAllocations = Component->GetWeightmapLayerAllocations();
 
-		for (int32 AllocIdx = 0; AllocIdx < Component->WeightmapLayerAllocations.Num(); AllocIdx++)
+		for (int32 AllocIdx = 0; AllocIdx < ComponentWeightmapLayerAllocations.Num(); AllocIdx++)
 		{
-			FWeightmapLayerAllocationInfo& AllocInfo = Component->WeightmapLayerAllocations[AllocIdx];
+			FWeightmapLayerAllocationInfo& AllocInfo = ComponentWeightmapLayerAllocations[AllocIdx];
 			if (AllocInfo.LayerInfo == ALandscapeProxy::VisibilityLayer)
 			{
 				CDI.GetWeightmapTextureData(AllocInfo.LayerInfo, VisDataMap);
@@ -3931,7 +3969,9 @@ void ALandscapeProxy::PostEditChangeProperty(FPropertyChangedEvent& PropertyChan
 	bool bRemovedAnyLayers = false;
 	for (ULandscapeComponent* Component : LandscapeComponents)
 	{
-		int32 NumNullLayers = Algo::CountIf(Component->WeightmapLayerAllocations, [](const FWeightmapLayerAllocationInfo& Allocation) { return Allocation.LayerInfo == nullptr; });
+		TArray<FWeightmapLayerAllocationInfo>& ComponentWeightmapLayerAllocations = Component->GetWeightmapLayerAllocations(false);
+
+		int32 NumNullLayers = Algo::CountIf(ComponentWeightmapLayerAllocations, [](const FWeightmapLayerAllocationInfo& Allocation) { return Allocation.LayerInfo == nullptr; });
 		if (NumNullLayers > 0)
 		{
 			FLandscapeEditDataInterface LandscapeEdit(Info);
@@ -4497,14 +4537,18 @@ void ULandscapeInfo::ClearSelectedRegion(bool bIsComponentwise /*= true*/)
 	}
 }
 
-void ULandscapeComponent::ReallocateWeightmaps(FLandscapeEditDataInterface* DataInterface)
+void ULandscapeComponent::ReallocateWeightmaps(FLandscapeEditDataInterface* DataInterface, bool InCanUseCurrentEditingWeightmap, bool InSaveToTransactionBuffer, bool InInitPlatformDataAsync, TArray<UTexture2D*>* OutNewCreatedTextures)
 {
 	ALandscapeProxy* Proxy = GetLandscapeProxy();
 
+	TArray<FWeightmapLayerAllocationInfo>& ComponentWeightmapLayerAllocations = GetWeightmapLayerAllocations(InCanUseCurrentEditingWeightmap);
+	TArray<UTexture2D*>& ComponentWeightmapTextures = GetWeightmapTextures(InCanUseCurrentEditingWeightmap);
+	TArray<FLandscapeWeightmapUsage*>& ComponentWeightmapTexturesUsage = GetWeightmapTexturesUsage(InCanUseCurrentEditingWeightmap);
+
 	int32 NeededNewChannels = 0;
-	for (int32 LayerIdx = 0; LayerIdx < WeightmapLayerAllocations.Num(); LayerIdx++)
+	for (int32 LayerIdx = 0; LayerIdx < ComponentWeightmapLayerAllocations.Num(); LayerIdx++)
 	{
-		if (WeightmapLayerAllocations[LayerIdx].WeightmapTextureIndex == 255)
+		if (ComponentWeightmapLayerAllocations[LayerIdx].WeightmapTextureIndex == 255)
 		{
 			NeededNewChannels++;
 		}
@@ -4516,18 +4560,20 @@ void ULandscapeComponent::ReallocateWeightmaps(FLandscapeEditDataInterface* Data
 		return;
 	}
 
-	Modify();
-	//Landscape->Modify();
-	Proxy->Modify();
+	if (InSaveToTransactionBuffer)
+	{
+		Modify();
+		Proxy->Modify();
+	}
 
 	// UE_LOG(LogLandscape, Log, TEXT("----------------------"));
 	// UE_LOG(LogLandscape, Log, TEXT("Component %s needs %d layers (%d new)"), *GetName(), WeightmapLayerAllocations.Num(), NeededNewChannels);
 
 	// See if our existing textures have sufficient space
 	int32 ExistingTexAvailableChannels = 0;
-	for (int32 TexIdx = 0; TexIdx < WeightmapTextures.Num(); TexIdx++)
+	for (int32 TexIdx = 0; TexIdx < ComponentWeightmapTextures.Num(); TexIdx++)
 	{
-		FLandscapeWeightmapUsage* Usage = Proxy->WeightmapUsageMap.Find(WeightmapTextures[TexIdx]);
+		FLandscapeWeightmapUsage* Usage = ComponentWeightmapTexturesUsage[TexIdx];
 		check(Usage);
 
 		ExistingTexAvailableChannels += Usage->FreeChannelCount();
@@ -4543,35 +4589,44 @@ void ULandscapeComponent::ReallocateWeightmaps(FLandscapeEditDataInterface* Data
 		// UE_LOG(LogLandscape, Log, TEXT("Existing texture has available channels"));
 
 		// Allocate using our existing textures' spare channels.
-		for (int32 TexIdx = 0; TexIdx < WeightmapTextures.Num(); TexIdx++)
+		int32 CurrentAlloc = 0;
+		for (int32 TexIdx = 0; TexIdx < ComponentWeightmapTextures.Num(); TexIdx++)
 		{
-			FLandscapeWeightmapUsage* Usage = Proxy->WeightmapUsageMap.Find(WeightmapTextures[TexIdx]);
+			FLandscapeWeightmapUsage* Usage = ComponentWeightmapTexturesUsage[TexIdx];
 
 			for (int32 ChanIdx = 0; ChanIdx < 4; ChanIdx++)
 			{
 				if (Usage->ChannelUsage[ChanIdx] == nullptr)
 				{
-					for (int32 LayerIdx = 0; LayerIdx < WeightmapLayerAllocations.Num(); LayerIdx++)
+					// Find next allocation to treat
+					for (; CurrentAlloc < ComponentWeightmapLayerAllocations.Num(); ++CurrentAlloc)
 					{
-						FWeightmapLayerAllocationInfo& AllocInfo = WeightmapLayerAllocations[LayerIdx];
+						FWeightmapLayerAllocationInfo& AllocInfo = ComponentWeightmapLayerAllocations[CurrentAlloc];
+						
 						if (AllocInfo.WeightmapTextureIndex == 255)
 						{
-							// Zero out the data for this texture channel
-							if (DataInterface)
-							{
-								DataInterface->ZeroTextureChannel(WeightmapTextures[TexIdx], ChanIdx);
-							}
+							break;
+						}
+					}
 
-							AllocInfo.WeightmapTextureIndex = TexIdx;
-							AllocInfo.WeightmapTextureChannel = ChanIdx;
-							Usage->ChannelUsage[ChanIdx] = this;
-							NeededNewChannels--;
+					FWeightmapLayerAllocationInfo& AllocInfo = ComponentWeightmapLayerAllocations[CurrentAlloc];
+					check(AllocInfo.WeightmapTextureIndex == 255);
 
-							if (NeededNewChannels == 0)
-							{
-								return;
-							}
-						}
+					// Zero out the data for this texture channel
+					if (DataInterface)
+					{
+						DataInterface->ZeroTextureChannel(ComponentWeightmapTextures[TexIdx], ChanIdx);
+					}
+
+					AllocInfo.WeightmapTextureIndex = TexIdx;
+					AllocInfo.WeightmapTextureChannel = ChanIdx;
+					Usage->ChannelUsage[ChanIdx] = this;
+
+					NeededNewChannels--;
+
+					if (NeededNewChannels == 0)
+					{
+						return;
 					}
 				}
 			}
@@ -4583,9 +4638,10 @@ void ULandscapeComponent::ReallocateWeightmaps(FLandscapeEditDataInterface* Data
 	// UE_LOG(LogLandscape, Log, TEXT("Reallocating."));
 
 	// We are totally reallocating the weightmap
-	int32 TotalNeededChannels = WeightmapLayerAllocations.Num();
+	int32 TotalNeededChannels = ComponentWeightmapLayerAllocations.Num();
 	int32 CurrentLayer = 0;
 	TArray<UTexture2D*> NewWeightmapTextures;
+	TArray<FLandscapeWeightmapUsage*> NewComponentWeightmapTexturesUsage;
 	while (TotalNeededChannels > 0)
 	{
 		// UE_LOG(LogLandscape, Log, TEXT("Still need %d channels"), TotalNeededChannels);
@@ -4599,22 +4655,34 @@ void ULandscapeComponent::ReallocateWeightmaps(FLandscapeEditDataInterface* Data
 
 			// see if we can find a suitable existing weightmap texture with sufficient channels
 			int32 BestDistanceSquared = MAX_int32;
-			for (TMap<UTexture2D*, struct FLandscapeWeightmapUsage>::TIterator It(Proxy->WeightmapUsageMap); It; ++It)
+			for (auto& ItPair : Proxy->WeightmapUsageMap)
 			{
-				FLandscapeWeightmapUsage* TryWeightmapUsage = &It.Value();
-				if (TryWeightmapUsage->FreeChannelCount() >= TotalNeededChannels)
+				FLandscapeWeightmapUsage* TryWeightmapUsage = ItPair.Value;
+				//
+				FName ProceduralLayerNameToSeek = InCanUseCurrentEditingWeightmap ? CurrentProceduralLayerName : NAME_None;
+
+				if (TryWeightmapUsage->FreeChannelCount() >= TotalNeededChannels && TryWeightmapUsage->ProceduralLayerName == ProceduralLayerNameToSeek)
 				{
-					// See if this candidate is closer than any others we've found
-					for (int32 ChanIdx = 0; ChanIdx < 4; ChanIdx++)
+					if (TryWeightmapUsage->FreeChannelCount() == 4)
 					{
-						if (TryWeightmapUsage->ChannelUsage[ChanIdx] != nullptr)
+						CurrentWeightmapTexture = ItPair.Key;
+						CurrentWeightmapUsage = TryWeightmapUsage;
+						break;
+					}
+					else
+					{
+						// See if this candidate is closer than any others we've found
+						for (int32 ChanIdx = 0; ChanIdx < 4; ChanIdx++)
 						{
-							int32 TryDistanceSquared = (TryWeightmapUsage->ChannelUsage[ChanIdx]->GetSectionBase() - GetSectionBase()).SizeSquared();
-							if (TryDistanceSquared < BestDistanceSquared)
+							if (TryWeightmapUsage->ChannelUsage[ChanIdx] != nullptr)
 							{
-								CurrentWeightmapTexture = It.Key();
-								CurrentWeightmapUsage = TryWeightmapUsage;
-								BestDistanceSquared = TryDistanceSquared;
+								int32 TryDistanceSquared = (TryWeightmapUsage->ChannelUsage[ChanIdx]->GetSectionBase() - GetSectionBase()).SizeSquared();
+								if (TryDistanceSquared < BestDistanceSquared)
+								{
+									CurrentWeightmapTexture = ItPair.Key;
+									CurrentWeightmapUsage = TryWeightmapUsage;
+									BestDistanceSquared = TryDistanceSquared;
+								}
 							}
 						}
 					}
@@ -4633,16 +4701,32 @@ void ULandscapeComponent::ReallocateWeightmaps(FLandscapeEditDataInterface* Data
 
 			// We need a new weightmap texture
 			CurrentWeightmapTexture = GetLandscapeProxy()->CreateLandscapeTexture(WeightmapSize, WeightmapSize, TEXTUREGROUP_Terrain_Weightmap, TSF_BGRA8);
+
 			// Alloc dummy mips
-			CreateEmptyTextureMips(CurrentWeightmapTexture);
-			CurrentWeightmapTexture->PostEditChange();
+			CreateEmptyTextureMips(CurrentWeightmapTexture, true);
 
-			// Store it in the usage map
-			CurrentWeightmapUsage = &Proxy->WeightmapUsageMap.Add(CurrentWeightmapTexture, FLandscapeWeightmapUsage());
+			if (InInitPlatformDataAsync)
+			{
+				CurrentWeightmapTexture->BeginCachePlatformData();
+				CurrentWeightmapTexture->ClearAllCachedCookedPlatformData();
+			}
+			else
+			{
+				CurrentWeightmapTexture->PostEditChange();
+			}
 
+			if (OutNewCreatedTextures != nullptr)
+			{
+				OutNewCreatedTextures->Add(CurrentWeightmapTexture);
+			}
+
+			// Store it in the usage map
+			CurrentWeightmapUsage = Proxy->WeightmapUsageMap.Add(CurrentWeightmapTexture, new FLandscapeWeightmapUsage());
+			CurrentWeightmapUsage->ProceduralLayerName = InCanUseCurrentEditingWeightmap ? CurrentProceduralLayerName : NAME_None;
 			// UE_LOG(LogLandscape, Log, TEXT("Making a new texture %s"), *CurrentWeightmapTexture->GetName());
 		}
 
+		NewComponentWeightmapTexturesUsage.Add(CurrentWeightmapUsage);
 		NewWeightmapTextures.Add(CurrentWeightmapTexture);
 
 		for (int32 ChanIdx = 0; ChanIdx < 4 && TotalNeededChannels > 0; ChanIdx++)
@@ -4652,7 +4736,7 @@ void ULandscapeComponent::ReallocateWeightmaps(FLandscapeEditDataInterface* Data
 			if (CurrentWeightmapUsage->ChannelUsage[ChanIdx] == nullptr)
 			{
 				// Use this allocation
-				FWeightmapLayerAllocationInfo& AllocInfo = WeightmapLayerAllocations[CurrentLayer];
+				FWeightmapLayerAllocationInfo& AllocInfo = ComponentWeightmapLayerAllocations[CurrentLayer];
 
 				if (AllocInfo.WeightmapTextureIndex == 255)
 				{
@@ -4665,7 +4749,7 @@ void ULandscapeComponent::ReallocateWeightmaps(FLandscapeEditDataInterface* Data
 				}
 				else
 				{
-					UTexture2D* OldWeightmapTexture = WeightmapTextures[AllocInfo.WeightmapTextureIndex];
+					UTexture2D* OldWeightmapTexture = ComponentWeightmapTextures[AllocInfo.WeightmapTextureIndex];
 
 					// Copy the data
 					if (ensure(DataInterface != nullptr)) // it's not safe to skip the copy
@@ -4676,7 +4760,7 @@ void ULandscapeComponent::ReallocateWeightmaps(FLandscapeEditDataInterface* Data
 					}
 
 					// Remove the old allocation
-					FLandscapeWeightmapUsage* OldWeightmapUsage = Proxy->WeightmapUsageMap.Find(OldWeightmapTexture);
+					FLandscapeWeightmapUsage* OldWeightmapUsage = ComponentWeightmapTexturesUsage[AllocInfo.WeightmapTextureIndex];
 					OldWeightmapUsage->ChannelUsage[AllocInfo.WeightmapTextureChannel] = nullptr;
 				}
 
@@ -4690,15 +4774,12 @@ void ULandscapeComponent::ReallocateWeightmaps(FLandscapeEditDataInterface* Data
 		}
 	}
 
-	// Replace the weightmap textures
-	WeightmapTextures = MoveTemp(NewWeightmapTextures);
-
 	if (DataInterface)
 	{
 		// Update the mipmaps for the textures we edited
-		for (int32 Idx = 0; Idx < WeightmapTextures.Num(); Idx++)
+		for (int32 Idx = 0; Idx < NewWeightmapTextures.Num(); Idx++)
 		{
-			UTexture2D* WeightmapTexture = WeightmapTextures[Idx];
+			UTexture2D* WeightmapTexture = NewWeightmapTextures[Idx];
 			FLandscapeTextureDataInfo* WeightmapDataInfo = DataInterface->GetTextureDataInfo(WeightmapTexture);
 
 			int32 NumMips = WeightmapTexture->Source.GetNumMips();
@@ -4712,17 +4793,21 @@ void ULandscapeComponent::ReallocateWeightmaps(FLandscapeEditDataInterface* Data
 			ULandscapeComponent::UpdateWeightmapMips(NumSubsections, SubsectionSizeQuads, WeightmapTexture, WeightmapTextureMipData, 0, 0, MAX_int32, MAX_int32, WeightmapDataInfo);
 		}
 	}
+
+	// Replace the weightmap textures
+	SetWeightmapTextures(MoveTemp(NewWeightmapTextures), InCanUseCurrentEditingWeightmap);
+	SetWeightmapTexturesUsage(MoveTemp(NewComponentWeightmapTexturesUsage), InCanUseCurrentEditingWeightmap);	
 }
 
 void ALandscapeProxy::RemoveInvalidWeightmaps()
 {
 	if (GIsEditor)
 	{
-		for (TMap< UTexture2D*, struct FLandscapeWeightmapUsage >::TIterator It(WeightmapUsageMap); It; ++It)
+		for (TMap< UTexture2D*, struct FLandscapeWeightmapUsage* >::TIterator It(WeightmapUsageMap); It; ++It)
 		{
 			UTexture2D* Tex = It.Key();
-			FLandscapeWeightmapUsage& Usage = It.Value();
-			if (Usage.FreeChannelCount() == 4) // Invalid Weight-map
+			FLandscapeWeightmapUsage* Usage = It.Value();
+			if (Usage->FreeChannelCount() == 4) // Invalid Weight-map
 			{
 				if (Tex)
 				{
@@ -4731,7 +4816,9 @@ void ALandscapeProxy::RemoveInvalidWeightmaps()
 					Tex->MarkPackageDirty();
 					Tex->ClearFlags(RF_Standalone);
 				}
-				WeightmapUsageMap.Remove(Tex);
+
+				delete Usage;
+				It.RemoveCurrent();
 			}
 		}
 
@@ -4746,16 +4833,20 @@ void ALandscapeProxy::RemoveInvalidWeightmaps()
 
 void ULandscapeComponent::RemoveInvalidWeightmaps()
 {
+	TArray<FWeightmapLayerAllocationInfo>& ComponentWeightmapLayerAllocations = GetWeightmapLayerAllocations();
+	TArray<UTexture2D*>& ComponentWeightmapTextures = GetWeightmapTextures();
+	TArray<FLandscapeWeightmapUsage*>& ComponentWeightmapTexturesUsage = GetWeightmapTexturesUsage();
+
 	// Adjust WeightmapTextureIndex index for other layers
 	TSet<int32> UnUsedTextureIndices;
 	{
 		TSet<int32> UsedTextureIndices;
-		for (int32 LayerIdx = 0; LayerIdx < WeightmapLayerAllocations.Num(); LayerIdx++)
+		for (int32 LayerIdx = 0; LayerIdx < ComponentWeightmapLayerAllocations.Num(); LayerIdx++)
 		{
-			UsedTextureIndices.Add(WeightmapLayerAllocations[LayerIdx].WeightmapTextureIndex);
+			UsedTextureIndices.Add(ComponentWeightmapLayerAllocations[LayerIdx].WeightmapTextureIndex);
 		}
 
-		for (int32 WeightIdx = 0; WeightIdx < WeightmapTextures.Num(); ++WeightIdx)
+		for (int32 WeightIdx = 0; WeightIdx < ComponentWeightmapTextures.Num(); ++WeightIdx)
 		{
 			if (!UsedTextureIndices.Contains(WeightIdx))
 			{
@@ -4768,16 +4859,18 @@ void ULandscapeComponent::RemoveInvalidWeightmaps()
 	for (int32 UnusedIndex : UnUsedTextureIndices)
 	{
 		int32 WeightmapTextureIndexToRemove = UnusedIndex - RemovedTextures;
-		WeightmapTextures[WeightmapTextureIndexToRemove]->SetFlags(RF_Transactional);
-		WeightmapTextures[WeightmapTextureIndexToRemove]->Modify();
-		WeightmapTextures[WeightmapTextureIndexToRemove]->MarkPackageDirty();
-		WeightmapTextures[WeightmapTextureIndexToRemove]->ClearFlags(RF_Standalone);
-		WeightmapTextures.RemoveAt(WeightmapTextureIndexToRemove);
+		ComponentWeightmapTextures[WeightmapTextureIndexToRemove]->SetFlags(RF_Transactional);
+		ComponentWeightmapTextures[WeightmapTextureIndexToRemove]->Modify();
+		ComponentWeightmapTextures[WeightmapTextureIndexToRemove]->MarkPackageDirty();
+		ComponentWeightmapTextures[WeightmapTextureIndexToRemove]->ClearFlags(RF_Standalone);
+		ComponentWeightmapTextures.RemoveAt(WeightmapTextureIndexToRemove);
+
+		ComponentWeightmapTexturesUsage.RemoveAt(WeightmapTextureIndexToRemove);
 
 		// Adjust WeightmapTextureIndex index for other layers
-		for (int32 LayerIdx = 0; LayerIdx < WeightmapLayerAllocations.Num(); LayerIdx++)
+		for (int32 LayerIdx = 0; LayerIdx < ComponentWeightmapLayerAllocations.Num(); LayerIdx++)
 		{
-			FWeightmapLayerAllocationInfo& Allocation = WeightmapLayerAllocations[LayerIdx];
+			FWeightmapLayerAllocationInfo& Allocation = ComponentWeightmapLayerAllocations[LayerIdx];
 
 			if (Allocation.WeightmapTextureIndex > WeightmapTextureIndexToRemove)
 			{
@@ -4899,7 +4992,7 @@ void ULandscapeComponent::InitWeightmapData(TArray<ULandscapeLayerInfoObject*>&
 		new (WeightmapLayerAllocations)FWeightmapLayerAllocationInfo(LayerInfos[Idx]);
 	}
 
-	ReallocateWeightmaps(nullptr);
+	ReallocateWeightmaps();
 
 	check(WeightmapLayerAllocations.Num() > 0 && WeightmapTextures.Num() > 0);
 
@@ -4959,6 +5052,8 @@ void ULandscapeComponent::InitWeightmapData(TArray<ULandscapeLayerInfoObject*>&
 
 	LODIndexToMaterialIndex.Empty(1);
 	LODIndexToMaterialIndex.Add(0);
+
+	//  TODO: need to update procedural?
 }
 
 #define MAX_LANDSCAPE_EXPORT_COMPONENTS_NUM		16
