When SetParentEditorOnly is called we no longer recompute shader in most case as we will call UpdateStaticPermutation that will do so.

diff --git a/Engine/Source/Runtime/Landscape/Private/LandscapeEditInterface.cpp b/Engine/Source/Runtime/Landscape/Private/LandscapeEditInterface.cpp
--- a/Engine/Source/Runtime/Landscape/Private/LandscapeEditInterface.cpp
+++ b/Engine/Source/Runtime/Landscape/Private/LandscapeEditInterface.cpp
@@ -22,6 +22,7 @@ LandscapeEditInterface.cpp: Landscape editing interface
 #include "LandscapeRender.h"
 #include "ComponentReregisterContext.h"
 #include "Algo/Transform.h"
+#include "Settings/EditorExperimentalSettings.h"
 
 // Channel remapping
 extern const size_t ChannelOffsets[4] = {STRUCT_OFFSET(FColor,R), STRUCT_OFFSET(FColor,G), STRUCT_OFFSET(FColor,B), STRUCT_OFFSET(FColor,A)};
@@ -1502,9 +1503,12 @@ void FLandscapeEditDataInterface::GetHeightDataFast(const int32 X1, const int32
 void ULandscapeComponent::DeleteLayer(ULandscapeLayerInfoObject* LayerInfo, FLandscapeEditDataInterface& LandscapeEdit)
 {
 	ULandscapeComponent* Component = this;
+	TArray<FWeightmapLayerAllocationInfo>& ComponentWeightmapLayerAllocations = Component->GetWeightmapLayerAllocations(true);
+	TArray<UTexture2D*>& ComponentWeightmapTextures = Component->GetWeightmapTextures(true);
+	TArray<FLandscapeWeightmapUsage*>& ComponentWeightmapTexturesUsage = Component->GetWeightmapTexturesUsage(true);
 
 	// Find the index for this layer in this component.
-	const int32 DeleteLayerIdx = Component->WeightmapLayerAllocations.IndexOfByPredicate(
+	const int32 DeleteLayerIdx = ComponentWeightmapLayerAllocations.IndexOfByPredicate(
 		[LayerInfo](const FWeightmapLayerAllocationInfo& Allocation) { return Allocation.LayerInfo == LayerInfo; });
 	if (DeleteLayerIdx == INDEX_NONE)
 	{
@@ -1512,14 +1516,14 @@ void ULandscapeComponent::DeleteLayer(ULandscapeLayerInfoObject* LayerInfo, FLan
 		return;
 	}
 
-	FWeightmapLayerAllocationInfo& DeleteLayerAllocation = Component->WeightmapLayerAllocations[DeleteLayerIdx];
+	FWeightmapLayerAllocationInfo& DeleteLayerAllocation = ComponentWeightmapLayerAllocations[DeleteLayerIdx];
 	int32 DeleteLayerWeightmapTextureIndex = DeleteLayerAllocation.WeightmapTextureIndex;
 
 	// See if we'll be able to remove the texture completely.
 	bool bCanRemoveLayerTexture = true;
-	for (int32 LayerIdx = 0; LayerIdx < Component->WeightmapLayerAllocations.Num(); LayerIdx++)
+	for (int32 LayerIdx = 0; LayerIdx < ComponentWeightmapLayerAllocations.Num(); LayerIdx++)
 	{
-		FWeightmapLayerAllocationInfo& Allocation = Component->WeightmapLayerAllocations[LayerIdx];
+		FWeightmapLayerAllocationInfo& Allocation = ComponentWeightmapLayerAllocations[LayerIdx];
 
 		// check if we will be able to remove the texture also
 		if (LayerIdx != DeleteLayerIdx && Allocation.WeightmapTextureIndex == DeleteLayerWeightmapTextureIndex)
@@ -1535,18 +1539,18 @@ void ULandscapeComponent::DeleteLayer(ULandscapeLayerInfoObject* LayerInfo, FLan
 	{
 		// Lock data for all the weightmaps
 		TArray<FLandscapeTextureDataInfo*> TexDataInfos;
-		for (int32 WeightmapIdx = 0; WeightmapIdx < Component->WeightmapTextures.Num(); WeightmapIdx++)
+		for (int32 WeightmapIdx = 0; WeightmapIdx < ComponentWeightmapTextures.Num(); WeightmapIdx++)
 		{
-			TexDataInfos.Add(LandscapeEdit.GetTextureDataInfo(Component->WeightmapTextures[WeightmapIdx]));
+			TexDataInfos.Add(LandscapeEdit.GetTextureDataInfo(ComponentWeightmapTextures[WeightmapIdx]));
 		}
 
 		TArray<bool> LayerNoWeightBlends;	// Array of NoWeightBlend flags
 		TArray<uint8*> LayerDataPtrs;		// Pointers to all layers' data 
 
 		// Get the data for each layer
-		for (int32 LayerIdx = 0; LayerIdx < Component->WeightmapLayerAllocations.Num(); LayerIdx++)
+		for (int32 LayerIdx = 0; LayerIdx < ComponentWeightmapLayerAllocations.Num(); LayerIdx++)
 		{
-			FWeightmapLayerAllocationInfo& Allocation = Component->WeightmapLayerAllocations[LayerIdx];
+			FWeightmapLayerAllocationInfo& Allocation = ComponentWeightmapLayerAllocations[LayerIdx];
 			LayerDataPtrs.Add((uint8*)TexDataInfos[Allocation.WeightmapTextureIndex]->GetMipData(0) + ChannelOffsets[Allocation.WeightmapTextureChannel]);
 
 			// Find the layer info and record if it is a bNoWeightBlend layer.
@@ -1573,7 +1577,7 @@ void ULandscapeComponent::DeleteLayer(ULandscapeLayerInfoObject* LayerInfo, FLan
 
 						// Calculate the sum of other layer weights
 						int32 OtherLayerWeightSum = 0;
-						for (int32 LayerIdx = 0; LayerIdx < Component->WeightmapLayerAllocations.Num(); LayerIdx++)
+						for (int32 LayerIdx = 0; LayerIdx < ComponentWeightmapLayerAllocations.Num(); LayerIdx++)
 						{
 							if (LayerIdx != DeleteLayerIdx && LayerNoWeightBlends[LayerIdx] == false)
 							{
@@ -1586,7 +1590,7 @@ void ULandscapeComponent::DeleteLayer(ULandscapeLayerInfoObject* LayerInfo, FLan
 							// Set the first other weight-blend layer we can find to 255 to avoid a black hole
 							// This isn't ideal but it's the best option
 							// There's nothing we can easily do if this was the only weight-blend layer on this component
-							for (int32 LayerIdx = 0; LayerIdx < Component->WeightmapLayerAllocations.Num(); LayerIdx++)
+							for (int32 LayerIdx = 0; LayerIdx < ComponentWeightmapLayerAllocations.Num(); LayerIdx++)
 							{
 								if (LayerIdx != DeleteLayerIdx && LayerNoWeightBlends[LayerIdx] == false)
 								{
@@ -1599,7 +1603,7 @@ void ULandscapeComponent::DeleteLayer(ULandscapeLayerInfoObject* LayerInfo, FLan
 						else
 						{
 							// Adjust other layer weights
-							for (int32 LayerIdx = 0; LayerIdx < Component->WeightmapLayerAllocations.Num(); LayerIdx++)
+							for (int32 LayerIdx = 0; LayerIdx < ComponentWeightmapLayerAllocations.Num(); LayerIdx++)
 							{
 								if (LayerIdx != DeleteLayerIdx && LayerNoWeightBlends[LayerIdx] == false)
 								{
@@ -1614,7 +1618,7 @@ void ULandscapeComponent::DeleteLayer(ULandscapeLayerInfoObject* LayerInfo, FLan
 		}
 
 		// Update all the textures and mips
-		for (int32 Idx = 0; Idx < Component->WeightmapTextures.Num(); Idx++)
+		for (int32 Idx = 0; Idx < ComponentWeightmapTextures.Num(); Idx++)
 		{
 			if (bCanRemoveLayerTexture && Idx == DeleteLayerWeightmapTextureIndex)
 			{
@@ -1622,7 +1626,7 @@ void ULandscapeComponent::DeleteLayer(ULandscapeLayerInfoObject* LayerInfo, FLan
 				continue;
 			}
 
-			UTexture2D* WeightmapTexture = Component->WeightmapTextures[Idx];
+			UTexture2D* WeightmapTexture = ComponentWeightmapTextures[Idx];
 			FLandscapeTextureDataInfo* WeightmapDataInfo = TexDataInfos[Idx];
 
 			const int32 NumMips = WeightmapTexture->Source.GetNumMips();
@@ -1644,36 +1648,36 @@ void ULandscapeComponent::DeleteLayer(ULandscapeLayerInfoObject* LayerInfo, FLan
 	Component->Modify();
 	Proxy->Modify();
 
-	FLandscapeWeightmapUsage* Usage = Proxy->WeightmapUsageMap.Find(Component->WeightmapTextures[DeleteLayerAllocation.WeightmapTextureIndex]);
+	FLandscapeWeightmapUsage* Usage = ComponentWeightmapTexturesUsage.IsValidIndex(DeleteLayerAllocation.WeightmapTextureIndex) ? ComponentWeightmapTexturesUsage[DeleteLayerAllocation.WeightmapTextureIndex] : nullptr;
 	if (Usage) // can be null if WeightmapUsageMap hasn't been built yet
 	{
 		Usage->ChannelUsage[DeleteLayerAllocation.WeightmapTextureChannel] = nullptr;
 	}
 
 	// Remove the layer
-	Component->WeightmapLayerAllocations.RemoveAt(DeleteLayerIdx);
+	ComponentWeightmapLayerAllocations.RemoveAt(DeleteLayerIdx);
 
 	// If this layer was the last usage for this channel in this layer, we can remove it.
 	if (bCanRemoveLayerTexture)
 	{
-		Component->WeightmapTextures[DeleteLayerWeightmapTextureIndex]->SetFlags(RF_Transactional);
-		Component->WeightmapTextures[DeleteLayerWeightmapTextureIndex]->Modify();
-		Component->WeightmapTextures[DeleteLayerWeightmapTextureIndex]->MarkPackageDirty();
-		Component->WeightmapTextures[DeleteLayerWeightmapTextureIndex]->ClearFlags(RF_Standalone);
+		ComponentWeightmapTextures[DeleteLayerWeightmapTextureIndex]->SetFlags(RF_Transactional);
+		ComponentWeightmapTextures[DeleteLayerWeightmapTextureIndex]->Modify();
+		ComponentWeightmapTextures[DeleteLayerWeightmapTextureIndex]->MarkPackageDirty();
+		ComponentWeightmapTextures[DeleteLayerWeightmapTextureIndex]->ClearFlags(RF_Standalone);
 
-		Component->WeightmapTextures.RemoveAt(DeleteLayerWeightmapTextureIndex);
+		ComponentWeightmapTextures.RemoveAt(DeleteLayerWeightmapTextureIndex);
 
 		// Adjust WeightmapTextureIndex index for other layers
-		for (int32 LayerIdx = 0; LayerIdx < Component->WeightmapLayerAllocations.Num(); LayerIdx++)
+		for (int32 LayerIdx = 0; LayerIdx < ComponentWeightmapLayerAllocations.Num(); LayerIdx++)
 		{
-			FWeightmapLayerAllocationInfo& Allocation = Component->WeightmapLayerAllocations[LayerIdx];
+			FWeightmapLayerAllocationInfo& Allocation = ComponentWeightmapLayerAllocations[LayerIdx];
 
 			if (Allocation.WeightmapTextureIndex > DeleteLayerWeightmapTextureIndex)
 			{
 				Allocation.WeightmapTextureIndex--;
 			}
 
-			check(Allocation.WeightmapTextureIndex < Component->WeightmapTextures.Num());
+			check(Allocation.WeightmapTextureIndex < ComponentWeightmapTextures.Num());
 		}
 	}
 
@@ -1684,16 +1688,16 @@ void ULandscapeComponent::DeleteLayer(ULandscapeLayerInfoObject* LayerInfo, FLan
 
 	// Update dominant layer info stored in collision component
 	TArray<FColor*> CollisionWeightmapMipData;
-	for (int32 WeightmapIdx = 0; WeightmapIdx < Component->WeightmapTextures.Num(); WeightmapIdx++)
+	for (int32 WeightmapIdx = 0; WeightmapIdx < ComponentWeightmapTextures.Num(); WeightmapIdx++)
 	{
-		CollisionWeightmapMipData.Add((FColor*)LandscapeEdit.GetTextureDataInfo(Component->WeightmapTextures[WeightmapIdx])->GetMipData(Component->CollisionMipLevel));
+		CollisionWeightmapMipData.Add((FColor*)LandscapeEdit.GetTextureDataInfo(ComponentWeightmapTextures[WeightmapIdx])->GetMipData(Component->CollisionMipLevel));
 	}
 	TArray<FColor*> SimpleCollisionWeightmapMipData;
 	if (Component->SimpleCollisionMipLevel > Component->CollisionMipLevel)
 	{
-		for (int32 WeightmapIdx = 0; WeightmapIdx < Component->WeightmapTextures.Num(); WeightmapIdx++)
+		for (int32 WeightmapIdx = 0; WeightmapIdx < ComponentWeightmapTextures.Num(); WeightmapIdx++)
 		{
-			SimpleCollisionWeightmapMipData.Add((FColor*)LandscapeEdit.GetTextureDataInfo(Component->WeightmapTextures[WeightmapIdx])->GetMipData(Component->SimpleCollisionMipLevel));
+			SimpleCollisionWeightmapMipData.Add((FColor*)LandscapeEdit.GetTextureDataInfo(ComponentWeightmapTextures[WeightmapIdx])->GetMipData(Component->SimpleCollisionMipLevel));
 		}
 	}
 	Component->UpdateCollisionLayerData(
@@ -1733,19 +1737,23 @@ void ULandscapeComponent::FillLayer(ULandscapeLayerInfoObject* LayerInfo, FLands
 	const bool bFillLayerIsNoWeightBlend = LayerInfo->bNoWeightBlend;
 	bool bClearOtherWeightBlendLayers = !bFillLayerIsNoWeightBlend;
 
+	TArray<FWeightmapLayerAllocationInfo>& ComponentWeightmapLayerAllocations = Component->GetWeightmapLayerAllocations(true);
+	TArray<UTexture2D*>& ComponentWeightmapTextures = Component->GetWeightmapTextures(true);
+	TArray<FLandscapeWeightmapUsage*>& ComponentWeightmapTexturesUsage = Component->GetWeightmapTexturesUsage(true);
+
 	// Find the index for this layer in this component.
-	int32 FillLayerIdx = Component->WeightmapLayerAllocations.IndexOfByPredicate(
+	int32 FillLayerIdx = ComponentWeightmapLayerAllocations.IndexOfByPredicate(
 		[LayerInfo](const FWeightmapLayerAllocationInfo& Allocation) { return Allocation.LayerInfo == LayerInfo; });
 
 	// if the layer isn't used on this component yet but is a weight-blend layer, then simply steal the allocation of another weight-blend layer!
 	if (FillLayerIdx == INDEX_NONE && !bFillLayerIsNoWeightBlend)
 	{
-		FillLayerIdx = Component->WeightmapLayerAllocations.IndexOfByPredicate(
+		FillLayerIdx = ComponentWeightmapLayerAllocations.IndexOfByPredicate(
 			[](const FWeightmapLayerAllocationInfo& Allocation) { return !Allocation.LayerInfo || !Allocation.LayerInfo->bNoWeightBlend; });
 
 		if (FillLayerIdx != INDEX_NONE)
 		{
-			Component->WeightmapLayerAllocations[FillLayerIdx].LayerInfo = LayerInfo;
+			ComponentWeightmapLayerAllocations[FillLayerIdx].LayerInfo = LayerInfo;
 		}
 		else
 		{
@@ -1757,9 +1765,14 @@ void ULandscapeComponent::FillLayer(ULandscapeLayerInfoObject* LayerInfo, FLands
 	// if the layer is still not found then we are forced to make a new allocation
 	if (FillLayerIdx == INDEX_NONE)
 	{
-		FillLayerIdx = Component->WeightmapLayerAllocations.Num();
-		Component->WeightmapLayerAllocations.Add(FWeightmapLayerAllocationInfo(LayerInfo));
+		FillLayerIdx = ComponentWeightmapLayerAllocations.Num();
+		ComponentWeightmapLayerAllocations.Add(FWeightmapLayerAllocationInfo(LayerInfo));
 		Component->ReallocateWeightmaps(&LandscapeEdit);
+
+		if (GetMutableDefault<UEditorExperimentalSettings>()->bProceduralLandscape)
+		{
+			GetLandscapeActor()->RequestProceduralContentUpdate(EProceduralContentUpdateFlag::Weightmap_Render);
+		}
 	}
 
 	check(FillLayerIdx != INDEX_NONE);
@@ -1772,8 +1785,8 @@ void ULandscapeComponent::FillLayer(ULandscapeLayerInfoObject* LayerInfo, FLands
 		const int32 WeightmapOffsetX = Component->WeightmapScaleBias.Z * (float)SizeU;
 		const int32 WeightmapOffsetY = Component->WeightmapScaleBias.W * (float)SizeV;
 
-		FWeightmapLayerAllocationInfo& FillLayerAllocation = Component->WeightmapLayerAllocations[FillLayerIdx];
-		uint8* LayerData = (uint8*)LandscapeEdit.GetTextureDataInfo(Component->WeightmapTextures[FillLayerAllocation.WeightmapTextureIndex])->GetMipData(0);
+		FWeightmapLayerAllocationInfo& FillLayerAllocation = ComponentWeightmapLayerAllocations[FillLayerIdx];
+		uint8* LayerData = (uint8*)LandscapeEdit.GetTextureDataInfo(ComponentWeightmapTextures[FillLayerAllocation.WeightmapTextureIndex])->GetMipData(0);
 
 		for (int32 Y = 0; Y < SizeV; ++Y)
 		{
@@ -1789,19 +1802,19 @@ void ULandscapeComponent::FillLayer(ULandscapeLayerInfoObject* LayerInfo, FLands
 	// clear other layers
 	if (bClearOtherWeightBlendLayers)
 	{
-		for (int32 LayerIdx = 0; LayerIdx < Component->WeightmapLayerAllocations.Num(); ++LayerIdx)
+		for (int32 LayerIdx = 0; LayerIdx < ComponentWeightmapLayerAllocations.Num(); ++LayerIdx)
 		{
 			if (LayerIdx == FillLayerIdx)
 			{
 				continue;
 			}
-			FWeightmapLayerAllocationInfo& Allocation = Component->WeightmapLayerAllocations[LayerIdx];
+			FWeightmapLayerAllocationInfo& Allocation = ComponentWeightmapLayerAllocations[LayerIdx];
 			if (Allocation.LayerInfo->bNoWeightBlend)
 			{
 				continue;
 			}
 
-			FLandscapeWeightmapUsage* Usage = Proxy->WeightmapUsageMap.Find(Component->WeightmapTextures[Allocation.WeightmapTextureIndex]);
+			FLandscapeWeightmapUsage* Usage = ComponentWeightmapTexturesUsage[Allocation.WeightmapTextureIndex];
 			if (Usage) // can be null if WeightmapUsageMap hasn't been built yet
 			{
 				Usage->ChannelUsage[Allocation.WeightmapTextureChannel] = nullptr;
@@ -1810,27 +1823,27 @@ void ULandscapeComponent::FillLayer(ULandscapeLayerInfoObject* LayerInfo, FLands
 			Allocation.WeightmapTextureIndex = 255;
 		}
 
-		Component->WeightmapLayerAllocations.RemoveAll(
+		ComponentWeightmapLayerAllocations.RemoveAll(
 			[](const FWeightmapLayerAllocationInfo& Allocation) { return Allocation.WeightmapTextureIndex == 255; });
 
 		// remove any textures we're no longer using
-		for (int32 TextureIdx = 0; TextureIdx < Component->WeightmapTextures.Num(); ++TextureIdx)
+		for (int32 TextureIdx = 0; TextureIdx < ComponentWeightmapTextures.Num(); ++TextureIdx)
 		{
-			if (!Component->WeightmapLayerAllocations.ContainsByPredicate(
+			if (!ComponentWeightmapLayerAllocations.ContainsByPredicate(
 				[TextureIdx](const FWeightmapLayerAllocationInfo& Allocation) { return Allocation.WeightmapTextureIndex == TextureIdx; }))
 			{
-				Component->WeightmapTextures[TextureIdx]->Modify();
+				ComponentWeightmapTextures[TextureIdx]->Modify();
 
-				for (int32 LayerIdx = 0; LayerIdx < Component->WeightmapLayerAllocations.Num(); ++LayerIdx)
+				for (int32 LayerIdx = 0; LayerIdx < ComponentWeightmapLayerAllocations.Num(); ++LayerIdx)
 				{
-					FWeightmapLayerAllocationInfo& Allocation = Component->WeightmapLayerAllocations[LayerIdx];
+					FWeightmapLayerAllocationInfo& Allocation = ComponentWeightmapLayerAllocations[LayerIdx];
 					if (Allocation.WeightmapTextureIndex > TextureIdx)
 					{
 						--Allocation.WeightmapTextureIndex;
 					}
 				}
 
-				Component->WeightmapTextures.RemoveAt(TextureIdx--);
+				ComponentWeightmapTextures.RemoveAt(TextureIdx--);
 			}
 		}
 	}
@@ -1840,9 +1853,9 @@ void ULandscapeComponent::FillLayer(ULandscapeLayerInfoObject* LayerInfo, FLands
 	// this can also happen with normal painting I believe
 
 	// update mips
-	for (int32 TextureIdx = 0; TextureIdx < Component->WeightmapTextures.Num(); ++TextureIdx)
+	for (int32 TextureIdx = 0; TextureIdx < ComponentWeightmapTextures.Num(); ++TextureIdx)
 	{
-		UTexture2D* WeightmapTexture = WeightmapTextures[TextureIdx];
+		UTexture2D* WeightmapTexture = ComponentWeightmapTextures[TextureIdx];
 		FLandscapeTextureDataInfo* WeightmapDataInfo = LandscapeEdit.GetTextureDataInfo(WeightmapTexture);
 
 		const int32 NumMips = WeightmapTexture->Source.GetNumMips();
@@ -1867,16 +1880,16 @@ void ULandscapeComponent::FillLayer(ULandscapeLayerInfoObject* LayerInfo, FLands
 
 	// Update dominant layer info stored in collision component
 	TArray<FColor*> CollisionWeightmapMipData;
-	for (int32 WeightmapIdx = 0; WeightmapIdx < Component->WeightmapTextures.Num(); WeightmapIdx++)
+	for (int32 WeightmapIdx = 0; WeightmapIdx < ComponentWeightmapTextures.Num(); WeightmapIdx++)
 	{
-		CollisionWeightmapMipData.Add((FColor*)LandscapeEdit.GetTextureDataInfo(Component->WeightmapTextures[WeightmapIdx])->GetMipData(Component->CollisionMipLevel));
+		CollisionWeightmapMipData.Add((FColor*)LandscapeEdit.GetTextureDataInfo(ComponentWeightmapTextures[WeightmapIdx])->GetMipData(Component->CollisionMipLevel));
 	}
 	TArray<FColor*> SimpleCollisionWeightmapMipData;
 	if (Component->SimpleCollisionMipLevel > Component->CollisionMipLevel)
 	{
-		for (int32 WeightmapIdx = 0; WeightmapIdx < Component->WeightmapTextures.Num(); WeightmapIdx++)
+		for (int32 WeightmapIdx = 0; WeightmapIdx < ComponentWeightmapTextures.Num(); WeightmapIdx++)
 		{
-			SimpleCollisionWeightmapMipData.Add((FColor*)LandscapeEdit.GetTextureDataInfo(Component->WeightmapTextures[WeightmapIdx])->GetMipData(Component->SimpleCollisionMipLevel));
+			SimpleCollisionWeightmapMipData.Add((FColor*)LandscapeEdit.GetTextureDataInfo(ComponentWeightmapTextures[WeightmapIdx])->GetMipData(Component->SimpleCollisionMipLevel));
 		}
 	}
 	Component->UpdateCollisionLayerData(
@@ -1920,8 +1933,20 @@ void FLandscapeEditDataInterface::FillEmptyLayers(ULandscapeLayerInfoObject* Lay
 	for (auto& XYComponentPair : LandscapeInfo->XYtoComponentMap)
 	{
 		ULandscapeComponent* Component = XYComponentPair.Value;
+		TArray<FWeightmapLayerAllocationInfo>& ComponentWeightmapLayerAllocations = Component->GetWeightmapLayerAllocations(true);
 
-		if (Component->WeightmapLayerAllocations.Num() == 0)
+		bool LayerNeedFilling = true;
+
+		for (const FWeightmapLayerAllocationInfo& Alloc : ComponentWeightmapLayerAllocations)
+		{
+			if (Alloc.LayerInfo != nullptr)
+			{
+				LayerNeedFilling = false;
+				break;
+			}
+		}
+
+		if (LayerNeedFilling)
 		{
 			Component->FillLayer(LayerInfo, *this);
 		}
@@ -1944,11 +1969,15 @@ void ULandscapeComponent::ReplaceLayer(ULandscapeLayerInfoObject* FromLayerInfo,
 		return;
 	}
 
+	TArray<FWeightmapLayerAllocationInfo>& ComponentWeightmapLayerAllocations = GetWeightmapLayerAllocations(true);
+	TArray<UTexture2D*>& ComponentWeightmapTextures = GetWeightmapTextures(true);
+	TArray<FLandscapeWeightmapUsage*>& ComponentWeightmapTexturesUsage = GetWeightmapTexturesUsage(true);
+
 	// Find the index for this layer in this component.
 	int32 FromLayerIdx = INDEX_NONE;
-	for (int32 LayerIdx = 0; LayerIdx < WeightmapLayerAllocations.Num(); LayerIdx++)
+	for (int32 LayerIdx = 0; LayerIdx < ComponentWeightmapLayerAllocations.Num(); LayerIdx++)
 	{
-		FWeightmapLayerAllocationInfo& Allocation = WeightmapLayerAllocations[LayerIdx];
+		FWeightmapLayerAllocationInfo& Allocation = ComponentWeightmapLayerAllocations[LayerIdx];
 		if (Allocation.LayerInfo == FromLayerInfo)
 		{
 			FromLayerIdx = LayerIdx;
@@ -1964,9 +1993,9 @@ void ULandscapeComponent::ReplaceLayer(ULandscapeLayerInfoObject* FromLayerInfo,
 
 	// Find the index for this layer in this component.
 	int32 ToLayerIdx = INDEX_NONE;
-	for (int32 LayerIdx = 0; LayerIdx < WeightmapLayerAllocations.Num(); LayerIdx++)
+	for (int32 LayerIdx = 0; LayerIdx < ComponentWeightmapLayerAllocations.Num(); LayerIdx++)
 	{
-		FWeightmapLayerAllocationInfo& Allocation = WeightmapLayerAllocations[LayerIdx];
+		FWeightmapLayerAllocationInfo& Allocation = ComponentWeightmapLayerAllocations[LayerIdx];
 		if (Allocation.LayerInfo == ToLayerInfo)
 		{
 			ToLayerIdx = LayerIdx;
@@ -1975,20 +2004,20 @@ void ULandscapeComponent::ReplaceLayer(ULandscapeLayerInfoObject* FromLayerInfo,
 	if (ToLayerIdx == INDEX_NONE)
 	{
 		// Layer not used for this component, so do trivial replace.
-		WeightmapLayerAllocations[FromLayerIdx].LayerInfo = ToLayerInfo;
+		ComponentWeightmapLayerAllocations[FromLayerIdx].LayerInfo = ToLayerInfo;
 		bMerging = false;
 	}
 
-	FWeightmapLayerAllocationInfo& FromLayerAllocation = WeightmapLayerAllocations[FromLayerIdx];
+	FWeightmapLayerAllocationInfo& FromLayerAllocation = ComponentWeightmapLayerAllocations[FromLayerIdx];
 
 	// See if we'll be able to remove the texture completely.
 	bool bCanRemoveLayerTexture = false;
 	if (bMerging)
 	{
 		bCanRemoveLayerTexture = true;
-		for (int32 LayerIdx = 0; LayerIdx < WeightmapLayerAllocations.Num(); LayerIdx++)
+		for (int32 LayerIdx = 0; LayerIdx < ComponentWeightmapLayerAllocations.Num(); LayerIdx++)
 		{
-			FWeightmapLayerAllocationInfo& Allocation = WeightmapLayerAllocations[LayerIdx];
+			FWeightmapLayerAllocationInfo& Allocation = ComponentWeightmapLayerAllocations[LayerIdx];
 
 			// check if we will be able to remove the texture also
 			if (LayerIdx != FromLayerIdx && Allocation.WeightmapTextureIndex == FromLayerAllocation.WeightmapTextureIndex)
@@ -2007,11 +2036,11 @@ void ULandscapeComponent::ReplaceLayer(ULandscapeLayerInfoObject* FromLayerInfo,
 
 	if (bMerging)
 	{
-		FWeightmapLayerAllocationInfo& ToLayerAllocation = WeightmapLayerAllocations[ToLayerIdx];
+		FWeightmapLayerAllocationInfo& ToLayerAllocation = ComponentWeightmapLayerAllocations[ToLayerIdx];
 
 		// Lock data for all the weightmaps
-		FLandscapeTextureDataInfo* FromTexDataInfo = LandscapeEdit.GetTextureDataInfo(WeightmapTextures[FromLayerAllocation.WeightmapTextureIndex]);
-		FLandscapeTextureDataInfo* ToTexDataInfo = LandscapeEdit.GetTextureDataInfo(WeightmapTextures[ToLayerAllocation.WeightmapTextureIndex]);
+		FLandscapeTextureDataInfo* FromTexDataInfo = LandscapeEdit.GetTextureDataInfo(ComponentWeightmapTextures[FromLayerAllocation.WeightmapTextureIndex]);
+		FLandscapeTextureDataInfo* ToTexDataInfo = LandscapeEdit.GetTextureDataInfo(ComponentWeightmapTextures[ToLayerAllocation.WeightmapTextureIndex]);
 
 		check(FromTexDataInfo->GetMipSizeX(0) == FromTexDataInfo->GetMipSizeY(0));
 		check(ToTexDataInfo->GetMipSizeX(0) == ToTexDataInfo->GetMipSizeY(0));
@@ -2029,7 +2058,7 @@ void ULandscapeComponent::ReplaceLayer(ULandscapeLayerInfoObject* FromLayerInfo,
 		// Update all mips
 		if (!bCanRemoveLayerTexture)
 		{
-			UTexture2D* WeightmapTexture = WeightmapTextures[FromLayerAllocation.WeightmapTextureIndex];
+			UTexture2D* WeightmapTexture = ComponentWeightmapTextures[FromLayerAllocation.WeightmapTextureIndex];
 			FLandscapeTextureDataInfo* WeightmapDataInfo = FromTexDataInfo;
 
 			const int32 NumMips = WeightmapTexture->Source.GetNumMips();
@@ -2047,7 +2076,7 @@ void ULandscapeComponent::ReplaceLayer(ULandscapeLayerInfoObject* FromLayerInfo,
 
 		if (FromTexDataInfo != ToTexDataInfo)
 		{
-			UTexture2D* WeightmapTexture = WeightmapTextures[ToLayerAllocation.WeightmapTextureIndex];
+			UTexture2D* WeightmapTexture = ComponentWeightmapTextures[ToLayerAllocation.WeightmapTextureIndex];
 			FLandscapeTextureDataInfo* WeightmapDataInfo = ToTexDataInfo;
 
 			const int32 NumMips = WeightmapTexture->Source.GetNumMips();
@@ -2074,7 +2103,7 @@ void ULandscapeComponent::ReplaceLayer(ULandscapeLayerInfoObject* FromLayerInfo,
 	{
 		ALandscapeProxy* Proxy = GetLandscapeProxy();
 		// Mark the channel as unallocated, so we can reuse it later
-		FLandscapeWeightmapUsage* Usage = Proxy->WeightmapUsageMap.Find(WeightmapTextures[FromLayerAllocation.WeightmapTextureIndex]);
+		FLandscapeWeightmapUsage* Usage = ComponentWeightmapTexturesUsage[FromLayerAllocation.WeightmapTextureIndex];
 		//check(Usage);
 		if (Usage)
 		{
@@ -2084,34 +2113,34 @@ void ULandscapeComponent::ReplaceLayer(ULandscapeLayerInfoObject* FromLayerInfo,
 		// If this layer was the last usage for this texture, we can remove it.
 		if (bCanRemoveLayerTexture)
 		{
-			WeightmapTextures[FromLayerAllocation.WeightmapTextureIndex]->SetFlags(RF_Transactional);
-			WeightmapTextures[FromLayerAllocation.WeightmapTextureIndex]->Modify();
-			WeightmapTextures[FromLayerAllocation.WeightmapTextureIndex]->MarkPackageDirty();
-			WeightmapTextures[FromLayerAllocation.WeightmapTextureIndex]->ClearFlags(RF_Standalone);
+			ComponentWeightmapTextures[FromLayerAllocation.WeightmapTextureIndex]->SetFlags(RF_Transactional);
+			ComponentWeightmapTextures[FromLayerAllocation.WeightmapTextureIndex]->Modify();
+			ComponentWeightmapTextures[FromLayerAllocation.WeightmapTextureIndex]->MarkPackageDirty();
+			ComponentWeightmapTextures[FromLayerAllocation.WeightmapTextureIndex]->ClearFlags(RF_Standalone);
 
-			WeightmapTextures.RemoveAt(FromLayerAllocation.WeightmapTextureIndex);
+			ComponentWeightmapTextures.RemoveAt(FromLayerAllocation.WeightmapTextureIndex);
 
 			// Adjust WeightmapTextureIndex index for other layers
-			for (int32 LayerIdx = 0; LayerIdx < WeightmapLayerAllocations.Num(); LayerIdx++)
+			for (int32 LayerIdx = 0; LayerIdx < ComponentWeightmapLayerAllocations.Num(); LayerIdx++)
 			{
 				if (LayerIdx == FromLayerIdx)
 				{
 					continue;
 				}
 
-				FWeightmapLayerAllocationInfo& Allocation = WeightmapLayerAllocations[LayerIdx];
+				FWeightmapLayerAllocationInfo& Allocation = ComponentWeightmapLayerAllocations[LayerIdx];
 
 				if (Allocation.WeightmapTextureIndex > FromLayerAllocation.WeightmapTextureIndex)
 				{
 					Allocation.WeightmapTextureIndex--;
 				}
 
-				check(Allocation.WeightmapTextureIndex < WeightmapTextures.Num());
+				check(Allocation.WeightmapTextureIndex < ComponentWeightmapTextures.Num());
 			}
 		}
 
 		// Remove the layer
-		WeightmapLayerAllocations.RemoveAt(FromLayerIdx);
+		ComponentWeightmapLayerAllocations.RemoveAt(FromLayerIdx);
 
 		// Update the shaders for this component
 		UpdateMaterialInstances();
@@ -2129,18 +2158,20 @@ void FLandscapeEditDataInterface::ReplaceLayer(ULandscapeLayerInfoObject* FromLa
 		ULandscapeComponent* Component = It.Value();
 		Component->ReplaceLayer(FromLayerInfo, ToLayerInfo, *this);
 
+		TArray<UTexture2D*>& ComponentWeightmapTextures = Component->GetWeightmapTextures(true);
+
 		// Update dominant layer info stored in collision component
 		TArray<FColor*> CollisionWeightmapMipData;
-		for( int32 WeightmapIdx=0;WeightmapIdx < Component->WeightmapTextures.Num();WeightmapIdx++ )
+		for( int32 WeightmapIdx=0;WeightmapIdx < ComponentWeightmapTextures.Num();WeightmapIdx++ )
 		{
-			CollisionWeightmapMipData.Add((FColor*)GetTextureDataInfo(Component->WeightmapTextures[WeightmapIdx])->GetMipData(Component->CollisionMipLevel));
+			CollisionWeightmapMipData.Add((FColor*)GetTextureDataInfo(ComponentWeightmapTextures[WeightmapIdx])->GetMipData(Component->CollisionMipLevel));
 		}
 		TArray<FColor*> SimpleCollisionWeightmapMipData;
 		if (Component->SimpleCollisionMipLevel > Component->CollisionMipLevel)
 		{
-			for (int32 WeightmapIdx = 0; WeightmapIdx < Component->WeightmapTextures.Num(); WeightmapIdx++)
+			for (int32 WeightmapIdx = 0; WeightmapIdx < ComponentWeightmapTextures.Num(); WeightmapIdx++)
 			{
-				SimpleCollisionWeightmapMipData.Add((FColor*)GetTextureDataInfo(Component->WeightmapTextures[WeightmapIdx])->GetMipData(Component->SimpleCollisionMipLevel));
+				SimpleCollisionWeightmapMipData.Add((FColor*)GetTextureDataInfo(ComponentWeightmapTextures[WeightmapIdx])->GetMipData(Component->SimpleCollisionMipLevel));
 			}
 		}
 		Component->UpdateCollisionLayerData(
@@ -2316,24 +2347,28 @@ bool DeleteLayerIfAllZero(ULandscapeComponent* const Component, const uint8* con
 	Component->Modify();
 	Proxy->Modify();
 
+	TArray<FWeightmapLayerAllocationInfo>& ComponentWeightmapLayerAllocations = Component->GetWeightmapLayerAllocations(true);
+	TArray<UTexture2D*>& ComponentWeightmapTextures = Component->GetWeightmapTextures(true);
+	TArray<FLandscapeWeightmapUsage*>& ComponentWeightmapTexturesUsage = Component->GetWeightmapTexturesUsage(true);
+
 	// Mark the channel as unallocated, so we can reuse it later
-	const int32 DeleteLayerWeightmapTextureIndex = Component->WeightmapLayerAllocations[LayerIdx].WeightmapTextureIndex;
-	FLandscapeWeightmapUsage& Usage = Proxy->WeightmapUsageMap.FindChecked(Component->WeightmapTextures[DeleteLayerWeightmapTextureIndex]);
-	Usage.ChannelUsage[Component->WeightmapLayerAllocations[LayerIdx].WeightmapTextureChannel] = NULL;
+	const int32 DeleteLayerWeightmapTextureIndex = ComponentWeightmapLayerAllocations[LayerIdx].WeightmapTextureIndex;
+	FLandscapeWeightmapUsage& Usage = *ComponentWeightmapTexturesUsage[DeleteLayerWeightmapTextureIndex];
+	Usage.ChannelUsage[ComponentWeightmapLayerAllocations[LayerIdx].WeightmapTextureChannel] = NULL;
 
 	// Remove the layer as it's totally painted away.
-	Component->WeightmapLayerAllocations.RemoveAt(LayerIdx);
+	ComponentWeightmapLayerAllocations.RemoveAt(LayerIdx);
 
 	// Check if the weightmap texture used by the layer we just removed is used by any other layer, and if so, remove the texture too
-	bool bCanRemoveLayerTexture = !Component->WeightmapLayerAllocations.ContainsByPredicate([DeleteLayerWeightmapTextureIndex](const FWeightmapLayerAllocationInfo& Allocation){ return Allocation.WeightmapTextureIndex == DeleteLayerWeightmapTextureIndex; });
+	bool bCanRemoveLayerTexture = !ComponentWeightmapLayerAllocations.ContainsByPredicate([DeleteLayerWeightmapTextureIndex](const FWeightmapLayerAllocationInfo& Allocation){ return Allocation.WeightmapTextureIndex == DeleteLayerWeightmapTextureIndex; });
 	if (bCanRemoveLayerTexture)
 	{
-		Component->WeightmapTextures[DeleteLayerWeightmapTextureIndex]->MarkPackageDirty();
-		Component->WeightmapTextures[DeleteLayerWeightmapTextureIndex]->ClearFlags(RF_Standalone);
-		Component->WeightmapTextures.RemoveAt(DeleteLayerWeightmapTextureIndex);
+		ComponentWeightmapTextures[DeleteLayerWeightmapTextureIndex]->MarkPackageDirty();
+		ComponentWeightmapTextures[DeleteLayerWeightmapTextureIndex]->ClearFlags(RF_Standalone);
+		ComponentWeightmapTextures.RemoveAt(DeleteLayerWeightmapTextureIndex);
 
 		// Adjust WeightmapTextureChannel index for other layers
-		for (auto It = Component->WeightmapLayerAllocations.CreateIterator(); It; ++It)
+		for (auto It = ComponentWeightmapLayerAllocations.CreateIterator(); It; ++It)
 		{
 			FWeightmapLayerAllocationInfo& Allocation = *It;
 			if (Allocation.WeightmapTextureIndex > DeleteLayerWeightmapTextureIndex)
@@ -2429,6 +2464,8 @@ inline TMap<const ULandscapeLayerInfoObject*, uint32> FLandscapeEditDataInterfac
 	TMap<const ULandscapeLayerInfoObject*, uint32> LayerInfluenceMap;
 
 	ULandscapeComponent* Component = LandscapeInfo->XYtoComponentMap.FindChecked(FIntPoint(ComponentIndexX,ComponentIndexY));
+	TArray<FWeightmapLayerAllocationInfo>& ComponentWeightmapLayerAllocations = Component->GetWeightmapLayerAllocations(true);
+	TArray<UTexture2D*>& ComponentWeightmapTextures = Component->GetWeightmapTextures(true);
 
 	// used if InOptionalLayerDataPtrs is null
 	TArray<FLandscapeTextureDataInfo*, TInlineAllocator<2>> InternalTexDataInfos;
@@ -2436,21 +2473,21 @@ inline TMap<const ULandscapeLayerInfoObject*, uint32> FLandscapeEditDataInterfac
 	TArrayView<const uint8* const> LayerDataPtrs;
 	if (InOptionalLayerDataPtrs)
 	{
-		check(InOptionalLayerDataPtrs->Num() == Component->WeightmapLayerAllocations.Num());
+		check(InOptionalLayerDataPtrs->Num() == ComponentWeightmapLayerAllocations.Num());
 		LayerDataPtrs = InOptionalLayerDataPtrs.GetValue();
 	}
 	else
 	{
-		InternalTexDataInfos.AddUninitialized(Component->WeightmapTextures.Num());
-		for (int32 WeightmapIdx = 0; WeightmapIdx < Component->WeightmapTextures.Num(); ++WeightmapIdx)
+		InternalTexDataInfos.AddUninitialized(ComponentWeightmapTextures.Num());
+		for (int32 WeightmapIdx = 0; WeightmapIdx < ComponentWeightmapTextures.Num(); ++WeightmapIdx)
 		{
-			InternalTexDataInfos[WeightmapIdx] = GetTextureDataInfo(Component->WeightmapTextures[WeightmapIdx]);
+			InternalTexDataInfos[WeightmapIdx] = GetTextureDataInfo(ComponentWeightmapTextures[WeightmapIdx]);
 		}
 
-		InternalLayerDataPtrs.AddUninitialized(Component->WeightmapLayerAllocations.Num());
-		for (int32 LayerIdx = 0; LayerIdx < Component->WeightmapLayerAllocations.Num(); LayerIdx++)
+		InternalLayerDataPtrs.AddUninitialized(ComponentWeightmapLayerAllocations.Num());
+		for (int32 LayerIdx = 0; LayerIdx < ComponentWeightmapLayerAllocations.Num(); LayerIdx++)
 		{
-			const FWeightmapLayerAllocationInfo& Allocation = Component->WeightmapLayerAllocations[LayerIdx];
+			const FWeightmapLayerAllocationInfo& Allocation = ComponentWeightmapLayerAllocations[LayerIdx];
 			InternalLayerDataPtrs[LayerIdx] = (uint8*)InternalTexDataInfos[Allocation.WeightmapTextureIndex]->GetMipData(0) + ChannelOffsets[Allocation.WeightmapTextureChannel];
 		}
 
@@ -2459,9 +2496,9 @@ inline TMap<const ULandscapeLayerInfoObject*, uint32> FLandscapeEditDataInterfac
 
 	const int32 ScanlineSize = (SubsectionSizeQuads + 1) * ComponentNumSubsections * 4;
 
-	for (int32 LayerIdx = 0; LayerIdx < Component->WeightmapLayerAllocations.Num(); LayerIdx++)
+	for (int32 LayerIdx = 0; LayerIdx < ComponentWeightmapLayerAllocations.Num(); LayerIdx++)
 	{
-		const FWeightmapLayerAllocationInfo& Allocation = Component->WeightmapLayerAllocations[LayerIdx];
+		const FWeightmapLayerAllocationInfo& Allocation = ComponentWeightmapLayerAllocations[LayerIdx];
 		if (Allocation.LayerInfo->bNoWeightBlend)
 		{
 			continue;
@@ -2622,7 +2659,8 @@ void FLandscapeEditDataInterface::SetAlphaData(ULandscapeLayerInfoObject* const
 
 			Component->Modify();
 
-			int32 UpdateLayerIdx = Component->WeightmapLayerAllocations.IndexOfByPredicate([LayerInfo](const FWeightmapLayerAllocationInfo& Allocation){ return Allocation.LayerInfo == LayerInfo; });
+			TArray<FWeightmapLayerAllocationInfo>& ComponentWeightmapLayerAllocations = Component->GetWeightmapLayerAllocations(true);
+			int32 UpdateLayerIdx = ComponentWeightmapLayerAllocations.IndexOfByPredicate([LayerInfo](const FWeightmapLayerAllocationInfo& Allocation){ return Allocation.LayerInfo == LayerInfo; });
 
 			// Need allocation for weightmap
 			if (UpdateLayerIdx == INDEX_NONE)
@@ -2631,13 +2669,13 @@ void FLandscapeEditDataInterface::SetAlphaData(ULandscapeLayerInfoObject* const
 
 				// if we can't allocate a layer, then there is nothing to paint
 				if (PaintingRestriction == ELandscapeLayerPaintingRestriction::ExistingOnly ||
-					(PaintingRestriction == ELandscapeLayerPaintingRestriction::UseMaxLayers && LayerLimit > 0 && Component->WeightmapLayerAllocations.Num() >= LayerLimit))
+					(PaintingRestriction == ELandscapeLayerPaintingRestriction::UseMaxLayers && LayerLimit > 0 && ComponentWeightmapLayerAllocations.Num() >= LayerLimit))
 				{
 					continue;
 				}
 
-				UpdateLayerIdx = Component->WeightmapLayerAllocations.Num();
-				new (Component->WeightmapLayerAllocations) FWeightmapLayerAllocationInfo(LayerInfo);
+				UpdateLayerIdx = ComponentWeightmapLayerAllocations.Num();
+				new (ComponentWeightmapLayerAllocations) FWeightmapLayerAllocationInfo(LayerInfo);
 				Component->ReallocateWeightmaps(this);
 
 				Component->UpdateMaterialInstances();
@@ -2649,23 +2687,24 @@ void FLandscapeEditDataInterface::SetAlphaData(ULandscapeLayerInfoObject* const
 
 			// Lock data for all the weightmaps
 			TexDataInfos.Reset();
-			TexDataInfos.AddUninitialized(Component->WeightmapTextures.Num());
+			TArray<UTexture2D*>& ComponentWeightmapTextures = Component->GetWeightmapTextures(true);
+			TexDataInfos.AddUninitialized(ComponentWeightmapTextures.Num());
 
-			for (int32 WeightmapIdx = 0; WeightmapIdx < Component->WeightmapTextures.Num(); ++WeightmapIdx)
+			for (int32 WeightmapIdx = 0; WeightmapIdx < ComponentWeightmapTextures.Num(); ++WeightmapIdx)
 			{
-				TexDataInfos[WeightmapIdx] = GetTextureDataInfo(Component->WeightmapTextures[WeightmapIdx]);
+				TexDataInfos[WeightmapIdx] = GetTextureDataInfo(ComponentWeightmapTextures[WeightmapIdx]);
 			}
 
 			LayerDataPtrs.Reset();        // Pointers to all layers' data
-			LayerDataPtrs.AddUninitialized(Component->WeightmapLayerAllocations.Num());
+			LayerDataPtrs.AddUninitialized(ComponentWeightmapLayerAllocations.Num());
 			LayerNoWeightBlends.Reset();  // NoWeightBlend flags
-			LayerNoWeightBlends.AddUninitialized(Component->WeightmapLayerAllocations.Num());
+			LayerNoWeightBlends.AddUninitialized(ComponentWeightmapLayerAllocations.Num());
 			LayerEditDataAllZero.Reset(); // Whether the data we are editing for this layer is all zero 
-			LayerEditDataAllZero.AddUninitialized(Component->WeightmapLayerAllocations.Num());
+			LayerEditDataAllZero.AddUninitialized(ComponentWeightmapLayerAllocations.Num());
 
-			for (int32 LayerIdx = 0; LayerIdx < Component->WeightmapLayerAllocations.Num(); LayerIdx++)
+			for (int32 LayerIdx = 0; LayerIdx < ComponentWeightmapLayerAllocations.Num(); LayerIdx++)
 			{
-				FWeightmapLayerAllocationInfo& Allocation = Component->WeightmapLayerAllocations[LayerIdx];
+				FWeightmapLayerAllocationInfo& Allocation = ComponentWeightmapLayerAllocations[LayerIdx];
 
 				if (Allocation.LayerInfo != nullptr) // only take into account valid layer
 				{
@@ -2746,7 +2785,7 @@ void FLandscapeEditDataInterface::SetAlphaData(ULandscapeLayerInfoObject* const
 									int32 MaxWeight = INT_MIN;
 
 									// Adjust other layers' weights accordingly
-									for (int32 LayerIdx = 0; LayerIdx < Component->WeightmapLayerAllocations.Num(); LayerIdx++)
+									for (int32 LayerIdx = 0; LayerIdx < ComponentWeightmapLayerAllocations.Num(); LayerIdx++)
 									{
 										uint8& ExistingWeight = LayerDataPtrs[LayerIdx][TexDataIndex];
 
@@ -2772,7 +2811,7 @@ void FLandscapeEditDataInterface::SetAlphaData(ULandscapeLayerInfoObject* const
 										OtherLayerWeightSum = 0;
 
 										// Normalize
-										for (int32 LayerIdx = 0; LayerIdx<Component->WeightmapLayerAllocations.Num(); LayerIdx++)
+										for (int32 LayerIdx = 0; LayerIdx < ComponentWeightmapLayerAllocations.Num(); LayerIdx++)
 										{
 											uint8& ExistingWeight = LayerDataPtrs[LayerIdx][TexDataIndex];
 
@@ -2798,7 +2837,7 @@ void FLandscapeEditDataInterface::SetAlphaData(ULandscapeLayerInfoObject* const
 								else
 								{
 									// Adjust other layers' weights accordingly
-									for (int32 LayerIdx = 0; LayerIdx < Component->WeightmapLayerAllocations.Num(); LayerIdx++)
+									for (int32 LayerIdx = 0; LayerIdx < ComponentWeightmapLayerAllocations.Num(); LayerIdx++)
 									{
 										// Exclude bNoWeightBlend layers
 										if (LayerIdx != UpdateLayerIdx && LayerNoWeightBlends[LayerIdx] == false)
@@ -2815,7 +2854,7 @@ void FLandscapeEditDataInterface::SetAlphaData(ULandscapeLayerInfoObject* const
 											const ULandscapeLayerInfoObject* ReplacementLayer = ChooseReplacementLayer(LayerInfo, ComponentIndexX, SubIndexX, SubX, ComponentIndexY, SubIndexY, SubY, LayerInfluenceCache, LayerDataPtrs);
 											if (ReplacementLayer)
 											{
-												const int32 ReplacementLayerIndex = Component->WeightmapLayerAllocations.IndexOfByPredicate([&](const FWeightmapLayerAllocationInfo& AllocationInfo) { return AllocationInfo.LayerInfo == ReplacementLayer; });
+												const int32 ReplacementLayerIndex = ComponentWeightmapLayerAllocations.IndexOfByPredicate([&](const FWeightmapLayerAllocationInfo& AllocationInfo) { return AllocationInfo.LayerInfo == ReplacementLayer; });
 
 												LayerDataPtrs[ReplacementLayerIndex][TexDataIndex] = 255 - NewWeight;
 												LayerEditDataAllZero[ReplacementLayerIndex] = false;
@@ -2836,7 +2875,7 @@ void FLandscapeEditDataInterface::SetAlphaData(ULandscapeLayerInfoObject* const
 									}
 									else
 									{
-										for (int32 LayerIdx = 0; LayerIdx < Component->WeightmapLayerAllocations.Num(); LayerIdx++)
+										for (int32 LayerIdx = 0; LayerIdx < ComponentWeightmapLayerAllocations.Num(); LayerIdx++)
 										{
 											uint8& Weight = LayerDataPtrs[LayerIdx][TexDataIndex];
 
@@ -2879,7 +2918,7 @@ void FLandscapeEditDataInterface::SetAlphaData(ULandscapeLayerInfoObject* const
 					const int32 TexY1 = (SubsectionSizeQuads+1) * SubIndexY + SubY1;
 					const int32 TexX2 = (SubsectionSizeQuads+1) * SubIndexX + SubX2;
 					const int32 TexY2 = (SubsectionSizeQuads+1) * SubIndexY + SubY2;
-					for (int32 WeightmapIdx = 0; WeightmapIdx < Component->WeightmapTextures.Num(); WeightmapIdx++)
+					for (int32 WeightmapIdx = 0; WeightmapIdx < ComponentWeightmapTextures.Num(); WeightmapIdx++)
 					{
 						if (TexDataInfos[WeightmapIdx] != NULL)
 						{
@@ -2891,10 +2930,10 @@ void FLandscapeEditDataInterface::SetAlphaData(ULandscapeLayerInfoObject* const
 
 			// Update mipmaps
 			CollisionWeightmapMipData.Reset();
-			CollisionWeightmapMipData.AddUninitialized(Component->WeightmapTextures.Num());
-			for (int32 WeightmapIdx = 0; WeightmapIdx < Component->WeightmapTextures.Num(); WeightmapIdx++)
+			CollisionWeightmapMipData.AddUninitialized(ComponentWeightmapTextures.Num());
+			for (int32 WeightmapIdx = 0; WeightmapIdx < ComponentWeightmapTextures.Num(); WeightmapIdx++)
 			{
-				UTexture2D* const WeightmapTexture = Component->WeightmapTextures[WeightmapIdx];
+				UTexture2D* const WeightmapTexture = ComponentWeightmapTextures[WeightmapIdx];
 
 				const int32 NumMips = WeightmapTexture->Source.GetNumMips();
 				WeightmapTextureMipData.Reset();
@@ -2912,7 +2951,7 @@ void FLandscapeEditDataInterface::SetAlphaData(ULandscapeLayerInfoObject* const
 
 			if (Component->SimpleCollisionMipLevel > Component->CollisionMipLevel)
 			{
-				for (int32 WeightmapIdx = 0; WeightmapIdx < Component->WeightmapTextures.Num(); WeightmapIdx++)
+				for (int32 WeightmapIdx = 0; WeightmapIdx < ComponentWeightmapTextures.Num(); WeightmapIdx++)
 				{
 					SimpleCollisionWeightmapMipData.Add((FColor*)TexDataInfos[WeightmapIdx]->GetMipData(Component->SimpleCollisionMipLevel));
 				}
@@ -2928,7 +2967,7 @@ void FLandscapeEditDataInterface::SetAlphaData(ULandscapeLayerInfoObject* const
 
 			// Check if we need to remove weightmap allocations for layers that were completely painted away
 			bool bRemovedLayer = false;
-			for (int32 LayerIdx = 0; LayerIdx < Component->WeightmapLayerAllocations.Num(); LayerIdx++)
+			for (int32 LayerIdx = 0; LayerIdx < ComponentWeightmapLayerAllocations.Num(); LayerIdx++)
 			{
 				if (LayerEditDataAllZero[LayerIdx])
 				{
@@ -3017,9 +3056,11 @@ void FLandscapeEditDataInterface::SetAlphaData(const TSet<ULandscapeLayerInfoObj
 
 			NeedAllocationInfos.Reset();
 
+			TArray<FWeightmapLayerAllocationInfo>& ComponentWeightmapLayerAllocations = Component->GetWeightmapLayerAllocations(true);
+
 			for (ULandscapeLayerInfoObject* LayerInfo : DirtyLayerInfos)
 			{
-				const bool bFound = Component->WeightmapLayerAllocations.ContainsByPredicate([LayerInfo](const FWeightmapLayerAllocationInfo& Allocation){ return Allocation.LayerInfo == LayerInfo; });
+				const bool bFound = ComponentWeightmapLayerAllocations.ContainsByPredicate([LayerInfo](const FWeightmapLayerAllocationInfo& Allocation){ return Allocation.LayerInfo == LayerInfo; });
 				if (!bFound)
 				{
 					NeedAllocationInfos.Add(LayerInfo);
@@ -3033,7 +3074,7 @@ void FLandscapeEditDataInterface::SetAlphaData(const TSet<ULandscapeLayerInfoObj
 				{
 					if (PaintingRestriction == ELandscapeLayerPaintingRestriction::ExistingOnly ||
 						(PaintingRestriction == ELandscapeLayerPaintingRestriction::UseMaxLayers &&
-						 Component->WeightmapLayerAllocations.Num() >= LayerLimit))
+							ComponentWeightmapLayerAllocations.Num() >= LayerLimit))
 					{
 						// nothing to paint to this component due to layer limit
 						continue;
@@ -3045,11 +3086,11 @@ void FLandscapeEditDataInterface::SetAlphaData(const TSet<ULandscapeLayerInfoObj
 					for (ULandscapeLayerInfoObject* LayerInfoNeedingAllocation : NeedAllocationInfos)
 					{
 						if (PaintingRestriction == ELandscapeLayerPaintingRestriction::UseMaxLayers &&
-							LayerLimit > 0 && Component->WeightmapLayerAllocations.Num() >= LayerLimit)
+							LayerLimit > 0 && ComponentWeightmapLayerAllocations.Num() >= LayerLimit)
 						{
 							break;
 						}
-						Component->WeightmapLayerAllocations.Emplace(LayerInfoNeedingAllocation);
+						ComponentWeightmapLayerAllocations.Emplace(LayerInfoNeedingAllocation);
 					}
 					Component->ReallocateWeightmaps(this);
 					Component->UpdateMaterialInstances();
@@ -3062,22 +3103,23 @@ void FLandscapeEditDataInterface::SetAlphaData(const TSet<ULandscapeLayerInfoObj
 
 			// Lock data for all the weightmaps
 			TexDataInfos.Reset();
-			TexDataInfos.AddUninitialized(Component->WeightmapTextures.Num());
+			TArray<UTexture2D*>& ComponentWeightmapTextures = Component->GetWeightmapTextures(true);
+			TexDataInfos.AddUninitialized(ComponentWeightmapTextures.Num());
 
-			for (int32 WeightmapIdx = 0; WeightmapIdx < Component->WeightmapTextures.Num(); ++WeightmapIdx)
+			for (int32 WeightmapIdx = 0; WeightmapIdx < ComponentWeightmapTextures.Num(); ++WeightmapIdx)
 			{
-				TexDataInfos[WeightmapIdx] = GetTextureDataInfo(Component->WeightmapTextures[WeightmapIdx]);
+				TexDataInfos[WeightmapIdx] = GetTextureDataInfo(ComponentWeightmapTextures[WeightmapIdx]);
 			}
 
 			LayerDataInfos.Reset();        // Pointers to all layers' data
-			LayerDataInfos.AddUninitialized(Component->WeightmapLayerAllocations.Num());
+			LayerDataInfos.AddUninitialized(ComponentWeightmapLayerAllocations.Num());
 			LayerEditDataAllZero.Reset(); // Whether the data we are editing for this layer is all zero 
-			LayerEditDataAllZero.AddUninitialized(Component->WeightmapLayerAllocations.Num());
+			LayerEditDataAllZero.AddUninitialized(ComponentWeightmapLayerAllocations.Num());
 
-			for (int32 LayerIdx = 0; LayerIdx < Component->WeightmapLayerAllocations.Num(); LayerIdx++)
+			for (int32 LayerIdx = 0; LayerIdx < ComponentWeightmapLayerAllocations.Num(); LayerIdx++)
 			{
-				FWeightmapLayerAllocationInfo& Allocation = Component->WeightmapLayerAllocations[LayerIdx];
-				const int32 LayerDataIdx = LandscapeInfo->GetLayerInfoIndex(Component->WeightmapLayerAllocations[LayerIdx].LayerInfo);
+				FWeightmapLayerAllocationInfo& Allocation = ComponentWeightmapLayerAllocations[LayerIdx];
+				const int32 LayerDataIdx = LandscapeInfo->GetLayerInfoIndex(ComponentWeightmapLayerAllocations[LayerIdx].LayerInfo);
 				check(LayerDataIdx != INDEX_NONE);
 				LayerDataInfos[LayerIdx].InDataPtr = Data + LayerDataIdx;
 				LayerDataInfos[LayerIdx].TexDataPtr = (uint8*)TexDataInfos[Allocation.WeightmapTextureIndex]->GetMipData(0) + ChannelOffsets[Allocation.WeightmapTextureChannel];
@@ -3131,7 +3173,7 @@ void FLandscapeEditDataInterface::SetAlphaData(const TSet<ULandscapeLayerInfoObj
 							int32 OtherLayerWeightSum = 0;
 
 							// Apply weights to all layers simultaneously.
-							for (int32 LayerIdx = 0; LayerIdx < Component->WeightmapLayerAllocations.Num(); LayerIdx++)
+							for (int32 LayerIdx = 0; LayerIdx < ComponentWeightmapLayerAllocations.Num(); LayerIdx++)
 							{
 								// this is equivalent to saying if (DirtyLayerInfos.Contains(Allocation.LayerInfo))
 								// which is what we really mean here, but this is quicker
@@ -3155,7 +3197,7 @@ void FLandscapeEditDataInterface::SetAlphaData(const TSet<ULandscapeLayerInfoObj
 					const int32 TexY1 = (SubsectionSizeQuads+1) * SubIndexY + SubY1;
 					const int32 TexX2 = (SubsectionSizeQuads+1) * SubIndexX + SubX2;
 					const int32 TexY2 = (SubsectionSizeQuads+1) * SubIndexY + SubY2;
-					for (int32 WeightmapIdx = 0; WeightmapIdx < Component->WeightmapTextures.Num(); WeightmapIdx++)
+					for (int32 WeightmapIdx = 0; WeightmapIdx < ComponentWeightmapTextures.Num(); WeightmapIdx++)
 					{
 						if (TexDataInfos[WeightmapIdx] != NULL)
 						{
@@ -3167,10 +3209,10 @@ void FLandscapeEditDataInterface::SetAlphaData(const TSet<ULandscapeLayerInfoObj
 
 			// Update mipmaps
 			CollisionWeightmapMipData.Reset();
-			CollisionWeightmapMipData.AddUninitialized(Component->WeightmapTextures.Num());
-			for (int32 WeightmapIdx = 0; WeightmapIdx < Component->WeightmapTextures.Num(); WeightmapIdx++)
+			CollisionWeightmapMipData.AddUninitialized(ComponentWeightmapTextures.Num());
+			for (int32 WeightmapIdx = 0; WeightmapIdx < ComponentWeightmapTextures.Num(); WeightmapIdx++)
 			{
-				UTexture2D* const WeightmapTexture = Component->WeightmapTextures[WeightmapIdx];
+				UTexture2D* const WeightmapTexture = ComponentWeightmapTextures[WeightmapIdx];
 
 				const int32 NumMips = WeightmapTexture->Source.GetNumMips();
 				WeightmapTextureMipData.Reset();
@@ -3188,7 +3230,7 @@ void FLandscapeEditDataInterface::SetAlphaData(const TSet<ULandscapeLayerInfoObj
 
 			if (Component->SimpleCollisionMipLevel > Component->CollisionMipLevel)
 			{
-				for (int32 WeightmapIdx = 0; WeightmapIdx < Component->WeightmapTextures.Num(); WeightmapIdx++)
+				for (int32 WeightmapIdx = 0; WeightmapIdx < ComponentWeightmapTextures.Num(); WeightmapIdx++)
 				{
 					SimpleCollisionWeightmapMipData.Add((FColor*)TexDataInfos[WeightmapIdx]->GetMipData(Component->SimpleCollisionMipLevel));
 				}
@@ -3204,7 +3246,7 @@ void FLandscapeEditDataInterface::SetAlphaData(const TSet<ULandscapeLayerInfoObj
 
 			// Check if we need to remove weightmap allocations for layers that were completely painted away
 			bool bRemovedLayer = false;
-			for (int32 LayerIdx = 0; LayerIdx < Component->WeightmapLayerAllocations.Num(); LayerIdx++)
+			for (int32 LayerIdx = 0; LayerIdx < ComponentWeightmapLayerAllocations.Num(); LayerIdx++)
 			{
 				if (LayerEditDataAllZero[LayerIdx])
 				{
@@ -3256,16 +3298,19 @@ void FLandscapeEditDataInterface::GetWeightDataTemplFast(ULandscapeLayerInfoObje
 			uint8 WeightmapChannelOffset = 0;
 			TArray<FLandscapeTextureDataInfo*> TexDataInfos; // added for whole weight case...
 
+			TArray<FWeightmapLayerAllocationInfo>& ComponentWeightmapLayerAllocations = Component->GetWeightmapLayerAllocations(true);
+			TArray<UTexture2D*>& ComponentWeightmapTextures = Component->GetWeightmapTextures(true);
+
 			if (LayerInfo != NULL)
 			{
-				for( int32 LayerIdx=0;LayerIdx<Component->WeightmapLayerAllocations.Num();LayerIdx++ )
+				for( int32 LayerIdx=0;LayerIdx<ComponentWeightmapLayerAllocations.Num();LayerIdx++ )
 				{
-					if( Component->WeightmapLayerAllocations[LayerIdx].LayerInfo == LayerInfo )
+					if(ComponentWeightmapLayerAllocations[LayerIdx].LayerInfo == LayerInfo )
 					{
-						WeightmapTexture = Component->WeightmapTextures[Component->WeightmapLayerAllocations[LayerIdx].WeightmapTextureIndex];
+						WeightmapTexture = ComponentWeightmapTextures[ComponentWeightmapLayerAllocations[LayerIdx].WeightmapTextureIndex];
 						TexDataInfo = GetTextureDataInfo(WeightmapTexture);
 						WeightmapTextureData = (uint8*)TexDataInfo->GetMipData(0);
-						WeightmapChannelOffset = ChannelOffsets[Component->WeightmapLayerAllocations[LayerIdx].WeightmapTextureChannel];
+						WeightmapChannelOffset = ChannelOffsets[ComponentWeightmapLayerAllocations[LayerIdx].WeightmapTextureChannel];
 						break;
 					}
 				}
@@ -3273,9 +3318,9 @@ void FLandscapeEditDataInterface::GetWeightDataTemplFast(ULandscapeLayerInfoObje
 			else
 			{
 				// Lock data for all the weightmaps
-				for( int32 WeightmapIdx=0;WeightmapIdx < Component->WeightmapTextures.Num();WeightmapIdx++ )
+				for( int32 WeightmapIdx=0;WeightmapIdx < ComponentWeightmapTextures.Num();WeightmapIdx++ )
 				{
-					TexDataInfos.Add(GetTextureDataInfo(Component->WeightmapTextures[WeightmapIdx]));
+					TexDataInfos.Add(GetTextureDataInfo(ComponentWeightmapTextures[WeightmapIdx]));
 				}
 			}
 
@@ -3335,12 +3380,13 @@ void FLandscapeEditDataInterface::GetWeightDataTemplFast(ULandscapeLayerInfoObje
 							else // Whole weight map case...
 							{
 								StoreData.PreInit(LandscapeInfo->Layers.Num());
-								for( int32 LayerIdx=0;LayerIdx<Component->WeightmapLayerAllocations.Num();LayerIdx++ )
+
+								for( int32 LayerIdx=0;LayerIdx<ComponentWeightmapLayerAllocations.Num();LayerIdx++ )
 								{
-									int32 Idx = Component->WeightmapLayerAllocations[LayerIdx].WeightmapTextureIndex;
-									UTexture2D* ComponentWeightmapTexture = Component->WeightmapTextures[Idx];
+									int32 Idx = ComponentWeightmapLayerAllocations[LayerIdx].WeightmapTextureIndex;
+									UTexture2D* ComponentWeightmapTexture = ComponentWeightmapTextures[Idx];
 									uint8* ComponentWeightmapTextureData = (uint8*)TexDataInfos[Idx]->GetMipData(0);
-									uint8 ComponentWeightmapChannelOffset = ChannelOffsets[Component->WeightmapLayerAllocations[LayerIdx].WeightmapTextureChannel];
+									uint8 ComponentWeightmapChannelOffset = ChannelOffsets[ComponentWeightmapLayerAllocations[LayerIdx].WeightmapTextureChannel];
 
 									// Find the texture data corresponding to this vertex
 									int32 SizeU = ComponentWeightmapTexture->Source.GetSizeX();
@@ -3355,7 +3401,7 @@ void FLandscapeEditDataInterface::GetWeightDataTemplFast(ULandscapeLayerInfoObje
 
 									// Find index in LayerInfos
 									{
-										int32 LayerInfoIdx = LandscapeInfo->GetLayerInfoIndex(Component->WeightmapLayerAllocations[LayerIdx].LayerInfo);
+										int32 LayerInfoIdx = LandscapeInfo->GetLayerInfoIndex(ComponentWeightmapLayerAllocations[LayerIdx].LayerInfo);
 										if (LayerInfoIdx != INDEX_NONE)
 										{
 											StoreData.Store(LandscapeX, LandscapeY, Weight, LayerInfoIdx);
@@ -3378,14 +3424,17 @@ uint8 FLandscapeEditDataInterface::GetWeightMapData(const ULandscapeComponent* C
 	{
 		if (LayerInfo != NULL)
 		{
-			for( int32 LayerIdx=0;LayerIdx<Component->WeightmapLayerAllocations.Num();LayerIdx++ )
+			const TArray<FWeightmapLayerAllocationInfo>& ComponentWeightmapLayerAllocations = Component->GetWeightmapLayerAllocations(true);
+			const TArray<UTexture2D*>& ComponentWeightmapTextures = Component->GetWeightmapTextures(true);
+
+			for( int32 LayerIdx=0;LayerIdx<ComponentWeightmapLayerAllocations.Num();LayerIdx++ )
 			{
-				if( Component->WeightmapLayerAllocations[LayerIdx].LayerInfo == LayerInfo )
+				if(ComponentWeightmapLayerAllocations[LayerIdx].LayerInfo == LayerInfo )
 				{
-					Texture = Component->WeightmapTextures[Component->WeightmapLayerAllocations[LayerIdx].WeightmapTextureIndex];
+					Texture = ComponentWeightmapTextures[ComponentWeightmapLayerAllocations[LayerIdx].WeightmapTextureIndex];
 					FLandscapeTextureDataInfo* TexDataInfo = GetTextureDataInfo(Texture);
 					TextureData = (uint8*)TexDataInfo->GetMipData(0);
-					Offset = ChannelOffsets[Component->WeightmapLayerAllocations[LayerIdx].WeightmapTextureChannel];
+					Offset = ChannelOffsets[ComponentWeightmapLayerAllocations[LayerIdx].WeightmapTextureChannel];
 					break;
 				}
 			}
@@ -3463,16 +3512,20 @@ void FLandscapeEditDataInterface::GetWeightDataTempl(ULandscapeLayerInfoObject*
 		
 			if( Component )
 			{
+				TArray<UTexture2D*>& ComponentWeightmapTextures = Component->GetWeightmapTextures(true);
+
 				if (LayerInfo != NULL)
 				{
-					for( int32 LayerIdx=0;LayerIdx<Component->WeightmapLayerAllocations.Num();LayerIdx++ )
+					TArray<FWeightmapLayerAllocationInfo>& ComponentWeightmapLayerAllocations = Component->GetWeightmapLayerAllocations(true);
+
+					for( int32 LayerIdx=0;LayerIdx<ComponentWeightmapLayerAllocations.Num();LayerIdx++ )
 					{
-						if( Component->WeightmapLayerAllocations[LayerIdx].LayerInfo == LayerInfo )
+						if(ComponentWeightmapLayerAllocations[LayerIdx].LayerInfo == LayerInfo )
 						{
-							WeightmapTexture = Component->WeightmapTextures[Component->WeightmapLayerAllocations[LayerIdx].WeightmapTextureIndex];
+							WeightmapTexture = ComponentWeightmapTextures[ComponentWeightmapLayerAllocations[LayerIdx].WeightmapTextureIndex];
 							TexDataInfo = GetTextureDataInfo(WeightmapTexture);
 							WeightmapTextureData = (uint8*)TexDataInfo->GetMipData(0);
-							WeightmapChannelOffset = ChannelOffsets[Component->WeightmapLayerAllocations[LayerIdx].WeightmapTextureChannel];
+							WeightmapChannelOffset = ChannelOffsets[ComponentWeightmapLayerAllocations[LayerIdx].WeightmapTextureChannel];
 							break;
 						}
 					}
@@ -3480,9 +3533,9 @@ void FLandscapeEditDataInterface::GetWeightDataTempl(ULandscapeLayerInfoObject*
 				else
 				{
 					// Lock data for all the weightmaps
-					for( int32 WeightmapIdx=0;WeightmapIdx < Component->WeightmapTextures.Num();WeightmapIdx++ )
+					for( int32 WeightmapIdx=0;WeightmapIdx < ComponentWeightmapTextures.Num();WeightmapIdx++ )
 					{
-						TexDataInfos.Add(GetTextureDataInfo(Component->WeightmapTextures[WeightmapIdx]));
+						TexDataInfos.Add(GetTextureDataInfo(ComponentWeightmapTextures[WeightmapIdx]));
 					}
 				}
 				ComponentDataExist[ComponentIndexXY] = true;
@@ -3530,14 +3583,17 @@ void FLandscapeEditDataInterface::GetWeightDataTempl(ULandscapeLayerInfoObject*
 							NoBorderX1 = false;
 							if (LayerInfo != NULL)
 							{
-								for( int32 LayerIdx=0;LayerIdx<BorderComponent[0]->WeightmapLayerAllocations.Num();LayerIdx++ )
+								TArray<FWeightmapLayerAllocationInfo>& BorderWeightmapLayerAllocations = BorderComponent[0]->GetWeightmapLayerAllocations(true);
+								TArray<UTexture2D*>& BorderWeightmapTextures = BorderComponent[0]->GetWeightmapTextures(true);
+
+								for( int32 LayerIdx=0;LayerIdx<BorderWeightmapLayerAllocations.Num();LayerIdx++ )
 								{
-									if( BorderComponent[0]->WeightmapLayerAllocations[LayerIdx].LayerInfo == LayerInfo )
+									if(BorderWeightmapLayerAllocations[LayerIdx].LayerInfo == LayerInfo )
 									{
-										NeighborWeightmapTexture[0] = BorderComponent[0]->WeightmapTextures[BorderComponent[0]->WeightmapLayerAllocations[LayerIdx].WeightmapTextureIndex];
+										NeighborWeightmapTexture[0] = BorderWeightmapTextures[BorderWeightmapLayerAllocations[LayerIdx].WeightmapTextureIndex];
 										NeighborTexDataInfo[0] = GetTextureDataInfo(NeighborWeightmapTexture[0]);
 										NeighborWeightmapTextureData[0] = (uint8*)NeighborTexDataInfo[0]->GetMipData(0);
-										NeighborWeightmapChannelOffset[0] = ChannelOffsets[BorderComponent[0]->WeightmapLayerAllocations[LayerIdx].WeightmapTextureChannel];
+										NeighborWeightmapChannelOffset[0] = ChannelOffsets[BorderWeightmapLayerAllocations[LayerIdx].WeightmapTextureChannel];
 										break;
 									}
 								}
@@ -3558,14 +3614,17 @@ void FLandscapeEditDataInterface::GetWeightDataTempl(ULandscapeLayerInfoObject*
 							NoBorderX2 = false;
 							if (LayerInfo != NULL)
 							{
-								for( int32 LayerIdx=0;LayerIdx<BorderComponent[1]->WeightmapLayerAllocations.Num();LayerIdx++ )
+								TArray<FWeightmapLayerAllocationInfo>& BorderWeightmapLayerAllocations = BorderComponent[1]->GetWeightmapLayerAllocations(true);
+								TArray<UTexture2D*>& BorderWeightmapTextures = BorderComponent[1]->GetWeightmapTextures(true);
+
+								for( int32 LayerIdx=0;LayerIdx<BorderWeightmapLayerAllocations.Num();LayerIdx++ )
 								{
-									if( BorderComponent[1]->WeightmapLayerAllocations[LayerIdx].LayerInfo == LayerInfo )
+									if(BorderWeightmapLayerAllocations[LayerIdx].LayerInfo == LayerInfo )
 									{
-										NeighborWeightmapTexture[1] = BorderComponent[1]->WeightmapTextures[BorderComponent[1]->WeightmapLayerAllocations[LayerIdx].WeightmapTextureIndex];
+										NeighborWeightmapTexture[1] = BorderWeightmapTextures[BorderWeightmapLayerAllocations[LayerIdx].WeightmapTextureIndex];
 										NeighborTexDataInfo[1] = GetTextureDataInfo(NeighborWeightmapTexture[1]);
 										NeighborWeightmapTextureData[1] = (uint8*)NeighborTexDataInfo[1]->GetMipData(0);
-										NeighborWeightmapChannelOffset[1] = ChannelOffsets[BorderComponent[1]->WeightmapLayerAllocations[LayerIdx].WeightmapTextureChannel];
+										NeighborWeightmapChannelOffset[1] = ChannelOffsets[BorderWeightmapLayerAllocations[LayerIdx].WeightmapTextureChannel];
 										break;
 									}
 								}
@@ -3585,14 +3644,17 @@ void FLandscapeEditDataInterface::GetWeightDataTempl(ULandscapeLayerInfoObject*
 							NoBorderY1[ComponentIndexXX] = false;
 							if (LayerInfo != NULL)
 							{
-								for( int32 LayerIdx=0;LayerIdx<BorderComponent[2]->WeightmapLayerAllocations.Num();LayerIdx++ )
+								TArray<FWeightmapLayerAllocationInfo>& BorderWeightmapLayerAllocations = BorderComponent[2]->GetWeightmapLayerAllocations(true);
+								TArray<UTexture2D*>& BorderWeightmapTextures = BorderComponent[2]->GetWeightmapTextures(true);
+
+								for( int32 LayerIdx=0;LayerIdx<BorderWeightmapLayerAllocations.Num();LayerIdx++ )
 								{
-									if( BorderComponent[2]->WeightmapLayerAllocations[LayerIdx].LayerInfo == LayerInfo )
+									if(BorderWeightmapLayerAllocations[LayerIdx].LayerInfo == LayerInfo )
 									{
-										NeighborWeightmapTexture[2] = BorderComponent[2]->WeightmapTextures[BorderComponent[2]->WeightmapLayerAllocations[LayerIdx].WeightmapTextureIndex];
+										NeighborWeightmapTexture[2] = BorderWeightmapTextures[BorderWeightmapLayerAllocations[LayerIdx].WeightmapTextureIndex];
 										NeighborTexDataInfo[2] = GetTextureDataInfo(NeighborWeightmapTexture[2]);
 										NeighborWeightmapTextureData[2] = (uint8*)NeighborTexDataInfo[2]->GetMipData(0);
-										NeighborWeightmapChannelOffset[2] = ChannelOffsets[BorderComponent[2]->WeightmapLayerAllocations[LayerIdx].WeightmapTextureChannel];
+										NeighborWeightmapChannelOffset[2] = ChannelOffsets[BorderWeightmapLayerAllocations[LayerIdx].WeightmapTextureChannel];
 										break;
 									}
 								}
@@ -3607,14 +3669,17 @@ void FLandscapeEditDataInterface::GetWeightDataTempl(ULandscapeLayerInfoObject*
 					{
 						if (LayerInfo != NULL)
 						{
-							for( int32 LayerIdx=0;LayerIdx<BorderComponent[2]->WeightmapLayerAllocations.Num();LayerIdx++ )
+							TArray<FWeightmapLayerAllocationInfo>& BorderWeightmapLayerAllocations = BorderComponent[2]->GetWeightmapLayerAllocations(true);
+							TArray<UTexture2D*>& BorderWeightmapTextures = BorderComponent[2]->GetWeightmapTextures(true);
+
+							for( int32 LayerIdx=0;LayerIdx<BorderWeightmapLayerAllocations.Num();LayerIdx++ )
 							{
-								if( BorderComponent[2]->WeightmapLayerAllocations[LayerIdx].LayerInfo == LayerInfo )
+								if(BorderWeightmapLayerAllocations[LayerIdx].LayerInfo == LayerInfo )
 								{
-									NeighborWeightmapTexture[2] = BorderComponent[2]->WeightmapTextures[BorderComponent[2]->WeightmapLayerAllocations[LayerIdx].WeightmapTextureIndex];
+									NeighborWeightmapTexture[2] = BorderWeightmapTextures[BorderWeightmapLayerAllocations[LayerIdx].WeightmapTextureIndex];
 									NeighborTexDataInfo[2] = GetTextureDataInfo(NeighborWeightmapTexture[2]);
 									NeighborWeightmapTextureData[2] = (uint8*)NeighborTexDataInfo[2]->GetMipData(0);
-									NeighborWeightmapChannelOffset[2] = ChannelOffsets[BorderComponent[2]->WeightmapLayerAllocations[LayerIdx].WeightmapTextureChannel];
+									NeighborWeightmapChannelOffset[2] = ChannelOffsets[BorderWeightmapLayerAllocations[LayerIdx].WeightmapTextureChannel];
 									break;
 								}
 							}
@@ -3633,14 +3698,17 @@ void FLandscapeEditDataInterface::GetWeightDataTempl(ULandscapeLayerInfoObject*
 							NoBorderY2[ComponentIndexXX] = false;
 							if (LayerInfo != NULL)
 							{
-								for( int32 LayerIdx=0;LayerIdx<BorderComponent[3]->WeightmapLayerAllocations.Num();LayerIdx++ )
+								TArray<FWeightmapLayerAllocationInfo>& BorderWeightmapLayerAllocations = BorderComponent[3]->GetWeightmapLayerAllocations(true);
+								TArray<UTexture2D*>& BorderWeightmapTextures = BorderComponent[3]->GetWeightmapTextures(true);
+
+								for( int32 LayerIdx=0;LayerIdx<BorderWeightmapLayerAllocations.Num();LayerIdx++ )
 								{
-									if( BorderComponent[3]->WeightmapLayerAllocations[LayerIdx].LayerInfo == LayerInfo )
+									if(BorderWeightmapLayerAllocations[LayerIdx].LayerInfo == LayerInfo )
 									{
-										NeighborWeightmapTexture[3] = BorderComponent[3]->WeightmapTextures[BorderComponent[3]->WeightmapLayerAllocations[LayerIdx].WeightmapTextureIndex];
+										NeighborWeightmapTexture[3] = BorderWeightmapTextures[BorderWeightmapLayerAllocations[LayerIdx].WeightmapTextureIndex];
 										NeighborTexDataInfo[3] = GetTextureDataInfo(NeighborWeightmapTexture[3]);
 										NeighborWeightmapTextureData[3] = (uint8*)NeighborTexDataInfo[3]->GetMipData(0);
-										NeighborWeightmapChannelOffset[3] = ChannelOffsets[BorderComponent[3]->WeightmapLayerAllocations[LayerIdx].WeightmapTextureChannel];
+										NeighborWeightmapChannelOffset[3] = ChannelOffsets[BorderWeightmapLayerAllocations[LayerIdx].WeightmapTextureChannel];
 										break;
 									}
 								}
@@ -3656,14 +3724,17 @@ void FLandscapeEditDataInterface::GetWeightDataTempl(ULandscapeLayerInfoObject*
 					{
 						if (LayerInfo != NULL)
 						{
-							for( int32 LayerIdx=0;LayerIdx<BorderComponent[3]->WeightmapLayerAllocations.Num();LayerIdx++ )
+							TArray<FWeightmapLayerAllocationInfo>& BorderWeightmapLayerAllocations = BorderComponent[3]->GetWeightmapLayerAllocations(true);
+							TArray<UTexture2D*>& BorderWeightmapTextures = BorderComponent[3]->GetWeightmapTextures(true);
+
+							for( int32 LayerIdx=0;LayerIdx<BorderWeightmapLayerAllocations.Num();LayerIdx++ )
 							{
-								if( BorderComponent[3]->WeightmapLayerAllocations[LayerIdx].LayerInfo == LayerInfo )
+								if(BorderWeightmapLayerAllocations[LayerIdx].LayerInfo == LayerInfo )
 								{
-									NeighborWeightmapTexture[3] = BorderComponent[3]->WeightmapTextures[BorderComponent[3]->WeightmapLayerAllocations[LayerIdx].WeightmapTextureIndex];
+									NeighborWeightmapTexture[3] = BorderWeightmapTextures[BorderWeightmapLayerAllocations[LayerIdx].WeightmapTextureIndex];
 									NeighborTexDataInfo[3] = GetTextureDataInfo(NeighborWeightmapTexture[3]);
 									NeighborWeightmapTextureData[3] = (uint8*)NeighborTexDataInfo[3]->GetMipData(0);
-									NeighborWeightmapChannelOffset[3] = ChannelOffsets[BorderComponent[3]->WeightmapLayerAllocations[LayerIdx].WeightmapTextureChannel];
+									NeighborWeightmapChannelOffset[3] = ChannelOffsets[BorderWeightmapLayerAllocations[LayerIdx].WeightmapTextureChannel];
 									break;
 								}
 							}
@@ -3925,12 +3996,16 @@ void FLandscapeEditDataInterface::GetWeightDataTempl(ULandscapeLayerInfoObject*
 							else // Whole weight map case... no interpolation now...
 							{
 								StoreData.PreInit(LandscapeInfo->Layers.Num());
-								for( int32 LayerIdx=0;LayerIdx<Component->WeightmapLayerAllocations.Num();LayerIdx++ )
+
+								TArray<FWeightmapLayerAllocationInfo>& ComponentWeightmapLayerAllocations = Component->GetWeightmapLayerAllocations(true);
+								TArray<UTexture2D*>& ComponentWeightmapTextures = Component->GetWeightmapTextures(true);
+
+								for( int32 LayerIdx=0;LayerIdx<ComponentWeightmapLayerAllocations.Num();LayerIdx++ )
 								{
-									int32 Idx = Component->WeightmapLayerAllocations[LayerIdx].WeightmapTextureIndex;
-									UTexture2D* ComponentWeightmapTexture = Component->WeightmapTextures[Idx];
+									int32 Idx = ComponentWeightmapLayerAllocations[LayerIdx].WeightmapTextureIndex;
+									UTexture2D* ComponentWeightmapTexture = ComponentWeightmapTextures[Idx];
 									uint8* ComponentWeightmapTextureData = (uint8*)TexDataInfos[Idx]->GetMipData(0);
-									uint8 ComponentWeightmapChannelOffset = ChannelOffsets[Component->WeightmapLayerAllocations[LayerIdx].WeightmapTextureChannel];
+									uint8 ComponentWeightmapChannelOffset = ChannelOffsets[ComponentWeightmapLayerAllocations[LayerIdx].WeightmapTextureChannel];
 
 									// Find the texture data corresponding to this vertex
 									int32 SizeU = ComponentWeightmapTexture->Source.GetSizeX();
@@ -3945,7 +4020,7 @@ void FLandscapeEditDataInterface::GetWeightDataTempl(ULandscapeLayerInfoObject*
 
 									// Find index in LayerInfos
 									{
-										int32 LayerInfoIdx = LandscapeInfo->GetLayerInfoIndex(Component->WeightmapLayerAllocations[LayerIdx].LayerInfo);
+										int32 LayerInfoIdx = LandscapeInfo->GetLayerInfoIndex(ComponentWeightmapLayerAllocations[LayerIdx].LayerInfo);
 										if (LayerInfoIdx != INDEX_NONE)
 										{
 											StoreData.Store(LandscapeX, LandscapeY, Weight, LayerInfoIdx);
@@ -4446,6 +4521,9 @@ void FLandscapeEditDataInterface::SetXYOffsetDataTempl(int32 X1, int32 Y1, int32
 					XYOffsetTexture = Component->GetLandscapeProxy()->CreateLandscapeTexture(WeightmapSize, WeightmapSize, TEXTUREGROUP_Terrain_Weightmap, TSF_BGRA8);
 					// Alloc dummy mips
 					ULandscapeComponent::CreateEmptyTextureMips(XYOffsetTexture, true);
+
+					// TODO: we flag somewhere that we need to use XYOffset Factory, so we can recompute shaders based on this change
+
 					XYOffsetTexture->PostEditChange();
 
 					//FlushRenderingCommands();
@@ -5277,4 +5355,4 @@ FLandscapeTextureDataInfo::~FLandscapeTextureDataInfo()
 	Texture->ClearFlags(RF_Transactional);
 }
 
-#endif // WITH_EDITOR
+#endif // WITH_EDITOR
\ No newline at end of file
