When SetParentEditorOnly is called we no longer recompute shader in most case as we will call UpdateStaticPermutation that will do so.

diff --git a/Engine/Source/Runtime/Landscape/Private/LandscapeEditProcedural.cpp b/Engine/Source/Runtime/Landscape/Private/LandscapeEditProcedural.cpp
--- a/Engine/Source/Runtime/Landscape/Private/LandscapeEditProcedural.cpp
+++ b/Engine/Source/Runtime/Landscape/Private/LandscapeEditProcedural.cpp
@@ -12,90 +12,27 @@
 #include "LandscapeComponent.h"
 #include "LandscapeLayerInfoObject.h"
 #include "LandscapeDataAccess.h"
-#include "LandscapeRender.h"
-#include "LandscapeRenderMobile.h"
+#include "LandscapePrivate.h"
 
-#if WITH_EDITOR
-#include "LandscapeEditorModule.h"
-#include "ComponentRecreateRenderStateContext.h"
-#include "Interfaces/ITargetPlatform.h"
-#include "Engine/TextureRenderTarget2D.h"
-#include "ScopedTransaction.h"
-#include "Editor.h"
-#include "Settings/EditorExperimentalSettings.h"
-#endif
-
 #include "Shader.h"
 #include "GlobalShader.h"
-#include "RendererInterface.h"
-#include "PipelineStateCache.h"
-#include "MaterialShaderType.h"
-#include "EngineModule.h"
 #include "ShaderParameterUtils.h"
+#include "Engine/TextureRenderTarget2D.h"
+
+#if WITH_EDITOR
+#include "LandscapeEditorModule.h"
+#include "ComponentRecreateRenderStateContext.h"
+#include "Settings/EditorExperimentalSettings.h"
 #include "LandscapeBPCustomBrush.h"
+#include "Materials/MaterialInstanceConstant.h"
+#include "LandscapeMaterialInstanceConstant.h"
+#include "Materials/MaterialExpressionLandscapeVisibilityMask.h"
+#include "ShaderCompiler.h"
+#endif
 
 #define LOCTEXT_NAMESPACE "Landscape"
 
-void ALandscapeProxy::BeginDestroy()
-{
-	Super::BeginDestroy();
+ENGINE_API extern bool GDisableAutomaticTextureMaterialUpdateDependencies;
 
-#if WITH_EDITORONLY_DATA
-	if (GetMutableDefault<UEditorExperimentalSettings>()->bProceduralLandscape)
-	{
-		for (auto& ItPair : RenderDataPerHeightmap)
-		{
-			FRenderDataPerHeightmap& HeightmapRenderData = ItPair.Value;
-
-			if (HeightmapRenderData.HeightmapsCPUReadBack != nullptr)
-			{
-				BeginReleaseResource(HeightmapRenderData.HeightmapsCPUReadBack);
-			}
-		}
-
-		ReleaseResourceFence.BeginFence();
-	}
-#endif
-}
-
-bool ALandscapeProxy::IsReadyForFinishDestroy()
-{
-	bool bReadyForFinishDestroy = Super::IsReadyForFinishDestroy();
-
-#if WITH_EDITORONLY_DATA
-	if (GetMutableDefault<UEditorExperimentalSettings>()->bProceduralLandscape)
-	{
-		if (bReadyForFinishDestroy)
-		{
-			bReadyForFinishDestroy = ReleaseResourceFence.IsFenceComplete();
-		}
-	}
-#endif
-
-	return bReadyForFinishDestroy;
-}
-
-void ALandscapeProxy::FinishDestroy()
-{
-	Super::FinishDestroy();
-
-#if WITH_EDITORONLY_DATA
-	if (GetMutableDefault<UEditorExperimentalSettings>()->bProceduralLandscape)
-	{
-		check(ReleaseResourceFence.IsFenceComplete());
-
-		for (auto& ItPair : RenderDataPerHeightmap)
-		{
-			FRenderDataPerHeightmap& HeightmapRenderData = ItPair.Value;
-
-			delete HeightmapRenderData.HeightmapsCPUReadBack;
-			HeightmapRenderData.HeightmapsCPUReadBack = nullptr;
-		}
-	}
-#endif
-}
-
-#if WITH_EDITOR
-
 static TAutoConsoleVariable<int32> CVarOutputProceduralDebugDrawCallName(
 	TEXT("landscape.OutputProceduralDebugDrawCallName"),
@@ -107,6 +44,16 @@
 	0,
 	TEXT("This will output the content of render target. This is used for debugging only."));
 
+static TAutoConsoleVariable<int32> CVarOutputProceduralWeightmapsRTContent(
+	TEXT("landscape.OutputProceduralWeightmapsRTContent"),
+	0,
+	TEXT("This will output the content of render target used for weightmap. This is used for debugging only."));
+
+DECLARE_GPU_STAT_NAMED(LandscapeProceduralRender, TEXT("Landscape Procedural Render"));
+DECLARE_GPU_STAT_NAMED(LandscapeProceduralCopy, TEXT("Landscape Procedural Copy"));
+
+// Vertex format and vertex buffer
+
 struct FLandscapeProceduralVertex
 {
 	FVector2D Position;
@@ -120,7 +67,6 @@
 	FLandscapeProceduralVertex V2;
 };
 
-/** The filter vertex declaration resource type. */
 class FLandscapeProceduralVertexDeclaration : public FRenderResource
 {
 public:
@@ -144,7 +90,6 @@
 	}
 };
 
-
 class FLandscapeProceduralVertexBuffer : public FVertexBuffer
 {
 public:
@@ -176,7 +121,7 @@
 	TArray<FLandscapeProceduralTriangle> TriangleList;
 };
 
+// Custom Pixel and Vertex shaders
 
-
 class FLandscapeProceduralVS : public FGlobalShader
 {
@@ -186,7 +131,7 @@
 
 	static bool ShouldCompilePermutation(const FGlobalShaderPermutationParameters& Parameters)
 	{
-		return IsFeatureLevelSupported(Parameters.Platform, ERHIFeatureLevel::SM4) && !IsConsolePlatform(Parameters.Platform);
+		return IsFeatureLevelSupported(Parameters.Platform, ERHIFeatureLevel::SM5) && !IsConsolePlatform(Parameters.Platform);
 	}
 
 	static void ModifyCompilationEnvironment(const FGlobalShaderPermutationParameters& Parameters, FShaderCompilerEnvironment& OutEnvironment)
@@ -227,13 +172,13 @@
 		, ReadHeightmap2(nullptr)
 		, HeightmapSize(0, 0)
 		, ApplyLayerModifiers(false)
-		, LayerWeight(1.0f)
+		, LayerAlpha(1.0f)
 		, LayerVisible(true)
 		, OutputAsDelta(false)
 		, GenerateNormals(false)
 		, GridSize(0.0f, 0.0f, 0.0f)
-		, CurrentMipHeightmapSize(0, 0)
-		, ParentMipHeightmapSize(0, 0)
+		, CurrentMipSize(0, 0)
+		, ParentMipSize(0, 0)
 		, CurrentMipComponentVertexCount(0)
 	{}
 
@@ -241,13 +186,13 @@
 	UTexture* ReadHeightmap2;
 	FIntPoint HeightmapSize;
 	bool ApplyLayerModifiers;
-	float LayerWeight;
+	float LayerAlpha;
 	bool LayerVisible;
 	bool OutputAsDelta;
 	bool GenerateNormals;
 	FVector GridSize;
-	FIntPoint CurrentMipHeightmapSize;
-	FIntPoint ParentMipHeightmapSize;
+	FIntPoint CurrentMipSize;
+	FIntPoint ParentMipSize;
 	int32 CurrentMipComponentVertexCount;
 };
 
@@ -258,7 +203,7 @@
 
 	static bool ShouldCompilePermutation(const FGlobalShaderPermutationParameters& Parameters)
 	{
-		return IsFeatureLevelSupported(Parameters.Platform, ERHIFeatureLevel::SM4) && !IsConsolePlatform(Parameters.Platform);
+		return IsFeatureLevelSupported(Parameters.Platform, ERHIFeatureLevel::SM5) && !IsConsolePlatform(Parameters.Platform);
 	}
 
 	static void ModifyCompilationEnvironment(const FGlobalShaderPermutationParameters& Parameters, FShaderCompilerEnvironment& OutEnvironment)
@@ -268,13 +213,14 @@
 	FLandscapeHeightmapProceduralPS(const ShaderMetaType::CompiledShaderInitializerType& Initializer)
 		: FGlobalShader(Initializer)
 	{
-		ReadHeightmapTexture1Param.Bind(Initializer.ParameterMap, TEXT("ReadHeightmapTexture1"));
-		ReadHeightmapTexture2Param.Bind(Initializer.ParameterMap, TEXT("ReadHeightmapTexture2"));
-		ReadHeightmapTexture1SamplerParam.Bind(Initializer.ParameterMap, TEXT("ReadHeightmapTexture1Sampler"));
-		ReadHeightmapTexture2SamplerParam.Bind(Initializer.ParameterMap, TEXT("ReadHeightmapTexture2Sampler"));
+		ReadTexture1Param.Bind(Initializer.ParameterMap, TEXT("ReadTexture1"));
+		ReadTexture2Param.Bind(Initializer.ParameterMap, TEXT("ReadTexture2"));
+		ReadTexture1SamplerParam.Bind(Initializer.ParameterMap, TEXT("ReadTexture1Sampler"));
+		ReadTexture2SamplerParam.Bind(Initializer.ParameterMap, TEXT("ReadTexture2Sampler"));
+
 		LayerInfoParam.Bind(Initializer.ParameterMap, TEXT("LayerInfo"));
 		OutputConfigParam.Bind(Initializer.ParameterMap, TEXT("OutputConfig"));
-		TextureSizeParam.Bind(Initializer.ParameterMap, TEXT("HeightmapTextureSize"));
+		TextureSizeParam.Bind(Initializer.ParameterMap, TEXT("TextureSize"));
 		LandscapeGridScaleParam.Bind(Initializer.ParameterMap, TEXT("LandscapeGridScale"));
 		ComponentVertexCountParam.Bind(Initializer.ParameterMap, TEXT("CurrentMipComponentVertexCount"));
 	}
@@ -284,14 +230,10 @@
 
 	void SetParameters(FRHICommandList& RHICmdList, const FLandscapeHeightmapProceduralShaderParameters& InParams)
 	{
-		SetTextureParameter(RHICmdList, GetPixelShader(), ReadHeightmapTexture1Param, ReadHeightmapTexture1SamplerParam, TStaticSamplerState<SF_Point, AM_Clamp, AM_Clamp, AM_Clamp>::GetRHI(), InParams.ReadHeightmap1->Resource->TextureRHI);
+		SetTextureParameter(RHICmdList, GetPixelShader(), ReadTexture1Param, ReadTexture1SamplerParam, TStaticSamplerState<SF_Point, AM_Clamp, AM_Clamp, AM_Clamp>::GetRHI(), InParams.ReadHeightmap1->Resource->TextureRHI);
+		SetTextureParameter(RHICmdList, GetPixelShader(), ReadTexture2Param, ReadTexture2SamplerParam, TStaticSamplerState<SF_Point, AM_Clamp, AM_Clamp, AM_Clamp>::GetRHI(), InParams.ReadHeightmap2 != nullptr ? InParams.ReadHeightmap2->Resource->TextureRHI : GWhiteTexture->TextureRHI);
 
-		if (InParams.ReadHeightmap2 != nullptr)
-		{
-			SetTextureParameter(RHICmdList, GetPixelShader(), ReadHeightmapTexture2Param, ReadHeightmapTexture2SamplerParam, TStaticSamplerState<SF_Point, AM_Clamp, AM_Clamp, AM_Clamp>::GetRHI(), InParams.ReadHeightmap2->Resource->TextureRHI);
-		}
-
-		FVector2D LayerInfo(InParams.LayerWeight, InParams.LayerVisible ? 1.0f : 0.0f);
+		FVector2D LayerInfo(InParams.LayerAlpha, InParams.LayerVisible ? 1.0f : 0.0f);
 		FVector4 OutputConfig(InParams.ApplyLayerModifiers ? 1.0f : 0.0f, InParams.OutputAsDelta ? 1.0f : 0.0f, InParams.ReadHeightmap2 ? 1.0f : 0.0f, InParams.GenerateNormals ? 1.0f : 0.0f);
 		FVector2D TextureSize(InParams.HeightmapSize.X, InParams.HeightmapSize.Y);
 
@@ -305,10 +247,10 @@
 	virtual bool Serialize(FArchive& Ar) override
 	{
 		bool bShaderHasOutdatedParameters = FShader::Serialize(Ar);
-		Ar << ReadHeightmapTexture1Param;
-		Ar << ReadHeightmapTexture2Param;
-		Ar << ReadHeightmapTexture1SamplerParam;
-		Ar << ReadHeightmapTexture2SamplerParam;
+		Ar << ReadTexture1Param;
+		Ar << ReadTexture2Param;
+		Ar << ReadTexture1SamplerParam;
+		Ar << ReadTexture2SamplerParam;
 		Ar << LayerInfoParam;
 		Ar << OutputConfigParam;
 		Ar << TextureSizeParam;
@@ -318,10 +260,10 @@
 	}
 
 private:
-	FShaderResourceParameter ReadHeightmapTexture1Param;
-	FShaderResourceParameter ReadHeightmapTexture2Param;
-	FShaderResourceParameter ReadHeightmapTexture1SamplerParam;
-	FShaderResourceParameter ReadHeightmapTexture2SamplerParam;
+	FShaderResourceParameter ReadTexture1Param;
+	FShaderResourceParameter ReadTexture2Param;
+	FShaderResourceParameter ReadTexture1SamplerParam;
+	FShaderResourceParameter ReadTexture2SamplerParam;
 	FShaderParameter LayerInfoParam;
 	FShaderParameter OutputConfigParam;
 	FShaderParameter TextureSizeParam;
@@ -329,7 +271,7 @@
 	FShaderParameter ComponentVertexCountParam;
 };
 
-IMPLEMENT_GLOBAL_SHADER(FLandscapeHeightmapProceduralPS, "/Engine/Private/LandscapeProceduralPS.usf", "PSMain", SF_Pixel);
+IMPLEMENT_GLOBAL_SHADER(FLandscapeHeightmapProceduralPS, "/Engine/Private/LandscapeProceduralPS.usf", "PSHeightmapMain", SF_Pixel);
 
 class FLandscapeHeightmapMipsProceduralPS : public FGlobalShader
 {
@@ -338,7 +280,7 @@
 
 	static bool ShouldCompilePermutation(const FGlobalShaderPermutationParameters& Parameters)
 	{
-		return IsFeatureLevelSupported(Parameters.Platform, ERHIFeatureLevel::SM4) && !IsConsolePlatform(Parameters.Platform);
+		return IsFeatureLevelSupported(Parameters.Platform, ERHIFeatureLevel::SM5) && !IsConsolePlatform(Parameters.Platform);
 	}
 
 	static void ModifyCompilationEnvironment(const FGlobalShaderPermutationParameters& Parameters, FShaderCompilerEnvironment& OutEnvironment)
@@ -348,10 +290,10 @@
 	FLandscapeHeightmapMipsProceduralPS(const ShaderMetaType::CompiledShaderInitializerType& Initializer)
 		: FGlobalShader(Initializer)
 	{
-		ReadHeightmapTexture1Param.Bind(Initializer.ParameterMap, TEXT("ReadHeightmapTexture1"));
-		ReadHeightmapTexture1SamplerParam.Bind(Initializer.ParameterMap, TEXT("ReadHeightmapTexture1Sampler"));
-		CurrentMipHeightmapSizeParam.Bind(Initializer.ParameterMap, TEXT("CurrentMipTextureSize"));
-		ParentMipHeightmapSizeParam.Bind(Initializer.ParameterMap, TEXT("ParentMipTextureSize"));
+		ReadTexture1Param.Bind(Initializer.ParameterMap, TEXT("ReadTexture1"));
+		ReadTexture1SamplerParam.Bind(Initializer.ParameterMap, TEXT("ReadTexture1Sampler"));
+		CurrentMipSizeParam.Bind(Initializer.ParameterMap, TEXT("CurrentMipTextureSize"));
+		ParentMipSizeParam.Bind(Initializer.ParameterMap, TEXT("ParentMipTextureSize"));
 		CurrentMipComponentVertexCountParam.Bind(Initializer.ParameterMap, TEXT("CurrentMipComponentVertexCount"));
 	}
 
@@ -360,129 +302,803 @@
 
 	void SetParameters(FRHICommandList& RHICmdList, const FLandscapeHeightmapProceduralShaderParameters& InParams)
 	{
-		SetTextureParameter(RHICmdList, GetPixelShader(), ReadHeightmapTexture1Param, ReadHeightmapTexture1SamplerParam, TStaticSamplerState<SF_Point, AM_Clamp, AM_Clamp, AM_Clamp>::GetRHI(), InParams.ReadHeightmap1->Resource->TextureRHI);
+		SetTextureParameter(RHICmdList, GetPixelShader(), ReadTexture1Param, ReadTexture1SamplerParam, TStaticSamplerState<SF_Point, AM_Clamp, AM_Clamp, AM_Clamp>::GetRHI(), InParams.ReadHeightmap1->Resource->TextureRHI);
 
-		SetShaderValue(RHICmdList, GetPixelShader(), CurrentMipHeightmapSizeParam, FVector2D(InParams.CurrentMipHeightmapSize.X, InParams.CurrentMipHeightmapSize.Y));
-		SetShaderValue(RHICmdList, GetPixelShader(), ParentMipHeightmapSizeParam, FVector2D(InParams.ParentMipHeightmapSize.X, InParams.ParentMipHeightmapSize.Y));
+		SetShaderValue(RHICmdList, GetPixelShader(), CurrentMipSizeParam, FVector2D(InParams.CurrentMipSize.X, InParams.CurrentMipSize.Y));
+		SetShaderValue(RHICmdList, GetPixelShader(), ParentMipSizeParam, FVector2D(InParams.ParentMipSize.X, InParams.ParentMipSize.Y));
 		SetShaderValue(RHICmdList, GetPixelShader(), CurrentMipComponentVertexCountParam, (float)InParams.CurrentMipComponentVertexCount);
 	}
 
 	virtual bool Serialize(FArchive& Ar) override
 	{
 		bool bShaderHasOutdatedParameters = FShader::Serialize(Ar);
-		Ar << ReadHeightmapTexture1Param;
-		Ar << ReadHeightmapTexture1SamplerParam;
-		Ar << CurrentMipHeightmapSizeParam;
-		Ar << ParentMipHeightmapSizeParam;
+		Ar << ReadTexture1Param;
+		Ar << ReadTexture1SamplerParam;
+		Ar << CurrentMipSizeParam;
+		Ar << ParentMipSizeParam;
 		Ar << CurrentMipComponentVertexCountParam;
 
 		return bShaderHasOutdatedParameters;
 	}
 
 private:
-	FShaderResourceParameter ReadHeightmapTexture1Param;
-	FShaderResourceParameter ReadHeightmapTexture1SamplerParam;
-	FShaderParameter CurrentMipHeightmapSizeParam;
-	FShaderParameter ParentMipHeightmapSizeParam;
+	FShaderResourceParameter ReadTexture1Param;
+	FShaderResourceParameter ReadTexture1SamplerParam;
+	FShaderParameter CurrentMipSizeParam;
+	FShaderParameter ParentMipSizeParam;
 	FShaderParameter CurrentMipComponentVertexCountParam;
 };
 
-IMPLEMENT_GLOBAL_SHADER(FLandscapeHeightmapMipsProceduralPS, "/Engine/Private/LandscapeProceduralPS.usf", "PSMainMips", SF_Pixel);
+IMPLEMENT_GLOBAL_SHADER(FLandscapeHeightmapMipsProceduralPS, "/Engine/Private/LandscapeProceduralPS.usf", "PSHeightmapMainMips", SF_Pixel);
 
-/** The filter vertex declaration resource type. */
+struct FLandscapeWeightmapProceduralShaderParameters
+{
+	FLandscapeWeightmapProceduralShaderParameters()
+		: ReadWeightmap1(nullptr)
+		, ReadWeightmap2(nullptr)
+		, ApplyLayerModifiers(false)
+		, LayerAlpha(1.0f)
+		, LayerVisible(true)
+		, OutputAsDelta(false)
+		, OutputAsSubstractive(false)
+		, OutputAsNormalized(false)
+		, CurrentMipSize(0, 0)
+		, ParentMipSize(0, 0)
+		, CurrentMipComponentVertexCount(0)
+	{}
 
+	UTexture* ReadWeightmap1;
+	UTexture* ReadWeightmap2;
+	bool ApplyLayerModifiers;
+	float LayerAlpha;
+	bool LayerVisible;
+	bool OutputAsDelta;
+	bool OutputAsSubstractive;
+	bool OutputAsNormalized;
+	FIntPoint CurrentMipSize;
+	FIntPoint ParentMipSize;
+	int32 CurrentMipComponentVertexCount;
+};
 
-DECLARE_GPU_STAT_NAMED(LandscapeProceduralRender, TEXT("Landscape Procedural Render"));
+class FLandscapeWeightmapProceduralPS : public FGlobalShader
+{
+	DECLARE_GLOBAL_SHADER(FLandscapeWeightmapProceduralPS);
+public:
 
-class FLandscapeProceduralCopyResource_RenderThread
+	static bool ShouldCompilePermutation(const FGlobalShaderPermutationParameters& Parameters)
+	{
+		return IsFeatureLevelSupported(Parameters.Platform, ERHIFeatureLevel::SM5) && !IsConsolePlatform(Parameters.Platform);
+	}
+
+	static void ModifyCompilationEnvironment(const FGlobalShaderPermutationParameters& Parameters, FShaderCompilerEnvironment& OutEnvironment)
+	{
+	}
+
+	FLandscapeWeightmapProceduralPS(const ShaderMetaType::CompiledShaderInitializerType& Initializer)
+		: FGlobalShader(Initializer)
+	{
+		ReadTexture1Param.Bind(Initializer.ParameterMap, TEXT("ReadTexture1"));
+		ReadTexture2Param.Bind(Initializer.ParameterMap, TEXT("ReadTexture2"));
+		ReadTexture1SamplerParam.Bind(Initializer.ParameterMap, TEXT("ReadTexture1Sampler"));
+		ReadTexture2SamplerParam.Bind(Initializer.ParameterMap, TEXT("ReadTexture2Sampler"));
+		LayerInfoParam.Bind(Initializer.ParameterMap, TEXT("LayerInfo"));
+		OutputConfigParam.Bind(Initializer.ParameterMap, TEXT("OutputConfig"));
+		ComponentVertexCountParam.Bind(Initializer.ParameterMap, TEXT("CurrentMipComponentVertexCount"));
+	}
+
+	FLandscapeWeightmapProceduralPS()
+	{}
+
+	void SetParameters(FRHICommandList& RHICmdList, const FLandscapeWeightmapProceduralShaderParameters& InParams)
+	{
+		SetTextureParameter(RHICmdList, GetPixelShader(), ReadTexture1Param, ReadTexture1SamplerParam, TStaticSamplerState<SF_Point, AM_Clamp, AM_Clamp, AM_Clamp>::GetRHI(), InParams.ReadWeightmap1->Resource->TextureRHI);
+		SetTextureParameter(RHICmdList, GetPixelShader(), ReadTexture2Param, ReadTexture2SamplerParam, TStaticSamplerState<SF_Point, AM_Clamp, AM_Clamp, AM_Clamp>::GetRHI(), InParams.ReadWeightmap2 != nullptr ? InParams.ReadWeightmap2->Resource->TextureRHI : GWhiteTexture->TextureRHI);
+
+		FVector2D LayerInfo(InParams.LayerAlpha, InParams.LayerVisible ? 1.0f : 0.0f);
+		FVector4 OutputConfig(InParams.ApplyLayerModifiers ? 1.0f : 0.0f, InParams.OutputAsSubstractive ? 1.0f : 0.0f, InParams.ReadWeightmap2 != nullptr ? 1.0f : 0.0f, InParams.OutputAsNormalized ? 1.0f : 0.0f);
+
+		SetShaderValue(RHICmdList, GetPixelShader(), LayerInfoParam, LayerInfo);
+		SetShaderValue(RHICmdList, GetPixelShader(), OutputConfigParam, OutputConfig);
+		SetShaderValue(RHICmdList, GetPixelShader(), ComponentVertexCountParam, (float)InParams.CurrentMipComponentVertexCount);
+	}
+
+	virtual bool Serialize(FArchive& Ar) override
+	{
+		bool bShaderHasOutdatedParameters = FShader::Serialize(Ar);
+		Ar << ReadTexture1Param;
+		Ar << ReadTexture2Param;
+		Ar << ReadTexture1SamplerParam;
+		Ar << ReadTexture2SamplerParam;
+		Ar << LayerInfoParam;
+		Ar << OutputConfigParam;
+		Ar << ComponentVertexCountParam;
+		return bShaderHasOutdatedParameters;
+	}
+
+private:
+	FShaderResourceParameter ReadTexture1Param;
+	FShaderResourceParameter ReadTexture2Param;
+	FShaderResourceParameter ReadTexture1SamplerParam;
+	FShaderResourceParameter ReadTexture2SamplerParam;
+	FShaderParameter LayerInfoParam;
+	FShaderParameter OutputConfigParam;
+	FShaderParameter ComponentVertexCountParam;
+};
+
+IMPLEMENT_GLOBAL_SHADER(FLandscapeWeightmapProceduralPS, "/Engine/Private/LandscapeProceduralPS.usf", "PSWeightmapMain", SF_Pixel);
+
+class FLandscapeWeightmapMipsProceduralPS : public FGlobalShader
 {
+	DECLARE_GLOBAL_SHADER(FLandscapeWeightmapMipsProceduralPS);
+public:
+
+	static bool ShouldCompilePermutation(const FGlobalShaderPermutationParameters& Parameters)
+	{
+		return IsFeatureLevelSupported(Parameters.Platform, ERHIFeatureLevel::SM5) && !IsConsolePlatform(Parameters.Platform);
+	}
+
+	static void ModifyCompilationEnvironment(const FGlobalShaderPermutationParameters& Parameters, FShaderCompilerEnvironment& OutEnvironment)
+	{
+	}
+
+	FLandscapeWeightmapMipsProceduralPS(const ShaderMetaType::CompiledShaderInitializerType& Initializer)
+		: FGlobalShader(Initializer)
+	{
+		ReadTexture1Param.Bind(Initializer.ParameterMap, TEXT("ReadTexture1"));
+		ReadTexture1SamplerParam.Bind(Initializer.ParameterMap, TEXT("ReadTexture1Sampler"));
+		CurrentMipSizeParam.Bind(Initializer.ParameterMap, TEXT("CurrentMipTextureSize"));
+		ParentMipSizeParam.Bind(Initializer.ParameterMap, TEXT("ParentMipTextureSize"));
+		CurrentMipComponentVertexCountParam.Bind(Initializer.ParameterMap, TEXT("CurrentMipComponentVertexCount"));
+	}
+
+	FLandscapeWeightmapMipsProceduralPS()
+	{}
+
+	void SetParameters(FRHICommandList& RHICmdList, const FLandscapeWeightmapProceduralShaderParameters& InParams)
+	{
+		SetTextureParameter(RHICmdList, GetPixelShader(), ReadTexture1Param, ReadTexture1SamplerParam, TStaticSamplerState<SF_Point, AM_Clamp, AM_Clamp, AM_Clamp>::GetRHI(), InParams.ReadWeightmap1->Resource->TextureRHI);
+
+		SetShaderValue(RHICmdList, GetPixelShader(), CurrentMipSizeParam, FVector2D(InParams.CurrentMipSize.X, InParams.CurrentMipSize.Y));
+		SetShaderValue(RHICmdList, GetPixelShader(), ParentMipSizeParam, FVector2D(InParams.ParentMipSize.X, InParams.ParentMipSize.Y));
+		SetShaderValue(RHICmdList, GetPixelShader(), CurrentMipComponentVertexCountParam, (float)InParams.CurrentMipComponentVertexCount);
+	}
+
+	virtual bool Serialize(FArchive& Ar) override
+	{
+		bool bShaderHasOutdatedParameters = FShader::Serialize(Ar);
+		Ar << ReadTexture1Param;
+		Ar << ReadTexture1SamplerParam;
+		Ar << CurrentMipSizeParam;
+		Ar << ParentMipSizeParam;
+		Ar << CurrentMipComponentVertexCountParam;
+
+		return bShaderHasOutdatedParameters;
+	}
+
+private:
+	FShaderResourceParameter ReadTexture1Param;
+	FShaderResourceParameter ReadTexture1SamplerParam;
+	FShaderParameter CurrentMipSizeParam;
+	FShaderParameter ParentMipSizeParam;
+	FShaderParameter CurrentMipComponentVertexCountParam;
+};
+
+IMPLEMENT_GLOBAL_SHADER(FLandscapeWeightmapMipsProceduralPS, "/Engine/Private/LandscapeProceduralPS.usf", "PSWeightmapMainMips", SF_Pixel);
+
+// Custom Resources
+
+class FLandscapeTexture2DResource : public FTextureResource
+{
+public:
+	FLandscapeTexture2DResource(uint32 InSizeX, uint32 InSizeY, EPixelFormat InFormat, uint32 InNumMips, bool InNeedUAV)
+		: SizeX(InSizeX)
+		, SizeY(InSizeY)
+		, Format(InFormat)
+		, NumMips(InNumMips)
+		, CreateUAV(InNeedUAV)
+	{}
+
+	virtual uint32 GetSizeX() const override
+	{
+		return SizeX;
+	}
+
+	virtual uint32 GetSizeY() const override
+	{
+		return SizeY;
+	}
+
+	/** Called when the resource is initialized. This is only called by the rendering thread. */
+	virtual void InitRHI() override
+	{
+		FTextureResource::InitRHI();
+
+		FRHIResourceCreateInfo CreateInfo;
+		uint32 Flags = TexCreate_NoTiling | TexCreate_OfflineProcessed;
+
+		if (CreateUAV)
+		{
+			Flags |= TexCreate_UAV;
+		}
+
+		TextureRHI = RHICreateTexture2D(SizeX, SizeY, Format, NumMips, 1, Flags, CreateInfo);
+
+		if (CreateUAV)
+		{
+			TextureUAV = RHICreateUnorderedAccessView(TextureRHI, 0);
+		}
+	}
+
+	FUnorderedAccessViewRHIRef TextureUAV;
+
+private:
+	uint32 SizeX;
+	uint32 SizeY;
+	EPixelFormat Format;
+	uint32 NumMips;
+	bool CreateUAV;
+};
+
+class FLandscapeTexture2DArrayResource : public FTextureResource
+{
+public:
+	FLandscapeTexture2DArrayResource(uint32 InSizeX, uint32 InSizeY, uint32 InSizeZ, EPixelFormat InFormat, uint32 InNumMips, bool InNeedUAV)
+		: SizeX(InSizeX)
+		, SizeY(InSizeY)
+		, SizeZ(InSizeZ)
+		, Format(InFormat)
+		, NumMips(InNumMips)
+		, CreateUAV(InNeedUAV)
+	{}
+
+	virtual uint32 GetSizeX() const override
+	{
+		return SizeX;
+	}
+
+	virtual uint32 GetSizeY() const override
+	{
+		return SizeY;
+	}
+
+	virtual uint32 GetSizeZ() const
+	{
+		return SizeZ;
+	}
+
+	/** Called when the resource is initialized. This is only called by the rendering thread. */
+	virtual void InitRHI() override
+	{
+		FTextureResource::InitRHI();
+
+		FRHIResourceCreateInfo CreateInfo;
+		uint32 Flags = TexCreate_NoTiling | TexCreate_OfflineProcessed;
+
+		if (CreateUAV)
+		{
+			Flags |= TexCreate_UAV;
+		}
+
+		TextureRHI = RHICreateTexture2DArray(SizeX, SizeY, SizeZ, Format, NumMips, Flags, CreateInfo);
+
+		if (CreateUAV)
+		{
+			TextureUAV = RHICreateUnorderedAccessView(TextureRHI, 0);
+		}
+	}
+
+	virtual void ReleaseRHI() override
+	{
+		FTextureResource::ReleaseRHI();
+
+		TextureUAV.SafeRelease();
+	}
+
+	FUnorderedAccessViewRHIRef TextureUAV;
+
+private:
+	uint32 SizeX;
+	uint32 SizeY;
+	uint32 SizeZ;
+	EPixelFormat Format;
+	uint32 NumMips;
+	bool CreateUAV;
+};
+
+// Compute shaders data
+
+int32 GLandscapeWeightmapThreadGroupSizeX = 16;
+int32 GLandscapeWeightmapThreadGroupSizeY = 16;
+
+struct FLandscapeProceduralWeightmapExtractLayersComponentData
+{
+	FIntPoint ComponentVertexPosition;	// Section base converted to vertex instead of quad
+	uint32 DestinationPaintLayerIndex;	// correspond to which layer info object index the data should be stored in the texture 2d array
+	uint32 WeightmapChannelToProcess;	// correspond to which RGBA channel to process
+	FIntPoint AtlasTexturePositionOutput;	// This represent the location we will write layer information
+};
+
+class FLandscapeProceduralWeightmapExtractLayersComputeShaderResource : public FRenderResource
+{
+public:
+	FLandscapeProceduralWeightmapExtractLayersComputeShaderResource(const TArray<FLandscapeProceduralWeightmapExtractLayersComponentData>& InComponentsData)
+		: OriginalComponentsData(InComponentsData)
+		, ComponentsDataCount(OriginalComponentsData.Num())
+	{}
+
+	~FLandscapeProceduralWeightmapExtractLayersComputeShaderResource()
+	{
+		ComponentsData.SafeRelease();
+		ComponentsDataSRV.SafeRelease();
+	}
+
+	/** Called when the resource is initialized. This is only called by the rendering thread. */
+	virtual void InitDynamicRHI() override
+	{
+		FRHIResourceCreateInfo CreateInfo;
+		ComponentsData = RHICreateStructuredBuffer(sizeof(FLandscapeProceduralWeightmapExtractLayersComponentData), OriginalComponentsData.Num() * sizeof(FLandscapeProceduralWeightmapExtractLayersComponentData), BUF_ShaderResource | BUF_Volatile, CreateInfo);
+		ComponentsDataSRV = RHICreateShaderResourceView(ComponentsData);
+
+		uint8* Buffer = (uint8*)RHILockStructuredBuffer(ComponentsData, 0, OriginalComponentsData.Num() * sizeof(FLandscapeProceduralWeightmapExtractLayersComponentData), RLM_WriteOnly);
+		FMemory::Memcpy(Buffer, OriginalComponentsData.GetData(), OriginalComponentsData.Num() * sizeof(FLandscapeProceduralWeightmapExtractLayersComponentData));
+		RHIUnlockStructuredBuffer(ComponentsData);
+	}
+
+	virtual void ReleaseDynamicRHI() override
+	{
+		ComponentsData.SafeRelease();
+		ComponentsDataSRV.SafeRelease();
+	}
+
+	int32 GetComponentsDataCount() const
+	{
+		return ComponentsDataCount;
+	}
+
+private:
+	friend class FLandscapeProceduralWeightmapExtractLayersCS;
+
+	FStructuredBufferRHIRef ComponentsData;
+	FShaderResourceViewRHIRef ComponentsDataSRV;
+	TArray<FLandscapeProceduralWeightmapExtractLayersComponentData> OriginalComponentsData;
+	int32 ComponentsDataCount;
+};
+
+struct FLandscapeWeightmapProceduralWeightmapExtractLayersComputeShaderParameters
+{
+	FLandscapeWeightmapProceduralWeightmapExtractLayersComputeShaderParameters()
+		: ComponentWeightmapResource(nullptr)
+		, ComputeShaderResource(nullptr)
+		, AtlasWeightmapsPerLayer(nullptr)
+		, ComponentSize(0)
+	{}
+
+	FLandscapeTexture2DResource* ComponentWeightmapResource;
+	FLandscapeProceduralWeightmapExtractLayersComputeShaderResource* ComputeShaderResource;
+	FLandscapeTexture2DArrayResource* AtlasWeightmapsPerLayer;
+	uint32 ComponentSize;
+};
+
+class FLandscapeProceduralWeightmapExtractLayersCS : public FGlobalShader
+{
+	DECLARE_GLOBAL_SHADER(FLandscapeProceduralWeightmapExtractLayersCS);
+public:
+
+	static bool ShouldCompilePermutation(const FGlobalShaderPermutationParameters& Parameters)
+	{
+		return IsFeatureLevelSupported(Parameters.Platform, ERHIFeatureLevel::SM5) && !IsConsolePlatform(Parameters.Platform);
+	}
+
+	static void ModifyCompilationEnvironment(const FGlobalShaderPermutationParameters& Parameters, FShaderCompilerEnvironment& OutEnvironment)
+	{
+		OutEnvironment.SetDefine(TEXT("THREADGROUP_SIZEX"), GLandscapeWeightmapThreadGroupSizeX);
+		OutEnvironment.SetDefine(TEXT("THREADGROUP_SIZEY"), GLandscapeWeightmapThreadGroupSizeY);
+	}
+
+	FLandscapeProceduralWeightmapExtractLayersCS(const ShaderMetaType::CompiledShaderInitializerType& Initializer)
+		: FGlobalShader(Initializer)
+	{
+		ComponentWeightmapParam.Bind(Initializer.ParameterMap, TEXT("InComponentWeightMaps"));
+		AtlasPaintListsParam.Bind(Initializer.ParameterMap, TEXT("OutAtlasPaintLayers"));
+		ComponentsDataParam.Bind(Initializer.ParameterMap, TEXT("InExtractLayersComponentsData"));
+		ComponentSizeParam.Bind(Initializer.ParameterMap, TEXT("ComponentSize"));
+	}
+
+	FLandscapeProceduralWeightmapExtractLayersCS()
+	{}
+
+	void SetParameters(FRHICommandList& RHICmdList, const FLandscapeWeightmapProceduralWeightmapExtractLayersComputeShaderParameters& InParams)
+	{
+		SetTextureParameter(RHICmdList, GetComputeShader(), ComponentWeightmapParam, InParams.ComponentWeightmapResource->TextureRHI);
+		SetUAVParameter(RHICmdList, GetComputeShader(), AtlasPaintListsParam, InParams.AtlasWeightmapsPerLayer->TextureUAV);
+		SetSRVParameter(RHICmdList, GetComputeShader(), ComponentsDataParam, InParams.ComputeShaderResource->ComponentsDataSRV);
+		SetShaderValue(RHICmdList, GetComputeShader(), ComponentSizeParam, InParams.ComponentSize);
+	}
+
+	void UnsetParameters(FRHICommandList& RHICmdList)
+	{
+		SetUAVParameter(RHICmdList, GetComputeShader(), AtlasPaintListsParam, FUnorderedAccessViewRHIParamRef());
+	}
+
+	virtual bool Serialize(FArchive& Ar) override
+	{
+		bool bShaderHasOutdatedParameters = FShader::Serialize(Ar);
+		Ar << ComponentWeightmapParam;
+		Ar << AtlasPaintListsParam;
+		Ar << ComponentsDataParam;
+		Ar << ComponentSizeParam;
+
+		return bShaderHasOutdatedParameters;
+	}
+
+private:
+	FShaderResourceParameter ComponentWeightmapParam;
+	FShaderResourceParameter AtlasPaintListsParam;
+	FShaderResourceParameter ComponentsDataParam;
+	FShaderParameter ComponentSizeParam;
+};
+
+IMPLEMENT_GLOBAL_SHADER(FLandscapeProceduralWeightmapExtractLayersCS, "/Engine/Private/LandscapeProceduralCS.usf", "ComputeWeightmapPerPaintLayer", SF_Compute);
+
+class FLandscapeProceduralWeightmapExtractLayersCSDispatch_RenderThread
+{
+public:
+	FLandscapeProceduralWeightmapExtractLayersCSDispatch_RenderThread(const FLandscapeWeightmapProceduralWeightmapExtractLayersComputeShaderParameters& InShaderParams)
+		: ShaderParams(InShaderParams)
+	{}
+
+	void ExtractLayers(FRHICommandListImmediate& InRHICmdList)
+	{
+		SCOPE_CYCLE_COUNTER(STAT_LandscapeRegenerateProcedural_RenderThread);
+		SCOPED_DRAW_EVENTF(InRHICmdList, LandscapeProceduralRender, TEXT("ExtractLayers"));
+
+		TShaderMapRef<FLandscapeProceduralWeightmapExtractLayersCS> ComputeShader(GetGlobalShaderMap(GMaxRHIFeatureLevel));
+		InRHICmdList.SetComputeShader(ComputeShader->GetComputeShader());
+		ComputeShader->SetParameters(InRHICmdList, ShaderParams);
+
+		uint32 ThreadGroupCountX = FMath::CeilToInt((float)ShaderParams.ComponentSize / (float)GLandscapeWeightmapThreadGroupSizeX);
+		uint32 ThreadGroupCountY = FMath::CeilToInt((float)ShaderParams.ComponentSize / (float)GLandscapeWeightmapThreadGroupSizeY);
+		check(ThreadGroupCountX > 0 && ThreadGroupCountY > 0);
+
+		DispatchComputeShader(InRHICmdList, *ComputeShader, ThreadGroupCountX, ThreadGroupCountY, ShaderParams.ComputeShaderResource->GetComponentsDataCount());
+		ComputeShader->UnsetParameters(InRHICmdList);
+		ShaderParams.ComputeShaderResource->ReleaseResource();
+		delete ShaderParams.ComputeShaderResource;
+	}
+
+private:
+	FLandscapeWeightmapProceduralWeightmapExtractLayersComputeShaderParameters ShaderParams;
+};
+
+struct FLandscapeProceduralWeightmapPackLayersComponentData
+{
+	int32 ComponentVertexPositionX[4];		// Section base converted to vertex instead of quad
+	int32 ComponentVertexPositionY[4];		// Section base converted to vertex instead of quad
+	int32 SourcePaintLayerIndex[4];			// correspond to which layer info object index the data should be loaded from the texture 2d array
+	int32 WeightmapChannelToProcess[4];		// correspond to which RGBA channel to process
+};
+
+class FLandscapeProceduralWeightmapPackLayersComputeShaderResource : public FRenderResource
+{
+public:
+	FLandscapeProceduralWeightmapPackLayersComputeShaderResource(const TArray<FLandscapeProceduralWeightmapPackLayersComponentData>& InComponentsData, const TArray<float>& InWeightmapWeightBlendModeData, const TArray<FVector2D>& InTextureOutputOffset)
+		: OriginalComponentsData(InComponentsData)
+		, ComponentsDataCount(OriginalComponentsData.Num())
+		, OriginalWeightmapWeightBlendModeData(InWeightmapWeightBlendModeData)
+		, OriginalTextureOutputOffset(InTextureOutputOffset)
+	{}
+
+	~FLandscapeProceduralWeightmapPackLayersComputeShaderResource()
+	{
+		ComponentsData.SafeRelease();
+		ComponentsDataSRV.SafeRelease();
+		WeightmapWeightBlendModeSRV.SafeRelease();
+		WeightmapTextureOutputOffsetSRV.SafeRelease();
+	}
+
+	/** Called when the resource is initialized. This is only called by the rendering thread. */
+	virtual void InitDynamicRHI() override
+	{
+		FRHIResourceCreateInfo CreateInfo;
+		uint32 ComponentsDataMemSize = OriginalComponentsData.Num() * sizeof(FLandscapeProceduralWeightmapPackLayersComponentData);
+		ComponentsData = RHICreateStructuredBuffer(sizeof(FLandscapeProceduralWeightmapPackLayersComponentData), ComponentsDataMemSize, BUF_ShaderResource | BUF_Volatile, CreateInfo);
+		ComponentsDataSRV = RHICreateShaderResourceView(ComponentsData);
+
+		uint8* Buffer = (uint8*)RHILockStructuredBuffer(ComponentsData, 0, ComponentsDataMemSize, RLM_WriteOnly);
+		FMemory::Memcpy(Buffer, OriginalComponentsData.GetData(), ComponentsDataMemSize);
+		RHIUnlockStructuredBuffer(ComponentsData);
+
+		FRHIResourceCreateInfo WeightBlendCreateInfo;
+		uint32 WeightBlendMemSize = OriginalWeightmapWeightBlendModeData.Num() * sizeof(float);
+		WeightmapWeightBlendMode = RHICreateVertexBuffer(WeightBlendMemSize, BUF_ShaderResource | BUF_Volatile | BUF_Dynamic, WeightBlendCreateInfo);
+		WeightmapWeightBlendModeSRV = RHICreateShaderResourceView(WeightmapWeightBlendMode, sizeof(float), PF_R32_FLOAT);
+
+		void* WeightmapWeightBlendModePtr = RHILockVertexBuffer(WeightmapWeightBlendMode, 0, WeightBlendMemSize, RLM_WriteOnly);
+		FMemory::Memcpy(WeightmapWeightBlendModePtr, OriginalWeightmapWeightBlendModeData.GetData(), WeightBlendMemSize);
+		RHIUnlockVertexBuffer(WeightmapWeightBlendMode);
+
+		FRHIResourceCreateInfo TextureOutputOffsetCreateInfo;
+		uint32 TextureOutputOffsetMemSize = OriginalTextureOutputOffset.Num() * sizeof(FVector2D);
+		WeightmapTextureOutputOffset = RHICreateVertexBuffer(TextureOutputOffsetMemSize, BUF_ShaderResource | BUF_Volatile | BUF_Dynamic, TextureOutputOffsetCreateInfo);
+		WeightmapTextureOutputOffsetSRV = RHICreateShaderResourceView(WeightmapTextureOutputOffset, sizeof(FVector2D), PF_G32R32F);
+
+		void* TextureOutputOffsetPtr = RHILockVertexBuffer(WeightmapTextureOutputOffset, 0, TextureOutputOffsetMemSize, RLM_WriteOnly);
+		FMemory::Memcpy(TextureOutputOffsetPtr, OriginalTextureOutputOffset.GetData(), TextureOutputOffsetMemSize);
+		RHIUnlockVertexBuffer(WeightmapTextureOutputOffset);
+	}
+
+	virtual void ReleaseDynamicRHI() override
+	{
+		ComponentsData.SafeRelease();
+		ComponentsDataSRV.SafeRelease();
+		WeightmapWeightBlendModeSRV.SafeRelease();
+		WeightmapTextureOutputOffsetSRV.SafeRelease();
+	}
+
+	int32 GetComponentsDataCount() const
+	{
+		return ComponentsDataCount;
+	}
+
+private:
+	friend class FLandscapeProceduralWeightmapPackLayersCS;
+
+	FStructuredBufferRHIRef ComponentsData;
+	FShaderResourceViewRHIRef ComponentsDataSRV;
+	TArray<FLandscapeProceduralWeightmapPackLayersComponentData> OriginalComponentsData;
+	int32 ComponentsDataCount;
+
+	TArray<float> OriginalWeightmapWeightBlendModeData;
+	FVertexBufferRHIRef WeightmapWeightBlendMode;
+	FShaderResourceViewRHIRef WeightmapWeightBlendModeSRV;
+
+	TArray<FVector2D> OriginalTextureOutputOffset;
+	FVertexBufferRHIRef WeightmapTextureOutputOffset;
+	FShaderResourceViewRHIRef WeightmapTextureOutputOffsetSRV;
+};
+
+struct FLandscapeProceduralWeightmapPackLayersComputeShaderParameters
+{
+	FLandscapeProceduralWeightmapPackLayersComputeShaderParameters()
+		: ComponentWeightmapResource(nullptr)
+		, ComputeShaderResource(nullptr)
+		, AtlasWeightmapsPerLayer(nullptr)
+		, ComponentSize(0)
+	{}
+
+	FLandscapeTexture2DResource* ComponentWeightmapResource;
+	FLandscapeProceduralWeightmapPackLayersComputeShaderResource* ComputeShaderResource;
+	FLandscapeTexture2DArrayResource* AtlasWeightmapsPerLayer;
+	uint32 ComponentSize;
+};
+
+class FLandscapeProceduralWeightmapPackLayersCS : public FGlobalShader
+{
+	DECLARE_GLOBAL_SHADER(FLandscapeProceduralWeightmapPackLayersCS);
+public:
+
+	static bool ShouldCompilePermutation(const FGlobalShaderPermutationParameters& Parameters)
+	{
+		return IsFeatureLevelSupported(Parameters.Platform, ERHIFeatureLevel::SM5) && !IsConsolePlatform(Parameters.Platform);
+	}
+
+	static void ModifyCompilationEnvironment(const FGlobalShaderPermutationParameters& Parameters, FShaderCompilerEnvironment& OutEnvironment)
+	{
+		OutEnvironment.SetDefine(TEXT("THREADGROUP_SIZEX"), GLandscapeWeightmapThreadGroupSizeX);
+		OutEnvironment.SetDefine(TEXT("THREADGROUP_SIZEY"), GLandscapeWeightmapThreadGroupSizeY);
+	}
+
+	FLandscapeProceduralWeightmapPackLayersCS(const ShaderMetaType::CompiledShaderInitializerType& Initializer)
+		: FGlobalShader(Initializer)
+	{
+		ComponentWeightmapParam.Bind(Initializer.ParameterMap, TEXT("OutComponentWeightMaps"));
+		AtlasPaintListsParam.Bind(Initializer.ParameterMap, TEXT("InAtlasPaintLayers"));
+		ComponentsDataParam.Bind(Initializer.ParameterMap, TEXT("InPackLayersComponentsData"));
+		ComponentSizeParam.Bind(Initializer.ParameterMap, TEXT("ComponentSize"));
+		WeightmapWeightBlendModeParam.Bind(Initializer.ParameterMap, TEXT("InWeightmapWeightBlendMode"));
+		WeightmapTextureOutputOffsetParam.Bind(Initializer.ParameterMap, TEXT("InWeightmapTextureOutputOffset"));
+	}
+
+	FLandscapeProceduralWeightmapPackLayersCS()
+	{}
+
+	void SetParameters(FRHICommandList& RHICmdList, const FLandscapeProceduralWeightmapPackLayersComputeShaderParameters& InParams)
+	{
+		SetUAVParameter(RHICmdList, GetComputeShader(), ComponentWeightmapParam, InParams.ComponentWeightmapResource->TextureUAV);
+		SetTextureParameter(RHICmdList, GetComputeShader(), AtlasPaintListsParam, InParams.AtlasWeightmapsPerLayer->TextureRHI);
+		SetSRVParameter(RHICmdList, GetComputeShader(), ComponentsDataParam, InParams.ComputeShaderResource->ComponentsDataSRV);
+		SetShaderValue(RHICmdList, GetComputeShader(), ComponentSizeParam, InParams.ComponentSize);
+		SetSRVParameter(RHICmdList, GetComputeShader(), WeightmapWeightBlendModeParam, InParams.ComputeShaderResource->WeightmapWeightBlendModeSRV);
+		SetSRVParameter(RHICmdList, GetComputeShader(), WeightmapTextureOutputOffsetParam, InParams.ComputeShaderResource->WeightmapTextureOutputOffsetSRV);
+	}
+
+	void UnsetParameters(FRHICommandList& RHICmdList)
+	{
+		SetUAVParameter(RHICmdList, GetComputeShader(), ComponentWeightmapParam, FUnorderedAccessViewRHIParamRef());
+	}
+
+	virtual bool Serialize(FArchive& Ar) override
+	{
+		bool bShaderHasOutdatedParameters = FShader::Serialize(Ar);
+		Ar << ComponentWeightmapParam;
+		Ar << AtlasPaintListsParam;
+		Ar << ComponentsDataParam;
+		Ar << ComponentSizeParam;
+		Ar << WeightmapWeightBlendModeParam;
+		Ar << WeightmapTextureOutputOffsetParam;
+		return bShaderHasOutdatedParameters;
+	}
+
+private:
+	FShaderResourceParameter ComponentWeightmapParam;
+	FShaderResourceParameter AtlasPaintListsParam;
+	FShaderResourceParameter ComponentsDataParam;
+	FShaderParameter ComponentSizeParam;
+	FShaderResourceParameter WeightmapWeightBlendModeParam;
+	FShaderResourceParameter WeightmapTextureOutputOffsetParam;
+};
+
+IMPLEMENT_GLOBAL_SHADER(FLandscapeProceduralWeightmapPackLayersCS, "/Engine/Private/LandscapeProceduralCS.usf", "PackPaintLayerToWeightmap", SF_Compute);
+
+class FLandscapeProceduralWeightmapPackLayerCSDispatch_RenderThread
+{
+public:
+	FLandscapeProceduralWeightmapPackLayerCSDispatch_RenderThread(const FLandscapeProceduralWeightmapPackLayersComputeShaderParameters& InShaderParams)
+		: ShaderParams(InShaderParams)
+	{}
+
+	void PackLayers(FRHICommandListImmediate& InRHICmdList)
+	{
+		SCOPE_CYCLE_COUNTER(STAT_LandscapeRegenerateProcedural_RenderThread);
+		SCOPED_DRAW_EVENTF(InRHICmdList, LandscapeProceduralRender, TEXT("PackLayers"));
+
+		TShaderMapRef<FLandscapeProceduralWeightmapPackLayersCS> ComputeShader(GetGlobalShaderMap(GMaxRHIFeatureLevel));
+		InRHICmdList.SetComputeShader(ComputeShader->GetComputeShader());
+		ComputeShader->SetParameters(InRHICmdList, ShaderParams);
+
+		uint32 ThreadGroupCountX = FMath::CeilToInt((float)ShaderParams.ComponentSize / (float)GLandscapeWeightmapThreadGroupSizeX);
+		uint32 ThreadGroupCountY = FMath::CeilToInt((float)ShaderParams.ComponentSize / (float)GLandscapeWeightmapThreadGroupSizeY);
+		check(ThreadGroupCountX > 0 && ThreadGroupCountY > 0);
+
+		DispatchComputeShader(InRHICmdList, *ComputeShader, ThreadGroupCountX, ThreadGroupCountY, ShaderParams.ComputeShaderResource->GetComponentsDataCount());
+		ComputeShader->UnsetParameters(InRHICmdList);
+		ShaderParams.ComputeShaderResource->ReleaseResource();
+		delete ShaderParams.ComputeShaderResource;
+	}
+
+private:
+	FLandscapeProceduralWeightmapPackLayersComputeShaderParameters ShaderParams;
+};
+
+// Copy texture render command
+
+class FLandscapeProceduralCopyTexture_RenderThread
+{
 public: 
-	FLandscapeProceduralCopyResource_RenderThread(UTexture* InHeightmapRTRead, UTexture* InCopyResolveTarget, FTextureResource* InCopyResolveTargetCPUResource, FIntPoint InComponentSectionBase, int32 InSubSectionSizeQuad, int32 InNumSubSections, int32 InCurrentMip)
-		: SourceResource(InHeightmapRTRead != nullptr ? InHeightmapRTRead->Resource : nullptr)
-		, CopyResolveTargetResource(InCopyResolveTarget != nullptr ? InCopyResolveTarget->Resource : nullptr)
-		, CopyResolveTargetCPUResource(InCopyResolveTargetCPUResource)
-		, CurrentMip(InCurrentMip)
-		, ComponentSectionBase(InComponentSectionBase)
+	FLandscapeProceduralCopyTexture_RenderThread(const FString& InSourceResourceDebugName, FTextureResource* InSourceResource, const FString& InDestResourceDebugName, FTextureResource* InDestResource, FTextureResource* InDestCPUResource, 
+											     const FIntPoint& InFirstComponentSectionBase, int32 InSubSectionSizeQuad, int32 InNumSubSections, uint8 InSourceCurrentMip, uint8 InDestCurrentMip, uint32 InSourceArrayIndex, uint32 InDestArrayIndex)
+		: SourceResource(InSourceResource)
+		, DestResource(InDestResource)
+		, DestCPUResource(InDestCPUResource)
+		, SourceMip(InSourceCurrentMip)
+		, DestMip(InDestCurrentMip)
+		, SourceArrayIndex(InSourceArrayIndex)
+		, DestArrayIndex(InDestArrayIndex)
+		, ComponentSectionBase(InFirstComponentSectionBase)
 		, SubSectionSizeQuad(InSubSectionSizeQuad)
 		, NumSubSections(InNumSubSections)
-		, SourceDebugName(SourceResource != nullptr ? InHeightmapRTRead->GetName() : TEXT(""))
-		, CopyResolveDebugName(InCopyResolveTarget != nullptr ? InCopyResolveTarget->GetName() : TEXT(""))
+		, SourceDebugName(InSourceResourceDebugName)
+		, DestResourceDebugName(InDestResourceDebugName)
 	{}
 
-	void CopyToResolveTarget(FRHICommandListImmediate& InRHICmdList)
+	void Copy(FRHICommandListImmediate& InRHICmdList)
 	{
-		if (SourceResource == nullptr || CopyResolveTargetResource == nullptr)
-		{
-			return;
-		}
+		SCOPE_CYCLE_COUNTER(STAT_LandscapeRegenerateProcedural_RenderThread);
+		SCOPED_DRAW_EVENTF(InRHICmdList, LandscapeProceduralCopy, TEXT("LS Copy %s -> %s, Mip (%d -> %d), Array Index (%d -> %d)"), *SourceDebugName, *DestResourceDebugName, SourceMip, DestMip, SourceArrayIndex, DestArrayIndex);
+		SCOPED_GPU_STAT(InRHICmdList, LandscapeProceduralCopy);
 
-		SCOPE_CYCLE_COUNTER(STAT_LandscapeRegenerateProceduralHeightmaps_RenderThread);
-		SCOPED_DRAW_EVENTF(InRHICmdList, LandscapeProceduralCopy, TEXT("LS Copy %s -> %s, Mip: %d"), *SourceDebugName, *CopyResolveDebugName, CurrentMip);
-		SCOPED_GPU_STAT(InRHICmdList, LandscapeProceduralRender);
+		FIntPoint SourceSize(SourceResource->GetSizeX(), SourceResource->GetSizeY()); // SourceResource is always proper size, as it's always the good MIP we want to copy from
+		FIntPoint DestSize(DestResource->GetSizeX() >> DestMip, DestResource->GetSizeY() >> DestMip);
 
-		FIntPoint SourceReadTextureSize(SourceResource->GetSizeX(), SourceResource->GetSizeY());
-		FIntPoint CopyResolveWriteTextureSize(CopyResolveTargetResource->GetSizeX() >> CurrentMip, CopyResolveTargetResource->GetSizeY() >> CurrentMip);
-
 		int32 LocalComponentSizeQuad = SubSectionSizeQuad * NumSubSections;
-		FVector2D HeightmapPositionOffset(FMath::RoundToInt(ComponentSectionBase.X / LocalComponentSizeQuad), FMath::RoundToInt(ComponentSectionBase.Y / LocalComponentSizeQuad));
+		FVector2D PositionOffset(FMath::RoundToInt(ComponentSectionBase.X / LocalComponentSizeQuad), FMath::RoundToInt(ComponentSectionBase.Y / LocalComponentSizeQuad));
 
-		FResolveParams Params;
-		Params.SourceArrayIndex = 0;
-		Params.DestArrayIndex = CurrentMip;
+		FRHICopyTextureInfo Params;
+		Params.NumSlices = 1;
+		Params.Size.Z = 1;
+		Params.SourceSliceIndex = SourceArrayIndex;
+		Params.DestSliceIndex = DestArrayIndex;
+		Params.SourceMipIndex = 0; // In my case, always assume we copy from mip 0 to something else as in my case each mips will be stored into individual texture/RT
+		Params.DestMipIndex = DestMip;
 
-		if (SourceReadTextureSize.X <= CopyResolveWriteTextureSize.X)
+		if (SourceSize.X <= DestSize.X)
 		{
-			Params.Rect.X1 = 0;
-			Params.Rect.X2 = SourceReadTextureSize.X;
-			Params.DestRect.X1 = FMath::RoundToInt(HeightmapPositionOffset.X * (((SubSectionSizeQuad + 1) * NumSubSections) >> CurrentMip));
+			Params.SourcePosition.X = 0;
+			Params.Size.X = SourceSize.X;
+			Params.DestPosition.X = FMath::RoundToInt(PositionOffset.X * (((SubSectionSizeQuad + 1) * NumSubSections) >> DestMip));
 		}
 		else
 		{
-			Params.Rect.X1 = FMath::RoundToInt(HeightmapPositionOffset.X * (((SubSectionSizeQuad + 1) * NumSubSections) >> CurrentMip));
-			Params.Rect.X2 = Params.Rect.X1 + CopyResolveWriteTextureSize.X;
-			Params.DestRect.X1 = 0;
+			Params.SourcePosition.X = FMath::RoundToInt(PositionOffset.X * (((SubSectionSizeQuad + 1) * NumSubSections) >> SourceMip));
+			Params.Size.X = DestSize.X;
+			Params.DestPosition.X = 0;
 		}
 
-		if (SourceReadTextureSize.Y <= CopyResolveWriteTextureSize.Y)
+		if (SourceSize.Y <= DestSize.Y)
 		{
-			Params.Rect.Y1 = 0;
-			Params.Rect.Y2 = SourceReadTextureSize.Y;
-			Params.DestRect.Y1 = FMath::RoundToInt(HeightmapPositionOffset.Y * (((SubSectionSizeQuad + 1) * NumSubSections) >> CurrentMip));
+			Params.SourcePosition.Y = 0;
+			Params.Size.Y = SourceSize.Y;
+			Params.DestPosition.Y = FMath::RoundToInt(PositionOffset.Y * (((SubSectionSizeQuad + 1) * NumSubSections) >> DestMip));
 		}
 		else
 		{
-			Params.Rect.Y1 = FMath::RoundToInt(HeightmapPositionOffset.Y * (((SubSectionSizeQuad + 1) * NumSubSections) >> CurrentMip));
-			Params.Rect.Y2 = Params.Rect.Y1 + CopyResolveWriteTextureSize.Y;
-			Params.DestRect.Y1 = 0;
-		}
+			Params.SourcePosition.Y = FMath::RoundToInt(PositionOffset.Y * (((SubSectionSizeQuad + 1) * NumSubSections) >> SourceMip));
+			Params.Size.Y = DestSize.Y;
+			Params.DestPosition.Y = 0;
+		}		
 
-		InRHICmdList.CopyToResolveTarget(SourceResource->TextureRHI, CopyResolveTargetResource->TextureRHI, Params);
+		InRHICmdList.CopyTexture(SourceResource->TextureRHI, DestResource->TextureRHI, Params);
 
-		if (CopyResolveTargetCPUResource != nullptr)
+		if (DestCPUResource != nullptr)
 		{
-			InRHICmdList.CopyToResolveTarget(SourceResource->TextureRHI, CopyResolveTargetCPUResource->TextureRHI, Params);
+			InRHICmdList.CopyTexture(SourceResource->TextureRHI, DestCPUResource->TextureRHI, Params);
 		}
 	}
 
 private:
 	FTextureResource* SourceResource;
-	FTextureResource* CopyResolveTargetResource;
-	FTextureResource* CopyResolveTargetCPUResource;
-	FIntPoint ReadRenderTargetSize;
-	int32 CurrentMip;
+	FTextureResource* DestResource;
+	FTextureResource* DestCPUResource;
+	uint8 SourceMip;
+	uint8 DestMip;
+	uint32 SourceArrayIndex;
+	uint32 DestArrayIndex;
 	FIntPoint ComponentSectionBase;
 	int32 SubSectionSizeQuad;
 	int32 NumSubSections;
 	FString SourceDebugName;
-	FString CopyResolveDebugName;
+	FString DestResourceDebugName;
 };
 
-class FLandscapeHeightmapProceduralRender_RenderThread
+// Clear command
+
+class LandscapeProceduralWeightmapClear_RenderThread
+{
+public:
+	LandscapeProceduralWeightmapClear_RenderThread(const FString& InDebugName, FTextureRenderTargetResource* InTextureResourceToClear)
+		: DebugName(InDebugName)
+		, RenderTargetResource(InTextureResourceToClear)
+	{}
+
+	virtual ~LandscapeProceduralWeightmapClear_RenderThread()
+	{}
+
+	void Clear(FRHICommandListImmediate& InRHICmdList)
+	{
+		SCOPE_CYCLE_COUNTER(STAT_LandscapeRegenerateProcedural_RenderThread);
+		SCOPED_DRAW_EVENTF(InRHICmdList, LandscapeProceduralRender, TEXT("%s"), DebugName.Len() > 0 ? *DebugName : TEXT("LandscapeProceduralClear"));
+		SCOPED_GPU_STAT(InRHICmdList, LandscapeProceduralRender);
+
+		check(IsInRenderingThread());
+
+		FRHIRenderPassInfo RPInfo(RenderTargetResource->TextureRHI, ERenderTargetActions::Clear_Store);
+		InRHICmdList.BeginRenderPass(RPInfo, TEXT("Clear"));
+		InRHICmdList.EndRenderPass();
+	}
+
+	FString DebugName;
+	FTextureRenderTargetResource* RenderTargetResource;
+};
+
+// Render command
+
+template<typename ShaderDataType, typename ShaderPixelClass, typename ShaderPixelMipsClass>
+class FLandscapeProceduralRender_RenderThread
 {
 public:
 
-	FLandscapeHeightmapProceduralRender_RenderThread(const FString& InDebugName, UTextureRenderTarget2D* InWriteRenderTarget, const FIntPoint& InWriteRenderTargetSize, const FIntPoint& InReadRenderTargetSize, const FMatrix& InProjectionMatrix, 
-													 const FLandscapeHeightmapProceduralShaderParameters& InShaderParams, int32 InCurrentMip, const TArray<FLandscapeProceduralTriangle>& InTriangleList)
+	FLandscapeProceduralRender_RenderThread(const FString& InDebugName, UTextureRenderTarget2D* InWriteRenderTarget, const FIntPoint& InWriteRenderTargetSize, const FIntPoint& InReadRenderTargetSize, const FMatrix& InProjectionMatrix,
+											const ShaderDataType& InShaderParams, uint8 InCurrentMip, const TArray<FLandscapeProceduralTriangle>& InTriangleList)
 		: RenderTargetResource(InWriteRenderTarget->GameThread_GetRenderTargetResource())
 		, WriteRenderTargetSize(InWriteRenderTargetSize)
 		, ReadRenderTargetSize(InReadRenderTargetSize)
@@ -495,15 +1111,15 @@
 		VertexBufferResource.Init(InTriangleList);
 	}
 
-	virtual ~FLandscapeHeightmapProceduralRender_RenderThread()
+	virtual ~FLandscapeProceduralRender_RenderThread()
 	{}
 
 	void Render(FRHICommandListImmediate& InRHICmdList, bool InClearRT)
 	{
-		SCOPE_CYCLE_COUNTER(STAT_LandscapeRegenerateProceduralHeightmaps_RenderThread);
-		SCOPED_DRAW_EVENTF(InRHICmdList, LandscapeProceduralHeightmapRender, TEXT("%s"), DebugName.Len() > 0 ? *DebugName : TEXT("LandscapeProceduralHeightmapRender"));
+		SCOPE_CYCLE_COUNTER(STAT_LandscapeRegenerateProcedural_RenderThread);
+		SCOPED_DRAW_EVENTF(InRHICmdList, LandscapeProceduralRender, TEXT("%s"), DebugName.Len() > 0 ? *DebugName : TEXT("LandscapeProceduralRender"));
 		SCOPED_GPU_STAT(InRHICmdList, LandscapeProceduralRender);
-		INC_DWORD_STAT(STAT_LandscapeRegenerateProceduralHeightmapsDrawCalls);
+		INC_DWORD_STAT(STAT_LandscapeRegenerateProceduralDrawCalls);
 
 		check(IsInRenderingThread());
 
@@ -531,19 +1147,18 @@
 		GraphicsPSOInit.BoundShaderState.VertexDeclarationRHI = VertexDeclaration.VertexDeclarationRHI;
 		GraphicsPSOInit.PrimitiveType = PT_TriangleList;
 
-		//RT0ColorWriteMask, RT0ColorBlendOp, RT0ColorSrcBlend, RT0ColorDestBlend, RT0AlphaBlendOp, RT0AlphaSrcBlend, RT0AlphaDestBlend,
 		GraphicsPSOInit.BlendState = TStaticBlendState<CW_RGBA, BO_Add, BF_One, BF_Zero, BO_Add, BF_One, BF_Zero>::GetRHI();
 		GraphicsPSOInit.RasterizerState = TStaticRasterizerState<>::GetRHI();
 		GraphicsPSOInit.DepthStencilState = TStaticDepthStencilState<false, CF_Always>::GetRHI();
 
 		FRHIRenderPassInfo RenderPassInfo(ViewFamily.RenderTarget->GetRenderTargetTexture(), CurrentMip == 0 ? ERenderTargetActions::Clear_Store : ERenderTargetActions::Load_Store, nullptr, 0, 0);
-		InRHICmdList.BeginRenderPass(RenderPassInfo, TEXT("DrawProceduralHeightmaps"));
+		InRHICmdList.BeginRenderPass(RenderPassInfo, TEXT("DrawProcedural"));
 
 		if (CurrentMip == 0)
 		{
 			// Setup Shaders
 			TShaderMapRef<FLandscapeProceduralVS> VertexShader(GetGlobalShaderMap(View->GetFeatureLevel()));
-			TShaderMapRef<FLandscapeHeightmapProceduralPS> PixelShader(GetGlobalShaderMap(View->GetFeatureLevel()));
+			TShaderMapRef<ShaderPixelClass> PixelShader(GetGlobalShaderMap(View->GetFeatureLevel()));
 
 			GraphicsPSOInit.BoundShaderState.VertexShaderRHI = GETSAFERHISHADER_VERTEX(*VertexShader);
 			GraphicsPSOInit.BoundShaderState.PixelShaderRHI = GETSAFERHISHADER_PIXEL(*PixelShader);
@@ -561,7 +1176,7 @@
 		{
 			// Setup Shaders
 			TShaderMapRef<FLandscapeProceduralVS> VertexShader(GetGlobalShaderMap(View->GetFeatureLevel()));
-			TShaderMapRef<FLandscapeHeightmapMipsProceduralPS> PixelShader(GetGlobalShaderMap(View->GetFeatureLevel()));
+			TShaderMapRef<ShaderPixelMipsClass> PixelShader(GetGlobalShaderMap(View->GetFeatureLevel()));
 
 			GraphicsPSOInit.BoundShaderState.VertexShaderRHI = GETSAFERHISHADER_VERTEX(*VertexShader);
 			GraphicsPSOInit.BoundShaderState.PixelShaderRHI = GETSAFERHISHADER_PIXEL(*PixelShader);
@@ -593,14 +1208,18 @@
 	FIntPoint WriteRenderTargetSize;
 	FIntPoint ReadRenderTargetSize;
 	FMatrix ProjectionMatrix;
-	FLandscapeHeightmapProceduralShaderParameters ShaderParams;
+	ShaderDataType ShaderParams;
 	FLandscapeProceduralVertexBuffer VertexBufferResource;
 	int32 PrimitiveCount;
 	FLandscapeProceduralVertexDeclaration VertexDeclaration;
 	FString DebugName;
-	int32 CurrentMip;
+	uint8 CurrentMip;
 };
 
+typedef FLandscapeProceduralRender_RenderThread<FLandscapeHeightmapProceduralShaderParameters, FLandscapeHeightmapProceduralPS, FLandscapeHeightmapMipsProceduralPS> LandscapeProceduralHeightmapRender_RenderThread;
+typedef FLandscapeProceduralRender_RenderThread<FLandscapeWeightmapProceduralShaderParameters, FLandscapeWeightmapProceduralPS, FLandscapeWeightmapMipsProceduralPS> LandscapeProceduralWeightmapRender_RenderThread;
+
+#if WITH_EDITOR
 void ALandscapeProxy::SetupProceduralLayers(int32 InNumComponentsX, int32 InNumComponentsY)
 {
 	ALandscape* Landscape = GetLandscapeActor();
@@ -664,13 +1283,13 @@
 	}	
 	///// TEMP STUFF END
 
+	// TODO: When using Change Component Size, we will call Setup again, currently all the existing data will get collapsed into the layer 1, it should keep the layers data.
+
 	int32 NumComponentsX = InNumComponentsX;
 	int32 NumComponentsY = InNumComponentsY;
 	bool GenerateComponentCounts = NumComponentsX == INDEX_NONE || NumComponentsY == INDEX_NONE;
 	FIntPoint MaxSectionBase(0, 0);
 
-	uint32 UpdateFlags = 0;
-
 	// Setup all Heightmap data
 	for (ALandscapeProxy* LandscapeProxy : AllLandscapes)
 	{
@@ -714,20 +1333,21 @@
 
 	if (Landscape->HeightmapRTList.Num() == 0)
 	{
-		Landscape->HeightmapRTList.Init(nullptr, EHeightmapRTType::Count);
+		Landscape->HeightmapRTList.Init(nullptr, EHeightmapRTType::HeightmapRT_Count);
 
 		int32 CurrentMipSizeX = ((SubsectionSizeQuads + 1) * NumSubsections) * NumComponentsX;
 		int32 CurrentMipSizeY = ((SubsectionSizeQuads + 1) * NumSubsections) * NumComponentsY;
 
-		for (int32 i = 0; i < EHeightmapRTType::Count; ++i)
+		for (int32 i = 0; i < EHeightmapRTType::HeightmapRT_Count; ++i)
 		{
 			Landscape->HeightmapRTList[i] = NewObject<UTextureRenderTarget2D>(Landscape->GetOutermost());
 			check(Landscape->HeightmapRTList[i]);
 			Landscape->HeightmapRTList[i]->RenderTargetFormat = RTF_RGBA8;
 			Landscape->HeightmapRTList[i]->AddressX = TextureAddress::TA_Clamp;
 			Landscape->HeightmapRTList[i]->AddressY = TextureAddress::TA_Clamp;
+			Landscape->HeightmapRTList[i]->ClearColor = FLinearColor(0.0f, 0.0f, 0.0f, 0.0f);
 
-			if (i < LandscapeSizeMip1) // Landscape size RT
+			if (i < EHeightmapRTType::HeightmapRT_Mip1) // Landscape size RT
 			{
 				Landscape->HeightmapRTList[i]->InitAutoFormat(FMath::RoundUpToPowerOfTwo(TotalVertexCountX), FMath::RoundUpToPowerOfTwo(TotalVertexCountY));
 			}
@@ -751,8 +1371,6 @@
 	TArray<FVector> VertexNormals;
 	TArray<uint16> EmptyHeightmapData;
 
-	UpdateFlags |= EProceduralContentUpdateFlag::Heightmap_Render;
-
 	// Setup all Heightmap data
 	for (ALandscapeProxy* LandscapeProxy : AllLandscapes)
 	{
@@ -782,12 +1400,10 @@
 					int32 MipSizeU = Heightmap->Source.GetSizeX();
 					int32 MipSizeV = Heightmap->Source.GetSizeY();
 
-					UpdateFlags |= EProceduralContentUpdateFlag::Heightmap_ResolveToTexture | EProceduralContentUpdateFlag::Heightmap_BoundsAndCollision;
-
 					// Copy data from Heightmap to first layer, after that all other layer will get init to empty layer
 					if (FirstLayer)
 					{
-						int32 MipIndex = 0;
+						uint8 MipIndex = 0;
 						TArray<uint8> MipData;
 						MipData.Reserve(MipSizeU*MipSizeV * sizeof(FColor));
 
@@ -948,19 +1564,226 @@
 		}
 	}
 
-	// Setup all Weightmap data
-	// TODO
+	// Weightmaps handling
+	if (Landscape->WeightmapRTList.Num() == 0)
+	{
+		Landscape->WeightmapRTList.Init(nullptr, EWeightmapRTType::WeightmapRT_Count);
 
+		int32 CurrentMipSizeX = ((SubsectionSizeQuads + 1) * NumSubsections) * NumComponentsX;
+		int32 CurrentMipSizeY = ((SubsectionSizeQuads + 1) * NumSubsections) * NumComponentsY;
+
+		for (int32 i = 0; i < EWeightmapRTType::WeightmapRT_Count; ++i)
+		{
+			Landscape->WeightmapRTList[i] = NewObject<UTextureRenderTarget2D>(Landscape->GetOutermost());
+			check(Landscape->WeightmapRTList[i]);
+			Landscape->WeightmapRTList[i]->AddressX = TextureAddress::TA_Clamp;
+			Landscape->WeightmapRTList[i]->AddressY = TextureAddress::TA_Clamp;
+			Landscape->WeightmapRTList[i]->ClearColor = FLinearColor(0.0f, 0.0f, 0.0f, 0.0f);
+			Landscape->WeightmapRTList[i]->RenderTargetFormat = RTF_RGBA8;
+
+			if (i < EWeightmapRTType::WeightmapRT_Mip0) // Landscape size RT, only create the number of layer we have
+			{
+				Landscape->WeightmapRTList[i]->RenderTargetFormat = i == WeightmapRT_Scratch_RGBA ? RTF_RGBA8 : RTF_R8;
+				Landscape->WeightmapRTList[i]->InitAutoFormat(FMath::RoundUpToPowerOfTwo(TotalVertexCountX), FMath::RoundUpToPowerOfTwo(TotalVertexCountY));
+			}
+			else // Mips
+			{
+				Landscape->WeightmapRTList[i]->InitAutoFormat(FMath::RoundUpToPowerOfTwo(CurrentMipSizeX), FMath::RoundUpToPowerOfTwo(CurrentMipSizeY));
+
+				CurrentMipSizeX >>= 1;
+				CurrentMipSizeY >>= 1;
+			}
+
+			if (Landscape->WeightmapRTList[i] != nullptr)
+			{
+				Landscape->WeightmapRTList[i]->UpdateResourceImmediate(true);
+			}
+
+			// Only generate required mips RT
+			if (CurrentMipSizeX == NumComponentsX && CurrentMipSizeY == NumComponentsY)
+			{
+				break;
+			}
+		}
+	}
+
+	TArray<ULandscapeComponent*> ComponentsToCleanup;
+
+	for (ALandscapeProxy* LandscapeProxy : AllLandscapes)
+	{
+		bool FirstLayer = true;
+
+		for (auto& ItLayerDataPair : LandscapeProxy->ProceduralLayersData)
+		{
+			FName ProceduralLayerName = ItLayerDataPair.Key;
+			FProceduralLayerData& ProceduralLayerData = ItLayerDataPair.Value;
+
+			struct FTextureData
+			{
+				UTexture2D* Texture;
+				FLandscapeWeightmapUsage* Usage;
+				FLandscapeWeightmapUsage* OriginalUsage;
+			};
+
+			TMap<UTexture2D*, FTextureData> ProcessedTextures;
+
+			for (ULandscapeComponent* Component : LandscapeProxy->LandscapeComponents)
+			{
+				FWeightmapLayerData* WeightmapLayer = ProceduralLayerData.WeightmapData.Find(Component);
+
+				if (WeightmapLayer == nullptr)
+				{					
+					FWeightmapLayerData& NewWeightmapData = ProceduralLayerData.WeightmapData.Add(Component, FWeightmapLayerData());
+
+					// If no data exist for this weightmap and that data exist in the base weightmap, simply copy it to the first layer, and clear the data in the base (as it will become the final weightmap)
+					if (FirstLayer)
+					{
+						ComponentsToCleanup.Add(Component);
+
+						const TArray<UTexture2D*>& ComponentWeightmapTextures = Component->GetWeightmapTextures();
+						TArray<FWeightmapLayerAllocationInfo>& ComponentLayerAllocations = Component->GetWeightmapLayerAllocations();
+						TArray<FLandscapeWeightmapUsage*>& ComponentWeightmapTexturesUsage = Component->GetWeightmapTexturesUsage();
+
+						NewWeightmapData.Weightmaps.AddDefaulted(ComponentWeightmapTextures.Num());
+						NewWeightmapData.WeightmapTextureUsages.AddDefaulted(ComponentWeightmapTexturesUsage.Num());
+
+						for (int32 TextureIndex = 0; TextureIndex < ComponentWeightmapTextures.Num(); ++TextureIndex)
+						{
+							UTexture2D* OriginalWeightmapTexture = ComponentWeightmapTextures[TextureIndex];
+							FTextureData* TextureData = ProcessedTextures.Find(OriginalWeightmapTexture);
+
+							if (TextureData != nullptr)
+							{
+								ComponentWeightmapTexturesUsage[TextureIndex] = TextureData->OriginalUsage;
+
+								NewWeightmapData.Weightmaps[TextureIndex] = TextureData->Texture;
+								NewWeightmapData.WeightmapTextureUsages[TextureIndex] = TextureData->Usage;
+								check(TextureData->Usage->ProceduralLayerName == ProceduralLayerName);
+
+								for (int32 ChannelIndex = 0; ChannelIndex < 4; ++ChannelIndex)
+								{
+									ULandscapeComponent* ChannelLandscapeComponent = NewWeightmapData.WeightmapTextureUsages.Last()->ChannelUsage[ChannelIndex];
+
+									if (ChannelLandscapeComponent != nullptr && ChannelLandscapeComponent == Component)
+									{
+										for (FWeightmapLayerAllocationInfo& Allocation : ComponentLayerAllocations)
+										{
+											if (Allocation.WeightmapTextureIndex == TextureIndex)
+											{
+												NewWeightmapData.WeightmapLayerAllocations.Add(Allocation);
+											}
+										}
+
+										break;
+									}
+								}
+							}
+							else
+							{
+								UTexture2D* NewWeightmapTexture = LandscapeProxy->CreateLandscapeTexture(OriginalWeightmapTexture->Source.GetSizeX(), OriginalWeightmapTexture->Source.GetSizeY(), TEXTUREGROUP_Terrain_Weightmap, OriginalWeightmapTexture->Source.GetFormat());
+
+								int32 MipSubsectionSizeQuads = SubsectionSizeQuads;
+								int32 MipSizeU = OriginalWeightmapTexture->Source.GetSizeX();
+								int32 MipSizeV = OriginalWeightmapTexture->Source.GetSizeY();
+
+								uint8 MipIndex = 0;
+								TArray<uint8> MipData;
+								MipData.Reserve(MipSizeU * MipSizeV * sizeof(FColor));
+
+								while (MipSizeU > 1 && MipSizeV > 1 && MipSubsectionSizeQuads >= 1)
+								{
+									MipData.Reset();
+									OriginalWeightmapTexture->Source.GetMipData(MipData, MipIndex);
+
+									FColor* WeightmapTextureData = (FColor*)NewWeightmapTexture->Source.LockMip(MipIndex);
+									FMemory::Memcpy(WeightmapTextureData, MipData.GetData(), MipData.Num());
+									NewWeightmapTexture->Source.UnlockMip(MipIndex);
+
+									MipSizeU >>= 1;
+									MipSizeV >>= 1;
+
+									MipSubsectionSizeQuads = ((MipSubsectionSizeQuads + 1) >> 1) - 1;
+									++MipIndex;
+								}
+
+								NewWeightmapData.Weightmaps[TextureIndex] = NewWeightmapTexture;
+								NewWeightmapData.WeightmapTextureUsages[TextureIndex] = ComponentWeightmapTexturesUsage[TextureIndex];
+								NewWeightmapData.WeightmapTextureUsages[TextureIndex]->ProceduralLayerName = ProceduralLayerName;
+
+								// Create new Usage for the base layer as the other one will now be used by the Layer 1
+								ComponentWeightmapTexturesUsage[TextureIndex] = LandscapeProxy->WeightmapUsageMap.Add(NewWeightmapTexture, new FLandscapeWeightmapUsage());
+
+								for (FWeightmapLayerAllocationInfo& Allocation : ComponentLayerAllocations)
+								{
+									if (Allocation.WeightmapTextureIndex == TextureIndex)
+									{
+										NewWeightmapData.WeightmapLayerAllocations.Add(Allocation);
+									}
+								}
+
+								FTextureData NewTextureData;
+								NewTextureData.Texture = NewWeightmapTexture;
+								NewTextureData.Usage = NewWeightmapData.WeightmapTextureUsages[TextureIndex];
+								NewTextureData.OriginalUsage = ComponentWeightmapTexturesUsage[TextureIndex];
+
+								ProcessedTextures.Add(OriginalWeightmapTexture, NewTextureData);
+
+								NewWeightmapTexture->BeginCachePlatformData();
+								NewWeightmapTexture->ClearAllCachedCookedPlatformData();
+							}							
+						}
+					}
+				}
+				else
+				{
+					WeightmapLayer->WeightmapTextureUsages.AddDefaulted(WeightmapLayer->Weightmaps.Num());
+
+					// regenerate the weightmap usage
+					for (int32 LayerIdx = 0; LayerIdx < WeightmapLayer->WeightmapLayerAllocations.Num(); LayerIdx++)
+					{
+						FWeightmapLayerAllocationInfo& Allocation = WeightmapLayer->WeightmapLayerAllocations[LayerIdx];
+						UTexture2D* WeightmapTexture = WeightmapLayer->Weightmaps[Allocation.WeightmapTextureIndex];
+						FLandscapeWeightmapUsage** TempUsage = LandscapeProxy->WeightmapUsageMap.Find(WeightmapTexture);
+
+						if (TempUsage == nullptr)
+						{
+							TempUsage = &LandscapeProxy->WeightmapUsageMap.Add(WeightmapTexture, new FLandscapeWeightmapUsage());							
+							(*TempUsage)->ProceduralLayerName = ProceduralLayerName;
+						}
+
+						FLandscapeWeightmapUsage* Usage = *TempUsage;
+						WeightmapLayer->WeightmapTextureUsages[Allocation.WeightmapTextureIndex] = Usage; // Keep a ref to it for faster access
+
+						check(Usage->ChannelUsage[Allocation.WeightmapTextureChannel] == nullptr || Usage->ChannelUsage[Allocation.WeightmapTextureChannel] == Component);
+
+						Usage->ChannelUsage[Allocation.WeightmapTextureChannel] = Component;
+					}
+				}
+			}
+
+			FirstLayer = false;
+		}
+	}
+
+	for (ULandscapeComponent* Component : ComponentsToCleanup)
+	{
+		TArray<FWeightmapLayerAllocationInfo>& ComponentLayerAllocations = Component->GetWeightmapLayerAllocations();
+		ComponentLayerAllocations.Reset();		
+	}
+
 	// Fix Owning actor for Brushes. It can happen after save as operation, for example
 	for (FProceduralLayer& Layer : Landscape->ProceduralLayers)
 	{
-		for (int32 i = 0; i < Layer.Brushes.Num(); ++i)
+		for (int32 i = Layer.Brushes.Num() - 1; i >= 0; --i)
 		{
 			FLandscapeProceduralLayerBrush& Brush = Layer.Brushes[i];
 
-			if (Brush.BPCustomBrush->GetOwningLandscape() == nullptr)
+			if (Brush.BPCustomBrush != nullptr)
 			{
-				Brush.BPCustomBrush->SetOwningLandscape(Landscape);
+				if (Brush.BPCustomBrush->GetOwningLandscape() == nullptr)
+				{
+					Brush.BPCustomBrush->SetOwningLandscape(Landscape);
+				}
 			}
 		}
 
@@ -971,7 +1794,7 @@
 			{
 				FLandscapeProceduralLayerBrush& Brush = Layer.Brushes[i];
 
-				if (Brush.BPCustomBrush->IsAffectingHeightmap())
+				if (Brush.BPCustomBrush != nullptr && Brush.BPCustomBrush->IsAffectingHeightmap())
 				{
 					Layer.HeightmapBrushOrderIndices.Add(i);
 				}
@@ -984,7 +1807,7 @@
 			{
 				FLandscapeProceduralLayerBrush& Brush = Layer.Brushes[i];
 
-				if (Brush.BPCustomBrush->IsAffectingWeightmap())
+				if (Brush.BPCustomBrush != nullptr && Brush.BPCustomBrush->IsAffectingWeightmap())
 				{
 					Layer.WeightmapBrushOrderIndices.Add(i);
 				}
@@ -992,18 +1815,178 @@
 		}		
 		// TEMP stuff
 	}
-
-	Landscape->RequestProceduralContentUpdate(UpdateFlags);
 }
 
-void ALandscape::CopyProceduralTargetToResolveTarget(UTexture* InHeightmapRTRead, UTexture* InCopyResolveTarget, FTextureResource* InCopyResolveTargetCPUResource, const FIntPoint& InFirstComponentSectionBase, int32 InCurrentMip) const
+void ALandscape::CopyProceduralTexture(UTexture* InSourceTexture, UTexture* InDestTexture, FTextureResource* InDestCPUResource, const FIntPoint& InFirstComponentSectionBase, uint8 InSourceCurrentMip, uint8 InDestCurrentMip, uint32 InSourceArrayIndex, uint32 InDestArrayIndex) const
 {
-	FLandscapeProceduralCopyResource_RenderThread CopyResource(InHeightmapRTRead, InCopyResolveTarget, InCopyResolveTargetCPUResource, InFirstComponentSectionBase, SubsectionSizeQuads, NumSubsections, InCurrentMip);
+	CopyProceduralTexture(InSourceTexture != nullptr ? InSourceTexture->GetName() : TEXT("Invalid"), InSourceTexture->Resource, InDestTexture != nullptr ? InDestTexture->GetName() : TEXT("Invalid"), InDestTexture->Resource, InDestCPUResource, 
+						  InFirstComponentSectionBase, InSourceCurrentMip, InDestCurrentMip, InSourceArrayIndex, InDestArrayIndex);
+}
 
-	ENQUEUE_RENDER_COMMAND(FLandscapeProceduralCopyResultCommand)(
-		[CopyResource](FRHICommandListImmediate& RHICmdList) mutable
+void ALandscape::CopyProceduralTexture(const FString& InSourceDebugName, FTextureResource* InSourceResource, const FString& InDestDebugName, FTextureResource* InDestResource, FTextureResource* InDestCPUResource, const FIntPoint& InFirstComponentSectionBase, 
+										uint8 InSourceCurrentMip, uint8 InDestCurrentMip, uint32 InSourceArrayIndex, uint32 InDestArrayIndex) const
+{
+	FLandscapeProceduralCopyTexture_RenderThread CopyTexture(InSourceDebugName, InSourceResource, InDestDebugName, InDestResource, InDestCPUResource, InFirstComponentSectionBase, SubsectionSizeQuads, NumSubsections, InSourceCurrentMip, InDestCurrentMip, InSourceArrayIndex, InDestArrayIndex);
+
+	ENQUEUE_RENDER_COMMAND(FLandscapeProceduralCopyCommand)(
+		[CopyTexture](FRHICommandListImmediate& RHICmdList) mutable
+	{
+		CopyTexture.Copy(RHICmdList);
+	});	
+}
+
+void ALandscape::DrawWeightmapComponentsToRenderTarget(const FString& InDebugName, const FIntPoint& InSectionBase, const FVector2D& InScaleBias, UTexture* InWeightmapRTRead, UTextureRenderTarget2D* InOptionalWeightmapRTRead2, UTextureRenderTarget2D* InWeightmapRTWrite,
+													   bool InClearRTWrite, FLandscapeWeightmapProceduralShaderParameters& InShaderParams, uint8 InMipRender) const
+{
+	check(InWeightmapRTRead != nullptr);
+	check(InWeightmapRTWrite != nullptr);
+
+	FIntPoint WeightmapWriteTextureSize(InWeightmapRTWrite->SizeX, InWeightmapRTWrite->SizeY);
+	FIntPoint WeightmapReadTextureSize(InWeightmapRTRead->Source.GetSizeX(), InWeightmapRTRead->Source.GetSizeY());
+	UTextureRenderTarget2D* WeightmapRTRead = Cast<UTextureRenderTarget2D>(InWeightmapRTRead);
+
+	if (WeightmapRTRead != nullptr)
+	{
+		WeightmapReadTextureSize.X = WeightmapRTRead->SizeX;
+		WeightmapReadTextureSize.Y = WeightmapRTRead->SizeY;
+	}
+
+	// Quad Setup
+	TArray<FLandscapeProceduralTriangle> TriangleList;
+	TriangleList.Reserve(1 * 2 * NumSubsections);
+
+	if (InMipRender == 0)
+	{
+		GenerateProceduralRenderQuadsAtlas(InSectionBase, InScaleBias, SubsectionSizeQuads, WeightmapReadTextureSize, WeightmapWriteTextureSize, TriangleList);
+	}
+	else
+	{
+		GenerateProceduralRenderQuadsMip(InSectionBase, InScaleBias, SubsectionSizeQuads, WeightmapReadTextureSize, WeightmapWriteTextureSize, InMipRender, TriangleList);
+	}
+
+	InShaderParams.ReadWeightmap1 = InWeightmapRTRead;
+	InShaderParams.ReadWeightmap2 = InOptionalWeightmapRTRead2;
+	InShaderParams.CurrentMipComponentVertexCount = (((SubsectionSizeQuads + 1) * NumSubsections) >> InMipRender);
+
+	if (InMipRender > 0)
+	{
+		InShaderParams.CurrentMipSize = WeightmapWriteTextureSize;
+		InShaderParams.ParentMipSize = WeightmapReadTextureSize;
+	}
+
+	FMatrix ProjectionMatrix = AdjustProjectionMatrixForRHI(FTranslationMatrix(FVector(0, 0, 0)) *
+		FMatrix(FPlane(1.0f / (FMath::Max<uint32>(WeightmapWriteTextureSize.X, 1.f) / 2.0f), 0.0, 0.0f, 0.0f), FPlane(0.0f, -1.0f / (FMath::Max<uint32>(WeightmapWriteTextureSize.Y, 1.f) / 2.0f), 0.0f, 0.0f), FPlane(0.0f, 0.0f, 1.0f, 0.0f), FPlane(-1.0f, 1.0f, 0.0f, 1.0f)));
+
+	LandscapeProceduralWeightmapRender_RenderThread ProceduralRender(InDebugName, InWeightmapRTWrite, WeightmapWriteTextureSize, WeightmapReadTextureSize, ProjectionMatrix, InShaderParams, InMipRender, TriangleList);
+
+	ENQUEUE_RENDER_COMMAND(FDrawLandscapeProceduralWeightmapCommand)(
+		[ProceduralRender, ClearRT = InClearRTWrite](FRHICommandListImmediate& RHICmdList) mutable
+	{
+		ProceduralRender.Render(RHICmdList, ClearRT);
+	});
+
+	PrintProceduralDebugRT(InDebugName, InWeightmapRTWrite, InMipRender, false);
+}
+
+void ALandscape::DrawWeightmapComponentsToRenderTarget(const FString& InDebugName, const TArray<ULandscapeComponent*>& InComponentsToDraw, UTexture* InWeightmapRTRead, UTextureRenderTarget2D* InOptionalWeightmapRTRead2, UTextureRenderTarget2D* InWeightmapRTWrite,
+													   bool InClearRTWrite, FLandscapeWeightmapProceduralShaderParameters& InShaderParams, uint8 InMipRender) const
+{
+	check(InWeightmapRTRead != nullptr);
+	check(InWeightmapRTWrite != nullptr);
+
+	FIntPoint WeightmapWriteTextureSize(InWeightmapRTWrite->SizeX, InWeightmapRTWrite->SizeY);
+	FIntPoint WeightmapReadTextureSize(InWeightmapRTRead->Source.GetSizeX(), InWeightmapRTRead->Source.GetSizeY());
+	UTextureRenderTarget2D* WeightmapRTRead = Cast<UTextureRenderTarget2D>(InWeightmapRTRead);
+
+	if (WeightmapRTRead != nullptr)
+	{
+		WeightmapReadTextureSize.X = WeightmapRTRead->SizeX;
+		WeightmapReadTextureSize.Y = WeightmapRTRead->SizeY;
+	}
+
+	// Quad Setup
+	TArray<FLandscapeProceduralTriangle> TriangleList;
+	TriangleList.Reserve(InComponentsToDraw.Num() * 2 * NumSubsections);
+
+	if (InMipRender == 0)
+	{
+		for (ULandscapeComponent* Component : InComponentsToDraw)
 		{
-			CopyResource.CopyToResolveTarget(RHICmdList);
+			// TODO: check what to do with WeightmapSubsectionOffset
+			FVector2D WeightmapScaleBias(Component->WeightmapScaleBias.Z, Component->WeightmapScaleBias.W);
+			GenerateProceduralRenderQuadsAtlas(Component->GetSectionBase(), WeightmapScaleBias, SubsectionSizeQuads, WeightmapReadTextureSize, WeightmapWriteTextureSize, TriangleList);
+		}
+	}
+	else
+	{
+		for (ULandscapeComponent* Component : InComponentsToDraw)
+		{
+			// TODO: check what to do with WeightmapSubsectionOffset
+			FVector2D WeightmapScaleBias(Component->WeightmapScaleBias.Z, Component->WeightmapScaleBias.W);
+			GenerateProceduralRenderQuadsMip(Component->GetSectionBase(), WeightmapScaleBias, SubsectionSizeQuads, WeightmapReadTextureSize, WeightmapWriteTextureSize, InMipRender, TriangleList);
+		}
+	}
+
+	InShaderParams.ReadWeightmap1 = InWeightmapRTRead;
+	InShaderParams.ReadWeightmap2 = InOptionalWeightmapRTRead2;
+	InShaderParams.CurrentMipComponentVertexCount = (((SubsectionSizeQuads + 1) * NumSubsections) >> InMipRender);
+
+	if (InMipRender > 0)
+	{
+		InShaderParams.CurrentMipSize = WeightmapWriteTextureSize;
+		InShaderParams.ParentMipSize = WeightmapReadTextureSize;
+	}
+
+	FMatrix ProjectionMatrix = AdjustProjectionMatrixForRHI(FTranslationMatrix(FVector(0, 0, 0)) *
+															FMatrix(FPlane(1.0f / (FMath::Max<uint32>(WeightmapWriteTextureSize.X, 1.f) / 2.0f), 0.0, 0.0f, 0.0f), FPlane(0.0f, -1.0f / (FMath::Max<uint32>(WeightmapWriteTextureSize.Y, 1.f) / 2.0f), 0.0f, 0.0f), FPlane(0.0f, 0.0f, 1.0f, 0.0f), FPlane(-1.0f, 1.0f, 0.0f, 1.0f)));
+
+	LandscapeProceduralWeightmapRender_RenderThread ProceduralRender(InDebugName, InWeightmapRTWrite, WeightmapWriteTextureSize, WeightmapReadTextureSize, ProjectionMatrix, InShaderParams, InMipRender, TriangleList);
+
+	ENQUEUE_RENDER_COMMAND(FDrawLandscapeProceduralWeightmapCommand)(
+		[ProceduralRender, ClearRT = InClearRTWrite](FRHICommandListImmediate& RHICmdList) mutable
+	{
+		ProceduralRender.Render(RHICmdList, ClearRT);
+	});
+
+	PrintProceduralDebugRT(InDebugName, InWeightmapRTWrite, InMipRender, false);
+}
+
+void ALandscape::DrawWeightmapComponentToRenderTargetMips(const FIntPoint& TopLeftTexturePosition, UTexture* InReadWeightmap, bool InClearRTWrite, struct FLandscapeWeightmapProceduralShaderParameters& InShaderParams) const
+{
+	bool OutputDebugName = CVarOutputProceduralDebugDrawCallName.GetValueOnAnyThread() == 1 || CVarOutputProceduralRTContent.GetValueOnAnyThread() == 1 ? true : false;
+	int32 CurrentMip = 1;
+	UTexture* ReadMipRT = InReadWeightmap;
+
+	// Convert from Texture position to SectionBase
+	int32 LocalComponentSizeQuad = SubsectionSizeQuads * NumSubsections;
+	int32 LocalComponentSizeVerts = SubsectionSizeQuads + 1 * NumSubsections;
+
+	FVector2D PositionOffset(FMath::RoundToInt(TopLeftTexturePosition.X / LocalComponentSizeVerts), FMath::RoundToInt(TopLeftTexturePosition.Y / LocalComponentSizeVerts));
+	FIntPoint ComponentSectionBase(PositionOffset.X * LocalComponentSizeQuad, PositionOffset.Y * LocalComponentSizeQuad);
+	FVector2D WeightmapScaleBias(0.0f, 0.0f);
+
+	for (int32 MipRTIndex = EWeightmapRTType::WeightmapRT_Mip1; MipRTIndex < EWeightmapRTType::WeightmapRT_Count; ++MipRTIndex)
+	{
+		UTextureRenderTarget2D* WriteMipRT = WeightmapRTList[MipRTIndex];
+
+		if (WriteMipRT != nullptr)
+		{
+			DrawWeightmapComponentsToRenderTarget(OutputDebugName ? FString::Printf(TEXT("LS Weight: %s = -> %s Mips %d"), *ReadMipRT->GetName(), *WriteMipRT->GetName(), CurrentMip) : TEXT(""),
+												  ComponentSectionBase, WeightmapScaleBias, ReadMipRT, nullptr, WriteMipRT, InClearRTWrite, InShaderParams, CurrentMip++);
+		}
+
+		ReadMipRT = WeightmapRTList[MipRTIndex];
+	}
+}
+
+void ALandscape::ClearWeightmapTextureResource(const FString& InDebugName, FTextureRenderTargetResource* InTextureResourceToClear)
+{
+	LandscapeProceduralWeightmapClear_RenderThread ProceduralClear(InDebugName, InTextureResourceToClear);
+
+	ENQUEUE_RENDER_COMMAND(FLandscapeProceduralClearWeightmapCommand)(
+		[ProceduralClear](FRHICommandListImmediate& RHICmdList) mutable
+		{
+			ProceduralClear.Clear(RHICmdList);
 		});
 }
 
@@ -1013,7 +1996,7 @@
 	int32 CurrentMip = 1;
 	UTexture* ReadMipRT = InReadHeightmap;
 
-	for (int32 MipRTIndex = EHeightmapRTType::LandscapeSizeMip1; MipRTIndex < EHeightmapRTType::Count; ++MipRTIndex)
+	for (int32 MipRTIndex = EHeightmapRTType::HeightmapRT_Mip1; MipRTIndex < EHeightmapRTType::HeightmapRT_Count; ++MipRTIndex)
 	{
 		UTextureRenderTarget2D* WriteMipRT = HeightmapRTList[MipRTIndex];
 
@@ -1027,8 +2010,8 @@
 	}
 }
 
-void ALandscape::DrawHeightmapComponentsToRenderTarget(const FString& InDebugName, TArray<ULandscapeComponent*>& InComponentsToDraw, UTexture* InHeightmapRTRead, UTextureRenderTarget2D* InOptionalHeightmapRTRead2, UTextureRenderTarget2D* InHeightmapRTWrite,
-													  ERTDrawingType InDrawType, bool InClearRTWrite, FLandscapeHeightmapProceduralShaderParameters& InShaderParams, int32 InMipRender) const
+void ALandscape::DrawHeightmapComponentsToRenderTarget(const FString& InDebugName, const TArray<ULandscapeComponent*>& InComponentsToDraw, UTexture* InHeightmapRTRead, UTextureRenderTarget2D* InOptionalHeightmapRTRead2, UTextureRenderTarget2D* InHeightmapRTWrite,
+													  ERTDrawingType InDrawType, bool InClearRTWrite, FLandscapeHeightmapProceduralShaderParameters& InShaderParams, uint8 InMipRender) const
 {
 	check(InHeightmapRTRead != nullptr);
 	check(InHeightmapRTWrite != nullptr);
@@ -1054,7 +2037,7 @@
 			for (ULandscapeComponent* Component : InComponentsToDraw)
 			{
 				FVector2D HeightmapScaleBias(Component->HeightmapScaleBias.Z, Component->HeightmapScaleBias.W);
-				GenerateHeightmapQuadsAtlas(Component->GetSectionBase(), HeightmapScaleBias, SubsectionSizeQuads, HeightmapReadTextureSize, HeightmapWriteTextureSize, TriangleList);
+				GenerateProceduralRenderQuadsAtlas(Component->GetSectionBase(), HeightmapScaleBias, SubsectionSizeQuads, HeightmapReadTextureSize, HeightmapWriteTextureSize, TriangleList);
 			}
 		} break;
 
@@ -1063,7 +2046,7 @@
 			for (ULandscapeComponent* Component : InComponentsToDraw)
 			{
 				FVector2D HeightmapScaleBias(Component->HeightmapScaleBias.Z, Component->HeightmapScaleBias.W);
-				GenerateHeightmapQuadsAtlasToNonAtlas(Component->GetSectionBase(), HeightmapScaleBias, SubsectionSizeQuads, HeightmapReadTextureSize, HeightmapWriteTextureSize, TriangleList);
+				GenerateProceduralRenderQuadsAtlasToNonAtlas(Component->GetSectionBase(), HeightmapScaleBias, SubsectionSizeQuads, HeightmapReadTextureSize, HeightmapWriteTextureSize, TriangleList);
 			}
 		} break;
 
@@ -1072,7 +2055,7 @@
 			for (ULandscapeComponent* Component : InComponentsToDraw)
 			{
 				FVector2D HeightmapScaleBias(Component->HeightmapScaleBias.Z, Component->HeightmapScaleBias.W);
-				GenerateHeightmapQuadsNonAtlas(Component->GetSectionBase(), HeightmapScaleBias, SubsectionSizeQuads, HeightmapReadTextureSize, HeightmapWriteTextureSize, TriangleList);
+				GenerateProceduralRenderQuadsNonAtlas(Component->GetSectionBase(), HeightmapScaleBias, SubsectionSizeQuads, HeightmapReadTextureSize, HeightmapWriteTextureSize, TriangleList);
 			}
 		} break;
 
@@ -1081,7 +2064,7 @@
 			for (ULandscapeComponent* Component : InComponentsToDraw)
 			{
 				FVector2D HeightmapScaleBias(Component->HeightmapScaleBias.Z, Component->HeightmapScaleBias.W);
-				GenerateHeightmapQuadsNonAtlasToAtlas(Component->GetSectionBase(), HeightmapScaleBias, SubsectionSizeQuads, HeightmapReadTextureSize, HeightmapWriteTextureSize, TriangleList);
+				GenerateProceduralRenderQuadsNonAtlasToAtlas(Component->GetSectionBase(), HeightmapScaleBias, SubsectionSizeQuads, HeightmapReadTextureSize, HeightmapWriteTextureSize, TriangleList);
 			}
 		} break;
 
@@ -1090,7 +2073,7 @@
 			for (ULandscapeComponent* Component : InComponentsToDraw)
 			{
 				FVector2D HeightmapScaleBias(Component->HeightmapScaleBias.Z, Component->HeightmapScaleBias.W);
-				GenerateHeightmapQuadsMip(Component->GetSectionBase(), HeightmapScaleBias, SubsectionSizeQuads, HeightmapReadTextureSize, HeightmapWriteTextureSize, InMipRender, TriangleList);
+				GenerateProceduralRenderQuadsMip(Component->GetSectionBase(), HeightmapScaleBias, SubsectionSizeQuads, HeightmapReadTextureSize, HeightmapWriteTextureSize, InMipRender, TriangleList);
 			}			
 		} break;
 
@@ -1108,25 +2091,25 @@
 
 	if (InMipRender > 0)
 	{
-		InShaderParams.CurrentMipHeightmapSize = HeightmapWriteTextureSize;
-		InShaderParams.ParentMipHeightmapSize = HeightmapReadTextureSize;
+		InShaderParams.CurrentMipSize = HeightmapWriteTextureSize;
+		InShaderParams.ParentMipSize = HeightmapReadTextureSize;
 	}
 
 	FMatrix ProjectionMatrix = AdjustProjectionMatrixForRHI(FTranslationMatrix(FVector(0, 0, 0)) *
 															FMatrix(FPlane(1.0f / (FMath::Max<uint32>(HeightmapWriteTextureSize.X, 1.f) / 2.0f), 0.0, 0.0f, 0.0f), FPlane(0.0f, -1.0f / (FMath::Max<uint32>(HeightmapWriteTextureSize.Y, 1.f) / 2.0f), 0.0f, 0.0f), FPlane(0.0f, 0.0f, 1.0f, 0.0f), FPlane(-1.0f, 1.0f, 0.0f, 1.0f)));
 
-	FLandscapeHeightmapProceduralRender_RenderThread ProceduralRender(InDebugName, InHeightmapRTWrite, HeightmapWriteTextureSize, HeightmapReadTextureSize, ProjectionMatrix, InShaderParams, InMipRender, TriangleList);
+	LandscapeProceduralHeightmapRender_RenderThread ProceduralRender(InDebugName, InHeightmapRTWrite, HeightmapWriteTextureSize, HeightmapReadTextureSize, ProjectionMatrix, InShaderParams, InMipRender, TriangleList);
 
-	ENQUEUE_RENDER_COMMAND(FDrawSceneCommand)(
+	ENQUEUE_RENDER_COMMAND(FDrawLandscapeProceduralHeightmapCommand)(
 		[ProceduralRender, ClearRT = InClearRTWrite](FRHICommandListImmediate& RHICmdList) mutable
 		{
 			ProceduralRender.Render(RHICmdList, ClearRT);
 		});
 	
-	PrintDebugRTHeightmap(InDebugName, InHeightmapRTWrite, InMipRender, InShaderParams.GenerateNormals);
+	PrintProceduralDebugRT(InDebugName, InHeightmapRTWrite, InMipRender, true, InShaderParams.GenerateNormals);
 }
 
-void ALandscape::GenerateHeightmapQuad(const FIntPoint& InVertexPosition, const float InVertexSize, const FVector2D& InUVStart, const FVector2D& InUVSize, TArray<FLandscapeProceduralTriangle>& OutTriangles) const
+void ALandscape::GenerateProceduralRenderQuad(const FIntPoint& InVertexPosition, float InVertexSize, const FVector2D& InUVStart, const FVector2D& InUVSize, TArray<FLandscapeProceduralTriangle>& OutTriangles) const
 {
 	FLandscapeProceduralTriangle Tri1;
 	
@@ -1151,15 +2134,16 @@
 	OutTriangles.Add(Tri2);
 }
 
-void ALandscape::GenerateHeightmapQuadsAtlas(const FIntPoint& InSectionBase, const FVector2D& InScaleBias, float InSubSectionSizeQuad, const FIntPoint& InReadSize, const FIntPoint& InWriteSize, TArray<FLandscapeProceduralTriangle>& OutTriangles) const
+void ALandscape::GenerateProceduralRenderQuadsAtlas(const FIntPoint& InSectionBase, const FVector2D& InScaleBias, float InSubSectionSizeQuad, const FIntPoint& InReadSize, const FIntPoint& InWriteSize, TArray<FLandscapeProceduralTriangle>& OutTriangles) const
 {
 	FIntPoint ComponentSectionBase = InSectionBase;
 	FIntPoint UVComponentSectionBase = InSectionBase;
 
+	int32 SubsectionSizeVerts = InSubSectionSizeQuad + 1;
 	int32 LocalComponentSizeQuad = InSubSectionSizeQuad * NumSubsections;
-	int32 SubsectionSizeVerts = InSubSectionSizeQuad + 1;
+	int32 LocalComponentSizeVerts = SubsectionSizeVerts * NumSubsections;
 
-	FVector2D HeightmapPositionOffset(FMath::RoundToInt(ComponentSectionBase.X / LocalComponentSizeQuad), FMath::RoundToInt(ComponentSectionBase.Y / LocalComponentSizeQuad));
+	FVector2D PositionOffset(FMath::RoundToInt(ComponentSectionBase.X / LocalComponentSizeQuad), FMath::RoundToInt(ComponentSectionBase.Y / LocalComponentSizeQuad));
 	FVector2D ComponentsPerTexture(FMath::RoundToInt(InWriteSize.X / LocalComponentSizeQuad), FMath::RoundToInt(InWriteSize.Y / LocalComponentSizeQuad));
 
 	if (InReadSize.X >= InWriteSize.X)
@@ -1168,16 +2152,16 @@
 		{
 			if (ComponentsPerTexture.X > 1.0f)
 			{
-				UVComponentSectionBase.X = HeightmapPositionOffset.X * (SubsectionSizeVerts * NumSubsections);
+				UVComponentSectionBase.X = PositionOffset.X * LocalComponentSizeVerts;
 			}
 			else
 			{
-				UVComponentSectionBase.X -= (UVComponentSectionBase.X + LocalComponentSizeQuad > InWriteSize.X) ? FMath::FloorToInt((HeightmapPositionOffset.X / ComponentsPerTexture.X)) * ComponentsPerTexture.X * LocalComponentSizeQuad : 0;
+				UVComponentSectionBase.X -= (UVComponentSectionBase.X + LocalComponentSizeQuad > InWriteSize.X) ? FMath::FloorToInt((PositionOffset.X / ComponentsPerTexture.X)) * ComponentsPerTexture.X * LocalComponentSizeQuad : 0;
 			}
 		}
 
-		ComponentSectionBase.X -= (ComponentSectionBase.X + LocalComponentSizeQuad > InWriteSize.X) ? FMath::FloorToInt((HeightmapPositionOffset.X / ComponentsPerTexture.X)) * ComponentsPerTexture.X * LocalComponentSizeQuad : 0;
-		HeightmapPositionOffset.X = ComponentSectionBase.X / LocalComponentSizeQuad;
+		ComponentSectionBase.X -= (ComponentSectionBase.X + LocalComponentSizeQuad > InWriteSize.X) ? FMath::FloorToInt((PositionOffset.X / ComponentsPerTexture.X)) * ComponentsPerTexture.X * LocalComponentSizeQuad : 0;
+		PositionOffset.X = ComponentSectionBase.X / LocalComponentSizeQuad;
 	}
 
 	if (InReadSize.Y >= InWriteSize.Y)
@@ -1186,22 +2170,22 @@
 		{
 			if (ComponentsPerTexture.Y > 1.0f)
 			{
-				UVComponentSectionBase.Y = HeightmapPositionOffset.Y * (SubsectionSizeVerts * NumSubsections);
+				UVComponentSectionBase.Y = PositionOffset.Y * LocalComponentSizeVerts;
 			}
 			else
 			{
-				UVComponentSectionBase.Y -= (UVComponentSectionBase.Y + LocalComponentSizeQuad > InWriteSize.Y) ? FMath::FloorToInt((HeightmapPositionOffset.Y / ComponentsPerTexture.Y)) * ComponentsPerTexture.Y * LocalComponentSizeQuad : 0;
+				UVComponentSectionBase.Y -= (UVComponentSectionBase.Y + LocalComponentSizeQuad > InWriteSize.Y) ? FMath::FloorToInt((PositionOffset.Y / ComponentsPerTexture.Y)) * ComponentsPerTexture.Y * LocalComponentSizeQuad : 0;
 			}
 		}
 
-		ComponentSectionBase.Y -= (ComponentSectionBase.Y + LocalComponentSizeQuad > InWriteSize.Y) ? FMath::FloorToInt((HeightmapPositionOffset.Y / ComponentsPerTexture.Y)) * ComponentsPerTexture.Y * LocalComponentSizeQuad : 0;
-		HeightmapPositionOffset.Y = ComponentSectionBase.Y / LocalComponentSizeQuad;
+		ComponentSectionBase.Y -= (ComponentSectionBase.Y + LocalComponentSizeQuad > InWriteSize.Y) ? FMath::FloorToInt((PositionOffset.Y / ComponentsPerTexture.Y)) * ComponentsPerTexture.Y * LocalComponentSizeQuad : 0;
+		PositionOffset.Y = ComponentSectionBase.Y / LocalComponentSizeQuad;
 	}
 
-	ComponentSectionBase.X = HeightmapPositionOffset.X * (SubsectionSizeVerts * NumSubsections);
-	ComponentSectionBase.Y = HeightmapPositionOffset.Y * (SubsectionSizeVerts * NumSubsections);
+	ComponentSectionBase.X = PositionOffset.X * LocalComponentSizeVerts;
+	ComponentSectionBase.Y = PositionOffset.Y * LocalComponentSizeVerts;
 
-	FVector2D HeightmapUVSize((float)SubsectionSizeVerts / (float)InReadSize.X, (float)SubsectionSizeVerts / (float)InReadSize.Y);
+	FVector2D UVSize((float)SubsectionSizeVerts / (float)InReadSize.X, (float)SubsectionSizeVerts / (float)InReadSize.Y);
 	FIntPoint SubSectionSectionBase;
 
 	for (int8 SubY = 0; SubY < NumSubsections; ++SubY)
@@ -1211,44 +2195,46 @@
 			SubSectionSectionBase.X = ComponentSectionBase.X + SubsectionSizeVerts * SubX;
 			SubSectionSectionBase.Y = ComponentSectionBase.Y + SubsectionSizeVerts * SubY;
 
-			// Offset for this component's data in heightmap texture
-			FVector2D HeightmapUVStart;
+			// Offset for this component's data in texture
+			FVector2D UVStart;
 
 			if (InReadSize.X >= InWriteSize.X)
 			{
-				HeightmapUVStart.X = ((float)UVComponentSectionBase.X / (float)InReadSize.X) + HeightmapUVSize.X * (float)SubX;
+				UVStart.X = ((float)UVComponentSectionBase.X / (float)InReadSize.X) + UVSize.X * (float)SubX;
 			}
 			else
 			{
-				HeightmapUVStart.X = InScaleBias.X + HeightmapUVSize.X * (float)SubX;
+				UVStart.X = InScaleBias.X + UVSize.X * (float)SubX;
 			}
 
 			if (InReadSize.Y >= InWriteSize.Y)
 			{
-				HeightmapUVStart.Y = ((float)UVComponentSectionBase.Y / (float)InReadSize.Y) + HeightmapUVSize.Y * (float)SubY;
+				UVStart.Y = ((float)UVComponentSectionBase.Y / (float)InReadSize.Y) + UVSize.Y * (float)SubY;
 			}
 			else
 			{
-				HeightmapUVStart.Y = InScaleBias.Y + HeightmapUVSize.Y * (float)SubY;
+				UVStart.Y = InScaleBias.Y + UVSize.Y * (float)SubY;
 			}
 
-			GenerateHeightmapQuad(SubSectionSectionBase, SubsectionSizeVerts, HeightmapUVStart, HeightmapUVSize, OutTriangles);
+			GenerateProceduralRenderQuad(SubSectionSectionBase, SubsectionSizeVerts, UVStart, UVSize, OutTriangles);
 		}
 	}
 }
 
-void ALandscape::GenerateHeightmapQuadsMip(const FIntPoint& InSectionBase, const FVector2D& InScaleBias, float InSubSectionSizeQuad, const FIntPoint& InReadSize, const FIntPoint& InWriteSize, int32 CurrentMip, TArray<FLandscapeProceduralTriangle>& OutTriangles) const
+void ALandscape::GenerateProceduralRenderQuadsMip(const FIntPoint& InSectionBase, const FVector2D& InScaleBias, float InSubSectionSizeQuad, const FIntPoint& InReadSize, const FIntPoint& InWriteSize, uint8 InCurrentMip, TArray<FLandscapeProceduralTriangle>& OutTriangles) const
 {
+	int32 SubsectionSizeVerts = InSubSectionSizeQuad + 1;
 	int32 LocalComponentSizeQuad = InSubSectionSizeQuad * NumSubsections;
-	int32 SubsectionSizeVerts = InSubSectionSizeQuad + 1;
-	int32 MipSubsectionSizeVerts = SubsectionSizeVerts >> CurrentMip;
-
-	FVector2D HeightmapPositionOffset(FMath::RoundToInt(InSectionBase.X / LocalComponentSizeQuad), FMath::RoundToInt(InSectionBase.Y / LocalComponentSizeQuad));
+	int32 LocalComponentSizeVerts = SubsectionSizeVerts * NumSubsections;	
+	int32 MipSubsectionSizeVerts = SubsectionSizeVerts >> InCurrentMip;
+	int32 MipLocalComponentSizeVerts = MipSubsectionSizeVerts * NumSubsections;	
+	
+	FVector2D PositionOffset(FMath::RoundToInt(InSectionBase.X / LocalComponentSizeQuad), FMath::RoundToInt(InSectionBase.Y / LocalComponentSizeQuad));
 	FVector2D ComponentsPerTexture(FMath::RoundToInt(InWriteSize.X / LocalComponentSizeQuad), FMath::RoundToInt(InWriteSize.Y / LocalComponentSizeQuad));
 
-	FIntPoint ComponentSectionBase(HeightmapPositionOffset.X * (MipSubsectionSizeVerts * NumSubsections), HeightmapPositionOffset.Y * (MipSubsectionSizeVerts * NumSubsections));
-	FIntPoint UVComponentSectionBase(HeightmapPositionOffset.X * (SubsectionSizeVerts * NumSubsections), HeightmapPositionOffset.Y * (SubsectionSizeVerts * NumSubsections));
-	FVector2D HeightmapUVSize((float)(SubsectionSizeVerts >> (CurrentMip - 1)) / (float)InReadSize.X, (float)(SubsectionSizeVerts >> (CurrentMip - 1)) / (float)InReadSize.Y);
+	FIntPoint ComponentSectionBase(PositionOffset.X * MipLocalComponentSizeVerts, PositionOffset.Y * MipLocalComponentSizeVerts);
+	FIntPoint UVComponentSectionBase(PositionOffset.X * LocalComponentSizeVerts, PositionOffset.Y * LocalComponentSizeVerts);
+	FVector2D UVSize((float)(SubsectionSizeVerts >> (InCurrentMip - 1)) / (float)InReadSize.X, (float)(SubsectionSizeVerts >> (InCurrentMip - 1)) / (float)InReadSize.Y);
 	FIntPoint SubSectionSectionBase;
 
 	for (int8 SubY = 0; SubY < NumSubsections; ++SubY)
@@ -1258,32 +2244,25 @@
 			SubSectionSectionBase.X = ComponentSectionBase.X + MipSubsectionSizeVerts * SubX;
 			SubSectionSectionBase.Y = ComponentSectionBase.Y + MipSubsectionSizeVerts * SubY;
 
-			// Offset for this component's data in heightmap texture
-			FVector2D HeightmapUVStart;
-			HeightmapUVStart.X = ((float)(UVComponentSectionBase.X >> (CurrentMip - 1)) / (float)InReadSize.X) + HeightmapUVSize.X * (float)SubX;
-			HeightmapUVStart.Y = ((float)(UVComponentSectionBase.Y >> (CurrentMip - 1)) / (float)InReadSize.Y) + HeightmapUVSize.Y * (float)SubY;
+			// Offset for this component's data in texture
+			FVector2D UVStart(((float)(UVComponentSectionBase.X >> (InCurrentMip - 1)) / (float)InReadSize.X) + UVSize.X * (float)SubX, ((float)(UVComponentSectionBase.Y >> (InCurrentMip - 1)) / (float)InReadSize.Y) + UVSize.Y * (float)SubY);
 
-			GenerateHeightmapQuad(SubSectionSectionBase, MipSubsectionSizeVerts, HeightmapUVStart, HeightmapUVSize, OutTriangles);
+			GenerateProceduralRenderQuad(SubSectionSectionBase, MipSubsectionSizeVerts, UVStart, UVSize, OutTriangles);
 		}
 	}
 }
 
-void ALandscape::GenerateHeightmapQuadsAtlasToNonAtlas(const FIntPoint& InSectionBase, const FVector2D& InScaleBias, float InSubSectionSizeQuad, const FIntPoint& InHeightmapReadTextureSize, const FIntPoint& InHeightmapWriteTextureSize, TArray<FLandscapeProceduralTriangle>& OutTriangles) const
+void ALandscape::GenerateProceduralRenderQuadsAtlasToNonAtlas(const FIntPoint& InSectionBase, const FVector2D& InScaleBias, float InSubSectionSizeQuad, const FIntPoint& InReadSize, const FIntPoint& InWriteSize, TArray<struct FLandscapeProceduralTriangle>& OutTriangles) const
 {
-	FIntPoint ComponentSectionBase = InSectionBase;
+	int32 SubsectionSizeVerts = InSubSectionSizeQuad + 1;
 	int32 LocalComponentSizeQuad = InSubSectionSizeQuad * NumSubsections;
-	int32 HeightmapPositionOffsetX = ComponentSectionBase.X / LocalComponentSizeQuad;
-	int32 HeightmapPositionOffsetY = ComponentSectionBase.Y / LocalComponentSizeQuad;
-	int32 SubsectionSizeVerts = InSubSectionSizeQuad + 1;
+	int32 LocalComponentSizeVerts = SubsectionSizeVerts * NumSubsections;
 
-	FIntPoint UVComponentSectionBase = InSectionBase;
-	UVComponentSectionBase.X = HeightmapPositionOffsetX * (SubsectionSizeVerts * NumSubsections);
-	UVComponentSectionBase.Y = HeightmapPositionOffsetY * (SubsectionSizeVerts * NumSubsections);
+	FVector2D PositionOffset(FMath::RoundToInt(InSectionBase.X / LocalComponentSizeQuad), FMath::RoundToInt(InSectionBase.Y / LocalComponentSizeQuad));
+	FIntPoint ComponentSectionBase(PositionOffset.X * LocalComponentSizeQuad, PositionOffset.Y * LocalComponentSizeQuad);
+	FIntPoint UVComponentSectionBase(PositionOffset.X * SubsectionSizeVerts, PositionOffset.Y * SubsectionSizeVerts);
+	FVector2D UVSize((float)SubsectionSizeVerts / (float)InReadSize.X, (float)SubsectionSizeVerts / (float)InReadSize.Y);
 
-	ComponentSectionBase.X = HeightmapPositionOffsetX * (InSubSectionSizeQuad * NumSubsections);
-	ComponentSectionBase.Y = HeightmapPositionOffsetY * (InSubSectionSizeQuad * NumSubsections);
-
-	FVector2D HeightmapUVSize((float)SubsectionSizeVerts / (float)InHeightmapReadTextureSize.X, (float)SubsectionSizeVerts / (float)InHeightmapReadTextureSize.Y);
 	FIntPoint SubSectionSectionBase;
 
 	for (int8 SubY = 0; SubY < NumSubsections; ++SubY)
@@ -1293,48 +2272,44 @@
 			SubSectionSectionBase.X = ComponentSectionBase.X + InSubSectionSizeQuad * SubX;
 			SubSectionSectionBase.Y = ComponentSectionBase.Y + InSubSectionSizeQuad * SubY;
 
-			// Offset for this component's data in heightmap texture
-			FVector2D HeightmapUVStart;
+			// Offset for this component's data in texture
+			FVector2D UVStart;
 
-			if (InHeightmapReadTextureSize.X >= InHeightmapWriteTextureSize.X)
+			if (InReadSize.X >= InWriteSize.X)
 			{
-				HeightmapUVStart.X = ((float)UVComponentSectionBase.X / (float)InHeightmapReadTextureSize.X) + HeightmapUVSize.X * (float)SubX;
+				UVStart.X = ((float)UVComponentSectionBase.X / (float)InReadSize.X) + UVSize.X * (float)SubX;
 			}
 			else
 			{
-				HeightmapUVStart.X = InScaleBias.X + HeightmapUVSize.X * (float)SubX;
+				UVStart.X = InScaleBias.X + UVSize.X * (float)SubX;
 			}
 
-			if (InHeightmapReadTextureSize.Y >= InHeightmapWriteTextureSize.Y)
+			if (InReadSize.Y >= InWriteSize.Y)
 			{
-				HeightmapUVStart.Y = ((float)UVComponentSectionBase.Y / (float)InHeightmapReadTextureSize.Y) + HeightmapUVSize.Y * (float)SubY;
+				UVStart.Y = ((float)UVComponentSectionBase.Y / (float)InReadSize.Y) + UVSize.Y * (float)SubY;
 			}
 			else
 			{
-				HeightmapUVStart.Y = InScaleBias.Y + HeightmapUVSize.Y * (float)SubY;
+				UVStart.Y = InScaleBias.Y + UVSize.Y * (float)SubY;
 			}
 
-			GenerateHeightmapQuad(SubSectionSectionBase, SubsectionSizeVerts, HeightmapUVStart, HeightmapUVSize, OutTriangles);
+			GenerateProceduralRenderQuad(SubSectionSectionBase, SubsectionSizeVerts, UVStart, UVSize, OutTriangles);
 		}
 	}
 }
 
-void ALandscape::GenerateHeightmapQuadsNonAtlas(const FIntPoint& InSectionBase, const FVector2D& InScaleBias, float InSubSectionSizeQuad, const FIntPoint& InHeightmapReadTextureSize, const FIntPoint& InHeightmapWriteTextureSize, TArray<FLandscapeProceduralTriangle>& OutTriangles) const
+void ALandscape::GenerateProceduralRenderQuadsNonAtlas(const FIntPoint& InSectionBase, const FVector2D& InScaleBias, float InSubSectionSizeQuad, const FIntPoint& InReadSize, const FIntPoint& InWriteSize, TArray<struct FLandscapeProceduralTriangle>& OutTriangles) const
 {
 	// We currently only support drawing in non atlas mode with the same texture size
-	check(InHeightmapReadTextureSize.X == InHeightmapWriteTextureSize.X && InHeightmapReadTextureSize.Y == InHeightmapWriteTextureSize.Y);
+	check(InReadSize.X == InWriteSize.X && InReadSize.Y == InWriteSize.Y);
 
+	int32 SubsectionSizeVerts = InSubSectionSizeQuad + 1;
+	int32 LocalComponentSizeQuad = InSubSectionSizeQuad * NumSubsections;
+
+	FVector2D PositionOffset(FMath::RoundToInt(InSectionBase.X / LocalComponentSizeQuad), FMath::RoundToInt(InSectionBase.Y / LocalComponentSizeQuad));
 	FIntPoint ComponentSectionBase = InSectionBase;
-	int32 LocalComponentSizeQuad = InSubSectionSizeQuad * NumSubsections;
-	int32 HeightmapPositionOffsetX = ComponentSectionBase.X / LocalComponentSizeQuad;
-	int32 HeightmapPositionOffsetY = ComponentSectionBase.Y / LocalComponentSizeQuad;
-	int32 SubsectionSizeVerts = InSubSectionSizeQuad + 1;
-
-	FIntPoint UVComponentSectionBase = InSectionBase;
-	UVComponentSectionBase.X = HeightmapPositionOffsetX * (InSubSectionSizeQuad * NumSubsections);
-	UVComponentSectionBase.Y = HeightmapPositionOffsetY * (InSubSectionSizeQuad * NumSubsections);
-
-	FVector2D HeightmapUVSize((float)SubsectionSizeVerts / (float)InHeightmapReadTextureSize.X, (float)SubsectionSizeVerts / (float)InHeightmapReadTextureSize.Y);
+	FIntPoint UVComponentSectionBase(PositionOffset.X * LocalComponentSizeQuad, PositionOffset.Y * LocalComponentSizeQuad);
+	FVector2D UVSize((float)SubsectionSizeVerts / (float)InReadSize.X, (float)SubsectionSizeVerts / (float)InReadSize.Y);
 	FIntPoint SubSectionSectionBase;
 
 	for (int8 SubY = 0; SubY < NumSubsections; ++SubY)
@@ -1344,25 +2319,23 @@
 			SubSectionSectionBase.X = ComponentSectionBase.X + InSubSectionSizeQuad * SubX;
 			SubSectionSectionBase.Y = ComponentSectionBase.Y + InSubSectionSizeQuad * SubY;
 
-			// Offset for this component's data in heightmap texture
-			FVector2D HeightmapUVStart(((float)UVComponentSectionBase.X / (float)InHeightmapReadTextureSize.X) + HeightmapUVSize.X * (float)SubX, ((float)UVComponentSectionBase.Y / (float)InHeightmapReadTextureSize.Y) + HeightmapUVSize.Y * (float)SubY);
-			GenerateHeightmapQuad(SubSectionSectionBase, SubsectionSizeVerts, HeightmapUVStart, HeightmapUVSize, OutTriangles);
+			// Offset for this component's data in texture
+			FVector2D UVStart(((float)UVComponentSectionBase.X / (float)InReadSize.X) + UVSize.X * (float)SubX, ((float)UVComponentSectionBase.Y / (float)InReadSize.Y) + UVSize.Y * (float)SubY);
+			GenerateProceduralRenderQuad(SubSectionSectionBase, SubsectionSizeVerts, UVStart, UVSize, OutTriangles);
 		}
 	}
 }
 
-void ALandscape::GenerateHeightmapQuadsNonAtlasToAtlas(const FIntPoint& InSectionBase, const FVector2D& InScaleBias, float InSubSectionSizeQuad, const FIntPoint& InHeightmapReadTextureSize, const FIntPoint& InHeightmapWriteTextureSize, TArray<FLandscapeProceduralTriangle>& OutTriangles) const
+void ALandscape::GenerateProceduralRenderQuadsNonAtlasToAtlas(const FIntPoint& InSectionBase, const FVector2D& InScaleBias, float InSubSectionSizeQuad, const FIntPoint& InReadSize, const FIntPoint& InWriteSize, TArray<struct FLandscapeProceduralTriangle>& OutTriangles) const
 {
-	FIntPoint ComponentSectionBase = InSectionBase;
+	int32 SubsectionSizeVerts = InSubSectionSizeQuad + 1;
 	int32 LocalComponentSizeQuad = InSubSectionSizeQuad * NumSubsections;
-	int32 HeightmapPositionOffsetX = ComponentSectionBase.X / LocalComponentSizeQuad;
-	int32 HeightmapPositionOffsetY = ComponentSectionBase.Y / LocalComponentSizeQuad;
-	int32 SubsectionSizeVerts = InSubSectionSizeQuad + 1;
+	int32 LocalComponentSizeVerts = SubsectionSizeVerts * NumSubsections;
 
-	ComponentSectionBase.X = HeightmapPositionOffsetX * (SubsectionSizeVerts * NumSubsections);
-	ComponentSectionBase.Y = HeightmapPositionOffsetY * (SubsectionSizeVerts * NumSubsections);
+	FVector2D PositionOffset(FMath::RoundToInt(InSectionBase.X / LocalComponentSizeQuad), FMath::RoundToInt(InSectionBase.Y / LocalComponentSizeQuad));
+	FIntPoint ComponentSectionBase(PositionOffset.X * LocalComponentSizeVerts, PositionOffset.Y * LocalComponentSizeVerts);
+	FVector2D UVSize((float)SubsectionSizeVerts / (float)InReadSize.X, (float)SubsectionSizeVerts / (float)InReadSize.Y);
 
-	FVector2D HeightmapUVSize((float)SubsectionSizeVerts / (float)InHeightmapReadTextureSize.X, (float)SubsectionSizeVerts / (float)InHeightmapReadTextureSize.Y);
 	FIntPoint SubSectionSectionBase;
 
 	for (int8 SubY = 0; SubY < NumSubsections; ++SubY)
@@ -1372,17 +2345,17 @@
 			SubSectionSectionBase.X = ComponentSectionBase.X + SubsectionSizeVerts * SubX;
 			SubSectionSectionBase.Y = ComponentSectionBase.Y + SubsectionSizeVerts * SubY;
 
-			// Offset for this component's data in heightmap texture
-			float HeightmapScaleBiasZ = (float)InSectionBase.X / (float)InHeightmapReadTextureSize.X;
-			float HeightmapScaleBiasW = (float)InSectionBase.Y / (float)InHeightmapReadTextureSize.Y;
-			FVector2D HeightmapUVStart(HeightmapScaleBiasZ + ((float)InSubSectionSizeQuad / (float)InHeightmapReadTextureSize.X) * (float)SubX, HeightmapScaleBiasW + ((float)InSubSectionSizeQuad / (float)InHeightmapReadTextureSize.Y) * (float)SubY);
+			// Offset for this component's data in texture
+			float ScaleBiasZ = (float)InSectionBase.X / (float)InReadSize.X;
+			float ScaleBiasW = (float)InSectionBase.Y / (float)InReadSize.Y;
+			FVector2D UVStart(ScaleBiasZ + ((float)InSubSectionSizeQuad / (float)InReadSize.X) * (float)SubX, ScaleBiasW + ((float)InSubSectionSizeQuad / (float)InReadSize.Y) * (float)SubY);
 
-			GenerateHeightmapQuad(SubSectionSectionBase, SubsectionSizeVerts, HeightmapUVStart, HeightmapUVSize, OutTriangles);
+			GenerateProceduralRenderQuad(SubSectionSectionBase, SubsectionSizeVerts, UVStart, UVSize, OutTriangles);
 		}
 	}
 }
 
-void ALandscape::PrintDebugHeightData(const FString& InContext, const TArray<FColor>& InHeightmapData, const FIntPoint& InDataSize, int32 InMipRender, bool InOutputNormals) const
+void ALandscape::PrintProceduralDebugHeightData(const FString& InContext, const TArray<FColor>& InHeightmapData, const FIntPoint& InDataSize, uint8 InMipRender, bool InOutputNormals) const
 {
 	bool DisplayDebugPrint = CVarOutputProceduralRTContent.GetValueOnAnyThread() == 1 ? true : false;
 	bool DisplayHeightAsDelta = false;
@@ -1484,18 +2457,56 @@
 	}
 }
 
-void ALandscape::PrintDebugRTHeightmap(FString Context, UTextureRenderTarget2D* InDebugRT, int32 InMipRender, bool InOutputNormals) const
+void ALandscape::PrintProceduralDebugWeightData(const FString& InContext, const TArray<FColor>& InWeightmapData, const FIntPoint& InDataSize, uint8 InMipRender) const
 {
-	bool DisplayDebugPrint = CVarOutputProceduralRTContent.GetValueOnAnyThread() == 1 ? true : false;
+	bool DisplayDebugPrint = (CVarOutputProceduralRTContent.GetValueOnAnyThread() == 1 || CVarOutputProceduralWeightmapsRTContent.GetValueOnAnyThread() == 1) ? true : false;
 
 	if (!DisplayDebugPrint)
 	{
 		return;
 	}
 
+	UE_LOG(LogLandscapeBP, Display, TEXT("Context: %s"), *InContext);
+
+	int32 MipSize = ((SubsectionSizeQuads + 1) >> InMipRender);
+
+	for (int32 Y = 0; Y < InDataSize.Y; ++Y)
+	{
+		FString WeightmapOutput;
+
+		for (int32 X = 0; X < InDataSize.X; ++X)
+		{
+			const FColor& Weight = InWeightmapData[X + Y * InDataSize.X];
+
+			if (X > 0 && MipSize > 0 && X % MipSize == 0)
+			{
+				WeightmapOutput += FString::Printf(TEXT("  "));
+			}
+
+			WeightmapOutput += FString::Printf(TEXT("%s "), *Weight.ToString());
+		}
+
+		if (Y > 0 && MipSize > 0 && Y % MipSize == 0)
+		{
+			UE_LOG(LogLandscapeBP, Display, TEXT(""));
+		}
+
+		UE_LOG(LogLandscapeBP, Display, TEXT("%s"), *WeightmapOutput);
+	}
+}
+
+void ALandscape::PrintProceduralDebugRT(const FString& InContext, UTextureRenderTarget2D* InDebugRT, uint8 InMipRender, bool InOutputHeight, bool InOutputNormals) const
+{
+	bool DisplayDebugPrint = (CVarOutputProceduralRTContent.GetValueOnAnyThread() == 1 || CVarOutputProceduralWeightmapsRTContent.GetValueOnAnyThread() == 1) ? true : false;
+
+	if (!DisplayDebugPrint)
+	{
+		return;
+	}
+
 	FTextureRenderTargetResource* RenderTargetResource = InDebugRT->GameThread_GetRenderTargetResource();
-	ENQUEUE_RENDER_COMMAND(HeightmapRTCanvasRenderTargetResolveCommand)(
-		[RenderTargetResource](FRHICommandListImmediate& RHICmdList)
+	ENQUEUE_RENDER_COMMAND(FProceduralDebugRenderTargetResolveCommand)(
+		[RenderTargetResource](FRHICommandListImmediate& RHICmdList) mutable
 		{
 			// Copy (resolve) the rendered image from the frame buffer to its render target texture
 			RHICmdList.CopyToResolveTarget(RenderTargetResource->GetRenderTargetTexture(), RenderTargetResource->TextureRHI, FResolveParams());
@@ -1509,34 +2520,62 @@
 
 	FReadSurfaceDataFlags Flags(RCM_UNorm, CubeFace_MAX);
 
-	TArray<FColor> OutputRTHeightmap;
-	OutputRTHeightmap.Reserve(SampleRect.Width() * SampleRect.Height());
+	TArray<FColor> OutputRT;
+	OutputRT.Reserve(SampleRect.Width() * SampleRect.Height());
 
-	InDebugRT->GameThread_GetRenderTargetResource()->ReadPixels(OutputRTHeightmap, Flags, SampleRect);
+	InDebugRT->GameThread_GetRenderTargetResource()->ReadPixels(OutputRT, Flags, SampleRect);
 
-	PrintDebugHeightData(Context, OutputRTHeightmap, FIntPoint(SampleRect.Width(), SampleRect.Height()), InMipRender, InOutputNormals);
+	if (InOutputHeight)
+	{
+		PrintProceduralDebugHeightData(InContext, OutputRT, FIntPoint(SampleRect.Width(), SampleRect.Height()), InMipRender, InOutputNormals);
+	}
+	else
+	{
+		PrintProceduralDebugWeightData(InContext, OutputRT, FIntPoint(SampleRect.Width(), SampleRect.Height()), InMipRender);
+	}
 }
 
-void ALandscape::RegenerateProceduralHeightmaps()
+void ALandscape::PrintProceduralDebugTextureResource(const FString& InContext, FTextureResource* InTextureResource, uint8 InMipRender, bool InOutputHeight, bool InOutputNormals) const
 {
-	SCOPE_CYCLE_COUNTER(STAT_LandscapeRegenerateProceduralHeightmaps);
+	bool DisplayDebugPrint = (CVarOutputProceduralRTContent.GetValueOnAnyThread() == 1 || CVarOutputProceduralWeightmapsRTContent.GetValueOnAnyThread() == 1) ? true : false;
 
-	ULandscapeInfo* Info = GetLandscapeInfo();
-
-	if (ProceduralContentUpdateFlags == 0 || Info == nullptr)
+	if (!DisplayDebugPrint)
 	{
 		return;
 	}
 
-	TArray<ALandscapeProxy*> AllLandscapes;
-	AllLandscapes.Add(this);
+	int32 MinX, MinY, MaxX, MaxY;
+	const ULandscapeInfo* LandscapeInfo = GetLandscapeInfo();
+	LandscapeInfo->GetLandscapeExtent(MinX, MinY, MaxX, MaxY);
+	FIntRect SampleRect = FIntRect(0, 0, InTextureResource->GetSizeX(), InTextureResource->GetSizeY());
 
-	for (auto& It : Info->Proxies)
+	TArray<FColor> OutputTexels;
+	OutputTexels.Reserve(SampleRect.Width() * SampleRect.Height());
+
+	FReadSurfaceDataFlags Flags(RCM_UNorm, CubeFace_MAX);
+	Flags.SetMip(InMipRender);
+
+	ENQUEUE_RENDER_COMMAND(FProceduralDebugReadSurfaceCommand)(
+		[SourceTextureRHI = InTextureResource->TextureRHI, SampleRect = SampleRect, OutData = &OutputTexels, ReadFlags = Flags](FRHICommandListImmediate& RHICmdList) mutable
 	{
-		AllLandscapes.Add(It);
+		RHICmdList.ReadSurfaceData(SourceTextureRHI, SampleRect, *OutData, ReadFlags);
+	});
+
+	FlushRenderingCommands();
+
+	if (InOutputHeight)
+	{
+		PrintProceduralDebugHeightData(InContext, OutputTexels, FIntPoint(SampleRect.Width(), SampleRect.Height()), InMipRender, InOutputNormals);
 	}
+	else
+	{
+		PrintProceduralDebugWeightData(InContext, OutputTexels, FIntPoint(SampleRect.Width(), SampleRect.Height()), InMipRender);
+	}
+}
 
-	for (ALandscapeProxy* Landscape : AllLandscapes)
+bool ALandscape::AreHeightmapTextureResourcesReady(const TArray<ALandscapeProxy*>& InAllLandscapes) const
+{
+	for (ALandscapeProxy* Landscape : InAllLandscapes)
 	{
 		for (auto& ItLayerDataPair : Landscape->ProceduralLayersData)
 		{
@@ -1547,7 +2586,7 @@
 
 				if (!LayerHeightmap->IsAsyncCacheComplete() || !OriginalHeightmap->IsFullyStreamedIn())
 				{
-					return;
+					return false;
 				}
 
 				if (LayerHeightmap->Resource == nullptr)
@@ -1563,10 +2602,37 @@
 
 				if (!LayerHeightmap->Resource->IsInitialized() || !LayerHeightmap->IsFullyStreamedIn())
 				{
-					return;
+					return false;
 				}
 			}
 		}
+	}
+
+	return true;
+}
+
+void ALandscape::RegenerateProceduralHeightmaps()
+{
+	SCOPE_CYCLE_COUNTER(STAT_LandscapeRegenerateProceduralHeightmaps);
+
+	ULandscapeInfo* Info = GetLandscapeInfo();
+
+	if (ProceduralContentUpdateFlags == 0 || Info == nullptr)
+	{
+		return;
+	}
+
+	TArray<ALandscapeProxy*> AllLandscapes;
+	AllLandscapes.Add(this);
+
+	for (auto& It : Info->Proxies)
+	{
+		AllLandscapes.Add(It);
+	}
+
+	if (!AreHeightmapTextureResourcesReady(AllLandscapes))
+	{
+		return;
 	}
 
 	TArray<ULandscapeComponent*> AllLandscapeComponents;
@@ -1581,20 +2647,20 @@
 		FLandscapeHeightmapProceduralShaderParameters ShaderParams;
 
 		bool FirstLayer = true;
-		UTextureRenderTarget2D* CombinedHeightmapAtlasRT = HeightmapRTList[EHeightmapRTType::LandscapeSizeCombinedAtlas];
-		UTextureRenderTarget2D* CombinedHeightmapNonAtlasRT = HeightmapRTList[EHeightmapRTType::LandscapeSizeCombinedNonAtlas];
-		UTextureRenderTarget2D* LandscapeScratchRT1 = HeightmapRTList[EHeightmapRTType::LandscapeSizeScratch1];
-		UTextureRenderTarget2D* LandscapeScratchRT2 = HeightmapRTList[EHeightmapRTType::LandscapeSizeScratch2];
-		UTextureRenderTarget2D* LandscapeScratchRT3 = HeightmapRTList[EHeightmapRTType::LandscapeSizeScratch3];
+		UTextureRenderTarget2D* CombinedHeightmapAtlasRT = HeightmapRTList[EHeightmapRTType::HeightmapRT_CombinedAtlas];
+		UTextureRenderTarget2D* CombinedHeightmapNonAtlasRT = HeightmapRTList[EHeightmapRTType::HeightmapRT_CombinedNonAtlas];
+		UTextureRenderTarget2D* LandscapeScratchRT1 = HeightmapRTList[EHeightmapRTType::HeightmapRT_Scratch1];
+		UTextureRenderTarget2D* LandscapeScratchRT2 = HeightmapRTList[EHeightmapRTType::HeightmapRT_Scratch2];
+		UTextureRenderTarget2D* LandscapeScratchRT3 = HeightmapRTList[EHeightmapRTType::HeightmapRT_Scratch3];
 
-		bool OutputDebugName = CVarOutputProceduralDebugDrawCallName.GetValueOnAnyThread() == 1 || CVarOutputProceduralRTContent.GetValueOnAnyThread() == 1 ? true : false;
+		bool OutputDebugName = (CVarOutputProceduralDebugDrawCallName.GetValueOnAnyThread() == 1 || CVarOutputProceduralRTContent.GetValueOnAnyThread() == 1) ? true : false;
 
 		for (FProceduralLayer& Layer : ProceduralLayers)
 		{
 			//Draw Layer heightmap to Combined RT Atlas
 			ShaderParams.ApplyLayerModifiers = true;
 			ShaderParams.LayerVisible = Layer.Visible;
-			ShaderParams.LayerWeight = Layer.Weight;
+			ShaderParams.LayerAlpha = Layer.HeightmapAlpha;
 
 			for (ALandscapeProxy* Landscape : AllLandscapes)
 			{
@@ -1607,9 +2673,9 @@
 						FRenderDataPerHeightmap& HeightmapRenderData = *Landscape->RenderDataPerHeightmap.Find(ItPair.Key);
 						UTexture2D* Heightmap = ItPair.Value;
 
-						CopyProceduralTargetToResolveTarget(Heightmap, LandscapeScratchRT1, nullptr, HeightmapRenderData.TopLeftSectionBase, 0);
+						CopyProceduralTexture(Heightmap, LandscapeScratchRT1, nullptr, HeightmapRenderData.TopLeftSectionBase);
 
-						PrintDebugRTHeightmap(OutputDebugName ? FString::Printf(TEXT("LS Height: %s Component %s += -> CombinedAtlas %s"), *Layer.Name.ToString(), *Heightmap->GetName(), *LandscapeScratchRT1->GetName()) : TEXT(""), LandscapeScratchRT1);
+						PrintProceduralDebugRT(OutputDebugName ? FString::Printf(TEXT("LS Height: %s Component %s += -> CombinedAtlas %s"), *Layer.Name.ToString(), *Heightmap->GetName(), *LandscapeScratchRT1->GetName()) : TEXT(""), LandscapeScratchRT1);
 					}
 				}
 			}
@@ -1629,8 +2695,8 @@
 				// Draw each Combined RT into a Non Atlas RT format to be use as base for all brush rendering
 				if (Layer.Brushes.Num() > 0)
 				{
-					CopyProceduralTargetToResolveTarget(CombinedHeightmapNonAtlasRT, LandscapeScratchRT1, nullptr, FIntPoint(0,0), 0);
-					PrintDebugRTHeightmap(OutputDebugName ? FString::Printf(TEXT("LS Height: %s Component %s += -> CombinedNonAtlas %s"), *Layer.Name.ToString(), *CombinedHeightmapNonAtlasRT->GetName(), *LandscapeScratchRT1->GetName()) : TEXT(""), LandscapeScratchRT1);
+					CopyProceduralTexture(CombinedHeightmapNonAtlasRT, LandscapeScratchRT1);
+					PrintProceduralDebugRT(OutputDebugName ? FString::Printf(TEXT("LS Height: %s Component %s += -> CombinedNonAtlas %s"), *Layer.Name.ToString(), *CombinedHeightmapNonAtlasRT->GetName(), *LandscapeScratchRT1->GetName()) : TEXT(""), LandscapeScratchRT1);
 				}
 
 				// Draw each brushes				
@@ -1641,13 +2707,11 @@
 
 					FLandscapeProceduralLayerBrush& Brush = Layer.Brushes[Layer.HeightmapBrushOrderIndices[i]];
 
-					if (Brush.BPCustomBrush == nullptr)
+					if (Brush.BPCustomBrush == nullptr || !Brush.BPCustomBrush->IsAffectingHeightmap())
 					{
 						continue;
 					}
 
-					check(Brush.BPCustomBrush->IsAffectingHeightmap());
-
 					if (!Brush.IsInitialized())
 					{
 						Brush.Initialize(GetBoundingRect(), FIntPoint(CombinedHeightmapNonAtlasRT->SizeX, CombinedHeightmapNonAtlasRT->SizeY));
@@ -1660,18 +2724,18 @@
 						continue;
 					}
 
-					INC_DWORD_STAT(STAT_LandscapeRegenerateProceduralHeightmapsDrawCalls); // Brush Render
+					INC_DWORD_STAT(STAT_LandscapeRegenerateProceduralDrawCalls); // Brush Render
 
-					PrintDebugRTHeightmap(OutputDebugName ? FString::Printf(TEXT("LS Height: %s %s -> BrushNonAtlas %s"), *Layer.Name.ToString(), *Brush.BPCustomBrush->GetName(), *BrushOutputNonAtlasRT->GetName()) : TEXT(""), BrushOutputNonAtlasRT);
+					PrintProceduralDebugRT(OutputDebugName ? FString::Printf(TEXT("LS Height: %s %s -> BrushNonAtlas %s"), *Layer.Name.ToString(), *Brush.BPCustomBrush->GetName(), *BrushOutputNonAtlasRT->GetName()) : TEXT(""), BrushOutputNonAtlasRT);
 
 					// Resolve back to Combined heightmap
-					CopyProceduralTargetToResolveTarget(BrushOutputNonAtlasRT, CombinedHeightmapNonAtlasRT, nullptr, FIntPoint(0, 0), 0);
-					PrintDebugRTHeightmap(OutputDebugName ? FString::Printf(TEXT("LS Height: %s Component %s += -> CombinedNonAtlas %s"), *Layer.Name.ToString(), *BrushOutputNonAtlasRT->GetName(), *CombinedHeightmapNonAtlasRT->GetName()) : TEXT(""), CombinedHeightmapNonAtlasRT);
+					CopyProceduralTexture(BrushOutputNonAtlasRT, CombinedHeightmapNonAtlasRT);
+					PrintProceduralDebugRT(OutputDebugName ? FString::Printf(TEXT("LS Height: %s Component %s += -> CombinedNonAtlas %s"), *Layer.Name.ToString(), *BrushOutputNonAtlasRT->GetName(), *CombinedHeightmapNonAtlasRT->GetName()) : TEXT(""), CombinedHeightmapNonAtlasRT);
 				}
 			}
 
-			CopyProceduralTargetToResolveTarget(CombinedHeightmapNonAtlasRT, LandscapeScratchRT3, nullptr, FIntPoint(0, 0), 0);
-			PrintDebugRTHeightmap(OutputDebugName ? FString::Printf(TEXT("LS Height: %s Component %s += -> CombinedNonAtlas %s"), *Layer.Name.ToString(), *CombinedHeightmapNonAtlasRT->GetName(), *LandscapeScratchRT3->GetName()) : TEXT(""), LandscapeScratchRT3);
+			CopyProceduralTexture(CombinedHeightmapNonAtlasRT, LandscapeScratchRT3);
+			PrintProceduralDebugRT(OutputDebugName ? FString::Printf(TEXT("LS Height: %s Component %s += -> CombinedNonAtlas %s"), *Layer.Name.ToString(), *CombinedHeightmapNonAtlasRT->GetName(), *LandscapeScratchRT3->GetName()) : TEXT(""), LandscapeScratchRT3);
 
 			FirstLayer = false;
 		}
@@ -1697,22 +2761,24 @@
 				int32 CurrentMip = 0;
 				FRenderDataPerHeightmap& HeightmapRenderData = ItPair.Value;
 
-				CopyProceduralTargetToResolveTarget(CombinedHeightmapAtlasRT, HeightmapRenderData.OriginalHeightmap, HeightmapRenderData.HeightmapsCPUReadBack, HeightmapRenderData.TopLeftSectionBase, CurrentMip++);
+				CopyProceduralTexture(CombinedHeightmapAtlasRT, HeightmapRenderData.OriginalHeightmap, HeightmapRenderData.HeightmapsCPUReadBack, HeightmapRenderData.TopLeftSectionBase, CurrentMip, CurrentMip);
+				++CurrentMip;
 
-				for (int32 MipRTIndex = EHeightmapRTType::LandscapeSizeMip1; MipRTIndex < EHeightmapRTType::Count; ++MipRTIndex)
+				for (int32 MipRTIndex = EHeightmapRTType::HeightmapRT_Mip1; MipRTIndex < EHeightmapRTType::HeightmapRT_Count; ++MipRTIndex)
 				{
 					if (HeightmapRTList[MipRTIndex] != nullptr)
 					{
-						CopyProceduralTargetToResolveTarget(HeightmapRTList[MipRTIndex], HeightmapRenderData.OriginalHeightmap, HeightmapRenderData.HeightmapsCPUReadBack, HeightmapRenderData.TopLeftSectionBase, CurrentMip++);
+						CopyProceduralTexture(HeightmapRTList[MipRTIndex], HeightmapRenderData.OriginalHeightmap, HeightmapRenderData.HeightmapsCPUReadBack, HeightmapRenderData.TopLeftSectionBase, CurrentMip, CurrentMip);
+						++CurrentMip;
 					}
 				}
 			}
 		}
 	}
 
-	if ((ProceduralContentUpdateFlags & EProceduralContentUpdateFlag::Heightmap_ResolveToTexture) != 0 || (ProceduralContentUpdateFlags & EProceduralContentUpdateFlag::Heightmap_ResolveToTextureDDC) != 0)
+	if ((ProceduralContentUpdateFlags & EProceduralContentUpdateFlag::Heightmap_ResolveToTexture) != 0)
 	{
-		ResolveProceduralHeightmapTexture((ProceduralContentUpdateFlags & EProceduralContentUpdateFlag::Heightmap_ResolveToTextureDDC) != 0);
+		ResolveProceduralHeightmapTexture(AllLandscapes);
 	}
 
 	if ((ProceduralContentUpdateFlags & EProceduralContentUpdateFlag::Heightmap_BoundsAndCollision) != 0)
@@ -1726,32 +2792,20 @@
 		}
 	}
 
-	ProceduralContentUpdateFlags = 0;
+	ProceduralContentUpdateFlags &= ~EProceduralContentUpdateFlag::Heightmap_All;
 
 	// If doing rendering debug, keep doing the render only
 	if (CVarOutputProceduralDebugDrawCallName.GetValueOnAnyThread() == 1)
 	{
-		ProceduralContentUpdateFlags = EProceduralContentUpdateFlag::Heightmap_Render;
+		ProceduralContentUpdateFlags |= EProceduralContentUpdateFlag::Heightmap_Render;
 	}
 }
 
-void ALandscape::ResolveProceduralHeightmapTexture(bool InUpdateDDC)
+void ALandscape::ResolveProceduralHeightmapTexture(const TArray<ALandscapeProxy*>& InAllLandscapes)
 {
 	SCOPE_CYCLE_COUNTER(STAT_LandscapeResolveProceduralHeightmap);
 
-	ULandscapeInfo* Info = GetLandscapeInfo();
-
-	TArray<ALandscapeProxy*> AllLandscapes;
-	AllLandscapes.Add(this);
-
-	for (auto& It : Info->Proxies)
-	{
-		AllLandscapes.Add(It);
-	}
-
-	TArray<UTexture2D*> PendingDDCUpdateTextureList;
-
-	for (ALandscapeProxy* Landscape : AllLandscapes)
+	for (ALandscapeProxy* Landscape : InAllLandscapes)
 	{
 		TArray<TArray<FColor>> MipData;
 
@@ -1764,88 +2818,1319 @@
 				continue;
 			}
 
-			if (MipData.Num() == 0)
+			ResolveProceduralTexture(HeightmapRenderData.HeightmapsCPUReadBack, HeightmapRenderData.OriginalHeightmap);
+		}
+	}		
+}
+
+void ALandscape::ResolveProceduralTexture(FLandscapeProceduralTexture2DCPUReadBackResource* InCPUReadBackTexture, UTexture2D* InOriginalTexture)
+{
+	TArray<TArray<FColor>> MipData;
+	MipData.AddDefaulted(InCPUReadBackTexture->TextureRHI->GetNumMips());
+
+	int32 MipSizeU = InCPUReadBackTexture->GetSizeX();
+	int32 MipSizeV = InCPUReadBackTexture->GetSizeY();
+	int32 MipIndex = 0;
+
+	while (MipSizeU >= 1 && MipSizeV >= 1)
+	{
+		MipData[MipIndex].Reset();
+
+		FReadSurfaceDataFlags Flags(RCM_UNorm, CubeFace_MAX);
+		Flags.SetMip(MipIndex);
+		FIntRect Rect(0, 0, MipSizeU, MipSizeV);
+
+		ENQUEUE_RENDER_COMMAND(FProceduralReadSurfaceCommand)(
+			[SourceTextureRHI = InCPUReadBackTexture->TextureRHI, Rect = Rect, OutData = &MipData[MipIndex], ReadFlags = Flags](FRHICommandListImmediate& RHICmdList) mutable
+		{
+			RHICmdList.ReadSurfaceData(SourceTextureRHI, Rect, *OutData, ReadFlags);
+		});
+
+
+		MipSizeU >>= 1;
+		MipSizeV >>= 1;
+		++MipIndex;
+	}
+
+	// TODO: find a way to NOT have to flush the rendering command as this create hic up of ~10-15ms
+	FlushRenderingCommands();
+
+	for (MipIndex = 0; MipIndex < MipData.Num(); ++MipIndex)
+	{
+		if (MipData[MipIndex].Num() > 0)
+		{
+			FColor* TextureData = (FColor*)InOriginalTexture->Source.LockMip(MipIndex);
+			FMemory::Memcpy(TextureData, MipData[MipIndex].GetData(), MipData[MipIndex].Num() * sizeof(FColor));
+			InOriginalTexture->Source.UnlockMip(MipIndex);
+		}
+	}
+}
+
+void ALandscape::PrepareProceduralComponentDataForExtractLayersCS(const FName& InProceduralLayerName, int32 InCurrentWeightmapToProcessIndex, bool InOutputDebugName, const TArray<ALandscapeProxy*>& InAllLandscape, FLandscapeTexture2DResource* InOutTextureData,
+																  TArray<FLandscapeProceduralWeightmapExtractLayersComponentData>& OutComponentData, TMap<ULandscapeLayerInfoObject*, int32>& OutLayerInfoObjects)
+{
+	ULandscapeInfo* Info = GetLandscapeInfo();
+	
+	for (const ALandscapeProxy* Landscape : InAllLandscape)
+	{
+		const FProceduralLayerData* ProceduralLayerData = Landscape->ProceduralLayersData.Find(InProceduralLayerName);
+
+		if (ProceduralLayerData != nullptr)
+		{
+			for (const auto& ItPair : ProceduralLayerData->WeightmapData)
 			{
-				MipData.AddDefaulted(HeightmapRenderData.HeightmapsCPUReadBack->TextureRHI->GetNumMips());
+				ULandscapeComponent* Component = ItPair.Key;
+				const FWeightmapLayerData& WeightLayerData = ItPair.Value;
+
+				if (WeightLayerData.Weightmaps.IsValidIndex(InCurrentWeightmapToProcessIndex))
+				{
+					UTexture2D* Weightmap = WeightLayerData.Weightmaps[InCurrentWeightmapToProcessIndex];
+					check(Weightmap != nullptr);
+
+					const FLandscapeWeightmapUsage* WeightmapUsage = WeightLayerData.WeightmapTextureUsages[InCurrentWeightmapToProcessIndex];
+					check(WeightmapUsage != nullptr);
+
+					CopyProceduralTexture(*Weightmap->GetName(), Weightmap->Resource, InOutputDebugName ? FString::Printf(TEXT("%s WeightmapScratchTexture"), *InProceduralLayerName.ToString()) : TEXT(""), InOutTextureData, nullptr, Component->GetSectionBase(), 0);
+					PrintProceduralDebugTextureResource(InOutputDebugName ? FString::Printf(TEXT("LS Weight: %s WeightmapScratchTexture %s"), *InProceduralLayerName.ToString(), TEXT("WeightmapScratchTextureResource")) : TEXT(""), InOutTextureData, 0, false);
+
+					for (const FWeightmapLayerAllocationInfo& WeightmapLayerAllocation : WeightLayerData.WeightmapLayerAllocations)
+					{
+						if (WeightmapLayerAllocation.LayerInfo != nullptr && WeightmapLayerAllocation.WeightmapTextureIndex != 255 && WeightLayerData.Weightmaps[WeightmapLayerAllocation.WeightmapTextureIndex] == Weightmap)
+						{
+							FLandscapeProceduralWeightmapExtractLayersComponentData Data;
+
+							const ULandscapeComponent* DestComponent = WeightmapUsage->ChannelUsage[WeightmapLayerAllocation.WeightmapTextureChannel];
+							check(DestComponent);
+
+							// Compute component top left vertex position from section base info
+							int32 LocalComponentSizeQuad = Component->SubsectionSizeQuads * NumSubsections;
+							int32 LocalComponentSizeVerts = (Component->SubsectionSizeQuads + 1) * NumSubsections;
+							FVector2D SourcePositionOffset(FMath::RoundToInt(Component->GetSectionBase().X / LocalComponentSizeQuad), FMath::RoundToInt(Component->GetSectionBase().Y / LocalComponentSizeQuad));
+							FVector2D DestPositionOffset(FMath::RoundToInt(DestComponent->GetSectionBase().X / LocalComponentSizeQuad), FMath::RoundToInt(DestComponent->GetSectionBase().Y / LocalComponentSizeQuad));
+
+							Data.ComponentVertexPosition = FIntPoint(SourcePositionOffset.X * LocalComponentSizeVerts, SourcePositionOffset.Y * LocalComponentSizeVerts);
+							Data.AtlasTexturePositionOutput = FIntPoint(DestPositionOffset.X * LocalComponentSizeVerts, DestPositionOffset.Y * LocalComponentSizeVerts);
+							Data.WeightmapChannelToProcess = WeightmapLayerAllocation.WeightmapTextureChannel;
+
+							if (WeightmapLayerAllocation.LayerInfo == ALandscapeProxy::VisibilityLayer)
+							{
+								Data.DestinationPaintLayerIndex = 0;
+								int32& NewLayerinfoObjectIndex = OutLayerInfoObjects.FindOrAdd(ALandscapeProxy::VisibilityLayer);
+								NewLayerinfoObjectIndex = 0;
+							}
+							else
+							{
+								for (int32 LayerInfoSettingsIndex = 0; LayerInfoSettingsIndex < Info->Layers.Num(); ++LayerInfoSettingsIndex)
+								{
+									const FLandscapeInfoLayerSettings& InfoLayerSettings = Info->Layers[LayerInfoSettingsIndex];
+
+									if (InfoLayerSettings.LayerInfoObj == WeightmapLayerAllocation.LayerInfo)
+									{
+										Data.DestinationPaintLayerIndex = LayerInfoSettingsIndex + 1; // due to visibility layer that is at 0
+										int32& NewLayerinfoObjectIndex = OutLayerInfoObjects.FindOrAdd(WeightmapLayerAllocation.LayerInfo);
+										NewLayerinfoObjectIndex = LayerInfoSettingsIndex + 1;
+
+										break;
+									}
+								}
+							}
+
+							OutComponentData.Add(Data);
+						}
+					}
+				}
 			}
+		}
+	}
+}
 
-			int32 MipSizeU = HeightmapRenderData.HeightmapsCPUReadBack->GetSizeX();
-			int32 MipSizeV = HeightmapRenderData.HeightmapsCPUReadBack->GetSizeY();
-			int32 MipIndex = 0;
+void ALandscape::PrepareProceduralComponentDataForPackLayersCS(int32 InCurrentWeightmapToProcessIndex, bool InOutputDebugName, const TArray<ULandscapeComponent*>& InAllLandscapeComponents, TArray<UTexture2D*>& InOutProcessedWeightmaps, 
+															  TArray<FLandscapeProceduralTexture2DCPUReadBackResource*>& InOutProcessedWeightmapCPUCopy, TArray<FLandscapeProceduralWeightmapPackLayersComponentData>& OutComponentData)
+{
+	ULandscapeInfo* Info = GetLandscapeInfo();
 
-			while (MipSizeU >= 1 && MipSizeV >= 1)
+	for (const ULandscapeComponent* Component : InAllLandscapeComponents)
+	{
+		const TArray<UTexture2D*>& WeightmapTextures = Component->GetWeightmapTextures();
+
+		if (WeightmapTextures.IsValidIndex(InCurrentWeightmapToProcessIndex))
+		{
+			UTexture2D* WeightmapTexture = WeightmapTextures[InCurrentWeightmapToProcessIndex];
+
+			if (!InOutProcessedWeightmaps.Contains(WeightmapTexture))
 			{
-				MipData[MipIndex].Reset();
+				InOutProcessedWeightmaps.Add(WeightmapTexture);
 
-				FReadSurfaceDataFlags Flags(RCM_UNorm, CubeFace_MAX);
-				Flags.SetMip(MipIndex);
-				FIntRect Rect(0, 0, MipSizeU, MipSizeV);
+				FLandscapeProceduralTexture2DCPUReadBackResource** WeightmapCPUCopy = Component->GetLandscapeProxy()->WeightmapCPUReadBackTextures.Find(WeightmapTexture);
+				check(WeightmapCPUCopy != nullptr);
 
+				InOutProcessedWeightmapCPUCopy.Add(*WeightmapCPUCopy);
+
+				const TArray<FLandscapeWeightmapUsage*>& WeightmapTexturesUsage = Component->GetWeightmapTexturesUsage();
+
+				const FLandscapeWeightmapUsage* WeightmapUsage = WeightmapTexturesUsage[InCurrentWeightmapToProcessIndex];
+				check(WeightmapUsage != nullptr);
+
+				TArray<const FWeightmapLayerAllocationInfo*> AlreadyProcessedAllocation;
+				FLandscapeProceduralWeightmapPackLayersComponentData Data;
+
+				for (int32 WeightmapChannelIndex = 0; WeightmapChannelIndex < 4; ++WeightmapChannelIndex)
 				{
-					TArray<FColor>* OutData = &MipData[MipIndex];
-					FTextureRHIRef SourceTextureRHI = HeightmapRenderData.HeightmapsCPUReadBack->TextureRHI;
-					ENQUEUE_RENDER_COMMAND(ReadSurfaceCommand)(
-						[SourceTextureRHI, Rect, OutData, Flags](FRHICommandListImmediate& RHICmdList)
+					// Clear out data to known values
+					Data.ComponentVertexPositionX[WeightmapChannelIndex] = INDEX_NONE;
+					Data.ComponentVertexPositionY[WeightmapChannelIndex] = INDEX_NONE;
+					Data.SourcePaintLayerIndex[WeightmapChannelIndex] = INDEX_NONE;
+					Data.WeightmapChannelToProcess[WeightmapChannelIndex] = INDEX_NONE;
+
+					if (WeightmapUsage->ChannelUsage[WeightmapChannelIndex] != nullptr)
+					{
+						const ULandscapeComponent* ChannelComponent = WeightmapUsage->ChannelUsage[WeightmapChannelIndex];
+						FIntPoint ComponentSectionBase = ChannelComponent->GetSectionBase();
+
+						// Compute component top left vertex position from section base info
+						int32 LocalComponentSizeQuad = ChannelComponent->SubsectionSizeQuads * NumSubsections;
+						int32 LocalComponentSizeVerts = (ChannelComponent->SubsectionSizeQuads + 1) * NumSubsections;
+						FVector2D PositionOffset(FMath::RoundToInt(ChannelComponent->GetSectionBase().X / LocalComponentSizeQuad), FMath::RoundToInt(ChannelComponent->GetSectionBase().Y / LocalComponentSizeQuad));
+
+						Data.ComponentVertexPositionX[WeightmapChannelIndex] = PositionOffset.X * LocalComponentSizeVerts;
+						Data.ComponentVertexPositionY[WeightmapChannelIndex] = PositionOffset.Y * LocalComponentSizeVerts;
+
+						const TArray<FWeightmapLayerAllocationInfo>& ChannelLayerAllocations = ChannelComponent->GetWeightmapLayerAllocations();
+						const TArray<UTexture2D*>& ChannelComponentWeightmapTextures = ChannelComponent->GetWeightmapTextures();
+						
+						for (const FWeightmapLayerAllocationInfo& ChannelLayerAllocation : ChannelLayerAllocations)
 						{
-							RHICmdList.ReadSurfaceData(SourceTextureRHI, Rect, *OutData, Flags);
-						});
+							if (ChannelLayerAllocation.LayerInfo != nullptr && !AlreadyProcessedAllocation.Contains(&ChannelLayerAllocation) && ChannelComponentWeightmapTextures[ChannelLayerAllocation.WeightmapTextureIndex] == WeightmapTexture)
+							{
+								Data.WeightmapChannelToProcess[WeightmapChannelIndex] = ChannelLayerAllocation.WeightmapTextureChannel;
+								AlreadyProcessedAllocation.Add(&ChannelLayerAllocation);
+
+								if (ChannelLayerAllocation.LayerInfo == ALandscapeProxy::VisibilityLayer)
+								{
+									Data.SourcePaintLayerIndex[WeightmapChannelIndex] = 0; // Always store after the last weightmap index
+								}
+								else
+								{
+									for (int32 LayerInfoSettingsIndex = 0; LayerInfoSettingsIndex < Info->Layers.Num(); ++LayerInfoSettingsIndex)
+									{
+										const FLandscapeInfoLayerSettings& LayerInfo = Info->Layers[LayerInfoSettingsIndex];
+
+										if (ChannelLayerAllocation.LayerInfo == LayerInfo.LayerInfoObj)
+										{
+											Data.SourcePaintLayerIndex[WeightmapChannelIndex] = LayerInfoSettingsIndex + 1; // due to visibility layer that is at 0
+											break;
+										}
+									}
+								}
+
+								break;
+							}
+						}
+					}
 				}
 
-				MipSizeU >>= 1;
-				MipSizeV >>= 1;
-				++MipIndex;
+				OutComponentData.Add(Data);
 			}
+		}
+	}
+}
 
-			FlushRenderingCommands();
+void ALandscape::ReallocateProceduralWeightmaps(const TArray<ALandscapeProxy*>& InAllLandscape, const TArray<ULandscapeLayerInfoObject*>& InBrushRequiredAllocations, TArray<ULandscapeComponent*>& OutComponentThatNeedMaterialRebuild)
+{
+	SCOPE_CYCLE_COUNTER(STAT_LandscapeReallocateProceduralWeightmaps);
 
-			for (MipIndex = 0; MipIndex < MipData.Num(); ++MipIndex)
+	TArray<ULandscapeComponent*> AllLandscapeComponents;
+
+	for (ALandscapeProxy* Landscape : InAllLandscape)
+	{
+		AllLandscapeComponents.Append(Landscape->LandscapeComponents);
+	}
+
+	// Copy Previous Usage, to know which texture need updating
+	TMap<UTexture2D*, FLandscapeWeightmapUsage> CurrentWeightmapsUsage;
+
+	for (ULandscapeComponent* Component : AllLandscapeComponents)
+	{
+		TArray<UTexture2D*>& ComponentWeightmapTextures = Component->GetWeightmapTextures();
+		TArray<FLandscapeWeightmapUsage*>& ComponentWeightmapTextureUsage = Component->GetWeightmapTexturesUsage();
+
+		for (int32 i = 0; i < ComponentWeightmapTextures.Num(); ++i)
+		{
+			UTexture2D* ComponentWeightmapTexture = ComponentWeightmapTextures[i];
+			FLandscapeWeightmapUsage* CurrentWeightmapTextureUsage = CurrentWeightmapsUsage.Find(ComponentWeightmapTexture);
+
+			if (CurrentWeightmapTextureUsage == nullptr)
 			{
-				if (MipData[MipIndex].Num() > 0)
+				FLandscapeWeightmapUsage* ComponentWeightmapUsage = ComponentWeightmapTextureUsage[i];
+
+				FLandscapeWeightmapUsage Usage;
+
+				for (int32 j = 0; j < 4; ++j)
 				{
-					PrintDebugHeightData(CVarOutputProceduralRTContent.GetValueOnAnyThread() == 1 ? FString::Printf(TEXT("CPUReadBack -> Source Heightmap %s, Mip: %d"), *HeightmapRenderData.OriginalHeightmap->GetName(), MipIndex) : TEXT(""),
-										 MipData[MipIndex], FIntPoint(HeightmapRenderData.HeightmapsCPUReadBack->GetSizeX() >> MipIndex, HeightmapRenderData.HeightmapsCPUReadBack->GetSizeY() >> MipIndex), MipIndex, true);
+					Usage.ChannelUsage[j] = ComponentWeightmapUsage->ChannelUsage[j];
+				}
 
-					FColor* HeightmapTextureData = (FColor*)HeightmapRenderData.OriginalHeightmap->Source.LockMip(MipIndex);
-					FMemory::Memzero(HeightmapTextureData, MipData[MipIndex].Num() * sizeof(FColor));
-					FMemory::Memcpy(HeightmapTextureData, MipData[MipIndex].GetData(), MipData[MipIndex].Num() * sizeof(FColor));
-					HeightmapRenderData.OriginalHeightmap->Source.UnlockMip(MipIndex);
+				CurrentWeightmapsUsage.Add(ComponentWeightmapTexture, MoveTemp(Usage));
+			}
+		}
+	}
+
+	// Clear allocation data
+	for (ULandscapeComponent* Component : AllLandscapeComponents)
+	{
+		TArray<FWeightmapLayerAllocationInfo>& BaseLayerAllocations = Component->GetWeightmapLayerAllocations();
+
+		for (FWeightmapLayerAllocationInfo& BaseWeightmapAllocation : BaseLayerAllocations)
+		{
+			BaseWeightmapAllocation.WeightmapTextureChannel = 255;
+			BaseWeightmapAllocation.WeightmapTextureIndex = 255;
+		}
+		
+		TArray<FLandscapeWeightmapUsage*>& WeightmapTexturesUsage = Component->GetWeightmapTexturesUsage();
+
+		for (int32 i = 0; i < WeightmapTexturesUsage.Num(); ++i)
+		{
+			FLandscapeWeightmapUsage* Usage = WeightmapTexturesUsage[i];
+			check(Usage != nullptr);
+
+			Usage->ClearUsage();
+		}
+	}	 
+
+	bool NeedMaterialInstanceRebuild = false;
+
+	// Build a map of all the allocation per components
+	TMap<ULandscapeComponent*, TArray<ULandscapeLayerInfoObject*>> LayerAllocsPerComponent;
+
+	for (const ALandscapeProxy* Landscape : InAllLandscape)
+	{
+		for (auto& ItLayerPair : Landscape->ProceduralLayersData)
+		{
+			const FProceduralLayerData& ProceduralLayerData = ItLayerPair.Value;
+
+			for (const auto& ItWeightmapPair : ProceduralLayerData.WeightmapData)
+			{
+				ULandscapeComponent* Component = ItWeightmapPair.Key;
+				const FWeightmapLayerData& WeightLayerData = ItWeightmapPair.Value;
+
+				TArray<ULandscapeLayerInfoObject*>* ComponentLayerAlloc = LayerAllocsPerComponent.Find(Component);
+
+				if (ComponentLayerAlloc == nullptr)
+				{
+					TArray<ULandscapeLayerInfoObject*> NewLayerAllocs;
+					ComponentLayerAlloc = &LayerAllocsPerComponent.Add(Component, NewLayerAllocs);
+				}
+
+				for (const FWeightmapLayerAllocationInfo& LayerWeightmapAllocation : WeightLayerData.WeightmapLayerAllocations)
+				{
+					if (LayerWeightmapAllocation.LayerInfo != nullptr)
+					{
+						ComponentLayerAlloc->AddUnique(LayerWeightmapAllocation.LayerInfo);
+					}
+				}
+
+				// Add the brush alloc also
+				for (ULandscapeLayerInfoObject* BrushLayerInfo : InBrushRequiredAllocations)
+				{
+					if (BrushLayerInfo != nullptr)
+					{
+						ComponentLayerAlloc->AddUnique(BrushLayerInfo);
+					}
 				}
 			}
+		}
+	}
 
-			if (InUpdateDDC)
+	// Determine if the Final layer need to add/remove some alloc
+	for (auto& ItPair : LayerAllocsPerComponent)
+	{
+		ULandscapeComponent* Component = ItPair.Key;
+		TArray<ULandscapeLayerInfoObject*>& ComponentLayerAlloc = ItPair.Value;
+		TArray<FWeightmapLayerAllocationInfo>& ComponentBaseLayerAlloc = Component->GetWeightmapLayerAllocations();
+
+		// Deal with the one that need removal
+		for (int32 i = ComponentBaseLayerAlloc.Num() - 1; i >= 0; --i)
+		{
+			FWeightmapLayerAllocationInfo& Alloc = ComponentBaseLayerAlloc[i];
+
+			if (!ComponentLayerAlloc.Contains(Alloc.LayerInfo))
 			{
-				HeightmapRenderData.OriginalHeightmap->BeginCachePlatformData();
-				HeightmapRenderData.OriginalHeightmap->ClearAllCachedCookedPlatformData();
-				PendingDDCUpdateTextureList.Add(HeightmapRenderData.OriginalHeightmap);
-				HeightmapRenderData.OriginalHeightmap->MarkPackageDirty();
+				ComponentBaseLayerAlloc.RemoveAt(i);
 			}
 		}
+
+		// Then add the new one
+		for (ULandscapeLayerInfoObject* LayerAlloc : ComponentLayerAlloc)
+		{
+			const bool AllocExist = ComponentBaseLayerAlloc.ContainsByPredicate([&LayerAlloc](FWeightmapLayerAllocationInfo& BaseLayerAlloc) { return (LayerAlloc == BaseLayerAlloc.LayerInfo); });
+
+			if (!AllocExist)
+			{
+				ComponentBaseLayerAlloc.Add(FWeightmapLayerAllocationInfo(LayerAlloc));
+			}
+		}
+	}
+
+	// Realloc the weightmap so it will create proper texture (if needed) and will set the allocations information
+	TArray<UTexture2D*> NewCreatedTextures;
+
+	for (ULandscapeComponent* Component : AllLandscapeComponents)
+	{
+		Component->ReallocateWeightmaps(nullptr, false, false, true, &NewCreatedTextures);
+	}
+
+	// TODO: correctly only recreate what is required instead of everything..
+	//GDisableAutomaticTextureMaterialUpdateDependencies = true;
+
+	for (UTexture2D* Texture : NewCreatedTextures)
+	{
+		Texture->FinishCachePlatformData();
+		Texture->PostEditChange();
+
+		Texture->WaitForStreaming();
 	}	
 
-	if (InUpdateDDC)
+	//GDisableAutomaticTextureMaterialUpdateDependencies = false;
+
+	// Determine which Component need updating
+	for (ULandscapeComponent* Component : AllLandscapeComponents)
 	{
-		// Wait for all texture to be finished, do them async, since we can have many to update but we still need to wait for all of them to be finish before continuing
-		for (UTexture* PendingDDCUpdateTexture : PendingDDCUpdateTextureList)
+		if (OutComponentThatNeedMaterialRebuild.Contains(Component))
 		{
-			PendingDDCUpdateTexture->FinishCachePlatformData();
+			continue;
+		}
 
-			PendingDDCUpdateTexture->Resource = PendingDDCUpdateTexture->CreateResource();
-			if (PendingDDCUpdateTexture->Resource)
+		TArray<UTexture2D*>& ComponentWeightmapTextures = Component->GetWeightmapTextures();
+		TArray<FLandscapeWeightmapUsage*>& ComponentWeightmapTextureUsage = Component->GetWeightmapTexturesUsage();
+
+		for (int32 i = 0; i < ComponentWeightmapTextures.Num(); ++i)
+		{
+			UTexture2D* ComponentWeightmapTexture = ComponentWeightmapTextures[i];
+			FLandscapeWeightmapUsage* ComponentWeightmapUsage = ComponentWeightmapTextureUsage[i];
+			FLandscapeWeightmapUsage* CurrentWeightmapTextureUsage = CurrentWeightmapsUsage.Find(ComponentWeightmapTexture);
+
+			if (CurrentWeightmapTextureUsage != nullptr)
 			{
-				BeginInitResource(PendingDDCUpdateTexture->Resource);
+				for (int32 j = 0; j < 4; ++j)
+				{
+					if (ComponentWeightmapUsage->ChannelUsage[j] != CurrentWeightmapTextureUsage->ChannelUsage[j] && ComponentWeightmapUsage->ChannelUsage[j] != nullptr)
+					{
+						OutComponentThatNeedMaterialRebuild.AddUnique(ComponentWeightmapUsage->ChannelUsage[j]);
+					}
+				}
+			}
+			else if (NewCreatedTextures.Contains(ComponentWeightmapTexture))
+			{
+				for (int32 j = 0; j < 4; ++j)
+				{
+					if (ComponentWeightmapUsage->ChannelUsage[j] != nullptr)
+					{
+						OutComponentThatNeedMaterialRebuild.AddUnique(ComponentWeightmapUsage->ChannelUsage[j]);
+					}
+				}
 			}
 		}
 	}
 }
 
+void ALandscape::InitProceduralWeightmapResources(uint8 InLayerCount)
+{
+	ULandscapeInfo* Info = GetLandscapeInfo();
+
+	// Use the 1st one to compute the resource as they are all the same anyway
+	UTextureRenderTarget2D* FirstWeightmapRT = WeightmapRTList[WeightmapRT_Scratch1];
+
+	if (CombinedProcLayerWeightmapAllLayersResource != nullptr && InLayerCount != CombinedProcLayerWeightmapAllLayersResource->GetSizeZ())
+	{
+		ReleaseResourceAndFlush(CombinedProcLayerWeightmapAllLayersResource);
+		CombinedProcLayerWeightmapAllLayersResource = nullptr;
+	}
+
+	if (CombinedProcLayerWeightmapAllLayersResource == nullptr)
+	{
+		CombinedProcLayerWeightmapAllLayersResource = new FLandscapeTexture2DArrayResource(FirstWeightmapRT->SizeX, FirstWeightmapRT->SizeY, InLayerCount, PF_G8, 1, true);
+		BeginInitResource(CombinedProcLayerWeightmapAllLayersResource);
+	}
+
+	if (CurrentProcLayerWeightmapAllLayersResource != nullptr && InLayerCount != CurrentProcLayerWeightmapAllLayersResource->GetSizeZ())
+	{
+		ReleaseResourceAndFlush(CurrentProcLayerWeightmapAllLayersResource);
+		CurrentProcLayerWeightmapAllLayersResource = nullptr;
+	}
+
+	if (CurrentProcLayerWeightmapAllLayersResource == nullptr)
+	{
+		CurrentProcLayerWeightmapAllLayersResource = new FLandscapeTexture2DArrayResource(FirstWeightmapRT->SizeX, FirstWeightmapRT->SizeY, InLayerCount, PF_G8, 1, true);
+		BeginInitResource(CurrentProcLayerWeightmapAllLayersResource);
+	}
+
+	if (WeightmapScratchExtractLayerTextureResource == nullptr)
+	{
+		WeightmapScratchExtractLayerTextureResource = new FLandscapeTexture2DResource(FirstWeightmapRT->SizeX, FirstWeightmapRT->SizeY, PF_B8G8R8A8, 1, false);
+		BeginInitResource(WeightmapScratchExtractLayerTextureResource);
+	}
+
+	if (WeightmapScratchPackLayerTextureResource == nullptr)
+	{
+		int32 MipCount = 0;
+
+		for (int32 MipRTIndex = EWeightmapRTType::WeightmapRT_Mip0; MipRTIndex < EWeightmapRTType::WeightmapRT_Count; ++MipRTIndex)
+		{
+			if (WeightmapRTList[MipRTIndex] != nullptr)
+			{
+				++MipCount;
+			}
+		}
+
+		WeightmapScratchPackLayerTextureResource = new FLandscapeTexture2DResource(FirstWeightmapRT->SizeX, FirstWeightmapRT->SizeY, PF_B8G8R8A8, MipCount, true);
+		BeginInitResource(WeightmapScratchPackLayerTextureResource);		
+	}
+}
+
+/* This code was removed as it's a fix to remove weightmap with 0 weightmap for a component, unfortunatly it will cause side effect, like perf issue, etc, so for now it's disabled until we can address this later
+bool ALandscape::GenerateZeroAllocationPerComponents(const TArray<ALandscapeProxy*>& InAllLandscape, const TMap<ULandscapeLayerInfoObject*, bool>& InWeightmapLayersBlendSubstractive, TMap<ULandscapeComponent*, TArray<ULandscapeLayerInfoObject*>>& OutZeroAllocationsPerComponents)
+{
+	// Make sure all CPU Readback texture are ready to be used
+	for (ALandscapeProxy* Landscape : InAllLandscape)
+	{
+		for (auto& ItPair : Landscape->WeightmapCPUReadBackTextures)
+		{
+			FLandscapeProceduralTexture2DCPUReadBackResource* CPuReadBackTexture = ItPair.Value;
+
+			if (CPuReadBackTexture != nullptr && !CPuReadBackTexture->IsInitialized())
+			{
+				return false;
+			}
+		}
+	}
+
+	TArray<ULandscapeComponent*> AllLandscapeComponents;
+
+	for (ALandscapeProxy* Landscape : InAllLandscape)
+	{
+		AllLandscapeComponents.Append(Landscape->LandscapeComponents);
+	}
+
+	// Compute the data to facilitate the process
+	struct FTextureChannelData
+	{
+		TArray<ULandscapeComponent*, TInlineAllocator<4>> ComponentPerChannel;
+		TArray<ULandscapeLayerInfoObject*, TInlineAllocator<4>> LayerInfoPerChannel;
+		TArray<FColor> TexelData;
+		FLandscapeProceduralTexture2DCPUReadBackResource* CPUReadbackTexture;
+	};
+
+	TMap<UTexture2D*, FTextureChannelData> TextureDataPerTextures;
+
+	for (ULandscapeComponent* LandscapeComponent : AllLandscapeComponents)
+	{
+		const TArray<FLandscapeWeightmapUsage*>& ComponentWeightmapUsage = LandscapeComponent->GetWeightmapTexturesUsage();
+		const TArray<UTexture2D*>& ComponentWeightmapTextures = LandscapeComponent->GetWeightmapTextures();
+
+		for (int32 TextureIndex = 0; TextureIndex < ComponentWeightmapTextures.Num(); ++TextureIndex)
+		{
+			UTexture2D* WeightmapTexture = ComponentWeightmapTextures[TextureIndex];
+
+			if (TextureDataPerTextures.Contains(WeightmapTexture))
+			{
+				continue;
+			}
+
+			FLandscapeWeightmapUsage* Usage = ComponentWeightmapUsage[TextureIndex];
+
+			for (int32 i = 0; i < 4; ++i)
+			{
+				ULandscapeComponent* Component = Usage->ChannelUsage[i];
+
+				if (Component != nullptr)
+				{
+					const TArray<FWeightmapLayerAllocationInfo>& ComponentUsageWeightmapAllocations = Component->GetWeightmapLayerAllocations();
+
+					for (const FWeightmapLayerAllocationInfo& Allocation : ComponentUsageWeightmapAllocations)
+					{
+						if (Allocation.WeightmapTextureIndex == TextureIndex)
+						{
+							const bool* IsSubtractive = InWeightmapLayersBlendSubstractive.Find(Allocation.LayerInfo);
+
+							if (IsSubtractive != nullptr && *IsSubtractive)
+							{
+								FTextureChannelData* TextureData = TextureDataPerTextures.Find(WeightmapTexture);
+
+								if (TextureData == nullptr)
+								{
+									FTextureChannelData NewData;
+									NewData.ComponentPerChannel.SetNumZeroed(4);
+									NewData.LayerInfoPerChannel.SetNumZeroed(4);
+
+									FLandscapeProceduralTexture2DCPUReadBackResource** CPUReadbackTexture = Component->GetLandscapeProxy()->WeightmapCPUReadBackTextures.Find(WeightmapTexture);
+									check(CPUReadbackTexture);
+
+									NewData.CPUReadbackTexture = *CPUReadbackTexture;
+
+									TextureData = &TextureDataPerTextures.Add(WeightmapTexture, NewData);
+								}
+
+								TextureData->LayerInfoPerChannel[Allocation.WeightmapTextureChannel] = Allocation.LayerInfo;
+								TextureData->ComponentPerChannel[Allocation.WeightmapTextureChannel] = Component;
+							}
+						}
+					}
+				}
+			}
+		}
+	}
+
+	// Read the data form the CPU texture
+	if (TextureDataPerTextures.Num() > 0)
+	{
+		for (auto& ItPair : TextureDataPerTextures)
+		{
+			FTextureChannelData& TextureChannelData = ItPair.Value;
+
+			int32 MipSizeU = TextureChannelData.CPUReadbackTexture->GetSizeX();
+			int32 MipSizeV = TextureChannelData.CPUReadbackTexture->GetSizeY();
+
+			FReadSurfaceDataFlags Flags(RCM_UNorm, CubeFace_MAX);
+			Flags.SetMip(0);
+			FIntRect Rect(0, 0, MipSizeU, MipSizeV);
+
+			ENQUEUE_UNIQUE_RENDER_COMMAND_FOURPARAMETER(
+				ReadSurfaceCommand,
+				FTextureRHIRef, SourceTextureRHI, TextureChannelData.CPUReadbackTexture->TextureRHI,
+				FIntRect, Rect, Rect,
+				TArray<FColor>*, OutData, &TextureChannelData.TexelData,
+				FReadSurfaceDataFlags, ReadFlags, Flags,
+				{
+					RHICmdList.ReadSurfaceData(SourceTextureRHI, Rect, *OutData, ReadFlags);
+				});
+		}
+
+		FlushRenderingCommands();
+
+		// Determine which texture channel is zero allocation to "remove" them
+		for (auto& ItPair : TextureDataPerTextures)
+		{
+			FTextureChannelData& TextureChannelData = ItPair.Value;
+
+			TArray<bool> AreLayerEmpty;
+			AreLayerEmpty.Init(true, 4);
+
+			for (FColor& Color : TextureChannelData.TexelData)
+			{
+				if (Color.R != 0)
+				{
+					AreLayerEmpty[0] = false;
+				}
+
+				if (Color.G != 0)
+				{
+					AreLayerEmpty[1] = false;
+				}
+
+				if (Color.B != 0)
+				{
+					AreLayerEmpty[2] = false;
+				}
+
+				if (Color.A != 0)
+				{
+					AreLayerEmpty[3] = false;
+				}
+			}
+			
+			for (int32 i = 0; i < 4; ++i)
+			{
+				if (AreLayerEmpty[i] && TextureChannelData.ComponentPerChannel[0] != nullptr)
+				{
+					TArray<ULandscapeLayerInfoObject*>& ZeroAllocations = OutZeroAllocationsPerComponents.FindOrAdd(TextureChannelData.ComponentPerChannel[0]);
+					ZeroAllocations.Add(TextureChannelData.LayerInfoPerChannel[0]);
+				}
+			}
+		}
+	}
+
+	return true;
+}
+*/
+
+bool ALandscape::AreWeightmapTextureResourcesReady(const TArray<ALandscapeProxy*>& InAllLandscapes) const
+{
+	// Make sure all our original weightmap textures are streamed in and ready to be used
+	for (const ALandscapeProxy* Landscape : InAllLandscapes)
+	{
+		for (const auto& ItLayerDataPair : Landscape->ProceduralLayersData)
+		{
+			for (const auto& ItWeightmapPair : ItLayerDataPair.Value.WeightmapData)
+			{
+				ULandscapeComponent* Component = ItWeightmapPair.Key;
+				const TArray<UTexture2D*>& OriginalWeightmaps = Component->GetWeightmapTextures();
+
+				for (UTexture2D* Weightmap : OriginalWeightmaps)
+				{
+					if (!Weightmap->IsFullyStreamedIn())
+					{
+						return false;
+					}
+				}
+			}
+		}
+	}
+
+	// Init all needed resources
+	for (const ALandscapeProxy* Landscape : InAllLandscapes)
+	{
+		for (const auto& ItLayerDataPair : Landscape->ProceduralLayersData)
+		{
+			for (const auto& ItWeightmapPair : ItLayerDataPair.Value.WeightmapData)
+			{
+				const FWeightmapLayerData& WeightmapLayerData = ItWeightmapPair.Value;
+
+				for (UTexture2D* Weightmap : WeightmapLayerData.Weightmaps)
+				{
+					if (Weightmap->Resource == nullptr)
+					{
+						Weightmap->FinishCachePlatformData();
+
+						Weightmap->Resource = Weightmap->CreateResource();
+						if (Weightmap->Resource)
+						{
+							BeginInitResource(Weightmap->Resource);
+						}
+					}
+				}
+			}
+		}
+	}
+
+	// Wait for the new resource to be fully initialized/streamed in
+	for (const ALandscapeProxy* Landscape : InAllLandscapes)
+	{
+		for (const auto& ItLayerDataPair : Landscape->ProceduralLayersData)
+		{
+			for (const auto& ItWeightmapPair : ItLayerDataPair.Value.WeightmapData)
+			{
+				const FWeightmapLayerData& WeightmapLayerData = ItWeightmapPair.Value;
+
+				for (UTexture2D* Weightmap : WeightmapLayerData.Weightmaps)
+				{
+					if (!Weightmap->Resource->IsInitialized() || !Weightmap->IsFullyStreamedIn())
+					{
+						return false;
+					}
+				}
+			}
+		}
+	}
+
+	return true;
+}
+
 void ALandscape::RegenerateProceduralWeightmaps()
 {
+	SCOPE_CYCLE_COUNTER(STAT_LandscapeRegenerateProceduralWeightmaps);
 
+	ULandscapeInfo* Info = GetLandscapeInfo();
+
+	if (ProceduralContentUpdateFlags == 0 || Info == nullptr || Info->Layers.Num() == 0)
+	{
+		return;
+	}
+
+	TArray<ALandscapeProxy*> AllLandscapes;
+	AllLandscapes.Add(this);
+
+	for (const auto& It : Info->Proxies)
+	{
+		AllLandscapes.Add(It);
+	}
+
+	if (!AreWeightmapTextureResourcesReady(AllLandscapes))
+	{
+		return;
+	}
+
+	TArray<ULandscapeComponent*> AllLandscapeComponents;
+
+	for (ALandscapeProxy* Landscape : AllLandscapes)
+	{
+		AllLandscapeComponents.Append(Landscape->LandscapeComponents);
+	}
+
+	TArray<ULandscapeComponent*> ComponentThatNeedMaterialRebuild;
+	TArray<ULandscapeLayerInfoObject*> BrushRequiredAllocations;
+	int32 LayerCount = Info->Layers.Num() + 1; // due to visibility being stored at 0
+	bool ClearFlagsAfterUpdate = true;
+
+	if ((ProceduralContentUpdateFlags & EProceduralContentUpdateFlag::Weightmap_Render) != 0 && WeightmapRTList.Num() > 0)
+	{
+		UTextureRenderTarget2D* LandscapeScratchRT1 = WeightmapRTList[EWeightmapRTType::WeightmapRT_Scratch1];
+		UTextureRenderTarget2D* LandscapeScratchRT2 = WeightmapRTList[EWeightmapRTType::WeightmapRT_Scratch2];
+		UTextureRenderTarget2D* LandscapeScratchRT3 = WeightmapRTList[EWeightmapRTType::WeightmapRT_Scratch3];
+		UTextureRenderTarget2D* EmptyRT = WeightmapRTList[EWeightmapRTType::WeightmapRT_Scratch_RGBA];
+		FLandscapeWeightmapProceduralShaderParameters PSShaderParams;
+		bool OutputDebugName = (CVarOutputProceduralDebugDrawCallName.GetValueOnAnyThread() == 1 || CVarOutputProceduralRTContent.GetValueOnAnyThread() == 1 || CVarOutputProceduralWeightmapsRTContent.GetValueOnAnyThread() == 1) ? true : false;
+
+		InitProceduralWeightmapResources(LayerCount);
+
+		ClearWeightmapTextureResource(TEXT("ClearRT RGBA"), EmptyRT->GameThread_GetRenderTargetResource());
+		ClearWeightmapTextureResource(TEXT("ClearRT R"), LandscapeScratchRT1->GameThread_GetRenderTargetResource());
+
+		for (int32 LayerIndex = 0; LayerIndex < LayerCount; ++LayerIndex)
+		{
+			CopyProceduralTexture(*LandscapeScratchRT1->GetName(), LandscapeScratchRT1->GameThread_GetRenderTargetResource(), OutputDebugName ? FString::Printf(TEXT("Weight: Clear CombinedProcLayerWeightmapAllLayersResource %d, "), LayerIndex) : TEXT(""), CombinedProcLayerWeightmapAllLayersResource, nullptr, FIntPoint(0, 0), 0, 0, 0, LayerIndex);
+		}
+
+		bool ComputeShaderGeneratedData = false;
+		bool FirstLayer = true;		
+		TMap<ULandscapeLayerInfoObject*, bool> WeightmapLayersBlendSubstractive;
+
+		for (FProceduralLayer& ProceduralLayer : ProceduralLayers)
+		{
+			int8 CurrentWeightmapToProcessIndex = 0;
+			bool HasFoundWeightmapToProcess = true; // try processing at least once
+
+			TMap<ULandscapeLayerInfoObject*, int32> LayerInfoObjects; // <LayerInfoObj, LayerIndex>
+
+			// Determine if some brush want to write to layer that we have currently no data on
+			if (ProceduralLayer.Visible)
+			{
+				for (int32 LayerInfoSettingsIndex = 0; LayerInfoSettingsIndex < Info->Layers.Num(); ++LayerInfoSettingsIndex)
+				{
+					const FLandscapeInfoLayerSettings& InfoLayerSettings = Info->Layers[LayerInfoSettingsIndex];
+
+					for (int32 i = 0; i < ProceduralLayer.WeightmapBrushOrderIndices.Num(); ++i)
+					{
+						FLandscapeProceduralLayerBrush& Brush = ProceduralLayer.Brushes[ProceduralLayer.WeightmapBrushOrderIndices[i]];
+
+						if (Brush.BPCustomBrush == nullptr)
+						{
+							continue;
+						}
+
+						if (Brush.BPCustomBrush->IsAffectingWeightmapLayer(InfoLayerSettings.GetLayerName()) && !LayerInfoObjects.Contains(InfoLayerSettings.LayerInfoObj))
+						{
+							LayerInfoObjects.Add(InfoLayerSettings.LayerInfoObj, LayerInfoSettingsIndex + 1); // due to visibility layer that is at 0
+						}
+					}
+				}
+			}
+
+			// Loop until there is no more weightmap texture to process
+			while (HasFoundWeightmapToProcess)
+			{
+				CopyProceduralTexture(*EmptyRT->GetName(), EmptyRT->GameThread_GetRenderTargetResource(), OutputDebugName ? FString::Printf(TEXT("Weight: %s Clear WeightmapScratchExtractLayerTextureResource"), *ProceduralLayer.Name.ToString()) : TEXT(""), WeightmapScratchExtractLayerTextureResource);
+
+				// Prepare compute shader data
+				TArray<FLandscapeProceduralWeightmapExtractLayersComponentData> ComponentsData;	
+				PrepareProceduralComponentDataForExtractLayersCS(ProceduralLayer.Name, CurrentWeightmapToProcessIndex, OutputDebugName, AllLandscapes, WeightmapScratchExtractLayerTextureResource, ComponentsData, LayerInfoObjects);
+
+				HasFoundWeightmapToProcess = ComponentsData.Num() > 0;
+
+				// Perform the compute shader
+				if (ComponentsData.Num() > 0)
+				{
+					PrintProceduralDebugTextureResource(OutputDebugName ? FString::Printf(TEXT("LS Weight: %s WeightmapScratchTexture %s"), *ProceduralLayer.Name.ToString(), TEXT("WeightmapScratchTextureResource")) : TEXT(""), WeightmapScratchExtractLayerTextureResource, 0, false);
+
+					// Clear the current atlas if required
+					if (CurrentWeightmapToProcessIndex == 0)
+					{
+						ClearWeightmapTextureResource(TEXT("ClearRT"), LandscapeScratchRT1->GameThread_GetRenderTargetResource());
+
+						// Important: for performance reason we only clear the layer we will write to, the other one might contain data but they will not be read during the blend phase
+						for (auto& ItPair : LayerInfoObjects)
+						{
+							int32 LayerIndex = ItPair.Value;
+							CopyProceduralTexture(*LandscapeScratchRT1->GetName(), LandscapeScratchRT1->GameThread_GetRenderTargetResource(), OutputDebugName ? FString::Printf(TEXT("Weight: %s Clear CurrentProcLayerWeightmapAllLayersResource %d, "), *ProceduralLayer.Name.ToString(), LayerIndex) : TEXT(""), CurrentProcLayerWeightmapAllLayersResource, nullptr, FIntPoint(0, 0), 0, 0, 0, LayerIndex);
+						}
+					}
+
+					FLandscapeWeightmapProceduralWeightmapExtractLayersComputeShaderParameters CSExtractLayersShaderParams;
+					CSExtractLayersShaderParams.AtlasWeightmapsPerLayer = CurrentProcLayerWeightmapAllLayersResource;
+					CSExtractLayersShaderParams.ComponentWeightmapResource = WeightmapScratchExtractLayerTextureResource;
+					CSExtractLayersShaderParams.ComputeShaderResource = new FLandscapeProceduralWeightmapExtractLayersComputeShaderResource(ComponentsData);
+					CSExtractLayersShaderParams.ComponentSize = (SubsectionSizeQuads + 1) * NumSubsections;
+
+					BeginInitResource(CSExtractLayersShaderParams.ComputeShaderResource);
+
+					FLandscapeProceduralWeightmapExtractLayersCSDispatch_RenderThread CSDispatch(CSExtractLayersShaderParams);
+
+					ENQUEUE_RENDER_COMMAND(FLandscapeProceduralExtractLayersCSCommand)(
+						[CSDispatch](FRHICommandListImmediate& RHICmdList) mutable
+					{
+						CSDispatch.ExtractLayers(RHICmdList);						
+					});
+					
+					++CurrentWeightmapToProcessIndex;
+					ComputeShaderGeneratedData = true; // at least 1 CS was executed, so we can continue the processing
+				}
+			}
+
+			// If we did process at least one compute shader
+			if (LayerInfoObjects.Num() > 0)
+			{
+				for (auto& LayerInfoObject : LayerInfoObjects)
+				{
+					int32 LayerIndex = LayerInfoObject.Value;
+					ULandscapeLayerInfoObject* LayerInfoObj = LayerInfoObject.Key;
+
+					// Copy the layer we are working on
+					CopyProceduralTexture(OutputDebugName ? FString::Printf(TEXT("Weight: %s PaintLayer: %s, CurrentProcLayerWeightmapAllLayersResource"), *ProceduralLayer.Name.ToString(), *LayerInfoObj->LayerName.ToString()) : TEXT(""), CurrentProcLayerWeightmapAllLayersResource, *LandscapeScratchRT1->GetName(), LandscapeScratchRT1->GameThread_GetRenderTargetResource(), nullptr, FIntPoint(0, 0), 0, 0, LayerIndex, 0);
+					PrintProceduralDebugRT(OutputDebugName ? FString::Printf(TEXT("LS Weight: %s CurrentProcLayerWeightmapAllLayersResource -> Paint Layer RT %s"), *ProceduralLayer.Name.ToString(), *LandscapeScratchRT1->GetName()) : TEXT(""), LandscapeScratchRT1, 0, false);
+
+					PSShaderParams.ApplyLayerModifiers = true;
+					PSShaderParams.LayerVisible = ProceduralLayer.Visible;
+					PSShaderParams.LayerAlpha = LayerInfoObj == ALandscapeProxy::VisibilityLayer ? 1.0f : ProceduralLayer.WeightmapAlpha; // visibility can't be affected by weight
+
+					DrawWeightmapComponentsToRenderTarget(OutputDebugName ? FString::Printf(TEXT("LS Weight: %s Paint: %s += -> %s"), *ProceduralLayer.Name.ToString(), *LayerInfoObj->LayerName.ToString(), *LandscapeScratchRT1->GetName(), *LandscapeScratchRT2->GetName()) : TEXT(""),
+														  AllLandscapeComponents, LandscapeScratchRT1, nullptr, LandscapeScratchRT2, true, PSShaderParams, 0);
+
+					PSShaderParams.ApplyLayerModifiers = false;
+
+					// Combined Layer data with current stack
+					CopyProceduralTexture(OutputDebugName ? FString::Printf(TEXT("Weight: %s PaintLayer: %s CombinedProcLayerWeightmap"), *ProceduralLayer.Name.ToString(), *LayerInfoObj->LayerName.ToString()) : TEXT(""), CombinedProcLayerWeightmapAllLayersResource, *LandscapeScratchRT1->GetName(), LandscapeScratchRT1->GameThread_GetRenderTargetResource(), nullptr, FIntPoint(0, 0), 0, 0, LayerIndex, 0);
+					PrintProceduralDebugRT(OutputDebugName ? FString::Printf(TEXT("LS Weight: %s CombinedProcLayerWeightmap -> Paint Layer RT %s"), *ProceduralLayer.Name.ToString(), *LandscapeScratchRT1->GetName()) : TEXT(""), LandscapeScratchRT1, 0, false);
+
+					// Combine with current status and copy back to the combined 2d resource array
+					PSShaderParams.OutputAsSubstractive = false;
+
+					if (!FirstLayer)
+					{
+						const bool* BlendSubstractive = ProceduralLayer.WeightmapLayerAllocationBlend.Find(LayerInfoObj);
+						PSShaderParams.OutputAsSubstractive = BlendSubstractive != nullptr ? *BlendSubstractive : false;
+
+						if (PSShaderParams.OutputAsSubstractive)
+						{
+							bool& IsSubstractiveBlend = WeightmapLayersBlendSubstractive.FindOrAdd(LayerInfoObj);
+							IsSubstractiveBlend = true;
+						}
+					}
+
+					DrawWeightmapComponentsToRenderTarget(OutputDebugName ? FString::Printf(TEXT("LS Weight: %s PaintLayer: %s, %s += -> Combined %s"), *ProceduralLayer.Name.ToString(), *LayerInfoObj->LayerName.ToString(), *LandscapeScratchRT2->GetName(), *LandscapeScratchRT3->GetName()) : TEXT(""),
+														  AllLandscapeComponents, LandscapeScratchRT2, FirstLayer ? nullptr : LandscapeScratchRT1, LandscapeScratchRT3, true, PSShaderParams, 0);
+
+					PSShaderParams.OutputAsSubstractive = false;
+
+					CopyProceduralTexture(OutputDebugName ? FString::Printf(TEXT("Weight: %s PaintLayer: %s %s"), *ProceduralLayer.Name.ToString(), *LayerInfoObj->LayerName.ToString(), *LandscapeScratchRT3->GetName()) : TEXT(""), LandscapeScratchRT3->GameThread_GetRenderTargetResource(), TEXT("CombinedProcLayerWeightmap"), CombinedProcLayerWeightmapAllLayersResource, nullptr, FIntPoint(0, 0), 0, 0, 0, LayerIndex);
+
+					// Handle brush blending
+					if (ProceduralLayer.Visible)
+					{
+						// Draw each brushes				
+						for (int32 i = 0; i < ProceduralLayer.WeightmapBrushOrderIndices.Num(); ++i)
+						{
+							// TODO: handle conversion/handling of RT not same size as internal size
+
+							FLandscapeProceduralLayerBrush& Brush = ProceduralLayer.Brushes[ProceduralLayer.WeightmapBrushOrderIndices[i]];
+
+							if (Brush.BPCustomBrush == nullptr || !Brush.BPCustomBrush->IsAffectingWeightmap() || !Brush.BPCustomBrush->IsAffectingWeightmapLayer(LayerInfoObj->LayerName))
+							{
+								continue;
+							}
+
+							BrushRequiredAllocations.AddUnique(LayerInfoObj);
+
+							if (!Brush.IsInitialized())
+							{
+								Brush.Initialize(GetBoundingRect(), FIntPoint(LandscapeScratchRT3->SizeX, LandscapeScratchRT3->SizeY));
+							}
+
+							UTextureRenderTarget2D* BrushOutputRT = Brush.Render(false, LandscapeScratchRT3);
+
+							if (BrushOutputRT == nullptr || BrushOutputRT->SizeX != LandscapeScratchRT3->SizeX || BrushOutputRT->SizeY != LandscapeScratchRT3->SizeY)
+							{
+								continue;
+							}
+
+							INC_DWORD_STAT(STAT_LandscapeRegenerateProceduralDrawCalls); // Brush Render
+
+							PrintProceduralDebugRT(OutputDebugName ? FString::Printf(TEXT("LS Weight: %s %s -> Brush %s"), *ProceduralLayer.Name.ToString(), *Brush.BPCustomBrush->GetName(), *BrushOutputRT->GetName()) : TEXT(""), BrushOutputRT);
+
+							CopyProceduralTexture(OutputDebugName ? FString::Printf(TEXT("Weight: %s PaintLayer: %s Brush: %s"), *ProceduralLayer.Name.ToString(), *LayerInfoObj->LayerName.ToString(), *BrushOutputRT->GetName()) : TEXT(""), BrushOutputRT->GameThread_GetRenderTargetResource(), *LandscapeScratchRT3->GetName(), LandscapeScratchRT3->GameThread_GetRenderTargetResource());
+							PrintProceduralDebugRT(OutputDebugName ? FString::Printf(TEXT("LS Weight: %s Component %s += -> Combined %s"), *ProceduralLayer.Name.ToString(), *BrushOutputRT->GetName(), *LandscapeScratchRT3->GetName()) : TEXT(""), LandscapeScratchRT3);
+						}
+
+						PrintProceduralDebugRT(OutputDebugName ? FString::Printf(TEXT("LS Weight: %s CombinedPostBrushProcLayerWeightmap -> Paint Layer RT %s"), *ProceduralLayer.Name.ToString(), *LandscapeScratchRT3->GetName()) : TEXT(""), LandscapeScratchRT3, 0, false);
+						CopyProceduralTexture(OutputDebugName ? FString::Printf(TEXT("Weight: %s PaintLayer: %s %s"), *ProceduralLayer.Name.ToString(), *LayerInfoObj->LayerName.ToString(), *LandscapeScratchRT3->GetName()) : TEXT(""), LandscapeScratchRT3->GameThread_GetRenderTargetResource(), TEXT("CombinedProcLayerWeightmap"), CombinedProcLayerWeightmapAllLayersResource, nullptr, FIntPoint(0, 0), 0, 0, 0, LayerIndex);
+					}
+				}
+
+				PSShaderParams.ApplyLayerModifiers = false;
+			}
+
+			FirstLayer = false;
+		}
+
+		// TODO:  if editing a Brush affecting layers, since we don't have any bounds to brush, right now ReallocateProceduralWeightmaps wont ask a rebuild of the component affected by Brushes, which mean ComponentThatNeedMaterialRebuild wont contains Brush affected component!
+		ReallocateProceduralWeightmaps(AllLandscapes, BrushRequiredAllocations, ComponentThatNeedMaterialRebuild);
+
+		// Allocation that will need to be excluded when we update materials
+		TMap<ULandscapeComponent*, TArray<ULandscapeLayerInfoObject*>> ZeroAllocationsPerComponents;
+
+		if (ComputeShaderGeneratedData)
+		{
+			// Will generate CPU read back resource, if required
+			for (ALandscapeProxy* LandscapeProxy : AllLandscapes)
+			{
+				for (const ULandscapeComponent* Component : LandscapeProxy->LandscapeComponents)
+				{
+					const TArray<UTexture2D*>& ComponentWeightmapTextures = Component->GetWeightmapTextures();
+
+					for (UTexture2D* WeightmapTexture : ComponentWeightmapTextures)
+					{
+						FLandscapeProceduralTexture2DCPUReadBackResource** WeightmapCPUReadBack = LandscapeProxy->WeightmapCPUReadBackTextures.Find(WeightmapTexture);
+
+						if (WeightmapCPUReadBack == nullptr)
+						{
+							FLandscapeProceduralTexture2DCPUReadBackResource* NewWeightmapCPUReadBack = new FLandscapeProceduralTexture2DCPUReadBackResource(WeightmapTexture->Source.GetSizeX(), WeightmapTexture->Source.GetSizeY(), WeightmapTexture->GetPixelFormat(), WeightmapTexture->Source.GetNumMips());
+							BeginInitResource(NewWeightmapCPUReadBack);
+
+							LandscapeProxy->WeightmapCPUReadBackTextures.Add(WeightmapTexture, NewWeightmapCPUReadBack);
+						}
+					}
+				}
+			}
+
+			int8 CurrentWeightmapToProcessIndex = 0;
+			bool HasFoundWeightmapToProcess = true; // try processing at least once
+
+			TArray<float> WeightmapLayerWeightBlend;
+			TArray<UTexture2D*> ProcessedWeightmaps;
+			TArray<FLandscapeProceduralTexture2DCPUReadBackResource*> ProcessedWeightmapsCPUCopy;
+			int32 NextTextureIndexToProcess = 0;
+
+			// Generate the component data from the weightmap allocation that were done earlier and weight blend them if required (i.e renormalize)
+			while (HasFoundWeightmapToProcess)
+			{
+				TArray<FLandscapeProceduralWeightmapPackLayersComponentData> PackLayersComponentsData;
+				PrepareProceduralComponentDataForPackLayersCS(CurrentWeightmapToProcessIndex, OutputDebugName, AllLandscapeComponents, ProcessedWeightmaps, ProcessedWeightmapsCPUCopy, PackLayersComponentsData);
+				HasFoundWeightmapToProcess = PackLayersComponentsData.Num() > 0;
+
+				// Perform the compute shader
+				if (PackLayersComponentsData.Num() > 0)
+				{
+					// Compute the weightblend mode of each layer for the compute shader
+					if (WeightmapLayerWeightBlend.Num() != LayerCount)
+					{
+						WeightmapLayerWeightBlend.SetNum(LayerCount);
+
+						for (int32 LayerInfoSettingsIndex = 0; LayerInfoSettingsIndex < Info->Layers.Num(); ++LayerInfoSettingsIndex)
+						{
+							const FLandscapeInfoLayerSettings& LayerInfo = Info->Layers[LayerInfoSettingsIndex];
+							WeightmapLayerWeightBlend[LayerInfoSettingsIndex + 1] = LayerInfo.LayerInfoObj != nullptr ? (LayerInfo.LayerInfoObj->bNoWeightBlend ? 0.0f : 1.0f) : 1.0f;
+						}
+
+						WeightmapLayerWeightBlend[0] = 0.0f; // Blend of Visibility 
+					}
+
+					TArray<FVector2D> WeightmapTextureOutputOffset;
+
+					// Compute each weightmap location so compute shader will be able to output at expected location
+					int32 ComponentSize = (SubsectionSizeQuads + 1) * NumSubsections;
+					float ComponentY = 0;
+					float ComponentX = 0;
+
+					for (int32 i = 0; i < PackLayersComponentsData.Num(); ++i)
+					{
+						check(ComponentY < WeightmapScratchPackLayerTextureResource->GetSizeY()); // This should never happen as it would be a bug in the algo
+
+						if (ComponentX >= WeightmapScratchPackLayerTextureResource->GetSizeX())
+						{
+							ComponentY += ComponentSize;
+							ComponentX = 0;
+						}
+
+						WeightmapTextureOutputOffset.Add(FVector2D(ComponentX, ComponentY));
+						ComponentX += ComponentSize;
+					}
+
+					// Clear Pack texture
+					CopyProceduralTexture(*EmptyRT->GetName(), EmptyRT->GameThread_GetRenderTargetResource(), TEXT("Weight: Clear WeightmapScratchPackLayerTextureResource"), WeightmapScratchPackLayerTextureResource);
+
+					FLandscapeProceduralWeightmapPackLayersComputeShaderParameters CSPackLayersShaderParams;
+					CSPackLayersShaderParams.AtlasWeightmapsPerLayer = CombinedProcLayerWeightmapAllLayersResource;
+					CSPackLayersShaderParams.ComponentWeightmapResource = WeightmapScratchPackLayerTextureResource;
+					CSPackLayersShaderParams.ComputeShaderResource = new FLandscapeProceduralWeightmapPackLayersComputeShaderResource(PackLayersComponentsData, WeightmapLayerWeightBlend, WeightmapTextureOutputOffset);
+					CSPackLayersShaderParams.ComponentSize = ComponentSize;
+					BeginInitResource(CSPackLayersShaderParams.ComputeShaderResource);
+
+					FLandscapeProceduralWeightmapPackLayerCSDispatch_RenderThread CSDispatch(CSPackLayersShaderParams);
+
+					ENQUEUE_RENDER_COMMAND(FLandscapeProceduralPackLayersCSCommand)(
+						[CSDispatch](FRHICommandListImmediate& RHICmdList) mutable
+					{
+						CSDispatch.PackLayers(RHICmdList);
+					});
+
+					int32 StartTextureIndex = NextTextureIndexToProcess;
+
+					for (; NextTextureIndexToProcess < ProcessedWeightmaps.Num(); ++NextTextureIndexToProcess)
+					{
+						UTexture2D* WeightmapTexture = ProcessedWeightmaps[NextTextureIndexToProcess];
+						FLandscapeProceduralTexture2DCPUReadBackResource* WeightmapCPUReadBack = ProcessedWeightmapsCPUCopy[NextTextureIndexToProcess];
+						FIntPoint TextureTopLeftPositionInAtlas(WeightmapTextureOutputOffset[NextTextureIndexToProcess - StartTextureIndex].X, WeightmapTextureOutputOffset[NextTextureIndexToProcess - StartTextureIndex].Y);
+
+						UTextureRenderTarget2D* CurrentRT = WeightmapRTList[WeightmapRT_Mip0];
+						CopyProceduralTexture(TEXT("WeightmapScratchTexture"), WeightmapScratchPackLayerTextureResource, *CurrentRT->GetName(), CurrentRT->GameThread_GetRenderTargetResource());
+
+						DrawWeightmapComponentToRenderTargetMips(TextureTopLeftPositionInAtlas, CurrentRT, true, PSShaderParams);
+
+						int32 CurrentMip = 0;
+
+						for (int32 MipRTIndex = EWeightmapRTType::WeightmapRT_Mip0; MipRTIndex < EWeightmapRTType::WeightmapRT_Count; ++MipRTIndex)
+						{
+							CurrentRT = WeightmapRTList[MipRTIndex];
+
+							if (CurrentRT != nullptr)
+							{
+								CopyProceduralTexture(*CurrentRT->GetName(), CurrentRT->GameThread_GetRenderTargetResource(), OutputDebugName ? FString::Printf(TEXT("Weightmap Mip: %d"), CurrentMip) : TEXT(""), WeightmapTexture->Resource, WeightmapCPUReadBack, TextureTopLeftPositionInAtlas, CurrentMip, CurrentMip);
+								++CurrentMip;
+							}
+						}
+					}
+				}
+
+				++CurrentWeightmapToProcessIndex;
+			}
+
+			// Generate all the allocation with zero painted data for each components, it will be used to update the material instance to not sample those textures
+			/*if (ComputeShaderGeneratedData)
+			{
+				ClearFlagsAfterUpdate = GenerateZeroAllocationPerComponents(AllLandscapes, WeightmapLayersBlendSubstractive, ZeroAllocationsPerComponents);
+			}
+			*/
+		}
+
+		UpdateProceduralMaterialInstances(ProceduralUpdateAllMaterials ? AllLandscapeComponents : ComponentThatNeedMaterialRebuild, ZeroAllocationsPerComponents);
+		ProceduralUpdateAllMaterials = false;
+	}
+
+	if ((ProceduralContentUpdateFlags & EProceduralContentUpdateFlag::Weightmap_ResolveToTexture))
+	{
+		ResolveProceduralWeightmapTexture(AllLandscapes);
+	}
+
+	if ((ProceduralContentUpdateFlags & EProceduralContentUpdateFlag::Weightmap_Collision) != 0)
+	{
+		for (ULandscapeComponent* Component : AllLandscapeComponents)
+		{
+			Component->UpdateCollisionLayerData();
+		}
+	}
+
+	if (ClearFlagsAfterUpdate)
+	{
+		ProceduralContentUpdateFlags &= ~EProceduralContentUpdateFlag::Weightmap_All;
+	}
+
+	// If doing rendering debug, keep doing the render only
+	if (CVarOutputProceduralDebugDrawCallName.GetValueOnAnyThread() == 1)
+	{
+		ProceduralContentUpdateFlags |= EProceduralContentUpdateFlag::Weightmap_Render;
+	}
 }
 
-void ALandscape::RequestProceduralContentUpdate(uint32 InDataFlags)
+void ALandscape::UpdateProceduralMaterialInstances(const TArray<ULandscapeComponent*>& InComponentsToUpdate, const TMap<ULandscapeComponent*, TArray<ULandscapeLayerInfoObject*>>& InZeroAllocationsPerComponents)
 {
+	if (InComponentsToUpdate.Num() == 0 && InZeroAllocationsPerComponents.Num() == 0)
+	{
+		return;
+	}
+
+	TArray<ULandscapeComponent*> ComponentsToUpdate;
+	//InZeroAllocationsPerComponents.GenerateKeyArray(ComponentsToUpdate);
+	ComponentsToUpdate.Append(InComponentsToUpdate);	
+
+	SCOPE_CYCLE_COUNTER(STAT_LandscapeProceduralUpdateMaterialInstance);
+
+	// we're not having the material update context recreate render states because we will manually do it for only our components
+	TArray<FComponentRecreateRenderStateContext> RecreateRenderStateContexts;
+	RecreateRenderStateContexts.Reserve(ComponentsToUpdate.Num());
+
+	for (ULandscapeComponent* Component : ComponentsToUpdate)
+	{
+		RecreateRenderStateContexts.Emplace(Component);
+	}
+	TOptional<FMaterialUpdateContext> MaterialUpdateContext;
+	MaterialUpdateContext.Emplace(FMaterialUpdateContext::EOptions::Default & ~FMaterialUpdateContext::EOptions::RecreateRenderStates);
+
+	for (ULandscapeComponent* Component : ComponentsToUpdate)
+	{
+		int32 MaxLOD = FMath::CeilLogTwo(SubsectionSizeQuads + 1) - 1;
+		TMap<UMaterialInterface*, int8> NewMaterialPerLOD;
+		Component->LODIndexToMaterialIndex.SetNumUninitialized(MaxLOD + 1);
+		int8 LastLODIndex = INDEX_NONE;
+
+		UMaterialInterface* BaseMaterial = GetLandscapeMaterial();
+		UMaterialInterface* LOD0Material = GetLandscapeMaterial(0);
+
+		for (int32 LODIndex = 0; LODIndex <= MaxLOD; ++LODIndex)
+		{
+			UMaterialInterface* CurrentMaterial = GetLandscapeMaterial(LODIndex);
+
+			// if we have a LOD0 override, do not let the base material override it, it should override everything!
+			if (CurrentMaterial == BaseMaterial && BaseMaterial != LOD0Material)
+			{
+				CurrentMaterial = LOD0Material;
+			}
+
+			const int8* MaterialLOD = NewMaterialPerLOD.Find(CurrentMaterial);
+
+			if (MaterialLOD != nullptr)
+			{
+				Component->LODIndexToMaterialIndex[LODIndex] = *MaterialLOD > LastLODIndex ? *MaterialLOD : LastLODIndex;
+			}
+			else
+			{
+				int32 AddedIndex = NewMaterialPerLOD.Num();
+				NewMaterialPerLOD.Add(CurrentMaterial, LODIndex);
+				Component->LODIndexToMaterialIndex[LODIndex] = AddedIndex;
+				LastLODIndex = AddedIndex;
+			}
+		}
+
+		Component->MaterialPerLOD = NewMaterialPerLOD;
+
+		Component->MaterialInstances.SetNumZeroed(Component->MaterialPerLOD.Num()/* * 2*/); // over allocate in case we are using tessellation
+		Component->MaterialIndexToDisabledTessellationMaterial.Init(INDEX_NONE, MaxLOD + 1);
+		int8 TessellatedMaterialCount = 0;
+		int8 MaterialIndex = 0;
+
+		TArray<FWeightmapLayerAllocationInfo> WeightmapBaseLayerAllocation = Component->GetWeightmapLayerAllocations(); // We copy the array here
+		/*const TArray<ULandscapeLayerInfoObject*>* ZeroAllocations = InZeroAllocationsPerComponents.Find(Component);
+
+		// We have some allocation to remove
+		if (ZeroAllocations != nullptr)
+		{
+			for (int32 i = WeightmapBaseLayerAllocation.Num() - 1; i >= 0; --i)
+			{
+				const FWeightmapLayerAllocationInfo& Allocation = WeightmapBaseLayerAllocation[i];
+
+				if (ZeroAllocations->Contains(Allocation.LayerInfo))
+				{
+					WeightmapBaseLayerAllocation.RemoveAt(i);
+				}
+			}
+		}
+		*/
+
+		TArray<UTexture2D*>& ComponentWeightmapTextures = Component->GetWeightmapTextures();
+		UTexture2D* Heightmap = Component->GetHeightmap();
+
+		for (auto& ItPair : Component->MaterialPerLOD)
+		{
+			const int8 MaterialLOD = ItPair.Value;
+
+			// Find or set a matching MIC in the Landscape's map.
+			UMaterialInstanceConstant* CombinationMaterialInstance = Component->GetCombinationMaterial(nullptr, WeightmapBaseLayerAllocation, MaterialLOD, false);
+
+			if (CombinationMaterialInstance != nullptr)
+			{
+				UMaterialInstanceConstant* MaterialInstance = Component->MaterialInstances[MaterialIndex];
+				bool NeedToCreateMIC = MaterialInstance == nullptr;
+
+				if (NeedToCreateMIC)
+				{
+					// Create the instance for this component, that will use the layer combination instance.
+					MaterialInstance = NewObject<ULandscapeMaterialInstanceConstant>(GetOutermost());
+					Component->MaterialInstances[MaterialIndex] = MaterialInstance;
+				}
+
+				MaterialInstance->SetParentEditorOnly(CombinationMaterialInstance);
+
+				MaterialUpdateContext.GetValue().AddMaterialInstance(MaterialInstance); // must be done after SetParent				
+
+				FLinearColor Masks[4] = { FLinearColor(1.0f, 0.0f, 0.0f, 0.0f), FLinearColor(0.0f, 1.0f, 0.0f, 0.0f), FLinearColor(0.0f, 0.0f, 1.0f, 0.0f), FLinearColor(0.0f, 0.0f, 0.0f, 1.0f) };
+
+				// Set the layer mask
+				for (int32 AllocIdx = 0; AllocIdx < WeightmapBaseLayerAllocation.Num(); AllocIdx++)
+				{
+					FWeightmapLayerAllocationInfo& Allocation = WeightmapBaseLayerAllocation[AllocIdx];
+
+					FName LayerName = Allocation.LayerInfo == ALandscapeProxy::VisibilityLayer ? UMaterialExpressionLandscapeVisibilityMask::ParameterName : Allocation.LayerInfo ? Allocation.LayerInfo->LayerName : NAME_None;
+					MaterialInstance->SetVectorParameterValueEditorOnly(FName(*FString::Printf(TEXT("LayerMask_%s"), *LayerName.ToString())), Masks[Allocation.WeightmapTextureChannel]);
+				}
+
+				// Set the weightmaps
+				for (int32 i = 0; i < ComponentWeightmapTextures.Num(); i++)
+				{
+					MaterialInstance->SetTextureParameterValueEditorOnly(FName(*FString::Printf(TEXT("Weightmap%d"), i)), ComponentWeightmapTextures[i]);
+				}
+
+				if (NeedToCreateMIC)
+				{
+					MaterialInstance->PostEditChange();
+				}
+
+				/*// Setup material instance with disabled tessellation
+				if (CombinationMaterialInstance->GetMaterial()->D3D11TessellationMode != EMaterialTessellationMode::MTM_NoTessellation)
+				{
+					int32 TessellatedMaterialIndex = MaterialPerLOD.Num() + TessellatedMaterialCount++;
+					ULandscapeMaterialInstanceConstant* TessellationMaterialInstance = Cast<ULandscapeMaterialInstanceConstant>(MaterialInstances[TessellatedMaterialIndex]);
+
+					if (NeedToCreateMIC || TessellationMaterialInstance == nullptr)
+					{
+						TessellationMaterialInstance = NewObject<ULandscapeMaterialInstanceConstant>(GetOutermost());
+						TessellationMaterialInstance->SetParentEditorOnly(MaterialInstance);
+
+						MaterialInstances[TessellatedMaterialIndex] = TessellationMaterialInstance;
+						MaterialIndexToDisabledTessellationMaterial[MaterialIndex] = TessellatedMaterialIndex;
+
+						TessellationMaterialInstance->bDisableTessellation = true;
+						TessellationMaterialInstance->PostEditChange();
+					}
+
+					Context.AddMaterialInstance(TessellationMaterialInstance); // must be done after SetParent
+				}
+				*/
+			}
+
+			++MaterialIndex;
+		}
+
+		if (Component->MaterialPerLOD.Num() == 0)
+		{
+			Component->MaterialInstances.Empty(1);
+			Component->MaterialInstances.Add(nullptr);
+			Component->LODIndexToMaterialIndex.Empty(1);
+			Component->LODIndexToMaterialIndex.Add(0);
+		}
+	}
+
+	// End material update
+	MaterialUpdateContext.Reset();
+
+	// Recreate the render state for our components, needed to update the static drawlist which has cached the MaterialRenderProxies
+	// Must be after the FMaterialUpdateContext is destroyed
+	RecreateRenderStateContexts.Empty();
+}
+
+void ALandscape::ResolveProceduralWeightmapTexture(const TArray<ALandscapeProxy*>& InAllLandscapes)
+{
+	SCOPE_CYCLE_COUNTER(STAT_LandscapeResolveProceduralWeightmap);
+	       
+	for (ALandscapeProxy* Landscape : InAllLandscapes)
+	{
+		TArray<TArray<FColor>> MipData;
+
+		for (auto& ItPair : Landscape->WeightmapCPUReadBackTextures)
+		{
+			UTexture2D* OriginalWeightmap = ItPair.Key;
+			FLandscapeProceduralTexture2DCPUReadBackResource* WeightmapsCPUReadBack = ItPair.Value;
+
+			if (WeightmapsCPUReadBack == nullptr)
+			{
+				continue;
+			}
+
+			ResolveProceduralTexture(WeightmapsCPUReadBack, OriginalWeightmap);
+		}
+	}
+}
+
+void ALandscape::RequestProceduralContentUpdate(uint32 InDataFlags, bool InUpdateAllMaterials)
+{
 	ProceduralContentUpdateFlags = InDataFlags;
+	ProceduralUpdateAllMaterials = InUpdateAllMaterials;
 }
 
 void ALandscape::RegenerateProceduralContent()
@@ -1853,35 +4138,255 @@
 	if ((ProceduralContentUpdateFlags & Heightmap_Setup) != 0 || (ProceduralContentUpdateFlags & Weightmap_Setup) != 0)
 	{
 		SetupProceduralLayers();
+		ProceduralContentUpdateFlags &= ~(EProceduralContentUpdateFlag::All_Setup);
 	}
 
 	RegenerateProceduralHeightmaps();
 	RegenerateProceduralWeightmaps();
 }
 
-void ALandscape::OnPreSaveWorld(uint32 SaveFlags, UWorld* World)
+void ALandscape::TickProcedural(float DeltaTime, ELevelTick TickType, FActorTickFunction& ThisTickFunction)
 {
-	if (GetMutableDefault<UEditorExperimentalSettings>()->bProceduralLandscape)
+	check(GIsEditor);
+
+	UWorld* World = GetWorld();
+	if (World && !World->IsPlayInEditor())
 	{
-		// Need to perform setup here, as it's possible to get here with the data not setup, when doing a Save As on a level
-		if (PreviousExperimentalLandscapeProcedural != GetMutableDefault<UEditorExperimentalSettings>()->bProceduralLandscape)
+		if (GetMutableDefault<UEditorExperimentalSettings>()->bProceduralLandscape)
 		{
-			PreviousExperimentalLandscapeProcedural = GetMutableDefault<UEditorExperimentalSettings>()->bProceduralLandscape;
-			RequestProceduralContentUpdate(EProceduralContentUpdateFlag::All_Setup | EProceduralContentUpdateFlag::All_WithDDCUpdate);
+			if (PreviousExperimentalLandscapeProcedural != GetMutableDefault<UEditorExperimentalSettings>()->bProceduralLandscape)
+			{
+				PreviousExperimentalLandscapeProcedural = GetMutableDefault<UEditorExperimentalSettings>()->bProceduralLandscape;
+
+				RequestProceduralContentUpdate(EProceduralContentUpdateFlag::All_Setup);
+			}
+
+			// If doing editing while shader are compiling or at load of a map, it's possible we will need another update pass after shader are completed to see the correct result
+			const int32 RemainingShadersThisFrame = GShaderCompilingManager->GetNumRemainingJobs();
+
+			if (!WasCompilingShaders && RemainingShadersThisFrame > 0)
+			{
+				WasCompilingShaders = true;
+			}
+			else if (WasCompilingShaders)
+			{
+				WasCompilingShaders = false;
+				RequestProceduralContentUpdate(EProceduralContentUpdateFlag::All);
+			}
+
+			RegenerateProceduralContent();
 		}
 		else
 		{
-			RequestProceduralContentUpdate(EProceduralContentUpdateFlag::Heightmap_ResolveToTextureDDC | EProceduralContentUpdateFlag::Weightmap_ResolveToTextureDDC);
+			if (PreviousExperimentalLandscapeProcedural != GetMutableDefault<UEditorExperimentalSettings>()->bProceduralLandscape)
+			{
+				PreviousExperimentalLandscapeProcedural = GetMutableDefault<UEditorExperimentalSettings>()->bProceduralLandscape;
+
+			#if WITH_EDITORONLY_DATA
+				for (auto& ItPair : RenderDataPerHeightmap)
+				{
+					FRenderDataPerHeightmap& HeightmapRenderData = ItPair.Value;
+
+					if (HeightmapRenderData.HeightmapsCPUReadBack != nullptr)
+					{
+						BeginReleaseResource(HeightmapRenderData.HeightmapsCPUReadBack);
+					}
+				}
+
+				for (auto& ItPair : WeightmapCPUReadBackTextures)
+				{
+					FLandscapeProceduralTexture2DCPUReadBackResource* WeightmapCPUReadBack = ItPair.Value;
+
+					if (WeightmapCPUReadBack != nullptr)
+					{
+						BeginReleaseResource(WeightmapCPUReadBack);
+					}
+				}
+
+				if (CombinedProcLayerWeightmapAllLayersResource != nullptr)
+				{
+					BeginReleaseResource(CombinedProcLayerWeightmapAllLayersResource);
+				}
+
+				if (CurrentProcLayerWeightmapAllLayersResource != nullptr)
+				{
+					BeginReleaseResource(CurrentProcLayerWeightmapAllLayersResource);
+				}
+
+				if (WeightmapScratchExtractLayerTextureResource != nullptr)
+				{
+					BeginReleaseResource(WeightmapScratchExtractLayerTextureResource);
+				}
+
+				if (WeightmapScratchPackLayerTextureResource != nullptr)
+				{
+					BeginReleaseResource(WeightmapScratchPackLayerTextureResource);
+				}
+
+				FlushRenderingCommands();
+
+				for (auto& ItPair : RenderDataPerHeightmap)
+				{
+					FRenderDataPerHeightmap& HeightmapRenderData = ItPair.Value;
+
+					delete HeightmapRenderData.HeightmapsCPUReadBack;
+					HeightmapRenderData.HeightmapsCPUReadBack = nullptr;
+				}
+
+				for (auto& ItPair : WeightmapCPUReadBackTextures)
+				{
+					FLandscapeProceduralTexture2DCPUReadBackResource* WeightmapCPUReadBack = ItPair.Value;
+
+					delete WeightmapCPUReadBack;
+					WeightmapCPUReadBack = nullptr;
+				}
+
+				delete CombinedProcLayerWeightmapAllLayersResource;
+				delete CurrentProcLayerWeightmapAllLayersResource;
+				delete WeightmapScratchExtractLayerTextureResource;
+				delete WeightmapScratchPackLayerTextureResource;
+
+				CombinedProcLayerWeightmapAllLayersResource = nullptr;
+				CurrentProcLayerWeightmapAllLayersResource = nullptr;
+				WeightmapScratchExtractLayerTextureResource = nullptr;
+				WeightmapScratchPackLayerTextureResource = nullptr;
+			#endif
+			}
 		}
-
-		RegenerateProceduralContent();
-		ProceduralContentUpdateFlags = 0; // Force reset so we don't end up performing save info at the next Tick
 	}
 }
 
-void ALandscape::OnPostSaveWorld(uint32 SaveFlags, UWorld* World, bool bSuccess)
-{	
+#endif
+
+void ALandscapeProxy::BeginDestroy()
+{
+#if WITH_EDITORONLY_DATA
+	if (GetMutableDefault<UEditorExperimentalSettings>()->bProceduralLandscape)
+	{
+		for (auto& ItPair : RenderDataPerHeightmap)
+		{
+			FRenderDataPerHeightmap& HeightmapRenderData = ItPair.Value;
+
+			if (HeightmapRenderData.HeightmapsCPUReadBack != nullptr)
+			{
+				BeginReleaseResource(HeightmapRenderData.HeightmapsCPUReadBack);
+			}
+		}
+
+		for (auto& ItPair : WeightmapCPUReadBackTextures)
+		{
+			FLandscapeProceduralTexture2DCPUReadBackResource* WeightmapCPUReadBack = ItPair.Value;
+
+			if (WeightmapCPUReadBack != nullptr)
+			{
+				BeginReleaseResource(WeightmapCPUReadBack);
+			}
+		}
+
+		ReleaseResourceFence.BeginFence();
+	}
+#endif
+
+	Super::BeginDestroy();
 }
-#endif //WITH_EDITOR
 
+bool ALandscapeProxy::IsReadyForFinishDestroy()
+{
+	bool bReadyForFinishDestroy = Super::IsReadyForFinishDestroy();
+
+#if WITH_EDITORONLY_DATA
+	if (GetMutableDefault<UEditorExperimentalSettings>()->bProceduralLandscape)
+	{
+		if (bReadyForFinishDestroy)
+		{
+			bReadyForFinishDestroy = ReleaseResourceFence.IsFenceComplete();
+		}
+	}
+#endif
+
+	return bReadyForFinishDestroy;
+}
+
+void ALandscapeProxy::FinishDestroy()
+{
+#if WITH_EDITORONLY_DATA
+	if (GetMutableDefault<UEditorExperimentalSettings>()->bProceduralLandscape)
+	{
+		check(ReleaseResourceFence.IsFenceComplete());
+
+		for (auto& ItPair : RenderDataPerHeightmap)
+		{
+			FRenderDataPerHeightmap& HeightmapRenderData = ItPair.Value;
+
+			delete HeightmapRenderData.HeightmapsCPUReadBack;
+			HeightmapRenderData.HeightmapsCPUReadBack = nullptr;
+		}
+
+		for (auto& ItPair : WeightmapCPUReadBackTextures)
+		{
+			FLandscapeProceduralTexture2DCPUReadBackResource* WeightmapCPUReadBack = ItPair.Value;
+
+			delete WeightmapCPUReadBack;
+			WeightmapCPUReadBack = nullptr;
+		}
+	}
+#endif
+
+	Super::FinishDestroy();
+}
+
+void ALandscape::BeginDestroy()
+{
+#if WITH_EDITOR
+	if (GetMutableDefault<UEditorExperimentalSettings>()->bProceduralLandscape)
+	{
+		if (CombinedProcLayerWeightmapAllLayersResource != nullptr)
+		{
+			BeginReleaseResource(CombinedProcLayerWeightmapAllLayersResource);
+		}
+
+		if (CurrentProcLayerWeightmapAllLayersResource != nullptr)
+		{
+			BeginReleaseResource(CurrentProcLayerWeightmapAllLayersResource);
+		}
+
+		if (WeightmapScratchExtractLayerTextureResource != nullptr)
+		{
+			BeginReleaseResource(WeightmapScratchExtractLayerTextureResource);
+		}
+
+		if (WeightmapScratchPackLayerTextureResource != nullptr)
+		{
+			BeginReleaseResource(WeightmapScratchPackLayerTextureResource);
+		}
+
+		// Use ResourceFence from base class		
+	}
+#endif
+
+	Super::BeginDestroy();
+}
+
+void ALandscape::FinishDestroy()
+{
+#if WITH_EDITORONLY_DATA
+	if (GetMutableDefault<UEditorExperimentalSettings>()->bProceduralLandscape)
+	{
+		check(ReleaseResourceFence.IsFenceComplete());
+
+		delete CombinedProcLayerWeightmapAllLayersResource;
+		delete CurrentProcLayerWeightmapAllLayersResource;
+		delete WeightmapScratchExtractLayerTextureResource;
+		delete WeightmapScratchPackLayerTextureResource;
+
+		CombinedProcLayerWeightmapAllLayersResource = nullptr;
+		CurrentProcLayerWeightmapAllLayersResource = nullptr;
+		WeightmapScratchExtractLayerTextureResource = nullptr;
+		WeightmapScratchPackLayerTextureResource = nullptr;
+	}
+#endif
+
+	Super::FinishDestroy();
+}
+
 #undef LOCTEXT_NAMESPACE
