When SetParentEditorOnly is called we no longer recompute shader in most case as we will call UpdateStaticPermutation that will do so.

diff --git a/Engine/Source/Runtime/Landscape/Private/LandscapeGrass.cpp b/Engine/Source/Runtime/Landscape/Private/LandscapeGrass.cpp
--- a/Engine/Source/Runtime/Landscape/Private/LandscapeGrass.cpp
+++ b/Engine/Source/Runtime/Landscape/Private/LandscapeGrass.cpp
@@ -2120,7 +2120,9 @@ void ALandscapeProxy::UpdateGrass(const TArray<FVector>& Cameras, bool bForceSyn
 						{
 							CurrentForcedStreamedTextures.Add(Heightmap);
 						}
-						for (auto WeightmapTexture : Component->WeightmapTextures)
+						TArray<UTexture2D*>& ComponentWeightmapTextures = Component->GetWeightmapTextures();
+
+						for (auto WeightmapTexture : ComponentWeightmapTextures)
 						{
 							if (WeightmapTexture->bForceMiplevelsToBeResident)
 							{
@@ -2353,7 +2355,9 @@ void ALandscapeProxy::UpdateGrass(const TArray<FVector>& Cameras, bool bForceSyn
 											{
 												// we're ready to generate but our textures need streaming in
 												DesiredForceStreamedTextures.Add(Component->GetHeightmap());
-												for (auto WeightmapTexture : Component->WeightmapTextures)
+												TArray<UTexture2D*>& ComponentWeightmapTextures = Component->GetWeightmapTextures();
+												
+												for (auto WeightmapTexture : ComponentWeightmapTextures)
 												{
 													DesiredForceStreamedTextures.Add(WeightmapTexture);
 												}
@@ -2526,7 +2530,9 @@ void ALandscapeProxy::UpdateGrass(const TArray<FVector>& Cameras, bool bForceSyn
 								// Force stream in other heightmaps but only if we're not waiting for the textures 
 								// near the camera to stream in
 								DesiredForceStreamedTextures.Add(Component->GetHeightmap());
-								for (auto WeightmapTexture : Component->WeightmapTextures)
+								TArray<UTexture2D*>& ComponentWeightmapTextures = Component->GetWeightmapTextures();
+								
+								for (auto WeightmapTexture : ComponentWeightmapTextures)
 								{
 									DesiredForceStreamedTextures.Add(WeightmapTexture);
 								}
diff --git a/Engine/Source/Runtime/Landscape/Private/LandscapePrivate.h b/Engine/Source/Runtime/Landscape/Private/LandscapePrivate.h
--- a/Engine/Source/Runtime/Landscape/Private/LandscapePrivate.h
+++ b/Engine/Source/Runtime/Landscape/Private/LandscapePrivate.h
@@ -21,6 +21,18 @@ DECLARE_CYCLE_STAT_EXTERN(TEXT("PostInit View Custom Data"), STAT_LandscapePostI
 DECLARE_CYCLE_STAT_EXTERN(TEXT("Compute Custom Mesh Batch LOD"), STAT_LandscapeComputeCustomMeshBatchLOD, STATGROUP_Landscape, );
 DECLARE_CYCLE_STAT_EXTERN(TEXT("Compute Custom Shadow Mesh Batch LOD"), STAT_LandscapeComputeCustomShadowMeshBatchLOD, STATGROUP_Landscape, );
 
+DECLARE_CYCLE_STAT_EXTERN(TEXT("Regenerate Procedural (RenderThread)"), STAT_LandscapeRegenerateProcedural_RenderThread, STATGROUP_Landscape, );
+DECLARE_CYCLE_STAT_EXTERN(TEXT("Regenerate Procedural Heightmap (GameThread)"), STAT_LandscapeRegenerateProceduralHeightmaps, STATGROUP_Landscape, );
+DECLARE_CYCLE_STAT_EXTERN(TEXT("Resolve Procedural Heightmap"), STAT_LandscapeResolveProceduralHeightmap, STATGROUP_Landscape, );
+
+DECLARE_CYCLE_STAT_EXTERN(TEXT("Regenerate Procedural Weightmap (GameThread)"), STAT_LandscapeRegenerateProceduralWeightmaps, STATGROUP_Landscape, );
+
+DECLARE_CYCLE_STAT_EXTERN(TEXT("Resolve Procedural Weightmap"), STAT_LandscapeResolveProceduralWeightmap, STATGROUP_Landscape, );
+DECLARE_CYCLE_STAT_EXTERN(TEXT("Reallocate Procedural Weightmaps"), STAT_LandscapeReallocateProceduralWeightmaps, STATGROUP_Landscape, );
+DECLARE_CYCLE_STAT_EXTERN(TEXT("Update Procedural Material Instances"), STAT_LandscapeProceduralUpdateMaterialInstance, STATGROUP_Landscape, );
+
+DECLARE_DWORD_COUNTER_STAT_EXTERN(TEXT("Regenerate Procedural DrawCalls"), STAT_LandscapeRegenerateProceduralDrawCalls, STATGROUP_Landscape, );
+
 DECLARE_DWORD_COUNTER_STAT_EXTERN(TEXT("Components Using SubSection DrawCall"), STAT_LandscapeComponentUsingSubSectionDrawCalls, STATGROUP_Landscape, );
 DECLARE_DWORD_COUNTER_STAT_EXTERN(TEXT("Tessellated Shadow Cascade"), STAT_LandscapeTessellatedShadowCascade, STATGROUP_Landscape, );
 DECLARE_DWORD_COUNTER_STAT_EXTERN(TEXT("Tessellated Components"), STAT_LandscapeTessellatedComponents, STATGROUP_Landscape, );
@@ -28,12 +40,6 @@ DECLARE_DWORD_COUNTER_STAT_EXTERN(TEXT("Processed Triangles"), STAT_LandscapeTri
 DECLARE_DWORD_COUNTER_STAT_EXTERN(TEXT("Render Passes"), STAT_LandscapeComponentRenderPasses, STATGROUP_Landscape, );
 DECLARE_DWORD_COUNTER_STAT_EXTERN(TEXT("DrawCalls"), STAT_LandscapeDrawCalls, STATGROUP_Landscape, );
 
-DECLARE_CYCLE_STAT_EXTERN(TEXT("Regenerate Procedural Heightmap (GameThread)"), STAT_LandscapeRegenerateProceduralHeightmaps, STATGROUP_Landscape, );
-DECLARE_CYCLE_STAT_EXTERN(TEXT("Regenerate Procedural Heightmap (RenderThread)"), STAT_LandscapeRegenerateProceduralHeightmaps_RenderThread, STATGROUP_Landscape, );
-DECLARE_CYCLE_STAT_EXTERN(TEXT("Resolve Procedural Heightmap"), STAT_LandscapeResolveProceduralHeightmap, STATGROUP_Landscape, );
-
-DECLARE_DWORD_COUNTER_STAT_EXTERN(TEXT("Regenerate Procedural Heightmap DrawCalls"), STAT_LandscapeRegenerateProceduralHeightmapsDrawCalls, STATGROUP_Landscape, );
-
 DECLARE_MEMORY_STAT_EXTERN(TEXT("Vertex Mem"), STAT_LandscapeVertexMem, STATGROUP_Landscape, );
 DECLARE_MEMORY_STAT_EXTERN(TEXT("Occluder Mem"), STAT_LandscapeOccluderMem, STATGROUP_Landscape, );
 DECLARE_DWORD_COUNTER_STAT_EXTERN(TEXT("Component Mem"), STAT_LandscapeComponentMem, STATGROUP_Landscape, );
diff --git a/Engine/Source/Runtime/Landscape/Private/LandscapeRender.cpp b/Engine/Source/Runtime/Landscape/Private/LandscapeRender.cpp
--- a/Engine/Source/Runtime/Landscape/Private/LandscapeRender.cpp
+++ b/Engine/Source/Runtime/Landscape/Private/LandscapeRender.cpp
@@ -614,7 +614,7 @@ FLandscapeComponentSceneProxy::FLandscapeComponentSceneProxy(ULandscapeComponent
 	, FLandscapeNeighborInfo(InComponent->GetWorld(), InComponent->GetLandscapeProxy()->GetLandscapeGuid(), InComponent->GetSectionBase() / InComponent->ComponentSizeQuads, InComponent->GetHeightmap(), InComponent->ForcedLOD, InComponent->LODBias)
 	, MaxLOD(FMath::CeilLogTwo(InComponent->SubsectionSizeQuads + 1) - 1)
 	, UseTessellationComponentScreenSizeFalloff(InComponent->GetLandscapeProxy()->UseTessellationComponentScreenSizeFalloff)
-	, NumWeightmapLayerAllocations(InComponent->WeightmapLayerAllocations.Num())
+	, NumWeightmapLayerAllocations(InComponent->GetWeightmapLayerAllocations().Num())
 	, StaticLightingLOD(InComponent->GetLandscapeProxy()->StaticLightingLOD)
 	, WeightmapSubsectionOffset(InComponent->WeightmapSubsectionOffset)
 	, FirstLOD(0)
@@ -631,7 +631,7 @@ FLandscapeComponentSceneProxy::FLandscapeComponentSceneProxy(ULandscapeComponent
 	, SectionBase(InComponent->GetSectionBase())
 	, LandscapeComponent(InComponent)
 	, WeightmapScaleBias(InComponent->WeightmapScaleBias)
-	, WeightmapTextures(InComponent->WeightmapTextures)
+	, WeightmapTextures(InComponent->GetWeightmapTextures())
 	, NormalmapTexture(InComponent->GetHeightmap())
 	, BaseColorForGITexture(InComponent->GIBakedBaseColorTexture)
 	, HeightmapScaleBias(InComponent->HeightmapScaleBias)
@@ -839,7 +839,9 @@ FLandscapeComponentSceneProxy::FLandscapeComponentSceneProxy(ULandscapeComponent
 	bSupportsHeightfieldRepresentation = true;
 
 #if WITH_EDITOR
-	for (auto& Allocation : InComponent->WeightmapLayerAllocations)
+	TArray<FWeightmapLayerAllocationInfo>& ComponentWeightmapLayerAllocations = InComponent->GetWeightmapLayerAllocations();
+	
+	for (auto& Allocation : ComponentWeightmapLayerAllocations)
 	{
 		if (Allocation.LayerInfo != nullptr)
 		{
diff --git a/Engine/Source/Runtime/Landscape/Private/LandscapeRenderMobile.cpp b/Engine/Source/Runtime/Landscape/Private/LandscapeRenderMobile.cpp
--- a/Engine/Source/Runtime/Landscape/Private/LandscapeRenderMobile.cpp
+++ b/Engine/Source/Runtime/Landscape/Private/LandscapeRenderMobile.cpp
@@ -339,7 +339,7 @@ FLandscapeComponentSceneProxyMobile::FLandscapeComponentSceneProxyMobile(ULandsc
 	BlendableLayerMask = InComponent->MobileBlendableLayerMask;
 
 #if WITH_EDITOR
-	TArray<FWeightmapLayerAllocationInfo>& LayerAllocations = InComponent->MobileWeightmapLayerAllocations.Num() ? InComponent->MobileWeightmapLayerAllocations : InComponent->WeightmapLayerAllocations;
+	TArray<FWeightmapLayerAllocationInfo>& LayerAllocations = InComponent->MobileWeightmapLayerAllocations.Num() ? InComponent->MobileWeightmapLayerAllocations : InComponent->GetWeightmapLayerAllocations();
 	LayerColors.Empty();
 	for (auto& Allocation : LayerAllocations)
 	{
diff --git a/Engine/Source/Runtime/Windows/D3D11RHI/Private/D3D11RenderTarget.cpp b/Engine/Source/Runtime/Windows/D3D11RHI/Private/D3D11RenderTarget.cpp
--- a/Engine/Source/Runtime/Windows/D3D11RHI/Private/D3D11RenderTarget.cpp
+++ b/Engine/Source/Runtime/Windows/D3D11RHI/Private/D3D11RenderTarget.cpp
@@ -978,6 +978,21 @@ static void ConvertRAWSurfaceDataToFColor(DXGI_FORMAT Format, uint32 Width, uint
 			}
 		}
 	}
+	else if (Format == DXGI_FORMAT_R8_UNORM)
+	{
+		// Read the data out of the buffer, converting it from ABGR to ARGB.
+		for (uint32 Y = 0; Y < Height; Y++)
+		{
+			uint8* SrcPtr = (uint8*)(In + Y * SrcPitch);
+			FColor* DestPtr = Out + Y * Width;
+			for (uint32 X = 0; X < Width; X++)
+			{
+				*DestPtr = FColor(*SrcPtr, 0, 0, 0);
+				++SrcPtr;
+				++DestPtr;
+			}
+		}
+	}
 	else
 	{
 		// not supported yet
