Fixed Static Analysis issues

diff --git a/Engine/Source/Runtime/Landscape/Private/LandscapeEditProcedural.cpp b/Engine/Source/Runtime/Landscape/Private/LandscapeEditProcedural.cpp
--- a/Engine/Source/Runtime/Landscape/Private/LandscapeEditProcedural.cpp
+++ b/Engine/Source/Runtime/Landscape/Private/LandscapeEditProcedural.cpp
@@ -1819,13 +1819,18 @@ void ALandscapeProxy::SetupProceduralLayers(int32 InNumComponentsX, int32 InNumC
 
 void ALandscape::CopyProceduralTexture(UTexture* InSourceTexture, UTexture* InDestTexture, FTextureResource* InDestCPUResource, const FIntPoint& InFirstComponentSectionBase, uint8 InSourceCurrentMip, uint8 InDestCurrentMip, uint32 InSourceArrayIndex, uint32 InDestArrayIndex) const
 {
-	CopyProceduralTexture(InSourceTexture != nullptr ? InSourceTexture->GetName() : TEXT("Invalid"), InSourceTexture->Resource, InDestTexture != nullptr ? InDestTexture->GetName() : TEXT("Invalid"), InDestTexture->Resource, InDestCPUResource, 
-						  InFirstComponentSectionBase, InSourceCurrentMip, InDestCurrentMip, InSourceArrayIndex, InDestArrayIndex);
+	if (InSourceTexture != nullptr && InDestTexture != nullptr)
+	{
+		CopyProceduralTexture(InSourceTexture->GetName(), InSourceTexture->Resource, InDestTexture->GetName(), InDestTexture->Resource, InDestCPUResource, InFirstComponentSectionBase, InSourceCurrentMip, InDestCurrentMip, InSourceArrayIndex, InDestArrayIndex);
+	}
 }
 
 void ALandscape::CopyProceduralTexture(const FString& InSourceDebugName, FTextureResource* InSourceResource, const FString& InDestDebugName, FTextureResource* InDestResource, FTextureResource* InDestCPUResource, const FIntPoint& InFirstComponentSectionBase, 
 										uint8 InSourceCurrentMip, uint8 InDestCurrentMip, uint32 InSourceArrayIndex, uint32 InDestArrayIndex) const
 {
+	check(InSourceResource != nullptr);
+	check(InDestResource != nullptr);
+
 	FLandscapeProceduralCopyTexture_RenderThread CopyTexture(InSourceDebugName, InSourceResource, InDestDebugName, InDestResource, InDestCPUResource, InFirstComponentSectionBase, SubsectionSizeQuads, NumSubsections, InSourceCurrentMip, InDestCurrentMip, InSourceArrayIndex, InDestArrayIndex);
 
 	ENQUEUE_RENDER_COMMAND(FLandscapeProceduralCopyCommand)(
@@ -2594,13 +2599,14 @@ bool ALandscape::AreHeightmapTextureResourcesReady(const TArray<ALandscapeProxy*
 					LayerHeightmap->FinishCachePlatformData();
 
 					LayerHeightmap->Resource = LayerHeightmap->CreateResource();
-					if (LayerHeightmap->Resource)
+
+					if (LayerHeightmap->Resource != nullptr)
 					{
 						BeginInitResource(LayerHeightmap->Resource);
 					}
 				}
 
-				if (!LayerHeightmap->Resource->IsInitialized() || !LayerHeightmap->IsFullyStreamedIn())
+				if (LayerHeightmap->Resource == nullptr || !LayerHeightmap->Resource->IsInitialized() || !LayerHeightmap->IsFullyStreamedIn())
 				{
 					return false;
 				}
@@ -3492,7 +3498,8 @@ bool ALandscape::AreWeightmapTextureResourcesReady(const TArray<ALandscapeProxy*
 						Weightmap->FinishCachePlatformData();
 
 						Weightmap->Resource = Weightmap->CreateResource();
-						if (Weightmap->Resource)
+
+						if (Weightmap->Resource != nullptr)
 						{
 							BeginInitResource(Weightmap->Resource);
 						}
@@ -3513,7 +3520,7 @@ bool ALandscape::AreWeightmapTextureResourcesReady(const TArray<ALandscapeProxy*
 
 				for (UTexture2D* Weightmap : WeightmapLayerData.Weightmaps)
 				{
-					if (!Weightmap->Resource->IsInitialized() || !Weightmap->IsFullyStreamedIn())
+					if (Weightmap->Resource == nullptr || !Weightmap->Resource->IsInitialized() || !Weightmap->IsFullyStreamedIn())
 					{
 						return false;
 					}
