Fixed WeightmapUsage crash caused by doing an Undo and using a newed ptr in a TMap that was transacted

diff --git a/Engine/Source/Editor/LandscapeEditor/Private/LandscapeEdModeComponentTools.cpp b/Engine/Source/Editor/LandscapeEditor/Private/LandscapeEdModeComponentTools.cpp
--- a/Engine/Source/Editor/LandscapeEditor/Private/LandscapeEdModeComponentTools.cpp
+++ b/Engine/Source/Editor/LandscapeEditor/Private/LandscapeEdModeComponentTools.cpp
@@ -648,7 +648,7 @@ public:
 							// UE_LOG(LogLandscape, Log, TEXT("Still need %d channels"), TotalNeededChannels);
 
 							UTexture2D* CurrentWeightmapTexture = nullptr;
-							FLandscapeWeightmapUsage* CurrentWeightmapUsage = nullptr;
+							ULandscapeWeightmapUsage* CurrentWeightmapUsage = nullptr;
 
 							if (TotalNeededChannels < 4)
 							{
@@ -658,7 +658,7 @@ public:
 								int32 BestDistanceSquared = MAX_int32;
 								for (auto& WeightmapUsagePair : LandscapeProxy->WeightmapUsageMap)
 								{
-									FLandscapeWeightmapUsage* TryWeightmapUsage = WeightmapUsagePair.Value;
+									ULandscapeWeightmapUsage* TryWeightmapUsage = WeightmapUsagePair.Value;
 									if (TryWeightmapUsage->FreeChannelCount() >= TotalNeededChannels) // TODO: handle procedural layer
 									{
 										// See if this candidate is closer than any others we've found
@@ -695,7 +695,7 @@ public:
 								CurrentWeightmapTexture->PostEditChange();
 
 								// Store it in the usage map
-								CurrentWeightmapUsage = LandscapeProxy->WeightmapUsageMap.Add(CurrentWeightmapTexture, new FLandscapeWeightmapUsage());
+								CurrentWeightmapUsage = LandscapeProxy->WeightmapUsageMap.Add(CurrentWeightmapTexture, NewObject<ULandscapeWeightmapUsage>(LandscapeProxy));
 
 								// UE_LOG(LogLandscape, Log, TEXT("Making a new texture %s"), *CurrentWeightmapTexture->GetName());
 							}
@@ -726,7 +726,7 @@ public:
 										LandscapeEdit.ZeroTextureChannel(OldWeightmapTexture, AllocInfo.WeightmapTextureChannel);
 
 										// Remove the old allocation
-										FLandscapeWeightmapUsage** OldWeightmapUsage = Component->GetLandscapeProxy()->WeightmapUsageMap.Find(OldWeightmapTexture);
+										ULandscapeWeightmapUsage** OldWeightmapUsage = Component->GetLandscapeProxy()->WeightmapUsageMap.Find(OldWeightmapTexture);
 
 										if (OldWeightmapUsage != nullptr)
 										{
diff --git a/Engine/Source/Runtime/Landscape/Classes/LandscapeComponent.h b/Engine/Source/Runtime/Landscape/Classes/LandscapeComponent.h
--- a/Engine/Source/Runtime/Landscape/Classes/LandscapeComponent.h
+++ b/Engine/Source/Runtime/Landscape/Classes/LandscapeComponent.h
@@ -8,6 +8,7 @@
 #include "Engine/TextureStreamingTypes.h"
 #include "Components/PrimitiveComponent.h"
 #include "PerPlatformProperties.h"
+#include "LandscapeWeightmapUsage.h"
 
 #include "LandscapeComponent.generated.h"
 
@@ -251,46 +252,6 @@ struct FLandscapeComponentMaterialOverride
 	UMaterialInterface* Material;
 };
 
-/** Structure storing channel usage for weightmap textures */
-USTRUCT(NotBlueprintable)
-struct FLandscapeWeightmapUsage
-{
-	GENERATED_USTRUCT_BODY()
-
-	UPROPERTY()
-	ULandscapeComponent* ChannelUsage[4];
-
-	UPROPERTY()
-	FName ProceduralLayerName;
-
-	FLandscapeWeightmapUsage()
-	{
-		ClearUsage();
-	}
-
-	friend FArchive& operator<<(FArchive& Ar, FLandscapeWeightmapUsage& U);
-
-	int32 FreeChannelCount() const
-	{
-		int32 Count = 0;
-
-		for (int8 i = 0; i < 4; ++i)
-		{
-			Count += (ChannelUsage[i] == nullptr) ? 1 : 0;
-		}
-
-		return	Count;
-	}
-
-	void ClearUsage()
-	{
-		for (int8 i = 0; i < 4; ++i)
-		{
-			ChannelUsage[i] = nullptr;
-		}
-	}
-};
-
 UCLASS(hidecategories=(Display, Attachment, Physics, Debug, Collision, Movement, Rendering, PrimitiveComponent, Object, Transform, Mobility), showcategories=("Rendering|Material"), MinimalAPI, Within=LandscapeProxy)
 class ULandscapeComponent : public UPrimitiveComponent
 {
@@ -378,10 +339,10 @@ private:
 	UTexture2D* CurrentEditingHeightmapTexture;
 	TArray<FWeightmapLayerAllocationInfo>* CurrentEditingWeightmapLayerAllocations;
 	TArray<UTexture2D*>* CurrentEditingWeightmapTextures;
-	TArray<FLandscapeWeightmapUsage*>* CurrentEditingWeightmapTexturesUsage;
+	TArray<ULandscapeWeightmapUsage*>* CurrentEditingWeightmapTexturesUsage;
 	FName CurrentProceduralLayerName;
 
-	TArray<FLandscapeWeightmapUsage*> WeightmapTexturesUsage;
+	TArray<ULandscapeWeightmapUsage*> WeightmapTexturesUsage;
 #endif // WITH_EDITORONLY_DATA
 
 	/** Heightmap texture reference */
@@ -563,11 +524,11 @@ public:
 
 	LANDSCAPE_API void SetWeightmapLayerAllocations(const TArray<FWeightmapLayerAllocationInfo>& InNewWeightmapLayerAllocations);
 	LANDSCAPE_API void SetCurrentEditingWeightmapLayerAllocations(TArray<FWeightmapLayerAllocationInfo>* InNewWeightmapLayerAllocations);
-	LANDSCAPE_API void SetWeightmapTexturesUsage(const TArray<FLandscapeWeightmapUsage*>& InNewWeightmapTexturesUsage, bool InApplyToCurrentEditingWeightmap = false);
-	LANDSCAPE_API void SetCurrentEditingWeightmapTexturesUsage(TArray<FLandscapeWeightmapUsage*>* InNewWeightmapTexturesUsage);
+	LANDSCAPE_API void SetWeightmapTexturesUsage(const TArray<ULandscapeWeightmapUsage*>& InNewWeightmapTexturesUsage, bool InApplyToCurrentEditingWeightmap = false);
+	LANDSCAPE_API void SetCurrentEditingWeightmapTexturesUsage(TArray<ULandscapeWeightmapUsage*>* InNewWeightmapTexturesUsage);
 
-	LANDSCAPE_API TArray<FLandscapeWeightmapUsage*>& GetWeightmapTexturesUsage(bool InReturnCurrentEditingWeightmap = false);
-	LANDSCAPE_API const TArray<FLandscapeWeightmapUsage*>& GetWeightmapTexturesUsage(bool InReturnCurrentEditingWeightmap = false) const;
+	LANDSCAPE_API TArray<ULandscapeWeightmapUsage*>& GetWeightmapTexturesUsage(bool InReturnCurrentEditingWeightmap = false);
+	LANDSCAPE_API const TArray<ULandscapeWeightmapUsage*>& GetWeightmapTexturesUsage(bool InReturnCurrentEditingWeightmap = false) const;
 
 #endif 
 
diff --git a/Engine/Source/Runtime/Landscape/Classes/LandscapeProxy.h b/Engine/Source/Runtime/Landscape/Classes/LandscapeProxy.h
--- a/Engine/Source/Runtime/Landscape/Classes/LandscapeProxy.h
+++ b/Engine/Source/Runtime/Landscape/Classes/LandscapeProxy.h
@@ -14,6 +14,8 @@
 #include "PerPlatformProperties.h"
 #include "LandscapeBPCustomBrush.h"
 #include "LandscapeComponent.h"
+#include "LandscapeWeightmapUsage.h"
+
 #include "LandscapeProxy.generated.h"
 
 class ALandscape;
@@ -385,7 +387,7 @@ struct FWeightmapLayerData
 	UPROPERTY()
 	TArray<FWeightmapLayerAllocationInfo> WeightmapLayerAllocations;
 
-	TArray<FLandscapeWeightmapUsage*> WeightmapTextureUsages;	// Easy Access ref to data stored into the LandscapeProxy weightmap usage map
+	TArray<ULandscapeWeightmapUsage*> WeightmapTextureUsages;	// Easy Access ref to data stored into the LandscapeProxy weightmap usage map
 };
 
 USTRUCT()
@@ -696,7 +698,8 @@ public:
 #endif
 
 	/** Map of weightmap usage */
-	TMap<UTexture2D*, FLandscapeWeightmapUsage*> WeightmapUsageMap;
+	UPROPERTY(Transient)
+	TMap<UTexture2D*, ULandscapeWeightmapUsage*> WeightmapUsageMap;
 
 	// Blueprint functions
 
diff --git a/Engine/Source/Runtime/Landscape/Classes/LandscapeWeightmapUsage.h b/Engine/Source/Runtime/Landscape/Classes/LandscapeWeightmapUsage.h
new file mode 100644
--- /dev/null
+++ b/Engine/Source/Runtime/Landscape/Classes/LandscapeWeightmapUsage.h
@@ -0,0 +1,43 @@
+// Copyright 1998-2019 Epic Games, Inc. All Rights Reserved.
+
+#pragma once
+
+#include "CoreMinimal.h"
+#include "UObject/ObjectMacros.h"
+
+#include "LandscapeWeightmapUsage.generated.h"
+
+class ULandscapeComponent;
+
+UCLASS(MinimalAPI, NotBlueprintable)
+class ULandscapeWeightmapUsage : public UObject
+{
+	GENERATED_UCLASS_BODY()
+
+public:
+	UPROPERTY()
+	ULandscapeComponent* ChannelUsage[4];
+
+	UPROPERTY()
+	FName ProceduralLayerName;
+
+	int32 FreeChannelCount() const
+	{
+		int32 Count = 0;
+
+		for (int8 i = 0; i < 4; ++i)
+		{
+			Count += (ChannelUsage[i] == nullptr) ? 1 : 0;
+		}
+
+		return	Count;
+	}
+
+	void ClearUsage()
+	{
+		for (int8 i = 0; i < 4; ++i)
+		{
+			ChannelUsage[i] = nullptr;
+		}
+	}
+};
diff --git a/Engine/Source/Runtime/Landscape/Private/Landscape.cpp b/Engine/Source/Runtime/Landscape/Private/Landscape.cpp
--- a/Engine/Source/Runtime/Landscape/Private/Landscape.cpp
+++ b/Engine/Source/Runtime/Landscape/Private/Landscape.cpp
@@ -49,6 +49,7 @@ Landscape.cpp: Terrain rendering
 #include "Engine/Engine.h"
 #include "EngineUtils.h"
 #include "ComponentRecreateRenderStateContext.h"
+#include "LandscapeWeightmapUsage.h"
 
 #if WITH_EDITOR
 #include "MaterialUtilities.h"
@@ -235,7 +236,7 @@ void ULandscapeComponent::AddReferencedObjects(UObject* InThis, FReferenceCollec
 
 	if (This->CurrentEditingWeightmapTexturesUsage != nullptr)
 	{
-		for (FLandscapeWeightmapUsage* TextureUsage : *This->CurrentEditingWeightmapTexturesUsage)
+		for (ULandscapeWeightmapUsage* TextureUsage : *This->CurrentEditingWeightmapTexturesUsage)
 		{
 			if (TextureUsage != nullptr)
 			{
@@ -1215,7 +1216,7 @@ void ULandscapeComponent::BeginDestroy()
 			if (WeightmapTextures.IsValidIndex(WeightmapIndex))
 			{
 				UTexture2D* WeightmapTexture = WeightmapTextures[WeightmapIndex];
-				FLandscapeWeightmapUsage** Usage = Proxy->WeightmapUsageMap.Find(WeightmapTexture);
+				ULandscapeWeightmapUsage** Usage = Proxy->WeightmapUsageMap.Find(WeightmapTexture);
 				if (Usage != nullptr && (*Usage) != nullptr)
 				{
 					(*Usage)->ChannelUsage[WeightmapLayerAllocations[LayerIdx].WeightmapTextureChannel] = nullptr;
@@ -1223,7 +1224,6 @@ void ULandscapeComponent::BeginDestroy()
 					if ((*Usage)->FreeChannelCount() == 4)
 					{
 						Proxy->WeightmapUsageMap.Remove(WeightmapTexture);
-						delete (*Usage);
 					}
 				}
 			}
@@ -1464,7 +1464,7 @@ void ULandscapeComponent::SetCurrentEditingWeightmapLayerAllocations(TArray<FWei
 #endif
 }
 
-TArray<FLandscapeWeightmapUsage*>& ULandscapeComponent::GetWeightmapTexturesUsage(bool InReturnCurrentEditingWeightmap)
+TArray<ULandscapeWeightmapUsage*>& ULandscapeComponent::GetWeightmapTexturesUsage(bool InReturnCurrentEditingWeightmap)
 {
 #if WITH_EDITORONLY_DATA
 	if (InReturnCurrentEditingWeightmap && GetMutableDefault<UEditorExperimentalSettings>()->bProceduralLandscape)
@@ -1479,7 +1479,7 @@ TArray<FLandscapeWeightmapUsage*>& ULandscapeComponent::GetWeightmapTexturesUsag
 	return WeightmapTexturesUsage;
 }
 
-const TArray<FLandscapeWeightmapUsage*>& ULandscapeComponent::GetWeightmapTexturesUsage(bool InReturnCurrentEditingWeightmap) const
+const TArray<ULandscapeWeightmapUsage*>& ULandscapeComponent::GetWeightmapTexturesUsage(bool InReturnCurrentEditingWeightmap) const
 {
 #if WITH_EDITORONLY_DATA
 	if (InReturnCurrentEditingWeightmap && GetMutableDefault<UEditorExperimentalSettings>()->bProceduralLandscape)
@@ -1494,14 +1494,14 @@ const TArray<FLandscapeWeightmapUsage*>& ULandscapeComponent::GetWeightmapTextur
 	return WeightmapTexturesUsage;
 }
 
-void ULandscapeComponent::SetCurrentEditingWeightmapTexturesUsage(TArray<FLandscapeWeightmapUsage*>* InNewWeightmapTexturesUsage)
+void ULandscapeComponent::SetCurrentEditingWeightmapTexturesUsage(TArray<ULandscapeWeightmapUsage*>* InNewWeightmapTexturesUsage)
 {
 #if WITH_EDITORONLY_DATA
 	CurrentEditingWeightmapTexturesUsage = InNewWeightmapTexturesUsage;
 #endif
 }
 
-void ULandscapeComponent::SetWeightmapTexturesUsage(const TArray<FLandscapeWeightmapUsage*>& InNewWeightmapTexturesUsage, bool InApplyToCurrentEditingWeightmap )
+void ULandscapeComponent::SetWeightmapTexturesUsage(const TArray<ULandscapeWeightmapUsage*>& InNewWeightmapTexturesUsage, bool InApplyToCurrentEditingWeightmap )
 {
 #if WITH_EDITORONLY_DATA
 	if (InApplyToCurrentEditingWeightmap && CurrentEditingWeightmapTexturesUsage != nullptr)
@@ -1568,21 +1568,6 @@ void ALandscapeProxy::UnregisterAllComponents(const bool bForReregister)
 	Super::UnregisterAllComponents(bForReregister);
 }
 
-// FLandscapeWeightmapUsage serializer
-FArchive& operator<<(FArchive& Ar, FLandscapeWeightmapUsage& U)
-{
-	return Ar << U.ChannelUsage[0] << U.ChannelUsage[1] << U.ChannelUsage[2] << U.ChannelUsage[3];
-}
-
-FArchive& operator<<(FArchive& Ar, FLandscapeWeightmapUsage* U)
-{
-	if (U != nullptr)
-	{
-		return Ar << U->ChannelUsage[0] << U->ChannelUsage[1] << U->ChannelUsage[2] << U->ChannelUsage[3];
-	}
-
-	return Ar;
-}
 
 FArchive& operator<<(FArchive& Ar, FWeightmapLayerAllocationInfo& U)
 {
@@ -1732,13 +1717,6 @@ void ALandscapeProxy::Serialize(FArchive& Ar)
 			}
 		}
 	}
-
-#if WITH_EDITOR
-	if (Ar.IsTransacting())
-	{
-		Ar << WeightmapUsageMap;
-	}
-#endif
 }
 
 void ALandscapeProxy::AddReferencedObjects(UObject* InThis, FReferenceCollector& Collector)
@@ -1750,15 +1728,6 @@ void ALandscapeProxy::AddReferencedObjects(UObject* InThis, FReferenceCollector&
 #if WITH_EDITORONLY_DATA
 	Collector.AddReferencedObjects(This->MaterialInstanceConstantMap, This);
 #endif
-
-	for (auto It = This->WeightmapUsageMap.CreateIterator(); It; ++It)
-	{
-		Collector.AddReferencedObject(It.Key(), This);
-		Collector.AddReferencedObject(It.Value()->ChannelUsage[0], This);
-		Collector.AddReferencedObject(It.Value()->ChannelUsage[1], This);
-		Collector.AddReferencedObject(It.Value()->ChannelUsage[2], This);
-		Collector.AddReferencedObject(It.Value()->ChannelUsage[3], This);
-	}
 }
 
 #if WITH_EDITOR
@@ -2855,6 +2824,12 @@ void ULandscapeComponent::PostDuplicate(bool bDuplicateForPIE)
 	}
 }
 
+ULandscapeWeightmapUsage::ULandscapeWeightmapUsage(const FObjectInitializer& ObjectInitializer)
+	: Super(ObjectInitializer)
+{
+	ClearUsage();
+}
+
 // Generate a new guid to force a recache of all landscape derived data
 #define LANDSCAPE_FULL_DERIVEDDATA_VER			TEXT("016D326F3A954BBA9CCDFA00CEFA31E9")
 
diff --git a/Engine/Source/Runtime/Landscape/Private/LandscapeEdit.cpp b/Engine/Source/Runtime/Landscape/Private/LandscapeEdit.cpp
--- a/Engine/Source/Runtime/Landscape/Private/LandscapeEdit.cpp
+++ b/Engine/Source/Runtime/Landscape/Private/LandscapeEdit.cpp
@@ -564,11 +564,6 @@ void ULandscapeComponent::PostEditUndo()
 
 void ALandscapeProxy::FixupWeightmaps()
 {
-	for (auto& ItPair : WeightmapUsageMap)
-	{
-		delete ItPair.Value;
-	}
-
 	WeightmapUsageMap.Empty();
 
 	for (ULandscapeComponent* Component : LandscapeComponents)
@@ -651,15 +646,15 @@ void ULandscapeComponent::FixupWeightmaps()
 
 				UTexture2D* WeightmapTexture = WeightmapTextures[Allocation.WeightmapTextureIndex];
 
-				FLandscapeWeightmapUsage** TempUsage = Proxy->WeightmapUsageMap.Find(WeightmapTexture);
+				ULandscapeWeightmapUsage** TempUsage = Proxy->WeightmapUsageMap.Find(WeightmapTexture);
 
 				if (TempUsage == nullptr)
 				{
-					TempUsage = &Proxy->WeightmapUsageMap.Add(WeightmapTexture, new FLandscapeWeightmapUsage());
+					TempUsage = &Proxy->WeightmapUsageMap.Add(WeightmapTexture, NewObject<ULandscapeWeightmapUsage>(GetLandscapeProxy()));
 					(*TempUsage)->ProceduralLayerName = NAME_None;
 				}
 
-				FLandscapeWeightmapUsage* Usage = *TempUsage;
+				ULandscapeWeightmapUsage* Usage = *TempUsage;
 				WeightmapTexturesUsage[Allocation.WeightmapTextureIndex] = Usage; // Keep a ref to it for faster access
 
 				// Detect a shared layer allocation, caused by a previous undo or layer deletion bugs
@@ -2543,12 +2538,12 @@ LANDSCAPE_API void ALandscapeProxy::Import(
 
 				TArray<FWeightmapLayerAllocationInfo>& ComponentWeightmapLayerAllocations = LandscapeComponent->GetWeightmapLayerAllocations();
 				TArray<UTexture2D*>& ComponentWeightmapTextures = LandscapeComponent->GetWeightmapTextures();
-				TArray<FLandscapeWeightmapUsage*>& ComponentWeightmapTexturesUsage = LandscapeComponent->GetWeightmapTexturesUsage();
+				TArray<ULandscapeWeightmapUsage*>& ComponentWeightmapTexturesUsage = LandscapeComponent->GetWeightmapTexturesUsage();
 
 				if (BestAllocationIndex != -1)
 				{
 					FWeightmapTextureAllocation& Allocation = TextureAllocations[BestAllocationIndex];
-					FLandscapeWeightmapUsage* WeightmapUsage = WeightmapUsageMap.FindChecked(Allocation.Texture);
+					ULandscapeWeightmapUsage* WeightmapUsage = WeightmapUsageMap.FindChecked(Allocation.Texture);
 					ComponentWeightmapTexturesUsage.Add(WeightmapUsage);
 
 					UE_LOG(LogLandscape, Log, TEXT("  ==> Storing %d channels starting at %s[%d]"), RemainingLayers, *Allocation.Texture->GetName(), Allocation.ChannelsInUse);
@@ -2588,7 +2583,7 @@ LANDSCAPE_API void ALandscapeProxy::Import(
 
 					const int32 ThisAllocationLayers = FMath::Min<int32>(RemainingLayers, 4);
 					new(TextureAllocations) FWeightmapTextureAllocation(ComponentX, ComponentY, ThisAllocationLayers, WeightmapTexture, MipData);
-					FLandscapeWeightmapUsage* WeightmapUsage = WeightmapUsageMap.Add(WeightmapTexture, new FLandscapeWeightmapUsage());
+					ULandscapeWeightmapUsage* WeightmapUsage = WeightmapUsageMap.Add(WeightmapTexture, NewObject<ULandscapeWeightmapUsage>(this));
 					ComponentWeightmapTexturesUsage.Add(WeightmapUsage);
 
 					UE_LOG(LogLandscape, Log, TEXT("  ==> Storing %d channels in new texture %s"), ThisAllocationLayers, *WeightmapTexture->GetName());
@@ -4543,7 +4538,7 @@ void ULandscapeComponent::ReallocateWeightmaps(FLandscapeEditDataInterface* Data
 
 	TArray<FWeightmapLayerAllocationInfo>& ComponentWeightmapLayerAllocations = GetWeightmapLayerAllocations(InCanUseCurrentEditingWeightmap);
 	TArray<UTexture2D*>& ComponentWeightmapTextures = GetWeightmapTextures(InCanUseCurrentEditingWeightmap);
-	TArray<FLandscapeWeightmapUsage*>& ComponentWeightmapTexturesUsage = GetWeightmapTexturesUsage(InCanUseCurrentEditingWeightmap);
+	TArray<ULandscapeWeightmapUsage*>& ComponentWeightmapTexturesUsage = GetWeightmapTexturesUsage(InCanUseCurrentEditingWeightmap);
 
 	int32 NeededNewChannels = 0;
 	for (int32 LayerIdx = 0; LayerIdx < ComponentWeightmapLayerAllocations.Num(); LayerIdx++)
@@ -4573,7 +4568,7 @@ void ULandscapeComponent::ReallocateWeightmaps(FLandscapeEditDataInterface* Data
 	int32 ExistingTexAvailableChannels = 0;
 	for (int32 TexIdx = 0; TexIdx < ComponentWeightmapTextures.Num(); TexIdx++)
 	{
-		FLandscapeWeightmapUsage* Usage = ComponentWeightmapTexturesUsage[TexIdx];
+		ULandscapeWeightmapUsage* Usage = ComponentWeightmapTexturesUsage[TexIdx];
 		check(Usage);
 
 		ExistingTexAvailableChannels += Usage->FreeChannelCount();
@@ -4592,7 +4587,7 @@ void ULandscapeComponent::ReallocateWeightmaps(FLandscapeEditDataInterface* Data
 		int32 CurrentAlloc = 0;
 		for (int32 TexIdx = 0; TexIdx < ComponentWeightmapTextures.Num(); TexIdx++)
 		{
-			FLandscapeWeightmapUsage* Usage = ComponentWeightmapTexturesUsage[TexIdx];
+			ULandscapeWeightmapUsage* Usage = ComponentWeightmapTexturesUsage[TexIdx];
 
 			for (int32 ChanIdx = 0; ChanIdx < 4; ChanIdx++)
 			{
@@ -4641,13 +4636,13 @@ void ULandscapeComponent::ReallocateWeightmaps(FLandscapeEditDataInterface* Data
 	int32 TotalNeededChannels = ComponentWeightmapLayerAllocations.Num();
 	int32 CurrentLayer = 0;
 	TArray<UTexture2D*> NewWeightmapTextures;
-	TArray<FLandscapeWeightmapUsage*> NewComponentWeightmapTexturesUsage;
+	TArray<ULandscapeWeightmapUsage*> NewComponentWeightmapTexturesUsage;
 	while (TotalNeededChannels > 0)
 	{
 		// UE_LOG(LogLandscape, Log, TEXT("Still need %d channels"), TotalNeededChannels);
 
 		UTexture2D* CurrentWeightmapTexture = nullptr;
-		FLandscapeWeightmapUsage* CurrentWeightmapUsage = nullptr;
+		ULandscapeWeightmapUsage* CurrentWeightmapUsage = nullptr;
 
 		if (TotalNeededChannels < 4)
 		{
@@ -4657,7 +4652,7 @@ void ULandscapeComponent::ReallocateWeightmaps(FLandscapeEditDataInterface* Data
 			int32 BestDistanceSquared = MAX_int32;
 			for (auto& ItPair : Proxy->WeightmapUsageMap)
 			{
-				FLandscapeWeightmapUsage* TryWeightmapUsage = ItPair.Value;
+				ULandscapeWeightmapUsage* TryWeightmapUsage = ItPair.Value;
 				//
 				FName ProceduralLayerNameToSeek = InCanUseCurrentEditingWeightmap ? CurrentProceduralLayerName : NAME_None;
 
@@ -4721,7 +4716,7 @@ void ULandscapeComponent::ReallocateWeightmaps(FLandscapeEditDataInterface* Data
 			}
 
 			// Store it in the usage map
-			CurrentWeightmapUsage = Proxy->WeightmapUsageMap.Add(CurrentWeightmapTexture, new FLandscapeWeightmapUsage());
+			CurrentWeightmapUsage = Proxy->WeightmapUsageMap.Add(CurrentWeightmapTexture, NewObject<ULandscapeWeightmapUsage>(GetLandscapeProxy()));
 			CurrentWeightmapUsage->ProceduralLayerName = InCanUseCurrentEditingWeightmap ? CurrentProceduralLayerName : NAME_None;
 			// UE_LOG(LogLandscape, Log, TEXT("Making a new texture %s"), *CurrentWeightmapTexture->GetName());
 		}
@@ -4760,7 +4755,7 @@ void ULandscapeComponent::ReallocateWeightmaps(FLandscapeEditDataInterface* Data
 					}
 
 					// Remove the old allocation
-					FLandscapeWeightmapUsage* OldWeightmapUsage = ComponentWeightmapTexturesUsage[AllocInfo.WeightmapTextureIndex];
+					ULandscapeWeightmapUsage* OldWeightmapUsage = ComponentWeightmapTexturesUsage[AllocInfo.WeightmapTextureIndex];
 					OldWeightmapUsage->ChannelUsage[AllocInfo.WeightmapTextureChannel] = nullptr;
 				}
 
@@ -4803,10 +4798,10 @@ void ALandscapeProxy::RemoveInvalidWeightmaps()
 {
 	if (GIsEditor)
 	{
-		for (TMap< UTexture2D*, struct FLandscapeWeightmapUsage* >::TIterator It(WeightmapUsageMap); It; ++It)
+		for (TMap< UTexture2D*, ULandscapeWeightmapUsage* >::TIterator It(WeightmapUsageMap); It; ++It)
 		{
 			UTexture2D* Tex = It.Key();
-			FLandscapeWeightmapUsage* Usage = It.Value();
+			ULandscapeWeightmapUsage* Usage = It.Value();
 			if (Usage->FreeChannelCount() == 4) // Invalid Weight-map
 			{
 				if (Tex)
@@ -4817,7 +4812,6 @@ void ALandscapeProxy::RemoveInvalidWeightmaps()
 					Tex->ClearFlags(RF_Standalone);
 				}
 
-				delete Usage;
 				It.RemoveCurrent();
 			}
 		}
@@ -4835,7 +4829,7 @@ void ULandscapeComponent::RemoveInvalidWeightmaps()
 {
 	TArray<FWeightmapLayerAllocationInfo>& ComponentWeightmapLayerAllocations = GetWeightmapLayerAllocations();
 	TArray<UTexture2D*>& ComponentWeightmapTextures = GetWeightmapTextures();
-	TArray<FLandscapeWeightmapUsage*>& ComponentWeightmapTexturesUsage = GetWeightmapTexturesUsage();
+	TArray<ULandscapeWeightmapUsage*>& ComponentWeightmapTexturesUsage = GetWeightmapTexturesUsage();
 
 	// Adjust WeightmapTextureIndex index for other layers
 	TSet<int32> UnUsedTextureIndices;
diff --git a/Engine/Source/Runtime/Landscape/Private/LandscapeEditInterface.cpp b/Engine/Source/Runtime/Landscape/Private/LandscapeEditInterface.cpp
--- a/Engine/Source/Runtime/Landscape/Private/LandscapeEditInterface.cpp
+++ b/Engine/Source/Runtime/Landscape/Private/LandscapeEditInterface.cpp
@@ -1505,7 +1505,7 @@ void ULandscapeComponent::DeleteLayer(ULandscapeLayerInfoObject* LayerInfo, FLan
 	ULandscapeComponent* Component = this;
 	TArray<FWeightmapLayerAllocationInfo>& ComponentWeightmapLayerAllocations = Component->GetWeightmapLayerAllocations(true);
 	TArray<UTexture2D*>& ComponentWeightmapTextures = Component->GetWeightmapTextures(true);
-	TArray<FLandscapeWeightmapUsage*>& ComponentWeightmapTexturesUsage = Component->GetWeightmapTexturesUsage(true);
+	TArray<ULandscapeWeightmapUsage*>& ComponentWeightmapTexturesUsage = Component->GetWeightmapTexturesUsage(true);
 
 	// Find the index for this layer in this component.
 	const int32 DeleteLayerIdx = ComponentWeightmapLayerAllocations.IndexOfByPredicate(
@@ -1648,7 +1648,7 @@ void ULandscapeComponent::DeleteLayer(ULandscapeLayerInfoObject* LayerInfo, FLan
 	Component->Modify();
 	Proxy->Modify();
 
-	FLandscapeWeightmapUsage* Usage = ComponentWeightmapTexturesUsage.IsValidIndex(DeleteLayerAllocation.WeightmapTextureIndex) ? ComponentWeightmapTexturesUsage[DeleteLayerAllocation.WeightmapTextureIndex] : nullptr;
+	ULandscapeWeightmapUsage* Usage = ComponentWeightmapTexturesUsage.IsValidIndex(DeleteLayerAllocation.WeightmapTextureIndex) ? ComponentWeightmapTexturesUsage[DeleteLayerAllocation.WeightmapTextureIndex] : nullptr;
 	if (Usage) // can be null if WeightmapUsageMap hasn't been built yet
 	{
 		Usage->ChannelUsage[DeleteLayerAllocation.WeightmapTextureChannel] = nullptr;
@@ -1739,7 +1739,7 @@ void ULandscapeComponent::FillLayer(ULandscapeLayerInfoObject* LayerInfo, FLands
 
 	TArray<FWeightmapLayerAllocationInfo>& ComponentWeightmapLayerAllocations = Component->GetWeightmapLayerAllocations(true);
 	TArray<UTexture2D*>& ComponentWeightmapTextures = Component->GetWeightmapTextures(true);
-	TArray<FLandscapeWeightmapUsage*>& ComponentWeightmapTexturesUsage = Component->GetWeightmapTexturesUsage(true);
+	TArray<ULandscapeWeightmapUsage*>& ComponentWeightmapTexturesUsage = Component->GetWeightmapTexturesUsage(true);
 
 	// Find the index for this layer in this component.
 	int32 FillLayerIdx = ComponentWeightmapLayerAllocations.IndexOfByPredicate(
@@ -1814,7 +1814,7 @@ void ULandscapeComponent::FillLayer(ULandscapeLayerInfoObject* LayerInfo, FLands
 				continue;
 			}
 
-			FLandscapeWeightmapUsage* Usage = ComponentWeightmapTexturesUsage[Allocation.WeightmapTextureIndex];
+			ULandscapeWeightmapUsage* Usage = ComponentWeightmapTexturesUsage[Allocation.WeightmapTextureIndex];
 			if (Usage) // can be null if WeightmapUsageMap hasn't been built yet
 			{
 				Usage->ChannelUsage[Allocation.WeightmapTextureChannel] = nullptr;
@@ -1971,7 +1971,7 @@ void ULandscapeComponent::ReplaceLayer(ULandscapeLayerInfoObject* FromLayerInfo,
 
 	TArray<FWeightmapLayerAllocationInfo>& ComponentWeightmapLayerAllocations = GetWeightmapLayerAllocations(true);
 	TArray<UTexture2D*>& ComponentWeightmapTextures = GetWeightmapTextures(true);
-	TArray<FLandscapeWeightmapUsage*>& ComponentWeightmapTexturesUsage = GetWeightmapTexturesUsage(true);
+	TArray<ULandscapeWeightmapUsage*>& ComponentWeightmapTexturesUsage = GetWeightmapTexturesUsage(true);
 
 	// Find the index for this layer in this component.
 	int32 FromLayerIdx = INDEX_NONE;
@@ -2103,7 +2103,7 @@ void ULandscapeComponent::ReplaceLayer(ULandscapeLayerInfoObject* FromLayerInfo,
 	{
 		ALandscapeProxy* Proxy = GetLandscapeProxy();
 		// Mark the channel as unallocated, so we can reuse it later
-		FLandscapeWeightmapUsage* Usage = ComponentWeightmapTexturesUsage[FromLayerAllocation.WeightmapTextureIndex];
+		ULandscapeWeightmapUsage* Usage = ComponentWeightmapTexturesUsage[FromLayerAllocation.WeightmapTextureIndex];
 		//check(Usage);
 		if (Usage)
 		{
@@ -2349,11 +2349,11 @@ bool DeleteLayerIfAllZero(ULandscapeComponent* const Component, const uint8* con
 
 	TArray<FWeightmapLayerAllocationInfo>& ComponentWeightmapLayerAllocations = Component->GetWeightmapLayerAllocations(true);
 	TArray<UTexture2D*>& ComponentWeightmapTextures = Component->GetWeightmapTextures(true);
-	TArray<FLandscapeWeightmapUsage*>& ComponentWeightmapTexturesUsage = Component->GetWeightmapTexturesUsage(true);
+	TArray<ULandscapeWeightmapUsage*>& ComponentWeightmapTexturesUsage = Component->GetWeightmapTexturesUsage(true);
 
 	// Mark the channel as unallocated, so we can reuse it later
 	const int32 DeleteLayerWeightmapTextureIndex = ComponentWeightmapLayerAllocations[LayerIdx].WeightmapTextureIndex;
-	FLandscapeWeightmapUsage& Usage = *ComponentWeightmapTexturesUsage[DeleteLayerWeightmapTextureIndex];
+	ULandscapeWeightmapUsage& Usage = *ComponentWeightmapTexturesUsage[DeleteLayerWeightmapTextureIndex];
 	Usage.ChannelUsage[ComponentWeightmapLayerAllocations[LayerIdx].WeightmapTextureChannel] = NULL;
 
 	// Remove the layer as it's totally painted away.
diff --git a/Engine/Source/Runtime/Landscape/Private/LandscapeEditProcedural.cpp b/Engine/Source/Runtime/Landscape/Private/LandscapeEditProcedural.cpp
--- a/Engine/Source/Runtime/Landscape/Private/LandscapeEditProcedural.cpp
+++ b/Engine/Source/Runtime/Landscape/Private/LandscapeEditProcedural.cpp
@@ -1621,8 +1621,8 @@ void ALandscapeProxy::SetupProceduralLayers(int32 InNumComponentsX, int32 InNumC
 			struct FTextureData
 			{
 				UTexture2D* Texture;
-				FLandscapeWeightmapUsage* Usage;
-				FLandscapeWeightmapUsage* OriginalUsage;
+				ULandscapeWeightmapUsage* Usage;
+				ULandscapeWeightmapUsage* OriginalUsage;
 			};
 
 			TMap<UTexture2D*, FTextureData> ProcessedTextures;
@@ -1642,7 +1642,7 @@ void ALandscapeProxy::SetupProceduralLayers(int32 InNumComponentsX, int32 InNumC
 
 						const TArray<UTexture2D*>& ComponentWeightmapTextures = Component->GetWeightmapTextures();
 						TArray<FWeightmapLayerAllocationInfo>& ComponentLayerAllocations = Component->GetWeightmapLayerAllocations();
-						TArray<FLandscapeWeightmapUsage*>& ComponentWeightmapTexturesUsage = Component->GetWeightmapTexturesUsage();
+						TArray<ULandscapeWeightmapUsage*>& ComponentWeightmapTexturesUsage = Component->GetWeightmapTexturesUsage();
 
 						NewWeightmapData.Weightmaps.AddDefaulted(ComponentWeightmapTextures.Num());
 						NewWeightmapData.WeightmapTextureUsages.AddDefaulted(ComponentWeightmapTexturesUsage.Num());
@@ -1711,7 +1711,7 @@ void ALandscapeProxy::SetupProceduralLayers(int32 InNumComponentsX, int32 InNumC
 								NewWeightmapData.WeightmapTextureUsages[TextureIndex]->ProceduralLayerName = ProceduralLayerName;
 
 								// Create new Usage for the base layer as the other one will now be used by the Layer 1
-								ComponentWeightmapTexturesUsage[TextureIndex] = LandscapeProxy->WeightmapUsageMap.Add(NewWeightmapTexture, new FLandscapeWeightmapUsage());
+								ComponentWeightmapTexturesUsage[TextureIndex] = LandscapeProxy->WeightmapUsageMap.Add(NewWeightmapTexture, NewObject<ULandscapeWeightmapUsage>(LandscapeProxy));
 
 								for (FWeightmapLayerAllocationInfo& Allocation : ComponentLayerAllocations)
 								{
@@ -1743,15 +1743,15 @@ void ALandscapeProxy::SetupProceduralLayers(int32 InNumComponentsX, int32 InNumC
 					{
 						FWeightmapLayerAllocationInfo& Allocation = WeightmapLayer->WeightmapLayerAllocations[LayerIdx];
 						UTexture2D* WeightmapTexture = WeightmapLayer->Weightmaps[Allocation.WeightmapTextureIndex];
-						FLandscapeWeightmapUsage** TempUsage = LandscapeProxy->WeightmapUsageMap.Find(WeightmapTexture);
+						ULandscapeWeightmapUsage** TempUsage = LandscapeProxy->WeightmapUsageMap.Find(WeightmapTexture);
 
 						if (TempUsage == nullptr)
 						{
-							TempUsage = &LandscapeProxy->WeightmapUsageMap.Add(WeightmapTexture, new FLandscapeWeightmapUsage());							
+							TempUsage = &LandscapeProxy->WeightmapUsageMap.Add(WeightmapTexture, NewObject<ULandscapeWeightmapUsage>(LandscapeProxy));
 							(*TempUsage)->ProceduralLayerName = ProceduralLayerName;
 						}
 
-						FLandscapeWeightmapUsage* Usage = *TempUsage;
+						ULandscapeWeightmapUsage* Usage = *TempUsage;
 						WeightmapLayer->WeightmapTextureUsages[Allocation.WeightmapTextureIndex] = Usage; // Keep a ref to it for faster access
 
 						check(Usage->ChannelUsage[Allocation.WeightmapTextureChannel] == nullptr || Usage->ChannelUsage[Allocation.WeightmapTextureChannel] == Component);
@@ -2893,7 +2893,7 @@ void ALandscape::PrepareProceduralComponentDataForExtractLayersCS(const FName& I
 					UTexture2D* Weightmap = WeightLayerData.Weightmaps[InCurrentWeightmapToProcessIndex];
 					check(Weightmap != nullptr);
 
-					const FLandscapeWeightmapUsage* WeightmapUsage = WeightLayerData.WeightmapTextureUsages[InCurrentWeightmapToProcessIndex];
+					const ULandscapeWeightmapUsage* WeightmapUsage = WeightLayerData.WeightmapTextureUsages[InCurrentWeightmapToProcessIndex];
 					check(WeightmapUsage != nullptr);
 
 					CopyProceduralTexture(*Weightmap->GetName(), Weightmap->Resource, InOutputDebugName ? FString::Printf(TEXT("%s WeightmapScratchTexture"), *InProceduralLayerName.ToString()) : TEXT(""), InOutTextureData, nullptr, Component->GetSectionBase(), 0);
@@ -2972,9 +2972,9 @@ void ALandscape::PrepareProceduralComponentDataForPackLayersCS(int32 InCurrentWe
 
 				InOutProcessedWeightmapCPUCopy.Add(*WeightmapCPUCopy);
 
-				const TArray<FLandscapeWeightmapUsage*>& WeightmapTexturesUsage = Component->GetWeightmapTexturesUsage();
+				const TArray<ULandscapeWeightmapUsage*>& WeightmapTexturesUsage = Component->GetWeightmapTexturesUsage();
 
-				const FLandscapeWeightmapUsage* WeightmapUsage = WeightmapTexturesUsage[InCurrentWeightmapToProcessIndex];
+				const ULandscapeWeightmapUsage* WeightmapUsage = WeightmapTexturesUsage[InCurrentWeightmapToProcessIndex];
 				check(WeightmapUsage != nullptr);
 
 				TArray<const FWeightmapLayerAllocationInfo*> AlreadyProcessedAllocation;
@@ -3053,30 +3053,29 @@ void ALandscape::ReallocateProceduralWeightmaps(const TArray<ALandscapeProxy*>&
 	}
 
 	// Copy Previous Usage, to know which texture need updating
-	TMap<UTexture2D*, FLandscapeWeightmapUsage> CurrentWeightmapsUsage;
+	TMap<UTexture2D*, ULandscapeWeightmapUsage*> CurrentWeightmapsUsage;
 
 	for (ULandscapeComponent* Component : AllLandscapeComponents)
 	{
 		TArray<UTexture2D*>& ComponentWeightmapTextures = Component->GetWeightmapTextures();
-		TArray<FLandscapeWeightmapUsage*>& ComponentWeightmapTextureUsage = Component->GetWeightmapTexturesUsage();
+		TArray<ULandscapeWeightmapUsage*>& ComponentWeightmapTextureUsage = Component->GetWeightmapTexturesUsage();
 
 		for (int32 i = 0; i < ComponentWeightmapTextures.Num(); ++i)
 		{
 			UTexture2D* ComponentWeightmapTexture = ComponentWeightmapTextures[i];
-			FLandscapeWeightmapUsage* CurrentWeightmapTextureUsage = CurrentWeightmapsUsage.Find(ComponentWeightmapTexture);
+			ULandscapeWeightmapUsage** CurrentWeightmapTextureUsage = CurrentWeightmapsUsage.Find(ComponentWeightmapTexture);
 
 			if (CurrentWeightmapTextureUsage == nullptr)
 			{
-				FLandscapeWeightmapUsage* ComponentWeightmapUsage = ComponentWeightmapTextureUsage[i];
-
-				FLandscapeWeightmapUsage Usage;
+				ULandscapeWeightmapUsage* ComponentWeightmapUsage = ComponentWeightmapTextureUsage[i];
+				ULandscapeWeightmapUsage* Usage = NewObject<ULandscapeWeightmapUsage>(Component->GetLandscapeProxy());
 
 				for (int32 j = 0; j < 4; ++j)
 				{
-					Usage.ChannelUsage[j] = ComponentWeightmapUsage->ChannelUsage[j];
+					Usage->ChannelUsage[j] = ComponentWeightmapUsage->ChannelUsage[j];
 				}
 
-				CurrentWeightmapsUsage.Add(ComponentWeightmapTexture, MoveTemp(Usage));
+				CurrentWeightmapsUsage.Add(ComponentWeightmapTexture, Usage);
 			}
 		}
 	}
@@ -3092,11 +3091,11 @@ void ALandscape::ReallocateProceduralWeightmaps(const TArray<ALandscapeProxy*>&
 			BaseWeightmapAllocation.WeightmapTextureIndex = 255;
 		}
 		
-		TArray<FLandscapeWeightmapUsage*>& WeightmapTexturesUsage = Component->GetWeightmapTexturesUsage();
+		TArray<ULandscapeWeightmapUsage*>& WeightmapTexturesUsage = Component->GetWeightmapTexturesUsage();
 
 		for (int32 i = 0; i < WeightmapTexturesUsage.Num(); ++i)
 		{
-			FLandscapeWeightmapUsage* Usage = WeightmapTexturesUsage[i];
+			ULandscapeWeightmapUsage* Usage = WeightmapTexturesUsage[i];
 			check(Usage != nullptr);
 
 			Usage->ClearUsage();
@@ -3207,19 +3206,19 @@ void ALandscape::ReallocateProceduralWeightmaps(const TArray<ALandscapeProxy*>&
 		}
 
 		TArray<UTexture2D*>& ComponentWeightmapTextures = Component->GetWeightmapTextures();
-		TArray<FLandscapeWeightmapUsage*>& ComponentWeightmapTextureUsage = Component->GetWeightmapTexturesUsage();
+		TArray<ULandscapeWeightmapUsage*>& ComponentWeightmapTextureUsage = Component->GetWeightmapTexturesUsage();
 
 		for (int32 i = 0; i < ComponentWeightmapTextures.Num(); ++i)
 		{
 			UTexture2D* ComponentWeightmapTexture = ComponentWeightmapTextures[i];
-			FLandscapeWeightmapUsage* ComponentWeightmapUsage = ComponentWeightmapTextureUsage[i];
-			FLandscapeWeightmapUsage* CurrentWeightmapTextureUsage = CurrentWeightmapsUsage.Find(ComponentWeightmapTexture);
+			ULandscapeWeightmapUsage* ComponentWeightmapUsage = ComponentWeightmapTextureUsage[i];
+			ULandscapeWeightmapUsage** CurrentWeightmapTextureUsage = CurrentWeightmapsUsage.Find(ComponentWeightmapTexture);
 
 			if (CurrentWeightmapTextureUsage != nullptr)
 			{
 				for (int32 j = 0; j < 4; ++j)
 				{
-					if (ComponentWeightmapUsage->ChannelUsage[j] != CurrentWeightmapTextureUsage->ChannelUsage[j] && ComponentWeightmapUsage->ChannelUsage[j] != nullptr)
+					if (ComponentWeightmapUsage->ChannelUsage[j] != (*CurrentWeightmapTextureUsage)->ChannelUsage[j] && ComponentWeightmapUsage->ChannelUsage[j] != nullptr)
 					{
 						OutComponentThatNeedMaterialRebuild.AddUnique(ComponentWeightmapUsage->ChannelUsage[j]);
 					}
