Removed usage of Procedural Layer's name as key in Landscape Procedural Layer system (ULandscapeWeightmapUsage, ALandscapeProxy's ProceduralLayersData map, ULandscapeComponent's CurrentProceduralLayer) and use Layer's Guid instead.

diff --git a/Engine/Source/Editor/LandscapeEditor/Private/LandscapeEdMode.cpp b/Engine/Source/Editor/LandscapeEditor/Private/LandscapeEdMode.cpp
--- a/Engine/Source/Editor/LandscapeEditor/Private/LandscapeEdMode.cpp
+++ b/Engine/Source/Editor/LandscapeEditor/Private/LandscapeEdMode.cpp
@@ -4183,11 +4183,6 @@ void FEdModeLandscape::SetProceduralLayerName(int32 InLayerIndex, const FName& I
 	}
 }
 
-FName FEdModeLandscape::GetCurrentProceduralLayerName() const
-{
-	return GetProceduralLayerName(CurrentToolTarget.CurrentProceduralLayerIndex);
-}
-
 float FEdModeLandscape::GetProceduralLayerAlpha(int32 InLayerIndex) const
 {
 	FProceduralLayer* Layer = GetProceduralLayer(InLayerIndex);
diff --git a/Engine/Source/Editor/LandscapeEditor/Private/LandscapeEdMode.h b/Engine/Source/Editor/LandscapeEditor/Private/LandscapeEdMode.h
--- a/Engine/Source/Editor/LandscapeEditor/Private/LandscapeEdMode.h
+++ b/Engine/Source/Editor/LandscapeEditor/Private/LandscapeEdMode.h
@@ -493,7 +493,6 @@ public:
 	int32 GetProceduralLayerCount() const;
 	void SetCurrentProceduralLayer(int32 InLayerIndex);
 	int32 GetCurrentProceduralLayerIndex() const;
-	FName GetCurrentProceduralLayerName() const;
 	ALandscape* GetLandscape() const;
 	struct FProceduralLayer* GetProceduralLayer(int32 InLayerIndex) const;
 	FName GetProceduralLayerName(int32 InLayerIndex) const;
diff --git a/Engine/Source/Runtime/Landscape/Classes/Landscape.h b/Engine/Source/Runtime/Landscape/Classes/Landscape.h
--- a/Engine/Source/Runtime/Landscape/Classes/Landscape.h
+++ b/Engine/Source/Runtime/Landscape/Classes/Landscape.h
@@ -258,7 +258,7 @@ private:
 
 	void UpdateProceduralMaterialInstances(const TArray<ULandscapeComponent*>& InComponentsToUpdate, const TMap<ULandscapeComponent*, TArray<ULandscapeLayerInfoObject*>>& InZeroAllocationsPerComponents);
 
-	void PrepareProceduralComponentDataForExtractLayersCS(const FName& InProceduralLayerName, int32 InCurrentWeightmapToProcessIndex, bool InOutputDebugName, const TArray<ALandscapeProxy*>& InAllLandscape, class FLandscapeTexture2DResource* InOutTextureData,
+	void PrepareProceduralComponentDataForExtractLayersCS(const FProceduralLayer& InProceduralLayer, int32 InCurrentWeightmapToProcessIndex, bool InOutputDebugName, const TArray<ALandscapeProxy*>& InAllLandscape, class FLandscapeTexture2DResource* InOutTextureData,
 														  TArray<struct FLandscapeProceduralWeightmapExtractLayersComponentData>& OutComponentData, TMap<ULandscapeLayerInfoObject*, int32>& OutLayerInfoObjects);
 	void PrepareProceduralComponentDataForPackLayersCS(int32 InCurrentWeightmapToProcessIndex, bool InOutputDebugName, const TArray<ULandscapeComponent*>& InAllLandscapeComponents, 
 													   TArray<UTexture2D*>& InOutProcessedWeightmaps, TArray<class FLandscapeProceduralTexture2DCPUReadBackResource*>& InOutProcessedWeightmapCPUCopy, TArray<struct FLandscapeProceduralWeightmapPackLayersComponentData>& OutComponentData);
diff --git a/Engine/Source/Runtime/Landscape/Classes/LandscapeComponent.h b/Engine/Source/Runtime/Landscape/Classes/LandscapeComponent.h
--- a/Engine/Source/Runtime/Landscape/Classes/LandscapeComponent.h
+++ b/Engine/Source/Runtime/Landscape/Classes/LandscapeComponent.h
@@ -340,7 +340,7 @@ private:
 	TArray<FWeightmapLayerAllocationInfo>* CurrentEditingWeightmapLayerAllocations;
 	TArray<UTexture2D*>* CurrentEditingWeightmapTextures;
 	TArray<ULandscapeWeightmapUsage*>* CurrentEditingWeightmapTexturesUsage;
-	FName CurrentProceduralLayerName;
+	FGuid CurrentProceduralLayerGuid;
 
 	TArray<ULandscapeWeightmapUsage*> WeightmapTexturesUsage;
 #endif // WITH_EDITORONLY_DATA
@@ -515,21 +515,16 @@ public:
 
 #if WITH_EDITOR
 	LANDSCAPE_API void SetHeightmap(UTexture2D* NewHeightmap);
-	LANDSCAPE_API void SetCurrentEditingHeightmap(UTexture2D* InNewHeightmap);
 
 	LANDSCAPE_API void SetWeightmapTextures(const TArray<UTexture2D*>& InNewWeightmapTextures, bool InApplyToCurrentEditingWeightmap = false);
-	LANDSCAPE_API void SetCurrentEditingWeightmaps(TArray<UTexture2D*>* InNewWeightmapTextures);
-	LANDSCAPE_API void SetCurrentProceduralLayerName(const FName& InName);
-	LANDSCAPE_API const FName& GetCurrentProceduralLayerName() const;
+	LANDSCAPE_API const FGuid& GetCurrentProceduralLayerGuid() const;
 
 	LANDSCAPE_API void SetWeightmapLayerAllocations(const TArray<FWeightmapLayerAllocationInfo>& InNewWeightmapLayerAllocations);
-	LANDSCAPE_API void SetCurrentEditingWeightmapLayerAllocations(TArray<FWeightmapLayerAllocationInfo>* InNewWeightmapLayerAllocations);
 	LANDSCAPE_API void SetWeightmapTexturesUsage(const TArray<ULandscapeWeightmapUsage*>& InNewWeightmapTexturesUsage, bool InApplyToCurrentEditingWeightmap = false);
-	LANDSCAPE_API void SetCurrentEditingWeightmapTexturesUsage(TArray<ULandscapeWeightmapUsage*>* InNewWeightmapTexturesUsage);
-
 	LANDSCAPE_API TArray<ULandscapeWeightmapUsage*>& GetWeightmapTexturesUsage(bool InReturnCurrentEditingWeightmap = false);
 	LANDSCAPE_API const TArray<ULandscapeWeightmapUsage*>& GetWeightmapTexturesUsage(bool InReturnCurrentEditingWeightmap = false) const;
 
+	LANDSCAPE_API void SetCurrentEditingProceduralLayer(struct FProceduralLayer* Layer, struct FProceduralLayerData* LayerData);
 #endif 
 
 #if WITH_EDITOR
@@ -803,6 +798,14 @@ public:
 
 protected:
 
+#if WITH_EDITOR
+	void SetCurrentEditingHeightmap(UTexture2D* InNewHeightmap);
+	void SetCurrentEditingWeightmaps(TArray<UTexture2D*>* InNewWeightmapTextures);
+	void SetCurrentProceduralLayerGuid(const FGuid& InLayerGuid);
+	void SetCurrentEditingWeightmapLayerAllocations(TArray<FWeightmapLayerAllocationInfo>* InNewWeightmapLayerAllocations);
+	void SetCurrentEditingWeightmapTexturesUsage(TArray<ULandscapeWeightmapUsage*>* InNewWeightmapTexturesUsage);
+#endif
+
 	/** Whether the component type supports static lighting. */
 	virtual bool SupportsStaticLighting() const override
 	{
diff --git a/Engine/Source/Runtime/Landscape/Classes/LandscapeProxy.h b/Engine/Source/Runtime/Landscape/Classes/LandscapeProxy.h
--- a/Engine/Source/Runtime/Landscape/Classes/LandscapeProxy.h
+++ b/Engine/Source/Runtime/Landscape/Classes/LandscapeProxy.h
@@ -643,7 +643,7 @@ public:
 	TArray<FLandscapeEditorLayerSettings> EditorLayerSettings;
 
 	UPROPERTY(TextExportTransient)
-	TMap<FName, FProceduralLayerData> ProceduralLayersData;
+	TMap<FGuid, FProceduralLayerData> ProceduralLayersData;
 
 	UPROPERTY()
 	bool HasProceduralContent;
diff --git a/Engine/Source/Runtime/Landscape/Classes/LandscapeWeightmapUsage.h b/Engine/Source/Runtime/Landscape/Classes/LandscapeWeightmapUsage.h
--- a/Engine/Source/Runtime/Landscape/Classes/LandscapeWeightmapUsage.h
+++ b/Engine/Source/Runtime/Landscape/Classes/LandscapeWeightmapUsage.h
@@ -19,7 +19,7 @@ public:
 	ULandscapeComponent* ChannelUsage[4];
 
 	UPROPERTY()
-	FName ProceduralLayerName;
+	FGuid ProceduralLayerGuid;
 
 	int32 FreeChannelCount() const
 	{
diff --git a/Engine/Source/Runtime/Landscape/Private/Landscape.cpp b/Engine/Source/Runtime/Landscape/Private/Landscape.cpp
--- a/Engine/Source/Runtime/Landscape/Private/Landscape.cpp
+++ b/Engine/Source/Runtime/Landscape/Private/Landscape.cpp
@@ -169,7 +169,7 @@ ULandscapeComponent::ULandscapeComponent(const FObjectInitializer& ObjectInitial
 
 #if WITH_EDITORONLY_DATA
 	EditToolRenderData = FLandscapeEditToolRenderData();
-	CurrentProceduralLayerName = NAME_None;
+	CurrentProceduralLayerGuid.Invalidate();
 #endif
 
 	LpvBiasMultiplier = 0.0f; // Bias is 0 for landscape, since it's single sided
@@ -1416,6 +1416,20 @@ TArray<FWeightmapLayerAllocationInfo>& ULandscapeComponent::GetWeightmapLayerAll
 
 #if WITH_EDITOR
 
+void ULandscapeComponent::SetCurrentEditingProceduralLayer(FProceduralLayer* Layer, FProceduralLayerData* LayerData)
+{
+	// Update Current Heightmap
+	UTexture2D** LayerHeightmap = LayerData ? LayerData->Heightmaps.Find(GetHeightmap()) : nullptr;
+	SetCurrentEditingHeightmap(LayerHeightmap ? *LayerHeightmap : nullptr);
+
+	// Update Current Weightmaps
+	FWeightmapLayerData* WeightmapData = LayerData ? LayerData->WeightmapData.Find(this) : nullptr;
+	SetCurrentEditingWeightmapLayerAllocations(WeightmapData ? &WeightmapData->WeightmapLayerAllocations : nullptr);
+	SetCurrentEditingWeightmaps(WeightmapData ? &WeightmapData->Weightmaps : nullptr);
+	SetCurrentProceduralLayerGuid(WeightmapData && Layer ? Layer->Guid : FGuid());
+	SetCurrentEditingWeightmapTexturesUsage(WeightmapData ? &WeightmapData->WeightmapTextureUsages : nullptr);
+}
+
 void ULandscapeComponent::SetCurrentEditingHeightmap(UTexture2D* InNewHeightmap)
 {
 #if WITH_EDITORONLY_DATA
@@ -1516,14 +1530,14 @@ void ULandscapeComponent::SetWeightmapTexturesUsage(const TArray<ULandscapeWeigh
 	}
 }
 
-void ULandscapeComponent::SetCurrentProceduralLayerName(const FName& InName)
+void ULandscapeComponent::SetCurrentProceduralLayerGuid(const FGuid& InLayerGuid)
 {
-	CurrentProceduralLayerName = InName;
+	CurrentProceduralLayerGuid = InLayerGuid;
 }
 
-const FName& ULandscapeComponent::GetCurrentProceduralLayerName() const
+const FGuid& ULandscapeComponent::GetCurrentProceduralLayerGuid() const
 {
-	return CurrentProceduralLayerName;
+	return CurrentProceduralLayerGuid;
 }
 #endif
 
diff --git a/Engine/Source/Runtime/Landscape/Private/LandscapeEdit.cpp b/Engine/Source/Runtime/Landscape/Private/LandscapeEdit.cpp
--- a/Engine/Source/Runtime/Landscape/Private/LandscapeEdit.cpp
+++ b/Engine/Source/Runtime/Landscape/Private/LandscapeEdit.cpp
@@ -651,7 +651,7 @@ void ULandscapeComponent::FixupWeightmaps()
 				if (TempUsage == nullptr)
 				{
 					TempUsage = &Proxy->WeightmapUsageMap.Add(WeightmapTexture, NewObject<ULandscapeWeightmapUsage>(GetLandscapeProxy()));
-					(*TempUsage)->ProceduralLayerName = NAME_None;
+					(*TempUsage)->ProceduralLayerGuid.Invalidate();
 				}
 
 				ULandscapeWeightmapUsage* Usage = *TempUsage;
@@ -4654,9 +4654,9 @@ void ULandscapeComponent::ReallocateWeightmaps(FLandscapeEditDataInterface* Data
 			{
 				ULandscapeWeightmapUsage* TryWeightmapUsage = ItPair.Value;
 				//
-				FName ProceduralLayerNameToSeek = InCanUseCurrentEditingWeightmap ? CurrentProceduralLayerName : NAME_None;
+				FGuid ProceduralLayerGuidToSeek = InCanUseCurrentEditingWeightmap ? CurrentProceduralLayerGuid : FGuid();
 
-				if (TryWeightmapUsage->FreeChannelCount() >= TotalNeededChannels && TryWeightmapUsage->ProceduralLayerName == ProceduralLayerNameToSeek)
+				if (TryWeightmapUsage->FreeChannelCount() >= TotalNeededChannels && TryWeightmapUsage->ProceduralLayerGuid == ProceduralLayerGuidToSeek)
 				{
 					if (TryWeightmapUsage->FreeChannelCount() == 4)
 					{
@@ -4717,7 +4717,7 @@ void ULandscapeComponent::ReallocateWeightmaps(FLandscapeEditDataInterface* Data
 
 			// Store it in the usage map
 			CurrentWeightmapUsage = Proxy->WeightmapUsageMap.Add(CurrentWeightmapTexture, NewObject<ULandscapeWeightmapUsage>(GetLandscapeProxy()));
-			CurrentWeightmapUsage->ProceduralLayerName = InCanUseCurrentEditingWeightmap ? CurrentProceduralLayerName : NAME_None;
+			CurrentWeightmapUsage->ProceduralLayerGuid = InCanUseCurrentEditingWeightmap ? CurrentProceduralLayerGuid : FGuid();
 			// UE_LOG(LogLandscape, Log, TEXT("Making a new texture %s"), *CurrentWeightmapTexture->GetName());
 		}
 
diff --git a/Engine/Source/Runtime/Landscape/Private/LandscapeEditProcedural.cpp b/Engine/Source/Runtime/Landscape/Private/LandscapeEditProcedural.cpp
--- a/Engine/Source/Runtime/Landscape/Private/LandscapeEditProcedural.cpp
+++ b/Engine/Source/Runtime/Landscape/Private/LandscapeEditProcedural.cpp
@@ -1577,7 +1577,7 @@ void ALandscapeProxy::SetupProceduralLayers(int32 InNumComponentsX, int32 InNumC
 
 		for (auto& ItLayerDataPair : LandscapeProxy->ProceduralLayersData)
 		{
-			FName ProceduralLayerName = ItLayerDataPair.Key;
+			FGuid ProceduralLayerGuid = ItLayerDataPair.Key;
 			FProceduralLayerData& ProceduralLayerData = ItLayerDataPair.Value;
 
 			struct FTextureData
@@ -1620,7 +1620,7 @@ void ALandscapeProxy::SetupProceduralLayers(int32 InNumComponentsX, int32 InNumC
 
 								NewWeightmapData.Weightmaps[TextureIndex] = TextureData->Texture;
 								NewWeightmapData.WeightmapTextureUsages[TextureIndex] = TextureData->Usage;
-								check(TextureData->Usage->ProceduralLayerName == ProceduralLayerName);
+								check(TextureData->Usage->ProceduralLayerGuid == ProceduralLayerGuid);
 
 								for (int32 ChannelIndex = 0; ChannelIndex < 4; ++ChannelIndex)
 								{
@@ -1670,7 +1670,7 @@ void ALandscapeProxy::SetupProceduralLayers(int32 InNumComponentsX, int32 InNumC
 
 								NewWeightmapData.Weightmaps[TextureIndex] = NewWeightmapTexture;
 								NewWeightmapData.WeightmapTextureUsages[TextureIndex] = ComponentWeightmapTexturesUsage[TextureIndex];
-								NewWeightmapData.WeightmapTextureUsages[TextureIndex]->ProceduralLayerName = ProceduralLayerName;
+								NewWeightmapData.WeightmapTextureUsages[TextureIndex]->ProceduralLayerGuid = ProceduralLayerGuid;
 
 								// Create new Usage for the base layer as the other one will now be used by the Layer 1
 								ComponentWeightmapTexturesUsage[TextureIndex] = LandscapeProxy->WeightmapUsageMap.Add(NewWeightmapTexture, NewObject<ULandscapeWeightmapUsage>(LandscapeProxy));
@@ -1710,7 +1710,7 @@ void ALandscapeProxy::SetupProceduralLayers(int32 InNumComponentsX, int32 InNumC
 						if (TempUsage == nullptr)
 						{
 							TempUsage = &LandscapeProxy->WeightmapUsageMap.Add(WeightmapTexture, NewObject<ULandscapeWeightmapUsage>(LandscapeProxy));
-							(*TempUsage)->ProceduralLayerName = ProceduralLayerName;
+							(*TempUsage)->ProceduralLayerGuid = ProceduralLayerGuid;
 						}
 
 						ULandscapeWeightmapUsage* Usage = *TempUsage;
@@ -2632,7 +2632,7 @@ void ALandscape::RegenerateProceduralHeightmaps()
 
 			for (ALandscapeProxy* Landscape : AllLandscapes)
 			{
-				FProceduralLayerData* LayerData = Landscape->ProceduralLayersData.Find(Layer.Name);
+				FProceduralLayerData* LayerData = Landscape->ProceduralLayersData.Find(Layer.Guid);
 
 				if (LayerData != nullptr)
 				{
@@ -2834,14 +2834,14 @@ void ALandscape::ResolveProceduralTexture(FLandscapeProceduralTexture2DCPUReadBa
 	}
 }
 
-void ALandscape::PrepareProceduralComponentDataForExtractLayersCS(const FName& InProceduralLayerName, int32 InCurrentWeightmapToProcessIndex, bool InOutputDebugName, const TArray<ALandscapeProxy*>& InAllLandscape, FLandscapeTexture2DResource* InOutTextureData,
+void ALandscape::PrepareProceduralComponentDataForExtractLayersCS(const FProceduralLayer& InProceduralLayer, int32 InCurrentWeightmapToProcessIndex, bool InOutputDebugName, const TArray<ALandscapeProxy*>& InAllLandscape, FLandscapeTexture2DResource* InOutTextureData,
 																  TArray<FLandscapeProceduralWeightmapExtractLayersComponentData>& OutComponentData, TMap<ULandscapeLayerInfoObject*, int32>& OutLayerInfoObjects)
 {
 	ULandscapeInfo* Info = GetLandscapeInfo();
 	
 	for (const ALandscapeProxy* Landscape : InAllLandscape)
 	{
-		const FProceduralLayerData* ProceduralLayerData = Landscape->ProceduralLayersData.Find(InProceduralLayerName);
+		const FProceduralLayerData* ProceduralLayerData = Landscape->ProceduralLayersData.Find(InProceduralLayer.Guid);
 
 		if (ProceduralLayerData != nullptr)
 		{
@@ -2858,8 +2858,8 @@ void ALandscape::PrepareProceduralComponentDataForExtractLayersCS(const FName& I
 					const ULandscapeWeightmapUsage* WeightmapUsage = WeightLayerData.WeightmapTextureUsages[InCurrentWeightmapToProcessIndex];
 					check(WeightmapUsage != nullptr);
 
-					CopyProceduralTexture(*Weightmap->GetName(), Weightmap->Resource, InOutputDebugName ? FString::Printf(TEXT("%s WeightmapScratchTexture"), *InProceduralLayerName.ToString()) : TEXT(""), InOutTextureData, nullptr, Component->GetSectionBase(), 0);
-					PrintProceduralDebugTextureResource(InOutputDebugName ? FString::Printf(TEXT("LS Weight: %s WeightmapScratchTexture %s"), *InProceduralLayerName.ToString(), TEXT("WeightmapScratchTextureResource")) : TEXT(""), InOutTextureData, 0, false);
+					CopyProceduralTexture(*Weightmap->GetName(), Weightmap->Resource, InOutputDebugName ? FString::Printf(TEXT("%s WeightmapScratchTexture"), *InProceduralLayer.Name.ToString()) : TEXT(""), InOutTextureData, nullptr, Component->GetSectionBase(), 0);
+					PrintProceduralDebugTextureResource(InOutputDebugName ? FString::Printf(TEXT("LS Weight: %s WeightmapScratchTexture %s"), *InProceduralLayer.Name.ToString(), TEXT("WeightmapScratchTextureResource")) : TEXT(""), InOutTextureData, 0, false);
 
 					for (const FWeightmapLayerAllocationInfo& WeightmapLayerAllocation : WeightLayerData.WeightmapLayerAllocations)
 					{
@@ -3590,7 +3590,7 @@ void ALandscape::RegenerateProceduralWeightmaps()
 
 				// Prepare compute shader data
 				TArray<FLandscapeProceduralWeightmapExtractLayersComponentData> ComponentsData;	
-				PrepareProceduralComponentDataForExtractLayersCS(ProceduralLayer.Name, CurrentWeightmapToProcessIndex, OutputDebugName, AllLandscapes, WeightmapScratchExtractLayerTextureResource, ComponentsData, LayerInfoObjects);
+				PrepareProceduralComponentDataForExtractLayersCS(ProceduralLayer, CurrentWeightmapToProcessIndex, OutputDebugName, AllLandscapes, WeightmapScratchExtractLayerTextureResource, ComponentsData, LayerInfoObjects);
 
 				HasFoundWeightmapToProcess = ComponentsData.Num() > 0;
 
@@ -4377,16 +4377,7 @@ void ALandscape::SetProceduralLayerName(int32 InLayerIndex, const FName& InName)
 		return;
 	}
 
-	// TODO: ProceduralLayersData should be changed to map GUID to Data
-	FName OldName = ProceduralLayers[InLayerIndex].Name;
 	ProceduralLayers[InLayerIndex].Name = InName;
-	LandscapeInfo->ForAllLandscapeProxies([OldName, InName](ALandscapeProxy* Proxy)
-	{
-		FProceduralLayerData* Value = Proxy->ProceduralLayersData.Find(OldName);
-		check(Value);
-		Proxy->ProceduralLayersData.Add(InName, *Value);
-		Proxy->ProceduralLayersData.Remove(OldName);
-	});
 }
 
 void ALandscape::SetProceduralLayerAlpha(int32 InLayerIndex, const float InAlpha, bool bInHeightmap)
@@ -4448,12 +4439,12 @@ void ALandscape::DeleteProceduralLayer(int32 InLayerIndex)
 	}
 	
 	Modify();
-	FName LayerName = Layer->Name;
+	FGuid LayerGuid = Layer->Guid;
 
 	// Clean up Weightmap usage in LandscapeProxies
-	LandscapeInfo->ForAllLandscapeProxies([LayerName](ALandscapeProxy* Proxy)
+	LandscapeInfo->ForAllLandscapeProxies([LayerGuid](ALandscapeProxy* Proxy)
 	{
-		const FProceduralLayerData* LayerData = Proxy->ProceduralLayersData.Find(LayerName);
+		const FProceduralLayerData* LayerData = Proxy->ProceduralLayersData.Find(LayerGuid);
 		if (LayerData)
 		{
 			for (ULandscapeComponent* Component : Proxy->LandscapeComponents)
@@ -4480,9 +4471,9 @@ void ALandscape::DeleteProceduralLayer(int32 InLayerIndex)
 	});
 	
 	// Remove associated layer data of each landscape proxy
-	LandscapeInfo->ForAllLandscapeProxies([LayerName](ALandscapeProxy* Proxy)
+	LandscapeInfo->ForAllLandscapeProxies([LayerGuid](ALandscapeProxy* Proxy)
 	{
-		Proxy->ProceduralLayersData.Remove(LayerName);
+		Proxy->ProceduralLayersData.Remove(LayerGuid);
 	});
 
 	// Remove layer from list
@@ -4614,23 +4605,12 @@ void ALandscape::SetCurrentEditingProceduralLayer(FGuid InLayerGuid)
 
 	LandscapeInfo->ForAllLandscapeProxies([InLayerGuid, this](ALandscapeProxy* Proxy)
 	{
-		// TODO: ProceduralLayersData should be changed to map GUID to Data
 		FProceduralLayer* Layer = ProceduralLayers.FindByPredicate([InLayerGuid](const FProceduralLayer& Other) { return Other.Guid == InLayerGuid; });
-		FProceduralLayerData* CurrentLayerData = Layer ? Proxy->ProceduralLayersData.Find(Layer->Name) : nullptr;
+		FProceduralLayerData* LayerData = Layer ? Proxy->ProceduralLayersData.Find(Layer->Guid) : nullptr;
 
 		for (ULandscapeComponent* Component : Proxy->LandscapeComponents)
 		{
-			// Update Current Heightmap
-			UTexture2D** LayerHeightmap = CurrentLayerData ? CurrentLayerData->Heightmaps.Find(Component->GetHeightmap()) : nullptr;
-			Component->SetCurrentEditingHeightmap(LayerHeightmap ? *LayerHeightmap : nullptr);
-
-			// Update Current Weightmaps
-			FWeightmapLayerData* WeightmapData = CurrentLayerData ? CurrentLayerData->WeightmapData.Find(Component) : nullptr;
-			Component->SetCurrentEditingWeightmapLayerAllocations(WeightmapData ? &WeightmapData->WeightmapLayerAllocations : nullptr);
-			Component->SetCurrentEditingWeightmaps(WeightmapData ? &WeightmapData->Weightmaps : nullptr);
-			Component->SetCurrentProceduralLayerName(WeightmapData && Layer ? Layer->Name : NAME_None);
-			Component->SetCurrentEditingWeightmapTexturesUsage(WeightmapData ? &WeightmapData->WeightmapTextureUsages : nullptr);
-
+			Component->SetCurrentEditingProceduralLayer(Layer, LayerData);
 			Component->MarkRenderStateDirty();
 		}
 	});
@@ -4652,7 +4632,7 @@ void ALandscape::CreateProceduralLayer(FName InName, bool bInUpdateProceduralCon
 	// Create associated layer data in each landscape proxy
 	LandscapeInfo->ForAllLandscapeProxies([NewLayer](ALandscapeProxy* Proxy)
 	{
-		Proxy->ProceduralLayersData.Add(NewLayer.Name, FProceduralLayerData());
+		Proxy->ProceduralLayersData.Add(NewLayer.Guid, FProceduralLayerData());
 	});
 
 	if (bInUpdateProceduralContent)
