Fixed: Crash when changing Preview Rendering Level in ARPG.

diff --git a/Engine/Source/Runtime/Landscape/Private/LandscapeEdit.cpp b/Engine/Source/Runtime/Landscape/Private/LandscapeEdit.cpp
--- a/Engine/Source/Runtime/Landscape/Private/LandscapeEdit.cpp
+++ b/Engine/Source/Runtime/Landscape/Private/LandscapeEdit.cpp
@@ -421,7 +421,10 @@ void ULandscapeComponent::UpdateMaterialInstances_Internal(FMaterialUpdateContex
 		{
 			const int8 MaterialLOD = It.Value();
 
-			MobileCombinationMaterialInstances[MobileMaterialIndex] = GetCombinationMaterial(&Context, MobileWeightmapLayerAllocations, MaterialLOD, true);
+			UMaterialInstanceConstant* MobileCombinationMaterialInstance = GetCombinationMaterial(&Context, MobileWeightmapLayerAllocations, MaterialLOD, true);
+			MobileCombinationMaterialInstances[MobileMaterialIndex] = MobileCombinationMaterialInstance;
+			Context.AddMaterialInstance(MobileCombinationMaterialInstance);
+						
 			++MobileMaterialIndex;
 		}
 	}
