Fixed a Static Analysis issue

diff --git a/Engine/Source/Runtime/Landscape/Private/LandscapeEditProcedural.cpp b/Engine/Source/Runtime/Landscape/Private/LandscapeEditProcedural.cpp
--- a/Engine/Source/Runtime/Landscape/Private/LandscapeEditProcedural.cpp
+++ b/Engine/Source/Runtime/Landscape/Private/LandscapeEditProcedural.cpp
@@ -1537,6 +1537,7 @@ void ALandscapeProxy::SetupProceduralLayers(int32 InNumComponentsX, int32 InNumC
 		for (int32 i = 0; i < EWeightmapRTType::WeightmapRT_Count; ++i)
 		{
 			Landscape->WeightmapRTList[i] = NewObject<UTextureRenderTarget2D>(Landscape->GetOutermost());
+
 			check(Landscape->WeightmapRTList[i]);
 			Landscape->WeightmapRTList[i]->AddressX = TextureAddress::TA_Clamp;
 			Landscape->WeightmapRTList[i]->AddressY = TextureAddress::TA_Clamp;
@@ -1556,10 +1557,7 @@ void ALandscapeProxy::SetupProceduralLayers(int32 InNumComponentsX, int32 InNumC
 				CurrentMipSizeY >>= 1;
 			}
 
-			if (Landscape->WeightmapRTList[i] != nullptr)
-			{
-				Landscape->WeightmapRTList[i]->UpdateResourceImmediate(true);
-			}
+			Landscape->WeightmapRTList[i]->UpdateResourceImmediate(true);
 
 			// Only generate required mips RT
 			if (CurrentMipSizeX == NumComponentsX && CurrentMipSizeY == NumComponentsY)
