Changed the name from Landscape Procedural to Landscape Layers

diff --git a/Engine/Source/Editor/LandscapeEditor/Private/LandscapeEditorDetailCustomization_Layers.h b/Engine/Source/Editor/LandscapeEditor/Private/LandscapeEditorDetailCustomization_Layers.h
new file mode 100644
--- /dev/null
+++ b/Engine/Source/Editor/LandscapeEditor/Private/LandscapeEditorDetailCustomization_Layers.h
@@ -0,0 +1,93 @@
+// Copyright 1998-2019 Epic Games, Inc. All Rights Reserved.
+
+#pragma once
+
+#include "CoreMinimal.h"
+#include "Misc/Attribute.h"
+#include "Layout/Visibility.h"
+#include "Layout/Margin.h"
+#include "Styling/SlateColor.h"
+#include "Styling/SlateBrush.h"
+#include "Input/Reply.h"
+#include "Widgets/DeclarativeSyntaxSupport.h"
+#include "Widgets/SWidget.h"
+#include "Widgets/Layout/SBorder.h"
+#include "Editor/LandscapeEditor/Private/LandscapeEdMode.h"
+#include "IDetailCustomNodeBuilder.h"
+#include "IDetailCustomization.h"
+#include "AssetThumbnail.h"
+#include "Framework/SlateDelegates.h"
+#include "Editor/LandscapeEditor/Private/LandscapeEditorDetailCustomization_Base.h"
+
+class FDetailWidgetRow;
+class IDetailChildrenBuilder;
+class IDetailLayoutBuilder;
+class SDragAndDropVerticalBox;
+
+/**
+ * Slate widgets customizer for the layers list in the Landscape Editor
+ */
+
+class FLandscapeEditorDetailCustomization_Layers : public FLandscapeEditorDetailCustomization_Base
+{
+public:
+	/** Makes a new instance of this detail layout class for a specific detail view requesting it */
+	static TSharedRef<IDetailCustomization> MakeInstance();
+
+	/** IDetailCustomization interface */
+	virtual void CustomizeDetails(IDetailLayoutBuilder& DetailBuilder) override;
+};
+
+class FLandscapeEditorCustomNodeBuilder_Layers : public IDetailCustomNodeBuilder, public TSharedFromThis<FLandscapeEditorCustomNodeBuilder_Layers>
+{
+public:
+	FLandscapeEditorCustomNodeBuilder_Layers(TSharedRef<FAssetThumbnailPool> ThumbnailPool);
+	~FLandscapeEditorCustomNodeBuilder_Layers();
+
+	virtual void SetOnRebuildChildren( FSimpleDelegate InOnRegenerateChildren ) override;
+	virtual void GenerateHeaderRowContent( FDetailWidgetRow& NodeRow ) override;
+	virtual void GenerateChildContent( IDetailChildrenBuilder& ChildrenBuilder ) override;
+	virtual void Tick( float DeltaTime ) override {}
+	virtual bool RequiresTick() const override { return false; }
+	virtual bool InitiallyCollapsed() const override { return false; }
+	virtual FName GetName() const override { return "Layers"; }
+
+protected:
+	TSharedRef<FAssetThumbnailPool> ThumbnailPool;
+
+	static class FEdModeLandscape* GetEditorMode();
+
+	TSharedPtr<SWidget> GenerateRow(int32 InLayerIndex);
+
+	// Drag/Drop handling
+	FReply HandleDragDetected(const FGeometry& MyGeometry, const FPointerEvent& MouseEvent, int32 SlotIndex, SVerticalBox::FSlot* Slot);
+	TOptional<SDragAndDropVerticalBox::EItemDropZone> HandleCanAcceptDrop(const FDragDropEvent& DragDropEvent, SDragAndDropVerticalBox::EItemDropZone DropZone, SVerticalBox::FSlot* Slot);
+	FReply HandleAcceptDrop(FDragDropEvent const& DragDropEvent, SDragAndDropVerticalBox::EItemDropZone DropZone, int32 SlotIndex, SVerticalBox::FSlot* Slot);
+
+	bool IsLayerSelected(int32 LayerIndex);
+	void OnLayerSelectionChanged(int32 LayerIndex);
+	TSharedPtr<SWidget> OnLayerContextMenuOpening(int32 InLayerIndex);
+	void CreateLayer();
+	void ClearLayer(int32 InLayerIndex);
+	void RenameLayer(int32 InLayerIndex);
+	void DeleteLayer(int32 InLayerIndex);
+	void ShowOnlySelectedLayer(int32 InLayerIndex);
+	void ShowAllLayers();
+	bool CanRenameLayerTo(const FText& NewText, FText& OutErrorMessage, int32 InLayerIndex);
+	void SetLayerName(const FText& InText, ETextCommit::Type InCommitType, int32 InLayerIndex);
+	FText GetLayerText(int32 InLayerIndex) const;
+
+	TOptional<float> GetLayerAlpha(int32 InLayerIndex) const;
+	void SetLayerAlpha(float InAlpha, int32 InLayerIndex);
+
+	FReply OnToggleVisibility(int32 InLayerIndex);
+	const FSlateBrush* GetVisibilityBrushForLayer(int32 InLayerIndex) const;
+	
+	FReply OnToggleLock(int32 InLayerIndex);
+	const FSlateBrush* GetLockBrushForLayer(int32 InLayerIndex) const;
+
+private:
+
+	/** Widgets for displaying and editing the layer name */
+	TArray< TSharedPtr< SInlineEditableTextBlock > > InlineTextBlocks;
+};
\ No newline at end of file
diff --git a/Engine/Source/Editor/LandscapeEditor/Private/LandscapeEditorDetailCustomization_LayersBrushStack.cpp b/Engine/Source/Editor/LandscapeEditor/Private/LandscapeEditorDetailCustomization_LayersBrushStack.cpp
new file mode 100644
--- /dev/null
+++ b/Engine/Source/Editor/LandscapeEditor/Private/LandscapeEditorDetailCustomization_LayersBrushStack.cpp
@@ -0,0 +1,404 @@
+// Copyright 1998-2019 Epic Games, Inc. All Rights Reserved.
+
+#include "LandscapeEditorDetailCustomization_LayersBrushStack.h"
+#include "IDetailChildrenBuilder.h"
+#include "Framework/Commands/UIAction.h"
+#include "Widgets/Text/STextBlock.h"
+#include "Framework/MultiBox/MultiBoxBuilder.h"
+#include "Misc/MessageDialog.h"
+#include "Modules/ModuleManager.h"
+#include "Brushes/SlateColorBrush.h"
+#include "Layout/WidgetPath.h"
+#include "SlateOptMacros.h"
+#include "Framework/Application/MenuStack.h"
+#include "Framework/Application/SlateApplication.h"
+#include "Widgets/Images/SImage.h"
+#include "Widgets/Input/SButton.h"
+#include "EditorModeManager.h"
+#include "EditorModes.h"
+#include "DetailLayoutBuilder.h"
+#include "IDetailPropertyRow.h"
+#include "DetailCategoryBuilder.h"
+#include "PropertyCustomizationHelpers.h"
+
+#include "ScopedTransaction.h"
+
+#include "LandscapeEditorDetailCustomization_TargetLayers.h"
+#include "Widgets/Input/SEditableText.h"
+#include "LandscapeBPCustomBrush.h"
+
+#define LOCTEXT_NAMESPACE "LandscapeEditor.Layers"
+
+TSharedRef<IDetailCustomization> FLandscapeEditorDetailCustomization_LayersBrushStack::MakeInstance()
+{
+	return MakeShareable(new FLandscapeEditorDetailCustomization_LayersBrushStack);
+}
+
+BEGIN_SLATE_FUNCTION_BUILD_OPTIMIZATION
+void FLandscapeEditorDetailCustomization_LayersBrushStack::CustomizeDetails(IDetailLayoutBuilder& DetailBuilder)
+{
+	IDetailCategoryBuilder& LayerCategory = DetailBuilder.EditCategory("Current Layer Brushes");
+
+	FEdModeLandscape* LandscapeEdMode = GetEditorMode();
+	if (LandscapeEdMode && LandscapeEdMode->CurrentToolMode != nullptr)
+	{
+		const FName CurrentToolName = LandscapeEdMode->CurrentTool->GetToolName();
+
+		if (LandscapeEdMode->CurrentToolMode->SupportedTargetTypes != 0 && CurrentToolName == TEXT("BPCustom"))
+		{
+			LayerCategory.AddCustomBuilder(MakeShareable(new FLandscapeEditorCustomNodeBuilder_LayersBrushStack(DetailBuilder.GetThumbnailPool().ToSharedRef())));
+		}
+	}
+}
+END_SLATE_FUNCTION_BUILD_OPTIMIZATION
+
+//////////////////////////////////////////////////////////////////////////
+
+FEdModeLandscape* FLandscapeEditorCustomNodeBuilder_LayersBrushStack::GetEditorMode()
+{
+	return (FEdModeLandscape*)GLevelEditorModeTools().GetActiveMode(FBuiltinEditorModes::EM_Landscape);
+}
+
+FLandscapeEditorCustomNodeBuilder_LayersBrushStack::FLandscapeEditorCustomNodeBuilder_LayersBrushStack(TSharedRef<FAssetThumbnailPool> InThumbnailPool)
+	: ThumbnailPool(InThumbnailPool)
+{
+}
+
+FLandscapeEditorCustomNodeBuilder_LayersBrushStack::~FLandscapeEditorCustomNodeBuilder_LayersBrushStack()
+{
+	
+}
+
+void FLandscapeEditorCustomNodeBuilder_LayersBrushStack::SetOnRebuildChildren(FSimpleDelegate InOnRegenerateChildren)
+{
+}
+
+void FLandscapeEditorCustomNodeBuilder_LayersBrushStack::GenerateHeaderRowContent(FDetailWidgetRow& NodeRow)
+{
+	FEdModeLandscape* LandscapeEdMode = GetEditorMode();
+	
+	if (LandscapeEdMode == NULL)
+	{
+		return;	
+	}
+
+	NodeRow.NameWidget
+		[
+			SNew(STextBlock)
+			.Font(IDetailLayoutBuilder::GetDetailFont())
+			.Text(FText::FromString(TEXT("Stack")))
+		];
+}
+
+BEGIN_SLATE_FUNCTION_BUILD_OPTIMIZATION
+void FLandscapeEditorCustomNodeBuilder_LayersBrushStack::GenerateChildContent(IDetailChildrenBuilder& ChildrenBuilder)
+{
+	FEdModeLandscape* LandscapeEdMode = GetEditorMode();
+	if (LandscapeEdMode != NULL)
+	{
+		TSharedPtr<SDragAndDropVerticalBox> BrushesList = SNew(SDragAndDropVerticalBox)
+			.OnCanAcceptDrop(this, &FLandscapeEditorCustomNodeBuilder_LayersBrushStack::HandleCanAcceptDrop)
+			.OnAcceptDrop(this, &FLandscapeEditorCustomNodeBuilder_LayersBrushStack::HandleAcceptDrop)
+			.OnDragDetected(this, &FLandscapeEditorCustomNodeBuilder_LayersBrushStack::HandleDragDetected);
+
+		BrushesList->SetDropIndicator_Above(*FEditorStyle::GetBrush("LandscapeEditor.TargetList.DropZone.Above"));
+		BrushesList->SetDropIndicator_Below(*FEditorStyle::GetBrush("LandscapeEditor.TargetList.DropZone.Below"));
+
+		ChildrenBuilder.AddCustomRow(FText::FromString(FString(TEXT("Brush Stack"))))
+			.Visibility(EVisibility::Visible)
+			[
+				SNew(SVerticalBox)
+
+				+ SVerticalBox::Slot()
+				.AutoHeight()
+				.VAlign(VAlign_Center)
+				.Padding(0, 2)
+				[
+					BrushesList.ToSharedRef()
+				]
+
+				+ SVerticalBox::Slot()
+				.AutoHeight()
+				.VAlign(VAlign_Center)
+				.Padding(0, 2)
+				[
+					SNew(SHorizontalBox)
+
+					+SHorizontalBox::Slot()
+					.HAlign(HAlign_Right)
+					//.Padding(4, 0)
+					[
+						SNew(SButton)
+						.Text(this, &FLandscapeEditorCustomNodeBuilder_LayersBrushStack::GetCommitBrushesButtonText)
+						.OnClicked(this, &FLandscapeEditorCustomNodeBuilder_LayersBrushStack::ToggleCommitBrushes)
+						.IsEnabled(this, &FLandscapeEditorCustomNodeBuilder_LayersBrushStack::IsCommitBrushesButtonEnabled)
+					]
+				]
+			];
+
+		if (LandscapeEdMode->CurrentToolMode != nullptr)
+		{
+			const TArray<int8>& BrushOrderStack = LandscapeEdMode->GetBrushesOrderForCurrentLayer(LandscapeEdMode->CurrentToolTarget.TargetType);
+
+			for (int32 i = 0; i < BrushOrderStack.Num(); ++i)
+			{
+				TSharedPtr<SWidget> GeneratedRowWidget = GenerateRow(i);
+
+				if (GeneratedRowWidget.IsValid())
+				{
+					BrushesList->AddSlot()
+						.AutoHeight()
+						[
+							GeneratedRowWidget.ToSharedRef()
+						];
+				}
+			}
+		}
+	}
+}
+END_SLATE_FUNCTION_BUILD_OPTIMIZATION
+
+BEGIN_SLATE_FUNCTION_BUILD_OPTIMIZATION
+TSharedPtr<SWidget> FLandscapeEditorCustomNodeBuilder_LayersBrushStack::GenerateRow(int32 InBrushIndex)
+{
+	FEdModeLandscape* LandscapeEdMode = GetEditorMode();
+	TSharedPtr<SWidget> RowWidget = SNew(SLandscapeEditorSelectableBorder)
+		.Padding(0)
+		.VAlign(VAlign_Center)
+		.OnSelected(this, &FLandscapeEditorCustomNodeBuilder_LayersBrushStack::OnBrushSelectionChanged, InBrushIndex)
+		.IsSelected(TAttribute<bool>::Create(TAttribute<bool>::FGetter::CreateSP(this, &FLandscapeEditorCustomNodeBuilder_LayersBrushStack::IsBrushSelected, InBrushIndex)))
+		[	
+			SNew(SHorizontalBox)
+
+			+ SHorizontalBox::Slot()
+			.VAlign(VAlign_Center)
+			.Padding(4, 0)
+			[
+				SNew(SVerticalBox)
+				+ SVerticalBox::Slot()
+				.AutoHeight()
+				.VAlign(VAlign_Center)
+				.Padding(0, 2)
+				[
+					SNew(STextBlock)
+					.ColorAndOpacity(TAttribute<FSlateColor>::Create(TAttribute<FSlateColor>::FGetter::CreateSP(this, &FLandscapeEditorCustomNodeBuilder_LayersBrushStack::GetBrushTextColor, InBrushIndex)))
+					.Text(this, &FLandscapeEditorCustomNodeBuilder_LayersBrushStack::GetBrushText, InBrushIndex)
+				]
+			]
+		];
+	
+	return RowWidget;
+}
+END_SLATE_FUNCTION_BUILD_OPTIMIZATION
+
+bool FLandscapeEditorCustomNodeBuilder_LayersBrushStack::IsBrushSelected(int32 InBrushIndex) const
+{
+	ALandscapeBlueprintCustomBrush* Brush = GetBrush(InBrushIndex);
+
+	return Brush != nullptr ? Brush->IsSelected() : false;
+}
+
+void FLandscapeEditorCustomNodeBuilder_LayersBrushStack::OnBrushSelectionChanged(int32 InBrushIndex)
+{
+	FEdModeLandscape* LandscapeEdMode = GetEditorMode();
+
+	if (LandscapeEdMode != nullptr && LandscapeEdMode->AreAllBrushesCommitedToCurrentLayer(LandscapeEdMode->CurrentToolTarget.TargetType))
+	{
+		return;
+	}
+
+	ALandscapeBlueprintCustomBrush* Brush = GetBrush(InBrushIndex);
+
+	if (Brush != nullptr && !Brush->IsCommited())
+	{
+		GEditor->SelectNone(true, true);
+		GEditor->SelectActor(Brush, true, true);
+	}
+}
+
+FText FLandscapeEditorCustomNodeBuilder_LayersBrushStack::GetBrushText(int32 InBrushIndex) const
+{
+	ALandscapeBlueprintCustomBrush* Brush = GetBrush(InBrushIndex);
+
+	if (Brush != nullptr)
+	{
+		return FText::FromString(Brush->GetActorLabel());
+	}
+
+	return FText::FromName(NAME_None);
+}
+
+FSlateColor FLandscapeEditorCustomNodeBuilder_LayersBrushStack::GetBrushTextColor(int32 InBrushIndex) const
+{
+	ALandscapeBlueprintCustomBrush* Brush = GetBrush(InBrushIndex);
+
+	if (Brush != nullptr)
+	{
+		return Brush->IsCommited() ? FSlateColor::UseSubduedForeground() : FSlateColor::UseForeground();
+	}
+
+	return FSlateColor::UseSubduedForeground();
+}
+
+ALandscapeBlueprintCustomBrush* FLandscapeEditorCustomNodeBuilder_LayersBrushStack::GetBrush(int32 InBrushIndex) const
+{
+	FEdModeLandscape* LandscapeEdMode = GetEditorMode();
+
+	if (LandscapeEdMode != nullptr)
+	{
+		return LandscapeEdMode->GetBrushForCurrentLayer(LandscapeEdMode->CurrentToolTarget.TargetType, InBrushIndex);
+	}
+
+	return nullptr;
+}
+
+FReply FLandscapeEditorCustomNodeBuilder_LayersBrushStack::ToggleCommitBrushes()
+{
+	FEdModeLandscape* LandscapeEdMode = GetEditorMode();
+
+	if (LandscapeEdMode != nullptr)
+	{
+		bool CommitBrushes = !LandscapeEdMode->AreAllBrushesCommitedToCurrentLayer(LandscapeEdMode->CurrentToolTarget.TargetType);
+
+		if (CommitBrushes)
+		{
+			TArray<ALandscapeBlueprintCustomBrush*> BrushStack = LandscapeEdMode->GetBrushesForCurrentLayer(LandscapeEdMode->CurrentToolTarget.TargetType);
+
+			for (ALandscapeBlueprintCustomBrush* Brush : BrushStack)
+			{
+				GEditor->SelectActor(Brush, false, true);
+			}
+		}
+
+		LandscapeEdMode->SetBrushesCommitStateForCurrentLayer(LandscapeEdMode->CurrentToolTarget.TargetType, CommitBrushes);
+	}
+
+	return FReply::Handled();
+}
+
+bool FLandscapeEditorCustomNodeBuilder_LayersBrushStack::IsCommitBrushesButtonEnabled() const
+{
+	FEdModeLandscape* LandscapeEdMode = GetEditorMode();
+
+	if (LandscapeEdMode != nullptr)
+	{
+		TArray<ALandscapeBlueprintCustomBrush*> BrushStack = LandscapeEdMode->GetBrushesForCurrentLayer(LandscapeEdMode->CurrentToolTarget.TargetType);
+
+		return BrushStack.Num() > 0;
+	}
+
+	return false;
+}
+
+FText FLandscapeEditorCustomNodeBuilder_LayersBrushStack::GetCommitBrushesButtonText() const
+{
+	FEdModeLandscape* LandscapeEdMode = GetEditorMode();
+
+	if (LandscapeEdMode != nullptr)
+	{
+		return LandscapeEdMode->AreAllBrushesCommitedToCurrentLayer(LandscapeEdMode->CurrentToolTarget.TargetType) ? LOCTEXT("UnCommitBrushesText", "Uncommit") : LOCTEXT("CommitBrushesText", "Commit");
+	}
+
+	return FText::FromName(NAME_None);
+}
+
+FReply FLandscapeEditorCustomNodeBuilder_LayersBrushStack::HandleDragDetected(const FGeometry& MyGeometry, const FPointerEvent& MouseEvent, int32 SlotIndex, SVerticalBox::FSlot* Slot)
+{
+	FEdModeLandscape* LandscapeEdMode = GetEditorMode();
+
+	if (LandscapeEdMode != nullptr)
+	{
+		const TArray<int8>& BrushOrderStack = LandscapeEdMode->GetBrushesOrderForCurrentLayer(LandscapeEdMode->CurrentToolTarget.TargetType);
+
+		if (BrushOrderStack.IsValidIndex(SlotIndex))
+		{
+			TSharedPtr<SWidget> Row = GenerateRow(SlotIndex);
+
+			if (Row.IsValid())
+			{
+				return FReply::Handled().BeginDragDrop(FLandscapeBrushDragDropOp::New(SlotIndex, Slot, Row));
+			}
+		}
+	}
+
+	return FReply::Unhandled();
+}
+
+TOptional<SDragAndDropVerticalBox::EItemDropZone> FLandscapeEditorCustomNodeBuilder_LayersBrushStack::HandleCanAcceptDrop(const FDragDropEvent& DragDropEvent, SDragAndDropVerticalBox::EItemDropZone DropZone, SVerticalBox::FSlot* Slot)
+{
+	TSharedPtr<FLandscapeBrushDragDropOp> DragDropOperation = DragDropEvent.GetOperationAs<FLandscapeBrushDragDropOp>();
+
+	if (DragDropOperation.IsValid())
+	{
+		return DropZone;
+	}
+
+	return TOptional<SDragAndDropVerticalBox::EItemDropZone>();
+}
+
+FReply FLandscapeEditorCustomNodeBuilder_LayersBrushStack::HandleAcceptDrop(FDragDropEvent const& DragDropEvent, SDragAndDropVerticalBox::EItemDropZone DropZone, int32 SlotIndex, SVerticalBox::FSlot* Slot)
+{
+	TSharedPtr<FLandscapeBrushDragDropOp> DragDropOperation = DragDropEvent.GetOperationAs<FLandscapeBrushDragDropOp>();
+
+	if (DragDropOperation.IsValid())
+	{
+		FEdModeLandscape* LandscapeEdMode = GetEditorMode();
+
+		if (LandscapeEdMode != nullptr)
+		{
+			TArray<int8>& BrushOrderStack = LandscapeEdMode->GetBrushesOrderForCurrentLayer(LandscapeEdMode->CurrentToolTarget.TargetType);
+
+			if (BrushOrderStack.IsValidIndex(DragDropOperation->SlotIndexBeingDragged) && BrushOrderStack.IsValidIndex(SlotIndex))
+			{
+				int32 StartingLayerIndex = DragDropOperation->SlotIndexBeingDragged;
+				int32 DestinationLayerIndex = SlotIndex;
+
+				if (StartingLayerIndex != INDEX_NONE && DestinationLayerIndex != INDEX_NONE)
+				{
+					int8 MovingBrushIndex = BrushOrderStack[StartingLayerIndex];
+					 
+					BrushOrderStack.RemoveAt(StartingLayerIndex);
+					BrushOrderStack.Insert(MovingBrushIndex, DestinationLayerIndex);
+
+					LandscapeEdMode->RefreshDetailPanel();
+					LandscapeEdMode->RequestLayersContentUpdate();
+
+					return FReply::Handled();
+				}
+			}
+		}
+	}
+
+	return FReply::Unhandled();
+}
+
+TSharedRef<FLandscapeBrushDragDropOp> FLandscapeBrushDragDropOp::New(int32 InSlotIndexBeingDragged, SVerticalBox::FSlot* InSlotBeingDragged, TSharedPtr<SWidget> WidgetToShow)
+{
+	TSharedRef<FLandscapeBrushDragDropOp> Operation = MakeShareable(new FLandscapeBrushDragDropOp);
+
+	Operation->MouseCursor = EMouseCursor::GrabHandClosed;
+	Operation->SlotIndexBeingDragged = InSlotIndexBeingDragged;
+	Operation->SlotBeingDragged = InSlotBeingDragged;
+	Operation->WidgetToShow = WidgetToShow;
+
+	Operation->Construct();
+
+	return Operation;
+}
+
+FLandscapeBrushDragDropOp::~FLandscapeBrushDragDropOp()
+{
+}
+
+TSharedPtr<SWidget> FLandscapeBrushDragDropOp::GetDefaultDecorator() const
+{
+	return SNew(SBorder)
+		.BorderImage(FEditorStyle::GetBrush("ContentBrowser.AssetDragDropTooltipBackground"))
+		.Content()
+		[
+			WidgetToShow.ToSharedRef()
+		];
+
+}
+
+#undef LOCTEXT_NAMESPACE
\ No newline at end of file
diff --git a/Engine/Source/Editor/LandscapeEditor/Private/LandscapeEditorDetailCustomization_LayersBrushStack.h b/Engine/Source/Editor/LandscapeEditor/Private/LandscapeEditorDetailCustomization_LayersBrushStack.h
new file mode 100644
--- /dev/null
+++ b/Engine/Source/Editor/LandscapeEditor/Private/LandscapeEditorDetailCustomization_LayersBrushStack.h
@@ -0,0 +1,90 @@
+// Copyright 1998-2019 Epic Games, Inc. All Rights Reserved.
+
+#pragma once
+
+#include "CoreMinimal.h"
+#include "Misc/Attribute.h"
+#include "Layout/Visibility.h"
+#include "Layout/Margin.h"
+#include "Styling/SlateColor.h"
+#include "Input/Reply.h"
+#include "Widgets/DeclarativeSyntaxSupport.h"
+#include "Widgets/SWidget.h"
+#include "Widgets/Layout/SBorder.h"
+#include "Editor/LandscapeEditor/Private/LandscapeEdMode.h"
+#include "IDetailCustomNodeBuilder.h"
+#include "IDetailCustomization.h"
+#include "AssetThumbnail.h"
+#include "Framework/SlateDelegates.h"
+#include "Editor/LandscapeEditor/Private/LandscapeEditorDetailCustomization_Base.h"
+
+class FDetailWidgetRow;
+class IDetailChildrenBuilder;
+class IDetailLayoutBuilder;
+class SDragAndDropVerticalBox;
+
+/**
+ * Slate widgets customizer for the layers list in the Landscape Editor
+ */
+
+class FLandscapeEditorDetailCustomization_LayersBrushStack : public FLandscapeEditorDetailCustomization_Base
+{
+public:
+	/** Makes a new instance of this detail layout class for a specific detail view requesting it */
+	static TSharedRef<IDetailCustomization> MakeInstance();
+
+	/** IDetailCustomization interface */
+	virtual void CustomizeDetails(IDetailLayoutBuilder& DetailBuilder) override;
+};
+
+class FLandscapeEditorCustomNodeBuilder_LayersBrushStack : public IDetailCustomNodeBuilder, public TSharedFromThis<FLandscapeEditorCustomNodeBuilder_LayersBrushStack>
+{
+public:
+	FLandscapeEditorCustomNodeBuilder_LayersBrushStack(TSharedRef<FAssetThumbnailPool> ThumbnailPool);
+	~FLandscapeEditorCustomNodeBuilder_LayersBrushStack();
+
+	virtual void SetOnRebuildChildren( FSimpleDelegate InOnRegenerateChildren ) override;
+	virtual void GenerateHeaderRowContent( FDetailWidgetRow& NodeRow ) override;
+	virtual void GenerateChildContent( IDetailChildrenBuilder& ChildrenBuilder ) override;
+	virtual void Tick( float DeltaTime ) override {}
+	virtual bool RequiresTick() const override { return false; }
+	virtual bool InitiallyCollapsed() const override { return false; }
+	virtual FName GetName() const override { return "Brush Stack"; }
+
+protected:
+	TSharedRef<FAssetThumbnailPool> ThumbnailPool;
+
+	static class FEdModeLandscape* GetEditorMode();
+
+	TSharedPtr<SWidget> GenerateRow(int32 InLayerIndex);
+
+	bool IsBrushSelected(int32 InBrushIndex) const;
+	void OnBrushSelectionChanged(int32 InBrushIndex);
+	FText GetBrushText(int32 InBrushIndex) const;
+	FSlateColor GetBrushTextColor(int32 InBrushIndex) const;
+	ALandscapeBlueprintCustomBrush* GetBrush(int32 InBrushIndex) const;
+
+	FReply ToggleCommitBrushes();
+	bool IsCommitBrushesButtonEnabled() const;
+	FText GetCommitBrushesButtonText() const;
+
+	// Drag/Drop handling
+	FReply HandleDragDetected(const FGeometry& MyGeometry, const FPointerEvent& MouseEvent, int32 SlotIndex, SVerticalBox::FSlot* Slot);
+	TOptional<SDragAndDropVerticalBox::EItemDropZone> HandleCanAcceptDrop(const FDragDropEvent& DragDropEvent, SDragAndDropVerticalBox::EItemDropZone DropZone, SVerticalBox::FSlot* Slot);
+	FReply HandleAcceptDrop(FDragDropEvent const& DragDropEvent, SDragAndDropVerticalBox::EItemDropZone DropZone, int32 SlotIndex, SVerticalBox::FSlot* Slot);
+};
+
+class FLandscapeBrushDragDropOp : public FDragAndDropVerticalBoxOp
+{
+public:
+	DRAG_DROP_OPERATOR_TYPE(FLandscapeBrushDragDropOp, FDragAndDropVerticalBoxOp)
+
+	TSharedPtr<SWidget> WidgetToShow;
+
+	static TSharedRef<FLandscapeBrushDragDropOp> New(int32 InSlotIndexBeingDragged, SVerticalBox::FSlot* InSlotBeingDragged, TSharedPtr<SWidget> InWidgetToShow);
+
+public:
+	virtual ~FLandscapeBrushDragDropOp();
+
+	virtual TSharedPtr<SWidget> GetDefaultDecorator() const override;
+};
diff --git a/Engine/Source/Editor/LandscapeEditor/Private/LandscapeEditorDetailCustomization_NewLandscape.cpp b/Engine/Source/Editor/LandscapeEditor/Private/LandscapeEditorDetailCustomization_NewLandscape.cpp
--- a/Engine/Source/Editor/LandscapeEditor/Private/LandscapeEditorDetailCustomization_NewLandscape.cpp
+++ b/Engine/Source/Editor/LandscapeEditor/Private/LandscapeEditorDetailCustomization_NewLandscape.cpp
@@ -859,7 +859,7 @@ FReply FLandscapeEditorDetailCustomization_NewLandscape::OnCreateButtonClicked()
 
 		LandscapeEdMode->SetCurrentTool("Select"); // change tool so switching back to the manage mode doesn't give "New Landscape" again
 		LandscapeEdMode->SetCurrentTool("Sculpt"); // change to sculpting mode and tool
-		LandscapeEdMode->SetCurrentProceduralLayer(0);
+		LandscapeEdMode->SetCurrentLayer(0);
 
 		if (LandscapeEdMode->CurrentToolTarget.LandscapeInfo.IsValid())
 		{
diff --git a/Engine/Source/Editor/LandscapeEditor/Private/LandscapeEditorDetailCustomization_TargetLayers.cpp b/Engine/Source/Editor/LandscapeEditor/Private/LandscapeEditorDetailCustomization_TargetLayers.cpp
--- a/Engine/Source/Editor/LandscapeEditor/Private/LandscapeEditorDetailCustomization_TargetLayers.cpp
+++ b/Engine/Source/Editor/LandscapeEditor/Private/LandscapeEditorDetailCustomization_TargetLayers.cpp
@@ -702,14 +702,14 @@ TSharedPtr<SWidget> FLandscapeEditorCustomNodeBuilder_TargetLayers::GenerateRow(
 					.AutoHeight()
 					[
 						SNew(SHorizontalBox)
-						.Visibility_Static(&FLandscapeEditorCustomNodeBuilder_TargetLayers::GetProceduralLayerSubstractiveBlendVisibility, Target)
+						.Visibility_Static(&FLandscapeEditorCustomNodeBuilder_TargetLayers::GetLayersSubstractiveBlendVisibility, Target)
 						+ SHorizontalBox::Slot()
 						.AutoWidth()
 						.Padding(0, 2, 2, 2)
 						[
 							SNew(SCheckBox)
-							.IsChecked_Static(&FLandscapeEditorCustomNodeBuilder_TargetLayers::ProceduralLayerSubstractiveBlendIsChecked, Target)
-							.OnCheckStateChanged_Static(&FLandscapeEditorCustomNodeBuilder_TargetLayers::OnProceduralLayerSubstractiveBlendChanged, Target)
+							.IsChecked_Static(&FLandscapeEditorCustomNodeBuilder_TargetLayers::IsLayersSubstractiveBlendChecked, Target)
+							.OnCheckStateChanged_Static(&FLandscapeEditorCustomNodeBuilder_TargetLayers::OnLayersSubstractiveBlendChanged, Target)
 							[
 								SNew(STextBlock)
 								.Text(LOCTEXT("ViewMode.Debug_None", "Substractive Blend"))
@@ -1094,7 +1094,7 @@ void FLandscapeEditorCustomNodeBuilder_TargetLayers::OnFillLayer(const TSharedRe
 		FEdModeLandscape* LandscapeEdMode = GetEditorMode();
 		if (LandscapeEdMode)
 		{
-			FScopedSetLandscapeCurrentEditingProceduralLayer Scope(LandscapeEdMode->GetLandscape(), LandscapeEdMode->GetCurrentProceduralLayerGuid(), [&] { LandscapeEdMode->RequestProceduralContentUpdate(true); });
+			FScopedSetLandscapeEditingLayer Scope(LandscapeEdMode->GetLandscape(), LandscapeEdMode->GetCurrentLayerGuid(), [&] { LandscapeEdMode->RequestLayersContentUpdate(true); });
 			LandscapeEdit.FillLayer(Target->LayerInfoObj.Get());
 		}
 	}
@@ -1107,11 +1107,11 @@ void FLandscapeEditorCustomNodeBuilder_TargetLayers::FillEmptyLayers(ULandscapeI
 	{
 		FLandscapeEditDataInterface LandscapeEdit(LandscapeInfo);
 		
-		if (GetMutableDefault<UEditorExperimentalSettings>()->bProceduralLandscape)
+		if (GetMutableDefault<UEditorExperimentalSettings>()->bLandscapeLayerSystem)
 		{
-			if (LandscapeEdMode->NeedToFillEmptyLayersForProcedural())
+			if (LandscapeEdMode->NeedToFillEmptyMaterialLayers())
 			{
-				FScopedSetLandscapeCurrentEditingProceduralLayer Scope(LandscapeEdMode->GetLandscape(), LandscapeEdMode->GetCurrentProceduralLayerGuid());
+				FScopedSetLandscapeEditingLayer Scope(LandscapeEdMode->GetLandscape(), LandscapeEdMode->GetCurrentLayerGuid());
 				LandscapeEdit.FillEmptyLayers(LandscapeInfoObject);
 			}
 		}
@@ -1132,7 +1132,7 @@ void FLandscapeEditorCustomNodeBuilder_TargetLayers::OnClearLayer(const TSharedR
 		FEdModeLandscape* LandscapeEdMode = GetEditorMode();
 		if (LandscapeEdMode)
 		{
-			FScopedSetLandscapeCurrentEditingProceduralLayer Scope(LandscapeEdMode->GetLandscape(), LandscapeEdMode->GetCurrentProceduralLayerGuid(), [&] { LandscapeEdMode->RequestProceduralContentUpdate(true); });
+			FScopedSetLandscapeEditingLayer Scope(LandscapeEdMode->GetLandscape(), LandscapeEdMode->GetCurrentLayerGuid(), [&] { LandscapeEdMode->RequestLayersContentUpdate(true); });
 			FLandscapeEditDataInterface LandscapeEdit(Target->LandscapeInfo.Get());
 			LandscapeEdit.DeleteLayer(Target->LayerInfoObj.Get());
 		}
@@ -1425,9 +1425,9 @@ EVisibility FLandscapeEditorCustomNodeBuilder_TargetLayers::GetDebugModeLayerUsa
 	return EVisibility::Visible;
 }
 
-EVisibility FLandscapeEditorCustomNodeBuilder_TargetLayers::GetProceduralLayerSubstractiveBlendVisibility(const TSharedRef<FLandscapeTargetListInfo> Target)
+EVisibility FLandscapeEditorCustomNodeBuilder_TargetLayers::GetLayersSubstractiveBlendVisibility(const TSharedRef<FLandscapeTargetListInfo> Target)
 {
-	if (GetMutableDefault<UEditorExperimentalSettings>()->bProceduralLandscape && Target->TargetType != ELandscapeToolTargetType::Heightmap && Target->LayerInfoObj.IsValid())
+	if (GetMutableDefault<UEditorExperimentalSettings>()->bLandscapeLayerSystem && Target->TargetType != ELandscapeToolTargetType::Heightmap && Target->LayerInfoObj.IsValid())
 	{
 		return EVisibility::Visible;
 	}
@@ -1435,23 +1435,23 @@ EVisibility FLandscapeEditorCustomNodeBuilder_TargetLayers::GetProceduralLayerSu
 	return EVisibility::Collapsed;
 }
 
-ECheckBoxState FLandscapeEditorCustomNodeBuilder_TargetLayers::ProceduralLayerSubstractiveBlendIsChecked(const TSharedRef<FLandscapeTargetListInfo> Target)
+ECheckBoxState FLandscapeEditorCustomNodeBuilder_TargetLayers::IsLayersSubstractiveBlendChecked(const TSharedRef<FLandscapeTargetListInfo> Target)
 {
 	FEdModeLandscape* LandscapeEdMode = GetEditorMode();
 	if (LandscapeEdMode)
 	{
-		return LandscapeEdMode->IsCurrentProceduralLayerBlendSubstractive(Target.Get().LayerInfoObj) ? ECheckBoxState::Checked : ECheckBoxState::Unchecked;
+		return LandscapeEdMode->IsCurrentLayerBlendSubstractive(Target.Get().LayerInfoObj) ? ECheckBoxState::Checked : ECheckBoxState::Unchecked;
 	}
 
 	return ECheckBoxState::Unchecked;
 }
 
-void FLandscapeEditorCustomNodeBuilder_TargetLayers::OnProceduralLayerSubstractiveBlendChanged(ECheckBoxState NewCheckedState, const TSharedRef<FLandscapeTargetListInfo> Target)
+void FLandscapeEditorCustomNodeBuilder_TargetLayers::OnLayersSubstractiveBlendChanged(ECheckBoxState NewCheckedState, const TSharedRef<FLandscapeTargetListInfo> Target)
 {
 	FEdModeLandscape* LandscapeEdMode = GetEditorMode();
 	if (LandscapeEdMode)
 	{
-		LandscapeEdMode->SetCurrentProceduralLayerSubstractiveBlendStatus(NewCheckedState == ECheckBoxState::Checked, Target.Get().LayerInfoObj);
+		LandscapeEdMode->SetCurrentLayerSubstractiveBlendStatus(NewCheckedState == ECheckBoxState::Checked, Target.Get().LayerInfoObj);
 	}
 }
 
diff --git a/Engine/Source/Editor/LandscapeEditor/Private/LandscapeEditorDetailCustomization_TargetLayers.h b/Engine/Source/Editor/LandscapeEditor/Private/LandscapeEditorDetailCustomization_TargetLayers.h
--- a/Engine/Source/Editor/LandscapeEditor/Private/LandscapeEditorDetailCustomization_TargetLayers.h
+++ b/Engine/Source/Editor/LandscapeEditor/Private/LandscapeEditorDetailCustomization_TargetLayers.h
@@ -94,9 +94,9 @@ protected:
 	static ECheckBoxState DebugModeColorChannelIsChecked(const TSharedRef<FLandscapeTargetListInfo> Target, int32 Channel);
 	static void OnDebugModeColorChannelChanged(ECheckBoxState NewCheckedState, const TSharedRef<FLandscapeTargetListInfo> Target, int32 Channel);
 
-	static EVisibility GetProceduralLayerSubstractiveBlendVisibility(const TSharedRef<FLandscapeTargetListInfo> Target);
-	static ECheckBoxState ProceduralLayerSubstractiveBlendIsChecked(const TSharedRef<FLandscapeTargetListInfo> Target);
-	static void OnProceduralLayerSubstractiveBlendChanged(ECheckBoxState NewCheckedState, const TSharedRef<FLandscapeTargetListInfo> Target);
+	static EVisibility GetLayersSubstractiveBlendVisibility(const TSharedRef<FLandscapeTargetListInfo> Target);
+	static ECheckBoxState IsLayersSubstractiveBlendChecked(const TSharedRef<FLandscapeTargetListInfo> Target);
+	static void OnLayersSubstractiveBlendChanged(ECheckBoxState NewCheckedState, const TSharedRef<FLandscapeTargetListInfo> Target);
 
 	// Drag/Drop handling
 	FReply HandleDragDetected(const FGeometry& MyGeometry, const FPointerEvent& MouseEvent, int32 SlotIndex, SVerticalBox::FSlot* Slot);
diff --git a/Engine/Source/Editor/LandscapeEditor/Private/LandscapeEditorDetails.cpp b/Engine/Source/Editor/LandscapeEditor/Private/LandscapeEditorDetails.cpp
--- a/Engine/Source/Editor/LandscapeEditor/Private/LandscapeEditorDetails.cpp
+++ b/Engine/Source/Editor/LandscapeEditor/Private/LandscapeEditorDetails.cpp
@@ -19,7 +19,7 @@
 
 #include "LandscapeEditorCommands.h"
 #include "LandscapeEditorDetailWidgets.h"
-#include "LandscapeEditorDetailCustomization_ProceduralBrushStack.h"
+#include "LandscapeEditorDetailCustomization_LayersBrushStack.h"
 #include "Settings/EditorExperimentalSettings.h"
 
 #define LOCTEXT_NAMESPACE "LandscapeEditor"
@@ -108,15 +108,15 @@ void FLandscapeEditorDetails::CustomizeDetails(IDetailLayoutBuilder& DetailBuild
 	Customization_AlphaBrush = MakeShareable(new FLandscapeEditorDetailCustomization_AlphaBrush);
 	Customization_AlphaBrush->CustomizeDetails(DetailBuilder);
 
-	if (GetMutableDefault<UEditorExperimentalSettings>()->bProceduralLandscape)
+	if (GetMutableDefault<UEditorExperimentalSettings>()->bLandscapeLayerSystem)
 	{
-		// Procedural Layers
-		Customization_ProceduralLayers = MakeShareable(new FLandscapeEditorDetailCustomization_ProceduralLayers);
-		Customization_ProceduralLayers->CustomizeDetails(DetailBuilder);
+		// Layers
+		Customization_Layers = MakeShareable(new FLandscapeEditorDetailCustomization_Layers);
+		Customization_Layers->CustomizeDetails(DetailBuilder);
 
 		// Brush Stack
-		Customization_ProceduralBrushStack = MakeShareable(new FLandscapeEditorDetailCustomization_ProceduralBrushStack);
-		Customization_ProceduralBrushStack->CustomizeDetails(DetailBuilder);
+		Customization_LayersBrushStack = MakeShareable(new FLandscapeEditorDetailCustomization_LayersBrushStack);
+		Customization_LayersBrushStack->CustomizeDetails(DetailBuilder);
 	}
 
 	// Target Layers:
@@ -143,7 +143,7 @@ FText FLandscapeEditorDetails::GetLocalizedName(FString Name)
 		LOCTEXT("ToolSet_Retopologize", "Retopologize");
 		LOCTEXT("ToolSet_Visibility", "Visibility");
 
-		if (GetMutableDefault<UEditorExperimentalSettings>()->bProceduralLandscape)
+		if (GetMutableDefault<UEditorExperimentalSettings>()->bLandscapeLayerSystem)
 		{
 			LOCTEXT("ToolSet_BPCustom", "Blueprint Custom");
 		}
@@ -324,7 +324,7 @@ TSharedRef<SWidget> FLandscapeEditorDetails::GetToolSelector()
 			MenuBuilder.AddToolButton(NameToCommandMap.FindChecked("Tool_Retopologize"), NAME_None, LOCTEXT("Tool.Retopologize", "Retopologize"), LOCTEXT("Tool.Retopologize.Tooltip", "Automatically adjusts landscape vertices with an X/Y offset map to improve vertex density on cliffs, reducing texture stretching.\nNote: An X/Y offset map makes the landscape slower to render and paint on with other tools, so only use if needed"));
 			MenuBuilder.AddToolButton(NameToCommandMap.FindChecked("Tool_Visibility"), NAME_None, LOCTEXT("Tool.Visibility", "Visibility"), LOCTEXT("Tool.Visibility.Tooltip", "Mask out individual quads in the landscape, leaving a hole."));
 
-			if (GetMutableDefault<UEditorExperimentalSettings>()->bProceduralLandscape)
+			if (GetMutableDefault<UEditorExperimentalSettings>()->bLandscapeLayerSystem)
 			{
 				MenuBuilder.AddToolButton(NameToCommandMap.FindChecked("Tool_BPCustom"), NAME_None, LOCTEXT("Tool.SculptBPCustom", "Blueprint Custom"), LOCTEXT("Tool.SculptBPCustom.Tooltip", "Custom sculpting tools created using Blueprint."));
 			}
@@ -346,7 +346,7 @@ TSharedRef<SWidget> FLandscapeEditorDetails::GetToolSelector()
 			MenuBuilder.AddToolButton(NameToCommandMap.FindChecked("Tool_Flatten"), NAME_None, LOCTEXT("Tool.Flatten", "Flatten"), LOCTEXT("Tool.Flatten.Tooltip", "Flattens an area of heightmap or blend layer"));
 			MenuBuilder.AddToolButton(NameToCommandMap.FindChecked("Tool_Noise"), NAME_None, LOCTEXT("Tool.Noise", "Noise"), LOCTEXT("Tool.Noise.Tooltip", "Adds noise to the heightmap or blend layer"));
 
-			if (GetMutableDefault<UEditorExperimentalSettings>()->bProceduralLandscape)
+			if (GetMutableDefault<UEditorExperimentalSettings>()->bLandscapeLayerSystem)
 			{
 				MenuBuilder.AddToolButton(NameToCommandMap.FindChecked("Tool_BPCustom"), NAME_None, LOCTEXT("Tool.PaintBPCustom", "Blueprint Custom"), LOCTEXT("Tool.PaintBPCustom.Tooltip", "Custom painting tools created using Blueprint."));
 			}
diff --git a/Engine/Source/Editor/LandscapeEditor/Private/LandscapeEditorDetails.h b/Engine/Source/Editor/LandscapeEditor/Private/LandscapeEditorDetails.h
--- a/Engine/Source/Editor/LandscapeEditor/Private/LandscapeEditorDetails.h
+++ b/Engine/Source/Editor/LandscapeEditor/Private/LandscapeEditorDetails.h
@@ -9,8 +9,8 @@
 #include "UnrealClient.h"
 #include "IDetailCustomization.h"
 #include "Editor/LandscapeEditor/Private/LandscapeEditorDetailCustomization_Base.h"
-#include "LandscapeEditorDetailCustomization_ProceduralLayers.h"
-#include "LandscapeEditorDetailCustomization_ProceduralBrushStack.h"
+#include "LandscapeEditorDetailCustomization_Layers.h"
+#include "LandscapeEditorDetailCustomization_LayersBrushStack.h"
 
 class FLandscapeEditorDetailCustomization_AlphaBrush;
 class FLandscapeEditorDetailCustomization_CopyPaste;
@@ -63,6 +63,6 @@ protected:
 	TSharedPtr<FLandscapeEditorDetailCustomization_MiscTools> Customization_MiscTools;
 	TSharedPtr<FLandscapeEditorDetailCustomization_AlphaBrush> Customization_AlphaBrush;
 	TSharedPtr<FLandscapeEditorDetailCustomization_TargetLayers> Customization_TargetLayers;
-	TSharedPtr<FLandscapeEditorDetailCustomization_ProceduralLayers> Customization_ProceduralLayers;
-	TSharedPtr<FLandscapeEditorDetailCustomization_ProceduralBrushStack> Customization_ProceduralBrushStack;
+	TSharedPtr<FLandscapeEditorDetailCustomization_Layers> Customization_Layers;
+	TSharedPtr<FLandscapeEditorDetailCustomization_LayersBrushStack> Customization_LayersBrushStack;
 };
diff --git a/Engine/Source/Editor/LandscapeEditor/Public/LandscapeToolInterface.h b/Engine/Source/Editor/LandscapeEditor/Public/LandscapeToolInterface.h
--- a/Engine/Source/Editor/LandscapeEditor/Public/LandscapeToolInterface.h
+++ b/Engine/Source/Editor/LandscapeEditor/Public/LandscapeToolInterface.h
@@ -176,14 +176,14 @@ struct FLandscapeToolTarget
 	ELandscapeToolTargetType::Type TargetType;
 	TWeakObjectPtr<ULandscapeLayerInfoObject> LayerInfo;
 	FName LayerName;
-	int32 CurrentProceduralLayerIndex;
+	int32 CurrentLayerIndex;
 
 	FLandscapeToolTarget()
 		: LandscapeInfo()
 		, TargetType(ELandscapeToolTargetType::Heightmap)
 		, LayerInfo()
 		, LayerName(NAME_None)
-		, CurrentProceduralLayerIndex(INDEX_NONE)
+		, CurrentLayerIndex(INDEX_NONE)
 	{
 	}
 };
diff --git a/Engine/Source/Editor/UnrealEd/Classes/Settings/EditorExperimentalSettings.h b/Engine/Source/Editor/UnrealEd/Classes/Settings/EditorExperimentalSettings.h
--- a/Engine/Source/Editor/UnrealEd/Classes/Settings/EditorExperimentalSettings.h
+++ b/Engine/Source/Editor/UnrealEd/Classes/Settings/EditorExperimentalSettings.h
@@ -25,8 +25,8 @@ public:
 	bool bProceduralFoliage;
 
 	/** Allows usage of the procedural landscape system*/
-	UPROPERTY(EditAnywhere, config, Category = Landscape, meta = (DisplayName = "Procedural Landscape"))
-	bool bProceduralLandscape;
+	UPROPERTY(EditAnywhere, config, Category = Landscape, meta = (DisplayName = "Landscape Layer System"))
+	bool bLandscapeLayerSystem;
 
 	/** Allows usage of the Localization Dashboard */
 	UPROPERTY(EditAnywhere, config, Category = Tools, meta = (DisplayName = "Localization Dashboard"))
diff --git a/Engine/Source/Editor/UnrealEd/Private/Settings/SettingsClasses.cpp b/Engine/Source/Editor/UnrealEd/Private/Settings/SettingsClasses.cpp
--- a/Engine/Source/Editor/UnrealEd/Private/Settings/SettingsClasses.cpp
+++ b/Engine/Source/Editor/UnrealEd/Private/Settings/SettingsClasses.cpp
@@ -113,6 +113,7 @@ USkeletalMeshEditorSettings::USkeletalMeshEditorSettings(const FObjectInitialize
 
 UEditorExperimentalSettings::UEditorExperimentalSettings( const FObjectInitializer& ObjectInitializer )
 	: Super(ObjectInitializer)
+	, bLandscapeLayerSystem(false)
 	, bEnableLocalizationDashboard(true)
 	, bUseOpenCLForConvexHullDecomp(false)
 	, bAllowPotentiallyUnsafePropertyEditing(false)
diff --git a/Engine/Source/Runtime/Engine/Private/Texture2D.cpp b/Engine/Source/Runtime/Engine/Private/Texture2D.cpp
--- a/Engine/Source/Runtime/Engine/Private/Texture2D.cpp
+++ b/Engine/Source/Runtime/Engine/Private/Texture2D.cpp
@@ -1166,7 +1166,7 @@ bool UTexture2D::ShouldMipLevelsBeForcedResident() const
 	}
 
 #if WITH_EDITOR
-	if (GIsEditor && GetMutableDefault<UEditorExperimentalSettings>()->bProceduralLandscape && (LODGroup == TEXTUREGROUP_Terrain_Heightmap || LODGroup == TEXTUREGROUP_Terrain_Weightmap))
+	if (GIsEditor && GetMutableDefault<UEditorExperimentalSettings>()->bLandscapeLayerSystem && (LODGroup == TEXTUREGROUP_Terrain_Heightmap || LODGroup == TEXTUREGROUP_Terrain_Weightmap))
 	{
 		return true;
 	}
diff --git a/Engine/Source/Runtime/Landscape/Classes/Landscape.h b/Engine/Source/Runtime/Landscape/Classes/Landscape.h
--- a/Engine/Source/Runtime/Landscape/Classes/Landscape.h
+++ b/Engine/Source/Runtime/Landscape/Classes/Landscape.h
@@ -72,7 +72,7 @@ enum EWeightmapRTType : uint8
 	WeightmapRT_Count
 };
 
-enum EProceduralContentUpdateFlag : uint32
+enum ELandscapeLayersContentUpdateFlag : uint32
 {
 	Heightmap_Setup					= 0x00000001u,
 	Heightmap_Render				= 0x00000002u,
@@ -94,15 +94,15 @@ enum EProceduralContentUpdateFlag : uint32
 };
 
 USTRUCT()
-struct FLandscapeProceduralLayerBrush
+struct FLandscapeLayerBrush
 {
 	GENERATED_USTRUCT_BODY()
 
-	FLandscapeProceduralLayerBrush()
+	FLandscapeLayerBrush()
 		: BPCustomBrush(nullptr)
 	{}
 
-	FLandscapeProceduralLayerBrush(ALandscapeBlueprintCustomBrush* InBrush)
+	FLandscapeLayerBrush(ALandscapeBlueprintCustomBrush* InBrush)
 		: BPCustomBrush(InBrush)
 	{}
 
@@ -140,11 +140,11 @@ struct FLandscapeProceduralLayerBrush
 };
 
 USTRUCT()
-struct FProceduralLayer
+struct FLandscapeLayer
 {
 	GENERATED_USTRUCT_BODY()
 
-	FProceduralLayer()
+	FLandscapeLayer()
 		: Guid(FGuid::NewGuid())
 		, Name(NAME_None)
 		, bVisible(true)
@@ -172,7 +172,7 @@ struct FProceduralLayer
 	float WeightmapAlpha;
 
 	UPROPERTY()
-	TArray<FLandscapeProceduralLayerBrush> Brushes;
+	TArray<FLandscapeLayerBrush> Brushes;
 
 	UPROPERTY()
 	TArray<int8> HeightmapBrushOrderIndices;
@@ -227,81 +227,93 @@ public:
 	//~ End UObject Interface
 
 
-	// Procedural stuff
+	// Layers stuff
 #if WITH_EDITOR
-	LANDSCAPE_API void RequestProceduralContentUpdate(uint32 InDataFlags, bool InUpdateAllMaterials = false);
-	LANDSCAPE_API void CreateProceduralLayer(FName InName = NAME_None, bool bInUpdateProceduralContent = true);
-	LANDSCAPE_API bool IsProceduralLayerNameUnique(const FName& InName) const;
-	LANDSCAPE_API void SetProceduralLayerName(int32 InLayerIndex, const FName& InName);
-	LANDSCAPE_API void SetProceduralLayerAlpha(int32 InLayerIndex, const float InAlpha, bool bInHeightmap);
-	LANDSCAPE_API void SetProceduralLayerVisibility(int32 InLayerIndex, bool bInVisible);
-	LANDSCAPE_API struct FProceduralLayer* GetProceduralLayer(int32 InLayerIndex);
-	LANDSCAPE_API const struct FProceduralLayer* GetProceduralLayer(int32 InLayerIndex) const;
-	LANDSCAPE_API void ClearProceduralLayer(int32 InLayerIndex);
-	LANDSCAPE_API void ClearProceduralLayer(const FGuid& InLayerGuid);
-	LANDSCAPE_API void DeleteProceduralLayer(int32 InLayerIndex);
-	LANDSCAPE_API void SetCurrentEditingProceduralLayer(FGuid InLayerGuid = FGuid());
-	LANDSCAPE_API void ShowOnlySelectedProceduralLayer(int32 InLayerIndex);
-	LANDSCAPE_API void ShowAllProceduralLayers();
+	LANDSCAPE_API void RequestLayersContentUpdate(uint32 InDataFlags, bool InUpdateAllMaterials = false);
+	LANDSCAPE_API void CreateLayer(FName InName = NAME_None, bool bInUpdateLayersContent = true);
+	LANDSCAPE_API bool IsLayerNameUnique(const FName& InName) const;
+	LANDSCAPE_API void SetLayerName(int32 InLayerIndex, const FName& InName);
+	LANDSCAPE_API void SetLayerAlpha(int32 InLayerIndex, const float InAlpha, bool bInHeightmap);
+	LANDSCAPE_API void SetLayerVisibility(int32 InLayerIndex, bool bInVisible);
+	LANDSCAPE_API struct FLandscapeLayer* GetLayer(int32 InLayerIndex);
+	LANDSCAPE_API const struct FLandscapeLayer* GetLayer(int32 InLayerIndex) const;
+	LANDSCAPE_API void ClearLayer(int32 InLayerIndex);
+	LANDSCAPE_API void ClearLayer(const FGuid& InLayerGuid);
+	LANDSCAPE_API void DeleteLayer(int32 InLayerIndex);
+	LANDSCAPE_API void SetEditingLayer(FGuid InLayerGuid = FGuid());
+	LANDSCAPE_API void ShowOnlySelectedLayer(int32 InLayerIndex);
+	LANDSCAPE_API void ShowAllLayers();
+
+	LANDSCAPE_API bool IsLayerBlendSubstractive(int32 InLayerIndex, const TWeakObjectPtr<ULandscapeLayerInfoObject>& InLayerInfoObj) const;
+	LANDSCAPE_API void SetLayerSubstractiveBlendStatus(int32 InLayerIndex, bool InStatus, const TWeakObjectPtr<ULandscapeLayerInfoObject>& InLayerInfoObj);
+
+	LANDSCAPE_API void AddBrushToLayer(int32 InLayerIndex, int32 InTargetType, class ALandscapeBlueprintCustomBrush* InBrush);
+	LANDSCAPE_API void RemoveBrushFromLayer(int32 InLayerIndex, int32 InTargetType, class ALandscapeBlueprintCustomBrush* InBrush);
+	LANDSCAPE_API bool AreAllBrushesCommitedToLayer(int32 InLayerIndex, int32 InTargetType);
+	LANDSCAPE_API void SetBrushesCommitStateForLayer(int32 InLayerIndex, int32 InTargetType, bool InCommited);
+	LANDSCAPE_API TArray<int8>& GetBrushesOrderForLayer(int32 InLayerIndex, int32 InTargetType);
+	LANDSCAPE_API class ALandscapeBlueprintCustomBrush* GetBrushForLayer(int32 InLayerIndex, int32 InTargetType, int8 BrushIndex) const;
+	LANDSCAPE_API TArray<class ALandscapeBlueprintCustomBrush*> GetBrushesForLayer(int32 InLayerIndex, int32 InTargetType) const;
+		
 
 private:
-	void TickProcedural(float DeltaTime, ELevelTick TickType, FActorTickFunction& ThisTickFunction);
-	void RegenerateProceduralContent();
-	void RegenerateProceduralHeightmaps();
-	void ResolveProceduralHeightmapTexture(const TArray<ALandscapeProxy*>& InAllLandscapes);
-	void ResolveProceduralWeightmapTexture(const TArray<ALandscapeProxy*>& InAllLandscapes);
-	void ResolveProceduralTexture(FLandscapeProceduralTexture2DCPUReadBackResource* InCPUReadBackTexture, UTexture2D* InOriginalTexture);
-	void RegenerateProceduralWeightmaps();
-
-	bool AreHeightmapTextureResourcesReady(const TArray<ALandscapeProxy*>& InAllLandscapes) const;
-	bool AreWeightmapTextureResourcesReady(const TArray<ALandscapeProxy*>& InAllLandscapes) const;
-
-	void UpdateProceduralMaterialInstances(const TArray<ULandscapeComponent*>& InComponentsToUpdate, const TMap<ULandscapeComponent*, TArray<ULandscapeLayerInfoObject*>>& InZeroAllocationsPerComponents);
-
-	void PrepareProceduralComponentDataForExtractLayersCS(const FProceduralLayer& InProceduralLayer, int32 InCurrentWeightmapToProcessIndex, bool InOutputDebugName, const TArray<ALandscapeProxy*>& InAllLandscape, class FLandscapeTexture2DResource* InOutTextureData,
-														  TArray<struct FLandscapeProceduralWeightmapExtractLayersComponentData>& OutComponentData, TMap<ULandscapeLayerInfoObject*, int32>& OutLayerInfoObjects);
-	void PrepareProceduralComponentDataForPackLayersCS(int32 InCurrentWeightmapToProcessIndex, bool InOutputDebugName, const TArray<ULandscapeComponent*>& InAllLandscapeComponents, 
-													   TArray<UTexture2D*>& InOutProcessedWeightmaps, TArray<class FLandscapeProceduralTexture2DCPUReadBackResource*>& InOutProcessedWeightmapCPUCopy, TArray<struct FLandscapeProceduralWeightmapPackLayersComponentData>& OutComponentData);
-	void ReallocateProceduralWeightmaps(const TArray<ALandscapeProxy*>& InAllLandscape, const TArray<ULandscapeLayerInfoObject*>& InBrushRequiredAllocations, TArray<ULandscapeComponent*>& OutComponentThatNeedMaterialRebuild);
-	void InitProceduralWeightmapResources(uint8 InLayerCount);
+	void TickLayers(float DeltaTime, ELevelTick TickType, FActorTickFunction& ThisTickFunction);
+	void RegenerateLayersContent();
+	void RegenerateLayersHeightmaps();
+	void RegenerateLayersWeightmaps();
+	void ResolveLayersHeightmapTexture(const TArray<ALandscapeProxy*>& InAllLandscapes);
+	void ResolveLayersWeightmapTexture(const TArray<ALandscapeProxy*>& InAllLandscapes);
+	void ResolveLayersTexture(FLandscapeLayersTexture2DCPUReadBackResource* InCPUReadBackTexture, UTexture2D* InOriginalTexture);
+
+	bool AreLayersHeightmapTextureResourcesReady(const TArray<ALandscapeProxy*>& InAllLandscapes) const;
+	bool AreLayersWeightmapTextureResourcesReady(const TArray<ALandscapeProxy*>& InAllLandscapes) const;
+
+	void UpdateLayersMaterialInstances(const TArray<ULandscapeComponent*>& InComponentsToUpdate, const TMap<ULandscapeComponent*, TArray<ULandscapeLayerInfoObject*>>& InZeroAllocationsPerComponents);
+
+	void PrepareComponentDataToExtractMaterialLayersCS(const FLandscapeLayer& InLayer, int32 InCurrentWeightmapToProcessIndex, bool InOutputDebugName, const TArray<ALandscapeProxy*>& InAllLandscape, class FLandscapeTexture2DResource* InOutTextureData,
+														  TArray<struct FLandscapeLayerWeightmapExtractMaterialLayersComponentData>& OutComponentData, TMap<ULandscapeLayerInfoObject*, int32>& OutLayerInfoObjects);
+	void PrepareComponentDataToPackMaterialLayersCS(int32 InCurrentWeightmapToProcessIndex, bool InOutputDebugName, const TArray<ULandscapeComponent*>& InAllLandscapeComponents, 
+													   TArray<UTexture2D*>& InOutProcessedWeightmaps, TArray<class FLandscapeLayersTexture2DCPUReadBackResource*>& InOutProcessedWeightmapCPUCopy, TArray<struct FLandscapeLayerWeightmapPackMaterialLayersComponentData>& OutComponentData);
+	void ReallocateLayersWeightmaps(const TArray<ALandscapeProxy*>& InAllLandscape, const TArray<ULandscapeLayerInfoObject*>& InBrushRequiredAllocations, TArray<ULandscapeComponent*>& OutComponentThatNeedMaterialRebuild);
+	void InitLayersWeightmapResources(uint8 InLayerCount);
 	bool GenerateZeroAllocationPerComponents(const TArray<ALandscapeProxy*>& InAllLandscape, const TMap<ULandscapeLayerInfoObject*, bool>& InWeightmapLayersBlendSubstractive, TMap<ULandscapeComponent*, TArray<ULandscapeLayerInfoObject*>>& OutZeroAllocationsPerComponents);
 
-	void GenerateProceduralRenderQuad(const FIntPoint& InVertexPosition, float InVertexSize, const FVector2D& InUVStart, const FVector2D& InUVSize, TArray<struct FLandscapeProceduralTriangle>& OutTriangles) const;
-	void GenerateProceduralRenderQuadsAtlas(const FIntPoint& InSectionBase, const FVector2D& InScaleBias, float InSubSectionSizeQuad, const FIntPoint& InReadSize, const FIntPoint& InWriteSize, TArray<struct FLandscapeProceduralTriangle>& OutTriangles) const;
-	void GenerateProceduralRenderQuadsAtlasToNonAtlas(const FIntPoint& InSectionBase, const FVector2D& InScaleBias, float InSubSectionSizeQuad, const FIntPoint& InReadSize, const FIntPoint& InWriteSize, TArray<struct FLandscapeProceduralTriangle>& OutTriangles) const;
-	void GenerateProceduralRenderQuadsNonAtlas(const FIntPoint& InSectionBase, const FVector2D& InScaleBias, float InSubSectionSizeQuad, const FIntPoint& InReadSize, const FIntPoint& InWriteSize, TArray<struct FLandscapeProceduralTriangle>& OutTriangles) const;
-	void GenerateProceduralRenderQuadsNonAtlasToAtlas(const FIntPoint& InSectionBase, const FVector2D& InScaleBias, float InSubSectionSizeQuad, const FIntPoint& InReadSize, const FIntPoint& InWriteSize, TArray<struct FLandscapeProceduralTriangle>& OutTriangles) const;
-	void GenerateProceduralRenderQuadsMip(const FIntPoint& InSectionBase, const FVector2D& InScaleBias, float InSubSectionSizeQuad, const FIntPoint& InReadSize, const FIntPoint& InWriteSize, uint8 InCurrentMip, TArray<FLandscapeProceduralTriangle>& OutTriangles) const;
+	void GenerateLayersRenderQuad(const FIntPoint& InVertexPosition, float InVertexSize, const FVector2D& InUVStart, const FVector2D& InUVSize, TArray<struct FLandscapeLayersTriangle>& OutTriangles) const;
+	void GenerateLayersRenderQuadsAtlas(const FIntPoint& InSectionBase, const FVector2D& InScaleBias, float InSubSectionSizeQuad, const FIntPoint& InReadSize, const FIntPoint& InWriteSize, TArray<struct FLandscapeLayersTriangle>& OutTriangles) const;
+	void GenerateLayersRenderQuadsAtlasToNonAtlas(const FIntPoint& InSectionBase, const FVector2D& InScaleBias, float InSubSectionSizeQuad, const FIntPoint& InReadSize, const FIntPoint& InWriteSize, TArray<struct FLandscapeLayersTriangle>& OutTriangles) const;
+	void GenerateLayersRenderQuadsNonAtlas(const FIntPoint& InSectionBase, const FVector2D& InScaleBias, float InSubSectionSizeQuad, const FIntPoint& InReadSize, const FIntPoint& InWriteSize, TArray<struct FLandscapeLayersTriangle>& OutTriangles) const;
+	void GenerateLayersRenderQuadsNonAtlasToAtlas(const FIntPoint& InSectionBase, const FVector2D& InScaleBias, float InSubSectionSizeQuad, const FIntPoint& InReadSize, const FIntPoint& InWriteSize, TArray<struct FLandscapeLayersTriangle>& OutTriangles) const;
+	void GenerateLayersRenderQuadsMip(const FIntPoint& InSectionBase, const FVector2D& InScaleBias, float InSubSectionSizeQuad, const FIntPoint& InReadSize, const FIntPoint& InWriteSize, uint8 InCurrentMip, TArray<FLandscapeLayersTriangle>& OutTriangles) const;
 
-	void ClearWeightmapTextureResource(const FString& InDebugName, FTextureRenderTargetResource* InTextureResourceToClear);
+	void ClearLayersWeightmapTextureResource(const FString& InDebugName, FTextureRenderTargetResource* InTextureResourceToClear);
 	void DrawHeightmapComponentsToRenderTarget(const FString& InDebugName, const TArray<ULandscapeComponent*>& InComponentsToDraw, UTexture* InHeightmapRTRead, UTextureRenderTarget2D* InOptionalHeightmapRTRead2, UTextureRenderTarget2D* InHeightmapRTWrite, ERTDrawingType InDrawType,
-											   bool InClearRTWrite, struct FLandscapeHeightmapProceduralShaderParameters& InShaderParams, uint8 InMipRender = 0) const;
+											   bool InClearRTWrite, struct FLandscapeLayersHeightmapShaderParameters& InShaderParams, uint8 InMipRender = 0) const;
 
 	void DrawWeightmapComponentsToRenderTarget(const FString& InDebugName, const TArray<ULandscapeComponent*>& InComponentsToDraw, UTexture* InWeightmapRTRead, UTextureRenderTarget2D* InOptionalWeightmapRTRead2, UTextureRenderTarget2D* InWeightmapRTWrite,
-											   bool InClearRTWrite, struct FLandscapeWeightmapProceduralShaderParameters& InShaderParams, uint8 InMipRender) const;
+											   bool InClearRTWrite, struct FLandscapeLayersWeightmapShaderParameters& InShaderParams, uint8 InMipRender) const;
 
 	void DrawWeightmapComponentsToRenderTarget(const FString& InDebugName, const FIntPoint& InSectionBase, const FVector2D& InScaleBias, UTexture* InWeightmapRTRead, UTextureRenderTarget2D* InOptionalWeightmapRTRead2, UTextureRenderTarget2D* InWeightmapRTWrite,
-									 		   bool InClearRTWrite, struct FLandscapeWeightmapProceduralShaderParameters& InShaderParams, uint8 InMipRender) const;
+									 		   bool InClearRTWrite, struct FLandscapeLayersWeightmapShaderParameters& InShaderParams, uint8 InMipRender) const;
 
-	void DrawHeightmapComponentsToRenderTargetMips(TArray<ULandscapeComponent*>& InComponentsToDraw, UTexture* InReadHeightmap, bool InClearRTWrite, struct FLandscapeHeightmapProceduralShaderParameters& InShaderParams) const;
-	void DrawWeightmapComponentToRenderTargetMips(const FIntPoint& TopLeftTexturePosition, UTexture* InReadWeightmap, bool InClearRTWrite, struct FLandscapeWeightmapProceduralShaderParameters& InShaderParams) const;
+	void DrawHeightmapComponentsToRenderTargetMips(TArray<ULandscapeComponent*>& InComponentsToDraw, UTexture* InReadHeightmap, bool InClearRTWrite, struct FLandscapeLayersHeightmapShaderParameters& InShaderParams) const;
+	void DrawWeightmapComponentToRenderTargetMips(const FIntPoint& TopLeftTexturePosition, UTexture* InReadWeightmap, bool InClearRTWrite, struct FLandscapeLayersWeightmapShaderParameters& InShaderParams) const;
 
-	void CopyProceduralTexture(UTexture* InSourceTexture, UTexture* InDestTexture, FTextureResource* InDestCPUResource = nullptr, const FIntPoint& InFirstComponentSectionBase = FIntPoint(0, 0), uint8 InSourceCurrentMip = 0, uint8 InDestCurrentMip = 0,
+	void CopyLayersTexture(UTexture* InSourceTexture, UTexture* InDestTexture, FTextureResource* InDestCPUResource = nullptr, const FIntPoint& InFirstComponentSectionBase = FIntPoint(0, 0), uint8 InSourceCurrentMip = 0, uint8 InDestCurrentMip = 0,
 							   uint32 InSourceArrayIndex = 0, uint32 InDestArrayIndex = 0) const;
-	void CopyProceduralTexture(const FString& InSourceDebugName, FTextureResource* InSourceResource, const FString& InDestDebugName, FTextureResource* InDestResource, FTextureResource* InDestCPUResource = nullptr, const FIntPoint& InFirstComponentSectionBase = FIntPoint(0, 0),
+	void CopyLayersTexture(const FString& InSourceDebugName, FTextureResource* InSourceResource, const FString& InDestDebugName, FTextureResource* InDestResource, FTextureResource* InDestCPUResource = nullptr, const FIntPoint& InFirstComponentSectionBase = FIntPoint(0, 0),
 							   uint8 InSourceCurrentMip = 0, uint8 InDestCurrentMip = 0, uint32 InSourceArrayIndex = 0, uint32 uInDestArrayIndex = 0) const;
 
-	void PrintProceduralDebugRT(const FString& InContext, UTextureRenderTarget2D* InDebugRT, uint8 InMipRender = 0, bool InOutputHeight = true, bool InOutputNormals = false) const;
-	void PrintProceduralDebugTextureResource(const FString& InContext, FTextureResource* InTextureResource, uint8 InMipRender = 0, bool InOutputHeight = true, bool InOutputNormals = false) const;
-	void PrintProceduralDebugHeightData(const FString& InContext, const TArray<FColor>& InHeightmapData, const FIntPoint& InDataSize, uint8 InMipRender, bool InOutputNormals = false) const;
-	void PrintProceduralDebugWeightData(const FString& InContext, const TArray<FColor>& InWeightmapData, const FIntPoint& InDataSize, uint8 InMipRender) const;
+	void PrintLayersDebugRT(const FString& InContext, UTextureRenderTarget2D* InDebugRT, uint8 InMipRender = 0, bool InOutputHeight = true, bool InOutputNormals = false) const;
+	void PrintLayersDebugTextureResource(const FString& InContext, FTextureResource* InTextureResource, uint8 InMipRender = 0, bool InOutputHeight = true, bool InOutputNormals = false) const;
+	void PrintLayersDebugHeightData(const FString& InContext, const TArray<FColor>& InHeightmapData, const FIntPoint& InDataSize, uint8 InMipRender, bool InOutputNormals = false) const;
+	void PrintLayersDebugWeightData(const FString& InContext, const TArray<FColor>& InWeightmapData, const FIntPoint& InDataSize, uint8 InMipRender) const;
 #endif
 
 public:
 
 #if WITH_EDITORONLY_DATA
 	UPROPERTY(TextExportTransient)
-	TArray<FProceduralLayer> ProceduralLayers;
+	TArray<FLandscapeLayer> LandscapeLayers;
 
 	UPROPERTY(Transient)
 	TArray<UTextureRenderTarget2D*> HeightmapRTList;
@@ -310,7 +322,7 @@ public:
 	TArray<UTextureRenderTarget2D*> WeightmapRTList;
 
 	UPROPERTY(Transient)
-	bool PreviousExperimentalLandscapeProcedural;
+	bool PreviousExperimentalLandscapeLayers;
 
 private:
 
@@ -318,40 +330,40 @@ private:
 	bool WasCompilingShaders;
 
 	UPROPERTY(Transient)
-	uint32 ProceduralContentUpdateFlags;
+	uint32 LayersContentUpdateFlags;
 
 	UPROPERTY(Transient)
-	bool ProceduralUpdateAllMaterials;
+	bool LayersUpdateAllMaterials;
 
-	// Represent all the resolved paint layer, from all procedural layer blended together (size of the landscape x paint layer count)
-	class FLandscapeTexture2DArrayResource* CombinedProcLayerWeightmapAllLayersResource;
+	// Represent all the resolved paint layer, from all layers blended together (size of the landscape x material layer count)
+	class FLandscapeTexture2DArrayResource* CombinedLayersWeightmapAllMaterialLayersResource;
 	
-	// Represent all the resolved paint layer, from the current procedual layer only (size of the landscape x paint layer count)
-	class FLandscapeTexture2DArrayResource* CurrentProcLayerWeightmapAllLayersResource;	
+	// Represent all the resolved paint layer, from the current layer only (size of the landscape x material layer count)
+	class FLandscapeTexture2DArrayResource* CurrentLayersWeightmapAllMaterialLayersResource;	
 	
-	// Used in extracting the paint layers data from procedural layer weightmaps (size of the landscape)
+	// Used in extracting the material layers data from layer weightmaps (size of the landscape)
 	class FLandscapeTexture2DResource* WeightmapScratchExtractLayerTextureResource;	
 	
-	// Used in packing the paint layers data contained into CombinedProcLayerWeightmapAllLayersResource to be set again for each component weightmap (size of the landscape)
+	// Used in packing the material layer data contained into CombinedLayersWeightmapAllMaterialLayersResource to be set again for each component weightmap (size of the landscape)
 	class FLandscapeTexture2DResource* WeightmapScratchPackLayerTextureResource;	
 #endif
 
 protected:
 #if WITH_EDITOR
-	FName GenerateUniqueProceduralLayerName(FName InName = NAME_None) const;
+	FName GenerateUniqueLayerName(FName InName = NAME_None) const;
 #endif
 };
 
 #if WITH_EDITOR
-class LANDSCAPE_API FScopedSetLandscapeCurrentEditingProceduralLayer
+class LANDSCAPE_API FScopedSetLandscapeEditingLayer
 {
 public:
-	FScopedSetLandscapeCurrentEditingProceduralLayer(ALandscape* InLandscape, const FGuid& InProceduralLayer, TFunction<void()> InCompletionCallback = TFunction<void()>());
-	~FScopedSetLandscapeCurrentEditingProceduralLayer();
+	FScopedSetLandscapeEditingLayer(ALandscape* InLandscape, const FGuid& InLayerGUID, TFunction<void()> InCompletionCallback = TFunction<void()>());
+	~FScopedSetLandscapeEditingLayer();
 
 private:
 	TWeakObjectPtr<ALandscape> Landscape;
-	const FGuid& ProceduralLayer;
+	const FGuid& LayerGUID;
 	TFunction<void()> CompletionCallback;
 };
 #endif
diff --git a/Engine/Source/Runtime/Landscape/Classes/LandscapeComponent.h b/Engine/Source/Runtime/Landscape/Classes/LandscapeComponent.h
--- a/Engine/Source/Runtime/Landscape/Classes/LandscapeComponent.h
+++ b/Engine/Source/Runtime/Landscape/Classes/LandscapeComponent.h
@@ -335,12 +335,12 @@ private:
 	UPROPERTY()
 	FGuid LightingGuid;
 
-	/** Current data we're working on (only used in Procedural mode) This data will get set by tools as needed */
+	/** Current data we're working on (only used in Layering mode) This data will get set by tools as needed */
 	UTexture2D* CurrentEditingHeightmapTexture;
 	TArray<FWeightmapLayerAllocationInfo>* CurrentEditingWeightmapLayerAllocations;
 	TArray<UTexture2D*>* CurrentEditingWeightmapTextures;
 	TArray<ULandscapeWeightmapUsage*>* CurrentEditingWeightmapTexturesUsage;
-	FGuid CurrentProceduralLayerGuid;
+	FGuid CurrentLayerGuid;
 
 	TArray<ULandscapeWeightmapUsage*> WeightmapTexturesUsage;
 #endif // WITH_EDITORONLY_DATA
@@ -517,14 +517,14 @@ public:
 	LANDSCAPE_API void SetHeightmap(UTexture2D* NewHeightmap);
 
 	LANDSCAPE_API void SetWeightmapTextures(const TArray<UTexture2D*>& InNewWeightmapTextures, bool InApplyToCurrentEditingWeightmap = false);
-	LANDSCAPE_API const FGuid& GetCurrentProceduralLayerGuid() const;
+	LANDSCAPE_API const FGuid& GetCurrentLayerGuid() const;
 
 	LANDSCAPE_API void SetWeightmapLayerAllocations(const TArray<FWeightmapLayerAllocationInfo>& InNewWeightmapLayerAllocations);
 	LANDSCAPE_API void SetWeightmapTexturesUsage(const TArray<ULandscapeWeightmapUsage*>& InNewWeightmapTexturesUsage, bool InApplyToCurrentEditingWeightmap = false);
 	LANDSCAPE_API TArray<ULandscapeWeightmapUsage*>& GetWeightmapTexturesUsage(bool InReturnCurrentEditingWeightmap = false);
 	LANDSCAPE_API const TArray<ULandscapeWeightmapUsage*>& GetWeightmapTexturesUsage(bool InReturnCurrentEditingWeightmap = false) const;
 
-	LANDSCAPE_API void SetCurrentEditingProceduralLayer(struct FProceduralLayer* Layer, struct FProceduralLayerData* LayerData);
+	LANDSCAPE_API void SetEditingLayer(struct FLandscapeLayer* InLayer, struct FLandscapeLayerData* InLayerData);
 #endif 
 
 #if WITH_EDITOR
@@ -799,11 +799,9 @@ public:
 protected:
 
 #if WITH_EDITOR
-	void SetCurrentEditingHeightmap(UTexture2D* InNewHeightmap);
-	void SetCurrentEditingWeightmaps(TArray<UTexture2D*>* InNewWeightmapTextures);
-	void SetCurrentProceduralLayerGuid(const FGuid& InLayerGuid);
-	void SetCurrentEditingWeightmapLayerAllocations(TArray<FWeightmapLayerAllocationInfo>* InNewWeightmapLayerAllocations);
-	void SetCurrentEditingWeightmapTexturesUsage(TArray<ULandscapeWeightmapUsage*>* InNewWeightmapTexturesUsage);
+	void SetEditingHeightmap(UTexture2D* InNewHeightmap);
+	void SetEditingWeightmaps(TArray<UTexture2D*>* InNewWeightmapTextures, TArray<FWeightmapLayerAllocationInfo>* InNewWeightmapLayerAllocations, TArray<ULandscapeWeightmapUsage*>* InNewWeightmapTexturesUsage);
+	void SetEditingLayerGuid(const FGuid& InLayerGuid);
 #endif
 
 	/** Whether the component type supports static lighting. */
diff --git a/Engine/Source/Runtime/Landscape/Classes/LandscapeProxy.h b/Engine/Source/Runtime/Landscape/Classes/LandscapeProxy.h
--- a/Engine/Source/Runtime/Landscape/Classes/LandscapeProxy.h
+++ b/Engine/Source/Runtime/Landscape/Classes/LandscapeProxy.h
@@ -323,10 +323,10 @@ struct FLandscapeProxyMaterialOverride
 	UMaterialInterface* Material;
 };
 
-class FLandscapeProceduralTexture2DCPUReadBackResource : public FTextureResource
+class FLandscapeLayersTexture2DCPUReadBackResource : public FTextureResource
 {
 public:
-	FLandscapeProceduralTexture2DCPUReadBackResource(uint32 InSizeX, uint32 InSizeY, EPixelFormat InFormat, uint32 InNumMips)
+	FLandscapeLayersTexture2DCPUReadBackResource(uint32 InSizeX, uint32 InSizeY, EPixelFormat InFormat, uint32 InNumMips)
 		: SizeX(InSizeX)
 		, SizeY(InSizeY)
 		, Format(InFormat)
@@ -367,7 +367,7 @@ struct FRenderDataPerHeightmap
 	UPROPERTY(Transient)
 	UTexture2D* OriginalHeightmap;
 
-	FLandscapeProceduralTexture2DCPUReadBackResource* HeightmapsCPUReadBack;
+	FLandscapeLayersTexture2DCPUReadBackResource* HeightmapsCPUReadBack;
 
 	UPROPERTY(Transient)
 	TArray<ULandscapeComponent*> Components;
@@ -391,11 +391,11 @@ struct FWeightmapLayerData
 };
 
 USTRUCT()
-struct FProceduralLayerData
+struct FLandscapeLayerData
 {
 	GENERATED_USTRUCT_BODY()
 
-	FProceduralLayerData()
+	FLandscapeLayerData()
 	{}
 
 	UPROPERTY()
@@ -643,15 +643,15 @@ public:
 	TArray<FLandscapeEditorLayerSettings> EditorLayerSettings;
 
 	UPROPERTY(TextExportTransient)
-	TMap<FGuid, FProceduralLayerData> ProceduralLayersData;
+	TMap<FGuid, FLandscapeLayerData> LandscapeLayersData;
 
 	UPROPERTY()
-	bool HasProceduralContent;
+	bool HasLayersContent;
 
 	UPROPERTY(Transient)
 	TMap<UTexture2D*, FRenderDataPerHeightmap> RenderDataPerHeightmap; // Mapping between Original heightmap and general render data
 
-	TMap<UTexture2D*, FLandscapeProceduralTexture2DCPUReadBackResource*> WeightmapCPUReadBackTextures; // Mapping between Original weightmap and tyhe CPU readback resource
+	TMap<UTexture2D*, FLandscapeLayersTexture2DCPUReadBackResource*> WeightmapCPUReadBackTextures; // Mapping between Original weightmap and tyhe CPU readback resource
 
 	FRenderCommandFence ReleaseResourceFence;
 #endif
@@ -986,6 +986,6 @@ public:
 protected:
 	FLandscapeMaterialChangedDelegate LandscapeMaterialChangedDelegate;
 	
-	LANDSCAPE_API void SetupProceduralLayers(int32 InNumComponentsX = INDEX_NONE, int32 InNumComponentsY = INDEX_NONE);
+	LANDSCAPE_API void SetupLayers(int32 InNumComponentsX = INDEX_NONE, int32 InNumComponentsY = INDEX_NONE);
 #endif
 };
diff --git a/Engine/Source/Runtime/Landscape/Classes/LandscapeWeightmapUsage.h b/Engine/Source/Runtime/Landscape/Classes/LandscapeWeightmapUsage.h
--- a/Engine/Source/Runtime/Landscape/Classes/LandscapeWeightmapUsage.h
+++ b/Engine/Source/Runtime/Landscape/Classes/LandscapeWeightmapUsage.h
@@ -19,7 +19,7 @@ public:
 	ULandscapeComponent* ChannelUsage[4];
 
 	UPROPERTY()
-	FGuid ProceduralLayerGuid;
+	FGuid LayerGuid;
 
 	int32 FreeChannelCount() const
 	{
diff --git a/Engine/Source/Runtime/Landscape/Private/Landscape.cpp b/Engine/Source/Runtime/Landscape/Private/Landscape.cpp
--- a/Engine/Source/Runtime/Landscape/Private/Landscape.cpp
+++ b/Engine/Source/Runtime/Landscape/Private/Landscape.cpp
@@ -76,17 +76,17 @@ DEFINE_STAT(STAT_LandscapeComponentUsingSubSectionDrawCalls);
 DEFINE_STAT(STAT_LandscapeDrawCalls);
 DEFINE_STAT(STAT_LandscapeTriangles);
 
-DEFINE_STAT(STAT_LandscapeRegenerateProcedural_RenderThread);
-DEFINE_STAT(STAT_LandscapeRegenerateProceduralDrawCalls);
+DEFINE_STAT(STAT_LandscapeLayersRegenerate_RenderThread);
+DEFINE_STAT(STAT_LandscapeLayersRegenerateDrawCalls);
 
-DEFINE_STAT(STAT_LandscapeRegenerateProceduralHeightmaps);
-DEFINE_STAT(STAT_LandscapeResolveProceduralHeightmap);
+DEFINE_STAT(STAT_LandscapeLayersRegenerateHeightmaps);
+DEFINE_STAT(STAT_LandscapeLayersResolveHeightmaps);
 
-DEFINE_STAT(STAT_LandscapeProceduralUpdateMaterialInstance);
-DEFINE_STAT(STAT_LandscapeReallocateProceduralWeightmaps);
+DEFINE_STAT(STAT_LandscapeLayersUpdateMaterialInstance);
+DEFINE_STAT(STAT_LandscapeLayersReallocateWeightmaps);
 
-DEFINE_STAT(STAT_LandscapeResolveProceduralWeightmap);
-DEFINE_STAT(STAT_LandscapeRegenerateProceduralWeightmaps);
+DEFINE_STAT(STAT_LandscapeLayersResolveWeightmaps);
+DEFINE_STAT(STAT_LandscapeLayersRegenerateWeightmaps);
 
 DEFINE_STAT(STAT_LandscapeVertexMem);
 DEFINE_STAT(STAT_LandscapeOccluderMem);
@@ -169,7 +169,7 @@ ULandscapeComponent::ULandscapeComponent(const FObjectInitializer& ObjectInitial
 
 #if WITH_EDITORONLY_DATA
 	EditToolRenderData = FLandscapeEditToolRenderData();
-	CurrentProceduralLayerGuid.Invalidate();
+	CurrentLayerGuid.Invalidate();
 #endif
 
 	LpvBiasMultiplier = 0.0f; // Bias is 0 for landscape, since it's single sided
@@ -944,7 +944,7 @@ ALandscapeProxy::ALandscapeProxy(const FObjectInitializer& ObjectInitializer)
 	bGenerateOverlapEvents = false;
 #if WITH_EDITORONLY_DATA
 	MaxPaintedLayersPerComponent = 0;
-	HasProceduralContent = false;
+	HasLayersContent = false;
 #endif
 
 #if WITH_EDITOR
@@ -987,11 +987,11 @@ ALandscape::ALandscape(const FObjectInitializer& ObjectInitializer)
 {
 #if WITH_EDITORONLY_DATA
 	bLockLocation = false;
-	PreviousExperimentalLandscapeProcedural = false;
+	PreviousExperimentalLandscapeLayers = false;
 	WasCompilingShaders = false;
-	ProceduralContentUpdateFlags = 0;
-	CombinedProcLayerWeightmapAllLayersResource = nullptr;
-	CurrentProcLayerWeightmapAllLayersResource = nullptr;
+	LayersContentUpdateFlags = 0;
+	CombinedLayersWeightmapAllMaterialLayersResource = nullptr;
+	CurrentLayersWeightmapAllMaterialLayersResource = nullptr;
 	WeightmapScratchExtractLayerTextureResource = nullptr;
 	WeightmapScratchPackLayerTextureResource = nullptr;
 #endif // WITH_EDITORONLY_DATA
@@ -1342,7 +1342,7 @@ void ULandscapeComponent::OnUnregister()
 UTexture2D* ULandscapeComponent::GetHeightmap(bool InReturnCurrentEditingHeightmap) const
 {
 #if WITH_EDITORONLY_DATA
-	if (InReturnCurrentEditingHeightmap && GetMutableDefault<UEditorExperimentalSettings>()->bProceduralLandscape)
+	if (InReturnCurrentEditingHeightmap && GetMutableDefault<UEditorExperimentalSettings>()->bLandscapeLayerSystem)
 	{
 		if (CurrentEditingHeightmapTexture != nullptr)
 		{
@@ -1357,7 +1357,7 @@ UTexture2D* ULandscapeComponent::GetHeightmap(bool InReturnCurrentEditingHeightm
 const TArray<UTexture2D*>& ULandscapeComponent::GetWeightmapTextures(bool InReturnCurrentEditingWeightmap) const
 {
 #if WITH_EDITORONLY_DATA
-	if (InReturnCurrentEditingWeightmap && GetMutableDefault<UEditorExperimentalSettings>()->bProceduralLandscape)
+	if (InReturnCurrentEditingWeightmap && GetMutableDefault<UEditorExperimentalSettings>()->bLandscapeLayerSystem)
 	{
 		if (CurrentEditingWeightmapTextures != nullptr)
 		{
@@ -1372,7 +1372,7 @@ const TArray<UTexture2D*>& ULandscapeComponent::GetWeightmapTextures(bool InRetu
 TArray<UTexture2D*>& ULandscapeComponent::GetWeightmapTextures(bool InReturnCurrentEditingWeightmap)
 {
 #if WITH_EDITORONLY_DATA
-	if (InReturnCurrentEditingWeightmap && GetMutableDefault<UEditorExperimentalSettings>()->bProceduralLandscape)
+	if (InReturnCurrentEditingWeightmap && GetMutableDefault<UEditorExperimentalSettings>()->bLandscapeLayerSystem)
 	{
 		if (CurrentEditingWeightmapTextures != nullptr)
 		{
@@ -1387,7 +1387,7 @@ TArray<UTexture2D*>& ULandscapeComponent::GetWeightmapTextures(bool InReturnCurr
 const TArray<FWeightmapLayerAllocationInfo>& ULandscapeComponent::GetWeightmapLayerAllocations(bool InReturnCurrentEditingWeightmap) const
 {
 #if WITH_EDITORONLY_DATA
-	if (InReturnCurrentEditingWeightmap && GetMutableDefault<UEditorExperimentalSettings>()->bProceduralLandscape)
+	if (InReturnCurrentEditingWeightmap && GetMutableDefault<UEditorExperimentalSettings>()->bLandscapeLayerSystem)
 	{
 		if (CurrentEditingWeightmapLayerAllocations != nullptr)
 		{
@@ -1402,7 +1402,7 @@ const TArray<FWeightmapLayerAllocationInfo>& ULandscapeComponent::GetWeightmapLa
 TArray<FWeightmapLayerAllocationInfo>& ULandscapeComponent::GetWeightmapLayerAllocations(bool InReturnCurrentEditingWeightmap)
 {
 #if WITH_EDITORONLY_DATA
-	if (InReturnCurrentEditingWeightmap && GetMutableDefault<UEditorExperimentalSettings>()->bProceduralLandscape)
+	if (InReturnCurrentEditingWeightmap && GetMutableDefault<UEditorExperimentalSettings>()->bLandscapeLayerSystem)
 	{
 		if (CurrentEditingWeightmapLayerAllocations != nullptr)
 		{
@@ -1416,21 +1416,20 @@ TArray<FWeightmapLayerAllocationInfo>& ULandscapeComponent::GetWeightmapLayerAll
 
 #if WITH_EDITOR
 
-void ULandscapeComponent::SetCurrentEditingProceduralLayer(FProceduralLayer* Layer, FProceduralLayerData* LayerData)
+void ULandscapeComponent::SetEditingLayer(FLandscapeLayer* InLayer, FLandscapeLayerData* InLayerData)
 {
 	// Update Current Heightmap
-	UTexture2D** LayerHeightmap = LayerData ? LayerData->Heightmaps.Find(GetHeightmap()) : nullptr;
-	SetCurrentEditingHeightmap(LayerHeightmap ? *LayerHeightmap : nullptr);
+	UTexture2D** LayerHeightmap = InLayerData ? InLayerData->Heightmaps.Find(GetHeightmap()) : nullptr;
+	SetEditingHeightmap(LayerHeightmap ? *LayerHeightmap : nullptr);
 
 	// Update Current Weightmaps
-	FWeightmapLayerData* WeightmapData = LayerData ? LayerData->WeightmapData.Find(this) : nullptr;
-	SetCurrentEditingWeightmapLayerAllocations(WeightmapData ? &WeightmapData->WeightmapLayerAllocations : nullptr);
-	SetCurrentEditingWeightmaps(WeightmapData ? &WeightmapData->Weightmaps : nullptr);
-	SetCurrentProceduralLayerGuid(WeightmapData && Layer ? Layer->Guid : FGuid());
-	SetCurrentEditingWeightmapTexturesUsage(WeightmapData ? &WeightmapData->WeightmapTextureUsages : nullptr);
+	FWeightmapLayerData* WeightmapData = InLayerData ? InLayerData->WeightmapData.Find(this) : nullptr;
+
+	SetEditingWeightmaps(WeightmapData ? &WeightmapData->Weightmaps : nullptr, WeightmapData ? &WeightmapData->WeightmapLayerAllocations : nullptr, WeightmapData ? &WeightmapData->WeightmapTextureUsages : nullptr);
+	SetEditingLayerGuid(WeightmapData && InLayer ? InLayer->Guid : FGuid());
 }
 
-void ULandscapeComponent::SetCurrentEditingHeightmap(UTexture2D* InNewHeightmap)
+void ULandscapeComponent::SetEditingHeightmap(UTexture2D* InNewHeightmap)
 {
 #if WITH_EDITORONLY_DATA
 	CurrentEditingHeightmapTexture = InNewHeightmap;
@@ -1458,10 +1457,12 @@ void ULandscapeComponent::SetWeightmapTextures(const TArray<UTexture2D*>& InNewW
 	}
 }
 
-void ULandscapeComponent::SetCurrentEditingWeightmaps(TArray<UTexture2D*>* InNewWeightmapTextures)
+void ULandscapeComponent::SetEditingWeightmaps(TArray<UTexture2D*>* InNewWeightmapTextures, TArray<FWeightmapLayerAllocationInfo>* InNewWeightmapLayerAllocations, TArray<ULandscapeWeightmapUsage*>* InNewWeightmapTexturesUsage)
 {
 #if WITH_EDITORONLY_DATA
 	CurrentEditingWeightmapTextures = InNewWeightmapTextures;
+	CurrentEditingWeightmapLayerAllocations = InNewWeightmapLayerAllocations;
+	CurrentEditingWeightmapTexturesUsage = InNewWeightmapTexturesUsage;
 #endif
 }
 
@@ -1470,18 +1471,10 @@ void ULandscapeComponent::SetWeightmapLayerAllocations(const TArray<FWeightmapLa
 	WeightmapLayerAllocations = InNewWeightmapLayerAllocations;
 }
 
-
-void ULandscapeComponent::SetCurrentEditingWeightmapLayerAllocations(TArray<FWeightmapLayerAllocationInfo>* InNewWeightmapLayerAllocations)
-{
-#if WITH_EDITORONLY_DATA
-	CurrentEditingWeightmapLayerAllocations = InNewWeightmapLayerAllocations;
-#endif
-}
-
 TArray<ULandscapeWeightmapUsage*>& ULandscapeComponent::GetWeightmapTexturesUsage(bool InReturnCurrentEditingWeightmap)
 {
 #if WITH_EDITORONLY_DATA
-	if (InReturnCurrentEditingWeightmap && GetMutableDefault<UEditorExperimentalSettings>()->bProceduralLandscape)
+	if (InReturnCurrentEditingWeightmap && GetMutableDefault<UEditorExperimentalSettings>()->bLandscapeLayerSystem)
 	{
 		if (CurrentEditingWeightmapTexturesUsage != nullptr)
 		{
@@ -1496,7 +1489,7 @@ TArray<ULandscapeWeightmapUsage*>& ULandscapeComponent::GetWeightmapTexturesUsag
 const TArray<ULandscapeWeightmapUsage*>& ULandscapeComponent::GetWeightmapTexturesUsage(bool InReturnCurrentEditingWeightmap) const
 {
 #if WITH_EDITORONLY_DATA
-	if (InReturnCurrentEditingWeightmap && GetMutableDefault<UEditorExperimentalSettings>()->bProceduralLandscape)
+	if (InReturnCurrentEditingWeightmap && GetMutableDefault<UEditorExperimentalSettings>()->bLandscapeLayerSystem)
 	{
 		if (CurrentEditingWeightmapTexturesUsage != nullptr)
 		{
@@ -1508,13 +1501,6 @@ const TArray<ULandscapeWeightmapUsage*>& ULandscapeComponent::GetWeightmapTextur
 	return WeightmapTexturesUsage;
 }
 
-void ULandscapeComponent::SetCurrentEditingWeightmapTexturesUsage(TArray<ULandscapeWeightmapUsage*>* InNewWeightmapTexturesUsage)
-{
-#if WITH_EDITORONLY_DATA
-	CurrentEditingWeightmapTexturesUsage = InNewWeightmapTexturesUsage;
-#endif
-}
-
 void ULandscapeComponent::SetWeightmapTexturesUsage(const TArray<ULandscapeWeightmapUsage*>& InNewWeightmapTexturesUsage, bool InApplyToCurrentEditingWeightmap )
 {
 #if WITH_EDITORONLY_DATA
@@ -1530,14 +1516,14 @@ void ULandscapeComponent::SetWeightmapTexturesUsage(const TArray<ULandscapeWeigh
 	}
 }
 
-void ULandscapeComponent::SetCurrentProceduralLayerGuid(const FGuid& InLayerGuid)
+void ULandscapeComponent::SetEditingLayerGuid(const FGuid& InLayerGuid)
 {
-	CurrentProceduralLayerGuid = InLayerGuid;
+	CurrentLayerGuid = InLayerGuid;
 }
 
-const FGuid& ULandscapeComponent::GetCurrentProceduralLayerGuid() const
+const FGuid& ULandscapeComponent::GetCurrentLayerGuid() const
 {
-	return CurrentProceduralLayerGuid;
+	return CurrentLayerGuid;
 }
 #endif
 
@@ -1697,9 +1683,9 @@ void ALandscapeProxy::PreSave(const class ITargetPlatform* TargetPlatform)
 		bHasLandscapeGrass = LandscapeComponents.ContainsByPredicate([](ULandscapeComponent* Component) { return Component->MaterialHasGrass(); });
 	}
 
-	if (GetMutableDefault<UEditorExperimentalSettings>()->bProceduralLandscape)
+	if (GetMutableDefault<UEditorExperimentalSettings>()->bLandscapeLayerSystem)
 	{
-		HasProceduralContent = true;
+		HasLayersContent = true;
 	}
 #endif
 }
@@ -2073,13 +2059,13 @@ void ALandscapeProxy::PostLoad()
 	FOnFeatureLevelChanged::FDelegate FeatureLevelChangedDelegate = FOnFeatureLevelChanged::FDelegate::CreateUObject(this, &ALandscapeProxy::OnFeatureLevelChanged);
 	FeatureLevelChangedDelegateHandle = GetWorld()->AddOnFeatureLevelChangedHandler(FeatureLevelChangedDelegate);
 
-	if (GetMutableDefault<UEditorExperimentalSettings>()->bProceduralLandscape)
+	if (GetMutableDefault<UEditorExperimentalSettings>()->bLandscapeLayerSystem)
 	{
 		ALandscape* Landscape = GetLandscapeActor();
 
 		if (Landscape != nullptr)
 		{
-			Landscape->RequestProceduralContentUpdate(EProceduralContentUpdateFlag::All_Setup);
+			Landscape->RequestLayersContentUpdate(ELandscapeLayersContentUpdateFlag::All_Setup);
 		}
 	}
 #endif
@@ -2957,7 +2943,7 @@ void ALandscape::TickActor(float DeltaTime, ELevelTick TickType, FActorTickFunct
 #if WITH_EDITOR
 	if (GIsEditor)
 	{
-		TickProcedural(DeltaTime, TickType, ThisTickFunction);
+		TickLayers(DeltaTime, TickType, ThisTickFunction);
 	}
 #endif
 }
diff --git a/Engine/Source/Runtime/Landscape/Private/LandscapeBPCustomBrush.cpp b/Engine/Source/Runtime/Landscape/Private/LandscapeBPCustomBrush.cpp
--- a/Engine/Source/Runtime/Landscape/Private/LandscapeBPCustomBrush.cpp
+++ b/Engine/Source/Runtime/Landscape/Private/LandscapeBPCustomBrush.cpp
@@ -78,7 +78,7 @@ void ALandscapeBlueprintCustomBrush::PostEditMove(bool bFinished)
 
 	if (OwningLandscape != nullptr)
 	{
-		OwningLandscape->RequestProceduralContentUpdate(bFinished ? EProceduralContentUpdateFlag::All : EProceduralContentUpdateFlag::All_Render);
+		OwningLandscape->RequestLayersContentUpdate(bFinished ? ELandscapeLayersContentUpdateFlag::All : ELandscapeLayersContentUpdateFlag::All_Render);
 	}
 }
 
@@ -103,7 +103,7 @@ void ALandscapeBlueprintCustomBrush::PostEditChangeProperty(FPropertyChangedEven
 	if (PropertyName == GET_MEMBER_NAME_CHECKED(ALandscapeBlueprintCustomBrush, AffectHeightmap)
 		|| PropertyName == GET_MEMBER_NAME_CHECKED(ALandscapeBlueprintCustomBrush, AffectWeightmap))
 	{
-		for (FProceduralLayer& Layer : OwningLandscape->ProceduralLayers)
+		for (FLandscapeLayer& Layer : OwningLandscape->LandscapeLayers)
 		{
 			for (int32 i = 0; i < Layer.Brushes.Num(); ++i)
 			{
@@ -157,7 +157,7 @@ void ALandscapeBlueprintCustomBrush::PostEditChangeProperty(FPropertyChangedEven
 
 	if (OwningLandscape != nullptr)
 	{
-		OwningLandscape->RequestProceduralContentUpdate(EProceduralContentUpdateFlag::All, true); // For now since we dont have brush bounds, we have to force an update of all components, otherwise we will have not right layer being rendered as materials are not updated
+		OwningLandscape->RequestLayersContentUpdate(ELandscapeLayersContentUpdateFlag::All, true); // For now since we dont have brush bounds, we have to force an update of all components, otherwise we will have not right layer being rendered as materials are not updated
 	}
 }
 #endif
diff --git a/Engine/Source/Runtime/Landscape/Private/LandscapeEdit.cpp b/Engine/Source/Runtime/Landscape/Private/LandscapeEdit.cpp
--- a/Engine/Source/Runtime/Landscape/Private/LandscapeEdit.cpp
+++ b/Engine/Source/Runtime/Landscape/Private/LandscapeEdit.cpp
@@ -654,7 +654,7 @@ void ULandscapeComponent::FixupWeightmaps()
 				if (TempUsage == nullptr)
 				{
 					TempUsage = &Proxy->WeightmapUsageMap.Add(WeightmapTexture, NewObject<ULandscapeWeightmapUsage>(GetLandscapeProxy()));
-					(*TempUsage)->ProceduralLayerGuid.Invalidate();
+					(*TempUsage)->LayerGuid.Invalidate();
 				}
 
 				ULandscapeWeightmapUsage* Usage = *TempUsage;
@@ -2817,9 +2817,9 @@ LANDSCAPE_API void ALandscapeProxy::Import(
 	// Must be after the FMaterialUpdateContext is destroyed
 	RecreateRenderStateContexts.Reset();
 
-	if (GetMutableDefault<UEditorExperimentalSettings>()->bProceduralLandscape)
+	if (GetMutableDefault<UEditorExperimentalSettings>()->bLandscapeLayerSystem)
 	{
-		SetupProceduralLayers(NumComponentsX, NumComponentsY);
+		SetupLayers(NumComponentsX, NumComponentsY);
 	}
 
 	if (GetLevel()->bIsVisible)
@@ -4657,9 +4657,9 @@ void ULandscapeComponent::ReallocateWeightmaps(FLandscapeEditDataInterface* Data
 			{
 				ULandscapeWeightmapUsage* TryWeightmapUsage = ItPair.Value;
 				//
-				FGuid ProceduralLayerGuidToSeek = InCanUseCurrentEditingWeightmap ? CurrentProceduralLayerGuid : FGuid();
+				FGuid LayerGuidToSeek = InCanUseCurrentEditingWeightmap ? CurrentLayerGuid : FGuid();
 
-				if (TryWeightmapUsage->FreeChannelCount() >= TotalNeededChannels && TryWeightmapUsage->ProceduralLayerGuid == ProceduralLayerGuidToSeek)
+				if (TryWeightmapUsage->FreeChannelCount() >= TotalNeededChannels && TryWeightmapUsage->LayerGuid == LayerGuidToSeek)
 				{
 					if (TryWeightmapUsage->FreeChannelCount() == 4)
 					{
@@ -4720,7 +4720,7 @@ void ULandscapeComponent::ReallocateWeightmaps(FLandscapeEditDataInterface* Data
 
 			// Store it in the usage map
 			CurrentWeightmapUsage = Proxy->WeightmapUsageMap.Add(CurrentWeightmapTexture, NewObject<ULandscapeWeightmapUsage>(GetLandscapeProxy()));
-			CurrentWeightmapUsage->ProceduralLayerGuid = InCanUseCurrentEditingWeightmap ? CurrentProceduralLayerGuid : FGuid();
+			CurrentWeightmapUsage->LayerGuid = InCanUseCurrentEditingWeightmap ? CurrentLayerGuid : FGuid();
 			// UE_LOG(LogLandscape, Log, TEXT("Making a new texture %s"), *CurrentWeightmapTexture->GetName());
 		}
 
@@ -5050,7 +5050,7 @@ void ULandscapeComponent::InitWeightmapData(TArray<ULandscapeLayerInfoObject*>&
 	LODIndexToMaterialIndex.Empty(1);
 	LODIndexToMaterialIndex.Add(0);
 
-	//  TODO: need to update procedural?
+	//  TODO: need to update layer system?
 }
 
 #define MAX_LANDSCAPE_EXPORT_COMPONENTS_NUM		16
diff --git a/Engine/Source/Runtime/Landscape/Private/LandscapeEditInterface.cpp b/Engine/Source/Runtime/Landscape/Private/LandscapeEditInterface.cpp
--- a/Engine/Source/Runtime/Landscape/Private/LandscapeEditInterface.cpp
+++ b/Engine/Source/Runtime/Landscape/Private/LandscapeEditInterface.cpp
@@ -1769,9 +1769,9 @@ void ULandscapeComponent::FillLayer(ULandscapeLayerInfoObject* LayerInfo, FLands
 		ComponentWeightmapLayerAllocations.Add(FWeightmapLayerAllocationInfo(LayerInfo));
 		Component->ReallocateWeightmaps(&LandscapeEdit);
 
-		if (GetMutableDefault<UEditorExperimentalSettings>()->bProceduralLandscape)
+		if (GetMutableDefault<UEditorExperimentalSettings>()->bLandscapeLayerSystem)
 		{
-			GetLandscapeActor()->RequestProceduralContentUpdate(EProceduralContentUpdateFlag::Weightmap_Render);
+			GetLandscapeActor()->RequestLayersContentUpdate(ELandscapeLayersContentUpdateFlag::Weightmap_Render);
 		}
 	}
 
