Changed the name from Landscape Procedural to Landscape Layers

diff --git a/Engine/Source/Runtime/Landscape/Private/LandscapeEditLayers.cpp b/Engine/Source/Runtime/Landscape/Private/LandscapeEditLayers.cpp
new file mode 100644
--- /dev/null
+++ b/Engine/Source/Runtime/Landscape/Private/LandscapeEditLayers.cpp
@@ -0,0 +1,4905 @@
+// Copyright 1998-2019 Epic Games, Inc. All Rights Reserved.
+
+/*=============================================================================
+LandscapeEditLayers.cpp: Landscape editing layers mode
+=============================================================================*/
+
+#include "LandscapeEdit.h"
+#include "Landscape.h"
+#include "LandscapeProxy.h"
+#include "LandscapeStreamingProxy.h"
+#include "LandscapeInfo.h"
+#include "LandscapeComponent.h"
+#include "LandscapeLayerInfoObject.h"
+#include "LandscapeDataAccess.h"
+#include "LandscapePrivate.h"
+
+#include "Shader.h"
+#include "GlobalShader.h"
+#include "ShaderParameterUtils.h"
+#include "Engine/TextureRenderTarget2D.h"
+
+#if WITH_EDITOR
+#include "LandscapeEditorModule.h"
+#include "ComponentRecreateRenderStateContext.h"
+#include "Settings/EditorExperimentalSettings.h"
+#include "LandscapeBPCustomBrush.h"
+#include "Materials/MaterialInstanceConstant.h"
+#include "LandscapeMaterialInstanceConstant.h"
+#include "Materials/MaterialExpressionLandscapeVisibilityMask.h"
+#include "ShaderCompiler.h"
+#include "Algo/Count.h"
+#endif
+#include "LandscapeToolInterface.h"
+
+#define LOCTEXT_NAMESPACE "Landscape"
+
+ENGINE_API extern bool GDisableAutomaticTextureMaterialUpdateDependencies;
+
+static TAutoConsoleVariable<int32> CVarOutputLayersDebugDrawCallName(
+	TEXT("landscape.OutputLayersDebugDrawCallName"),
+	0,
+	TEXT("This will output the name of each draw call for Scope Draw call event. This will allow readable draw call info through RenderDoc, for example."));
+
+static TAutoConsoleVariable<int32> CVarOutputLayersRTContent(
+	TEXT("landscape.OutputLayersRTContent"),
+	0,
+	TEXT("This will output the content of render target. This is used for debugging only."));
+
+static TAutoConsoleVariable<int32> CVarOutputLayersWeightmapsRTContent(
+	TEXT("landscape.OutputLayersWeightmapsRTContent"),
+	0,
+	TEXT("This will output the content of render target used for weightmap. This is used for debugging only."));
+
+DECLARE_GPU_STAT_NAMED(LandscapeLayersRender, TEXT("Landscape Layer System Render"));
+DECLARE_GPU_STAT_NAMED(LandscapeLayersCopy, TEXT("Landscape Layer System Copy"));
+
+// Vertex format and vertex buffer
+
+struct FLandscapeLayersVertex
+{
+	FVector2D Position;
+	FVector2D UV;
+};
+
+struct FLandscapeLayersTriangle
+{
+	FLandscapeLayersVertex V0;
+	FLandscapeLayersVertex V1;
+	FLandscapeLayersVertex V2;
+};
+
+class FLandscapeLayersVertexDeclaration : public FRenderResource
+{
+public:
+	FVertexDeclarationRHIRef VertexDeclarationRHI;
+
+	/** Destructor. */
+	virtual ~FLandscapeLayersVertexDeclaration() {}
+
+	virtual void InitRHI()
+	{
+		FVertexDeclarationElementList Elements;
+		uint32 Stride = sizeof(FLandscapeLayersVertex);
+		Elements.Add(FVertexElement(0, STRUCT_OFFSET(FLandscapeLayersVertex, Position), VET_Float2, 0, Stride));
+		Elements.Add(FVertexElement(0, STRUCT_OFFSET(FLandscapeLayersVertex, UV), VET_Float2, 1, Stride));
+		VertexDeclarationRHI = PipelineStateCache::GetOrCreateVertexDeclaration(Elements);
+	}
+
+	virtual void ReleaseRHI()
+	{
+		VertexDeclarationRHI.SafeRelease();
+	}
+};
+
+class FLandscapeLayersVertexBuffer : public FVertexBuffer
+{
+public:
+	void Init(const TArray<FLandscapeLayersTriangle>& InTriangleList)
+	{
+		TriangleList = InTriangleList;
+	}
+
+private:
+
+	/** Initialize the RHI for this rendering resource */
+	void InitRHI() override
+	{
+		TResourceArray<FLandscapeLayersVertex, VERTEXBUFFER_ALIGNMENT> Vertices;
+		Vertices.SetNumUninitialized(TriangleList.Num() * 3);
+
+		for (int32 i = 0; i < TriangleList.Num(); ++i)
+		{
+			Vertices[i * 3 + 0] = TriangleList[i].V0;
+			Vertices[i * 3 + 1] = TriangleList[i].V1;
+			Vertices[i * 3 + 2] = TriangleList[i].V2;
+		}
+		
+		// Create vertex buffer. Fill buffer with initial data upon creation
+		FRHIResourceCreateInfo CreateInfo(&Vertices);
+		VertexBufferRHI = RHICreateVertexBuffer(Vertices.GetResourceDataSize(), BUF_Static, CreateInfo);
+	}
+
+	TArray<FLandscapeLayersTriangle> TriangleList;
+};
+
+// Custom Pixel and Vertex shaders
+
+class FLandscapeLayersVS : public FGlobalShader
+{
+	DECLARE_GLOBAL_SHADER(FLandscapeLayersVS)
+
+public:
+
+	static bool ShouldCompilePermutation(const FGlobalShaderPermutationParameters& Parameters)
+	{
+		return IsFeatureLevelSupported(Parameters.Platform, ERHIFeatureLevel::SM5) && !IsConsolePlatform(Parameters.Platform) && !IsMetalPlatform(Parameters.Platform);
+	}
+
+	static void ModifyCompilationEnvironment(const FGlobalShaderPermutationParameters& Parameters, FShaderCompilerEnvironment& OutEnvironment)
+	{
+	}
+
+	FLandscapeLayersVS(const ShaderMetaType::CompiledShaderInitializerType& Initializer)
+		: FGlobalShader(Initializer)
+	{
+		TransformParam.Bind(Initializer.ParameterMap, TEXT("Transform"), SPF_Mandatory);
+	}
+
+	FLandscapeLayersVS()
+	{}
+
+	void SetParameters(FRHICommandList& RHICmdList, const FMatrix& InProjectionMatrix)
+	{
+		SetShaderValue(RHICmdList, GetVertexShader(), TransformParam, InProjectionMatrix);
+	}
+
+	virtual bool Serialize(FArchive& Ar) override
+	{
+		bool bShaderHasOutdatedParameters = FShader::Serialize(Ar);
+		Ar << TransformParam;
+		return bShaderHasOutdatedParameters;
+	}
+
+private:
+	FShaderParameter TransformParam;
+};
+
+IMPLEMENT_GLOBAL_SHADER(FLandscapeLayersVS, "/Engine/Private/LandscapeLayersVS.usf", "VSMain", SF_Vertex);
+
+struct FLandscapeLayersHeightmapShaderParameters
+{
+	FLandscapeLayersHeightmapShaderParameters()
+		: ReadHeightmap1(nullptr)
+		, ReadHeightmap2(nullptr)
+		, HeightmapSize(0, 0)
+		, ApplyLayerModifiers(false)
+		, LayerAlpha(1.0f)
+		, LayerVisible(true)
+		, OutputAsDelta(false)
+		, GenerateNormals(false)
+		, GridSize(0.0f, 0.0f, 0.0f)
+		, CurrentMipSize(0, 0)
+		, ParentMipSize(0, 0)
+		, CurrentMipComponentVertexCount(0)
+	{}
+
+	UTexture* ReadHeightmap1;
+	UTexture* ReadHeightmap2;
+	FIntPoint HeightmapSize;
+	bool ApplyLayerModifiers;
+	float LayerAlpha;
+	bool LayerVisible;
+	bool OutputAsDelta;
+	bool GenerateNormals;
+	FVector GridSize;
+	FIntPoint CurrentMipSize;
+	FIntPoint ParentMipSize;
+	int32 CurrentMipComponentVertexCount;
+};
+
+class FLandscapeLayersHeightmapPS : public FGlobalShader
+{
+	DECLARE_GLOBAL_SHADER(FLandscapeLayersHeightmapPS);
+public:
+
+	static bool ShouldCompilePermutation(const FGlobalShaderPermutationParameters& Parameters)
+	{
+		return IsFeatureLevelSupported(Parameters.Platform, ERHIFeatureLevel::SM5) && !IsConsolePlatform(Parameters.Platform) && !IsMetalPlatform(Parameters.Platform);
+	}
+
+	static void ModifyCompilationEnvironment(const FGlobalShaderPermutationParameters& Parameters, FShaderCompilerEnvironment& OutEnvironment)
+	{
+	}
+
+	FLandscapeLayersHeightmapPS(const ShaderMetaType::CompiledShaderInitializerType& Initializer)
+		: FGlobalShader(Initializer)
+	{
+		ReadTexture1Param.Bind(Initializer.ParameterMap, TEXT("ReadTexture1"));
+		ReadTexture2Param.Bind(Initializer.ParameterMap, TEXT("ReadTexture2"));
+		ReadTexture1SamplerParam.Bind(Initializer.ParameterMap, TEXT("ReadTexture1Sampler"));
+		ReadTexture2SamplerParam.Bind(Initializer.ParameterMap, TEXT("ReadTexture2Sampler"));
+
+		LayerInfoParam.Bind(Initializer.ParameterMap, TEXT("LayerInfo"));
+		OutputConfigParam.Bind(Initializer.ParameterMap, TEXT("OutputConfig"));
+		TextureSizeParam.Bind(Initializer.ParameterMap, TEXT("TextureSize"));
+		LandscapeGridScaleParam.Bind(Initializer.ParameterMap, TEXT("LandscapeGridScale"));
+		ComponentVertexCountParam.Bind(Initializer.ParameterMap, TEXT("CurrentMipComponentVertexCount"));
+	}
+
+	FLandscapeLayersHeightmapPS()
+	{}
+
+	void SetParameters(FRHICommandList& RHICmdList, const FLandscapeLayersHeightmapShaderParameters& InParams)
+	{
+		SetTextureParameter(RHICmdList, GetPixelShader(), ReadTexture1Param, ReadTexture1SamplerParam, TStaticSamplerState<SF_Point, AM_Clamp, AM_Clamp, AM_Clamp>::GetRHI(), InParams.ReadHeightmap1->Resource->TextureRHI);
+		SetTextureParameter(RHICmdList, GetPixelShader(), ReadTexture2Param, ReadTexture2SamplerParam, TStaticSamplerState<SF_Point, AM_Clamp, AM_Clamp, AM_Clamp>::GetRHI(), InParams.ReadHeightmap2 != nullptr ? InParams.ReadHeightmap2->Resource->TextureRHI : GWhiteTexture->TextureRHI);
+
+		FVector2D LayerInfo(InParams.LayerAlpha, InParams.LayerVisible ? 1.0f : 0.0f);
+		FVector4 OutputConfig(InParams.ApplyLayerModifiers ? 1.0f : 0.0f, InParams.OutputAsDelta ? 1.0f : 0.0f, InParams.ReadHeightmap2 ? 1.0f : 0.0f, InParams.GenerateNormals ? 1.0f : 0.0f);
+		FVector2D TextureSize(InParams.HeightmapSize.X, InParams.HeightmapSize.Y);
+
+		SetShaderValue(RHICmdList, GetPixelShader(), LayerInfoParam, LayerInfo);
+		SetShaderValue(RHICmdList, GetPixelShader(), OutputConfigParam, OutputConfig);
+		SetShaderValue(RHICmdList, GetPixelShader(), TextureSizeParam, TextureSize);
+		SetShaderValue(RHICmdList, GetPixelShader(), LandscapeGridScaleParam, InParams.GridSize);
+		SetShaderValue(RHICmdList, GetPixelShader(), ComponentVertexCountParam, (float)InParams.CurrentMipComponentVertexCount);
+	}
+
+	virtual bool Serialize(FArchive& Ar) override
+	{
+		bool bShaderHasOutdatedParameters = FShader::Serialize(Ar);
+		Ar << ReadTexture1Param;
+		Ar << ReadTexture2Param;
+		Ar << ReadTexture1SamplerParam;
+		Ar << ReadTexture2SamplerParam;
+		Ar << LayerInfoParam;
+		Ar << OutputConfigParam;
+		Ar << TextureSizeParam;
+		Ar << LandscapeGridScaleParam;
+		Ar << ComponentVertexCountParam;
+		return bShaderHasOutdatedParameters;
+	}
+
+private:
+	FShaderResourceParameter ReadTexture1Param;
+	FShaderResourceParameter ReadTexture2Param;
+	FShaderResourceParameter ReadTexture1SamplerParam;
+	FShaderResourceParameter ReadTexture2SamplerParam;
+	FShaderParameter LayerInfoParam;
+	FShaderParameter OutputConfigParam;
+	FShaderParameter TextureSizeParam;
+	FShaderParameter LandscapeGridScaleParam;
+	FShaderParameter ComponentVertexCountParam;
+};
+
+IMPLEMENT_GLOBAL_SHADER(FLandscapeLayersHeightmapPS, "/Engine/Private/LandscapeLayersPS.usf", "PSHeightmapMain", SF_Pixel);
+
+class FLandscapeLayersHeightmapMipsPS : public FGlobalShader
+{
+	DECLARE_GLOBAL_SHADER(FLandscapeLayersHeightmapMipsPS);
+public:
+
+	static bool ShouldCompilePermutation(const FGlobalShaderPermutationParameters& Parameters)
+	{
+		return IsFeatureLevelSupported(Parameters.Platform, ERHIFeatureLevel::SM5) && !IsConsolePlatform(Parameters.Platform) && !IsMetalPlatform(Parameters.Platform);
+	}
+
+	static void ModifyCompilationEnvironment(const FGlobalShaderPermutationParameters& Parameters, FShaderCompilerEnvironment& OutEnvironment)
+	{
+	}
+
+	FLandscapeLayersHeightmapMipsPS(const ShaderMetaType::CompiledShaderInitializerType& Initializer)
+		: FGlobalShader(Initializer)
+	{
+		ReadTexture1Param.Bind(Initializer.ParameterMap, TEXT("ReadTexture1"));
+		ReadTexture1SamplerParam.Bind(Initializer.ParameterMap, TEXT("ReadTexture1Sampler"));
+		CurrentMipSizeParam.Bind(Initializer.ParameterMap, TEXT("CurrentMipTextureSize"));
+		ParentMipSizeParam.Bind(Initializer.ParameterMap, TEXT("ParentMipTextureSize"));
+		CurrentMipComponentVertexCountParam.Bind(Initializer.ParameterMap, TEXT("CurrentMipComponentVertexCount"));
+	}
+
+	FLandscapeLayersHeightmapMipsPS()
+	{}
+
+	void SetParameters(FRHICommandList& RHICmdList, const FLandscapeLayersHeightmapShaderParameters& InParams)
+	{
+		SetTextureParameter(RHICmdList, GetPixelShader(), ReadTexture1Param, ReadTexture1SamplerParam, TStaticSamplerState<SF_Point, AM_Clamp, AM_Clamp, AM_Clamp>::GetRHI(), InParams.ReadHeightmap1->Resource->TextureRHI);
+
+		SetShaderValue(RHICmdList, GetPixelShader(), CurrentMipSizeParam, FVector2D(InParams.CurrentMipSize.X, InParams.CurrentMipSize.Y));
+		SetShaderValue(RHICmdList, GetPixelShader(), ParentMipSizeParam, FVector2D(InParams.ParentMipSize.X, InParams.ParentMipSize.Y));
+		SetShaderValue(RHICmdList, GetPixelShader(), CurrentMipComponentVertexCountParam, (float)InParams.CurrentMipComponentVertexCount);
+	}
+
+	virtual bool Serialize(FArchive& Ar) override
+	{
+		bool bShaderHasOutdatedParameters = FShader::Serialize(Ar);
+		Ar << ReadTexture1Param;
+		Ar << ReadTexture1SamplerParam;
+		Ar << CurrentMipSizeParam;
+		Ar << ParentMipSizeParam;
+		Ar << CurrentMipComponentVertexCountParam;
+
+		return bShaderHasOutdatedParameters;
+	}
+
+private:
+	FShaderResourceParameter ReadTexture1Param;
+	FShaderResourceParameter ReadTexture1SamplerParam;
+	FShaderParameter CurrentMipSizeParam;
+	FShaderParameter ParentMipSizeParam;
+	FShaderParameter CurrentMipComponentVertexCountParam;
+};
+
+IMPLEMENT_GLOBAL_SHADER(FLandscapeLayersHeightmapMipsPS, "/Engine/Private/LandscapeLayersPS.usf", "PSHeightmapMainMips", SF_Pixel);
+
+struct FLandscapeLayersWeightmapShaderParameters
+{
+	FLandscapeLayersWeightmapShaderParameters()
+		: ReadWeightmap1(nullptr)
+		, ReadWeightmap2(nullptr)
+		, ApplyLayerModifiers(false)
+		, LayerAlpha(1.0f)
+		, LayerVisible(true)
+		, OutputAsDelta(false)
+		, OutputAsSubstractive(false)
+		, OutputAsNormalized(false)
+		, CurrentMipSize(0, 0)
+		, ParentMipSize(0, 0)
+		, CurrentMipComponentVertexCount(0)
+	{}
+
+	UTexture* ReadWeightmap1;
+	UTexture* ReadWeightmap2;
+	bool ApplyLayerModifiers;
+	float LayerAlpha;
+	bool LayerVisible;
+	bool OutputAsDelta;
+	bool OutputAsSubstractive;
+	bool OutputAsNormalized;
+	FIntPoint CurrentMipSize;
+	FIntPoint ParentMipSize;
+	int32 CurrentMipComponentVertexCount;
+};
+
+class FLandscapeLayersWeightmapPS : public FGlobalShader
+{
+	DECLARE_GLOBAL_SHADER(FLandscapeLayersWeightmapPS);
+public:
+
+	static bool ShouldCompilePermutation(const FGlobalShaderPermutationParameters& Parameters)
+	{
+		return IsFeatureLevelSupported(Parameters.Platform, ERHIFeatureLevel::SM5) && !IsConsolePlatform(Parameters.Platform) && !IsMetalPlatform(Parameters.Platform);
+	}
+
+	static void ModifyCompilationEnvironment(const FGlobalShaderPermutationParameters& Parameters, FShaderCompilerEnvironment& OutEnvironment)
+	{
+	}
+
+	FLandscapeLayersWeightmapPS(const ShaderMetaType::CompiledShaderInitializerType& Initializer)
+		: FGlobalShader(Initializer)
+	{
+		ReadTexture1Param.Bind(Initializer.ParameterMap, TEXT("ReadTexture1"));
+		ReadTexture2Param.Bind(Initializer.ParameterMap, TEXT("ReadTexture2"));
+		ReadTexture1SamplerParam.Bind(Initializer.ParameterMap, TEXT("ReadTexture1Sampler"));
+		ReadTexture2SamplerParam.Bind(Initializer.ParameterMap, TEXT("ReadTexture2Sampler"));
+		LayerInfoParam.Bind(Initializer.ParameterMap, TEXT("LayerInfo"));
+		OutputConfigParam.Bind(Initializer.ParameterMap, TEXT("OutputConfig"));
+		ComponentVertexCountParam.Bind(Initializer.ParameterMap, TEXT("CurrentMipComponentVertexCount"));
+	}
+
+	FLandscapeLayersWeightmapPS()
+	{}
+
+	void SetParameters(FRHICommandList& RHICmdList, const FLandscapeLayersWeightmapShaderParameters& InParams)
+	{
+		SetTextureParameter(RHICmdList, GetPixelShader(), ReadTexture1Param, ReadTexture1SamplerParam, TStaticSamplerState<SF_Point, AM_Clamp, AM_Clamp, AM_Clamp>::GetRHI(), InParams.ReadWeightmap1->Resource->TextureRHI);
+		SetTextureParameter(RHICmdList, GetPixelShader(), ReadTexture2Param, ReadTexture2SamplerParam, TStaticSamplerState<SF_Point, AM_Clamp, AM_Clamp, AM_Clamp>::GetRHI(), InParams.ReadWeightmap2 != nullptr ? InParams.ReadWeightmap2->Resource->TextureRHI : GWhiteTexture->TextureRHI);
+
+		FVector2D LayerInfo(InParams.LayerAlpha, InParams.LayerVisible ? 1.0f : 0.0f);
+		FVector4 OutputConfig(InParams.ApplyLayerModifiers ? 1.0f : 0.0f, InParams.OutputAsSubstractive ? 1.0f : 0.0f, InParams.ReadWeightmap2 != nullptr ? 1.0f : 0.0f, InParams.OutputAsNormalized ? 1.0f : 0.0f);
+
+		SetShaderValue(RHICmdList, GetPixelShader(), LayerInfoParam, LayerInfo);
+		SetShaderValue(RHICmdList, GetPixelShader(), OutputConfigParam, OutputConfig);
+		SetShaderValue(RHICmdList, GetPixelShader(), ComponentVertexCountParam, (float)InParams.CurrentMipComponentVertexCount);
+	}
+
+	virtual bool Serialize(FArchive& Ar) override
+	{
+		bool bShaderHasOutdatedParameters = FShader::Serialize(Ar);
+		Ar << ReadTexture1Param;
+		Ar << ReadTexture2Param;
+		Ar << ReadTexture1SamplerParam;
+		Ar << ReadTexture2SamplerParam;
+		Ar << LayerInfoParam;
+		Ar << OutputConfigParam;
+		Ar << ComponentVertexCountParam;
+		return bShaderHasOutdatedParameters;
+	}
+
+private:
+	FShaderResourceParameter ReadTexture1Param;
+	FShaderResourceParameter ReadTexture2Param;
+	FShaderResourceParameter ReadTexture1SamplerParam;
+	FShaderResourceParameter ReadTexture2SamplerParam;
+	FShaderParameter LayerInfoParam;
+	FShaderParameter OutputConfigParam;
+	FShaderParameter ComponentVertexCountParam;
+};
+
+IMPLEMENT_GLOBAL_SHADER(FLandscapeLayersWeightmapPS, "/Engine/Private/LandscapeLayersPS.usf", "PSWeightmapMain", SF_Pixel);
+
+class FLandscapeLayersWeightmapMipsPS : public FGlobalShader
+{
+	DECLARE_GLOBAL_SHADER(FLandscapeLayersWeightmapMipsPS);
+public:
+
+	static bool ShouldCompilePermutation(const FGlobalShaderPermutationParameters& Parameters)
+	{
+		return IsFeatureLevelSupported(Parameters.Platform, ERHIFeatureLevel::SM5) && !IsConsolePlatform(Parameters.Platform) && !IsMetalPlatform(Parameters.Platform);
+	}
+
+	static void ModifyCompilationEnvironment(const FGlobalShaderPermutationParameters& Parameters, FShaderCompilerEnvironment& OutEnvironment)
+	{
+	}
+
+	FLandscapeLayersWeightmapMipsPS(const ShaderMetaType::CompiledShaderInitializerType& Initializer)
+		: FGlobalShader(Initializer)
+	{
+		ReadTexture1Param.Bind(Initializer.ParameterMap, TEXT("ReadTexture1"));
+		ReadTexture1SamplerParam.Bind(Initializer.ParameterMap, TEXT("ReadTexture1Sampler"));
+		CurrentMipSizeParam.Bind(Initializer.ParameterMap, TEXT("CurrentMipTextureSize"));
+		ParentMipSizeParam.Bind(Initializer.ParameterMap, TEXT("ParentMipTextureSize"));
+		CurrentMipComponentVertexCountParam.Bind(Initializer.ParameterMap, TEXT("CurrentMipComponentVertexCount"));
+	}
+
+	FLandscapeLayersWeightmapMipsPS()
+	{}
+
+	void SetParameters(FRHICommandList& RHICmdList, const FLandscapeLayersWeightmapShaderParameters& InParams)
+	{
+		SetTextureParameter(RHICmdList, GetPixelShader(), ReadTexture1Param, ReadTexture1SamplerParam, TStaticSamplerState<SF_Point, AM_Clamp, AM_Clamp, AM_Clamp>::GetRHI(), InParams.ReadWeightmap1->Resource->TextureRHI);
+
+		SetShaderValue(RHICmdList, GetPixelShader(), CurrentMipSizeParam, FVector2D(InParams.CurrentMipSize.X, InParams.CurrentMipSize.Y));
+		SetShaderValue(RHICmdList, GetPixelShader(), ParentMipSizeParam, FVector2D(InParams.ParentMipSize.X, InParams.ParentMipSize.Y));
+		SetShaderValue(RHICmdList, GetPixelShader(), CurrentMipComponentVertexCountParam, (float)InParams.CurrentMipComponentVertexCount);
+	}
+
+	virtual bool Serialize(FArchive& Ar) override
+	{
+		bool bShaderHasOutdatedParameters = FShader::Serialize(Ar);
+		Ar << ReadTexture1Param;
+		Ar << ReadTexture1SamplerParam;
+		Ar << CurrentMipSizeParam;
+		Ar << ParentMipSizeParam;
+		Ar << CurrentMipComponentVertexCountParam;
+
+		return bShaderHasOutdatedParameters;
+	}
+
+private:
+	FShaderResourceParameter ReadTexture1Param;
+	FShaderResourceParameter ReadTexture1SamplerParam;
+	FShaderParameter CurrentMipSizeParam;
+	FShaderParameter ParentMipSizeParam;
+	FShaderParameter CurrentMipComponentVertexCountParam;
+};
+
+IMPLEMENT_GLOBAL_SHADER(FLandscapeLayersWeightmapMipsPS, "/Engine/Private/LandscapeLayersPS.usf", "PSWeightmapMainMips", SF_Pixel);
+
+// Custom Resources
+
+class FLandscapeTexture2DResource : public FTextureResource
+{
+public:
+	FLandscapeTexture2DResource(uint32 InSizeX, uint32 InSizeY, EPixelFormat InFormat, uint32 InNumMips, bool InNeedUAV)
+		: SizeX(InSizeX)
+		, SizeY(InSizeY)
+		, Format(InFormat)
+		, NumMips(InNumMips)
+		, CreateUAV(InNeedUAV)
+	{}
+
+	virtual uint32 GetSizeX() const override
+	{
+		return SizeX;
+	}
+
+	virtual uint32 GetSizeY() const override
+	{
+		return SizeY;
+	}
+
+	/** Called when the resource is initialized. This is only called by the rendering thread. */
+	virtual void InitRHI() override
+	{
+		FTextureResource::InitRHI();
+
+		FRHIResourceCreateInfo CreateInfo;
+		uint32 Flags = TexCreate_NoTiling | TexCreate_OfflineProcessed;
+
+		if (CreateUAV)
+		{
+			Flags |= TexCreate_UAV;
+		}
+
+		TextureRHI = RHICreateTexture2D(SizeX, SizeY, Format, NumMips, 1, Flags, CreateInfo);
+
+		if (CreateUAV)
+		{
+			TextureUAV = RHICreateUnorderedAccessView(TextureRHI, 0);
+		}
+	}
+
+	FUnorderedAccessViewRHIRef TextureUAV;
+
+private:
+	uint32 SizeX;
+	uint32 SizeY;
+	EPixelFormat Format;
+	uint32 NumMips;
+	bool CreateUAV;
+};
+
+class FLandscapeTexture2DArrayResource : public FTextureResource
+{
+public:
+	FLandscapeTexture2DArrayResource(uint32 InSizeX, uint32 InSizeY, uint32 InSizeZ, EPixelFormat InFormat, uint32 InNumMips, bool InNeedUAV)
+		: SizeX(InSizeX)
+		, SizeY(InSizeY)
+		, SizeZ(InSizeZ)
+		, Format(InFormat)
+		, NumMips(InNumMips)
+		, CreateUAV(InNeedUAV)
+	{}
+
+	virtual uint32 GetSizeX() const override
+	{
+		return SizeX;
+	}
+
+	virtual uint32 GetSizeY() const override
+	{
+		return SizeY;
+	}
+
+	virtual uint32 GetSizeZ() const
+	{
+		return SizeZ;
+	}
+
+	/** Called when the resource is initialized. This is only called by the rendering thread. */
+	virtual void InitRHI() override
+	{
+		FTextureResource::InitRHI();
+
+		FRHIResourceCreateInfo CreateInfo;
+		uint32 Flags = TexCreate_NoTiling | TexCreate_OfflineProcessed;
+
+		if (CreateUAV)
+		{
+			Flags |= TexCreate_UAV;
+		}
+
+		TextureRHI = RHICreateTexture2DArray(SizeX, SizeY, SizeZ, Format, NumMips, Flags, CreateInfo);
+
+		if (CreateUAV)
+		{
+			TextureUAV = RHICreateUnorderedAccessView(TextureRHI, 0);
+		}
+	}
+
+	virtual void ReleaseRHI() override
+	{
+		FTextureResource::ReleaseRHI();
+
+		TextureUAV.SafeRelease();
+	}
+
+	FUnorderedAccessViewRHIRef TextureUAV;
+
+private:
+	uint32 SizeX;
+	uint32 SizeY;
+	uint32 SizeZ;
+	EPixelFormat Format;
+	uint32 NumMips;
+	bool CreateUAV;
+};
+
+// Compute shaders data
+
+int32 GLandscapeLayerWeightmapThreadGroupSizeX = 16;
+int32 GLandscapeLayerWeightmapThreadGroupSizeY = 16;
+
+struct FLandscapeLayerWeightmapExtractMaterialLayersComponentData
+{
+	FIntPoint ComponentVertexPosition;	// Section base converted to vertex instead of quad
+	uint32 DestinationPaintLayerIndex;	// correspond to which layer info object index the data should be stored in the texture 2d array
+	uint32 WeightmapChannelToProcess;	// correspond to which RGBA channel to process
+	FIntPoint AtlasTexturePositionOutput;	// This represent the location we will write layer information
+};
+
+class FLandscapeLayerWeightmapExtractMaterialLayersComputeShaderResource : public FRenderResource
+{
+public:
+	FLandscapeLayerWeightmapExtractMaterialLayersComputeShaderResource(const TArray<FLandscapeLayerWeightmapExtractMaterialLayersComponentData>& InComponentsData)
+		: OriginalComponentsData(InComponentsData)
+		, ComponentsDataCount(OriginalComponentsData.Num())
+	{}
+
+	~FLandscapeLayerWeightmapExtractMaterialLayersComputeShaderResource()
+	{
+		ComponentsData.SafeRelease();
+		ComponentsDataSRV.SafeRelease();
+	}
+
+	/** Called when the resource is initialized. This is only called by the rendering thread. */
+	virtual void InitDynamicRHI() override
+	{
+		FRHIResourceCreateInfo CreateInfo;
+		ComponentsData = RHICreateStructuredBuffer(sizeof(FLandscapeLayerWeightmapExtractMaterialLayersComponentData), OriginalComponentsData.Num() * sizeof(FLandscapeLayerWeightmapExtractMaterialLayersComponentData), BUF_ShaderResource | BUF_Volatile, CreateInfo);
+		ComponentsDataSRV = RHICreateShaderResourceView(ComponentsData);
+
+		uint8* Buffer = (uint8*)RHILockStructuredBuffer(ComponentsData, 0, OriginalComponentsData.Num() * sizeof(FLandscapeLayerWeightmapExtractMaterialLayersComponentData), RLM_WriteOnly);
+		FMemory::Memcpy(Buffer, OriginalComponentsData.GetData(), OriginalComponentsData.Num() * sizeof(FLandscapeLayerWeightmapExtractMaterialLayersComponentData));
+		RHIUnlockStructuredBuffer(ComponentsData);
+	}
+
+	virtual void ReleaseDynamicRHI() override
+	{
+		ComponentsData.SafeRelease();
+		ComponentsDataSRV.SafeRelease();
+	}
+
+	int32 GetComponentsDataCount() const
+	{
+		return ComponentsDataCount;
+	}
+
+private:
+	friend class FLandscapeLayerWeightmapExtractMaterialLayersCS;
+
+	FStructuredBufferRHIRef ComponentsData;
+	FShaderResourceViewRHIRef ComponentsDataSRV;
+	TArray<FLandscapeLayerWeightmapExtractMaterialLayersComponentData> OriginalComponentsData;
+	int32 ComponentsDataCount;
+};
+
+struct FLandscapeLayerWeightmapExtractMaterialLayersComputeShaderParameters
+{
+	FLandscapeLayerWeightmapExtractMaterialLayersComputeShaderParameters()
+		: ComponentWeightmapResource(nullptr)
+		, ComputeShaderResource(nullptr)
+		, AtlasWeightmapsPerLayer(nullptr)
+		, ComponentSize(0)
+	{}
+
+	FLandscapeTexture2DResource* ComponentWeightmapResource;
+	FLandscapeLayerWeightmapExtractMaterialLayersComputeShaderResource* ComputeShaderResource;
+	FLandscapeTexture2DArrayResource* AtlasWeightmapsPerLayer;
+	uint32 ComponentSize;
+};
+
+class FLandscapeLayerWeightmapExtractMaterialLayersCS : public FGlobalShader
+{
+	DECLARE_GLOBAL_SHADER(FLandscapeLayerWeightmapExtractMaterialLayersCS);
+public:
+
+	static bool ShouldCompilePermutation(const FGlobalShaderPermutationParameters& Parameters)
+	{
+		return IsFeatureLevelSupported(Parameters.Platform, ERHIFeatureLevel::SM5) && !IsConsolePlatform(Parameters.Platform) && !IsMetalPlatform(Parameters.Platform);
+	}
+
+	static void ModifyCompilationEnvironment(const FGlobalShaderPermutationParameters& Parameters, FShaderCompilerEnvironment& OutEnvironment)
+	{
+		OutEnvironment.SetDefine(TEXT("THREADGROUP_SIZEX"), GLandscapeLayerWeightmapThreadGroupSizeX);
+		OutEnvironment.SetDefine(TEXT("THREADGROUP_SIZEY"), GLandscapeLayerWeightmapThreadGroupSizeY);
+	}
+
+	FLandscapeLayerWeightmapExtractMaterialLayersCS(const ShaderMetaType::CompiledShaderInitializerType& Initializer)
+		: FGlobalShader(Initializer)
+	{
+		ComponentWeightmapParam.Bind(Initializer.ParameterMap, TEXT("InComponentWeightMaps"));
+		AtlasPaintListsParam.Bind(Initializer.ParameterMap, TEXT("OutAtlasPaintLayers"));
+		ComponentsDataParam.Bind(Initializer.ParameterMap, TEXT("InExtractLayersComponentsData"));
+		ComponentSizeParam.Bind(Initializer.ParameterMap, TEXT("ComponentSize"));
+	}
+
+	FLandscapeLayerWeightmapExtractMaterialLayersCS()
+	{}
+
+	void SetParameters(FRHICommandList& RHICmdList, const FLandscapeLayerWeightmapExtractMaterialLayersComputeShaderParameters& InParams)
+	{
+		SetTextureParameter(RHICmdList, GetComputeShader(), ComponentWeightmapParam, InParams.ComponentWeightmapResource->TextureRHI);
+		SetUAVParameter(RHICmdList, GetComputeShader(), AtlasPaintListsParam, InParams.AtlasWeightmapsPerLayer->TextureUAV);
+		SetSRVParameter(RHICmdList, GetComputeShader(), ComponentsDataParam, InParams.ComputeShaderResource->ComponentsDataSRV);
+		SetShaderValue(RHICmdList, GetComputeShader(), ComponentSizeParam, InParams.ComponentSize);
+	}
+
+	void UnsetParameters(FRHICommandList& RHICmdList)
+	{
+		SetUAVParameter(RHICmdList, GetComputeShader(), AtlasPaintListsParam, FUnorderedAccessViewRHIParamRef());
+	}
+
+	virtual bool Serialize(FArchive& Ar) override
+	{
+		bool bShaderHasOutdatedParameters = FShader::Serialize(Ar);
+		Ar << ComponentWeightmapParam;
+		Ar << AtlasPaintListsParam;
+		Ar << ComponentsDataParam;
+		Ar << ComponentSizeParam;
+
+		return bShaderHasOutdatedParameters;
+	}
+
+private:
+	FShaderResourceParameter ComponentWeightmapParam;
+	FShaderResourceParameter AtlasPaintListsParam;
+	FShaderResourceParameter ComponentsDataParam;
+	FShaderParameter ComponentSizeParam;
+};
+
+IMPLEMENT_GLOBAL_SHADER(FLandscapeLayerWeightmapExtractMaterialLayersCS, "/Engine/Private/LandscapeLayersCS.usf", "ComputeWeightmapPerPaintLayer", SF_Compute);
+
+class FLandscapeLayerWeightmapExtractMaterialLayersCSDispatch_RenderThread
+{
+public:
+	FLandscapeLayerWeightmapExtractMaterialLayersCSDispatch_RenderThread(const FLandscapeLayerWeightmapExtractMaterialLayersComputeShaderParameters& InShaderParams)
+		: ShaderParams(InShaderParams)
+	{}
+
+	void ExtractLayers(FRHICommandListImmediate& InRHICmdList)
+	{
+		SCOPE_CYCLE_COUNTER(STAT_LandscapeLayersRegenerate_RenderThread);
+		SCOPED_DRAW_EVENTF(InRHICmdList, LandscapeLayersRender, TEXT("ExtractLayers"));
+
+		TShaderMapRef<FLandscapeLayerWeightmapExtractMaterialLayersCS> ComputeShader(GetGlobalShaderMap(GMaxRHIFeatureLevel));
+		InRHICmdList.SetComputeShader(ComputeShader->GetComputeShader());
+		ComputeShader->SetParameters(InRHICmdList, ShaderParams);
+
+		uint32 ThreadGroupCountX = FMath::CeilToInt((float)ShaderParams.ComponentSize / (float)GLandscapeLayerWeightmapThreadGroupSizeX);
+		uint32 ThreadGroupCountY = FMath::CeilToInt((float)ShaderParams.ComponentSize / (float)GLandscapeLayerWeightmapThreadGroupSizeY);
+		check(ThreadGroupCountX > 0 && ThreadGroupCountY > 0);
+
+		DispatchComputeShader(InRHICmdList, *ComputeShader, ThreadGroupCountX, ThreadGroupCountY, ShaderParams.ComputeShaderResource->GetComponentsDataCount());
+		ComputeShader->UnsetParameters(InRHICmdList);
+		ShaderParams.ComputeShaderResource->ReleaseResource();
+		delete ShaderParams.ComputeShaderResource;
+	}
+
+private:
+	FLandscapeLayerWeightmapExtractMaterialLayersComputeShaderParameters ShaderParams;
+};
+
+struct FLandscapeLayerWeightmapPackMaterialLayersComponentData
+{
+	int32 ComponentVertexPositionX[4];		// Section base converted to vertex instead of quad
+	int32 ComponentVertexPositionY[4];		// Section base converted to vertex instead of quad
+	int32 SourcePaintLayerIndex[4];			// correspond to which layer info object index the data should be loaded from the texture 2d array
+	int32 WeightmapChannelToProcess[4];		// correspond to which RGBA channel to process
+};
+
+class FLandscapeLayerWeightmapPackMaterialLayersComputeShaderResource : public FRenderResource
+{
+public:
+	FLandscapeLayerWeightmapPackMaterialLayersComputeShaderResource(const TArray<FLandscapeLayerWeightmapPackMaterialLayersComponentData>& InComponentsData, const TArray<float>& InWeightmapWeightBlendModeData, const TArray<FVector2D>& InTextureOutputOffset)
+		: OriginalComponentsData(InComponentsData)
+		, ComponentsDataCount(OriginalComponentsData.Num())
+		, OriginalWeightmapWeightBlendModeData(InWeightmapWeightBlendModeData)
+		, OriginalTextureOutputOffset(InTextureOutputOffset)
+	{}
+
+	~FLandscapeLayerWeightmapPackMaterialLayersComputeShaderResource()
+	{
+		ComponentsData.SafeRelease();
+		ComponentsDataSRV.SafeRelease();
+		WeightmapWeightBlendModeSRV.SafeRelease();
+		WeightmapTextureOutputOffsetSRV.SafeRelease();
+	}
+
+	/** Called when the resource is initialized. This is only called by the rendering thread. */
+	virtual void InitDynamicRHI() override
+	{
+		FRHIResourceCreateInfo CreateInfo;
+		uint32 ComponentsDataMemSize = OriginalComponentsData.Num() * sizeof(FLandscapeLayerWeightmapPackMaterialLayersComponentData);
+		ComponentsData = RHICreateStructuredBuffer(sizeof(FLandscapeLayerWeightmapPackMaterialLayersComponentData), ComponentsDataMemSize, BUF_ShaderResource | BUF_Volatile, CreateInfo);
+		ComponentsDataSRV = RHICreateShaderResourceView(ComponentsData);
+
+		uint8* Buffer = (uint8*)RHILockStructuredBuffer(ComponentsData, 0, ComponentsDataMemSize, RLM_WriteOnly);
+		FMemory::Memcpy(Buffer, OriginalComponentsData.GetData(), ComponentsDataMemSize);
+		RHIUnlockStructuredBuffer(ComponentsData);
+
+		FRHIResourceCreateInfo WeightBlendCreateInfo;
+		uint32 WeightBlendMemSize = OriginalWeightmapWeightBlendModeData.Num() * sizeof(float);
+		WeightmapWeightBlendMode = RHICreateVertexBuffer(WeightBlendMemSize, BUF_ShaderResource | BUF_Volatile | BUF_Dynamic, WeightBlendCreateInfo);
+		WeightmapWeightBlendModeSRV = RHICreateShaderResourceView(WeightmapWeightBlendMode, sizeof(float), PF_R32_FLOAT);
+
+		void* WeightmapWeightBlendModePtr = RHILockVertexBuffer(WeightmapWeightBlendMode, 0, WeightBlendMemSize, RLM_WriteOnly);
+		FMemory::Memcpy(WeightmapWeightBlendModePtr, OriginalWeightmapWeightBlendModeData.GetData(), WeightBlendMemSize);
+		RHIUnlockVertexBuffer(WeightmapWeightBlendMode);
+
+		FRHIResourceCreateInfo TextureOutputOffsetCreateInfo;
+		uint32 TextureOutputOffsetMemSize = OriginalTextureOutputOffset.Num() * sizeof(FVector2D);
+		WeightmapTextureOutputOffset = RHICreateVertexBuffer(TextureOutputOffsetMemSize, BUF_ShaderResource | BUF_Volatile | BUF_Dynamic, TextureOutputOffsetCreateInfo);
+		WeightmapTextureOutputOffsetSRV = RHICreateShaderResourceView(WeightmapTextureOutputOffset, sizeof(FVector2D), PF_G32R32F);
+
+		void* TextureOutputOffsetPtr = RHILockVertexBuffer(WeightmapTextureOutputOffset, 0, TextureOutputOffsetMemSize, RLM_WriteOnly);
+		FMemory::Memcpy(TextureOutputOffsetPtr, OriginalTextureOutputOffset.GetData(), TextureOutputOffsetMemSize);
+		RHIUnlockVertexBuffer(WeightmapTextureOutputOffset);
+	}
+
+	virtual void ReleaseDynamicRHI() override
+	{
+		ComponentsData.SafeRelease();
+		ComponentsDataSRV.SafeRelease();
+		WeightmapWeightBlendModeSRV.SafeRelease();
+		WeightmapTextureOutputOffsetSRV.SafeRelease();
+	}
+
+	int32 GetComponentsDataCount() const
+	{
+		return ComponentsDataCount;
+	}
+
+private:
+	friend class FLandscapeLayerWeightmapPackMaterialLayersCS;
+
+	FStructuredBufferRHIRef ComponentsData;
+	FShaderResourceViewRHIRef ComponentsDataSRV;
+	TArray<FLandscapeLayerWeightmapPackMaterialLayersComponentData> OriginalComponentsData;
+	int32 ComponentsDataCount;
+
+	TArray<float> OriginalWeightmapWeightBlendModeData;
+	FVertexBufferRHIRef WeightmapWeightBlendMode;
+	FShaderResourceViewRHIRef WeightmapWeightBlendModeSRV;
+
+	TArray<FVector2D> OriginalTextureOutputOffset;
+	FVertexBufferRHIRef WeightmapTextureOutputOffset;
+	FShaderResourceViewRHIRef WeightmapTextureOutputOffsetSRV;
+};
+
+struct FLandscapeLayerWeightmapPackMaterialLayersComputeShaderParameters
+{
+	FLandscapeLayerWeightmapPackMaterialLayersComputeShaderParameters()
+		: ComponentWeightmapResource(nullptr)
+		, ComputeShaderResource(nullptr)
+		, AtlasWeightmapsPerLayer(nullptr)
+		, ComponentSize(0)
+	{}
+
+	FLandscapeTexture2DResource* ComponentWeightmapResource;
+	FLandscapeLayerWeightmapPackMaterialLayersComputeShaderResource* ComputeShaderResource;
+	FLandscapeTexture2DArrayResource* AtlasWeightmapsPerLayer;
+	uint32 ComponentSize;
+};
+
+class FLandscapeLayerWeightmapPackMaterialLayersCS : public FGlobalShader
+{
+	DECLARE_GLOBAL_SHADER(FLandscapeLayerWeightmapPackMaterialLayersCS);
+public:
+
+	static bool ShouldCompilePermutation(const FGlobalShaderPermutationParameters& Parameters)
+	{
+		return IsFeatureLevelSupported(Parameters.Platform, ERHIFeatureLevel::SM5) && !IsConsolePlatform(Parameters.Platform) && !IsMetalPlatform(Parameters.Platform);
+	}
+
+	static void ModifyCompilationEnvironment(const FGlobalShaderPermutationParameters& Parameters, FShaderCompilerEnvironment& OutEnvironment)
+	{
+		OutEnvironment.SetDefine(TEXT("THREADGROUP_SIZEX"), GLandscapeLayerWeightmapThreadGroupSizeX);
+		OutEnvironment.SetDefine(TEXT("THREADGROUP_SIZEY"), GLandscapeLayerWeightmapThreadGroupSizeY);
+	}
+
+	FLandscapeLayerWeightmapPackMaterialLayersCS(const ShaderMetaType::CompiledShaderInitializerType& Initializer)
+		: FGlobalShader(Initializer)
+	{
+		ComponentWeightmapParam.Bind(Initializer.ParameterMap, TEXT("OutComponentWeightMaps"));
+		AtlasPaintListsParam.Bind(Initializer.ParameterMap, TEXT("InAtlasPaintLayers"));
+		ComponentsDataParam.Bind(Initializer.ParameterMap, TEXT("InPackLayersComponentsData"));
+		ComponentSizeParam.Bind(Initializer.ParameterMap, TEXT("ComponentSize"));
+		WeightmapWeightBlendModeParam.Bind(Initializer.ParameterMap, TEXT("InWeightmapWeightBlendMode"));
+		WeightmapTextureOutputOffsetParam.Bind(Initializer.ParameterMap, TEXT("InWeightmapTextureOutputOffset"));
+	}
+
+	FLandscapeLayerWeightmapPackMaterialLayersCS()
+	{}
+
+	void SetParameters(FRHICommandList& RHICmdList, const FLandscapeLayerWeightmapPackMaterialLayersComputeShaderParameters& InParams)
+	{
+		SetUAVParameter(RHICmdList, GetComputeShader(), ComponentWeightmapParam, InParams.ComponentWeightmapResource->TextureUAV);
+		SetTextureParameter(RHICmdList, GetComputeShader(), AtlasPaintListsParam, InParams.AtlasWeightmapsPerLayer->TextureRHI);
+		SetSRVParameter(RHICmdList, GetComputeShader(), ComponentsDataParam, InParams.ComputeShaderResource->ComponentsDataSRV);
+		SetShaderValue(RHICmdList, GetComputeShader(), ComponentSizeParam, InParams.ComponentSize);
+		SetSRVParameter(RHICmdList, GetComputeShader(), WeightmapWeightBlendModeParam, InParams.ComputeShaderResource->WeightmapWeightBlendModeSRV);
+		SetSRVParameter(RHICmdList, GetComputeShader(), WeightmapTextureOutputOffsetParam, InParams.ComputeShaderResource->WeightmapTextureOutputOffsetSRV);
+	}
+
+	void UnsetParameters(FRHICommandList& RHICmdList)
+	{
+		SetUAVParameter(RHICmdList, GetComputeShader(), ComponentWeightmapParam, FUnorderedAccessViewRHIParamRef());
+	}
+
+	virtual bool Serialize(FArchive& Ar) override
+	{
+		bool bShaderHasOutdatedParameters = FShader::Serialize(Ar);
+		Ar << ComponentWeightmapParam;
+		Ar << AtlasPaintListsParam;
+		Ar << ComponentsDataParam;
+		Ar << ComponentSizeParam;
+		Ar << WeightmapWeightBlendModeParam;
+		Ar << WeightmapTextureOutputOffsetParam;
+		return bShaderHasOutdatedParameters;
+	}
+
+private:
+	FShaderResourceParameter ComponentWeightmapParam;
+	FShaderResourceParameter AtlasPaintListsParam;
+	FShaderResourceParameter ComponentsDataParam;
+	FShaderParameter ComponentSizeParam;
+	FShaderResourceParameter WeightmapWeightBlendModeParam;
+	FShaderResourceParameter WeightmapTextureOutputOffsetParam;
+};
+
+IMPLEMENT_GLOBAL_SHADER(FLandscapeLayerWeightmapPackMaterialLayersCS, "/Engine/Private/LandscapeLayersCS.usf", "PackPaintLayerToWeightmap", SF_Compute);
+
+class FLandscapeLayerWeightmapPackMaterialLayersCSDispatch_RenderThread
+{
+public:
+	FLandscapeLayerWeightmapPackMaterialLayersCSDispatch_RenderThread(const FLandscapeLayerWeightmapPackMaterialLayersComputeShaderParameters& InShaderParams)
+		: ShaderParams(InShaderParams)
+	{}
+
+	void PackLayers(FRHICommandListImmediate& InRHICmdList)
+	{
+		SCOPE_CYCLE_COUNTER(STAT_LandscapeLayersRegenerate_RenderThread);
+		SCOPED_DRAW_EVENTF(InRHICmdList, LandscapeLayersRender, TEXT("PackLayers"));
+
+		TShaderMapRef<FLandscapeLayerWeightmapPackMaterialLayersCS> ComputeShader(GetGlobalShaderMap(GMaxRHIFeatureLevel));
+		InRHICmdList.SetComputeShader(ComputeShader->GetComputeShader());
+		ComputeShader->SetParameters(InRHICmdList, ShaderParams);
+
+		uint32 ThreadGroupCountX = FMath::CeilToInt((float)ShaderParams.ComponentSize / (float)GLandscapeLayerWeightmapThreadGroupSizeX);
+		uint32 ThreadGroupCountY = FMath::CeilToInt((float)ShaderParams.ComponentSize / (float)GLandscapeLayerWeightmapThreadGroupSizeY);
+		check(ThreadGroupCountX > 0 && ThreadGroupCountY > 0);
+
+		DispatchComputeShader(InRHICmdList, *ComputeShader, ThreadGroupCountX, ThreadGroupCountY, ShaderParams.ComputeShaderResource->GetComponentsDataCount());
+		ComputeShader->UnsetParameters(InRHICmdList);
+		ShaderParams.ComputeShaderResource->ReleaseResource();
+		delete ShaderParams.ComputeShaderResource;
+	}
+
+private:
+	FLandscapeLayerWeightmapPackMaterialLayersComputeShaderParameters ShaderParams;
+};
+
+// Copy texture render command
+
+class FLandscapeLayersCopyTexture_RenderThread
+{
+public: 
+	FLandscapeLayersCopyTexture_RenderThread(const FString& InSourceResourceDebugName, FTextureResource* InSourceResource, const FString& InDestResourceDebugName, FTextureResource* InDestResource, FTextureResource* InDestCPUResource,
+											     const FIntPoint& InFirstComponentSectionBase, int32 InSubSectionSizeQuad, int32 InNumSubSections, uint8 InSourceCurrentMip, uint8 InDestCurrentMip, uint32 InSourceArrayIndex, uint32 InDestArrayIndex)
+		: SourceResource(InSourceResource)
+		, DestResource(InDestResource)
+		, DestCPUResource(InDestCPUResource)
+		, SourceMip(InSourceCurrentMip)
+		, DestMip(InDestCurrentMip)
+		, SourceArrayIndex(InSourceArrayIndex)
+		, DestArrayIndex(InDestArrayIndex)
+		, ComponentSectionBase(InFirstComponentSectionBase)
+		, SubSectionSizeQuad(InSubSectionSizeQuad)
+		, NumSubSections(InNumSubSections)
+		, SourceDebugName(InSourceResourceDebugName)
+		, DestResourceDebugName(InDestResourceDebugName)
+	{}
+
+	void Copy(FRHICommandListImmediate& InRHICmdList)
+	{
+		SCOPE_CYCLE_COUNTER(STAT_LandscapeLayersRegenerate_RenderThread);
+		SCOPED_DRAW_EVENTF(InRHICmdList, LandscapeLayersCopy, TEXT("LS Copy %s -> %s, Mip (%d -> %d), Array Index (%d -> %d)"), *SourceDebugName, *DestResourceDebugName, SourceMip, DestMip, SourceArrayIndex, DestArrayIndex);
+		SCOPED_GPU_STAT(InRHICmdList, LandscapeLayersCopy);
+
+		FIntPoint SourceSize(SourceResource->GetSizeX(), SourceResource->GetSizeY()); // SourceResource is always proper size, as it's always the good MIP we want to copy from
+		FIntPoint DestSize(DestResource->GetSizeX() >> DestMip, DestResource->GetSizeY() >> DestMip);
+
+		int32 LocalComponentSizeQuad = SubSectionSizeQuad * NumSubSections;
+		FVector2D PositionOffset(FMath::RoundToInt(ComponentSectionBase.X / LocalComponentSizeQuad), FMath::RoundToInt(ComponentSectionBase.Y / LocalComponentSizeQuad));
+
+		FRHICopyTextureInfo Params;
+		Params.NumSlices = 1;
+		Params.Size.Z = 1;
+		Params.SourceSliceIndex = SourceArrayIndex;
+		Params.DestSliceIndex = DestArrayIndex;
+		Params.SourceMipIndex = 0; // In my case, always assume we copy from mip 0 to something else as in my case each mips will be stored into individual texture/RT
+		Params.DestMipIndex = DestMip;
+
+		if (SourceSize.X <= DestSize.X)
+		{
+			Params.SourcePosition.X = 0;
+			Params.Size.X = SourceSize.X;
+			Params.DestPosition.X = FMath::RoundToInt(PositionOffset.X * (((SubSectionSizeQuad + 1) * NumSubSections) >> DestMip));
+		}
+		else
+		{
+			Params.SourcePosition.X = FMath::RoundToInt(PositionOffset.X * (((SubSectionSizeQuad + 1) * NumSubSections) >> SourceMip));
+			Params.Size.X = DestSize.X;
+			Params.DestPosition.X = 0;
+		}
+
+		if (SourceSize.Y <= DestSize.Y)
+		{
+			Params.SourcePosition.Y = 0;
+			Params.Size.Y = SourceSize.Y;
+			Params.DestPosition.Y = FMath::RoundToInt(PositionOffset.Y * (((SubSectionSizeQuad + 1) * NumSubSections) >> DestMip));
+		}
+		else
+		{
+			Params.SourcePosition.Y = FMath::RoundToInt(PositionOffset.Y * (((SubSectionSizeQuad + 1) * NumSubSections) >> SourceMip));
+			Params.Size.Y = DestSize.Y;
+			Params.DestPosition.Y = 0;
+		}		
+
+		InRHICmdList.CopyTexture(SourceResource->TextureRHI, DestResource->TextureRHI, Params);
+
+		if (DestCPUResource != nullptr)
+		{
+			InRHICmdList.CopyTexture(SourceResource->TextureRHI, DestCPUResource->TextureRHI, Params);
+		}
+	}
+
+private:
+	FTextureResource* SourceResource;
+	FTextureResource* DestResource;
+	FTextureResource* DestCPUResource;
+	uint8 SourceMip;
+	uint8 DestMip;
+	uint32 SourceArrayIndex;
+	uint32 DestArrayIndex;
+	FIntPoint ComponentSectionBase;
+	int32 SubSectionSizeQuad;
+	int32 NumSubSections;
+	FString SourceDebugName;
+	FString DestResourceDebugName;
+};
+
+// Clear command
+
+class LandscapeLayersWeightmapClear_RenderThread
+{
+public:
+	LandscapeLayersWeightmapClear_RenderThread(const FString& InDebugName, FTextureRenderTargetResource* InTextureResourceToClear)
+		: DebugName(InDebugName)
+		, RenderTargetResource(InTextureResourceToClear)
+	{}
+
+	virtual ~LandscapeLayersWeightmapClear_RenderThread()
+	{}
+
+	void Clear(FRHICommandListImmediate& InRHICmdList)
+	{
+		SCOPE_CYCLE_COUNTER(STAT_LandscapeLayersRegenerate_RenderThread);
+		SCOPED_DRAW_EVENTF(InRHICmdList, LandscapeLayersRender, TEXT("%s"), DebugName.Len() > 0 ? *DebugName : TEXT("LandscapeLayersClear"));
+		SCOPED_GPU_STAT(InRHICmdList, LandscapeLayersRender);
+
+		check(IsInRenderingThread());
+
+		FRHIRenderPassInfo RPInfo(RenderTargetResource->TextureRHI, ERenderTargetActions::Clear_Store);
+		InRHICmdList.BeginRenderPass(RPInfo, TEXT("Clear"));
+		InRHICmdList.EndRenderPass();
+	}
+
+	FString DebugName;
+	FTextureRenderTargetResource* RenderTargetResource;
+};
+
+// Render command
+
+template<typename ShaderDataType, typename ShaderPixelClass, typename ShaderPixelMipsClass>
+class FLandscapeLayersRender_RenderThread
+{
+public:
+
+	FLandscapeLayersRender_RenderThread(const FString& InDebugName, UTextureRenderTarget2D* InWriteRenderTarget, const FIntPoint& InWriteRenderTargetSize, const FIntPoint& InReadRenderTargetSize, const FMatrix& InProjectionMatrix,
+											const ShaderDataType& InShaderParams, uint8 InCurrentMip, const TArray<FLandscapeLayersTriangle>& InTriangleList)
+		: RenderTargetResource(InWriteRenderTarget->GameThread_GetRenderTargetResource())
+		, WriteRenderTargetSize(InWriteRenderTargetSize)
+		, ReadRenderTargetSize(InReadRenderTargetSize)
+		, ProjectionMatrix(InProjectionMatrix)
+		, ShaderParams(InShaderParams)
+		, PrimitiveCount(InTriangleList.Num())
+		, DebugName(InDebugName)
+		, CurrentMip(InCurrentMip)
+	{
+		VertexBufferResource.Init(InTriangleList);
+	}
+
+	virtual ~FLandscapeLayersRender_RenderThread()
+	{}
+
+	void Render(FRHICommandListImmediate& InRHICmdList, bool InClearRT)
+	{
+		SCOPE_CYCLE_COUNTER(STAT_LandscapeLayersRegenerate_RenderThread);
+		SCOPED_DRAW_EVENTF(InRHICmdList, LandscapeLayersRender, TEXT("%s"), DebugName.Len() > 0 ? *DebugName : TEXT("LandscapeLayersRender"));
+		SCOPED_GPU_STAT(InRHICmdList, LandscapeLayersRender);
+		INC_DWORD_STAT(STAT_LandscapeLayersRegenerateDrawCalls);
+
+		check(IsInRenderingThread());
+
+		FSceneViewFamilyContext ViewFamily(FSceneViewFamily::ConstructionValues(RenderTargetResource, NULL, FEngineShowFlags(ESFIM_Game)).SetWorldTimes(FApp::GetCurrentTime() - GStartTime, FApp::GetDeltaTime(), FApp::GetCurrentTime() - GStartTime));
+
+		FSceneViewInitOptions ViewInitOptions;
+		ViewInitOptions.SetViewRectangle(FIntRect(0, 0, WriteRenderTargetSize.X, WriteRenderTargetSize.Y));
+		ViewInitOptions.ViewOrigin = FVector::ZeroVector;
+		ViewInitOptions.ViewRotationMatrix = FMatrix::Identity;
+		ViewInitOptions.ProjectionMatrix = ProjectionMatrix;
+		ViewInitOptions.ViewFamily = &ViewFamily;
+		ViewInitOptions.BackgroundColor = FLinearColor::Black;
+		ViewInitOptions.OverlayColor = FLinearColor::White;
+
+		// Create and add the new view
+		FSceneView* View = new FSceneView(ViewInitOptions);
+		ViewFamily.Views.Add(View);
+
+		// Init VB/IB Resource
+		VertexDeclaration.InitResource();
+		VertexBufferResource.InitResource();
+
+		// Setup Pipeline
+		FGraphicsPipelineStateInitializer GraphicsPSOInit;
+		GraphicsPSOInit.BoundShaderState.VertexDeclarationRHI = VertexDeclaration.VertexDeclarationRHI;
+		GraphicsPSOInit.PrimitiveType = PT_TriangleList;
+
+		GraphicsPSOInit.BlendState = TStaticBlendState<CW_RGBA, BO_Add, BF_One, BF_Zero, BO_Add, BF_One, BF_Zero>::GetRHI();
+		GraphicsPSOInit.RasterizerState = TStaticRasterizerState<>::GetRHI();
+		GraphicsPSOInit.DepthStencilState = TStaticDepthStencilState<false, CF_Always>::GetRHI();
+
+		FRHIRenderPassInfo RenderPassInfo(ViewFamily.RenderTarget->GetRenderTargetTexture(), CurrentMip == 0 ? ERenderTargetActions::Clear_Store : ERenderTargetActions::Load_Store, nullptr, 0, 0);
+		InRHICmdList.BeginRenderPass(RenderPassInfo, TEXT("DrawLayers"));
+
+		if (CurrentMip == 0)
+		{
+			// Setup Shaders
+			TShaderMapRef<FLandscapeLayersVS> VertexShader(GetGlobalShaderMap(View->GetFeatureLevel()));
+			TShaderMapRef<ShaderPixelClass> PixelShader(GetGlobalShaderMap(View->GetFeatureLevel()));
+
+			GraphicsPSOInit.BoundShaderState.VertexShaderRHI = GETSAFERHISHADER_VERTEX(*VertexShader);
+			GraphicsPSOInit.BoundShaderState.PixelShaderRHI = GETSAFERHISHADER_PIXEL(*PixelShader);
+
+			InRHICmdList.SetViewport(View->UnscaledViewRect.Min.X, View->UnscaledViewRect.Min.Y, 0.0f, View->UnscaledViewRect.Max.X, View->UnscaledViewRect.Max.Y, 1.0f);
+
+			InRHICmdList.ApplyCachedRenderTargets(GraphicsPSOInit);
+			SetGraphicsPipelineState(InRHICmdList, GraphicsPSOInit);
+
+			// Set shader params
+			VertexShader->SetParameters(InRHICmdList, ProjectionMatrix);
+			PixelShader->SetParameters(InRHICmdList, ShaderParams);
+		}
+		else
+		{
+			// Setup Shaders
+			TShaderMapRef<FLandscapeLayersVS> VertexShader(GetGlobalShaderMap(View->GetFeatureLevel()));
+			TShaderMapRef<ShaderPixelMipsClass> PixelShader(GetGlobalShaderMap(View->GetFeatureLevel()));
+
+			GraphicsPSOInit.BoundShaderState.VertexShaderRHI = GETSAFERHISHADER_VERTEX(*VertexShader);
+			GraphicsPSOInit.BoundShaderState.PixelShaderRHI = GETSAFERHISHADER_PIXEL(*PixelShader);
+
+			InRHICmdList.SetViewport(0.0f, 0.0f, 0.0f, WriteRenderTargetSize.X, WriteRenderTargetSize.Y, 1.0f);
+
+			InRHICmdList.ApplyCachedRenderTargets(GraphicsPSOInit);
+			SetGraphicsPipelineState(InRHICmdList, GraphicsPSOInit);
+
+			// Set shader params
+			VertexShader->SetParameters(InRHICmdList, ProjectionMatrix);
+			PixelShader->SetParameters(InRHICmdList, ShaderParams);			
+		}
+
+		InRHICmdList.SetStencilRef(0);
+		InRHICmdList.SetScissorRect(false, 0, 0, 0, 0);
+		InRHICmdList.SetStreamSource(0, VertexBufferResource.VertexBufferRHI, 0);
+
+		InRHICmdList.DrawPrimitive(0, PrimitiveCount, 1);
+
+		InRHICmdList.EndRenderPass();
+
+		VertexDeclaration.ReleaseResource();
+		VertexBufferResource.ReleaseResource();
+	}
+
+private:
+	FTextureRenderTargetResource* RenderTargetResource;
+	FIntPoint WriteRenderTargetSize;
+	FIntPoint ReadRenderTargetSize;
+	FMatrix ProjectionMatrix;
+	ShaderDataType ShaderParams;
+	FLandscapeLayersVertexBuffer VertexBufferResource;
+	int32 PrimitiveCount;
+	FLandscapeLayersVertexDeclaration VertexDeclaration;
+	FString DebugName;
+	uint8 CurrentMip;
+};
+
+typedef FLandscapeLayersRender_RenderThread<FLandscapeLayersHeightmapShaderParameters, FLandscapeLayersHeightmapPS, FLandscapeLayersHeightmapMipsPS> FLandscapeLayersHeightmapRender_RenderThread;
+typedef FLandscapeLayersRender_RenderThread<FLandscapeLayersWeightmapShaderParameters, FLandscapeLayersWeightmapPS, FLandscapeLayersWeightmapMipsPS> FLandscapeLayersWeightmapRender_RenderThread;
+
+#if WITH_EDITOR
+void ALandscapeProxy::SetupLayers(int32 InNumComponentsX, int32 InNumComponentsY)
+{
+	ALandscape* Landscape = GetLandscapeActor();
+	check(Landscape);
+
+	ULandscapeInfo* Info = GetLandscapeInfo();
+	if (Info == nullptr)
+	{
+		return;
+	}
+
+	TArray<ALandscapeProxy*> AllLandscapes;
+	AllLandscapes.Add(Landscape);
+	for (auto& It : Info->Proxies)
+	{
+		AllLandscapes.Add(It);
+	}
+
+	// Make sure we have at least 1 layer
+	if (Landscape->LandscapeLayers.Num() == 0)
+	{
+		Landscape->CreateLayer(FName(TEXT("Layer")), false);
+	}	
+
+	// TODO: When using Change Component Size, we will call Setup again, currently all the existing data will get collapsed into the layer 1, it should keep the layers data.
+
+	int32 NumComponentsX = InNumComponentsX;
+	int32 NumComponentsY = InNumComponentsY;
+	bool GenerateComponentCounts = NumComponentsX == INDEX_NONE || NumComponentsY == INDEX_NONE;
+	FIntPoint MaxSectionBase(0, 0);
+
+	// Setup all Heightmap data
+	for (ALandscapeProxy* LandscapeProxy : AllLandscapes)
+	{
+		for (ULandscapeComponent* Component : LandscapeProxy->LandscapeComponents)
+		{
+			UTexture2D* ComponentHeightmapTexture = Component->GetHeightmap();
+
+			FRenderDataPerHeightmap* Data = LandscapeProxy->RenderDataPerHeightmap.Find(ComponentHeightmapTexture);
+
+			if (Data == nullptr)
+			{
+				FRenderDataPerHeightmap NewData;
+				NewData.Components.Add(Component);
+				NewData.OriginalHeightmap = ComponentHeightmapTexture;
+				NewData.HeightmapsCPUReadBack = new FLandscapeLayersTexture2DCPUReadBackResource(ComponentHeightmapTexture->Source.GetSizeX(), ComponentHeightmapTexture->Source.GetSizeY(), ComponentHeightmapTexture->GetPixelFormat(), ComponentHeightmapTexture->Source.GetNumMips());
+				BeginInitResource(NewData.HeightmapsCPUReadBack);
+
+				LandscapeProxy->RenderDataPerHeightmap.Add(ComponentHeightmapTexture, NewData);
+			}
+			else
+			{
+				Data->Components.AddUnique(Component);
+			}
+
+			if (GenerateComponentCounts)
+			{
+				MaxSectionBase.X = FMath::Max(MaxSectionBase.X, Component->SectionBaseX);
+				MaxSectionBase.Y = FMath::Max(MaxSectionBase.Y, Component->SectionBaseY);
+			}
+		}
+	}
+
+	if (GenerateComponentCounts)
+	{
+		NumComponentsX = (MaxSectionBase.X / ComponentSizeQuads) + 1;
+		NumComponentsY = (MaxSectionBase.Y / ComponentSizeQuads) + 1;
+	}
+
+	const int32 TotalVertexCountX = (SubsectionSizeQuads * NumSubsections) * NumComponentsX + 1;
+	const int32 TotalVertexCountY = (SubsectionSizeQuads * NumSubsections) * NumComponentsY + 1;
+
+	if (Landscape->HeightmapRTList.Num() == 0)
+	{
+		Landscape->HeightmapRTList.Init(nullptr, EHeightmapRTType::HeightmapRT_Count);
+
+		int32 CurrentMipSizeX = ((SubsectionSizeQuads + 1) * NumSubsections) * NumComponentsX;
+		int32 CurrentMipSizeY = ((SubsectionSizeQuads + 1) * NumSubsections) * NumComponentsY;
+
+		for (int32 i = 0; i < EHeightmapRTType::HeightmapRT_Count; ++i)
+		{
+			Landscape->HeightmapRTList[i] = NewObject<UTextureRenderTarget2D>(Landscape->GetOutermost());
+			check(Landscape->HeightmapRTList[i]);
+			Landscape->HeightmapRTList[i]->RenderTargetFormat = RTF_RGBA8;
+			Landscape->HeightmapRTList[i]->AddressX = TextureAddress::TA_Clamp;
+			Landscape->HeightmapRTList[i]->AddressY = TextureAddress::TA_Clamp;
+			Landscape->HeightmapRTList[i]->ClearColor = FLinearColor(0.0f, 0.0f, 0.0f, 0.0f);
+
+			if (i < EHeightmapRTType::HeightmapRT_Mip1) // Landscape size RT
+			{
+				Landscape->HeightmapRTList[i]->InitAutoFormat(FMath::RoundUpToPowerOfTwo(TotalVertexCountX), FMath::RoundUpToPowerOfTwo(TotalVertexCountY));
+			}
+			else // Mips
+			{
+				CurrentMipSizeX >>= 1;
+				CurrentMipSizeY >>= 1;
+				Landscape->HeightmapRTList[i]->InitAutoFormat(FMath::RoundUpToPowerOfTwo(CurrentMipSizeX), FMath::RoundUpToPowerOfTwo(CurrentMipSizeY));
+			}
+
+			Landscape->HeightmapRTList[i]->UpdateResourceImmediate(true);
+
+			// Only generate required mips RT
+			if (CurrentMipSizeX == NumComponentsX && CurrentMipSizeY == NumComponentsY)
+			{
+				break;
+			}
+		}
+	}
+
+	TArray<FVector> VertexNormals;
+	TArray<uint16> EmptyHeightmapData;
+
+	// Setup all Heightmap data
+	for (ALandscapeProxy* LandscapeProxy : AllLandscapes)
+	{
+		for (auto& ItPair : LandscapeProxy->RenderDataPerHeightmap)
+		{
+			FRenderDataPerHeightmap& HeightmapRenderData = ItPair.Value;
+			HeightmapRenderData.TopLeftSectionBase = FIntPoint(TotalVertexCountX, TotalVertexCountY);
+
+			for (ULandscapeComponent* Component : HeightmapRenderData.Components)
+			{
+				HeightmapRenderData.TopLeftSectionBase.X = FMath::Min(HeightmapRenderData.TopLeftSectionBase.X, Component->GetSectionBase().X);
+				HeightmapRenderData.TopLeftSectionBase.Y = FMath::Min(HeightmapRenderData.TopLeftSectionBase.Y, Component->GetSectionBase().Y);
+			}
+
+			bool FirstLayer = true;
+
+			for (auto& ItLayerDataPair : LandscapeProxy->LandscapeLayersData)
+			{
+				FLandscapeLayerData& LayerData = ItLayerDataPair.Value;
+
+				if (LayerData.Heightmaps.Find(HeightmapRenderData.OriginalHeightmap) == nullptr)
+				{
+					UTexture2D* Heightmap = LandscapeProxy->CreateLandscapeTexture(HeightmapRenderData.OriginalHeightmap->Source.GetSizeX(), HeightmapRenderData.OriginalHeightmap->Source.GetSizeY(), TEXTUREGROUP_Terrain_Heightmap, HeightmapRenderData.OriginalHeightmap->Source.GetFormat());
+					LayerData.Heightmaps.Add(HeightmapRenderData.OriginalHeightmap, Heightmap);
+
+					int32 MipSubsectionSizeQuads = SubsectionSizeQuads;
+					int32 MipSizeU = Heightmap->Source.GetSizeX();
+					int32 MipSizeV = Heightmap->Source.GetSizeY();
+
+					// Copy data from Heightmap to first layer, after that all other layer will get init to empty layer
+					if (FirstLayer)
+					{
+						uint8 MipIndex = 0;
+						TArray<uint8> MipData;
+						MipData.Reserve(MipSizeU*MipSizeV * sizeof(FColor));
+
+						while (MipSizeU > 1 && MipSizeV > 1 && MipSubsectionSizeQuads >= 1)
+						{
+							MipData.Reset();
+							HeightmapRenderData.OriginalHeightmap->Source.GetMipData(MipData, MipIndex);
+
+							FColor* HeightmapTextureData = (FColor*)Heightmap->Source.LockMip(MipIndex);
+							FMemory::Memcpy(HeightmapTextureData, MipData.GetData(), MipData.Num());
+							Heightmap->Source.UnlockMip(MipIndex);
+
+							MipSizeU >>= 1;
+							MipSizeV >>= 1;
+
+							MipSubsectionSizeQuads = ((MipSubsectionSizeQuads + 1) >> 1) - 1;
+							++MipIndex;
+						}
+					}
+					else
+					{
+						TArray<FColor*> HeightmapMipMapData;
+
+						while (MipSizeU > 1 && MipSizeV > 1 && MipSubsectionSizeQuads >= 1)
+						{
+							int32 MipIndex = HeightmapMipMapData.Num();
+							FColor* HeightmapTextureData = (FColor*)Heightmap->Source.LockMip(MipIndex);
+							FMemory::Memzero(HeightmapTextureData, MipSizeU*MipSizeV * sizeof(FColor));
+							HeightmapMipMapData.Add(HeightmapTextureData);
+
+							MipSizeU >>= 1;
+							MipSizeV >>= 1;
+
+							MipSubsectionSizeQuads = ((MipSubsectionSizeQuads + 1) >> 1) - 1;
+						}
+
+						// Initialize blank heightmap data as if ALL components were in the same heightmap to prevent creating many allocations
+						if (EmptyHeightmapData.Num() == 0)
+						{
+							EmptyHeightmapData.Init(32768, TotalVertexCountX * TotalVertexCountY);
+						}
+
+						const FVector DrawScale3D = GetRootComponent()->RelativeScale3D;
+
+						// Init vertex normal data if required
+						if (VertexNormals.Num() == 0)
+						{
+							VertexNormals.AddZeroed(TotalVertexCountX * TotalVertexCountY);
+							for (int32 QuadY = 0; QuadY < TotalVertexCountY - 1; QuadY++)
+							{
+								for (int32 QuadX = 0; QuadX < TotalVertexCountX - 1; QuadX++)
+								{
+									const FVector Vert00 = FVector(0.0f, 0.0f, ((float)EmptyHeightmapData[FMath::Clamp<int32>(QuadY + 0, 0, TotalVertexCountY) * TotalVertexCountX + FMath::Clamp<int32>(QuadX + 0, 0, TotalVertexCountX)] - 32768.0f)*LANDSCAPE_ZSCALE) * DrawScale3D;
+									const FVector Vert01 = FVector(0.0f, 1.0f, ((float)EmptyHeightmapData[FMath::Clamp<int32>(QuadY + 1, 0, TotalVertexCountY) * TotalVertexCountX + FMath::Clamp<int32>(QuadX + 0, 0, TotalVertexCountX)] - 32768.0f)*LANDSCAPE_ZSCALE) * DrawScale3D;
+									const FVector Vert10 = FVector(1.0f, 0.0f, ((float)EmptyHeightmapData[FMath::Clamp<int32>(QuadY + 0, 0, TotalVertexCountY) * TotalVertexCountX + FMath::Clamp<int32>(QuadX + 1, 0, TotalVertexCountX)] - 32768.0f)*LANDSCAPE_ZSCALE) * DrawScale3D;
+									const FVector Vert11 = FVector(1.0f, 1.0f, ((float)EmptyHeightmapData[FMath::Clamp<int32>(QuadY + 1, 0, TotalVertexCountY) * TotalVertexCountX + FMath::Clamp<int32>(QuadX + 1, 0, TotalVertexCountX)] - 32768.0f)*LANDSCAPE_ZSCALE) * DrawScale3D;
+
+									const FVector FaceNormal1 = ((Vert00 - Vert10) ^ (Vert10 - Vert11)).GetSafeNormal();
+									const FVector FaceNormal2 = ((Vert11 - Vert01) ^ (Vert01 - Vert00)).GetSafeNormal();
+
+									// contribute to the vertex normals.
+									VertexNormals[(QuadX + 1 + TotalVertexCountX * (QuadY + 0))] += FaceNormal1;
+									VertexNormals[(QuadX + 0 + TotalVertexCountX * (QuadY + 1))] += FaceNormal2;
+									VertexNormals[(QuadX + 0 + TotalVertexCountX * (QuadY + 0))] += FaceNormal1 + FaceNormal2;
+									VertexNormals[(QuadX + 1 + TotalVertexCountX * (QuadY + 1))] += FaceNormal1 + FaceNormal2;
+								}
+							}
+						}
+
+						for (ULandscapeComponent* Component : HeightmapRenderData.Components)
+						{
+							int32 HeightmapComponentOffsetX = FMath::RoundToInt((float)Heightmap->Source.GetSizeX() * Component->HeightmapScaleBias.Z);
+							int32 HeightmapComponentOffsetY = FMath::RoundToInt((float)Heightmap->Source.GetSizeY() * Component->HeightmapScaleBias.W);
+
+							for (int32 SubsectionY = 0; SubsectionY < NumSubsections; SubsectionY++)
+							{
+								for (int32 SubsectionX = 0; SubsectionX < NumSubsections; SubsectionX++)
+								{
+									for (int32 SubY = 0; SubY <= SubsectionSizeQuads; SubY++)
+									{
+										for (int32 SubX = 0; SubX <= SubsectionSizeQuads; SubX++)
+										{
+											// X/Y of the vertex we're looking at in component's coordinates.
+											const int32 CompX = SubsectionSizeQuads * SubsectionX + SubX;
+											const int32 CompY = SubsectionSizeQuads * SubsectionY + SubY;
+
+											// X/Y of the vertex we're looking indexed into the texture data
+											const int32 TexX = (SubsectionSizeQuads + 1) * SubsectionX + SubX;
+											const int32 TexY = (SubsectionSizeQuads + 1) * SubsectionY + SubY;
+
+											const int32 HeightTexDataIdx = (HeightmapComponentOffsetX + TexX) + (HeightmapComponentOffsetY + TexY) * Heightmap->Source.GetSizeX();
+
+											// copy height and normal data
+											int32 Value = FMath::Clamp<int32>(CompY + Component->GetSectionBase().Y, 0, TotalVertexCountY) * TotalVertexCountX + FMath::Clamp<int32>(CompX + Component->GetSectionBase().X, 0, TotalVertexCountX);
+											const uint16 HeightValue = EmptyHeightmapData[Value];
+											const FVector Normal = VertexNormals[CompX + Component->GetSectionBase().X + TotalVertexCountX * (CompY + Component->GetSectionBase().Y)].GetSafeNormal();
+
+											HeightmapMipMapData[0][HeightTexDataIdx].R = HeightValue >> 8;
+											HeightmapMipMapData[0][HeightTexDataIdx].G = HeightValue & 255;
+											HeightmapMipMapData[0][HeightTexDataIdx].B = FMath::RoundToInt(127.5f * (Normal.X + 1.0f));
+											HeightmapMipMapData[0][HeightTexDataIdx].A = FMath::RoundToInt(127.5f * (Normal.Y + 1.0f));
+										}
+									}
+								}
+							}
+
+							bool IsBorderComponentX = (Component->GetSectionBase().X + 1 * NumSubsections) * InNumComponentsX == TotalVertexCountX;
+							bool IsBorderComponentY = (Component->GetSectionBase().Y + 1 * NumSubsections) * InNumComponentsY == TotalVertexCountY;
+
+							Component->GenerateHeightmapMips(HeightmapMipMapData, IsBorderComponentX ? MAX_int32 : 0, IsBorderComponentY ? MAX_int32 : 0);
+						}
+
+						// Add remaining mips down to 1x1 to heightmap texture.These do not represent quads and are just a simple averages of the previous mipmaps.
+						// These mips are not used for sampling in the vertex shader but could be sampled in the pixel shader.
+						int32 Mip = HeightmapMipMapData.Num();
+						MipSizeU = (Heightmap->Source.GetSizeX()) >> Mip;
+						MipSizeV = (Heightmap->Source.GetSizeY()) >> Mip;
+						while (MipSizeU > 1 && MipSizeV > 1)
+						{
+							HeightmapMipMapData.Add((FColor*)Heightmap->Source.LockMip(Mip));
+							const int32 PrevMipSizeU = (Heightmap->Source.GetSizeX()) >> (Mip - 1);
+							const int32 PrevMipSizeV = (Heightmap->Source.GetSizeY()) >> (Mip - 1);
+
+							for (int32 Y = 0; Y < MipSizeV; Y++)
+							{
+								for (int32 X = 0; X < MipSizeU; X++)
+								{
+									FColor* const TexData = &(HeightmapMipMapData[Mip])[X + Y * MipSizeU];
+
+									const FColor* const PreMipTexData00 = &(HeightmapMipMapData[Mip - 1])[(X * 2 + 0) + (Y * 2 + 0) * PrevMipSizeU];
+									const FColor* const PreMipTexData01 = &(HeightmapMipMapData[Mip - 1])[(X * 2 + 0) + (Y * 2 + 1) * PrevMipSizeU];
+									const FColor* const PreMipTexData10 = &(HeightmapMipMapData[Mip - 1])[(X * 2 + 1) + (Y * 2 + 0) * PrevMipSizeU];
+									const FColor* const PreMipTexData11 = &(HeightmapMipMapData[Mip - 1])[(X * 2 + 1) + (Y * 2 + 1) * PrevMipSizeU];
+
+									TexData->R = (((int32)PreMipTexData00->R + (int32)PreMipTexData01->R + (int32)PreMipTexData10->R + (int32)PreMipTexData11->R) >> 2);
+									TexData->G = (((int32)PreMipTexData00->G + (int32)PreMipTexData01->G + (int32)PreMipTexData10->G + (int32)PreMipTexData11->G) >> 2);
+									TexData->B = (((int32)PreMipTexData00->B + (int32)PreMipTexData01->B + (int32)PreMipTexData10->B + (int32)PreMipTexData11->B) >> 2);
+									TexData->A = (((int32)PreMipTexData00->A + (int32)PreMipTexData01->A + (int32)PreMipTexData10->A + (int32)PreMipTexData11->A) >> 2);
+								}
+							}
+							Mip++;
+							MipSizeU >>= 1;
+							MipSizeV >>= 1;
+						}
+
+						for (int32 i = 0; i < HeightmapMipMapData.Num(); i++)
+						{
+							Heightmap->Source.UnlockMip(i);
+						}
+					}
+
+					Heightmap->BeginCachePlatformData();
+					Heightmap->ClearAllCachedCookedPlatformData();					
+				}
+
+				FirstLayer = false;
+			}
+		}
+	}
+
+	// Weightmaps handling
+	if (Landscape->WeightmapRTList.Num() == 0)
+	{
+		Landscape->WeightmapRTList.Init(nullptr, EWeightmapRTType::WeightmapRT_Count);
+
+		int32 CurrentMipSizeX = ((SubsectionSizeQuads + 1) * NumSubsections) * NumComponentsX;
+		int32 CurrentMipSizeY = ((SubsectionSizeQuads + 1) * NumSubsections) * NumComponentsY;
+
+		for (int32 i = 0; i < EWeightmapRTType::WeightmapRT_Count; ++i)
+		{
+			Landscape->WeightmapRTList[i] = NewObject<UTextureRenderTarget2D>(Landscape->GetOutermost());
+
+			check(Landscape->WeightmapRTList[i]);
+			Landscape->WeightmapRTList[i]->AddressX = TextureAddress::TA_Clamp;
+			Landscape->WeightmapRTList[i]->AddressY = TextureAddress::TA_Clamp;
+			Landscape->WeightmapRTList[i]->ClearColor = FLinearColor(0.0f, 0.0f, 0.0f, 0.0f);
+			Landscape->WeightmapRTList[i]->RenderTargetFormat = RTF_RGBA8;
+
+			if (i < EWeightmapRTType::WeightmapRT_Mip0) // Landscape size RT, only create the number of layer we have
+			{
+				Landscape->WeightmapRTList[i]->RenderTargetFormat = i == WeightmapRT_Scratch_RGBA ? RTF_RGBA8 : RTF_R8;
+				Landscape->WeightmapRTList[i]->InitAutoFormat(FMath::RoundUpToPowerOfTwo(TotalVertexCountX), FMath::RoundUpToPowerOfTwo(TotalVertexCountY));
+			}
+			else // Mips
+			{
+				Landscape->WeightmapRTList[i]->InitAutoFormat(FMath::RoundUpToPowerOfTwo(CurrentMipSizeX), FMath::RoundUpToPowerOfTwo(CurrentMipSizeY));
+
+				CurrentMipSizeX >>= 1;
+				CurrentMipSizeY >>= 1;
+			}
+
+			Landscape->WeightmapRTList[i]->UpdateResourceImmediate(true);
+
+			// Only generate required mips RT
+			if (CurrentMipSizeX == NumComponentsX && CurrentMipSizeY == NumComponentsY)
+			{
+				break;
+			}
+		}
+	}
+
+	TArray<ULandscapeComponent*> ComponentsToCleanup;
+
+	for (ALandscapeProxy* LandscapeProxy : AllLandscapes)
+	{
+		bool FirstLayer = true;
+
+		for (auto& ItLayerDataPair : LandscapeProxy->LandscapeLayersData)
+		{
+			FGuid LayerGuid = ItLayerDataPair.Key;
+			FLandscapeLayerData& LayerData = ItLayerDataPair.Value;
+
+			struct FTextureData
+			{
+				UTexture2D* Texture;
+				ULandscapeWeightmapUsage* Usage;
+				ULandscapeWeightmapUsage* OriginalUsage;
+			};
+
+			TMap<UTexture2D*, FTextureData> ProcessedTextures;
+
+			for (ULandscapeComponent* Component : LandscapeProxy->LandscapeComponents)
+			{
+				FWeightmapLayerData* WeightmapLayer = LayerData.WeightmapData.Find(Component);
+
+				if (WeightmapLayer == nullptr)
+				{					
+					FWeightmapLayerData& NewWeightmapData = LayerData.WeightmapData.Add(Component, FWeightmapLayerData());
+
+					// If no data exist for this weightmap and that data exist in the base weightmap, simply copy it to the first layer, and clear the data in the base (as it will become the final weightmap)
+					if (FirstLayer)
+					{
+						ComponentsToCleanup.Add(Component);
+
+						const TArray<UTexture2D*>& ComponentWeightmapTextures = Component->GetWeightmapTextures();
+						TArray<FWeightmapLayerAllocationInfo>& ComponentLayerAllocations = Component->GetWeightmapLayerAllocations();
+						TArray<ULandscapeWeightmapUsage*>& ComponentWeightmapTexturesUsage = Component->GetWeightmapTexturesUsage();
+
+						NewWeightmapData.Weightmaps.AddDefaulted(ComponentWeightmapTextures.Num());
+						NewWeightmapData.WeightmapTextureUsages.AddDefaulted(ComponentWeightmapTexturesUsage.Num());
+
+						for (int32 TextureIndex = 0; TextureIndex < ComponentWeightmapTextures.Num(); ++TextureIndex)
+						{
+							UTexture2D* OriginalWeightmapTexture = ComponentWeightmapTextures[TextureIndex];
+							FTextureData* TextureData = ProcessedTextures.Find(OriginalWeightmapTexture);
+
+							if (TextureData != nullptr)
+							{
+								ComponentWeightmapTexturesUsage[TextureIndex] = TextureData->OriginalUsage;
+
+								NewWeightmapData.Weightmaps[TextureIndex] = TextureData->Texture;
+								NewWeightmapData.WeightmapTextureUsages[TextureIndex] = TextureData->Usage;
+								check(TextureData->Usage->LayerGuid == LayerGuid);
+
+								for (int32 ChannelIndex = 0; ChannelIndex < 4; ++ChannelIndex)
+								{
+									ULandscapeComponent* ChannelLandscapeComponent = NewWeightmapData.WeightmapTextureUsages.Last()->ChannelUsage[ChannelIndex];
+
+									if (ChannelLandscapeComponent != nullptr && ChannelLandscapeComponent == Component)
+									{
+										for (FWeightmapLayerAllocationInfo& Allocation : ComponentLayerAllocations)
+										{
+											if (Allocation.WeightmapTextureIndex == TextureIndex)
+											{
+												NewWeightmapData.WeightmapLayerAllocations.Add(Allocation);
+											}
+										}
+
+										break;
+									}
+								}
+							}
+							else
+							{
+								UTexture2D* NewWeightmapTexture = LandscapeProxy->CreateLandscapeTexture(OriginalWeightmapTexture->Source.GetSizeX(), OriginalWeightmapTexture->Source.GetSizeY(), TEXTUREGROUP_Terrain_Weightmap, OriginalWeightmapTexture->Source.GetFormat());
+
+								int32 MipSubsectionSizeQuads = SubsectionSizeQuads;
+								int32 MipSizeU = OriginalWeightmapTexture->Source.GetSizeX();
+								int32 MipSizeV = OriginalWeightmapTexture->Source.GetSizeY();
+
+								uint8 MipIndex = 0;
+								TArray<uint8> MipData;
+								MipData.Reserve(MipSizeU * MipSizeV * sizeof(FColor));
+
+								while (MipSizeU > 1 && MipSizeV > 1 && MipSubsectionSizeQuads >= 1)
+								{
+									MipData.Reset();
+									OriginalWeightmapTexture->Source.GetMipData(MipData, MipIndex);
+
+									FColor* WeightmapTextureData = (FColor*)NewWeightmapTexture->Source.LockMip(MipIndex);
+									FMemory::Memcpy(WeightmapTextureData, MipData.GetData(), MipData.Num());
+									NewWeightmapTexture->Source.UnlockMip(MipIndex);
+
+									MipSizeU >>= 1;
+									MipSizeV >>= 1;
+
+									MipSubsectionSizeQuads = ((MipSubsectionSizeQuads + 1) >> 1) - 1;
+									++MipIndex;
+								}
+
+								NewWeightmapData.Weightmaps[TextureIndex] = NewWeightmapTexture;
+								NewWeightmapData.WeightmapTextureUsages[TextureIndex] = ComponentWeightmapTexturesUsage[TextureIndex];
+								NewWeightmapData.WeightmapTextureUsages[TextureIndex]->LayerGuid = LayerGuid;
+
+								// Create new Usage for the base layer as the other one will now be used by the Layer 1
+								ComponentWeightmapTexturesUsage[TextureIndex] = LandscapeProxy->WeightmapUsageMap.Add(NewWeightmapTexture, NewObject<ULandscapeWeightmapUsage>(LandscapeProxy));
+
+								for (FWeightmapLayerAllocationInfo& Allocation : ComponentLayerAllocations)
+								{
+									if (Allocation.WeightmapTextureIndex == TextureIndex)
+									{
+										NewWeightmapData.WeightmapLayerAllocations.Add(Allocation);
+									}
+								}
+
+								FTextureData NewTextureData;
+								NewTextureData.Texture = NewWeightmapTexture;
+								NewTextureData.Usage = NewWeightmapData.WeightmapTextureUsages[TextureIndex];
+								NewTextureData.OriginalUsage = ComponentWeightmapTexturesUsage[TextureIndex];
+
+								ProcessedTextures.Add(OriginalWeightmapTexture, NewTextureData);
+
+								NewWeightmapTexture->BeginCachePlatformData();
+								NewWeightmapTexture->ClearAllCachedCookedPlatformData();
+							}							
+						}
+					}
+				}
+				else
+				{
+					WeightmapLayer->WeightmapTextureUsages.AddDefaulted(WeightmapLayer->Weightmaps.Num());
+
+					// regenerate the weightmap usage
+					for (int32 LayerIdx = 0; LayerIdx < WeightmapLayer->WeightmapLayerAllocations.Num(); LayerIdx++)
+					{
+						FWeightmapLayerAllocationInfo& Allocation = WeightmapLayer->WeightmapLayerAllocations[LayerIdx];
+						UTexture2D* WeightmapTexture = WeightmapLayer->Weightmaps[Allocation.WeightmapTextureIndex];
+						ULandscapeWeightmapUsage** TempUsage = LandscapeProxy->WeightmapUsageMap.Find(WeightmapTexture);
+
+						if (TempUsage == nullptr)
+						{
+							TempUsage = &LandscapeProxy->WeightmapUsageMap.Add(WeightmapTexture, NewObject<ULandscapeWeightmapUsage>(LandscapeProxy));
+							(*TempUsage)->LayerGuid = LayerGuid;
+						}
+
+						ULandscapeWeightmapUsage* Usage = *TempUsage;
+						WeightmapLayer->WeightmapTextureUsages[Allocation.WeightmapTextureIndex] = Usage; // Keep a ref to it for faster access
+
+						check(Usage->ChannelUsage[Allocation.WeightmapTextureChannel] == nullptr || Usage->ChannelUsage[Allocation.WeightmapTextureChannel] == Component);
+
+						Usage->ChannelUsage[Allocation.WeightmapTextureChannel] = Component;
+					}
+				}
+			}
+
+			FirstLayer = false;
+		}
+	}
+
+	for (ULandscapeComponent* Component : ComponentsToCleanup)
+	{
+		TArray<FWeightmapLayerAllocationInfo>& ComponentLayerAllocations = Component->GetWeightmapLayerAllocations();
+		ComponentLayerAllocations.Reset();		
+	}
+
+	// Fix Owning actor for Brushes. It can happen after save as operation, for example
+	for (FLandscapeLayer& Layer : Landscape->LandscapeLayers)
+	{
+		for (int32 i = Layer.Brushes.Num() - 1; i >= 0; --i)
+		{
+			FLandscapeLayerBrush& Brush = Layer.Brushes[i];
+
+			if (Brush.BPCustomBrush != nullptr)
+			{
+				if (Brush.BPCustomBrush->GetOwningLandscape() == nullptr)
+				{
+					Brush.BPCustomBrush->SetOwningLandscape(Landscape);
+				}
+			}
+		}
+
+		// TEMP stuff
+		if (Layer.HeightmapBrushOrderIndices.Num() == 0)
+		{
+			for (int32 i = 0; i < Layer.Brushes.Num(); ++i)
+			{
+				FLandscapeLayerBrush& Brush = Layer.Brushes[i];
+
+				if (Brush.BPCustomBrush != nullptr && Brush.BPCustomBrush->IsAffectingHeightmap())
+				{
+					Layer.HeightmapBrushOrderIndices.Add(i);
+				}
+			}
+		}
+
+		if (Layer.WeightmapBrushOrderIndices.Num() == 0)
+		{
+			for (int32 i = 0; i < Layer.Brushes.Num(); ++i)
+			{
+				FLandscapeLayerBrush& Brush = Layer.Brushes[i];
+
+				if (Brush.BPCustomBrush != nullptr && Brush.BPCustomBrush->IsAffectingWeightmap())
+				{
+					Layer.WeightmapBrushOrderIndices.Add(i);
+				}
+			}
+		}		
+		// TEMP stuff
+	}
+}
+
+void ALandscape::CopyLayersTexture(UTexture* InSourceTexture, UTexture* InDestTexture, FTextureResource* InDestCPUResource, const FIntPoint& InFirstComponentSectionBase, uint8 InSourceCurrentMip, uint8 InDestCurrentMip, uint32 InSourceArrayIndex, uint32 InDestArrayIndex) const
+{
+	if (InSourceTexture != nullptr && InDestTexture != nullptr)
+	{
+		CopyLayersTexture(InSourceTexture->GetName(), InSourceTexture->Resource, InDestTexture->GetName(), InDestTexture->Resource, InDestCPUResource, InFirstComponentSectionBase, InSourceCurrentMip, InDestCurrentMip, InSourceArrayIndex, InDestArrayIndex);
+	}
+}
+
+void ALandscape::CopyLayersTexture(const FString& InSourceDebugName, FTextureResource* InSourceResource, const FString& InDestDebugName, FTextureResource* InDestResource, FTextureResource* InDestCPUResource, const FIntPoint& InFirstComponentSectionBase, 
+										uint8 InSourceCurrentMip, uint8 InDestCurrentMip, uint32 InSourceArrayIndex, uint32 InDestArrayIndex) const
+{
+	check(InSourceResource != nullptr);
+	check(InDestResource != nullptr);
+
+	FLandscapeLayersCopyTexture_RenderThread CopyTexture(InSourceDebugName, InSourceResource, InDestDebugName, InDestResource, InDestCPUResource, InFirstComponentSectionBase, SubsectionSizeQuads, NumSubsections, InSourceCurrentMip, InDestCurrentMip, InSourceArrayIndex, InDestArrayIndex);
+
+	ENQUEUE_RENDER_COMMAND(FLandscapeLayersCopyCommand)(
+		[CopyTexture](FRHICommandListImmediate& RHICmdList) mutable
+	{
+		CopyTexture.Copy(RHICmdList);
+	});	
+}
+
+void ALandscape::DrawWeightmapComponentsToRenderTarget(const FString& InDebugName, const FIntPoint& InSectionBase, const FVector2D& InScaleBias, UTexture* InWeightmapRTRead, UTextureRenderTarget2D* InOptionalWeightmapRTRead2, UTextureRenderTarget2D* InWeightmapRTWrite,
+													   bool InClearRTWrite, FLandscapeLayersWeightmapShaderParameters& InShaderParams, uint8 InMipRender) const
+{
+	check(InWeightmapRTRead != nullptr);
+	check(InWeightmapRTWrite != nullptr);
+
+	FIntPoint WeightmapWriteTextureSize(InWeightmapRTWrite->SizeX, InWeightmapRTWrite->SizeY);
+	FIntPoint WeightmapReadTextureSize(InWeightmapRTRead->Source.GetSizeX(), InWeightmapRTRead->Source.GetSizeY());
+	UTextureRenderTarget2D* WeightmapRTRead = Cast<UTextureRenderTarget2D>(InWeightmapRTRead);
+
+	if (WeightmapRTRead != nullptr)
+	{
+		WeightmapReadTextureSize.X = WeightmapRTRead->SizeX;
+		WeightmapReadTextureSize.Y = WeightmapRTRead->SizeY;
+	}
+
+	// Quad Setup
+	TArray<FLandscapeLayersTriangle> TriangleList;
+	TriangleList.Reserve(1 * 2 * NumSubsections);
+
+	if (InMipRender == 0)
+	{
+		GenerateLayersRenderQuadsAtlas(InSectionBase, InScaleBias, SubsectionSizeQuads, WeightmapReadTextureSize, WeightmapWriteTextureSize, TriangleList);
+	}
+	else
+	{
+		GenerateLayersRenderQuadsMip(InSectionBase, InScaleBias, SubsectionSizeQuads, WeightmapReadTextureSize, WeightmapWriteTextureSize, InMipRender, TriangleList);
+	}
+
+	InShaderParams.ReadWeightmap1 = InWeightmapRTRead;
+	InShaderParams.ReadWeightmap2 = InOptionalWeightmapRTRead2;
+	InShaderParams.CurrentMipComponentVertexCount = (((SubsectionSizeQuads + 1) * NumSubsections) >> InMipRender);
+
+	if (InMipRender > 0)
+	{
+		InShaderParams.CurrentMipSize = WeightmapWriteTextureSize;
+		InShaderParams.ParentMipSize = WeightmapReadTextureSize;
+	}
+
+	FMatrix ProjectionMatrix = AdjustProjectionMatrixForRHI(FTranslationMatrix(FVector(0, 0, 0)) *
+		FMatrix(FPlane(1.0f / (FMath::Max<uint32>(WeightmapWriteTextureSize.X, 1.f) / 2.0f), 0.0, 0.0f, 0.0f), FPlane(0.0f, -1.0f / (FMath::Max<uint32>(WeightmapWriteTextureSize.Y, 1.f) / 2.0f), 0.0f, 0.0f), FPlane(0.0f, 0.0f, 1.0f, 0.0f), FPlane(-1.0f, 1.0f, 0.0f, 1.0f)));
+
+	FLandscapeLayersWeightmapRender_RenderThread LayersRender(InDebugName, InWeightmapRTWrite, WeightmapWriteTextureSize, WeightmapReadTextureSize, ProjectionMatrix, InShaderParams, InMipRender, TriangleList);
+
+	ENQUEUE_RENDER_COMMAND(FDrawLandscapeLayersWeightmapCommand)(
+		[LayersRender, ClearRT = InClearRTWrite](FRHICommandListImmediate& RHICmdList) mutable
+	{
+		LayersRender.Render(RHICmdList, ClearRT);
+	});
+
+	PrintLayersDebugRT(InDebugName, InWeightmapRTWrite, InMipRender, false);
+}
+
+void ALandscape::DrawWeightmapComponentsToRenderTarget(const FString& InDebugName, const TArray<ULandscapeComponent*>& InComponentsToDraw, UTexture* InWeightmapRTRead, UTextureRenderTarget2D* InOptionalWeightmapRTRead2, UTextureRenderTarget2D* InWeightmapRTWrite,
+													   bool InClearRTWrite, FLandscapeLayersWeightmapShaderParameters& InShaderParams, uint8 InMipRender) const
+{
+	check(InWeightmapRTRead != nullptr);
+	check(InWeightmapRTWrite != nullptr);
+
+	FIntPoint WeightmapWriteTextureSize(InWeightmapRTWrite->SizeX, InWeightmapRTWrite->SizeY);
+	FIntPoint WeightmapReadTextureSize(InWeightmapRTRead->Source.GetSizeX(), InWeightmapRTRead->Source.GetSizeY());
+	UTextureRenderTarget2D* WeightmapRTRead = Cast<UTextureRenderTarget2D>(InWeightmapRTRead);
+
+	if (WeightmapRTRead != nullptr)
+	{
+		WeightmapReadTextureSize.X = WeightmapRTRead->SizeX;
+		WeightmapReadTextureSize.Y = WeightmapRTRead->SizeY;
+	}
+
+	// Quad Setup
+	TArray<FLandscapeLayersTriangle> TriangleList;
+	TriangleList.Reserve(InComponentsToDraw.Num() * 2 * NumSubsections);
+
+	if (InMipRender == 0)
+	{
+		for (ULandscapeComponent* Component : InComponentsToDraw)
+		{
+			// TODO: check what to do with WeightmapSubsectionOffset
+			FVector2D WeightmapScaleBias(Component->WeightmapScaleBias.Z, Component->WeightmapScaleBias.W);
+			GenerateLayersRenderQuadsAtlas(Component->GetSectionBase(), WeightmapScaleBias, SubsectionSizeQuads, WeightmapReadTextureSize, WeightmapWriteTextureSize, TriangleList);
+		}
+	}
+	else
+	{
+		for (ULandscapeComponent* Component : InComponentsToDraw)
+		{
+			// TODO: check what to do with WeightmapSubsectionOffset
+			FVector2D WeightmapScaleBias(Component->WeightmapScaleBias.Z, Component->WeightmapScaleBias.W);
+			GenerateLayersRenderQuadsMip(Component->GetSectionBase(), WeightmapScaleBias, SubsectionSizeQuads, WeightmapReadTextureSize, WeightmapWriteTextureSize, InMipRender, TriangleList);
+		}
+	}
+
+	InShaderParams.ReadWeightmap1 = InWeightmapRTRead;
+	InShaderParams.ReadWeightmap2 = InOptionalWeightmapRTRead2;
+	InShaderParams.CurrentMipComponentVertexCount = (((SubsectionSizeQuads + 1) * NumSubsections) >> InMipRender);
+
+	if (InMipRender > 0)
+	{
+		InShaderParams.CurrentMipSize = WeightmapWriteTextureSize;
+		InShaderParams.ParentMipSize = WeightmapReadTextureSize;
+	}
+
+	FMatrix ProjectionMatrix = AdjustProjectionMatrixForRHI(FTranslationMatrix(FVector(0, 0, 0)) *
+															FMatrix(FPlane(1.0f / (FMath::Max<uint32>(WeightmapWriteTextureSize.X, 1.f) / 2.0f), 0.0, 0.0f, 0.0f), FPlane(0.0f, -1.0f / (FMath::Max<uint32>(WeightmapWriteTextureSize.Y, 1.f) / 2.0f), 0.0f, 0.0f), FPlane(0.0f, 0.0f, 1.0f, 0.0f), FPlane(-1.0f, 1.0f, 0.0f, 1.0f)));
+
+	FLandscapeLayersWeightmapRender_RenderThread LayersRender(InDebugName, InWeightmapRTWrite, WeightmapWriteTextureSize, WeightmapReadTextureSize, ProjectionMatrix, InShaderParams, InMipRender, TriangleList);
+
+	ENQUEUE_RENDER_COMMAND(FDrawLandscapeLayersWeightmapCommand)(
+		[LayersRender, ClearRT = InClearRTWrite](FRHICommandListImmediate& RHICmdList) mutable
+	{
+		LayersRender.Render(RHICmdList, ClearRT);
+	});
+
+	PrintLayersDebugRT(InDebugName, InWeightmapRTWrite, InMipRender, false);
+}
+
+void ALandscape::DrawWeightmapComponentToRenderTargetMips(const FIntPoint& TopLeftTexturePosition, UTexture* InReadWeightmap, bool InClearRTWrite, struct FLandscapeLayersWeightmapShaderParameters& InShaderParams) const
+{
+	bool OutputDebugName = CVarOutputLayersDebugDrawCallName.GetValueOnAnyThread() == 1 || CVarOutputLayersRTContent.GetValueOnAnyThread() == 1 ? true : false;
+	int32 CurrentMip = 1;
+	UTexture* ReadMipRT = InReadWeightmap;
+
+	// Convert from Texture position to SectionBase
+	int32 LocalComponentSizeQuad = SubsectionSizeQuads * NumSubsections;
+	int32 LocalComponentSizeVerts = SubsectionSizeQuads + 1 * NumSubsections;
+
+	FVector2D PositionOffset(FMath::RoundToInt(TopLeftTexturePosition.X / LocalComponentSizeVerts), FMath::RoundToInt(TopLeftTexturePosition.Y / LocalComponentSizeVerts));
+	FIntPoint ComponentSectionBase(PositionOffset.X * LocalComponentSizeQuad, PositionOffset.Y * LocalComponentSizeQuad);
+	FVector2D WeightmapScaleBias(0.0f, 0.0f);
+
+	for (int32 MipRTIndex = EWeightmapRTType::WeightmapRT_Mip1; MipRTIndex < EWeightmapRTType::WeightmapRT_Count; ++MipRTIndex)
+	{
+		UTextureRenderTarget2D* WriteMipRT = WeightmapRTList[MipRTIndex];
+
+		if (WriteMipRT != nullptr)
+		{
+			DrawWeightmapComponentsToRenderTarget(OutputDebugName ? FString::Printf(TEXT("LS Weight: %s = -> %s Mips %d"), *ReadMipRT->GetName(), *WriteMipRT->GetName(), CurrentMip) : TEXT(""),
+												  ComponentSectionBase, WeightmapScaleBias, ReadMipRT, nullptr, WriteMipRT, InClearRTWrite, InShaderParams, CurrentMip++);
+		}
+
+		ReadMipRT = WeightmapRTList[MipRTIndex];
+	}
+}
+
+void ALandscape::ClearLayersWeightmapTextureResource(const FString& InDebugName, FTextureRenderTargetResource* InTextureResourceToClear)
+{
+	LandscapeLayersWeightmapClear_RenderThread LayersClear(InDebugName, InTextureResourceToClear);
+
+	ENQUEUE_RENDER_COMMAND(FLandscapeLayersClearWeightmapCommand)(
+		[LayersClear](FRHICommandListImmediate& RHICmdList) mutable
+		{
+			LayersClear.Clear(RHICmdList);
+		});
+}
+
+void ALandscape::DrawHeightmapComponentsToRenderTargetMips(TArray<ULandscapeComponent*>& InComponentsToDraw, UTexture* InReadHeightmap, bool InClearRTWrite, struct FLandscapeLayersHeightmapShaderParameters& InShaderParams) const
+{
+	bool OutputDebugName = CVarOutputLayersDebugDrawCallName.GetValueOnAnyThread() == 1 || CVarOutputLayersRTContent.GetValueOnAnyThread() == 1 ? true : false;
+	int32 CurrentMip = 1;
+	UTexture* ReadMipRT = InReadHeightmap;
+
+	for (int32 MipRTIndex = EHeightmapRTType::HeightmapRT_Mip1; MipRTIndex < EHeightmapRTType::HeightmapRT_Count; ++MipRTIndex)
+	{
+		UTextureRenderTarget2D* WriteMipRT = HeightmapRTList[MipRTIndex];
+
+		if (WriteMipRT != nullptr)
+		{
+			DrawHeightmapComponentsToRenderTarget(OutputDebugName ? FString::Printf(TEXT("LS Height: %s = -> %s CombinedAtlasWithMips %d"), *ReadMipRT->GetName(), *WriteMipRT->GetName(), CurrentMip) : TEXT(""),
+												  InComponentsToDraw, ReadMipRT, nullptr, WriteMipRT, ERTDrawingType::RTMips, InClearRTWrite, InShaderParams, CurrentMip++);
+		}
+
+		ReadMipRT = HeightmapRTList[MipRTIndex];
+	}
+}
+
+void ALandscape::DrawHeightmapComponentsToRenderTarget(const FString& InDebugName, const TArray<ULandscapeComponent*>& InComponentsToDraw, UTexture* InHeightmapRTRead, UTextureRenderTarget2D* InOptionalHeightmapRTRead2, UTextureRenderTarget2D* InHeightmapRTWrite,
+													  ERTDrawingType InDrawType, bool InClearRTWrite, FLandscapeLayersHeightmapShaderParameters& InShaderParams, uint8 InMipRender) const
+{
+	check(InHeightmapRTRead != nullptr);
+	check(InHeightmapRTWrite != nullptr);
+
+	FIntPoint HeightmapWriteTextureSize(InHeightmapRTWrite->SizeX, InHeightmapRTWrite->SizeY);
+	FIntPoint HeightmapReadTextureSize(InHeightmapRTRead->Source.GetSizeX(), InHeightmapRTRead->Source.GetSizeY());
+	UTextureRenderTarget2D* HeightmapRTRead = Cast<UTextureRenderTarget2D>(InHeightmapRTRead);
+
+	if (HeightmapRTRead != nullptr)
+	{
+		HeightmapReadTextureSize.X = HeightmapRTRead->SizeX;
+		HeightmapReadTextureSize.Y = HeightmapRTRead->SizeY;
+	}
+
+	// Quad Setup
+	TArray<FLandscapeLayersTriangle> TriangleList;
+	TriangleList.Reserve(InComponentsToDraw.Num() * 2 * NumSubsections);
+
+	switch (InDrawType)
+	{
+		case ERTDrawingType::RTAtlas:
+		{
+			for (ULandscapeComponent* Component : InComponentsToDraw)
+			{
+				FVector2D HeightmapScaleBias(Component->HeightmapScaleBias.Z, Component->HeightmapScaleBias.W);
+				GenerateLayersRenderQuadsAtlas(Component->GetSectionBase(), HeightmapScaleBias, SubsectionSizeQuads, HeightmapReadTextureSize, HeightmapWriteTextureSize, TriangleList);
+			}
+		} break;
+
+		case ERTDrawingType::RTAtlasToNonAtlas:
+		{
+			for (ULandscapeComponent* Component : InComponentsToDraw)
+			{
+				FVector2D HeightmapScaleBias(Component->HeightmapScaleBias.Z, Component->HeightmapScaleBias.W);
+				GenerateLayersRenderQuadsAtlasToNonAtlas(Component->GetSectionBase(), HeightmapScaleBias, SubsectionSizeQuads, HeightmapReadTextureSize, HeightmapWriteTextureSize, TriangleList);
+			}
+		} break;
+
+		case ERTDrawingType::RTNonAtlas:
+		{
+			for (ULandscapeComponent* Component : InComponentsToDraw)
+			{
+				FVector2D HeightmapScaleBias(Component->HeightmapScaleBias.Z, Component->HeightmapScaleBias.W);
+				GenerateLayersRenderQuadsNonAtlas(Component->GetSectionBase(), HeightmapScaleBias, SubsectionSizeQuads, HeightmapReadTextureSize, HeightmapWriteTextureSize, TriangleList);
+			}
+		} break;
+
+		case ERTDrawingType::RTNonAtlasToAtlas:
+		{
+			for (ULandscapeComponent* Component : InComponentsToDraw)
+			{
+				FVector2D HeightmapScaleBias(Component->HeightmapScaleBias.Z, Component->HeightmapScaleBias.W);
+				GenerateLayersRenderQuadsNonAtlasToAtlas(Component->GetSectionBase(), HeightmapScaleBias, SubsectionSizeQuads, HeightmapReadTextureSize, HeightmapWriteTextureSize, TriangleList);
+			}
+		} break;
+
+		case ERTDrawingType::RTMips:
+		{
+			for (ULandscapeComponent* Component : InComponentsToDraw)
+			{
+				FVector2D HeightmapScaleBias(Component->HeightmapScaleBias.Z, Component->HeightmapScaleBias.W);
+				GenerateLayersRenderQuadsMip(Component->GetSectionBase(), HeightmapScaleBias, SubsectionSizeQuads, HeightmapReadTextureSize, HeightmapWriteTextureSize, InMipRender, TriangleList);
+			}			
+		} break;
+
+		default:
+		{
+			check(false);
+			return;
+		}
+	}
+
+	InShaderParams.ReadHeightmap1 = InHeightmapRTRead;
+	InShaderParams.ReadHeightmap2 = InOptionalHeightmapRTRead2;
+	InShaderParams.HeightmapSize = HeightmapReadTextureSize;
+	InShaderParams.CurrentMipComponentVertexCount = (((SubsectionSizeQuads + 1) * NumSubsections) >> InMipRender);
+
+	if (InMipRender > 0)
+	{
+		InShaderParams.CurrentMipSize = HeightmapWriteTextureSize;
+		InShaderParams.ParentMipSize = HeightmapReadTextureSize;
+	}
+
+	FMatrix ProjectionMatrix = AdjustProjectionMatrixForRHI(FTranslationMatrix(FVector(0, 0, 0)) *
+															FMatrix(FPlane(1.0f / (FMath::Max<uint32>(HeightmapWriteTextureSize.X, 1.f) / 2.0f), 0.0, 0.0f, 0.0f), FPlane(0.0f, -1.0f / (FMath::Max<uint32>(HeightmapWriteTextureSize.Y, 1.f) / 2.0f), 0.0f, 0.0f), FPlane(0.0f, 0.0f, 1.0f, 0.0f), FPlane(-1.0f, 1.0f, 0.0f, 1.0f)));
+
+	FLandscapeLayersHeightmapRender_RenderThread LayersRender(InDebugName, InHeightmapRTWrite, HeightmapWriteTextureSize, HeightmapReadTextureSize, ProjectionMatrix, InShaderParams, InMipRender, TriangleList);
+
+	ENQUEUE_RENDER_COMMAND(FDrawLandscapeLayersHeightmapCommand)(
+		[LayersRender, ClearRT = InClearRTWrite](FRHICommandListImmediate& RHICmdList) mutable
+		{
+		LayersRender.Render(RHICmdList, ClearRT);
+		});
+	
+	PrintLayersDebugRT(InDebugName, InHeightmapRTWrite, InMipRender, true, InShaderParams.GenerateNormals);
+}
+
+void ALandscape::GenerateLayersRenderQuad(const FIntPoint& InVertexPosition, float InVertexSize, const FVector2D& InUVStart, const FVector2D& InUVSize, TArray<FLandscapeLayersTriangle>& OutTriangles) const
+{
+	FLandscapeLayersTriangle Tri1;
+	
+	Tri1.V0.Position = FVector2D(InVertexPosition.X, InVertexPosition.Y);
+	Tri1.V1.Position = FVector2D(InVertexPosition.X + InVertexSize, InVertexPosition.Y);
+	Tri1.V2.Position = FVector2D(InVertexPosition.X + InVertexSize, InVertexPosition.Y + InVertexSize);
+
+	Tri1.V0.UV = FVector2D(InUVStart.X, InUVStart.Y);
+	Tri1.V1.UV = FVector2D(InUVStart.X + InUVSize.X, InUVStart.Y);
+	Tri1.V2.UV = FVector2D(InUVStart.X + InUVSize.X, InUVStart.Y + InUVSize.Y);
+	OutTriangles.Add(Tri1);
+
+	FLandscapeLayersTriangle Tri2;
+	Tri2.V0.Position = FVector2D(InVertexPosition.X + InVertexSize, InVertexPosition.Y + InVertexSize);
+	Tri2.V1.Position = FVector2D(InVertexPosition.X, InVertexPosition.Y + InVertexSize);
+	Tri2.V2.Position = FVector2D(InVertexPosition.X, InVertexPosition.Y);
+
+	Tri2.V0.UV = FVector2D(InUVStart.X + InUVSize.X, InUVStart.Y + InUVSize.Y);
+	Tri2.V1.UV = FVector2D(InUVStart.X, InUVStart.Y + InUVSize.Y);
+	Tri2.V2.UV = FVector2D(InUVStart.X, InUVStart.Y);
+
+	OutTriangles.Add(Tri2);
+}
+
+void ALandscape::GenerateLayersRenderQuadsAtlas(const FIntPoint& InSectionBase, const FVector2D& InScaleBias, float InSubSectionSizeQuad, const FIntPoint& InReadSize, const FIntPoint& InWriteSize, TArray<FLandscapeLayersTriangle>& OutTriangles) const
+{
+	FIntPoint ComponentSectionBase = InSectionBase;
+	FIntPoint UVComponentSectionBase = InSectionBase;
+
+	int32 SubsectionSizeVerts = InSubSectionSizeQuad + 1;
+	int32 LocalComponentSizeQuad = InSubSectionSizeQuad * NumSubsections;
+	int32 LocalComponentSizeVerts = SubsectionSizeVerts * NumSubsections;
+
+	FVector2D PositionOffset(FMath::RoundToInt(ComponentSectionBase.X / LocalComponentSizeQuad), FMath::RoundToInt(ComponentSectionBase.Y / LocalComponentSizeQuad));
+	FVector2D ComponentsPerTexture(FMath::RoundToInt(InWriteSize.X / LocalComponentSizeQuad), FMath::RoundToInt(InWriteSize.Y / LocalComponentSizeQuad));
+
+	if (InReadSize.X >= InWriteSize.X)
+	{
+		if (InReadSize.X == InWriteSize.X)
+		{
+			if (ComponentsPerTexture.X > 1.0f)
+			{
+				UVComponentSectionBase.X = PositionOffset.X * LocalComponentSizeVerts;
+			}
+			else
+			{
+				UVComponentSectionBase.X -= (UVComponentSectionBase.X + LocalComponentSizeQuad > InWriteSize.X) ? FMath::FloorToInt((PositionOffset.X / ComponentsPerTexture.X)) * ComponentsPerTexture.X * LocalComponentSizeQuad : 0;
+			}
+		}
+
+		ComponentSectionBase.X -= (ComponentSectionBase.X + LocalComponentSizeQuad > InWriteSize.X) ? FMath::FloorToInt((PositionOffset.X / ComponentsPerTexture.X)) * ComponentsPerTexture.X * LocalComponentSizeQuad : 0;
+		PositionOffset.X = ComponentSectionBase.X / LocalComponentSizeQuad;
+	}
+
+	if (InReadSize.Y >= InWriteSize.Y)
+	{
+		if (InReadSize.Y == InWriteSize.Y)
+		{
+			if (ComponentsPerTexture.Y > 1.0f)
+			{
+				UVComponentSectionBase.Y = PositionOffset.Y * LocalComponentSizeVerts;
+			}
+			else
+			{
+				UVComponentSectionBase.Y -= (UVComponentSectionBase.Y + LocalComponentSizeQuad > InWriteSize.Y) ? FMath::FloorToInt((PositionOffset.Y / ComponentsPerTexture.Y)) * ComponentsPerTexture.Y * LocalComponentSizeQuad : 0;
+			}
+		}
+
+		ComponentSectionBase.Y -= (ComponentSectionBase.Y + LocalComponentSizeQuad > InWriteSize.Y) ? FMath::FloorToInt((PositionOffset.Y / ComponentsPerTexture.Y)) * ComponentsPerTexture.Y * LocalComponentSizeQuad : 0;
+		PositionOffset.Y = ComponentSectionBase.Y / LocalComponentSizeQuad;
+	}
+
+	ComponentSectionBase.X = PositionOffset.X * LocalComponentSizeVerts;
+	ComponentSectionBase.Y = PositionOffset.Y * LocalComponentSizeVerts;
+
+	FVector2D UVSize((float)SubsectionSizeVerts / (float)InReadSize.X, (float)SubsectionSizeVerts / (float)InReadSize.Y);
+	FIntPoint SubSectionSectionBase;
+
+	for (int8 SubY = 0; SubY < NumSubsections; ++SubY)
+	{
+		for (int8 SubX = 0; SubX < NumSubsections; ++SubX)
+		{
+			SubSectionSectionBase.X = ComponentSectionBase.X + SubsectionSizeVerts * SubX;
+			SubSectionSectionBase.Y = ComponentSectionBase.Y + SubsectionSizeVerts * SubY;
+
+			// Offset for this component's data in texture
+			FVector2D UVStart;
+
+			if (InReadSize.X >= InWriteSize.X)
+			{
+				UVStart.X = ((float)UVComponentSectionBase.X / (float)InReadSize.X) + UVSize.X * (float)SubX;
+			}
+			else
+			{
+				UVStart.X = InScaleBias.X + UVSize.X * (float)SubX;
+			}
+
+			if (InReadSize.Y >= InWriteSize.Y)
+			{
+				UVStart.Y = ((float)UVComponentSectionBase.Y / (float)InReadSize.Y) + UVSize.Y * (float)SubY;
+			}
+			else
+			{
+				UVStart.Y = InScaleBias.Y + UVSize.Y * (float)SubY;
+			}
+
+			GenerateLayersRenderQuad(SubSectionSectionBase, SubsectionSizeVerts, UVStart, UVSize, OutTriangles);
+		}
+	}
+}
+
+void ALandscape::GenerateLayersRenderQuadsMip(const FIntPoint& InSectionBase, const FVector2D& InScaleBias, float InSubSectionSizeQuad, const FIntPoint& InReadSize, const FIntPoint& InWriteSize, uint8 InCurrentMip, TArray<FLandscapeLayersTriangle>& OutTriangles) const
+{
+	int32 SubsectionSizeVerts = InSubSectionSizeQuad + 1;
+	int32 LocalComponentSizeQuad = InSubSectionSizeQuad * NumSubsections;
+	int32 LocalComponentSizeVerts = SubsectionSizeVerts * NumSubsections;	
+	int32 MipSubsectionSizeVerts = SubsectionSizeVerts >> InCurrentMip;
+	int32 MipLocalComponentSizeVerts = MipSubsectionSizeVerts * NumSubsections;	
+	
+	FVector2D PositionOffset(FMath::RoundToInt(InSectionBase.X / LocalComponentSizeQuad), FMath::RoundToInt(InSectionBase.Y / LocalComponentSizeQuad));
+	FVector2D ComponentsPerTexture(FMath::RoundToInt(InWriteSize.X / LocalComponentSizeQuad), FMath::RoundToInt(InWriteSize.Y / LocalComponentSizeQuad));
+
+	FIntPoint ComponentSectionBase(PositionOffset.X * MipLocalComponentSizeVerts, PositionOffset.Y * MipLocalComponentSizeVerts);
+	FIntPoint UVComponentSectionBase(PositionOffset.X * LocalComponentSizeVerts, PositionOffset.Y * LocalComponentSizeVerts);
+	FVector2D UVSize((float)(SubsectionSizeVerts >> (InCurrentMip - 1)) / (float)InReadSize.X, (float)(SubsectionSizeVerts >> (InCurrentMip - 1)) / (float)InReadSize.Y);
+	FIntPoint SubSectionSectionBase;
+
+	for (int8 SubY = 0; SubY < NumSubsections; ++SubY)
+	{
+		for (int8 SubX = 0; SubX < NumSubsections; ++SubX)
+		{
+			SubSectionSectionBase.X = ComponentSectionBase.X + MipSubsectionSizeVerts * SubX;
+			SubSectionSectionBase.Y = ComponentSectionBase.Y + MipSubsectionSizeVerts * SubY;
+
+			// Offset for this component's data in texture
+			FVector2D UVStart(((float)(UVComponentSectionBase.X >> (InCurrentMip - 1)) / (float)InReadSize.X) + UVSize.X * (float)SubX, ((float)(UVComponentSectionBase.Y >> (InCurrentMip - 1)) / (float)InReadSize.Y) + UVSize.Y * (float)SubY);
+
+			GenerateLayersRenderQuad(SubSectionSectionBase, MipSubsectionSizeVerts, UVStart, UVSize, OutTriangles);
+		}
+	}
+}
+
+void ALandscape::GenerateLayersRenderQuadsAtlasToNonAtlas(const FIntPoint& InSectionBase, const FVector2D& InScaleBias, float InSubSectionSizeQuad, const FIntPoint& InReadSize, const FIntPoint& InWriteSize, TArray<struct FLandscapeLayersTriangle>& OutTriangles) const
+{
+	int32 SubsectionSizeVerts = InSubSectionSizeQuad + 1;
+	int32 LocalComponentSizeQuad = InSubSectionSizeQuad * NumSubsections;
+	int32 LocalComponentSizeVerts = SubsectionSizeVerts * NumSubsections;
+
+	FVector2D PositionOffset(FMath::RoundToInt(InSectionBase.X / LocalComponentSizeQuad), FMath::RoundToInt(InSectionBase.Y / LocalComponentSizeQuad));
+	FIntPoint ComponentSectionBase(PositionOffset.X * LocalComponentSizeQuad, PositionOffset.Y * LocalComponentSizeQuad);
+	FIntPoint UVComponentSectionBase(PositionOffset.X * SubsectionSizeVerts, PositionOffset.Y * SubsectionSizeVerts);
+	FVector2D UVSize((float)SubsectionSizeVerts / (float)InReadSize.X, (float)SubsectionSizeVerts / (float)InReadSize.Y);
+
+	FIntPoint SubSectionSectionBase;
+
+	for (int8 SubY = 0; SubY < NumSubsections; ++SubY)
+	{
+		for (int8 SubX = 0; SubX < NumSubsections; ++SubX)
+		{
+			SubSectionSectionBase.X = ComponentSectionBase.X + InSubSectionSizeQuad * SubX;
+			SubSectionSectionBase.Y = ComponentSectionBase.Y + InSubSectionSizeQuad * SubY;
+
+			// Offset for this component's data in texture
+			FVector2D UVStart;
+
+			if (InReadSize.X >= InWriteSize.X)
+			{
+				UVStart.X = ((float)UVComponentSectionBase.X / (float)InReadSize.X) + UVSize.X * (float)SubX;
+			}
+			else
+			{
+				UVStart.X = InScaleBias.X + UVSize.X * (float)SubX;
+			}
+
+			if (InReadSize.Y >= InWriteSize.Y)
+			{
+				UVStart.Y = ((float)UVComponentSectionBase.Y / (float)InReadSize.Y) + UVSize.Y * (float)SubY;
+			}
+			else
+			{
+				UVStart.Y = InScaleBias.Y + UVSize.Y * (float)SubY;
+			}
+
+			GenerateLayersRenderQuad(SubSectionSectionBase, SubsectionSizeVerts, UVStart, UVSize, OutTriangles);
+		}
+	}
+}
+
+void ALandscape::GenerateLayersRenderQuadsNonAtlas(const FIntPoint& InSectionBase, const FVector2D& InScaleBias, float InSubSectionSizeQuad, const FIntPoint& InReadSize, const FIntPoint& InWriteSize, TArray<struct FLandscapeLayersTriangle>& OutTriangles) const
+{
+	// We currently only support drawing in non atlas mode with the same texture size
+	check(InReadSize.X == InWriteSize.X && InReadSize.Y == InWriteSize.Y);
+
+	int32 SubsectionSizeVerts = InSubSectionSizeQuad + 1;
+	int32 LocalComponentSizeQuad = InSubSectionSizeQuad * NumSubsections;
+
+	FVector2D PositionOffset(FMath::RoundToInt(InSectionBase.X / LocalComponentSizeQuad), FMath::RoundToInt(InSectionBase.Y / LocalComponentSizeQuad));
+	FIntPoint ComponentSectionBase = InSectionBase;
+	FIntPoint UVComponentSectionBase(PositionOffset.X * LocalComponentSizeQuad, PositionOffset.Y * LocalComponentSizeQuad);
+	FVector2D UVSize((float)SubsectionSizeVerts / (float)InReadSize.X, (float)SubsectionSizeVerts / (float)InReadSize.Y);
+	FIntPoint SubSectionSectionBase;
+
+	for (int8 SubY = 0; SubY < NumSubsections; ++SubY)
+	{
+		for (int8 SubX = 0; SubX < NumSubsections; ++SubX)
+		{
+			SubSectionSectionBase.X = ComponentSectionBase.X + InSubSectionSizeQuad * SubX;
+			SubSectionSectionBase.Y = ComponentSectionBase.Y + InSubSectionSizeQuad * SubY;
+
+			// Offset for this component's data in texture
+			FVector2D UVStart(((float)UVComponentSectionBase.X / (float)InReadSize.X) + UVSize.X * (float)SubX, ((float)UVComponentSectionBase.Y / (float)InReadSize.Y) + UVSize.Y * (float)SubY);
+			GenerateLayersRenderQuad(SubSectionSectionBase, SubsectionSizeVerts, UVStart, UVSize, OutTriangles);
+		}
+	}
+}
+
+void ALandscape::GenerateLayersRenderQuadsNonAtlasToAtlas(const FIntPoint& InSectionBase, const FVector2D& InScaleBias, float InSubSectionSizeQuad, const FIntPoint& InReadSize, const FIntPoint& InWriteSize, TArray<struct FLandscapeLayersTriangle>& OutTriangles) const
+{
+	int32 SubsectionSizeVerts = InSubSectionSizeQuad + 1;
+	int32 LocalComponentSizeQuad = InSubSectionSizeQuad * NumSubsections;
+	int32 LocalComponentSizeVerts = SubsectionSizeVerts * NumSubsections;
+
+	FVector2D PositionOffset(FMath::RoundToInt(InSectionBase.X / LocalComponentSizeQuad), FMath::RoundToInt(InSectionBase.Y / LocalComponentSizeQuad));
+	FIntPoint ComponentSectionBase(PositionOffset.X * LocalComponentSizeVerts, PositionOffset.Y * LocalComponentSizeVerts);
+	FVector2D UVSize((float)SubsectionSizeVerts / (float)InReadSize.X, (float)SubsectionSizeVerts / (float)InReadSize.Y);
+
+	FIntPoint SubSectionSectionBase;
+
+	for (int8 SubY = 0; SubY < NumSubsections; ++SubY)
+	{
+		for (int8 SubX = 0; SubX < NumSubsections; ++SubX)
+		{
+			SubSectionSectionBase.X = ComponentSectionBase.X + SubsectionSizeVerts * SubX;
+			SubSectionSectionBase.Y = ComponentSectionBase.Y + SubsectionSizeVerts * SubY;
+
+			// Offset for this component's data in texture
+			float ScaleBiasZ = (float)InSectionBase.X / (float)InReadSize.X;
+			float ScaleBiasW = (float)InSectionBase.Y / (float)InReadSize.Y;
+			FVector2D UVStart(ScaleBiasZ + ((float)InSubSectionSizeQuad / (float)InReadSize.X) * (float)SubX, ScaleBiasW + ((float)InSubSectionSizeQuad / (float)InReadSize.Y) * (float)SubY);
+
+			GenerateLayersRenderQuad(SubSectionSectionBase, SubsectionSizeVerts, UVStart, UVSize, OutTriangles);
+		}
+	}
+}
+
+void ALandscape::PrintLayersDebugHeightData(const FString& InContext, const TArray<FColor>& InHeightmapData, const FIntPoint& InDataSize, uint8 InMipRender, bool InOutputNormals) const
+{
+	bool DisplayDebugPrint = CVarOutputLayersRTContent.GetValueOnAnyThread() == 1 ? true : false;
+	bool DisplayHeightAsDelta = false;
+
+	if (!DisplayDebugPrint)
+	{
+		return;
+	}
+
+	TArray<uint16> HeightData;
+	TArray<FVector> NormalData;
+	HeightData.Reserve(InHeightmapData.Num());
+	NormalData.Reserve(InHeightmapData.Num());
+
+	for (FColor Color : InHeightmapData)
+	{
+		uint16 Height = ((Color.R << 8) | Color.G);
+		HeightData.Add(Height);
+
+		if (InOutputNormals)
+		{
+			FVector Normal;
+			Normal.X = Color.B > 0.0f ? ((float)Color.B / 127.5f - 1.0f) : 0.0f;
+			Normal.Y = Color.A > 0.0f ? ((float)Color.A / 127.5f - 1.0f) : 0.0f;
+			Normal.Z = 0.0f;
+
+			NormalData.Add(Normal);
+		}
+	}
+
+	UE_LOG(LogLandscapeBP, Display, TEXT("Context: %s"), *InContext);
+
+	int32 MipSize = ((SubsectionSizeQuads + 1) >> InMipRender);
+
+	for (int32 Y = 0; Y < InDataSize.Y; ++Y)
+	{
+		FString HeightmapHeightOutput;
+
+		for (int32 X = 0; X < InDataSize.X; ++X)
+		{
+			int32 HeightDelta = HeightData[X + Y * InDataSize.X];
+
+			if (DisplayHeightAsDelta)
+			{
+				HeightDelta = HeightDelta >= 32768 ? HeightDelta - 32768 : HeightDelta;
+			}
+
+			if (X > 0 && MipSize > 0 && X % MipSize == 0)
+			{
+				HeightmapHeightOutput += FString::Printf(TEXT("  "));
+			}
+
+			FString HeightStr = FString::Printf(TEXT("%d"), HeightDelta);
+
+			int32 PadCount = 5 - HeightStr.Len();
+			if (PadCount > 0)
+			{
+				HeightStr = FString::ChrN(PadCount, '0') + HeightStr;
+			}
+
+			HeightmapHeightOutput += HeightStr + TEXT(" ");
+		}
+
+		if (Y > 0 && MipSize > 0 && Y % MipSize == 0)
+		{
+			UE_LOG(LogLandscapeBP, Display, TEXT(""));
+		}
+
+		UE_LOG(LogLandscapeBP, Display, TEXT("%s"), *HeightmapHeightOutput);
+	}
+
+	if (InOutputNormals)
+	{
+		UE_LOG(LogLandscapeBP, Display, TEXT(""));
+
+		for (int32 Y = 0; Y < InDataSize.Y; ++Y)
+		{
+			FString HeightmapNormaltOutput;
+
+			for (int32 X = 0; X < InDataSize.X; ++X)
+			{
+				FVector Normal = NormalData[X + Y * InDataSize.X];
+
+				if (X > 0 && MipSize > 0 && X % MipSize == 0)
+				{
+					HeightmapNormaltOutput += FString::Printf(TEXT("  "));
+				}
+
+				HeightmapNormaltOutput += FString::Printf(TEXT(" %s"), *Normal.ToString());
+			}
+
+			if (Y > 0 && MipSize > 0 && Y % MipSize == 0)
+			{
+				UE_LOG(LogLandscapeBP, Display, TEXT(""));
+			}
+
+			UE_LOG(LogLandscapeBP, Display, TEXT("%s"), *HeightmapNormaltOutput);
+		}
+	}
+}
+
+void ALandscape::PrintLayersDebugWeightData(const FString& InContext, const TArray<FColor>& InWeightmapData, const FIntPoint& InDataSize, uint8 InMipRender) const
+{
+	bool DisplayDebugPrint = (CVarOutputLayersRTContent.GetValueOnAnyThread() == 1 || CVarOutputLayersWeightmapsRTContent.GetValueOnAnyThread() == 1) ? true : false;
+
+	if (!DisplayDebugPrint)
+	{
+		return;
+	}
+
+	UE_LOG(LogLandscapeBP, Display, TEXT("Context: %s"), *InContext);
+
+	int32 MipSize = ((SubsectionSizeQuads + 1) >> InMipRender);
+
+	for (int32 Y = 0; Y < InDataSize.Y; ++Y)
+	{
+		FString WeightmapOutput;
+
+		for (int32 X = 0; X < InDataSize.X; ++X)
+		{
+			const FColor& Weight = InWeightmapData[X + Y * InDataSize.X];
+
+			if (X > 0 && MipSize > 0 && X % MipSize == 0)
+			{
+				WeightmapOutput += FString::Printf(TEXT("  "));
+			}
+
+			WeightmapOutput += FString::Printf(TEXT("%s "), *Weight.ToString());
+		}
+
+		if (Y > 0 && MipSize > 0 && Y % MipSize == 0)
+		{
+			UE_LOG(LogLandscapeBP, Display, TEXT(""));
+		}
+
+		UE_LOG(LogLandscapeBP, Display, TEXT("%s"), *WeightmapOutput);
+	}
+}
+
+void ALandscape::PrintLayersDebugRT(const FString& InContext, UTextureRenderTarget2D* InDebugRT, uint8 InMipRender, bool InOutputHeight, bool InOutputNormals) const
+{
+	bool DisplayDebugPrint = (CVarOutputLayersRTContent.GetValueOnAnyThread() == 1 || CVarOutputLayersWeightmapsRTContent.GetValueOnAnyThread() == 1) ? true : false;
+
+	if (!DisplayDebugPrint)
+	{
+		return;
+	}
+
+	FTextureRenderTargetResource* RenderTargetResource = InDebugRT->GameThread_GetRenderTargetResource();
+	ENQUEUE_RENDER_COMMAND(FLandscapeLayersDebugRenderTargetResolveCommand)(
+		[RenderTargetResource](FRHICommandListImmediate& RHICmdList) mutable
+		{
+			// Copy (resolve) the rendered image from the frame buffer to its render target texture
+			RHICmdList.CopyToResolveTarget(RenderTargetResource->GetRenderTargetTexture(), RenderTargetResource->TextureRHI, FResolveParams());
+		});
+
+	FlushRenderingCommands();
+	int32 MinX, MinY, MaxX, MaxY;
+	const ULandscapeInfo* LandscapeInfo = GetLandscapeInfo();
+	LandscapeInfo->GetLandscapeExtent(MinX, MinY, MaxX, MaxY);
+	FIntRect SampleRect = FIntRect(0, 0, InDebugRT->SizeX, InDebugRT->SizeY);
+
+	FReadSurfaceDataFlags Flags(RCM_UNorm, CubeFace_MAX);
+
+	TArray<FColor> OutputRT;
+	OutputRT.Reserve(SampleRect.Width() * SampleRect.Height());
+
+	InDebugRT->GameThread_GetRenderTargetResource()->ReadPixels(OutputRT, Flags, SampleRect);
+
+	if (InOutputHeight)
+	{
+		PrintLayersDebugHeightData(InContext, OutputRT, FIntPoint(SampleRect.Width(), SampleRect.Height()), InMipRender, InOutputNormals);
+	}
+	else
+	{
+		PrintLayersDebugWeightData(InContext, OutputRT, FIntPoint(SampleRect.Width(), SampleRect.Height()), InMipRender);
+	}
+}
+
+void ALandscape::PrintLayersDebugTextureResource(const FString& InContext, FTextureResource* InTextureResource, uint8 InMipRender, bool InOutputHeight, bool InOutputNormals) const
+{
+	bool DisplayDebugPrint = (CVarOutputLayersRTContent.GetValueOnAnyThread() == 1 || CVarOutputLayersWeightmapsRTContent.GetValueOnAnyThread() == 1) ? true : false;
+
+	if (!DisplayDebugPrint)
+	{
+		return;
+	}
+
+	int32 MinX, MinY, MaxX, MaxY;
+	const ULandscapeInfo* LandscapeInfo = GetLandscapeInfo();
+	LandscapeInfo->GetLandscapeExtent(MinX, MinY, MaxX, MaxY);
+	FIntRect SampleRect = FIntRect(0, 0, InTextureResource->GetSizeX(), InTextureResource->GetSizeY());
+
+	TArray<FColor> OutputTexels;
+	OutputTexels.Reserve(SampleRect.Width() * SampleRect.Height());
+
+	FReadSurfaceDataFlags Flags(RCM_UNorm, CubeFace_MAX);
+	Flags.SetMip(InMipRender);
+
+	ENQUEUE_RENDER_COMMAND(FLandscapeLayersDebugReadSurfaceCommand)(
+		[SourceTextureRHI = InTextureResource->TextureRHI, SampleRect = SampleRect, OutData = &OutputTexels, ReadFlags = Flags](FRHICommandListImmediate& RHICmdList) mutable
+	{
+		RHICmdList.ReadSurfaceData(SourceTextureRHI, SampleRect, *OutData, ReadFlags);
+	});
+
+	FlushRenderingCommands();
+
+	if (InOutputHeight)
+	{
+		PrintLayersDebugHeightData(InContext, OutputTexels, FIntPoint(SampleRect.Width(), SampleRect.Height()), InMipRender, InOutputNormals);
+	}
+	else
+	{
+		PrintLayersDebugWeightData(InContext, OutputTexels, FIntPoint(SampleRect.Width(), SampleRect.Height()), InMipRender);
+	}
+}
+
+bool ALandscape::AreLayersHeightmapTextureResourcesReady(const TArray<ALandscapeProxy*>& InAllLandscapes) const
+{
+	for (ALandscapeProxy* Landscape : InAllLandscapes)
+	{
+		for (auto& ItLayerDataPair : Landscape->LandscapeLayersData)
+		{
+			for (auto& ItHeightmapPair : ItLayerDataPair.Value.Heightmaps)
+			{
+				UTexture2D* OriginalHeightmap = ItHeightmapPair.Key;
+				UTexture2D* LayerHeightmap = ItHeightmapPair.Value;
+
+				if (!LayerHeightmap->IsAsyncCacheComplete() || !OriginalHeightmap->IsFullyStreamedIn())
+				{
+					return false;
+				}
+
+				if (LayerHeightmap->Resource == nullptr)
+				{
+					LayerHeightmap->FinishCachePlatformData();
+
+					LayerHeightmap->Resource = LayerHeightmap->CreateResource();
+
+					if (LayerHeightmap->Resource != nullptr)
+					{
+						BeginInitResource(LayerHeightmap->Resource);
+					}
+				}
+
+				if (LayerHeightmap->Resource == nullptr || !LayerHeightmap->Resource->IsInitialized() || !LayerHeightmap->IsFullyStreamedIn())
+				{
+					return false;
+				}
+			}
+		}
+	}
+
+	return true;
+}
+
+void ALandscape::RegenerateLayersHeightmaps()
+{
+	SCOPE_CYCLE_COUNTER(STAT_LandscapeLayersRegenerateHeightmaps);
+
+	ULandscapeInfo* Info = GetLandscapeInfo();
+
+	if (LayersContentUpdateFlags == 0 || Info == nullptr)
+	{
+		return;
+	}
+
+	TArray<ALandscapeProxy*> AllLandscapes;
+	AllLandscapes.Add(this);
+
+	for (auto& It : Info->Proxies)
+	{
+		AllLandscapes.Add(It);
+	}
+
+	if (!AreLayersHeightmapTextureResourcesReady(AllLandscapes))
+	{
+		return;
+	}
+
+	TArray<ULandscapeComponent*> AllLandscapeComponents;
+
+	for (ALandscapeProxy* Landscape : AllLandscapes)
+	{
+		AllLandscapeComponents.Append(Landscape->LandscapeComponents);
+	}
+
+	if ((LayersContentUpdateFlags & ELandscapeLayersContentUpdateFlag::Heightmap_Render) != 0 && HeightmapRTList.Num() > 0)
+	{
+		FLandscapeLayersHeightmapShaderParameters ShaderParams;
+
+		bool FirstLayer = true;
+		UTextureRenderTarget2D* CombinedHeightmapAtlasRT = HeightmapRTList[EHeightmapRTType::HeightmapRT_CombinedAtlas];
+		UTextureRenderTarget2D* CombinedHeightmapNonAtlasRT = HeightmapRTList[EHeightmapRTType::HeightmapRT_CombinedNonAtlas];
+		UTextureRenderTarget2D* LandscapeScratchRT1 = HeightmapRTList[EHeightmapRTType::HeightmapRT_Scratch1];
+		UTextureRenderTarget2D* LandscapeScratchRT2 = HeightmapRTList[EHeightmapRTType::HeightmapRT_Scratch2];
+		UTextureRenderTarget2D* LandscapeScratchRT3 = HeightmapRTList[EHeightmapRTType::HeightmapRT_Scratch3];
+
+		bool OutputDebugName = (CVarOutputLayersDebugDrawCallName.GetValueOnAnyThread() == 1 || CVarOutputLayersRTContent.GetValueOnAnyThread() == 1) ? true : false;
+
+		for (FLandscapeLayer& Layer : LandscapeLayers)
+		{
+			//Draw Layer heightmap to Combined RT Atlas
+			ShaderParams.ApplyLayerModifiers = true;
+			ShaderParams.LayerVisible = Layer.bVisible;
+			ShaderParams.LayerAlpha = Layer.HeightmapAlpha;
+
+			for (ALandscapeProxy* Landscape : AllLandscapes)
+			{
+				FLandscapeLayerData* LayerData = Landscape->LandscapeLayersData.Find(Layer.Guid);
+
+				if (LayerData != nullptr)
+				{
+					for (auto& ItPair : LayerData->Heightmaps)
+					{
+						FRenderDataPerHeightmap& HeightmapRenderData = *Landscape->RenderDataPerHeightmap.Find(ItPair.Key);
+						UTexture2D* Heightmap = ItPair.Value;
+
+						CopyLayersTexture(Heightmap, LandscapeScratchRT1, nullptr, HeightmapRenderData.TopLeftSectionBase);
+
+						PrintLayersDebugRT(OutputDebugName ? FString::Printf(TEXT("LS Height: %s Component %s += -> CombinedAtlas %s"), *Layer.Name.ToString(), *Heightmap->GetName(), *LandscapeScratchRT1->GetName()) : TEXT(""), LandscapeScratchRT1);
+					}
+				}
+			}
+
+			// NOTE: From this point on, we always work in non atlas, we'll convert back at the end to atlas only
+			DrawHeightmapComponentsToRenderTarget(OutputDebugName ? FString::Printf(TEXT("LS Height: %s += -> NonAtlas %s"), *Layer.Name.ToString(), *LandscapeScratchRT1->GetName(), *LandscapeScratchRT2->GetName()) : TEXT(""),
+												  AllLandscapeComponents, LandscapeScratchRT1, nullptr, LandscapeScratchRT2, ERTDrawingType::RTAtlasToNonAtlas, true, ShaderParams);
+
+			// Combine Current layer with current result
+			DrawHeightmapComponentsToRenderTarget(OutputDebugName ? FString::Printf(TEXT("LS Height: %s += -> CombinedNonAtlas %s"), *Layer.Name.ToString(), *LandscapeScratchRT2->GetName(), *CombinedHeightmapNonAtlasRT->GetName()) : TEXT(""),
+												AllLandscapeComponents, LandscapeScratchRT2, FirstLayer ? nullptr : LandscapeScratchRT3, CombinedHeightmapNonAtlasRT, ERTDrawingType::RTNonAtlas, FirstLayer, ShaderParams);
+
+			ShaderParams.ApplyLayerModifiers = false;
+
+			if (Layer.bVisible)
+			{
+				// Draw each Combined RT into a Non Atlas RT format to be use as base for all brush rendering
+				if (Layer.Brushes.Num() > 0)
+				{
+					CopyLayersTexture(CombinedHeightmapNonAtlasRT, LandscapeScratchRT1);
+					PrintLayersDebugRT(OutputDebugName ? FString::Printf(TEXT("LS Height: %s Component %s += -> CombinedNonAtlas %s"), *Layer.Name.ToString(), *CombinedHeightmapNonAtlasRT->GetName(), *LandscapeScratchRT1->GetName()) : TEXT(""), LandscapeScratchRT1);
+				}
+
+				// Draw each brushes				
+				for (int32 i = 0; i < Layer.HeightmapBrushOrderIndices.Num(); ++i)
+				{
+					// TODO: handle conversion from float to RG8 by using material params to write correct values
+					// TODO: handle conversion/handling of RT not same size as internal size
+
+					FLandscapeLayerBrush& Brush = Layer.Brushes[Layer.HeightmapBrushOrderIndices[i]];
+
+					if (Brush.BPCustomBrush == nullptr || !Brush.BPCustomBrush->IsAffectingHeightmap())
+					{
+						continue;
+					}
+
+					if (!Brush.IsInitialized())
+					{
+						Brush.Initialize(GetBoundingRect(), FIntPoint(CombinedHeightmapNonAtlasRT->SizeX, CombinedHeightmapNonAtlasRT->SizeY));
+					}
+
+					UTextureRenderTarget2D* BrushOutputNonAtlasRT = Brush.Render(true, CombinedHeightmapNonAtlasRT);
+
+					if (BrushOutputNonAtlasRT == nullptr || BrushOutputNonAtlasRT->SizeX != CombinedHeightmapNonAtlasRT->SizeX || BrushOutputNonAtlasRT->SizeY != CombinedHeightmapNonAtlasRT->SizeY)
+					{
+						continue;
+					}
+
+					INC_DWORD_STAT(STAT_LandscapeLayersRegenerateDrawCalls); // Brush Render
+
+					PrintLayersDebugRT(OutputDebugName ? FString::Printf(TEXT("LS Height: %s %s -> BrushNonAtlas %s"), *Layer.Name.ToString(), *Brush.BPCustomBrush->GetName(), *BrushOutputNonAtlasRT->GetName()) : TEXT(""), BrushOutputNonAtlasRT);
+
+					// Resolve back to Combined heightmap
+					CopyLayersTexture(BrushOutputNonAtlasRT, CombinedHeightmapNonAtlasRT);
+					PrintLayersDebugRT(OutputDebugName ? FString::Printf(TEXT("LS Height: %s Component %s += -> CombinedNonAtlas %s"), *Layer.Name.ToString(), *BrushOutputNonAtlasRT->GetName(), *CombinedHeightmapNonAtlasRT->GetName()) : TEXT(""), CombinedHeightmapNonAtlasRT);
+				}
+			}
+
+			CopyLayersTexture(CombinedHeightmapNonAtlasRT, LandscapeScratchRT3);
+			PrintLayersDebugRT(OutputDebugName ? FString::Printf(TEXT("LS Height: %s Component %s += -> CombinedNonAtlas %s"), *Layer.Name.ToString(), *CombinedHeightmapNonAtlasRT->GetName(), *LandscapeScratchRT3->GetName()) : TEXT(""), LandscapeScratchRT3);
+
+			FirstLayer = false;
+		}
+
+		ShaderParams.GenerateNormals = true;
+		ShaderParams.GridSize = GetRootComponent()->RelativeScale3D;
+
+		DrawHeightmapComponentsToRenderTarget(OutputDebugName ? FString::Printf(TEXT("LS Height: %s = -> CombinedNonAtlasNormals : %s"), *CombinedHeightmapNonAtlasRT->GetName(), *LandscapeScratchRT1->GetName()) : TEXT(""),
+											  AllLandscapeComponents, CombinedHeightmapNonAtlasRT, nullptr, LandscapeScratchRT1, ERTDrawingType::RTNonAtlas, true, ShaderParams);
+
+		ShaderParams.GenerateNormals = false;
+
+		DrawHeightmapComponentsToRenderTarget(OutputDebugName ? FString::Printf(TEXT("LS Height: %s = -> CombinedAtlasFinal : %s"), *LandscapeScratchRT1->GetName(), *CombinedHeightmapAtlasRT->GetName()) : TEXT(""),
+											  AllLandscapeComponents, LandscapeScratchRT1, nullptr, CombinedHeightmapAtlasRT, ERTDrawingType::RTNonAtlasToAtlas, true, ShaderParams);
+
+		DrawHeightmapComponentsToRenderTargetMips(AllLandscapeComponents, CombinedHeightmapAtlasRT, true, ShaderParams);
+
+		// Copy back all Mips to original heightmap data
+		for (ALandscapeProxy* Landscape : AllLandscapes)
+		{
+			for (auto& ItPair : Landscape->RenderDataPerHeightmap)
+			{
+				int32 CurrentMip = 0;
+				FRenderDataPerHeightmap& HeightmapRenderData = ItPair.Value;
+
+				CopyLayersTexture(CombinedHeightmapAtlasRT, HeightmapRenderData.OriginalHeightmap, HeightmapRenderData.HeightmapsCPUReadBack, HeightmapRenderData.TopLeftSectionBase, CurrentMip, CurrentMip);
+				++CurrentMip;
+
+				for (int32 MipRTIndex = EHeightmapRTType::HeightmapRT_Mip1; MipRTIndex < EHeightmapRTType::HeightmapRT_Count; ++MipRTIndex)
+				{
+					if (HeightmapRTList[MipRTIndex] != nullptr)
+					{
+						CopyLayersTexture(HeightmapRTList[MipRTIndex], HeightmapRenderData.OriginalHeightmap, HeightmapRenderData.HeightmapsCPUReadBack, HeightmapRenderData.TopLeftSectionBase, CurrentMip, CurrentMip);
+						++CurrentMip;
+					}
+				}
+			}
+		}
+	}
+
+	if ((LayersContentUpdateFlags & ELandscapeLayersContentUpdateFlag::Heightmap_ResolveToTexture) != 0)
+	{
+		ResolveLayersHeightmapTexture(AllLandscapes);
+	}
+
+	if ((LayersContentUpdateFlags & ELandscapeLayersContentUpdateFlag::Heightmap_BoundsAndCollision) != 0)
+	{
+		for (ULandscapeComponent* Component : AllLandscapeComponents)
+		{
+			Component->UpdateCachedBounds();
+			Component->UpdateComponentToWorld();
+
+			Component->UpdateCollisionData(false);
+		}
+	}
+
+	LayersContentUpdateFlags &= ~ELandscapeLayersContentUpdateFlag::Heightmap_All;
+
+	// If doing rendering debug, keep doing the render only
+	if (CVarOutputLayersDebugDrawCallName.GetValueOnAnyThread() == 1)
+	{
+		LayersContentUpdateFlags |= ELandscapeLayersContentUpdateFlag::Heightmap_Render;
+	}
+}
+
+void ALandscape::ResolveLayersHeightmapTexture(const TArray<ALandscapeProxy*>& InAllLandscapes)
+{
+	SCOPE_CYCLE_COUNTER(STAT_LandscapeLayersResolveHeightmaps);
+
+	for (ALandscapeProxy* Landscape : InAllLandscapes)
+	{
+		TArray<TArray<FColor>> MipData;
+
+		for (auto& ItPair : Landscape->RenderDataPerHeightmap)
+		{
+			FRenderDataPerHeightmap& HeightmapRenderData = ItPair.Value;
+
+			if (HeightmapRenderData.HeightmapsCPUReadBack == nullptr)
+			{
+				continue;
+			}
+
+			ResolveLayersTexture(HeightmapRenderData.HeightmapsCPUReadBack, HeightmapRenderData.OriginalHeightmap);
+		}
+	}		
+}
+
+void ALandscape::ResolveLayersTexture(FLandscapeLayersTexture2DCPUReadBackResource* InCPUReadBackTexture, UTexture2D* InOriginalTexture)
+{
+	TArray<TArray<FColor>> MipData;
+	MipData.AddDefaulted(InCPUReadBackTexture->TextureRHI->GetNumMips());
+
+	int32 MipSizeU = InCPUReadBackTexture->GetSizeX();
+	int32 MipSizeV = InCPUReadBackTexture->GetSizeY();
+	int32 MipIndex = 0;
+
+	while (MipSizeU >= 1 && MipSizeV >= 1)
+	{
+		MipData[MipIndex].Reset();
+
+		FReadSurfaceDataFlags Flags(RCM_UNorm, CubeFace_MAX);
+		Flags.SetMip(MipIndex);
+		FIntRect Rect(0, 0, MipSizeU, MipSizeV);
+
+		ENQUEUE_RENDER_COMMAND(FLandscapeLayersReadSurfaceCommand)(
+			[SourceTextureRHI = InCPUReadBackTexture->TextureRHI, Rect = Rect, OutData = &MipData[MipIndex], ReadFlags = Flags](FRHICommandListImmediate& RHICmdList) mutable
+		{
+			RHICmdList.ReadSurfaceData(SourceTextureRHI, Rect, *OutData, ReadFlags);
+		});
+
+
+		MipSizeU >>= 1;
+		MipSizeV >>= 1;
+		++MipIndex;
+	}
+
+	// TODO: find a way to NOT have to flush the rendering command as this create hic up of ~10-15ms
+	FlushRenderingCommands();
+
+	for (MipIndex = 0; MipIndex < MipData.Num(); ++MipIndex)
+	{
+		if (MipData[MipIndex].Num() > 0)
+		{
+			FColor* TextureData = (FColor*)InOriginalTexture->Source.LockMip(MipIndex);
+			FMemory::Memcpy(TextureData, MipData[MipIndex].GetData(), MipData[MipIndex].Num() * sizeof(FColor));
+			InOriginalTexture->Source.UnlockMip(MipIndex);
+		}
+	}
+}
+
+void ALandscape::PrepareComponentDataToExtractMaterialLayersCS(const FLandscapeLayer& InLayer, int32 InCurrentWeightmapToProcessIndex, bool InOutputDebugName, const TArray<ALandscapeProxy*>& InAllLandscape, FLandscapeTexture2DResource* InOutTextureData,
+																  TArray<FLandscapeLayerWeightmapExtractMaterialLayersComponentData>& OutComponentData, TMap<ULandscapeLayerInfoObject*, int32>& OutLayerInfoObjects)
+{
+	ULandscapeInfo* Info = GetLandscapeInfo();
+	
+	for (const ALandscapeProxy* Landscape : InAllLandscape)
+	{
+		const FLandscapeLayerData* LayerData = Landscape->LandscapeLayersData.Find(InLayer.Guid);
+
+		if (LayerData != nullptr)
+		{
+			for (const auto& ItPair : LayerData->WeightmapData)
+			{
+				ULandscapeComponent* Component = ItPair.Key;
+				const FWeightmapLayerData& WeightLayerData = ItPair.Value;
+
+				if (WeightLayerData.Weightmaps.IsValidIndex(InCurrentWeightmapToProcessIndex))
+				{
+					UTexture2D* Weightmap = WeightLayerData.Weightmaps[InCurrentWeightmapToProcessIndex];
+					check(Weightmap != nullptr);
+
+					const ULandscapeWeightmapUsage* WeightmapUsage = WeightLayerData.WeightmapTextureUsages[InCurrentWeightmapToProcessIndex];
+					check(WeightmapUsage != nullptr);
+
+					CopyLayersTexture(*Weightmap->GetName(), Weightmap->Resource, InOutputDebugName ? FString::Printf(TEXT("%s WeightmapScratchTexture"), *InLayer.Name.ToString()) : TEXT(""), InOutTextureData, nullptr, Component->GetSectionBase(), 0);
+					PrintLayersDebugTextureResource(InOutputDebugName ? FString::Printf(TEXT("LS Weight: %s WeightmapScratchTexture %s"), *InLayer.Name.ToString(), TEXT("WeightmapScratchTextureResource")) : TEXT(""), InOutTextureData, 0, false);
+
+					for (const FWeightmapLayerAllocationInfo& WeightmapLayerAllocation : WeightLayerData.WeightmapLayerAllocations)
+					{
+						if (WeightmapLayerAllocation.LayerInfo != nullptr && WeightmapLayerAllocation.WeightmapTextureIndex != 255 && WeightLayerData.Weightmaps[WeightmapLayerAllocation.WeightmapTextureIndex] == Weightmap)
+						{
+							FLandscapeLayerWeightmapExtractMaterialLayersComponentData Data;
+
+							const ULandscapeComponent* DestComponent = WeightmapUsage->ChannelUsage[WeightmapLayerAllocation.WeightmapTextureChannel];
+							check(DestComponent);
+
+							// Compute component top left vertex position from section base info
+							int32 LocalComponentSizeQuad = Component->SubsectionSizeQuads * NumSubsections;
+							int32 LocalComponentSizeVerts = (Component->SubsectionSizeQuads + 1) * NumSubsections;
+							FVector2D SourcePositionOffset(FMath::RoundToInt(Component->GetSectionBase().X / LocalComponentSizeQuad), FMath::RoundToInt(Component->GetSectionBase().Y / LocalComponentSizeQuad));
+							FVector2D DestPositionOffset(FMath::RoundToInt(DestComponent->GetSectionBase().X / LocalComponentSizeQuad), FMath::RoundToInt(DestComponent->GetSectionBase().Y / LocalComponentSizeQuad));
+
+							Data.ComponentVertexPosition = FIntPoint(SourcePositionOffset.X * LocalComponentSizeVerts, SourcePositionOffset.Y * LocalComponentSizeVerts);
+							Data.AtlasTexturePositionOutput = FIntPoint(DestPositionOffset.X * LocalComponentSizeVerts, DestPositionOffset.Y * LocalComponentSizeVerts);
+							Data.WeightmapChannelToProcess = WeightmapLayerAllocation.WeightmapTextureChannel;
+
+							if (WeightmapLayerAllocation.LayerInfo == ALandscapeProxy::VisibilityLayer)
+							{
+								Data.DestinationPaintLayerIndex = 0;
+								int32& NewLayerinfoObjectIndex = OutLayerInfoObjects.FindOrAdd(ALandscapeProxy::VisibilityLayer);
+								NewLayerinfoObjectIndex = 0;
+							}
+							else
+							{
+								for (int32 LayerInfoSettingsIndex = 0; LayerInfoSettingsIndex < Info->Layers.Num(); ++LayerInfoSettingsIndex)
+								{
+									const FLandscapeInfoLayerSettings& InfoLayerSettings = Info->Layers[LayerInfoSettingsIndex];
+
+									if (InfoLayerSettings.LayerInfoObj == WeightmapLayerAllocation.LayerInfo)
+									{
+										Data.DestinationPaintLayerIndex = LayerInfoSettingsIndex + 1; // due to visibility layer that is at 0
+										int32& NewLayerinfoObjectIndex = OutLayerInfoObjects.FindOrAdd(WeightmapLayerAllocation.LayerInfo);
+										NewLayerinfoObjectIndex = LayerInfoSettingsIndex + 1;
+
+										break;
+									}
+								}
+							}
+
+							OutComponentData.Add(Data);
+						}
+					}
+				}
+			}
+		}
+	}
+}
+
+void ALandscape::PrepareComponentDataToPackMaterialLayersCS(int32 InCurrentWeightmapToProcessIndex, bool InOutputDebugName, const TArray<ULandscapeComponent*>& InAllLandscapeComponents, TArray<UTexture2D*>& InOutProcessedWeightmaps, 
+															  TArray<FLandscapeLayersTexture2DCPUReadBackResource*>& InOutProcessedWeightmapCPUCopy, TArray<FLandscapeLayerWeightmapPackMaterialLayersComponentData>& OutComponentData)
+{
+	ULandscapeInfo* Info = GetLandscapeInfo();
+
+	for (const ULandscapeComponent* Component : InAllLandscapeComponents)
+	{
+		const TArray<UTexture2D*>& WeightmapTextures = Component->GetWeightmapTextures();
+
+		if (WeightmapTextures.IsValidIndex(InCurrentWeightmapToProcessIndex))
+		{
+			UTexture2D* WeightmapTexture = WeightmapTextures[InCurrentWeightmapToProcessIndex];
+
+			if (!InOutProcessedWeightmaps.Contains(WeightmapTexture))
+			{
+				InOutProcessedWeightmaps.Add(WeightmapTexture);
+
+				FLandscapeLayersTexture2DCPUReadBackResource** WeightmapCPUCopy = Component->GetLandscapeProxy()->WeightmapCPUReadBackTextures.Find(WeightmapTexture);
+				check(WeightmapCPUCopy != nullptr);
+
+				InOutProcessedWeightmapCPUCopy.Add(*WeightmapCPUCopy);
+
+				const TArray<ULandscapeWeightmapUsage*>& WeightmapTexturesUsage = Component->GetWeightmapTexturesUsage();
+
+				const ULandscapeWeightmapUsage* WeightmapUsage = WeightmapTexturesUsage[InCurrentWeightmapToProcessIndex];
+				check(WeightmapUsage != nullptr);
+
+				TArray<const FWeightmapLayerAllocationInfo*> AlreadyProcessedAllocation;
+				FLandscapeLayerWeightmapPackMaterialLayersComponentData Data;
+
+				for (int32 WeightmapChannelIndex = 0; WeightmapChannelIndex < 4; ++WeightmapChannelIndex)
+				{
+					// Clear out data to known values
+					Data.ComponentVertexPositionX[WeightmapChannelIndex] = INDEX_NONE;
+					Data.ComponentVertexPositionY[WeightmapChannelIndex] = INDEX_NONE;
+					Data.SourcePaintLayerIndex[WeightmapChannelIndex] = INDEX_NONE;
+					Data.WeightmapChannelToProcess[WeightmapChannelIndex] = INDEX_NONE;
+
+					if (WeightmapUsage->ChannelUsage[WeightmapChannelIndex] != nullptr)
+					{
+						const ULandscapeComponent* ChannelComponent = WeightmapUsage->ChannelUsage[WeightmapChannelIndex];
+						FIntPoint ComponentSectionBase = ChannelComponent->GetSectionBase();
+
+						// Compute component top left vertex position from section base info
+						int32 LocalComponentSizeQuad = ChannelComponent->SubsectionSizeQuads * NumSubsections;
+						int32 LocalComponentSizeVerts = (ChannelComponent->SubsectionSizeQuads + 1) * NumSubsections;
+						FVector2D PositionOffset(FMath::RoundToInt(ChannelComponent->GetSectionBase().X / LocalComponentSizeQuad), FMath::RoundToInt(ChannelComponent->GetSectionBase().Y / LocalComponentSizeQuad));
+
+						Data.ComponentVertexPositionX[WeightmapChannelIndex] = PositionOffset.X * LocalComponentSizeVerts;
+						Data.ComponentVertexPositionY[WeightmapChannelIndex] = PositionOffset.Y * LocalComponentSizeVerts;
+
+						const TArray<FWeightmapLayerAllocationInfo>& ChannelLayerAllocations = ChannelComponent->GetWeightmapLayerAllocations();
+						const TArray<UTexture2D*>& ChannelComponentWeightmapTextures = ChannelComponent->GetWeightmapTextures();
+						
+						for (const FWeightmapLayerAllocationInfo& ChannelLayerAllocation : ChannelLayerAllocations)
+						{
+							if (ChannelLayerAllocation.LayerInfo != nullptr && !AlreadyProcessedAllocation.Contains(&ChannelLayerAllocation) && ChannelComponentWeightmapTextures[ChannelLayerAllocation.WeightmapTextureIndex] == WeightmapTexture)
+							{
+								Data.WeightmapChannelToProcess[WeightmapChannelIndex] = ChannelLayerAllocation.WeightmapTextureChannel;
+								AlreadyProcessedAllocation.Add(&ChannelLayerAllocation);
+
+								if (ChannelLayerAllocation.LayerInfo == ALandscapeProxy::VisibilityLayer)
+								{
+									Data.SourcePaintLayerIndex[WeightmapChannelIndex] = 0; // Always store after the last weightmap index
+								}
+								else
+								{
+									for (int32 LayerInfoSettingsIndex = 0; LayerInfoSettingsIndex < Info->Layers.Num(); ++LayerInfoSettingsIndex)
+									{
+										const FLandscapeInfoLayerSettings& LayerInfo = Info->Layers[LayerInfoSettingsIndex];
+
+										if (ChannelLayerAllocation.LayerInfo == LayerInfo.LayerInfoObj)
+										{
+											Data.SourcePaintLayerIndex[WeightmapChannelIndex] = LayerInfoSettingsIndex + 1; // due to visibility layer that is at 0
+											break;
+										}
+									}
+								}
+
+								break;
+							}
+						}
+					}
+				}
+
+				OutComponentData.Add(Data);
+			}
+		}
+	}
+}
+
+void ALandscape::ReallocateLayersWeightmaps(const TArray<ALandscapeProxy*>& InAllLandscape, const TArray<ULandscapeLayerInfoObject*>& InBrushRequiredAllocations, TArray<ULandscapeComponent*>& OutComponentThatNeedMaterialRebuild)
+{
+	SCOPE_CYCLE_COUNTER(STAT_LandscapeLayersReallocateWeightmaps);
+
+	TArray<ULandscapeComponent*> AllLandscapeComponents;
+
+	for (ALandscapeProxy* Landscape : InAllLandscape)
+	{
+		AllLandscapeComponents.Append(Landscape->LandscapeComponents);
+	}
+
+	// Copy Previous Usage, to know which texture need updating
+	TMap<UTexture2D*, ULandscapeWeightmapUsage*> CurrentWeightmapsUsage;
+
+	for (ULandscapeComponent* Component : AllLandscapeComponents)
+	{
+		TArray<UTexture2D*>& ComponentWeightmapTextures = Component->GetWeightmapTextures();
+		TArray<ULandscapeWeightmapUsage*>& ComponentWeightmapTextureUsage = Component->GetWeightmapTexturesUsage();
+
+		for (int32 i = 0; i < ComponentWeightmapTextures.Num(); ++i)
+		{
+			UTexture2D* ComponentWeightmapTexture = ComponentWeightmapTextures[i];
+			ULandscapeWeightmapUsage** CurrentWeightmapTextureUsage = CurrentWeightmapsUsage.Find(ComponentWeightmapTexture);
+
+			if (CurrentWeightmapTextureUsage == nullptr)
+			{
+				ULandscapeWeightmapUsage* ComponentWeightmapUsage = ComponentWeightmapTextureUsage[i];
+				ULandscapeWeightmapUsage* Usage = NewObject<ULandscapeWeightmapUsage>(Component->GetLandscapeProxy());
+
+				for (int32 j = 0; j < 4; ++j)
+				{
+					Usage->ChannelUsage[j] = ComponentWeightmapUsage->ChannelUsage[j];
+				}
+
+				CurrentWeightmapsUsage.Add(ComponentWeightmapTexture, Usage);
+			}
+		}
+	}
+
+	// Clear allocation data
+	for (ULandscapeComponent* Component : AllLandscapeComponents)
+	{
+		TArray<FWeightmapLayerAllocationInfo>& BaseLayerAllocations = Component->GetWeightmapLayerAllocations();
+
+		for (FWeightmapLayerAllocationInfo& BaseWeightmapAllocation : BaseLayerAllocations)
+		{
+			BaseWeightmapAllocation.WeightmapTextureChannel = 255;
+			BaseWeightmapAllocation.WeightmapTextureIndex = 255;
+		}
+		
+		TArray<ULandscapeWeightmapUsage*>& WeightmapTexturesUsage = Component->GetWeightmapTexturesUsage();
+
+		for (int32 i = 0; i < WeightmapTexturesUsage.Num(); ++i)
+		{
+			ULandscapeWeightmapUsage* Usage = WeightmapTexturesUsage[i];
+			check(Usage != nullptr);
+
+			Usage->ClearUsage();
+		}
+	}	 
+
+	bool NeedMaterialInstanceRebuild = false;
+
+	// Build a map of all the allocation per components
+	TMap<ULandscapeComponent*, TArray<ULandscapeLayerInfoObject*>> LayerAllocsPerComponent;
+
+	for (const ALandscapeProxy* Landscape : InAllLandscape)
+	{
+		for (auto& ItLayerPair : Landscape->LandscapeLayersData)
+		{
+			const FLandscapeLayerData& LayerData = ItLayerPair.Value;
+
+			for (const auto& ItWeightmapPair : LayerData.WeightmapData)
+			{
+				ULandscapeComponent* Component = ItWeightmapPair.Key;
+				const FWeightmapLayerData& WeightLayerData = ItWeightmapPair.Value;
+
+				TArray<ULandscapeLayerInfoObject*>* ComponentLayerAlloc = LayerAllocsPerComponent.Find(Component);
+
+				if (ComponentLayerAlloc == nullptr)
+				{
+					TArray<ULandscapeLayerInfoObject*> NewLayerAllocs;
+					ComponentLayerAlloc = &LayerAllocsPerComponent.Add(Component, NewLayerAllocs);
+				}
+
+				for (const FWeightmapLayerAllocationInfo& LayerWeightmapAllocation : WeightLayerData.WeightmapLayerAllocations)
+				{
+					if (LayerWeightmapAllocation.LayerInfo != nullptr)
+					{
+						ComponentLayerAlloc->AddUnique(LayerWeightmapAllocation.LayerInfo);
+					}
+				}
+
+				// Add the brush alloc also
+				for (ULandscapeLayerInfoObject* BrushLayerInfo : InBrushRequiredAllocations)
+				{
+					if (BrushLayerInfo != nullptr)
+					{
+						ComponentLayerAlloc->AddUnique(BrushLayerInfo);
+					}
+				}
+			}
+		}
+	}
+
+	// Determine if the Final layer need to add/remove some alloc
+	for (auto& ItPair : LayerAllocsPerComponent)
+	{
+		ULandscapeComponent* Component = ItPair.Key;
+		TArray<ULandscapeLayerInfoObject*>& ComponentLayerAlloc = ItPair.Value;
+		TArray<FWeightmapLayerAllocationInfo>& ComponentBaseLayerAlloc = Component->GetWeightmapLayerAllocations();
+
+		// Deal with the one that need removal
+		for (int32 i = ComponentBaseLayerAlloc.Num() - 1; i >= 0; --i)
+		{
+			FWeightmapLayerAllocationInfo& Alloc = ComponentBaseLayerAlloc[i];
+
+			if (!ComponentLayerAlloc.Contains(Alloc.LayerInfo))
+			{
+				ComponentBaseLayerAlloc.RemoveAt(i);
+			}
+		}
+
+		// Then add the new one
+		for (ULandscapeLayerInfoObject* LayerAlloc : ComponentLayerAlloc)
+		{
+			const bool AllocExist = ComponentBaseLayerAlloc.ContainsByPredicate([&LayerAlloc](FWeightmapLayerAllocationInfo& BaseLayerAlloc) { return (LayerAlloc == BaseLayerAlloc.LayerInfo); });
+
+			if (!AllocExist)
+			{
+				ComponentBaseLayerAlloc.Add(FWeightmapLayerAllocationInfo(LayerAlloc));
+			}
+		}
+	}
+
+	// Realloc the weightmap so it will create proper texture (if needed) and will set the allocations information
+	TArray<UTexture2D*> NewCreatedTextures;
+
+	for (ULandscapeComponent* Component : AllLandscapeComponents)
+	{
+		Component->ReallocateWeightmaps(nullptr, false, false, true, &NewCreatedTextures);
+	}
+
+	// TODO: correctly only recreate what is required instead of everything..
+	//GDisableAutomaticTextureMaterialUpdateDependencies = true;
+
+	for (UTexture2D* Texture : NewCreatedTextures)
+	{
+		Texture->FinishCachePlatformData();
+		Texture->PostEditChange();
+
+		Texture->WaitForStreaming();
+	}	
+
+	//GDisableAutomaticTextureMaterialUpdateDependencies = false;
+
+	// Determine which Component need updating
+	for (ULandscapeComponent* Component : AllLandscapeComponents)
+	{
+		if (OutComponentThatNeedMaterialRebuild.Contains(Component))
+		{
+			continue;
+		}
+
+		TArray<UTexture2D*>& ComponentWeightmapTextures = Component->GetWeightmapTextures();
+		TArray<ULandscapeWeightmapUsage*>& ComponentWeightmapTextureUsage = Component->GetWeightmapTexturesUsage();
+
+		for (int32 i = 0; i < ComponentWeightmapTextures.Num(); ++i)
+		{
+			UTexture2D* ComponentWeightmapTexture = ComponentWeightmapTextures[i];
+			ULandscapeWeightmapUsage* ComponentWeightmapUsage = ComponentWeightmapTextureUsage[i];
+			ULandscapeWeightmapUsage** CurrentWeightmapTextureUsage = CurrentWeightmapsUsage.Find(ComponentWeightmapTexture);
+
+			if (CurrentWeightmapTextureUsage != nullptr)
+			{
+				for (int32 j = 0; j < 4; ++j)
+				{
+					if (ComponentWeightmapUsage->ChannelUsage[j] != (*CurrentWeightmapTextureUsage)->ChannelUsage[j] && ComponentWeightmapUsage->ChannelUsage[j] != nullptr)
+					{
+						OutComponentThatNeedMaterialRebuild.AddUnique(ComponentWeightmapUsage->ChannelUsage[j]);
+					}
+				}
+			}
+			else if (NewCreatedTextures.Contains(ComponentWeightmapTexture))
+			{
+				for (int32 j = 0; j < 4; ++j)
+				{
+					if (ComponentWeightmapUsage->ChannelUsage[j] != nullptr)
+					{
+						OutComponentThatNeedMaterialRebuild.AddUnique(ComponentWeightmapUsage->ChannelUsage[j]);
+					}
+				}
+			}
+		}
+	}
+}
+
+void ALandscape::InitLayersWeightmapResources(uint8 InLayerCount)
+{
+	ULandscapeInfo* Info = GetLandscapeInfo();
+
+	// Use the 1st one to compute the resource as they are all the same anyway
+	UTextureRenderTarget2D* FirstWeightmapRT = WeightmapRTList[WeightmapRT_Scratch1];
+
+	if (CombinedLayersWeightmapAllMaterialLayersResource != nullptr && InLayerCount != CombinedLayersWeightmapAllMaterialLayersResource->GetSizeZ())
+	{
+		ReleaseResourceAndFlush(CombinedLayersWeightmapAllMaterialLayersResource);
+		CombinedLayersWeightmapAllMaterialLayersResource = nullptr;
+	}
+
+	if (CombinedLayersWeightmapAllMaterialLayersResource == nullptr)
+	{
+		CombinedLayersWeightmapAllMaterialLayersResource = new FLandscapeTexture2DArrayResource(FirstWeightmapRT->SizeX, FirstWeightmapRT->SizeY, InLayerCount, PF_G8, 1, true);
+		BeginInitResource(CombinedLayersWeightmapAllMaterialLayersResource);
+	}
+
+	if (CurrentLayersWeightmapAllMaterialLayersResource != nullptr && InLayerCount != CurrentLayersWeightmapAllMaterialLayersResource->GetSizeZ())
+	{
+		ReleaseResourceAndFlush(CurrentLayersWeightmapAllMaterialLayersResource);
+		CurrentLayersWeightmapAllMaterialLayersResource = nullptr;
+	}
+
+	if (CurrentLayersWeightmapAllMaterialLayersResource == nullptr)
+	{
+		CurrentLayersWeightmapAllMaterialLayersResource = new FLandscapeTexture2DArrayResource(FirstWeightmapRT->SizeX, FirstWeightmapRT->SizeY, InLayerCount, PF_G8, 1, true);
+		BeginInitResource(CurrentLayersWeightmapAllMaterialLayersResource);
+	}
+
+	if (WeightmapScratchExtractLayerTextureResource == nullptr)
+	{
+		WeightmapScratchExtractLayerTextureResource = new FLandscapeTexture2DResource(FirstWeightmapRT->SizeX, FirstWeightmapRT->SizeY, PF_B8G8R8A8, 1, false);
+		BeginInitResource(WeightmapScratchExtractLayerTextureResource);
+	}
+
+	if (WeightmapScratchPackLayerTextureResource == nullptr)
+	{
+		int32 MipCount = 0;
+
+		for (int32 MipRTIndex = EWeightmapRTType::WeightmapRT_Mip0; MipRTIndex < EWeightmapRTType::WeightmapRT_Count; ++MipRTIndex)
+		{
+			if (WeightmapRTList[MipRTIndex] != nullptr)
+			{
+				++MipCount;
+			}
+		}
+
+		WeightmapScratchPackLayerTextureResource = new FLandscapeTexture2DResource(FirstWeightmapRT->SizeX, FirstWeightmapRT->SizeY, PF_B8G8R8A8, MipCount, true);
+		BeginInitResource(WeightmapScratchPackLayerTextureResource);		
+	}
+}
+
+/* This code was removed as it's a fix to remove weightmap with 0 weightmap for a component, unfortunatly it will cause side effect, like perf issue, etc, so for now it's disabled until we can address this later
+bool ALandscape::GenerateZeroAllocationPerComponents(const TArray<ALandscapeProxy*>& InAllLandscape, const TMap<ULandscapeLayerInfoObject*, bool>& InWeightmapLayersBlendSubstractive, TMap<ULandscapeComponent*, TArray<ULandscapeLayerInfoObject*>>& OutZeroAllocationsPerComponents)
+{
+	// Make sure all CPU Readback texture are ready to be used
+	for (ALandscapeProxy* Landscape : InAllLandscape)
+	{
+		for (auto& ItPair : Landscape->WeightmapCPUReadBackTextures)
+		{
+			FLandscapeLayersTexture2DCPUReadBackResource* CPuReadBackTexture = ItPair.Value;
+
+			if (CPuReadBackTexture != nullptr && !CPuReadBackTexture->IsInitialized())
+			{
+				return false;
+			}
+		}
+	}
+
+	TArray<ULandscapeComponent*> AllLandscapeComponents;
+
+	for (ALandscapeProxy* Landscape : InAllLandscape)
+	{
+		AllLandscapeComponents.Append(Landscape->LandscapeComponents);
+	}
+
+	// Compute the data to facilitate the process
+	struct FTextureChannelData
+	{
+		TArray<ULandscapeComponent*, TInlineAllocator<4>> ComponentPerChannel;
+		TArray<ULandscapeLayerInfoObject*, TInlineAllocator<4>> LayerInfoPerChannel;
+		TArray<FColor> TexelData;
+		FLandscapeLayersTexture2DCPUReadBackResource* CPUReadbackTexture;
+	};
+
+	TMap<UTexture2D*, FTextureChannelData> TextureDataPerTextures;
+
+	for (ULandscapeComponent* LandscapeComponent : AllLandscapeComponents)
+	{
+		const TArray<FLandscapeWeightmapUsage*>& ComponentWeightmapUsage = LandscapeComponent->GetWeightmapTexturesUsage();
+		const TArray<UTexture2D*>& ComponentWeightmapTextures = LandscapeComponent->GetWeightmapTextures();
+
+		for (int32 TextureIndex = 0; TextureIndex < ComponentWeightmapTextures.Num(); ++TextureIndex)
+		{
+			UTexture2D* WeightmapTexture = ComponentWeightmapTextures[TextureIndex];
+
+			if (TextureDataPerTextures.Contains(WeightmapTexture))
+			{
+				continue;
+			}
+
+			FLandscapeWeightmapUsage* Usage = ComponentWeightmapUsage[TextureIndex];
+
+			for (int32 i = 0; i < 4; ++i)
+			{
+				ULandscapeComponent* Component = Usage->ChannelUsage[i];
+
+				if (Component != nullptr)
+				{
+					const TArray<FWeightmapLayerAllocationInfo>& ComponentUsageWeightmapAllocations = Component->GetWeightmapLayerAllocations();
+
+					for (const FWeightmapLayerAllocationInfo& Allocation : ComponentUsageWeightmapAllocations)
+					{
+						if (Allocation.WeightmapTextureIndex == TextureIndex)
+						{
+							const bool* IsSubtractive = InWeightmapLayersBlendSubstractive.Find(Allocation.LayerInfo);
+
+							if (IsSubtractive != nullptr && *IsSubtractive)
+							{
+								FTextureChannelData* TextureData = TextureDataPerTextures.Find(WeightmapTexture);
+
+								if (TextureData == nullptr)
+								{
+									FTextureChannelData NewData;
+									NewData.ComponentPerChannel.SetNumZeroed(4);
+									NewData.LayerInfoPerChannel.SetNumZeroed(4);
+
+									FLandscapeLayersTexture2DCPUReadBackResource** CPUReadbackTexture = Component->GetLandscapeProxy()->WeightmapCPUReadBackTextures.Find(WeightmapTexture);
+									check(CPUReadbackTexture);
+
+									NewData.CPUReadbackTexture = *CPUReadbackTexture;
+
+									TextureData = &TextureDataPerTextures.Add(WeightmapTexture, NewData);
+								}
+
+								TextureData->LayerInfoPerChannel[Allocation.WeightmapTextureChannel] = Allocation.LayerInfo;
+								TextureData->ComponentPerChannel[Allocation.WeightmapTextureChannel] = Component;
+							}
+						}
+					}
+				}
+			}
+		}
+	}
+
+	// Read the data form the CPU texture
+	if (TextureDataPerTextures.Num() > 0)
+	{
+		for (auto& ItPair : TextureDataPerTextures)
+		{
+			FTextureChannelData& TextureChannelData = ItPair.Value;
+
+			int32 MipSizeU = TextureChannelData.CPUReadbackTexture->GetSizeX();
+			int32 MipSizeV = TextureChannelData.CPUReadbackTexture->GetSizeY();
+
+			FReadSurfaceDataFlags Flags(RCM_UNorm, CubeFace_MAX);
+			Flags.SetMip(0);
+			FIntRect Rect(0, 0, MipSizeU, MipSizeV);
+
+			ENQUEUE_UNIQUE_RENDER_COMMAND_FOURPARAMETER(
+				ReadSurfaceCommand,
+				FTextureRHIRef, SourceTextureRHI, TextureChannelData.CPUReadbackTexture->TextureRHI,
+				FIntRect, Rect, Rect,
+				TArray<FColor>*, OutData, &TextureChannelData.TexelData,
+				FReadSurfaceDataFlags, ReadFlags, Flags,
+				{
+					RHICmdList.ReadSurfaceData(SourceTextureRHI, Rect, *OutData, ReadFlags);
+				});
+		}
+
+		FlushRenderingCommands();
+
+		// Determine which texture channel is zero allocation to "remove" them
+		for (auto& ItPair : TextureDataPerTextures)
+		{
+			FTextureChannelData& TextureChannelData = ItPair.Value;
+
+			TArray<bool> AreLayerEmpty;
+			AreLayerEmpty.Init(true, 4);
+
+			for (FColor& Color : TextureChannelData.TexelData)
+			{
+				if (Color.R != 0)
+				{
+					AreLayerEmpty[0] = false;
+				}
+
+				if (Color.G != 0)
+				{
+					AreLayerEmpty[1] = false;
+				}
+
+				if (Color.B != 0)
+				{
+					AreLayerEmpty[2] = false;
+				}
+
+				if (Color.A != 0)
+				{
+					AreLayerEmpty[3] = false;
+				}
+			}
+			
+			for (int32 i = 0; i < 4; ++i)
+			{
+				if (AreLayerEmpty[i] && TextureChannelData.ComponentPerChannel[0] != nullptr)
+				{
+					TArray<ULandscapeLayerInfoObject*>& ZeroAllocations = OutZeroAllocationsPerComponents.FindOrAdd(TextureChannelData.ComponentPerChannel[0]);
+					ZeroAllocations.Add(TextureChannelData.LayerInfoPerChannel[0]);
+				}
+			}
+		}
+	}
+
+	return true;
+}
+*/
+
+bool ALandscape::AreLayersWeightmapTextureResourcesReady(const TArray<ALandscapeProxy*>& InAllLandscapes) const
+{
+	// Make sure all our original weightmap textures are streamed in and ready to be used
+	for (const ALandscapeProxy* Landscape : InAllLandscapes)
+	{
+		for (const auto& ItLayerDataPair : Landscape->LandscapeLayersData)
+		{
+			for (const auto& ItWeightmapPair : ItLayerDataPair.Value.WeightmapData)
+			{
+				ULandscapeComponent* Component = ItWeightmapPair.Key;
+				const TArray<UTexture2D*>& OriginalWeightmaps = Component->GetWeightmapTextures();
+
+				for (UTexture2D* Weightmap : OriginalWeightmaps)
+				{
+					if (!Weightmap->IsFullyStreamedIn())
+					{
+						return false;
+					}
+				}
+			}
+		}
+	}
+
+	// Init all needed resources
+	for (const ALandscapeProxy* Landscape : InAllLandscapes)
+	{
+		for (const auto& ItLayerDataPair : Landscape->LandscapeLayersData)
+		{
+			for (const auto& ItWeightmapPair : ItLayerDataPair.Value.WeightmapData)
+			{
+				const FWeightmapLayerData& WeightmapLayerData = ItWeightmapPair.Value;
+
+				for (UTexture2D* Weightmap : WeightmapLayerData.Weightmaps)
+				{
+					if (Weightmap->Resource == nullptr)
+					{
+						Weightmap->FinishCachePlatformData();
+
+						Weightmap->Resource = Weightmap->CreateResource();
+
+						if (Weightmap->Resource != nullptr)
+						{
+							BeginInitResource(Weightmap->Resource);
+						}
+					}
+				}
+			}
+		}
+	}
+
+	// Wait for the new resource to be fully initialized/streamed in
+	for (const ALandscapeProxy* Landscape : InAllLandscapes)
+	{
+		for (const auto& ItLayerDataPair : Landscape->LandscapeLayersData)
+		{
+			for (const auto& ItWeightmapPair : ItLayerDataPair.Value.WeightmapData)
+			{
+				const FWeightmapLayerData& WeightmapLayerData = ItWeightmapPair.Value;
+
+				for (UTexture2D* Weightmap : WeightmapLayerData.Weightmaps)
+				{
+					if (Weightmap->Resource == nullptr || !Weightmap->Resource->IsInitialized() || !Weightmap->IsFullyStreamedIn())
+					{
+						return false;
+					}
+				}
+			}
+		}
+	}
+
+	return true;
+}
+
+void ALandscape::RegenerateLayersWeightmaps()
+{
+	SCOPE_CYCLE_COUNTER(STAT_LandscapeLayersRegenerateWeightmaps);
+
+	ULandscapeInfo* Info = GetLandscapeInfo();
+
+	if (LayersContentUpdateFlags == 0 || Info == nullptr || Info->Layers.Num() == 0)
+	{
+		return;
+	}
+
+	TArray<ALandscapeProxy*> AllLandscapes;
+	AllLandscapes.Add(this);
+
+	for (const auto& It : Info->Proxies)
+	{
+		AllLandscapes.Add(It);
+	}
+
+	if (!AreLayersWeightmapTextureResourcesReady(AllLandscapes))
+	{
+		return;
+	}
+
+	TArray<ULandscapeComponent*> AllLandscapeComponents;
+
+	for (ALandscapeProxy* Landscape : AllLandscapes)
+	{
+		AllLandscapeComponents.Append(Landscape->LandscapeComponents);
+	}
+
+	TArray<ULandscapeComponent*> ComponentThatNeedMaterialRebuild;
+	TArray<ULandscapeLayerInfoObject*> BrushRequiredAllocations;
+	int32 LayerCount = Info->Layers.Num() + 1; // due to visibility being stored at 0
+	bool ClearFlagsAfterUpdate = true;
+
+	if ((LayersContentUpdateFlags & ELandscapeLayersContentUpdateFlag::Weightmap_Render) != 0 && WeightmapRTList.Num() > 0)
+	{
+		UTextureRenderTarget2D* LandscapeScratchRT1 = WeightmapRTList[EWeightmapRTType::WeightmapRT_Scratch1];
+		UTextureRenderTarget2D* LandscapeScratchRT2 = WeightmapRTList[EWeightmapRTType::WeightmapRT_Scratch2];
+		UTextureRenderTarget2D* LandscapeScratchRT3 = WeightmapRTList[EWeightmapRTType::WeightmapRT_Scratch3];
+		UTextureRenderTarget2D* EmptyRT = WeightmapRTList[EWeightmapRTType::WeightmapRT_Scratch_RGBA];
+		FLandscapeLayersWeightmapShaderParameters PSShaderParams;
+		bool OutputDebugName = (CVarOutputLayersDebugDrawCallName.GetValueOnAnyThread() == 1 || CVarOutputLayersRTContent.GetValueOnAnyThread() == 1 || CVarOutputLayersWeightmapsRTContent.GetValueOnAnyThread() == 1) ? true : false;
+
+		InitLayersWeightmapResources(LayerCount);
+
+		ClearLayersWeightmapTextureResource(TEXT("ClearRT RGBA"), EmptyRT->GameThread_GetRenderTargetResource());
+		ClearLayersWeightmapTextureResource(TEXT("ClearRT R"), LandscapeScratchRT1->GameThread_GetRenderTargetResource());
+
+		for (int32 LayerIndex = 0; LayerIndex < LayerCount; ++LayerIndex)
+		{
+			CopyLayersTexture(*LandscapeScratchRT1->GetName(), LandscapeScratchRT1->GameThread_GetRenderTargetResource(), OutputDebugName ? FString::Printf(TEXT("Weight: Clear CombinedProcLayerWeightmapAllLayersResource %d, "), LayerIndex) : TEXT(""), CombinedLayersWeightmapAllMaterialLayersResource, nullptr, FIntPoint(0, 0), 0, 0, 0, LayerIndex);
+		}
+
+		bool ComputeShaderGeneratedData = false;
+		bool FirstLayer = true;		
+		TMap<ULandscapeLayerInfoObject*, bool> WeightmapLayersBlendSubstractive;
+
+		for (FLandscapeLayer& Layer : LandscapeLayers)
+		{
+			int8 CurrentWeightmapToProcessIndex = 0;
+			bool HasFoundWeightmapToProcess = true; // try processing at least once
+
+			TMap<ULandscapeLayerInfoObject*, int32> LayerInfoObjects; // <LayerInfoObj, LayerIndex>
+
+			// Determine if some brush want to write to layer that we have currently no data on
+			if (Layer.bVisible)
+			{
+				for (int32 LayerInfoSettingsIndex = 0; LayerInfoSettingsIndex < Info->Layers.Num(); ++LayerInfoSettingsIndex)
+				{
+					const FLandscapeInfoLayerSettings& InfoLayerSettings = Info->Layers[LayerInfoSettingsIndex];
+
+					for (int32 i = 0; i < Layer.WeightmapBrushOrderIndices.Num(); ++i)
+					{
+						FLandscapeLayerBrush& Brush = Layer.Brushes[Layer.WeightmapBrushOrderIndices[i]];
+
+						if (Brush.BPCustomBrush == nullptr)
+						{
+							continue;
+						}
+
+						if (Brush.BPCustomBrush->IsAffectingWeightmapLayer(InfoLayerSettings.GetLayerName()) && !LayerInfoObjects.Contains(InfoLayerSettings.LayerInfoObj))
+						{
+							LayerInfoObjects.Add(InfoLayerSettings.LayerInfoObj, LayerInfoSettingsIndex + 1); // due to visibility layer that is at 0
+						}
+					}
+				}
+			}
+
+			// Loop until there is no more weightmap texture to process
+			while (HasFoundWeightmapToProcess)
+			{
+				CopyLayersTexture(*EmptyRT->GetName(), EmptyRT->GameThread_GetRenderTargetResource(), OutputDebugName ? FString::Printf(TEXT("Weight: %s Clear WeightmapScratchExtractLayerTextureResource"), *Layer.Name.ToString()) : TEXT(""), WeightmapScratchExtractLayerTextureResource);
+
+				// Prepare compute shader data
+				TArray<FLandscapeLayerWeightmapExtractMaterialLayersComponentData> ComponentsData;	
+				PrepareComponentDataToExtractMaterialLayersCS(Layer, CurrentWeightmapToProcessIndex, OutputDebugName, AllLandscapes, WeightmapScratchExtractLayerTextureResource, ComponentsData, LayerInfoObjects);
+
+				HasFoundWeightmapToProcess = ComponentsData.Num() > 0;
+
+				// Perform the compute shader
+				if (ComponentsData.Num() > 0)
+				{
+					PrintLayersDebugTextureResource(OutputDebugName ? FString::Printf(TEXT("LS Weight: %s WeightmapScratchTexture %s"), *Layer.Name.ToString(), TEXT("WeightmapScratchTextureResource")) : TEXT(""), WeightmapScratchExtractLayerTextureResource, 0, false);
+
+					// Clear the current atlas if required
+					if (CurrentWeightmapToProcessIndex == 0)
+					{
+						ClearLayersWeightmapTextureResource(TEXT("ClearRT"), LandscapeScratchRT1->GameThread_GetRenderTargetResource());
+
+						// Important: for performance reason we only clear the layer we will write to, the other one might contain data but they will not be read during the blend phase
+						for (auto& ItPair : LayerInfoObjects)
+						{
+							int32 LayerIndex = ItPair.Value;
+							CopyLayersTexture(*LandscapeScratchRT1->GetName(), LandscapeScratchRT1->GameThread_GetRenderTargetResource(), OutputDebugName ? FString::Printf(TEXT("Weight: %s Clear CurrentProcLayerWeightmapAllLayersResource %d, "), *Layer.Name.ToString(), LayerIndex) : TEXT(""), CurrentLayersWeightmapAllMaterialLayersResource, nullptr, FIntPoint(0, 0), 0, 0, 0, LayerIndex);
+						}
+					}
+
+					FLandscapeLayerWeightmapExtractMaterialLayersComputeShaderParameters CSExtractLayersShaderParams;
+					CSExtractLayersShaderParams.AtlasWeightmapsPerLayer = CurrentLayersWeightmapAllMaterialLayersResource;
+					CSExtractLayersShaderParams.ComponentWeightmapResource = WeightmapScratchExtractLayerTextureResource;
+					CSExtractLayersShaderParams.ComputeShaderResource = new FLandscapeLayerWeightmapExtractMaterialLayersComputeShaderResource(ComponentsData);
+					CSExtractLayersShaderParams.ComponentSize = (SubsectionSizeQuads + 1) * NumSubsections;
+
+					BeginInitResource(CSExtractLayersShaderParams.ComputeShaderResource);
+
+					FLandscapeLayerWeightmapExtractMaterialLayersCSDispatch_RenderThread CSDispatch(CSExtractLayersShaderParams);
+
+					ENQUEUE_RENDER_COMMAND(FLandscapeLayersExtractMaterialLayersCSCommand)(
+						[CSDispatch](FRHICommandListImmediate& RHICmdList) mutable
+					{
+						CSDispatch.ExtractLayers(RHICmdList);						
+					});
+					
+					++CurrentWeightmapToProcessIndex;
+					ComputeShaderGeneratedData = true; // at least 1 CS was executed, so we can continue the processing
+				}
+			}
+
+			// If we did process at least one compute shader
+			if (LayerInfoObjects.Num() > 0)
+			{
+				for (auto& LayerInfoObject : LayerInfoObjects)
+				{
+					int32 LayerIndex = LayerInfoObject.Value;
+					ULandscapeLayerInfoObject* LayerInfoObj = LayerInfoObject.Key;
+
+					// Copy the layer we are working on
+					CopyLayersTexture(OutputDebugName ? FString::Printf(TEXT("Weight: %s PaintLayer: %s, CurrentProcLayerWeightmapAllLayersResource"), *Layer.Name.ToString(), *LayerInfoObj->LayerName.ToString()) : TEXT(""), CurrentLayersWeightmapAllMaterialLayersResource, *LandscapeScratchRT1->GetName(), LandscapeScratchRT1->GameThread_GetRenderTargetResource(), nullptr, FIntPoint(0, 0), 0, 0, LayerIndex, 0);
+					PrintLayersDebugRT(OutputDebugName ? FString::Printf(TEXT("LS Weight: %s CurrentProcLayerWeightmapAllLayersResource -> Paint Layer RT %s"), *Layer.Name.ToString(), *LandscapeScratchRT1->GetName()) : TEXT(""), LandscapeScratchRT1, 0, false);
+
+					PSShaderParams.ApplyLayerModifiers = true;
+					PSShaderParams.LayerVisible = Layer.bVisible;
+					PSShaderParams.LayerAlpha = LayerInfoObj == ALandscapeProxy::VisibilityLayer ? 1.0f : Layer.WeightmapAlpha; // visibility can't be affected by weight
+
+					DrawWeightmapComponentsToRenderTarget(OutputDebugName ? FString::Printf(TEXT("LS Weight: %s Paint: %s += -> %s"), *Layer.Name.ToString(), *LayerInfoObj->LayerName.ToString(), *LandscapeScratchRT1->GetName(), *LandscapeScratchRT2->GetName()) : TEXT(""),
+														  AllLandscapeComponents, LandscapeScratchRT1, nullptr, LandscapeScratchRT2, true, PSShaderParams, 0);
+
+					PSShaderParams.ApplyLayerModifiers = false;
+
+					// Combined Layer data with current stack
+					CopyLayersTexture(OutputDebugName ? FString::Printf(TEXT("Weight: %s PaintLayer: %s CombinedProcLayerWeightmap"), *Layer.Name.ToString(), *LayerInfoObj->LayerName.ToString()) : TEXT(""), CombinedLayersWeightmapAllMaterialLayersResource, *LandscapeScratchRT1->GetName(), LandscapeScratchRT1->GameThread_GetRenderTargetResource(), nullptr, FIntPoint(0, 0), 0, 0, LayerIndex, 0);
+					PrintLayersDebugRT(OutputDebugName ? FString::Printf(TEXT("LS Weight: %s CombinedProcLayerWeightmap -> Paint Layer RT %s"), *Layer.Name.ToString(), *LandscapeScratchRT1->GetName()) : TEXT(""), LandscapeScratchRT1, 0, false);
+
+					// Combine with current status and copy back to the combined 2d resource array
+					PSShaderParams.OutputAsSubstractive = false;
+
+					if (!FirstLayer)
+					{
+						const bool* BlendSubstractive = Layer.WeightmapLayerAllocationBlend.Find(LayerInfoObj);
+						PSShaderParams.OutputAsSubstractive = BlendSubstractive != nullptr ? *BlendSubstractive : false;
+
+						if (PSShaderParams.OutputAsSubstractive)
+						{
+							bool& IsSubstractiveBlend = WeightmapLayersBlendSubstractive.FindOrAdd(LayerInfoObj);
+							IsSubstractiveBlend = true;
+						}
+					}
+
+					DrawWeightmapComponentsToRenderTarget(OutputDebugName ? FString::Printf(TEXT("LS Weight: %s PaintLayer: %s, %s += -> Combined %s"), *Layer.Name.ToString(), *LayerInfoObj->LayerName.ToString(), *LandscapeScratchRT2->GetName(), *LandscapeScratchRT3->GetName()) : TEXT(""),
+														  AllLandscapeComponents, LandscapeScratchRT2, FirstLayer ? nullptr : LandscapeScratchRT1, LandscapeScratchRT3, true, PSShaderParams, 0);
+
+					PSShaderParams.OutputAsSubstractive = false;
+
+					CopyLayersTexture(OutputDebugName ? FString::Printf(TEXT("Weight: %s PaintLayer: %s %s"), *Layer.Name.ToString(), *LayerInfoObj->LayerName.ToString(), *LandscapeScratchRT3->GetName()) : TEXT(""), LandscapeScratchRT3->GameThread_GetRenderTargetResource(), TEXT("CombinedProcLayerWeightmap"), CombinedLayersWeightmapAllMaterialLayersResource, nullptr, FIntPoint(0, 0), 0, 0, 0, LayerIndex);
+
+					// Handle brush blending
+					if (Layer.bVisible)
+					{
+						// Draw each brushes				
+						for (int32 i = 0; i < Layer.WeightmapBrushOrderIndices.Num(); ++i)
+						{
+							// TODO: handle conversion/handling of RT not same size as internal size
+
+							FLandscapeLayerBrush& Brush = Layer.Brushes[Layer.WeightmapBrushOrderIndices[i]];
+
+							if (Brush.BPCustomBrush == nullptr || !Brush.BPCustomBrush->IsAffectingWeightmap() || !Brush.BPCustomBrush->IsAffectingWeightmapLayer(LayerInfoObj->LayerName))
+							{
+								continue;
+							}
+
+							BrushRequiredAllocations.AddUnique(LayerInfoObj);
+
+							if (!Brush.IsInitialized())
+							{
+								Brush.Initialize(GetBoundingRect(), FIntPoint(LandscapeScratchRT3->SizeX, LandscapeScratchRT3->SizeY));
+							}
+
+							UTextureRenderTarget2D* BrushOutputRT = Brush.Render(false, LandscapeScratchRT3);
+
+							if (BrushOutputRT == nullptr || BrushOutputRT->SizeX != LandscapeScratchRT3->SizeX || BrushOutputRT->SizeY != LandscapeScratchRT3->SizeY)
+							{
+								continue;
+							}
+
+							INC_DWORD_STAT(STAT_LandscapeLayersRegenerateDrawCalls); // Brush Render
+
+							PrintLayersDebugRT(OutputDebugName ? FString::Printf(TEXT("LS Weight: %s %s -> Brush %s"), *Layer.Name.ToString(), *Brush.BPCustomBrush->GetName(), *BrushOutputRT->GetName()) : TEXT(""), BrushOutputRT);
+
+							CopyLayersTexture(OutputDebugName ? FString::Printf(TEXT("Weight: %s PaintLayer: %s Brush: %s"), *Layer.Name.ToString(), *LayerInfoObj->LayerName.ToString(), *BrushOutputRT->GetName()) : TEXT(""), BrushOutputRT->GameThread_GetRenderTargetResource(), *LandscapeScratchRT3->GetName(), LandscapeScratchRT3->GameThread_GetRenderTargetResource());
+							PrintLayersDebugRT(OutputDebugName ? FString::Printf(TEXT("LS Weight: %s Component %s += -> Combined %s"), *Layer.Name.ToString(), *BrushOutputRT->GetName(), *LandscapeScratchRT3->GetName()) : TEXT(""), LandscapeScratchRT3);
+						}
+
+						PrintLayersDebugRT(OutputDebugName ? FString::Printf(TEXT("LS Weight: %s CombinedPostBrushProcLayerWeightmap -> Paint Layer RT %s"), *Layer.Name.ToString(), *LandscapeScratchRT3->GetName()) : TEXT(""), LandscapeScratchRT3, 0, false);
+						CopyLayersTexture(OutputDebugName ? FString::Printf(TEXT("Weight: %s PaintLayer: %s %s"), *Layer.Name.ToString(), *LayerInfoObj->LayerName.ToString(), *LandscapeScratchRT3->GetName()) : TEXT(""), LandscapeScratchRT3->GameThread_GetRenderTargetResource(), TEXT("CombinedProcLayerWeightmap"), CombinedLayersWeightmapAllMaterialLayersResource, nullptr, FIntPoint(0, 0), 0, 0, 0, LayerIndex);
+					}
+				}
+
+				PSShaderParams.ApplyLayerModifiers = false;
+			}
+
+			FirstLayer = false;
+		}
+
+		// TODO:  if editing a Brush affecting layers, since we don't have any bounds to brush, right now ReallocateLayersWeightmaps wont ask a rebuild of the component affected by Brushes, which mean ComponentThatNeedMaterialRebuild wont contains Brush affected component!
+		ReallocateLayersWeightmaps(AllLandscapes, BrushRequiredAllocations, ComponentThatNeedMaterialRebuild);
+
+		// Allocation that will need to be excluded when we update materials
+		TMap<ULandscapeComponent*, TArray<ULandscapeLayerInfoObject*>> ZeroAllocationsPerComponents;
+
+		if (ComputeShaderGeneratedData)
+		{
+			// Will generate CPU read back resource, if required
+			for (ALandscapeProxy* LandscapeProxy : AllLandscapes)
+			{
+				for (const ULandscapeComponent* Component : LandscapeProxy->LandscapeComponents)
+				{
+					const TArray<UTexture2D*>& ComponentWeightmapTextures = Component->GetWeightmapTextures();
+
+					for (UTexture2D* WeightmapTexture : ComponentWeightmapTextures)
+					{
+						FLandscapeLayersTexture2DCPUReadBackResource** WeightmapCPUReadBack = LandscapeProxy->WeightmapCPUReadBackTextures.Find(WeightmapTexture);
+
+						if (WeightmapCPUReadBack == nullptr)
+						{
+							FLandscapeLayersTexture2DCPUReadBackResource* NewWeightmapCPUReadBack = new FLandscapeLayersTexture2DCPUReadBackResource(WeightmapTexture->Source.GetSizeX(), WeightmapTexture->Source.GetSizeY(), WeightmapTexture->GetPixelFormat(), WeightmapTexture->Source.GetNumMips());
+							BeginInitResource(NewWeightmapCPUReadBack);
+
+							LandscapeProxy->WeightmapCPUReadBackTextures.Add(WeightmapTexture, NewWeightmapCPUReadBack);
+						}
+					}
+				}
+			}
+
+			int8 CurrentWeightmapToProcessIndex = 0;
+			bool HasFoundWeightmapToProcess = true; // try processing at least once
+
+			TArray<float> WeightmapLayerWeightBlend;
+			TArray<UTexture2D*> ProcessedWeightmaps;
+			TArray<FLandscapeLayersTexture2DCPUReadBackResource*> ProcessedWeightmapsCPUCopy;
+			int32 NextTextureIndexToProcess = 0;
+
+			// Generate the component data from the weightmap allocation that were done earlier and weight blend them if required (i.e renormalize)
+			while (HasFoundWeightmapToProcess)
+			{
+				TArray<FLandscapeLayerWeightmapPackMaterialLayersComponentData> PackLayersComponentsData;
+				PrepareComponentDataToPackMaterialLayersCS(CurrentWeightmapToProcessIndex, OutputDebugName, AllLandscapeComponents, ProcessedWeightmaps, ProcessedWeightmapsCPUCopy, PackLayersComponentsData);
+				HasFoundWeightmapToProcess = PackLayersComponentsData.Num() > 0;
+
+				// Perform the compute shader
+				if (PackLayersComponentsData.Num() > 0)
+				{
+					// Compute the weightblend mode of each layer for the compute shader
+					if (WeightmapLayerWeightBlend.Num() != LayerCount)
+					{
+						WeightmapLayerWeightBlend.SetNum(LayerCount);
+
+						for (int32 LayerInfoSettingsIndex = 0; LayerInfoSettingsIndex < Info->Layers.Num(); ++LayerInfoSettingsIndex)
+						{
+							const FLandscapeInfoLayerSettings& LayerInfo = Info->Layers[LayerInfoSettingsIndex];
+							WeightmapLayerWeightBlend[LayerInfoSettingsIndex + 1] = LayerInfo.LayerInfoObj != nullptr ? (LayerInfo.LayerInfoObj->bNoWeightBlend ? 0.0f : 1.0f) : 1.0f;
+						}
+
+						WeightmapLayerWeightBlend[0] = 0.0f; // Blend of Visibility 
+					}
+
+					TArray<FVector2D> WeightmapTextureOutputOffset;
+
+					// Compute each weightmap location so compute shader will be able to output at expected location
+					int32 ComponentSize = (SubsectionSizeQuads + 1) * NumSubsections;
+					float ComponentY = 0;
+					float ComponentX = 0;
+
+					for (int32 i = 0; i < PackLayersComponentsData.Num(); ++i)
+					{
+						check(ComponentY < WeightmapScratchPackLayerTextureResource->GetSizeY()); // This should never happen as it would be a bug in the algo
+
+						if (ComponentX >= WeightmapScratchPackLayerTextureResource->GetSizeX())
+						{
+							ComponentY += ComponentSize;
+							ComponentX = 0;
+						}
+
+						WeightmapTextureOutputOffset.Add(FVector2D(ComponentX, ComponentY));
+						ComponentX += ComponentSize;
+					}
+
+					// Clear Pack texture
+					CopyLayersTexture(*EmptyRT->GetName(), EmptyRT->GameThread_GetRenderTargetResource(), TEXT("Weight: Clear WeightmapScratchPackLayerTextureResource"), WeightmapScratchPackLayerTextureResource);
+
+					FLandscapeLayerWeightmapPackMaterialLayersComputeShaderParameters CSPackLayersShaderParams;
+					CSPackLayersShaderParams.AtlasWeightmapsPerLayer = CombinedLayersWeightmapAllMaterialLayersResource;
+					CSPackLayersShaderParams.ComponentWeightmapResource = WeightmapScratchPackLayerTextureResource;
+					CSPackLayersShaderParams.ComputeShaderResource = new FLandscapeLayerWeightmapPackMaterialLayersComputeShaderResource(PackLayersComponentsData, WeightmapLayerWeightBlend, WeightmapTextureOutputOffset);
+					CSPackLayersShaderParams.ComponentSize = ComponentSize;
+					BeginInitResource(CSPackLayersShaderParams.ComputeShaderResource);
+
+					FLandscapeLayerWeightmapPackMaterialLayersCSDispatch_RenderThread CSDispatch(CSPackLayersShaderParams);
+
+					ENQUEUE_RENDER_COMMAND(FLandscapeLayersPackMaterialLayersCSCommand)(
+						[CSDispatch](FRHICommandListImmediate& RHICmdList) mutable
+					{
+						CSDispatch.PackLayers(RHICmdList);
+					});
+
+					int32 StartTextureIndex = NextTextureIndexToProcess;
+
+					for (; NextTextureIndexToProcess < ProcessedWeightmaps.Num(); ++NextTextureIndexToProcess)
+					{
+						UTexture2D* WeightmapTexture = ProcessedWeightmaps[NextTextureIndexToProcess];
+						FLandscapeLayersTexture2DCPUReadBackResource* WeightmapCPUReadBack = ProcessedWeightmapsCPUCopy[NextTextureIndexToProcess];
+						FIntPoint TextureTopLeftPositionInAtlas(WeightmapTextureOutputOffset[NextTextureIndexToProcess - StartTextureIndex].X, WeightmapTextureOutputOffset[NextTextureIndexToProcess - StartTextureIndex].Y);
+
+						UTextureRenderTarget2D* CurrentRT = WeightmapRTList[WeightmapRT_Mip0];
+						CopyLayersTexture(TEXT("WeightmapScratchTexture"), WeightmapScratchPackLayerTextureResource, *CurrentRT->GetName(), CurrentRT->GameThread_GetRenderTargetResource());
+
+						DrawWeightmapComponentToRenderTargetMips(TextureTopLeftPositionInAtlas, CurrentRT, true, PSShaderParams);
+
+						int32 CurrentMip = 0;
+
+						for (int32 MipRTIndex = EWeightmapRTType::WeightmapRT_Mip0; MipRTIndex < EWeightmapRTType::WeightmapRT_Count; ++MipRTIndex)
+						{
+							CurrentRT = WeightmapRTList[MipRTIndex];
+
+							if (CurrentRT != nullptr)
+							{
+								CopyLayersTexture(*CurrentRT->GetName(), CurrentRT->GameThread_GetRenderTargetResource(), OutputDebugName ? FString::Printf(TEXT("Weightmap Mip: %d"), CurrentMip) : TEXT(""), WeightmapTexture->Resource, WeightmapCPUReadBack, TextureTopLeftPositionInAtlas, CurrentMip, CurrentMip);
+								++CurrentMip;
+							}
+						}
+					}
+				}
+
+				++CurrentWeightmapToProcessIndex;
+			}
+
+			// Generate all the allocation with zero painted data for each components, it will be used to update the material instance to not sample those textures
+			/*if (ComputeShaderGeneratedData)
+			{
+				ClearFlagsAfterUpdate = GenerateZeroAllocationPerComponents(AllLandscapes, WeightmapLayersBlendSubstractive, ZeroAllocationsPerComponents);
+			}
+			*/
+		}
+
+		UpdateLayersMaterialInstances(LayersUpdateAllMaterials ? AllLandscapeComponents : ComponentThatNeedMaterialRebuild, ZeroAllocationsPerComponents);
+		LayersUpdateAllMaterials = false;
+	}
+
+	if ((LayersContentUpdateFlags & ELandscapeLayersContentUpdateFlag::Weightmap_ResolveToTexture))
+	{
+		ResolveLayersWeightmapTexture(AllLandscapes);
+	}
+
+	if ((LayersContentUpdateFlags & ELandscapeLayersContentUpdateFlag::Weightmap_Collision) != 0)
+	{
+		for (ULandscapeComponent* Component : AllLandscapeComponents)
+		{
+			Component->UpdateCollisionLayerData();
+		}
+	}
+
+	if (ClearFlagsAfterUpdate)
+	{
+		LayersContentUpdateFlags &= ~ELandscapeLayersContentUpdateFlag::Weightmap_All;
+	}
+
+	// If doing rendering debug, keep doing the render only
+	if (CVarOutputLayersDebugDrawCallName.GetValueOnAnyThread() == 1)
+	{
+		LayersContentUpdateFlags |= ELandscapeLayersContentUpdateFlag::Weightmap_Render;
+	}
+}
+
+void ALandscape::UpdateLayersMaterialInstances(const TArray<ULandscapeComponent*>& InComponentsToUpdate, const TMap<ULandscapeComponent*, TArray<ULandscapeLayerInfoObject*>>& InZeroAllocationsPerComponents)
+{
+	if (InComponentsToUpdate.Num() == 0 && InZeroAllocationsPerComponents.Num() == 0)
+	{
+		return;
+	}
+
+	TArray<ULandscapeComponent*> ComponentsToUpdate;
+	//InZeroAllocationsPerComponents.GenerateKeyArray(ComponentsToUpdate);
+	ComponentsToUpdate.Append(InComponentsToUpdate);	
+
+	SCOPE_CYCLE_COUNTER(STAT_LandscapeLayersUpdateMaterialInstance);
+
+	// we're not having the material update context recreate render states because we will manually do it for only our components
+	TArray<FComponentRecreateRenderStateContext> RecreateRenderStateContexts;
+	RecreateRenderStateContexts.Reserve(ComponentsToUpdate.Num());
+
+	for (ULandscapeComponent* Component : ComponentsToUpdate)
+	{
+		RecreateRenderStateContexts.Emplace(Component);
+	}
+	TOptional<FMaterialUpdateContext> MaterialUpdateContext;
+	MaterialUpdateContext.Emplace(FMaterialUpdateContext::EOptions::Default & ~FMaterialUpdateContext::EOptions::RecreateRenderStates);
+
+	for (ULandscapeComponent* Component : ComponentsToUpdate)
+	{
+		int32 MaxLOD = FMath::CeilLogTwo(SubsectionSizeQuads + 1) - 1;
+		TMap<UMaterialInterface*, int8> NewMaterialPerLOD;
+		Component->LODIndexToMaterialIndex.SetNumUninitialized(MaxLOD + 1);
+		int8 LastLODIndex = INDEX_NONE;
+
+		UMaterialInterface* BaseMaterial = GetLandscapeMaterial();
+		UMaterialInterface* LOD0Material = GetLandscapeMaterial(0);
+
+		for (int32 LODIndex = 0; LODIndex <= MaxLOD; ++LODIndex)
+		{
+			UMaterialInterface* CurrentMaterial = GetLandscapeMaterial(LODIndex);
+
+			// if we have a LOD0 override, do not let the base material override it, it should override everything!
+			if (CurrentMaterial == BaseMaterial && BaseMaterial != LOD0Material)
+			{
+				CurrentMaterial = LOD0Material;
+			}
+
+			const int8* MaterialLOD = NewMaterialPerLOD.Find(CurrentMaterial);
+
+			if (MaterialLOD != nullptr)
+			{
+				Component->LODIndexToMaterialIndex[LODIndex] = *MaterialLOD > LastLODIndex ? *MaterialLOD : LastLODIndex;
+			}
+			else
+			{
+				int32 AddedIndex = NewMaterialPerLOD.Num();
+				NewMaterialPerLOD.Add(CurrentMaterial, LODIndex);
+				Component->LODIndexToMaterialIndex[LODIndex] = AddedIndex;
+				LastLODIndex = AddedIndex;
+			}
+		}
+
+		Component->MaterialPerLOD = NewMaterialPerLOD;
+
+		Component->MaterialInstances.SetNumZeroed(Component->MaterialPerLOD.Num()/* * 2*/); // over allocate in case we are using tessellation
+		Component->MaterialIndexToDisabledTessellationMaterial.Init(INDEX_NONE, MaxLOD + 1);
+		int8 TessellatedMaterialCount = 0;
+		int8 MaterialIndex = 0;
+
+		TArray<FWeightmapLayerAllocationInfo> WeightmapBaseLayerAllocation = Component->GetWeightmapLayerAllocations(); // We copy the array here
+		/*const TArray<ULandscapeLayerInfoObject*>* ZeroAllocations = InZeroAllocationsPerComponents.Find(Component);
+
+		// We have some allocation to remove
+		if (ZeroAllocations != nullptr)
+		{
+			for (int32 i = WeightmapBaseLayerAllocation.Num() - 1; i >= 0; --i)
+			{
+				const FWeightmapLayerAllocationInfo& Allocation = WeightmapBaseLayerAllocation[i];
+
+				if (ZeroAllocations->Contains(Allocation.LayerInfo))
+				{
+					WeightmapBaseLayerAllocation.RemoveAt(i);
+				}
+			}
+		}
+		*/
+
+		TArray<UTexture2D*>& ComponentWeightmapTextures = Component->GetWeightmapTextures();
+		UTexture2D* Heightmap = Component->GetHeightmap();
+
+		for (auto& ItPair : Component->MaterialPerLOD)
+		{
+			const int8 MaterialLOD = ItPair.Value;
+
+			// Find or set a matching MIC in the Landscape's map.
+			UMaterialInstanceConstant* CombinationMaterialInstance = Component->GetCombinationMaterial(nullptr, WeightmapBaseLayerAllocation, MaterialLOD, false);
+
+			if (CombinationMaterialInstance != nullptr)
+			{
+				UMaterialInstanceConstant* MaterialInstance = Component->MaterialInstances[MaterialIndex];
+				bool NeedToCreateMIC = MaterialInstance == nullptr;
+
+				if (NeedToCreateMIC)
+				{
+					// Create the instance for this component, that will use the layer combination instance.
+					MaterialInstance = NewObject<ULandscapeMaterialInstanceConstant>(GetOutermost());
+					Component->MaterialInstances[MaterialIndex] = MaterialInstance;
+				}
+
+				MaterialInstance->SetParentEditorOnly(CombinationMaterialInstance);
+
+				MaterialUpdateContext.GetValue().AddMaterialInstance(MaterialInstance); // must be done after SetParent				
+
+				FLinearColor Masks[4] = { FLinearColor(1.0f, 0.0f, 0.0f, 0.0f), FLinearColor(0.0f, 1.0f, 0.0f, 0.0f), FLinearColor(0.0f, 0.0f, 1.0f, 0.0f), FLinearColor(0.0f, 0.0f, 0.0f, 1.0f) };
+
+				// Set the layer mask
+				for (int32 AllocIdx = 0; AllocIdx < WeightmapBaseLayerAllocation.Num(); AllocIdx++)
+				{
+					FWeightmapLayerAllocationInfo& Allocation = WeightmapBaseLayerAllocation[AllocIdx];
+
+					FName LayerName = Allocation.LayerInfo == ALandscapeProxy::VisibilityLayer ? UMaterialExpressionLandscapeVisibilityMask::ParameterName : Allocation.LayerInfo ? Allocation.LayerInfo->LayerName : NAME_None;
+					MaterialInstance->SetVectorParameterValueEditorOnly(FName(*FString::Printf(TEXT("LayerMask_%s"), *LayerName.ToString())), Masks[Allocation.WeightmapTextureChannel]);
+				}
+
+				// Set the weightmaps
+				for (int32 i = 0; i < ComponentWeightmapTextures.Num(); i++)
+				{
+					MaterialInstance->SetTextureParameterValueEditorOnly(FName(*FString::Printf(TEXT("Weightmap%d"), i)), ComponentWeightmapTextures[i]);
+				}
+
+				if (NeedToCreateMIC)
+				{
+					MaterialInstance->PostEditChange();
+				}
+
+				/*// Setup material instance with disabled tessellation
+				if (CombinationMaterialInstance->GetMaterial()->D3D11TessellationMode != EMaterialTessellationMode::MTM_NoTessellation)
+				{
+					int32 TessellatedMaterialIndex = MaterialPerLOD.Num() + TessellatedMaterialCount++;
+					ULandscapeMaterialInstanceConstant* TessellationMaterialInstance = Cast<ULandscapeMaterialInstanceConstant>(MaterialInstances[TessellatedMaterialIndex]);
+
+					if (NeedToCreateMIC || TessellationMaterialInstance == nullptr)
+					{
+						TessellationMaterialInstance = NewObject<ULandscapeMaterialInstanceConstant>(GetOutermost());
+						TessellationMaterialInstance->SetParentEditorOnly(MaterialInstance);
+
+						MaterialInstances[TessellatedMaterialIndex] = TessellationMaterialInstance;
+						MaterialIndexToDisabledTessellationMaterial[MaterialIndex] = TessellatedMaterialIndex;
+
+						TessellationMaterialInstance->bDisableTessellation = true;
+						TessellationMaterialInstance->PostEditChange();
+					}
+
+					Context.AddMaterialInstance(TessellationMaterialInstance); // must be done after SetParent
+				}
+				*/
+			}
+
+			++MaterialIndex;
+		}
+
+		if (Component->MaterialPerLOD.Num() == 0)
+		{
+			Component->MaterialInstances.Empty(1);
+			Component->MaterialInstances.Add(nullptr);
+			Component->LODIndexToMaterialIndex.Empty(1);
+			Component->LODIndexToMaterialIndex.Add(0);
+		}
+	}
+
+	// End material update
+	MaterialUpdateContext.Reset();
+
+	// Recreate the render state for our components, needed to update the static drawlist which has cached the MaterialRenderProxies
+	// Must be after the FMaterialUpdateContext is destroyed
+	RecreateRenderStateContexts.Empty();
+}
+
+void ALandscape::ResolveLayersWeightmapTexture(const TArray<ALandscapeProxy*>& InAllLandscapes)
+{
+	SCOPE_CYCLE_COUNTER(STAT_LandscapeLayersResolveWeightmaps);
+	       
+	for (ALandscapeProxy* Landscape : InAllLandscapes)
+	{
+		TArray<TArray<FColor>> MipData;
+
+		for (auto& ItPair : Landscape->WeightmapCPUReadBackTextures)
+		{
+			UTexture2D* OriginalWeightmap = ItPair.Key;
+			FLandscapeLayersTexture2DCPUReadBackResource* WeightmapsCPUReadBack = ItPair.Value;
+
+			if (WeightmapsCPUReadBack == nullptr)
+			{
+				continue;
+			}
+
+			ResolveLayersTexture(WeightmapsCPUReadBack, OriginalWeightmap);
+		}
+	}
+}
+
+void ALandscape::RequestLayersContentUpdate(uint32 InDataFlags, bool InUpdateAllMaterials)
+{
+	LayersContentUpdateFlags = InDataFlags;
+	LayersUpdateAllMaterials = InUpdateAllMaterials;
+}
+
+void ALandscape::RegenerateLayersContent()
+{
+	if ((LayersContentUpdateFlags & Heightmap_Setup) != 0 || (LayersContentUpdateFlags & Weightmap_Setup) != 0)
+	{
+		SetupLayers();
+		LayersContentUpdateFlags &= ~(ELandscapeLayersContentUpdateFlag::All_Setup);
+	}
+
+	RegenerateLayersHeightmaps();
+	RegenerateLayersWeightmaps();
+}
+
+void ALandscape::TickLayers(float DeltaTime, ELevelTick TickType, FActorTickFunction& ThisTickFunction)
+{
+	check(GIsEditor);
+
+	UWorld* World = GetWorld();
+	if (World && !World->IsPlayInEditor())
+	{
+		if (GetMutableDefault<UEditorExperimentalSettings>()->bLandscapeLayerSystem)
+		{
+			if (PreviousExperimentalLandscapeLayers != GetMutableDefault<UEditorExperimentalSettings>()->bLandscapeLayerSystem)
+			{
+				PreviousExperimentalLandscapeLayers = GetMutableDefault<UEditorExperimentalSettings>()->bLandscapeLayerSystem;
+
+				RequestLayersContentUpdate(ELandscapeLayersContentUpdateFlag::All_Setup);
+			}
+
+			// If doing editing while shader are compiling or at load of a map, it's possible we will need another update pass after shader are completed to see the correct result
+			const int32 RemainingShadersThisFrame = GShaderCompilingManager->GetNumRemainingJobs();
+
+			if (!WasCompilingShaders && RemainingShadersThisFrame > 0)
+			{
+				WasCompilingShaders = true;
+			}
+			else if (WasCompilingShaders)
+			{
+				WasCompilingShaders = false;
+				RequestLayersContentUpdate(ELandscapeLayersContentUpdateFlag::All);
+			}
+
+			RegenerateLayersContent();
+		}
+		else
+		{
+			if (PreviousExperimentalLandscapeLayers != GetMutableDefault<UEditorExperimentalSettings>()->bLandscapeLayerSystem)
+			{
+				PreviousExperimentalLandscapeLayers = GetMutableDefault<UEditorExperimentalSettings>()->bLandscapeLayerSystem;
+
+			#if WITH_EDITORONLY_DATA
+				for (auto& ItPair : RenderDataPerHeightmap)
+				{
+					FRenderDataPerHeightmap& HeightmapRenderData = ItPair.Value;
+
+					if (HeightmapRenderData.HeightmapsCPUReadBack != nullptr)
+					{
+						BeginReleaseResource(HeightmapRenderData.HeightmapsCPUReadBack);
+					}
+				}
+
+				for (auto& ItPair : WeightmapCPUReadBackTextures)
+				{
+					FLandscapeLayersTexture2DCPUReadBackResource* WeightmapCPUReadBack = ItPair.Value;
+
+					if (WeightmapCPUReadBack != nullptr)
+					{
+						BeginReleaseResource(WeightmapCPUReadBack);
+					}
+				}
+
+				if (CombinedLayersWeightmapAllMaterialLayersResource != nullptr)
+				{
+					BeginReleaseResource(CombinedLayersWeightmapAllMaterialLayersResource);
+				}
+
+				if (CurrentLayersWeightmapAllMaterialLayersResource != nullptr)
+				{
+					BeginReleaseResource(CurrentLayersWeightmapAllMaterialLayersResource);
+				}
+
+				if (WeightmapScratchExtractLayerTextureResource != nullptr)
+				{
+					BeginReleaseResource(WeightmapScratchExtractLayerTextureResource);
+				}
+
+				if (WeightmapScratchPackLayerTextureResource != nullptr)
+				{
+					BeginReleaseResource(WeightmapScratchPackLayerTextureResource);
+				}
+
+				FlushRenderingCommands();
+
+				for (auto& ItPair : RenderDataPerHeightmap)
+				{
+					FRenderDataPerHeightmap& HeightmapRenderData = ItPair.Value;
+
+					delete HeightmapRenderData.HeightmapsCPUReadBack;
+					HeightmapRenderData.HeightmapsCPUReadBack = nullptr;
+				}
+
+				for (auto& ItPair : WeightmapCPUReadBackTextures)
+				{
+					FLandscapeLayersTexture2DCPUReadBackResource* WeightmapCPUReadBack = ItPair.Value;
+
+					delete WeightmapCPUReadBack;
+					WeightmapCPUReadBack = nullptr;
+				}
+
+				delete CombinedLayersWeightmapAllMaterialLayersResource;
+				delete CurrentLayersWeightmapAllMaterialLayersResource;
+				delete WeightmapScratchExtractLayerTextureResource;
+				delete WeightmapScratchPackLayerTextureResource;
+
+				CombinedLayersWeightmapAllMaterialLayersResource = nullptr;
+				CurrentLayersWeightmapAllMaterialLayersResource = nullptr;
+				WeightmapScratchExtractLayerTextureResource = nullptr;
+				WeightmapScratchPackLayerTextureResource = nullptr;
+			#endif
+			}
+		}
+	}
+}
+
+#endif
+
+void ALandscapeProxy::BeginDestroy()
+{
+#if WITH_EDITORONLY_DATA
+	if (GetMutableDefault<UEditorExperimentalSettings>()->bLandscapeLayerSystem)
+	{
+		for (auto& ItPair : RenderDataPerHeightmap)
+		{
+			FRenderDataPerHeightmap& HeightmapRenderData = ItPair.Value;
+
+			if (HeightmapRenderData.HeightmapsCPUReadBack != nullptr)
+			{
+				BeginReleaseResource(HeightmapRenderData.HeightmapsCPUReadBack);
+			}
+		}
+
+		for (auto& ItPair : WeightmapCPUReadBackTextures)
+		{
+			FLandscapeLayersTexture2DCPUReadBackResource* WeightmapCPUReadBack = ItPair.Value;
+
+			if (WeightmapCPUReadBack != nullptr)
+			{
+				BeginReleaseResource(WeightmapCPUReadBack);
+			}
+		}
+
+		ReleaseResourceFence.BeginFence();
+	}
+#endif
+
+	Super::BeginDestroy();
+}
+
+bool ALandscapeProxy::IsReadyForFinishDestroy()
+{
+	bool bReadyForFinishDestroy = Super::IsReadyForFinishDestroy();
+
+#if WITH_EDITORONLY_DATA
+	if (GetMutableDefault<UEditorExperimentalSettings>()->bLandscapeLayerSystem)
+	{
+		if (bReadyForFinishDestroy)
+		{
+			bReadyForFinishDestroy = ReleaseResourceFence.IsFenceComplete();
+		}
+	}
+#endif
+
+	return bReadyForFinishDestroy;
+}
+
+void ALandscapeProxy::FinishDestroy()
+{
+#if WITH_EDITORONLY_DATA
+	if (GetMutableDefault<UEditorExperimentalSettings>()->bLandscapeLayerSystem)
+	{
+		check(ReleaseResourceFence.IsFenceComplete());
+
+		for (auto& ItPair : RenderDataPerHeightmap)
+		{
+			FRenderDataPerHeightmap& HeightmapRenderData = ItPair.Value;
+
+			delete HeightmapRenderData.HeightmapsCPUReadBack;
+			HeightmapRenderData.HeightmapsCPUReadBack = nullptr;
+		}
+
+		for (auto& ItPair : WeightmapCPUReadBackTextures)
+		{
+			FLandscapeLayersTexture2DCPUReadBackResource* WeightmapCPUReadBack = ItPair.Value;
+
+			delete WeightmapCPUReadBack;
+			WeightmapCPUReadBack = nullptr;
+		}
+	}
+#endif
+
+	Super::FinishDestroy();
+}
+
+void ALandscape::BeginDestroy()
+{
+#if WITH_EDITOR
+	if (GetMutableDefault<UEditorExperimentalSettings>()->bLandscapeLayerSystem)
+	{
+		if (CombinedLayersWeightmapAllMaterialLayersResource != nullptr)
+		{
+			BeginReleaseResource(CombinedLayersWeightmapAllMaterialLayersResource);
+		}
+
+		if (CurrentLayersWeightmapAllMaterialLayersResource != nullptr)
+		{
+			BeginReleaseResource(CurrentLayersWeightmapAllMaterialLayersResource);
+		}
+
+		if (WeightmapScratchExtractLayerTextureResource != nullptr)
+		{
+			BeginReleaseResource(WeightmapScratchExtractLayerTextureResource);
+		}
+
+		if (WeightmapScratchPackLayerTextureResource != nullptr)
+		{
+			BeginReleaseResource(WeightmapScratchPackLayerTextureResource);
+		}
+
+		// Use ResourceFence from base class		
+	}
+#endif
+
+	Super::BeginDestroy();
+}
+
+void ALandscape::FinishDestroy()
+{
+#if WITH_EDITORONLY_DATA
+	if (GetMutableDefault<UEditorExperimentalSettings>()->bLandscapeLayerSystem)
+	{
+		check(ReleaseResourceFence.IsFenceComplete());
+
+		delete CombinedLayersWeightmapAllMaterialLayersResource;
+		delete CurrentLayersWeightmapAllMaterialLayersResource;
+		delete WeightmapScratchExtractLayerTextureResource;
+		delete WeightmapScratchPackLayerTextureResource;
+
+		CombinedLayersWeightmapAllMaterialLayersResource = nullptr;
+		CurrentLayersWeightmapAllMaterialLayersResource = nullptr;
+		WeightmapScratchExtractLayerTextureResource = nullptr;
+		WeightmapScratchPackLayerTextureResource = nullptr;
+	}
+#endif
+
+	Super::FinishDestroy();
+}
+
+#if WITH_EDITOR
+bool ALandscape::IsLayerNameUnique(const FName& InName) const
+{
+	return Algo::CountIf(LandscapeLayers, [InName](const FLandscapeLayer& Layer) { return (Layer.Name == InName); }) == 0;
+}
+
+void ALandscape::SetLayerName(int32 InLayerIndex, const FName& InName)
+{
+	FLandscapeLayer* Layer = GetLayer(InLayerIndex);
+	ULandscapeInfo* LandscapeInfo = GetLandscapeInfo();
+	if (!LandscapeInfo || !Layer || Layer->Name == InName)
+	{
+		return;
+	}
+
+	if (!IsLayerNameUnique(InName))
+	{
+		return;
+	}
+
+	LandscapeLayers[InLayerIndex].Name = InName;
+}
+
+void ALandscape::SetLayerAlpha(int32 InLayerIndex, const float InAlpha, bool bInHeightmap)
+{
+	FLandscapeLayer* Layer = GetLayer(InLayerIndex);
+	if (!Layer)
+	{
+		return;
+	}
+	float& LayerAlpha = bInHeightmap ? Layer->HeightmapAlpha : Layer->WeightmapAlpha;
+	if (LayerAlpha == InAlpha)
+	{
+		return;
+	}
+
+	LayerAlpha = InAlpha;
+	RequestLayersContentUpdate(ELandscapeLayersContentUpdateFlag::All, true);
+}
+
+void ALandscape::SetLayerVisibility(int32 InLayerIndex, bool bInVisible)
+{
+	FLandscapeLayer* Layer = GetLayer(InLayerIndex);
+	if (!Layer || Layer->bVisible == bInVisible)
+	{
+		return;
+	}
+
+	Layer->bVisible = bInVisible;
+	RequestLayersContentUpdate(ELandscapeLayersContentUpdateFlag::All, true);
+}
+
+FLandscapeLayer* ALandscape::GetLayer(int32 InLayerIndex)
+{
+	if (LandscapeLayers.IsValidIndex(InLayerIndex))
+	{
+		return &LandscapeLayers[InLayerIndex];
+	}
+	return nullptr;
+}
+
+const FLandscapeLayer* ALandscape::GetLayer(int32 InLayerIndex) const
+{
+	if (LandscapeLayers.IsValidIndex(InLayerIndex))
+	{
+		return &LandscapeLayers[InLayerIndex];
+	}
+	return nullptr;
+}
+
+void ALandscape::DeleteLayer(int32 InLayerIndex)
+{
+	ensure(GetMutableDefault<UEditorExperimentalSettings>()->bLandscapeLayerSystem);
+
+	ULandscapeInfo* LandscapeInfo = GetLandscapeInfo();
+	const FLandscapeLayer* Layer = GetLayer(InLayerIndex);
+	if (!LandscapeInfo || !Layer || LandscapeLayers.Num() <= 1)
+	{
+		return;
+	}
+	
+	Modify();
+	FGuid LayerGuid = Layer->Guid;
+
+	// Clean up Weightmap usage in LandscapeProxies
+	LandscapeInfo->ForAllLandscapeProxies([LayerGuid](ALandscapeProxy* Proxy)
+	{
+		const FLandscapeLayerData* LayerData = Proxy->LandscapeLayersData.Find(LayerGuid);
+		if (LayerData)
+		{
+			for (ULandscapeComponent* Component : Proxy->LandscapeComponents)
+			{
+				const FWeightmapLayerData* WeightmapLayer = LayerData->WeightmapData.Find(Component);
+				if (WeightmapLayer)
+				{
+					for (const FWeightmapLayerAllocationInfo& Allocation : WeightmapLayer->WeightmapLayerAllocations)
+					{
+						UTexture2D* WeightmapTexture = WeightmapLayer->Weightmaps[Allocation.WeightmapTextureIndex];
+						ULandscapeWeightmapUsage** Usage = Proxy->WeightmapUsageMap.Find(WeightmapTexture);
+						if (Usage != nullptr && (*Usage) != nullptr)
+						{
+							(*Usage)->ChannelUsage[Allocation.WeightmapTextureChannel] = nullptr;
+							if ((*Usage)->FreeChannelCount() == 4)
+							{
+								Proxy->WeightmapUsageMap.Remove(WeightmapTexture);
+							}
+						}
+					}
+				}
+			}
+		}
+	});
+	
+	// Remove associated layer data of each landscape proxy
+	LandscapeInfo->ForAllLandscapeProxies([LayerGuid](ALandscapeProxy* Proxy)
+	{
+		Proxy->LandscapeLayersData.Remove(LayerGuid);
+	});
+
+	// Remove layer from list
+	LandscapeLayers.RemoveAt(InLayerIndex);
+
+	// Request Update
+	RequestLayersContentUpdate(ELandscapeLayersContentUpdateFlag::All_Setup | All, true);
+}
+
+void ALandscape::ClearLayer(int32 InLayerIndex)
+{
+	const FLandscapeLayer* Layer = GetLayer(InLayerIndex);
+	if (Layer)
+	{
+		ClearLayer(Layer->Guid);
+	}
+}
+
+void ALandscape::ClearLayer(const FGuid& InLayerGuid)
+{
+	ensure(GetMutableDefault<UEditorExperimentalSettings>()->bLandscapeLayerSystem);
+
+	FLandscapeLayer* Layer = LandscapeLayers.FindByPredicate([InLayerGuid](const FLandscapeLayer& Other) { return Other.Guid == InLayerGuid; });
+	ULandscapeInfo* LandscapeInfo = GetLandscapeInfo();
+	if (!LandscapeInfo || !Layer)
+	{
+		return;
+	}
+	
+	Modify();
+	FScopedSetLandscapeEditingLayer Scope(this, Layer ? Layer->Guid : FGuid(), [=] { RequestLayersContentUpdate(ELandscapeLayersContentUpdateFlag::All, true); });
+
+	TArray<uint16> NewData;
+	NewData.AddZeroed(FMath::Square(ComponentSizeQuads + 1));
+	uint16 ZeroValue = LandscapeDataAccess::GetTexHeight(0.f);
+	for (uint16& NewDataValue : NewData)
+	{
+		NewDataValue = ZeroValue;
+	}
+
+	FLandscapeEditDataInterface LandscapeEdit(LandscapeInfo);
+	LandscapeInfo->ForAllLandscapeProxies([&](ALandscapeProxy* Proxy)
+	{
+		for (ULandscapeComponent* Component : Proxy->LandscapeComponents)
+		{
+			int32 MinX = MAX_int32;
+			int32 MinY = MAX_int32;
+			int32 MaxX = MIN_int32;
+			int32 MaxY = MIN_int32;
+			Component->GetComponentExtent(MinX, MinY, MaxX, MaxY);
+			check(ComponentSizeQuads == (MaxX - MinX));
+			check(ComponentSizeQuads == (MaxY - MinY));
+
+			TArray<uint16> OldData;
+			OldData.AddZeroed((1 + MaxY - MinY) * (1 + MaxX - MinX));
+
+			LandscapeEdit.GetHeightData(MinX, MinY, MaxX, MaxY, OldData.GetData(), 0);
+			if (FMemory::Memcmp(OldData.GetData(), NewData.GetData(), NewData.Num() * NewData.GetTypeSize()) != 0)
+			{
+				LandscapeEdit.SetHeightData(MinX, MinY, MaxX, MaxY, NewData.GetData(), 0, true);
+			}
+
+			// Clear weight maps
+			for (FLandscapeInfoLayerSettings& LayerSettings : LandscapeInfo->Layers)
+			{
+				Component->DeleteLayer(LayerSettings.LayerInfoObj, LandscapeEdit);
+			}
+		}
+	});
+}
+
+void ALandscape::ShowOnlySelectedLayer(int32 InLayerIndex)
+{
+	const FLandscapeLayer* VisibleLayer = GetLayer(InLayerIndex);
+	if (VisibleLayer)
+	{
+		for (FLandscapeLayer& Layer : LandscapeLayers)
+		{
+			Layer.bVisible = (&Layer == VisibleLayer);
+		}
+		RequestLayersContentUpdate(ELandscapeLayersContentUpdateFlag::All, true);
+	}
+}
+
+void ALandscape::ShowAllLayers()
+{
+	if (LandscapeLayers.Num() > 0)
+	{
+		for (FLandscapeLayer& Layer : LandscapeLayers)
+		{
+			Layer.bVisible = true;
+		}
+		RequestLayersContentUpdate(ELandscapeLayersContentUpdateFlag::All, true);
+	}
+}
+
+FScopedSetLandscapeEditingLayer::FScopedSetLandscapeEditingLayer(ALandscape* InLandscape, const FGuid& InLayerGUID, TFunction<void()> InCompletionCallback)
+	: Landscape(InLandscape)
+	, LayerGUID(InLayerGUID)
+	, CompletionCallback(MoveTemp(InCompletionCallback))
+{
+	if (GetMutableDefault<UEditorExperimentalSettings>()->bLandscapeLayerSystem && Landscape.IsValid() && LayerGUID.IsValid())
+	{
+		Landscape->SetEditingLayer(LayerGUID);
+	}
+}
+
+FScopedSetLandscapeEditingLayer::~FScopedSetLandscapeEditingLayer()
+{
+	if (GetMutableDefault<UEditorExperimentalSettings>()->bLandscapeLayerSystem && Landscape.IsValid() && LayerGUID.IsValid())
+	{
+		Landscape->SetEditingLayer();
+		if (CompletionCallback)
+		{
+			CompletionCallback();
+		}
+	}
+}
+
+void ALandscape::SetEditingLayer(FGuid InLayerGuid)
+{
+	ensure(GetMutableDefault<UEditorExperimentalSettings>()->bLandscapeLayerSystem);
+
+	ULandscapeInfo* LandscapeInfo = GetLandscapeInfo();
+	if (!LandscapeInfo)
+	{
+		return;
+	}
+
+	LandscapeInfo->ForAllLandscapeProxies([InLayerGuid, this](ALandscapeProxy* Proxy)
+	{
+		FLandscapeLayer* Layer = InLayerGuid.IsValid() ? LandscapeLayers.FindByPredicate([InLayerGuid](const FLandscapeLayer& Other) { return Other.Guid == InLayerGuid; }) : nullptr;
+		FLandscapeLayerData* LayerData = Layer ? Proxy->LandscapeLayersData.Find(Layer->Guid) : nullptr;
+
+		for (ULandscapeComponent* Component : Proxy->LandscapeComponents)
+		{
+			Component->SetEditingLayer(Layer, LayerData);
+			Component->MarkRenderStateDirty();
+		}
+	});
+}
+
+void ALandscape::CreateLayer(FName InName, bool bInUpdateLayersContent)
+{
+	ULandscapeInfo* LandscapeInfo = GetLandscapeInfo();
+	if (!LandscapeInfo || !GetMutableDefault<UEditorExperimentalSettings>()->bLandscapeLayerSystem)
+	{
+		return;
+	}
+
+	Modify();
+	FLandscapeLayer NewLayer;
+	NewLayer.Name = GenerateUniqueLayerName(InName);
+	LandscapeLayers.Add(NewLayer);
+
+	// Create associated layer data in each landscape proxy
+	LandscapeInfo->ForAllLandscapeProxies([NewLayer](ALandscapeProxy* Proxy)
+	{
+		Proxy->LandscapeLayersData.Add(NewLayer.Guid, FLandscapeLayerData());
+	});
+
+	if (bInUpdateLayersContent)
+	{
+		// Request Update
+		RequestLayersContentUpdate(ELandscapeLayersContentUpdateFlag::All_Setup);
+		RegenerateLayersContent();
+	}
+}
+
+FName ALandscape::GenerateUniqueLayerName(FName InName) const
+{
+	FString BaseName = InName == NAME_None ? "Layer" : InName.ToString();
+	FName NewName;
+	int32 LayerIndex = 0;
+	do
+	{
+		++LayerIndex;
+		NewName = FName(*FString::Printf(TEXT("%s%d"), *BaseName, LayerIndex));
+	} while (LandscapeLayers.ContainsByPredicate([NewName](const FLandscapeLayer& Layer) { return Layer.Name == NewName; }));
+
+	return NewName;
+}
+
+bool ALandscape::IsLayerBlendSubstractive(int32 InLayerIndex, const TWeakObjectPtr<ULandscapeLayerInfoObject>& InLayerInfoObj) const
+{
+	const FLandscapeLayer* Layer = GetLayer(InLayerIndex);
+
+	if (Layer == nullptr)
+	{
+		return false;
+	}
+
+	const bool* AllocationBlend = Layer->WeightmapLayerAllocationBlend.Find(InLayerInfoObj.Get());
+
+	if (AllocationBlend != nullptr)
+	{
+		return (*AllocationBlend);
+	}
+
+	return false;
+}
+
+void ALandscape::SetLayerSubstractiveBlendStatus(int32 InLayerIndex, bool InStatus, const TWeakObjectPtr<ULandscapeLayerInfoObject>& InLayerInfoObj)
+{
+	FLandscapeLayer* Layer = GetLayer(InLayerIndex);
+
+	if (Layer == nullptr)
+	{
+		return;
+	}
+
+	bool* AllocationBlend = Layer->WeightmapLayerAllocationBlend.Find(InLayerInfoObj.Get());
+
+	if (AllocationBlend == nullptr)
+	{
+		Layer->WeightmapLayerAllocationBlend.Add(InLayerInfoObj.Get(), InStatus);
+	}
+	else
+	{
+		*AllocationBlend = InStatus;
+	}
+
+	RequestLayersContentUpdate(true);
+}
+
+void ALandscape::AddBrushToLayer(int32 InLayerIndex, int32 InTargetType, ALandscapeBlueprintCustomBrush* InBrush)
+{
+	FLandscapeLayer* Layer = GetLayer(InLayerIndex);
+
+	if (Layer == nullptr)
+	{
+		return;
+	}
+
+	int32 AddedIndex = Layer->Brushes.Add(FLandscapeLayerBrush(InBrush));
+
+	if (InTargetType == ELandscapeToolTargetType::Type::Heightmap)
+	{
+		Layer->HeightmapBrushOrderIndices.Add(AddedIndex);
+	}
+	else
+	{
+		Layer->WeightmapBrushOrderIndices.Add(AddedIndex);
+	}
+
+	InBrush->SetOwningLandscape(this);
+
+	RequestLayersContentUpdate(ELandscapeLayersContentUpdateFlag::All_Render);
+}
+
+void ALandscape::RemoveBrushFromLayer(int32 InLayerIndex, int32 InTargetType, ALandscapeBlueprintCustomBrush* InBrush)
+{
+	FLandscapeLayer* Layer = GetLayer(InLayerIndex);
+
+	if (Layer == nullptr)
+	{
+		return;
+	}
+
+	int32 IndexToRemove = INDEX_NONE;
+	for (int32 i = 0; i < Layer->Brushes.Num(); ++i)
+	{
+		if (Layer->Brushes[i].BPCustomBrush == InBrush)
+		{
+			IndexToRemove = i;
+			break;
+		}
+	}
+
+	if (IndexToRemove != INDEX_NONE)
+	{
+		Layer->Brushes.RemoveAt(IndexToRemove);
+
+		for (int32 i = 0; i < Layer->HeightmapBrushOrderIndices.Num(); ++i)
+		{
+			if (Layer->HeightmapBrushOrderIndices[i] == IndexToRemove)
+			{
+				// Update the value of the index of all the one after the one we removed, so index still correctly match actual brushes list
+				for (int32 j = 0; j < Layer->HeightmapBrushOrderIndices.Num(); ++j)
+				{
+					if (Layer->HeightmapBrushOrderIndices[j] > IndexToRemove)
+					{
+						--Layer->HeightmapBrushOrderIndices[j];
+					}
+				}
+
+				Layer->HeightmapBrushOrderIndices.RemoveAt(i);
+				break;
+			}
+		}
+
+		for (int32 i = 0; i < Layer->WeightmapBrushOrderIndices.Num(); ++i)
+		{
+			if (Layer->WeightmapBrushOrderIndices[i] == IndexToRemove)
+			{
+				// Update the value of the index of all the one after the one we removed, so index still correctly match actual brushes list
+				for (int32 j = 0; j < Layer->WeightmapBrushOrderIndices.Num(); ++j)
+				{
+					if (Layer->WeightmapBrushOrderIndices[j] > IndexToRemove)
+					{
+						--Layer->WeightmapBrushOrderIndices[j];
+					}
+				}
+
+				Layer->WeightmapBrushOrderIndices.RemoveAt(i);
+				break;
+			}
+		}
+
+		InBrush->SetOwningLandscape(nullptr);
+	}
+
+	RequestLayersContentUpdate(ELandscapeLayersContentUpdateFlag::All_Render);
+}
+
+bool ALandscape::AreAllBrushesCommitedToLayer(int32 InLayerIndex, int32 InTargetType)
+{
+	FLandscapeLayer* Layer = GetLayer(InLayerIndex);
+
+	if (Layer == nullptr)
+	{
+		return false;
+	}
+
+	for (FLandscapeLayerBrush& Brush : Layer->Brushes)
+	{
+		if (Brush.BPCustomBrush != nullptr && !Brush.BPCustomBrush->IsCommited()
+			&& ((InTargetType == ELandscapeToolTargetType::Type::Heightmap && Brush.BPCustomBrush->IsAffectingHeightmap()) || (InTargetType == ELandscapeToolTargetType::Type::Weightmap && Brush.BPCustomBrush->IsAffectingWeightmap())))
+		{
+			return false;
+		}
+	}
+
+	return true;
+}
+
+void ALandscape::SetBrushesCommitStateForLayer(int32 InLayerIndex, int32 InTargetType, bool InCommited)
+{
+	FLandscapeLayer* Layer = GetLayer(InLayerIndex);
+
+	if (Layer == nullptr)
+	{
+		return;
+	}
+
+	for (FLandscapeLayerBrush& Brush : Layer->Brushes)
+	{
+		if (Brush.BPCustomBrush != nullptr)
+		{
+			Brush.BPCustomBrush->SetCommitState(InCommited);
+		}
+	}
+}
+
+TArray<int8>& ALandscape::GetBrushesOrderForLayer(int32 InLayerIndex, int32 InTargetType)
+{
+	FLandscapeLayer* Layer = GetLayer(InLayerIndex);
+	check(Layer != nullptr);
+
+	if (InTargetType == ELandscapeToolTargetType::Type::Heightmap)
+	{
+		return Layer->HeightmapBrushOrderIndices;
+	}
+
+	return Layer->WeightmapBrushOrderIndices;
+}
+
+ALandscapeBlueprintCustomBrush* ALandscape::GetBrushForLayer(int32 InLayerIndex, int32 InTargetType, int8 InBrushIndex) const
+{
+	const FLandscapeLayer* Layer = GetLayer(InLayerIndex);
+
+	if (Layer == nullptr)
+	{
+		return nullptr;
+	}
+
+	if (InTargetType == ELandscapeToolTargetType::Type::Heightmap)
+	{
+		if (Layer->HeightmapBrushOrderIndices.IsValidIndex(InBrushIndex))
+		{
+			int8 ActualBrushIndex = Layer->HeightmapBrushOrderIndices[InBrushIndex];
+			if (Layer->Brushes.IsValidIndex(ActualBrushIndex))
+			{
+				return Layer->Brushes[ActualBrushIndex].BPCustomBrush;
+			}
+		}
+	}
+	else
+	{
+		if (Layer->WeightmapBrushOrderIndices.IsValidIndex(InBrushIndex))
+		{
+			int8 ActualBrushIndex = Layer->WeightmapBrushOrderIndices[InBrushIndex];
+			if (Layer->Brushes.IsValidIndex(ActualBrushIndex))
+			{
+				return Layer->Brushes[ActualBrushIndex].BPCustomBrush;
+			}
+		}
+	}
+
+	return nullptr;
+}
+
+TArray<ALandscapeBlueprintCustomBrush*> ALandscape::GetBrushesForLayer(int32 InLayerIndex, int32 InTargetType) const
+{
+	TArray<ALandscapeBlueprintCustomBrush*> Brushes;
+
+	const FLandscapeLayer* Layer = GetLayer(InLayerIndex);
+
+	if (Layer == nullptr)
+	{
+		return Brushes;
+	}
+
+	Brushes.Reserve(Layer->Brushes.Num());
+
+	for (const FLandscapeLayerBrush& Brush : Layer->Brushes)
+	{
+		if (Brush.BPCustomBrush != nullptr && ((Brush.BPCustomBrush->IsAffectingHeightmap() && InTargetType == ELandscapeToolTargetType::Type::Heightmap)
+			|| (Brush.BPCustomBrush->IsAffectingWeightmap() && InTargetType == ELandscapeToolTargetType::Type::Weightmap)))
+		{
+			Brushes.Add(Brush.BPCustomBrush);
+		}
+	}
+
+	return Brushes;
+}
+
+#endif // WITH_EDITOR
+
+#undef LOCTEXT_NAMESPACE
