Changed the name from Landscape Procedural to Landscape Layers

diff --git a/Engine/Source/Runtime/Landscape/Private/LandscapeModule.cpp b/Engine/Source/Runtime/Landscape/Private/LandscapeModule.cpp
--- a/Engine/Source/Runtime/Landscape/Private/LandscapeModule.cpp
+++ b/Engine/Source/Runtime/Landscape/Private/LandscapeModule.cpp
@@ -114,13 +114,13 @@ void GetLandscapeTexturesAndMaterials(ULevel* Level, TArray<UObject*>& OutTextur
 			LandscapeComponent->GetGeneratedTexturesAndMaterialInstances(OutTexturesAndMaterials);
 		}
 
-		if (GetMutableDefault<UEditorExperimentalSettings>()->bProceduralLandscape)
+		if (GetMutableDefault<UEditorExperimentalSettings>()->bLandscapeLayerSystem)
 		{
 			ALandscapeProxy* Landscape = Cast<ALandscapeProxy>(ObjInLevel);
 
 			if (Landscape != nullptr)
 			{
-				for (auto& ItLayerDataPair : Landscape->ProceduralLayersData)
+				for (auto& ItLayerDataPair : Landscape->LandscapeLayersData)
 				{
 					for (auto& ItHeightmapPair : ItLayerDataPair.Value.Heightmaps)
 					{
diff --git a/Engine/Source/Runtime/Landscape/Private/LandscapePrivate.h b/Engine/Source/Runtime/Landscape/Private/LandscapePrivate.h
--- a/Engine/Source/Runtime/Landscape/Private/LandscapePrivate.h
+++ b/Engine/Source/Runtime/Landscape/Private/LandscapePrivate.h
@@ -21,17 +21,17 @@ DECLARE_CYCLE_STAT_EXTERN(TEXT("PostInit View Custom Data"), STAT_LandscapePostI
 DECLARE_CYCLE_STAT_EXTERN(TEXT("Compute Custom Mesh Batch LOD"), STAT_LandscapeComputeCustomMeshBatchLOD, STATGROUP_Landscape, );
 DECLARE_CYCLE_STAT_EXTERN(TEXT("Compute Custom Shadow Mesh Batch LOD"), STAT_LandscapeComputeCustomShadowMeshBatchLOD, STATGROUP_Landscape, );
 
-DECLARE_CYCLE_STAT_EXTERN(TEXT("Regenerate Procedural (RenderThread)"), STAT_LandscapeRegenerateProcedural_RenderThread, STATGROUP_Landscape, );
-DECLARE_CYCLE_STAT_EXTERN(TEXT("Regenerate Procedural Heightmap (GameThread)"), STAT_LandscapeRegenerateProceduralHeightmaps, STATGROUP_Landscape, );
-DECLARE_CYCLE_STAT_EXTERN(TEXT("Resolve Procedural Heightmap"), STAT_LandscapeResolveProceduralHeightmap, STATGROUP_Landscape, );
+DECLARE_CYCLE_STAT_EXTERN(TEXT("Regenerate Layers (RenderThread)"), STAT_LandscapeLayersRegenerate_RenderThread, STATGROUP_Landscape, );
+DECLARE_CYCLE_STAT_EXTERN(TEXT("Regenerate Layers Heightmap (GameThread)"), STAT_LandscapeLayersRegenerateHeightmaps, STATGROUP_Landscape, );
+DECLARE_CYCLE_STAT_EXTERN(TEXT("Resolve Layers Heightmaps"), STAT_LandscapeLayersResolveHeightmaps, STATGROUP_Landscape, );
 
-DECLARE_CYCLE_STAT_EXTERN(TEXT("Regenerate Procedural Weightmap (GameThread)"), STAT_LandscapeRegenerateProceduralWeightmaps, STATGROUP_Landscape, );
+DECLARE_CYCLE_STAT_EXTERN(TEXT("Regenerate Layers Weightmap (GameThread)"), STAT_LandscapeLayersRegenerateWeightmaps, STATGROUP_Landscape, );
 
-DECLARE_CYCLE_STAT_EXTERN(TEXT("Resolve Procedural Weightmap"), STAT_LandscapeResolveProceduralWeightmap, STATGROUP_Landscape, );
-DECLARE_CYCLE_STAT_EXTERN(TEXT("Reallocate Procedural Weightmaps"), STAT_LandscapeReallocateProceduralWeightmaps, STATGROUP_Landscape, );
-DECLARE_CYCLE_STAT_EXTERN(TEXT("Update Procedural Material Instances"), STAT_LandscapeProceduralUpdateMaterialInstance, STATGROUP_Landscape, );
+DECLARE_CYCLE_STAT_EXTERN(TEXT("Resolve Layers Weightmaps"), STAT_LandscapeLayersResolveWeightmaps, STATGROUP_Landscape, );
+DECLARE_CYCLE_STAT_EXTERN(TEXT("Reallocate Layers Weightmaps"), STAT_LandscapeLayersReallocateWeightmaps, STATGROUP_Landscape, );
+DECLARE_CYCLE_STAT_EXTERN(TEXT("Update Layers Material Instances"), STAT_LandscapeLayersUpdateMaterialInstance, STATGROUP_Landscape, );
 
-DECLARE_DWORD_COUNTER_STAT_EXTERN(TEXT("Regenerate Procedural DrawCalls"), STAT_LandscapeRegenerateProceduralDrawCalls, STATGROUP_Landscape, );
+DECLARE_DWORD_COUNTER_STAT_EXTERN(TEXT("Regenerate Layers DrawCalls"), STAT_LandscapeLayersRegenerateDrawCalls, STATGROUP_Landscape, );
 
 DECLARE_DWORD_COUNTER_STAT_EXTERN(TEXT("Components Using SubSection DrawCall"), STAT_LandscapeComponentUsingSubSectionDrawCalls, STATGROUP_Landscape, );
 DECLARE_DWORD_COUNTER_STAT_EXTERN(TEXT("Tessellated Shadow Cascade"), STAT_LandscapeTessellatedShadowCascade, STATGROUP_Landscape, );
