- Landscape Splines can now reserve a Landscape Layer (using the Landscape Layer System). In that use-case, their terraforming/terrapainting becomes non-destructive.
- Added automatic update of landscape deformation when landscape splines are modified (only in Landscape Layer System Mode)
- Added a new blend mode (alpha blend) to allow landscape splines to blend with previous landscape layers (blending and raise/lower flags is handled by GPU)

diff --git a/Engine/Shaders/Private/LandscapeLayersPS.usf b/Engine/Shaders/Private/LandscapeLayersPS.usf
--- a/Engine/Shaders/Private/LandscapeLayersPS.usf
+++ b/Engine/Shaders/Private/LandscapeLayersPS.usf
@@ -3,13 +3,13 @@
 #include "Common.ush"
 
 Texture2D ReadTexture1;
-Texture2D ReadTexture2;					// This one is optional, it's valid when OutputConfig.z == 1
+Texture2D ReadTexture2;							// This one is optional, it's valid when OutputConfig.z == 1
 SamplerState ReadTexture1Sampler;		
 SamplerState ReadTexture2Sampler;
-float2 LayerInfo;					// x == weight, y == visibility
-float4 OutputConfig;					// x == ApplyLayerModifiers, y == OutputAsDelta or Output weightmap as substractive/additive, z == Use ReadTexture2 and out Delta with ReadTexture1, w == Output Normals or Output weightmap as normalized weight blend
-float2 TextureSize;					// x == read texture 1 width, y == read texture 1 height
-float3 LandscapeGridScale;				// x == LS Actor DrawScale.X, y == LS Actor DrawScale.y, z == LS Actor DrawScale.z / 128.0f (ZSCALE)
+float4 LayerInfo;								// x == Layer Alpha value, y == Layer Visibility flag, z == Layer Blend Mode (0: Additive Blend Mode, 1: Alpha Blend Mode), w == Unused
+float4 OutputConfig;							// x == ApplyLayerModifiers, y == Alpha Blend Mode or Output weightmap as substractive/additive, z == Use ReadTexture2 and out Delta with ReadTexture1, w == Output Normals or Output weightmap as normalized weight blend
+float2 TextureSize;								// x == ReadTexture1 width, y == ReadTexture1 height
+float3 LandscapeGridScale;						// x == LS Actor DrawScale.X, y == LS Actor DrawScale.y, z == LS Actor DrawScale.z / 128.0f (ZSCALE)
 float CurrentMipComponentVertexCount;
 float2 CurrentMipTextureSize;
 float2 ParentMipTextureSize;
@@ -48,7 +48,7 @@ void PSHeightmapMain(float2 InTextureCoordinates : TEXCOORD0, out float4 OutColo
 {
 	float4 SourceColor = ReadTexture1.Sample(ReadTexture1Sampler, InTextureCoordinates);
 	float Height = ExtractHeight(SourceColor.rg);
-
+	
 	// Perform calculation 0 based
 	Height -= 32768.0f;
 
@@ -63,17 +63,34 @@ void PSHeightmapMain(float2 InTextureCoordinates : TEXCOORD0, out float4 OutColo
 	{
 		float2 BaseColor = ReadTexture2.Sample(ReadTexture2Sampler, InTextureCoordinates).rg;
 		float BaseHeight = ExtractHeight(BaseColor.rg);
-		Height = OutputConfig.y == 1 ? BaseHeight + (Height - BaseHeight) : BaseHeight + Height;
-	}
-	else
-	{
-		// Output as Delta to Source
-		if (OutputConfig.y == 0.0)
+		BaseHeight -= 32768.0;
+		if (LayerInfo.z == 1 && LayerInfo.y == 1)
 		{
-			Height += 32768.0f;
+			// Alpha Blend Mode
+			float HeightAlphaBlend = SourceColor.b;
+			float NewHeight = Height + (HeightAlphaBlend * BaseHeight);
+			uint RaiseLower = floor(SourceColor.a * 255.f + 0.5f);
+			bool bLowerTerrain = (RaiseLower & 1);
+			bool bRaiseTerrain = (RaiseLower & 2);
+			if ((bRaiseTerrain && NewHeight > BaseHeight) || (bLowerTerrain && NewHeight < BaseHeight))
+			{
+				Height = NewHeight;
+			}
+			else
+			{
+				Height = BaseHeight;
+			}
+		}
+		else
+		{
+			// Additive Blend Mode
+			Height += BaseHeight;
 		}
+		Height = clamp(Height, -32768.0, 32767.0);
 	}
 
+	Height += 32768.0f;
+
 	float2 PackedHeight = PackHeight(Height);
 
 	OutColor.r = PackedHeight.x;
diff --git a/Engine/Source/Editor/LandscapeEditor/Private/LandscapeEdMode.cpp b/Engine/Source/Editor/LandscapeEditor/Private/LandscapeEdMode.cpp
--- a/Engine/Source/Editor/LandscapeEditor/Private/LandscapeEdMode.cpp
+++ b/Engine/Source/Editor/LandscapeEditor/Private/LandscapeEdMode.cpp
@@ -1900,7 +1900,11 @@ bool FEdModeLandscape::InputKey(FEditorViewportClient* ViewportClient, FViewport
 						ALandscape* Landscape = GetLandscape();
 						bool bLandscapeLayerSystem = GetMutableDefault<UEditorExperimentalSettings>()->bLandscapeLayerSystem;
 
-						if (bLandscapeLayerSystem && GetCurrentLayer() && !GetCurrentLayer()->bVisible)
+						if (bLandscapeLayerSystem && !GetCurrentLayer())
+						{
+							FMessageDialog::Open(EAppMsgType::Ok, NSLOCTEXT("UnrealEd", "LandscapeInvalidLayer", "No layer selected. You must first chose a layer to work on."));
+						}
+						else if (bLandscapeLayerSystem && GetCurrentLayer() && !GetCurrentLayer()->bVisible)
 						{
 							FMessageDialog::Open(EAppMsgType::Ok, NSLOCTEXT("UnrealEd", "LandscapeLayerHidden", "Painting in a hidden layer is not allowed."));
 						}
@@ -1908,6 +1912,10 @@ bool FEdModeLandscape::InputKey(FEditorViewportClient* ViewportClient, FViewport
 						{
 							FMessageDialog::Open(EAppMsgType::Ok, NSLOCTEXT("UnrealEd", "LandscapeLayerLocked", "This layer is locked. You must unlock it before you can work on this layer."));
 						}
+						else if (bLandscapeLayerSystem && (CurrentTool != (FLandscapeTool*)SplinesTool) && GetCurrentLayer() && Landscape && (GetCurrentLayer() == Landscape->GetLandscapeSplinesReservedLayer()))
+						{
+							FMessageDialog::Open(EAppMsgType::Ok, NSLOCTEXT("UnrealEd", "LandscapeLayerReservedForSplines", "This layer is reserved for Landscape Splines."));
+						}
 						else if (CurrentToolTarget.TargetType == ELandscapeToolTargetType::Weightmap && CurrentToolTarget.LayerInfo == NULL && CurrentToolName != TEXT("BPCustom"))
 						{
 							FMessageDialog::Open(EAppMsgType::Ok,
@@ -4183,6 +4191,11 @@ void FEdModeLandscape::SetLayerName(int32 InLayerIndex, const FName& InName)
 	}
 }
 
+bool FEdModeLandscape::IsLayerAlphaVisible(int32 InLayerIndex) const
+{
+	return (CurrentToolTarget.TargetType == ELandscapeToolTargetType::Heightmap || CurrentToolTarget.TargetType == ELandscapeToolTargetType::Weightmap);
+}
+
 float FEdModeLandscape::GetLayerAlpha(int32 InLayerIndex) const
 {
 	FLandscapeLayer* Layer = GetLayer(InLayerIndex);
@@ -4198,7 +4211,10 @@ void FEdModeLandscape::SetLayerAlpha(int32 InLayerIndex, float InAlpha)
 	ALandscape* Landscape = GetLandscape();
 	if (Landscape)
 	{
-		Landscape->SetLayerAlpha(InLayerIndex, InAlpha, CurrentToolTarget.TargetType == ELandscapeToolTargetType::Heightmap);
+		if (CurrentToolTarget.TargetType == ELandscapeToolTargetType::Heightmap || CurrentToolTarget.TargetType == ELandscapeToolTargetType::Weightmap)
+		{
+			Landscape->SetLayerAlpha(InLayerIndex, InAlpha, CurrentToolTarget.TargetType == ELandscapeToolTargetType::Heightmap);
+		}
 	}
 }
 
@@ -4219,8 +4235,10 @@ void FEdModeLandscape::SetLayerVisibility(bool bInVisible, int32 InLayerIndex)
 
 bool FEdModeLandscape::IsLayerLocked(int32 InLayerIndex) const
 {
-	FLandscapeLayer* Layer = GetLayer(InLayerIndex);
-	return Layer ? Layer->bLocked : false;
+	ALandscape* Landscape = GetLandscape();
+	const FLandscapeLayer* Layer = GetLayer(InLayerIndex);
+	const FLandscapeLayer* LayerReservedForSplines = Landscape ? Landscape->GetLandscapeSplinesReservedLayer() : nullptr;
+	return Layer ? (Layer->bLocked || Layer == LayerReservedForSplines) : false;
 }
 
 void FEdModeLandscape::SetLayerLocked(int32 InLayerIndex, bool bInLocked)
@@ -4350,6 +4368,39 @@ FLandscapeLayer* FEdModeLandscape::GetCurrentLayer() const
 	return GetLayer(GetCurrentLayerIndex());
 }
 
+void FEdModeLandscape::AutoUpdateDirtyLandscapeSplines()
+{
+	if (GetMutableDefault<UEditorExperimentalSettings>()->bLandscapeLayerSystem && GEditor->IsTransactionActive())
+	{
+		// Only auto-update if a layer is reserved for landscape splines
+		ALandscape* Landscape = GetLandscape();
+		if (Landscape && Landscape->GetLandscapeSplinesReservedLayer())
+		{
+			// TODO : Only update dirty regions
+			UpdateLandscapeSplines();
+		}
+	}
+}
+
+void FEdModeLandscape::UpdateLandscapeSplines(bool bUpdateOnlySelected /* = false*/)
+{
+	if (GetMutableDefault<UEditorExperimentalSettings>()->bLandscapeLayerSystem)
+	{
+		ALandscape* Landscape = GetLandscape();
+		if (Landscape)
+		{
+			Landscape->UpdateLandscapeSplines(GetCurrentLayerGuid(), bUpdateOnlySelected);
+		}
+	}
+	else
+	{
+		if (CurrentToolTarget.LandscapeInfo.IsValid())
+		{
+			CurrentToolTarget.LandscapeInfo->ApplySplines(bUpdateOnlySelected);
+		}
+	}
+}
+
 FGuid FEdModeLandscape::GetCurrentLayerGuid() const
 {
 	FLandscapeLayer* CurrentLayer = GetCurrentLayer();
@@ -4450,4 +4501,4 @@ FName FLandscapeTargetListInfo::GetLayerName() const
 	return LayerInfoObj.IsValid() ? LayerInfoObj->LayerName : LayerName;
 }
 
-#undef LOCTEXT_NAMESPACE
+#undef LOCTEXT_NAMESPACE
\ No newline at end of file
diff --git a/Engine/Source/Editor/LandscapeEditor/Private/LandscapeEdMode.h b/Engine/Source/Editor/LandscapeEditor/Private/LandscapeEdMode.h
--- a/Engine/Source/Editor/LandscapeEditor/Private/LandscapeEdMode.h
+++ b/Engine/Source/Editor/LandscapeEditor/Private/LandscapeEdMode.h
@@ -503,11 +503,14 @@ public:
 	void SetLayerVisibility(bool InVisible, int32 InLayerIndex);
 	bool IsLayerVisible(int32 InLayerIndex) const;
 	bool IsLayerLocked(int32 InLayerIndex) const;
+	bool IsLayerAlphaVisible(int32 InLayerIndex) const;
 	void SetLayerLocked(int32 InLayerIndex, bool bInLocked);
 	struct FLandscapeLayer* GetCurrentLayer() const;
 	FGuid GetCurrentLayerGuid() const;
 	bool IsCurrentLayerBlendSubstractive(const TWeakObjectPtr<ULandscapeLayerInfoObject>& InLayerInfoObj) const;
 	void SetCurrentLayerSubstractiveBlendStatus(bool InStatus, const TWeakObjectPtr<ULandscapeLayerInfoObject>& InLayerInfoObj);
+	void UpdateLandscapeSplines(bool bUpdateOnlySelected = false);
+	void AutoUpdateDirtyLandscapeSplines();
 
 	void AddBrushToCurrentLayer(int32 InTargetType, class ALandscapeBlueprintCustomBrush* InBrush);
 	void RemoveBrushFromCurrentLayer(int32 InTargetType, class ALandscapeBlueprintCustomBrush* InBrush);
@@ -583,4 +586,4 @@ private:
 
 	/** The interactor that is currently painting, prevents multiple interactors from sculpting when one actually is */
 	class UViewportInteractor* InteractorPainting;
-};
+};
\ No newline at end of file
diff --git a/Engine/Source/Editor/LandscapeEditor/Private/LandscapeEdModeSplineTools.cpp b/Engine/Source/Editor/LandscapeEditor/Private/LandscapeEdModeSplineTools.cpp
--- a/Engine/Source/Editor/LandscapeEditor/Private/LandscapeEdModeSplineTools.cpp
+++ b/Engine/Source/Editor/LandscapeEditor/Private/LandscapeEdModeSplineTools.cpp
@@ -491,6 +491,7 @@ public:
 		SelectControlPoint(NewControlPoint);
 		UpdatePropertiesWindows();
 
+		EdMode->AutoUpdateDirtyLandscapeSplines();
 		if (!SplinesComponent->IsRegistered())
 		{
 			SplinesComponent->RegisterComponent();
@@ -541,6 +542,8 @@ public:
 			ToDelete->Connections[1].ControlPoint->UpdateSplinePoints();
 		}
 
+
+		EdMode->AutoUpdateDirtyLandscapeSplines();
 		SplinesComponent->MarkRenderStateDirty();
 	}
 
@@ -584,6 +587,7 @@ public:
 				ToDelete->ConnectedSegments.Empty();
 
 				SplinesComponent->ControlPoints.Remove(ToDelete);
+				EdMode->AutoUpdateDirtyLandscapeSplines();
 				SplinesComponent->MarkRenderStateDirty();
 
 				return;
@@ -622,6 +626,7 @@ public:
 		ToDelete->ConnectedSegments.Empty();
 
 		SplinesComponent->ControlPoints.Remove(ToDelete);
+		EdMode->AutoUpdateDirtyLandscapeSplines();
 		SplinesComponent->MarkRenderStateDirty();
 	}
 
@@ -1208,6 +1213,7 @@ public:
 					Segment->Connections[1].ControlPoint->UpdateSplinePoints();
 				}
 
+				EdMode->AutoUpdateDirtyLandscapeSplines();
 				return true;
 			}
 		}
@@ -1223,6 +1229,7 @@ public:
 					FlipSegment(Segment);
 				}
 
+				EdMode->AutoUpdateDirtyLandscapeSplines();
 				return true;
 			}
 		}
@@ -1247,6 +1254,7 @@ public:
 					Segment->Connections[1].ControlPoint->UpdateSplinePoints();
 				}
 
+				EdMode->AutoUpdateDirtyLandscapeSplines();
 				return true;
 			}
 		}
@@ -1268,6 +1276,7 @@ public:
 				}
 				UpdatePropertiesWindows();
 
+				EdMode->AutoUpdateDirtyLandscapeSplines();
 				GUnrealEd->RedrawLevelEditingViewports();
 				return true;
 			}
@@ -1335,6 +1344,7 @@ public:
 							AddSegment(ControlPoint, ClickedControlPoint, bAutoRotateOnJoin, bAutoRotateOnJoin);
 						}
 
+						EdMode->AutoUpdateDirtyLandscapeSplines();
 						GUnrealEd->RedrawLevelEditingViewports();
 
 						return true;
@@ -1478,6 +1488,7 @@ public:
 
 					ResetAllowDuplication();
 
+					EdMode->AutoUpdateDirtyLandscapeSplines();
 					GEditor->EndTransaction();
 
 					return false; // We're not actually handling this case ourselves, just wrapping it in a transaction
@@ -1491,6 +1502,7 @@ public:
 					InViewportClient->SetWidgetCoordSystemSpace(DraggingTangent_CacheCoordSpace);
 					InViewportClient->SetRequiredCursorOverride(false);
 
+					EdMode->AutoUpdateDirtyLandscapeSplines();
 					GEditor->EndTransaction();
 
 					return false; // false to let FEditorViewportClient.InputKey end mouse tracking
diff --git a/Engine/Source/Editor/LandscapeEditor/Private/LandscapeEditorDetailCustomization_Layers.cpp b/Engine/Source/Editor/LandscapeEditor/Private/LandscapeEditorDetailCustomization_Layers.cpp
--- a/Engine/Source/Editor/LandscapeEditor/Private/LandscapeEditorDetailCustomization_Layers.cpp
+++ b/Engine/Source/Editor/LandscapeEditor/Private/LandscapeEditorDetailCustomization_Layers.cpp
@@ -101,7 +101,7 @@ void FLandscapeEditorCustomNodeBuilder_Layers::GenerateHeaderRowContent(FDetailW
 		[
 			SNew(STextBlock)
 			.Font(IDetailLayoutBuilder::GetDetailFont())
-		.Text(FText::FromString(TEXT("Layers")))
+			.Text(FText::FromString(TEXT("Layers")))
 		];
 }
 
@@ -157,96 +157,104 @@ TSharedPtr<SWidget> FLandscapeEditorCustomNodeBuilder_Layers::GenerateRow(int32
 		.Visibility(EVisibility::Visible)
 		[
 			SNew(SHorizontalBox)
-
+			
 			+ SHorizontalBox::Slot()
-		.AutoWidth()
-		.VAlign(VAlign_Center)
-		[
-			SNew(SButton)
-			.ButtonStyle(FEditorStyle::Get(), "NoBorder")
-		.OnClicked(this, &FLandscapeEditorCustomNodeBuilder_Layers::OnToggleLock, InLayerIndex)
-		.ToolTipText(LOCTEXT("FLandscapeEditorCustomNodeBuilder_LayerLock", "Locks the current layer"))
-		[
-			SNew(SImage)
-			.Image(this, &FLandscapeEditorCustomNodeBuilder_Layers::GetLockBrushForLayer, InLayerIndex)
-		]
-		]
-
-	+ SHorizontalBox::Slot()
-		.AutoWidth()
-		.VAlign(VAlign_Center)
-		[
-			SNew(SButton)
-			.ContentPadding(0)
-		.ButtonStyle(FEditorStyle::Get(), "NoBorder")
-		.OnClicked(this, &FLandscapeEditorCustomNodeBuilder_Layers::OnToggleVisibility, InLayerIndex)
-		.ToolTipText(LOCTEXT("FLandscapeEditorCustomNodeBuilder_LayerVisibility", "Toggle Layer Visibility"))
-		.HAlign(HAlign_Center)
-		.VAlign(VAlign_Center)
-		.Content()
-		[
-			SNew(SImage)
+			.AutoWidth()
+			.VAlign(VAlign_Center)
+			[
+				SNew(SButton)
+				.ButtonStyle(FEditorStyle::Get(), "NoBorder")
+				.OnClicked(this, &FLandscapeEditorCustomNodeBuilder_Layers::OnToggleLock, InLayerIndex)
+				.ToolTipText(LOCTEXT("LandscapeLayerLock", "Locks the current layer"))
+				[
+					SNew(SImage)
+					.Image(this, &FLandscapeEditorCustomNodeBuilder_Layers::GetLockBrushForLayer, InLayerIndex)
+				]
+			]
+
+			+SHorizontalBox::Slot()
+			.AutoWidth()
+			.VAlign(VAlign_Center)
+			[
+				SNew(SButton)
+				.ContentPadding(0)
+				.ButtonStyle(FEditorStyle::Get(), "NoBorder")
+				.OnClicked(this, &FLandscapeEditorCustomNodeBuilder_Layers::OnToggleVisibility, InLayerIndex)
+				.ToolTipText(LOCTEXT("LandscapeLayerVisibility", "Toggle Layer Visibility"))
+				.HAlign(HAlign_Center)
+				.VAlign(VAlign_Center)
+				.Content()
+				[
+					SNew(SImage)
 			.Image(this, &FLandscapeEditorCustomNodeBuilder_Layers::GetVisibilityBrushForLayer, InLayerIndex)
-		]
-		]
+				]
+			]
 
-	+ SHorizontalBox::Slot()
-		.FillWidth(1.0)
-		.VAlign(VAlign_Center)
-		.Padding(4, 0)
-		[
-			SAssignNew(InlineTextBlocks[InLayerIndex], SInlineEditableTextBlock)
-			.Text(this, &FLandscapeEditorCustomNodeBuilder_Layers::GetLayerText, InLayerIndex)
-		.ToolTipText(LOCTEXT("FLandscapeEditorCustomNodeBuilder_Layers_tooltip", "Name of the Layer"))
-		.OnVerifyTextChanged(FOnVerifyTextChanged::CreateSP(this, &FLandscapeEditorCustomNodeBuilder_Layers::CanRenameLayerTo, InLayerIndex))
-		.OnTextCommitted(FOnTextCommitted::CreateSP(this, &FLandscapeEditorCustomNodeBuilder_Layers::SetLayerName, InLayerIndex))
-		]
-
-	+ SHorizontalBox::Slot()
-		.FillWidth(1.0f)
-		.VAlign(VAlign_Center)
-		.Padding(0, 2)
-		.HAlign(HAlign_Right)
-		[
-			SNew(SHorizontalBox)
 			+ SHorizontalBox::Slot()
-		.Padding(0)
-		.FillWidth(1.0f)
-		.VAlign(VAlign_Center)
-		.HAlign(HAlign_Left)
-		[
-			SNew(STextBlock)
-			.Text(LOCTEXT("FLandscapeEditorCustomNodeBuilder_LayerAlpha", "Alpha"))
-		]
-	+ SHorizontalBox::Slot()
-		.VAlign(VAlign_Center)
-		.Padding(0, 2)
-		.HAlign(HAlign_Left)
-		.FillWidth(1.0f)
-		[
-			SNew(SNumericEntryBox<float>)
-			.AllowSpin(true)
-		.MinValue(0.0f)
-		.MaxValue(100.0f)
-		.MaxSliderValue(100.0f)
-		.MinDesiredValueWidth(60.0f)
-		.Value(this, &FLandscapeEditorCustomNodeBuilder_Layers::GetLayerAlpha, InLayerIndex)
-		.OnValueChanged(this, &FLandscapeEditorCustomNodeBuilder_Layers::SetLayerAlpha, InLayerIndex)
-		.IsEnabled(true)
-		]
-		]
-		];
+			.FillWidth(1.0)
+			.VAlign(VAlign_Center)
+			.Padding(4, 0)
+			[
+				SAssignNew(InlineTextBlocks[InLayerIndex], SInlineEditableTextBlock)
+				.IsEnabled(this, &FLandscapeEditorCustomNodeBuilder_Layers::IsLayerEditionEnabled, InLayerIndex)
+				.Text(this, &FLandscapeEditorCustomNodeBuilder_Layers::GetLayerDisplayName, InLayerIndex)
+				.ToolTipText(LOCTEXT("LandscapeLayers_tooltip", "Name of the Layer"))
+				.OnVerifyTextChanged(FOnVerifyTextChanged::CreateSP(this, &FLandscapeEditorCustomNodeBuilder_Layers::CanRenameLayerTo, InLayerIndex))
+				.OnTextCommitted(FOnTextCommitted::CreateSP(this, &FLandscapeEditorCustomNodeBuilder_Layers::SetLayerName, InLayerIndex))
+			]
+
+			+ SHorizontalBox::Slot()
+			.FillWidth(1.0f)
+			.VAlign(VAlign_Center)
+			.Padding(0, 2)
+			.HAlign(HAlign_Right)
+			[
+				SNew(SHorizontalBox)
+				+ SHorizontalBox::Slot()
+				.Padding(0)
+				.FillWidth(1.0f)
+				.VAlign(VAlign_Center)
+				.HAlign(HAlign_Left)
+				[
+					SNew(STextBlock)
+					.IsEnabled(this, &FLandscapeEditorCustomNodeBuilder_Layers::IsLayerEditionEnabled, InLayerIndex)
+					.Visibility(this, &FLandscapeEditorCustomNodeBuilder_Layers::GetLayerAlphaVisibility, InLayerIndex)
+					.Text(LOCTEXT("LandscapeLayerAlpha", "Alpha"))
+				]
+				+ SHorizontalBox::Slot()
+				.VAlign(VAlign_Center)
+				.Padding(0, 2)
+				.HAlign(HAlign_Left)
+				.FillWidth(1.0f)
+				[
+					SNew(SNumericEntryBox<float>)
+					.AllowSpin(true)
+					.MinValue(0.0f)
+					.MaxValue(100.0f)
+					.MaxSliderValue(100.0f)
+					.MinDesiredValueWidth(60.0f)
+					.IsEnabled(this, &FLandscapeEditorCustomNodeBuilder_Layers::IsLayerEditionEnabled, InLayerIndex)
+					.Visibility(this, &FLandscapeEditorCustomNodeBuilder_Layers::GetLayerAlphaVisibility, InLayerIndex)
+					.Value(this, &FLandscapeEditorCustomNodeBuilder_Layers::GetLayerAlpha, InLayerIndex)
+					.OnValueChanged(this, &FLandscapeEditorCustomNodeBuilder_Layers::SetLayerAlpha, InLayerIndex)
+				]		
+			]
+		];	
 
 	return RowWidget;
 }
 END_SLATE_FUNCTION_BUILD_OPTIMIZATION
 
-FText FLandscapeEditorCustomNodeBuilder_Layers::GetLayerText(int32 InLayerIndex) const
+FText FLandscapeEditorCustomNodeBuilder_Layers::GetLayerDisplayName(int32 InLayerIndex) const
 {
 	FEdModeLandscape* LandscapeEdMode = GetEditorMode();
-	if (LandscapeEdMode)
+	ALandscape* Landscape = LandscapeEdMode ? LandscapeEdMode->GetLandscape() : nullptr;
+	if (LandscapeEdMode && Landscape)
 	{
-		return FText::FromName(LandscapeEdMode->GetLayerName(InLayerIndex));
+		const FLandscapeLayer* Layer = LandscapeEdMode->GetLayer(InLayerIndex);
+		const FLandscapeLayer* ReservedLayer = Landscape->GetLandscapeSplinesReservedLayer();
+		bool bIsReserved = Layer && Layer == ReservedLayer;
+		return FText::Format(LOCTEXT("DisplayName", "{0} {1}"), FText::FromName(LandscapeEdMode->GetLayerName(InLayerIndex)), bIsReserved ? LOCTEXT("ReservedForSplines", "(Reserved for Splines)") : FText::GetEmpty());
 	}
 
 	return FText::FromString(TEXT("None"));
@@ -270,7 +278,7 @@ bool FLandscapeEditorCustomNodeBuilder_Layers::CanRenameLayerTo(const FText& InN
 	{
 		if (!LandscapeEdMode->CanRenameLayerTo(InLayerIndex, *InNewText.ToString()))
 		{
-			OutErrorMessage = LOCTEXT("RenameFailed_AlreadyExists", "This layer already exists");
+			OutErrorMessage = LOCTEXT("Landscape_Layers_RenameFailed_AlreadyExists", "This layer already exists");
 			return false;
 		}
 	}
@@ -317,10 +325,22 @@ TSharedPtr<SWidget> FLandscapeEditorCustomNodeBuilder_Layers::OnLayerContextMenu
 					if (Landscape->LandscapeLayers.Num() > 1)
 					{
 						// Delete Layer
-						FUIAction DeleteLayerAction = FUIAction(FExecuteAction::CreateLambda([SharedThis, InLayerIndex] { SharedThis->DeleteLayer(InLayerIndex); }));
+						FUIAction DeleteLayerAction = FUIAction(FExecuteAction::CreateLambda([SharedThis, InLayerIndex] { SharedThis->DeleteLayer(InLayerIndex); } ));
 						MenuBuilder.AddMenuEntry(LOCTEXT("DeleteLayer", "Delete..."), LOCTEXT("DeleteLayerTooltip", "Delete Layer"), FSlateIcon(), DeleteLayerAction);
 					}
 				}
+
+				if (Landscape->GetLandscapeSplinesReservedLayer() != Landscape->GetLayer(InLayerIndex))
+				{
+					// Reserve for Landscape Splines
+					FUIAction ReserveLayerAction = FUIAction(FExecuteAction::CreateLambda([SharedThis, InLayerIndex] { SharedThis->SetLandscapeSplinesReservedLayer(InLayerIndex); }));
+					MenuBuilder.AddMenuEntry(LOCTEXT("ReserveLayerForSplines", "Reserve for Splines"), LOCTEXT("ReserveLayerForSplinesTooltip", "Reserve Layer for Landscape Splines"), FSlateIcon(), ReserveLayerAction);
+				}
+				else
+				{
+					FUIAction RemoveReserveLayerAction = FUIAction(FExecuteAction::CreateLambda([SharedThis, InLayerIndex] { SharedThis->SetLandscapeSplinesReservedLayer(INDEX_NONE); }));
+					MenuBuilder.AddMenuEntry(LOCTEXT("RemoveReserveLayerForSplines", "Remove Reserve for Splines"), LOCTEXT("RemoveReserveLayerForSplinesTooltip", "Remove reservation of Layer for Landscape Splines"), FSlateIcon(), RemoveReserveLayerAction);
+				}
 			}
 		}
 		MenuBuilder.EndSection();
@@ -358,6 +378,37 @@ TSharedPtr<SWidget> FLandscapeEditorCustomNodeBuilder_Layers::OnLayerContextMenu
 	return NULL;
 }
 
+void FLandscapeEditorCustomNodeBuilder_Layers::SetLandscapeSplinesReservedLayer(int32 InLayerIndex)
+{
+	FEdModeLandscape* LandscapeEdMode = GetEditorMode();
+	ALandscape* Landscape = LandscapeEdMode ? LandscapeEdMode->GetLandscape() : nullptr;
+	if (Landscape)
+	{
+		const FLandscapeLayer* CurrenReservedLayer = Landscape->GetLandscapeSplinesReservedLayer();
+		const FLandscapeLayer* NewReservedLayer = Landscape->GetLayer(InLayerIndex);
+		if (NewReservedLayer != CurrenReservedLayer)
+		{
+			EAppReturnType::Type Result = EAppReturnType::No;
+			if (NewReservedLayer)
+			{
+				Result = FMessageDialog::Open(EAppMsgType::YesNo, FText::Format(LOCTEXT("Landscape_SetReservedForSplines_Message", "Reserving layer {0} for landscape splines will clear it from its content and no edition will be allowed.  Continue?"), FText::FromName(NewReservedLayer->Name)));
+			}
+			else if (CurrenReservedLayer)
+			{
+				Result = FMessageDialog::Open(EAppMsgType::YesNo, FText::Format(LOCTEXT("Landscape_UnsetReservedForSplines_Message", "Removing reservation of layer {0} for landscape splines will clear its content.  Continue?"), FText::FromName(CurrenReservedLayer->Name)));
+			}
+
+			if (Result == EAppReturnType::Yes)
+			{
+				const FScopedTransaction Transaction(LOCTEXT("ReserveForSplines", "Reserve Landscape Layer for Splines"));
+				Landscape->SetLandscapeSplinesReservedLayer(InLayerIndex);
+				LandscapeEdMode->RefreshDetailPanel();
+				LandscapeEdMode->AutoUpdateDirtyLandscapeSplines();
+			}
+		}
+	}
+}
+
 void FLandscapeEditorCustomNodeBuilder_Layers::RenameLayer(int32 InLayerIndex)
 {
 	if (InlineTextBlocks.IsValidIndex(InLayerIndex) && InlineTextBlocks[InLayerIndex].IsValid())
@@ -375,7 +426,7 @@ void FLandscapeEditorCustomNodeBuilder_Layers::ClearLayer(int32 InLayerIndex)
 		FLandscapeLayer* Layer = LandscapeEdMode->GetLayer(InLayerIndex);
 		if (Layer)
 		{
-			EAppReturnType::Type Result = FMessageDialog::Open(EAppMsgType::YesNo, FText::Format(LOCTEXT("LandscapeMode_Message", "The layer {0} content will be completely cleared.  Continue?"), FText::FromName(Layer->Name)));
+			EAppReturnType::Type Result = FMessageDialog::Open(EAppMsgType::YesNo, FText::Format(LOCTEXT("Landscape_ClearLayer_Message", "The layer {0} content will be completely cleared.  Continue?"), FText::FromName(Layer->Name)));
 			if (Result == EAppReturnType::Yes)
 			{
 				const FScopedTransaction Transaction(LOCTEXT("Landscape_Layers_Clean", "Clear Layer"));
@@ -394,7 +445,7 @@ void FLandscapeEditorCustomNodeBuilder_Layers::DeleteLayer(int32 InLayerIndex)
 		FLandscapeLayer* Layer = LandscapeEdMode->GetLayer(InLayerIndex);
 		if (Layer)
 		{
-			EAppReturnType::Type Result = FMessageDialog::Open(EAppMsgType::YesNo, FText::Format(LOCTEXT("LandscapeMode_Message", "The layer {0} will be deleted.  Continue?"), FText::FromName(Layer->Name)));
+			EAppReturnType::Type Result = FMessageDialog::Open(EAppMsgType::YesNo, FText::Format(LOCTEXT("Landscape_DeleteLayer_Message", "The layer {0} will be deleted.  Continue?"), FText::FromName(Layer->Name)));
 			if (Result == EAppReturnType::Yes)
 			{
 				const FScopedTransaction Transaction(LOCTEXT("Landscape_Layers_Delete", "Delete Layer"));
@@ -503,6 +554,20 @@ FReply FLandscapeEditorCustomNodeBuilder_Layers::OnToggleLock(int32 InLayerIndex
 	return FReply::Handled();
 }
 
+EVisibility FLandscapeEditorCustomNodeBuilder_Layers::GetLayerAlphaVisibility(int32 InLayerIndex) const
+{
+	FEdModeLandscape* LandscapeEdMode = GetEditorMode();
+	bool bIsVisible = LandscapeEdMode && LandscapeEdMode->IsLayerAlphaVisible(InLayerIndex);
+	return bIsVisible ? EVisibility::Visible : EVisibility::Hidden;
+}
+
+bool FLandscapeEditorCustomNodeBuilder_Layers::IsLayerEditionEnabled(int32 InLayerIndex) const
+{
+	FEdModeLandscape* LandscapeEdMode = GetEditorMode();
+	bool bIsLocked = LandscapeEdMode && LandscapeEdMode->IsLayerLocked(InLayerIndex);
+	return !bIsLocked;
+}
+
 const FSlateBrush* FLandscapeEditorCustomNodeBuilder_Layers::GetLockBrushForLayer(int32 InLayerIndex) const
 {
 	FEdModeLandscape* LandscapeEdMode = GetEditorMode();
@@ -534,4 +599,4 @@ FReply FLandscapeEditorCustomNodeBuilder_Layers::HandleAcceptDrop(FDragDropEvent
 	return FReply::Unhandled();
 }
 
-#undef LOCTEXT_NAMESPACE
+#undef LOCTEXT_NAMESPACE
\ No newline at end of file
diff --git a/Engine/Source/Editor/LandscapeEditor/Private/LandscapeEditorDetailCustomization_Layers.h b/Engine/Source/Editor/LandscapeEditor/Private/LandscapeEditorDetailCustomization_Layers.h
--- a/Engine/Source/Editor/LandscapeEditor/Private/LandscapeEditorDetailCustomization_Layers.h
+++ b/Engine/Source/Editor/LandscapeEditor/Private/LandscapeEditorDetailCustomization_Layers.h
@@ -76,6 +76,10 @@ protected:
 	bool CanRenameLayerTo(const FText& NewText, FText& OutErrorMessage, int32 InLayerIndex);
 	void SetLayerName(const FText& InText, ETextCommit::Type InCommitType, int32 InLayerIndex);
 	FText GetLayerText(int32 InLayerIndex) const;
+	void SetLandscapeSplinesReservedLayer(int32 InLayerIndex);
+	FText GetLayerDisplayName(int32 InLayerIndex) const;
+	bool IsLayerEditionEnabled(int32 InLayerIndex) const;
+	EVisibility GetLayerAlphaVisibility(int32 InLayerIndex) const;
 
 	TOptional<float> GetLayerAlpha(int32 InLayerIndex) const;
 	void SetLayerAlpha(float InAlpha, int32 InLayerIndex);
diff --git a/Engine/Source/Editor/LandscapeEditor/Private/LandscapeEditorDetailCustomization_MiscTools.cpp b/Engine/Source/Editor/LandscapeEditor/Private/LandscapeEditorDetailCustomization_MiscTools.cpp
--- a/Engine/Source/Editor/LandscapeEditor/Private/LandscapeEditorDetailCustomization_MiscTools.cpp
+++ b/Engine/Source/Editor/LandscapeEditor/Private/LandscapeEditorDetailCustomization_MiscTools.cpp
@@ -303,20 +303,19 @@ FReply FLandscapeEditorDetailCustomization_MiscTools::OnClearRegionSelectionButt
 FReply FLandscapeEditorDetailCustomization_MiscTools::OnApplyAllSplinesButtonClicked()
 {
 	FEdModeLandscape* LandscapeEdMode = GetEditorMode();
-	if (LandscapeEdMode && LandscapeEdMode->CurrentToolTarget.LandscapeInfo.IsValid())
+	if (LandscapeEdMode)
 	{
-		LandscapeEdMode->CurrentToolTarget.LandscapeInfo->ApplySplines(false);
+		LandscapeEdMode->UpdateLandscapeSplines(false);
 	}
-
 	return FReply::Handled();
 }
 
 FReply FLandscapeEditorDetailCustomization_MiscTools::OnApplySelectedSplinesButtonClicked()
 {
 	FEdModeLandscape* LandscapeEdMode = GetEditorMode();
-	if (LandscapeEdMode && LandscapeEdMode->CurrentToolTarget.LandscapeInfo.IsValid())
+	if (LandscapeEdMode)
 	{
-		LandscapeEdMode->CurrentToolTarget.LandscapeInfo->ApplySplines(true);
+		LandscapeEdMode->UpdateLandscapeSplines(true);
 	}
 
 	return FReply::Handled();
diff --git a/Engine/Source/Runtime/Engine/Public/Raster.h b/Engine/Source/Runtime/Engine/Public/Raster.h
--- a/Engine/Source/Runtime/Engine/Public/Raster.h
+++ b/Engine/Source/Runtime/Engine/Public/Raster.h
@@ -23,6 +23,14 @@ public:
 		InterpolantType	Interpolants[3] = { I0, I1, I2 };
 		FVector2D		Points[3] = { P0, P1, P2 };
 
+		// Avoid any division by zero
+		if (FMath::IsNearlyEqual(Points[1].Y, Points[0].Y) ||
+			FMath::IsNearlyEqual(Points[2].Y, Points[0].Y) ||
+			FMath::IsNearlyEqual(Points[2].Y, Points[1].Y))
+		{
+			return;
+		}
+
 		// Find the top point.
 
 		if(Points[1].Y < Points[0].Y && Points[1].Y <= Points[2].Y)
diff --git a/Engine/Source/Runtime/Landscape/Classes/Landscape.h b/Engine/Source/Runtime/Landscape/Classes/Landscape.h
--- a/Engine/Source/Runtime/Landscape/Classes/Landscape.h
+++ b/Engine/Source/Runtime/Landscape/Classes/Landscape.h
@@ -139,6 +139,14 @@ struct FLandscapeLayerBrush
 	ALandscapeBlueprintCustomBrush* BPCustomBrush;
 };
 
+UENUM()
+enum ELandscapeBlendMode
+{
+	LSBM_AdditiveBlend,
+	LSBM_AlphaBlend,
+	LSBM_MAX,
+};
+
 USTRUCT()
 struct FLandscapeLayer
 {
@@ -151,6 +159,7 @@ struct FLandscapeLayer
 		, bLocked(false)
 		, HeightmapAlpha(1.0f)
 		, WeightmapAlpha(1.0f)
+		, BlendMode(LSBM_AdditiveBlend)
 	{}
 
 	UPROPERTY(meta = (IgnoreForMemberInitializationTest))
@@ -171,6 +180,9 @@ struct FLandscapeLayer
 	UPROPERTY()
 	float WeightmapAlpha;
 
+	UPROPERTY()
+	TEnumAsByte<enum ELandscapeBlendMode> BlendMode;
+
 	UPROPERTY()
 	TArray<FLandscapeLayerBrush> Brushes;
 
@@ -237,12 +249,18 @@ public:
 	LANDSCAPE_API void SetLayerVisibility(int32 InLayerIndex, bool bInVisible);
 	LANDSCAPE_API struct FLandscapeLayer* GetLayer(int32 InLayerIndex);
 	LANDSCAPE_API const struct FLandscapeLayer* GetLayer(int32 InLayerIndex) const;
-	LANDSCAPE_API void ClearLayer(int32 InLayerIndex);
-	LANDSCAPE_API void ClearLayer(const FGuid& InLayerGuid);
+	LANDSCAPE_API const struct FLandscapeLayer* GetLayer(const FGuid& InLayerGuid) const;
+	LANDSCAPE_API void ClearLayer(int32 InLayerIndex, bool bInUpdateCollision = true);
+	LANDSCAPE_API void ClearLayer(const FGuid& InLayerGuid, bool bInUpdateCollision = true);
 	LANDSCAPE_API void DeleteLayer(int32 InLayerIndex);
-	LANDSCAPE_API void SetEditingLayer(FGuid InLayerGuid = FGuid());
+	LANDSCAPE_API void SetEditingLayer(const FGuid& InLayerGuid = FGuid());
 	LANDSCAPE_API void ShowOnlySelectedLayer(int32 InLayerIndex);
 	LANDSCAPE_API void ShowAllLayers();
+	LANDSCAPE_API void UpdateLandscapeSplines(const FGuid& InLayerGuid = FGuid(), bool bUpdateOnlySelected = false);
+	LANDSCAPE_API void SetLandscapeSplinesReservedLayer(int32 InLayerIndex);
+	LANDSCAPE_API struct FLandscapeLayer* GetLandscapeSplinesReservedLayer();
+	LANDSCAPE_API const struct FLandscapeLayer* GetLandscapeSplinesReservedLayer() const;
+	LANDSCAPE_API bool IsEditingLayerReservedForSplines() const;
 
 	LANDSCAPE_API bool IsLayerBlendSubstractive(int32 InLayerIndex, const TWeakObjectPtr<ULandscapeLayerInfoObject>& InLayerInfoObj) const;
 	LANDSCAPE_API void SetLayerSubstractiveBlendStatus(int32 InLayerIndex, bool InStatus, const TWeakObjectPtr<ULandscapeLayerInfoObject>& InLayerInfoObj);
@@ -254,7 +272,6 @@ public:
 	LANDSCAPE_API TArray<int8>& GetBrushesOrderForLayer(int32 InLayerIndex, int32 InTargetType);
 	LANDSCAPE_API class ALandscapeBlueprintCustomBrush* GetBrushForLayer(int32 InLayerIndex, int32 InTargetType, int8 BrushIndex) const;
 	LANDSCAPE_API TArray<class ALandscapeBlueprintCustomBrush*> GetBrushesForLayer(int32 InLayerIndex, int32 InTargetType) const;
-		
 
 private:
 	void TickLayers(float DeltaTime, ELevelTick TickType, FActorTickFunction& ThisTickFunction);
@@ -312,6 +329,13 @@ private:
 public:
 
 #if WITH_EDITORONLY_DATA
+	/** Target Landscape Layer for Landscape Splines */
+	UPROPERTY()
+	FGuid LandscapeSplinesTargetLayerGuid;
+	
+	/** Current Editing Landscape Layer*/
+	FGuid EditingLayer;
+
 	UPROPERTY(TextExportTransient)
 	TArray<FLandscapeLayer> LandscapeLayers;
 
diff --git a/Engine/Source/Runtime/Landscape/Classes/LandscapeComponent.h b/Engine/Source/Runtime/Landscape/Classes/LandscapeComponent.h
--- a/Engine/Source/Runtime/Landscape/Classes/LandscapeComponent.h
+++ b/Engine/Source/Runtime/Landscape/Classes/LandscapeComponent.h
@@ -524,7 +524,7 @@ public:
 	LANDSCAPE_API TArray<ULandscapeWeightmapUsage*>& GetWeightmapTexturesUsage(bool InReturnCurrentEditingWeightmap = false);
 	LANDSCAPE_API const TArray<ULandscapeWeightmapUsage*>& GetWeightmapTexturesUsage(bool InReturnCurrentEditingWeightmap = false) const;
 
-	LANDSCAPE_API void SetEditingLayer(struct FLandscapeLayer* InLayer, struct FLandscapeLayerData* InLayerData);
+	LANDSCAPE_API void SetEditingLayer(const struct FLandscapeLayer* InLayer, struct FLandscapeLayerData* InLayerData);
 #endif 
 
 #if WITH_EDITOR
diff --git a/Engine/Source/Runtime/Landscape/Private/Landscape.cpp b/Engine/Source/Runtime/Landscape/Private/Landscape.cpp
--- a/Engine/Source/Runtime/Landscape/Private/Landscape.cpp
+++ b/Engine/Source/Runtime/Landscape/Private/Landscape.cpp
@@ -1416,7 +1416,7 @@ TArray<FWeightmapLayerAllocationInfo>& ULandscapeComponent::GetWeightmapLayerAll
 
 #if WITH_EDITOR
 
-void ULandscapeComponent::SetEditingLayer(FLandscapeLayer* InLayer, FLandscapeLayerData* InLayerData)
+void ULandscapeComponent::SetEditingLayer(const FLandscapeLayer* InLayer, FLandscapeLayerData* InLayerData)
 {
 	// Update Current Heightmap
 	UTexture2D** LayerHeightmap = InLayerData ? InLayerData->Heightmaps.Find(GetHeightmap()) : nullptr;
@@ -1649,6 +1649,14 @@ void ALandscape::PostLoad()
 #endif
 	}
 
+#if WITH_EDITOR
+	for (FLandscapeLayer& Layer : LandscapeLayers)
+	{
+		// For now, only Layer reserved for Landscape Spline uses AlphaBlend
+		Layer.BlendMode = (Layer.Guid == LandscapeSplinesTargetLayerGuid) ? LSBM_AlphaBlend : LSBM_AdditiveBlend;
+	}
+#endif
+
 	Super::PostLoad();
 }
 
diff --git a/Engine/Source/Runtime/Landscape/Private/LandscapeEdit.cpp b/Engine/Source/Runtime/Landscape/Private/LandscapeEdit.cpp
--- a/Engine/Source/Runtime/Landscape/Private/LandscapeEdit.cpp
+++ b/Engine/Source/Runtime/Landscape/Private/LandscapeEdit.cpp
@@ -76,6 +76,10 @@ FAutoConsoleVariableRef CVarMobileCompressLanscapeWeightMaps(
 
 #if WITH_EDITOR
 
+// Used to temporarily disable material instance updates (typically used for cases where multiple updates are called on sample component)
+// Instead, one call per component is done at the end
+LANDSCAPE_API bool GDisableUpdateLandscapeMaterialInstances = false;
+
 // Channel remapping
 extern const size_t ChannelOffsets[4];
 
@@ -432,6 +436,9 @@ void ULandscapeComponent::UpdateMaterialInstances_Internal(FMaterialUpdateContex
 
 void ULandscapeComponent::UpdateMaterialInstances()
 {
+	if (GDisableUpdateLandscapeMaterialInstances)
+		return;
+
 	// we're not having the material update context recreate the render state because we will manually do it for only this component
 	TOptional<FComponentRecreateRenderStateContext> RecreateRenderStateContext;
 	RecreateRenderStateContext.Emplace(this);
diff --git a/Engine/Source/Runtime/Landscape/Private/LandscapeEditInterface.cpp b/Engine/Source/Runtime/Landscape/Private/LandscapeEditInterface.cpp
--- a/Engine/Source/Runtime/Landscape/Private/LandscapeEditInterface.cpp
+++ b/Engine/Source/Runtime/Landscape/Private/LandscapeEditInterface.cpp
@@ -211,7 +211,7 @@ bool FLandscapeEditDataInterface::GetComponentsInRegion(int32 X1, int32 Y1, int3
 	return bNotLocked;
 }
 
-void FLandscapeEditDataInterface::SetHeightData(int32 X1, int32 Y1, int32 X2, int32 Y2, const uint16* InData, int32 InStride, bool InCalcNormals, const uint16* InNormalData, bool InCreateComponents, UTexture2D* InHeightmap, UTexture2D* InXYOffsetmapTexture,
+void FLandscapeEditDataInterface::SetHeightData(int32 X1, int32 Y1, int32 X2, int32 Y2, const uint16* InData, int32 InStride, bool InCalcNormals, const uint16* InNormalData, const uint8* InHeightAlphaBlendData, const uint8* InHeightFlagsData, bool InCreateComponents, UTexture2D* InHeightmap, UTexture2D* InXYOffsetmapTexture,
 											   bool InUpdateBounds, bool InUpdateCollision, bool InGenerateMips)
 {
 	const int32 NumVertsX = 1 + X2 - X1;
@@ -379,6 +379,12 @@ void FLandscapeEditDataInterface::SetHeightData(int32 X1, int32 Y1, int32 X2, in
 								TexData.B = Normal >> 8;
 								TexData.A = Normal & 255;
 							}
+							else if (InHeightAlphaBlendData)
+							{
+								const uint8& HeightAlphaBlend = InHeightAlphaBlendData[DataIndex];
+								TexData.B = HeightAlphaBlend;
+								TexData.A = InHeightFlagsData ? InHeightFlagsData[DataIndex] : 3;
+							}
 						}
 					}
 
@@ -927,7 +933,7 @@ void FLandscapeEditDataInterface::CalcMissingValues(const int32& X1, const int32
 	}
 }
 
-uint16 FLandscapeEditDataInterface::GetHeightMapData(const ULandscapeComponent* Component, int32 TexU, int32 TexV, FColor* TextureData /*= NULL*/)
+FColor& FLandscapeEditDataInterface::GetHeightMapColor(const ULandscapeComponent* Component, int32 TexU, int32 TexV, FColor* TextureData)
 {
 	check(Component);
 	if (!TextureData)
@@ -944,12 +950,29 @@ uint16 FLandscapeEditDataInterface::GetHeightMapData(const ULandscapeComponent*
 	int32 TexX = HeightmapOffsetX + TexU;
 	int32 TexY = HeightmapOffsetY + TexV;
 	FColor& TexData = TextureData[ TexX + TexY * SizeU ];
+	return TexData;
+}
 
+uint16 FLandscapeEditDataInterface::GetHeightMapData(const ULandscapeComponent* Component, int32 TexU, int32 TexV, FColor* TextureData /*= NULL*/)
+{
+	const FColor& TexData = GetHeightMapColor(Component, TexU, TexV, TextureData);
 	return ((((uint16)TexData.R) << 8) | TexData.G);
 }
 
-template<typename TStoreData>
-void FLandscapeEditDataInterface::GetHeightDataTempl(int32& ValidX1, int32& ValidY1, int32& ValidX2, int32& ValidY2, TStoreData& StoreData)
+uint8 FLandscapeEditDataInterface::GetHeightMapAlphaBlendData(const ULandscapeComponent* Component, int32 TexU, int32 TexV, FColor* TextureData /*= NULL*/)
+{
+	const FColor& TexData = GetHeightMapColor(Component, TexU, TexV, TextureData);
+	return TexData.B;
+}
+
+uint8 FLandscapeEditDataInterface::GetHeightMapFlagsData(const ULandscapeComponent* Component, int32 TexU, int32 TexV, FColor* TextureData /*= NULL*/)
+{
+	const FColor& TexData = GetHeightMapColor(Component, TexU, TexV, TextureData);
+	return TexData.A;
+}
+
+template<typename TDataAccess>
+void FLandscapeEditDataInterface::GetHeightDataInternal(int32& ValidX1, int32& ValidY1, int32& ValidX2, int32& ValidY2, TDataAccess& StoreData, TFunctionRef<typename TDataAccess::DataType(const ULandscapeComponent*, int32, int32, FColor*)> GetHeightMapDataFunction)
 {
 	// Copy variables
 	int32 X1 = ValidX1, X2 = ValidX2, Y1 = ValidY1, Y2 = ValidY2;
@@ -973,7 +996,7 @@ void FLandscapeEditDataInterface::GetHeightDataTempl(int32& ValidX1, int32& Vali
 
 	FLandscapeTextureDataInfo* NeighborTexDataInfo[4] = {0, 0, 0, 0};
 	FColor* NeighborHeightmapTextureData[4] = {0, 0, 0, 0};
-	uint16 CornerValues[4] = {0, 0, 0, 0};
+	typename TDataAccess::DataType CornerValues[4] = {0, 0, 0, 0};
 
 	int32 EdgeCoord = (SubsectionSizeQuads+1) * ComponentNumSubsections - 1; //ComponentSizeQuads;
 
@@ -1122,28 +1145,28 @@ void FLandscapeEditDataInterface::GetHeightDataTempl(int32& ValidX1, int32& Vali
 				if (CornerComponent[0])
 				{
 					CornerSet |= 1;
-					CornerValues[0] = GetHeightMapData(CornerComponent[0], EdgeCoord, EdgeCoord);
+					CornerValues[0] = GetHeightMapDataFunction(CornerComponent[0], EdgeCoord, EdgeCoord, nullptr);
 				}
 				else if ((ExistLeft || ExistUp) && X1 <= ComponentIndexX*ComponentSizeQuads && Y1 <= ComponentIndexY*ComponentSizeQuads  )
 				{
 					CornerSet |= 1;
-					CornerValues[0] = StoreData.Load( ComponentIndexX*ComponentSizeQuads, ComponentIndexY*ComponentSizeQuads);
+					CornerValues[0] = StoreData.Load(ComponentIndexX*ComponentSizeQuads, ComponentIndexY*ComponentSizeQuads);
 				}
 				else if (BorderComponent[0])
 				{
 					CornerSet |= 1;
-					CornerValues[0] = GetHeightMapData(BorderComponent[0], EdgeCoord, 0, NeighborHeightmapTextureData[0]);
+					CornerValues[0] = GetHeightMapDataFunction(BorderComponent[0], EdgeCoord, 0, NeighborHeightmapTextureData[0]);
 				}
 				else if (BorderComponent[2])
 				{
 					CornerSet |= 1;
-					CornerValues[0] = GetHeightMapData(BorderComponent[2], 0, EdgeCoord, NeighborHeightmapTextureData[2]);
+					CornerValues[0] = GetHeightMapDataFunction(BorderComponent[2], 0, EdgeCoord, NeighborHeightmapTextureData[2]);
 				}
 
 				if (CornerComponent[1])
 				{
 					CornerSet |= 1 << 1;
-					CornerValues[1] = GetHeightMapData(CornerComponent[1], 0, EdgeCoord);
+					CornerValues[1] = GetHeightMapDataFunction(CornerComponent[1], 0, EdgeCoord, nullptr);
 				}
 				else if (ExistUp && X2 >= (ComponentIndexX+1)*ComponentSizeQuads)
 				{
@@ -1153,49 +1176,49 @@ void FLandscapeEditDataInterface::GetHeightDataTempl(int32& ValidX1, int32& Vali
 				else if (BorderComponent[1])
 				{
 					CornerSet |= 1 << 1;
-					CornerValues[1] = GetHeightMapData(BorderComponent[1], 0, 0, NeighborHeightmapTextureData[1]);
+					CornerValues[1] = GetHeightMapDataFunction(BorderComponent[1], 0, 0, NeighborHeightmapTextureData[1]);
 				}
 				else if (BorderComponent[2])
 				{
 					CornerSet |= 1 << 1;
-					CornerValues[1] = GetHeightMapData(BorderComponent[2], EdgeCoord, EdgeCoord, NeighborHeightmapTextureData[2]);
+					CornerValues[1] = GetHeightMapDataFunction(BorderComponent[2], EdgeCoord, EdgeCoord, NeighborHeightmapTextureData[2]);
 				}
 
 				if (CornerComponent[2])
 				{
 					CornerSet |= 1 << 2;
-					CornerValues[2] = GetHeightMapData(CornerComponent[2], EdgeCoord, 0);
+					CornerValues[2] = GetHeightMapDataFunction(CornerComponent[2], EdgeCoord, 0, nullptr);
 				}
 				else if (ExistLeft && Y2 >= (ComponentIndexY+1)*ComponentSizeQuads) // Use data already stored for 0, 2
 				{
 					CornerSet |= 1 << 2;
-					CornerValues[2] = StoreData.Load( ComponentIndexX*ComponentSizeQuads, (ComponentIndexY+1)*ComponentSizeQuads);
+					CornerValues[2] = StoreData.Load(ComponentIndexX*ComponentSizeQuads, (ComponentIndexY + 1)*ComponentSizeQuads);
 				}
 				else if (BorderComponent[0])
 				{
 					CornerSet |= 1 << 2;
-					CornerValues[2] = GetHeightMapData(BorderComponent[0], EdgeCoord, EdgeCoord, NeighborHeightmapTextureData[0]);
+					CornerValues[2] = GetHeightMapDataFunction(BorderComponent[0], EdgeCoord, EdgeCoord, NeighborHeightmapTextureData[0]);
 				}
 				else if (BorderComponent[3])
 				{
 					CornerSet |= 1 << 2;
-					CornerValues[2] = GetHeightMapData(BorderComponent[3], 0, 0, NeighborHeightmapTextureData[3]);
+					CornerValues[2] = GetHeightMapDataFunction(BorderComponent[3], 0, 0, NeighborHeightmapTextureData[3]);
 				}
 
 				if (CornerComponent[3])
 				{
 					CornerSet |= 1 << 3;
-					CornerValues[3] = GetHeightMapData(CornerComponent[3], 0, 0);
+					CornerValues[3] = GetHeightMapDataFunction(CornerComponent[3], 0, 0, nullptr);
 				}
 				else if (BorderComponent[1])
 				{
 					CornerSet |= 1 << 3;
-					CornerValues[3] = GetHeightMapData(BorderComponent[1], 0, EdgeCoord, NeighborHeightmapTextureData[1]);
+					CornerValues[3] = GetHeightMapDataFunction(BorderComponent[1], 0, EdgeCoord, NeighborHeightmapTextureData[1]);
 				}
 				else if (BorderComponent[3])
 				{
 					CornerSet |= 1 << 3;
-					CornerValues[3] = GetHeightMapData(BorderComponent[3], EdgeCoord, 0, NeighborHeightmapTextureData[3]);
+					CornerValues[3] = GetHeightMapDataFunction(BorderComponent[3], EdgeCoord, 0, NeighborHeightmapTextureData[3]);
 				}
 
 				FillCornerValues(CornerSet, CornerValues);
@@ -1241,13 +1264,13 @@ void FLandscapeEditDataInterface::GetHeightDataTempl(int32& ValidX1, int32& Vali
 							if( Component )
 							{
 								// Find the texture data corresponding to this vertex
-								uint16 Height = GetHeightMapData(Component, (SubsectionSizeQuads+1) * SubIndexX + SubX, (SubsectionSizeQuads+1) * SubIndexY + SubY, HeightmapTextureData);
+								typename TDataAccess::DataType Height = GetHeightMapDataFunction(Component, (SubsectionSizeQuads + 1) * SubIndexX + SubX, (SubsectionSizeQuads + 1) * SubIndexY + SubY, HeightmapTextureData);
 								StoreData.Store(LandscapeX, LandscapeY, Height);
 							}
 							else
 							{
 								// Find the texture data corresponding to this vertex
-								uint16 Value[4] = {0, 0, 0, 0};
+								typename TDataAccess::DataType Value[4] = {0, 0, 0, 0};
 								int32 Dist[4] = {INT_MAX, INT_MAX, INT_MAX, INT_MAX};
 								float ValueX = 0.0f, ValueY = 0.0f;
 								bool Exist[4] = {false, false, false, false};
@@ -1261,11 +1284,11 @@ void FLandscapeEditDataInterface::GetHeightDataTempl(int32& ValidX1, int32& Vali
 								}
 								else if (BorderComponent[0])
 								{
-									Value[0] = GetHeightMapData(BorderComponent[0], EdgeCoord, (SubsectionSizeQuads+1) * SubIndexY + SubY, NeighborHeightmapTextureData[0]);
+									Value[0] = GetHeightMapDataFunction(BorderComponent[0], EdgeCoord, (SubsectionSizeQuads + 1) * SubIndexY + SubY, NeighborHeightmapTextureData[0]);
 									Dist[0] = LandscapeX - (BorderComponent[0]->GetSectionBase().X + ComponentSizeQuads);
 									Exist[0] = true;
 								}
-								else 
+								else
 								{
 									if ((CornerSet & 1) && (CornerSet & 1 << 2))
 									{
@@ -1279,7 +1302,7 @@ void FLandscapeEditDataInterface::GetHeightDataTempl(int32& ValidX1, int32& Vali
 
 								if (BorderComponent[1])
 								{
-									Value[1] = GetHeightMapData(BorderComponent[1], 0, (SubsectionSizeQuads+1) * SubIndexY + SubY, NeighborHeightmapTextureData[1]);
+									Value[1] = GetHeightMapDataFunction(BorderComponent[1], 0, (SubsectionSizeQuads + 1) * SubIndexY + SubY, NeighborHeightmapTextureData[1]);
 									Dist[1] = (BorderComponent[1]->GetSectionBase().X) - LandscapeX;
 									Exist[1] = true;
 								}
@@ -1297,13 +1320,13 @@ void FLandscapeEditDataInterface::GetHeightDataTempl(int32& ValidX1, int32& Vali
 
 								if (ExistUp)
 								{
-									Value[2] = StoreData.Load( LandscapeX, ComponentIndexY*ComponentSizeQuads);
+									Value[2] = StoreData.Load(LandscapeX, ComponentIndexY*ComponentSizeQuads);
 									Dist[2] = LandscapeY - (ComponentIndexY*ComponentSizeQuads);
 									Exist[2] = true;
 								}
 								else if (BorderComponent[2])
 								{
-									Value[2] = GetHeightMapData(BorderComponent[2], (SubsectionSizeQuads+1) * SubIndexX + SubX, EdgeCoord, NeighborHeightmapTextureData[2]);
+									Value[2] = GetHeightMapDataFunction(BorderComponent[2], (SubsectionSizeQuads + 1) * SubIndexX + SubX, EdgeCoord, NeighborHeightmapTextureData[2]);
 									Dist[2] = LandscapeY - (BorderComponent[2]->GetSectionBase().Y + ComponentSizeQuads);
 									Exist[2] = true;
 								}
@@ -1321,7 +1344,7 @@ void FLandscapeEditDataInterface::GetHeightDataTempl(int32& ValidX1, int32& Vali
 
 								if (BorderComponent[3])
 								{
-									Value[3] = GetHeightMapData(BorderComponent[3], (SubsectionSizeQuads+1) * SubIndexX + SubX, 0, NeighborHeightmapTextureData[3]);
+									Value[3] = GetHeightMapDataFunction(BorderComponent[3], (SubsectionSizeQuads + 1) * SubIndexX + SubX, 0, NeighborHeightmapTextureData[3]);
 									Dist[3] = (BorderComponent[3]->GetSectionBase().Y) - LandscapeY;
 									Exist[3] = true;
 								}
@@ -1337,12 +1360,12 @@ void FLandscapeEditDataInterface::GetHeightDataTempl(int32& ValidX1, int32& Vali
 									}
 								}
 
-								CalcInterpValue<uint16>(Dist, Exist, Value, ValueX, ValueY);
+								CalcInterpValue<typename TDataAccess::DataType>(Dist, Exist, Value, ValueX, ValueY);
 
-								uint16 FinalValue = 0; // Default Value
+								typename TDataAccess::DataType FinalValue = 0; // Default Value
 								if ( (Exist[0] || Exist[1]) && (Exist[2] || Exist[3]) )
 								{
-									FinalValue = CalcValueFromValueXY<uint16>(Dist, ValueX, ValueY, CornerSet, CornerValues);
+									FinalValue = CalcValueFromValueXY<typename TDataAccess::DataType>(Dist, ValueX, ValueY, CornerSet, CornerValues);
 								}
 								else if ( (BorderComponent[0] || BorderComponent[1]) )
 								{
@@ -1362,7 +1385,6 @@ void FLandscapeEditDataInterface::GetHeightDataTempl(int32& ValidX1, int32& Vali
 								}
 
 								StoreData.Store(LandscapeX, LandscapeY, FinalValue);
-								//StoreData.StoreDefault(LandscapeX, LandscapeY);
 							}
 						}
 					}
@@ -1373,10 +1395,10 @@ void FLandscapeEditDataInterface::GetHeightDataTempl(int32& ValidX1, int32& Vali
 
 	if (bHasMissingValue)
 	{
-		CalcMissingValues<uint16, TStoreData, float>( X1, X2, Y1, Y2,
-			ComponentIndexX1, ComponentIndexX2, ComponentIndexY1, ComponentIndexY2, 
+		CalcMissingValues<typename TDataAccess::DataType, TDataAccess, float>(X1, X2, Y1, Y2,
+			ComponentIndexX1, ComponentIndexX2, ComponentIndexY1, ComponentIndexY2,
 			ComponentSizeX, ComponentSizeY, CornerValues,
-			NoBorderY1, NoBorderY2, ComponentDataExist, StoreData );
+			NoBorderY1, NoBorderY2, ComponentDataExist, StoreData);
 		// Update valid region
 		ValidX1 = FMath::Max<int32>(X1, ValidX1);
 		ValidX2 = FMath::Min<int32>(X2, ValidX2);
@@ -1394,62 +1416,73 @@ void FLandscapeEditDataInterface::GetHeightDataTempl(int32& ValidX1, int32& Vali
 
 namespace
 {
-	struct FArrayStoreData
+	template<typename T>
+	struct TArrayDataAccess
 	{
+		typedef T DataType;
 		int32 X1;
 		int32 Y1;
-		uint16* Data;
+		DataType* Data;
 		int32 Stride;
 
-		FArrayStoreData(int32 InX1, int32 InY1, uint16* InData, int32 InStride)
+		TArrayDataAccess(int32 InX1, int32 InY1, DataType* InData, int32 InStride)
 			:	X1(InX1)
 			,	Y1(InY1)
 			,	Data(InData)
 			,	Stride(InStride)
 		{}
 
-		inline void Store(int32 LandscapeX, int32 LandscapeY, uint16 Height)
+		inline void Store(int32 LandscapeX, int32 LandscapeY, DataType NewValue)
 		{
-			Data[ (LandscapeY-Y1) * Stride + (LandscapeX-X1) ] = Height;
+			Data[ (LandscapeY-Y1) * Stride + (LandscapeX-X1) ] = NewValue;
 		}
 
-		// for interpolation
-		inline uint16 Load(int32 LandscapeX, int32 LandscapeY)
+		inline DataType Load(int32 LandscapeX, int32 LandscapeY)
 		{
 			return Data[ (LandscapeY-Y1) * Stride + (LandscapeX-X1) ];
 		}
-
-		inline void StoreDefault(int32 LandscapeX, int32 LandscapeY)
-		{
-			Data[ (LandscapeY-Y1) * Stride + (LandscapeX-X1) ] = 0;
-		}
 	};
 
-	struct FSparseStoreData
+	template<typename T>
+	struct TSparseDataAccess
 	{
-		TMap<FIntPoint, uint16>& SparseData;
+		typedef T DataType;
+		TMap<FIntPoint, DataType>& SparseData;
 
-		FSparseStoreData(TMap<FIntPoint, uint16>& InSparseData)
+		TSparseDataAccess(TMap<FIntPoint, DataType>& InSparseData)
 			:	SparseData(InSparseData)
 		{}
 
-		inline void Store(int32 LandscapeX, int32 LandscapeY, uint16 Height)
+		inline void Store(int32 LandscapeX, int32 LandscapeY, DataType NewValue)
 		{
-			SparseData.Add(FIntPoint(LandscapeX,LandscapeY), Height);
+			SparseData.Add(FIntPoint(LandscapeX,LandscapeY), NewValue);
 		}
 
-		inline uint16 Load(int32 LandscapeX, int32 LandscapeY)
+		inline DataType Load(int32 LandscapeX, int32 LandscapeY)
 		{
 			return SparseData.FindRef(FIntPoint(LandscapeX,LandscapeY));
 		}
-
-		inline void StoreDefault(int32 LandscapeX, int32 LandscapeY)
-		{
-		}
 	};
-
 };
 
+template<typename TDataAccess>
+void FLandscapeEditDataInterface::GetHeightDataTempl(int32& ValidX1, int32& ValidY1, int32& ValidX2, int32& ValidY2, TDataAccess& StoreData)
+{
+	GetHeightDataInternal(ValidX1, ValidY1, ValidX2, ValidY2, StoreData, [&](const ULandscapeComponent* Component, int32 TexU, int32 TexV, FColor* TextureData)->uint16 { return GetHeightMapData(Component, TexU, TexV, TextureData); });
+}
+
+template<typename TDataAccess>
+void FLandscapeEditDataInterface::GetHeightAlphaBlendDataTempl(int32& ValidX1, int32& ValidY1, int32& ValidX2, int32& ValidY2, TDataAccess& StoreData)
+{
+	GetHeightDataInternal(ValidX1, ValidY1, ValidX2, ValidY2, StoreData, [&](const ULandscapeComponent* Component, int32 TexU, int32 TexV, FColor* TextureData)->uint8 { return GetHeightMapAlphaBlendData(Component, TexU, TexV, TextureData); });
+}
+
+template<typename TDataAccess>
+void FLandscapeEditDataInterface::GetHeightFlagsDataTempl(int32& ValidX1, int32& ValidY1, int32& ValidX2, int32& ValidY2, TDataAccess& StoreData)
+{
+	GetHeightDataInternal(ValidX1, ValidY1, ValidX2, ValidY2, StoreData, [&](const ULandscapeComponent* Component, int32 TexU, int32 TexV, FColor* TextureData)->uint8 { return GetHeightMapFlagsData(Component, TexU, TexV, TextureData); });
+}
+
 void FLandscapeEditDataInterface::GetHeightData(int32& X1, int32& Y1, int32& X2, int32& Y2, uint16* Data, int32 Stride)
 {
 	if( Stride==0 )
@@ -1457,10 +1490,32 @@ void FLandscapeEditDataInterface::GetHeightData(int32& X1, int32& Y1, int32& X2,
 		Stride = (1+X2-X1);
 	}
 
-	FArrayStoreData ArrayStoreData(X1, Y1, Data, Stride);
+	TArrayDataAccess<uint16> ArrayStoreData(X1, Y1, Data, Stride);
 	GetHeightDataTempl(X1, Y1, X2, Y2, ArrayStoreData);
 }
 
+void FLandscapeEditDataInterface::GetHeightAlphaBlendData(int32& X1, int32& Y1, int32& X2, int32& Y2, uint8* Data, int32 Stride)
+{
+	if (Stride == 0)
+	{
+		Stride = (1 + X2 - X1);
+	}
+
+	TArrayDataAccess<uint8> ArrayStoreData(X1, Y1, Data, Stride);
+	GetHeightAlphaBlendDataTempl(X1, Y1, X2, Y2, ArrayStoreData);
+}
+
+void FLandscapeEditDataInterface::GetHeightFlagsData(int32& X1, int32& Y1, int32& X2, int32& Y2, uint8* Data, int32 Stride)
+{
+	if (Stride == 0)
+	{
+		Stride = (1 + X2 - X1);
+	}
+
+	TArrayDataAccess<uint8> ArrayStoreData(X1, Y1, Data, Stride);
+	GetHeightFlagsDataTempl(X1, Y1, X2, Y2, ArrayStoreData);
+}
+
 void FLandscapeEditDataInterface::GetHeightDataFast(const int32 X1, const int32 Y1, const int32 X2, const int32 Y2, uint16* Data, int32 Stride, uint16* NormalData /*= NULL*/, UTexture2D* InHeightmap)
 {
 	if( Stride==0 )
@@ -1468,10 +1523,10 @@ void FLandscapeEditDataInterface::GetHeightDataFast(const int32 X1, const int32
 		Stride = (1+X2-X1);
 	}
 
-	FArrayStoreData ArrayStoreData(X1, Y1, Data, Stride);
+	TArrayDataAccess<uint16> ArrayStoreData(X1, Y1, Data, Stride);
 	if (NormalData)
 	{
-		FArrayStoreData ArrayNormalData(X1, Y1, NormalData, Stride);
+		TArrayDataAccess<uint16> ArrayNormalData(X1, Y1, NormalData, Stride);
 		GetHeightDataTemplFast(X1, Y1, X2, Y2, ArrayStoreData, InHeightmap, &ArrayNormalData);
 	}
 	else
@@ -1482,16 +1537,16 @@ void FLandscapeEditDataInterface::GetHeightDataFast(const int32 X1, const int32
 
 void FLandscapeEditDataInterface::GetHeightData(int32& X1, int32& Y1, int32& X2, int32& Y2, TMap<FIntPoint, uint16>& SparseData)
 {
-	FSparseStoreData SparseStoreData(SparseData);
+	TSparseDataAccess<uint16> SparseStoreData(SparseData);
 	GetHeightDataTempl(X1, Y1, X2, Y2, SparseStoreData);
 }
 
 void FLandscapeEditDataInterface::GetHeightDataFast(const int32 X1, const int32 Y1, const int32 X2, const int32 Y2, TMap<FIntPoint, uint16>& SparseData, TMap<FIntPoint, uint16>* NormalData /*= NULL*/, UTexture2D* InHeightmap)
 {
-	FSparseStoreData SparseStoreData(SparseData);
+	TSparseDataAccess<uint16> SparseStoreData(SparseData);
 	if (NormalData)
 	{
-		FSparseStoreData SparseNormalData(*NormalData);
+		TSparseDataAccess<uint16> SparseNormalData(*NormalData);
 		GetHeightDataTemplFast(X1, Y1, X2, Y2, SparseStoreData, InHeightmap, &SparseNormalData);
 	}
 	else
@@ -5134,7 +5189,6 @@ void FLandscapeEditDataInterface::GetXYOffsetDataTempl(int32& ValidX1, int32& Va
 								}
 
 								StoreData.Store(LandscapeX, LandscapeY, FinalValue);
-								//StoreData.StoreDefault(LandscapeX, LandscapeY);
 							}
 						}
 					}
@@ -5264,6 +5318,11 @@ void FLandscapeEditDataInterface::GetXYOffsetDataTemplFast(const int32 X1, const
 	}
 }
 
+const ALandscape* FLandscapeEditDataInterface::GetTargetLandscape() const
+{
+	return LandscapeInfo ? LandscapeInfo->LandscapeActor.Get() : nullptr;
+}
+
 void FLandscapeEditDataInterface::GetXYOffsetDataFast(const int32 X1, const int32 Y1, const int32 X2, const int32 Y2, FVector2D* Data, int32 Stride)
 {
 	if( Stride==0 )
diff --git a/Engine/Source/Runtime/Landscape/Private/LandscapeEditLayers.cpp b/Engine/Source/Runtime/Landscape/Private/LandscapeEditLayers.cpp
--- a/Engine/Source/Runtime/Landscape/Private/LandscapeEditLayers.cpp
+++ b/Engine/Source/Runtime/Landscape/Private/LandscapeEditLayers.cpp
@@ -176,7 +176,7 @@ struct FLandscapeLayersHeightmapShaderParameters
 		, ApplyLayerModifiers(false)
 		, LayerAlpha(1.0f)
 		, LayerVisible(true)
-		, OutputAsDelta(false)
+		, LayerBlendMode(LSBM_AdditiveBlend)
 		, GenerateNormals(false)
 		, GridSize(0.0f, 0.0f, 0.0f)
 		, CurrentMipSize(0, 0)
@@ -190,7 +190,7 @@ struct FLandscapeLayersHeightmapShaderParameters
 	bool ApplyLayerModifiers;
 	float LayerAlpha;
 	bool LayerVisible;
-	bool OutputAsDelta;
+	ELandscapeBlendMode LayerBlendMode;
 	bool GenerateNormals;
 	FVector GridSize;
 	FIntPoint CurrentMipSize;
@@ -235,8 +235,8 @@ public:
 		SetTextureParameter(RHICmdList, GetPixelShader(), ReadTexture1Param, ReadTexture1SamplerParam, TStaticSamplerState<SF_Point, AM_Clamp, AM_Clamp, AM_Clamp>::GetRHI(), InParams.ReadHeightmap1->Resource->TextureRHI);
 		SetTextureParameter(RHICmdList, GetPixelShader(), ReadTexture2Param, ReadTexture2SamplerParam, TStaticSamplerState<SF_Point, AM_Clamp, AM_Clamp, AM_Clamp>::GetRHI(), InParams.ReadHeightmap2 != nullptr ? InParams.ReadHeightmap2->Resource->TextureRHI : GWhiteTexture->TextureRHI);
 
-		FVector2D LayerInfo(InParams.LayerAlpha, InParams.LayerVisible ? 1.0f : 0.0f);
-		FVector4 OutputConfig(InParams.ApplyLayerModifiers ? 1.0f : 0.0f, InParams.OutputAsDelta ? 1.0f : 0.0f, InParams.ReadHeightmap2 ? 1.0f : 0.0f, InParams.GenerateNormals ? 1.0f : 0.0f);
+		FVector4 LayerInfo(InParams.LayerAlpha, InParams.LayerVisible ? 1.0f : 0.0f, InParams.LayerBlendMode == LSBM_AlphaBlend ? 1.0f : 0.f, 0.f);
+		FVector4 OutputConfig(InParams.ApplyLayerModifiers ? 1.0f : 0.0f, 0.0f /*unused*/, InParams.ReadHeightmap2 ? 1.0f : 0.0f, InParams.GenerateNormals ? 1.0f : 0.0f);
 		FVector2D TextureSize(InParams.HeightmapSize.X, InParams.HeightmapSize.Y);
 
 		SetShaderValue(RHICmdList, GetPixelShader(), LayerInfoParam, LayerInfo);
@@ -341,7 +341,7 @@ struct FLandscapeLayersWeightmapShaderParameters
 		, ApplyLayerModifiers(false)
 		, LayerAlpha(1.0f)
 		, LayerVisible(true)
-		, OutputAsDelta(false)
+		, LayerBlendMode(LSBM_AdditiveBlend)
 		, OutputAsSubstractive(false)
 		, OutputAsNormalized(false)
 		, CurrentMipSize(0, 0)
@@ -354,7 +354,7 @@ struct FLandscapeLayersWeightmapShaderParameters
 	bool ApplyLayerModifiers;
 	float LayerAlpha;
 	bool LayerVisible;
-	bool OutputAsDelta;
+	ELandscapeBlendMode LayerBlendMode;
 	bool OutputAsSubstractive;
 	bool OutputAsNormalized;
 	FIntPoint CurrentMipSize;
@@ -396,7 +396,7 @@ public:
 		SetTextureParameter(RHICmdList, GetPixelShader(), ReadTexture1Param, ReadTexture1SamplerParam, TStaticSamplerState<SF_Point, AM_Clamp, AM_Clamp, AM_Clamp>::GetRHI(), InParams.ReadWeightmap1->Resource->TextureRHI);
 		SetTextureParameter(RHICmdList, GetPixelShader(), ReadTexture2Param, ReadTexture2SamplerParam, TStaticSamplerState<SF_Point, AM_Clamp, AM_Clamp, AM_Clamp>::GetRHI(), InParams.ReadWeightmap2 != nullptr ? InParams.ReadWeightmap2->Resource->TextureRHI : GWhiteTexture->TextureRHI);
 
-		FVector2D LayerInfo(InParams.LayerAlpha, InParams.LayerVisible ? 1.0f : 0.0f);
+		FVector4 LayerInfo(InParams.LayerAlpha, InParams.LayerVisible ? 1.0f : 0.0f, InParams.LayerBlendMode == LSBM_AlphaBlend ? 1.0f : 0.f, 0.f);
 		FVector4 OutputConfig(InParams.ApplyLayerModifiers ? 1.0f : 0.0f, InParams.OutputAsSubstractive ? 1.0f : 0.0f, InParams.ReadWeightmap2 != nullptr ? 1.0f : 0.0f, InParams.OutputAsNormalized ? 1.0f : 0.0f);
 
 		SetShaderValue(RHICmdList, GetPixelShader(), LayerInfoParam, LayerInfo);
@@ -2627,7 +2627,20 @@ void ALandscape::RegenerateLayersHeightmaps()
 			//Draw Layer heightmap to Combined RT Atlas
 			ShaderParams.ApplyLayerModifiers = true;
 			ShaderParams.LayerVisible = Layer.bVisible;
-			ShaderParams.LayerAlpha = Layer.HeightmapAlpha;
+			ShaderParams.GenerateNormals = false;
+			ShaderParams.LayerBlendMode = Layer.BlendMode;
+			if (Layer.BlendMode == LSBM_AlphaBlend)
+			{
+				// For now, only Layer reserved for Landscape Splines will use the AlphaBlendMode
+				const FLandscapeLayer* SplinesReservedLayer = GetLandscapeSplinesReservedLayer();
+				check(&Layer == SplinesReservedLayer);
+				ShaderParams.LayerAlpha = 1.0f;
+			}
+			else
+			{
+				check(Layer.BlendMode == LSBM_AdditiveBlend);
+				ShaderParams.LayerAlpha = Layer.HeightmapAlpha;
+			}
 
 			for (ALandscapeProxy* Landscape : AllLandscapes)
 			{
@@ -4426,6 +4439,11 @@ const FLandscapeLayer* ALandscape::GetLayer(int32 InLayerIndex) const
 	return nullptr;
 }
 
+const FLandscapeLayer* ALandscape::GetLayer(const FGuid& InLayerGuid) const
+{
+	return LandscapeLayers.FindByPredicate([InLayerGuid](const FLandscapeLayer& Other) { return Other.Guid == InLayerGuid; });
+}
+
 void ALandscape::DeleteLayer(int32 InLayerIndex)
 {
 	ensure(GetMutableDefault<UEditorExperimentalSettings>()->bLandscapeLayerSystem);
@@ -4475,6 +4493,12 @@ void ALandscape::DeleteLayer(int32 InLayerIndex)
 		Proxy->LandscapeLayersData.Remove(LayerGuid);
 	});
 
+	const FLandscapeLayer* SplinesReservedLayer = GetLandscapeSplinesReservedLayer();
+	if (SplinesReservedLayer == Layer)
+	{
+		LandscapeSplinesTargetLayerGuid.Invalidate();
+	}
+
 	// Remove layer from list
 	LandscapeLayers.RemoveAt(InLayerIndex);
 
@@ -4482,20 +4506,20 @@ void ALandscape::DeleteLayer(int32 InLayerIndex)
 	RequestLayersContentUpdate(ELandscapeLayersContentUpdateFlag::All_Setup | All, true);
 }
 
-void ALandscape::ClearLayer(int32 InLayerIndex)
+void ALandscape::ClearLayer(int32 InLayerIndex, bool bInUpdateCollision)
 {
 	const FLandscapeLayer* Layer = GetLayer(InLayerIndex);
 	if (Layer)
 	{
-		ClearLayer(Layer->Guid);
+		ClearLayer(Layer->Guid, bInUpdateCollision);
 	}
 }
 
-void ALandscape::ClearLayer(const FGuid& InLayerGuid)
+void ALandscape::ClearLayer(const FGuid& InLayerGuid, bool bInUpdateCollision)
 {
 	ensure(GetMutableDefault<UEditorExperimentalSettings>()->bLandscapeLayerSystem);
 
-	FLandscapeLayer* Layer = LandscapeLayers.FindByPredicate([InLayerGuid](const FLandscapeLayer& Other) { return Other.Guid == InLayerGuid; });
+	const FLandscapeLayer* Layer = GetLayer(InLayerGuid);
 	ULandscapeInfo* LandscapeInfo = GetLandscapeInfo();
 	if (!LandscapeInfo || !Layer)
 	{
@@ -4505,12 +4529,22 @@ void ALandscape::ClearLayer(const FGuid& InLayerGuid)
 	Modify();
 	FScopedSetLandscapeEditingLayer Scope(this, Layer ? Layer->Guid : FGuid(), [=] { RequestLayersContentUpdate(ELandscapeLayersContentUpdateFlag::All, true); });
 
-	TArray<uint16> NewData;
-	NewData.AddZeroed(FMath::Square(ComponentSizeQuads + 1));
+	TArray<uint16> NewHeightData;
+	NewHeightData.AddZeroed(FMath::Square(ComponentSizeQuads + 1));
 	uint16 ZeroValue = LandscapeDataAccess::GetTexHeight(0.f);
-	for (uint16& NewDataValue : NewData)
+	for (uint16& NewHeightDataValue : NewHeightData)
+	{
+		NewHeightDataValue = ZeroValue;
+	}
+
+	TArray<uint8> NewHeightAlphaBlendData;
+	TArray<uint8> NewHeightFlagsData;
+
+	if (Layer->BlendMode == LSBM_AlphaBlend)
 	{
-		NewDataValue = ZeroValue;
+		NewHeightAlphaBlendData.AddZeroed(FMath::Square(ComponentSizeQuads + 1));
+		FMemory::Memset(NewHeightAlphaBlendData.GetData(), 255, NewHeightAlphaBlendData.Num());
+		NewHeightFlagsData.AddZeroed(FMath::Square(ComponentSizeQuads + 1));
 	}
 
 	FLandscapeEditDataInterface LandscapeEdit(LandscapeInfo);
@@ -4526,13 +4560,25 @@ void ALandscape::ClearLayer(const FGuid& InLayerGuid)
 			check(ComponentSizeQuads == (MaxX - MinX));
 			check(ComponentSizeQuads == (MaxY - MinY));
 
-			TArray<uint16> OldData;
-			OldData.AddZeroed((1 + MaxY - MinY) * (1 + MaxX - MinX));
+			TArray<uint16> OldHeightData;
+			OldHeightData.AddZeroed((1 + MaxY - MinY) * (1 + MaxX - MinX));
+			LandscapeEdit.GetHeightData(MinX, MinY, MaxX, MaxY, OldHeightData.GetData(), 0);
 
-			LandscapeEdit.GetHeightData(MinX, MinY, MaxX, MaxY, OldData.GetData(), 0);
-			if (FMemory::Memcmp(OldData.GetData(), NewData.GetData(), NewData.Num() * NewData.GetTypeSize()) != 0)
+			TArray<uint8> OldHeightAlphaBlendData;
+			TArray<uint8> OldHeightFlagsData;
+			if (Layer->BlendMode == LSBM_AlphaBlend)
 			{
-				LandscapeEdit.SetHeightData(MinX, MinY, MaxX, MaxY, NewData.GetData(), 0, true);
+				OldHeightAlphaBlendData.AddZeroed((1 + MaxY - MinY) * (1 + MaxX - MinX));
+				LandscapeEdit.GetHeightAlphaBlendData(MinX, MinY, MaxX, MaxY, OldHeightAlphaBlendData.GetData(), 0);
+				OldHeightFlagsData.AddZeroed((1 + MaxY - MinY) * (1 + MaxX - MinX));
+				LandscapeEdit.GetHeightFlagsData(MinX, MinY, MaxX, MaxY, OldHeightFlagsData.GetData(), 0);
+			}
+
+			if ((FMemory::Memcmp(OldHeightData.GetData(), NewHeightData.GetData(), NewHeightData.Num() * NewHeightData.GetTypeSize()) != 0) ||
+				(FMemory::Memcmp(OldHeightAlphaBlendData.GetData(), NewHeightAlphaBlendData.GetData(), NewHeightAlphaBlendData.Num() * NewHeightAlphaBlendData.GetTypeSize()) != 0) ||
+				(FMemory::Memcmp(OldHeightFlagsData.GetData(), NewHeightFlagsData.GetData(), NewHeightFlagsData.Num() * NewHeightFlagsData.GetTypeSize()) != 0))
+			{
+				LandscapeEdit.SetHeightData(MinX, MinY, MaxX, MaxY, NewHeightData.GetData(), 0, false, nullptr, NewHeightAlphaBlendData.GetData(), NewHeightFlagsData.GetData(), false, nullptr, nullptr, true, bInUpdateCollision);
 			}
 
 			// Clear weight maps
@@ -4569,6 +4615,74 @@ void ALandscape::ShowAllLayers()
 	}
 }
 
+void ALandscape::SetLandscapeSplinesReservedLayer(int32 InLayerIndex)
+{
+	FLandscapeLayer* NewLayer = GetLayer(InLayerIndex);
+	FLandscapeLayer* PreviousLayer = GetLandscapeSplinesReservedLayer();
+	if (NewLayer != PreviousLayer)
+	{
+		if (PreviousLayer)
+		{
+			ClearLayer(LandscapeSplinesTargetLayerGuid);
+			PreviousLayer->BlendMode = LSBM_AdditiveBlend;
+		}
+		if (NewLayer)
+		{
+			NewLayer->HeightmapAlpha = 1.0f;
+			NewLayer->WeightmapAlpha = 1.0f;
+			NewLayer->BlendMode = LSBM_AlphaBlend;
+			LandscapeSplinesTargetLayerGuid = NewLayer->Guid;
+			ClearLayer(LandscapeSplinesTargetLayerGuid);
+		}
+		else
+		{
+			LandscapeSplinesTargetLayerGuid.Invalidate();
+		}
+	}
+}
+
+const FLandscapeLayer* ALandscape::GetLandscapeSplinesReservedLayer() const
+{
+	if (LandscapeSplinesTargetLayerGuid.IsValid())
+	{
+		return LandscapeLayers.FindByPredicate([this](const FLandscapeLayer& Other) { return Other.Guid == LandscapeSplinesTargetLayerGuid; });
+	}
+	return nullptr;
+}
+
+FLandscapeLayer* ALandscape::GetLandscapeSplinesReservedLayer()
+{
+	if (LandscapeSplinesTargetLayerGuid.IsValid())
+	{
+		return LandscapeLayers.FindByPredicate([this](const FLandscapeLayer& Other) { return Other.Guid == LandscapeSplinesTargetLayerGuid; });
+	}
+	return nullptr;
+}
+
+LANDSCAPE_API extern bool GDisableUpdateLandscapeMaterialInstances;
+
+void ALandscape::UpdateLandscapeSplines(const FGuid& InTargetLayer, bool bUpdateOnlySelected)
+{
+	check(GetMutableDefault<UEditorExperimentalSettings>()->bLandscapeLayerSystem);
+	ULandscapeInfo* LandscapeInfo = GetLandscapeInfo();
+	FGuid TargetLayerGuid = LandscapeSplinesTargetLayerGuid.IsValid() ? LandscapeSplinesTargetLayerGuid : InTargetLayer;
+	const FLandscapeLayer* TargetLayer = GetLayer(TargetLayerGuid);
+	if (LandscapeInfo && TargetLayer)
+	{
+		FScopedSetLandscapeEditingLayer Scope(this, TargetLayerGuid, [=] { this->RequestLayersContentUpdate(ELandscapeLayersContentUpdateFlag::All, true); });
+		// Temporarily disable material instance updates since it will be done once at the end (requested by RequestLayersContentUpdate)
+		GDisableUpdateLandscapeMaterialInstances = true;
+		if (LandscapeSplinesTargetLayerGuid.IsValid())
+		{
+			ClearLayer(LandscapeSplinesTargetLayerGuid, false);
+			// For now, in Landscape Layer System Mode with a reserved layer for splines, we always update all the splines since we clear the whole layer first
+			bUpdateOnlySelected = false;
+		}
+		LandscapeInfo->ApplySplines(bUpdateOnlySelected);
+		GDisableUpdateLandscapeMaterialInstances = false;
+	}
+}
+
 FScopedSetLandscapeEditingLayer::FScopedSetLandscapeEditingLayer(ALandscape* InLandscape, const FGuid& InLayerGUID, TFunction<void()> InCompletionCallback)
 	: Landscape(InLandscape)
 	, LayerGUID(InLayerGUID)
@@ -4592,19 +4706,31 @@ FScopedSetLandscapeEditingLayer::~FScopedSetLandscapeEditingLayer()
 	}
 }
 
-void ALandscape::SetEditingLayer(FGuid InLayerGuid)
+bool ALandscape::IsEditingLayerReservedForSplines() const
+{
+	if (GetMutableDefault<UEditorExperimentalSettings>()->bLandscapeLayerSystem)
+	{
+		const FLandscapeLayer* SplinesReservedLayer = GetLandscapeSplinesReservedLayer();
+		return SplinesReservedLayer && SplinesReservedLayer->Guid == EditingLayer;
+	}
+	return false;
+}
+
+void ALandscape::SetEditingLayer(const FGuid& InLayerGuid)
 {
 	ensure(GetMutableDefault<UEditorExperimentalSettings>()->bLandscapeLayerSystem);
 
 	ULandscapeInfo* LandscapeInfo = GetLandscapeInfo();
 	if (!LandscapeInfo)
 	{
+		EditingLayer.Invalidate();
 		return;
 	}
 
+	EditingLayer = InLayerGuid;
 	LandscapeInfo->ForAllLandscapeProxies([InLayerGuid, this](ALandscapeProxy* Proxy)
 	{
-		FLandscapeLayer* Layer = InLayerGuid.IsValid() ? LandscapeLayers.FindByPredicate([InLayerGuid](const FLandscapeLayer& Other) { return Other.Guid == InLayerGuid; }) : nullptr;
+		const FLandscapeLayer* Layer = InLayerGuid.IsValid() ? GetLayer(InLayerGuid) : nullptr;
 		FLandscapeLayerData* LayerData = Layer ? Proxy->LandscapeLayersData.Find(Layer->Guid) : nullptr;
 
 		for (ULandscapeComponent* Component : Proxy->LandscapeComponents)
diff --git a/Engine/Source/Runtime/Landscape/Private/LandscapeSplineRaster.cpp b/Engine/Source/Runtime/Landscape/Private/LandscapeSplineRaster.cpp
--- a/Engine/Source/Runtime/Landscape/Private/LandscapeSplineRaster.cpp
+++ b/Engine/Source/Runtime/Landscape/Private/LandscapeSplineRaster.cpp
@@ -19,6 +19,8 @@
 #if WITH_EDITOR
 #include "ScopedTransaction.h"
 #include "Raster.h"
+#include "Landscape.h"
+#include "Settings/EditorExperimentalSettings.h"
 #endif
 
 #define LOCTEXT_NAMESPACE "Landscape"
@@ -35,8 +37,10 @@ public:
 	typedef FVector InterpolantType;
 
 	/** Initialization constructor. */
-	FLandscapeSplineHeightsRasterPolicy(TArray<uint16>& InData, int32 InMinX, int32 InMinY, int32 InMaxX, int32 InMaxY, bool InbRaiseTerrain, bool InbLowerTerrain) :
-		Data(InData),
+	FLandscapeSplineHeightsRasterPolicy(TArray<uint16>& InHeightData, int32 InMinX, int32 InMinY, int32 InMaxX, int32 InMaxY, bool InbRaiseTerrain, bool InbLowerTerrain, TArray<uint8>* InHeightAlphaBlendData = nullptr, TArray<uint8>* InHeightFlagsData = nullptr) :
+		HeightData(InHeightData),
+		HeightAlphaBlendData(InHeightAlphaBlendData),
+		HeightFlagsData(InHeightFlagsData),
 		MinX(InMinX),
 		MinY(InMinY),
 		MaxX(InMaxX),
@@ -49,7 +53,6 @@ public:
 protected:
 
 	// FTriangleRasterizer policy interface.
-
 	int32 GetMinX() const { return MinX; }
 	int32 GetMaxX() const { return MaxX; }
 	int32 GetMinY() const { return MinY; }
@@ -57,21 +60,55 @@ protected:
 
 	inline void ProcessPixel(int32 X, int32 Y, const InterpolantType& Interpolant, bool BackFacing)
 	{
+		if (!bRaiseTerrain && !bLowerTerrain)
+		{
+			return;
+		}
+
 		const float CosInterpX = (Interpolant.X >= 1 ? 1 : 0.5f - 0.5f * FMath::Cos(Interpolant.X * PI));
 		const float CosInterpY = (Interpolant.Y >= 1 ? 1 : 0.5f - 0.5f * FMath::Cos(Interpolant.Y * PI));
-		const float Alpha = CosInterpX * CosInterpY;
-		uint16& Dest = Data[(Y - MinY)*(1 + MaxX - MinX) + X - MinX];
-		float Value = FMath::Lerp((float)Dest, Interpolant.Z, Alpha);
-		uint16 DValue = (uint16)FMath::Clamp<float>(Value, 0, (float)LandscapeDataAccess::MaxValue);
-		if ((bRaiseTerrain && DValue > Dest) ||
-			(bLowerTerrain && DValue < Dest))
+		const float Alpha = FMath::Clamp<float>(CosInterpX * CosInterpY, 0.f, 1.f);
+
+		int32 DataIndex = (Y - MinY)*(1 + MaxX - MinX) + X - MinX;
+		uint16& Dest = HeightData[DataIndex];
+
+		if (HeightAlphaBlendData)
 		{
-			Dest = DValue;
+			uint16 NewHeight = (uint16)FMath::Clamp<float>(Interpolant.Z, 0, (float)LandscapeDataAccess::MaxValue);
+			float InterpValue = (NewHeight * Alpha) + (Dest * (1.f - Alpha));
+			Dest = (uint16)FMath::Clamp<float>(InterpValue, 0, (float)LandscapeDataAccess::MaxValue);
+
+			uint8& DestAlphaValue = (*HeightAlphaBlendData)[DataIndex];
+			float InterpAlphaValue = DestAlphaValue * (1.f - Alpha);
+			DestAlphaValue = (uint8)FMath::Clamp<float>(InterpAlphaValue, 0.f, 255.f);
+
+			if (HeightFlagsData)
+			{
+				uint8& DestHeightFlagsValue = (*HeightFlagsData)[DataIndex];
+				uint8 Flags = DestHeightFlagsValue;
+				if (bLowerTerrain)
+					Flags |= 1; 
+				if (bRaiseTerrain)
+					Flags |= 2;
+				DestHeightFlagsValue = Flags;
+			}
+		}
+		else
+		{
+			float Value = FMath::Lerp((float)Dest, Interpolant.Z, Alpha);
+			uint16 DValue = (uint16)FMath::Clamp<float>(Value, 0, (float)LandscapeDataAccess::MaxValue);
+			if ((bRaiseTerrain && DValue > Dest) ||
+				(bLowerTerrain && DValue < Dest))
+			{
+				Dest = DValue;
+			}
 		}
 	}
 
 private:
-	TArray<uint16>& Data;
+	TArray<uint16>& HeightData;
+	TArray<uint8>* HeightAlphaBlendData;
+	TArray<uint8>* HeightFlagsData;
 	int32 MinX, MinY, MaxX, MaxY;
 	uint32 bRaiseTerrain : 1, bLowerTerrain : 1;
 };
@@ -116,7 +153,7 @@ private:
 	int32 MinX, MinY, MaxX, MaxY;
 };
 
-void RasterizeControlPointHeights(int32& MinX, int32& MinY, int32& MaxX, int32& MaxY, FLandscapeEditDataInterface& LandscapeEdit, FVector ControlPointLocation, const TArray<FLandscapeSplineInterpPoint>& Points, bool bRaiseTerrain, bool bLowerTerrain, TSet<ULandscapeComponent*>& ModifiedComponents)
+void RasterizeHeight(int32& MinX, int32& MinY, int32& MaxX, int32& MaxY, FLandscapeEditDataInterface& LandscapeEdit, bool bRaiseTerrain, bool bLowerTerrain, TSet<ULandscapeComponent*>& ModifiedComponents, TFunctionRef<void(FTriangleRasterizer<FLandscapeSplineHeightsRasterPolicy>&)> RasterizerFunction)
 {
 	if (!(bRaiseTerrain || bLowerTerrain))
 	{
@@ -128,18 +165,18 @@ void RasterizeControlPointHeights(int32& MinX, int32& MinY, int32& MaxX, int32&
 		return;
 	}
 
-	TArray<uint16> Data;
-	Data.AddZeroed((1 + MaxY - MinY) * (1 + MaxX - MinX));
+	TArray<uint16> HeightData;
+	HeightData.AddZeroed((1 + MaxY - MinY) * (1 + MaxX - MinX));
 
 	int32 ValidMinX = MinX;
 	int32 ValidMinY = MinY;
 	int32 ValidMaxX = MaxX;
 	int32 ValidMaxY = MaxY;
-	LandscapeEdit.GetHeightData(ValidMinX, ValidMinY, ValidMaxX, ValidMaxY, Data.GetData(), 0);
+	LandscapeEdit.GetHeightData(ValidMinX, ValidMinY, ValidMaxX, ValidMaxY, HeightData.GetData(), 0);
 
 	if (ValidMinX > ValidMaxX || ValidMinY > ValidMaxY)
 	{
-		// The control point's bounds don't intersect any data, so skip it
+		// The bounds don't intersect any data, so skip it
 		MinX = ValidMinX;
 		MinY = ValidMinY;
 		MaxX = ValidMaxX;
@@ -148,43 +185,85 @@ void RasterizeControlPointHeights(int32& MinX, int32& MinY, int32& MaxX, int32&
 		return;
 	}
 
-	FLandscapeEditDataInterface::ShrinkData(Data, MinX, MinY, MaxX, MaxY, ValidMinX, ValidMinY, ValidMaxX, ValidMaxY);
+	FLandscapeEditDataInterface::ShrinkData(HeightData, MinX, MinY, MaxX, MaxY, ValidMinX, ValidMinY, ValidMaxX, ValidMaxY);
+
+	const ALandscape* Landscape = LandscapeEdit.GetTargetLandscape();
+	bool bIsEditingLayerReservedForSplines = Landscape && Landscape->IsEditingLayerReservedForSplines();
+	check(!bIsEditingLayerReservedForSplines || (Landscape->GetLandscapeSplinesReservedLayer()->BlendMode == LSBM_AlphaBlend));
+
+	TArray<uint8> HeightAlphaBlendData;
+	TArray<uint8> HeightFlagsData;
+	TArray<uint8>* HeightAlphaBlendDataPtr = nullptr;
+	TArray<uint8>* HeightFlagsDataPtr = nullptr;
+
+	bool bCalculateNormals = true;
+	if (bIsEditingLayerReservedForSplines)
+	{
+		bCalculateNormals = false;
+
+		HeightAlphaBlendData.AddZeroed((1 + MaxY - MinY) * (1 + MaxX - MinX));
+		int32 AlphaValidMinX = MinX;
+		int32 AlphaValidMinY = MinY;
+		int32 AlphaValidMaxX = MaxX;
+		int32 AlphaValidMaxY = MaxY;
+		LandscapeEdit.GetHeightAlphaBlendData(AlphaValidMinX, AlphaValidMinY, AlphaValidMaxX, AlphaValidMaxY, HeightAlphaBlendData.GetData(), 0);
+		check(AlphaValidMinX == ValidMinX && AlphaValidMinY == ValidMinY && AlphaValidMaxX == ValidMaxX && AlphaValidMaxY == ValidMaxY);
+		FLandscapeEditDataInterface::ShrinkData(HeightAlphaBlendData, MinX, MinY, MaxX, MaxY, AlphaValidMinX, AlphaValidMinY, AlphaValidMaxX, AlphaValidMaxY);
+		HeightAlphaBlendDataPtr = &HeightAlphaBlendData;
+
+		HeightFlagsData.AddZeroed((1 + MaxY - MinY) * (1 + MaxX - MinX));
+		int32 FlagsValidMinX = MinX;
+		int32 FlagsValidMinY = MinY;
+		int32 FlagsValidMaxX = MaxX;
+		int32 FlagsValidMaxY = MaxY;
+		LandscapeEdit.GetHeightFlagsData(FlagsValidMinX, FlagsValidMinY, FlagsValidMaxX, FlagsValidMaxY, HeightFlagsData.GetData(), 0);
+		check(FlagsValidMinX == ValidMinX && FlagsValidMinY == ValidMinY && FlagsValidMaxX == ValidMaxX && FlagsValidMaxY == ValidMaxY);
+		FLandscapeEditDataInterface::ShrinkData(HeightFlagsData, MinX, MinY, MaxX, MaxY, FlagsValidMinX, FlagsValidMinY, FlagsValidMaxX, FlagsValidMaxY);
+		HeightFlagsDataPtr = &HeightFlagsData;
+	}
 
 	MinX = ValidMinX;
 	MinY = ValidMinY;
 	MaxX = ValidMaxX;
 	MaxY = ValidMaxY;
 
-	FTriangleRasterizer<FLandscapeSplineHeightsRasterPolicy> Rasterizer(
-		FLandscapeSplineHeightsRasterPolicy(Data, MinX, MinY, MaxX, MaxY, bRaiseTerrain, bLowerTerrain));
+	FTriangleRasterizer<FLandscapeSplineHeightsRasterPolicy> Rasterizer(FLandscapeSplineHeightsRasterPolicy(HeightData, MinX, MinY, MaxX, MaxY, bRaiseTerrain, bLowerTerrain, HeightAlphaBlendDataPtr, HeightFlagsDataPtr));
 
-	const FVector2D CenterPos = FVector2D(ControlPointLocation);
-	const FVector Center = FVector(1.0f, Points[0].StartEndFalloff, ControlPointLocation.Z * LANDSCAPE_INV_ZSCALE + LandscapeDataAccess::MidValue);
-
-	for (int32 i = Points.Num() - 1, j = 0; j < Points.Num(); i = j++)
-	{
-		// Solid center
-		const FVector2D Right0Pos = FVector2D(Points[i].Right);
-		const FVector2D Left1Pos = FVector2D(Points[j].Left);
-		const FVector2D Right1Pos = FVector2D(Points[j].Right);
-		const FVector Right0 = FVector(1.0f, Points[i].StartEndFalloff, Points[i].Right.Z);
-		const FVector Left1 = FVector(1.0f, Points[j].StartEndFalloff, Points[j].Left.Z);
-		const FVector Right1 = FVector(1.0f, Points[j].StartEndFalloff, Points[j].Right.Z);
+	RasterizerFunction(Rasterizer);
 
-		Rasterizer.DrawTriangle(Center, Right0, Left1, CenterPos, Right0Pos, Left1Pos, false);
-		Rasterizer.DrawTriangle(Center, Left1, Right1, CenterPos, Left1Pos, Right1Pos, false);
+	LandscapeEdit.SetHeightData(MinX, MinY, MaxX, MaxY, HeightData.GetData(), 0, bCalculateNormals, nullptr, HeightAlphaBlendDataPtr ? HeightAlphaBlendDataPtr->GetData() : nullptr, HeightFlagsDataPtr ? HeightFlagsDataPtr->GetData() : nullptr, false, nullptr, nullptr, true, !bIsEditingLayerReservedForSplines);
+	LandscapeEdit.GetComponentsInRegion(MinX, MinY, MaxX, MaxY, &ModifiedComponents);
+}
 
-		// Falloff
-		FVector2D FalloffRight0Pos = FVector2D(Points[i].FalloffRight);
-		FVector2D FalloffLeft1Pos = FVector2D(Points[j].FalloffLeft);
-		FVector FalloffRight0 = FVector(0.0f, Points[i].StartEndFalloff, Points[i].FalloffRight.Z);
-		FVector FalloffLeft1 = FVector(0.0f, Points[j].StartEndFalloff, Points[j].FalloffLeft.Z);
-		Rasterizer.DrawTriangle(Right0, FalloffRight0, Left1, Right0Pos, FalloffRight0Pos, Left1Pos, false);
-		Rasterizer.DrawTriangle(FalloffRight0, Left1, FalloffLeft1, FalloffRight0Pos, Left1Pos, FalloffLeft1Pos, false);
-	}
+void RasterizeControlPointHeights(int32& MinX, int32& MinY, int32& MaxX, int32& MaxY, FLandscapeEditDataInterface& LandscapeEdit, FVector ControlPointLocation, const TArray<FLandscapeSplineInterpPoint>& Points, bool bRaiseTerrain, bool bLowerTerrain, TSet<ULandscapeComponent*>& ModifiedComponents)
+{
+	RasterizeHeight(MinX, MinY, MaxX, MaxY, LandscapeEdit, bRaiseTerrain, bLowerTerrain, ModifiedComponents, [&](FTriangleRasterizer<FLandscapeSplineHeightsRasterPolicy>& Rasterizer)
+	{
+		const FVector2D CenterPos = FVector2D(ControlPointLocation);
+		const FVector Center = FVector(1.0f, Points[0].StartEndFalloff, ControlPointLocation.Z * LANDSCAPE_INV_ZSCALE + LandscapeDataAccess::MidValue);
 
-	LandscapeEdit.SetHeightData(MinX, MinY, MaxX, MaxY, Data.GetData(), 0, true);
-	LandscapeEdit.GetComponentsInRegion(MinX, MinY, MaxX, MaxY, &ModifiedComponents);
+		for (int32 i = Points.Num() - 1, j = 0; j < Points.Num(); i = j++)
+		{
+			// Solid center
+			const FVector2D Right0Pos = FVector2D(Points[i].Right);
+			const FVector2D Left1Pos = FVector2D(Points[j].Left);
+			const FVector2D Right1Pos = FVector2D(Points[j].Right);
+			const FVector Right0 = FVector(1.0f, Points[i].StartEndFalloff, Points[i].Right.Z);
+			const FVector Left1 = FVector(1.0f, Points[j].StartEndFalloff, Points[j].Left.Z);
+			const FVector Right1 = FVector(1.0f, Points[j].StartEndFalloff, Points[j].Right.Z);
+
+			Rasterizer.DrawTriangle(Center, Right0, Left1, CenterPos, Right0Pos, Left1Pos, false);
+			Rasterizer.DrawTriangle(Center, Left1, Right1, CenterPos, Left1Pos, Right1Pos, false);
+
+			// Falloff
+			FVector2D FalloffRight0Pos = FVector2D(Points[i].FalloffRight);
+			FVector2D FalloffLeft1Pos = FVector2D(Points[j].FalloffLeft);
+			FVector FalloffRight0 = FVector(0.0f, Points[i].StartEndFalloff, Points[i].FalloffRight.Z);
+			FVector FalloffLeft1 = FVector(0.0f, Points[j].StartEndFalloff, Points[j].FalloffLeft.Z);
+			Rasterizer.DrawTriangle(Right0, FalloffRight0, Left1, Right0Pos, FalloffRight0Pos, Left1Pos, false);
+			Rasterizer.DrawTriangle(FalloffRight0, Left1, FalloffLeft1, FalloffRight0Pos, Left1Pos, FalloffLeft1Pos, false);
+		}
+	});
 }
 
 void RasterizeControlPointAlpha(int32& MinX, int32& MinY, int32& MaxX, int32& MaxY, FLandscapeEditDataInterface& LandscapeEdit, FVector ControlPointLocation, const TArray<FLandscapeSplineInterpPoint>& Points, ULandscapeLayerInfoObject* LayerInfo, TSet<ULandscapeComponent*>& ModifiedComponents)
@@ -263,79 +342,39 @@ void RasterizeControlPointAlpha(int32& MinX, int32& MinY, int32& MaxX, int32& Ma
 
 void RasterizeSegmentHeight(int32& MinX, int32& MinY, int32& MaxX, int32& MaxY, FLandscapeEditDataInterface& LandscapeEdit, const TArray<FLandscapeSplineInterpPoint>& Points, bool bRaiseTerrain, bool bLowerTerrain, TSet<ULandscapeComponent*>& ModifiedComponents)
 {
-	if (!(bRaiseTerrain || bLowerTerrain))
+	RasterizeHeight(MinX, MinY, MaxX, MaxY, LandscapeEdit, bRaiseTerrain, bLowerTerrain, ModifiedComponents, [&](FTriangleRasterizer<FLandscapeSplineHeightsRasterPolicy>& Rasterizer)
 	{
-		return;
-	}
-
-	if (MinX > MaxX || MinY > MaxY)
-	{
-		return;
-	}
-
-	TArray<uint16> Data;
-	Data.AddZeroed((1 + MaxY - MinY) * (1 + MaxX - MinX));
-
-	int32 ValidMinX = MinX;
-	int32 ValidMinY = MinY;
-	int32 ValidMaxX = MaxX;
-	int32 ValidMaxY = MaxY;
-	LandscapeEdit.GetHeightData(ValidMinX, ValidMinY, ValidMaxX, ValidMaxY, Data.GetData(), 0);
-
-	if (ValidMinX > ValidMaxX || ValidMinY > ValidMaxY)
-	{
-		// The segment's bounds don't intersect any data, so skip it
-		MinX = ValidMinX;
-		MinY = ValidMinY;
-		MaxX = ValidMaxX;
-		MaxY = ValidMaxY;
-
-		return;
-	}
-
-	FLandscapeEditDataInterface::ShrinkData(Data, MinX, MinY, MaxX, MaxY, ValidMinX, ValidMinY, ValidMaxX, ValidMaxY);
-
-	MinX = ValidMinX;
-	MinY = ValidMinY;
-	MaxX = ValidMaxX;
-	MaxY = ValidMaxY;
-
-	FTriangleRasterizer<FLandscapeSplineHeightsRasterPolicy> Rasterizer(
-		FLandscapeSplineHeightsRasterPolicy(Data, MinX, MinY, MaxX, MaxY, bRaiseTerrain, bLowerTerrain));
-
-	for (int32 j = 1; j < Points.Num(); j++)
-	{
-		// Middle
-		FVector2D Left0Pos = FVector2D(Points[j - 1].Left);
-		FVector2D Right0Pos = FVector2D(Points[j - 1].Right);
-		FVector2D Left1Pos = FVector2D(Points[j].Left);
-		FVector2D Right1Pos = FVector2D(Points[j].Right);
-		FVector Left0 = FVector(1.0f, Points[j - 1].StartEndFalloff, Points[j - 1].Left.Z);
-		FVector Right0 = FVector(1.0f, Points[j - 1].StartEndFalloff, Points[j - 1].Right.Z);
-		FVector Left1 = FVector(1.0f, Points[j].StartEndFalloff, Points[j].Left.Z);
-		FVector Right1 = FVector(1.0f, Points[j].StartEndFalloff, Points[j].Right.Z);
-		Rasterizer.DrawTriangle(Left0, Right0, Left1, Left0Pos, Right0Pos, Left1Pos, false);
-		Rasterizer.DrawTriangle(Right0, Left1, Right1, Right0Pos, Left1Pos, Right1Pos, false);
-
-		// Left Falloff
-		FVector2D FalloffLeft0Pos = FVector2D(Points[j - 1].FalloffLeft);
-		FVector2D FalloffLeft1Pos = FVector2D(Points[j].FalloffLeft);
-		FVector FalloffLeft0 = FVector(0.0f, Points[j - 1].StartEndFalloff, Points[j - 1].FalloffLeft.Z);
-		FVector FalloffLeft1 = FVector(0.0f, Points[j].StartEndFalloff, Points[j].FalloffLeft.Z);
-		Rasterizer.DrawTriangle(FalloffLeft0, Left0, FalloffLeft1, FalloffLeft0Pos, Left0Pos, FalloffLeft1Pos, false);
-		Rasterizer.DrawTriangle(Left0, FalloffLeft1, Left1, Left0Pos, FalloffLeft1Pos, Left1Pos, false);
-
-		// Right Falloff
-		FVector2D FalloffRight0Pos = FVector2D(Points[j - 1].FalloffRight);
-		FVector2D FalloffRight1Pos = FVector2D(Points[j].FalloffRight);
-		FVector FalloffRight0 = FVector(0.0f, Points[j - 1].StartEndFalloff, Points[j - 1].FalloffRight.Z);
-		FVector FalloffRight1 = FVector(0.0f, Points[j].StartEndFalloff, Points[j].FalloffRight.Z);
-		Rasterizer.DrawTriangle(Right0, FalloffRight0, Right1, Right0Pos, FalloffRight0Pos, Right1Pos, false);
-		Rasterizer.DrawTriangle(FalloffRight0, Right1, FalloffRight1, FalloffRight0Pos, Right1Pos, FalloffRight1Pos, false);
-	}
-
-	LandscapeEdit.SetHeightData(MinX, MinY, MaxX, MaxY, Data.GetData(), 0, true);
-	LandscapeEdit.GetComponentsInRegion(MinX, MinY, MaxX, MaxY, &ModifiedComponents);
+		for (int32 j = 1; j < Points.Num(); j++)
+		{
+			// Middle
+			FVector2D Left0Pos = FVector2D(Points[j - 1].Left);
+			FVector2D Right0Pos = FVector2D(Points[j - 1].Right);
+			FVector2D Left1Pos = FVector2D(Points[j].Left);
+			FVector2D Right1Pos = FVector2D(Points[j].Right);
+			FVector Left0 = FVector(1.0f, Points[j - 1].StartEndFalloff, Points[j - 1].Left.Z);
+			FVector Right0 = FVector(1.0f, Points[j - 1].StartEndFalloff, Points[j - 1].Right.Z);
+			FVector Left1 = FVector(1.0f, Points[j].StartEndFalloff, Points[j].Left.Z);
+			FVector Right1 = FVector(1.0f, Points[j].StartEndFalloff, Points[j].Right.Z);
+			Rasterizer.DrawTriangle(Left0, Right0, Left1, Left0Pos, Right0Pos, Left1Pos, false);
+			Rasterizer.DrawTriangle(Right0, Left1, Right1, Right0Pos, Left1Pos, Right1Pos, false);
+
+			// Left Falloff
+			FVector2D FalloffLeft0Pos = FVector2D(Points[j - 1].FalloffLeft);
+			FVector2D FalloffLeft1Pos = FVector2D(Points[j].FalloffLeft);
+			FVector FalloffLeft0 = FVector(0.0f, Points[j - 1].StartEndFalloff, Points[j - 1].FalloffLeft.Z);
+			FVector FalloffLeft1 = FVector(0.0f, Points[j].StartEndFalloff, Points[j].FalloffLeft.Z);
+			Rasterizer.DrawTriangle(FalloffLeft0, Left0, FalloffLeft1, FalloffLeft0Pos, Left0Pos, FalloffLeft1Pos, false);
+			Rasterizer.DrawTriangle(Left0, FalloffLeft1, Left1, Left0Pos, FalloffLeft1Pos, Left1Pos, false);
+
+			// Right Falloff
+			FVector2D FalloffRight0Pos = FVector2D(Points[j - 1].FalloffRight);
+			FVector2D FalloffRight1Pos = FVector2D(Points[j].FalloffRight);
+			FVector FalloffRight0 = FVector(0.0f, Points[j - 1].StartEndFalloff, Points[j - 1].FalloffRight.Z);
+			FVector FalloffRight1 = FVector(0.0f, Points[j].StartEndFalloff, Points[j].FalloffRight.Z);
+			Rasterizer.DrawTriangle(Right0, FalloffRight0, Right1, Right0Pos, FalloffRight0Pos, Right1Pos, false);
+			Rasterizer.DrawTriangle(FalloffRight0, Right1, FalloffRight1, FalloffRight0Pos, Right1Pos, FalloffRight1Pos, false);
+		}
+	});
 }
 
 
@@ -423,6 +462,11 @@ bool ULandscapeInfo::ApplySplines(bool bOnlySelected)
 {
 	bool bResult = false;
 
+	ALandscape* Landscape = LandscapeActor.Get();
+	const FLandscapeLayer* Layer = Landscape ? Landscape->GetLandscapeSplinesReservedLayer() : nullptr;
+	FGuid SplinesTargetLayerGuid = Layer ? Layer->Guid : FGuid();
+	FScopedSetLandscapeEditingLayer Scope(Landscape, SplinesTargetLayerGuid, [=] { Landscape->RequestLayersContentUpdate(ELandscapeLayersContentUpdateFlag::All); });
+
 	ForAllLandscapeProxies([&bResult, bOnlySelected, this](ALandscapeProxy* Proxy)
 	{
 		bResult |= ApplySplinesInternal(bOnlySelected, Proxy);
@@ -431,16 +475,16 @@ bool ULandscapeInfo::ApplySplines(bool bOnlySelected)
 	return bResult;
 }
 
-bool ULandscapeInfo::ApplySplinesInternal(bool bOnlySelected, ALandscapeProxy* Landscape)
+bool ULandscapeInfo::ApplySplinesInternal(bool bOnlySelected, ALandscapeProxy* Proxy)
 {
-	if (!Landscape || !Landscape->SplineComponent || Landscape->SplineComponent->ControlPoints.Num() == 0 || Landscape->SplineComponent->Segments.Num() == 0)
+	if (!Proxy || !Proxy->SplineComponent || Proxy->SplineComponent->ControlPoints.Num() == 0 || Proxy->SplineComponent->Segments.Num() == 0)
 	{
 		return false;
 	}
 
 	FScopedTransaction Transaction(LOCTEXT("LandscapeSpline_ApplySplines", "Apply Splines to Landscape"));
 
-	const FTransform SplineToLandscape = Landscape->SplineComponent->GetComponentTransform().GetRelativeTransform(Landscape->LandscapeActorToWorld());
+	const FTransform SplineToLandscape = Proxy->SplineComponent->GetComponentTransform().GetRelativeTransform(Proxy->LandscapeActorToWorld());
 
 	FLandscapeEditDataInterface LandscapeEdit(this);
 	TSet<ULandscapeComponent*> ModifiedComponents;
@@ -452,7 +496,7 @@ bool ULandscapeInfo::ApplySplinesInternal(bool bOnlySelected, ALandscapeProxy* L
 		return false;
 	}
 
-	for (const ULandscapeSplineControlPoint* ControlPoint : Landscape->SplineComponent->ControlPoints)
+	for (const ULandscapeSplineControlPoint* ControlPoint : Proxy->SplineComponent->ControlPoints)
 	{
 		if (bOnlySelected && !ControlPoint->IsSplineSelected())
 		{
@@ -517,7 +561,7 @@ bool ULandscapeInfo::ApplySplinesInternal(bool bOnlySelected, ALandscapeProxy* L
 		}
 	}
 
-	for (const ULandscapeSplineSegment* Segment : Landscape->SplineComponent->Segments)
+	for (const ULandscapeSplineSegment* Segment : Proxy->SplineComponent->Segments)
 	{
 		if (bOnlySelected && !Segment->IsSplineSelected())
 		{
@@ -581,14 +625,19 @@ bool ULandscapeInfo::ApplySplinesInternal(bool bOnlySelected, ALandscapeProxy* L
 
 	LandscapeEdit.Flush();
 
-	for (ULandscapeComponent* Component : ModifiedComponents)
+	ALandscape* Landscape = Proxy->GetLandscapeActor();
+	bool bSkipCollisionAndNavUpdate = Landscape && Landscape->IsEditingLayerReservedForSplines();
+	if (!bSkipCollisionAndNavUpdate)
 	{
-		// Recreate collision for modified components and update the navmesh
-		ULandscapeHeightfieldCollisionComponent* CollisionComponent = Component->CollisionComponent.Get();
-		if (CollisionComponent)
+		for (ULandscapeComponent* Component : ModifiedComponents)
 		{
-			CollisionComponent->RecreateCollision();
-			FNavigationSystem::UpdateComponentData(*CollisionComponent);
+			// Recreate collision for modified components and update the navmesh
+			ULandscapeHeightfieldCollisionComponent* CollisionComponent = Component->CollisionComponent.Get();
+			if (CollisionComponent)
+			{
+				CollisionComponent->RecreateCollision();
+				FNavigationSystem::UpdateComponentData(*CollisionComponent);
+			}
 		}
 	}
 
diff --git a/Engine/Source/Runtime/Landscape/Private/LandscapeSplines.cpp b/Engine/Source/Runtime/Landscape/Private/LandscapeSplines.cpp
--- a/Engine/Source/Runtime/Landscape/Private/LandscapeSplines.cpp
+++ b/Engine/Source/Runtime/Landscape/Private/LandscapeSplines.cpp
@@ -37,6 +37,8 @@
 #include "Logging/TokenizedMessage.h"
 #include "Logging/MessageLog.h"
 #include "Misc/UObjectToken.h"
+#include "Settings/EditorExperimentalSettings.h"
+#include "Landscape.h"
 #endif
 
 IMPLEMENT_HIT_PROXY(HLandscapeSplineProxy, HHitProxy);
@@ -1876,6 +1878,19 @@ void ULandscapeSplineControlPoint::PostEditChangeProperty(FPropertyChangedEvent&
 		const bool bUpdateCollision = PropertyChangedEvent.ChangeType != EPropertyChangeType::Interactive;
 		UpdateSplinePoints(bUpdateCollision);
 	}
+
+	if (PropertyChangedEvent.ChangeType == EPropertyChangeType::ValueSet)
+	{
+		if (GetMutableDefault<UEditorExperimentalSettings>()->bLandscapeLayerSystem)
+		{
+			ALandscapeProxy* OuterLandscape = Cast<ALandscapeProxy>(GetOuterULandscapeSplinesComponent()->GetOwner());
+			ALandscape* Landscape = OuterLandscape ? OuterLandscape->GetLandscapeActor() : nullptr;
+			if (Landscape && Landscape->GetLandscapeSplinesReservedLayer())
+			{
+				Landscape->UpdateLandscapeSplines();
+			}
+		}
+	}
 }
 #endif // WITH_EDITOR
 
@@ -2819,6 +2834,19 @@ void ULandscapeSplineSegment::PostEditChangeProperty(FPropertyChangedEvent& Prop
 		const bool bUpdateCollision = PropertyChangedEvent.ChangeType != EPropertyChangeType::Interactive;
 		UpdateSplinePoints(bUpdateCollision);
 	}
+
+	if (PropertyChangedEvent.ChangeType == EPropertyChangeType::ValueSet)
+	{
+		if (GetMutableDefault<UEditorExperimentalSettings>()->bLandscapeLayerSystem)
+		{
+			ALandscapeProxy* OuterLandscape = Cast<ALandscapeProxy>(GetOuterULandscapeSplinesComponent()->GetOwner());
+			ALandscape* Landscape = OuterLandscape ? OuterLandscape->GetLandscapeActor() : nullptr;
+			if (Landscape && Landscape->GetLandscapeSplinesReservedLayer())
+			{
+				Landscape->UpdateLandscapeSplines();
+			}
+		}
+	}
 }
 #endif // WITH_EDITOR
 
diff --git a/Engine/Source/Runtime/Landscape/Public/LandscapeEdit.h b/Engine/Source/Runtime/Landscape/Public/LandscapeEdit.h
--- a/Engine/Source/Runtime/Landscape/Public/LandscapeEdit.h
+++ b/Engine/Source/Runtime/Landscape/Public/LandscapeEdit.h
@@ -119,20 +119,30 @@ struct LANDSCAPE_API FLandscapeEditDataInterface : public FLandscapeTextureDataI
 	//
 	// Heightmap access
 	//
-	void SetHeightData(int32 X1, int32 Y1, int32 X2, int32 Y2, const uint16* InData, int32 InStride, bool InCalcNormals, const uint16* InNormalData = nullptr, bool InCreateComponents = false, UTexture2D* InHeightmap = nullptr, UTexture2D* InXYOffsetmapTexture = nullptr,
+	void SetHeightData(int32 X1, int32 Y1, int32 X2, int32 Y2, const uint16* InData, int32 InStride, bool InCalcNormals, const uint16* InNormalData = nullptr, const uint8* InHeightAlphaBlendData = nullptr, const uint8* InHeightRaiseLowerData = nullptr, bool InCreateComponents = false, UTexture2D* InHeightmap = nullptr, UTexture2D* InXYOffsetmapTexture = nullptr,
 					   bool InUpdateBounds = true, bool InUpdateCollision = true, bool InGenerateMips = true);
 
 	// Helper accessor
 	FORCEINLINE uint16 GetHeightMapData(const ULandscapeComponent* Component, int32 TexU, int32 TexV, FColor* TextureData = NULL);
+	// Helper accessor
+	FORCEINLINE uint8 GetHeightMapAlphaBlendData(const ULandscapeComponent* Component, int32 TexU, int32 TexV, FColor* TextureData = NULL);
+	// Helper accessor
+	FORCEINLINE uint8 GetHeightMapFlagsData(const ULandscapeComponent* Component, int32 TexU, int32 TexV, FColor* TextureData = NULL);
 	// Generic
 	template<typename TStoreData>
 	void GetHeightDataTempl(int32& X1, int32& Y1, int32& X2, int32& Y2, TStoreData& StoreData);
+	template<typename TStoreData>
+	void GetHeightAlphaBlendDataTempl(int32& X1, int32& Y1, int32& X2, int32& Y2, TStoreData& StoreData);
+	template<typename TStoreData>
+	void GetHeightFlagsDataTempl(int32& X1, int32& Y1, int32& X2, int32& Y2, TStoreData& StoreData);
 	// Without data interpolation, able to get normal data
 	template<typename TStoreData>
 	void GetHeightDataTemplFast(const int32 X1, const int32 Y1, const int32 X2, const int32 Y2, TStoreData& StoreData, UTexture2D* InHeightmap = nullptr, TStoreData* NormalData = NULL);
 	// Implementation for fixed array
 	void GetHeightData(int32& X1, int32& Y1, int32& X2, int32& Y2, uint16* Data, int32 Stride);
 	void GetHeightDataFast(const int32 X1, const int32 Y1, const int32 X2, const int32 Y2, uint16* Data, int32 Stride, uint16* NormalData = NULL, UTexture2D* InHeightmap = nullptr);
+	void GetHeightAlphaBlendData(int32& X1, int32& Y1, int32& X2, int32& Y2, uint8* Data, int32 Stride);
+	void GetHeightFlagsData(int32& X1, int32& Y1, int32& X2, int32& Y2, uint8* Data, int32 Stride);
 	// Implementation for sparse array
 	void GetHeightData(int32& X1, int32& Y1, int32& X2, int32& Y2, TMap<FIntPoint, uint16>& SparseData);
 	void GetHeightDataFast(const int32 X1, const int32 Y1, const int32 X2, const int32 Y2, TMap<FIntPoint, uint16>& SparseData, TMap<FIntPoint, uint16>* NormalData = NULL, UTexture2D* InHeightmap = nullptr);
@@ -208,6 +218,7 @@ struct LANDSCAPE_API FLandscapeEditDataInterface : public FLandscapeTextureDataI
 	template<typename T>
 	static void ShrinkData(TArray<T>& Data, int32 OldMinX, int32 OldMinY, int32 OldMaxX, int32 OldMaxY, int32 NewMinX, int32 NewMinY, int32 NewMaxX, int32 NewMaxY);
 
+	const ALandscape* GetTargetLandscape() const;
 private:
 	int32 ComponentSizeQuads;
 	int32 SubsectionSizeQuads;
@@ -223,6 +234,12 @@ private:
 		const int32& ComponentSizeX, const int32& ComponentSizeY, TData* CornerValues,
 		TArray<bool>& NoBorderY1, TArray<bool>& NoBorderY2, TArray<bool>& ComponentDataExist, TStoreData& StoreData);
 
+	// Generic Height Data access
+	template<typename TStoreData>
+	void GetHeightDataInternal(int32& ValidX1, int32& ValidY1, int32& ValidX2, int32& ValidY2, TStoreData& StoreData, TFunctionRef<typename TStoreData::DataType(const ULandscapeComponent*, int32, int32, FColor*)> GetHeightDataFunction);
+
+	FORCEINLINE FColor& GetHeightMapColor(const ULandscapeComponent* Component, int32 TexU, int32 TexV, FColor* TextureData);
+
 	// test if layer is whitelisted for a given texel
 	inline bool IsWhitelisted(const ULandscapeLayerInfoObject* LayerInfo,
 	                          int32 ComponentIndexX, int32 SubIndexX, int32 SubX,
