Fixed grass not rendering
Added back broken feature by rendering refactor

diff --git a/Engine/Source/Runtime/Landscape/Private/LandscapeGrass.cpp b/Engine/Source/Runtime/Landscape/Private/LandscapeGrass.cpp
--- a/Engine/Source/Runtime/Landscape/Private/LandscapeGrass.cpp
+++ b/Engine/Source/Runtime/Landscape/Private/LandscapeGrass.cpp
@@ -322,6 +322,8 @@ public:
 		int32 NumPasses,
 		FVector2D ViewOffset,
 		float PassOffsetX,
+		int32 FirstHeightMipsPassIndex, 
+		const TArray<int32>& HeightMips,
 		const FPrimitiveSceneProxy* RESTRICT PrimitiveSceneProxy);
 
 	virtual void AddMeshBatch(const FMeshBatch& RESTRICT MeshBatch, uint64 BatchElementMask, const FPrimitiveSceneProxy* RESTRICT PrimitiveSceneProxy, int32 StaticMeshId = -1) override final
@@ -339,7 +341,9 @@ private:
 		const FMaterial& RESTRICT MaterialResource,
 		int32 NumPasses,
 		FVector2D ViewOffset,
-		float PassOffsetX);
+		float PassOffsetX,
+		int32 FirstHeightMipsPassIndex,
+		const TArray<int32>& HeightMips);
 
 	FMeshPassProcessorRenderState PassDrawRenderState;
 };
@@ -359,6 +363,8 @@ void FLandscapeGrassWeightMeshProcessor::AddMeshBatch(const FMeshBatch& RESTRICT
 	int32 NumPasses,
 	FVector2D ViewOffset,
 	float PassOffsetX,
+	int32 FirstHeightMipsPassIndex, 
+	const TArray<int32>& HeightMips, 
 	const FPrimitiveSceneProxy* RESTRICT PrimitiveSceneProxy)
 {
 	// Determine the mesh's material and blend mode.
@@ -367,7 +373,7 @@ void FLandscapeGrassWeightMeshProcessor::AddMeshBatch(const FMeshBatch& RESTRICT
 
 	const FMaterialRenderProxy& MaterialRenderProxy = FallbackMaterialRenderProxyPtr ? *FallbackMaterialRenderProxyPtr : *MeshBatch.MaterialRenderProxy;
 
-	Process(MeshBatch, BatchElementMask, PrimitiveSceneProxy, MaterialRenderProxy, Material, NumPasses, ViewOffset, PassOffsetX);
+	Process(MeshBatch, BatchElementMask, PrimitiveSceneProxy, MaterialRenderProxy, Material, NumPasses, ViewOffset, PassOffsetX, FirstHeightMipsPassIndex, HeightMips);
 }
 
 void FLandscapeGrassWeightMeshProcessor::Process(
@@ -378,7 +384,9 @@ void FLandscapeGrassWeightMeshProcessor::Process(
 	const FMaterial& RESTRICT MaterialResource,
 	int32 NumPasses,
 	FVector2D ViewOffset,
-	float PassOffsetX)
+	float PassOffsetX,
+	int32 FirstHeightMipsPassIndex,
+	const TArray<int32>& HeightMips)
 {
 	const FVertexFactory* VertexFactory = MeshBatch.VertexFactory;
 
@@ -401,12 +409,14 @@ void FLandscapeGrassWeightMeshProcessor::Process(
 
 	for (int32 PassIndex = 0; PassIndex < NumPasses; ++PassIndex)
 	{
-		ShaderElementData.OutputPass = PassIndex;
+		ShaderElementData.OutputPass = (PassIndex >= FirstHeightMipsPassIndex) ? 0 : PassIndex;
 		ShaderElementData.RenderOffset = ViewOffset + FVector2D(PassOffsetX * PassIndex, 0);
 
+		uint64 Mask = (PassIndex >= FirstHeightMipsPassIndex) ? HeightMips[PassIndex - FirstHeightMipsPassIndex] : BatchElementMask;
+
 		BuildMeshDrawCommands(
 			MeshBatch,
-			BatchElementMask,
+			Mask,
 			PrimitiveSceneProxy,
 			MaterialRenderProxy,
 			MaterialResource,
@@ -495,21 +505,21 @@ public:
 		FMemMark Mark(FMemStack::Get());
 
 		DrawDynamicMeshPass(*View, RHICmdList,
-			[View, PassOffsetX = PassOffsetX, &ComponentInfos = ComponentInfos, NumPasses = NumPasses](FDynamicPassMeshDrawListContext* DynamicMeshPassContext)
+			[View, PassOffsetX = PassOffsetX, &ComponentInfos = ComponentInfos, NumPasses = NumPasses, FirstHeightMipsPassIndex = FirstHeightMipsPassIndex, HeightMips = HeightMips](FDynamicPassMeshDrawListContext* DynamicMeshPassContext)
 		{
 			FLandscapeGrassWeightMeshProcessor PassMeshProcessor(
 				nullptr,
 				View,
 				DynamicMeshPassContext);
 
-			const uint64 DefaultBatchElementMask = ~0ull;
+			const uint64 DefaultBatchElementMask = ~0ul;
 
 			for (auto& ComponentInfo : ComponentInfos)
 			{
 				const FMeshBatch& Mesh = ComponentInfo.SceneProxy->GetGrassMeshBatch();
 				Mesh.MaterialRenderProxy->UpdateUniformExpressionCacheIfNeeded(View->GetFeatureLevel());
 
-				PassMeshProcessor.AddMeshBatch(Mesh, DefaultBatchElementMask, NumPasses, ComponentInfo.ViewOffset, PassOffsetX, ComponentInfo.SceneProxy);
+				PassMeshProcessor.AddMeshBatch(Mesh, DefaultBatchElementMask, NumPasses, ComponentInfo.ViewOffset, PassOffsetX, FirstHeightMipsPassIndex, HeightMips, ComponentInfo.SceneProxy);
 			}
 		});
 	}
