Fixes for Landscape Layer System
- Proper updating of Heightmap Colllision
- Multiple Undo/Redo fixes

diff --git a/Engine/Source/Editor/LandscapeEditor/Private/LandscapeEdModeComponentTools.cpp b/Engine/Source/Editor/LandscapeEditor/Private/LandscapeEdModeComponentTools.cpp
--- a/Engine/Source/Editor/LandscapeEditor/Private/LandscapeEdModeComponentTools.cpp
+++ b/Engine/Source/Editor/LandscapeEditor/Private/LandscapeEdModeComponentTools.cpp
@@ -656,7 +656,7 @@ public:
 								CurrentWeightmapTexture->PostEditChange();
 
 								// Store it in the usage map
-								CurrentWeightmapUsage = LandscapeProxy->WeightmapUsageMap.Add(CurrentWeightmapTexture, NewObject<ULandscapeWeightmapUsage>(LandscapeProxy));
+								CurrentWeightmapUsage = LandscapeProxy->WeightmapUsageMap.Add(CurrentWeightmapTexture, LandscapeProxy->CreateWeightmapUsage());
 
 								// UE_LOG(LogLandscape, Log, TEXT("Making a new texture %s"), *CurrentWeightmapTexture->GetName());
 							}
@@ -691,11 +691,13 @@ public:
 
 										if (OldWeightmapUsage != nullptr)
 										{
+											(*OldWeightmapUsage)->Modify();
 											(*OldWeightmapUsage)->ChannelUsage[AllocInfo.WeightmapTextureChannel] = nullptr;
 										}
 									}
 
 									// Assign the new allocation
+									CurrentWeightmapUsage->Modify();
 									CurrentWeightmapUsage->ChannelUsage[ChanIdx] = Component;
 									AllocInfo.WeightmapTextureIndex = NewWeightmapTextures.Num() - 1;
 									AllocInfo.WeightmapTextureChannel = ChanIdx;
diff --git a/Engine/Source/Runtime/Landscape/Classes/LandscapeComponent.h b/Engine/Source/Runtime/Landscape/Classes/LandscapeComponent.h
--- a/Engine/Source/Runtime/Landscape/Classes/LandscapeComponent.h
+++ b/Engine/Source/Runtime/Landscape/Classes/LandscapeComponent.h
@@ -342,6 +342,7 @@ private:
 	TArray<ULandscapeWeightmapUsage*>* CurrentEditingWeightmapTexturesUsage;
 	FGuid CurrentLayerGuid;
 
+	UPROPERTY(Transient)
 	TArray<ULandscapeWeightmapUsage*> WeightmapTexturesUsage;
 #endif // WITH_EDITORONLY_DATA
 
@@ -721,12 +722,18 @@ public:
 	 * @param ComponentX1, ComponentY1, ComponentX2, ComponentY2: region to update
 	 * @param bUpdateBounds: Whether to update bounds from render component.
 	 * @param XYOffsetTextureMipData: xy-offset map data
+	 * @returns True if CollisionComponent was created in this update.
 	 */
-	void UpdateCollisionHeightData(const FColor* HeightmapTextureMipData, const FColor* SimpleCollisionHeightmapTextureData, int32 ComponentX1=0, int32 ComponentY1=0, int32 ComponentX2=MAX_int32, int32 ComponentY2=MAX_int32, bool bUpdateBounds=false, const FColor* XYOffsetTextureMipData=nullptr);
+	bool UpdateCollisionHeightData(const FColor* HeightmapTextureMipData, const FColor* SimpleCollisionHeightmapTextureData, int32 ComponentX1=0, int32 ComponentY1=0, int32 ComponentX2=MAX_int32, int32 ComponentY2=MAX_int32, bool bUpdateBounds=false, const FColor* XYOffsetTextureMipData=nullptr);
+
+	/**
+	 * Deletes Collision Component
+	 */
+	void DestroyCollisionData();
 
 	/** Updates collision component height data for the entire component, locking and unlocking heightmap textures
 	 * @param: bRebuild: If true, recreates the collision component */
-	void UpdateCollisionData(bool bRebuild);
+	void UpdateCollisionData();
 
 	/**
 	 * Update collision component dominant layer data
diff --git a/Engine/Source/Runtime/Landscape/Classes/LandscapeHeightfieldCollisionComponent.h b/Engine/Source/Runtime/Landscape/Classes/LandscapeHeightfieldCollisionComponent.h
--- a/Engine/Source/Runtime/Landscape/Classes/LandscapeHeightfieldCollisionComponent.h
+++ b/Engine/Source/Runtime/Landscape/Classes/LandscapeHeightfieldCollisionComponent.h
@@ -243,6 +243,8 @@ public:
 #if WITH_EDITORONLY_DATA
 	// Called from editor code to manage foliage instances on landscape.
 	LANDSCAPE_API void SnapFoliageInstances(const FBox& InInstanceBox);
+
+	LANDSCAPE_API void SnapFoliageInstances();
 #endif
 };
 
diff --git a/Engine/Source/Runtime/Landscape/Classes/LandscapeProxy.h b/Engine/Source/Runtime/Landscape/Classes/LandscapeProxy.h
--- a/Engine/Source/Runtime/Landscape/Classes/LandscapeProxy.h
+++ b/Engine/Source/Runtime/Landscape/Classes/LandscapeProxy.h
@@ -366,8 +366,9 @@ struct FRenderDataPerHeightmap
 
 	UPROPERTY(Transient)
 	UTexture2D* OriginalHeightmap;
-
-	FLandscapeLayersTexture2DCPUReadBackResource* HeightmapsCPUReadBack;
+	
+    UPROPERTY(Transient)
+	int32 HeightmapsCPUReadBackResourceIndex;
 
 	UPROPERTY(Transient)
 	TArray<ULandscapeComponent*> Components;
@@ -387,6 +388,7 @@ struct FWeightmapLayerData
 	UPROPERTY()
 	TArray<FWeightmapLayerAllocationInfo> WeightmapLayerAllocations;
 
+	UPROPERTY(Transient)
 	TArray<ULandscapeWeightmapUsage*> WeightmapTextureUsages;	// Easy Access ref to data stored into the LandscapeProxy weightmap usage map
 };
 
@@ -651,6 +653,8 @@ public:
 	UPROPERTY(Transient)
 	TMap<UTexture2D*, FRenderDataPerHeightmap> RenderDataPerHeightmap; // Mapping between Original heightmap and general render data
 
+	TArray<TUniquePtr<FLandscapeLayersTexture2DCPUReadBackResource>> HeightmapsCPUReadBackResources;
+
 	TMap<UTexture2D*, FLandscapeLayersTexture2DCPUReadBackResource*> WeightmapCPUReadBackTextures; // Mapping between Original weightmap and tyhe CPU readback resource
 
 	FRenderCommandFence ReleaseResourceFence;
@@ -925,6 +929,9 @@ public:
 	/** Creates a Texture2D for use by this landscape proxy or one of it's components. If OptionalOverrideOuter is not specified, the level is used. */
 	LANDSCAPE_API UTexture2D* CreateLandscapeTexture(int32 InSizeX, int32 InSizeY, TextureGroup InLODGroup, ETextureSourceFormat InFormat, UObject* OptionalOverrideOuter = nullptr, bool bCompress = false) const;
 
+	/** Creates a LandscapeWeightMapUsage object outered to this proxy. */
+	LANDSCAPE_API ULandscapeWeightmapUsage* CreateWeightmapUsage();
+
 	/* For the grassmap rendering notification */
 	int32 NumComponentsNeedingGrassMapRender;
 	LANDSCAPE_API static int32 TotalComponentsNeedingGrassMapRender;
diff --git a/Engine/Source/Runtime/Landscape/Private/LandscapeCollision.cpp b/Engine/Source/Runtime/Landscape/Private/LandscapeCollision.cpp
--- a/Engine/Source/Runtime/Landscape/Private/LandscapeCollision.cpp
+++ b/Engine/Source/Runtime/Landscape/Private/LandscapeCollision.cpp
@@ -45,6 +45,10 @@
 #include "Physics/PhysicsInterfaceCore.h"
 #include "Physics/PhysicsInterfaceUtils.h"
 
+#if WITH_EDITOR
+#include "Settings/EditorExperimentalSettings.h"
+#endif
+
 #if WITH_EDITOR && WITH_PHYSX
 	#include "Physics/IPhysXCooking.h"
 #endif
@@ -1225,6 +1229,11 @@ void ULandscapeHeightfieldCollisionComponent::RecreateCollision()
 }
 
 #if WITH_EDITORONLY_DATA
+void ULandscapeHeightfieldCollisionComponent::SnapFoliageInstances()
+{
+	SnapFoliageInstances(FBox(FVector(-WORLD_MAX), FVector(WORLD_MAX)));
+}
+
 void ULandscapeHeightfieldCollisionComponent::SnapFoliageInstances(const FBox& InInstanceBox)
 {
 	UWorld* ComponentWorld = GetWorld();
@@ -1465,13 +1474,17 @@ void ULandscapeHeightfieldCollisionComponent::PostEditUndo()
 {
 	Super::PostEditUndo();
 
-	// Reinitialize physics after undo
-	if (CollisionSizeQuads > 0)
+    // Landscape Layers are updates are delayed and done in  ALandscape::TickLayers
+	if (!GetMutableDefault<UEditorExperimentalSettings>()->bLandscapeLayerSystem)
 	{
-		RecreateCollision();
-	}
+		// Reinitialize physics after undo
+		if (CollisionSizeQuads > 0)
+		{
+			RecreateCollision();
+		}
 
-	FNavigationSystem::UpdateComponentData(*this);
+		FNavigationSystem::UpdateComponentData(*this);
+	}
 }
 
 bool ULandscapeHeightfieldCollisionComponent::ComponentIsTouchingSelectionBox(const FBox& InSelBBox, const FEngineShowFlags& ShowFlags, const bool bConsiderOnlyBSP, const bool bMustEncompassEntireComponent) const
diff --git a/Engine/Source/Runtime/Landscape/Private/LandscapeEdit.cpp b/Engine/Source/Runtime/Landscape/Private/LandscapeEdit.cpp
--- a/Engine/Source/Runtime/Landscape/Private/LandscapeEdit.cpp
+++ b/Engine/Source/Runtime/Landscape/Private/LandscapeEdit.cpp
@@ -119,7 +119,11 @@ void ULandscapeComponent::UpdateCachedBounds()
 	ULandscapeHeightfieldCollisionComponent* HFCollisionComponent = CollisionComponent.Get();
 	if (HFCollisionComponent)
 	{
-		HFCollisionComponent->Modify();
+        // In Landscape Layers the Collision Component is slave and doesn't need to be transacted
+		if (!GetMutableDefault<UEditorExperimentalSettings>()->bLandscapeLayerSystem)
+		{
+			HFCollisionComponent->Modify();
+		}
 		HFCollisionComponent->CachedLocalBox = CachedLocalBox;
 		HFCollisionComponent->UpdateComponentToWorld();
 	}
@@ -570,6 +574,7 @@ void ULandscapeComponent::PostEditUndo()
 	TSet<ULandscapeComponent*> Components;
 	Components.Add(this);
 	GetLandscapeProxy()->FlushGrassComponents(&Components);
+	GetLandscapeActor()->RequestLayersContentUpdate(ELandscapeLayersContentUpdateFlag::All, true);
 }
 
 void ALandscapeProxy::FixupWeightmaps()
@@ -660,7 +665,7 @@ void ULandscapeComponent::FixupWeightmaps()
 
 				if (TempUsage == nullptr)
 				{
-					TempUsage = &Proxy->WeightmapUsageMap.Add(WeightmapTexture, NewObject<ULandscapeWeightmapUsage>(GetLandscapeProxy()));
+					TempUsage = &Proxy->WeightmapUsageMap.Add(WeightmapTexture, GetLandscapeProxy()->CreateWeightmapUsage());
 					(*TempUsage)->LayerGuid.Invalidate();
 				}
 
@@ -743,7 +748,7 @@ struct FLandscapeComponentAlphaInfo
 };
 
 
-void ULandscapeComponent::UpdateCollisionHeightData(const FColor* const HeightmapTextureMipData, const FColor* const SimpleCollisionHeightmapTextureData, int32 ComponentX1/*=0*/, int32 ComponentY1/*=0*/, int32 ComponentX2/*=MAX_int32*/, int32 ComponentY2/*=MAX_int32*/, bool bUpdateBounds/*=false*/, const FColor* XYOffsetTextureMipData/*=nullptr*/)
+bool ULandscapeComponent::UpdateCollisionHeightData(const FColor* const HeightmapTextureMipData, const FColor* const SimpleCollisionHeightmapTextureData, int32 ComponentX1/*=0*/, int32 ComponentY1/*=0*/, int32 ComponentX2/*=MAX_int32*/, int32 ComponentY2/*=MAX_int32*/, bool bUpdateBounds/*=false*/, const FColor* XYOffsetTextureMipData/*=nullptr*/)
 {
 	ULandscapeInfo* Info = GetLandscapeInfo();
 	ALandscapeProxy* Proxy = GetLandscapeProxy();
@@ -769,23 +774,27 @@ void ULandscapeComponent::UpdateCollisionHeightData(const FColor* const Heightma
 	bool CreatedNew = false;
 	bool ChangeType = false;
 
-	if (CollisionComp)
+    // In Landscape Layers the Collision Component is slave and doesn't need to be transacted
+	if (!GetMutableDefault<UEditorExperimentalSettings>()->bLandscapeLayerSystem)
 	{
-		CollisionComp->Modify();
+		if (CollisionComp)
+		{
+			CollisionComp->Modify();
+		}
 	}
 
 	// Existing collision component is same type with collision
 	if (CollisionComp && ((XYOffsetmapTexture == nullptr) == (MeshCollisionComponent == nullptr)))
 	{
-		ComponentX1 = FMath::Min(ComponentX1, ComponentSizeQuads);
-		ComponentY1 = FMath::Min(ComponentY1, ComponentSizeQuads);
-		ComponentX2 = FMath::Max(ComponentX2, 0);
-		ComponentY2 = FMath::Max(ComponentY2, 0);
+		ComponentX1 = FMath::Clamp(ComponentX1, 0, ComponentSizeQuads);
+		ComponentY1 = FMath::Clamp(ComponentY1, 0, ComponentSizeQuads);
+		ComponentX2 = FMath::Clamp(ComponentX2, 0, ComponentSizeQuads);
+		ComponentY2 = FMath::Clamp(ComponentY2, 0, ComponentSizeQuads);
 
 		if (ComponentX2 < ComponentX1 || ComponentY2 < ComponentY1)
 		{
 			// nothing to do
-			return;
+			return false;
 		}
 
 		if (bUpdateBounds)
@@ -1125,17 +1134,22 @@ void ULandscapeComponent::UpdateCollisionHeightData(const FColor* const Heightma
 	{
 		CollisionComp->RegisterComponent();
 	}
+
+	return CreatedNew;
 }
 
-void ULandscapeComponent::UpdateCollisionData(bool bRebuild)
+void ULandscapeComponent::DestroyCollisionData()
 {
 	ULandscapeHeightfieldCollisionComponent* CollisionComp = CollisionComponent.Get();
-	if (CollisionComp && bRebuild)
+	if (CollisionComp)
 	{
 		CollisionComp->DestroyComponent();
 		CollisionComponent = CollisionComp = nullptr;
 	}
+}
 
+void ULandscapeComponent::UpdateCollisionData()
+{
 	TArray<uint8> CollisionMipData;
 	TArray<uint8> SimpleCollisionMipData;
 	TArray<uint8> XYOffsetMipData;
@@ -1150,11 +1164,21 @@ void ULandscapeComponent::UpdateCollisionData(bool bRebuild)
 		XYOffsetmapTexture->Source.GetMipData(XYOffsetMipData, CollisionMipLevel);
 	}
 
-	UpdateCollisionHeightData(
+	bool bCreatedNewComponent = UpdateCollisionHeightData(
 		(FColor*)CollisionMipData.GetData(),
 		SimpleCollisionMipLevel > CollisionMipLevel ? (FColor*)SimpleCollisionMipData.GetData() : nullptr,
 		0, 0, MAX_int32, MAX_int32, true,
 		XYOffsetmapTexture ? (FColor*)XYOffsetMipData.GetData() : nullptr);
+
+    // If the Collision Component hasn't changed we need to reflect its data to the underlying physics implementation
+	if (!bCreatedNewComponent)
+	{
+		ULandscapeHeightfieldCollisionComponent* CollisionComp = CollisionComponent.Get();
+		if (CollisionComp)
+		{
+			CollisionComp->RecreateCollision();
+		}
+	}
 }
 
 void ULandscapeComponent::UpdateCollisionLayerData(const FColor* const* const WeightmapTextureMipData, const FColor* const* const SimpleCollisionWeightmapTextureMipData, int32 ComponentX1, int32 ComponentY1, int32 ComponentX2, int32 ComponentY2)
@@ -2593,7 +2617,7 @@ LANDSCAPE_API void ALandscapeProxy::Import(
 
 					const int32 ThisAllocationLayers = FMath::Min<int32>(RemainingLayers, 4);
 					new(TextureAllocations) FWeightmapTextureAllocation(ComponentX, ComponentY, ThisAllocationLayers, WeightmapTexture, MipData);
-					ULandscapeWeightmapUsage* WeightmapUsage = WeightmapUsageMap.Add(WeightmapTexture, NewObject<ULandscapeWeightmapUsage>(this));
+					ULandscapeWeightmapUsage* WeightmapUsage = WeightmapUsageMap.Add(WeightmapTexture, CreateWeightmapUsage());
 					ComponentWeightmapTexturesUsage.Add(WeightmapUsage);
 
 					UE_LOG(LogLandscape, Log, TEXT("  ==> Storing %d channels in new texture %s"), ThisAllocationLayers, *WeightmapTexture->GetName());
@@ -3745,7 +3769,8 @@ void ALandscapeProxy::RecreateCollisionComponents()
 		{
 			Comp->CollisionMipLevel = CollisionMipLevel;
 			Comp->SimpleCollisionMipLevel = SimpleCollisionMipLevel;
-			Comp->UpdateCollisionData(true);
+			Comp->DestroyCollisionData();
+			Comp->UpdateCollisionData();
 		}
 	}
 }
@@ -4433,7 +4458,8 @@ void ULandscapeComponent::PostEditChangeProperty(FPropertyChangedEvent& Property
 		SimpleCollisionMipLevel = FMath::Clamp<int32>(SimpleCollisionMipLevel, 0, FMath::CeilLogTwo(SubsectionSizeQuads + 1) - 1);
 		if (PropertyChangedEvent.ChangeType != EPropertyChangeType::Interactive)
 		{
-			UpdateCollisionData(true); // Rebuild for new CollisionMipLevel
+			DestroyCollisionData();
+			UpdateCollisionData(); // Rebuild for new CollisionMipLevel
 		}
 	}
 
@@ -4625,6 +4651,11 @@ void ULandscapeComponent::ReallocateWeightmaps(FLandscapeEditDataInterface* Data
 
 					AllocInfo.WeightmapTextureIndex = TexIdx;
 					AllocInfo.WeightmapTextureChannel = ChanIdx;
+
+					if (InSaveToTransactionBuffer)
+					{
+						Usage->Modify();
+					}
 					Usage->ChannelUsage[ChanIdx] = this;
 
 					NeededNewChannels--;
@@ -4726,7 +4757,7 @@ void ULandscapeComponent::ReallocateWeightmaps(FLandscapeEditDataInterface* Data
 			}
 
 			// Store it in the usage map
-			CurrentWeightmapUsage = Proxy->WeightmapUsageMap.Add(CurrentWeightmapTexture, NewObject<ULandscapeWeightmapUsage>(GetLandscapeProxy()));
+			CurrentWeightmapUsage = Proxy->WeightmapUsageMap.Add(CurrentWeightmapTexture, GetLandscapeProxy()->CreateWeightmapUsage());
 			CurrentWeightmapUsage->LayerGuid = InCanUseCurrentEditingWeightmap ? CurrentLayerGuid : FGuid();
 			// UE_LOG(LogLandscape, Log, TEXT("Making a new texture %s"), *CurrentWeightmapTexture->GetName());
 		}
@@ -4766,10 +4797,18 @@ void ULandscapeComponent::ReallocateWeightmaps(FLandscapeEditDataInterface* Data
 
 					// Remove the old allocation
 					ULandscapeWeightmapUsage* OldWeightmapUsage = ComponentWeightmapTexturesUsage[AllocInfo.WeightmapTextureIndex];
+					if (InSaveToTransactionBuffer)
+					{
+						OldWeightmapUsage->Modify();
+					}
 					OldWeightmapUsage->ChannelUsage[AllocInfo.WeightmapTextureChannel] = nullptr;
 				}
 
 				// Assign the new allocation
+				if (InSaveToTransactionBuffer)
+				{
+					CurrentWeightmapUsage->Modify();
+				}
 				CurrentWeightmapUsage->ChannelUsage[ChanIdx] = this;
 				AllocInfo.WeightmapTextureIndex = NewWeightmapTextures.Num() - 1;
 				AllocInfo.WeightmapTextureChannel = ChanIdx;
@@ -5765,6 +5804,11 @@ UTexture2D* ALandscapeProxy::CreateLandscapeTexture(int32 InSizeX, int32 InSizeY
 	return NewTexture;
 }
 
+ULandscapeWeightmapUsage* ALandscapeProxy::CreateWeightmapUsage()
+{
+	return NewObject<ULandscapeWeightmapUsage>(this, ULandscapeWeightmapUsage::StaticClass(), NAME_None, RF_Transactional);
+}
+
 void ALandscapeProxy::RemoveOverlappingComponent(ULandscapeComponent* Component)
 {
 	Modify();
diff --git a/Engine/Source/Runtime/Landscape/Private/LandscapeEditInterface.cpp b/Engine/Source/Runtime/Landscape/Private/LandscapeEditInterface.cpp
--- a/Engine/Source/Runtime/Landscape/Private/LandscapeEditInterface.cpp
+++ b/Engine/Source/Runtime/Landscape/Private/LandscapeEditInterface.cpp
@@ -1706,6 +1706,7 @@ void ULandscapeComponent::DeleteLayer(ULandscapeLayerInfoObject* LayerInfo, FLan
 	ULandscapeWeightmapUsage* Usage = ComponentWeightmapTexturesUsage.IsValidIndex(DeleteLayerAllocation.WeightmapTextureIndex) ? ComponentWeightmapTexturesUsage[DeleteLayerAllocation.WeightmapTextureIndex] : nullptr;
 	if (Usage) // can be null if WeightmapUsageMap hasn't been built yet
 	{
+		Usage->Modify();
 		Usage->ChannelUsage[DeleteLayerAllocation.WeightmapTextureChannel] = nullptr;
 	}
 
@@ -1872,6 +1873,7 @@ void ULandscapeComponent::FillLayer(ULandscapeLayerInfoObject* LayerInfo, FLands
 			ULandscapeWeightmapUsage* Usage = ComponentWeightmapTexturesUsage[Allocation.WeightmapTextureIndex];
 			if (Usage) // can be null if WeightmapUsageMap hasn't been built yet
 			{
+				Usage->Modify();
 				Usage->ChannelUsage[Allocation.WeightmapTextureChannel] = nullptr;
 			}
 
@@ -2162,7 +2164,8 @@ void ULandscapeComponent::ReplaceLayer(ULandscapeLayerInfoObject* FromLayerInfo,
 		//check(Usage);
 		if (Usage)
 		{
-			Usage->ChannelUsage[FromLayerAllocation.WeightmapTextureChannel] = NULL;
+			Usage->Modify();
+			Usage->ChannelUsage[FromLayerAllocation.WeightmapTextureChannel] = nullptr;
 		}
 
 		// If this layer was the last usage for this texture, we can remove it.
@@ -2409,7 +2412,8 @@ bool DeleteLayerIfAllZero(ULandscapeComponent* const Component, const uint8* con
 	// Mark the channel as unallocated, so we can reuse it later
 	const int32 DeleteLayerWeightmapTextureIndex = ComponentWeightmapLayerAllocations[LayerIdx].WeightmapTextureIndex;
 	ULandscapeWeightmapUsage& Usage = *ComponentWeightmapTexturesUsage[DeleteLayerWeightmapTextureIndex];
-	Usage.ChannelUsage[ComponentWeightmapLayerAllocations[LayerIdx].WeightmapTextureChannel] = NULL;
+	Usage.Modify();
+	Usage.ChannelUsage[ComponentWeightmapLayerAllocations[LayerIdx].WeightmapTextureChannel] = nullptr;
 
 	// Remove the layer as it's totally painted away.
 	ComponentWeightmapLayerAllocations.RemoveAt(LayerIdx);
diff --git a/Engine/Source/Runtime/Landscape/Private/LandscapeEditLayers.cpp b/Engine/Source/Runtime/Landscape/Private/LandscapeEditLayers.cpp
--- a/Engine/Source/Runtime/Landscape/Private/LandscapeEditLayers.cpp
+++ b/Engine/Source/Runtime/Landscape/Private/LandscapeEditLayers.cpp
@@ -51,6 +51,13 @@ static TAutoConsoleVariable<int32> CVarOutputLayersWeightmapsRTContent(
 	0,
 	TEXT("This will output the content of render target used for weightmap. This is used for debugging only."));
 
+#if WITH_EDITOR
+static TAutoConsoleVariable<int32> CVarLandscapeSimulatePhysics(
+	TEXT("landscape.SimulatePhysics"),
+	0,
+	TEXT("This will enable physic simulation on worlds containing landscape."));
+#endif
+
 DECLARE_GPU_STAT_NAMED(LandscapeLayersRender, TEXT("Landscape Layer System Render"));
 DECLARE_GPU_STAT_NAMED(LandscapeLayersCopy, TEXT("Landscape Layer System Copy"));
 
@@ -1267,9 +1274,10 @@ void ALandscapeProxy::SetupLayers(int32 InNumComponentsX, int32 InNumComponentsY
 				FRenderDataPerHeightmap NewData;
 				NewData.Components.Add(Component);
 				NewData.OriginalHeightmap = ComponentHeightmapTexture;
-				NewData.HeightmapsCPUReadBack = new FLandscapeLayersTexture2DCPUReadBackResource(ComponentHeightmapTexture->Source.GetSizeX(), ComponentHeightmapTexture->Source.GetSizeY(), ComponentHeightmapTexture->GetPixelFormat(), ComponentHeightmapTexture->Source.GetNumMips());
-				BeginInitResource(NewData.HeightmapsCPUReadBack);
-
+				FLandscapeLayersTexture2DCPUReadBackResource* CPUReadBackResource = new FLandscapeLayersTexture2DCPUReadBackResource(ComponentHeightmapTexture->Source.GetSizeX(), ComponentHeightmapTexture->Source.GetSizeY(), ComponentHeightmapTexture->GetPixelFormat(), ComponentHeightmapTexture->Source.GetNumMips());
+				NewData.HeightmapsCPUReadBackResourceIndex = LandscapeProxy->HeightmapsCPUReadBackResources.Num();
+				LandscapeProxy->HeightmapsCPUReadBackResources.Add(TUniquePtr<FLandscapeLayersTexture2DCPUReadBackResource>(CPUReadBackResource));
+				BeginInitResource(CPUReadBackResource);
 				LandscapeProxy->RenderDataPerHeightmap.Add(ComponentHeightmapTexture, NewData);
 			}
 			else
@@ -1672,7 +1680,7 @@ void ALandscapeProxy::SetupLayers(int32 InNumComponentsX, int32 InNumComponentsY
 								NewWeightmapData.WeightmapTextureUsages[TextureIndex]->LayerGuid = LayerGuid;
 
 								// Create new Usage for the base layer as the other one will now be used by the Layer 1
-								ComponentWeightmapTexturesUsage[TextureIndex] = LandscapeProxy->WeightmapUsageMap.Add(NewWeightmapTexture, NewObject<ULandscapeWeightmapUsage>(LandscapeProxy));
+								ComponentWeightmapTexturesUsage[TextureIndex] = LandscapeProxy->WeightmapUsageMap.Add(NewWeightmapTexture, LandscapeProxy->CreateWeightmapUsage());
 
 								for (FWeightmapLayerAllocationInfo& Allocation : ComponentLayerAllocations)
 								{
@@ -1708,7 +1716,7 @@ void ALandscapeProxy::SetupLayers(int32 InNumComponentsX, int32 InNumComponentsY
 
 						if (TempUsage == nullptr)
 						{
-							TempUsage = &LandscapeProxy->WeightmapUsageMap.Add(WeightmapTexture, NewObject<ULandscapeWeightmapUsage>(LandscapeProxy));
+							TempUsage = &LandscapeProxy->WeightmapUsageMap.Add(WeightmapTexture, LandscapeProxy->CreateWeightmapUsage());
 							(*TempUsage)->LayerGuid = LayerGuid;
 						}
 
@@ -2740,15 +2748,14 @@ void ALandscape::RegenerateLayersHeightmaps()
 			{
 				int32 CurrentMip = 0;
 				FRenderDataPerHeightmap& HeightmapRenderData = ItPair.Value;
-
-				CopyLayersTexture(CombinedHeightmapAtlasRT, HeightmapRenderData.OriginalHeightmap, HeightmapRenderData.HeightmapsCPUReadBack, HeightmapRenderData.TopLeftSectionBase, CurrentMip, CurrentMip);
+				CopyLayersTexture(CombinedHeightmapAtlasRT, HeightmapRenderData.OriginalHeightmap, Landscape->HeightmapsCPUReadBackResources[HeightmapRenderData.HeightmapsCPUReadBackResourceIndex].Get(), HeightmapRenderData.TopLeftSectionBase, CurrentMip, CurrentMip);
 				++CurrentMip;
 
 				for (int32 MipRTIndex = EHeightmapRTType::HeightmapRT_Mip1; MipRTIndex < EHeightmapRTType::HeightmapRT_Count; ++MipRTIndex)
 				{
 					if (HeightmapRTList[MipRTIndex] != nullptr)
 					{
-						CopyLayersTexture(HeightmapRTList[MipRTIndex], HeightmapRenderData.OriginalHeightmap, HeightmapRenderData.HeightmapsCPUReadBack, HeightmapRenderData.TopLeftSectionBase, CurrentMip, CurrentMip);
+						CopyLayersTexture(HeightmapRTList[MipRTIndex], HeightmapRenderData.OriginalHeightmap, Landscape->HeightmapsCPUReadBackResources[HeightmapRenderData.HeightmapsCPUReadBackResourceIndex].Get(), HeightmapRenderData.TopLeftSectionBase, CurrentMip, CurrentMip);
 						++CurrentMip;
 					}
 				}
@@ -2761,6 +2768,7 @@ void ALandscape::RegenerateLayersHeightmaps()
 		ResolveLayersHeightmapTexture(AllLandscapes);
 	}
 
+	uint32 BoundsAndCollisionFlag = 0;
 	if ((LayersContentUpdateFlags & ELandscapeLayersContentUpdateFlag::Heightmap_BoundsAndCollision) != 0)
 	{
 		for (ULandscapeComponent* Component : AllLandscapeComponents)
@@ -2768,11 +2776,28 @@ void ALandscape::RegenerateLayersHeightmaps()
 			Component->UpdateCachedBounds();
 			Component->UpdateComponentToWorld();
 
-			Component->UpdateCollisionData(false);
+            Component->UpdateCollisionData();
+			
+			// Recreate collision for modified components to update the physical materials
+			ULandscapeHeightfieldCollisionComponent* CollisionComponent = Component->CollisionComponent.Get();
+			if (CollisionComponent)
+			{
+				// Only reflect outside of transaction
+				if (!GUndo)
+				{
+					FNavigationSystem::UpdateComponentData(*CollisionComponent);
+					CollisionComponent->SnapFoliageInstances();
+				}
+				else
+				{
+					BoundsAndCollisionFlag = ELandscapeLayersContentUpdateFlag::Heightmap_BoundsAndCollision;
+				}
+			}		
 		}
 	}
 
 	LayersContentUpdateFlags &= ~ELandscapeLayersContentUpdateFlag::Heightmap_All;
+	LayersContentUpdateFlags |= BoundsAndCollisionFlag; // Delay update if we were transacting
 
 	// If doing rendering debug, keep doing the render only
 	if (CVarOutputLayersDebugDrawCallName.GetValueOnAnyThread() == 1)
@@ -2792,13 +2817,8 @@ void ALandscape::ResolveLayersHeightmapTexture(const TArray<ALandscapeProxy*>& I
 		for (auto& ItPair : Landscape->RenderDataPerHeightmap)
 		{
 			FRenderDataPerHeightmap& HeightmapRenderData = ItPair.Value;
-
-			if (HeightmapRenderData.HeightmapsCPUReadBack == nullptr)
-			{
-				continue;
-			}
-
-			ResolveLayersTexture(HeightmapRenderData.HeightmapsCPUReadBack, HeightmapRenderData.OriginalHeightmap);
+						
+			ResolveLayersTexture(Landscape->HeightmapsCPUReadBackResources[HeightmapRenderData.HeightmapsCPUReadBackResourceIndex].Get(), HeightmapRenderData.OriginalHeightmap);
 		}
 	}		
 }
@@ -3042,7 +3062,7 @@ void ALandscape::ReallocateLayersWeightmaps(const TArray<ALandscapeProxy*>& InAl
 			if (CurrentWeightmapTextureUsage == nullptr)
 			{
 				ULandscapeWeightmapUsage* ComponentWeightmapUsage = ComponentWeightmapTextureUsage[i];
-				ULandscapeWeightmapUsage* Usage = NewObject<ULandscapeWeightmapUsage>(Component->GetLandscapeProxy());
+				ULandscapeWeightmapUsage* Usage = Component->GetLandscapeProxy()->CreateWeightmapUsage();
 
 				for (int32 j = 0; j < 4; ++j)
 				{
@@ -4109,8 +4129,8 @@ void ALandscape::ResolveLayersWeightmapTexture(const TArray<ALandscapeProxy*>& I
 
 void ALandscape::RequestLayersContentUpdate(uint32 InDataFlags, bool InUpdateAllMaterials)
 {
-	LayersContentUpdateFlags = InDataFlags;
-	LayersUpdateAllMaterials = InUpdateAllMaterials;
+	LayersContentUpdateFlags |= InDataFlags;
+	LayersUpdateAllMaterials |= InUpdateAllMaterials;
 }
 
 void ALandscape::RegenerateLayersContent()
@@ -4132,6 +4152,13 @@ void ALandscape::TickLayers(float DeltaTime, ELevelTick TickType, FActorTickFunc
 	UWorld* World = GetWorld();
 	if (World && !World->IsPlayInEditor())
 	{
+#if WITH_EDITOR
+		if (CVarLandscapeSimulatePhysics.GetValueOnAnyThread() == 1)
+		{
+			World->bShouldSimulatePhysics = true;
+		}
+#endif
+
 		if (GetMutableDefault<UEditorExperimentalSettings>()->bLandscapeLayerSystem)
 		{
 			if (PreviousExperimentalLandscapeLayers != GetMutableDefault<UEditorExperimentalSettings>()->bLandscapeLayerSystem)
@@ -4166,11 +4193,7 @@ void ALandscape::TickLayers(float DeltaTime, ELevelTick TickType, FActorTickFunc
 				for (auto& ItPair : RenderDataPerHeightmap)
 				{
 					FRenderDataPerHeightmap& HeightmapRenderData = ItPair.Value;
-
-					if (HeightmapRenderData.HeightmapsCPUReadBack != nullptr)
-					{
-						BeginReleaseResource(HeightmapRenderData.HeightmapsCPUReadBack);
-					}
+					BeginReleaseResource(HeightmapsCPUReadBackResources[HeightmapRenderData.HeightmapsCPUReadBackResourceIndex].Get());
 				}
 
 				for (auto& ItPair : WeightmapCPUReadBackTextures)
@@ -4204,14 +4227,9 @@ void ALandscape::TickLayers(float DeltaTime, ELevelTick TickType, FActorTickFunc
 				}
 
 				FlushRenderingCommands();
-
-				for (auto& ItPair : RenderDataPerHeightmap)
-				{
-					FRenderDataPerHeightmap& HeightmapRenderData = ItPair.Value;
-
-					delete HeightmapRenderData.HeightmapsCPUReadBack;
-					HeightmapRenderData.HeightmapsCPUReadBack = nullptr;
-				}
+				
+				RenderDataPerHeightmap.Empty();
+				HeightmapsCPUReadBackResources.Empty();
 
 				for (auto& ItPair : WeightmapCPUReadBackTextures)
 				{
@@ -4243,16 +4261,11 @@ void ALandscapeProxy::BeginDestroy()
 #if WITH_EDITORONLY_DATA
 	if (GetMutableDefault<UEditorExperimentalSettings>()->bLandscapeLayerSystem)
 	{
-		for (auto& ItPair : RenderDataPerHeightmap)
+		for (const TUniquePtr<FLandscapeLayersTexture2DCPUReadBackResource>& Ptr : HeightmapsCPUReadBackResources)
 		{
-			FRenderDataPerHeightmap& HeightmapRenderData = ItPair.Value;
-
-			if (HeightmapRenderData.HeightmapsCPUReadBack != nullptr)
-			{
-				BeginReleaseResource(HeightmapRenderData.HeightmapsCPUReadBack);
-			}
+			BeginReleaseResource(Ptr.Get());
 		}
-
+				
 		for (auto& ItPair : WeightmapCPUReadBackTextures)
 		{
 			FLandscapeLayersTexture2DCPUReadBackResource* WeightmapCPUReadBack = ItPair.Value;
@@ -4294,13 +4307,7 @@ void ALandscapeProxy::FinishDestroy()
 	{
 		check(ReleaseResourceFence.IsFenceComplete());
 
-		for (auto& ItPair : RenderDataPerHeightmap)
-		{
-			FRenderDataPerHeightmap& HeightmapRenderData = ItPair.Value;
-
-			delete HeightmapRenderData.HeightmapsCPUReadBack;
-			HeightmapRenderData.HeightmapsCPUReadBack = nullptr;
-		}
+		HeightmapsCPUReadBackResources.Empty();
 
 		for (auto& ItPair : WeightmapCPUReadBackTextures)
 		{
@@ -4475,6 +4482,7 @@ void ALandscape::DeleteLayer(int32 InLayerIndex)
 						ULandscapeWeightmapUsage** Usage = Proxy->WeightmapUsageMap.Find(WeightmapTexture);
 						if (Usage != nullptr && (*Usage) != nullptr)
 						{
+							(*Usage)->Modify();
 							(*Usage)->ChannelUsage[Allocation.WeightmapTextureChannel] = nullptr;
 							if ((*Usage)->FreeChannelCount() == 4)
 							{
diff --git a/Engine/Source/Runtime/Landscape/Private/LandscapeGrass.cpp b/Engine/Source/Runtime/Landscape/Private/LandscapeGrass.cpp
--- a/Engine/Source/Runtime/Landscape/Private/LandscapeGrass.cpp
+++ b/Engine/Source/Runtime/Landscape/Private/LandscapeGrass.cpp
@@ -764,7 +764,8 @@ public:
 
 			if (Proxy->bBakeMaterialPositionOffsetIntoCollision)
 			{
-				Component->UpdateCollisionData(true);
+				Component->DestroyCollisionData();
+				Component->UpdateCollisionData();
 			}
 		}
 	}
diff --git a/Engine/Source/Runtime/Landscape/Public/LandscapeEdit.h b/Engine/Source/Runtime/Landscape/Public/LandscapeEdit.h
--- a/Engine/Source/Runtime/Landscape/Public/LandscapeEdit.h
+++ b/Engine/Source/Runtime/Landscape/Public/LandscapeEdit.h
@@ -321,13 +321,18 @@ struct FHeightmapAccessor
 
 			// Notify foliage to move any attached instances
 			bool bUpdateFoliage = false;
-			for (ULandscapeComponent* Component : Components)
+
+            // Landscape Layers are updates are delayed and done in  ALandscape::TickLayers
+			if (!GetMutableDefault<UEditorExperimentalSettings>()->bLandscapeLayerSystem)
 			{
-				ULandscapeHeightfieldCollisionComponent* CollisionComponent = Component->CollisionComponent.Get();
-				if (CollisionComponent && AInstancedFoliageActor::HasFoliageAttached(CollisionComponent))
+				for (ULandscapeComponent* Component : Components)
 				{
-					bUpdateFoliage = true;
-					break;
+					ULandscapeHeightfieldCollisionComponent* CollisionComponent = Component->CollisionComponent.Get();
+					if (CollisionComponent && AInstancedFoliageActor::HasFoliageAttached(CollisionComponent))
+					{
+						bUpdateFoliage = true;
+						break;
+					}
 				}
 			}
 
@@ -377,18 +382,22 @@ struct FHeightmapAccessor
 		delete LandscapeEdit;
 		LandscapeEdit = NULL;
 
-		// Update the bounds and navmesh for the components we edited
-		for (TSet<ULandscapeComponent*>::TConstIterator It(ChangedComponents); It; ++It)
+		// Landscape Layers are updates are delayed and done in  ALandscape::TickLayers
+		if (!GetMutableDefault<UEditorExperimentalSettings>()->bLandscapeLayerSystem)
 		{
-			(*It)->UpdateCachedBounds();
-			(*It)->UpdateComponentToWorld();
-
-			// Recreate collision for modified components to update the physical materials
-			ULandscapeHeightfieldCollisionComponent* CollisionComponent = (*It)->CollisionComponent.Get();
-			if (CollisionComponent)
+			// Update the bounds and navmesh for the components we edited
+			for (TSet<ULandscapeComponent*>::TConstIterator It(ChangedComponents); It; ++It)
 			{
-				CollisionComponent->RecreateCollision();
-				FNavigationSystem::UpdateComponentData(*CollisionComponent);
+				(*It)->UpdateCachedBounds();
+				(*It)->UpdateComponentToWorld();
+
+				// Recreate collision for modified components to update the physical materials
+				ULandscapeHeightfieldCollisionComponent* CollisionComponent = (*It)->CollisionComponent.Get();
+				if (CollisionComponent)
+				{
+					CollisionComponent->RecreateCollision();
+					FNavigationSystem::UpdateComponentData(*CollisionComponent);
+				}
 			}
 		}
 	}
