Fix Editor Mac and Editor Linux

diff --git a/Engine/Source/Runtime/Landscape/Public/LandscapeEdit.h b/Engine/Source/Runtime/Landscape/Public/LandscapeEdit.h
--- a/Engine/Source/Runtime/Landscape/Public/LandscapeEdit.h
+++ b/Engine/Source/Runtime/Landscape/Public/LandscapeEdit.h
@@ -19,7 +19,10 @@ LandscapeEdit.h: Classes for the editor to access to Landscape data
 #include "LandscapeLayerInfoObject.h"
 
 #if WITH_EDITOR
-	#include "Containers/ArrayView.h"
+
+#include "Containers/ArrayView.h"
+#include "Settings/EditorExperimentalSettings.h"
+
 #endif
 
 class ULandscapeComponent;
