- Fixed physic material update when changing Physic Material on a LayerInfo.
- Fixed missing physic material update when Clearing, Filling and Replacing a Layer.

diff --git a/Engine/Source/Runtime/Landscape/Private/LandscapeEdit.cpp b/Engine/Source/Runtime/Landscape/Private/LandscapeEdit.cpp
--- a/Engine/Source/Runtime/Landscape/Private/LandscapeEdit.cpp
+++ b/Engine/Source/Runtime/Landscape/Private/LandscapeEdit.cpp
@@ -1270,11 +1270,10 @@ void ULandscapeComponent::UpdateCollisionLayerData(const FColor* const* const We
 			}
 			else
 			{
-				ComponentX1 = FMath::Min(ComponentX1, ComponentSizeQuads);
-				ComponentY1 = FMath::Min(ComponentY1, ComponentSizeQuads);
-				ComponentX2 = FMath::Max(ComponentX2, 0);
-				ComponentY2 = FMath::Max(ComponentY2, 0);
-
+				ComponentX1 = FMath::Clamp(ComponentX1, 0, ComponentSizeQuads);
+				ComponentY1 = FMath::Clamp(ComponentY1, 0, ComponentSizeQuads);
+				ComponentX2 = FMath::Clamp(ComponentX2, 0, ComponentSizeQuads);
+				ComponentY2 = FMath::Clamp(ComponentY2, 0, ComponentSizeQuads);
 				DominantLayerData = (uint8*)CollisionComp->DominantLayerData.Lock(LOCK_READ_WRITE);
 			}
 
diff --git a/Engine/Source/Runtime/Landscape/Private/LandscapeEditInterface.cpp b/Engine/Source/Runtime/Landscape/Private/LandscapeEditInterface.cpp
--- a/Engine/Source/Runtime/Landscape/Private/LandscapeEditInterface.cpp
+++ b/Engine/Source/Runtime/Landscape/Private/LandscapeEditInterface.cpp
@@ -1759,6 +1759,12 @@ void ULandscapeComponent::DeleteLayer(ULandscapeLayerInfoObject* LayerInfo, FLan
 	Component->UpdateCollisionLayerData(
 		CollisionWeightmapMipData.GetData(),
 		Component->SimpleCollisionMipLevel > Component->CollisionMipLevel ? SimpleCollisionWeightmapMipData.GetData() : nullptr);
+
+	ULandscapeHeightfieldCollisionComponent* CollisionComponent = Component->CollisionComponent.Get();
+	if (CollisionComponent)
+	{
+		CollisionComponent->RecreateCollision();
+	}
 }
 
 void FLandscapeEditDataInterface::DeleteLayer(ULandscapeLayerInfoObject* LayerInfo)
@@ -1952,6 +1958,12 @@ void ULandscapeComponent::FillLayer(ULandscapeLayerInfoObject* LayerInfo, FLands
 	Component->UpdateCollisionLayerData(
 		CollisionWeightmapMipData.GetData(),
 		Component->SimpleCollisionMipLevel > Component->CollisionMipLevel ? SimpleCollisionWeightmapMipData.GetData() : nullptr);
+
+	ULandscapeHeightfieldCollisionComponent* CollisionComponent = Component->CollisionComponent.Get();
+	if (CollisionComponent)
+	{
+		CollisionComponent->RecreateCollision();
+	}
 }
 
 void FLandscapeEditDataInterface::FillLayer(ULandscapeLayerInfoObject* LayerInfo)
@@ -2235,6 +2247,12 @@ void FLandscapeEditDataInterface::ReplaceLayer(ULandscapeLayerInfoObject* FromLa
 		Component->UpdateCollisionLayerData(
 			CollisionWeightmapMipData.GetData(),
 			Component->SimpleCollisionMipLevel > Component->CollisionMipLevel ? SimpleCollisionWeightmapMipData.GetData() : nullptr);
+
+		ULandscapeHeightfieldCollisionComponent* CollisionComponent = Component->CollisionComponent.Get();
+		if (CollisionComponent)
+		{
+			CollisionComponent->RecreateCollision();
+		}
 	}
 }
 
