#rnx

diff --git a/Engine/Source/Runtime/Landscape/Private/LandscapeEdit.cpp b/Engine/Source/Runtime/Landscape/Private/LandscapeEdit.cpp
--- a/Engine/Source/Runtime/Landscape/Private/LandscapeEdit.cpp
+++ b/Engine/Source/Runtime/Landscape/Private/LandscapeEdit.cpp
@@ -574,7 +574,13 @@ void ULandscapeComponent::PostEditUndo()
 	TSet<ULandscapeComponent*> Components;
 	Components.Add(this);
 	GetLandscapeProxy()->FlushGrassComponents(&Components);
-	GetLandscapeActor()->RequestLayersContentUpdate(ELandscapeLayersContentUpdateFlag::All, true);
+
+	ALandscape* LandscapeActor = GetLandscapeActor();
+	// Might be a ALandscapeStreamingProxy
+	if (LandscapeActor)
+	{
+		LandscapeActor->RequestLayersContentUpdate(ELandscapeLayersContentUpdateFlag::All, true);
+	}
 }
 
 void ALandscapeProxy::FixupWeightmaps()
