#UE4 Fix shadowed variable

diff --git a/Engine/Source/Runtime/Landscape/Private/LandscapeEditInterface.cpp b/Engine/Source/Runtime/Landscape/Private/LandscapeEditInterface.cpp
--- a/Engine/Source/Runtime/Landscape/Private/LandscapeEditInterface.cpp
+++ b/Engine/Source/Runtime/Landscape/Private/LandscapeEditInterface.cpp
@@ -1760,10 +1760,10 @@ void ULandscapeComponent::DeleteLayer(ULandscapeLayerInfoObject* LayerInfo, FLan
 		CollisionWeightmapMipData.GetData(),
 		Component->SimpleCollisionMipLevel > Component->CollisionMipLevel ? SimpleCollisionWeightmapMipData.GetData() : nullptr);
 
-	ULandscapeHeightfieldCollisionComponent* CollisionComponent = Component->CollisionComponent.Get();
-	if (CollisionComponent)
+	ULandscapeHeightfieldCollisionComponent* LocalCollisionComponent = Component->CollisionComponent.Get();
+	if (LocalCollisionComponent)
 	{
-		CollisionComponent->RecreateCollision();
+		LocalCollisionComponent->RecreateCollision();
 	}
 }
 
@@ -1959,10 +1959,10 @@ void ULandscapeComponent::FillLayer(ULandscapeLayerInfoObject* LayerInfo, FLands
 		CollisionWeightmapMipData.GetData(),
 		Component->SimpleCollisionMipLevel > Component->CollisionMipLevel ? SimpleCollisionWeightmapMipData.GetData() : nullptr);
 
-	ULandscapeHeightfieldCollisionComponent* CollisionComponent = Component->CollisionComponent.Get();
-	if (CollisionComponent)
+	ULandscapeHeightfieldCollisionComponent* LocalCollisionComponent = Component->CollisionComponent.Get();
+	if (LocalCollisionComponent)
 	{
-		CollisionComponent->RecreateCollision();
+		LocalCollisionComponent->RecreateCollision();
 	}
 }
 
