Landscape Layer System:
Split BoundsAndCollision Flag from ClientUpdate Flag
Split Public and Internal Update Flags to avoid invalid update requests

#rnx

diff --git a/Engine/Source/Runtime/Landscape/Classes/Landscape.h b/Engine/Source/Runtime/Landscape/Classes/Landscape.h
--- a/Engine/Source/Runtime/Landscape/Classes/Landscape.h
+++ b/Engine/Source/Runtime/Landscape/Classes/Landscape.h
@@ -72,24 +72,32 @@ enum EWeightmapRTType : uint8
 	WeightmapRT_Count
 };
 
-enum ELandscapeLayersContentUpdateFlag : uint32
+// Internal Update Flags that shouldn't be used alone
+enum EInternalLandscapeLayersContentUpdateFlag : uint32
 {
-	Heightmap_Setup					= 0x00000001u,
-	Heightmap_Render				= 0x00000002u,
-	Heightmap_BoundsAndCollision	= 0x00000004u,
-	Heightmap_ResolveToTexture		= 0x00000008u,
+	Internal_Heightmap_Setup = 0x00000001u,
+	Internal_Heightmap_Render = 0x00000002u,
+	Internal_Heightmap_BoundsAndCollision = 0x00000004u,
+	Internal_Heightmap_ResolveToTexture = 0x00000008u,
+	Internal_Heightmap_ClientUpdate = 0x00000010u,
+
+	Internal_Weightmap_Setup = 0x00000100u,
+	Internal_Weightmap_Render = 0x00000200u,
+	Internal_Weightmap_Collision = 0x00000400u,
+	Internal_Weightmap_ResolveToTexture = 0x00000800u,
+};
 
-	Weightmap_Setup					= 0x00000100u,
-	Weightmap_Render				= 0x00000200u,
-	Weightmap_Collision				= 0x00000400u,
-	Weightmap_ResolveToTexture		= 0x00000800u,
+enum ELandscapeLayersContentUpdateFlag : uint32
+{
+	Heightmap_Render = Internal_Heightmap_Render,
+	Weightmap_Render = Internal_Weightmap_Render,
 
 	// Combinations
-	Heightmap_All = Heightmap_Render | Heightmap_BoundsAndCollision | Heightmap_ResolveToTexture,
-	Weightmap_All = Weightmap_Render | Weightmap_Collision | Weightmap_ResolveToTexture,
+	Heightmap_All = Internal_Heightmap_Render | Internal_Heightmap_BoundsAndCollision | Internal_Heightmap_ResolveToTexture | Internal_Heightmap_ClientUpdate,
+	Weightmap_All = Internal_Weightmap_Render | Internal_Weightmap_Collision | Internal_Weightmap_ResolveToTexture,
 
 	All = Heightmap_All | Weightmap_All,
-	All_Setup = Heightmap_Setup | Weightmap_Setup,
+	All_Setup = Internal_Heightmap_Setup | Internal_Weightmap_Setup,
 	All_Render = Heightmap_Render | Weightmap_Render,
 };
 
@@ -241,7 +249,7 @@ public:
 
 	// Layers stuff
 #if WITH_EDITOR
-	LANDSCAPE_API void RequestLayersContentUpdate(uint32 InDataFlags, bool InUpdateAllMaterials = false);
+	LANDSCAPE_API void RequestLayersContentUpdate(ELandscapeLayersContentUpdateFlag InDataFlags, bool InUpdateAllMaterials = false);
 	LANDSCAPE_API void CreateLayer(FName InName = NAME_None, bool bInUpdateLayersContent = true);
 	LANDSCAPE_API bool ReorderLayer(int32 InStartingLayerIndex, int32 InDestinationLayerIndex);
 	LANDSCAPE_API bool IsLayerNameUnique(const FName& InName) const;
diff --git a/Engine/Source/Runtime/Landscape/Private/LandscapeEditLayers.cpp b/Engine/Source/Runtime/Landscape/Private/LandscapeEditLayers.cpp
--- a/Engine/Source/Runtime/Landscape/Private/LandscapeEditLayers.cpp
+++ b/Engine/Source/Runtime/Landscape/Private/LandscapeEditLayers.cpp
@@ -2763,41 +2763,49 @@ void ALandscape::RegenerateLayersHeightmaps()
 		}
 	}
 
-	if ((LayersContentUpdateFlags & ELandscapeLayersContentUpdateFlag::Heightmap_ResolveToTexture) != 0)
+	if ((LayersContentUpdateFlags & EInternalLandscapeLayersContentUpdateFlag::Internal_Heightmap_ResolveToTexture) != 0)
 	{
 		ResolveLayersHeightmapTexture(AllLandscapes);
 	}
 
-	uint32 BoundsAndCollisionFlag = 0;
-	if ((LayersContentUpdateFlags & ELandscapeLayersContentUpdateFlag::Heightmap_BoundsAndCollision) != 0)
+	bool bUpdateBoundsAndCollsion = (LayersContentUpdateFlags & EInternalLandscapeLayersContentUpdateFlag::Internal_Heightmap_BoundsAndCollision) != 0;
+	bool bUpdateClients = (LayersContentUpdateFlags & EInternalLandscapeLayersContentUpdateFlag::Internal_Heightmap_ClientUpdate) != 0;
+	int32 PostponeUpdateClientFlag = 0;
+	if (bUpdateBoundsAndCollsion || bUpdateClients)
 	{
 		for (ULandscapeComponent* Component : AllLandscapeComponents)
 		{
-			Component->UpdateCachedBounds();
-			Component->UpdateComponentToWorld();
+			if (bUpdateBoundsAndCollsion)
+			{
+				Component->UpdateCachedBounds();
+				Component->UpdateComponentToWorld();
 
-            Component->UpdateCollisionData();
-			
-			// Recreate collision for modified components to update the physical materials
-			ULandscapeHeightfieldCollisionComponent* CollisionComponent = Component->CollisionComponent.Get();
-			if (CollisionComponent)
+				Component->UpdateCollisionData();
+			}
+
+			if (bUpdateClients)
 			{
-				// Only reflect outside of transaction
-				if (!GUndo)
+				// Post-pone client update to next tick while in transaction
+				if (GUndo)
 				{
-					FNavigationSystem::UpdateComponentData(*CollisionComponent);
-					CollisionComponent->SnapFoliageInstances();
+					PostponeUpdateClientFlag |= EInternalLandscapeLayersContentUpdateFlag::Internal_Heightmap_ClientUpdate;
 				}
 				else
 				{
-					BoundsAndCollisionFlag = ELandscapeLayersContentUpdateFlag::Heightmap_BoundsAndCollision;
+					// Recreate collision for modified components to update the physical materials
+					ULandscapeHeightfieldCollisionComponent* CollisionComponent = Component->CollisionComponent.Get();
+					if (CollisionComponent)
+					{
+						FNavigationSystem::UpdateComponentData(*CollisionComponent);
+						CollisionComponent->SnapFoliageInstances();
+					}
 				}
-			}		
+			}
 		}
 	}
 
 	LayersContentUpdateFlags &= ~ELandscapeLayersContentUpdateFlag::Heightmap_All;
-	LayersContentUpdateFlags |= BoundsAndCollisionFlag; // Delay update if we were transacting
+	LayersContentUpdateFlags |= PostponeUpdateClientFlag; // if non-zero postpone client update (transacting)
 
 	// If doing rendering debug, keep doing the render only
 	if (CVarOutputLayersDebugDrawCallName.GetValueOnAnyThread() == 1)
@@ -3904,12 +3912,12 @@ void ALandscape::RegenerateLayersWeightmaps()
 		LayersUpdateAllMaterials = false;
 	}
 
-	if ((LayersContentUpdateFlags & ELandscapeLayersContentUpdateFlag::Weightmap_ResolveToTexture))
+	if ((LayersContentUpdateFlags & EInternalLandscapeLayersContentUpdateFlag::Internal_Weightmap_ResolveToTexture))
 	{
 		ResolveLayersWeightmapTexture(AllLandscapes);
 	}
 
-	if ((LayersContentUpdateFlags & ELandscapeLayersContentUpdateFlag::Weightmap_Collision) != 0)
+	if ((LayersContentUpdateFlags & EInternalLandscapeLayersContentUpdateFlag::Internal_Weightmap_Collision) != 0)
 	{
 		for (ULandscapeComponent* Component : AllLandscapeComponents)
 		{
@@ -4127,7 +4135,7 @@ void ALandscape::ResolveLayersWeightmapTexture(const TArray<ALandscapeProxy*>& I
 	}
 }
 
-void ALandscape::RequestLayersContentUpdate(uint32 InDataFlags, bool InUpdateAllMaterials)
+void ALandscape::RequestLayersContentUpdate(ELandscapeLayersContentUpdateFlag InDataFlags, bool InUpdateAllMaterials)
 {
 	LayersContentUpdateFlags |= InDataFlags;
 	LayersUpdateAllMaterials |= InUpdateAllMaterials;
@@ -4135,7 +4143,7 @@ void ALandscape::RequestLayersContentUpdate(uint32 InDataFlags, bool InUpdateAll
 
 void ALandscape::RegenerateLayersContent()
 {
-	if ((LayersContentUpdateFlags & Heightmap_Setup) != 0 || (LayersContentUpdateFlags & Weightmap_Setup) != 0)
+	if ((LayersContentUpdateFlags & Internal_Heightmap_Setup) != 0 || (LayersContentUpdateFlags & Internal_Weightmap_Setup) != 0)
 	{
 		SetupLayers();
 		LayersContentUpdateFlags &= ~(ELandscapeLayersContentUpdateFlag::All_Setup);
@@ -4511,7 +4519,8 @@ void ALandscape::DeleteLayer(int32 InLayerIndex)
 	LandscapeLayers.RemoveAt(InLayerIndex);
 
 	// Request Update
-	RequestLayersContentUpdate(ELandscapeLayersContentUpdateFlag::All_Setup | All, true);
+	RequestLayersContentUpdate(ELandscapeLayersContentUpdateFlag::All_Setup, false);
+	RequestLayersContentUpdate(ELandscapeLayersContentUpdateFlag::All, true);
 }
 
 void ALandscape::ClearLayer(int32 InLayerIndex, bool bInUpdateCollision)
@@ -4844,7 +4853,7 @@ void ALandscape::SetLayerSubstractiveBlendStatus(int32 InLayerIndex, bool InStat
 		*AllocationBlend = InStatus;
 	}
 
-	RequestLayersContentUpdate(true);
+	RequestLayersContentUpdate(ELandscapeLayersContentUpdateFlag::All, true);
 }
 
 void ALandscape::AddBrushToLayer(int32 InLayerIndex, int32 InTargetType, ALandscapeBlueprintCustomBrush* InBrush)
