- Fixed crash adding an array element to Landscape Materials Override

diff --git a/Engine/Source/Runtime/Landscape/Private/Landscape.cpp b/Engine/Source/Runtime/Landscape/Private/Landscape.cpp
--- a/Engine/Source/Runtime/Landscape/Private/Landscape.cpp
+++ b/Engine/Source/Runtime/Landscape/Private/Landscape.cpp
@@ -3083,10 +3083,13 @@ void ULandscapeComponent::SerializeStateHashes(FArchive& Ar)
 
 	for (FLandscapeProxyMaterialOverride& MaterialOverride : Proxy->LandscapeMaterialsOverride)
 	{
-		UMaterialInterface::TMicRecursionGuard RecursionGuard;
-		FGuid LocalStateId = MaterialOverride.Material->GetMaterial_Concurrent(RecursionGuard)->StateId;
-		Ar << LocalStateId;
-		Ar << MaterialOverride.LODIndex;
+		if (MaterialOverride.Material != nullptr)
+		{
+			UMaterialInterface::TMicRecursionGuard RecursionGuard;
+			FGuid LocalStateId = MaterialOverride.Material->GetMaterial_Concurrent(RecursionGuard)->StateId;
+			Ar << LocalStateId;
+			Ar << MaterialOverride.LODIndex;
+		}
 	}
 }
 
