- Landscape Layer System: Undo/Redo Fixes (mainly Texture Usage)

#rnx

diff --git a/Engine/Source/Runtime/Landscape/Classes/Landscape.h b/Engine/Source/Runtime/Landscape/Classes/Landscape.h
--- a/Engine/Source/Runtime/Landscape/Classes/Landscape.h
+++ b/Engine/Source/Runtime/Landscape/Classes/Landscape.h
@@ -237,6 +237,7 @@ public:
 	virtual void PreSave(const class ITargetPlatform* TargetPlatform) override;
 	virtual void PostEditChangeProperty(FPropertyChangedEvent& PropertyChangedEvent) override;
 	virtual void PostEditMove(bool bFinished) override;
+	virtual void PostEditUndo() override;
 	virtual bool ShouldImport(FString* ActorPropString, bool IsMovingLevel) override;
 	virtual void PostEditImport() override;
 	virtual void PostDuplicate(bool bDuplicateForPIE) override;
diff --git a/Engine/Source/Runtime/Landscape/Classes/LandscapeProxy.h b/Engine/Source/Runtime/Landscape/Classes/LandscapeProxy.h
--- a/Engine/Source/Runtime/Landscape/Classes/LandscapeProxy.h
+++ b/Engine/Source/Runtime/Landscape/Classes/LandscapeProxy.h
@@ -775,6 +775,7 @@ public:
 	virtual void PostEditMove(bool bFinished) override;
 	virtual bool ShouldImport(FString* ActorPropString, bool IsMovingLevel) override;
 	virtual bool ShouldExport() override;
+	virtual bool Modify(bool bAlwaysMarkDirty = true) override;
 	//~ End AActor Interface
 #endif	//WITH_EDITOR
 
diff --git a/Engine/Source/Runtime/Landscape/Private/LandscapeEdit.cpp b/Engine/Source/Runtime/Landscape/Private/LandscapeEdit.cpp
--- a/Engine/Source/Runtime/Landscape/Private/LandscapeEdit.cpp
+++ b/Engine/Source/Runtime/Landscape/Private/LandscapeEdit.cpp
@@ -3582,6 +3582,13 @@ void ALandscape::PostEditMove(bool bFinished)
 	Super::PostEditMove(bFinished);
 }
 
+void ALandscape::PostEditUndo()
+{
+	Super::PostEditUndo();
+
+	RequestLayersContentUpdate(ELandscapeLayersContentUpdateFlag::All, true);
+}
+
 bool ALandscape::ShouldImport(FString* ActorPropString, bool IsMovingLevel)
 {
 	return GetWorld() != nullptr && !GetWorld()->IsGameWorld();
@@ -4763,6 +4770,10 @@ void ULandscapeComponent::ReallocateWeightmaps(FLandscapeEditDataInterface* Data
 
 			// Store it in the usage map
 			CurrentWeightmapUsage = Proxy->WeightmapUsageMap.Add(CurrentWeightmapTexture, GetLandscapeProxy()->CreateWeightmapUsage());
+			if (InSaveToTransactionBuffer)
+			{
+				CurrentWeightmapUsage->Modify();
+			}
 			CurrentWeightmapUsage->LayerGuid = InCanUseCurrentEditingWeightmap ? CurrentLayerGuid : FGuid();
 			// UE_LOG(LogLandscape, Log, TEXT("Making a new texture %s"), *CurrentWeightmapTexture->GetName());
 		}
@@ -5108,6 +5119,18 @@ void ULandscapeComponent::InitWeightmapData(TArray<ULandscapeLayerInfoObject*>&
 #define MAX_LANDSCAPE_PROP_TEXT_LENGTH			1024*1024*16
 
 
+bool ALandscapeProxy::Modify(bool bAlwaysMarkDirty)
+{
+	for (auto& Pair : WeightmapUsageMap)
+	{
+		if (Pair.Value)
+		{
+			Pair.Value->Modify(bAlwaysMarkDirty);
+		}
+	}
+	return Super::Modify(bAlwaysMarkDirty);
+}
+
 bool ALandscapeProxy::ShouldExport()
 {
 	if (!bIsMovingToLevel && LandscapeComponents.Num() > MAX_LANDSCAPE_EXPORT_COMPONENTS_NUM)
diff --git a/Engine/Source/Runtime/Landscape/Private/LandscapeEditLayers.cpp b/Engine/Source/Runtime/Landscape/Private/LandscapeEditLayers.cpp
--- a/Engine/Source/Runtime/Landscape/Private/LandscapeEditLayers.cpp
+++ b/Engine/Source/Runtime/Landscape/Private/LandscapeEditLayers.cpp
@@ -4494,6 +4494,7 @@ void ALandscape::DeleteLayer(int32 InLayerIndex)
 							(*Usage)->ChannelUsage[Allocation.WeightmapTextureChannel] = nullptr;
 							if ((*Usage)->FreeChannelCount() == 4)
 							{
+								Proxy->Modify();
 								Proxy->WeightmapUsageMap.Remove(WeightmapTexture);
 							}
 						}
@@ -4506,6 +4507,7 @@ void ALandscape::DeleteLayer(int32 InLayerIndex)
 	// Remove associated layer data of each landscape proxy
 	LandscapeInfo->ForAllLandscapeProxies([LayerGuid](ALandscapeProxy* Proxy)
 	{
+		Proxy->Modify();
 		Proxy->LandscapeLayersData.Remove(LayerGuid);
 	});
 
@@ -4567,8 +4569,11 @@ void ALandscape::ClearLayer(const FGuid& InLayerGuid, bool bInUpdateCollision)
 	FLandscapeEditDataInterface LandscapeEdit(LandscapeInfo);
 	LandscapeInfo->ForAllLandscapeProxies([&](ALandscapeProxy* Proxy)
 	{
+		Proxy->Modify();
 		for (ULandscapeComponent* Component : Proxy->LandscapeComponents)
 		{
+			Component->Modify();
+
 			int32 MinX = MAX_int32;
 			int32 MinY = MAX_int32;
 			int32 MaxX = MIN_int32;
@@ -4774,6 +4779,7 @@ void ALandscape::CreateLayer(FName InName, bool bInUpdateLayersContent)
 	// Create associated layer data in each landscape proxy
 	LandscapeInfo->ForAllLandscapeProxies([NewLayer](ALandscapeProxy* Proxy)
 	{
+		Proxy->Modify();
 		Proxy->LandscapeLayersData.Add(NewLayer.Guid, FLandscapeLayerData());
 	});
 
@@ -4781,7 +4787,7 @@ void ALandscape::CreateLayer(FName InName, bool bInUpdateLayersContent)
 	{
 		// Request Update
 		RequestLayersContentUpdate(ELandscapeLayersContentUpdateFlag::All_Setup);
-		RegenerateLayersContent();
+		RequestLayersContentUpdate(ELandscapeLayersContentUpdateFlag::All, true);
 	}
 }
 
