- Landscape Layer System: Undo/Redo fix for Reserve Layer for Splines

diff --git a/Engine/Source/Runtime/Landscape/Private/LandscapeEditLayers.cpp b/Engine/Source/Runtime/Landscape/Private/LandscapeEditLayers.cpp
--- a/Engine/Source/Runtime/Landscape/Private/LandscapeEditLayers.cpp
+++ b/Engine/Source/Runtime/Landscape/Private/LandscapeEditLayers.cpp
@@ -4654,6 +4654,7 @@ void ALandscape::ShowAllLayers()
 
 void ALandscape::SetLandscapeSplinesReservedLayer(int32 InLayerIndex)
 {
+	Modify();
 	FLandscapeLayer* NewLayer = GetLayer(InLayerIndex);
 	FLandscapeLayer* PreviousLayer = GetLandscapeSplinesReservedLayer();
 	if (NewLayer != PreviousLayer)
