- Landscape Layer System: Replace Layer Info now propagated to all layers

#rnx

diff --git a/Engine/Source/Runtime/Landscape/Classes/Landscape.h b/Engine/Source/Runtime/Landscape/Classes/Landscape.h
--- a/Engine/Source/Runtime/Landscape/Classes/Landscape.h
+++ b/Engine/Source/Runtime/Landscape/Classes/Landscape.h
@@ -261,6 +261,7 @@ public:
 	LANDSCAPE_API struct FLandscapeLayer* GetLayer(int32 InLayerIndex);
 	LANDSCAPE_API const struct FLandscapeLayer* GetLayer(int32 InLayerIndex) const;
 	LANDSCAPE_API const struct FLandscapeLayer* GetLayer(const FGuid& InLayerGuid) const;
+	LANDSCAPE_API void ForEachLayer(TFunctionRef<void(const struct FLandscapeLayer&)> Fn);
 	LANDSCAPE_API void ClearLayer(int32 InLayerIndex, bool bInUpdateCollision = true);
 	LANDSCAPE_API void ClearLayer(const FGuid& InLayerGuid, bool bInUpdateCollision = true);
 	LANDSCAPE_API void DeleteLayer(int32 InLayerIndex);
diff --git a/Engine/Source/Runtime/Landscape/Private/LandscapeEditInterface.cpp b/Engine/Source/Runtime/Landscape/Private/LandscapeEditInterface.cpp
--- a/Engine/Source/Runtime/Landscape/Private/LandscapeEditInterface.cpp
+++ b/Engine/Source/Runtime/Landscape/Private/LandscapeEditInterface.cpp
@@ -2223,37 +2223,61 @@ void ULandscapeComponent::ReplaceLayer(ULandscapeLayerInfoObject* FromLayerInfo,
 void FLandscapeEditDataInterface::ReplaceLayer(ULandscapeLayerInfoObject* FromLayerInfo, ULandscapeLayerInfoObject* ToLayerInfo)
 {
 	if (!LandscapeInfo) return;
-	for( auto It = LandscapeInfo->XYtoComponentMap.CreateIterator(); It; ++It )
+
+	auto DoReplace = [&]()
 	{
-		ULandscapeComponent* Component = It.Value();
-		Component->ReplaceLayer(FromLayerInfo, ToLayerInfo, *this);
+		for (auto It = LandscapeInfo->XYtoComponentMap.CreateIterator(); It; ++It)
+		{
+			ULandscapeComponent* Component = It.Value();
+			Component->ReplaceLayer(FromLayerInfo, ToLayerInfo, *this);
 
-		TArray<UTexture2D*>& ComponentWeightmapTextures = Component->GetWeightmapTextures(true);
+			TArray<UTexture2D*>& ComponentWeightmapTextures = Component->GetWeightmapTextures(true);
 
-		// Update dominant layer info stored in collision component
-		TArray<FColor*> CollisionWeightmapMipData;
-		for( int32 WeightmapIdx=0;WeightmapIdx < ComponentWeightmapTextures.Num();WeightmapIdx++ )
-		{
-			CollisionWeightmapMipData.Add((FColor*)GetTextureDataInfo(ComponentWeightmapTextures[WeightmapIdx])->GetMipData(Component->CollisionMipLevel));
-		}
-		TArray<FColor*> SimpleCollisionWeightmapMipData;
-		if (Component->SimpleCollisionMipLevel > Component->CollisionMipLevel)
-		{
+			// Update dominant layer info stored in collision component
+			TArray<FColor*> CollisionWeightmapMipData;
 			for (int32 WeightmapIdx = 0; WeightmapIdx < ComponentWeightmapTextures.Num(); WeightmapIdx++)
 			{
-				SimpleCollisionWeightmapMipData.Add((FColor*)GetTextureDataInfo(ComponentWeightmapTextures[WeightmapIdx])->GetMipData(Component->SimpleCollisionMipLevel));
+				CollisionWeightmapMipData.Add((FColor*)GetTextureDataInfo(ComponentWeightmapTextures[WeightmapIdx])->GetMipData(Component->CollisionMipLevel));
+			}
+			TArray<FColor*> SimpleCollisionWeightmapMipData;
+			if (Component->SimpleCollisionMipLevel > Component->CollisionMipLevel)
+			{
+				for (int32 WeightmapIdx = 0; WeightmapIdx < ComponentWeightmapTextures.Num(); WeightmapIdx++)
+				{
+					SimpleCollisionWeightmapMipData.Add((FColor*)GetTextureDataInfo(ComponentWeightmapTextures[WeightmapIdx])->GetMipData(Component->SimpleCollisionMipLevel));
+				}
+			}
+			Component->UpdateCollisionLayerData(
+				CollisionWeightmapMipData.GetData(),
+				Component->SimpleCollisionMipLevel > Component->CollisionMipLevel ? SimpleCollisionWeightmapMipData.GetData() : nullptr);
+
+			if (!GetMutableDefault<UEditorExperimentalSettings>()->bLandscapeLayerSystem)
+			{
+				ULandscapeHeightfieldCollisionComponent* CollisionComponent = Component->CollisionComponent.Get();
+				if (CollisionComponent)
+				{
+					CollisionComponent->RecreateCollision();
+				}
 			}
 		}
-		Component->UpdateCollisionLayerData(
-			CollisionWeightmapMipData.GetData(),
-			Component->SimpleCollisionMipLevel > Component->CollisionMipLevel ? SimpleCollisionWeightmapMipData.GetData() : nullptr);
+	};
 
-		ULandscapeHeightfieldCollisionComponent* CollisionComponent = Component->CollisionComponent.Get();
-		if (CollisionComponent)
+	if (GetMutableDefault<UEditorExperimentalSettings>()->bLandscapeLayerSystem)
+	{
+		if (LandscapeInfo->LandscapeActor)
 		{
-			CollisionComponent->RecreateCollision();
+			LandscapeInfo->LandscapeActor->ForEachLayer([&](const FLandscapeLayer& CurrentLayer)
+			{
+				FScopedSetLandscapeEditingLayer Scope(LandscapeInfo->LandscapeActor.Get(), CurrentLayer.Guid);
+				DoReplace();
+			});
+			LandscapeInfo->LandscapeActor->RequestLayersContentUpdate(ELandscapeLayersContentUpdateFlag::All, true);
 		}
 	}
+	else
+	{
+		DoReplace();
+	}
 }
 
 // simple classes for the template....
diff --git a/Engine/Source/Runtime/Landscape/Private/LandscapeEditLayers.cpp b/Engine/Source/Runtime/Landscape/Private/LandscapeEditLayers.cpp
--- a/Engine/Source/Runtime/Landscape/Private/LandscapeEditLayers.cpp
+++ b/Engine/Source/Runtime/Landscape/Private/LandscapeEditLayers.cpp
@@ -4474,6 +4474,14 @@ const FLandscapeLayer* ALandscape::GetLayer(const FGuid& InLayerGuid) const
 	return LandscapeLayers.FindByPredicate([InLayerGuid](const FLandscapeLayer& Other) { return Other.Guid == InLayerGuid; });
 }
 
+void ALandscape::ForEachLayer(TFunctionRef<void(const struct FLandscapeLayer&)> Fn)
+{
+	for (const FLandscapeLayer& Layer : LandscapeLayers)
+	{
+		Fn(Layer);
+	}
+}
+
 void ALandscape::DeleteLayer(int32 InLayerIndex)
 {
 	ensure(GetMutableDefault<UEditorExperimentalSettings>()->bLandscapeLayerSystem);
