- Landscape Layer System: Grass Update Fixes

diff --git a/Engine/Source/Runtime/Landscape/Classes/Landscape.h b/Engine/Source/Runtime/Landscape/Classes/Landscape.h
--- a/Engine/Source/Runtime/Landscape/Classes/Landscape.h
+++ b/Engine/Source/Runtime/Landscape/Classes/Landscape.h
@@ -251,6 +251,7 @@ public:
 	// Layers stuff
 #if WITH_EDITOR
 	LANDSCAPE_API void RequestLayersContentUpdate(ELandscapeLayersContentUpdateFlag InDataFlags, bool InUpdateAllMaterials = false);
+	LANDSCAPE_API bool HasPendingLayersContentUpdate() const { return LayersContentUpdateFlags != 0; }
 	LANDSCAPE_API void CreateLayer(FName InName = NAME_None, bool bInUpdateLayersContent = true);
 	LANDSCAPE_API bool ReorderLayer(int32 InStartingLayerIndex, int32 InDestinationLayerIndex);
 	LANDSCAPE_API bool IsLayerNameUnique(const FName& InName) const;
diff --git a/Engine/Source/Runtime/Landscape/Classes/LandscapeProxy.h b/Engine/Source/Runtime/Landscape/Classes/LandscapeProxy.h
--- a/Engine/Source/Runtime/Landscape/Classes/LandscapeProxy.h
+++ b/Engine/Source/Runtime/Landscape/Classes/LandscapeProxy.h
@@ -807,6 +807,9 @@ public:
 	/* Invalidate the precomputed grass and baked texture data for the specified components */
 	LANDSCAPE_API static void InvalidateGeneratedComponentData(const TSet<ULandscapeComponent*>& Components);
 
+	/* Invalidate the precomputed grass and baked texture data on all components */
+	LANDSCAPE_API void InvalidateGeneratedComponentData();
+
 #if WITH_EDITOR
 	/** Render grass maps for the specified components */
 	void RenderGrassMaps(const TArray<ULandscapeComponent*>& LandscapeComponents, const TArray<ULandscapeGrassType*>& GrassTypes);
diff --git a/Engine/Source/Runtime/Landscape/Private/Landscape.cpp b/Engine/Source/Runtime/Landscape/Private/Landscape.cpp
--- a/Engine/Source/Runtime/Landscape/Private/Landscape.cpp
+++ b/Engine/Source/Runtime/Landscape/Private/Landscape.cpp
@@ -3275,19 +3275,31 @@ void ALandscapeProxy::UpdateBakedTextures()
 }
 #endif
 
-void ALandscapeProxy::InvalidateGeneratedComponentData(const TSet<ULandscapeComponent*>& Components)
+template<class ContainerType>
+void InvalidateGeneratedComponentDataImpl(const ContainerType& Components)
 {
 	TMap<ALandscapeProxy*, TSet<ULandscapeComponent*>> ByProxy;
-	for (auto Component : Components)
+	for (auto Iter = Components.CreateConstIterator(); Iter; ++Iter)
 	{
+		ULandscapeComponent* Component = *Iter;
 		Component->BakedTextureMaterialGuid.Invalidate();
-
 		ByProxy.FindOrAdd(Component->GetLandscapeProxy()).Add(Component);
 	}
-	for (auto It = ByProxy.CreateConstIterator(); It; ++It)
+
+	for (auto Iter = ByProxy.CreateConstIterator(); Iter; ++Iter)
 	{
-		It.Key()->FlushGrassComponents(&It.Value());
+		Iter.Key()->FlushGrassComponents(&Iter.Value());
 	}
 }
 
+void ALandscapeProxy::InvalidateGeneratedComponentData()
+{
+	InvalidateGeneratedComponentDataImpl(LandscapeComponents);
+}
+
+void ALandscapeProxy::InvalidateGeneratedComponentData(const TSet<ULandscapeComponent*>& Components)
+{
+	InvalidateGeneratedComponentDataImpl(Components);
+}
+
 #undef LOCTEXT_NAMESPACE
diff --git a/Engine/Source/Runtime/Landscape/Private/LandscapeEditLayers.cpp b/Engine/Source/Runtime/Landscape/Private/LandscapeEditLayers.cpp
--- a/Engine/Source/Runtime/Landscape/Private/LandscapeEditLayers.cpp
+++ b/Engine/Source/Runtime/Landscape/Private/LandscapeEditLayers.cpp
@@ -4411,7 +4411,8 @@ void ALandscape::SetLayerName(int32 InLayerIndex, const FName& InName)
 void ALandscape::SetLayerAlpha(int32 InLayerIndex, const float InAlpha, bool bInHeightmap)
 {
 	FLandscapeLayer* Layer = GetLayer(InLayerIndex);
-	if (!Layer)
+	ULandscapeInfo* LandscapeInfo = GetLandscapeInfo();
+	if (!LandscapeInfo || !Layer)
 	{
 		return;
 	}
@@ -4424,12 +4425,21 @@ void ALandscape::SetLayerAlpha(int32 InLayerIndex, const float InAlpha, bool bIn
 	Modify();
 	LayerAlpha = InAlpha;
 	RequestLayersContentUpdate(ELandscapeLayersContentUpdateFlag::All, true);
+
+	LandscapeInfo->ForAllLandscapeProxies([](ALandscapeProxy* Proxy)
+	{
+		if (Proxy)
+		{
+			Proxy->InvalidateGeneratedComponentData();
+		}
+	});
 }
 
 void ALandscape::SetLayerVisibility(int32 InLayerIndex, bool bInVisible)
 {
 	FLandscapeLayer* Layer = GetLayer(InLayerIndex);
-	if (!Layer || Layer->bVisible == bInVisible)
+	ULandscapeInfo* LandscapeInfo = GetLandscapeInfo();
+	if (!LandscapeInfo || !Layer || Layer->bVisible == bInVisible)
 	{
 		return;
 	}
@@ -4437,6 +4447,14 @@ void ALandscape::SetLayerVisibility(int32 InLayerIndex, bool bInVisible)
 	Modify();
 	Layer->bVisible = bInVisible;
 	RequestLayersContentUpdate(ELandscapeLayersContentUpdateFlag::All, true);
+
+	LandscapeInfo->ForAllLandscapeProxies([](ALandscapeProxy* Proxy)
+	{
+		if (Proxy)
+		{
+			Proxy->InvalidateGeneratedComponentData();
+		}
+	});
 }
 
 void ALandscape::SetLayerLocked(int32 InLayerIndex, bool bLocked)
@@ -4476,7 +4494,7 @@ const FLandscapeLayer* ALandscape::GetLayer(const FGuid& InLayerGuid) const
 
 void ALandscape::ForEachLayer(TFunctionRef<void(const struct FLandscapeLayer&)> Fn)
 {
-	for (const FLandscapeLayer& Layer : LandscapeLayers)
+	for (FLandscapeLayer& Layer : LandscapeLayers)
 	{
 		Fn(Layer);
 	}
diff --git a/Engine/Source/Runtime/Landscape/Private/LandscapeGrass.cpp b/Engine/Source/Runtime/Landscape/Private/LandscapeGrass.cpp
--- a/Engine/Source/Runtime/Landscape/Private/LandscapeGrass.cpp
+++ b/Engine/Source/Runtime/Landscape/Private/LandscapeGrass.cpp
@@ -24,6 +24,7 @@
 #include "RHIStaticStates.h"
 #include "SceneView.h"
 #include "Shader.h"
+#include "Landscape.h"
 #include "LandscapeProxy.h"
 #include "LightMap.h"
 #include "Engine/MapBuildDataRegistry.h"
@@ -54,6 +55,7 @@
 #include "InstancedStaticMesh.h"
 #include "MeshPassProcessor.h"
 #include "MeshPassProcessor.inl"
+#include "Settings/EditorExperimentalSettings.h"
 
 #define LOCTEXT_NAMESPACE "Landscape"
 
@@ -1334,6 +1336,19 @@ void ALandscapeProxy::TickGrass()
 			return;
 		}
 	}
+
+	if (GetMutableDefault<UEditorExperimentalSettings>()->bLandscapeLayerSystem)
+	{
+		if (ALandscape* Landscape = GetLandscapeActor())
+		{
+			if (Landscape->HasPendingLayersContentUpdate())
+			{
+				return;
+			}
+		}
+	}
+
+
 	// Update foliage
 	static TArray<FVector> OldCameras;
 	if (CVarUseStreamingManagerForCameras.GetValueOnGameThread() == 0)
