-Landscape Layer System: Minor cleanup to the Grass Update check

#rnx

diff --git a/Engine/Source/Runtime/Landscape/Classes/Landscape.h b/Engine/Source/Runtime/Landscape/Classes/Landscape.h
--- a/Engine/Source/Runtime/Landscape/Classes/Landscape.h
+++ b/Engine/Source/Runtime/Landscape/Classes/Landscape.h
@@ -247,11 +247,11 @@ public:
 	virtual void FinishDestroy() override;
 	//~ End UObject Interface
 
+	LANDSCAPE_API bool IsUpToDate() const;
 
 	// Layers stuff
 #if WITH_EDITOR
 	LANDSCAPE_API void RequestLayersContentUpdate(ELandscapeLayersContentUpdateFlag InDataFlags, bool InUpdateAllMaterials = false);
-	LANDSCAPE_API bool HasPendingLayersContentUpdate() const { return LayersContentUpdateFlags != 0; }
 	LANDSCAPE_API void CreateLayer(FName InName = NAME_None, bool bInUpdateLayersContent = true);
 	LANDSCAPE_API bool ReorderLayer(int32 InStartingLayerIndex, int32 InDestinationLayerIndex);
 	LANDSCAPE_API bool IsLayerNameUnique(const FName& InName) const;
diff --git a/Engine/Source/Runtime/Landscape/Private/LandscapeEditLayers.cpp b/Engine/Source/Runtime/Landscape/Private/LandscapeEditLayers.cpp
--- a/Engine/Source/Runtime/Landscape/Private/LandscapeEditLayers.cpp
+++ b/Engine/Source/Runtime/Landscape/Private/LandscapeEditLayers.cpp
@@ -4384,6 +4384,18 @@ void ALandscape::FinishDestroy()
 	Super::FinishDestroy();
 }
 
+bool ALandscape::IsUpToDate() const
+{
+#if WITH_EDITORONLY_DATA
+	if (GetMutableDefault<UEditorExperimentalSettings>()->bLandscapeLayerSystem)
+	{
+		return LayersContentUpdateFlags == 0;
+	}
+#endif
+
+	return true;
+}
+
 #if WITH_EDITOR
 bool ALandscape::IsLayerNameUnique(const FName& InName) const
 {
diff --git a/Engine/Source/Runtime/Landscape/Private/LandscapeGrass.cpp b/Engine/Source/Runtime/Landscape/Private/LandscapeGrass.cpp
--- a/Engine/Source/Runtime/Landscape/Private/LandscapeGrass.cpp
+++ b/Engine/Source/Runtime/Landscape/Private/LandscapeGrass.cpp
@@ -56,10 +56,6 @@
 #include "MeshPassProcessor.h"
 #include "MeshPassProcessor.inl"
 
-#if WITH_EDITOR
-#include "Settings/EditorExperimentalSettings.h"
-#endif
-
 #define LOCTEXT_NAMESPACE "Landscape"
 
 DEFINE_LOG_CATEGORY_STATIC(LogGrass, Log, All);
@@ -1340,19 +1336,13 @@ void ALandscapeProxy::TickGrass()
 		}
 	}
 
-#if WITH_EDITOR
-	if (GetMutableDefault<UEditorExperimentalSettings>()->bLandscapeLayerSystem)
+	if (ALandscape* Landscape = GetLandscapeActor())
 	{
-		if (ALandscape* Landscape = GetLandscapeActor())
+		if (!Landscape->IsUpToDate())
 		{
-			if (Landscape->HasPendingLayersContentUpdate())
-			{
-				return;
-			}
+			return;
 		}
 	}
-#endif
-
 
 	// Update foliage
 	static TArray<FVector> OldCameras;
