Data structure change to the Landscape Layers system to be component based instead of storing the info into the landscape actor/proxy
Change how the Setup fonction work to split each behavior into seperate function and calling each one in the correct code path
Fixed some bugs in the shaders

diff --git a/Engine/Shaders/Private/LandscapeLayersCS.usf b/Engine/Shaders/Private/LandscapeLayersCS.usf
--- a/Engine/Shaders/Private/LandscapeLayersCS.usf
+++ b/Engine/Shaders/Private/LandscapeLayersCS.usf
@@ -81,7 +81,24 @@ void PackPaintLayerToWeightmap(uint3 DispatchThreadId : SV_DispatchThreadID)
 
 		for (int i = 0; i < 4; ++i)
 		{
-			Channels[i] = ComputeWeightmapChannel(ComponentData.WeightmapChannelToProcess[i], DispatchThreadId.xy + int2(ComponentData.ComponenVertexPositionX[i], ComponentData.ComponenVertexPositionY[i]), ComponentData.SourcePaintLayerIndex[i]);
+			float ChannelResult = ComputeWeightmapChannel(ComponentData.WeightmapChannelToProcess[i], DispatchThreadId.xy + int2(ComponentData.ComponenVertexPositionX[i], ComponentData.ComponenVertexPositionY[i]), ComponentData.SourcePaintLayerIndex[i]);
+
+			if (ComponentData.WeightmapChannelToProcess[i] == 0)
+			{ 
+				Channels.x = ChannelResult;
+			}
+			else if (ComponentData.WeightmapChannelToProcess[i] == 1)
+			{
+				Channels.y = ChannelResult;
+			}
+			else if (ComponentData.WeightmapChannelToProcess[i] == 2)
+			{
+				Channels.z = ChannelResult;
+			}
+			else if (ComponentData.WeightmapChannelToProcess[i] == 3)
+			{
+				Channels.w = ChannelResult;
+			}
 		}
 
 		float2 TextureOutputOffset = InWeightmapTextureOutputOffset[DispatchThreadId.z];
diff --git a/Engine/Shaders/Private/LandscapeLayersPS.usf b/Engine/Shaders/Private/LandscapeLayersPS.usf
--- a/Engine/Shaders/Private/LandscapeLayersPS.usf
+++ b/Engine/Shaders/Private/LandscapeLayersPS.usf
@@ -21,7 +21,7 @@ float ExtractHeight(float2 InPackedHeight)
 
 float2 PackHeight(float InExtractedHeight)
 {
-	InExtractedHeight = clamp(InExtractedHeight, 0.0f, 65536.0f);
+	InExtractedHeight = clamp(InExtractedHeight, 0.0f, 65535.0f);
 	int iHeight = (int)InExtractedHeight;
 	return float2((float)((iHeight - (iHeight & 255)) >> 8) / 255.0, (float)(iHeight & 255) / 255.0);
 }
@@ -266,7 +266,7 @@ void PSHeightmapMainMips(float2 InTextureCoordinates : TEXCOORD0, out float4 Out
 		}
 		else
 		{
-			HeightResult = clamp(round(lerp(SourceHeight, DownHeight, Alpha)), 0.0, 65536.0);
+			HeightResult = clamp(round(lerp(SourceHeight, DownHeight, Alpha)), 0.0, 65535.0);
 			NormalResult.x = lerp(SourceSample.b, DownSample.b, Alpha);
 			NormalResult.y = lerp(SourceSample.a, DownSample.a, Alpha);
 		}
@@ -290,7 +290,7 @@ void PSHeightmapMainMips(float2 InTextureCoordinates : TEXCOORD0, out float4 Out
 		}
 		else
 		{
-			HeightResult = clamp(round(lerp(RightHeight, DownRightHeight, Alpha)), 0.0, 65536.0);;
+			HeightResult = clamp(round(lerp(RightHeight, DownRightHeight, Alpha)), 0.0, 65535.0);;
 			NormalResult.x = lerp(RightSample.b, DownRightSample.b, Alpha);
 			NormalResult.y = lerp(RightSample.a, DownRightSample.a, Alpha);
 		}
@@ -309,13 +309,13 @@ void PSHeightmapMainMips(float2 InTextureCoordinates : TEXCOORD0, out float4 Out
 
 		if (IsMinBorderTexelY) // on top border
 		{
-			HeightResult = clamp(round(lerp(SourceHeight, RightHeight, Alpha)), 0.0, 65536.0);
+			HeightResult = clamp(round(lerp(SourceHeight, RightHeight, Alpha)), 0.0, 65535.0);
 			NormalResult.x = lerp(SourceSample.b, RightSample.b, Alpha);
 			NormalResult.y = lerp(SourceSample.a, RightSample.a, Alpha);
 		}
 		else if (IsMaxBorderTexelY) // on bottom border
 		{
-			HeightResult = clamp(round(lerp(DownHeight, DownRightHeight, Alpha)), 0.0, 65536.0);
+			HeightResult = clamp(round(lerp(DownHeight, DownRightHeight, Alpha)), 0.0, 65535.0);
 			NormalResult.x = lerp(DownSample.b, DownRightSample.b, Alpha);
 			NormalResult.y = lerp(DownSample.a, DownRightSample.a, Alpha);
 		}
@@ -324,7 +324,7 @@ void PSHeightmapMainMips(float2 InTextureCoordinates : TEXCOORD0, out float4 Out
 			float SourceToRightHeight = lerp(SourceHeight, RightHeight, Alpha);
 			float DownToDownRightHeight = lerp(DownHeight, DownRightHeight, Alpha);
 
-			HeightResult = clamp(round(lerp(SourceToRightHeight, DownToDownRightHeight, Alpha)), 0.0, 65536.0);
+			HeightResult = clamp(round(lerp(SourceToRightHeight, DownToDownRightHeight, Alpha)), 0.0, 65535.0);
 			NormalResult.x = lerp(lerp(SourceSample.b, RightSample.b, Alpha), lerp(DownSample.b, DownRightSample.b, Alpha), Alpha);
 			NormalResult.y = lerp(lerp(SourceSample.a, RightSample.a, Alpha), lerp(DownSample.a, DownRightSample.a, Alpha), Alpha);
 		}
diff --git a/Engine/Source/Editor/LandscapeEditor/Private/LandscapeEdMode.cpp b/Engine/Source/Editor/LandscapeEditor/Private/LandscapeEdMode.cpp
--- a/Engine/Source/Editor/LandscapeEditor/Private/LandscapeEdMode.cpp
+++ b/Engine/Source/Editor/LandscapeEditor/Private/LandscapeEdMode.cpp
@@ -4513,36 +4513,42 @@ bool FEdModeLandscape::NeedToFillEmptyMaterialLayers() const
 		return false;
 	}
 
-	TArray<ALandscapeProxy*> AllLandscapes;
-	AllLandscapes.Add(CurrentToolTarget.LandscapeInfo->LandscapeActor.Get());
+	bool bCanFill = true;
 
-	for (const auto& It : CurrentToolTarget.LandscapeInfo->Proxies)
+	CurrentToolTarget.LandscapeInfo->ForAllLandscapeProxies([&](ALandscapeProxy* Proxy)
 	{
-		AllLandscapes.Add(It);
-	}
-
-	for (const ALandscapeProxy* LandscapeProxy : AllLandscapes)
-	{
-		for (const auto& ItLayerPair : LandscapeProxy->LandscapeLayersData)
+		if (!bCanFill)
 		{
-			const FLandscapeLayerData& LayerData = ItLayerPair.Value;
+			return;
+		}
 
-			for (const auto& ItWeightmapPair : LayerData.WeightmapData)
-			{
-				const FWeightmapLayerData& WeightmapData = ItWeightmapPair.Value;
+		ALandscape* Landscape = Proxy->GetLandscapeActor();
 
-				for (const FWeightmapLayerAllocationInfo& Alloc : WeightmapData.WeightmapLayerAllocations)
+		if (Landscape != nullptr)
+		{
+			for (FLandscapeLayer& Layer : Landscape->LandscapeLayers)
+			{
+				for (ULandscapeComponent* Component : Proxy->LandscapeComponents)
 				{
-					if (Alloc.LayerInfo != nullptr)
+					const FLandscapeLayerComponentData* LayerComponentData = Component->GetLayerData(Layer.Guid);
+
+					if (LayerComponentData != nullptr)
 					{
-						return false;
+						for (const FWeightmapLayerAllocationInfo& Alloc : LayerComponentData->WeightmapData.LayerAllocations)
+						{
+							if (Alloc.LayerInfo != nullptr)
+							{
+								bCanFill = false;
+								return;
+							}
+						}
 					}
 				}
 			}
 		}
-	}
+	});	
 
-	return true;
+	return bCanFill;
 }
 
 void FEdModeLandscape::OnLevelActorAdded(AActor* InActor)
diff --git a/Engine/Source/Editor/LandscapeEditor/Private/LandscapeEditorDetailCustomization_NewLandscape.cpp b/Engine/Source/Editor/LandscapeEditor/Private/LandscapeEditorDetailCustomization_NewLandscape.cpp
--- a/Engine/Source/Editor/LandscapeEditor/Private/LandscapeEditorDetailCustomization_NewLandscape.cpp
+++ b/Engine/Source/Editor/LandscapeEditor/Private/LandscapeEditorDetailCustomization_NewLandscape.cpp
@@ -40,6 +40,7 @@
 #include "Framework/Application/SlateApplication.h"
 #include "Widgets/Input/SNumericEntryBox.h"
 #include "LandscapeDataAccess.h"
+#include "Settings/EditorExperimentalSettings.h"
 
 #define LOCTEXT_NAMESPACE "LandscapeEditor.NewLandscape"
 
@@ -808,6 +809,11 @@ FReply FLandscapeEditorDetailCustomization_NewLandscape::OnCreateButtonClicked()
 		Landscape->LandscapeMaterial = LandscapeEdMode->UISettings->NewLandscape_Material.Get();
 		Landscape->SetActorRelativeScale3D(LandscapeEdMode->UISettings->NewLandscape_Scale);
 
+		if (GetMutableDefault<UEditorExperimentalSettings>()->bLandscapeLayerSystem)
+		{
+			Landscape->PreviousExperimentalLandscapeLayers = true;
+		}
+
 		Landscape->Import(FGuid::NewGuid(), 0, 0, SizeX-1, SizeY-1, LandscapeEdMode->UISettings->NewLandscape_SectionsPerComponent, LandscapeEdMode->UISettings->NewLandscape_QuadsPerSection, Data.GetData(),
 			nullptr, ImportLayers.GetValue(), LandscapeEdMode->UISettings->ImportLandscape_AlphamapType);
 
diff --git a/Engine/Source/Runtime/Landscape/Classes/Landscape.h b/Engine/Source/Runtime/Landscape/Classes/Landscape.h
--- a/Engine/Source/Runtime/Landscape/Classes/Landscape.h
+++ b/Engine/Source/Runtime/Landscape/Classes/Landscape.h
@@ -75,16 +75,14 @@ enum EWeightmapRTType : uint8
 // Internal Update Flags that shouldn't be used alone
 enum EInternalLandscapeLayersContentUpdateFlag : uint32
 {
-	Internal_Heightmap_Setup = 0x00000001u,
-	Internal_Heightmap_Render = 0x00000002u,
-	Internal_Heightmap_BoundsAndCollision = 0x00000004u,
-	Internal_Heightmap_ResolveToTexture = 0x00000008u,
-	Internal_Heightmap_ClientUpdate = 0x00000010u,
-
-	Internal_Weightmap_Setup = 0x00000100u,
-	Internal_Weightmap_Render = 0x00000200u,
-	Internal_Weightmap_Collision = 0x00000400u,
-	Internal_Weightmap_ResolveToTexture = 0x00000800u,
+	Internal_Heightmap_Render = 0x00000001u,
+	Internal_Heightmap_BoundsAndCollision = 0x00000002u,
+	Internal_Heightmap_ResolveToTexture = 0x00000004u,
+	Internal_Heightmap_ClientUpdate = 0x000000008u,
+
+	Internal_Weightmap_Render = 0x00000100u,
+	Internal_Weightmap_Collision = 0x00000200u,
+	Internal_Weightmap_ResolveToTexture = 0x00000400u,
 };
 
 enum ELandscapeLayersContentUpdateFlag : uint32
@@ -97,7 +95,6 @@ enum ELandscapeLayersContentUpdateFlag : uint32
 	Weightmap_All = Internal_Weightmap_Render | Internal_Weightmap_Collision | Internal_Weightmap_ResolveToTexture,
 
 	All = Heightmap_All | Weightmap_All,
-	All_Setup = Internal_Heightmap_Setup | Internal_Weightmap_Setup,
 	All_Render = Heightmap_Render | Weightmap_Render,
 };
 
@@ -252,8 +249,13 @@ public:
 	// Layers stuff
 #if WITH_EDITOR
 	LANDSCAPE_API void RequestLayersContentUpdate(ELandscapeLayersContentUpdateFlag InDataFlags, bool InUpdateAllMaterials = false);
-	LANDSCAPE_API void CreateLayer(FName InName = NAME_None, bool bInUpdateLayersContent = true);
 	LANDSCAPE_API bool ReorderLayer(int32 InStartingLayerIndex, int32 InDestinationLayerIndex);
+	LANDSCAPE_API void CreateLayer(FName InName = NAME_None);
+	LANDSCAPE_API void CreateDefaultLayer(const FIntPoint& InComponentCounts);
+	LANDSCAPE_API void InitializeLayerWithEmptyContent(const FGuid& InLayerGUID);
+	LANDSCAPE_API void CopyOldDataToDefaultLayer();
+	LANDSCAPE_API TMap<UTexture2D*, TArray<ULandscapeComponent*>> GenerateComponentsPerHeightmaps() const;
+	LANDSCAPE_API FIntPoint ComputeComponentCounts() const;
 	LANDSCAPE_API bool IsLayerNameUnique(const FName& InName) const;
 	LANDSCAPE_API void SetLayerName(int32 InLayerIndex, const FName& InName);
 	LANDSCAPE_API void SetLayerAlpha(int32 InLayerIndex, const float InAlpha, bool bInHeightmap);
@@ -267,6 +269,7 @@ public:
 	LANDSCAPE_API void ClearLayer(const FGuid& InLayerGuid, bool bInUpdateCollision = true);
 	LANDSCAPE_API void DeleteLayer(int32 InLayerIndex);
 	LANDSCAPE_API void SetEditingLayer(const FGuid& InLayerGuid = FGuid());
+	LANDSCAPE_API const FGuid& GetEditingLayer() const;
 	LANDSCAPE_API void ShowOnlySelectedLayer(int32 InLayerIndex);
 	LANDSCAPE_API void ShowAllLayers();
 	LANDSCAPE_API void UpdateLandscapeSplines(const FGuid& InLayerGuid = FGuid(), bool bUpdateOnlySelected = false);
@@ -291,22 +294,22 @@ private:
 	void RegenerateLayersContent();
 	void RegenerateLayersHeightmaps();
 	void RegenerateLayersWeightmaps();
-	void ResolveLayersHeightmapTexture(const TArray<ALandscapeProxy*>& InAllLandscapes);
-	void ResolveLayersWeightmapTexture(const TArray<ALandscapeProxy*>& InAllLandscapes);
-	void ResolveLayersTexture(FLandscapeLayersTexture2DCPUReadBackResource* InCPUReadBackTexture, UTexture2D* InOriginalTexture);
+	void ResolveLayersHeightmapTexture();
+	void ResolveLayersWeightmapTexture();
+	void ResolveLayersTexture(class FLandscapeLayersTexture2DCPUReadBackResource* InCPUReadBackTexture, UTexture2D* InOutputTexture);
 
-	bool AreLayersHeightmapTextureResourcesReady(const TArray<ALandscapeProxy*>& InAllLandscapes) const;
-	bool AreLayersWeightmapTextureResourcesReady(const TArray<ALandscapeProxy*>& InAllLandscapes) const;
+	bool PrepareLayersHeightmapTextureResources() const;
+	bool PrepareLayersWeightmapTextureResources() const;
 
-	void UpdateLayersMaterialInstances(const TArray<ULandscapeComponent*>& InComponentsToUpdate, const TMap<ULandscapeComponent*, TArray<ULandscapeLayerInfoObject*>>& InZeroAllocationsPerComponents);
+	void UpdateLayersMaterialInstances(const TArray<ULandscapeComponent*>& InComponentsToUpdate);
 
-	void PrepareComponentDataToExtractMaterialLayersCS(const FLandscapeLayer& InLayer, int32 InCurrentWeightmapToProcessIndex, bool InOutputDebugName, const TArray<ALandscapeProxy*>& InAllLandscape, class FLandscapeTexture2DResource* InOutTextureData,
+	void PrepareComponentDataToExtractMaterialLayersCS(const FLandscapeLayer& InLayer, int32 InCurrentWeightmapToProcessIndex, bool InOutputDebugName, class FLandscapeTexture2DResource* InOutTextureData,
 														  TArray<struct FLandscapeLayerWeightmapExtractMaterialLayersComponentData>& OutComponentData, TMap<ULandscapeLayerInfoObject*, int32>& OutLayerInfoObjects);
-	void PrepareComponentDataToPackMaterialLayersCS(int32 InCurrentWeightmapToProcessIndex, bool InOutputDebugName, const TArray<ULandscapeComponent*>& InAllLandscapeComponents, 
-													   TArray<UTexture2D*>& InOutProcessedWeightmaps, TArray<class FLandscapeLayersTexture2DCPUReadBackResource*>& InOutProcessedWeightmapCPUCopy, TArray<struct FLandscapeLayerWeightmapPackMaterialLayersComponentData>& OutComponentData);
-	void ReallocateLayersWeightmaps(const TArray<ALandscapeProxy*>& InAllLandscape, const TArray<ULandscapeLayerInfoObject*>& InBrushRequiredAllocations, TArray<ULandscapeComponent*>& OutComponentThatNeedMaterialRebuild);
+	void PrepareComponentDataToPackMaterialLayersCS(int32 InCurrentWeightmapToProcessIndex, bool InOutputDebugName, const TArray<ULandscapeComponent*>& InAllLandscapeComponents, TArray<UTexture2D*>& InOutProcessedWeightmaps,
+													TArray<class FLandscapeLayersTexture2DCPUReadBackResource*>& OutProcessedCPUReadBackTexture, TArray<struct FLandscapeLayerWeightmapPackMaterialLayersComponentData>& OutComponentData);
+	void ReallocateLayersWeightmaps(const TArray<ULandscapeLayerInfoObject*>& InBrushRequiredAllocations, TArray<ULandscapeComponent*>& OutComponentThatNeedMaterialRebuild);
 	void InitLayersWeightmapResources(uint8 InLayerCount);
-	bool GenerateZeroAllocationPerComponents(const TArray<ALandscapeProxy*>& InAllLandscape, const TMap<ULandscapeLayerInfoObject*, bool>& InWeightmapLayersBlendSubstractive, TMap<ULandscapeComponent*, TArray<ULandscapeLayerInfoObject*>>& OutZeroAllocationsPerComponents);
+	bool GenerateZeroAllocationPerComponents(const TArray<ALandscapeProxy*>& InAllLandscape, const TMap<ULandscapeLayerInfoObject*, bool>& InWeightmapLayersBlendSubstractive);
 
 	void GenerateLayersRenderQuad(const FIntPoint& InVertexPosition, float InVertexSize, const FVector2D& InUVStart, const FVector2D& InUVSize, TArray<struct FLandscapeLayersTriangle>& OutTriangles) const;
 	void GenerateLayersRenderQuadsAtlas(const FIntPoint& InSectionBase, const FVector2D& InScaleBias, float InSubSectionSizeQuad, const FIntPoint& InReadSize, const FIntPoint& InWriteSize, TArray<struct FLandscapeLayersTriangle>& OutTriangles) const;
@@ -315,24 +318,28 @@ private:
 	void GenerateLayersRenderQuadsNonAtlasToAtlas(const FIntPoint& InSectionBase, const FVector2D& InScaleBias, float InSubSectionSizeQuad, const FIntPoint& InReadSize, const FIntPoint& InWriteSize, TArray<struct FLandscapeLayersTriangle>& OutTriangles) const;
 	void GenerateLayersRenderQuadsMip(const FIntPoint& InSectionBase, const FVector2D& InScaleBias, float InSubSectionSizeQuad, const FIntPoint& InReadSize, const FIntPoint& InWriteSize, uint8 InCurrentMip, TArray<FLandscapeLayersTriangle>& OutTriangles) const;
 
-	void ClearLayersWeightmapTextureResource(const FString& InDebugName, FTextureRenderTargetResource* InTextureResourceToClear);
+	void ClearLayersWeightmapTextureResource(const FString& InDebugName, FTextureRenderTargetResource* InTextureResourceToClear) const;
 	void DrawHeightmapComponentsToRenderTarget(const FString& InDebugName, const TArray<ULandscapeComponent*>& InComponentsToDraw, UTexture* InHeightmapRTRead, UTextureRenderTarget2D* InOptionalHeightmapRTRead2, UTextureRenderTarget2D* InHeightmapRTWrite, ERTDrawingType InDrawType,
 											   bool InClearRTWrite, struct FLandscapeLayersHeightmapShaderParameters& InShaderParams, uint8 InMipRender = 0) const;
 
-	void DrawWeightmapComponentsToRenderTarget(const FString& InDebugName, const TArray<ULandscapeComponent*>& InComponentsToDraw, UTexture* InWeightmapRTRead, UTextureRenderTarget2D* InOptionalWeightmapRTRead2, UTextureRenderTarget2D* InWeightmapRTWrite,
+	void DrawWeightmapComponentsToRenderTarget(const FString& InDebugName, const TArray<FIntPoint>& InSectionBaseList, const FVector2D& InScaleBias, UTexture* InWeightmapRTRead, UTextureRenderTarget2D* InOptionalWeightmapRTRead2, UTextureRenderTarget2D* InWeightmapRTWrite,
 											   bool InClearRTWrite, struct FLandscapeLayersWeightmapShaderParameters& InShaderParams, uint8 InMipRender) const;
 
-	void DrawWeightmapComponentsToRenderTarget(const FString& InDebugName, const FIntPoint& InSectionBase, const FVector2D& InScaleBias, UTexture* InWeightmapRTRead, UTextureRenderTarget2D* InOptionalWeightmapRTRead2, UTextureRenderTarget2D* InWeightmapRTWrite,
-									 		   bool InClearRTWrite, struct FLandscapeLayersWeightmapShaderParameters& InShaderParams, uint8 InMipRender) const;
+	void DrawWeightmapComponentsToRenderTarget(const FString& InDebugName, const TArray<ULandscapeComponent*>& InComponentsToDraw, UTexture* InWeightmapRTRead, UTextureRenderTarget2D* InOptionalWeightmapRTRead2, UTextureRenderTarget2D* InWeightmapRTWrite,
+											   bool InClearRTWrite, struct FLandscapeLayersWeightmapShaderParameters& InShaderParams, uint8 InMipRender) const;
 
-	void DrawHeightmapComponentsToRenderTargetMips(TArray<ULandscapeComponent*>& InComponentsToDraw, UTexture* InReadHeightmap, bool InClearRTWrite, struct FLandscapeLayersHeightmapShaderParameters& InShaderParams) const;
-	void DrawWeightmapComponentToRenderTargetMips(const FIntPoint& TopLeftTexturePosition, UTexture* InReadWeightmap, bool InClearRTWrite, struct FLandscapeLayersWeightmapShaderParameters& InShaderParams) const;
+	void DrawHeightmapComponentsToRenderTargetMips(const TArray<ULandscapeComponent*>& InComponentsToDraw, UTexture* InReadHeightmap, bool InClearRTWrite, struct FLandscapeLayersHeightmapShaderParameters& InShaderParams) const;
+	void DrawWeightmapComponentToRenderTargetMips(const TArray<FVector2D>& InTexturePositionsToDraw, UTexture* InReadWeightmap, bool InClearRTWrite, struct FLandscapeLayersWeightmapShaderParameters& InShaderParams) const;
 
 	void CopyLayersTexture(UTexture* InSourceTexture, UTexture* InDestTexture, FTextureResource* InDestCPUResource = nullptr, const FIntPoint& InFirstComponentSectionBase = FIntPoint(0, 0), uint8 InSourceCurrentMip = 0, uint8 InDestCurrentMip = 0,
 							   uint32 InSourceArrayIndex = 0, uint32 InDestArrayIndex = 0) const;
 	void CopyLayersTexture(const FString& InSourceDebugName, FTextureResource* InSourceResource, const FString& InDestDebugName, FTextureResource* InDestResource, FTextureResource* InDestCPUResource = nullptr, const FIntPoint& InFirstComponentSectionBase = FIntPoint(0, 0),
 							   uint8 InSourceCurrentMip = 0, uint8 InDestCurrentMip = 0, uint32 InSourceArrayIndex = 0, uint32 uInDestArrayIndex = 0) const;
 
+	void InitializeLayers();
+	void InitializeLayersWeightmapUsage();
+	void CreateLayersRenderingResource(const FIntPoint& InComponentCounts);
+
 	void PrintLayersDebugRT(const FString& InContext, UTextureRenderTarget2D* InDebugRT, uint8 InMipRender = 0, bool InOutputHeight = true, bool InOutputNormals = false) const;
 	void PrintLayersDebugTextureResource(const FString& InContext, FTextureResource* InTextureResource, uint8 InMipRender = 0, bool InOutputHeight = true, bool InOutputNormals = false) const;
 	void PrintLayersDebugHeightData(const FString& InContext, const TArray<FColor>& InHeightmapData, const FIntPoint& InDataSize, uint8 InMipRender, bool InOutputNormals = false) const;
@@ -361,6 +368,9 @@ public:
 	UPROPERTY(Transient)
 	bool PreviousExperimentalLandscapeLayers;
 
+	UPROPERTY(Transient)
+	bool bLandscapeLayersAreInitialized;
+
 private:
 
 	UPROPERTY(Transient)
diff --git a/Engine/Source/Runtime/Landscape/Classes/LandscapeComponent.h b/Engine/Source/Runtime/Landscape/Classes/LandscapeComponent.h
--- a/Engine/Source/Runtime/Landscape/Classes/LandscapeComponent.h
+++ b/Engine/Source/Runtime/Landscape/Classes/LandscapeComponent.h
@@ -252,6 +252,53 @@ struct FLandscapeComponentMaterialOverride
 	UMaterialInterface* Material;
 };
 
+USTRUCT(NotBlueprintable)
+struct FWeightmapData
+{
+	GENERATED_USTRUCT_BODY()
+
+	UPROPERTY()
+	TArray<UTexture2D*> Textures;
+	
+	UPROPERTY()
+	TArray<FWeightmapLayerAllocationInfo> LayerAllocations;
+
+	UPROPERTY(Transient)
+	TArray<ULandscapeWeightmapUsage*> TextureUsages;
+};
+
+USTRUCT(NotBlueprintable)
+struct FHeightmapData
+{
+	GENERATED_USTRUCT_BODY()
+
+	UPROPERTY()
+	UTexture2D* Texture;
+};
+
+USTRUCT(NotBlueprintable)
+struct FLandscapeLayerComponentData
+{
+	GENERATED_USTRUCT_BODY()
+
+	UPROPERTY()
+	FHeightmapData HeightmapData;
+
+	UPROPERTY()
+	FWeightmapData WeightmapData;
+
+	bool IsInitialized() const { return HeightmapData.Texture != nullptr || WeightmapData.Textures.Num() > 0;  }
+};
+
+USTRUCT(NotBlueprintable)
+struct FLandscapeLayersGlobalComponentData
+{
+	GENERATED_USTRUCT_BODY()
+
+	UPROPERTY()
+	FIntPoint TopLeftSectionBase;
+};
+
 UCLASS(hidecategories=(Display, Attachment, Physics, Debug, Collision, Movement, Rendering, PrimitiveComponent, Object, Transform, Mobility), showcategories=("Rendering|Material"), MinimalAPI, Within=LandscapeProxy)
 class ULandscapeComponent : public UPrimitiveComponent
 {
@@ -335,13 +382,13 @@ private:
 	UPROPERTY()
 	FGuid LightingGuid;
 
-	/** Current data we're working on (only used in Layering mode) This data will get set by tools as needed */
-	UTexture2D* CurrentEditingHeightmapTexture;
-	TArray<FWeightmapLayerAllocationInfo>* CurrentEditingWeightmapLayerAllocations;
-	TArray<UTexture2D*>* CurrentEditingWeightmapTextures;
-	TArray<ULandscapeWeightmapUsage*>* CurrentEditingWeightmapTexturesUsage;
-	FGuid CurrentLayerGuid;
+	UPROPERTY()
+	TMap<FGuid, FLandscapeLayerComponentData> LayersData;
+
+	UPROPERTY()
+	FLandscapeLayersGlobalComponentData GlobalLayersData;
 
+	// Final layer data
 	UPROPERTY(Transient)
 	TArray<ULandscapeWeightmapUsage*> WeightmapTexturesUsage;
 #endif // WITH_EDITORONLY_DATA
@@ -475,7 +522,6 @@ public:
 	virtual void Serialize(FArchive& Ar) override;
 	virtual void GetResourceSizeEx(FResourceSizeEx& CumulativeResourceSize) override;
 	virtual void BeginDestroy() override;
-	static void AddReferencedObjects(UObject* InThis, FReferenceCollector& Collector);
 	virtual void PostDuplicate(bool bDuplicateForPIE) override;
 #if WITH_EDITOR
 	virtual void BeginCacheForCookedPlatformData(const ITargetPlatform* TargetPlatform) override;
@@ -507,25 +553,34 @@ public:
 	virtual void GetStreamingRenderAssetInfo(FStreamingTextureLevelContext& LevelContext, TArray<FStreamingRenderAssetPrimitiveInfo>& OutStreamingRenderAssets) const override;
 	virtual bool IsPrecomputedLightingValid() const override;
 
-	LANDSCAPE_API UTexture2D* GetHeightmap(bool InReturnCurrentEditingHeightmap = false) const;
-	LANDSCAPE_API TArray<UTexture2D*>& GetWeightmapTextures(bool InReturnCurrentEditingWeightmap = false);
-	LANDSCAPE_API const TArray<UTexture2D*>& GetWeightmapTextures(bool InReturnCurrentEditingWeightmap = false) const;
+	LANDSCAPE_API UTexture2D* GetHeightmap(bool InReturnEditingHeightmap = false) const;
+	LANDSCAPE_API TArray<UTexture2D*>& GetWeightmapTextures(bool InReturnEditingWeightmap = false);
+	LANDSCAPE_API const TArray<UTexture2D*>& GetWeightmapTextures(bool InReturnEditingWeightmap = false) const;
 
-	LANDSCAPE_API TArray<FWeightmapLayerAllocationInfo>& GetWeightmapLayerAllocations(bool InReturnCurrentEditingWeightmap = false);
-	LANDSCAPE_API const TArray<FWeightmapLayerAllocationInfo>& GetWeightmapLayerAllocations(bool InReturnCurrentEditingWeightmap = false) const;
+	LANDSCAPE_API TArray<FWeightmapLayerAllocationInfo>& GetWeightmapLayerAllocations(bool InReturnEditingWeightmap = false);
+	LANDSCAPE_API const TArray<FWeightmapLayerAllocationInfo>& GetWeightmapLayerAllocations(bool InReturnEditingWeightmap = false) const;
 
 #if WITH_EDITOR
 	LANDSCAPE_API void SetHeightmap(UTexture2D* NewHeightmap);
 
-	LANDSCAPE_API void SetWeightmapTextures(const TArray<UTexture2D*>& InNewWeightmapTextures, bool InApplyToCurrentEditingWeightmap = false);
-	LANDSCAPE_API const FGuid& GetCurrentLayerGuid() const;
+	LANDSCAPE_API void SetWeightmapTextures(const TArray<UTexture2D*>& InNewWeightmapTextures, bool InApplyToEditingWeightmap = false);
 
 	LANDSCAPE_API void SetWeightmapLayerAllocations(const TArray<FWeightmapLayerAllocationInfo>& InNewWeightmapLayerAllocations);
-	LANDSCAPE_API void SetWeightmapTexturesUsage(const TArray<ULandscapeWeightmapUsage*>& InNewWeightmapTexturesUsage, bool InApplyToCurrentEditingWeightmap = false);
-	LANDSCAPE_API TArray<ULandscapeWeightmapUsage*>& GetWeightmapTexturesUsage(bool InReturnCurrentEditingWeightmap = false);
-	LANDSCAPE_API const TArray<ULandscapeWeightmapUsage*>& GetWeightmapTexturesUsage(bool InReturnCurrentEditingWeightmap = false) const;
+	LANDSCAPE_API void SetWeightmapTexturesUsage(const TArray<ULandscapeWeightmapUsage*>& InNewWeightmapTexturesUsage, bool InApplyToEditingWeightmap = false);
+	LANDSCAPE_API TArray<ULandscapeWeightmapUsage*>& GetWeightmapTexturesUsage(bool InReturnEditingWeightmap = false);
+	LANDSCAPE_API const TArray<ULandscapeWeightmapUsage*>& GetWeightmapTexturesUsage(bool InReturnEditingWeightmap = false) const;
+
+	LANDSCAPE_API const FLandscapeLayersGlobalComponentData& GetGlobalLayersData() const;
+	LANDSCAPE_API FLandscapeLayersGlobalComponentData& GetGlobalLayersData();
 
-	LANDSCAPE_API void SetEditingLayer(const struct FLandscapeLayer* InLayer, struct FLandscapeLayerData* InLayerData);
+	LANDSCAPE_API const FLandscapeLayerComponentData* GetLayerData(const FGuid& InLayerGuid) const;
+	LANDSCAPE_API FLandscapeLayerComponentData* GetLayerData(const FGuid& InLayerGuid);
+	LANDSCAPE_API void AddLayerData(const FGuid& InLayerGuid, const FLandscapeLayerComponentData& InData);
+	LANDSCAPE_API void RemoveLayerData(const FGuid& InLayerGuid);
+
+	FLandscapeLayerComponentData* GetEditingLayer();
+	const FLandscapeLayerComponentData* GetEditingLayer() const;
+	FGuid GetEditingLayerGUID() const;
 #endif 
 
 #if WITH_EDITOR
@@ -751,7 +806,7 @@ public:
 	/**
 	 * Create weightmaps for this component for the layers specified in the WeightmapLayerAllocations array
 	 */
-	void ReallocateWeightmaps(FLandscapeEditDataInterface* DataInterface= nullptr, bool InCanUseCurrentEditingWeightmap = true, bool InSaveToTransactionBuffer = true, bool InInitPlatformDataAsync = false, TArray<UTexture2D*>* OutNewCreatedTextures = nullptr);
+	void ReallocateWeightmaps(FLandscapeEditDataInterface* DataInterface= nullptr, bool InCanUseEditingWeightmap = true, bool InSaveToTransactionBuffer = true, bool InInitPlatformDataAsync = false, TArray<UTexture2D*>* OutNewCreatedTextures = nullptr);
 
 	/** Returns the actor's LandscapeMaterial, or the Component's OverrideLandscapeMaterial if set */
 	LANDSCAPE_API UMaterialInterface* GetLandscapeMaterial(int8 InLODIndex = INDEX_NONE) const;
@@ -805,16 +860,9 @@ public:
 
 protected:
 
-#if WITH_EDITOR
-	void SetEditingHeightmap(UTexture2D* InNewHeightmap);
-	void SetEditingWeightmaps(TArray<UTexture2D*>* InNewWeightmapTextures, TArray<FWeightmapLayerAllocationInfo>* InNewWeightmapLayerAllocations, TArray<ULandscapeWeightmapUsage*>* InNewWeightmapTexturesUsage);
-	void SetEditingLayerGuid(const FGuid& InLayerGuid);
-#endif
-
 	/** Whether the component type supports static lighting. */
 	virtual bool SupportsStaticLighting() const override
 	{
 		return true;
 	}
 };
-
diff --git a/Engine/Source/Runtime/Landscape/Classes/LandscapeProxy.h b/Engine/Source/Runtime/Landscape/Classes/LandscapeProxy.h
--- a/Engine/Source/Runtime/Landscape/Classes/LandscapeProxy.h
+++ b/Engine/Source/Runtime/Landscape/Classes/LandscapeProxy.h
@@ -359,54 +359,6 @@ private:
 	uint32 NumMips;
 };
 
-USTRUCT()
-struct FRenderDataPerHeightmap
-{
-	GENERATED_USTRUCT_BODY()
-
-	UPROPERTY(Transient)
-	UTexture2D* OriginalHeightmap;
-	
-    UPROPERTY(Transient)
-	int32 HeightmapsCPUReadBackResourceIndex;
-
-	UPROPERTY(Transient)
-	TArray<ULandscapeComponent*> Components;
-
-	UPROPERTY(Transient)
-	FIntPoint TopLeftSectionBase;
-};
-
-USTRUCT()
-struct FWeightmapLayerData
-{
-	GENERATED_USTRUCT_BODY()
-
-	UPROPERTY()
-	TArray<UTexture2D*> Weightmaps;
-
-	UPROPERTY()
-	TArray<FWeightmapLayerAllocationInfo> WeightmapLayerAllocations;
-
-	UPROPERTY(Transient)
-	TArray<ULandscapeWeightmapUsage*> WeightmapTextureUsages;	// Easy Access ref to data stored into the LandscapeProxy weightmap usage map
-};
-
-USTRUCT()
-struct FLandscapeLayerData
-{
-	GENERATED_USTRUCT_BODY()
-
-	FLandscapeLayerData()
-	{}
-
-	UPROPERTY()
-	TMap<UTexture2D*, UTexture2D*> Heightmaps; // Mapping between Original Heightmap -> Layer Heightmap
-
-	UPROPERTY()
-	TMap<ULandscapeComponent*, FWeightmapLayerData> WeightmapData; // Weightmaps per components
-};
-
 UCLASS(Abstract, MinimalAPI, NotBlueprintable, hidecategories=(Display, Attachment, Physics, Debug, Lighting, LOD), showcategories=(Lighting, Rendering, "Utilities|Transformation"), hidecategories=(Mobility))
 class ALandscapeProxy : public AActor
 {
@@ -644,19 +596,11 @@ public:
 	UPROPERTY()
 	TArray<FLandscapeEditorLayerSettings> EditorLayerSettings;
 
-	UPROPERTY(TextExportTransient)
-	TMap<FGuid, FLandscapeLayerData> LandscapeLayersData;
-
 	UPROPERTY()
 	bool HasLayersContent;
 
-	UPROPERTY(Transient)
-	TMap<UTexture2D*, FRenderDataPerHeightmap> RenderDataPerHeightmap; // Mapping between Original heightmap and general render data
-
-	TArray<TUniquePtr<FLandscapeLayersTexture2DCPUReadBackResource>> HeightmapsCPUReadBackResources;
-
-	TMap<UTexture2D*, FLandscapeLayersTexture2DCPUReadBackResource*> WeightmapCPUReadBackTextures; // Mapping between Original weightmap and tyhe CPU readback resource
-
+	TMap<UTexture2D*, FLandscapeLayersTexture2DCPUReadBackResource*> HeightmapsCPUReadBack;
+	TMap<UTexture2D*, FLandscapeLayersTexture2DCPUReadBackResource*> WeightmapsCPUReadBack;
 	FRenderCommandFence ReleaseResourceFence;
 #endif
 
@@ -768,6 +712,10 @@ public:
 	virtual void RerunConstructionScripts() override {}
 	virtual bool IsLevelBoundsRelevant() const override { return true; }
 
+	virtual void BeginDestroy() override;
+	virtual bool IsReadyForFinishDestroy() override;
+	virtual void FinishDestroy() override;
+
 #if WITH_EDITOR
 	virtual void Destroyed() override;
 	virtual void EditorApplyScale(const FVector& DeltaScale, const FVector* PivotLocation, bool bAltDown, bool bShiftDown, bool bCtrlDown) override;
@@ -837,9 +785,6 @@ public:
 	virtual void Serialize(FArchive& Ar) override;
 	static void AddReferencedObjects(UObject* InThis, FReferenceCollector& Collector);
 	virtual void PostLoad() override;
-	virtual void BeginDestroy() override;
-	virtual bool IsReadyForFinishDestroy() override;
-	virtual void FinishDestroy() override;
 
 #if WITH_EDITOR
 	virtual void PostEditChangeProperty(FPropertyChangedEvent& PropertyChangedEvent) override;
@@ -996,7 +941,6 @@ public:
 
 protected:
 	FLandscapeMaterialChangedDelegate LandscapeMaterialChangedDelegate;
-	
-	LANDSCAPE_API void SetupLayers(int32 InNumComponentsX = INDEX_NONE, int32 InNumComponentsY = INDEX_NONE);
+
 #endif
 };
diff --git a/Engine/Source/Runtime/Landscape/Private/Landscape.cpp b/Engine/Source/Runtime/Landscape/Private/Landscape.cpp
--- a/Engine/Source/Runtime/Landscape/Private/Landscape.cpp
+++ b/Engine/Source/Runtime/Landscape/Private/Landscape.cpp
@@ -169,7 +169,6 @@ ULandscapeComponent::ULandscapeComponent(const FObjectInitializer& ObjectInitial
 
 #if WITH_EDITORONLY_DATA
 	EditToolRenderData = FLandscapeEditToolRenderData();
-	CurrentLayerGuid.Invalidate();
 #endif
 
 	LpvBiasMultiplier = 0.0f; // Bias is 0 for landscape, since it's single sided
@@ -219,46 +218,6 @@ UMaterialInstanceDynamic* ULandscapeComponent::GetMaterialInstanceDynamic(int32
 	return nullptr;
 }
 
-void ULandscapeComponent::AddReferencedObjects(UObject* InThis, FReferenceCollector& Collector)
-{
-	ULandscapeComponent* This = CastChecked<ULandscapeComponent>(InThis);
-
-#if WITH_EDITORONLY_DATA
-	Collector.AddReferencedObject(This->CurrentEditingHeightmapTexture);
-
-	if (This->CurrentEditingWeightmapLayerAllocations != nullptr)
-	{
-		for (FWeightmapLayerAllocationInfo& WeightmapLayerAllocation : *This->CurrentEditingWeightmapLayerAllocations)
-		{
-			Collector.AddReferencedObject(WeightmapLayerAllocation.LayerInfo);
-		}
-	}
-
-	if (This->CurrentEditingWeightmapTexturesUsage != nullptr)
-	{
-		for (ULandscapeWeightmapUsage* TextureUsage : *This->CurrentEditingWeightmapTexturesUsage)
-		{
-			if (TextureUsage != nullptr)
-			{
-				for (int32 i = 0; i < 4; ++i)
-				{
-					Collector.AddReferencedObject(TextureUsage->ChannelUsage[i]);
-				}
-			}
-		}
-	}	
-
-	if (This->CurrentEditingWeightmapTextures != nullptr)
-	{
-		for (UTexture2D* Weightmap : *This->CurrentEditingWeightmapTextures)
-		{
-			Collector.AddReferencedObject(Weightmap);
-		}
-	}
-#endif
-
-	Super::AddReferencedObjects(This, Collector);
-}
 
 #if WITH_EDITOR
 void ULandscapeComponent::BeginCacheForCookedPlatformData(const ITargetPlatform* TargetPlatform)
@@ -994,6 +953,7 @@ ALandscape::ALandscape(const FObjectInitializer& ObjectInitializer)
 	CurrentLayersWeightmapAllMaterialLayersResource = nullptr;
 	WeightmapScratchExtractLayerTextureResource = nullptr;
 	WeightmapScratchPackLayerTextureResource = nullptr;
+	bLandscapeLayersAreInitialized = false;
 #endif // WITH_EDITORONLY_DATA
 }
 
@@ -1083,9 +1043,12 @@ void ULandscapeComponent::GetGeneratedTexturesAndMaterialInstances(TArray<UObjec
 		OutTexturesAndMaterials.Add(HeightmapTexture);
 	}
 
-	if (CurrentEditingHeightmapTexture != nullptr)
+	for (const auto& ItPair : LayersData)
 	{
-		OutTexturesAndMaterials.Add(CurrentEditingHeightmapTexture);
+		const FLandscapeLayerComponentData& LayerComponentData = ItPair.Value;
+
+		OutTexturesAndMaterials.Add(LayerComponentData.HeightmapData.Texture);
+		OutTexturesAndMaterials.Append(LayerComponentData.WeightmapData.Textures);
 	}
 
 	for (UTexture2D* Weightmap : WeightmapTextures)
@@ -1339,14 +1302,16 @@ void ULandscapeComponent::OnUnregister()
 	}
 }
 
-UTexture2D* ULandscapeComponent::GetHeightmap(bool InReturnCurrentEditingHeightmap) const
+UTexture2D* ULandscapeComponent::GetHeightmap(bool InReturnEditingHeightmap) const
 {
 #if WITH_EDITORONLY_DATA
-	if (InReturnCurrentEditingHeightmap && GetMutableDefault<UEditorExperimentalSettings>()->bLandscapeLayerSystem)
+	if (InReturnEditingHeightmap && GetMutableDefault<UEditorExperimentalSettings>()->bLandscapeLayerSystem)
 	{
-		if (CurrentEditingHeightmapTexture != nullptr)
+		const FLandscapeLayerComponentData* EditingLayer = GetEditingLayer();
+
+		if (EditingLayer != nullptr)
 		{
-			return CurrentEditingHeightmapTexture;
+			return EditingLayer->HeightmapData.Texture;
 		}
 	}
 #endif
@@ -1354,14 +1319,16 @@ UTexture2D* ULandscapeComponent::GetHeightmap(bool InReturnCurrentEditingHeightm
 	return HeightmapTexture;
 }
 
-const TArray<UTexture2D*>& ULandscapeComponent::GetWeightmapTextures(bool InReturnCurrentEditingWeightmap) const
+const TArray<UTexture2D*>& ULandscapeComponent::GetWeightmapTextures(bool InReturnEditingWeightmap) const
 {
 #if WITH_EDITORONLY_DATA
-	if (InReturnCurrentEditingWeightmap && GetMutableDefault<UEditorExperimentalSettings>()->bLandscapeLayerSystem)
+	if (InReturnEditingWeightmap && GetMutableDefault<UEditorExperimentalSettings>()->bLandscapeLayerSystem)
 	{
-		if (CurrentEditingWeightmapTextures != nullptr)
+		const FLandscapeLayerComponentData* EditingLayer = GetEditingLayer();
+
+		if (EditingLayer != nullptr)
 		{
-			return *CurrentEditingWeightmapTextures;
+			return EditingLayer->WeightmapData.Textures;
 		}
 	}
 #endif
@@ -1369,14 +1336,16 @@ const TArray<UTexture2D*>& ULandscapeComponent::GetWeightmapTextures(bool InRetu
 	return WeightmapTextures;
 }
 
-TArray<UTexture2D*>& ULandscapeComponent::GetWeightmapTextures(bool InReturnCurrentEditingWeightmap)
+TArray<UTexture2D*>& ULandscapeComponent::GetWeightmapTextures(bool InReturnEditingWeightmap)
 {
 #if WITH_EDITORONLY_DATA
-	if (InReturnCurrentEditingWeightmap && GetMutableDefault<UEditorExperimentalSettings>()->bLandscapeLayerSystem)
+	if (InReturnEditingWeightmap && GetMutableDefault<UEditorExperimentalSettings>()->bLandscapeLayerSystem)
 	{
-		if (CurrentEditingWeightmapTextures != nullptr)
+		FLandscapeLayerComponentData* EditingLayer = GetEditingLayer();
+
+		if (EditingLayer != nullptr)
 		{
-			return *CurrentEditingWeightmapTextures;
+			return EditingLayer->WeightmapData.Textures;
 		}
 	}
 #endif
@@ -1384,14 +1353,16 @@ TArray<UTexture2D*>& ULandscapeComponent::GetWeightmapTextures(bool InReturnCurr
 	return WeightmapTextures;
 }
 
-const TArray<FWeightmapLayerAllocationInfo>& ULandscapeComponent::GetWeightmapLayerAllocations(bool InReturnCurrentEditingWeightmap) const
+const TArray<FWeightmapLayerAllocationInfo>& ULandscapeComponent::GetWeightmapLayerAllocations(bool InReturnEditingWeightmap) const
 {
 #if WITH_EDITORONLY_DATA
-	if (InReturnCurrentEditingWeightmap && GetMutableDefault<UEditorExperimentalSettings>()->bLandscapeLayerSystem)
+	if (InReturnEditingWeightmap && GetMutableDefault<UEditorExperimentalSettings>()->bLandscapeLayerSystem)
 	{
-		if (CurrentEditingWeightmapLayerAllocations != nullptr)
+		const FLandscapeLayerComponentData* EditingLayer = GetEditingLayer();
+
+		if (EditingLayer != nullptr)
 		{
-			return *CurrentEditingWeightmapLayerAllocations;
+			return EditingLayer->WeightmapData.LayerAllocations;
 		}
 	}
 #endif
@@ -1399,14 +1370,16 @@ const TArray<FWeightmapLayerAllocationInfo>& ULandscapeComponent::GetWeightmapLa
 	return WeightmapLayerAllocations;
 }
 
-TArray<FWeightmapLayerAllocationInfo>& ULandscapeComponent::GetWeightmapLayerAllocations(bool InReturnCurrentEditingWeightmap)
+TArray<FWeightmapLayerAllocationInfo>& ULandscapeComponent::GetWeightmapLayerAllocations(bool InReturnEditingWeightmap)
 {
 #if WITH_EDITORONLY_DATA
-	if (InReturnCurrentEditingWeightmap && GetMutableDefault<UEditorExperimentalSettings>()->bLandscapeLayerSystem)
+	if (InReturnEditingWeightmap && GetMutableDefault<UEditorExperimentalSettings>()->bLandscapeLayerSystem)
 	{
-		if (CurrentEditingWeightmapLayerAllocations != nullptr)
+		FLandscapeLayerComponentData* EditingLayer = GetEditingLayer();
+
+		if (EditingLayer != nullptr)
 		{
-			return *CurrentEditingWeightmapLayerAllocations;
+			return EditingLayer->WeightmapData.LayerAllocations;
 		}
 	}
 #endif
@@ -1416,24 +1389,53 @@ TArray<FWeightmapLayerAllocationInfo>& ULandscapeComponent::GetWeightmapLayerAll
 
 #if WITH_EDITOR
 
-void ULandscapeComponent::SetEditingLayer(const FLandscapeLayer* InLayer, FLandscapeLayerData* InLayerData)
+FLandscapeLayerComponentData* ULandscapeComponent::GetEditingLayer()
+{	
+	ALandscape* Landscape = GetLandscapeActor();
+	return Landscape != nullptr ? LayersData.Find(Landscape->GetEditingLayer()) : nullptr;
+}
+
+const FLandscapeLayerComponentData* ULandscapeComponent::GetEditingLayer() const
+{
+	ALandscape* Landscape = GetLandscapeActor();
+	return Landscape != nullptr ? LayersData.Find(Landscape->GetEditingLayer()) : nullptr;
+}
+
+FGuid ULandscapeComponent::GetEditingLayerGUID() const
 {
-	// Update Current Heightmap
-	UTexture2D** LayerHeightmap = InLayerData ? InLayerData->Heightmaps.Find(GetHeightmap()) : nullptr;
-	SetEditingHeightmap(LayerHeightmap ? *LayerHeightmap : nullptr);
+	ALandscape* Landscape = GetLandscapeActor();
+	return Landscape != nullptr ? Landscape->GetEditingLayer() : FGuid();
+}
 
-	// Update Current Weightmaps
-	FWeightmapLayerData* WeightmapData = InLayerData ? InLayerData->WeightmapData.Find(this) : nullptr;
+const FLandscapeLayersGlobalComponentData& ULandscapeComponent::GetGlobalLayersData() const
+{
+	return GlobalLayersData;
+}
 
-	SetEditingWeightmaps(WeightmapData ? &WeightmapData->Weightmaps : nullptr, WeightmapData ? &WeightmapData->WeightmapLayerAllocations : nullptr, WeightmapData ? &WeightmapData->WeightmapTextureUsages : nullptr);
-	SetEditingLayerGuid(WeightmapData && InLayer ? InLayer->Guid : FGuid());
+FLandscapeLayersGlobalComponentData& ULandscapeComponent::GetGlobalLayersData() 
+{
+	return GlobalLayersData;
 }
 
-void ULandscapeComponent::SetEditingHeightmap(UTexture2D* InNewHeightmap)
+const FLandscapeLayerComponentData* ULandscapeComponent::GetLayerData(const FGuid& InLayerGuid) const 
 {
-#if WITH_EDITORONLY_DATA
-	CurrentEditingHeightmapTexture = InNewHeightmap;
-#endif
+	return LayersData.Find(InLayerGuid);
+}
+
+FLandscapeLayerComponentData* ULandscapeComponent::GetLayerData(const FGuid& InLayerGuid)
+{
+	return LayersData.Find(InLayerGuid);
+}
+
+void ULandscapeComponent::AddLayerData(const FGuid& InLayerGuid, const FLandscapeLayerComponentData& InData)
+{
+	FLandscapeLayerComponentData& Data = LayersData.FindOrAdd(InLayerGuid);
+	Data = InData;
+}
+
+void ULandscapeComponent::RemoveLayerData(const FGuid& InLayerGuid)
+{
+	LayersData.Remove(InLayerGuid);
 }
 
 void ULandscapeComponent::SetHeightmap(UTexture2D* NewHeightmap)
@@ -1442,13 +1444,15 @@ void ULandscapeComponent::SetHeightmap(UTexture2D* NewHeightmap)
 	HeightmapTexture = NewHeightmap;
 }
 
-void ULandscapeComponent::SetWeightmapTextures(const TArray<UTexture2D*>& InNewWeightmapTextures, bool InApplyToCurrentEditingWeightmap)
+void ULandscapeComponent::SetWeightmapTextures(const TArray<UTexture2D*>& InNewWeightmapTextures, bool InApplyToEditingWeightmap)
 {
 #if WITH_EDITORONLY_DATA
-	if (InApplyToCurrentEditingWeightmap && CurrentEditingWeightmapTextures != nullptr)
+	FLandscapeLayerComponentData* EditingLayer = GetEditingLayer();
+
+	if (InApplyToEditingWeightmap && EditingLayer != nullptr)
 	{
-		CurrentEditingWeightmapTextures->Reset(InNewWeightmapTextures.Num());
-		CurrentEditingWeightmapTextures->Append(InNewWeightmapTextures);
+		EditingLayer->WeightmapData.Textures.Reset(InNewWeightmapTextures.Num());
+		EditingLayer->WeightmapData.Textures.Append(InNewWeightmapTextures);
 	}
 	else
 #endif
@@ -1457,28 +1461,21 @@ void ULandscapeComponent::SetWeightmapTextures(const TArray<UTexture2D*>& InNewW
 	}
 }
 
-void ULandscapeComponent::SetEditingWeightmaps(TArray<UTexture2D*>* InNewWeightmapTextures, TArray<FWeightmapLayerAllocationInfo>* InNewWeightmapLayerAllocations, TArray<ULandscapeWeightmapUsage*>* InNewWeightmapTexturesUsage)
-{
-#if WITH_EDITORONLY_DATA
-	CurrentEditingWeightmapTextures = InNewWeightmapTextures;
-	CurrentEditingWeightmapLayerAllocations = InNewWeightmapLayerAllocations;
-	CurrentEditingWeightmapTexturesUsage = InNewWeightmapTexturesUsage;
-#endif
-}
-
 void ULandscapeComponent::SetWeightmapLayerAllocations(const TArray<FWeightmapLayerAllocationInfo>& InNewWeightmapLayerAllocations)
 {
 	WeightmapLayerAllocations = InNewWeightmapLayerAllocations;
 }
 
-TArray<ULandscapeWeightmapUsage*>& ULandscapeComponent::GetWeightmapTexturesUsage(bool InReturnCurrentEditingWeightmap)
+TArray<ULandscapeWeightmapUsage*>& ULandscapeComponent::GetWeightmapTexturesUsage(bool InReturnEditingWeightmap)
 {
 #if WITH_EDITORONLY_DATA
-	if (InReturnCurrentEditingWeightmap && GetMutableDefault<UEditorExperimentalSettings>()->bLandscapeLayerSystem)
+	if (InReturnEditingWeightmap && GetMutableDefault<UEditorExperimentalSettings>()->bLandscapeLayerSystem)
 	{
-		if (CurrentEditingWeightmapTexturesUsage != nullptr)
+		FLandscapeLayerComponentData* EditingLayer = GetEditingLayer();
+
+		if (EditingLayer != nullptr)
 		{
-			return *CurrentEditingWeightmapTexturesUsage;
+			return EditingLayer->WeightmapData.TextureUsages;
 		}
 	}
 #endif
@@ -1486,14 +1483,16 @@ TArray<ULandscapeWeightmapUsage*>& ULandscapeComponent::GetWeightmapTexturesUsag
 	return WeightmapTexturesUsage;
 }
 
-const TArray<ULandscapeWeightmapUsage*>& ULandscapeComponent::GetWeightmapTexturesUsage(bool InReturnCurrentEditingWeightmap) const
+const TArray<ULandscapeWeightmapUsage*>& ULandscapeComponent::GetWeightmapTexturesUsage(bool InReturnEditingWeightmap) const
 {
 #if WITH_EDITORONLY_DATA
-	if (InReturnCurrentEditingWeightmap && GetMutableDefault<UEditorExperimentalSettings>()->bLandscapeLayerSystem)
+	if (InReturnEditingWeightmap && GetMutableDefault<UEditorExperimentalSettings>()->bLandscapeLayerSystem)
 	{
-		if (CurrentEditingWeightmapTexturesUsage != nullptr)
+		const FLandscapeLayerComponentData* EditingLayer = GetEditingLayer();
+
+		if (EditingLayer != nullptr)
 		{
-			return *CurrentEditingWeightmapTexturesUsage;
+			return EditingLayer->WeightmapData.TextureUsages;
 		}
 	}
 #endif
@@ -1501,13 +1500,15 @@ const TArray<ULandscapeWeightmapUsage*>& ULandscapeComponent::GetWeightmapTextur
 	return WeightmapTexturesUsage;
 }
 
-void ULandscapeComponent::SetWeightmapTexturesUsage(const TArray<ULandscapeWeightmapUsage*>& InNewWeightmapTexturesUsage, bool InApplyToCurrentEditingWeightmap )
+void ULandscapeComponent::SetWeightmapTexturesUsage(const TArray<ULandscapeWeightmapUsage*>& InNewWeightmapTexturesUsage, bool InApplyToEditingWeightmap )
 {
 #if WITH_EDITORONLY_DATA
-	if (InApplyToCurrentEditingWeightmap && CurrentEditingWeightmapTexturesUsage != nullptr)
+	FLandscapeLayerComponentData* EditingLayer = GetEditingLayer();
+
+	if (InApplyToEditingWeightmap && EditingLayer != nullptr)
 	{
-		CurrentEditingWeightmapTexturesUsage->Reset(InNewWeightmapTexturesUsage.Num());
-		CurrentEditingWeightmapTexturesUsage->Append(InNewWeightmapTexturesUsage);
+		EditingLayer->WeightmapData.TextureUsages.Reset(InNewWeightmapTexturesUsage.Num());
+		EditingLayer->WeightmapData.TextureUsages.Append(InNewWeightmapTexturesUsage);
 	}
 	else
 #endif
@@ -1516,15 +1517,6 @@ void ULandscapeComponent::SetWeightmapTexturesUsage(const TArray<ULandscapeWeigh
 	}
 }
 
-void ULandscapeComponent::SetEditingLayerGuid(const FGuid& InLayerGuid)
-{
-	CurrentLayerGuid = InLayerGuid;
-}
-
-const FGuid& ULandscapeComponent::GetCurrentLayerGuid() const
-{
-	return CurrentLayerGuid;
-}
 #endif
 
 void ALandscapeProxy::PostRegisterAllComponents()
@@ -2067,15 +2059,7 @@ void ALandscapeProxy::PostLoad()
 	FOnFeatureLevelChanged::FDelegate FeatureLevelChangedDelegate = FOnFeatureLevelChanged::FDelegate::CreateUObject(this, &ALandscapeProxy::OnFeatureLevelChanged);
 	FeatureLevelChangedDelegateHandle = GetWorld()->AddOnFeatureLevelChangedHandler(FeatureLevelChangedDelegate);
 
-	if (GetMutableDefault<UEditorExperimentalSettings>()->bLandscapeLayerSystem)
-	{
-		ALandscape* Landscape = GetLandscapeActor();
 
-		if (Landscape != nullptr)
-		{
-			Landscape->RequestLayersContentUpdate(ELandscapeLayersContentUpdateFlag::All_Setup);
-		}
-	}
 #endif
 }
 
diff --git a/Engine/Source/Runtime/Landscape/Private/LandscapeEdit.cpp b/Engine/Source/Runtime/Landscape/Private/LandscapeEdit.cpp
--- a/Engine/Source/Runtime/Landscape/Private/LandscapeEdit.cpp
+++ b/Engine/Source/Runtime/Landscape/Private/LandscapeEdit.cpp
@@ -1723,13 +1723,13 @@ void ULandscapeComponent::GenerateHeightmapMips(TArray<FColor*>& HeightmapTextur
 
 void ULandscapeComponent::CreateEmptyTextureMips(UTexture2D* Texture, bool bClear /*= false*/)
 {
-	ETextureSourceFormat WeightmapFormat = Texture->Source.GetFormat();
-	int32 WeightmapSizeU = Texture->Source.GetSizeX();
-	int32 WeightmapSizeV = Texture->Source.GetSizeY();
+	ETextureSourceFormat Format = Texture->Source.GetFormat();
+	int32 SizeU = Texture->Source.GetSizeX();
+	int32 SizeV = Texture->Source.GetSizeY();
 
 	if (bClear)
 	{
-		Texture->Source.Init2DWithMipChain(WeightmapSizeU, WeightmapSizeV, WeightmapFormat);
+		Texture->Source.Init2DWithMipChain(SizeU, SizeV, Format);
 		int32 NumMips = Texture->Source.GetNumMips();
 		for (int32 MipIndex = 0; MipIndex < NumMips; ++MipIndex)
 		{
@@ -1742,7 +1742,7 @@ void ULandscapeComponent::CreateEmptyTextureMips(UTexture2D* Texture, bool bClea
 	{
 		TArray<uint8> TopMipData;
 		Texture->Source.GetMipData(TopMipData, 0);
-		Texture->Source.Init2DWithMipChain(WeightmapSizeU, WeightmapSizeV, WeightmapFormat);
+		Texture->Source.Init2DWithMipChain(SizeU, SizeV, Format);
 		int32 NumMips = Texture->Source.GetNumMips();
 		uint8* MipData = Texture->Source.LockMip(0);
 		FMemory::Memcpy(MipData, TopMipData.GetData(), TopMipData.Num());
@@ -2855,7 +2855,12 @@ LANDSCAPE_API void ALandscapeProxy::Import(
 
 	if (GetMutableDefault<UEditorExperimentalSettings>()->bLandscapeLayerSystem)
 	{
-		SetupLayers(NumComponentsX, NumComponentsY);
+		ALandscape* LandscapeActor = GetLandscapeActor();
+		check(LandscapeActor != nullptr);
+
+		// Create the default layer
+		LandscapeActor->CreateDefaultLayer(FIntPoint(NumComponentsX, NumComponentsY));
+		HasLayersContent = true;
 	}
 
 	if (GetLevel()->bIsVisible)
@@ -4580,13 +4585,13 @@ void ULandscapeInfo::ClearSelectedRegion(bool bIsComponentwise /*= true*/)
 	}
 }
 
-void ULandscapeComponent::ReallocateWeightmaps(FLandscapeEditDataInterface* DataInterface, bool InCanUseCurrentEditingWeightmap, bool InSaveToTransactionBuffer, bool InInitPlatformDataAsync, TArray<UTexture2D*>* OutNewCreatedTextures)
+void ULandscapeComponent::ReallocateWeightmaps(FLandscapeEditDataInterface* DataInterface, bool InCanUseEditingWeightmap, bool InSaveToTransactionBuffer, bool InInitPlatformDataAsync, TArray<UTexture2D*>* OutNewCreatedTextures)
 {
 	ALandscapeProxy* Proxy = GetLandscapeProxy();
 
-	TArray<FWeightmapLayerAllocationInfo>& ComponentWeightmapLayerAllocations = GetWeightmapLayerAllocations(InCanUseCurrentEditingWeightmap);
-	TArray<UTexture2D*>& ComponentWeightmapTextures = GetWeightmapTextures(InCanUseCurrentEditingWeightmap);
-	TArray<ULandscapeWeightmapUsage*>& ComponentWeightmapTexturesUsage = GetWeightmapTexturesUsage(InCanUseCurrentEditingWeightmap);
+	TArray<FWeightmapLayerAllocationInfo>& ComponentWeightmapLayerAllocations = GetWeightmapLayerAllocations(InCanUseEditingWeightmap);
+	TArray<UTexture2D*>& ComponentWeightmapTextures = GetWeightmapTextures(InCanUseEditingWeightmap);
+	TArray<ULandscapeWeightmapUsage*>& ComponentWeightmapTexturesUsage = GetWeightmapTexturesUsage(InCanUseEditingWeightmap);
 
 	int32 NeededNewChannels = 0;
 	for (int32 LayerIdx = 0; LayerIdx < ComponentWeightmapLayerAllocations.Num(); LayerIdx++)
@@ -4690,6 +4695,7 @@ void ULandscapeComponent::ReallocateWeightmaps(FLandscapeEditDataInterface* Data
 	int32 CurrentLayer = 0;
 	TArray<UTexture2D*> NewWeightmapTextures;
 	TArray<ULandscapeWeightmapUsage*> NewComponentWeightmapTexturesUsage;
+	
 	while (TotalNeededChannels > 0)
 	{
 		// UE_LOG(LogLandscape, Log, TEXT("Still need %d channels"), TotalNeededChannels);
@@ -4707,7 +4713,10 @@ void ULandscapeComponent::ReallocateWeightmaps(FLandscapeEditDataInterface* Data
 			{
 				ULandscapeWeightmapUsage* TryWeightmapUsage = ItPair.Value;
 				//
-				FGuid LayerGuidToSeek = InCanUseCurrentEditingWeightmap ? CurrentLayerGuid : FGuid();
+				FGuid EditingLayerGUID = GetEditingLayerGUID();
+				check(!InCanUseEditingWeightmap || EditingLayerGUID.IsValid());
+
+				FGuid LayerGuidToSeek = InCanUseEditingWeightmap ? EditingLayerGUID : FGuid();
 
 				if (TryWeightmapUsage->FreeChannelCount() >= TotalNeededChannels && TryWeightmapUsage->LayerGuid == LayerGuidToSeek)
 				{
@@ -4774,7 +4783,11 @@ void ULandscapeComponent::ReallocateWeightmaps(FLandscapeEditDataInterface* Data
 			{
 				CurrentWeightmapUsage->Modify();
 			}
-			CurrentWeightmapUsage->LayerGuid = InCanUseCurrentEditingWeightmap ? CurrentLayerGuid : FGuid();
+
+			FGuid EditingLayerGUID = GetEditingLayerGUID();
+			check(!InCanUseEditingWeightmap || EditingLayerGUID.IsValid());
+
+			CurrentWeightmapUsage->LayerGuid = InCanUseEditingWeightmap ? EditingLayerGUID : FGuid();
 			// UE_LOG(LogLandscape, Log, TEXT("Making a new texture %s"), *CurrentWeightmapTexture->GetName());
 		}
 
@@ -4855,8 +4868,8 @@ void ULandscapeComponent::ReallocateWeightmaps(FLandscapeEditDataInterface* Data
 	}
 
 	// Replace the weightmap textures
-	SetWeightmapTextures(MoveTemp(NewWeightmapTextures), InCanUseCurrentEditingWeightmap);
-	SetWeightmapTexturesUsage(MoveTemp(NewComponentWeightmapTexturesUsage), InCanUseCurrentEditingWeightmap);	
+	SetWeightmapTextures(MoveTemp(NewWeightmapTextures), InCanUseEditingWeightmap);
+	SetWeightmapTexturesUsage(MoveTemp(NewComponentWeightmapTexturesUsage), InCanUseEditingWeightmap);	
 }
 
 void ALandscapeProxy::RemoveInvalidWeightmaps()
diff --git a/Engine/Source/Runtime/Landscape/Private/LandscapeEditInterface.cpp b/Engine/Source/Runtime/Landscape/Private/LandscapeEditInterface.cpp
--- a/Engine/Source/Runtime/Landscape/Private/LandscapeEditInterface.cpp
+++ b/Engine/Source/Runtime/Landscape/Private/LandscapeEditInterface.cpp
@@ -1848,12 +1848,12 @@ void ULandscapeComponent::FillLayer(ULandscapeLayerInfoObject* LayerInfo, FLands
 		const int32 WeightmapOffsetY = Component->WeightmapScaleBias.W * (float)SizeV;
 
 		FWeightmapLayerAllocationInfo& FillLayerAllocation = ComponentWeightmapLayerAllocations[FillLayerIdx];
-		uint8* LayerData = (uint8*)LandscapeEdit.GetTextureDataInfo(ComponentWeightmapTextures[FillLayerAllocation.WeightmapTextureIndex])->GetMipData(0);
+		uint8* MaterialLayerData = (uint8*)LandscapeEdit.GetTextureDataInfo(ComponentWeightmapTextures[FillLayerAllocation.WeightmapTextureIndex])->GetMipData(0);
 
 		for (int32 Y = 0; Y < SizeV; ++Y)
 		{
 			const int32 TexY = WeightmapOffsetY + Y;
-			uint8* RowData = LayerData + ((WeightmapOffsetY + Y) * SizeU + WeightmapOffsetX) * 4 + ChannelOffsets[FillLayerAllocation.WeightmapTextureChannel];
+			uint8* RowData = MaterialLayerData + ((WeightmapOffsetY + Y) * SizeU + WeightmapOffsetX) * 4 + ChannelOffsets[FillLayerAllocation.WeightmapTextureChannel];
 			for (int32 X = 0; X < SizeU; ++X)
 			{
 				RowData[X * 4] = 255;
diff --git a/Engine/Source/Runtime/Landscape/Private/LandscapeEditLayers.cpp b/Engine/Source/Runtime/Landscape/Private/LandscapeEditLayers.cpp
--- a/Engine/Source/Runtime/Landscape/Private/LandscapeEditLayers.cpp
+++ b/Engine/Source/Runtime/Landscape/Private/LandscapeEditLayers.cpp
@@ -36,6 +36,7 @@ LandscapeEditLayers.cpp: Landscape editing layers mode
 
 ENGINE_API extern bool GDisableAutomaticTextureMaterialUpdateDependencies;
 
+#if WITH_EDITOR
 static TAutoConsoleVariable<int32> CVarOutputLayersDebugDrawCallName(
 	TEXT("landscape.OutputLayersDebugDrawCallName"),
 	0,
@@ -51,7 +52,7 @@ static TAutoConsoleVariable<int32> CVarOutputLayersWeightmapsRTContent(
 	0,
 	TEXT("This will output the content of render target used for weightmap. This is used for debugging only."));
 
-#if WITH_EDITOR
+
 static TAutoConsoleVariable<int32> CVarLandscapeSimulatePhysics(
 	TEXT("landscape.SimulatePhysics"),
 	0,
@@ -1160,7 +1161,7 @@ public:
 		GraphicsPSOInit.RasterizerState = TStaticRasterizerState<>::GetRHI();
 		GraphicsPSOInit.DepthStencilState = TStaticDepthStencilState<false, CF_Always>::GetRHI();
 
-		FRHIRenderPassInfo RenderPassInfo(ViewFamily.RenderTarget->GetRenderTargetTexture(), CurrentMip == 0 ? ERenderTargetActions::Clear_Store : ERenderTargetActions::Load_Store, nullptr, 0, 0);
+		FRHIRenderPassInfo RenderPassInfo(ViewFamily.RenderTarget->GetRenderTargetTexture(), CurrentMip == 0 || InClearRT ? ERenderTargetActions::Clear_Store : ERenderTargetActions::Load_Store, nullptr, 0, 0);
 		InRHICmdList.BeginRenderPass(RenderPassInfo, TEXT("DrawLayers"));
 
 		if (CurrentMip == 0)
@@ -1229,85 +1230,49 @@ typedef FLandscapeLayersRender_RenderThread<FLandscapeLayersHeightmapShaderParam
 typedef FLandscapeLayersRender_RenderThread<FLandscapeLayersWeightmapShaderParameters, FLandscapeLayersWeightmapPS, FLandscapeLayersWeightmapMipsPS> FLandscapeLayersWeightmapRender_RenderThread;
 
 #if WITH_EDITOR
-void ALandscapeProxy::SetupLayers(int32 InNumComponentsX, int32 InNumComponentsY)
-{
-	ALandscape* Landscape = GetLandscapeActor();
-	check(Landscape);
 
+void ALandscape::CreateLayersRenderingResource(const FIntPoint& InComponentCounts)
+{
 	ULandscapeInfo* Info = GetLandscapeInfo();
+
 	if (Info == nullptr)
 	{
 		return;
 	}
 
-	TArray<ALandscapeProxy*> AllLandscapes;
-	AllLandscapes.Add(Landscape);
-	for (auto& It : Info->Proxies)
+	// Create & Set the CPU Readback to each component global data
+	Info->ForAllLandscapeProxies([&](ALandscapeProxy* Proxy)
 	{
-		AllLandscapes.Add(It);
-	}
-
-	// Make sure we have at least 1 layer
-	if (Landscape->LandscapeLayers.Num() == 0)
-	{
-		Landscape->CreateLayer(FName(TEXT("Layer")), false);
-	}	
+		TArray<UTexture2D*> ComponentsHeightmaps;
 
-	// TODO: When using Change Component Size, we will call Setup again, currently all the existing data will get collapsed into the layer 1, it should keep the layers data.
-
-	int32 NumComponentsX = InNumComponentsX;
-	int32 NumComponentsY = InNumComponentsY;
-	bool GenerateComponentCounts = NumComponentsX == INDEX_NONE || NumComponentsY == INDEX_NONE;
-	FIntPoint MaxSectionBase(0, 0);
-
-	// Setup all Heightmap data
-	for (ALandscapeProxy* LandscapeProxy : AllLandscapes)
-	{
-		for (ULandscapeComponent* Component : LandscapeProxy->LandscapeComponents)
+		for (ULandscapeComponent* Component : Proxy->LandscapeComponents)
 		{
-			UTexture2D* ComponentHeightmapTexture = Component->GetHeightmap();
+			UTexture2D* ComponentHeightmap = Component->GetHeightmap();
 
-			FRenderDataPerHeightmap* Data = LandscapeProxy->RenderDataPerHeightmap.Find(ComponentHeightmapTexture);
+			FLandscapeLayersTexture2DCPUReadBackResource** CPUReadback = Proxy->HeightmapsCPUReadBack.Find(ComponentHeightmap);
 
-			if (Data == nullptr)
-			{
-				FRenderDataPerHeightmap NewData;
-				NewData.Components.Add(Component);
-				NewData.OriginalHeightmap = ComponentHeightmapTexture;
-				FLandscapeLayersTexture2DCPUReadBackResource* CPUReadBackResource = new FLandscapeLayersTexture2DCPUReadBackResource(ComponentHeightmapTexture->Source.GetSizeX(), ComponentHeightmapTexture->Source.GetSizeY(), ComponentHeightmapTexture->GetPixelFormat(), ComponentHeightmapTexture->Source.GetNumMips());
-				NewData.HeightmapsCPUReadBackResourceIndex = LandscapeProxy->HeightmapsCPUReadBackResources.Num();
-				LandscapeProxy->HeightmapsCPUReadBackResources.Add(TUniquePtr<FLandscapeLayersTexture2DCPUReadBackResource>(CPUReadBackResource));
-				BeginInitResource(CPUReadBackResource);
-				LandscapeProxy->RenderDataPerHeightmap.Add(ComponentHeightmapTexture, NewData);
-			}
-			else
+			if (CPUReadback == nullptr)
 			{
-				Data->Components.AddUnique(Component);
-			}
+				FLandscapeLayersTexture2DCPUReadBackResource* NewCPUReadback = new FLandscapeLayersTexture2DCPUReadBackResource(ComponentHeightmap->Source.GetSizeX(), ComponentHeightmap->Source.GetSizeY(), ComponentHeightmap->GetPixelFormat(), ComponentHeightmap->Source.GetNumMips());
+				BeginInitResource(NewCPUReadback);
 
-			if (GenerateComponentCounts)
-			{
-				MaxSectionBase.X = FMath::Max(MaxSectionBase.X, Component->SectionBaseX);
-				MaxSectionBase.Y = FMath::Max(MaxSectionBase.Y, Component->SectionBaseY);
+				Proxy->HeightmapsCPUReadBack.Add(ComponentHeightmap, NewCPUReadback);
 			}
 		}
-	}
+	});	
 
-	if (GenerateComponentCounts)
-	{
-		NumComponentsX = (MaxSectionBase.X / ComponentSizeQuads) + 1;
-		NumComponentsY = (MaxSectionBase.Y / ComponentSizeQuads) + 1;
-	}
+	const int32 TotalVertexCountX = (SubsectionSizeQuads * NumSubsections) * InComponentCounts.X + 1;
+	const int32 TotalVertexCountY = (SubsectionSizeQuads * NumSubsections) * InComponentCounts.Y + 1;
 
-	const int32 TotalVertexCountX = (SubsectionSizeQuads * NumSubsections) * NumComponentsX + 1;
-	const int32 TotalVertexCountY = (SubsectionSizeQuads * NumSubsections) * NumComponentsY + 1;
+	ALandscape* Landscape = GetLandscapeActor();
+	check(Landscape);
 
 	if (Landscape->HeightmapRTList.Num() == 0)
 	{
 		Landscape->HeightmapRTList.Init(nullptr, EHeightmapRTType::HeightmapRT_Count);
 
-		int32 CurrentMipSizeX = ((SubsectionSizeQuads + 1) * NumSubsections) * NumComponentsX;
-		int32 CurrentMipSizeY = ((SubsectionSizeQuads + 1) * NumSubsections) * NumComponentsY;
+		int32 CurrentMipSizeX = ((SubsectionSizeQuads + 1) * NumSubsections) * InComponentCounts.X;
+		int32 CurrentMipSizeY = ((SubsectionSizeQuads + 1) * NumSubsections) * InComponentCounts.Y;
 
 		for (int32 i = 0; i < EHeightmapRTType::HeightmapRT_Count; ++i)
 		{
@@ -1332,216 +1297,19 @@ void ALandscapeProxy::SetupLayers(int32 InNumComponentsX, int32 InNumComponentsY
 			Landscape->HeightmapRTList[i]->UpdateResourceImmediate(true);
 
 			// Only generate required mips RT
-			if (CurrentMipSizeX == NumComponentsX && CurrentMipSizeY == NumComponentsY)
+			if (CurrentMipSizeX == InComponentCounts.X && CurrentMipSizeY == InComponentCounts.Y)
 			{
 				break;
 			}
 		}
 	}
 
-	TArray<FVector> VertexNormals;
-	TArray<uint16> EmptyHeightmapData;
-
-	// Setup all Heightmap data
-	for (ALandscapeProxy* LandscapeProxy : AllLandscapes)
-	{
-		for (auto& ItPair : LandscapeProxy->RenderDataPerHeightmap)
-		{
-			FRenderDataPerHeightmap& HeightmapRenderData = ItPair.Value;
-			HeightmapRenderData.TopLeftSectionBase = FIntPoint(TotalVertexCountX, TotalVertexCountY);
-
-			for (ULandscapeComponent* Component : HeightmapRenderData.Components)
-			{
-				HeightmapRenderData.TopLeftSectionBase.X = FMath::Min(HeightmapRenderData.TopLeftSectionBase.X, Component->GetSectionBase().X);
-				HeightmapRenderData.TopLeftSectionBase.Y = FMath::Min(HeightmapRenderData.TopLeftSectionBase.Y, Component->GetSectionBase().Y);
-			}
-
-			bool FirstLayer = true;
-
-			for (auto& ItLayerDataPair : LandscapeProxy->LandscapeLayersData)
-			{
-				FLandscapeLayerData& LayerData = ItLayerDataPair.Value;
-
-				if (LayerData.Heightmaps.Find(HeightmapRenderData.OriginalHeightmap) == nullptr)
-				{
-					UTexture2D* Heightmap = LandscapeProxy->CreateLandscapeTexture(HeightmapRenderData.OriginalHeightmap->Source.GetSizeX(), HeightmapRenderData.OriginalHeightmap->Source.GetSizeY(), TEXTUREGROUP_Terrain_Heightmap, HeightmapRenderData.OriginalHeightmap->Source.GetFormat());
-					LayerData.Heightmaps.Add(HeightmapRenderData.OriginalHeightmap, Heightmap);
-
-					int32 MipSubsectionSizeQuads = SubsectionSizeQuads;
-					int32 MipSizeU = Heightmap->Source.GetSizeX();
-					int32 MipSizeV = Heightmap->Source.GetSizeY();
-
-					// Copy data from Heightmap to first layer, after that all other layer will get init to empty layer
-					if (FirstLayer)
-					{
-						uint8 MipIndex = 0;
-						TArray<uint8> MipData;
-						MipData.Reserve(MipSizeU*MipSizeV * sizeof(FColor));
-
-						while (MipSizeU > 1 && MipSizeV > 1 && MipSubsectionSizeQuads >= 1)
-						{
-							MipData.Reset();
-							HeightmapRenderData.OriginalHeightmap->Source.GetMipData(MipData, MipIndex);
-
-							FColor* HeightmapTextureData = (FColor*)Heightmap->Source.LockMip(MipIndex);
-							FMemory::Memcpy(HeightmapTextureData, MipData.GetData(), MipData.Num());
-							Heightmap->Source.UnlockMip(MipIndex);
-
-							MipSizeU >>= 1;
-							MipSizeV >>= 1;
-
-							MipSubsectionSizeQuads = ((MipSubsectionSizeQuads + 1) >> 1) - 1;
-							++MipIndex;
-						}
-					}
-					else
-					{
-						TArray<FColor*> HeightmapMipMapData;
-
-						while (MipSizeU > 1 && MipSizeV > 1 && MipSubsectionSizeQuads >= 1)
-						{
-							int32 MipIndex = HeightmapMipMapData.Num();
-							FColor* HeightmapTextureData = (FColor*)Heightmap->Source.LockMip(MipIndex);
-							FMemory::Memzero(HeightmapTextureData, MipSizeU*MipSizeV * sizeof(FColor));
-							HeightmapMipMapData.Add(HeightmapTextureData);
-
-							MipSizeU >>= 1;
-							MipSizeV >>= 1;
-
-							MipSubsectionSizeQuads = ((MipSubsectionSizeQuads + 1) >> 1) - 1;
-						}
-
-						// Initialize blank heightmap data as if ALL components were in the same heightmap to prevent creating many allocations
-						if (EmptyHeightmapData.Num() == 0)
-						{
-							EmptyHeightmapData.Init(32768, TotalVertexCountX * TotalVertexCountY);
-						}
-
-						const FVector DrawScale3D = GetRootComponent()->RelativeScale3D;
-
-						// Init vertex normal data if required
-						if (VertexNormals.Num() == 0)
-						{
-							VertexNormals.AddZeroed(TotalVertexCountX * TotalVertexCountY);
-							for (int32 QuadY = 0; QuadY < TotalVertexCountY - 1; QuadY++)
-							{
-								for (int32 QuadX = 0; QuadX < TotalVertexCountX - 1; QuadX++)
-								{
-									const FVector Vert00 = FVector(0.0f, 0.0f, ((float)EmptyHeightmapData[FMath::Clamp<int32>(QuadY + 0, 0, TotalVertexCountY) * TotalVertexCountX + FMath::Clamp<int32>(QuadX + 0, 0, TotalVertexCountX)] - 32768.0f)*LANDSCAPE_ZSCALE) * DrawScale3D;
-									const FVector Vert01 = FVector(0.0f, 1.0f, ((float)EmptyHeightmapData[FMath::Clamp<int32>(QuadY + 1, 0, TotalVertexCountY) * TotalVertexCountX + FMath::Clamp<int32>(QuadX + 0, 0, TotalVertexCountX)] - 32768.0f)*LANDSCAPE_ZSCALE) * DrawScale3D;
-									const FVector Vert10 = FVector(1.0f, 0.0f, ((float)EmptyHeightmapData[FMath::Clamp<int32>(QuadY + 0, 0, TotalVertexCountY) * TotalVertexCountX + FMath::Clamp<int32>(QuadX + 1, 0, TotalVertexCountX)] - 32768.0f)*LANDSCAPE_ZSCALE) * DrawScale3D;
-									const FVector Vert11 = FVector(1.0f, 1.0f, ((float)EmptyHeightmapData[FMath::Clamp<int32>(QuadY + 1, 0, TotalVertexCountY) * TotalVertexCountX + FMath::Clamp<int32>(QuadX + 1, 0, TotalVertexCountX)] - 32768.0f)*LANDSCAPE_ZSCALE) * DrawScale3D;
-
-									const FVector FaceNormal1 = ((Vert00 - Vert10) ^ (Vert10 - Vert11)).GetSafeNormal();
-									const FVector FaceNormal2 = ((Vert11 - Vert01) ^ (Vert01 - Vert00)).GetSafeNormal();
-
-									// contribute to the vertex normals.
-									VertexNormals[(QuadX + 1 + TotalVertexCountX * (QuadY + 0))] += FaceNormal1;
-									VertexNormals[(QuadX + 0 + TotalVertexCountX * (QuadY + 1))] += FaceNormal2;
-									VertexNormals[(QuadX + 0 + TotalVertexCountX * (QuadY + 0))] += FaceNormal1 + FaceNormal2;
-									VertexNormals[(QuadX + 1 + TotalVertexCountX * (QuadY + 1))] += FaceNormal1 + FaceNormal2;
-								}
-							}
-						}
-
-						for (ULandscapeComponent* Component : HeightmapRenderData.Components)
-						{
-							int32 HeightmapComponentOffsetX = FMath::RoundToInt((float)Heightmap->Source.GetSizeX() * Component->HeightmapScaleBias.Z);
-							int32 HeightmapComponentOffsetY = FMath::RoundToInt((float)Heightmap->Source.GetSizeY() * Component->HeightmapScaleBias.W);
-
-							for (int32 SubsectionY = 0; SubsectionY < NumSubsections; SubsectionY++)
-							{
-								for (int32 SubsectionX = 0; SubsectionX < NumSubsections; SubsectionX++)
-								{
-									for (int32 SubY = 0; SubY <= SubsectionSizeQuads; SubY++)
-									{
-										for (int32 SubX = 0; SubX <= SubsectionSizeQuads; SubX++)
-										{
-											// X/Y of the vertex we're looking at in component's coordinates.
-											const int32 CompX = SubsectionSizeQuads * SubsectionX + SubX;
-											const int32 CompY = SubsectionSizeQuads * SubsectionY + SubY;
-
-											// X/Y of the vertex we're looking indexed into the texture data
-											const int32 TexX = (SubsectionSizeQuads + 1) * SubsectionX + SubX;
-											const int32 TexY = (SubsectionSizeQuads + 1) * SubsectionY + SubY;
-
-											const int32 HeightTexDataIdx = (HeightmapComponentOffsetX + TexX) + (HeightmapComponentOffsetY + TexY) * Heightmap->Source.GetSizeX();
-
-											// copy height and normal data
-											int32 Value = FMath::Clamp<int32>(CompY + Component->GetSectionBase().Y, 0, TotalVertexCountY) * TotalVertexCountX + FMath::Clamp<int32>(CompX + Component->GetSectionBase().X, 0, TotalVertexCountX);
-											const uint16 HeightValue = EmptyHeightmapData[Value];
-											const FVector Normal = VertexNormals[CompX + Component->GetSectionBase().X + TotalVertexCountX * (CompY + Component->GetSectionBase().Y)].GetSafeNormal();
-
-											HeightmapMipMapData[0][HeightTexDataIdx].R = HeightValue >> 8;
-											HeightmapMipMapData[0][HeightTexDataIdx].G = HeightValue & 255;
-											HeightmapMipMapData[0][HeightTexDataIdx].B = FMath::RoundToInt(127.5f * (Normal.X + 1.0f));
-											HeightmapMipMapData[0][HeightTexDataIdx].A = FMath::RoundToInt(127.5f * (Normal.Y + 1.0f));
-										}
-									}
-								}
-							}
-
-							bool IsBorderComponentX = (Component->GetSectionBase().X + 1 * NumSubsections) * InNumComponentsX == TotalVertexCountX;
-							bool IsBorderComponentY = (Component->GetSectionBase().Y + 1 * NumSubsections) * InNumComponentsY == TotalVertexCountY;
-
-							Component->GenerateHeightmapMips(HeightmapMipMapData, IsBorderComponentX ? MAX_int32 : 0, IsBorderComponentY ? MAX_int32 : 0);
-						}
-
-						// Add remaining mips down to 1x1 to heightmap texture.These do not represent quads and are just a simple averages of the previous mipmaps.
-						// These mips are not used for sampling in the vertex shader but could be sampled in the pixel shader.
-						int32 Mip = HeightmapMipMapData.Num();
-						MipSizeU = (Heightmap->Source.GetSizeX()) >> Mip;
-						MipSizeV = (Heightmap->Source.GetSizeY()) >> Mip;
-						while (MipSizeU > 1 && MipSizeV > 1)
-						{
-							HeightmapMipMapData.Add((FColor*)Heightmap->Source.LockMip(Mip));
-							const int32 PrevMipSizeU = (Heightmap->Source.GetSizeX()) >> (Mip - 1);
-							const int32 PrevMipSizeV = (Heightmap->Source.GetSizeY()) >> (Mip - 1);
-
-							for (int32 Y = 0; Y < MipSizeV; Y++)
-							{
-								for (int32 X = 0; X < MipSizeU; X++)
-								{
-									FColor* const TexData = &(HeightmapMipMapData[Mip])[X + Y * MipSizeU];
-
-									const FColor* const PreMipTexData00 = &(HeightmapMipMapData[Mip - 1])[(X * 2 + 0) + (Y * 2 + 0) * PrevMipSizeU];
-									const FColor* const PreMipTexData01 = &(HeightmapMipMapData[Mip - 1])[(X * 2 + 0) + (Y * 2 + 1) * PrevMipSizeU];
-									const FColor* const PreMipTexData10 = &(HeightmapMipMapData[Mip - 1])[(X * 2 + 1) + (Y * 2 + 0) * PrevMipSizeU];
-									const FColor* const PreMipTexData11 = &(HeightmapMipMapData[Mip - 1])[(X * 2 + 1) + (Y * 2 + 1) * PrevMipSizeU];
-
-									TexData->R = (((int32)PreMipTexData00->R + (int32)PreMipTexData01->R + (int32)PreMipTexData10->R + (int32)PreMipTexData11->R) >> 2);
-									TexData->G = (((int32)PreMipTexData00->G + (int32)PreMipTexData01->G + (int32)PreMipTexData10->G + (int32)PreMipTexData11->G) >> 2);
-									TexData->B = (((int32)PreMipTexData00->B + (int32)PreMipTexData01->B + (int32)PreMipTexData10->B + (int32)PreMipTexData11->B) >> 2);
-									TexData->A = (((int32)PreMipTexData00->A + (int32)PreMipTexData01->A + (int32)PreMipTexData10->A + (int32)PreMipTexData11->A) >> 2);
-								}
-							}
-							Mip++;
-							MipSizeU >>= 1;
-							MipSizeV >>= 1;
-						}
-
-						for (int32 i = 0; i < HeightmapMipMapData.Num(); i++)
-						{
-							Heightmap->Source.UnlockMip(i);
-						}
-					}
-
-					Heightmap->BeginCachePlatformData();
-					Heightmap->ClearAllCachedCookedPlatformData();					
-				}
-
-				FirstLayer = false;
-			}
-		}
-	}
-
-	// Weightmaps handling
 	if (Landscape->WeightmapRTList.Num() == 0)
 	{
 		Landscape->WeightmapRTList.Init(nullptr, EWeightmapRTType::WeightmapRT_Count);
 
-		int32 CurrentMipSizeX = ((SubsectionSizeQuads + 1) * NumSubsections) * NumComponentsX;
-		int32 CurrentMipSizeY = ((SubsectionSizeQuads + 1) * NumSubsections) * NumComponentsY;
+		int32 CurrentMipSizeX = ((SubsectionSizeQuads + 1) * NumSubsections) * InComponentCounts.X;
+		int32 CurrentMipSizeY = ((SubsectionSizeQuads + 1) * NumSubsections) * InComponentCounts.Y;
 
 		for (int32 i = 0; i < EWeightmapRTType::WeightmapRT_Count; ++i)
 		{
@@ -1569,159 +1337,226 @@ void ALandscapeProxy::SetupLayers(int32 InNumComponentsX, int32 InNumComponentsY
 			Landscape->WeightmapRTList[i]->UpdateResourceImmediate(true);
 
 			// Only generate required mips RT
-			if (CurrentMipSizeX == NumComponentsX && CurrentMipSizeY == NumComponentsY)
+			if (CurrentMipSizeX == InComponentCounts.X && CurrentMipSizeY == InComponentCounts.Y)
 			{
 				break;
 			}
 		}
 	}
+}
 
-	TArray<ULandscapeComponent*> ComponentsToCleanup;
+FIntPoint ALandscape::ComputeComponentCounts() const
+{
+	ULandscapeInfo* Info = GetLandscapeInfo();
 
-	for (ALandscapeProxy* LandscapeProxy : AllLandscapes)
+	if (Info == nullptr)
 	{
-		bool FirstLayer = true;
+		return FIntPoint(INDEX_NONE, INDEX_NONE);
+	}
 
-		for (auto& ItLayerDataPair : LandscapeProxy->LandscapeLayersData)
-		{
-			FGuid LayerGuid = ItLayerDataPair.Key;
-			FLandscapeLayerData& LayerData = ItLayerDataPair.Value;
+	FIntPoint NumComponents;
+	FIntPoint MaxSectionBase(0, 0);
 
-			struct FTextureData
-			{
-				UTexture2D* Texture;
-				ULandscapeWeightmapUsage* Usage;
-				ULandscapeWeightmapUsage* OriginalUsage;
-			};
+	Info->ForAllLandscapeProxies([&MaxSectionBase](ALandscapeProxy* Proxy)
+	{
+		for (ULandscapeComponent* Component : Proxy->LandscapeComponents)
+		{
+			MaxSectionBase.X = FMath::Max(MaxSectionBase.X, Component->SectionBaseX);
+			MaxSectionBase.Y = FMath::Max(MaxSectionBase.Y, Component->SectionBaseY);
+		}
+	});
 
-			TMap<UTexture2D*, FTextureData> ProcessedTextures;
+	NumComponents.X = (MaxSectionBase.X / ComponentSizeQuads) + 1;
+	NumComponents.Y = (MaxSectionBase.Y / ComponentSizeQuads) + 1;
 
-			for (ULandscapeComponent* Component : LandscapeProxy->LandscapeComponents)
-			{
-				FWeightmapLayerData* WeightmapLayer = LayerData.WeightmapData.Find(Component);
+	return NumComponents;
+}
 
-				if (WeightmapLayer == nullptr)
-				{					
-					FWeightmapLayerData& NewWeightmapData = LayerData.WeightmapData.Add(Component, FWeightmapLayerData());
+void ALandscape::CopyOldDataToDefaultLayer()
+{
+	ULandscapeInfo* Info = GetLandscapeInfo();
 
-					// If no data exist for this weightmap and that data exist in the base weightmap, simply copy it to the first layer, and clear the data in the base (as it will become the final weightmap)
-					if (FirstLayer)
-					{
-						ComponentsToCleanup.Add(Component);
+	if (Info == nullptr)
+	{
+		return;
+	}
 
-						const TArray<UTexture2D*>& ComponentWeightmapTextures = Component->GetWeightmapTextures();
-						TArray<FWeightmapLayerAllocationInfo>& ComponentLayerAllocations = Component->GetWeightmapLayerAllocations();
-						TArray<ULandscapeWeightmapUsage*>& ComponentWeightmapTexturesUsage = Component->GetWeightmapTexturesUsage();
+	FLandscapeLayer* DefaultLayer = GetLayer(0);
+	check(DefaultLayer != nullptr);
 
-						NewWeightmapData.Weightmaps.AddDefaulted(ComponentWeightmapTextures.Num());
-						NewWeightmapData.WeightmapTextureUsages.AddDefaulted(ComponentWeightmapTexturesUsage.Num());
+	Info->ForAllLandscapeProxies([&](ALandscapeProxy* Proxy)
+	{
+		struct FWeightmapTextureData
+		{
+			UTexture2D* Texture;
+			ULandscapeWeightmapUsage* Usage;
+		};
 
-						for (int32 TextureIndex = 0; TextureIndex < ComponentWeightmapTextures.Num(); ++TextureIndex)
-						{
-							UTexture2D* OriginalWeightmapTexture = ComponentWeightmapTextures[TextureIndex];
-							FTextureData* TextureData = ProcessedTextures.Find(OriginalWeightmapTexture);
+		TMap<UTexture2D*, FWeightmapTextureData> ProcessedWeightmaps;
+		TArray<UTexture2D*> ProcessedHeightmaps;
+		TArray<ULandscapeComponent*> WeightmapsComponentsToCleanup;
 
-							if (TextureData != nullptr)
-							{
-								ComponentWeightmapTexturesUsage[TextureIndex] = TextureData->OriginalUsage;
+		for (ULandscapeComponent* Component : Proxy->LandscapeComponents)
+		{
+			FLandscapeLayerComponentData* LayerData = Component->GetLayerData(DefaultLayer->Guid);
 
-								NewWeightmapData.Weightmaps[TextureIndex] = TextureData->Texture;
-								NewWeightmapData.WeightmapTextureUsages[TextureIndex] = TextureData->Usage;
-								check(TextureData->Usage->LayerGuid == LayerGuid);
+			if (ensure(LayerData != nullptr && LayerData->IsInitialized()))
+			{
+				// Heightmap
+				UTexture2D* ComponentHeightmap = Component->GetHeightmap();
 
-								for (int32 ChannelIndex = 0; ChannelIndex < 4; ++ChannelIndex)
-								{
-									ULandscapeComponent* ChannelLandscapeComponent = NewWeightmapData.WeightmapTextureUsages.Last()->ChannelUsage[ChannelIndex];
+				if (!ProcessedHeightmaps.Contains(ComponentHeightmap))
+				{
+					ProcessedHeightmaps.Add(ComponentHeightmap);
 
-									if (ChannelLandscapeComponent != nullptr && ChannelLandscapeComponent == Component)
-									{
-										for (FWeightmapLayerAllocationInfo& Allocation : ComponentLayerAllocations)
-										{
-											if (Allocation.WeightmapTextureIndex == TextureIndex)
-											{
-												NewWeightmapData.WeightmapLayerAllocations.Add(Allocation);
-											}
-										}
+					UTexture* DefaultLayerHeightmap = LayerData->HeightmapData.Texture;
+					check(DefaultLayerHeightmap != nullptr);
 
-										break;
-									}
-								}
-							}
-							else
-							{
-								UTexture2D* NewWeightmapTexture = LandscapeProxy->CreateLandscapeTexture(OriginalWeightmapTexture->Source.GetSizeX(), OriginalWeightmapTexture->Source.GetSizeY(), TEXTUREGROUP_Terrain_Weightmap, OriginalWeightmapTexture->Source.GetFormat());
+					// Only copy Mip0 as other mips will get regenerated
+					TArray<uint8> ExistingMip0Data;
+					ComponentHeightmap->Source.GetMipData(ExistingMip0Data, 0);
 
-								int32 MipSubsectionSizeQuads = SubsectionSizeQuads;
-								int32 MipSizeU = OriginalWeightmapTexture->Source.GetSizeX();
-								int32 MipSizeV = OriginalWeightmapTexture->Source.GetSizeY();
+					FColor* Mip0Data = (FColor*)DefaultLayerHeightmap->Source.LockMip(0);
+					FMemory::Memcpy(Mip0Data, ExistingMip0Data.GetData(), ExistingMip0Data.Num());
+					DefaultLayerHeightmap->Source.UnlockMip(0);
 
-								uint8 MipIndex = 0;
-								TArray<uint8> MipData;
-								MipData.Reserve(MipSizeU * MipSizeV * sizeof(FColor));
+					DefaultLayerHeightmap->BeginCachePlatformData();
+					DefaultLayerHeightmap->ClearAllCachedCookedPlatformData();
+				}
 
-								while (MipSizeU > 1 && MipSizeV > 1 && MipSubsectionSizeQuads >= 1)
-								{
-									MipData.Reset();
-									OriginalWeightmapTexture->Source.GetMipData(MipData, MipIndex);
+				// Weightmaps
+				WeightmapsComponentsToCleanup.Add(Component);
 
-									FColor* WeightmapTextureData = (FColor*)NewWeightmapTexture->Source.LockMip(MipIndex);
-									FMemory::Memcpy(WeightmapTextureData, MipData.GetData(), MipData.Num());
-									NewWeightmapTexture->Source.UnlockMip(MipIndex);
+				const TArray<UTexture2D*>& ComponentWeightmapTextures = Component->GetWeightmapTextures();
+				const TArray<FWeightmapLayerAllocationInfo>& ComponentLayerAllocations = Component->GetWeightmapLayerAllocations();
+				TArray<ULandscapeWeightmapUsage*>& ComponentWeightmapTexturesUsage = Component->GetWeightmapTexturesUsage();
 
-									MipSizeU >>= 1;
-									MipSizeV >>= 1;
+				LayerData->WeightmapData.Textures.AddDefaulted(ComponentWeightmapTextures.Num());
+				LayerData->WeightmapData.TextureUsages.AddDefaulted(ComponentWeightmapTexturesUsage.Num());
 
-									MipSubsectionSizeQuads = ((MipSubsectionSizeQuads + 1) >> 1) - 1;
-									++MipIndex;
-								}
+				for (int32 TextureIndex = 0; TextureIndex < ComponentWeightmapTextures.Num(); ++TextureIndex)
+				{
+					UTexture2D* ComponentWeightmap = ComponentWeightmapTextures[TextureIndex];
+					const FWeightmapTextureData* WeightmapTextureData = ProcessedWeightmaps.Find(ComponentWeightmap);
 
-								NewWeightmapData.Weightmaps[TextureIndex] = NewWeightmapTexture;
-								NewWeightmapData.WeightmapTextureUsages[TextureIndex] = ComponentWeightmapTexturesUsage[TextureIndex];
-								NewWeightmapData.WeightmapTextureUsages[TextureIndex]->LayerGuid = LayerGuid;
+					if (WeightmapTextureData != nullptr)
+					{
+						LayerData->WeightmapData.Textures[TextureIndex] = WeightmapTextureData->Texture;
+						LayerData->WeightmapData.TextureUsages[TextureIndex] = WeightmapTextureData->Usage;
+						check(WeightmapTextureData->Usage->LayerGuid == DefaultLayer->Guid);
 
-								// Create new Usage for the base layer as the other one will now be used by the Layer 1
-								ComponentWeightmapTexturesUsage[TextureIndex] = LandscapeProxy->WeightmapUsageMap.Add(NewWeightmapTexture, LandscapeProxy->CreateWeightmapUsage());
+						for (int32 ChannelIndex = 0; ChannelIndex < 4; ++ChannelIndex)
+						{
+							const ULandscapeComponent* ChannelLandscapeComponent = LayerData->WeightmapData.TextureUsages[TextureIndex]->ChannelUsage[ChannelIndex];
 
-								for (FWeightmapLayerAllocationInfo& Allocation : ComponentLayerAllocations)
+							if (ChannelLandscapeComponent != nullptr && ChannelLandscapeComponent == Component)
+							{
+								for (const FWeightmapLayerAllocationInfo& Allocation : ComponentLayerAllocations)
 								{
 									if (Allocation.WeightmapTextureIndex == TextureIndex)
 									{
-										NewWeightmapData.WeightmapLayerAllocations.Add(Allocation);
+										LayerData->WeightmapData.LayerAllocations.Add(Allocation);
 									}
 								}
 
-								FTextureData NewTextureData;
-								NewTextureData.Texture = NewWeightmapTexture;
-								NewTextureData.Usage = NewWeightmapData.WeightmapTextureUsages[TextureIndex];
-								NewTextureData.OriginalUsage = ComponentWeightmapTexturesUsage[TextureIndex];
+								break;
+							}
+						}
+					}
+					else
+					{
+						UTexture2D* NewLayerWeightmapTexture = Proxy->CreateLandscapeTexture(ComponentWeightmap->Source.GetSizeX(), ComponentWeightmap->Source.GetSizeY(), TEXTUREGROUP_Terrain_Weightmap, ComponentWeightmap->Source.GetFormat());
+
+						// Only copy Mip0 as other mips will get regenerated
+						TArray<uint8> ExistingMip0Data;
+						ComponentWeightmap->Source.GetMipData(ExistingMip0Data, 0);
+
+						FColor* Mip0Data = (FColor*)NewLayerWeightmapTexture->Source.LockMip(0);
+						FMemory::Memcpy(Mip0Data, ExistingMip0Data.GetData(), ExistingMip0Data.Num());
+						NewLayerWeightmapTexture->Source.UnlockMip(0);
+
+						LayerData->WeightmapData.Textures[TextureIndex] = NewLayerWeightmapTexture;
+						LayerData->WeightmapData.TextureUsages[TextureIndex] = Proxy->WeightmapUsageMap.Add(NewLayerWeightmapTexture, Proxy->CreateWeightmapUsage());
 
-								ProcessedTextures.Add(OriginalWeightmapTexture, NewTextureData);
+						for (int32 ChannelIndex = 0; ChannelIndex < 4; ++ChannelIndex)
+						{
+							LayerData->WeightmapData.TextureUsages[TextureIndex]->ChannelUsage[ChannelIndex] = ComponentWeightmapTexturesUsage[TextureIndex]->ChannelUsage[ChannelIndex];
+						}
 
-								NewWeightmapTexture->BeginCachePlatformData();
-								NewWeightmapTexture->ClearAllCachedCookedPlatformData();
-							}							
+						LayerData->WeightmapData.TextureUsages[TextureIndex]->LayerGuid = DefaultLayer->Guid;
+
+						// Create new Usage for the "final" layer as the other one will now be used by the Default layer
+						for (const FWeightmapLayerAllocationInfo& Allocation : ComponentLayerAllocations)
+						{
+							if (Allocation.WeightmapTextureIndex == TextureIndex)
+							{
+								LayerData->WeightmapData.LayerAllocations.Add(Allocation);
+							}
 						}
+
+						FWeightmapTextureData NewTextureData;
+						NewTextureData.Texture = NewLayerWeightmapTexture;
+						NewTextureData.Usage = LayerData->WeightmapData.TextureUsages[TextureIndex];
+
+						ProcessedWeightmaps.Add(ComponentWeightmap, NewTextureData);
+
+						NewLayerWeightmapTexture->BeginCachePlatformData();
+						NewLayerWeightmapTexture->ClearAllCachedCookedPlatformData();
 					}
 				}
-				else
+			}
+		}
+
+		for (ULandscapeComponent* Component : WeightmapsComponentsToCleanup)
+		{
+			TArray<FWeightmapLayerAllocationInfo>& ComponentLayerAllocations = Component->GetWeightmapLayerAllocations();
+
+			for (FWeightmapLayerAllocationInfo& Allocation : ComponentLayerAllocations)
+			{
+				Allocation.WeightmapTextureChannel = 255;
+				Allocation.WeightmapTextureIndex = 255;
+			}
+		}
+	});
+}
+
+void ALandscape::InitializeLayersWeightmapUsage()
+{
+	ULandscapeInfo* Info = GetLandscapeInfo();
+
+	if (Info == nullptr)
+	{
+		return;
+	}
+
+	Info->ForAllLandscapeProxies([&](ALandscapeProxy* Proxy)
+	{
+		for (const FLandscapeLayer& Layer : LandscapeLayers)
+		{
+			for (ULandscapeComponent* Component : Proxy->LandscapeComponents)
+			{
+				// Compute per layer data
+				FLandscapeLayerComponentData* LayerData = Component->GetLayerData(Layer.Guid);
+
+				if (LayerData != nullptr && LayerData->IsInitialized())
 				{
-					WeightmapLayer->WeightmapTextureUsages.AddDefaulted(WeightmapLayer->Weightmaps.Num());
+					LayerData->WeightmapData.TextureUsages.AddDefaulted(LayerData->WeightmapData.Textures.Num());
 
 					// regenerate the weightmap usage
-					for (int32 LayerIdx = 0; LayerIdx < WeightmapLayer->WeightmapLayerAllocations.Num(); LayerIdx++)
+					for (int32 LayerIdx = 0; LayerIdx < LayerData->WeightmapData.LayerAllocations.Num(); LayerIdx++)
 					{
-						FWeightmapLayerAllocationInfo& Allocation = WeightmapLayer->WeightmapLayerAllocations[LayerIdx];
-						UTexture2D* WeightmapTexture = WeightmapLayer->Weightmaps[Allocation.WeightmapTextureIndex];
-						ULandscapeWeightmapUsage** TempUsage = LandscapeProxy->WeightmapUsageMap.Find(WeightmapTexture);
+						FWeightmapLayerAllocationInfo& Allocation = LayerData->WeightmapData.LayerAllocations[LayerIdx];
+						UTexture2D* WeightmapTexture = LayerData->WeightmapData.Textures[Allocation.WeightmapTextureIndex];
+						ULandscapeWeightmapUsage** TempUsage = Proxy->WeightmapUsageMap.Find(WeightmapTexture);
 
 						if (TempUsage == nullptr)
 						{
-							TempUsage = &LandscapeProxy->WeightmapUsageMap.Add(WeightmapTexture, LandscapeProxy->CreateWeightmapUsage());
-							(*TempUsage)->LayerGuid = LayerGuid;
+							TempUsage = &Proxy->WeightmapUsageMap.Add(WeightmapTexture, Proxy->CreateWeightmapUsage());
+							(*TempUsage)->LayerGuid = Layer.Guid;
 						}
 
 						ULandscapeWeightmapUsage* Usage = *TempUsage;
-						WeightmapLayer->WeightmapTextureUsages[Allocation.WeightmapTextureIndex] = Usage; // Keep a ref to it for faster access
+						LayerData->WeightmapData.TextureUsages[Allocation.WeightmapTextureIndex] = Usage; // Keep a ref to it for faster access
 
 						check(Usage->ChannelUsage[Allocation.WeightmapTextureChannel] == nullptr || Usage->ChannelUsage[Allocation.WeightmapTextureChannel] == Component);
 
@@ -1729,61 +1564,131 @@ void ALandscapeProxy::SetupLayers(int32 InNumComponentsX, int32 InNumComponentsY
 					}
 				}
 			}
-
-			FirstLayer = false;
 		}
-	}
+	});
+}
+
+void ALandscape::InitializeLayerWithEmptyContent(const FGuid& InLayerGUID)
+{
+	ULandscapeInfo* Info = GetLandscapeInfo();
 
-	for (ULandscapeComponent* Component : ComponentsToCleanup)
+	if (Info == nullptr)
 	{
-		TArray<FWeightmapLayerAllocationInfo>& ComponentLayerAllocations = Component->GetWeightmapLayerAllocations();
-		ComponentLayerAllocations.Reset();		
-	}
+		return;
+	}	
+
+	FIntPoint NumComponents = ComputeComponentCounts();
+
+	// Build a mapping between each Heightmaps and Component in them
+	TMap<UTexture2D*, TArray<ULandscapeComponent*>> ComponentsPerHeightmaps;
 
-	// Fix Owning actor for Brushes. It can happen after save as operation, for example
-	for (FLandscapeLayer& Layer : Landscape->LandscapeLayers)
+	Info->ForAllLandscapeProxies([&ComponentsPerHeightmaps](ALandscapeProxy* Proxy)
 	{
-		for (int32 i = Layer.Brushes.Num() - 1; i >= 0; --i)
+		for (ULandscapeComponent* Component : Proxy->LandscapeComponents)
 		{
-			FLandscapeLayerBrush& Brush = Layer.Brushes[i];
-
-			if (Brush.BPCustomBrush != nullptr)
-			{
-				if (Brush.BPCustomBrush->GetOwningLandscape() == nullptr)
-				{
-					Brush.BPCustomBrush->SetOwningLandscape(Landscape);
-				}
-			}
+			UTexture2D* ComponentHeightmapTexture = Component->GetHeightmap();
+			TArray<ULandscapeComponent*>& ComponentList = ComponentsPerHeightmaps.FindOrAdd(ComponentHeightmapTexture);
+			ComponentList.Add(Component);
 		}
+	});
+
+	FIntPoint TotalVertexCount((SubsectionSizeQuads * NumSubsections) * NumComponents.X + 1, (SubsectionSizeQuads * NumSubsections) * NumComponents.Y + 1);
+
+	FLandscapeLayer* Layer = LandscapeLayers.FindByPredicate([&InLayerGUID](const FLandscapeLayer& Other) { return Other.Guid == InLayerGUID; });
+
+	// Init layers with valid "empty" data
+	Info->ForAllLandscapeProxies([&](ALandscapeProxy* Proxy)
+	{
+		TMap<UTexture2D*, UTexture2D*> CreatedHeightmapTextures; // < Final layer texture, New created texture for layer
 
-		// TEMP stuff
-		if (Layer.HeightmapBrushOrderIndices.Num() == 0)
+		for (ULandscapeComponent* Component : Proxy->LandscapeComponents)
 		{
-			for (int32 i = 0; i < Layer.Brushes.Num(); ++i)
-			{
-				FLandscapeLayerBrush& Brush = Layer.Brushes[i];
+			UTexture2D* ComponentHeightmap = Component->GetHeightmap();
+			const TArray<ULandscapeComponent*>* ComponentsUsingHeightmap = ComponentsPerHeightmaps.Find(ComponentHeightmap);
+			check(ComponentsUsingHeightmap != nullptr);
 
-				if (Brush.BPCustomBrush != nullptr && Brush.BPCustomBrush->IsAffectingHeightmap())
-				{
-					Layer.HeightmapBrushOrderIndices.Add(i);
-				}
+			// Compute Global layers data
+			FLandscapeLayersGlobalComponentData& GlobalLayersData = Component->GetGlobalLayersData();
+			GlobalLayersData.TopLeftSectionBase = TotalVertexCount;
+
+			for (ULandscapeComponent* ComponentUsingHeightmap : *ComponentsUsingHeightmap)
+			{
+				GlobalLayersData.TopLeftSectionBase.X = FMath::Min(GlobalLayersData.TopLeftSectionBase.X, ComponentUsingHeightmap->GetSectionBase().X);
+				GlobalLayersData.TopLeftSectionBase.Y = FMath::Min(GlobalLayersData.TopLeftSectionBase.Y, ComponentUsingHeightmap->GetSectionBase().Y);
 			}
-		}
 
-		if (Layer.WeightmapBrushOrderIndices.Num() == 0)
-		{
-			for (int32 i = 0; i < Layer.Brushes.Num(); ++i)
+			// Compute per layer data
+			FLandscapeLayerComponentData* LayerData = Component->GetLayerData(Layer->Guid);
+
+			if (LayerData == nullptr || !LayerData->IsInitialized())
 			{
-				FLandscapeLayerBrush& Brush = Layer.Brushes[i];
+				FLandscapeLayerComponentData NewData;
+
+				// Setup Heightmap data
+				UTexture2D** LayerHeightmap = CreatedHeightmapTextures.Find(ComponentHeightmap);
 
-				if (Brush.BPCustomBrush != nullptr && Brush.BPCustomBrush->IsAffectingWeightmap())
+				if (LayerHeightmap == nullptr)
 				{
-					Layer.WeightmapBrushOrderIndices.Add(i);
+					UTexture2D* NewLayerHeightmap = Proxy->CreateLandscapeTexture(ComponentHeightmap->Source.GetSizeX(), ComponentHeightmap->Source.GetSizeY(), TEXTUREGROUP_Terrain_Heightmap, ComponentHeightmap->Source.GetFormat());
+					LayerHeightmap = &CreatedHeightmapTextures.Add(ComponentHeightmap, NewLayerHeightmap);
+
+					ULandscapeComponent::CreateEmptyTextureMips(NewLayerHeightmap, true);
+
+					// Init Mip0 to be at 32768 which is equal to "0"
+					FColor* Mip0Data = (FColor*)NewLayerHeightmap->Source.LockMip(0);
+
+					for (ULandscapeComponent* ComponentUsingHeightmap : *ComponentsUsingHeightmap)
+					{
+						int32 HeightmapComponentOffsetX = FMath::RoundToInt((float)NewLayerHeightmap->Source.GetSizeX() * ComponentUsingHeightmap->HeightmapScaleBias.Z);
+						int32 HeightmapComponentOffsetY = FMath::RoundToInt((float)NewLayerHeightmap->Source.GetSizeY() * ComponentUsingHeightmap->HeightmapScaleBias.W);
+
+						for (int32 SubsectionY = 0; SubsectionY < NumSubsections; SubsectionY++)
+						{
+							for (int32 SubsectionX = 0; SubsectionX < NumSubsections; SubsectionX++)
+							{
+								for (int32 SubY = 0; SubY <= SubsectionSizeQuads; SubY++)
+								{
+									for (int32 SubX = 0; SubX <= SubsectionSizeQuads; SubX++)
+									{
+										// X/Y of the vertex we're looking at in component's coordinates.
+										const int32 CompX = SubsectionSizeQuads * SubsectionX + SubX;
+										const int32 CompY = SubsectionSizeQuads * SubsectionY + SubY;
+
+										// X/Y of the vertex we're looking indexed into the texture data
+										const int32 TexX = (SubsectionSizeQuads + 1) * SubsectionX + SubX;
+										const int32 TexY = (SubsectionSizeQuads + 1) * SubsectionY + SubY;
+
+										const int32 HeightTexDataIdx = (HeightmapComponentOffsetX + TexX) + (HeightmapComponentOffsetY + TexY) * NewLayerHeightmap->Source.GetSizeX();
+
+										// copy height and normal data
+										const uint16 HeightValue = LandscapeDataAccess::GetTexHeight(0.f);
+
+										Mip0Data[HeightTexDataIdx].R = HeightValue >> 8;
+										Mip0Data[HeightTexDataIdx].G = HeightValue & 255;
+
+										// Normal with get calculated later
+										Mip0Data[HeightTexDataIdx].B = 0.0f;
+										Mip0Data[HeightTexDataIdx].A = 0.0f;
+									}
+								}
+							}
+						}
+					}
+
+					NewLayerHeightmap->Source.UnlockMip(0);
+
+					NewLayerHeightmap->BeginCachePlatformData();
+					NewLayerHeightmap->ClearAllCachedCookedPlatformData();
 				}
-			}
-		}		
-		// TEMP stuff
-	}
+
+				NewData.HeightmapData.Texture = *LayerHeightmap;
+
+				// Nothing to do for Weightmap by default
+
+				Component->AddLayerData(Layer->Guid, MoveTemp(NewData));
+			}				
+		}
+	});
 }
 
 void ALandscape::CopyLayersTexture(UTexture* InSourceTexture, UTexture* InDestTexture, FTextureResource* InDestCPUResource, const FIntPoint& InFirstComponentSectionBase, uint8 InSourceCurrentMip, uint8 InDestCurrentMip, uint32 InSourceArrayIndex, uint32 InDestArrayIndex) const
@@ -1809,7 +1714,7 @@ void ALandscape::CopyLayersTexture(const FString& InSourceDebugName, FTextureRes
 	});	
 }
 
-void ALandscape::DrawWeightmapComponentsToRenderTarget(const FString& InDebugName, const FIntPoint& InSectionBase, const FVector2D& InScaleBias, UTexture* InWeightmapRTRead, UTextureRenderTarget2D* InOptionalWeightmapRTRead2, UTextureRenderTarget2D* InWeightmapRTWrite,
+void ALandscape::DrawWeightmapComponentsToRenderTarget(const FString& InDebugName, const TArray<FIntPoint>& InSectionBaseList, const FVector2D& InScaleBias, UTexture* InWeightmapRTRead, UTextureRenderTarget2D* InOptionalWeightmapRTRead2, UTextureRenderTarget2D* InWeightmapRTWrite,
 													   bool InClearRTWrite, FLandscapeLayersWeightmapShaderParameters& InShaderParams, uint8 InMipRender) const
 {
 	check(InWeightmapRTRead != nullptr);
@@ -1827,15 +1732,21 @@ void ALandscape::DrawWeightmapComponentsToRenderTarget(const FString& InDebugNam
 
 	// Quad Setup
 	TArray<FLandscapeLayersTriangle> TriangleList;
-	TriangleList.Reserve(1 * 2 * NumSubsections);
+	TriangleList.Reserve(InSectionBaseList.Num() * 2 * NumSubsections);
 
 	if (InMipRender == 0)
 	{
-		GenerateLayersRenderQuadsAtlas(InSectionBase, InScaleBias, SubsectionSizeQuads, WeightmapReadTextureSize, WeightmapWriteTextureSize, TriangleList);
-	}
+		for (const FIntPoint& SectionBase : InSectionBaseList)
+		{
+			GenerateLayersRenderQuadsAtlas(SectionBase, InScaleBias, SubsectionSizeQuads, WeightmapReadTextureSize, WeightmapWriteTextureSize, TriangleList);
+		}
+	}
 	else
 	{
-		GenerateLayersRenderQuadsMip(InSectionBase, InScaleBias, SubsectionSizeQuads, WeightmapReadTextureSize, WeightmapWriteTextureSize, InMipRender, TriangleList);
+		for (const FIntPoint& SectionBase : InSectionBaseList)
+		{
+			GenerateLayersRenderQuadsMip(SectionBase, InScaleBias, SubsectionSizeQuads, WeightmapReadTextureSize, WeightmapWriteTextureSize, InMipRender, TriangleList);
+		}
 	}
 
 	InShaderParams.ReadWeightmap1 = InWeightmapRTRead;
@@ -1854,9 +1765,9 @@ void ALandscape::DrawWeightmapComponentsToRenderTarget(const FString& InDebugNam
 	FLandscapeLayersWeightmapRender_RenderThread LayersRender(InDebugName, InWeightmapRTWrite, WeightmapWriteTextureSize, WeightmapReadTextureSize, ProjectionMatrix, InShaderParams, InMipRender, TriangleList);
 
 	ENQUEUE_RENDER_COMMAND(FDrawLandscapeLayersWeightmapCommand)(
-		[LayersRender, ClearRT = InClearRTWrite](FRHICommandListImmediate& RHICmdList) mutable
+		[LayersRender, InClearRTWrite](FRHICommandListImmediate& RHICmdList) mutable
 	{
-		LayersRender.Render(RHICmdList, ClearRT);
+		LayersRender.Render(RHICmdList, InClearRTWrite);
 	});
 
 	PrintLayersDebugRT(InDebugName, InWeightmapRTWrite, InMipRender, false);
@@ -1917,15 +1828,15 @@ void ALandscape::DrawWeightmapComponentsToRenderTarget(const FString& InDebugNam
 	FLandscapeLayersWeightmapRender_RenderThread LayersRender(InDebugName, InWeightmapRTWrite, WeightmapWriteTextureSize, WeightmapReadTextureSize, ProjectionMatrix, InShaderParams, InMipRender, TriangleList);
 
 	ENQUEUE_RENDER_COMMAND(FDrawLandscapeLayersWeightmapCommand)(
-		[LayersRender, ClearRT = InClearRTWrite](FRHICommandListImmediate& RHICmdList) mutable
+		[LayersRender, InClearRTWrite](FRHICommandListImmediate& RHICmdList) mutable
 	{
-		LayersRender.Render(RHICmdList, ClearRT);
+		LayersRender.Render(RHICmdList, InClearRTWrite);
 	});
 
 	PrintLayersDebugRT(InDebugName, InWeightmapRTWrite, InMipRender, false);
 }
 
-void ALandscape::DrawWeightmapComponentToRenderTargetMips(const FIntPoint& TopLeftTexturePosition, UTexture* InReadWeightmap, bool InClearRTWrite, struct FLandscapeLayersWeightmapShaderParameters& InShaderParams) const
+void ALandscape::DrawWeightmapComponentToRenderTargetMips(const TArray<FVector2D>& InTexturePositionsToDraw, UTexture* InReadWeightmap, bool InClearRTWrite, struct FLandscapeLayersWeightmapShaderParameters& InShaderParams) const
 {
 	bool OutputDebugName = CVarOutputLayersDebugDrawCallName.GetValueOnAnyThread() == 1 || CVarOutputLayersRTContent.GetValueOnAnyThread() == 1 ? true : false;
 	int32 CurrentMip = 1;
@@ -1935,9 +1846,16 @@ void ALandscape::DrawWeightmapComponentToRenderTargetMips(const FIntPoint& TopLe
 	int32 LocalComponentSizeQuad = SubsectionSizeQuads * NumSubsections;
 	int32 LocalComponentSizeVerts = SubsectionSizeQuads + 1 * NumSubsections;
 
-	FVector2D PositionOffset(FMath::RoundToInt(TopLeftTexturePosition.X / LocalComponentSizeVerts), FMath::RoundToInt(TopLeftTexturePosition.Y / LocalComponentSizeVerts));
-	FIntPoint ComponentSectionBase(PositionOffset.X * LocalComponentSizeQuad, PositionOffset.Y * LocalComponentSizeQuad);
-	FVector2D WeightmapScaleBias(0.0f, 0.0f);
+	TArray<FIntPoint> SectionBaseToDraw;
+	SectionBaseToDraw.Reserve(InTexturePositionsToDraw.Num());
+
+	for (const FVector2D& TexturePosition : InTexturePositionsToDraw)
+	{
+		FVector2D PositionOffset(FMath::RoundToInt(TexturePosition.X / LocalComponentSizeVerts), FMath::RoundToInt(TexturePosition.Y / LocalComponentSizeVerts));
+		SectionBaseToDraw.Add(FIntPoint(PositionOffset.X * LocalComponentSizeQuad, PositionOffset.Y * LocalComponentSizeQuad));
+	}
+
+	FVector2D WeightmapScaleBias(0.0f, 0.0f); // we dont need a scale bias for mip drawing
 
 	for (int32 MipRTIndex = EWeightmapRTType::WeightmapRT_Mip1; MipRTIndex < EWeightmapRTType::WeightmapRT_Count; ++MipRTIndex)
 	{
@@ -1946,14 +1864,14 @@ void ALandscape::DrawWeightmapComponentToRenderTargetMips(const FIntPoint& TopLe
 		if (WriteMipRT != nullptr)
 		{
 			DrawWeightmapComponentsToRenderTarget(OutputDebugName ? FString::Printf(TEXT("LS Weight: %s = -> %s Mips %d"), *ReadMipRT->GetName(), *WriteMipRT->GetName(), CurrentMip) : TEXT(""),
-												  ComponentSectionBase, WeightmapScaleBias, ReadMipRT, nullptr, WriteMipRT, InClearRTWrite, InShaderParams, CurrentMip++);
+												  SectionBaseToDraw, WeightmapScaleBias, ReadMipRT, nullptr, WriteMipRT, InClearRTWrite, InShaderParams, CurrentMip++);
 		}
 
 		ReadMipRT = WeightmapRTList[MipRTIndex];
 	}
 }
 
-void ALandscape::ClearLayersWeightmapTextureResource(const FString& InDebugName, FTextureRenderTargetResource* InTextureResourceToClear)
+void ALandscape::ClearLayersWeightmapTextureResource(const FString& InDebugName, FTextureRenderTargetResource* InTextureResourceToClear) const
 {
 	LandscapeLayersWeightmapClear_RenderThread LayersClear(InDebugName, InTextureResourceToClear);
 
@@ -1964,7 +1882,7 @@ void ALandscape::ClearLayersWeightmapTextureResource(const FString& InDebugName,
 		});
 }
 
-void ALandscape::DrawHeightmapComponentsToRenderTargetMips(TArray<ULandscapeComponent*>& InComponentsToDraw, UTexture* InReadHeightmap, bool InClearRTWrite, struct FLandscapeLayersHeightmapShaderParameters& InShaderParams) const
+void ALandscape::DrawHeightmapComponentsToRenderTargetMips(const TArray<ULandscapeComponent*>& InComponentsToDraw, UTexture* InReadHeightmap, bool InClearRTWrite, struct FLandscapeLayersHeightmapShaderParameters& InShaderParams) const
 {
 	bool OutputDebugName = CVarOutputLayersDebugDrawCallName.GetValueOnAnyThread() == 1 || CVarOutputLayersRTContent.GetValueOnAnyThread() == 1 ? true : false;
 	int32 CurrentMip = 1;
@@ -2075,9 +1993,9 @@ void ALandscape::DrawHeightmapComponentsToRenderTarget(const FString& InDebugNam
 	FLandscapeLayersHeightmapRender_RenderThread LayersRender(InDebugName, InHeightmapRTWrite, HeightmapWriteTextureSize, HeightmapReadTextureSize, ProjectionMatrix, InShaderParams, InMipRender, TriangleList);
 
 	ENQUEUE_RENDER_COMMAND(FDrawLandscapeLayersHeightmapCommand)(
-		[LayersRender, ClearRT = InClearRTWrite](FRHICommandListImmediate& RHICmdList) mutable
+		[LayersRender, InClearRTWrite](FRHICommandListImmediate& RHICmdList) mutable
 		{
-		LayersRender.Render(RHICmdList, ClearRT);
+			LayersRender.Render(RHICmdList, InClearRTWrite);
 		});
 	
 	PrintLayersDebugRT(InDebugName, InHeightmapRTWrite, InMipRender, true, InShaderParams.GenerateNormals);
@@ -2547,20 +2465,41 @@ void ALandscape::PrintLayersDebugTextureResource(const FString& InContext, FText
 	}
 }
 
-bool ALandscape::AreLayersHeightmapTextureResourcesReady(const TArray<ALandscapeProxy*>& InAllLandscapes) const
+bool ALandscape::PrepareLayersHeightmapTextureResources() const
 {
-	for (ALandscapeProxy* Landscape : InAllLandscapes)
+	ULandscapeInfo* Info = GetLandscapeInfo();
+
+	if (Info == nullptr)
+	{
+		return false;
+	}
+
+	bool IsReady = true;
+
+	Info->ForAllLandscapeProxies([&](ALandscapeProxy* Proxy)
 	{
-		for (auto& ItLayerDataPair : Landscape->LandscapeLayersData)
+		if (!IsReady)
 		{
-			for (auto& ItHeightmapPair : ItLayerDataPair.Value.Heightmaps)
+			return;
+		}
+
+		for (ULandscapeComponent* Component : Proxy->LandscapeComponents)
+		{
+			if (!Component->GetHeightmap()->IsFullyStreamedIn())
+			{
+				IsReady = false;
+				return;
+			}
+
+			for (const FLandscapeLayer& Layer : LandscapeLayers)
 			{
-				UTexture2D* OriginalHeightmap = ItHeightmapPair.Key;
-				UTexture2D* LayerHeightmap = ItHeightmapPair.Value;
+				FLandscapeLayerComponentData* ComponentLayerData = Component->GetLayerData(Layer.Guid);
+				UTexture2D* LayerHeightmap = ComponentLayerData->HeightmapData.Texture;
 
-				if (!LayerHeightmap->IsAsyncCacheComplete() || !OriginalHeightmap->IsFullyStreamedIn())
+				if (!LayerHeightmap->IsAsyncCacheComplete())
 				{
-					return false;
+					IsReady = false;
+					return;
 				}
 
 				if (LayerHeightmap->Resource == nullptr)
@@ -2577,13 +2516,14 @@ bool ALandscape::AreLayersHeightmapTextureResourcesReady(const TArray<ALandscape
 
 				if (LayerHeightmap->Resource == nullptr || !LayerHeightmap->Resource->IsInitialized() || !LayerHeightmap->IsFullyStreamedIn())
 				{
-					return false;
+					IsReady = false;
+					return;
 				}
 			}
 		}
-	}
+	});
 
-	return true;
+	return IsReady;
 }
 
 void ALandscape::RegenerateLayersHeightmaps()
@@ -2592,33 +2532,22 @@ void ALandscape::RegenerateLayersHeightmaps()
 
 	ULandscapeInfo* Info = GetLandscapeInfo();
 
-	if (LayersContentUpdateFlags == 0 || Info == nullptr)
-	{
-		return;
-	}
-
-	TArray<ALandscapeProxy*> AllLandscapes;
-	AllLandscapes.Add(this);
-
-	for (auto& It : Info->Proxies)
-	{
-		AllLandscapes.Add(It);
-	}
-
-	if (!AreLayersHeightmapTextureResourcesReady(AllLandscapes))
+	if (LayersContentUpdateFlags == 0 || Info == nullptr || !PrepareLayersHeightmapTextureResources())
 	{
 		return;
 	}
 
 	TArray<ULandscapeComponent*> AllLandscapeComponents;
 
-	for (ALandscapeProxy* Landscape : AllLandscapes)
+	Info->ForAllLandscapeProxies([&AllLandscapeComponents](ALandscapeProxy* Proxy)
 	{
-		AllLandscapeComponents.Append(Landscape->LandscapeComponents);
-	}
+		AllLandscapeComponents.Append(Proxy->LandscapeComponents);
+	});
 
-	if ((LayersContentUpdateFlags & ELandscapeLayersContentUpdateFlag::Heightmap_Render) != 0 && HeightmapRTList.Num() > 0)
+	if ((LayersContentUpdateFlags & ELandscapeLayersContentUpdateFlag::Heightmap_Render) != 0)
 	{
+		check(HeightmapRTList.Num() > 0);
+
 		FLandscapeLayersHeightmapShaderParameters ShaderParams;
 
 		bool FirstLayer = true;
@@ -2650,23 +2579,18 @@ void ALandscape::RegenerateLayersHeightmaps()
 				ShaderParams.LayerAlpha = Layer.HeightmapAlpha;
 			}
 
-			for (ALandscapeProxy* Landscape : AllLandscapes)
+			Info->ForAllLandscapeProxies([&](ALandscapeProxy* Proxy)
 			{
-				FLandscapeLayerData* LayerData = Landscape->LandscapeLayersData.Find(Layer.Guid);
-
-				if (LayerData != nullptr)
+				for (ULandscapeComponent* Component : Proxy->LandscapeComponents)
 				{
-					for (auto& ItPair : LayerData->Heightmaps)
-					{
-						FRenderDataPerHeightmap& HeightmapRenderData = *Landscape->RenderDataPerHeightmap.Find(ItPair.Key);
-						UTexture2D* Heightmap = ItPair.Value;
+					FLandscapeLayerComponentData* ComponentLayerData = Component->GetLayerData(Layer.Guid);
+					FLandscapeLayersGlobalComponentData& GlobalLayersData = Component->GetGlobalLayersData();
+					UTexture2D* LayerHeightmap = ComponentLayerData->HeightmapData.Texture;
 
-						CopyLayersTexture(Heightmap, LandscapeScratchRT1, nullptr, HeightmapRenderData.TopLeftSectionBase);
-
-						PrintLayersDebugRT(OutputDebugName ? FString::Printf(TEXT("LS Height: %s Component %s += -> CombinedAtlas %s"), *Layer.Name.ToString(), *Heightmap->GetName(), *LandscapeScratchRT1->GetName()) : TEXT(""), LandscapeScratchRT1);
-					}
+					CopyLayersTexture(LayerHeightmap, LandscapeScratchRT1, nullptr, GlobalLayersData.TopLeftSectionBase);
+					PrintLayersDebugRT(OutputDebugName ? FString::Printf(TEXT("LS Height: %s Component %s += -> CombinedAtlas %s"), *Layer.Name.ToString(), *LayerHeightmap->GetName(), *LandscapeScratchRT1->GetName()) : TEXT(""), LandscapeScratchRT1);
 				}
-			}
+			});			
 
 			// NOTE: From this point on, we always work in non atlas, we'll convert back at the end to atlas only
 			DrawHeightmapComponentsToRenderTarget(OutputDebugName ? FString::Printf(TEXT("LS Height: %s += -> NonAtlas %s"), *Layer.Name.ToString(), *LandscapeScratchRT1->GetName(), *LandscapeScratchRT2->GetName()) : TEXT(""),
@@ -2742,30 +2666,42 @@ void ALandscape::RegenerateLayersHeightmaps()
 		DrawHeightmapComponentsToRenderTargetMips(AllLandscapeComponents, CombinedHeightmapAtlasRT, true, ShaderParams);
 
 		// Copy back all Mips to original heightmap data
-		for (ALandscapeProxy* Landscape : AllLandscapes)
+		Info->ForAllLandscapeProxies([&](ALandscapeProxy* Proxy)
 		{
-			for (auto& ItPair : Landscape->RenderDataPerHeightmap)
+			TArray<UTexture2D*> ProcessedTextures;
+
+			for (ULandscapeComponent* Component : Proxy->LandscapeComponents)
 			{
-				int32 CurrentMip = 0;
-				FRenderDataPerHeightmap& HeightmapRenderData = ItPair.Value;
-				CopyLayersTexture(CombinedHeightmapAtlasRT, HeightmapRenderData.OriginalHeightmap, Landscape->HeightmapsCPUReadBackResources[HeightmapRenderData.HeightmapsCPUReadBackResourceIndex].Get(), HeightmapRenderData.TopLeftSectionBase, CurrentMip, CurrentMip);
-				++CurrentMip;
+				UTexture2D* ComponentHeightmap = Component->GetHeightmap();
 
-				for (int32 MipRTIndex = EHeightmapRTType::HeightmapRT_Mip1; MipRTIndex < EHeightmapRTType::HeightmapRT_Count; ++MipRTIndex)
+				if (!ProcessedTextures.Contains(ComponentHeightmap))
 				{
-					if (HeightmapRTList[MipRTIndex] != nullptr)
+					int32 CurrentMip = 0;
+
+					FLandscapeLayersTexture2DCPUReadBackResource** CPUReadback = Proxy->HeightmapsCPUReadBack.Find(ComponentHeightmap);
+					check(CPUReadback != nullptr);
+
+					ProcessedTextures.Add(ComponentHeightmap);
+					FLandscapeLayersGlobalComponentData& GlobalLayersData = Component->GetGlobalLayersData();
+					CopyLayersTexture(CombinedHeightmapAtlasRT, ComponentHeightmap, *CPUReadback, GlobalLayersData.TopLeftSectionBase, CurrentMip, CurrentMip);
+					++CurrentMip;
+
+					for (int32 MipRTIndex = EHeightmapRTType::HeightmapRT_Mip1; MipRTIndex < EHeightmapRTType::HeightmapRT_Count; ++MipRTIndex)
 					{
-						CopyLayersTexture(HeightmapRTList[MipRTIndex], HeightmapRenderData.OriginalHeightmap, Landscape->HeightmapsCPUReadBackResources[HeightmapRenderData.HeightmapsCPUReadBackResourceIndex].Get(), HeightmapRenderData.TopLeftSectionBase, CurrentMip, CurrentMip);
-						++CurrentMip;
+						if (HeightmapRTList[MipRTIndex] != nullptr)
+						{
+							CopyLayersTexture(HeightmapRTList[MipRTIndex], ComponentHeightmap, *CPUReadback, GlobalLayersData.TopLeftSectionBase, CurrentMip, CurrentMip);
+							++CurrentMip;
+						}
 					}
 				}
 			}
-		}
+		});
 	}
 
 	if ((LayersContentUpdateFlags & EInternalLandscapeLayersContentUpdateFlag::Internal_Heightmap_ResolveToTexture) != 0)
 	{
-		ResolveLayersHeightmapTexture(AllLandscapes);
+		ResolveLayersHeightmapTexture();
 	}
 
 	bool bUpdateBoundsAndCollsion = (LayersContentUpdateFlags & EInternalLandscapeLayersContentUpdateFlag::Internal_Heightmap_BoundsAndCollision) != 0;
@@ -2814,24 +2750,43 @@ void ALandscape::RegenerateLayersHeightmaps()
 	}
 }
 
-void ALandscape::ResolveLayersHeightmapTexture(const TArray<ALandscapeProxy*>& InAllLandscapes)
+void ALandscape::ResolveLayersHeightmapTexture()
 {
 	SCOPE_CYCLE_COUNTER(STAT_LandscapeLayersResolveHeightmaps);
 
-	for (ALandscapeProxy* Landscape : InAllLandscapes)
+	ULandscapeInfo* Info = GetLandscapeInfo();
+
+	if (Info == nullptr)
+	{
+		return;
+	}
+
+	Info->ForAllLandscapeProxies([&](ALandscapeProxy* Proxy)
 	{
-		TArray<TArray<FColor>> MipData;
+		TArray<UTexture*> ProcessedTexture;
 
-		for (auto& ItPair : Landscape->RenderDataPerHeightmap)
+		for (ULandscapeComponent* Component : Proxy->LandscapeComponents)
 		{
-			FRenderDataPerHeightmap& HeightmapRenderData = ItPair.Value;
-						
-			ResolveLayersTexture(Landscape->HeightmapsCPUReadBackResources[HeightmapRenderData.HeightmapsCPUReadBackResourceIndex].Get(), HeightmapRenderData.OriginalHeightmap);
+			UTexture2D* ComponentHeightmap = Component->GetHeightmap();
+
+			if (!ProcessedTexture.Contains(ComponentHeightmap))
+			{
+				FLandscapeLayersTexture2DCPUReadBackResource** CPUReadback = Proxy->HeightmapsCPUReadBack.Find(ComponentHeightmap);
+
+				if (CPUReadback == nullptr)
+				{
+					continue;
+				}
+
+				ProcessedTexture.Add(ComponentHeightmap);
+
+				ResolveLayersTexture(*CPUReadback, ComponentHeightmap);
+			}
 		}
-	}		
+	});
 }
 
-void ALandscape::ResolveLayersTexture(FLandscapeLayersTexture2DCPUReadBackResource* InCPUReadBackTexture, UTexture2D* InOriginalTexture)
+void ALandscape::ResolveLayersTexture(FLandscapeLayersTexture2DCPUReadBackResource* InCPUReadBackTexture, UTexture2D* InOutputTexture)
 {
 	TArray<TArray<FColor>> MipData;
 	MipData.AddDefaulted(InCPUReadBackTexture->TextureRHI->GetNumMips());
@@ -2854,7 +2809,6 @@ void ALandscape::ResolveLayersTexture(FLandscapeLayersTexture2DCPUReadBackResour
 			RHICmdList.ReadSurfaceData(SourceTextureRHI, Rect, *OutData, ReadFlags);
 		});
 
-
 		MipSizeU >>= 1;
 		MipSizeV >>= 1;
 		++MipIndex;
@@ -2867,47 +2821,49 @@ void ALandscape::ResolveLayersTexture(FLandscapeLayersTexture2DCPUReadBackResour
 	{
 		if (MipData[MipIndex].Num() > 0)
 		{
-			FColor* TextureData = (FColor*)InOriginalTexture->Source.LockMip(MipIndex);
+			FColor* TextureData = (FColor*)InOutputTexture->Source.LockMip(MipIndex);
 			FMemory::Memcpy(TextureData, MipData[MipIndex].GetData(), MipData[MipIndex].Num() * sizeof(FColor));
-			InOriginalTexture->Source.UnlockMip(MipIndex);
+			InOutputTexture->Source.UnlockMip(MipIndex);
 		}
 	}
 }
 
-void ALandscape::PrepareComponentDataToExtractMaterialLayersCS(const FLandscapeLayer& InLayer, int32 InCurrentWeightmapToProcessIndex, bool InOutputDebugName, const TArray<ALandscapeProxy*>& InAllLandscape, FLandscapeTexture2DResource* InOutTextureData,
-																  TArray<FLandscapeLayerWeightmapExtractMaterialLayersComponentData>& OutComponentData, TMap<ULandscapeLayerInfoObject*, int32>& OutLayerInfoObjects)
+void ALandscape::PrepareComponentDataToExtractMaterialLayersCS(const FLandscapeLayer& InLayer, int32 InCurrentWeightmapToProcessIndex, bool InOutputDebugName, FLandscapeTexture2DResource* InOutTextureData,
+																TArray<FLandscapeLayerWeightmapExtractMaterialLayersComponentData>& OutComponentData, TMap<ULandscapeLayerInfoObject*, int32>& OutLayerInfoObjects)
 {
 	ULandscapeInfo* Info = GetLandscapeInfo();
-	
-	for (const ALandscapeProxy* Landscape : InAllLandscape)
+
+	if (Info == nullptr)
 	{
-		const FLandscapeLayerData* LayerData = Landscape->LandscapeLayersData.Find(InLayer.Guid);
+		return;
+	}
 
-		if (LayerData != nullptr)
+	Info->ForAllLandscapeProxies([&](ALandscapeProxy* Proxy)
+	{
+		for (ULandscapeComponent* Component : Proxy->LandscapeComponents)
 		{
-			for (const auto& ItPair : LayerData->WeightmapData)
-			{
-				ULandscapeComponent* Component = ItPair.Key;
-				const FWeightmapLayerData& WeightLayerData = ItPair.Value;
+			FLandscapeLayerComponentData* ComponentLayerData = Component->GetLayerData(InLayer.Guid);
 
-				if (WeightLayerData.Weightmaps.IsValidIndex(InCurrentWeightmapToProcessIndex))
+			if (ComponentLayerData != nullptr)
+			{
+				if (ComponentLayerData->WeightmapData.Textures.IsValidIndex(InCurrentWeightmapToProcessIndex))
 				{
-					UTexture2D* Weightmap = WeightLayerData.Weightmaps[InCurrentWeightmapToProcessIndex];
-					check(Weightmap != nullptr);
+					UTexture2D* LayerWeightmap = ComponentLayerData->WeightmapData.Textures[InCurrentWeightmapToProcessIndex];
+					check(LayerWeightmap != nullptr);
 
-					const ULandscapeWeightmapUsage* WeightmapUsage = WeightLayerData.WeightmapTextureUsages[InCurrentWeightmapToProcessIndex];
-					check(WeightmapUsage != nullptr);
+					const ULandscapeWeightmapUsage* LayerWeightmapUsage = ComponentLayerData->WeightmapData.TextureUsages[InCurrentWeightmapToProcessIndex];
+					check(LayerWeightmapUsage != nullptr);
 
-					CopyLayersTexture(*Weightmap->GetName(), Weightmap->Resource, InOutputDebugName ? FString::Printf(TEXT("%s WeightmapScratchTexture"), *InLayer.Name.ToString()) : TEXT(""), InOutTextureData, nullptr, Component->GetSectionBase(), 0);
+					CopyLayersTexture(*LayerWeightmap->GetName(), LayerWeightmap->Resource, InOutputDebugName ? FString::Printf(TEXT("%s WeightmapScratchTexture"), *InLayer.Name.ToString()) : TEXT(""), InOutTextureData, nullptr, Component->GetSectionBase(), 0);
 					PrintLayersDebugTextureResource(InOutputDebugName ? FString::Printf(TEXT("LS Weight: %s WeightmapScratchTexture %s"), *InLayer.Name.ToString(), TEXT("WeightmapScratchTextureResource")) : TEXT(""), InOutTextureData, 0, false);
 
-					for (const FWeightmapLayerAllocationInfo& WeightmapLayerAllocation : WeightLayerData.WeightmapLayerAllocations)
+					for (const FWeightmapLayerAllocationInfo& WeightmapLayerAllocation : ComponentLayerData->WeightmapData.LayerAllocations)
 					{
-						if (WeightmapLayerAllocation.LayerInfo != nullptr && WeightmapLayerAllocation.WeightmapTextureIndex != 255 && WeightLayerData.Weightmaps[WeightmapLayerAllocation.WeightmapTextureIndex] == Weightmap)
+						if (WeightmapLayerAllocation.LayerInfo != nullptr && WeightmapLayerAllocation.WeightmapTextureIndex != 255 && ComponentLayerData->WeightmapData.Textures[WeightmapLayerAllocation.WeightmapTextureIndex] == LayerWeightmap)
 						{
 							FLandscapeLayerWeightmapExtractMaterialLayersComponentData Data;
 
-							const ULandscapeComponent* DestComponent = WeightmapUsage->ChannelUsage[WeightmapLayerAllocation.WeightmapTextureChannel];
+							const ULandscapeComponent* DestComponent = LayerWeightmapUsage->ChannelUsage[WeightmapLayerAllocation.WeightmapTextureChannel];
 							check(DestComponent);
 
 							// Compute component top left vertex position from section base info
@@ -2949,52 +2905,75 @@ void ALandscape::PrepareComponentDataToExtractMaterialLayersCS(const FLandscapeL
 				}
 			}
 		}
-	}
+	});
 }
 
-void ALandscape::PrepareComponentDataToPackMaterialLayersCS(int32 InCurrentWeightmapToProcessIndex, bool InOutputDebugName, const TArray<ULandscapeComponent*>& InAllLandscapeComponents, TArray<UTexture2D*>& InOutProcessedWeightmaps, 
-															  TArray<FLandscapeLayersTexture2DCPUReadBackResource*>& InOutProcessedWeightmapCPUCopy, TArray<FLandscapeLayerWeightmapPackMaterialLayersComponentData>& OutComponentData)
+void ALandscape::PrepareComponentDataToPackMaterialLayersCS(int32 InCurrentWeightmapToProcessIndex, bool InOutputDebugName, const TArray<ULandscapeComponent*>& InAllLandscapeComponents, TArray<UTexture2D*>& OutProcessedWeightmaps, 
+															TArray<FLandscapeLayersTexture2DCPUReadBackResource*>& OutProcessedCPUReadBackTexture, TArray<FLandscapeLayerWeightmapPackMaterialLayersComponentData>& OutComponentData)
 {
 	ULandscapeInfo* Info = GetLandscapeInfo();
 
-	for (const ULandscapeComponent* Component : InAllLandscapeComponents)
+	if (Info == nullptr)
+	{
+		return;
+	}
+
+	// Compute a mapping of all textures for the asked index and their usage
+	TMap<UTexture2D*, ULandscapeWeightmapUsage*> WeightmapsToProcess;
+
+	for (ULandscapeComponent* Component : InAllLandscapeComponents)
 	{
-		const TArray<UTexture2D*>& WeightmapTextures = Component->GetWeightmapTextures();
+		const TArray<UTexture2D*>& ComponentWeightmapTextures = Component->GetWeightmapTextures();
+		const TArray<ULandscapeWeightmapUsage*>& ComponentWeightmapTexturesUsage = Component->GetWeightmapTexturesUsage();
 
-		if (WeightmapTextures.IsValidIndex(InCurrentWeightmapToProcessIndex))
+		if (ComponentWeightmapTextures.IsValidIndex(InCurrentWeightmapToProcessIndex))
 		{
-			UTexture2D* WeightmapTexture = WeightmapTextures[InCurrentWeightmapToProcessIndex];
+			UTexture2D* ComponentWeightmapTexture = ComponentWeightmapTextures[InCurrentWeightmapToProcessIndex];
+			ULandscapeWeightmapUsage* ComponentWeightmapTextureUsage = ComponentWeightmapTexturesUsage[InCurrentWeightmapToProcessIndex];
+			check(ComponentWeightmapTextureUsage != nullptr);
 
-			if (!InOutProcessedWeightmaps.Contains(WeightmapTexture))
+			if (WeightmapsToProcess.Find(ComponentWeightmapTexture) == nullptr)
 			{
-				InOutProcessedWeightmaps.Add(WeightmapTexture);
+				WeightmapsToProcess.Add(ComponentWeightmapTexture, ComponentWeightmapTextureUsage);
+				OutProcessedWeightmaps.Add(ComponentWeightmapTexture);				
 
-				FLandscapeLayersTexture2DCPUReadBackResource** WeightmapCPUCopy = Component->GetLandscapeProxy()->WeightmapCPUReadBackTextures.Find(WeightmapTexture);
-				check(WeightmapCPUCopy != nullptr);
+				FLandscapeLayersTexture2DCPUReadBackResource** CPUReadback = Component->GetLandscapeProxy()->WeightmapsCPUReadBack.Find(ComponentWeightmapTexture);
+				check(CPUReadback != nullptr);
 
-				InOutProcessedWeightmapCPUCopy.Add(*WeightmapCPUCopy);
+				OutProcessedCPUReadBackTexture.Add(*CPUReadback);
+			}
+		}
+	}
 
-				const TArray<ULandscapeWeightmapUsage*>& WeightmapTexturesUsage = Component->GetWeightmapTexturesUsage();
+	TArray<const FWeightmapLayerAllocationInfo*> AlreadyProcessedAllocation;
 
-				const ULandscapeWeightmapUsage* WeightmapUsage = WeightmapTexturesUsage[InCurrentWeightmapToProcessIndex];
-				check(WeightmapUsage != nullptr);
+	// Build for each texture what each channel should contains
+	for (auto& ItPair : WeightmapsToProcess)
+	{
+		FLandscapeLayerWeightmapPackMaterialLayersComponentData Data;
 
-				TArray<const FWeightmapLayerAllocationInfo*> AlreadyProcessedAllocation;
-				FLandscapeLayerWeightmapPackMaterialLayersComponentData Data;
+		for (int32 WeightmapChannelIndex = 0; WeightmapChannelIndex < 4; ++WeightmapChannelIndex)
+		{
+			UTexture2D* ComponentWeightmapTexture = ItPair.Key;
+			ULandscapeWeightmapUsage* ComponentWeightmapTextureUsage = ItPair.Value;
 
-				for (int32 WeightmapChannelIndex = 0; WeightmapChannelIndex < 4; ++WeightmapChannelIndex)
-				{
-					// Clear out data to known values
-					Data.ComponentVertexPositionX[WeightmapChannelIndex] = INDEX_NONE;
-					Data.ComponentVertexPositionY[WeightmapChannelIndex] = INDEX_NONE;
-					Data.SourcePaintLayerIndex[WeightmapChannelIndex] = INDEX_NONE;
-					Data.WeightmapChannelToProcess[WeightmapChannelIndex] = INDEX_NONE;
+			// Clear out data to known values
+			Data.ComponentVertexPositionX[WeightmapChannelIndex] = INDEX_NONE;
+			Data.ComponentVertexPositionY[WeightmapChannelIndex] = INDEX_NONE;
+			Data.SourcePaintLayerIndex[WeightmapChannelIndex] = INDEX_NONE;
+			Data.WeightmapChannelToProcess[WeightmapChannelIndex] = INDEX_NONE;
 
-					if (WeightmapUsage->ChannelUsage[WeightmapChannelIndex] != nullptr)
-					{
-						const ULandscapeComponent* ChannelComponent = WeightmapUsage->ChannelUsage[WeightmapChannelIndex];
-						FIntPoint ComponentSectionBase = ChannelComponent->GetSectionBase();
+			if (ComponentWeightmapTextureUsage->ChannelUsage[WeightmapChannelIndex] != nullptr)
+			{
+				const ULandscapeComponent* ChannelComponent = ComponentWeightmapTextureUsage->ChannelUsage[WeightmapChannelIndex];
+
+				const TArray<FWeightmapLayerAllocationInfo>& ChannelLayerAllocations = ChannelComponent->GetWeightmapLayerAllocations();
+				const TArray<UTexture2D*>& ChannelComponentWeightmapTextures = ChannelComponent->GetWeightmapTextures();
 
+				for (const FWeightmapLayerAllocationInfo& ChannelLayerAllocation : ChannelLayerAllocations)
+				{
+					if (ChannelLayerAllocation.LayerInfo != nullptr && !AlreadyProcessedAllocation.Contains(&ChannelLayerAllocation) && ChannelComponentWeightmapTextures[ChannelLayerAllocation.WeightmapTextureIndex] == ComponentWeightmapTexture)
+					{
 						// Compute component top left vertex position from section base info
 						int32 LocalComponentSizeQuad = ChannelComponent->SubsectionSizeQuads * NumSubsections;
 						int32 LocalComponentSizeVerts = (ChannelComponent->SubsectionSizeQuads + 1) * NumSubsections;
@@ -3003,56 +2982,55 @@ void ALandscape::PrepareComponentDataToPackMaterialLayersCS(int32 InCurrentWeigh
 						Data.ComponentVertexPositionX[WeightmapChannelIndex] = PositionOffset.X * LocalComponentSizeVerts;
 						Data.ComponentVertexPositionY[WeightmapChannelIndex] = PositionOffset.Y * LocalComponentSizeVerts;
 
-						const TArray<FWeightmapLayerAllocationInfo>& ChannelLayerAllocations = ChannelComponent->GetWeightmapLayerAllocations();
-						const TArray<UTexture2D*>& ChannelComponentWeightmapTextures = ChannelComponent->GetWeightmapTextures();
-						
-						for (const FWeightmapLayerAllocationInfo& ChannelLayerAllocation : ChannelLayerAllocations)
+						Data.WeightmapChannelToProcess[WeightmapChannelIndex] = ChannelLayerAllocation.WeightmapTextureChannel;
+
+						AlreadyProcessedAllocation.Add(&ChannelLayerAllocation);
+
+						if (ChannelLayerAllocation.LayerInfo == ALandscapeProxy::VisibilityLayer)
 						{
-							if (ChannelLayerAllocation.LayerInfo != nullptr && !AlreadyProcessedAllocation.Contains(&ChannelLayerAllocation) && ChannelComponentWeightmapTextures[ChannelLayerAllocation.WeightmapTextureIndex] == WeightmapTexture)
+							Data.SourcePaintLayerIndex[WeightmapChannelIndex] = 0; // Always store after the last weightmap index
+						}
+						else
+						{
+							for (int32 LayerInfoSettingsIndex = 0; LayerInfoSettingsIndex < Info->Layers.Num(); ++LayerInfoSettingsIndex)
 							{
-								Data.WeightmapChannelToProcess[WeightmapChannelIndex] = ChannelLayerAllocation.WeightmapTextureChannel;
-								AlreadyProcessedAllocation.Add(&ChannelLayerAllocation);
+								const FLandscapeInfoLayerSettings& LayerInfo = Info->Layers[LayerInfoSettingsIndex];
 
-								if (ChannelLayerAllocation.LayerInfo == ALandscapeProxy::VisibilityLayer)
-								{
-									Data.SourcePaintLayerIndex[WeightmapChannelIndex] = 0; // Always store after the last weightmap index
-								}
-								else
+								if (ChannelLayerAllocation.LayerInfo == LayerInfo.LayerInfoObj)
 								{
-									for (int32 LayerInfoSettingsIndex = 0; LayerInfoSettingsIndex < Info->Layers.Num(); ++LayerInfoSettingsIndex)
-									{
-										const FLandscapeInfoLayerSettings& LayerInfo = Info->Layers[LayerInfoSettingsIndex];
-
-										if (ChannelLayerAllocation.LayerInfo == LayerInfo.LayerInfoObj)
-										{
-											Data.SourcePaintLayerIndex[WeightmapChannelIndex] = LayerInfoSettingsIndex + 1; // due to visibility layer that is at 0
-											break;
-										}
-									}
+									Data.SourcePaintLayerIndex[WeightmapChannelIndex] = LayerInfoSettingsIndex + 1; // due to visibility layer that is at 0
+									break;
 								}
-
-								break;
 							}
 						}
+
+						break;
 					}
 				}
-
-				OutComponentData.Add(Data);
 			}
 		}
+
+		OutComponentData.Add(Data);
 	}
 }
 
-void ALandscape::ReallocateLayersWeightmaps(const TArray<ALandscapeProxy*>& InAllLandscape, const TArray<ULandscapeLayerInfoObject*>& InBrushRequiredAllocations, TArray<ULandscapeComponent*>& OutComponentThatNeedMaterialRebuild)
+void ALandscape::ReallocateLayersWeightmaps(const TArray<ULandscapeLayerInfoObject*>& InBrushRequiredAllocations, TArray<ULandscapeComponent*>& OutComponentThatNeedMaterialRebuild)
 {
 	SCOPE_CYCLE_COUNTER(STAT_LandscapeLayersReallocateWeightmaps);
 
+	ULandscapeInfo* Info = GetLandscapeInfo();
+
+	if (Info == nullptr)
+	{
+		return;
+	}	
+
 	TArray<ULandscapeComponent*> AllLandscapeComponents;
 
-	for (ALandscapeProxy* Landscape : InAllLandscape)
+	Info->ForAllLandscapeProxies([&AllLandscapeComponents](ALandscapeProxy* Proxy)
 	{
-		AllLandscapeComponents.Append(Landscape->LandscapeComponents);
-	}
+		AllLandscapeComponents.Append(Proxy->LandscapeComponents);
+	});
 
 	// Copy Previous Usage, to know which texture need updating
 	TMap<UTexture2D*, ULandscapeWeightmapUsage*> CurrentWeightmapsUsage;
@@ -3109,17 +3087,12 @@ void ALandscape::ReallocateLayersWeightmaps(const TArray<ALandscapeProxy*>& InAl
 	// Build a map of all the allocation per components
 	TMap<ULandscapeComponent*, TArray<ULandscapeLayerInfoObject*>> LayerAllocsPerComponent;
 
-	for (const ALandscapeProxy* Landscape : InAllLandscape)
+	Info->ForAllLandscapeProxies([&](ALandscapeProxy* Proxy)
 	{
-		for (auto& ItLayerPair : Landscape->LandscapeLayersData)
+		for (FLandscapeLayer& Layer : LandscapeLayers)
 		{
-			const FLandscapeLayerData& LayerData = ItLayerPair.Value;
-
-			for (const auto& ItWeightmapPair : LayerData.WeightmapData)
+			for (ULandscapeComponent* Component : Proxy->LandscapeComponents)
 			{
-				ULandscapeComponent* Component = ItWeightmapPair.Key;
-				const FWeightmapLayerData& WeightLayerData = ItWeightmapPair.Value;
-
 				TArray<ULandscapeLayerInfoObject*>* ComponentLayerAlloc = LayerAllocsPerComponent.Find(Component);
 
 				if (ComponentLayerAlloc == nullptr)
@@ -3128,7 +3101,9 @@ void ALandscape::ReallocateLayersWeightmaps(const TArray<ALandscapeProxy*>& InAl
 					ComponentLayerAlloc = &LayerAllocsPerComponent.Add(Component, NewLayerAllocs);
 				}
 
-				for (const FWeightmapLayerAllocationInfo& LayerWeightmapAllocation : WeightLayerData.WeightmapLayerAllocations)
+				FLandscapeLayerComponentData* LayerComponentData = Component->GetLayerData(Layer.Guid);
+
+				for (const FWeightmapLayerAllocationInfo& LayerWeightmapAllocation : LayerComponentData->WeightmapData.LayerAllocations)
 				{
 					if (LayerWeightmapAllocation.LayerInfo != nullptr)
 					{
@@ -3146,7 +3121,7 @@ void ALandscape::ReallocateLayersWeightmaps(const TArray<ALandscapeProxy*>& InAl
 				}
 			}
 		}
-	}
+	});	
 
 	// Determine if the Final layer need to add/remove some alloc
 	for (auto& ItPair : LayerAllocsPerComponent)
@@ -3200,7 +3175,13 @@ void ALandscape::ReallocateLayersWeightmaps(const TArray<ALandscapeProxy*>& InAl
 	//GDisableAutomaticTextureMaterialUpdateDependencies = false;
 
 	// Determine which Component need updating
-	for (ULandscapeComponent* Component : AllLandscapeComponents)
+
+	// BEGIN HACK for now we disable the logic that will determine what need updating as it's buggy, and i prefer a correct behavior than an optimized non working logic.
+	OutComponentThatNeedMaterialRebuild.Reset();
+	OutComponentThatNeedMaterialRebuild.Append(AllLandscapeComponents);	
+	// END HACK
+
+	/*for (ULandscapeComponent* Component : AllLandscapeComponents)
 	{
 		if (OutComponentThatNeedMaterialRebuild.Contains(Component))
 		{
@@ -3238,12 +3219,11 @@ void ALandscape::ReallocateLayersWeightmaps(const TArray<ALandscapeProxy*>& InAl
 			}
 		}
 	}
+	*/
 }
 
 void ALandscape::InitLayersWeightmapResources(uint8 InLayerCount)
 {
-	ULandscapeInfo* Info = GetLandscapeInfo();
-
 	// Use the 1st one to compute the resource as they are all the same anyway
 	UTextureRenderTarget2D* FirstWeightmapRT = WeightmapRTList[WeightmapRT_Scratch1];
 
@@ -3294,243 +3274,117 @@ void ALandscape::InitLayersWeightmapResources(uint8 InLayerCount)
 	}
 }
 
-/* This code was removed as it's a fix to remove weightmap with 0 weightmap for a component, unfortunatly it will cause side effect, like perf issue, etc, so for now it's disabled until we can address this later
-bool ALandscape::GenerateZeroAllocationPerComponents(const TArray<ALandscapeProxy*>& InAllLandscape, const TMap<ULandscapeLayerInfoObject*, bool>& InWeightmapLayersBlendSubstractive, TMap<ULandscapeComponent*, TArray<ULandscapeLayerInfoObject*>>& OutZeroAllocationsPerComponents)
+bool ALandscape::PrepareLayersWeightmapTextureResources() const
 {
-	// Make sure all CPU Readback texture are ready to be used
-	for (ALandscapeProxy* Landscape : InAllLandscape)
-	{
-		for (auto& ItPair : Landscape->WeightmapCPUReadBackTextures)
-		{
-			FLandscapeLayersTexture2DCPUReadBackResource* CPuReadBackTexture = ItPair.Value;
-
-			if (CPuReadBackTexture != nullptr && !CPuReadBackTexture->IsInitialized())
-			{
-				return false;
-			}
-		}
-	}
-
-	TArray<ULandscapeComponent*> AllLandscapeComponents;
+	ULandscapeInfo* Info = GetLandscapeInfo();
 
-	for (ALandscapeProxy* Landscape : InAllLandscape)
+	if (Info == nullptr)
 	{
-		AllLandscapeComponents.Append(Landscape->LandscapeComponents);
+		return false;
 	}
 
-	// Compute the data to facilitate the process
-	struct FTextureChannelData
-	{
-		TArray<ULandscapeComponent*, TInlineAllocator<4>> ComponentPerChannel;
-		TArray<ULandscapeLayerInfoObject*, TInlineAllocator<4>> LayerInfoPerChannel;
-		TArray<FColor> TexelData;
-		FLandscapeLayersTexture2DCPUReadBackResource* CPUReadbackTexture;
-	};
-
-	TMap<UTexture2D*, FTextureChannelData> TextureDataPerTextures;
+	bool IsReady = true;
 
-	for (ULandscapeComponent* LandscapeComponent : AllLandscapeComponents)
+	// Make sure all our original weightmap textures are streamed in and ready to be used
+	Info->ForAllLandscapeProxies([&](ALandscapeProxy* Proxy)
 	{
-		const TArray<FLandscapeWeightmapUsage*>& ComponentWeightmapUsage = LandscapeComponent->GetWeightmapTexturesUsage();
-		const TArray<UTexture2D*>& ComponentWeightmapTextures = LandscapeComponent->GetWeightmapTextures();
-
-		for (int32 TextureIndex = 0; TextureIndex < ComponentWeightmapTextures.Num(); ++TextureIndex)
+		for (ULandscapeComponent* Component : Proxy->LandscapeComponents)
 		{
-			UTexture2D* WeightmapTexture = ComponentWeightmapTextures[TextureIndex];
-
-			if (TextureDataPerTextures.Contains(WeightmapTexture))
-			{
-				continue;
-			}
-
-			FLandscapeWeightmapUsage* Usage = ComponentWeightmapUsage[TextureIndex];
-
-			for (int32 i = 0; i < 4; ++i)
+			for (UTexture2D* ComponentWeightmap : Component->GetWeightmapTextures())
 			{
-				ULandscapeComponent* Component = Usage->ChannelUsage[i];
-
-				if (Component != nullptr)
+				if (!ComponentWeightmap->IsFullyStreamedIn())
 				{
-					const TArray<FWeightmapLayerAllocationInfo>& ComponentUsageWeightmapAllocations = Component->GetWeightmapLayerAllocations();
-
-					for (const FWeightmapLayerAllocationInfo& Allocation : ComponentUsageWeightmapAllocations)
-					{
-						if (Allocation.WeightmapTextureIndex == TextureIndex)
-						{
-							const bool* IsSubtractive = InWeightmapLayersBlendSubstractive.Find(Allocation.LayerInfo);
-
-							if (IsSubtractive != nullptr && *IsSubtractive)
-							{
-								FTextureChannelData* TextureData = TextureDataPerTextures.Find(WeightmapTexture);
-
-								if (TextureData == nullptr)
-								{
-									FTextureChannelData NewData;
-									NewData.ComponentPerChannel.SetNumZeroed(4);
-									NewData.LayerInfoPerChannel.SetNumZeroed(4);
-
-									FLandscapeLayersTexture2DCPUReadBackResource** CPUReadbackTexture = Component->GetLandscapeProxy()->WeightmapCPUReadBackTextures.Find(WeightmapTexture);
-									check(CPUReadbackTexture);
-
-									NewData.CPUReadbackTexture = *CPUReadbackTexture;
-
-									TextureData = &TextureDataPerTextures.Add(WeightmapTexture, NewData);
-								}
-
-								TextureData->LayerInfoPerChannel[Allocation.WeightmapTextureChannel] = Allocation.LayerInfo;
-								TextureData->ComponentPerChannel[Allocation.WeightmapTextureChannel] = Component;
-							}
-						}
-					}
+					IsReady = false;
+					return;
 				}
 			}
-		}
-	}
-
-	// Read the data form the CPU texture
-	if (TextureDataPerTextures.Num() > 0)
-	{
-		for (auto& ItPair : TextureDataPerTextures)
-		{
-			FTextureChannelData& TextureChannelData = ItPair.Value;
-
-			int32 MipSizeU = TextureChannelData.CPUReadbackTexture->GetSizeX();
-			int32 MipSizeV = TextureChannelData.CPUReadbackTexture->GetSizeY();
-
-			FReadSurfaceDataFlags Flags(RCM_UNorm, CubeFace_MAX);
-			Flags.SetMip(0);
-			FIntRect Rect(0, 0, MipSizeU, MipSizeV);
 
-			ENQUEUE_UNIQUE_RENDER_COMMAND_FOURPARAMETER(
-				ReadSurfaceCommand,
-				FTextureRHIRef, SourceTextureRHI, TextureChannelData.CPUReadbackTexture->TextureRHI,
-				FIntRect, Rect, Rect,
-				TArray<FColor>*, OutData, &TextureChannelData.TexelData,
-				FReadSurfaceDataFlags, ReadFlags, Flags,
-				{
-					RHICmdList.ReadSurfaceData(SourceTextureRHI, Rect, *OutData, ReadFlags);
-				});
-		}
-
-		FlushRenderingCommands();
-
-		// Determine which texture channel is zero allocation to "remove" them
-		for (auto& ItPair : TextureDataPerTextures)
-		{
-			FTextureChannelData& TextureChannelData = ItPair.Value;
-
-			TArray<bool> AreLayerEmpty;
-			AreLayerEmpty.Init(true, 4);
-
-			for (FColor& Color : TextureChannelData.TexelData)
+			for (const FLandscapeLayer& Layer : LandscapeLayers)
 			{
-				if (Color.R != 0)
-				{
-					AreLayerEmpty[0] = false;
-				}
+				FLandscapeLayerComponentData* ComponentLayerData = Component->GetLayerData(Layer.Guid);
 
-				if (Color.G != 0)
+				if (ComponentLayerData == nullptr)
 				{
-					AreLayerEmpty[1] = false;
+					IsReady = false;
+					return;
 				}
 
-				if (Color.B != 0)
+				for (UTexture2D* LayerWeightmap : ComponentLayerData->WeightmapData.Textures)
 				{
-					AreLayerEmpty[2] = false;
-				}
-
-				if (Color.A != 0)
-				{
-					AreLayerEmpty[3] = false;
-				}
-			}
-			
-			for (int32 i = 0; i < 4; ++i)
-			{
-				if (AreLayerEmpty[i] && TextureChannelData.ComponentPerChannel[0] != nullptr)
-				{
-					TArray<ULandscapeLayerInfoObject*>& ZeroAllocations = OutZeroAllocationsPerComponents.FindOrAdd(TextureChannelData.ComponentPerChannel[0]);
-					ZeroAllocations.Add(TextureChannelData.LayerInfoPerChannel[0]);
-				}
-			}
-		}
-	}
-
-	return true;
-}
-*/
-
-bool ALandscape::AreLayersWeightmapTextureResourcesReady(const TArray<ALandscapeProxy*>& InAllLandscapes) const
-{
-	// Make sure all our original weightmap textures are streamed in and ready to be used
-	for (const ALandscapeProxy* Landscape : InAllLandscapes)
-	{
-		for (const auto& ItLayerDataPair : Landscape->LandscapeLayersData)
-		{
-			for (const auto& ItWeightmapPair : ItLayerDataPair.Value.WeightmapData)
-			{
-				ULandscapeComponent* Component = ItWeightmapPair.Key;
-				const TArray<UTexture2D*>& OriginalWeightmaps = Component->GetWeightmapTextures();
-
-				for (UTexture2D* Weightmap : OriginalWeightmaps)
-				{
-					if (!Weightmap->IsFullyStreamedIn())
+					if (LayerWeightmap != nullptr && !LayerWeightmap->IsFullyStreamedIn())
 					{
-						return false;
+						IsReady = false;
+						return;
 					}
 				}
 			}
 		}
-	}
+	});
 
 	// Init all needed resources
-	for (const ALandscapeProxy* Landscape : InAllLandscapes)
+	if (IsReady)
 	{
-		for (const auto& ItLayerDataPair : Landscape->LandscapeLayersData)
+		Info->ForAllLandscapeProxies([&](ALandscapeProxy* Proxy)
 		{
-			for (const auto& ItWeightmapPair : ItLayerDataPair.Value.WeightmapData)
+			for (const FLandscapeLayer& Layer : LandscapeLayers)
 			{
-				const FWeightmapLayerData& WeightmapLayerData = ItWeightmapPair.Value;
-
-				for (UTexture2D* Weightmap : WeightmapLayerData.Weightmaps)
+				for (ULandscapeComponent* Component : Proxy->LandscapeComponents)
 				{
-					if (Weightmap->Resource == nullptr)
-					{
-						Weightmap->FinishCachePlatformData();
+					FLandscapeLayerComponentData* ComponentLayerData = Component->GetLayerData(Layer.Guid);
 
-						Weightmap->Resource = Weightmap->CreateResource();
+					if (ComponentLayerData == nullptr)
+					{
+						IsReady = false;
+						return;
+					}
 
-						if (Weightmap->Resource != nullptr)
+					for (UTexture2D* LayerWeightmap : ComponentLayerData->WeightmapData.Textures)
+					{
+						if (LayerWeightmap->Resource == nullptr)
 						{
-							BeginInitResource(Weightmap->Resource);
+							LayerWeightmap->FinishCachePlatformData();
+
+							LayerWeightmap->Resource = LayerWeightmap->CreateResource();
+
+							if (LayerWeightmap->Resource != nullptr)
+							{
+								BeginInitResource(LayerWeightmap->Resource);
+							}
 						}
 					}
 				}
 			}
-		}
-	}
+		});
 
-	// Wait for the new resource to be fully initialized/streamed in
-	for (const ALandscapeProxy* Landscape : InAllLandscapes)
-	{
-		for (const auto& ItLayerDataPair : Landscape->LandscapeLayersData)
+		// Wait for the new resource to be fully initialized/streamed in
+		Info->ForAllLandscapeProxies([&](ALandscapeProxy* Proxy)
 		{
-			for (const auto& ItWeightmapPair : ItLayerDataPair.Value.WeightmapData)
+			for (const FLandscapeLayer& Layer : LandscapeLayers)
 			{
-				const FWeightmapLayerData& WeightmapLayerData = ItWeightmapPair.Value;
-
-				for (UTexture2D* Weightmap : WeightmapLayerData.Weightmaps)
+				for (ULandscapeComponent* Component : Proxy->LandscapeComponents)
 				{
-					if (Weightmap->Resource == nullptr || !Weightmap->Resource->IsInitialized() || !Weightmap->IsFullyStreamedIn())
+					FLandscapeLayerComponentData* ComponentLayerData = Component->GetLayerData(Layer.Guid);
+
+					if (ComponentLayerData == nullptr)
+					{
+						IsReady = false;
+						return;
+					}
+
+					for (UTexture2D* LayerWeightmap : ComponentLayerData->WeightmapData.Textures)
 					{
-						return false;
+						if (LayerWeightmap->Resource == nullptr || !LayerWeightmap->Resource->IsInitialized() || !LayerWeightmap->IsFullyStreamedIn())
+						{
+							IsReady = false;
+							return;
+						}
 					}
 				}
 			}
-		}
+		});
 	}
 
-	return true;
+	return IsReady;
 }
 
 void ALandscape::RegenerateLayersWeightmaps()
@@ -3539,38 +3393,27 @@ void ALandscape::RegenerateLayersWeightmaps()
 
 	ULandscapeInfo* Info = GetLandscapeInfo();
 
-	if (LayersContentUpdateFlags == 0 || Info == nullptr || Info->Layers.Num() == 0)
-	{
-		return;
-	}
-
-	TArray<ALandscapeProxy*> AllLandscapes;
-	AllLandscapes.Add(this);
-
-	for (const auto& It : Info->Proxies)
-	{
-		AllLandscapes.Add(It);
-	}
-
-	if (!AreLayersWeightmapTextureResourcesReady(AllLandscapes))
+	if (LayersContentUpdateFlags == 0 || Info == nullptr || Info->Layers.Num() == 0 || !PrepareLayersWeightmapTextureResources())
 	{
 		return;
 	}
 
 	TArray<ULandscapeComponent*> AllLandscapeComponents;
 
-	for (ALandscapeProxy* Landscape : AllLandscapes)
+	Info->ForAllLandscapeProxies([&AllLandscapeComponents](ALandscapeProxy* Proxy)
 	{
-		AllLandscapeComponents.Append(Landscape->LandscapeComponents);
-	}
+		AllLandscapeComponents.Append(Proxy->LandscapeComponents);
+	});
 
 	TArray<ULandscapeComponent*> ComponentThatNeedMaterialRebuild;
 	TArray<ULandscapeLayerInfoObject*> BrushRequiredAllocations;
 	int32 LayerCount = Info->Layers.Num() + 1; // due to visibility being stored at 0
 	bool ClearFlagsAfterUpdate = true;
 
-	if ((LayersContentUpdateFlags & ELandscapeLayersContentUpdateFlag::Weightmap_Render) != 0 && WeightmapRTList.Num() > 0)
+	if ((LayersContentUpdateFlags & ELandscapeLayersContentUpdateFlag::Weightmap_Render) != 0)
 	{
+		check(WeightmapRTList.Num() > 0);
+
 		UTextureRenderTarget2D* LandscapeScratchRT1 = WeightmapRTList[EWeightmapRTType::WeightmapRT_Scratch1];
 		UTextureRenderTarget2D* LandscapeScratchRT2 = WeightmapRTList[EWeightmapRTType::WeightmapRT_Scratch2];
 		UTextureRenderTarget2D* LandscapeScratchRT3 = WeightmapRTList[EWeightmapRTType::WeightmapRT_Scratch3];
@@ -3630,7 +3473,7 @@ void ALandscape::RegenerateLayersWeightmaps()
 
 				// Prepare compute shader data
 				TArray<FLandscapeLayerWeightmapExtractMaterialLayersComponentData> ComponentsData;	
-				PrepareComponentDataToExtractMaterialLayersCS(Layer, CurrentWeightmapToProcessIndex, OutputDebugName, AllLandscapes, WeightmapScratchExtractLayerTextureResource, ComponentsData, LayerInfoObjects);
+				PrepareComponentDataToExtractMaterialLayersCS(Layer, CurrentWeightmapToProcessIndex, OutputDebugName, WeightmapScratchExtractLayerTextureResource, ComponentsData, LayerInfoObjects);
 
 				HasFoundWeightmapToProcess = ComponentsData.Num() > 0;
 
@@ -3769,48 +3612,45 @@ void ALandscape::RegenerateLayersWeightmaps()
 		}
 
 		// TODO:  if editing a Brush affecting layers, since we don't have any bounds to brush, right now ReallocateLayersWeightmaps wont ask a rebuild of the component affected by Brushes, which mean ComponentThatNeedMaterialRebuild wont contains Brush affected component!
-		ReallocateLayersWeightmaps(AllLandscapes, BrushRequiredAllocations, ComponentThatNeedMaterialRebuild);
-
-		// Allocation that will need to be excluded when we update materials
-		TMap<ULandscapeComponent*, TArray<ULandscapeLayerInfoObject*>> ZeroAllocationsPerComponents;
+		ReallocateLayersWeightmaps(BrushRequiredAllocations, ComponentThatNeedMaterialRebuild);
 
 		if (ComputeShaderGeneratedData)
 		{
 			// Will generate CPU read back resource, if required
-			for (ALandscapeProxy* LandscapeProxy : AllLandscapes)
+			Info->ForAllLandscapeProxies([&](ALandscapeProxy* Proxy)
 			{
-				for (const ULandscapeComponent* Component : LandscapeProxy->LandscapeComponents)
+				for (ULandscapeComponent* Component : Proxy->LandscapeComponents)
 				{
 					const TArray<UTexture2D*>& ComponentWeightmapTextures = Component->GetWeightmapTextures();
 
 					for (UTexture2D* WeightmapTexture : ComponentWeightmapTextures)
 					{
-						FLandscapeLayersTexture2DCPUReadBackResource** WeightmapCPUReadBack = LandscapeProxy->WeightmapCPUReadBackTextures.Find(WeightmapTexture);
+						FLandscapeLayersTexture2DCPUReadBackResource** CPUReadback = Component->GetLandscapeProxy()->WeightmapsCPUReadBack.Find(WeightmapTexture);
 
-						if (WeightmapCPUReadBack == nullptr)
+						if (CPUReadback == nullptr)
 						{
 							FLandscapeLayersTexture2DCPUReadBackResource* NewWeightmapCPUReadBack = new FLandscapeLayersTexture2DCPUReadBackResource(WeightmapTexture->Source.GetSizeX(), WeightmapTexture->Source.GetSizeY(), WeightmapTexture->GetPixelFormat(), WeightmapTexture->Source.GetNumMips());
 							BeginInitResource(NewWeightmapCPUReadBack);
 
-							LandscapeProxy->WeightmapCPUReadBackTextures.Add(WeightmapTexture, NewWeightmapCPUReadBack);
+							Component->GetLandscapeProxy()->WeightmapsCPUReadBack.Add(WeightmapTexture, NewWeightmapCPUReadBack);
 						}
 					}
 				}
-			}
+			});			
 
 			int8 CurrentWeightmapToProcessIndex = 0;
 			bool HasFoundWeightmapToProcess = true; // try processing at least once
 
 			TArray<float> WeightmapLayerWeightBlend;
 			TArray<UTexture2D*> ProcessedWeightmaps;
-			TArray<FLandscapeLayersTexture2DCPUReadBackResource*> ProcessedWeightmapsCPUCopy;
+			TArray<FLandscapeLayersTexture2DCPUReadBackResource*> ProcessedCPUReadbackTextures;
 			int32 NextTextureIndexToProcess = 0;
 
 			// Generate the component data from the weightmap allocation that were done earlier and weight blend them if required (i.e renormalize)
 			while (HasFoundWeightmapToProcess)
 			{
 				TArray<FLandscapeLayerWeightmapPackMaterialLayersComponentData> PackLayersComponentsData;
-				PrepareComponentDataToPackMaterialLayersCS(CurrentWeightmapToProcessIndex, OutputDebugName, AllLandscapeComponents, ProcessedWeightmaps, ProcessedWeightmapsCPUCopy, PackLayersComponentsData);
+				PrepareComponentDataToPackMaterialLayersCS(CurrentWeightmapToProcessIndex, OutputDebugName, AllLandscapeComponents, ProcessedWeightmaps, ProcessedCPUReadbackTextures, PackLayersComponentsData);
 				HasFoundWeightmapToProcess = PackLayersComponentsData.Num() > 0;
 
 				// Perform the compute shader
@@ -3869,18 +3709,17 @@ void ALandscape::RegenerateLayersWeightmaps()
 						CSDispatch.PackLayers(RHICmdList);
 					});
 
+					UTextureRenderTarget2D* CurrentRT = WeightmapRTList[WeightmapRT_Mip0];
+					CopyLayersTexture(TEXT("WeightmapScratchTexture"), WeightmapScratchPackLayerTextureResource, *CurrentRT->GetName(), CurrentRT->GameThread_GetRenderTargetResource());
+					DrawWeightmapComponentToRenderTargetMips(WeightmapTextureOutputOffset, CurrentRT, true, PSShaderParams);
+
 					int32 StartTextureIndex = NextTextureIndexToProcess;
 
 					for (; NextTextureIndexToProcess < ProcessedWeightmaps.Num(); ++NextTextureIndexToProcess)
 					{
 						UTexture2D* WeightmapTexture = ProcessedWeightmaps[NextTextureIndexToProcess];
-						FLandscapeLayersTexture2DCPUReadBackResource* WeightmapCPUReadBack = ProcessedWeightmapsCPUCopy[NextTextureIndexToProcess];
-						FIntPoint TextureTopLeftPositionInAtlas(WeightmapTextureOutputOffset[NextTextureIndexToProcess - StartTextureIndex].X, WeightmapTextureOutputOffset[NextTextureIndexToProcess - StartTextureIndex].Y);
-
-						UTextureRenderTarget2D* CurrentRT = WeightmapRTList[WeightmapRT_Mip0];
-						CopyLayersTexture(TEXT("WeightmapScratchTexture"), WeightmapScratchPackLayerTextureResource, *CurrentRT->GetName(), CurrentRT->GameThread_GetRenderTargetResource());
-
-						DrawWeightmapComponentToRenderTargetMips(TextureTopLeftPositionInAtlas, CurrentRT, true, PSShaderParams);
+						FLandscapeLayersTexture2DCPUReadBackResource* WeightmapCPUReadBack = ProcessedCPUReadbackTextures[NextTextureIndexToProcess];
+						FIntPoint TextureTopLeftPositionInAtlas(WeightmapTextureOutputOffset[NextTextureIndexToProcess - StartTextureIndex].X, WeightmapTextureOutputOffset[NextTextureIndexToProcess - StartTextureIndex].Y);						
 
 						int32 CurrentMip = 0;
 
@@ -3899,22 +3738,15 @@ void ALandscape::RegenerateLayersWeightmaps()
 
 				++CurrentWeightmapToProcessIndex;
 			}
-
-			// Generate all the allocation with zero painted data for each components, it will be used to update the material instance to not sample those textures
-			/*if (ComputeShaderGeneratedData)
-			{
-				ClearFlagsAfterUpdate = GenerateZeroAllocationPerComponents(AllLandscapes, WeightmapLayersBlendSubstractive, ZeroAllocationsPerComponents);
-			}
-			*/
 		}
 
-		UpdateLayersMaterialInstances(LayersUpdateAllMaterials ? AllLandscapeComponents : ComponentThatNeedMaterialRebuild, ZeroAllocationsPerComponents);
-		LayersUpdateAllMaterials = false;
+		UpdateLayersMaterialInstances(LayersUpdateAllMaterials ? AllLandscapeComponents : ComponentThatNeedMaterialRebuild);
+		LayersUpdateAllMaterials = false;		
 	}
 
 	if ((LayersContentUpdateFlags & EInternalLandscapeLayersContentUpdateFlag::Internal_Weightmap_ResolveToTexture))
 	{
-		ResolveLayersWeightmapTexture(AllLandscapes);
+		ResolveLayersWeightmapTexture();
 	}
 
 	if ((LayersContentUpdateFlags & EInternalLandscapeLayersContentUpdateFlag::Internal_Weightmap_Collision) != 0)
@@ -3937,15 +3769,14 @@ void ALandscape::RegenerateLayersWeightmaps()
 	}
 }
 
-void ALandscape::UpdateLayersMaterialInstances(const TArray<ULandscapeComponent*>& InComponentsToUpdate, const TMap<ULandscapeComponent*, TArray<ULandscapeLayerInfoObject*>>& InZeroAllocationsPerComponents)
+void ALandscape::UpdateLayersMaterialInstances(const TArray<ULandscapeComponent*>& InComponentsToUpdate)
 {
-	if (InComponentsToUpdate.Num() == 0 && InZeroAllocationsPerComponents.Num() == 0)
+	if (InComponentsToUpdate.Num() == 0)
 	{
 		return;
 	}
 
 	TArray<ULandscapeComponent*> ComponentsToUpdate;
-	//InZeroAllocationsPerComponents.GenerateKeyArray(ComponentsToUpdate);
 	ComponentsToUpdate.Append(InComponentsToUpdate);	
 
 	SCOPE_CYCLE_COUNTER(STAT_LandscapeLayersUpdateMaterialInstance);
@@ -4004,22 +3835,6 @@ void ALandscape::UpdateLayersMaterialInstances(const TArray<ULandscapeComponent*
 		int8 MaterialIndex = 0;
 
 		TArray<FWeightmapLayerAllocationInfo> WeightmapBaseLayerAllocation = Component->GetWeightmapLayerAllocations(); // We copy the array here
-		/*const TArray<ULandscapeLayerInfoObject*>* ZeroAllocations = InZeroAllocationsPerComponents.Find(Component);
-
-		// We have some allocation to remove
-		if (ZeroAllocations != nullptr)
-		{
-			for (int32 i = WeightmapBaseLayerAllocation.Num() - 1; i >= 0; --i)
-			{
-				const FWeightmapLayerAllocationInfo& Allocation = WeightmapBaseLayerAllocation[i];
-
-				if (ZeroAllocations->Contains(Allocation.LayerInfo))
-				{
-					WeightmapBaseLayerAllocation.RemoveAt(i);
-				}
-			}
-		}
-		*/
 
 		TArray<UTexture2D*>& ComponentWeightmapTextures = Component->GetWeightmapTextures();
 		UTexture2D* Heightmap = Component->GetHeightmap();
@@ -4112,27 +3927,43 @@ void ALandscape::UpdateLayersMaterialInstances(const TArray<ULandscapeComponent*
 	RecreateRenderStateContexts.Empty();
 }
 
-void ALandscape::ResolveLayersWeightmapTexture(const TArray<ALandscapeProxy*>& InAllLandscapes)
+void ALandscape::ResolveLayersWeightmapTexture()
 {
 	SCOPE_CYCLE_COUNTER(STAT_LandscapeLayersResolveWeightmaps);
-	       
-	for (ALandscapeProxy* Landscape : InAllLandscapes)
+
+	ULandscapeInfo* Info = GetLandscapeInfo();
+
+	if (Info == nullptr)
+	{
+		return;
+	}
+
+	Info->ForAllLandscapeProxies([&](ALandscapeProxy* Proxy)
 	{
-		TArray<TArray<FColor>> MipData;
+		TArray<UTexture2D*> ProcessedWeightmaps;
 
-		for (auto& ItPair : Landscape->WeightmapCPUReadBackTextures)
+		for (ULandscapeComponent* Component : Proxy->LandscapeComponents)
 		{
-			UTexture2D* OriginalWeightmap = ItPair.Key;
-			FLandscapeLayersTexture2DCPUReadBackResource* WeightmapsCPUReadBack = ItPair.Value;
+			TArray<UTexture2D*>& ComponentWeightmaps = Component->GetWeightmapTextures();
 
-			if (WeightmapsCPUReadBack == nullptr)
+			for (UTexture2D* ComponentLayerWeightmap : ComponentWeightmaps)			
 			{
-				continue;
-			}
+				if (!ProcessedWeightmaps.Contains(ComponentLayerWeightmap))
+				{
+					FLandscapeLayersTexture2DCPUReadBackResource** CPUReadback = Proxy->WeightmapsCPUReadBack.Find(ComponentLayerWeightmap);
 
-			ResolveLayersTexture(WeightmapsCPUReadBack, OriginalWeightmap);
+					if (CPUReadback == nullptr)
+					{
+						continue;
+					}
+
+					ProcessedWeightmaps.Add(ComponentLayerWeightmap);
+
+					ResolveLayersTexture(*CPUReadback, ComponentLayerWeightmap);
+				}
+			}
 		}
-	}
+	});
 }
 
 void ALandscape::RequestLayersContentUpdate(ELandscapeLayersContentUpdateFlag InDataFlags, bool InUpdateAllMaterials)
@@ -4143,16 +3974,29 @@ void ALandscape::RequestLayersContentUpdate(ELandscapeLayersContentUpdateFlag In
 
 void ALandscape::RegenerateLayersContent()
 {
-	if ((LayersContentUpdateFlags & Internal_Heightmap_Setup) != 0 || (LayersContentUpdateFlags & Internal_Weightmap_Setup) != 0)
-	{
-		SetupLayers();
-		LayersContentUpdateFlags &= ~(ELandscapeLayersContentUpdateFlag::All_Setup);
-	}
-
 	RegenerateLayersHeightmaps();
 	RegenerateLayersWeightmaps();
 }
 
+void ALandscape::InitializeLayers()
+{
+	check(!bLandscapeLayersAreInitialized || PreviousExperimentalLandscapeLayers != GetMutableDefault<UEditorExperimentalSettings>()->bLandscapeLayerSystem);
+
+	FIntPoint ComponentCounts = ComputeComponentCounts();
+
+	if (HasLayersContent)
+	{
+		CreateLayersRenderingResource(ComponentCounts);
+		InitializeLayersWeightmapUsage();
+	}
+	else
+	{
+		// Special case to load the existing content to the default layer
+		CreateDefaultLayer(ComponentCounts);
+		CopyOldDataToDefaultLayer();
+	}
+}
+
 void ALandscape::TickLayers(float DeltaTime, ELevelTick TickType, FActorTickFunction& ThisTickFunction)
 {
 	check(GIsEditor);
@@ -4169,11 +4013,11 @@ void ALandscape::TickLayers(float DeltaTime, ELevelTick TickType, FActorTickFunc
 
 		if (GetMutableDefault<UEditorExperimentalSettings>()->bLandscapeLayerSystem)
 		{
-			if (PreviousExperimentalLandscapeLayers != GetMutableDefault<UEditorExperimentalSettings>()->bLandscapeLayerSystem)
+			if (!bLandscapeLayersAreInitialized || PreviousExperimentalLandscapeLayers != GetMutableDefault<UEditorExperimentalSettings>()->bLandscapeLayerSystem)
 			{
 				PreviousExperimentalLandscapeLayers = GetMutableDefault<UEditorExperimentalSettings>()->bLandscapeLayerSystem;
-
-				RequestLayersContentUpdate(ELandscapeLayersContentUpdateFlag::All_Setup);
+				InitializeLayers();
+				bLandscapeLayersAreInitialized = true;
 			}
 
 			// If doing editing while shader are compiling or at load of a map, it's possible we will need another update pass after shader are completed to see the correct result
@@ -4195,24 +4039,39 @@ void ALandscape::TickLayers(float DeltaTime, ELevelTick TickType, FActorTickFunc
 		{
 			if (PreviousExperimentalLandscapeLayers != GetMutableDefault<UEditorExperimentalSettings>()->bLandscapeLayerSystem)
 			{
-				PreviousExperimentalLandscapeLayers = GetMutableDefault<UEditorExperimentalSettings>()->bLandscapeLayerSystem;
+				ULandscapeInfo* Info = GetLandscapeInfo();
 
-			#if WITH_EDITORONLY_DATA
-				for (auto& ItPair : RenderDataPerHeightmap)
+				if (Info == nullptr)
 				{
-					FRenderDataPerHeightmap& HeightmapRenderData = ItPair.Value;
-					BeginReleaseResource(HeightmapsCPUReadBackResources[HeightmapRenderData.HeightmapsCPUReadBackResourceIndex].Get());
+					return;
 				}
 
-				for (auto& ItPair : WeightmapCPUReadBackTextures)
+				PreviousExperimentalLandscapeLayers = GetMutableDefault<UEditorExperimentalSettings>()->bLandscapeLayerSystem;
+
+#if WITH_EDITORONLY_DATA
+
+				Info->ForAllLandscapeProxies([&](ALandscapeProxy* Proxy)
 				{
-					FLandscapeLayersTexture2DCPUReadBackResource* WeightmapCPUReadBack = ItPair.Value;
+					for (auto& ItPair : Proxy->HeightmapsCPUReadBack)
+					{
+						FLandscapeLayersTexture2DCPUReadBackResource* HeightmapCPUReadBack = ItPair.Value;
+
+						if (HeightmapCPUReadBack != nullptr)
+						{
+							BeginReleaseResource(HeightmapCPUReadBack);
+						}
+					}
 
-					if (WeightmapCPUReadBack != nullptr)
+					for (auto& ItPair : Proxy->WeightmapsCPUReadBack)
 					{
-						BeginReleaseResource(WeightmapCPUReadBack);
+						FLandscapeLayersTexture2DCPUReadBackResource* WeightmapCPUReadBack = ItPair.Value;
+
+						if (WeightmapCPUReadBack != nullptr)
+						{
+							BeginReleaseResource(WeightmapCPUReadBack);
+						}
 					}
-				}
+				});
 
 				if (CombinedLayersWeightmapAllMaterialLayersResource != nullptr)
 				{
@@ -4236,16 +4095,24 @@ void ALandscape::TickLayers(float DeltaTime, ELevelTick TickType, FActorTickFunc
 
 				FlushRenderingCommands();
 				
-				RenderDataPerHeightmap.Empty();
-				HeightmapsCPUReadBackResources.Empty();
-
-				for (auto& ItPair : WeightmapCPUReadBackTextures)
+				Info->ForAllLandscapeProxies([&](ALandscapeProxy* Proxy)
 				{
-					FLandscapeLayersTexture2DCPUReadBackResource* WeightmapCPUReadBack = ItPair.Value;
+					for (auto& ItPair : Proxy->HeightmapsCPUReadBack)
+					{
+						FLandscapeLayersTexture2DCPUReadBackResource* HeightmapCPUReadBack = ItPair.Value;
 
-					delete WeightmapCPUReadBack;
-					WeightmapCPUReadBack = nullptr;
-				}
+						delete HeightmapCPUReadBack;
+						HeightmapCPUReadBack = nullptr;
+					}
+
+					for (auto& ItPair : Proxy->WeightmapsCPUReadBack)
+					{
+						FLandscapeLayersTexture2DCPUReadBackResource* WeightmapCPUReadBack = ItPair.Value;
+
+						delete WeightmapCPUReadBack;
+						WeightmapCPUReadBack = nullptr;
+					}
+				});				
 
 				delete CombinedLayersWeightmapAllMaterialLayersResource;
 				delete CurrentLayersWeightmapAllMaterialLayersResource;
@@ -4269,12 +4136,17 @@ void ALandscapeProxy::BeginDestroy()
 #if WITH_EDITORONLY_DATA
 	if (GetMutableDefault<UEditorExperimentalSettings>()->bLandscapeLayerSystem)
 	{
-		for (const TUniquePtr<FLandscapeLayersTexture2DCPUReadBackResource>& Ptr : HeightmapsCPUReadBackResources)
+		for (auto& ItPair : HeightmapsCPUReadBack)
 		{
-			BeginReleaseResource(Ptr.Get());
+			FLandscapeLayersTexture2DCPUReadBackResource* HeightmapCPUReadBack = ItPair.Value;
+
+			if (HeightmapCPUReadBack != nullptr)
+			{
+				BeginReleaseResource(HeightmapCPUReadBack);
+			}
 		}
-				
-		for (auto& ItPair : WeightmapCPUReadBackTextures)
+
+		for (auto& ItPair : WeightmapsCPUReadBack)
 		{
 			FLandscapeLayersTexture2DCPUReadBackResource* WeightmapCPUReadBack = ItPair.Value;
 
@@ -4315,9 +4187,15 @@ void ALandscapeProxy::FinishDestroy()
 	{
 		check(ReleaseResourceFence.IsFenceComplete());
 
-		HeightmapsCPUReadBackResources.Empty();
+		for (auto& ItPair : HeightmapsCPUReadBack)
+		{
+			FLandscapeLayersTexture2DCPUReadBackResource* HeightmapCPUReadBack = ItPair.Value;
 
-		for (auto& ItPair : WeightmapCPUReadBackTextures)
+			delete HeightmapCPUReadBack;
+			HeightmapCPUReadBack = nullptr;
+		}
+		
+		for (auto& ItPair : WeightmapsCPUReadBack)
 		{
 			FLandscapeLayersTexture2DCPUReadBackResource* WeightmapCPUReadBack = ItPair.Value;
 
@@ -4501,7 +4379,7 @@ const FLandscapeLayer* ALandscape::GetLayer(int32 InLayerIndex) const
 
 const FLandscapeLayer* ALandscape::GetLayer(const FGuid& InLayerGuid) const
 {
-	return LandscapeLayers.FindByPredicate([InLayerGuid](const FLandscapeLayer& Other) { return Other.Guid == InLayerGuid; });
+	return LandscapeLayers.FindByPredicate([&InLayerGuid](const FLandscapeLayer& Other) { return Other.Guid == InLayerGuid; });
 }
 
 void ALandscape::ForEachLayer(TFunctionRef<void(struct FLandscapeLayer&)> Fn)
@@ -4527,41 +4405,35 @@ void ALandscape::DeleteLayer(int32 InLayerIndex)
 	FGuid LayerGuid = Layer->Guid;
 
 	// Clean up Weightmap usage in LandscapeProxies
-	LandscapeInfo->ForAllLandscapeProxies([LayerGuid](ALandscapeProxy* Proxy)
+	LandscapeInfo->ForAllLandscapeProxies([&LayerGuid](ALandscapeProxy* Proxy)
 	{
-		const FLandscapeLayerData* LayerData = Proxy->LandscapeLayersData.Find(LayerGuid);
-		if (LayerData)
+		for (ULandscapeComponent* Component : Proxy->LandscapeComponents)
 		{
-			for (ULandscapeComponent* Component : Proxy->LandscapeComponents)
+			const FLandscapeLayerComponentData* LayerComponentData = Component->GetLayerData(LayerGuid);
+
+			if (LayerComponentData != nullptr)
 			{
-				const FWeightmapLayerData* WeightmapLayer = LayerData->WeightmapData.Find(Component);
-				if (WeightmapLayer)
+				for (const FWeightmapLayerAllocationInfo& Allocation : LayerComponentData->WeightmapData.LayerAllocations)
 				{
-					for (const FWeightmapLayerAllocationInfo& Allocation : WeightmapLayer->WeightmapLayerAllocations)
+					UTexture2D* WeightmapTexture = LayerComponentData->WeightmapData.Textures[Allocation.WeightmapTextureIndex];
+					ULandscapeWeightmapUsage** Usage = Proxy->WeightmapUsageMap.Find(WeightmapTexture);
+
+					if (Usage != nullptr && (*Usage) != nullptr)
 					{
-						UTexture2D* WeightmapTexture = WeightmapLayer->Weightmaps[Allocation.WeightmapTextureIndex];
-						ULandscapeWeightmapUsage** Usage = Proxy->WeightmapUsageMap.Find(WeightmapTexture);
-						if (Usage != nullptr && (*Usage) != nullptr)
+						(*Usage)->Modify();
+						(*Usage)->ChannelUsage[Allocation.WeightmapTextureChannel] = nullptr;
+
+						if ((*Usage)->FreeChannelCount() == 4)
 						{
-							(*Usage)->Modify();
-							(*Usage)->ChannelUsage[Allocation.WeightmapTextureChannel] = nullptr;
-							if ((*Usage)->FreeChannelCount() == 4)
-							{
-								Proxy->Modify();
-								Proxy->WeightmapUsageMap.Remove(WeightmapTexture);
-							}
+							Proxy->Modify();
+							Proxy->WeightmapUsageMap.Remove(WeightmapTexture);
 						}
 					}
 				}
+				Component->Modify();
+				Component->RemoveLayerData(LayerGuid);
 			}
-		}
-	});
-	
-	// Remove associated layer data of each landscape proxy
-	LandscapeInfo->ForAllLandscapeProxies([LayerGuid](ALandscapeProxy* Proxy)
-	{
-		Proxy->Modify();
-		Proxy->LandscapeLayersData.Remove(LayerGuid);
+		}		
 	});
 
 	const FLandscapeLayer* SplinesReservedLayer = GetLandscapeSplinesReservedLayer();
@@ -4574,7 +4446,6 @@ void ALandscape::DeleteLayer(int32 InLayerIndex)
 	LandscapeLayers.RemoveAt(InLayerIndex);
 
 	// Request Update
-	RequestLayersContentUpdate(ELandscapeLayersContentUpdateFlag::All_Setup, false);
 	RequestLayersContentUpdate(ELandscapeLayersContentUpdateFlag::All, true);
 }
 
@@ -4803,21 +4674,52 @@ void ALandscape::SetEditingLayer(const FGuid& InLayerGuid)
 		return;
 	}
 
-	EditingLayer = InLayerGuid;
-	LandscapeInfo->ForAllLandscapeProxies([InLayerGuid, this](ALandscapeProxy* Proxy)
-	{
-		const FLandscapeLayer* Layer = InLayerGuid.IsValid() ? GetLayer(InLayerGuid) : nullptr;
-		FLandscapeLayerData* LayerData = Layer ? Proxy->LandscapeLayersData.Find(Layer->Guid) : nullptr;
+	EditingLayer = InLayerGuid;	
+}
 
-		for (ULandscapeComponent* Component : Proxy->LandscapeComponents)
+const FGuid& ALandscape::GetEditingLayer() const
+{
+	return EditingLayer;
+}
+
+TMap<UTexture2D*, TArray<ULandscapeComponent*>> ALandscape::GenerateComponentsPerHeightmaps() const
+{
+	ULandscapeInfo* LandscapeInfo = GetLandscapeInfo();
+
+	TMap<UTexture2D*, TArray<ULandscapeComponent*>> ComponentsPerHeightmaps;
+
+	if (LandscapeInfo != nullptr)
+	{
+		LandscapeInfo->ForAllLandscapeProxies([&](ALandscapeProxy* Proxy)
 		{
-			Component->SetEditingLayer(Layer, LayerData);
-			Component->MarkRenderStateDirty();
-		}
-	});
+			for (ULandscapeComponent* Component : Proxy->LandscapeComponents)
+			{
+				UTexture2D* ComponentHeightmapTexture = Component->GetHeightmap();
+				TArray<ULandscapeComponent*>& ComponentList = ComponentsPerHeightmaps.FindOrAdd(ComponentHeightmapTexture);
+				ComponentList.Add(Component);
+			}
+		});
+	}
+
+	return ComponentsPerHeightmaps;
 }
 
-void ALandscape::CreateLayer(FName InName, bool bInUpdateLayersContent)
+void ALandscape::CreateDefaultLayer(const FIntPoint& InComponentCounts)
+{
+	ULandscapeInfo* LandscapeInfo = GetLandscapeInfo();
+	if (!LandscapeInfo || !GetMutableDefault<UEditorExperimentalSettings>()->bLandscapeLayerSystem)
+	{
+		return;
+	}
+
+	check(LandscapeLayers.Num() == 0); // We can only call this function if we have no layers
+
+	CreateLayer(FName(TEXT("Layer")));
+	CreateLayersRenderingResource(InComponentCounts);
+}
+
+
+void ALandscape::CreateLayer(FName InName)
 {
 	ULandscapeInfo* LandscapeInfo = GetLandscapeInfo();
 	if (!LandscapeInfo || !GetMutableDefault<UEditorExperimentalSettings>()->bLandscapeLayerSystem)
@@ -4831,18 +4733,16 @@ void ALandscape::CreateLayer(FName InName, bool bInUpdateLayersContent)
 	LandscapeLayers.Add(NewLayer);
 
 	// Create associated layer data in each landscape proxy
-	LandscapeInfo->ForAllLandscapeProxies([NewLayer](ALandscapeProxy* Proxy)
+	LandscapeInfo->ForAllLandscapeProxies([&NewLayer](ALandscapeProxy* Proxy)
 	{
-		Proxy->Modify();
-		Proxy->LandscapeLayersData.Add(NewLayer.Guid, FLandscapeLayerData());
+		for (ULandscapeComponent* Component : Proxy->LandscapeComponents)
+		{
+			Component->Modify();
+			Component->AddLayerData(NewLayer.Guid, FLandscapeLayerComponentData());
+		}
 	});
 
-	if (bInUpdateLayersContent)
-	{
-		// Request Update
-		RequestLayersContentUpdate(ELandscapeLayersContentUpdateFlag::All_Setup);
-		RequestLayersContentUpdate(ELandscapeLayersContentUpdateFlag::All, true);
-	}
+	InitializeLayerWithEmptyContent(NewLayer.Guid);
 }
 
 bool ALandscape::ReorderLayer(int32 InStartingLayerIndex, int32 InDestinationLayerIndex)
@@ -5121,4 +5021,4 @@ TArray<ALandscapeBlueprintCustomBrush*> ALandscape::GetBrushesForLayer(int32 InL
 
 #endif // WITH_EDITOR
 
-#undef LOCTEXT_NAMESPACE
+#undef LOCTEXT_NAMESPACE
\ No newline at end of file
diff --git a/Engine/Source/Runtime/Landscape/Private/LandscapeGrass.cpp b/Engine/Source/Runtime/Landscape/Private/LandscapeGrass.cpp
--- a/Engine/Source/Runtime/Landscape/Private/LandscapeGrass.cpp
+++ b/Engine/Source/Runtime/Landscape/Private/LandscapeGrass.cpp
@@ -916,7 +916,7 @@ bool ULandscapeComponent::AreTexturesStreamedForGrassMapRender() const
 void ULandscapeComponent::RenderGrassMap()
 {
 	UMaterialInterface* Material = GetLandscapeMaterial();
-	if (CanRenderGrassMap())
+	if (ensure(CanRenderGrassMap()))
 	{
 		TArray<ULandscapeGrassType*> GrassTypes;
 
diff --git a/Engine/Source/Runtime/Landscape/Private/LandscapeModule.cpp b/Engine/Source/Runtime/Landscape/Private/LandscapeModule.cpp
--- a/Engine/Source/Runtime/Landscape/Private/LandscapeModule.cpp
+++ b/Engine/Source/Runtime/Landscape/Private/LandscapeModule.cpp
@@ -112,25 +112,7 @@ void GetLandscapeTexturesAndMaterials(ULevel* Level, TArray<UObject*>& OutTextur
 		if (LandscapeComponent)
 		{
 			LandscapeComponent->GetGeneratedTexturesAndMaterialInstances(OutTexturesAndMaterials);
-		}
-
-		if (GetMutableDefault<UEditorExperimentalSettings>()->bLandscapeLayerSystem)
-		{
-			ALandscapeProxy* Landscape = Cast<ALandscapeProxy>(ObjInLevel);
-
-			if (Landscape != nullptr)
-			{
-				for (auto& ItLayerDataPair : Landscape->LandscapeLayersData)
-				{
-					for (auto& ItHeightmapPair : ItLayerDataPair.Value.Heightmaps)
-					{
-						OutTexturesAndMaterials.AddUnique(ItHeightmapPair.Value);
-					}
-
-					// TODO: add support for weightmap
-				}
-			}
-		}
+		}		
 	}
 }
 
