Improve performance of PrepareXXXX functions

diff --git a/Engine/Source/Runtime/Landscape/Private/LandscapeEditLayers.cpp b/Engine/Source/Runtime/Landscape/Private/LandscapeEditLayers.cpp
--- a/Engine/Source/Runtime/Landscape/Private/LandscapeEditLayers.cpp
+++ b/Engine/Source/Runtime/Landscape/Private/LandscapeEditLayers.cpp
@@ -122,7 +122,7 @@ private:
 			Vertices[i * 3 + 1] = TriangleList[i].V1;
 			Vertices[i * 3 + 2] = TriangleList[i].V2;
 		}
-		
+
 		// Create vertex buffer. Fill buffer with initial data upon creation
 		FRHIResourceCreateInfo CreateInfo(&Vertices);
 		VertexBufferRHI = RHICreateVertexBuffer(Vertices.GetResourceDataSize(), BUF_Static, CreateInfo);
@@ -984,9 +984,9 @@ private:
 
 class FLandscapeLayersCopyTexture_RenderThread
 {
-public: 
+public:
 	FLandscapeLayersCopyTexture_RenderThread(const FString& InSourceResourceDebugName, FTextureResource* InSourceResource, const FString& InDestResourceDebugName, FTextureResource* InDestResource, FTextureResource* InDestCPUResource,
-											     const FIntPoint& InFirstComponentSectionBase, int32 InSubSectionSizeQuad, int32 InNumSubSections, uint8 InSourceCurrentMip, uint8 InDestCurrentMip, uint32 InSourceArrayIndex, uint32 InDestArrayIndex)
+											 const FIntPoint& InFirstComponentSectionBase, int32 InSubSectionSizeQuad, int32 InNumSubSections, uint8 InSourceCurrentMip, uint8 InDestCurrentMip, uint32 InSourceArrayIndex, uint32 InDestArrayIndex)
 		: SourceResource(InSourceResource)
 		, DestResource(InDestResource)
 		, DestCPUResource(InDestCPUResource)
@@ -1045,7 +1045,7 @@ public:
 			Params.SourcePosition.Y = FMath::RoundToInt(PositionOffset.Y * (((SubSectionSizeQuad + 1) * NumSubSections) >> SourceMip));
 			Params.Size.Y = DestSize.Y;
 			Params.DestPosition.Y = 0;
-		}		
+		}
 
 		InRHICmdList.CopyTexture(SourceResource->TextureRHI, DestResource->TextureRHI, Params);
 
@@ -1108,7 +1108,7 @@ class FLandscapeLayersRender_RenderThread
 public:
 
 	FLandscapeLayersRender_RenderThread(const FString& InDebugName, UTextureRenderTarget2D* InWriteRenderTarget, const FIntPoint& InWriteRenderTargetSize, const FIntPoint& InReadRenderTargetSize, const FMatrix& InProjectionMatrix,
-											const ShaderDataType& InShaderParams, uint8 InCurrentMip, const TArray<FLandscapeLayersTriangle>& InTriangleList)
+										const ShaderDataType& InShaderParams, uint8 InCurrentMip, const TArray<FLandscapeLayersTriangle>& InTriangleList)
 		: RenderTargetResource(InWriteRenderTarget->GameThread_GetRenderTargetResource())
 		, WriteRenderTargetSize(InWriteRenderTargetSize)
 		, ReadRenderTargetSize(InReadRenderTargetSize)
@@ -1198,7 +1198,7 @@ public:
 
 			// Set shader params
 			VertexShader->SetParameters(InRHICmdList, ProjectionMatrix);
-			PixelShader->SetParameters(InRHICmdList, ShaderParams);			
+			PixelShader->SetParameters(InRHICmdList, ShaderParams);
 		}
 
 		InRHICmdList.SetStencilRef(0);
@@ -1259,7 +1259,7 @@ void ALandscape::CreateLayersRenderingResource(const FIntPoint& InComponentCount
 				Proxy->HeightmapsCPUReadBack.Add(ComponentHeightmap, NewCPUReadback);
 			}
 		}
-	});	
+	});
 
 	const int32 TotalVertexCountX = (SubsectionSizeQuads * NumSubsections) * InComponentCounts.X + 1;
 	const int32 TotalVertexCountY = (SubsectionSizeQuads * NumSubsections) * InComponentCounts.Y + 1;
@@ -1575,7 +1575,7 @@ void ALandscape::InitializeLayerWithEmptyContent(const FGuid& InLayerGUID)
 	if (Info == nullptr)
 	{
 		return;
-	}	
+	}
 
 	FIntPoint NumComponents = ComputeComponentCounts();
 
@@ -1686,7 +1686,7 @@ void ALandscape::InitializeLayerWithEmptyContent(const FGuid& InLayerGUID)
 				// Nothing to do for Weightmap by default
 
 				Component->AddLayerData(Layer->Guid, MoveTemp(NewData));
-			}				
+			}
 		}
 	});
 }
@@ -1699,8 +1699,8 @@ void ALandscape::CopyLayersTexture(UTexture* InSourceTexture, UTexture* InDestTe
 	}
 }
 
-void ALandscape::CopyLayersTexture(const FString& InSourceDebugName, FTextureResource* InSourceResource, const FString& InDestDebugName, FTextureResource* InDestResource, FTextureResource* InDestCPUResource, const FIntPoint& InFirstComponentSectionBase, 
-										uint8 InSourceCurrentMip, uint8 InDestCurrentMip, uint32 InSourceArrayIndex, uint32 InDestArrayIndex) const
+void ALandscape::CopyLayersTexture(const FString& InSourceDebugName, FTextureResource* InSourceResource, const FString& InDestDebugName, FTextureResource* InDestResource, FTextureResource* InDestCPUResource, const FIntPoint& InFirstComponentSectionBase,
+	uint8 InSourceCurrentMip, uint8 InDestCurrentMip, uint32 InSourceArrayIndex, uint32 InDestArrayIndex) const
 {
 	check(InSourceResource != nullptr);
 	check(InDestResource != nullptr);
@@ -1711,11 +1711,11 @@ void ALandscape::CopyLayersTexture(const FString& InSourceDebugName, FTextureRes
 		[CopyTexture](FRHICommandListImmediate& RHICmdList) mutable
 	{
 		CopyTexture.Copy(RHICmdList);
-	});	
+	});
 }
 
 void ALandscape::DrawWeightmapComponentsToRenderTarget(const FString& InDebugName, const TArray<FIntPoint>& InSectionBaseList, const FVector2D& InScaleBias, UTexture* InWeightmapRTRead, UTextureRenderTarget2D* InOptionalWeightmapRTRead2, UTextureRenderTarget2D* InWeightmapRTWrite,
-													   bool InClearRTWrite, FLandscapeLayersWeightmapShaderParameters& InShaderParams, uint8 InMipRender) const
+	bool InClearRTWrite, FLandscapeLayersWeightmapShaderParameters& InShaderParams, uint8 InMipRender) const
 {
 	check(InWeightmapRTRead != nullptr);
 	check(InWeightmapRTWrite != nullptr);
@@ -1774,7 +1774,7 @@ void ALandscape::DrawWeightmapComponentsToRenderTarget(const FString& InDebugNam
 }
 
 void ALandscape::DrawWeightmapComponentsToRenderTarget(const FString& InDebugName, const TArray<ULandscapeComponent*>& InComponentsToDraw, UTexture* InWeightmapRTRead, UTextureRenderTarget2D* InOptionalWeightmapRTRead2, UTextureRenderTarget2D* InWeightmapRTWrite,
-													   bool InClearRTWrite, FLandscapeLayersWeightmapShaderParameters& InShaderParams, uint8 InMipRender) const
+														bool InClearRTWrite, FLandscapeLayersWeightmapShaderParameters& InShaderParams, uint8 InMipRender) const
 {
 	check(InWeightmapRTRead != nullptr);
 	check(InWeightmapRTWrite != nullptr);
@@ -1864,7 +1864,7 @@ void ALandscape::DrawWeightmapComponentToRenderTargetMips(const TArray<FVector2D
 		if (WriteMipRT != nullptr)
 		{
 			DrawWeightmapComponentsToRenderTarget(OutputDebugName ? FString::Printf(TEXT("LS Weight: %s = -> %s Mips %d"), *ReadMipRT->GetName(), *WriteMipRT->GetName(), CurrentMip) : TEXT(""),
-												  SectionBaseToDraw, WeightmapScaleBias, ReadMipRT, nullptr, WriteMipRT, InClearRTWrite, InShaderParams, CurrentMip++);
+													SectionBaseToDraw, WeightmapScaleBias, ReadMipRT, nullptr, WriteMipRT, InClearRTWrite, InShaderParams, CurrentMip++);
 		}
 
 		ReadMipRT = WeightmapRTList[MipRTIndex];
@@ -1877,9 +1877,9 @@ void ALandscape::ClearLayersWeightmapTextureResource(const FString& InDebugName,
 
 	ENQUEUE_RENDER_COMMAND(FLandscapeLayersClearWeightmapCommand)(
 		[LayersClear](FRHICommandListImmediate& RHICmdList) mutable
-		{
-			LayersClear.Clear(RHICmdList);
-		});
+	{
+		LayersClear.Clear(RHICmdList);
+	});
 }
 
 void ALandscape::DrawHeightmapComponentsToRenderTargetMips(const TArray<ULandscapeComponent*>& InComponentsToDraw, UTexture* InReadHeightmap, bool InClearRTWrite, struct FLandscapeLayersHeightmapShaderParameters& InShaderParams) const
@@ -1994,17 +1994,17 @@ void ALandscape::DrawHeightmapComponentsToRenderTarget(const FString& InDebugNam
 
 	ENQUEUE_RENDER_COMMAND(FDrawLandscapeLayersHeightmapCommand)(
 		[LayersRender, InClearRTWrite](FRHICommandListImmediate& RHICmdList) mutable
-		{
-			LayersRender.Render(RHICmdList, InClearRTWrite);
-		});
-	
+	{
+		LayersRender.Render(RHICmdList, InClearRTWrite);
+	});
+
 	PrintLayersDebugRT(InDebugName, InHeightmapRTWrite, InMipRender, true, InShaderParams.GenerateNormals);
 }
 
 void ALandscape::GenerateLayersRenderQuad(const FIntPoint& InVertexPosition, float InVertexSize, const FVector2D& InUVStart, const FVector2D& InUVSize, TArray<FLandscapeLayersTriangle>& OutTriangles) const
 {
 	FLandscapeLayersTriangle Tri1;
-	
+
 	Tri1.V0.Position = FVector2D(InVertexPosition.X, InVertexPosition.Y);
 	Tri1.V1.Position = FVector2D(InVertexPosition.X + InVertexSize, InVertexPosition.Y);
 	Tri1.V2.Position = FVector2D(InVertexPosition.X + InVertexSize, InVertexPosition.Y + InVertexSize);
@@ -2117,10 +2117,10 @@ void ALandscape::GenerateLayersRenderQuadsMip(const FIntPoint& InSectionBase, co
 {
 	int32 SubsectionSizeVerts = InSubSectionSizeQuad + 1;
 	int32 LocalComponentSizeQuad = InSubSectionSizeQuad * NumSubsections;
-	int32 LocalComponentSizeVerts = SubsectionSizeVerts * NumSubsections;	
+	int32 LocalComponentSizeVerts = SubsectionSizeVerts * NumSubsections;
 	int32 MipSubsectionSizeVerts = SubsectionSizeVerts >> InCurrentMip;
-	int32 MipLocalComponentSizeVerts = MipSubsectionSizeVerts * NumSubsections;	
-	
+	int32 MipLocalComponentSizeVerts = MipSubsectionSizeVerts * NumSubsections;
+
 	FVector2D PositionOffset(FMath::RoundToInt(InSectionBase.X / LocalComponentSizeQuad), FMath::RoundToInt(InSectionBase.Y / LocalComponentSizeQuad));
 	FVector2D ComponentsPerTexture(FMath::RoundToInt(InWriteSize.X / LocalComponentSizeQuad), FMath::RoundToInt(InWriteSize.Y / LocalComponentSizeQuad));
 
@@ -2399,10 +2399,10 @@ void ALandscape::PrintLayersDebugRT(const FString& InContext, UTextureRenderTarg
 	FTextureRenderTargetResource* RenderTargetResource = InDebugRT->GameThread_GetRenderTargetResource();
 	ENQUEUE_RENDER_COMMAND(FLandscapeLayersDebugRenderTargetResolveCommand)(
 		[RenderTargetResource](FRHICommandListImmediate& RHICmdList) mutable
-		{
-			// Copy (resolve) the rendered image from the frame buffer to its render target texture
-			RHICmdList.CopyToResolveTarget(RenderTargetResource->GetRenderTargetTexture(), RenderTargetResource->TextureRHI, FResolveParams());
-		});
+	{
+		// Copy (resolve) the rendered image from the frame buffer to its render target texture
+		RHICmdList.CopyToResolveTarget(RenderTargetResource->GetRenderTargetTexture(), RenderTargetResource->TextureRHI, FResolveParams());
+	});
 
 	FlushRenderingCommands();
 	int32 MinX, MinY, MaxX, MaxY;
@@ -2478,46 +2478,35 @@ bool ALandscape::PrepareLayersHeightmapTextureResources() const
 
 	Info->ForAllLandscapeProxies([&](ALandscapeProxy* Proxy)
 	{
-		if (!IsReady)
-		{
-			return;
-		}
-
 		for (ULandscapeComponent* Component : Proxy->LandscapeComponents)
 		{
-			if (!Component->GetHeightmap()->IsFullyStreamedIn())
-			{
-				IsReady = false;
-				return;
-			}
+			UTexture2D* ComponentHeightmap = Component->GetHeightmap();
+
+			IsReady &= ComponentHeightmap->IsFullyStreamedIn();
 
 			for (const FLandscapeLayer& Layer : LandscapeLayers)
 			{
 				FLandscapeLayerComponentData* ComponentLayerData = Component->GetLayerData(Layer.Guid);
-				UTexture2D* LayerHeightmap = ComponentLayerData->HeightmapData.Texture;
 
-				if (!LayerHeightmap->IsAsyncCacheComplete())
-				{
-					IsReady = false;
-					return;
-				}
+				IsReady &= ComponentLayerData != nullptr;
 
-				if (LayerHeightmap->Resource == nullptr)
+				if (ComponentLayerData != nullptr)
 				{
-					LayerHeightmap->FinishCachePlatformData();
-
-					LayerHeightmap->Resource = LayerHeightmap->CreateResource();
+					UTexture2D* LayerHeightmap = ComponentLayerData->HeightmapData.Texture;
 
-					if (LayerHeightmap->Resource != nullptr)
+					if (LayerHeightmap->IsAsyncCacheComplete() && LayerHeightmap->Resource == nullptr)
 					{
-						BeginInitResource(LayerHeightmap->Resource);
+						LayerHeightmap->FinishCachePlatformData();
+
+						LayerHeightmap->Resource = LayerHeightmap->CreateResource();
+
+						if (LayerHeightmap->Resource != nullptr)
+						{
+							BeginInitResource(LayerHeightmap->Resource);
+						}
 					}
-				}
 
-				if (LayerHeightmap->Resource == nullptr || !LayerHeightmap->Resource->IsInitialized() || !LayerHeightmap->IsFullyStreamedIn())
-				{
-					IsReady = false;
-					return;
+					IsReady &= LayerHeightmap->Resource != nullptr && LayerHeightmap->Resource->IsInitialized() && LayerHeightmap->IsFullyStreamedIn();
 				}
 			}
 		}
@@ -2590,7 +2579,7 @@ void ALandscape::RegenerateLayersHeightmaps()
 					CopyLayersTexture(LayerHeightmap, LandscapeScratchRT1, nullptr, GlobalLayersData.TopLeftSectionBase);
 					PrintLayersDebugRT(OutputDebugName ? FString::Printf(TEXT("LS Height: %s Component %s += -> CombinedAtlas %s"), *Layer.Name.ToString(), *LayerHeightmap->GetName(), *LandscapeScratchRT1->GetName()) : TEXT(""), LandscapeScratchRT1);
 				}
-			});			
+			});
 
 			// NOTE: From this point on, we always work in non atlas, we'll convert back at the end to atlas only
 			DrawHeightmapComponentsToRenderTarget(OutputDebugName ? FString::Printf(TEXT("LS Height: %s += -> NonAtlas %s"), *Layer.Name.ToString(), *LandscapeScratchRT1->GetName(), *LandscapeScratchRT2->GetName()) : TEXT(""),
@@ -2935,7 +2924,7 @@ void ALandscape::PrepareComponentDataToPackMaterialLayersCS(int32 InCurrentWeigh
 			if (WeightmapsToProcess.Find(ComponentWeightmapTexture) == nullptr)
 			{
 				WeightmapsToProcess.Add(ComponentWeightmapTexture, ComponentWeightmapTextureUsage);
-				OutProcessedWeightmaps.Add(ComponentWeightmapTexture);				
+				OutProcessedWeightmaps.Add(ComponentWeightmapTexture);
 
 				FLandscapeLayersTexture2DCPUReadBackResource** CPUReadback = Component->GetLandscapeProxy()->WeightmapsCPUReadBack.Find(ComponentWeightmapTexture);
 				check(CPUReadback != nullptr);
@@ -3023,7 +3012,7 @@ void ALandscape::ReallocateLayersWeightmaps(const TArray<ULandscapeLayerInfoObje
 	if (Info == nullptr)
 	{
 		return;
-	}	
+	}
 
 	TArray<ULandscapeComponent*> AllLandscapeComponents;
 
@@ -3070,7 +3059,7 @@ void ALandscape::ReallocateLayersWeightmaps(const TArray<ULandscapeLayerInfoObje
 			BaseWeightmapAllocation.WeightmapTextureChannel = 255;
 			BaseWeightmapAllocation.WeightmapTextureIndex = 255;
 		}
-		
+
 		TArray<ULandscapeWeightmapUsage*>& WeightmapTexturesUsage = Component->GetWeightmapTexturesUsage();
 
 		for (int32 i = 0; i < WeightmapTexturesUsage.Num(); ++i)
@@ -3080,7 +3069,7 @@ void ALandscape::ReallocateLayersWeightmaps(const TArray<ULandscapeLayerInfoObje
 
 			Usage->ClearUsage();
 		}
-	}	 
+	}
 
 	bool NeedMaterialInstanceRebuild = false;
 
@@ -3121,7 +3110,7 @@ void ALandscape::ReallocateLayersWeightmaps(const TArray<ULandscapeLayerInfoObje
 				}
 			}
 		}
-	});	
+	});
 
 	// Determine if the Final layer need to add/remove some alloc
 	for (auto& ItPair : LayerAllocsPerComponent)
@@ -3170,7 +3159,7 @@ void ALandscape::ReallocateLayersWeightmaps(const TArray<ULandscapeLayerInfoObje
 		Texture->PostEditChange();
 
 		Texture->WaitForStreaming();
-	}	
+	}
 
 	//GDisableAutomaticTextureMaterialUpdateDependencies = false;
 
@@ -3178,7 +3167,7 @@ void ALandscape::ReallocateLayersWeightmaps(const TArray<ULandscapeLayerInfoObje
 
 	// BEGIN HACK for now we disable the logic that will determine what need updating as it's buggy, and i prefer a correct behavior than an optimized non working logic.
 	OutComponentThatNeedMaterialRebuild.Reset();
-	OutComponentThatNeedMaterialRebuild.Append(AllLandscapeComponents);	
+	OutComponentThatNeedMaterialRebuild.Append(AllLandscapeComponents);
 	// END HACK
 
 	/*for (ULandscapeComponent* Component : AllLandscapeComponents)
@@ -3270,7 +3259,7 @@ void ALandscape::InitLayersWeightmapResources(uint8 InLayerCount)
 		}
 
 		WeightmapScratchPackLayerTextureResource = new FLandscapeTexture2DResource(FirstWeightmapRT->SizeX, FirstWeightmapRT->SizeY, PF_B8G8R8A8, MipCount, true);
-		BeginInitResource(WeightmapScratchPackLayerTextureResource);		
+		BeginInitResource(WeightmapScratchPackLayerTextureResource);
 	}
 }
 
@@ -3285,62 +3274,30 @@ bool ALandscape::PrepareLayersWeightmapTextureResources() const
 
 	bool IsReady = true;
 
-	// Make sure all our original weightmap textures are streamed in and ready to be used
 	Info->ForAllLandscapeProxies([&](ALandscapeProxy* Proxy)
 	{
-		for (ULandscapeComponent* Component : Proxy->LandscapeComponents)
+		for (const FLandscapeLayer& Layer : LandscapeLayers)
 		{
-			for (UTexture2D* ComponentWeightmap : Component->GetWeightmapTextures())
-			{
-				if (!ComponentWeightmap->IsFullyStreamedIn())
-				{
-					IsReady = false;
-					return;
-				}
-			}
-
-			for (const FLandscapeLayer& Layer : LandscapeLayers)
+			for (ULandscapeComponent* Component : Proxy->LandscapeComponents)
 			{
-				FLandscapeLayerComponentData* ComponentLayerData = Component->GetLayerData(Layer.Guid);
-
-				if (ComponentLayerData == nullptr)
-				{
-					IsReady = false;
-					return;
-				}
-
-				for (UTexture2D* LayerWeightmap : ComponentLayerData->WeightmapData.Textures)
+				for (UTexture2D* ComponentWeightmap : Component->GetWeightmapTextures())
 				{
-					if (LayerWeightmap != nullptr && !LayerWeightmap->IsFullyStreamedIn())
+					if (!ComponentWeightmap->IsFullyStreamedIn())
 					{
 						IsReady = false;
-						return;
+						break;
 					}
 				}
-			}
-		}
-	});
 
-	// Init all needed resources
-	if (IsReady)
-	{
-		Info->ForAllLandscapeProxies([&](ALandscapeProxy* Proxy)
-		{
-			for (const FLandscapeLayer& Layer : LandscapeLayers)
-			{
-				for (ULandscapeComponent* Component : Proxy->LandscapeComponents)
-				{
-					FLandscapeLayerComponentData* ComponentLayerData = Component->GetLayerData(Layer.Guid);
+				FLandscapeLayerComponentData* ComponentLayerData = Component->GetLayerData(Layer.Guid);
 
-					if (ComponentLayerData == nullptr)
-					{
-						IsReady = false;
-						return;
-					}
+				IsReady &= ComponentLayerData != nullptr;
 
+				if (ComponentLayerData != nullptr)
+				{
 					for (UTexture2D* LayerWeightmap : ComponentLayerData->WeightmapData.Textures)
 					{
-						if (LayerWeightmap->Resource == nullptr)
+						if (LayerWeightmap->IsAsyncCacheComplete() && LayerWeightmap->Resource == nullptr)
 						{
 							LayerWeightmap->FinishCachePlatformData();
 
@@ -3351,38 +3308,13 @@ bool ALandscape::PrepareLayersWeightmapTextureResources() const
 								BeginInitResource(LayerWeightmap->Resource);
 							}
 						}
-					}
-				}
-			}
-		});
 
-		// Wait for the new resource to be fully initialized/streamed in
-		Info->ForAllLandscapeProxies([&](ALandscapeProxy* Proxy)
-		{
-			for (const FLandscapeLayer& Layer : LandscapeLayers)
-			{
-				for (ULandscapeComponent* Component : Proxy->LandscapeComponents)
-				{
-					FLandscapeLayerComponentData* ComponentLayerData = Component->GetLayerData(Layer.Guid);
-
-					if (ComponentLayerData == nullptr)
-					{
-						IsReady = false;
-						return;
-					}
-
-					for (UTexture2D* LayerWeightmap : ComponentLayerData->WeightmapData.Textures)
-					{
-						if (LayerWeightmap->Resource == nullptr || !LayerWeightmap->Resource->IsInitialized() || !LayerWeightmap->IsFullyStreamedIn())
-						{
-							IsReady = false;
-							return;
-						}
+						IsReady &= LayerWeightmap->Resource != nullptr && LayerWeightmap->Resource->IsInitialized() && LayerWeightmap->IsFullyStreamedIn();
 					}
 				}
 			}
-		});
-	}
+		}
+	});
 
 	return IsReady;
 }
@@ -3432,7 +3364,7 @@ void ALandscape::RegenerateLayersWeightmaps()
 		}
 
 		bool ComputeShaderGeneratedData = false;
-		bool FirstLayer = true;		
+		bool FirstLayer = true;
 		TMap<ULandscapeLayerInfoObject*, bool> WeightmapLayersBlendSubstractive;
 
 		for (FLandscapeLayer& Layer : LandscapeLayers)
@@ -3472,7 +3404,7 @@ void ALandscape::RegenerateLayersWeightmaps()
 				CopyLayersTexture(*EmptyRT->GetName(), EmptyRT->GameThread_GetRenderTargetResource(), OutputDebugName ? FString::Printf(TEXT("Weight: %s Clear WeightmapScratchExtractLayerTextureResource"), *Layer.Name.ToString()) : TEXT(""), WeightmapScratchExtractLayerTextureResource);
 
 				// Prepare compute shader data
-				TArray<FLandscapeLayerWeightmapExtractMaterialLayersComponentData> ComponentsData;	
+				TArray<FLandscapeLayerWeightmapExtractMaterialLayersComponentData> ComponentsData;
 				PrepareComponentDataToExtractMaterialLayersCS(Layer, CurrentWeightmapToProcessIndex, OutputDebugName, WeightmapScratchExtractLayerTextureResource, ComponentsData, LayerInfoObjects);
 
 				HasFoundWeightmapToProcess = ComponentsData.Num() > 0;
@@ -3508,9 +3440,9 @@ void ALandscape::RegenerateLayersWeightmaps()
 					ENQUEUE_RENDER_COMMAND(FLandscapeLayersExtractMaterialLayersCSCommand)(
 						[CSDispatch](FRHICommandListImmediate& RHICmdList) mutable
 					{
-						CSDispatch.ExtractLayers(RHICmdList);						
+						CSDispatch.ExtractLayers(RHICmdList);
 					});
-					
+
 					++CurrentWeightmapToProcessIndex;
 					ComputeShaderGeneratedData = true; // at least 1 CS was executed, so we can continue the processing
 				}
@@ -3636,7 +3568,7 @@ void ALandscape::RegenerateLayersWeightmaps()
 						}
 					}
 				}
-			});			
+			});
 
 			int8 CurrentWeightmapToProcessIndex = 0;
 			bool HasFoundWeightmapToProcess = true; // try processing at least once
@@ -3719,7 +3651,7 @@ void ALandscape::RegenerateLayersWeightmaps()
 					{
 						UTexture2D* WeightmapTexture = ProcessedWeightmaps[NextTextureIndexToProcess];
 						FLandscapeLayersTexture2DCPUReadBackResource* WeightmapCPUReadBack = ProcessedCPUReadbackTextures[NextTextureIndexToProcess];
-						FIntPoint TextureTopLeftPositionInAtlas(WeightmapTextureOutputOffset[NextTextureIndexToProcess - StartTextureIndex].X, WeightmapTextureOutputOffset[NextTextureIndexToProcess - StartTextureIndex].Y);						
+						FIntPoint TextureTopLeftPositionInAtlas(WeightmapTextureOutputOffset[NextTextureIndexToProcess - StartTextureIndex].X, WeightmapTextureOutputOffset[NextTextureIndexToProcess - StartTextureIndex].Y);
 
 						int32 CurrentMip = 0;
 
@@ -3741,7 +3673,7 @@ void ALandscape::RegenerateLayersWeightmaps()
 		}
 
 		UpdateLayersMaterialInstances(LayersUpdateAllMaterials ? AllLandscapeComponents : ComponentThatNeedMaterialRebuild);
-		LayersUpdateAllMaterials = false;		
+		LayersUpdateAllMaterials = false;
 	}
 
 	if ((LayersContentUpdateFlags & EInternalLandscapeLayersContentUpdateFlag::Internal_Weightmap_ResolveToTexture))
@@ -3777,7 +3709,7 @@ void ALandscape::UpdateLayersMaterialInstances(const TArray<ULandscapeComponent*
 	}
 
 	TArray<ULandscapeComponent*> ComponentsToUpdate;
-	ComponentsToUpdate.Append(InComponentsToUpdate);	
+	ComponentsToUpdate.Append(InComponentsToUpdate);
 
 	SCOPE_CYCLE_COUNTER(STAT_LandscapeLayersUpdateMaterialInstance);
 
@@ -3946,7 +3878,7 @@ void ALandscape::ResolveLayersWeightmapTexture()
 		{
 			TArray<UTexture2D*>& ComponentWeightmaps = Component->GetWeightmapTextures();
 
-			for (UTexture2D* ComponentLayerWeightmap : ComponentWeightmaps)			
+			for (UTexture2D* ComponentLayerWeightmap : ComponentWeightmaps)
 			{
 				if (!ProcessedWeightmaps.Contains(ComponentLayerWeightmap))
 				{
@@ -4094,7 +4026,7 @@ void ALandscape::TickLayers(float DeltaTime, ELevelTick TickType, FActorTickFunc
 				}
 
 				FlushRenderingCommands();
-				
+
 				Info->ForAllLandscapeProxies([&](ALandscapeProxy* Proxy)
 				{
 					for (auto& ItPair : Proxy->HeightmapsCPUReadBack)
@@ -4112,7 +4044,7 @@ void ALandscape::TickLayers(float DeltaTime, ELevelTick TickType, FActorTickFunc
 						delete WeightmapCPUReadBack;
 						WeightmapCPUReadBack = nullptr;
 					}
-				});				
+				});
 
 				delete CombinedLayersWeightmapAllMaterialLayersResource;
 				delete CurrentLayersWeightmapAllMaterialLayersResource;
@@ -4123,7 +4055,7 @@ void ALandscape::TickLayers(float DeltaTime, ELevelTick TickType, FActorTickFunc
 				CurrentLayersWeightmapAllMaterialLayersResource = nullptr;
 				WeightmapScratchExtractLayerTextureResource = nullptr;
 				WeightmapScratchPackLayerTextureResource = nullptr;
-			#endif
+#endif
 			}
 		}
 	}
@@ -4194,7 +4126,7 @@ void ALandscapeProxy::FinishDestroy()
 			delete HeightmapCPUReadBack;
 			HeightmapCPUReadBack = nullptr;
 		}
-		
+
 		for (auto& ItPair : WeightmapsCPUReadBack)
 		{
 			FLandscapeLayersTexture2DCPUReadBackResource* WeightmapCPUReadBack = ItPair.Value;
@@ -4400,7 +4332,7 @@ void ALandscape::DeleteLayer(int32 InLayerIndex)
 	{
 		return;
 	}
-	
+
 	Modify();
 	FGuid LayerGuid = Layer->Guid;
 
@@ -4433,7 +4365,7 @@ void ALandscape::DeleteLayer(int32 InLayerIndex)
 				Component->Modify();
 				Component->RemoveLayerData(LayerGuid);
 			}
-		}		
+		}
 	});
 
 	const FLandscapeLayer* SplinesReservedLayer = GetLandscapeSplinesReservedLayer();
@@ -4468,7 +4400,7 @@ void ALandscape::ClearLayer(const FGuid& InLayerGuid, bool bInUpdateCollision)
 	{
 		return;
 	}
-	
+
 	Modify();
 	FScopedSetLandscapeEditingLayer Scope(this, Layer ? Layer->Guid : FGuid(), [=] { RequestLayersContentUpdate(ELandscapeLayersContentUpdateFlag::All, true); });
 
@@ -4674,7 +4606,7 @@ void ALandscape::SetEditingLayer(const FGuid& InLayerGuid)
 		return;
 	}
 
-	EditingLayer = InLayerGuid;	
+	EditingLayer = InLayerGuid;
 }
 
 const FGuid& ALandscape::GetEditingLayer() const
@@ -4748,7 +4680,7 @@ void ALandscape::CreateLayer(FName InName)
 bool ALandscape::ReorderLayer(int32 InStartingLayerIndex, int32 InDestinationLayerIndex)
 {
 	if (InStartingLayerIndex != InDestinationLayerIndex &&
-		LandscapeLayers.IsValidIndex(InStartingLayerIndex) && 
+		LandscapeLayers.IsValidIndex(InStartingLayerIndex) &&
 		LandscapeLayers.IsValidIndex(InDestinationLayerIndex))
 	{
 		Modify();
