- Fixed invalid checks when reallocating Landscape weightmaps without using Landscape Layer System (introduced by 5899368)

diff --git a/Engine/Source/Runtime/Landscape/Private/LandscapeEdit.cpp b/Engine/Source/Runtime/Landscape/Private/LandscapeEdit.cpp
--- a/Engine/Source/Runtime/Landscape/Private/LandscapeEdit.cpp
+++ b/Engine/Source/Runtime/Landscape/Private/LandscapeEdit.cpp
@@ -4714,7 +4714,7 @@ void ULandscapeComponent::ReallocateWeightmaps(FLandscapeEditDataInterface* Data
 				ULandscapeWeightmapUsage* TryWeightmapUsage = ItPair.Value;
 				//
 				FGuid EditingLayerGUID = GetEditingLayerGUID();
-				check(!InCanUseEditingWeightmap || EditingLayerGUID.IsValid());
+				check(!GetMutableDefault<UEditorExperimentalSettings>()->bLandscapeLayerSystem || !InCanUseEditingWeightmap || EditingLayerGUID.IsValid());
 
 				FGuid LayerGuidToSeek = InCanUseEditingWeightmap ? EditingLayerGUID : FGuid();
 
@@ -4785,7 +4785,7 @@ void ULandscapeComponent::ReallocateWeightmaps(FLandscapeEditDataInterface* Data
 			}
 
 			FGuid EditingLayerGUID = GetEditingLayerGUID();
-			check(!InCanUseEditingWeightmap || EditingLayerGUID.IsValid());
+			check(!GetMutableDefault<UEditorExperimentalSettings>()->bLandscapeLayerSystem || !InCanUseEditingWeightmap || EditingLayerGUID.IsValid());
 
 			CurrentWeightmapUsage->LayerGuid = InCanUseEditingWeightmap ? EditingLayerGUID : FGuid();
 			// UE_LOG(LogLandscape, Log, TEXT("Making a new texture %s"), *CurrentWeightmapTexture->GetName());
