- Landscape Layer System:  Refactor Update Code to simplify UpdateCollisionLayerData and UpdateCollisionHeightData

#rnx

diff --git a/Engine/Source/Runtime/Landscape/Classes/LandscapeComponent.h b/Engine/Source/Runtime/Landscape/Classes/LandscapeComponent.h
--- a/Engine/Source/Runtime/Landscape/Classes/LandscapeComponent.h
+++ b/Engine/Source/Runtime/Landscape/Classes/LandscapeComponent.h
@@ -860,6 +860,14 @@ public:
 
 protected:
 
+#if WITH_EDITOR
+	void RecreateCollisionComponent(bool bUseSimpleCollision);
+	void UpdateCollisionHeightBuffer(int32 InComponentX1, int32 InComponentY1, int32 InComponentX2, int32 InComponentY2, int32 InCollisionMipLevel, int32 InHeightmapSizeU, int32 InHeightmapSizeV,
+		const FColor* const InHeightmapTextureMipData, uint16* CollisionHeightData, uint16* GrassHeightData,
+		const FColor* const InXYOffsetTextureMipData, uint16* CollisionXYOffsetData);
+	void UpdateDominantLayerBuffer(int32 InComponentX1, int32 InComponentY1, int32 InComponentX2, int32 InComponentY2, int32 InCollisionMipLevel, int32 InWeightmapSizeU, int32 InDataLayerIdx, const TArray<uint8*>& InCollisionDataPtrs, uint8* DominantLayerData);
+#endif
+
 	/** Whether the component type supports static lighting. */
 	virtual bool SupportsStaticLighting() const override
 	{
diff --git a/Engine/Source/Runtime/Landscape/Private/LandscapeEdit.cpp b/Engine/Source/Runtime/Landscape/Private/LandscapeEdit.cpp
--- a/Engine/Source/Runtime/Landscape/Private/LandscapeEdit.cpp
+++ b/Engine/Source/Runtime/Landscape/Private/LandscapeEdit.cpp
@@ -753,6 +753,40 @@ struct FLandscapeComponentAlphaInfo
 	}
 };
 
+struct FCollisionSize
+{
+public:
+	static FCollisionSize CreateSimple(bool bUseSimpleCollision, int32 InNumSubSections, int32 InSubsectionSizeQuads, int32 InMipLevel)
+	{
+		return bUseSimpleCollision ? Create(InNumSubSections, InSubsectionSizeQuads, InMipLevel) : FCollisionSize();
+	}
+
+	static FCollisionSize Create(int32 InNumSubSections, int32 InSubsectionSizeQuads, int32 InMipLevel)
+	{
+		return FCollisionSize(InNumSubSections, InSubsectionSizeQuads, InMipLevel);
+	}
+
+	FCollisionSize(FCollisionSize&& Other) = default;
+	FCollisionSize& operator=(FCollisionSize&& Other) = default;
+private:
+	FCollisionSize(int32 InNumSubsections, int32 InSubsectionSizeQuads, int32 InMipLevel)
+	{
+		SubsectionSizeVerts = ((InSubsectionSizeQuads + 1) >> InMipLevel);
+		SubsectionSizeQuads = SubsectionSizeVerts - 1;
+		SizeVerts = InNumSubsections * SubsectionSizeQuads + 1;
+		SizeVertsSquare = FMath::Square(SizeVerts);
+	}
+
+	FCollisionSize()
+	{
+	}
+
+public:
+	int32 SubsectionSizeVerts = 0;
+	int32 SubsectionSizeQuads = 0;
+	int32 SizeVerts = 0;
+	int32 SizeVertsSquare = 0;
+};
 
 bool ULandscapeComponent::UpdateCollisionHeightData(const FColor* const HeightmapTextureMipData, const FColor* const SimpleCollisionHeightmapTextureData, int32 ComponentX1/*=0*/, int32 ComponentY1/*=0*/, int32 ComponentX2/*=MAX_int32*/, int32 ComponentY2/*=MAX_int32*/, bool bUpdateBounds/*=false*/, const FColor* XYOffsetTextureMipData/*=nullptr*/)
 {
@@ -765,15 +799,11 @@ bool ULandscapeComponent::UpdateCollisionHeightData(const FColor* const Heightma
 
 	// Simple collision is not currently supported with mesh collision components
 	const bool bUsingSimpleCollision = (SimpleCollisionMipLevel > CollisionMipLevel && SimpleCollisionHeightmapTextureData && !XYOffsetmapTexture);
-
-	const int32 CollisionSubsectionSizeVerts = ((SubsectionSizeQuads + 1) >> CollisionMipLevel);
-	const int32 CollisionSubsectionSizeQuads = CollisionSubsectionSizeVerts - 1;
-	const int32 CollisionSizeVerts = NumSubsections * CollisionSubsectionSizeQuads + 1;
-	const int32 SimpleCollisionSubsectionSizeVerts = bUsingSimpleCollision ? ((SubsectionSizeQuads + 1) >> SimpleCollisionMipLevel) : 0;
-	const int32 SimpleCollisionSubsectionSizeQuads = bUsingSimpleCollision ? SimpleCollisionSubsectionSizeVerts - 1 : 0;
-	const int32 SimpleCollisionSizeVerts = bUsingSimpleCollision ? NumSubsections * SimpleCollisionSubsectionSizeQuads + 1 : 0;
-
-	const int32 TotalCollisionSize = FMath::Square(CollisionSizeVerts) + FMath::Square(SimpleCollisionSizeVerts);
+	
+	FCollisionSize CollisionSize = FCollisionSize::Create(NumSubsections, SubsectionSizeQuads, CollisionMipLevel);
+	FCollisionSize SimpleCollisionSize = FCollisionSize::CreateSimple(bUsingSimpleCollision, NumSubsections, SubsectionSizeQuads, SimpleCollisionMipLevel);
+		
+	const int32 TotalCollisionSize = CollisionSize.SizeVertsSquare + SimpleCollisionSize.SizeVertsSquare;
 
 	uint16* CollisionHeightData = nullptr;
 	uint16* CollisionXYOffsetData = nullptr;
@@ -808,124 +838,35 @@ bool ULandscapeComponent::UpdateCollisionHeightData(const FColor* const Heightma
 			CollisionComp->CachedLocalBox = CachedLocalBox;
 			CollisionComp->UpdateComponentToWorld();
 		}
-
-		CollisionHeightData = (uint16*)CollisionComp->CollisionHeightData.Lock(LOCK_READ_WRITE);
-
-		if (XYOffsetmapTexture && MeshCollisionComponent)
-		{
-			CollisionXYOffsetData = (uint16*)MeshCollisionComponent->CollisionXYOffsetData.Lock(LOCK_READ_WRITE);
-		}
 	}
 	else
 	{
+		CreatedNew = true;
+		ChangeType = CollisionComp != nullptr;
 		ComponentX1 = 0;
 		ComponentY1 = 0;
 		ComponentX2 = ComponentSizeQuads;
 		ComponentY2 = ComponentSizeQuads;
 
-		TArray<uint8> DominantLayerData;
-		TArray<ULandscapeLayerInfoObject*> LayerInfos;
-		if (CollisionComp) // remove old component before changing to other type collision...
-		{
-			ChangeType = true;
-
-			if (CollisionComp->DominantLayerData.GetElementCount())
-			{
-				check(CollisionComp->DominantLayerData.GetElementCount() >= TotalCollisionSize);
-				DominantLayerData.AddUninitialized(TotalCollisionSize);
-
-				const uint8* SrcDominantLayerData = (uint8*)CollisionComp->DominantLayerData.Lock(LOCK_READ_ONLY);
-				FMemory::Memcpy(DominantLayerData.GetData(), SrcDominantLayerData, TotalCollisionSize * sizeof(uint8));
-				CollisionComp->DominantLayerData.Unlock();
-			}
-
-			if (CollisionComp->ComponentLayerInfos.Num())
-			{
-				LayerInfos = CollisionComp->ComponentLayerInfos;
-			}
-
-			if (Info)
-			{
-				Info->Modify();
-			}
-			Proxy->Modify();
-			CollisionComp->DestroyComponent();
-			CollisionComp = nullptr;
-		}
-
-		if (XYOffsetmapTexture)
-		{
-			MeshCollisionComponent = NewObject<ULandscapeMeshCollisionComponent>(Proxy, NAME_None, RF_Transactional);
-			CollisionComp = MeshCollisionComponent;
-		}
-		else
-		{
-			MeshCollisionComponent = nullptr;
-			CollisionComp = NewObject<ULandscapeHeightfieldCollisionComponent>(Proxy, NAME_None, RF_Transactional);
-		}
-
-		CollisionComp->SetRelativeLocation(RelativeLocation);
-		CollisionComp->SetupAttachment(Proxy->GetRootComponent(), NAME_None);
-		Proxy->CollisionComponents.Add(CollisionComp);
-
-		CollisionComp->RenderComponent = this;
-		CollisionComp->SetSectionBase(GetSectionBase());
-		CollisionComp->CollisionSizeQuads = CollisionSubsectionSizeQuads * NumSubsections;
-		CollisionComp->CollisionScale = (float)(ComponentSizeQuads) / (float)(CollisionComp->CollisionSizeQuads);
-		CollisionComp->SimpleCollisionSizeQuads = bUsingSimpleCollision ? SimpleCollisionSubsectionSizeQuads * NumSubsections : 0;
-		CollisionComp->CachedLocalBox = CachedLocalBox;
-		CollisionComp->SetGenerateOverlapEvents(Proxy->bGenerateOverlapEvents);
-		CreatedNew = true;
+		RecreateCollisionComponent(bUsingSimpleCollision);
+		CollisionComp = CollisionComponent.Get();
+        MeshCollisionComponent = Cast<ULandscapeMeshCollisionComponent>(CollisionComp);
+	}
 
-		// Reallocate raw collision data
-		CollisionComp->CollisionHeightData.Lock(LOCK_READ_WRITE);
-		CollisionHeightData = (uint16*)CollisionComp->CollisionHeightData.Realloc(TotalCollisionSize);
-		FMemory::Memzero(CollisionHeightData, TotalCollisionSize * sizeof(uint16));
+	CollisionHeightData = (uint16*)CollisionComp->CollisionHeightData.Lock(LOCK_READ_WRITE);
 
-		if (XYOffsetmapTexture && MeshCollisionComponent)
-		{
-			// Need XYOffsetData for Collision Component
-			MeshCollisionComponent->CollisionXYOffsetData.Lock(LOCK_READ_WRITE);
-			CollisionXYOffsetData = (uint16*)MeshCollisionComponent->CollisionXYOffsetData.Realloc(TotalCollisionSize * 2);
-			FMemory::Memzero(CollisionXYOffsetData, TotalCollisionSize * 2 * sizeof(uint16));
-		}
-
-		if (DominantLayerData.Num())
-		{
-			CollisionComp->DominantLayerData.Lock(LOCK_READ_WRITE);
-			uint8* DestDominantLayerData = (uint8*)CollisionComp->DominantLayerData.Realloc(TotalCollisionSize);
-			FMemory::Memcpy(DestDominantLayerData, DominantLayerData.GetData(), TotalCollisionSize * sizeof(uint8));
-			CollisionComp->DominantLayerData.Unlock();
-		}
-
-		if (LayerInfos.Num())
-		{
-			CollisionComp->ComponentLayerInfos = MoveTemp(LayerInfos);
-		}
+	if (XYOffsetmapTexture && MeshCollisionComponent)
+	{
+		CollisionXYOffsetData = (uint16*)MeshCollisionComponent->CollisionXYOffsetData.Lock(LOCK_READ_WRITE);
 	}
 
 	const int32 HeightmapSizeU = GetHeightmap()->Source.GetSizeX();
 	const int32 HeightmapSizeV = GetHeightmap()->Source.GetSizeY();
-	const int32 MipSizeU = HeightmapSizeU >> CollisionMipLevel;
-	const int32 MipSizeV = HeightmapSizeV >> CollisionMipLevel;
-
-	const int32 XYMipSizeU = XYOffsetmapTexture ? XYOffsetmapTexture->Source.GetSizeX() >> CollisionMipLevel : 0;
-	const int32 XYMipSizeV = XYOffsetmapTexture ? XYOffsetmapTexture->Source.GetSizeY() >> CollisionMipLevel : 0;
-
-	// Ratio to convert update region coordinate to collision mip coordinates
-	const float CollisionQuadRatio = (float)CollisionSubsectionSizeQuads / (float)SubsectionSizeQuads;
-
-	// XY offset into heightmap mip data
-	const int32 HeightmapOffsetX = FMath::RoundToInt(HeightmapScaleBias.Z * (float)HeightmapSizeU) >> CollisionMipLevel;
-	const int32 HeightmapOffsetY = FMath::RoundToInt(HeightmapScaleBias.W * (float)HeightmapSizeV) >> CollisionMipLevel;
-
-	//int32 WeightmapOffsetX = FMath::RoundToInt(WeightmapScaleBias.Z * (float)XYMipSizeU) >> CollisionMipLevel;
-	//int32 WeightmapOffsetY = FMath::RoundToInt(WeightmapScaleBias.W * (float)XYMipSizeV) >> CollisionMipLevel;
 
 	// Handle Material WPO baked into heightfield collision
 	// Material WPO is not currently supported for mesh collision components
 	const bool bUsingGrassMapHeights = Proxy->bBakeMaterialPositionOffsetIntoCollision && !MeshCollisionComponent && GrassData->HasData() && !IsGrassMapOutdated();
-	const uint16* GrassHeights = nullptr;
+	uint16* GrassHeights = nullptr;
 	if (bUsingGrassMapHeights)
 	{
 		if (CollisionMipLevel == 0)
@@ -941,150 +882,13 @@ bool ULandscapeComponent::UpdateCollisionHeightData(const FColor* const Heightma
 		}
 	}
 
-	const int32 SubSectionX1 = FMath::Max(0, FMath::DivideAndRoundDown(ComponentX1 - 1, SubsectionSizeQuads));
-	const int32 SubSectionY1 = FMath::Max(0, FMath::DivideAndRoundDown(ComponentY1 - 1, SubsectionSizeQuads));
-	const int32 SubSectionX2 = FMath::Min(FMath::DivideAndRoundUp(ComponentX2 + 1, SubsectionSizeQuads), NumSubsections);
-	const int32 SubSectionY2 = FMath::Min(FMath::DivideAndRoundUp(ComponentY2 + 1, SubsectionSizeQuads), NumSubsections);
-	for (int32 SubsectionY = SubSectionY1; SubsectionY < SubSectionY2; ++SubsectionY)
-	{
-		for (int32 SubsectionX = SubSectionX1; SubsectionX < SubSectionX2; ++SubsectionX)
-		{
-			// Area to update in subsection coordinates
-			const int32 SubX1 = ComponentX1 - SubsectionSizeQuads * SubsectionX;
-			const int32 SubY1 = ComponentY1 - SubsectionSizeQuads * SubsectionY;
-			const int32 SubX2 = ComponentX2 - SubsectionSizeQuads * SubsectionX;
-			const int32 SubY2 = ComponentY2 - SubsectionSizeQuads * SubsectionY;
-
-			// Area to update in collision mip level coords
-			const int32 CollisionSubX1 = FMath::FloorToInt((float)SubX1 * CollisionQuadRatio);
-			const int32 CollisionSubY1 = FMath::FloorToInt((float)SubY1 * CollisionQuadRatio);
-			const int32 CollisionSubX2 = FMath::CeilToInt( (float)SubX2 * CollisionQuadRatio);
-			const int32 CollisionSubY2 = FMath::CeilToInt( (float)SubY2 * CollisionQuadRatio);
-
-			// Clamp area to update
-			const int32 VertX1 = FMath::Clamp<int32>(CollisionSubX1, 0, CollisionSubsectionSizeQuads);
-			const int32 VertY1 = FMath::Clamp<int32>(CollisionSubY1, 0, CollisionSubsectionSizeQuads);
-			const int32 VertX2 = FMath::Clamp<int32>(CollisionSubX2, 0, CollisionSubsectionSizeQuads);
-			const int32 VertY2 = FMath::Clamp<int32>(CollisionSubY2, 0, CollisionSubsectionSizeQuads);
-
-			for (int32 VertY = VertY1; VertY <= VertY2; VertY++)
-			{
-				for (int32 VertX = VertX1; VertX <= VertX2; VertX++)
-				{
-					// this uses Quads as we don't want the duplicated vertices
-					const int32 CompVertX = CollisionSubsectionSizeQuads * SubsectionX + VertX;
-					const int32 CompVertY = CollisionSubsectionSizeQuads * SubsectionY + VertY;
-
-					if (GrassHeights)
-					{
-						uint16& CollisionHeight = CollisionHeightData[CompVertX + CompVertY * CollisionSizeVerts];
-						const uint16& NewHeight = GrassHeights[CompVertX + CompVertY * CollisionSizeVerts];
-						CollisionHeight = NewHeight;
-					}
-					else
-					{
-						// X/Y of the vertex we're looking indexed into the texture data
-						const int32 TexX = HeightmapOffsetX + CollisionSubsectionSizeVerts * SubsectionX + VertX;
-						const int32 TexY = HeightmapOffsetY + CollisionSubsectionSizeVerts * SubsectionY + VertY;
-						const FColor& TexData = HeightmapTextureMipData[TexX + TexY * MipSizeU];
-
-						// Copy collision data
-						uint16& CollisionHeight = CollisionHeightData[CompVertX + CompVertY * CollisionSizeVerts];
-						const uint16 NewHeight = TexData.R << 8 | TexData.G;
-
-						CollisionHeight = NewHeight;
-					}
-
-					if (XYOffsetmapTexture && XYOffsetTextureMipData && CollisionXYOffsetData)
-					{
-						const int32 TexX = CollisionSubsectionSizeVerts * SubsectionX + VertX;
-						const int32 TexY = CollisionSubsectionSizeVerts * SubsectionY + VertY;
-						const FColor& TexData = XYOffsetTextureMipData[TexX + TexY * XYMipSizeU];
-
-						// Copy collision data
-						const uint16 NewXOffset = TexData.R << 8 | TexData.G;
-						const uint16 NewYOffset = TexData.B << 8 | TexData.A;
-
-						const int32 XYIndex = CompVertX + CompVertY * CollisionSizeVerts;
-						CollisionXYOffsetData[XYIndex * 2] = NewXOffset;
-						CollisionXYOffsetData[XYIndex * 2 + 1] = NewYOffset;
-					}
-				}
-			}
-		}
-	}
-
+	UpdateCollisionHeightBuffer(ComponentX1, ComponentY1, ComponentX2, ComponentY2, CollisionMipLevel, HeightmapSizeU, HeightmapSizeV, HeightmapTextureMipData, CollisionHeightData, GrassHeights, XYOffsetTextureMipData, CollisionXYOffsetData);
+		
 	if (bUsingSimpleCollision)
 	{
-		const int32 SimpleMipSizeU = HeightmapSizeU >> SimpleCollisionMipLevel;
-		const int32 SimpleMipSizeV = HeightmapSizeV >> SimpleCollisionMipLevel;
-
-		// Ratio to convert update region coordinate to collision mip coordinates
-		const float SimpleCollisionQuadRatio = (float)SimpleCollisionSubsectionSizeQuads / (float)SubsectionSizeQuads;
-
-		// XY offset into heightmap mip data
-		const int32 SimpleHeightmapOffsetX = FMath::RoundToInt(HeightmapScaleBias.Z * (float)HeightmapSizeU) >> SimpleCollisionMipLevel;
-		const int32 SimpleHeightmapOffsetY = FMath::RoundToInt(HeightmapScaleBias.W * (float)HeightmapSizeV) >> SimpleCollisionMipLevel;
-
-		const uint16* SimpleCollisionGrassHeights = bUsingGrassMapHeights && GrassData->HeightMipData.Contains(SimpleCollisionMipLevel) ? GrassData->HeightMipData[SimpleCollisionMipLevel].GetData() : nullptr;
-		uint16* const SimpleCollisionHeightData = CollisionHeightData + FMath::Square(CollisionSizeVerts);
-
-		for (int32 SubsectionY = SubSectionY1; SubsectionY < SubSectionY2; ++SubsectionY)
-		{
-			for (int32 SubsectionX = SubSectionX1; SubsectionX < SubSectionX2; ++SubsectionX)
-			{
-				// Area to update in subsection coordinates
-				const int32 SubX1 = ComponentX1 - SubsectionSizeQuads * SubsectionX;
-				const int32 SubY1 = ComponentY1 - SubsectionSizeQuads * SubsectionY;
-				const int32 SubX2 = ComponentX2 - SubsectionSizeQuads * SubsectionX;
-				const int32 SubY2 = ComponentY2 - SubsectionSizeQuads * SubsectionY;
-
-				// Area to update in simple collision mip level coords
-				const int32 SimpleCollisionSubX1 = FMath::FloorToInt((float)SubX1 * SimpleCollisionQuadRatio);
-				const int32 SimpleCollisionSubY1 = FMath::FloorToInt((float)SubY1 * SimpleCollisionQuadRatio);
-				const int32 SimpleCollisionSubX2 = FMath::CeilToInt((float)SubX2 * SimpleCollisionQuadRatio);
-				const int32 SimpleCollisionSubY2 = FMath::CeilToInt((float)SubY2 * SimpleCollisionQuadRatio);
-
-				// Clamp area to update
-				const int32 SimpleVertX1 = FMath::Clamp<int32>(SimpleCollisionSubX1, 0, SimpleCollisionSubsectionSizeQuads);
-				const int32 SimpleVertY1 = FMath::Clamp<int32>(SimpleCollisionSubY1, 0, SimpleCollisionSubsectionSizeQuads);
-				const int32 SimpleVertX2 = FMath::Clamp<int32>(SimpleCollisionSubX2, 0, SimpleCollisionSubsectionSizeQuads);
-				const int32 SimpleVertY2 = FMath::Clamp<int32>(SimpleCollisionSubY2, 0, SimpleCollisionSubsectionSizeQuads);
-
-				for (int32 VertY = SimpleVertY1; VertY <= SimpleVertY2; VertY++)
-				{
-					for (int32 VertX = SimpleVertX1; VertX <= SimpleVertX2; VertX++)
-					{
-						// this uses Quads as we don't want the duplicated vertices
-						const int32 CompVertX = SimpleCollisionSubsectionSizeQuads * SubsectionX + VertX;
-						const int32 CompVertY = SimpleCollisionSubsectionSizeQuads * SubsectionY + VertY;
-
-						if (SimpleCollisionGrassHeights)
-						{
-							uint16& CollisionHeight = SimpleCollisionHeightData[CompVertX + CompVertY * SimpleCollisionSizeVerts];
-							const uint16& NewHeight = SimpleCollisionGrassHeights[CompVertX + CompVertY * SimpleCollisionSizeVerts];
-							CollisionHeight = NewHeight;
-						}
-						else
-						{
-							// X/Y of the vertex we're looking indexed into the texture data
-							const int32 TexX = SimpleHeightmapOffsetX + SimpleCollisionSubsectionSizeVerts * SubsectionX + VertX;
-							const int32 TexY = SimpleHeightmapOffsetY + SimpleCollisionSubsectionSizeVerts * SubsectionY + VertY;
-							const FColor& TexData = SimpleCollisionHeightmapTextureData[TexX + TexY * SimpleMipSizeU];
-
-							// Copy collision data
-							uint16& CollisionHeight = SimpleCollisionHeightData[CompVertX + CompVertY * SimpleCollisionSizeVerts];
-							uint16 NewHeight = TexData.R << 8 | TexData.G;
-
-							CollisionHeight = NewHeight;
-						}
-
-						// not supported
-						// if (XYOffsetmapTexture && XYOffsetTextureMipData && CollisionXYOffsetData)
-					}
-				}
-			}
-		}
+		uint16* SimpleCollisionGrassHeights = bUsingGrassMapHeights && GrassData->HeightMipData.Contains(SimpleCollisionMipLevel) ? GrassData->HeightMipData[SimpleCollisionMipLevel].GetData() : nullptr;
+		uint16* const SimpleCollisionHeightData = CollisionHeightData + CollisionSize.SizeVertsSquare;
+		UpdateCollisionHeightBuffer(ComponentX1, ComponentY1, ComponentX2, ComponentY2, SimpleCollisionMipLevel, HeightmapSizeU, HeightmapSizeV, SimpleCollisionHeightmapTextureData, SimpleCollisionHeightData, SimpleCollisionGrassHeights, nullptr, nullptr);
 	}
 
 	CollisionComp->CollisionHeightData.Unlock();
@@ -1094,7 +898,7 @@ bool ULandscapeComponent::UpdateCollisionHeightData(const FColor* const Heightma
 		MeshCollisionComponent->CollisionXYOffsetData.Unlock();
 	}
 
-	// If we updated an existing component, we need to update the PhysX copy of the data
+	// If we updated an existing component, we need to update the phys x heightfield edit data
 	if (!CreatedNew)
 	{
 		if (MeshCollisionComponent)
@@ -1108,6 +912,8 @@ bool ULandscapeComponent::UpdateCollisionHeightData(const FColor* const Heightma
 		}
 		else
 		{
+			// Ratio to convert update region coordinate to collision mip coordinates
+			const float CollisionQuadRatio = (float)CollisionSize.SubsectionSizeQuads / (float)SubsectionSizeQuads;
 			const int32 CollisionCompX1 = FMath::FloorToInt((float)ComponentX1 * CollisionQuadRatio);
 			const int32 CollisionCompY1 = FMath::FloorToInt((float)ComponentY1 * CollisionQuadRatio);
 			const int32 CollisionCompX2 = FMath::CeilToInt( (float)ComponentX2 * CollisionQuadRatio);
@@ -1127,10 +933,7 @@ bool ULandscapeComponent::UpdateCollisionHeightData(const FColor* const Heightma
 	{
 		AInstancedFoliageActor::MoveInstancesToNewComponent(Proxy->GetWorld(), OldCollisionComponent, CollisionComp);
 	}
-
-	// Set new collision component to pointer
-	CollisionComponent = CollisionComp;
-
+	
 	if (CreatedNew && !ChangeType)
 	{
 		UpdateCollisionLayerData();
@@ -1187,6 +990,266 @@ void ULandscapeComponent::UpdateCollisionData()
 	}
 }
 
+void ULandscapeComponent::RecreateCollisionComponent(bool bUseSimpleCollision)
+{
+	ULandscapeHeightfieldCollisionComponent* CollisionComp = CollisionComponent.Get();
+	ULandscapeMeshCollisionComponent* MeshCollisionComponent = nullptr;
+	TArray<uint8> DominantLayerData;
+	TArray<ULandscapeLayerInfoObject*> LayerInfos;
+	ALandscapeProxy* Proxy = GetLandscapeProxy();
+	ULandscapeInfo* Info = GetLandscapeInfo();
+	const FCollisionSize CollisionSize = FCollisionSize::Create(NumSubsections, SubsectionSizeQuads, CollisionMipLevel);
+	const FCollisionSize SimpleCollisionSize = FCollisionSize::CreateSimple(bUseSimpleCollision, NumSubsections, SubsectionSizeQuads, SimpleCollisionMipLevel);
+	const int32 TotalCollisionSize = CollisionSize.SizeVertsSquare + SimpleCollisionSize.SizeVertsSquare;
+
+	if (CollisionComp) // remove old component before changing to other type collision...
+	{
+		if (CollisionComp->DominantLayerData.GetElementCount())
+		{
+			check(CollisionComp->DominantLayerData.GetElementCount() >= TotalCollisionSize);
+			DominantLayerData.AddUninitialized(TotalCollisionSize);
+
+			const uint8* SrcDominantLayerData = (uint8*)CollisionComp->DominantLayerData.Lock(LOCK_READ_ONLY);
+			FMemory::Memcpy(DominantLayerData.GetData(), SrcDominantLayerData, TotalCollisionSize * CollisionComp->DominantLayerData.GetElementSize());
+			CollisionComp->DominantLayerData.Unlock();
+		}
+
+		if (CollisionComp->ComponentLayerInfos.Num())
+		{
+			LayerInfos = CollisionComp->ComponentLayerInfos;
+		}
+
+		if (Info)
+		{
+			Info->Modify();
+		}
+		Proxy->Modify();
+		CollisionComp->DestroyComponent();
+		CollisionComp = nullptr;
+	}
+
+	if (XYOffsetmapTexture)
+	{
+		MeshCollisionComponent = NewObject<ULandscapeMeshCollisionComponent>(Proxy, NAME_None, RF_Transactional);
+		CollisionComp = MeshCollisionComponent;
+	}
+	else
+	{
+		MeshCollisionComponent = nullptr;
+		CollisionComp = NewObject<ULandscapeHeightfieldCollisionComponent>(Proxy, NAME_None, RF_Transactional);
+	}
+
+	CollisionComp->SetRelativeLocation(RelativeLocation);
+	CollisionComp->SetupAttachment(Proxy->GetRootComponent(), NAME_None);
+	Proxy->CollisionComponents.Add(CollisionComp);
+
+	CollisionComp->RenderComponent = this;
+	CollisionComp->SetSectionBase(GetSectionBase());
+	CollisionComp->CollisionSizeQuads = CollisionSize.SubsectionSizeQuads * NumSubsections;
+	CollisionComp->CollisionScale = (float)(ComponentSizeQuads) / (float)(CollisionComp->CollisionSizeQuads);
+	CollisionComp->SimpleCollisionSizeQuads = SimpleCollisionSize.SubsectionSizeQuads * NumSubsections;
+	CollisionComp->CachedLocalBox = CachedLocalBox;
+	CollisionComp->SetGenerateOverlapEvents(Proxy->bGenerateOverlapEvents);
+
+	// Reallocate raw collision data
+	CollisionComp->CollisionHeightData.Lock(LOCK_READ_WRITE);
+	uint16* CollisionHeightData = (uint16*)CollisionComp->CollisionHeightData.Realloc(TotalCollisionSize);
+	FMemory::Memzero(CollisionHeightData, TotalCollisionSize * CollisionComp->CollisionHeightData.GetElementSize());
+	CollisionComp->CollisionHeightData.Unlock();
+	
+	if (XYOffsetmapTexture && MeshCollisionComponent)
+	{
+		// Need XYOffsetData for Collision Component
+		MeshCollisionComponent->CollisionXYOffsetData.Lock(LOCK_READ_WRITE);
+		uint16* CollisionXYOffsetData = (uint16*)MeshCollisionComponent->CollisionXYOffsetData.Realloc(TotalCollisionSize * 2);
+		FMemory::Memzero(CollisionXYOffsetData, TotalCollisionSize * 2 * MeshCollisionComponent->CollisionXYOffsetData.GetElementSize());
+		MeshCollisionComponent->CollisionXYOffsetData.Unlock();
+	}
+
+	if (DominantLayerData.Num())
+	{
+		CollisionComp->DominantLayerData.Lock(LOCK_READ_WRITE);
+		uint8* DestDominantLayerData = (uint8*)CollisionComp->DominantLayerData.Realloc(TotalCollisionSize);
+		FMemory::Memcpy(DestDominantLayerData, DominantLayerData.GetData(), TotalCollisionSize * CollisionComp->DominantLayerData.GetElementSize());
+		CollisionComp->DominantLayerData.Unlock();
+	}
+
+	if (LayerInfos.Num())
+	{
+		CollisionComp->ComponentLayerInfos = MoveTemp(LayerInfos);
+	}
+	CollisionComponent = CollisionComp;
+}
+
+void ULandscapeComponent::UpdateCollisionHeightBuffer(	int32 InComponentX1, int32 InComponentY1, int32 InComponentX2, int32 InComponentY2, int32 InCollisionMipLevel, int32 InHeightmapSizeU, int32 InHeightmapSizeV,
+														const FColor* const InHeightmapTextureMipData, uint16* OutCollisionHeightData, uint16* InGrassHeightData,
+														const FColor* const InXYOffsetTextureMipData, uint16* OutCollisionXYOffsetData)
+{
+	FCollisionSize CollisionSize = FCollisionSize::Create(NumSubsections, SubsectionSizeQuads, InCollisionMipLevel);
+
+	// Ratio to convert update region coordinate to collision mip coordinates
+	const float CollisionQuadRatio = (float)CollisionSize.SubsectionSizeQuads / (float)SubsectionSizeQuads;
+
+	const int32 SubSectionX1 = FMath::Max(0, FMath::DivideAndRoundDown(InComponentX1 - 1, SubsectionSizeQuads));
+	const int32 SubSectionY1 = FMath::Max(0, FMath::DivideAndRoundDown(InComponentY1 - 1, SubsectionSizeQuads));
+	const int32 SubSectionX2 = FMath::Min(FMath::DivideAndRoundUp(InComponentX2 + 1, SubsectionSizeQuads), NumSubsections);
+	const int32 SubSectionY2 = FMath::Min(FMath::DivideAndRoundUp(InComponentY2 + 1, SubsectionSizeQuads), NumSubsections);
+
+	const int32 MipSizeU = InHeightmapSizeU >> InCollisionMipLevel;
+	const int32 MipSizeV = InHeightmapSizeV >> InCollisionMipLevel;
+
+	const int32 HeightmapOffsetX = FMath::RoundToInt(HeightmapScaleBias.Z * (float)InHeightmapSizeU) >> InCollisionMipLevel;
+	const int32 HeightmapOffsetY = FMath::RoundToInt(HeightmapScaleBias.W * (float)InHeightmapSizeV) >> InCollisionMipLevel;
+
+	const int32 XYMipSizeU = XYOffsetmapTexture ? XYOffsetmapTexture->Source.GetSizeX() >> InCollisionMipLevel : 0;
+
+	for (int32 SubsectionY = SubSectionY1; SubsectionY < SubSectionY2; ++SubsectionY)
+	{
+		for (int32 SubsectionX = SubSectionX1; SubsectionX < SubSectionX2; ++SubsectionX)
+		{
+			// Area to update in subsection coordinates
+			const int32 SubX1 = InComponentX1 - SubsectionSizeQuads * SubsectionX;
+			const int32 SubY1 = InComponentY1 - SubsectionSizeQuads * SubsectionY;
+			const int32 SubX2 = InComponentX2 - SubsectionSizeQuads * SubsectionX;
+			const int32 SubY2 = InComponentY2 - SubsectionSizeQuads * SubsectionY;
+
+			// Area to update in collision mip level coords
+			const int32 CollisionSubX1 = FMath::FloorToInt((float)SubX1 * CollisionQuadRatio);
+			const int32 CollisionSubY1 = FMath::FloorToInt((float)SubY1 * CollisionQuadRatio);
+			const int32 CollisionSubX2 = FMath::CeilToInt((float)SubX2 * CollisionQuadRatio);
+			const int32 CollisionSubY2 = FMath::CeilToInt((float)SubY2 * CollisionQuadRatio);
+
+			// Clamp area to update
+			const int32 VertX1 = FMath::Clamp<int32>(CollisionSubX1, 0, CollisionSize.SubsectionSizeQuads);
+			const int32 VertY1 = FMath::Clamp<int32>(CollisionSubY1, 0, CollisionSize.SubsectionSizeQuads);
+			const int32 VertX2 = FMath::Clamp<int32>(CollisionSubX2, 0, CollisionSize.SubsectionSizeQuads);
+			const int32 VertY2 = FMath::Clamp<int32>(CollisionSubY2, 0, CollisionSize.SubsectionSizeQuads);
+
+			for (int32 VertY = VertY1; VertY <= VertY2; VertY++)
+			{
+				for (int32 VertX = VertX1; VertX <= VertX2; VertX++)
+				{
+					// this uses Quads as we don't want the duplicated vertices
+					const int32 CompVertX = CollisionSize.SubsectionSizeQuads * SubsectionX + VertX;
+					const int32 CompVertY = CollisionSize.SubsectionSizeQuads * SubsectionY + VertY;
+
+					if (InGrassHeightData)
+					{
+						uint16& CollisionHeight = OutCollisionHeightData[CompVertX + CompVertY * CollisionSize.SizeVerts];
+						const uint16& NewHeight = InGrassHeightData[CompVertX + CompVertY * CollisionSize.SizeVerts];
+						CollisionHeight = NewHeight;
+					}
+					else
+					{
+						// X/Y of the vertex we're looking indexed into the texture data
+						const int32 TexX = HeightmapOffsetX + CollisionSize.SubsectionSizeVerts * SubsectionX + VertX;
+						const int32 TexY = HeightmapOffsetY + CollisionSize.SubsectionSizeVerts * SubsectionY + VertY;
+						const FColor& TexData = InHeightmapTextureMipData[TexX + TexY * MipSizeU];
+
+						// Copy collision data
+						uint16& CollisionHeight = OutCollisionHeightData[CompVertX + CompVertY * CollisionSize.SizeVerts];
+						const uint16 NewHeight = TexData.R << 8 | TexData.G;
+
+						CollisionHeight = NewHeight;
+					}
+
+					if (XYOffsetmapTexture && InXYOffsetTextureMipData && OutCollisionXYOffsetData)
+					{
+						const int32 TexX = CollisionSize.SubsectionSizeVerts * SubsectionX + VertX;
+						const int32 TexY = CollisionSize.SubsectionSizeVerts * SubsectionY + VertY;
+						const FColor& TexData = InXYOffsetTextureMipData[TexX + TexY * XYMipSizeU];
+
+						// Copy collision data
+						const uint16 NewXOffset = TexData.R << 8 | TexData.G;
+						const uint16 NewYOffset = TexData.B << 8 | TexData.A;
+
+						const int32 XYIndex = CompVertX + CompVertY * CollisionSize.SizeVerts;
+						OutCollisionXYOffsetData[XYIndex * 2] = NewXOffset;
+						OutCollisionXYOffsetData[XYIndex * 2 + 1] = NewYOffset;
+					}
+				}
+			}
+		}
+	}
+}
+
+void ULandscapeComponent::UpdateDominantLayerBuffer(int32 InComponentX1, int32 InComponentY1, int32 InComponentX2, int32 InComponentY2, int32 InCollisionMipLevel, int32 InWeightmapSizeU, int32 InDataLayerIdx, const TArray<uint8*>& InCollisionDataPtrs, uint8* OutDominantLayerData)
+{
+	const int32 MipSizeU = InWeightmapSizeU >> InCollisionMipLevel;
+
+	FCollisionSize CollisionSize = FCollisionSize::Create(NumSubsections, SubsectionSizeQuads, CollisionMipLevel);
+	
+	// Ratio to convert update region coordinate to collision mip coordinates
+	const float CollisionQuadRatio = (float)CollisionSize.SubsectionSizeQuads / (float)SubsectionSizeQuads;
+
+	const int32 SubSectionX1 = FMath::Max(0, FMath::DivideAndRoundDown(InComponentX1 - 1, SubsectionSizeQuads));
+	const int32 SubSectionY1 = FMath::Max(0, FMath::DivideAndRoundDown(InComponentY1 - 1, SubsectionSizeQuads));
+	const int32 SubSectionX2 = FMath::Min(FMath::DivideAndRoundUp(InComponentX2 + 1, SubsectionSizeQuads), NumSubsections);
+	const int32 SubSectionY2 = FMath::Min(FMath::DivideAndRoundUp(InComponentY2 + 1, SubsectionSizeQuads), NumSubsections);
+	for (int32 SubsectionY = SubSectionY1; SubsectionY < SubSectionY2; ++SubsectionY)
+	{
+		for (int32 SubsectionX = SubSectionX1; SubsectionX < SubSectionX2; ++SubsectionX)
+		{
+			// Area to update in subsection coordinates
+			const int32 SubX1 = InComponentX1 - SubsectionSizeQuads * SubsectionX;
+			const int32 SubY1 = InComponentY1 - SubsectionSizeQuads * SubsectionY;
+			const int32 SubX2 = InComponentX2 - SubsectionSizeQuads * SubsectionX;
+			const int32 SubY2 = InComponentY2 - SubsectionSizeQuads * SubsectionY;
+
+			// Area to update in collision mip level coords
+			const int32 CollisionSubX1 = FMath::FloorToInt((float)SubX1 * CollisionQuadRatio);
+			const int32 CollisionSubY1 = FMath::FloorToInt((float)SubY1 * CollisionQuadRatio);
+			const int32 CollisionSubX2 = FMath::CeilToInt((float)SubX2 * CollisionQuadRatio);
+			const int32 CollisionSubY2 = FMath::CeilToInt((float)SubY2 * CollisionQuadRatio);
+
+			// Clamp area to update
+			const int32 VertX1 = FMath::Clamp<int32>(CollisionSubX1, 0, CollisionSize.SubsectionSizeQuads);
+			const int32 VertY1 = FMath::Clamp<int32>(CollisionSubY1, 0, CollisionSize.SubsectionSizeQuads);
+			const int32 VertX2 = FMath::Clamp<int32>(CollisionSubX2, 0, CollisionSize.SubsectionSizeQuads);
+			const int32 VertY2 = FMath::Clamp<int32>(CollisionSubY2, 0, CollisionSize.SubsectionSizeQuads);
+
+			for (int32 VertY = VertY1; VertY <= VertY2; VertY++)
+			{
+				for (int32 VertX = VertX1; VertX <= VertX2; VertX++)
+				{
+					// X/Y of the vertex we're looking indexed into the texture data
+					const int32 TexX = CollisionSize.SubsectionSizeVerts * SubsectionX + VertX;
+					const int32 TexY = CollisionSize.SubsectionSizeVerts * SubsectionY + VertY;
+					const int32 DataOffset = (TexX + TexY * MipSizeU) * sizeof(FColor);
+
+					uint8 DominantLayer = 255; // 255 as invalid value
+					int32 DominantWeight = 0;
+					for (int32 LayerIdx = 0; LayerIdx < InCollisionDataPtrs.Num(); LayerIdx++)
+					{
+						const uint8 LayerWeight = InCollisionDataPtrs[LayerIdx][DataOffset];
+
+						if (LayerIdx == InDataLayerIdx) // Override value for hole
+						{
+							if (LayerWeight > 170) // 255 * 0.66...
+							{
+								DominantLayer = LayerIdx;
+								DominantWeight = INT_MAX;
+							}
+						}
+						else if (LayerWeight > DominantWeight)
+						{
+							DominantLayer = LayerIdx;
+							DominantWeight = LayerWeight;
+						}
+					}
+
+					// this uses Quads as we don't want the duplicated vertices
+					const int32 CompVertX = CollisionSize.SubsectionSizeQuads * SubsectionX + VertX;
+					const int32 CompVertY = CollisionSize.SubsectionSizeQuads * SubsectionY + VertY;
+
+					// Set collision data
+					OutDominantLayerData[CompVertX + CompVertY * CollisionSize.SizeVerts] = DominantLayer;
+				}
+			}
+		}
+	}
+}
+
 void ULandscapeComponent::UpdateCollisionLayerData(const FColor* const* const WeightmapTextureMipData, const FColor* const* const SimpleCollisionWeightmapTextureMipData, int32 ComponentX1, int32 ComponentY1, int32 ComponentX2, int32 ComponentY2)
 {
 	ULandscapeInfo* Info = GetLandscapeInfo();
@@ -1249,17 +1312,11 @@ void ULandscapeComponent::UpdateCollisionLayerData(const FColor* const* const We
 		}
 		else
 		{
-			const int32 CollisionSubsectionSizeVerts = ((SubsectionSizeQuads + 1) >> CollisionMipLevel);
-			const int32 CollisionSubsectionSizeQuads = CollisionSubsectionSizeVerts - 1;
-			const int32 CollisionSizeVerts = NumSubsections * CollisionSubsectionSizeQuads + 1;
-			const int32 SimpleCollisionSubsectionSizeVerts = bUsingSimpleCollision ? ((SubsectionSizeQuads + 1) >> SimpleCollisionMipLevel) : 0;
-			const int32 SimpleCollisionSubsectionSizeQuads = bUsingSimpleCollision ? SimpleCollisionSubsectionSizeVerts - 1 : 0;
-			const int32 SimpleCollisionSizeVerts = bUsingSimpleCollision ? NumSubsections * SimpleCollisionSubsectionSizeQuads + 1 : 0;
-
-			const int32 TotalCollisionSize = FMath::Square(CollisionSizeVerts) + FMath::Square(SimpleCollisionSizeVerts);
-
-			uint8* DominantLayerData = nullptr;
-
+			uint8* DominantLayerData = (uint8*)CollisionComp->DominantLayerData.Lock(LOCK_READ_WRITE);
+			FCollisionSize CollisionSize = FCollisionSize::Create(NumSubsections, SubsectionSizeQuads, CollisionMipLevel);
+			FCollisionSize SimpleCollisionSize = FCollisionSize::CreateSimple(bUsingSimpleCollision, NumSubsections, SubsectionSizeQuads, SimpleCollisionMipLevel);
+		
+				
 			// If there's no existing data, or the layer allocations have changed, we need to update the data for the whole component.
 			if (bExistingLayerMismatch || CollisionComp->DominantLayerData.GetElementCount() == 0)
 			{
@@ -1267,11 +1324,12 @@ void ULandscapeComponent::UpdateCollisionLayerData(const FColor* const* const We
 				ComponentY1 = 0;
 				ComponentX2 = ComponentSizeQuads;
 				ComponentY2 = ComponentSizeQuads;
+											
+				const int32 TotalCollisionSize = CollisionSize.SizeVertsSquare + SimpleCollisionSize.SizeVertsSquare;
 
-				CollisionComp->DominantLayerData.Lock(LOCK_READ_WRITE);
+				
 				DominantLayerData = (uint8*)CollisionComp->DominantLayerData.Realloc(TotalCollisionSize);
 				FMemory::Memzero(DominantLayerData, TotalCollisionSize);
-
 				CollisionComp->ComponentLayerInfos = CandidateLayers;
 			}
 			else
@@ -1280,157 +1338,17 @@ void ULandscapeComponent::UpdateCollisionLayerData(const FColor* const* const We
 				ComponentY1 = FMath::Clamp(ComponentY1, 0, ComponentSizeQuads);
 				ComponentX2 = FMath::Clamp(ComponentX2, 0, ComponentSizeQuads);
 				ComponentY2 = FMath::Clamp(ComponentY2, 0, ComponentSizeQuads);
-				DominantLayerData = (uint8*)CollisionComp->DominantLayerData.Lock(LOCK_READ_WRITE);
 			}
 
 			const int32 WeightmapSizeU = ComponentWeightmapsTexture[0]->Source.GetSizeX();
-			const int32 MipSizeU = WeightmapSizeU >> CollisionMipLevel;
-
-			// Ratio to convert update region coordinate to collision mip coordinates
-			const float CollisionQuadRatio = (float)CollisionSubsectionSizeQuads / (float)SubsectionSizeQuads;
-
-			// gmartin: WeightmapScaleBias not handled?
-
-			const int32 SubSectionX1 = FMath::Max(0, FMath::DivideAndRoundDown(ComponentX1 - 1, SubsectionSizeQuads));
-			const int32 SubSectionY1 = FMath::Max(0, FMath::DivideAndRoundDown(ComponentY1 - 1, SubsectionSizeQuads));
-			const int32 SubSectionX2 = FMath::Min(FMath::DivideAndRoundUp(ComponentX2 + 1, SubsectionSizeQuads), NumSubsections);
-			const int32 SubSectionY2 = FMath::Min(FMath::DivideAndRoundUp(ComponentY2 + 1, SubsectionSizeQuads), NumSubsections);
-			for (int32 SubsectionY = SubSectionY1; SubsectionY < SubSectionY2; ++SubsectionY)
-			{
-				for (int32 SubsectionX = SubSectionX1; SubsectionX < SubSectionX2; ++SubsectionX)
-				{
-					// Area to update in subsection coordinates
-					const int32 SubX1 = ComponentX1 - SubsectionSizeQuads * SubsectionX;
-					const int32 SubY1 = ComponentY1 - SubsectionSizeQuads * SubsectionY;
-					const int32 SubX2 = ComponentX2 - SubsectionSizeQuads * SubsectionX;
-					const int32 SubY2 = ComponentY2 - SubsectionSizeQuads * SubsectionY;
-
-					// Area to update in collision mip level coords
-					const int32 CollisionSubX1 = FMath::FloorToInt((float)SubX1 * CollisionQuadRatio);
-					const int32 CollisionSubY1 = FMath::FloorToInt((float)SubY1 * CollisionQuadRatio);
-					const int32 CollisionSubX2 = FMath::CeilToInt( (float)SubX2 * CollisionQuadRatio);
-					const int32 CollisionSubY2 = FMath::CeilToInt( (float)SubY2 * CollisionQuadRatio);
-
-					// Clamp area to update
-					const int32 VertX1 = FMath::Clamp<int32>(CollisionSubX1, 0, CollisionSubsectionSizeQuads);
-					const int32 VertY1 = FMath::Clamp<int32>(CollisionSubY1, 0, CollisionSubsectionSizeQuads);
-					const int32 VertX2 = FMath::Clamp<int32>(CollisionSubX2, 0, CollisionSubsectionSizeQuads);
-					const int32 VertY2 = FMath::Clamp<int32>(CollisionSubY2, 0, CollisionSubsectionSizeQuads);
-
-					for (int32 VertY = VertY1; VertY <= VertY2; VertY++)
-					{
-						for (int32 VertX = VertX1; VertX <= VertX2; VertX++)
-						{
-							// X/Y of the vertex we're looking indexed into the texture data
-							const int32 TexX = CollisionSubsectionSizeVerts * SubsectionX + VertX;
-							const int32 TexY = CollisionSubsectionSizeVerts * SubsectionY + VertY;
-							const int32 DataOffset = (TexX + TexY * MipSizeU) * sizeof(FColor);
-
-							uint8 DominantLayer = 255; // 255 as invalid value
-							int32 DominantWeight = 0;
-							for (int32 LayerIdx = 0; LayerIdx < CandidateDataPtrs.Num(); LayerIdx++)
-							{
-								const uint8 LayerWeight = CandidateDataPtrs[LayerIdx][DataOffset];
-
-								if (LayerIdx == DataLayerIdx) // Override value for hole
-								{
-									if (LayerWeight > 170) // 255 * 0.66...
-									{
-										DominantLayer = LayerIdx;
-										DominantWeight = INT_MAX;
-									}
-								}
-								else if (LayerWeight > DominantWeight)
-								{
-									DominantLayer = LayerIdx;
-									DominantWeight = LayerWeight;
-								}
-							}
-
-							// this uses Quads as we don't want the duplicated vertices
-							const int32 CompVertX = CollisionSubsectionSizeQuads * SubsectionX + VertX;
-							const int32 CompVertY = CollisionSubsectionSizeQuads * SubsectionY + VertY;
-
-							// Set collision data
-							DominantLayerData[CompVertX + CompVertY*CollisionSizeVerts] = DominantLayer;
-						}
-					}
-				}
-			}
+						
+			// gmartin: WeightmapScaleBias not handled?			
+			UpdateDominantLayerBuffer(ComponentX1, ComponentY1, ComponentX2, ComponentY2, CollisionMipLevel, WeightmapSizeU, DataLayerIdx, CandidateDataPtrs, DominantLayerData);
 
 			if (bUsingSimpleCollision)
 			{
-				const int32 SimpleMipSizeU = WeightmapSizeU >> SimpleCollisionMipLevel;
-
-				// Ratio to convert update region coordinate to collision mip coordinates
-				const float SimpleCollisionQuadRatio = (float)SimpleCollisionSubsectionSizeQuads / (float)SubsectionSizeQuads;
-
-				// gmartin: WeightmapScaleBias not handled?
-
-				uint8* const SimpleCollisionHeightData = DominantLayerData + FMath::Square(CollisionSizeVerts);
-
-				for (int32 SubsectionY = SubSectionY1; SubsectionY < SubSectionY2; ++SubsectionY)
-				{
-					for (int32 SubsectionX = SubSectionX1; SubsectionX < SubSectionX2; ++SubsectionX)
-					{
-						// Area to update in subsection coordinates
-						const int32 SubX1 = ComponentX1 - SubsectionSizeQuads * SubsectionX;
-						const int32 SubY1 = ComponentY1 - SubsectionSizeQuads * SubsectionY;
-						const int32 SubX2 = ComponentX2 - SubsectionSizeQuads * SubsectionX;
-						const int32 SubY2 = ComponentY2 - SubsectionSizeQuads * SubsectionY;
-
-						// Area to update in simple collision mip level coords
-						const int32 SimpleCollisionSubX1 = FMath::FloorToInt((float)SubX1 * SimpleCollisionQuadRatio);
-						const int32 SimpleCollisionSubY1 = FMath::FloorToInt((float)SubY1 * SimpleCollisionQuadRatio);
-						const int32 SimpleCollisionSubX2 = FMath::CeilToInt( (float)SubX2 * SimpleCollisionQuadRatio);
-						const int32 SimpleCollisionSubY2 = FMath::CeilToInt( (float)SubY2 * SimpleCollisionQuadRatio);
-
-						// Clamp area to update
-						const int32 SimpleVertX1 = FMath::Clamp<int32>(SimpleCollisionSubX1, 0, SimpleCollisionSubsectionSizeQuads);
-						const int32 SimpleVertY1 = FMath::Clamp<int32>(SimpleCollisionSubY1, 0, SimpleCollisionSubsectionSizeQuads);
-						const int32 SimpleVertX2 = FMath::Clamp<int32>(SimpleCollisionSubX2, 0, SimpleCollisionSubsectionSizeQuads);
-						const int32 SimpleVertY2 = FMath::Clamp<int32>(SimpleCollisionSubY2, 0, SimpleCollisionSubsectionSizeQuads);
-
-						for (int32 VertY = SimpleVertY1; VertY <= SimpleVertY2; VertY++)
-						{
-							for (int32 VertX = SimpleVertX1; VertX <= SimpleVertX2; VertX++)
-							{
-								// X/Y of the vertex we're looking indexed into the texture data
-								const int32 TexX = SimpleCollisionSubsectionSizeVerts * SubsectionX + VertX;
-								const int32 TexY = SimpleCollisionSubsectionSizeVerts * SubsectionY + VertY;
-								const int32 DataOffset = (TexX + TexY * SimpleMipSizeU) * sizeof(FColor);
-
-								uint8 DominantLayer = 255; // 255 as invalid value
-								int32 DominantWeight = 0;
-								for (int32 LayerIdx = 0; LayerIdx < CandidateDataPtrs.Num(); LayerIdx++)
-								{
-									const uint8 LayerWeight = SimpleCollisionDataPtrs[LayerIdx][DataOffset];
-
-									if (LayerIdx == DataLayerIdx) // Override value for hole
-									{
-										if (LayerWeight > 170) // 255 * 0.66...
-										{
-											DominantLayer = LayerIdx;
-											DominantWeight = INT_MAX;
-										}
-									}
-									else if (LayerWeight > DominantWeight)
-									{
-										DominantLayer = LayerIdx;
-										DominantWeight = LayerWeight;
-									}
-								}
-
-								// this uses Quads as we don't want the duplicated vertices
-								const int32 CompVertX = SimpleCollisionSubsectionSizeQuads * SubsectionX + VertX;
-								const int32 CompVertY = SimpleCollisionSubsectionSizeQuads * SubsectionY + VertY;
-
-								// Set collision data
-								DominantLayerData[CompVertX + CompVertY * SimpleCollisionSizeVerts] = DominantLayer;
-							}
-						}
-					}
-				}
+				uint8* const SimpleCollisionHeightData = DominantLayerData + CollisionSize.SizeVertsSquare;
+				UpdateDominantLayerBuffer(ComponentX1, ComponentY1, ComponentX2, ComponentY2, SimpleCollisionMipLevel, WeightmapSizeU, DataLayerIdx, SimpleCollisionDataPtrs, SimpleCollisionHeightData);
 			}
 
 			CollisionComp->DominantLayerData.Unlock();
