Clean-up unused landscape weightmap CPUReadback resources when reallocating weightmaps

diff --git a/Engine/Source/Runtime/Landscape/Private/LandscapeEditLayers.cpp b/Engine/Source/Runtime/Landscape/Private/LandscapeEditLayers.cpp
--- a/Engine/Source/Runtime/Landscape/Private/LandscapeEditLayers.cpp
+++ b/Engine/Source/Runtime/Landscape/Private/LandscapeEditLayers.cpp
@@ -3170,6 +3170,50 @@ void ALandscape::ReallocateLayersWeightmaps(const TArray<ULandscapeLayerInfoObje
 	OutComponentThatNeedMaterialRebuild.Append(AllLandscapeComponents);
 	// END HACK
 
+	// Clean-up unused weightmap CPUReadback resources
+	Info->ForAllLandscapeProxies([&AllLandscapeComponents](ALandscapeProxy* Proxy)
+	{
+		TArray<UTexture2D*, TInlineAllocator<64>> EntriesToRemoveFromMap;
+		for (auto& Pair : Proxy->WeightmapsCPUReadBack)
+		{
+			UTexture2D* WeightmapTextureKey = Pair.Key;
+			bool IsTextureReferenced = false;
+			for (ULandscapeComponent* Component : AllLandscapeComponents)
+			{
+				for (UTexture2D* WeightmapTexture : Component->GetWeightmapTextures(false))
+				{
+					if (WeightmapTexture == WeightmapTextureKey)
+					{
+						IsTextureReferenced = true;
+						break;
+					}
+				}
+			}
+			if (!IsTextureReferenced)
+			{
+				EntriesToRemoveFromMap.Add(WeightmapTextureKey);
+				if (FLandscapeLayersTexture2DCPUReadBackResource* ResourceToDelete = Pair.Value)
+				{
+					BeginReleaseResource(ResourceToDelete);
+				}
+			}
+		}
+
+		if (EntriesToRemoveFromMap.Num())
+		{
+			FlushRenderingCommands();
+			for (UTexture2D* OldWeightmapTexture : EntriesToRemoveFromMap)
+			{
+				if (FLandscapeLayersTexture2DCPUReadBackResource** ResourceToDelete = Proxy->WeightmapsCPUReadBack.Find(OldWeightmapTexture))
+				{
+					check(*ResourceToDelete);
+					delete *ResourceToDelete;
+					Proxy->WeightmapsCPUReadBack.Remove(OldWeightmapTexture);
+				}
+			}
+		}
+	});
+
 	/*for (ULandscapeComponent* Component : AllLandscapeComponents)
 	{
 		if (OutComponentThatNeedMaterialRebuild.Contains(Component))
