Allow resizing of resources, this is required by many tools like AddComponent, MoveComponentToLevel, etc.

diff --git a/Engine/Source/Runtime/Landscape/Classes/Landscape.h b/Engine/Source/Runtime/Landscape/Classes/Landscape.h
--- a/Engine/Source/Runtime/Landscape/Classes/Landscape.h
+++ b/Engine/Source/Runtime/Landscape/Classes/Landscape.h
@@ -288,6 +288,7 @@ public:
 	LANDSCAPE_API TArray<int8>& GetBrushesOrderForLayer(int32 InLayerIndex, int32 InTargetType);
 	LANDSCAPE_API class ALandscapeBlueprintCustomBrush* GetBrushForLayer(int32 InLayerIndex, int32 InTargetType, int8 BrushIndex) const;
 	LANDSCAPE_API TArray<class ALandscapeBlueprintCustomBrush*> GetBrushesForLayer(int32 InLayerIndex, int32 InTargetType) const;
+	LANDSCAPE_API void CreateLayersRenderingResource(const FIntPoint& InComponentCounts);
 
 private:
 	void TickLayers(float DeltaTime, ELevelTick TickType, FActorTickFunction& ThisTickFunction);
@@ -308,7 +309,7 @@ private:
 	void PrepareComponentDataToPackMaterialLayersCS(int32 InCurrentWeightmapToProcessIndex, bool InOutputDebugName, const TArray<ULandscapeComponent*>& InAllLandscapeComponents, TArray<UTexture2D*>& InOutProcessedWeightmaps,
 													TArray<class FLandscapeLayersTexture2DCPUReadBackResource*>& OutProcessedCPUReadBackTexture, TArray<struct FLandscapeLayerWeightmapPackMaterialLayersComponentData>& OutComponentData);
 	void ReallocateLayersWeightmaps(const TArray<ULandscapeLayerInfoObject*>& InBrushRequiredAllocations, TArray<ULandscapeComponent*>& OutComponentThatNeedMaterialRebuild);
-	void InitLayersWeightmapResources(uint8 InLayerCount);
+	void InitializeLayersWeightmapResources();
 	bool GenerateZeroAllocationPerComponents(const TArray<ALandscapeProxy*>& InAllLandscape, const TMap<ULandscapeLayerInfoObject*, bool>& InWeightmapLayersBlendSubstractive);
 
 	void GenerateLayersRenderQuad(const FIntPoint& InVertexPosition, float InVertexSize, const FVector2D& InUVStart, const FVector2D& InUVSize, TArray<struct FLandscapeLayersTriangle>& OutTriangles) const;
@@ -338,7 +339,6 @@ private:
 
 	void InitializeLayers();
 	void InitializeLayersWeightmapUsage();
-	void CreateLayersRenderingResource(const FIntPoint& InComponentCounts);
 
 	void PrintLayersDebugRT(const FString& InContext, UTextureRenderTarget2D* InDebugRT, uint8 InMipRender = 0, bool InOutputHeight = true, bool InOutputNormals = false) const;
 	void PrintLayersDebugTextureResource(const FString& InContext, FTextureResource* InTextureResource, uint8 InMipRender = 0, bool InOutputHeight = true, bool InOutputNormals = false) const;
diff --git a/Engine/Source/Runtime/Landscape/Private/LandscapeEditLayers.cpp b/Engine/Source/Runtime/Landscape/Private/LandscapeEditLayers.cpp
--- a/Engine/Source/Runtime/Landscape/Private/LandscapeEditLayers.cpp
+++ b/Engine/Source/Runtime/Landscape/Private/LandscapeEditLayers.cpp
@@ -1251,7 +1251,13 @@ void ALandscape::CreateLayersRenderingResource(const FIntPoint& InComponentCount
 
 			FLandscapeLayersTexture2DCPUReadBackResource** CPUReadback = Proxy->HeightmapsCPUReadBack.Find(ComponentHeightmap);
 
-			if (CPUReadback == nullptr)
+			if (CPUReadback != nullptr)
+			{
+				ReleaseResourceAndFlush(*CPUReadback);
+				*CPUReadback = nullptr;
+			}
+
+			if (CPUReadback == nullptr || *CPUReadback == nullptr)
 			{
 				FLandscapeLayersTexture2DCPUReadBackResource* NewCPUReadback = new FLandscapeLayersTexture2DCPUReadBackResource(ComponentHeightmap->Source.GetSizeX(), ComponentHeightmap->Source.GetSizeY(), ComponentHeightmap->GetPixelFormat(), ComponentHeightmap->Source.GetNumMips());
 				BeginInitResource(NewCPUReadback);
@@ -1303,6 +1309,31 @@ void ALandscape::CreateLayersRenderingResource(const FIntPoint& InComponentCount
 			}
 		}
 	}
+	else // Simply resize the render target
+	{
+		int32 CurrentMipSizeX = ((SubsectionSizeQuads + 1) * NumSubsections) * InComponentCounts.X;
+		int32 CurrentMipSizeY = ((SubsectionSizeQuads + 1) * NumSubsections) * InComponentCounts.Y;
+
+		for (int32 i = 0; i < EHeightmapRTType::HeightmapRT_Count; ++i)
+		{
+			if (i < EHeightmapRTType::HeightmapRT_Mip1) // Landscape size RT
+			{
+				Landscape->HeightmapRTList[i]->ResizeTarget(FMath::RoundUpToPowerOfTwo(TotalVertexCountX), FMath::RoundUpToPowerOfTwo(TotalVertexCountY));
+			}
+			else // Mips
+			{
+				CurrentMipSizeX >>= 1;
+				CurrentMipSizeY >>= 1;
+				Landscape->HeightmapRTList[i]->ResizeTarget(FMath::RoundUpToPowerOfTwo(CurrentMipSizeX), FMath::RoundUpToPowerOfTwo(CurrentMipSizeY));
+			}
+
+			// Only generate required mips RT
+			if (CurrentMipSizeX == InComponentCounts.X && CurrentMipSizeY == InComponentCounts.Y)
+			{
+				break;
+			}
+		}
+	}
 
 	if (Landscape->WeightmapRTList.Num() == 0)
 	{
@@ -1343,6 +1374,34 @@ void ALandscape::CreateLayersRenderingResource(const FIntPoint& InComponentCount
 			}
 		}
 	}
+	else // Simply resize the render target
+	{
+		int32 CurrentMipSizeX = ((SubsectionSizeQuads + 1) * NumSubsections) * InComponentCounts.X;
+		int32 CurrentMipSizeY = ((SubsectionSizeQuads + 1) * NumSubsections) * InComponentCounts.Y;
+
+		for (int32 i = 0; i < EWeightmapRTType::WeightmapRT_Count; ++i)
+		{
+			if (i < EWeightmapRTType::WeightmapRT_Mip0) // Landscape size RT, only create the number of layer we have
+			{
+				Landscape->WeightmapRTList[i]->ResizeTarget(FMath::RoundUpToPowerOfTwo(TotalVertexCountX), FMath::RoundUpToPowerOfTwo(TotalVertexCountY));
+			}
+			else // Mips
+			{
+				Landscape->WeightmapRTList[i]->ResizeTarget(FMath::RoundUpToPowerOfTwo(CurrentMipSizeX), FMath::RoundUpToPowerOfTwo(CurrentMipSizeY));
+
+				CurrentMipSizeX >>= 1;
+				CurrentMipSizeY >>= 1;
+			}
+
+			// Only generate required mips RT
+			if (CurrentMipSizeX == InComponentCounts.X && CurrentMipSizeY == InComponentCounts.Y)
+			{
+				break;
+			}
+		}
+	}
+
+	InitializeLayersWeightmapResources();
 }
 
 FIntPoint ALandscape::ComputeComponentCounts() const
@@ -3255,56 +3314,63 @@ void ALandscape::ReallocateLayersWeightmaps(const TArray<ULandscapeLayerInfoObje
 	*/
 }
 
-void ALandscape::InitLayersWeightmapResources(uint8 InLayerCount)
+void ALandscape::InitializeLayersWeightmapResources()
 {
-	// Use the 1st one to compute the resource as they are all the same anyway
-	UTextureRenderTarget2D* FirstWeightmapRT = WeightmapRTList[WeightmapRT_Scratch1];
+	ULandscapeInfo* Info = GetLandscapeInfo();
 
-	if (CombinedLayersWeightmapAllMaterialLayersResource != nullptr && InLayerCount != CombinedLayersWeightmapAllMaterialLayersResource->GetSizeZ())
+	if (Info == nullptr || Info->Layers.Num() == 0)
 	{
-		ReleaseResourceAndFlush(CombinedLayersWeightmapAllMaterialLayersResource);
-		CombinedLayersWeightmapAllMaterialLayersResource = nullptr;
+		return;
 	}
 
-	if (CombinedLayersWeightmapAllMaterialLayersResource == nullptr)
-	{
-		CombinedLayersWeightmapAllMaterialLayersResource = new FLandscapeTexture2DArrayResource(FirstWeightmapRT->SizeX, FirstWeightmapRT->SizeY, InLayerCount, PF_G8, 1, true);
-		BeginInitResource(CombinedLayersWeightmapAllMaterialLayersResource);
-	}
+	// Destroy existing resource
+	TArray<FTextureResource*> ResourceToDestroy;
+	ResourceToDestroy.Add(CombinedLayersWeightmapAllMaterialLayersResource);
+	ResourceToDestroy.Add(CurrentLayersWeightmapAllMaterialLayersResource);
+	ResourceToDestroy.Add(WeightmapScratchExtractLayerTextureResource);
+	ResourceToDestroy.Add(WeightmapScratchPackLayerTextureResource);
 
-	if (CurrentLayersWeightmapAllMaterialLayersResource != nullptr && InLayerCount != CurrentLayersWeightmapAllMaterialLayersResource->GetSizeZ())
+	for (FTextureResource* Resource : ResourceToDestroy)
 	{
-		ReleaseResourceAndFlush(CurrentLayersWeightmapAllMaterialLayersResource);
-		CurrentLayersWeightmapAllMaterialLayersResource = nullptr;
+		if (Resource != nullptr)
+		{
+			ENQUEUE_RENDER_COMMAND(ReleaseCommand)(
+				[Resource](FRHICommandList& RHICmdList)
+			{
+				Resource->ReleaseResource();
+				delete Resource;
+			});
+		}
 	}
 
-	if (CurrentLayersWeightmapAllMaterialLayersResource == nullptr)
-	{
-		CurrentLayersWeightmapAllMaterialLayersResource = new FLandscapeTexture2DArrayResource(FirstWeightmapRT->SizeX, FirstWeightmapRT->SizeY, InLayerCount, PF_G8, 1, true);
-		BeginInitResource(CurrentLayersWeightmapAllMaterialLayersResource);
-	}
+	// Create resources
 
-	if (WeightmapScratchExtractLayerTextureResource == nullptr)
-	{
-		WeightmapScratchExtractLayerTextureResource = new FLandscapeTexture2DResource(FirstWeightmapRT->SizeX, FirstWeightmapRT->SizeY, PF_B8G8R8A8, 1, false);
-		BeginInitResource(WeightmapScratchExtractLayerTextureResource);
-	}
+	int32 LayerCount = Info->Layers.Num() + 1; // due to visibility being stored at 0
 
-	if (WeightmapScratchPackLayerTextureResource == nullptr)
-	{
-		int32 MipCount = 0;
+	// Use the 1st one to compute the resource as they are all the same anyway
+	UTextureRenderTarget2D* FirstWeightmapRT = WeightmapRTList[WeightmapRT_Scratch1];
+
+	CombinedLayersWeightmapAllMaterialLayersResource = new FLandscapeTexture2DArrayResource(FirstWeightmapRT->SizeX, FirstWeightmapRT->SizeY, LayerCount, PF_G8, 1, true);
+	BeginInitResource(CombinedLayersWeightmapAllMaterialLayersResource);
 
-		for (int32 MipRTIndex = EWeightmapRTType::WeightmapRT_Mip0; MipRTIndex < EWeightmapRTType::WeightmapRT_Count; ++MipRTIndex)
+	CurrentLayersWeightmapAllMaterialLayersResource = new FLandscapeTexture2DArrayResource(FirstWeightmapRT->SizeX, FirstWeightmapRT->SizeY, LayerCount, PF_G8, 1, true);
+	BeginInitResource(CurrentLayersWeightmapAllMaterialLayersResource);
+
+	WeightmapScratchExtractLayerTextureResource = new FLandscapeTexture2DResource(FirstWeightmapRT->SizeX, FirstWeightmapRT->SizeY, PF_B8G8R8A8, 1, false);
+	BeginInitResource(WeightmapScratchExtractLayerTextureResource);
+
+	int32 MipCount = 0;
+
+	for (int32 MipRTIndex = EWeightmapRTType::WeightmapRT_Mip0; MipRTIndex < EWeightmapRTType::WeightmapRT_Count; ++MipRTIndex)
+	{
+		if (WeightmapRTList[MipRTIndex] != nullptr)
 		{
-			if (WeightmapRTList[MipRTIndex] != nullptr)
-			{
-				++MipCount;
-			}
+			++MipCount;
 		}
-
-		WeightmapScratchPackLayerTextureResource = new FLandscapeTexture2DResource(FirstWeightmapRT->SizeX, FirstWeightmapRT->SizeY, PF_B8G8R8A8, MipCount, true);
-		BeginInitResource(WeightmapScratchPackLayerTextureResource);
 	}
+
+	WeightmapScratchPackLayerTextureResource = new FLandscapeTexture2DResource(FirstWeightmapRT->SizeX, FirstWeightmapRT->SizeY, PF_B8G8R8A8, MipCount, true);
+	BeginInitResource(WeightmapScratchPackLayerTextureResource);
 }
 
 bool ALandscape::PrepareLayersWeightmapTextureResources() const
@@ -3397,8 +3463,6 @@ void ALandscape::RegenerateLayersWeightmaps()
 		FLandscapeLayersWeightmapShaderParameters PSShaderParams;
 		bool OutputDebugName = (CVarOutputLayersDebugDrawCallName.GetValueOnAnyThread() == 1 || CVarOutputLayersRTContent.GetValueOnAnyThread() == 1 || CVarOutputLayersWeightmapsRTContent.GetValueOnAnyThread() == 1) ? true : false;
 
-		InitLayersWeightmapResources(LayerCount);
-
 		ClearLayersWeightmapTextureResource(TEXT("ClearRT RGBA"), EmptyRT->GameThread_GetRenderTargetResource());
 		ClearLayersWeightmapTextureResource(TEXT("ClearRT R"), LandscapeScratchRT1->GameThread_GetRenderTargetResource());
 
