- Landscape Layer System: Fix Layer Debug Visualizer

#rnx

diff --git a/Engine/Source/Runtime/Landscape/Private/Landscape.cpp b/Engine/Source/Runtime/Landscape/Private/Landscape.cpp
--- a/Engine/Source/Runtime/Landscape/Private/Landscape.cpp
+++ b/Engine/Source/Runtime/Landscape/Private/Landscape.cpp
@@ -540,7 +540,7 @@ void ULandscapeComponent::GetLayerDebugColorKey(int32& R, int32& G, int32& B) co
 			if (LayerStruct.DebugColorChannel > 0
 				&& LayerStruct.LayerInfoObj)
 			{
-				const TArray<FWeightmapLayerAllocationInfo>& ComponentWeightmapLayerAllocations = GetWeightmapLayerAllocations(true);
+				const TArray<FWeightmapLayerAllocationInfo>& ComponentWeightmapLayerAllocations = GetWeightmapLayerAllocations();
 
 				for (int32 LayerIdx = 0; LayerIdx < ComponentWeightmapLayerAllocations.Num(); LayerIdx++)
 				{
diff --git a/Engine/Source/Runtime/Landscape/Private/LandscapeEditLayers.cpp b/Engine/Source/Runtime/Landscape/Private/LandscapeEditLayers.cpp
--- a/Engine/Source/Runtime/Landscape/Private/LandscapeEditLayers.cpp
+++ b/Engine/Source/Runtime/Landscape/Private/LandscapeEditLayers.cpp
@@ -3957,6 +3957,8 @@ void ALandscape::UpdateLayersMaterialInstances(const TArray<ULandscapeComponent*
 			Component->LODIndexToMaterialIndex.Empty(1);
 			Component->LODIndexToMaterialIndex.Add(0);
 		}
+
+		Component->EditToolRenderData.UpdateDebugColorMaterial(Component);
 	}
 
 	// End material update
