- Landscape Layer System: Fixup Component Layer Data on Load

#rnx

diff --git a/Engine/Source/Editor/LandscapeEditor/Private/LandscapeEdModeComponentTools.cpp b/Engine/Source/Editor/LandscapeEditor/Private/LandscapeEdModeComponentTools.cpp
--- a/Engine/Source/Editor/LandscapeEditor/Private/LandscapeEdModeComponentTools.cpp
+++ b/Engine/Source/Editor/LandscapeEditor/Private/LandscapeEdModeComponentTools.cpp
@@ -611,7 +611,7 @@ public:
 							UTexture2D* CurrentWeightmapTexture = nullptr;
 							ULandscapeWeightmapUsage* CurrentWeightmapUsage = nullptr;
 
-							if (TotalNeededChannels < 4)
+							if (TotalNeededChannels < ULandscapeWeightmapUsage::NumChannels)
 							{
 								// UE_LOG(LogLandscape, Log, TEXT("Looking for nearest"));
 
@@ -623,7 +623,7 @@ public:
 									if (TryWeightmapUsage->FreeChannelCount() >= TotalNeededChannels) // TODO: handle layers
 									{
 										// See if this candidate is closer than any others we've found
-										for (int32 ChanIdx = 0; ChanIdx < 4; ChanIdx++)
+										for (int32 ChanIdx = 0; ChanIdx < ULandscapeWeightmapUsage::NumChannels; ChanIdx++)
 										{
 											if (TryWeightmapUsage->ChannelUsage[ChanIdx] != nullptr)
 											{
@@ -663,7 +663,7 @@ public:
 
 							NewWeightmapTextures.Add(CurrentWeightmapTexture);
 
-							for (int32 ChanIdx = 0; ChanIdx < 4 && TotalNeededChannels > 0; ChanIdx++)
+							for (int32 ChanIdx = 0; ChanIdx < ULandscapeWeightmapUsage::NumChannels && TotalNeededChannels > 0; ChanIdx++)
 							{
 								// UE_LOG(LogLandscape, Log, TEXT("Finding allocation for layer %d"), CurrentLayer);
 
diff --git a/Engine/Source/Runtime/Landscape/Classes/Landscape.h b/Engine/Source/Runtime/Landscape/Classes/Landscape.h
--- a/Engine/Source/Runtime/Landscape/Classes/Landscape.h
+++ b/Engine/Source/Runtime/Landscape/Classes/Landscape.h
@@ -252,7 +252,6 @@ public:
 	LANDSCAPE_API bool ReorderLayer(int32 InStartingLayerIndex, int32 InDestinationLayerIndex);
 	LANDSCAPE_API void CreateLayer(FName InName = NAME_None);
 	LANDSCAPE_API void CreateDefaultLayer(const FIntPoint& InComponentCounts);
-	LANDSCAPE_API void InitializeLayerWithEmptyContent(const FGuid& InLayerGUID);
 	LANDSCAPE_API void CopyOldDataToDefaultLayer();
 	LANDSCAPE_API TMap<UTexture2D*, TArray<ULandscapeComponent*>> GenerateComponentsPerHeightmaps() const;
 	LANDSCAPE_API FIntPoint ComputeComponentCounts() const;
diff --git a/Engine/Source/Runtime/Landscape/Classes/LandscapeProxy.h b/Engine/Source/Runtime/Landscape/Classes/LandscapeProxy.h
--- a/Engine/Source/Runtime/Landscape/Classes/LandscapeProxy.h
+++ b/Engine/Source/Runtime/Landscape/Classes/LandscapeProxy.h
@@ -816,8 +816,9 @@ public:
 
 	// Copy properties from parent Landscape actor
 	LANDSCAPE_API void GetSharedProperties(ALandscapeProxy* Landscape);
-	// Assign only mismatched properties and mark proxy package dirty
-	LANDSCAPE_API void ConditionalAssignCommonProperties(ALandscape* Landscape);
+
+	// Assign only mismatching data and mark proxy package dirty
+	LANDSCAPE_API void FixupSharedData(ALandscape* Landscape);
 	
 	/** Get the LandcapeActor-to-world transform with respect to landscape section offset*/
 	LANDSCAPE_API FTransform LandscapeActorToWorld() const;
@@ -939,6 +940,26 @@ public:
 	DECLARE_EVENT(ALandscape, FLandscapeMaterialChangedDelegate);
 	FLandscapeMaterialChangedDelegate& OnMaterialChangedDelegate() { return LandscapeMaterialChangedDelegate; }
 
+protected:
+	friend class ALandscape;
+
+	/** Add Layer if it doesn't exist yet.
+	* @return True if layer was added.
+	*/
+	LANDSCAPE_API bool AddLayer(const FGuid& InLayerGuid);
+
+	/** Delete Layer.
+	*/
+	LANDSCAPE_API void DeleteLayer(const FGuid& InLayerGuid);
+
+	/** Remove Layers not found in InExistingLayers
+	* @return True if some layers were removed.
+	*/
+	LANDSCAPE_API bool RemoveObsoleteLayers(const TSet<FGuid>& InExistingLayers);
+
+	/** Initialize Layer with empty content if it hasn't been initialized yet. */
+	LANDSCAPE_API void InitializeLayerWithEmptyContent(const FGuid& InLayerGuid);
+
 protected:
 	FLandscapeMaterialChangedDelegate LandscapeMaterialChangedDelegate;
 
diff --git a/Engine/Source/Runtime/Landscape/Classes/LandscapeWeightmapUsage.h b/Engine/Source/Runtime/Landscape/Classes/LandscapeWeightmapUsage.h
--- a/Engine/Source/Runtime/Landscape/Classes/LandscapeWeightmapUsage.h
+++ b/Engine/Source/Runtime/Landscape/Classes/LandscapeWeightmapUsage.h
@@ -14,9 +14,11 @@ class ULandscapeWeightmapUsage : public UObject
 {
 	GENERATED_UCLASS_BODY()
 
+	enum { NumChannels = 4 };
+
 public:
 	UPROPERTY()
-	ULandscapeComponent* ChannelUsage[4];
+	ULandscapeComponent* ChannelUsage[NumChannels];
 
 	UPROPERTY()
 	FGuid LayerGuid;
@@ -25,7 +27,7 @@ public:
 	{
 		int32 Count = 0;
 
-		for (int8 i = 0; i < 4; ++i)
+		for (int8 i = 0; i < NumChannels; ++i)
 		{
 			Count += (ChannelUsage[i] == nullptr) ? 1 : 0;
 		}
@@ -35,9 +37,14 @@ public:
 
 	void ClearUsage()
 	{
-		for (int8 i = 0; i < 4; ++i)
+		for (int8 i = 0; i < NumChannels; ++i)
 		{
 			ChannelUsage[i] = nullptr;
 		}
 	}
+
+	bool IsEmpty() const
+	{
+		return FreeChannelCount() == NumChannels;
+	}
 };
diff --git a/Engine/Source/Runtime/Landscape/Private/Landscape.cpp b/Engine/Source/Runtime/Landscape/Private/Landscape.cpp
--- a/Engine/Source/Runtime/Landscape/Private/Landscape.cpp
+++ b/Engine/Source/Runtime/Landscape/Private/Landscape.cpp
@@ -55,6 +55,7 @@ Landscape.cpp: Terrain rendering
 #include "MaterialUtilities.h"
 #include "Settings/EditorExperimentalSettings.h"
 #include "Editor.h"
+#include "Algo/Transform.h"
 #endif
 #include "LandscapeVersion.h"
 #include "UObject/FortniteMainBranchObjectVersion.h"
@@ -1184,7 +1185,7 @@ void ULandscapeComponent::BeginDestroy()
 				{
 					(*Usage)->ChannelUsage[WeightmapLayerAllocations[LayerIdx].WeightmapTextureChannel] = nullptr;
 
-					if ((*Usage)->FreeChannelCount() == 4)
+					if ((*Usage)->IsEmpty())
 					{
 						Proxy->WeightmapUsageMap.Remove(WeightmapTexture);
 					}
@@ -1437,12 +1438,14 @@ void ULandscapeComponent::ForEachLayer(TFunctionRef<void(const FGuid&, struct FL
 
 void ULandscapeComponent::AddLayerData(const FGuid& InLayerGuid, const FLandscapeLayerComponentData& InData)
 {
+	Modify();
 	FLandscapeLayerComponentData& Data = LayersData.FindOrAdd(InLayerGuid);
 	Data = InData;
 }
 
 void ULandscapeComponent::RemoveLayerData(const FGuid& InLayerGuid)
 {
+	Modify();
 	LayersData.Remove(InLayerGuid);
 }
 
@@ -2158,7 +2161,7 @@ void ALandscapeProxy::GetSharedProperties(ALandscapeProxy* Landscape)
 	}
 }
 
-void ALandscapeProxy::ConditionalAssignCommonProperties(ALandscape* Landscape)
+void ALandscapeProxy::FixupSharedData(ALandscape* Landscape)
 {
 	if (Landscape == nullptr)
 	{
@@ -2226,6 +2229,17 @@ void ALandscapeProxy::ConditionalAssignCommonProperties(ALandscape* Landscape)
 		TargetDisplayOrderList = Landscape->TargetDisplayOrderList;
 		bUpdated = true;
 	}
+		
+	
+	TSet<FGuid> LayerGuids;
+	Algo::Transform(Landscape->LandscapeLayers, LayerGuids, [](const FLandscapeLayer& Layer) { return Layer.Guid; });
+	bUpdated |= RemoveObsoleteLayers(LayerGuids);
+	
+	for (const FLandscapeLayer& Layer : Landscape->LandscapeLayers)
+	{
+		bUpdated |= AddLayer(Layer.Guid);
+	}
+
 
 	if (bUpdated)
 	{
@@ -2526,7 +2540,7 @@ void ULandscapeInfo::RegisterActor(ALandscapeProxy* Proxy, bool bMapCheck)
 		for (ALandscapeStreamingProxy* StreamingProxy : Proxies)
 		{
 			StreamingProxy->LandscapeActor = LandscapeActor;
-			StreamingProxy->ConditionalAssignCommonProperties(Landscape);
+			StreamingProxy->FixupSharedData(Landscape);
 		}
 	}
 	else
@@ -2535,7 +2549,7 @@ void ULandscapeInfo::RegisterActor(ALandscapeProxy* Proxy, bool bMapCheck)
 
 		Proxies.Add(StreamingProxy);
 		StreamingProxy->LandscapeActor = LandscapeActor;
-		StreamingProxy->ConditionalAssignCommonProperties(LandscapeActor.Get());
+		StreamingProxy->FixupSharedData(LandscapeActor.Get());
 	}
 
 	UpdateLayerInfoMap(Proxy);
diff --git a/Engine/Source/Runtime/Landscape/Private/LandscapeEdit.cpp b/Engine/Source/Runtime/Landscape/Private/LandscapeEdit.cpp
--- a/Engine/Source/Runtime/Landscape/Private/LandscapeEdit.cpp
+++ b/Engine/Source/Runtime/Landscape/Private/LandscapeEdit.cpp
@@ -4638,7 +4638,7 @@ void ULandscapeComponent::ReallocateWeightmaps(FLandscapeEditDataInterface* Data
 
 				if (TryWeightmapUsage->FreeChannelCount() >= TotalNeededChannels && TryWeightmapUsage->LayerGuid == LayerGuidToSeek)
 				{
-					if (TryWeightmapUsage->FreeChannelCount() == 4)
+					if (TryWeightmapUsage->IsEmpty())
 					{
 						CurrentWeightmapTexture = ItPair.Key;
 						CurrentWeightmapUsage = TryWeightmapUsage;
@@ -4647,7 +4647,7 @@ void ULandscapeComponent::ReallocateWeightmaps(FLandscapeEditDataInterface* Data
 					else
 					{
 						// See if this candidate is closer than any others we've found
-						for (int32 ChanIdx = 0; ChanIdx < 4; ChanIdx++)
+						for (int32 ChanIdx = 0; ChanIdx < ULandscapeWeightmapUsage::NumChannels; ChanIdx++)
 						{
 							if (TryWeightmapUsage->ChannelUsage[ChanIdx] != nullptr)
 							{
@@ -4798,7 +4798,7 @@ void ALandscapeProxy::RemoveInvalidWeightmaps()
 		{
 			UTexture2D* Tex = It.Key();
 			ULandscapeWeightmapUsage* Usage = It.Value();
-			if (Usage->FreeChannelCount() == 4) // Invalid Weight-map
+			if (Usage->IsEmpty()) // Invalid Weight-map
 			{
 				if (Tex)
 				{
diff --git a/Engine/Source/Runtime/Landscape/Private/LandscapeEditLayers.cpp b/Engine/Source/Runtime/Landscape/Private/LandscapeEditLayers.cpp
--- a/Engine/Source/Runtime/Landscape/Private/LandscapeEditLayers.cpp
+++ b/Engine/Source/Runtime/Landscape/Private/LandscapeEditLayers.cpp
@@ -1504,7 +1504,7 @@ void ALandscape::CopyOldDataToDefaultLayer()
 						LayerData->WeightmapData.TextureUsages[TextureIndex] = WeightmapTextureData->Usage;
 						check(WeightmapTextureData->Usage->LayerGuid == DefaultLayer->Guid);
 
-						for (int32 ChannelIndex = 0; ChannelIndex < 4; ++ChannelIndex)
+						for (int32 ChannelIndex = 0; ChannelIndex < ULandscapeWeightmapUsage::NumChannels; ++ChannelIndex)
 						{
 							const ULandscapeComponent* ChannelLandscapeComponent = LayerData->WeightmapData.TextureUsages[TextureIndex]->ChannelUsage[ChannelIndex];
 
@@ -1537,7 +1537,7 @@ void ALandscape::CopyOldDataToDefaultLayer()
 						LayerData->WeightmapData.Textures[TextureIndex] = NewLayerWeightmapTexture;
 						LayerData->WeightmapData.TextureUsages[TextureIndex] = Proxy->WeightmapUsageMap.Add(NewLayerWeightmapTexture, Proxy->CreateWeightmapUsage());
 
-						for (int32 ChannelIndex = 0; ChannelIndex < 4; ++ChannelIndex)
+						for (int32 ChannelIndex = 0; ChannelIndex < ULandscapeWeightmapUsage::NumChannels; ++ChannelIndex)
 						{
 							LayerData->WeightmapData.TextureUsages[TextureIndex]->ChannelUsage[ChannelIndex] = ComponentWeightmapTexturesUsage[TextureIndex]->ChannelUsage[ChannelIndex];
 						}
@@ -1627,129 +1627,6 @@ void ALandscape::InitializeLayersWeightmapUsage()
 	});
 }
 
-void ALandscape::InitializeLayerWithEmptyContent(const FGuid& InLayerGUID)
-{
-	ULandscapeInfo* Info = GetLandscapeInfo();
-
-	if (Info == nullptr)
-	{
-		return;
-	}
-
-	FIntPoint NumComponents = ComputeComponentCounts();
-
-	// Build a mapping between each Heightmaps and Component in them
-	TMap<UTexture2D*, TArray<ULandscapeComponent*>> ComponentsPerHeightmaps;
-
-	Info->ForAllLandscapeProxies([&ComponentsPerHeightmaps](ALandscapeProxy* Proxy)
-	{
-		for (ULandscapeComponent* Component : Proxy->LandscapeComponents)
-		{
-			UTexture2D* ComponentHeightmapTexture = Component->GetHeightmap();
-			TArray<ULandscapeComponent*>& ComponentList = ComponentsPerHeightmaps.FindOrAdd(ComponentHeightmapTexture);
-			ComponentList.Add(Component);
-		}
-	});
-
-	FIntPoint TotalVertexCount((SubsectionSizeQuads * NumSubsections) * NumComponents.X + 1, (SubsectionSizeQuads * NumSubsections) * NumComponents.Y + 1);
-
-	FLandscapeLayer* Layer = LandscapeLayers.FindByPredicate([&InLayerGUID](const FLandscapeLayer& Other) { return Other.Guid == InLayerGUID; });
-
-	// Init layers with valid "empty" data
-	Info->ForAllLandscapeProxies([&](ALandscapeProxy* Proxy)
-	{
-		TMap<UTexture2D*, UTexture2D*> CreatedHeightmapTextures; // < Final layer texture, New created texture for layer
-
-		for (ULandscapeComponent* Component : Proxy->LandscapeComponents)
-		{
-			UTexture2D* ComponentHeightmap = Component->GetHeightmap();
-			const TArray<ULandscapeComponent*>* ComponentsUsingHeightmap = ComponentsPerHeightmaps.Find(ComponentHeightmap);
-			check(ComponentsUsingHeightmap != nullptr);
-
-			// Compute Global layers data
-			FLandscapeLayersGlobalComponentData& GlobalLayersData = Component->GetGlobalLayersData();
-			GlobalLayersData.TopLeftSectionBase = TotalVertexCount;
-
-			for (ULandscapeComponent* ComponentUsingHeightmap : *ComponentsUsingHeightmap)
-			{
-				GlobalLayersData.TopLeftSectionBase.X = FMath::Min(GlobalLayersData.TopLeftSectionBase.X, ComponentUsingHeightmap->GetSectionBase().X);
-				GlobalLayersData.TopLeftSectionBase.Y = FMath::Min(GlobalLayersData.TopLeftSectionBase.Y, ComponentUsingHeightmap->GetSectionBase().Y);
-			}
-
-			// Compute per layer data
-			FLandscapeLayerComponentData* LayerData = Component->GetLayerData(Layer->Guid);
-
-			if (LayerData == nullptr || !LayerData->IsInitialized())
-			{
-				FLandscapeLayerComponentData NewData;
-
-				// Setup Heightmap data
-				UTexture2D** LayerHeightmap = CreatedHeightmapTextures.Find(ComponentHeightmap);
-
-				if (LayerHeightmap == nullptr)
-				{
-					UTexture2D* NewLayerHeightmap = Proxy->CreateLandscapeTexture(ComponentHeightmap->Source.GetSizeX(), ComponentHeightmap->Source.GetSizeY(), TEXTUREGROUP_Terrain_Heightmap, ComponentHeightmap->Source.GetFormat());
-					LayerHeightmap = &CreatedHeightmapTextures.Add(ComponentHeightmap, NewLayerHeightmap);
-
-					ULandscapeComponent::CreateEmptyTextureMips(NewLayerHeightmap, true);
-
-					// Init Mip0 to be at 32768 which is equal to "0"
-					FColor* Mip0Data = (FColor*)NewLayerHeightmap->Source.LockMip(0);
-
-					for (ULandscapeComponent* ComponentUsingHeightmap : *ComponentsUsingHeightmap)
-					{
-						int32 HeightmapComponentOffsetX = FMath::RoundToInt((float)NewLayerHeightmap->Source.GetSizeX() * ComponentUsingHeightmap->HeightmapScaleBias.Z);
-						int32 HeightmapComponentOffsetY = FMath::RoundToInt((float)NewLayerHeightmap->Source.GetSizeY() * ComponentUsingHeightmap->HeightmapScaleBias.W);
-
-						for (int32 SubsectionY = 0; SubsectionY < NumSubsections; SubsectionY++)
-						{
-							for (int32 SubsectionX = 0; SubsectionX < NumSubsections; SubsectionX++)
-							{
-								for (int32 SubY = 0; SubY <= SubsectionSizeQuads; SubY++)
-								{
-									for (int32 SubX = 0; SubX <= SubsectionSizeQuads; SubX++)
-									{
-										// X/Y of the vertex we're looking at in component's coordinates.
-										const int32 CompX = SubsectionSizeQuads * SubsectionX + SubX;
-										const int32 CompY = SubsectionSizeQuads * SubsectionY + SubY;
-
-										// X/Y of the vertex we're looking indexed into the texture data
-										const int32 TexX = (SubsectionSizeQuads + 1) * SubsectionX + SubX;
-										const int32 TexY = (SubsectionSizeQuads + 1) * SubsectionY + SubY;
-
-										const int32 HeightTexDataIdx = (HeightmapComponentOffsetX + TexX) + (HeightmapComponentOffsetY + TexY) * NewLayerHeightmap->Source.GetSizeX();
-
-										// copy height and normal data
-										const uint16 HeightValue = LandscapeDataAccess::GetTexHeight(0.f);
-
-										Mip0Data[HeightTexDataIdx].R = HeightValue >> 8;
-										Mip0Data[HeightTexDataIdx].G = HeightValue & 255;
-
-										// Normal with get calculated later
-										Mip0Data[HeightTexDataIdx].B = 0.0f;
-										Mip0Data[HeightTexDataIdx].A = 0.0f;
-									}
-								}
-							}
-						}
-					}
-
-					NewLayerHeightmap->Source.UnlockMip(0);
-
-					NewLayerHeightmap->BeginCachePlatformData();
-					NewLayerHeightmap->ClearAllCachedCookedPlatformData();
-				}
-
-				NewData.HeightmapData.Texture = *LayerHeightmap;
-
-				// Nothing to do for Weightmap by default
-
-				Component->AddLayerData(Layer->Guid, MoveTemp(NewData));
-			}
-		}
-	});
-}
-
 void ALandscape::CopyLayersTexture(UTexture* InSourceTexture, UTexture* InDestTexture, FTextureResource* InDestCPUResource, const FIntPoint& InFirstComponentSectionBase, uint8 InSourceCurrentMip, uint8 InDestCurrentMip, uint32 InSourceArrayIndex, uint32 InDestArrayIndex) const
 {
 	if (InSourceTexture != nullptr && InDestTexture != nullptr)
@@ -3000,7 +2877,7 @@ void ALandscape::PrepareComponentDataToPackMaterialLayersCS(int32 InCurrentWeigh
 	{
 		FLandscapeLayerWeightmapPackMaterialLayersComponentData Data;
 
-		for (int32 WeightmapChannelIndex = 0; WeightmapChannelIndex < 4; ++WeightmapChannelIndex)
+		for (int32 WeightmapChannelIndex = 0; WeightmapChannelIndex < ULandscapeWeightmapUsage::NumChannels; ++WeightmapChannelIndex)
 		{
 			UTexture2D* ComponentWeightmapTexture = ItPair.Key;
 			ULandscapeWeightmapUsage* ComponentWeightmapTextureUsage = ItPair.Value;
@@ -3098,7 +2975,7 @@ void ALandscape::ReallocateLayersWeightmaps(const TArray<ULandscapeLayerInfoObje
 				ULandscapeWeightmapUsage* ComponentWeightmapUsage = ComponentWeightmapTextureUsage[i];
 				ULandscapeWeightmapUsage* Usage = Component->GetLandscapeProxy()->CreateWeightmapUsage();
 
-				for (int32 j = 0; j < 4; ++j)
+				for (int32 j = 0; j < ULandscapeWeightmapUsage::NumChannels; ++j)
 				{
 					Usage->ChannelUsage[j] = ComponentWeightmapUsage->ChannelUsage[j];
 				}
@@ -4250,6 +4127,193 @@ void ALandscapeProxy::FinishDestroy()
 	Super::FinishDestroy();
 }
 
+#if WITH_EDITOR
+bool ALandscapeProxy::RemoveObsoleteLayers(const TSet<FGuid>& InExistingLayers)
+{
+	TSet<FGuid> ComponentLayers;
+	for (ULandscapeComponent* Component : LandscapeComponents)
+	{
+		Component->ForEachLayer([&](const FGuid& InGuid, FLandscapeLayerComponentData&) { ComponentLayers.Add(InGuid); });
+	}
+
+	bool bModified = false;
+
+	for (const FGuid& LayerGuid : ComponentLayers)
+	{
+		if (!InExistingLayers.Contains(LayerGuid))
+		{
+			DeleteLayer(LayerGuid);
+			bModified = true;
+		}
+	}
+
+	return bModified;
+}
+
+bool ALandscapeProxy::AddLayer(const FGuid& InLayerGuid)
+{
+	bool bModified = false;
+	for (ULandscapeComponent* Component : LandscapeComponents)
+	{
+		if (!Component->GetLayerData(InLayerGuid))
+		{
+			Component->AddLayerData(InLayerGuid, FLandscapeLayerComponentData());
+			bModified = true;
+		}
+	}
+
+	if (bModified)
+	{
+		InitializeLayerWithEmptyContent(InLayerGuid);
+	}
+
+	return bModified;
+}
+
+void ALandscapeProxy::DeleteLayer(const FGuid& InLayerGuid)
+{
+	for (ULandscapeComponent* Component : LandscapeComponents)
+	{
+		const FLandscapeLayerComponentData* LayerComponentData = Component->GetLayerData(InLayerGuid);
+
+		if (LayerComponentData != nullptr)
+		{
+			for (const FWeightmapLayerAllocationInfo& Allocation : LayerComponentData->WeightmapData.LayerAllocations)
+			{
+				UTexture2D* WeightmapTexture = LayerComponentData->WeightmapData.Textures[Allocation.WeightmapTextureIndex];
+				ULandscapeWeightmapUsage** Usage = WeightmapUsageMap.Find(WeightmapTexture);
+
+				if (Usage != nullptr && (*Usage) != nullptr)
+				{
+					(*Usage)->Modify();
+					(*Usage)->ChannelUsage[Allocation.WeightmapTextureChannel] = nullptr;
+
+					if ((*Usage)->IsEmpty())
+					{
+						Modify();
+						WeightmapUsageMap.Remove(WeightmapTexture);
+					}
+				}
+			}
+			Component->RemoveLayerData(InLayerGuid);
+		}
+	}
+}
+
+void ALandscapeProxy::InitializeLayerWithEmptyContent(const FGuid& InLayerGuid)
+{
+	ULandscapeInfo* Info = GetLandscapeInfo();
+
+	if (Info == nullptr)
+	{
+		return;
+	}
+		
+	// Build a mapping between each Heightmaps and Component in them
+	TMap<UTexture2D*, TArray<ULandscapeComponent*>> ComponentsPerHeightmaps;
+
+	for (ULandscapeComponent* Component : LandscapeComponents)
+	{
+		UTexture2D* ComponentHeightmapTexture = Component->GetHeightmap();
+		TArray<ULandscapeComponent*>& ComponentList = ComponentsPerHeightmaps.FindOrAdd(ComponentHeightmapTexture);
+		ComponentList.Add(Component);
+	}
+
+	const FIntPoint InitTopLeftSectionBase(INT_MAX, INT_MAX);
+
+	// Init layers with valid "empty" data
+	TMap<UTexture2D*, UTexture2D*> CreatedHeightmapTextures; // < Final layer texture, New created texture for layer
+
+	for (ULandscapeComponent* Component : LandscapeComponents)
+	{
+		UTexture2D* ComponentHeightmap = Component->GetHeightmap();
+		const TArray<ULandscapeComponent*>* ComponentsUsingHeightmap = ComponentsPerHeightmaps.Find(ComponentHeightmap);
+		check(ComponentsUsingHeightmap != nullptr);
+
+		// Compute Global layers data
+		FLandscapeLayersGlobalComponentData& GlobalLayersData = Component->GetGlobalLayersData();
+		GlobalLayersData.TopLeftSectionBase = InitTopLeftSectionBase;
+
+		for (ULandscapeComponent* ComponentUsingHeightmap : *ComponentsUsingHeightmap)
+		{
+			GlobalLayersData.TopLeftSectionBase.X = FMath::Min(GlobalLayersData.TopLeftSectionBase.X, ComponentUsingHeightmap->GetSectionBase().X);
+			GlobalLayersData.TopLeftSectionBase.Y = FMath::Min(GlobalLayersData.TopLeftSectionBase.Y, ComponentUsingHeightmap->GetSectionBase().Y);
+		}
+
+		// Compute per layer data
+		FLandscapeLayerComponentData* LayerData = Component->GetLayerData(InLayerGuid);
+
+		if (LayerData == nullptr || !LayerData->IsInitialized())
+		{
+			FLandscapeLayerComponentData NewData;
+
+			// Setup Heightmap data
+			UTexture2D** LayerHeightmap = CreatedHeightmapTextures.Find(ComponentHeightmap);
+
+			if (LayerHeightmap == nullptr)
+			{
+				UTexture2D* NewLayerHeightmap = CreateLandscapeTexture(ComponentHeightmap->Source.GetSizeX(), ComponentHeightmap->Source.GetSizeY(), TEXTUREGROUP_Terrain_Heightmap, ComponentHeightmap->Source.GetFormat());
+				LayerHeightmap = &CreatedHeightmapTextures.Add(ComponentHeightmap, NewLayerHeightmap);
+
+				ULandscapeComponent::CreateEmptyTextureMips(NewLayerHeightmap, true);
+
+				// Init Mip0 to be at 32768 which is equal to "0"
+				FColor* Mip0Data = (FColor*)NewLayerHeightmap->Source.LockMip(0);
+
+				for (ULandscapeComponent* ComponentUsingHeightmap : *ComponentsUsingHeightmap)
+				{
+					int32 HeightmapComponentOffsetX = FMath::RoundToInt((float)NewLayerHeightmap->Source.GetSizeX() * ComponentUsingHeightmap->HeightmapScaleBias.Z);
+					int32 HeightmapComponentOffsetY = FMath::RoundToInt((float)NewLayerHeightmap->Source.GetSizeY() * ComponentUsingHeightmap->HeightmapScaleBias.W);
+
+					for (int32 SubsectionY = 0; SubsectionY < NumSubsections; SubsectionY++)
+					{
+						for (int32 SubsectionX = 0; SubsectionX < NumSubsections; SubsectionX++)
+						{
+							for (int32 SubY = 0; SubY <= SubsectionSizeQuads; SubY++)
+							{
+								for (int32 SubX = 0; SubX <= SubsectionSizeQuads; SubX++)
+								{
+									// X/Y of the vertex we're looking at in component's coordinates.
+									const int32 CompX = SubsectionSizeQuads * SubsectionX + SubX;
+									const int32 CompY = SubsectionSizeQuads * SubsectionY + SubY;
+
+									// X/Y of the vertex we're looking indexed into the texture data
+									const int32 TexX = (SubsectionSizeQuads + 1) * SubsectionX + SubX;
+									const int32 TexY = (SubsectionSizeQuads + 1) * SubsectionY + SubY;
+
+									const int32 HeightTexDataIdx = (HeightmapComponentOffsetX + TexX) + (HeightmapComponentOffsetY + TexY) * NewLayerHeightmap->Source.GetSizeX();
+
+									// copy height and normal data
+									const uint16 HeightValue = LandscapeDataAccess::GetTexHeight(0.f);
+
+									Mip0Data[HeightTexDataIdx].R = HeightValue >> 8;
+									Mip0Data[HeightTexDataIdx].G = HeightValue & 255;
+
+									// Normal with get calculated later
+									Mip0Data[HeightTexDataIdx].B = 0.0f;
+									Mip0Data[HeightTexDataIdx].A = 0.0f;
+								}
+							}
+						}
+					}
+				}
+
+				NewLayerHeightmap->Source.UnlockMip(0);
+
+				NewLayerHeightmap->BeginCachePlatformData();
+				NewLayerHeightmap->ClearAllCachedCookedPlatformData();
+			}
+
+			NewData.HeightmapData.Texture = *LayerHeightmap;
+
+			// Nothing to do for Weightmap by default
+
+			Component->AddLayerData(InLayerGuid, MoveTemp(NewData));
+		}
+	}
+}
+#endif
+
 void ALandscape::BeginDestroy()
 {
 #if WITH_EDITOR
@@ -4449,33 +4513,7 @@ void ALandscape::DeleteLayer(int32 InLayerIndex)
 	// Clean up Weightmap usage in LandscapeProxies
 	LandscapeInfo->ForAllLandscapeProxies([&LayerGuid](ALandscapeProxy* Proxy)
 	{
-		for (ULandscapeComponent* Component : Proxy->LandscapeComponents)
-		{
-			const FLandscapeLayerComponentData* LayerComponentData = Component->GetLayerData(LayerGuid);
-
-			if (LayerComponentData != nullptr)
-			{
-				for (const FWeightmapLayerAllocationInfo& Allocation : LayerComponentData->WeightmapData.LayerAllocations)
-				{
-					UTexture2D* WeightmapTexture = LayerComponentData->WeightmapData.Textures[Allocation.WeightmapTextureIndex];
-					ULandscapeWeightmapUsage** Usage = Proxy->WeightmapUsageMap.Find(WeightmapTexture);
-
-					if (Usage != nullptr && (*Usage) != nullptr)
-					{
-						(*Usage)->Modify();
-						(*Usage)->ChannelUsage[Allocation.WeightmapTextureChannel] = nullptr;
-
-						if ((*Usage)->FreeChannelCount() == 4)
-						{
-							Proxy->Modify();
-							Proxy->WeightmapUsageMap.Remove(WeightmapTexture);
-						}
-					}
-				}
-				Component->Modify();
-				Component->RemoveLayerData(LayerGuid);
-			}
-		}
+		Proxy->DeleteLayer(LayerGuid);
 	});
 
 	const FLandscapeLayer* SplinesReservedLayer = GetLandscapeSplinesReservedLayer();
@@ -4777,14 +4815,8 @@ void ALandscape::CreateLayer(FName InName)
 	// Create associated layer data in each landscape proxy
 	LandscapeInfo->ForAllLandscapeProxies([&NewLayer](ALandscapeProxy* Proxy)
 	{
-		for (ULandscapeComponent* Component : Proxy->LandscapeComponents)
-		{
-			Component->Modify();
-			Component->AddLayerData(NewLayer.Guid, FLandscapeLayerComponentData());
-		}
+		Proxy->AddLayer(NewLayer.Guid);
 	});
-
-	InitializeLayerWithEmptyContent(NewLayer.Guid);
 }
 
 bool ALandscape::ReorderLayer(int32 InStartingLayerIndex, int32 InDestinationLayerIndex)
