- Implemented Moving LandscapeComponent to another Level in Landscape Layer System
- Partial duplicate of 5708384 to reparent landscape components owned textures and materials to the landscape actor instead of the level package. (Will be required if we save actors outside levels in the future)

diff --git a/Engine/Source/Editor/LandscapeEditor/Private/LandscapeEdMode.cpp b/Engine/Source/Editor/LandscapeEditor/Private/LandscapeEdMode.cpp
--- a/Engine/Source/Editor/LandscapeEditor/Private/LandscapeEdMode.cpp
+++ b/Engine/Source/Editor/LandscapeEditor/Private/LandscapeEdMode.cpp
@@ -99,7 +99,7 @@ IMPLEMENT_HIT_PROXY(HNewLandscapeGrabHandleProxy, HHitProxy)
 
 ENGINE_API extern bool GDisableAutomaticTextureMaterialUpdateDependencies;
 
-void ALandscape::SplitHeightmap(ULandscapeComponent* Comp, bool bMoveToCurrentLevel, FMaterialUpdateContext* InOutUpdateContext, TArray<FComponentRecreateRenderStateContext>* InOutRecreateRenderStateContext, bool InReregisterComponent)
+void ALandscape::SplitHeightmap(ULandscapeComponent* Comp, ALandscapeProxy* TargetProxy, FMaterialUpdateContext* InOutUpdateContext, TArray<FComponentRecreateRenderStateContext>* InOutRecreateRenderStateContext, bool InReregisterComponent)
 {
 	ULandscapeInfo* Info = Comp->GetLandscapeInfo();
 	
@@ -108,10 +108,10 @@ void ALandscape::SplitHeightmap(ULandscapeComponent* Comp, bool bMoveToCurrentLe
 	int32 HeightmapSizeU = (1 << FMath::CeilLogTwo(ComponentSizeVerts));
 	int32 HeightmapSizeV = (1 << FMath::CeilLogTwo(ComponentSizeVerts));
 
-	ALandscapeProxy* Proxy = Comp->GetLandscapeProxy();
-	Proxy->Modify();
-
-	UObject* TextureOuter = bMoveToCurrentLevel ? Comp->GetWorld()->GetCurrentLevel()->GetOutermost() : nullptr;
+	ALandscapeProxy* SrcProxy = Comp->GetLandscapeProxy();
+	ALandscapeProxy* DstProxy = TargetProxy ? TargetProxy : SrcProxy;
+	SrcProxy->Modify();
+	DstProxy->Modify();
 	
 	UTexture2D* OldHeightmapTexture = Comp->GetHeightmap(false);
 	UTexture2D* NewHeightmapTexture = NULL;
@@ -129,7 +129,7 @@ void ALandscape::SplitHeightmap(ULandscapeComponent* Comp, bool bMoveToCurrentLe
 		LandscapeEdit.GetHeightDataFast(Comp->GetSectionBase().X, Comp->GetSectionBase().Y, Comp->GetSectionBase().X + Comp->ComponentSizeQuads, Comp->GetSectionBase().Y + Comp->ComponentSizeQuads, (uint16*)HeightData.GetData(), 0, (uint16*)NormalData.GetData());
 
 		// Create the new heightmap texture
-		NewHeightmapTexture = Proxy->CreateLandscapeTexture(HeightmapSizeU, HeightmapSizeV, TEXTUREGROUP_Terrain_Heightmap, TSF_BGRA8, TextureOuter);
+		NewHeightmapTexture = DstProxy->CreateLandscapeTexture(HeightmapSizeU, HeightmapSizeV, TEXTUREGROUP_Terrain_Heightmap, TSF_BGRA8);
 		ULandscapeComponent::CreateEmptyTextureMips(NewHeightmapTexture, true);
 		Comp->HeightmapScaleBias = NewHeightmapScaleBias;
 		Comp->SetHeightmap(NewHeightmapTexture);
@@ -195,11 +195,11 @@ void ALandscape::SplitHeightmap(ULandscapeComponent* Comp, bool bMoveToCurrentLe
 	{
 		FLandscapeLayersTexture2DCPUReadBackResource* NewCPUReadBackResource = new FLandscapeLayersTexture2DCPUReadBackResource(NewHeightmapTexture->Source.GetSizeX(), NewHeightmapTexture->Source.GetSizeY(), NewHeightmapTexture->GetPixelFormat(), NewHeightmapTexture->Source.GetNumMips());
 		BeginInitResource(NewCPUReadBackResource);
-		Proxy->HeightmapsCPUReadBack.Add(NewHeightmapTexture, NewCPUReadBackResource);
+		DstProxy->HeightmapsCPUReadBack.Add(NewHeightmapTexture, NewCPUReadBackResource);
 
 		// Free OldHeightmapTexture's CPUReadBackResource if not used by any component
 		bool FreeCPUReadBack = true;
-		for (ULandscapeComponent* Component : Proxy->LandscapeComponents)
+		for (ULandscapeComponent* Component : SrcProxy->LandscapeComponents)
 		{
 			if (Component != Comp && Component->GetHeightmap(false) == OldHeightmapTexture)
 			{
@@ -209,14 +209,14 @@ void ALandscape::SplitHeightmap(ULandscapeComponent* Comp, bool bMoveToCurrentLe
 		}
 		if (FreeCPUReadBack)
 		{
-			FLandscapeLayersTexture2DCPUReadBackResource** OldCPUReadBackResource = Proxy->HeightmapsCPUReadBack.Find(OldHeightmapTexture);
+			FLandscapeLayersTexture2DCPUReadBackResource** OldCPUReadBackResource = SrcProxy->HeightmapsCPUReadBack.Find(OldHeightmapTexture);
 			if (OldCPUReadBackResource)
 			{
 				if (FLandscapeLayersTexture2DCPUReadBackResource* ResourceToDelete = *OldCPUReadBackResource)
 				{
 					ReleaseResourceAndFlush(ResourceToDelete);
 					delete ResourceToDelete;
-					Proxy->HeightmapsCPUReadBack.Remove(OldHeightmapTexture);
+					SrcProxy->HeightmapsCPUReadBack.Remove(OldHeightmapTexture);
 				}
 			}
 		}
@@ -243,7 +243,7 @@ void ALandscape::SplitHeightmap(ULandscapeComponent* Comp, bool bMoveToCurrentLe
 				// Restore new heightmap scale/bias
 				Comp->HeightmapScaleBias = NewHeightmapScaleBias;
 				{
-					UTexture2D* LayerHeightmapTexture = Proxy->CreateLandscapeTexture(HeightmapSizeU, HeightmapSizeV, TEXTUREGROUP_Terrain_Heightmap, TSF_BGRA8, TextureOuter);
+					UTexture2D* LayerHeightmapTexture = DstProxy->CreateLandscapeTexture(HeightmapSizeU, HeightmapSizeV, TEXTUREGROUP_Terrain_Heightmap, TSF_BGRA8);
 					ULandscapeComponent::CreateEmptyTextureMips(LayerHeightmapTexture, true);
 					LayerHeightmapTexture->PostEditChange();
 					// Set Layer heightmap texture
@@ -3646,7 +3646,7 @@ void FEdModeLandscape::DeleteLandscapeComponents(ULandscapeInfo* LandscapeInfo,
 		// Changing Heightmap format for selected components
 		for (ULandscapeComponent* Component : HeightmapUpdateComponents)
 		{
-			ALandscape::SplitHeightmap(Component, false, &MaterialUpdateContext, &RecreateRenderStateContexts, false);
+			ALandscape::SplitHeightmap(Component, nullptr, &MaterialUpdateContext, &RecreateRenderStateContexts, false);
 		}
 
 		FMultiComponentReregisterContext RegisterContext(ComponentsToReregister);
diff --git a/Engine/Source/Editor/LandscapeEditor/Private/LandscapeEdModeComponentTools.cpp b/Engine/Source/Editor/LandscapeEditor/Private/LandscapeEdModeComponentTools.cpp
--- a/Engine/Source/Editor/LandscapeEditor/Private/LandscapeEdModeComponentTools.cpp
+++ b/Engine/Source/Editor/LandscapeEditor/Private/LandscapeEdModeComponentTools.cpp
@@ -538,21 +538,6 @@ public:
 					}
 				}
 
-				// Changing Heightmap format for selected components
-				for (const auto& HeightmapUpdateComponentPair : HeightmapUpdateComponents)
-				{
-					ALandscape::SplitHeightmap(HeightmapUpdateComponentPair.Key, HeightmapUpdateComponentPair.Value);
-				}
-
-				// Delete if it is no referenced textures...
-				for (UTexture2D* Texture : OldHeightmapTextures)
-				{
-					Texture->SetFlags(RF_Transactional);
-					Texture->Modify();
-					Texture->MarkPackageDirty();
-					Texture->ClearFlags(RF_Standalone);
-				}
-
 				ALandscapeProxy* LandscapeProxy = LandscapeInfo->GetCurrentLevelLandscapeProxy(false);
 				if (!LandscapeProxy)
 				{
@@ -573,6 +558,21 @@ public:
 					}
 				}
 
+				// Changing Heightmap format for selected components
+				for (const auto& HeightmapUpdateComponentPair : HeightmapUpdateComponents)
+				{
+					ALandscape::SplitHeightmap(HeightmapUpdateComponentPair.Key, HeightmapUpdateComponentPair.Value ? LandscapeProxy : nullptr);
+				}
+
+				// Delete if it is no referenced textures...
+				for (UTexture2D* Texture : OldHeightmapTextures)
+				{
+					Texture->SetFlags(RF_Transactional);
+					Texture->Modify();
+					Texture->MarkPackageDirty();
+					Texture->ClearFlags(RF_Standalone);
+				}
+
 				for (ALandscapeProxy* Proxy : SelectProxies)
 				{
 					Proxy->Modify();
@@ -587,7 +587,7 @@ public:
 					if (Component->XYOffsetmapTexture)
 					{
 						Component->XYOffsetmapTexture->Modify();
-						Component->XYOffsetmapTexture->Rename(nullptr, LandscapeProxy->GetOutermost());
+						Component->XYOffsetmapTexture->Rename(nullptr, LandscapeProxy);
 					}
 				}
 
@@ -596,138 +596,20 @@ public:
 					FLandscapeEditDataInterface LandscapeEdit(LandscapeInfo);
 					for (ULandscapeComponent* Component : TargetSelectedComponents)
 					{
-						TArray<FWeightmapLayerAllocationInfo>& ComponentWeightmapLayerAllocations = Component->GetWeightmapLayerAllocations();
-
-						int32 TotalNeededChannels = ComponentWeightmapLayerAllocations.Num();
-						int32 CurrentLayer = 0;
-						TArray<UTexture2D*> NewWeightmapTextures;
-
-						// Code from ULandscapeComponent::ReallocateWeightmaps
-						// Move to other channels left
-						while (TotalNeededChannels > 0)
-						{
-							// UE_LOG(LogLandscape, Log, TEXT("Still need %d channels"), TotalNeededChannels);
-
-							UTexture2D* CurrentWeightmapTexture = nullptr;
-							ULandscapeWeightmapUsage* CurrentWeightmapUsage = nullptr;
-
-							if (TotalNeededChannels < ULandscapeWeightmapUsage::NumChannels)
-							{
-								// UE_LOG(LogLandscape, Log, TEXT("Looking for nearest"));
-
-								// see if we can find a suitable existing weightmap texture with sufficient channels
-								int32 BestDistanceSquared = MAX_int32;
-								for (auto& WeightmapUsagePair : LandscapeProxy->WeightmapUsageMap)
-								{
-									ULandscapeWeightmapUsage* TryWeightmapUsage = WeightmapUsagePair.Value;
-									if (TryWeightmapUsage->FreeChannelCount() >= TotalNeededChannels) // TODO: handle layers
-									{
-										// See if this candidate is closer than any others we've found
-										for (int32 ChanIdx = 0; ChanIdx < ULandscapeWeightmapUsage::NumChannels; ChanIdx++)
-										{
-											if (TryWeightmapUsage->ChannelUsage[ChanIdx] != nullptr)
-											{
-												int32 TryDistanceSquared = (TryWeightmapUsage->ChannelUsage[ChanIdx]->GetSectionBase() - Component->GetSectionBase()).SizeSquared();
-												if (TryDistanceSquared < BestDistanceSquared)
-												{
-													CurrentWeightmapTexture = WeightmapUsagePair.Key;
-													CurrentWeightmapUsage = TryWeightmapUsage;
-													BestDistanceSquared = TryDistanceSquared;
-												}
-											}
-										}
-									}
-								}
-							}
-
-							bool NeedsUpdateResource = false;
-							// No suitable weightmap texture
-							if (CurrentWeightmapTexture == nullptr)
-							{
-								Component->MarkPackageDirty();
-
-								// Weightmap is sized the same as the component
-								int32 WeightmapSize = (Component->SubsectionSizeQuads + 1) * Component->NumSubsections;
-
-								// We need a new weightmap texture
-								CurrentWeightmapTexture = LandscapeProxy->CreateLandscapeTexture(WeightmapSize, WeightmapSize, TEXTUREGROUP_Terrain_Weightmap, TSF_BGRA8);
-								// Alloc dummy mips
-								Component->CreateEmptyTextureMips(CurrentWeightmapTexture);
-								CurrentWeightmapTexture->PostEditChange();
-
-								// Store it in the usage map
-								CurrentWeightmapUsage = LandscapeProxy->WeightmapUsageMap.Add(CurrentWeightmapTexture, LandscapeProxy->CreateWeightmapUsage());
-
-								// UE_LOG(LogLandscape, Log, TEXT("Making a new texture %s"), *CurrentWeightmapTexture->GetName());
-							}
-
-							NewWeightmapTextures.Add(CurrentWeightmapTexture);
-
-							for (int32 ChanIdx = 0; ChanIdx < ULandscapeWeightmapUsage::NumChannels && TotalNeededChannels > 0; ChanIdx++)
-							{
-								// UE_LOG(LogLandscape, Log, TEXT("Finding allocation for layer %d"), CurrentLayer);
-
-								if (CurrentWeightmapUsage->ChannelUsage[ChanIdx] == nullptr)
-								{
-									// Use this allocation
-									FWeightmapLayerAllocationInfo& AllocInfo = ComponentWeightmapLayerAllocations[CurrentLayer];
-
-									if (AllocInfo.WeightmapTextureIndex == 255)
-									{
-										// New layer - zero out the data for this texture channel
-										LandscapeEdit.ZeroTextureChannel(CurrentWeightmapTexture, ChanIdx);
-									}
-									else
-									{
-										TArray<UTexture2D*>& ComponentWeightmapTextures = Component->GetWeightmapTextures();
-										UTexture2D* OldWeightmapTexture = ComponentWeightmapTextures[AllocInfo.WeightmapTextureIndex];
-
-										// Copy the data
-										LandscapeEdit.CopyTextureChannel(CurrentWeightmapTexture, ChanIdx, OldWeightmapTexture, AllocInfo.WeightmapTextureChannel);
-										LandscapeEdit.ZeroTextureChannel(OldWeightmapTexture, AllocInfo.WeightmapTextureChannel);
-
-										// Remove the old allocation
-										ULandscapeWeightmapUsage** OldWeightmapUsage = Component->GetLandscapeProxy()->WeightmapUsageMap.Find(OldWeightmapTexture);
-
-										if (OldWeightmapUsage != nullptr)
-										{
-											(*OldWeightmapUsage)->Modify();
-											(*OldWeightmapUsage)->ChannelUsage[AllocInfo.WeightmapTextureChannel] = nullptr;
-										}
-									}
-
-									// Assign the new allocation
-									CurrentWeightmapUsage->Modify();
-									CurrentWeightmapUsage->ChannelUsage[ChanIdx] = Component;
-									AllocInfo.WeightmapTextureIndex = NewWeightmapTextures.Num() - 1;
-									AllocInfo.WeightmapTextureChannel = ChanIdx;
-									CurrentLayer++;
-									TotalNeededChannels--;
-								}
-							}
-						}
-
-						// Replace the weightmap textures
-						Component->SetWeightmapTextures(NewWeightmapTextures);
-
-						// Update the mipmaps for the textures we edited
-						for (UTexture2D* WeightmapTexture : NewWeightmapTextures)
+						Component->ReallocateWeightmaps(&LandscapeEdit, false, true, false, true, LandscapeProxy);
+						Component->ForEachLayer([&](const FGuid& LayerGuid, FLandscapeLayerComponentData& LayerData)
 						{
-							FLandscapeTextureDataInfo* WeightmapDataInfo = LandscapeEdit.GetTextureDataInfo(WeightmapTexture);
-
-							int32 NumMips = WeightmapTexture->Source.GetNumMips();
-							TArray<FColor*> WeightmapTextureMipData;
-							WeightmapTextureMipData.AddUninitialized(NumMips);
-							for (int32 MipIdx = 0; MipIdx < NumMips; MipIdx++)
-							{
-								WeightmapTextureMipData[MipIdx] = (FColor*)WeightmapDataInfo->GetMipData(MipIdx);
-							}
-
-							ULandscapeComponent::UpdateWeightmapMips(Component->NumSubsections, Component->SubsectionSizeQuads, WeightmapTexture, WeightmapTextureMipData, 0, 0, MAX_int32, MAX_int32, WeightmapDataInfo);
-						}
+							FScopedSetLandscapeEditingLayer Scope(Landscape, LayerGuid);
+							Component->ReallocateWeightmaps(&LandscapeEdit, true, true, false, true, LandscapeProxy);
+						});
+						Landscape->RequestLayersContentUpdate(ELandscapeLayersContentUpdateFlag::All, true);
 					}
+
 					// Need to Repacking all the Weight map (to make it packed well...)
-					Landscape->RemoveInvalidWeightmaps();
+					for (ALandscapeProxy* Proxy : SelectProxies)
+					{
+						Proxy->RemoveInvalidWeightmaps();
+					}
 				}
 
 				// Move the components to the Proxy actor
@@ -809,6 +691,7 @@ public:
 		: FLandscapeToolBase<FLandscapeToolStrokeMoveToLevel>(InEdMode)
 	{
 	}
+	virtual bool ShouldUpdateEditingLayer() const override { return false; }
 
 	virtual const TCHAR* GetToolName() override { return TEXT("MoveToLevel"); }
 	virtual FText GetDisplayName() override { return NSLOCTEXT("UnrealEd", "LandscapeMode_MoveToLevel", "Move to Streaming Level"); };
diff --git a/Engine/Source/Runtime/Landscape/Classes/Landscape.h b/Engine/Source/Runtime/Landscape/Classes/Landscape.h
--- a/Engine/Source/Runtime/Landscape/Classes/Landscape.h
+++ b/Engine/Source/Runtime/Landscape/Classes/Landscape.h
@@ -228,7 +228,7 @@ public:
 	LANDSCAPE_API static void CalcComponentIndicesNoOverlap(const int32 X1, const int32 Y1, const int32 X2, const int32 Y2, const int32 ComponentSizeQuads,
 		int32& ComponentIndexX1, int32& ComponentIndexY1, int32& ComponentIndexX2, int32& ComponentIndexY2);
 
-	static void SplitHeightmap(ULandscapeComponent* Comp, bool bMoveToCurrentLevel = false, class FMaterialUpdateContext* InOutUpdateContext = nullptr, TArray<class FComponentRecreateRenderStateContext>* InOutRecreateRenderStateContext = nullptr, bool InReregisterComponent = true);
+	static void SplitHeightmap(ULandscapeComponent* Comp, ALandscapeProxy* TargetProxy = nullptr, class FMaterialUpdateContext* InOutUpdateContext = nullptr, TArray<class FComponentRecreateRenderStateContext>* InOutRecreateRenderStateContext = nullptr, bool InReregisterComponent = true);
 	
 	//~ Begin UObject Interface.
 	virtual void PreSave(const class ITargetPlatform* TargetPlatform) override;
diff --git a/Engine/Source/Runtime/Landscape/Classes/LandscapeComponent.h b/Engine/Source/Runtime/Landscape/Classes/LandscapeComponent.h
--- a/Engine/Source/Runtime/Landscape/Classes/LandscapeComponent.h
+++ b/Engine/Source/Runtime/Landscape/Classes/LandscapeComponent.h
@@ -807,7 +807,7 @@ public:
 	/**
 	 * Create weightmaps for this component for the layers specified in the WeightmapLayerAllocations array
 	 */
-	void ReallocateWeightmaps(FLandscapeEditDataInterface* DataInterface= nullptr, bool InCanUseEditingWeightmap = true, bool InSaveToTransactionBuffer = true, bool InInitPlatformDataAsync = false, TArray<UTexture2D*>* OutNewCreatedTextures = nullptr);
+	LANDSCAPE_API void ReallocateWeightmaps(FLandscapeEditDataInterface* DataInterface = nullptr, bool InCanUseEditingWeightmap = true, bool InSaveToTransactionBuffer = true, bool InInitPlatformDataAsync = false, bool InForceReallocate = false, ALandscapeProxy* InTargetProxy = nullptr, TArray<UTexture2D*>* OutNewCreatedTextures = nullptr);
 
 	/** Returns the actor's LandscapeMaterial, or the Component's OverrideLandscapeMaterial if set */
 	LANDSCAPE_API UMaterialInterface* GetLandscapeMaterial(int8 InLODIndex = INDEX_NONE) const;
diff --git a/Engine/Source/Runtime/Landscape/Classes/LandscapeProxy.h b/Engine/Source/Runtime/Landscape/Classes/LandscapeProxy.h
--- a/Engine/Source/Runtime/Landscape/Classes/LandscapeProxy.h
+++ b/Engine/Source/Runtime/Landscape/Classes/LandscapeProxy.h
@@ -876,8 +876,8 @@ public:
 	/** @return Current size of bounding rectangle in quads space */
 	LANDSCAPE_API FIntRect GetBoundingRect() const;
 
-	/** Creates a Texture2D for use by this landscape proxy or one of it's components. If OptionalOverrideOuter is not specified, the level is used. */
-	LANDSCAPE_API UTexture2D* CreateLandscapeTexture(int32 InSizeX, int32 InSizeY, TextureGroup InLODGroup, ETextureSourceFormat InFormat, UObject* OptionalOverrideOuter = nullptr, bool bCompress = false) const;
+	/** Creates a Texture2D for use by this landscape proxy or one of it's components. */
+	LANDSCAPE_API UTexture2D* CreateLandscapeTexture(int32 InSizeX, int32 InSizeY, TextureGroup InLODGroup, ETextureSourceFormat InFormat, bool bCompress = false) const;
 
 	/** Creates a LandscapeWeightMapUsage object outered to this proxy. */
 	LANDSCAPE_API ULandscapeWeightmapUsage* CreateWeightmapUsage();
diff --git a/Engine/Source/Runtime/Landscape/Private/Landscape.cpp b/Engine/Source/Runtime/Landscape/Private/Landscape.cpp
--- a/Engine/Source/Runtime/Landscape/Private/Landscape.cpp
+++ b/Engine/Source/Runtime/Landscape/Private/Landscape.cpp
@@ -705,6 +705,55 @@ void ULandscapeComponent::PostLoad()
 	}
 #endif
 
+#if WITH_EDITOR
+	auto ReparentObject = [this](UObject* Object)
+	{
+		if (Object && !Object->HasAllFlags(RF_Public | RF_Standalone) && (Object->GetOuter() != GetOuter()))
+		{
+			Object->Rename(nullptr, GetOuter(), REN_ForceNoResetLoaders);
+			return true;
+		}
+		return false;
+	};
+
+	for (UMaterialInstance* MaterialInstance : MaterialInstances)
+	{
+		while (ReparentObject(MaterialInstance))
+		{
+			MaterialInstance = Cast<UMaterialInstance>(MaterialInstance->Parent);
+		}
+	}
+
+	for (UMaterialInterface* MobileMaterialInterface : MobileMaterialInterfaces)
+	{
+		while (ReparentObject(MobileMaterialInterface))
+		{
+			MobileMaterialInterface = Cast<UMaterialInstance>(MobileMaterialInterface) ? Cast<UMaterialInstance>(((UMaterialInstance*)MobileMaterialInterface)->Parent) : nullptr;
+		}
+	}
+
+	for (UMaterialInstance* MobileCombinationMaterialInstance : MobileCombinationMaterialInstances)
+	{
+		while (ReparentObject(MobileCombinationMaterialInstance))
+		{
+			MobileCombinationMaterialInstance = Cast<UMaterialInstance>(MobileCombinationMaterialInstance->Parent);
+		}
+	}
+
+	ReparentObject(HeightmapTexture);
+	ReparentObject(XYOffsetmapTexture);
+
+	for (UTexture2D* WeightmapTexture : WeightmapTextures)
+	{
+		ReparentObject(WeightmapTexture);
+	}
+
+	for (UTexture2D* MobileWeightmapTexture : MobileWeightmapTextures)
+	{
+		ReparentObject(MobileWeightmapTexture);
+	}
+#endif
+
 #if !UE_BUILD_SHIPPING
 	// This will fix the data in case there is mismatch between save of asset/maps
 	int8 MaxLOD = FMath::CeilLogTwo(SubsectionSizeQuads + 1) - 1;
diff --git a/Engine/Source/Runtime/Landscape/Private/LandscapeEdit.cpp b/Engine/Source/Runtime/Landscape/Private/LandscapeEdit.cpp
--- a/Engine/Source/Runtime/Landscape/Private/LandscapeEdit.cpp
+++ b/Engine/Source/Runtime/Landscape/Private/LandscapeEdit.cpp
@@ -259,11 +259,11 @@ UMaterialInstanceConstant* ULandscapeComponent::GetCombinationMaterial(FMaterial
 
 		// Find or set a matching MIC in the Landscape's map.
 		UMaterialInstanceConstant* CombinationMaterialInstance = Proxy->MaterialInstanceConstantMap.FindRef(*LayerKey);
-		if (CombinationMaterialInstance == nullptr || CombinationMaterialInstance->Parent != MaterialToUse || GetOutermost() != CombinationMaterialInstance->GetOutermost())
+		if (CombinationMaterialInstance == nullptr || CombinationMaterialInstance->Parent != MaterialToUse || GetOuter() != CombinationMaterialInstance->GetOuter())
 		{
 			FlushRenderingCommands();
 
-			ULandscapeMaterialInstanceConstant* LandscapeCombinationMaterialInstance = NewObject<ULandscapeMaterialInstanceConstant>(GetOutermost());
+			ULandscapeMaterialInstanceConstant* LandscapeCombinationMaterialInstance = NewObject<ULandscapeMaterialInstanceConstant>(GetOuter());
 			LandscapeCombinationMaterialInstance->bMobile = bMobile;
 			CombinationMaterialInstance = LandscapeCombinationMaterialInstance;
 			UE_LOG(LogLandscape, Log, TEXT("Looking for key %s, making new combination %s"), *LayerKey, *CombinationMaterialInstance->GetName());
@@ -353,7 +353,7 @@ void ULandscapeComponent::UpdateMaterialInstances_Internal(FMaterialUpdateContex
 		if (CombinationMaterialInstance != nullptr)
 		{
 			// Create the instance for this component, that will use the layer combination instance.
-			UMaterialInstanceConstant* MaterialInstance = NewObject<ULandscapeMaterialInstanceConstant>(GetOutermost());
+			UMaterialInstanceConstant* MaterialInstance = NewObject<ULandscapeMaterialInstanceConstant>(this);
 			MaterialInstances[MaterialIndex] = MaterialInstance;
 
 			// Material Instances don't support Undo/Redo (the shader map goes out of sync and crashes happen)
@@ -392,7 +392,7 @@ void ULandscapeComponent::UpdateMaterialInstances_Internal(FMaterialUpdateContex
 			// Setup material instance with disabled tessellation
 			if (CombinationMaterialInstance->GetMaterial()->D3D11TessellationMode != EMaterialTessellationMode::MTM_NoTessellation)
 			{
-				ULandscapeMaterialInstanceConstant* TessellationMaterialInstance = NewObject<ULandscapeMaterialInstanceConstant>(GetOutermost());
+				ULandscapeMaterialInstanceConstant* TessellationMaterialInstance = NewObject<ULandscapeMaterialInstanceConstant>(this);
 				int32 TessellatedMaterialIndex = MaterialPerLOD.Num() + TessellatedMaterialCount++;
 				MaterialInstances[TessellatedMaterialIndex] = TessellationMaterialInstance;
 				MaterialIndexToDisabledTessellationMaterial[MaterialIndex] = TessellatedMaterialIndex;
@@ -4503,107 +4503,118 @@ void ULandscapeInfo::ClearSelectedRegion(bool bIsComponentwise /*= true*/)
 	}
 }
 
-void ULandscapeComponent::ReallocateWeightmaps(FLandscapeEditDataInterface* DataInterface, bool InCanUseEditingWeightmap, bool InSaveToTransactionBuffer, bool InInitPlatformDataAsync, TArray<UTexture2D*>* OutNewCreatedTextures)
+void ULandscapeComponent::ReallocateWeightmaps(FLandscapeEditDataInterface* DataInterface, bool InCanUseEditingWeightmap, bool InSaveToTransactionBuffer, bool InInitPlatformDataAsync, bool InForceReallocate, ALandscapeProxy* InTargetProxy, TArray<UTexture2D*>* OutNewCreatedTextures)
 {
-	ALandscapeProxy* Proxy = GetLandscapeProxy();
+	int32 NeededNewChannels = 0;
+	ALandscapeProxy* TargetProxy = InTargetProxy ? InTargetProxy : GetLandscapeProxy();
+
+	FGuid EditingLayerGUID = GetEditingLayerGUID();
+	check(!GetMutableDefault<UEditorExperimentalSettings>()->bLandscapeLayerSystem || !InCanUseEditingWeightmap || EditingLayerGUID.IsValid());
+	FGuid TargetLayerGuid = InCanUseEditingWeightmap ? EditingLayerGUID : FGuid();
 
 	TArray<FWeightmapLayerAllocationInfo>& ComponentWeightmapLayerAllocations = GetWeightmapLayerAllocations(InCanUseEditingWeightmap);
 	TArray<UTexture2D*>& ComponentWeightmapTextures = GetWeightmapTextures(InCanUseEditingWeightmap);
 	TArray<ULandscapeWeightmapUsage*>& ComponentWeightmapTexturesUsage = GetWeightmapTexturesUsage(InCanUseEditingWeightmap);
 
-	int32 NeededNewChannels = 0;
-	for (int32 LayerIdx = 0; LayerIdx < ComponentWeightmapLayerAllocations.Num(); LayerIdx++)
+	// When force reallocating, skip tests to see if allocations are necessary based on Component's WeightmapLayeAllocInfo
+	if (!InForceReallocate)
 	{
-		if (ComponentWeightmapLayerAllocations[LayerIdx].WeightmapTextureIndex == 255)
+		for (int32 LayerIdx = 0; LayerIdx < ComponentWeightmapLayerAllocations.Num(); LayerIdx++)
 		{
-			NeededNewChannels++;
+			if (ComponentWeightmapLayerAllocations[LayerIdx].WeightmapTextureIndex == 255)
+			{
+				NeededNewChannels++;
+			}
 		}
-	}
 
-	// All channels allocated!
-	if (NeededNewChannels == 0)
-	{
-		return;
+		// All channels allocated!
+		if (NeededNewChannels == 0)
+		{
+			return;
+		}
 	}
 
 	if (InSaveToTransactionBuffer)
 	{
 		Modify();
-		Proxy->Modify();
+		TargetProxy->Modify();
 	}
 
-	// UE_LOG(LogLandscape, Log, TEXT("----------------------"));
-	// UE_LOG(LogLandscape, Log, TEXT("Component %s needs %d layers (%d new)"), *GetName(), WeightmapLayerAllocations.Num(), NeededNewChannels);
-
-	// See if our existing textures have sufficient space
-	int32 ExistingTexAvailableChannels = 0;
-	for (int32 TexIdx = 0; TexIdx < ComponentWeightmapTextures.Num(); TexIdx++)
+	if (!InForceReallocate)
 	{
-		ULandscapeWeightmapUsage* Usage = ComponentWeightmapTexturesUsage[TexIdx];
-		check(Usage);
-
-		ExistingTexAvailableChannels += Usage->FreeChannelCount();
+		// UE_LOG(LogLandscape, Log, TEXT("----------------------"));
+		// UE_LOG(LogLandscape, Log, TEXT("Component %s needs %d layers (%d new)"), *GetName(), WeightmapLayerAllocations.Num(), NeededNewChannels);
 
-		if (ExistingTexAvailableChannels >= NeededNewChannels)
+		// See if our existing textures have sufficient space
+		int32 ExistingTexAvailableChannels = 0;
+		for (int32 TexIdx = 0; TexIdx < ComponentWeightmapTextures.Num(); TexIdx++)
 		{
-			break;
-		}
-	}
+			ULandscapeWeightmapUsage* Usage = ComponentWeightmapTexturesUsage[TexIdx];
+			check(Usage);
+			check(Usage->LayerGuid == TargetLayerGuid);
+			ExistingTexAvailableChannels += Usage->FreeChannelCount();
 
-	if (ExistingTexAvailableChannels >= NeededNewChannels)
-	{
-		// UE_LOG(LogLandscape, Log, TEXT("Existing texture has available channels"));
+			if (ExistingTexAvailableChannels >= NeededNewChannels)
+			{
+				break;
+			}
+		}
 
-		// Allocate using our existing textures' spare channels.
-		int32 CurrentAlloc = 0;
-		for (int32 TexIdx = 0; TexIdx < ComponentWeightmapTextures.Num(); TexIdx++)
+		if (ExistingTexAvailableChannels >= NeededNewChannels)
 		{
-			ULandscapeWeightmapUsage* Usage = ComponentWeightmapTexturesUsage[TexIdx];
+			// UE_LOG(LogLandscape, Log, TEXT("Existing texture has available channels"));
 
-			for (int32 ChanIdx = 0; ChanIdx < 4; ChanIdx++)
+			// Allocate using our existing textures' spare channels.
+			int32 CurrentAlloc = 0;
+			for (int32 TexIdx = 0; TexIdx < ComponentWeightmapTextures.Num(); TexIdx++)
 			{
-				if (Usage->ChannelUsage[ChanIdx] == nullptr)
+				ULandscapeWeightmapUsage* Usage = ComponentWeightmapTexturesUsage[TexIdx];
+
+				for (int32 ChanIdx = 0; ChanIdx < 4; ChanIdx++)
 				{
-					// Find next allocation to treat
-					for (; CurrentAlloc < ComponentWeightmapLayerAllocations.Num(); ++CurrentAlloc)
+					if (Usage->ChannelUsage[ChanIdx] == nullptr)
 					{
-						FWeightmapLayerAllocationInfo& AllocInfo = ComponentWeightmapLayerAllocations[CurrentAlloc];
-						
-						if (AllocInfo.WeightmapTextureIndex == 255)
+						// Find next allocation to treat
+						for (; CurrentAlloc < ComponentWeightmapLayerAllocations.Num(); ++CurrentAlloc)
 						{
-							break;
+							FWeightmapLayerAllocationInfo& AllocInfo = ComponentWeightmapLayerAllocations[CurrentAlloc];
+
+							if (AllocInfo.WeightmapTextureIndex == 255)
+							{
+								break;
+							}
 						}
-					}
 
-					FWeightmapLayerAllocationInfo& AllocInfo = ComponentWeightmapLayerAllocations[CurrentAlloc];
-					check(AllocInfo.WeightmapTextureIndex == 255);
+						FWeightmapLayerAllocationInfo& AllocInfo = ComponentWeightmapLayerAllocations[CurrentAlloc];
+						check(AllocInfo.WeightmapTextureIndex == 255);
 
-					// Zero out the data for this texture channel
-					if (DataInterface)
-					{
-						DataInterface->ZeroTextureChannel(ComponentWeightmapTextures[TexIdx], ChanIdx);
-					}
+						// Zero out the data for this texture channel
+						if (DataInterface)
+						{
+							DataInterface->ZeroTextureChannel(ComponentWeightmapTextures[TexIdx], ChanIdx);
+						}
 
-					AllocInfo.WeightmapTextureIndex = TexIdx;
-					AllocInfo.WeightmapTextureChannel = ChanIdx;
+						AllocInfo.WeightmapTextureIndex = TexIdx;
+						AllocInfo.WeightmapTextureChannel = ChanIdx;
 
-					if (InSaveToTransactionBuffer)
-					{
-						Usage->Modify();
-					}
-					Usage->ChannelUsage[ChanIdx] = this;
+						if (InSaveToTransactionBuffer)
+						{
+							Usage->Modify();
+						}
+						Usage->ChannelUsage[ChanIdx] = this;
 
-					NeededNewChannels--;
+						NeededNewChannels--;
 
-					if (NeededNewChannels == 0)
-					{
-						return;
+						if (NeededNewChannels == 0)
+						{
+							return;
+						}
 					}
 				}
 			}
+			// we should never get here.
+			check(false);
 		}
-		// we should never get here.
-		check(false);
 	}
 
 	// UE_LOG(LogLandscape, Log, TEXT("Reallocating."));
@@ -4627,16 +4638,11 @@ void ULandscapeComponent::ReallocateWeightmaps(FLandscapeEditDataInterface* Data
 
 			// see if we can find a suitable existing weightmap texture with sufficient channels
 			int32 BestDistanceSquared = MAX_int32;
-			for (auto& ItPair : Proxy->WeightmapUsageMap)
+			for (auto& ItPair : TargetProxy->WeightmapUsageMap)
 			{
 				ULandscapeWeightmapUsage* TryWeightmapUsage = ItPair.Value;
 				//
-				FGuid EditingLayerGUID = GetEditingLayerGUID();
-				check(!GetMutableDefault<UEditorExperimentalSettings>()->bLandscapeLayerSystem || !InCanUseEditingWeightmap || EditingLayerGUID.IsValid());
-
-				FGuid LayerGuidToSeek = InCanUseEditingWeightmap ? EditingLayerGUID : FGuid();
-
-				if (TryWeightmapUsage->FreeChannelCount() >= TotalNeededChannels && TryWeightmapUsage->LayerGuid == LayerGuidToSeek)
+				if (TryWeightmapUsage->FreeChannelCount() >= TotalNeededChannels && TryWeightmapUsage->LayerGuid == TargetLayerGuid)
 				{
 					if (TryWeightmapUsage->IsEmpty())
 					{
@@ -4675,7 +4681,7 @@ void ULandscapeComponent::ReallocateWeightmaps(FLandscapeEditDataInterface* Data
 			int32 WeightmapSize = (SubsectionSizeQuads + 1) * NumSubsections;
 
 			// We need a new weightmap texture
-			CurrentWeightmapTexture = GetLandscapeProxy()->CreateLandscapeTexture(WeightmapSize, WeightmapSize, TEXTUREGROUP_Terrain_Weightmap, TSF_BGRA8);
+			CurrentWeightmapTexture = TargetProxy->CreateLandscapeTexture(WeightmapSize, WeightmapSize, TEXTUREGROUP_Terrain_Weightmap, TSF_BGRA8);
 
 			// Alloc dummy mips
 			CreateEmptyTextureMips(CurrentWeightmapTexture, true);
@@ -4696,16 +4702,13 @@ void ULandscapeComponent::ReallocateWeightmaps(FLandscapeEditDataInterface* Data
 			}
 
 			// Store it in the usage map
-			CurrentWeightmapUsage = Proxy->WeightmapUsageMap.Add(CurrentWeightmapTexture, GetLandscapeProxy()->CreateWeightmapUsage());
+			CurrentWeightmapUsage = TargetProxy->WeightmapUsageMap.Add(CurrentWeightmapTexture, TargetProxy->CreateWeightmapUsage());
 			if (InSaveToTransactionBuffer)
 			{
 				CurrentWeightmapUsage->Modify();
 			}
 
-			FGuid EditingLayerGUID = GetEditingLayerGUID();
-			check(!GetMutableDefault<UEditorExperimentalSettings>()->bLandscapeLayerSystem || !InCanUseEditingWeightmap || EditingLayerGUID.IsValid());
-
-			CurrentWeightmapUsage->LayerGuid = InCanUseEditingWeightmap ? EditingLayerGUID : FGuid();
+			CurrentWeightmapUsage->LayerGuid = TargetLayerGuid;
 			// UE_LOG(LogLandscape, Log, TEXT("Making a new texture %s"), *CurrentWeightmapTexture->GetName());
 		}
 
@@ -5366,7 +5369,7 @@ void ULandscapeComponent::GeneratePlatformPixelData()
 
 	MobileWeightmapTextures.Empty();
 
-    UTexture2D* MobileWeightNormalmapTexture = GetLandscapeProxy()->CreateLandscapeTexture(WeightmapSize, WeightmapSize, TEXTUREGROUP_Terrain_Weightmap, TSF_BGRA8, nullptr, GMobileCompressLandscapeWeightMaps ? true : false );
+    UTexture2D* MobileWeightNormalmapTexture = GetLandscapeProxy()->CreateLandscapeTexture(WeightmapSize, WeightmapSize, TEXTUREGROUP_Terrain_Weightmap, TSF_BGRA8, GMobileCompressLandscapeWeightMaps ? true : false );
 	CreateEmptyTextureMips(MobileWeightNormalmapTexture);
 
 	{
@@ -5409,7 +5412,7 @@ void ULandscapeComponent::GeneratePlatformPixelData()
 					// create a new weightmap texture if we've run out of channels
 					CurrentChannel = 0;
 					RemainingChannels = 4;
-                    CurrentWeightmapTexture = GetLandscapeProxy()->CreateLandscapeTexture(WeightmapSize, WeightmapSize, TEXTUREGROUP_Terrain_Weightmap, TSF_BGRA8, nullptr, GMobileCompressLandscapeWeightMaps ? true : false);
+                    CurrentWeightmapTexture = GetLandscapeProxy()->CreateLandscapeTexture(WeightmapSize, WeightmapSize, TEXTUREGROUP_Terrain_Weightmap, TSF_BGRA8, GMobileCompressLandscapeWeightMaps ? true : false);
 					CreateEmptyTextureMips(CurrentWeightmapTexture);
 					MobileWeightmapTextures.Add(CurrentWeightmapTexture);
 				}
@@ -5452,7 +5455,7 @@ void ULandscapeComponent::GeneratePlatformPixelData()
 
 		for (int32 MaterialIndex = 0; MaterialIndex < MobileCombinationMaterialInstances.Num(); ++MaterialIndex)
 		{
-			UMaterialInstanceDynamic* NewMobileMaterialInstance = UMaterialInstanceDynamic::Create(MobileCombinationMaterialInstances[MaterialIndex], GetOutermost());
+			UMaterialInstanceDynamic* NewMobileMaterialInstance = UMaterialInstanceDynamic::Create(MobileCombinationMaterialInstances[MaterialIndex], this);
 
 			// Set the layer mask
 			for (const auto& Allocation : MobileWeightmapLayerAllocations)
@@ -5507,7 +5510,7 @@ void ULandscapeComponent::GeneratePlatformPixelData()
 			MobileCombinationMaterialInstances[MaterialIndex] = GetCombinationMaterial(nullptr, MobileWeightmapLayerAllocations, MaterialLOD, true);
 			check(MobileCombinationMaterialInstances[MaterialIndex] != nullptr);
 
-			UMaterialInstanceConstant* NewMobileMaterialInstance = NewObject<ULandscapeMaterialInstanceConstant>(GetOutermost());
+			UMaterialInstanceConstant* NewMobileMaterialInstance = NewObject<ULandscapeMaterialInstanceConstant>(this);
 
 			NewMobileMaterialInstance->SetParentEditorOnly(MobileCombinationMaterialInstances[MaterialIndex]);
 
@@ -5748,10 +5751,9 @@ void ULandscapeComponent::GeneratePlatformVertexData(const ITargetPlatform* Targ
 	PlatformData.InitializeFromUncompressedData(NewPlatformData);
 }
 
-UTexture2D* ALandscapeProxy::CreateLandscapeTexture(int32 InSizeX, int32 InSizeY, TextureGroup InLODGroup, ETextureSourceFormat InFormat, UObject* OptionalOverrideOuter, bool bCompress) const
+UTexture2D* ALandscapeProxy::CreateLandscapeTexture(int32 InSizeX, int32 InSizeY, TextureGroup InLODGroup, ETextureSourceFormat InFormat, bool bCompress) const
 {
-	UObject* TexOuter = OptionalOverrideOuter ? OptionalOverrideOuter : GetOutermost();
-	UTexture2D* NewTexture = NewObject<UTexture2D>(TexOuter);
+	UTexture2D* NewTexture = NewObject<UTexture2D>(const_cast<ALandscapeProxy*>(this));
 	NewTexture->Source.Init2DWithMipChain(InSizeX, InSizeY, InFormat);
 	NewTexture->SRGB = false;
 	NewTexture->CompressionNone = !bCompress;
diff --git a/Engine/Source/Runtime/Landscape/Private/LandscapeEditLayers.cpp b/Engine/Source/Runtime/Landscape/Private/LandscapeEditLayers.cpp
--- a/Engine/Source/Runtime/Landscape/Private/LandscapeEditLayers.cpp
+++ b/Engine/Source/Runtime/Landscape/Private/LandscapeEditLayers.cpp
@@ -1282,7 +1282,7 @@ void ALandscape::CreateLayersRenderingResource(const FIntPoint& InComponentCount
 
 		for (int32 i = 0; i < EHeightmapRTType::HeightmapRT_Count; ++i)
 		{
-			Landscape->HeightmapRTList[i] = NewObject<UTextureRenderTarget2D>(Landscape->GetOutermost());
+			Landscape->HeightmapRTList[i] = NewObject<UTextureRenderTarget2D>(Landscape);
 			check(Landscape->HeightmapRTList[i]);
 			Landscape->HeightmapRTList[i]->RenderTargetFormat = RTF_RGBA8;
 			Landscape->HeightmapRTList[i]->AddressX = TextureAddress::TA_Clamp;
@@ -1344,7 +1344,7 @@ void ALandscape::CreateLayersRenderingResource(const FIntPoint& InComponentCount
 
 		for (int32 i = 0; i < EWeightmapRTType::WeightmapRT_Count; ++i)
 		{
-			Landscape->WeightmapRTList[i] = NewObject<UTextureRenderTarget2D>(Landscape->GetOutermost());
+			Landscape->WeightmapRTList[i] = NewObject<UTextureRenderTarget2D>(Landscape);
 
 			check(Landscape->WeightmapRTList[i]);
 			Landscape->WeightmapRTList[i]->AddressX = TextureAddress::TA_Clamp;
@@ -3102,7 +3102,7 @@ void ALandscape::ReallocateLayersWeightmaps(const TArray<ULandscapeLayerInfoObje
 
 	for (ULandscapeComponent* Component : AllLandscapeComponents)
 	{
-		Component->ReallocateWeightmaps(nullptr, false, false, true, &NewCreatedTextures);
+		Component->ReallocateWeightmaps(nullptr, false, false, true, false, nullptr, &NewCreatedTextures);
 	}
 
 	// TODO: correctly only recreate what is required instead of everything..
@@ -3790,7 +3790,7 @@ void ALandscape::UpdateLayersMaterialInstances(const TArray<ULandscapeComponent*
 				if (NeedToCreateMIC)
 				{
 					// Create the instance for this component, that will use the layer combination instance.
-					MaterialInstance = NewObject<ULandscapeMaterialInstanceConstant>(GetOutermost());
+					MaterialInstance = NewObject<ULandscapeMaterialInstanceConstant>(this);
 					Component->MaterialInstances[MaterialIndex] = MaterialInstance;
 				}
 
@@ -3828,7 +3828,7 @@ void ALandscape::UpdateLayersMaterialInstances(const TArray<ULandscapeComponent*
 
 					if (NeedToCreateMIC || TessellationMaterialInstance == nullptr)
 					{
-						TessellationMaterialInstance = NewObject<ULandscapeMaterialInstanceConstant>(GetOutermost());
+						TessellationMaterialInstance = NewObject<ULandscapeMaterialInstanceConstant>(this);
 						TessellationMaterialInstance->SetParentEditorOnly(MaterialInstance);
 
 						MaterialInstances[TessellatedMaterialIndex] = TessellationMaterialInstance;
