#UE4 Added more info to a warning

diff --git a/Engine/Source/Runtime/Landscape/Private/LandscapeEdit.cpp b/Engine/Source/Runtime/Landscape/Private/LandscapeEdit.cpp
--- a/Engine/Source/Runtime/Landscape/Private/LandscapeEdit.cpp
+++ b/Engine/Source/Runtime/Landscape/Private/LandscapeEdit.cpp
@@ -623,7 +623,7 @@ void ULandscapeComponent::FixupWeightmaps()
 					if (!bFixedLayerDeletion)
 					{
 						FFormatNamedArguments Arguments;
-						Arguments.Add(TEXT("LandscapeName"), FText::FromString(GetName()));
+						Arguments.Add(TEXT("LandscapeName"), FText::FromString(GetPathName()));
 						FMessageLog("MapCheck").Warning()
 							->AddToken(FTextToken::Create(FText::Format(LOCTEXT("MapCheck_Message_FixedUpDeletedLayerWeightmap", "{LandscapeName} : Fixed up deleted layer weightmap"), Arguments)))
 							->AddToken(FMapErrorToken::Create(FMapErrors::FixedUpDeletedLayerWeightmap));
