Add support for AddComponent tool in the Layers System
Fixed a case where we were doing multiple copy of the same data when regenerating heightmap

diff --git a/Engine/Source/Editor/LandscapeEditor/Private/LandscapeEdModeComponentTools.cpp b/Engine/Source/Editor/LandscapeEditor/Private/LandscapeEdModeComponentTools.cpp
--- a/Engine/Source/Editor/LandscapeEditor/Private/LandscapeEdModeComponentTools.cpp
+++ b/Engine/Source/Editor/LandscapeEditor/Private/LandscapeEdModeComponentTools.cpp
@@ -722,8 +722,8 @@ public:
 
 	virtual void Apply(FEditorViewportClient* ViewportClient, FLandscapeBrush* Brush, const ULandscapeEditorObject* UISettings, const TArray<FLandscapeToolInteractorPosition>& InteractorPositions)
 	{
-		ALandscapeProxy* Landscape = LandscapeInfo ? LandscapeInfo->GetCurrentLevelLandscapeProxy(true) : nullptr;
-		if (Landscape && EdMode->LandscapeRenderAddCollision)
+		ALandscapeProxy* LandscapeProxy = LandscapeInfo ? LandscapeInfo->GetCurrentLevelLandscapeProxy(true) : nullptr;
+		if (LandscapeProxy && EdMode->LandscapeRenderAddCollision)
 		{
 			check(Brush->GetBrushType() == ELandscapeBrushType::Component);
 
@@ -740,7 +740,7 @@ public:
 
 			// Find component range for this block of data, non shared vertices
 			int32 ComponentIndexX1, ComponentIndexY1, ComponentIndexX2, ComponentIndexY2;
-			ALandscape::CalcComponentIndicesNoOverlap(X1, Y1, X2, Y2, Landscape->ComponentSizeQuads, ComponentIndexX1, ComponentIndexY1, ComponentIndexX2, ComponentIndexY2);
+			ALandscape::CalcComponentIndicesNoOverlap(X1, Y1, X2, Y2, LandscapeProxy->ComponentSizeQuads, ComponentIndexX1, ComponentIndexY1, ComponentIndexX2, ComponentIndexY2);
 
 			// expand the area by one vertex in each direction to ensure normals are calculated correctly
 			X1 -= 1;
@@ -756,7 +756,7 @@ public:
 			bool bHasXYOffset = XYOffsetCache.GetCachedData(X1, Y1, X2, Y2, XYOffsetData);
 
 			TArray<ULandscapeComponent*> NewComponents;
-			Landscape->Modify();
+			LandscapeProxy->Modify();
 			LandscapeInfo->Modify();
 			for (int32 ComponentIndexY = ComponentIndexY1; ComponentIndexY <= ComponentIndexY2; ComponentIndexY++)
 			{
@@ -766,22 +766,22 @@ public:
 					if (!LandscapeComponent)
 					{
 						// Add New component...
-						FIntPoint ComponentBase = FIntPoint(ComponentIndexX, ComponentIndexY)*Landscape->ComponentSizeQuads;
-						LandscapeComponent = NewObject<ULandscapeComponent>(Landscape, NAME_None, RF_Transactional);
-						Landscape->LandscapeComponents.Add(LandscapeComponent);
+						FIntPoint ComponentBase = FIntPoint(ComponentIndexX, ComponentIndexY)*LandscapeProxy->ComponentSizeQuads;
+						LandscapeComponent = NewObject<ULandscapeComponent>(LandscapeProxy, NAME_None, RF_Transactional);
+						LandscapeProxy->LandscapeComponents.Add(LandscapeComponent);
 						NewComponents.Add(LandscapeComponent);
 						LandscapeComponent->Init(
 							ComponentBase.X, ComponentBase.Y,
-							Landscape->ComponentSizeQuads,
-							Landscape->NumSubsections,
-							Landscape->SubsectionSizeQuads
+							LandscapeProxy->ComponentSizeQuads,
+							LandscapeProxy->NumSubsections,
+							LandscapeProxy->SubsectionSizeQuads
 							);
-						LandscapeComponent->AttachToComponent(Landscape->GetRootComponent(), FAttachmentTransformRules::KeepRelativeTransform);
+						LandscapeComponent->AttachToComponent(LandscapeProxy->GetRootComponent(), FAttachmentTransformRules::KeepRelativeTransform);
 
 						// Assign shared properties
 						LandscapeComponent->UpdatedSharedPropertiesFromActor();
 
-						int32 ComponentVerts = (Landscape->SubsectionSizeQuads + 1) * Landscape->NumSubsections;
+						int32 ComponentVerts = (LandscapeProxy->SubsectionSizeQuads + 1) * LandscapeProxy->NumSubsections;
 						// Update Weightmap Scale Bias
 						LandscapeComponent->WeightmapScaleBias = FVector4(1.0f / (float)ComponentVerts, 1.0f / (float)ComponentVerts, 0.5f / (float)ComponentVerts, 0.5f / (float)ComponentVerts);
 						LandscapeComponent->WeightmapSubsectionOffset = (float)(LandscapeComponent->SubsectionSizeQuads + 1) / (float)ComponentVerts;
@@ -820,80 +820,110 @@ public:
 			HeightCache.SetCachedData(X1, Y1, X2, Y2, Data);
 			HeightCache.Flush();
 
+			ALandscape* Landscape = LandscapeInfo->LandscapeActor.Get();
+
+			if (GetMutableDefault<UEditorExperimentalSettings>()->bLandscapeLayerSystem)
+			{
+				check(Landscape != nullptr); // Landscape actor is required if layer system is enabled
+
+				Landscape->CreateLayersRenderingResource(Landscape->ComputeComponentCounts());
+			}
+
 			for (ULandscapeComponent* NewComponent : NewComponents)
 			{
+				if (GetMutableDefault<UEditorExperimentalSettings>()->bLandscapeLayerSystem)
+				{
+					// Since we do not share heightmap when adding new component, we will provided the required array, but they will only be used for 1 component
+					TMap<UTexture2D*, UTexture2D*> CreatedHeightmapTextures;
+					TArray<ULandscapeComponent*> ComponentsUsingHeightmap;
+					ComponentsUsingHeightmap.Add(NewComponent);
+
+					for (const FLandscapeLayer& Layer : Landscape->LandscapeLayers)
+					{
+						NewComponent->AddDefaultLayerData(Layer.Guid, ComponentsUsingHeightmap, CreatedHeightmapTextures);
+					}
+				}
+
 				// Update Collision
 				NewComponent->UpdateCachedBounds();
 				NewComponent->UpdateBounds();
 				NewComponent->MarkRenderStateDirty();
-				ULandscapeHeightfieldCollisionComponent* CollisionComp = NewComponent->CollisionComponent.Get();
-				if (CollisionComp && !bHasXYOffset)
+
+				if (!GetMutableDefault<UEditorExperimentalSettings>()->bLandscapeLayerSystem)
 				{
-					CollisionComp->MarkRenderStateDirty();
-					CollisionComp->RecreateCollision();
+					ULandscapeHeightfieldCollisionComponent* CollisionComp = NewComponent->CollisionComponent.Get();
+					if (CollisionComp && !bHasXYOffset)
+					{
+						CollisionComp->MarkRenderStateDirty();
+						CollisionComp->RecreateCollision();
+					}
 				}
 
 				TMap<ULandscapeLayerInfoObject*, int32> NeighbourLayerInfoObjectCount;
 
-				// Cover 9 tiles around us to determine which object should we use by default
-				for (int32 ComponentIndexX = ComponentIndexX1 - 1; ComponentIndexX <= ComponentIndexX2 + 1; ++ComponentIndexX)
 				{
-					for (int32 ComponentIndexY = ComponentIndexY1 - 1; ComponentIndexY <= ComponentIndexY2 + 1; ++ComponentIndexY)
-					{
-						ULandscapeComponent* NeighbourComponent = LandscapeInfo->XYtoComponentMap.FindRef(FIntPoint(ComponentIndexX, ComponentIndexY));
+					FScopedSetLandscapeEditingLayer Scope(Landscape, Landscape->GetLayer(0)->Guid, [=] { });
 
-						if (NeighbourComponent != nullptr && NeighbourComponent != NewComponent)
+					// Cover 9 tiles around us to determine which object should we use by default
+					for (int32 ComponentIndexX = ComponentIndexX1 - 1; ComponentIndexX <= ComponentIndexX2 + 1; ++ComponentIndexX)
+					{
+						for (int32 ComponentIndexY = ComponentIndexY1 - 1; ComponentIndexY <= ComponentIndexY2 + 1; ++ComponentIndexY)
 						{
-							ULandscapeInfo* NeighbourLandscapeInfo = NeighbourComponent->GetLandscapeInfo();
+							ULandscapeComponent* NeighbourComponent = LandscapeInfo->XYtoComponentMap.FindRef(FIntPoint(ComponentIndexX, ComponentIndexY));
 
-							for (int32 i = 0; i < NeighbourLandscapeInfo->Layers.Num(); ++i)
+							if (NeighbourComponent != nullptr && NeighbourComponent != NewComponent)
 							{
-								ULandscapeLayerInfoObject* NeighbourLayerInfo = NeighbourLandscapeInfo->Layers[i].LayerInfoObj;
+								ULandscapeInfo* NeighbourLandscapeInfo = NeighbourComponent->GetLandscapeInfo();
 
-								if (NeighbourLayerInfo != nullptr)
+								for (int32 i = 0; i < NeighbourLandscapeInfo->Layers.Num(); ++i)
 								{
-									TArray<uint8> WeightmapTextureData;
+									ULandscapeLayerInfoObject* NeighbourLayerInfo = NeighbourLandscapeInfo->Layers[i].LayerInfoObj;
 
-									FLandscapeComponentDataInterface DataInterface(NeighbourComponent);
-									DataInterface.GetWeightmapTextureData(NeighbourLayerInfo, WeightmapTextureData);
-
-									if (WeightmapTextureData.Num() > 0)
+									if (NeighbourLayerInfo != nullptr)
 									{
-										int32* Count = NeighbourLayerInfoObjectCount.Find(NeighbourLayerInfo);
+										TArray<uint8> WeightmapTextureData;
 
-										if (Count == nullptr)
-										{
-											Count = &NeighbourLayerInfoObjectCount.Add(NeighbourLayerInfo, 1);
-										}
+										FLandscapeComponentDataInterface DataInterface(NeighbourComponent);
+										DataInterface.GetWeightmapTextureData(NeighbourLayerInfo, WeightmapTextureData, true);
 
-										for (uint8 Value : WeightmapTextureData)
+										if (WeightmapTextureData.Num() > 0)
 										{
-											(*Count) += Value;
+											int32* Count = NeighbourLayerInfoObjectCount.Find(NeighbourLayerInfo);
+
+											if (Count == nullptr)
+											{
+												Count = &NeighbourLayerInfoObjectCount.Add(NeighbourLayerInfo, 1);
+											}
+
+											for (uint8 Value : WeightmapTextureData)
+											{
+												(*Count) += Value;
+											}
 										}
 									}
 								}
-							}							
+							}
 						}
-					}					
-				}
+					}
 
-				int32 BestLayerInfoObjectCount = 0;
-				ULandscapeLayerInfoObject* BestLayerInfoObject = nullptr;
+					int32 BestLayerInfoObjectCount = 0;
+					ULandscapeLayerInfoObject* BestLayerInfoObject = nullptr;
 
-				for (auto& LayerInfoObjectCount : NeighbourLayerInfoObjectCount)
-				{
-					if (LayerInfoObjectCount.Value > BestLayerInfoObjectCount)
+					for (auto& LayerInfoObjectCount : NeighbourLayerInfoObjectCount)
 					{
-						BestLayerInfoObjectCount = LayerInfoObjectCount.Value;
-						BestLayerInfoObject = LayerInfoObjectCount.Key;
+						if (LayerInfoObjectCount.Value > BestLayerInfoObjectCount)
+						{
+							BestLayerInfoObjectCount = LayerInfoObjectCount.Value;
+							BestLayerInfoObject = LayerInfoObjectCount.Key;
+						}
 					}
-				}
-				
-				if (BestLayerInfoObject != nullptr)
-				{
-					FLandscapeEditDataInterface LandscapeEdit(LandscapeInfo);
-					NewComponent->FillLayer(BestLayerInfoObject, LandscapeEdit);
-				}
+
+					if (BestLayerInfoObject != nullptr)
+					{
+						FLandscapeEditDataInterface LandscapeEdit(LandscapeInfo);
+						NewComponent->FillLayer(BestLayerInfoObject, LandscapeEdit);
+					}
+				}				
 			}
 
 			EdMode->LandscapeRenderAddCollision = nullptr;
diff --git a/Engine/Source/Runtime/Landscape/Classes/Landscape.h b/Engine/Source/Runtime/Landscape/Classes/Landscape.h
--- a/Engine/Source/Runtime/Landscape/Classes/Landscape.h
+++ b/Engine/Source/Runtime/Landscape/Classes/Landscape.h
@@ -303,9 +303,9 @@ private:
 
 	void UpdateLayersMaterialInstances(const TArray<ULandscapeComponent*>& InComponentsToUpdate);
 
-	void PrepareComponentDataToExtractMaterialLayersCS(const FLandscapeLayer& InLayer, int32 InCurrentWeightmapToProcessIndex, bool InOutputDebugName, class FLandscapeTexture2DResource* InOutTextureData,
+	void PrepareComponentDataToExtractMaterialLayersCS(const FLandscapeLayer& InLayer, int32 InCurrentWeightmapToProcessIndex, const FIntPoint& InLandscapeBase, bool InOutputDebugName, class FLandscapeTexture2DResource* InOutTextureData,
 														  TArray<struct FLandscapeLayerWeightmapExtractMaterialLayersComponentData>& OutComponentData, TMap<ULandscapeLayerInfoObject*, int32>& OutLayerInfoObjects);
-	void PrepareComponentDataToPackMaterialLayersCS(int32 InCurrentWeightmapToProcessIndex, bool InOutputDebugName, const TArray<ULandscapeComponent*>& InAllLandscapeComponents, TArray<UTexture2D*>& InOutProcessedWeightmaps,
+	void PrepareComponentDataToPackMaterialLayersCS(int32 InCurrentWeightmapToProcessIndex, const FIntPoint& InLandscapeBase, bool InOutputDebugName, const TArray<ULandscapeComponent*>& InAllLandscapeComponents, TArray<UTexture2D*>& InOutProcessedWeightmaps,
 													TArray<class FLandscapeLayersTexture2DCPUReadBackResource*>& OutProcessedCPUReadBackTexture, TArray<struct FLandscapeLayerWeightmapPackMaterialLayersComponentData>& OutComponentData);
 	void ReallocateLayersWeightmaps(const TArray<ULandscapeLayerInfoObject*>& InBrushRequiredAllocations, TArray<ULandscapeComponent*>& OutComponentThatNeedMaterialRebuild);
 	void InitializeLayersWeightmapResources();
@@ -319,21 +319,21 @@ private:
 	void GenerateLayersRenderQuadsMip(const FIntPoint& InSectionBase, const FVector2D& InScaleBias, float InSubSectionSizeQuad, const FIntPoint& InReadSize, const FIntPoint& InWriteSize, uint8 InCurrentMip, TArray<FLandscapeLayersTriangle>& OutTriangles) const;
 
 	void ClearLayersWeightmapTextureResource(const FString& InDebugName, FTextureRenderTargetResource* InTextureResourceToClear) const;
-	void DrawHeightmapComponentsToRenderTarget(const FString& InDebugName, const TArray<ULandscapeComponent*>& InComponentsToDraw, UTexture* InHeightmapRTRead, UTextureRenderTarget2D* InOptionalHeightmapRTRead2, UTextureRenderTarget2D* InHeightmapRTWrite, ERTDrawingType InDrawType,
+	void DrawHeightmapComponentsToRenderTarget(const FString& InDebugName, const TArray<ULandscapeComponent*>& InComponentsToDraw, const FIntPoint& InLandscapeBase, UTexture* InHeightmapRTRead, UTextureRenderTarget2D* InOptionalHeightmapRTRead2, UTextureRenderTarget2D* InHeightmapRTWrite, ERTDrawingType InDrawType,
 											   bool InClearRTWrite, struct FLandscapeLayersHeightmapShaderParameters& InShaderParams, uint8 InMipRender = 0) const;
 
 	void DrawWeightmapComponentsToRenderTarget(const FString& InDebugName, const TArray<FIntPoint>& InSectionBaseList, const FVector2D& InScaleBias, UTexture* InWeightmapRTRead, UTextureRenderTarget2D* InOptionalWeightmapRTRead2, UTextureRenderTarget2D* InWeightmapRTWrite,
 											   bool InClearRTWrite, struct FLandscapeLayersWeightmapShaderParameters& InShaderParams, uint8 InMipRender) const;
 
-	void DrawWeightmapComponentsToRenderTarget(const FString& InDebugName, const TArray<ULandscapeComponent*>& InComponentsToDraw, UTexture* InWeightmapRTRead, UTextureRenderTarget2D* InOptionalWeightmapRTRead2, UTextureRenderTarget2D* InWeightmapRTWrite,
+	void DrawWeightmapComponentsToRenderTarget(const FString& InDebugName, const TArray<ULandscapeComponent*>& InComponentsToDraw, const FIntPoint& InLandscapeBase, UTexture* InWeightmapRTRead, UTextureRenderTarget2D* InOptionalWeightmapRTRead2, UTextureRenderTarget2D* InWeightmapRTWrite,
 											   bool InClearRTWrite, struct FLandscapeLayersWeightmapShaderParameters& InShaderParams, uint8 InMipRender) const;
 
-	void DrawHeightmapComponentsToRenderTargetMips(const TArray<ULandscapeComponent*>& InComponentsToDraw, UTexture* InReadHeightmap, bool InClearRTWrite, struct FLandscapeLayersHeightmapShaderParameters& InShaderParams) const;
+	void DrawHeightmapComponentsToRenderTargetMips(const TArray<ULandscapeComponent*>& InComponentsToDraw, const FIntPoint& InLandscapeBase, UTexture* InReadHeightmap, bool InClearRTWrite, struct FLandscapeLayersHeightmapShaderParameters& InShaderParams) const;
 	void DrawWeightmapComponentToRenderTargetMips(const TArray<FVector2D>& InTexturePositionsToDraw, UTexture* InReadWeightmap, bool InClearRTWrite, struct FLandscapeLayersWeightmapShaderParameters& InShaderParams) const;
 
-	void CopyLayersTexture(UTexture* InSourceTexture, UTexture* InDestTexture, FTextureResource* InDestCPUResource = nullptr, const FIntPoint& InFirstComponentSectionBase = FIntPoint(0, 0), uint8 InSourceCurrentMip = 0, uint8 InDestCurrentMip = 0,
+	void CopyLayersTexture(UTexture* InSourceTexture, UTexture* InDestTexture, FTextureResource* InDestCPUResource = nullptr, const FIntPoint& InInitialPositionOffset = FIntPoint(0, 0), uint8 InSourceCurrentMip = 0, uint8 InDestCurrentMip = 0,
 							   uint32 InSourceArrayIndex = 0, uint32 InDestArrayIndex = 0) const;
-	void CopyLayersTexture(const FString& InSourceDebugName, FTextureResource* InSourceResource, const FString& InDestDebugName, FTextureResource* InDestResource, FTextureResource* InDestCPUResource = nullptr, const FIntPoint& InFirstComponentSectionBase = FIntPoint(0, 0),
+	void CopyLayersTexture(const FString& InSourceDebugName, FTextureResource* InSourceResource, const FString& InDestDebugName, FTextureResource* InDestResource, FTextureResource* InDestCPUResource = nullptr, const FIntPoint& InInitialPositionOffset = FIntPoint(0, 0),
 							   uint8 InSourceCurrentMip = 0, uint8 InDestCurrentMip = 0, uint32 InSourceArrayIndex = 0, uint32 uInDestArrayIndex = 0) const;
 
 	void InitializeLayers();
diff --git a/Engine/Source/Runtime/Landscape/Classes/LandscapeComponent.h b/Engine/Source/Runtime/Landscape/Classes/LandscapeComponent.h
--- a/Engine/Source/Runtime/Landscape/Classes/LandscapeComponent.h
+++ b/Engine/Source/Runtime/Landscape/Classes/LandscapeComponent.h
@@ -576,6 +576,7 @@ public:
 	LANDSCAPE_API const FLandscapeLayerComponentData* GetLayerData(const FGuid& InLayerGuid) const;
 	LANDSCAPE_API FLandscapeLayerComponentData* GetLayerData(const FGuid& InLayerGuid);
 	LANDSCAPE_API void AddLayerData(const FGuid& InLayerGuid, const FLandscapeLayerComponentData& InData);
+	LANDSCAPE_API void AddDefaultLayerData(const FGuid& InLayerGuid, const TArray<ULandscapeComponent*>& InComponentsUsingHeightmap, TMap<UTexture2D*, UTexture2D*>& InOutCreatedHeightmapTextures);
 	LANDSCAPE_API void RemoveLayerData(const FGuid& InLayerGuid);
 	LANDSCAPE_API void ForEachLayer(TFunctionRef<void(const FGuid&, struct FLandscapeLayerComponentData&)> Fn);
 
diff --git a/Engine/Source/Runtime/Landscape/Private/Landscape.cpp b/Engine/Source/Runtime/Landscape/Private/Landscape.cpp
--- a/Engine/Source/Runtime/Landscape/Private/Landscape.cpp
+++ b/Engine/Source/Runtime/Landscape/Private/Landscape.cpp
@@ -59,6 +59,7 @@ Landscape.cpp: Terrain rendering
 #endif
 #include "LandscapeVersion.h"
 #include "UObject/FortniteMainBranchObjectVersion.h"
+#include "LandscapeDataAccess.h"
 
 /** Landscape stats */
 
@@ -1492,6 +1493,93 @@ void ULandscapeComponent::AddLayerData(const FGuid& InLayerGuid, const FLandscap
 	Data = InData;
 }
 
+void ULandscapeComponent::AddDefaultLayerData(const FGuid& InLayerGuid, const TArray<ULandscapeComponent*>& InComponentsUsingHeightmap, TMap<UTexture2D*, UTexture2D*>& InOutCreatedHeightmapTextures)
+{
+	Modify();
+
+	UTexture2D* ComponentHeightmap = GetHeightmap();
+
+	// Compute Global layers data
+	GlobalLayersData.TopLeftSectionBase = FIntPoint(INT_MAX, INT_MAX);
+
+	for (ULandscapeComponent* ComponentUsingHeightmap : InComponentsUsingHeightmap)
+	{
+		GlobalLayersData.TopLeftSectionBase.X = FMath::Min(GlobalLayersData.TopLeftSectionBase.X, ComponentUsingHeightmap->GetSectionBase().X);
+		GlobalLayersData.TopLeftSectionBase.Y = FMath::Min(GlobalLayersData.TopLeftSectionBase.Y, ComponentUsingHeightmap->GetSectionBase().Y);
+	}
+
+	// Compute per layer data
+	FLandscapeLayerComponentData* LayerData = GetLayerData(InLayerGuid);
+
+	if (LayerData == nullptr || !LayerData->IsInitialized())
+	{
+		FLandscapeLayerComponentData NewData;
+
+		// Setup Heightmap data
+		UTexture2D** LayerHeightmap = InOutCreatedHeightmapTextures.Find(ComponentHeightmap);
+
+		if (LayerHeightmap == nullptr)
+		{
+			UTexture2D* NewLayerHeightmap = GetLandscapeProxy()->CreateLandscapeTexture(ComponentHeightmap->Source.GetSizeX(), ComponentHeightmap->Source.GetSizeY(), TEXTUREGROUP_Terrain_Heightmap, ComponentHeightmap->Source.GetFormat());
+			LayerHeightmap = &InOutCreatedHeightmapTextures.Add(ComponentHeightmap, NewLayerHeightmap);
+
+			ULandscapeComponent::CreateEmptyTextureMips(NewLayerHeightmap, true);
+
+			// Init Mip0 to be at 32768 which is equal to "0"
+			FColor* Mip0Data = (FColor*)NewLayerHeightmap->Source.LockMip(0);
+
+			for (ULandscapeComponent* ComponentUsingHeightmap : InComponentsUsingHeightmap)
+			{
+				int32 HeightmapComponentOffsetX = FMath::RoundToInt((float)NewLayerHeightmap->Source.GetSizeX() * ComponentUsingHeightmap->HeightmapScaleBias.Z);
+				int32 HeightmapComponentOffsetY = FMath::RoundToInt((float)NewLayerHeightmap->Source.GetSizeY() * ComponentUsingHeightmap->HeightmapScaleBias.W);
+
+				for (int32 SubsectionY = 0; SubsectionY < NumSubsections; SubsectionY++)
+				{
+					for (int32 SubsectionX = 0; SubsectionX < NumSubsections; SubsectionX++)
+					{
+						for (int32 SubY = 0; SubY <= SubsectionSizeQuads; SubY++)
+						{
+							for (int32 SubX = 0; SubX <= SubsectionSizeQuads; SubX++)
+							{
+								// X/Y of the vertex we're looking at in component's coordinates.
+								const int32 CompX = SubsectionSizeQuads * SubsectionX + SubX;
+								const int32 CompY = SubsectionSizeQuads * SubsectionY + SubY;
+
+								// X/Y of the vertex we're looking indexed into the texture data
+								const int32 TexX = (SubsectionSizeQuads + 1) * SubsectionX + SubX;
+								const int32 TexY = (SubsectionSizeQuads + 1) * SubsectionY + SubY;
+
+								const int32 HeightTexDataIdx = (HeightmapComponentOffsetX + TexX) + (HeightmapComponentOffsetY + TexY) * NewLayerHeightmap->Source.GetSizeX();
+
+								// copy height and normal data
+								const uint16 HeightValue = LandscapeDataAccess::GetTexHeight(0.f);
+
+								Mip0Data[HeightTexDataIdx].R = HeightValue >> 8;
+								Mip0Data[HeightTexDataIdx].G = HeightValue & 255;
+
+								// Normal with get calculated later
+								Mip0Data[HeightTexDataIdx].B = 0.0f;
+								Mip0Data[HeightTexDataIdx].A = 0.0f;
+							}
+						}
+					}
+				}
+			}
+
+			NewLayerHeightmap->Source.UnlockMip(0);
+
+			NewLayerHeightmap->BeginCachePlatformData();
+			NewLayerHeightmap->ClearAllCachedCookedPlatformData();
+		}
+
+		NewData.HeightmapData.Texture = *LayerHeightmap;
+
+		// Nothing to do for Weightmap by default
+
+		AddLayerData(InLayerGuid, MoveTemp(NewData));
+	}
+}
+
 void ULandscapeComponent::RemoveLayerData(const FGuid& InLayerGuid)
 {
 	Modify();
diff --git a/Engine/Source/Runtime/Landscape/Private/LandscapeDataAccess.cpp b/Engine/Source/Runtime/Landscape/Private/LandscapeDataAccess.cpp
--- a/Engine/Source/Runtime/Landscape/Private/LandscapeDataAccess.cpp
+++ b/Engine/Source/Runtime/Landscape/Private/LandscapeDataAccess.cpp
@@ -71,11 +71,11 @@ LANDSCAPE_API void FLandscapeComponentDataInterface::GetHeightmapTextureData(TAr
 	}
 }
 
-LANDSCAPE_API bool FLandscapeComponentDataInterface::GetWeightmapTextureData(ULandscapeLayerInfoObject* LayerInfo, TArray<uint8>& OutData)
+LANDSCAPE_API bool FLandscapeComponentDataInterface::GetWeightmapTextureData(ULandscapeLayerInfoObject* LayerInfo, TArray<uint8>& OutData, bool InUseEditingWeightmap)
 {
 	int32 LayerIdx = INDEX_NONE;
-	TArray<FWeightmapLayerAllocationInfo>& ComponentWeightmapLayerAllocations = Component->GetWeightmapLayerAllocations();
-	TArray<UTexture2D*>& ComponentWeightmapTextures = Component->GetWeightmapTextures();
+	TArray<FWeightmapLayerAllocationInfo>& ComponentWeightmapLayerAllocations = Component->GetWeightmapLayerAllocations(InUseEditingWeightmap);
+	TArray<UTexture2D*>& ComponentWeightmapTextures = Component->GetWeightmapTextures(InUseEditingWeightmap);
 
 	for (int32 Idx = 0; Idx < ComponentWeightmapLayerAllocations.Num(); Idx++)
 	{
diff --git a/Engine/Source/Runtime/Landscape/Private/LandscapeEditLayers.cpp b/Engine/Source/Runtime/Landscape/Private/LandscapeEditLayers.cpp
--- a/Engine/Source/Runtime/Landscape/Private/LandscapeEditLayers.cpp
+++ b/Engine/Source/Runtime/Landscape/Private/LandscapeEditLayers.cpp
@@ -986,7 +986,7 @@ class FLandscapeLayersCopyTexture_RenderThread
 {
 public:
 	FLandscapeLayersCopyTexture_RenderThread(const FString& InSourceResourceDebugName, FTextureResource* InSourceResource, const FString& InDestResourceDebugName, FTextureResource* InDestResource, FTextureResource* InDestCPUResource,
-											 const FIntPoint& InFirstComponentSectionBase, int32 InSubSectionSizeQuad, int32 InNumSubSections, uint8 InSourceCurrentMip, uint8 InDestCurrentMip, uint32 InSourceArrayIndex, uint32 InDestArrayIndex)
+											 const FIntPoint& InInitialPositionOffset, int32 InSubSectionSizeQuad, int32 InNumSubSections, uint8 InSourceCurrentMip, uint8 InDestCurrentMip, uint32 InSourceArrayIndex, uint32 InDestArrayIndex)
 		: SourceResource(InSourceResource)
 		, DestResource(InDestResource)
 		, DestCPUResource(InDestCPUResource)
@@ -994,7 +994,7 @@ public:
 		, DestMip(InDestCurrentMip)
 		, SourceArrayIndex(InSourceArrayIndex)
 		, DestArrayIndex(InDestArrayIndex)
-		, ComponentSectionBase(InFirstComponentSectionBase)
+		, InitialPositionOffset(InInitialPositionOffset)
 		, SubSectionSizeQuad(InSubSectionSizeQuad)
 		, NumSubSections(InNumSubSections)
 		, SourceDebugName(InSourceResourceDebugName)
@@ -1011,7 +1011,7 @@ public:
 		FIntPoint DestSize(DestResource->GetSizeX() >> DestMip, DestResource->GetSizeY() >> DestMip);
 
 		int32 LocalComponentSizeQuad = SubSectionSizeQuad * NumSubSections;
-		FVector2D PositionOffset(FMath::RoundToInt(ComponentSectionBase.X / LocalComponentSizeQuad), FMath::RoundToInt(ComponentSectionBase.Y / LocalComponentSizeQuad));
+		FVector2D PositionOffset(FMath::RoundToInt(InitialPositionOffset.X / LocalComponentSizeQuad), FMath::RoundToInt(InitialPositionOffset.Y / LocalComponentSizeQuad));
 
 		FRHICopyTextureInfo Params;
 		Params.NumSlices = 1;
@@ -1063,7 +1063,7 @@ private:
 	uint8 DestMip;
 	uint32 SourceArrayIndex;
 	uint32 DestArrayIndex;
-	FIntPoint ComponentSectionBase;
+	FIntPoint InitialPositionOffset;
 	int32 SubSectionSizeQuad;
 	int32 NumSubSections;
 	FString SourceDebugName;
@@ -1413,20 +1413,24 @@ FIntPoint ALandscape::ComputeComponentCounts() const
 		return FIntPoint(INDEX_NONE, INDEX_NONE);
 	}
 
-	FIntPoint NumComponents;
+	FIntPoint NumComponents(0, 0);
 	FIntPoint MaxSectionBase(0, 0);
+	FIntPoint MinSectionBase(0, 0);
 
-	Info->ForAllLandscapeProxies([&MaxSectionBase](ALandscapeProxy* Proxy)
+	Info->ForAllLandscapeProxies([&MaxSectionBase, &MinSectionBase](ALandscapeProxy* Proxy)
 	{
 		for (ULandscapeComponent* Component : Proxy->LandscapeComponents)
 		{
 			MaxSectionBase.X = FMath::Max(MaxSectionBase.X, Component->SectionBaseX);
 			MaxSectionBase.Y = FMath::Max(MaxSectionBase.Y, Component->SectionBaseY);
+
+			MinSectionBase.X = FMath::Min(MinSectionBase.X, Component->SectionBaseX);
+			MinSectionBase.Y = FMath::Min(MinSectionBase.Y, Component->SectionBaseY);
 		}
 	});
 
-	NumComponents.X = (MaxSectionBase.X / ComponentSizeQuads) + 1;
-	NumComponents.Y = (MaxSectionBase.Y / ComponentSizeQuads) + 1;
+	NumComponents.X = ((MaxSectionBase.X - MinSectionBase.X) / ComponentSizeQuads) + 1;
+	NumComponents.Y = ((MaxSectionBase.Y - MinSectionBase.Y) / ComponentSizeQuads) + 1;
 
 	return NumComponents;
 }
@@ -1635,13 +1639,13 @@ void ALandscape::CopyLayersTexture(UTexture* InSourceTexture, UTexture* InDestTe
 	}
 }
 
-void ALandscape::CopyLayersTexture(const FString& InSourceDebugName, FTextureResource* InSourceResource, const FString& InDestDebugName, FTextureResource* InDestResource, FTextureResource* InDestCPUResource, const FIntPoint& InFirstComponentSectionBase,
-	uint8 InSourceCurrentMip, uint8 InDestCurrentMip, uint32 InSourceArrayIndex, uint32 InDestArrayIndex) const
+void ALandscape::CopyLayersTexture(const FString& InSourceDebugName, FTextureResource* InSourceResource, const FString& InDestDebugName, FTextureResource* InDestResource, FTextureResource* InDestCPUResource, const FIntPoint& InInitialPositionOffset,
+									uint8 InSourceCurrentMip, uint8 InDestCurrentMip, uint32 InSourceArrayIndex, uint32 InDestArrayIndex) const
 {
 	check(InSourceResource != nullptr);
 	check(InDestResource != nullptr);
 
-	FLandscapeLayersCopyTexture_RenderThread CopyTexture(InSourceDebugName, InSourceResource, InDestDebugName, InDestResource, InDestCPUResource, InFirstComponentSectionBase, SubsectionSizeQuads, NumSubsections, InSourceCurrentMip, InDestCurrentMip, InSourceArrayIndex, InDestArrayIndex);
+	FLandscapeLayersCopyTexture_RenderThread CopyTexture(InSourceDebugName, InSourceResource, InDestDebugName, InDestResource, InDestCPUResource, InInitialPositionOffset, SubsectionSizeQuads, NumSubsections, InSourceCurrentMip, InDestCurrentMip, InSourceArrayIndex, InDestArrayIndex);
 
 	ENQUEUE_RENDER_COMMAND(FLandscapeLayersCopyCommand)(
 		[CopyTexture](FRHICommandListImmediate& RHICmdList) mutable
@@ -1651,7 +1655,7 @@ void ALandscape::CopyLayersTexture(const FString& InSourceDebugName, FTextureRes
 }
 
 void ALandscape::DrawWeightmapComponentsToRenderTarget(const FString& InDebugName, const TArray<FIntPoint>& InSectionBaseList, const FVector2D& InScaleBias, UTexture* InWeightmapRTRead, UTextureRenderTarget2D* InOptionalWeightmapRTRead2, UTextureRenderTarget2D* InWeightmapRTWrite,
-	bool InClearRTWrite, FLandscapeLayersWeightmapShaderParameters& InShaderParams, uint8 InMipRender) const
+														bool InClearRTWrite, FLandscapeLayersWeightmapShaderParameters& InShaderParams, uint8 InMipRender) const
 {
 	check(InWeightmapRTRead != nullptr);
 	check(InWeightmapRTWrite != nullptr);
@@ -1709,7 +1713,7 @@ void ALandscape::DrawWeightmapComponentsToRenderTarget(const FString& InDebugNam
 	PrintLayersDebugRT(InDebugName, InWeightmapRTWrite, InMipRender, false);
 }
 
-void ALandscape::DrawWeightmapComponentsToRenderTarget(const FString& InDebugName, const TArray<ULandscapeComponent*>& InComponentsToDraw, UTexture* InWeightmapRTRead, UTextureRenderTarget2D* InOptionalWeightmapRTRead2, UTextureRenderTarget2D* InWeightmapRTWrite,
+void ALandscape::DrawWeightmapComponentsToRenderTarget(const FString& InDebugName, const TArray<ULandscapeComponent*>& InComponentsToDraw, const FIntPoint& InLandscapeBase, UTexture* InWeightmapRTRead, UTextureRenderTarget2D* InOptionalWeightmapRTRead2, UTextureRenderTarget2D* InWeightmapRTWrite,
 														bool InClearRTWrite, FLandscapeLayersWeightmapShaderParameters& InShaderParams, uint8 InMipRender) const
 {
 	check(InWeightmapRTRead != nullptr);
@@ -1735,7 +1739,9 @@ void ALandscape::DrawWeightmapComponentsToRenderTarget(const FString& InDebugNam
 		{
 			// TODO: check what to do with WeightmapSubsectionOffset
 			FVector2D WeightmapScaleBias(Component->WeightmapScaleBias.Z, Component->WeightmapScaleBias.W);
-			GenerateLayersRenderQuadsAtlas(Component->GetSectionBase(), WeightmapScaleBias, SubsectionSizeQuads, WeightmapReadTextureSize, WeightmapWriteTextureSize, TriangleList);
+			FIntPoint ComponentSectionBase = Component->GetSectionBase() - InLandscapeBase;
+
+			GenerateLayersRenderQuadsAtlas(ComponentSectionBase, WeightmapScaleBias, SubsectionSizeQuads, WeightmapReadTextureSize, WeightmapWriteTextureSize, TriangleList);
 		}
 	}
 	else
@@ -1744,7 +1750,9 @@ void ALandscape::DrawWeightmapComponentsToRenderTarget(const FString& InDebugNam
 		{
 			// TODO: check what to do with WeightmapSubsectionOffset
 			FVector2D WeightmapScaleBias(Component->WeightmapScaleBias.Z, Component->WeightmapScaleBias.W);
-			GenerateLayersRenderQuadsMip(Component->GetSectionBase(), WeightmapScaleBias, SubsectionSizeQuads, WeightmapReadTextureSize, WeightmapWriteTextureSize, InMipRender, TriangleList);
+			FIntPoint ComponentSectionBase = Component->GetSectionBase() - InLandscapeBase;
+
+			GenerateLayersRenderQuadsMip(ComponentSectionBase, WeightmapScaleBias, SubsectionSizeQuads, WeightmapReadTextureSize, WeightmapWriteTextureSize, InMipRender, TriangleList);
 		}
 	}
 
@@ -1818,7 +1826,7 @@ void ALandscape::ClearLayersWeightmapTextureResource(const FString& InDebugName,
 	});
 }
 
-void ALandscape::DrawHeightmapComponentsToRenderTargetMips(const TArray<ULandscapeComponent*>& InComponentsToDraw, UTexture* InReadHeightmap, bool InClearRTWrite, struct FLandscapeLayersHeightmapShaderParameters& InShaderParams) const
+void ALandscape::DrawHeightmapComponentsToRenderTargetMips(const TArray<ULandscapeComponent*>& InComponentsToDraw, const FIntPoint& InLandscapeBase, UTexture* InReadHeightmap, bool InClearRTWrite, struct FLandscapeLayersHeightmapShaderParameters& InShaderParams) const
 {
 	bool OutputDebugName = CVarOutputLayersDebugDrawCallName.GetValueOnAnyThread() == 1 || CVarOutputLayersRTContent.GetValueOnAnyThread() == 1 ? true : false;
 	int32 CurrentMip = 1;
@@ -1831,14 +1839,14 @@ void ALandscape::DrawHeightmapComponentsToRenderTargetMips(const TArray<ULandsca
 		if (WriteMipRT != nullptr)
 		{
 			DrawHeightmapComponentsToRenderTarget(OutputDebugName ? FString::Printf(TEXT("LS Height: %s = -> %s CombinedAtlasWithMips %d"), *ReadMipRT->GetName(), *WriteMipRT->GetName(), CurrentMip) : TEXT(""),
-												  InComponentsToDraw, ReadMipRT, nullptr, WriteMipRT, ERTDrawingType::RTMips, InClearRTWrite, InShaderParams, CurrentMip++);
+												  InComponentsToDraw, InLandscapeBase, ReadMipRT, nullptr, WriteMipRT, ERTDrawingType::RTMips, InClearRTWrite, InShaderParams, CurrentMip++);
 		}
 
 		ReadMipRT = HeightmapRTList[MipRTIndex];
 	}
 }
 
-void ALandscape::DrawHeightmapComponentsToRenderTarget(const FString& InDebugName, const TArray<ULandscapeComponent*>& InComponentsToDraw, UTexture* InHeightmapRTRead, UTextureRenderTarget2D* InOptionalHeightmapRTRead2, UTextureRenderTarget2D* InHeightmapRTWrite,
+void ALandscape::DrawHeightmapComponentsToRenderTarget(const FString& InDebugName, const TArray<ULandscapeComponent*>& InComponentsToDraw, const FIntPoint& InLandscapeBase, UTexture* InHeightmapRTRead, UTextureRenderTarget2D* InOptionalHeightmapRTRead2, UTextureRenderTarget2D* InHeightmapRTWrite,
 													  ERTDrawingType InDrawType, bool InClearRTWrite, FLandscapeLayersHeightmapShaderParameters& InShaderParams, uint8 InMipRender) const
 {
 	check(InHeightmapRTRead != nullptr);
@@ -1858,57 +1866,43 @@ void ALandscape::DrawHeightmapComponentsToRenderTarget(const FString& InDebugNam
 	TArray<FLandscapeLayersTriangle> TriangleList;
 	TriangleList.Reserve(InComponentsToDraw.Num() * 2 * NumSubsections);
 
-	switch (InDrawType)
+	for (ULandscapeComponent* Component : InComponentsToDraw)
 	{
-		case ERTDrawingType::RTAtlas:
+		FVector2D HeightmapScaleBias(Component->HeightmapScaleBias.Z, Component->HeightmapScaleBias.W);
+		FIntPoint ComponentSectionBase = Component->GetSectionBase() - InLandscapeBase;
+
+		switch (InDrawType)
 		{
-			for (ULandscapeComponent* Component : InComponentsToDraw)
+			case ERTDrawingType::RTAtlas:
 			{
-				FVector2D HeightmapScaleBias(Component->HeightmapScaleBias.Z, Component->HeightmapScaleBias.W);
-				GenerateLayersRenderQuadsAtlas(Component->GetSectionBase(), HeightmapScaleBias, SubsectionSizeQuads, HeightmapReadTextureSize, HeightmapWriteTextureSize, TriangleList);
-			}
-		} break;
+				GenerateLayersRenderQuadsAtlas(ComponentSectionBase, HeightmapScaleBias, SubsectionSizeQuads, HeightmapReadTextureSize, HeightmapWriteTextureSize, TriangleList);
+			} break;
 
-		case ERTDrawingType::RTAtlasToNonAtlas:
-		{
-			for (ULandscapeComponent* Component : InComponentsToDraw)
+			case ERTDrawingType::RTAtlasToNonAtlas:
 			{
-				FVector2D HeightmapScaleBias(Component->HeightmapScaleBias.Z, Component->HeightmapScaleBias.W);
-				GenerateLayersRenderQuadsAtlasToNonAtlas(Component->GetSectionBase(), HeightmapScaleBias, SubsectionSizeQuads, HeightmapReadTextureSize, HeightmapWriteTextureSize, TriangleList);
-			}
-		} break;
+				GenerateLayersRenderQuadsAtlasToNonAtlas(ComponentSectionBase, HeightmapScaleBias, SubsectionSizeQuads, HeightmapReadTextureSize, HeightmapWriteTextureSize, TriangleList);
+			} break;
 
-		case ERTDrawingType::RTNonAtlas:
-		{
-			for (ULandscapeComponent* Component : InComponentsToDraw)
+			case ERTDrawingType::RTNonAtlas:
 			{
-				FVector2D HeightmapScaleBias(Component->HeightmapScaleBias.Z, Component->HeightmapScaleBias.W);
-				GenerateLayersRenderQuadsNonAtlas(Component->GetSectionBase(), HeightmapScaleBias, SubsectionSizeQuads, HeightmapReadTextureSize, HeightmapWriteTextureSize, TriangleList);
-			}
-		} break;
+				GenerateLayersRenderQuadsNonAtlas(ComponentSectionBase, HeightmapScaleBias, SubsectionSizeQuads, HeightmapReadTextureSize, HeightmapWriteTextureSize, TriangleList);
+			} break;
 
-		case ERTDrawingType::RTNonAtlasToAtlas:
-		{
-			for (ULandscapeComponent* Component : InComponentsToDraw)
+			case ERTDrawingType::RTNonAtlasToAtlas:
 			{
-				FVector2D HeightmapScaleBias(Component->HeightmapScaleBias.Z, Component->HeightmapScaleBias.W);
-				GenerateLayersRenderQuadsNonAtlasToAtlas(Component->GetSectionBase(), HeightmapScaleBias, SubsectionSizeQuads, HeightmapReadTextureSize, HeightmapWriteTextureSize, TriangleList);
-			}
-		} break;
+				GenerateLayersRenderQuadsNonAtlasToAtlas(ComponentSectionBase, HeightmapScaleBias, SubsectionSizeQuads, HeightmapReadTextureSize, HeightmapWriteTextureSize, TriangleList);
+			} break;
 
-		case ERTDrawingType::RTMips:
-		{
-			for (ULandscapeComponent* Component : InComponentsToDraw)
+			case ERTDrawingType::RTMips:
 			{
-				FVector2D HeightmapScaleBias(Component->HeightmapScaleBias.Z, Component->HeightmapScaleBias.W);
-				GenerateLayersRenderQuadsMip(Component->GetSectionBase(), HeightmapScaleBias, SubsectionSizeQuads, HeightmapReadTextureSize, HeightmapWriteTextureSize, InMipRender, TriangleList);
-			}			
-		} break;
+				GenerateLayersRenderQuadsMip(ComponentSectionBase, HeightmapScaleBias, SubsectionSizeQuads, HeightmapReadTextureSize, HeightmapWriteTextureSize, InMipRender, TriangleList);
+			} break;
 
-		default:
-		{
-			check(false);
-			return;
+			default:
+			{
+				check(false);
+				return;
+			}
 		}
 	}
 
@@ -2492,6 +2486,14 @@ void ALandscape::RegenerateLayersHeightmaps()
 	{
 		check(HeightmapRTList.Num() > 0);
 
+		FIntPoint MinExtend;
+		FIntPoint MaxExtend;
+
+		if (!Info->GetLandscapeExtent(MinExtend.X, MinExtend.Y, MaxExtend.X, MaxExtend.Y))
+		{
+			return;
+		}
+
 		FLandscapeLayersHeightmapShaderParameters ShaderParams;
 
 		bool FirstLayer = true;
@@ -2510,6 +2512,7 @@ void ALandscape::RegenerateLayersHeightmaps()
 			ShaderParams.LayerVisible = Layer.bVisible;
 			ShaderParams.GenerateNormals = false;
 			ShaderParams.LayerBlendMode = Layer.BlendMode;
+
 			if (Layer.BlendMode == LSBM_AlphaBlend)
 			{
 				// For now, only Layer reserved for Landscape Splines will use the AlphaBlendMode
@@ -2525,24 +2528,33 @@ void ALandscape::RegenerateLayersHeightmaps()
 
 			Info->ForAllLandscapeProxies([&](ALandscapeProxy* Proxy)
 			{
+				TArray<UTexture2D*> ProcessedTextures;
+
 				for (ULandscapeComponent* Component : Proxy->LandscapeComponents)
 				{
 					FLandscapeLayerComponentData* ComponentLayerData = Component->GetLayerData(Layer.Guid);
 					FLandscapeLayersGlobalComponentData& GlobalLayersData = Component->GetGlobalLayersData();
-					UTexture2D* LayerHeightmap = ComponentLayerData->HeightmapData.Texture;
+					UTexture2D* LayerHeightmap = ComponentLayerData->HeightmapData.Texture;			
+
+					if (!ProcessedTextures.Contains(LayerHeightmap))
+					{
+						ProcessedTextures.Add(LayerHeightmap);
+
+						FIntPoint TextureTopLeftSectionBase = GlobalLayersData.TopLeftSectionBase - MinExtend;
 
-					CopyLayersTexture(LayerHeightmap, LandscapeScratchRT1, nullptr, GlobalLayersData.TopLeftSectionBase);
-					PrintLayersDebugRT(OutputDebugName ? FString::Printf(TEXT("LS Height: %s Component %s += -> CombinedAtlas %s"), *Layer.Name.ToString(), *LayerHeightmap->GetName(), *LandscapeScratchRT1->GetName()) : TEXT(""), LandscapeScratchRT1);
+						CopyLayersTexture(LayerHeightmap, LandscapeScratchRT1, nullptr, TextureTopLeftSectionBase);
+						PrintLayersDebugRT(OutputDebugName ? FString::Printf(TEXT("LS Height: %s Component %s += -> CombinedAtlas %s"), *Layer.Name.ToString(), *LayerHeightmap->GetName(), *LandscapeScratchRT1->GetName()) : TEXT(""), LandscapeScratchRT1);
+					}
 				}
 			});
 
 			// NOTE: From this point on, we always work in non atlas, we'll convert back at the end to atlas only
 			DrawHeightmapComponentsToRenderTarget(OutputDebugName ? FString::Printf(TEXT("LS Height: %s += -> NonAtlas %s"), *Layer.Name.ToString(), *LandscapeScratchRT1->GetName(), *LandscapeScratchRT2->GetName()) : TEXT(""),
-												  AllLandscapeComponents, LandscapeScratchRT1, nullptr, LandscapeScratchRT2, ERTDrawingType::RTAtlasToNonAtlas, true, ShaderParams);
+												  AllLandscapeComponents, MinExtend, LandscapeScratchRT1, nullptr, LandscapeScratchRT2, ERTDrawingType::RTAtlasToNonAtlas, true, ShaderParams);
 
 			// Combine Current layer with current result
 			DrawHeightmapComponentsToRenderTarget(OutputDebugName ? FString::Printf(TEXT("LS Height: %s += -> CombinedNonAtlas %s"), *Layer.Name.ToString(), *LandscapeScratchRT2->GetName(), *CombinedHeightmapNonAtlasRT->GetName()) : TEXT(""),
-												AllLandscapeComponents, LandscapeScratchRT2, FirstLayer ? nullptr : LandscapeScratchRT3, CombinedHeightmapNonAtlasRT, ERTDrawingType::RTNonAtlas, FirstLayer, ShaderParams);
+												AllLandscapeComponents, MinExtend, LandscapeScratchRT2, FirstLayer ? nullptr : LandscapeScratchRT3, CombinedHeightmapNonAtlasRT, ERTDrawingType::RTNonAtlas, FirstLayer, ShaderParams);
 
 			ShaderParams.ApplyLayerModifiers = false;
 
@@ -2600,14 +2612,14 @@ void ALandscape::RegenerateLayersHeightmaps()
 		ShaderParams.GridSize = GetRootComponent()->RelativeScale3D;
 
 		DrawHeightmapComponentsToRenderTarget(OutputDebugName ? FString::Printf(TEXT("LS Height: %s = -> CombinedNonAtlasNormals : %s"), *CombinedHeightmapNonAtlasRT->GetName(), *LandscapeScratchRT1->GetName()) : TEXT(""),
-											  AllLandscapeComponents, CombinedHeightmapNonAtlasRT, nullptr, LandscapeScratchRT1, ERTDrawingType::RTNonAtlas, true, ShaderParams);
+											  AllLandscapeComponents, MinExtend, CombinedHeightmapNonAtlasRT, nullptr, LandscapeScratchRT1, ERTDrawingType::RTNonAtlas, true, ShaderParams);
 
 		ShaderParams.GenerateNormals = false;
 
 		DrawHeightmapComponentsToRenderTarget(OutputDebugName ? FString::Printf(TEXT("LS Height: %s = -> CombinedAtlasFinal : %s"), *LandscapeScratchRT1->GetName(), *CombinedHeightmapAtlasRT->GetName()) : TEXT(""),
-											  AllLandscapeComponents, LandscapeScratchRT1, nullptr, CombinedHeightmapAtlasRT, ERTDrawingType::RTNonAtlasToAtlas, true, ShaderParams);
+											  AllLandscapeComponents, MinExtend, LandscapeScratchRT1, nullptr, CombinedHeightmapAtlasRT, ERTDrawingType::RTNonAtlasToAtlas, true, ShaderParams);
 
-		DrawHeightmapComponentsToRenderTargetMips(AllLandscapeComponents, CombinedHeightmapAtlasRT, true, ShaderParams);
+		DrawHeightmapComponentsToRenderTargetMips(AllLandscapeComponents, MinExtend, CombinedHeightmapAtlasRT, true, ShaderParams);
 
 		// Copy back all Mips to original heightmap data
 		Info->ForAllLandscapeProxies([&](ALandscapeProxy* Proxy)
@@ -2627,14 +2639,17 @@ void ALandscape::RegenerateLayersHeightmaps()
 
 					ProcessedTextures.Add(ComponentHeightmap);
 					FLandscapeLayersGlobalComponentData& GlobalLayersData = Component->GetGlobalLayersData();
-					CopyLayersTexture(CombinedHeightmapAtlasRT, ComponentHeightmap, *CPUReadback, GlobalLayersData.TopLeftSectionBase, CurrentMip, CurrentMip);
+
+					FIntPoint TextureTopLeftSectionBase = GlobalLayersData.TopLeftSectionBase - MinExtend;
+
+					CopyLayersTexture(CombinedHeightmapAtlasRT, ComponentHeightmap, *CPUReadback, TextureTopLeftSectionBase, CurrentMip, CurrentMip);
 					++CurrentMip;
 
 					for (int32 MipRTIndex = EHeightmapRTType::HeightmapRT_Mip1; MipRTIndex < EHeightmapRTType::HeightmapRT_Count; ++MipRTIndex)
 					{
 						if (HeightmapRTList[MipRTIndex] != nullptr)
 						{
-							CopyLayersTexture(HeightmapRTList[MipRTIndex], ComponentHeightmap, *CPUReadback, GlobalLayersData.TopLeftSectionBase, CurrentMip, CurrentMip);
+							CopyLayersTexture(HeightmapRTList[MipRTIndex], ComponentHeightmap, *CPUReadback, TextureTopLeftSectionBase, CurrentMip, CurrentMip);
 							++CurrentMip;
 						}
 					}
@@ -2772,7 +2787,7 @@ void ALandscape::ResolveLayersTexture(FLandscapeLayersTexture2DCPUReadBackResour
 	}
 }
 
-void ALandscape::PrepareComponentDataToExtractMaterialLayersCS(const FLandscapeLayer& InLayer, int32 InCurrentWeightmapToProcessIndex, bool InOutputDebugName, FLandscapeTexture2DResource* InOutTextureData,
+void ALandscape::PrepareComponentDataToExtractMaterialLayersCS(const FLandscapeLayer& InLayer, int32 InCurrentWeightmapToProcessIndex, const FIntPoint& InLandscapeBase, bool InOutputDebugName, FLandscapeTexture2DResource* InOutTextureData,
 																TArray<FLandscapeLayerWeightmapExtractMaterialLayersComponentData>& OutComponentData, TMap<ULandscapeLayerInfoObject*, int32>& OutLayerInfoObjects)
 {
 	ULandscapeInfo* Info = GetLandscapeInfo();
@@ -2798,7 +2813,9 @@ void ALandscape::PrepareComponentDataToExtractMaterialLayersCS(const FLandscapeL
 					const ULandscapeWeightmapUsage* LayerWeightmapUsage = ComponentLayerData->WeightmapData.TextureUsages[InCurrentWeightmapToProcessIndex];
 					check(LayerWeightmapUsage != nullptr);
 
-					CopyLayersTexture(*LayerWeightmap->GetName(), LayerWeightmap->Resource, InOutputDebugName ? FString::Printf(TEXT("%s WeightmapScratchTexture"), *InLayer.Name.ToString()) : TEXT(""), InOutTextureData, nullptr, Component->GetSectionBase(), 0);
+					FIntPoint ComponentSectionBase = Component->GetSectionBase() - InLandscapeBase;
+
+					CopyLayersTexture(*LayerWeightmap->GetName(), LayerWeightmap->Resource, InOutputDebugName ? FString::Printf(TEXT("%s WeightmapScratchTexture"), *InLayer.Name.ToString()) : TEXT(""), InOutTextureData, nullptr, ComponentSectionBase, 0);
 					PrintLayersDebugTextureResource(InOutputDebugName ? FString::Printf(TEXT("LS Weight: %s WeightmapScratchTexture %s"), *InLayer.Name.ToString(), TEXT("WeightmapScratchTextureResource")) : TEXT(""), InOutTextureData, 0, false);
 
 					for (const FWeightmapLayerAllocationInfo& WeightmapLayerAllocation : ComponentLayerData->WeightmapData.LayerAllocations)
@@ -2810,11 +2827,13 @@ void ALandscape::PrepareComponentDataToExtractMaterialLayersCS(const FLandscapeL
 							const ULandscapeComponent* DestComponent = LayerWeightmapUsage->ChannelUsage[WeightmapLayerAllocation.WeightmapTextureChannel];
 							check(DestComponent);
 
+							FIntPoint DestComponentSectionBase = DestComponent->GetSectionBase() - InLandscapeBase;
+
 							// Compute component top left vertex position from section base info
 							int32 LocalComponentSizeQuad = Component->SubsectionSizeQuads * NumSubsections;
 							int32 LocalComponentSizeVerts = (Component->SubsectionSizeQuads + 1) * NumSubsections;
-							FVector2D SourcePositionOffset(FMath::RoundToInt(Component->GetSectionBase().X / LocalComponentSizeQuad), FMath::RoundToInt(Component->GetSectionBase().Y / LocalComponentSizeQuad));
-							FVector2D DestPositionOffset(FMath::RoundToInt(DestComponent->GetSectionBase().X / LocalComponentSizeQuad), FMath::RoundToInt(DestComponent->GetSectionBase().Y / LocalComponentSizeQuad));
+							FVector2D SourcePositionOffset(FMath::RoundToInt(ComponentSectionBase.X / LocalComponentSizeQuad), FMath::RoundToInt(ComponentSectionBase.Y / LocalComponentSizeQuad));
+							FVector2D DestPositionOffset(FMath::RoundToInt(DestComponentSectionBase.X / LocalComponentSizeQuad), FMath::RoundToInt(DestComponentSectionBase.Y / LocalComponentSizeQuad));
 
 							Data.ComponentVertexPosition = FIntPoint(SourcePositionOffset.X * LocalComponentSizeVerts, SourcePositionOffset.Y * LocalComponentSizeVerts);
 							Data.AtlasTexturePositionOutput = FIntPoint(DestPositionOffset.X * LocalComponentSizeVerts, DestPositionOffset.Y * LocalComponentSizeVerts);
@@ -2852,7 +2871,7 @@ void ALandscape::PrepareComponentDataToExtractMaterialLayersCS(const FLandscapeL
 	});
 }
 
-void ALandscape::PrepareComponentDataToPackMaterialLayersCS(int32 InCurrentWeightmapToProcessIndex, bool InOutputDebugName, const TArray<ULandscapeComponent*>& InAllLandscapeComponents, TArray<UTexture2D*>& OutProcessedWeightmaps, 
+void ALandscape::PrepareComponentDataToPackMaterialLayersCS(int32 InCurrentWeightmapToProcessIndex, const FIntPoint& InLandscapeBase, bool InOutputDebugName, const TArray<ULandscapeComponent*>& InAllLandscapeComponents, TArray<UTexture2D*>& OutProcessedWeightmaps,
 															TArray<FLandscapeLayersTexture2DCPUReadBackResource*>& OutProcessedCPUReadBackTexture, TArray<FLandscapeLayerWeightmapPackMaterialLayersComponentData>& OutComponentData)
 {
 	ULandscapeInfo* Info = GetLandscapeInfo();
@@ -2918,10 +2937,12 @@ void ALandscape::PrepareComponentDataToPackMaterialLayersCS(int32 InCurrentWeigh
 				{
 					if (ChannelLayerAllocation.LayerInfo != nullptr && !AlreadyProcessedAllocation.Contains(&ChannelLayerAllocation) && ChannelComponentWeightmapTextures[ChannelLayerAllocation.WeightmapTextureIndex] == ComponentWeightmapTexture)
 					{
+						FIntPoint ComponentSectionBase = ChannelComponent->GetSectionBase() - InLandscapeBase;
+
 						// Compute component top left vertex position from section base info
 						int32 LocalComponentSizeQuad = ChannelComponent->SubsectionSizeQuads * NumSubsections;
 						int32 LocalComponentSizeVerts = (ChannelComponent->SubsectionSizeQuads + 1) * NumSubsections;
-						FVector2D PositionOffset(FMath::RoundToInt(ChannelComponent->GetSectionBase().X / LocalComponentSizeQuad), FMath::RoundToInt(ChannelComponent->GetSectionBase().Y / LocalComponentSizeQuad));
+						FVector2D PositionOffset(FMath::RoundToInt(ComponentSectionBase.X / LocalComponentSizeQuad), FMath::RoundToInt(ComponentSectionBase.Y / LocalComponentSizeQuad));
 
 						Data.ComponentVertexPositionX[WeightmapChannelIndex] = PositionOffset.X * LocalComponentSizeVerts;
 						Data.ComponentVertexPositionY[WeightmapChannelIndex] = PositionOffset.Y * LocalComponentSizeVerts;
@@ -3350,6 +3371,14 @@ void ALandscape::RegenerateLayersWeightmaps()
 
 	if ((LayersContentUpdateFlags & ELandscapeLayersContentUpdateFlag::Weightmap_Render) != 0)
 	{
+		FIntPoint MinExtend;
+		FIntPoint MaxExtend;
+
+		if (!Info->GetLandscapeExtent(MinExtend.X, MinExtend.Y, MaxExtend.X, MaxExtend.Y))
+		{
+			return;
+		}
+		
 		check(WeightmapRTList.Num() > 0);
 
 		UTextureRenderTarget2D* LandscapeScratchRT1 = WeightmapRTList[EWeightmapRTType::WeightmapRT_Scratch1];
@@ -3409,7 +3438,7 @@ void ALandscape::RegenerateLayersWeightmaps()
 
 				// Prepare compute shader data
 				TArray<FLandscapeLayerWeightmapExtractMaterialLayersComponentData> ComponentsData;
-				PrepareComponentDataToExtractMaterialLayersCS(Layer, CurrentWeightmapToProcessIndex, OutputDebugName, WeightmapScratchExtractLayerTextureResource, ComponentsData, LayerInfoObjects);
+				PrepareComponentDataToExtractMaterialLayersCS(Layer, CurrentWeightmapToProcessIndex, MinExtend, OutputDebugName, WeightmapScratchExtractLayerTextureResource, ComponentsData, LayerInfoObjects);
 
 				HasFoundWeightmapToProcess = ComponentsData.Num() > 0;
 
@@ -3469,7 +3498,7 @@ void ALandscape::RegenerateLayersWeightmaps()
 					PSShaderParams.LayerAlpha = LayerInfoObj == ALandscapeProxy::VisibilityLayer ? 1.0f : Layer.WeightmapAlpha; // visibility can't be affected by weight
 
 					DrawWeightmapComponentsToRenderTarget(OutputDebugName ? FString::Printf(TEXT("LS Weight: %s Paint: %s += -> %s"), *Layer.Name.ToString(), *LayerInfoObj->LayerName.ToString(), *LandscapeScratchRT1->GetName(), *LandscapeScratchRT2->GetName()) : TEXT(""),
-														  AllLandscapeComponents, LandscapeScratchRT1, nullptr, LandscapeScratchRT2, true, PSShaderParams, 0);
+														  AllLandscapeComponents, MinExtend, LandscapeScratchRT1, nullptr, LandscapeScratchRT2, true, PSShaderParams, 0);
 
 					PSShaderParams.ApplyLayerModifiers = false;
 
@@ -3493,7 +3522,7 @@ void ALandscape::RegenerateLayersWeightmaps()
 					}
 
 					DrawWeightmapComponentsToRenderTarget(OutputDebugName ? FString::Printf(TEXT("LS Weight: %s PaintLayer: %s, %s += -> Combined %s"), *Layer.Name.ToString(), *LayerInfoObj->LayerName.ToString(), *LandscapeScratchRT2->GetName(), *LandscapeScratchRT3->GetName()) : TEXT(""),
-														  AllLandscapeComponents, LandscapeScratchRT2, FirstLayer ? nullptr : LandscapeScratchRT1, LandscapeScratchRT3, true, PSShaderParams, 0);
+														  AllLandscapeComponents, MinExtend, LandscapeScratchRT2, FirstLayer ? nullptr : LandscapeScratchRT1, LandscapeScratchRT3, true, PSShaderParams, 0);
 
 					PSShaderParams.OutputAsSubstractive = false;
 
@@ -3586,7 +3615,7 @@ void ALandscape::RegenerateLayersWeightmaps()
 			while (HasFoundWeightmapToProcess)
 			{
 				TArray<FLandscapeLayerWeightmapPackMaterialLayersComponentData> PackLayersComponentsData;
-				PrepareComponentDataToPackMaterialLayersCS(CurrentWeightmapToProcessIndex, OutputDebugName, AllLandscapeComponents, ProcessedWeightmaps, ProcessedCPUReadbackTextures, PackLayersComponentsData);
+				PrepareComponentDataToPackMaterialLayersCS(CurrentWeightmapToProcessIndex, MinExtend, OutputDebugName, AllLandscapeComponents, ProcessedWeightmaps, ProcessedCPUReadbackTextures, PackLayersComponentsData);
 				HasFoundWeightmapToProcess = PackLayersComponentsData.Num() > 0;
 
 				// Perform the compute shader
@@ -4249,86 +4278,7 @@ void ALandscapeProxy::InitializeLayerWithEmptyContent(const FGuid& InLayerGuid)
 		const TArray<ULandscapeComponent*>* ComponentsUsingHeightmap = ComponentsPerHeightmaps.Find(ComponentHeightmap);
 		check(ComponentsUsingHeightmap != nullptr);
 
-		// Compute Global layers data
-		FLandscapeLayersGlobalComponentData& GlobalLayersData = Component->GetGlobalLayersData();
-		GlobalLayersData.TopLeftSectionBase = InitTopLeftSectionBase;
-
-		for (ULandscapeComponent* ComponentUsingHeightmap : *ComponentsUsingHeightmap)
-		{
-			GlobalLayersData.TopLeftSectionBase.X = FMath::Min(GlobalLayersData.TopLeftSectionBase.X, ComponentUsingHeightmap->GetSectionBase().X);
-			GlobalLayersData.TopLeftSectionBase.Y = FMath::Min(GlobalLayersData.TopLeftSectionBase.Y, ComponentUsingHeightmap->GetSectionBase().Y);
-		}
-
-		// Compute per layer data
-		FLandscapeLayerComponentData* LayerData = Component->GetLayerData(InLayerGuid);
-
-		if (LayerData == nullptr || !LayerData->IsInitialized())
-		{
-			FLandscapeLayerComponentData NewData;
-
-			// Setup Heightmap data
-			UTexture2D** LayerHeightmap = CreatedHeightmapTextures.Find(ComponentHeightmap);
-
-			if (LayerHeightmap == nullptr)
-			{
-				UTexture2D* NewLayerHeightmap = CreateLandscapeTexture(ComponentHeightmap->Source.GetSizeX(), ComponentHeightmap->Source.GetSizeY(), TEXTUREGROUP_Terrain_Heightmap, ComponentHeightmap->Source.GetFormat());
-				LayerHeightmap = &CreatedHeightmapTextures.Add(ComponentHeightmap, NewLayerHeightmap);
-
-				ULandscapeComponent::CreateEmptyTextureMips(NewLayerHeightmap, true);
-
-				// Init Mip0 to be at 32768 which is equal to "0"
-				FColor* Mip0Data = (FColor*)NewLayerHeightmap->Source.LockMip(0);
-
-				for (ULandscapeComponent* ComponentUsingHeightmap : *ComponentsUsingHeightmap)
-				{
-					int32 HeightmapComponentOffsetX = FMath::RoundToInt((float)NewLayerHeightmap->Source.GetSizeX() * ComponentUsingHeightmap->HeightmapScaleBias.Z);
-					int32 HeightmapComponentOffsetY = FMath::RoundToInt((float)NewLayerHeightmap->Source.GetSizeY() * ComponentUsingHeightmap->HeightmapScaleBias.W);
-
-					for (int32 SubsectionY = 0; SubsectionY < NumSubsections; SubsectionY++)
-					{
-						for (int32 SubsectionX = 0; SubsectionX < NumSubsections; SubsectionX++)
-						{
-							for (int32 SubY = 0; SubY <= SubsectionSizeQuads; SubY++)
-							{
-								for (int32 SubX = 0; SubX <= SubsectionSizeQuads; SubX++)
-								{
-									// X/Y of the vertex we're looking at in component's coordinates.
-									const int32 CompX = SubsectionSizeQuads * SubsectionX + SubX;
-									const int32 CompY = SubsectionSizeQuads * SubsectionY + SubY;
-
-									// X/Y of the vertex we're looking indexed into the texture data
-									const int32 TexX = (SubsectionSizeQuads + 1) * SubsectionX + SubX;
-									const int32 TexY = (SubsectionSizeQuads + 1) * SubsectionY + SubY;
-
-									const int32 HeightTexDataIdx = (HeightmapComponentOffsetX + TexX) + (HeightmapComponentOffsetY + TexY) * NewLayerHeightmap->Source.GetSizeX();
-
-									// copy height and normal data
-									const uint16 HeightValue = LandscapeDataAccess::GetTexHeight(0.f);
-
-									Mip0Data[HeightTexDataIdx].R = HeightValue >> 8;
-									Mip0Data[HeightTexDataIdx].G = HeightValue & 255;
-
-									// Normal with get calculated later
-									Mip0Data[HeightTexDataIdx].B = 0.0f;
-									Mip0Data[HeightTexDataIdx].A = 0.0f;
-								}
-							}
-						}
-					}
-				}
-
-				NewLayerHeightmap->Source.UnlockMip(0);
-
-				NewLayerHeightmap->BeginCachePlatformData();
-				NewLayerHeightmap->ClearAllCachedCookedPlatformData();
-			}
-
-			NewData.HeightmapData.Texture = *LayerHeightmap;
-
-			// Nothing to do for Weightmap by default
-
-			Component->AddLayerData(InLayerGuid, MoveTemp(NewData));
-		}
+		Component->AddDefaultLayerData(InLayerGuid, *ComponentsUsingHeightmap, CreatedHeightmapTextures);
 	}
 }
 #endif
diff --git a/Engine/Source/Runtime/Landscape/Public/LandscapeDataAccess.h b/Engine/Source/Runtime/Landscape/Public/LandscapeDataAccess.h
--- a/Engine/Source/Runtime/Landscape/Public/LandscapeDataAccess.h
+++ b/Engine/Source/Runtime/Landscape/Public/LandscapeDataAccess.h
@@ -213,7 +213,7 @@ struct FLandscapeComponentDataInterface
 	/* Return the raw heightmap data exactly same size for Heightmap texture which belong to only this component */
 	LANDSCAPE_API void GetHeightmapTextureData(TArray<FColor>& OutData, bool bOkToFail = false);
 
-	LANDSCAPE_API bool GetWeightmapTextureData(ULandscapeLayerInfoObject* LayerInfo, TArray<uint8>& OutData);
+	LANDSCAPE_API bool GetWeightmapTextureData(ULandscapeLayerInfoObject* LayerInfo, TArray<uint8>& OutData, bool InUseEditingWeightmap = false);
 
 	FColor* GetHeightData(int32 LocalX, int32 LocalY) const
 	{
