- Landscape Layer System: Fix Subsections

#rnx

diff --git a/Engine/Source/Runtime/Landscape/Private/LandscapeEditLayers.cpp b/Engine/Source/Runtime/Landscape/Private/LandscapeEditLayers.cpp
--- a/Engine/Source/Runtime/Landscape/Private/LandscapeEditLayers.cpp
+++ b/Engine/Source/Runtime/Landscape/Private/LandscapeEditLayers.cpp
@@ -1691,7 +1691,7 @@ void ALandscape::DrawWeightmapComponentsToRenderTarget(const FString& InDebugNam
 
 	InShaderParams.ReadWeightmap1 = InWeightmapRTRead;
 	InShaderParams.ReadWeightmap2 = InOptionalWeightmapRTRead2;
-	InShaderParams.CurrentMipComponentVertexCount = (((SubsectionSizeQuads + 1) * NumSubsections) >> InMipRender);
+	InShaderParams.CurrentMipComponentVertexCount = ((SubsectionSizeQuads + 1) >> InMipRender);
 
 	if (InMipRender > 0)
 	{
@@ -1758,7 +1758,7 @@ void ALandscape::DrawWeightmapComponentsToRenderTarget(const FString& InDebugNam
 
 	InShaderParams.ReadWeightmap1 = InWeightmapRTRead;
 	InShaderParams.ReadWeightmap2 = InOptionalWeightmapRTRead2;
-	InShaderParams.CurrentMipComponentVertexCount = (((SubsectionSizeQuads + 1) * NumSubsections) >> InMipRender);
+	InShaderParams.CurrentMipComponentVertexCount = ((SubsectionSizeQuads + 1) >> InMipRender);
 
 	if (InMipRender > 0)
 	{
@@ -1788,7 +1788,7 @@ void ALandscape::DrawWeightmapComponentToRenderTargetMips(const TArray<FVector2D
 
 	// Convert from Texture position to SectionBase
 	int32 LocalComponentSizeQuad = SubsectionSizeQuads * NumSubsections;
-	int32 LocalComponentSizeVerts = SubsectionSizeQuads + 1 * NumSubsections;
+	int32 LocalComponentSizeVerts = (SubsectionSizeQuads + 1) * NumSubsections;
 
 	TArray<FIntPoint> SectionBaseToDraw;
 	SectionBaseToDraw.Reserve(InTexturePositionsToDraw.Num());
@@ -1909,7 +1909,7 @@ void ALandscape::DrawHeightmapComponentsToRenderTarget(const FString& InDebugNam
 	InShaderParams.ReadHeightmap1 = InHeightmapRTRead;
 	InShaderParams.ReadHeightmap2 = InOptionalHeightmapRTRead2;
 	InShaderParams.HeightmapSize = HeightmapReadTextureSize;
-	InShaderParams.CurrentMipComponentVertexCount = (((SubsectionSizeQuads + 1) * NumSubsections) >> InMipRender);
+	InShaderParams.CurrentMipComponentVertexCount = ((SubsectionSizeQuads + 1) >> InMipRender);
 
 	if (InMipRender > 0)
 	{
@@ -2077,29 +2077,22 @@ void ALandscape::GenerateLayersRenderQuadsMip(const FIntPoint& InSectionBase, co
 void ALandscape::GenerateLayersRenderQuadsAtlasToNonAtlas(const FIntPoint& InSectionBase, const FVector2D& InScaleBias, float InSubSectionSizeQuad, const FIntPoint& InReadSize, const FIntPoint& InWriteSize, TArray<struct FLandscapeLayersTriangle>& OutTriangles) const
 {
 	int32 SubsectionSizeVerts = InSubSectionSizeQuad + 1;
-	int32 LocalComponentSizeQuad = InSubSectionSizeQuad * NumSubsections;
-	int32 LocalComponentSizeVerts = SubsectionSizeVerts * NumSubsections;
-
-	FVector2D PositionOffset(FMath::RoundToInt(InSectionBase.X / LocalComponentSizeQuad), FMath::RoundToInt(InSectionBase.Y / LocalComponentSizeQuad));
-	FIntPoint ComponentSectionBase(PositionOffset.X * LocalComponentSizeQuad, PositionOffset.Y * LocalComponentSizeQuad);
-	FIntPoint UVComponentSectionBase(PositionOffset.X * SubsectionSizeVerts, PositionOffset.Y * SubsectionSizeVerts);
 	FVector2D UVSize((float)SubsectionSizeVerts / (float)InReadSize.X, (float)SubsectionSizeVerts / (float)InReadSize.Y);
 
-	FIntPoint SubSectionSectionBase;
-
 	for (int8 SubY = 0; SubY < NumSubsections; ++SubY)
 	{
 		for (int8 SubX = 0; SubX < NumSubsections; ++SubX)
 		{
-			SubSectionSectionBase.X = ComponentSectionBase.X + InSubSectionSizeQuad * SubX;
-			SubSectionSectionBase.Y = ComponentSectionBase.Y + InSubSectionSizeQuad * SubY;
+			FIntPoint SubSectionSectionBase(InSectionBase.X + InSubSectionSizeQuad * SubX, InSectionBase.Y + InSubSectionSizeQuad * SubY);
+			FVector2D PositionOffset(FMath::RoundToInt(SubSectionSectionBase.X / InSubSectionSizeQuad), FMath::RoundToInt(SubSectionSectionBase.Y / InSubSectionSizeQuad));
+			FIntPoint UVComponentSectionBase(PositionOffset.X * SubsectionSizeVerts, PositionOffset.Y * SubsectionSizeVerts);
 
 			// Offset for this component's data in texture
 			FVector2D UVStart;
 
 			if (InReadSize.X >= InWriteSize.X)
 			{
-				UVStart.X = ((float)UVComponentSectionBase.X / (float)InReadSize.X) + UVSize.X * (float)SubX;
+				UVStart.X = ((float)UVComponentSectionBase.X / (float)InReadSize.X);
 			}
 			else
 			{
@@ -2108,7 +2101,7 @@ void ALandscape::GenerateLayersRenderQuadsAtlasToNonAtlas(const FIntPoint& InSec
 
 			if (InReadSize.Y >= InWriteSize.Y)
 			{
-				UVStart.Y = ((float)UVComponentSectionBase.Y / (float)InReadSize.Y) + UVSize.Y * (float)SubY;
+				UVStart.Y = ((float)UVComponentSectionBase.Y / (float)InReadSize.Y);
 			}
 			else
 			{
@@ -2126,23 +2119,19 @@ void ALandscape::GenerateLayersRenderQuadsNonAtlas(const FIntPoint& InSectionBas
 	check(InReadSize.X == InWriteSize.X && InReadSize.Y == InWriteSize.Y);
 
 	int32 SubsectionSizeVerts = InSubSectionSizeQuad + 1;
-	int32 LocalComponentSizeQuad = InSubSectionSizeQuad * NumSubsections;
-
-	FVector2D PositionOffset(FMath::RoundToInt(InSectionBase.X / LocalComponentSizeQuad), FMath::RoundToInt(InSectionBase.Y / LocalComponentSizeQuad));
-	FIntPoint ComponentSectionBase = InSectionBase;
-	FIntPoint UVComponentSectionBase(PositionOffset.X * LocalComponentSizeQuad, PositionOffset.Y * LocalComponentSizeQuad);
+	
 	FVector2D UVSize((float)SubsectionSizeVerts / (float)InReadSize.X, (float)SubsectionSizeVerts / (float)InReadSize.Y);
-	FIntPoint SubSectionSectionBase;
-
+	
 	for (int8 SubY = 0; SubY < NumSubsections; ++SubY)
 	{
 		for (int8 SubX = 0; SubX < NumSubsections; ++SubX)
 		{
-			SubSectionSectionBase.X = ComponentSectionBase.X + InSubSectionSizeQuad * SubX;
-			SubSectionSectionBase.Y = ComponentSectionBase.Y + InSubSectionSizeQuad * SubY;
+			FIntPoint SubSectionSectionBase(InSectionBase.X + SubsectionSizeQuads * SubX, InSectionBase.Y + SubsectionSizeQuads * SubY);
+			FVector2D PositionOffset(FMath::RoundToInt(SubSectionSectionBase.X / InSubSectionSizeQuad), FMath::RoundToInt(SubSectionSectionBase.Y / InSubSectionSizeQuad));
+			FIntPoint UVComponentSectionBase(PositionOffset.X * InSubSectionSizeQuad, PositionOffset.Y * InSubSectionSizeQuad);
 
 			// Offset for this component's data in texture
-			FVector2D UVStart(((float)UVComponentSectionBase.X / (float)InReadSize.X) + UVSize.X * (float)SubX, ((float)UVComponentSectionBase.Y / (float)InReadSize.Y) + UVSize.Y * (float)SubY);
+			FVector2D UVStart(((float)UVComponentSectionBase.X / (float)InReadSize.X), ((float)UVComponentSectionBase.Y / (float)InReadSize.Y));
 			GenerateLayersRenderQuad(SubSectionSectionBase, SubsectionSizeVerts, UVStart, UVSize, OutTriangles);
 		}
 	}
