PR #5761: Landscape Grass "Use Landscape Lightmap" fix. (Contributed by [REDACTED_PERSON])

diff --git a/Engine/Source/Runtime/Landscape/Private/LandscapeGrass.cpp b/Engine/Source/Runtime/Landscape/Private/LandscapeGrass.cpp
--- a/Engine/Source/Runtime/Landscape/Private/LandscapeGrass.cpp
+++ b/Engine/Source/Runtime/Landscape/Private/LandscapeGrass.cpp
@@ -2448,6 +2448,7 @@ void ALandscapeProxy::UpdateGrass(const TArray<FVector>& Cameras, bool bForceSyn
 												LOD.OverrideMapBuildData = MakeUnique<FMeshMapBuildData>();
 												LOD.OverrideMapBuildData->LightMap = GrassLightMap;
 												LOD.OverrideMapBuildData->ShadowMap = GrassShadowMap;
+												LOD.OverrideMapBuildData->ResourceCluster = MeshMapBuildData->ResourceCluster;
 											}
 										}
 
