- Support Rename/SaveAs/Duplicate Map in Landscape Layer System
- Duplicate for PIE now shares Landscape resources (Textures/MaterialInstances) instead of copying them
- Fixed crash when changing LandscapeStreamingProxy's Landscape Actor

diff --git a/Engine/Source/Runtime/Landscape/Classes/Landscape.h b/Engine/Source/Runtime/Landscape/Classes/Landscape.h
--- a/Engine/Source/Runtime/Landscape/Classes/Landscape.h
+++ b/Engine/Source/Runtime/Landscape/Classes/Landscape.h
@@ -337,7 +337,7 @@ private:
 							   uint8 InSourceCurrentMip = 0, uint8 InDestCurrentMip = 0, uint32 InSourceArrayIndex = 0, uint32 uInDestArrayIndex = 0) const;
 
 	void InitializeLayers();
-	void InitializeLayersWeightmapUsage();
+	void InitializeLandscapeLayersWeightmapUsage();
 
 	void PrintLayersDebugRT(const FString& InContext, UTextureRenderTarget2D* InDebugRT, uint8 InMipRender = 0, bool InOutputHeight = true, bool InOutputNormals = false) const;
 	void PrintLayersDebugTextureResource(const FString& InContext, FTextureResource* InTextureResource, uint8 InMipRender = 0, bool InOutputHeight = true, bool InOutputNormals = false) const;
diff --git a/Engine/Source/Runtime/Landscape/Classes/LandscapeProxy.h b/Engine/Source/Runtime/Landscape/Classes/LandscapeProxy.h
--- a/Engine/Source/Runtime/Landscape/Classes/LandscapeProxy.h
+++ b/Engine/Source/Runtime/Landscape/Classes/LandscapeProxy.h
@@ -791,6 +791,8 @@ public:
 	virtual void PostEditImport() override;
 	//~ End UObject Interface
 
+	LANDSCAPE_API void InitializeProxyLayersWeightmapUsage();
+
 	LANDSCAPE_API static TArray<FName> GetLayersFromMaterial(UMaterialInterface* Material);
 	LANDSCAPE_API TArray<FName> GetLayersFromMaterial() const;
 	LANDSCAPE_API static ULandscapeLayerInfoObject* CreateLayerInfo(const TCHAR* LayerName, ULevel* Level);
diff --git a/Engine/Source/Runtime/Landscape/Private/Landscape.cpp b/Engine/Source/Runtime/Landscape/Private/Landscape.cpp
--- a/Engine/Source/Runtime/Landscape/Private/Landscape.cpp
+++ b/Engine/Source/Runtime/Landscape/Private/Landscape.cpp
@@ -318,7 +318,70 @@ void ULandscapeComponent::Serialize(FArchive& Ar)
 		CheckGenerateLandscapePlatformData(true, Ar.CookingTarget());
 	}
 
-	if (Ar.IsCooking() && !HasAnyFlags(RF_ClassDefaultObject) && !Ar.CookingTarget()->SupportsFeature(ETargetPlatformFeatures::DeferredRendering))
+	// Avoid the archiver in the PIE duplicate writer case because we want to share landscape textures & materials
+	if (Ar.GetPortFlags() & PPF_DuplicateForPIE)
+	{
+		if (Ar.IsLoading())
+		{
+			Super::Serialize(Ar);
+		}
+
+		TArray<UObject**> TexturesAndMaterials;
+		TexturesAndMaterials.Add((UObject**)&HeightmapTexture);
+		TexturesAndMaterials.Add((UObject**)&XYOffsetmapTexture);
+		for (UTexture2D* WeightmapTexture : WeightmapTextures)
+		{
+			TexturesAndMaterials.Add((UObject**)&WeightmapTexture);
+		}
+		for (UTexture2D* MobileWeightmapTexture : MobileWeightmapTextures)
+		{
+			TexturesAndMaterials.Add((UObject**)&MobileWeightmapTexture);
+		}
+		for (auto& ItPair : LayersData)
+		{
+			FLandscapeLayerComponentData& LayerComponentData = ItPair.Value;
+			TexturesAndMaterials.Add((UObject**)&LayerComponentData.HeightmapData.Texture);
+			for (UTexture2D* WeightmapTexture : LayerComponentData.WeightmapData.Textures)
+			{
+				TexturesAndMaterials.Add((UObject**)&WeightmapTexture);
+			}
+		}
+		for (UMaterialInstance* MaterialInstance : MaterialInstances)
+		{
+			TexturesAndMaterials.Add((UObject**)&MaterialInstance);
+		}
+		for (UMaterialInterface* MobileMaterialInterface : MobileMaterialInterfaces)
+		{
+			TexturesAndMaterials.Add((UObject**)(&MobileMaterialInterface));
+		}
+		for (UMaterialInstance* MobileCombinationMaterialInstance : MobileCombinationMaterialInstances)
+		{
+			TexturesAndMaterials.Add((UObject**)&MobileCombinationMaterialInstance);
+		}
+
+		if (Ar.IsSaving())
+		{
+			TArray<UObject*> BackupTexturesAndMaterials;
+			BackupTexturesAndMaterials.AddZeroed(TexturesAndMaterials.Num());
+			for (int i = 0; i < TexturesAndMaterials.Num(); ++i)
+			{
+				Exchange(*TexturesAndMaterials[i], BackupTexturesAndMaterials[i]);
+			}
+
+			Super::Serialize(Ar);
+
+			for (int i = 0; i < TexturesAndMaterials.Num(); ++i)
+			{
+				Exchange(*TexturesAndMaterials[i], BackupTexturesAndMaterials[i]);
+			}
+		}
+		// Manually serialize pointers
+		for (UObject** Object : TexturesAndMaterials)
+		{
+			Ar.Serialize(Object, sizeof(UObject*));
+		}
+	}
+	else if (Ar.IsCooking() && !HasAnyFlags(RF_ClassDefaultObject) && !Ar.CookingTarget()->SupportsFeature(ETargetPlatformFeatures::DeferredRendering))
 	{
 		// These properties are only used for SM4+ so we back them up and clear them before serializing them.
 		UTexture2D* BackupHeightmapTexture = nullptr;
@@ -709,7 +772,7 @@ void ULandscapeComponent::PostLoad()
 #if WITH_EDITOR
 	auto ReparentObject = [this](UObject* Object)
 	{
-		if (Object && !Object->HasAllFlags(RF_Public | RF_Standalone) && (Object->GetOuter() != GetOuter()))
+		if (Object && !Object->HasAllFlags(RF_Public | RF_Standalone) && (Object->GetOuter() != GetOuter()) && (Object->GetOutermost() == GetOutermost()))
 		{
 			Object->Rename(nullptr, GetOuter(), REN_ForceNoResetLoaders);
 			return true;
@@ -717,11 +780,35 @@ void ULandscapeComponent::PostLoad()
 		return false;
 	};
 
+	ReparentObject(HeightmapTexture);
+	ReparentObject(XYOffsetmapTexture);
+
+	for (UTexture2D* WeightmapTexture : WeightmapTextures)
+	{
+		ReparentObject(WeightmapTexture);
+	}
+
+	for (UTexture2D* MobileWeightmapTexture : MobileWeightmapTextures)
+	{
+		ReparentObject(MobileWeightmapTexture);
+	}
+
+	for (const auto& ItPair : LayersData)
+	{
+		const FLandscapeLayerComponentData& LayerComponentData = ItPair.Value;
+		ReparentObject(LayerComponentData.HeightmapData.Texture);
+		for (UTexture2D* WeightmapTexture : LayerComponentData.WeightmapData.Textures)
+		{
+			ReparentObject(WeightmapTexture);
+		}
+	}
+
 	for (UMaterialInstance* MaterialInstance : MaterialInstances)
 	{
-		while (ReparentObject(MaterialInstance))
+		ULandscapeMaterialInstanceConstant* CurrentMIC = Cast<ULandscapeMaterialInstanceConstant>(MaterialInstance);
+		while (ReparentObject(CurrentMIC))
 		{
-			MaterialInstance = Cast<UMaterialInstance>(MaterialInstance->Parent);
+			CurrentMIC = Cast<ULandscapeMaterialInstanceConstant>(MaterialInstance->Parent);
 		}
 	}
 
@@ -740,19 +827,6 @@ void ULandscapeComponent::PostLoad()
 			MobileCombinationMaterialInstance = Cast<UMaterialInstance>(MobileCombinationMaterialInstance->Parent);
 		}
 	}
-
-	ReparentObject(HeightmapTexture);
-	ReparentObject(XYOffsetmapTexture);
-
-	for (UTexture2D* WeightmapTexture : WeightmapTextures)
-	{
-		ReparentObject(WeightmapTexture);
-	}
-
-	for (UTexture2D* MobileWeightmapTexture : MobileWeightmapTextures)
-	{
-		ReparentObject(MobileWeightmapTexture);
-	}
 #endif
 
 #if !UE_BUILD_SHIPPING
diff --git a/Engine/Source/Runtime/Landscape/Private/LandscapeEdit.cpp b/Engine/Source/Runtime/Landscape/Private/LandscapeEdit.cpp
--- a/Engine/Source/Runtime/Landscape/Private/LandscapeEdit.cpp
+++ b/Engine/Source/Runtime/Landscape/Private/LandscapeEdit.cpp
@@ -3974,6 +3974,8 @@ void ALandscapeStreamingProxy::PostEditChangeProperty(FPropertyChangedEvent& Pro
 			{
 				// TODO - only need to refresh the old and new landscape info
 				ULandscapeInfo::RecreateLandscapeInfo(GetWorld(), false);
+				FixupWeightmaps();
+				InitializeProxyLayersWeightmapUsage();
 			}
 		}
 		else
diff --git a/Engine/Source/Runtime/Landscape/Private/LandscapeEditLayers.cpp b/Engine/Source/Runtime/Landscape/Private/LandscapeEditLayers.cpp
--- a/Engine/Source/Runtime/Landscape/Private/LandscapeEditLayers.cpp
+++ b/Engine/Source/Runtime/Landscape/Private/LandscapeEditLayers.cpp
@@ -1583,7 +1583,7 @@ void ALandscape::CopyOldDataToDefaultLayer()
 	});
 }
 
-void ALandscape::InitializeLayersWeightmapUsage()
+void ALandscape::InitializeLandscapeLayersWeightmapUsage()
 {
 	ULandscapeInfo* Info = GetLandscapeInfo();
 
@@ -1594,15 +1594,24 @@ void ALandscape::InitializeLayersWeightmapUsage()
 
 	Info->ForAllLandscapeProxies([&](ALandscapeProxy* Proxy)
 	{
-		for (const FLandscapeLayer& Layer : LandscapeLayers)
+		Proxy->InitializeProxyLayersWeightmapUsage();
+	});
+}
+
+void ALandscapeProxy::InitializeProxyLayersWeightmapUsage()
+{
+	if (ALandscape* Landscape = GetLandscapeActor())
+	{
+		for (const FLandscapeLayer& Layer : Landscape->LandscapeLayers)
 		{
-			for (ULandscapeComponent* Component : Proxy->LandscapeComponents)
+			for (ULandscapeComponent* Component : LandscapeComponents)
 			{
 				// Compute per layer data
 				FLandscapeLayerComponentData* LayerData = Component->GetLayerData(Layer.Guid);
 
 				if (LayerData != nullptr && LayerData->IsInitialized())
 				{
+					LayerData->WeightmapData.TextureUsages.Reset();
 					LayerData->WeightmapData.TextureUsages.AddDefaulted(LayerData->WeightmapData.Textures.Num());
 
 					// regenerate the weightmap usage
@@ -1610,11 +1619,11 @@ void ALandscape::InitializeLayersWeightmapUsage()
 					{
 						FWeightmapLayerAllocationInfo& Allocation = LayerData->WeightmapData.LayerAllocations[LayerIdx];
 						UTexture2D* WeightmapTexture = LayerData->WeightmapData.Textures[Allocation.WeightmapTextureIndex];
-						ULandscapeWeightmapUsage** TempUsage = Proxy->WeightmapUsageMap.Find(WeightmapTexture);
+						ULandscapeWeightmapUsage** TempUsage = WeightmapUsageMap.Find(WeightmapTexture);
 
 						if (TempUsage == nullptr)
 						{
-							TempUsage = &Proxy->WeightmapUsageMap.Add(WeightmapTexture, Proxy->CreateWeightmapUsage());
+							TempUsage = &WeightmapUsageMap.Add(WeightmapTexture, CreateWeightmapUsage());
 							(*TempUsage)->LayerGuid = Layer.Guid;
 						}
 
@@ -1628,7 +1637,7 @@ void ALandscape::InitializeLayersWeightmapUsage()
 				}
 			}
 		}
-	});
+	}
 }
 
 void ALandscape::CopyLayersTexture(UTexture* InSourceTexture, UTexture* InDestTexture, FTextureResource* InDestCPUResource, const FIntPoint& InFirstComponentSectionBase, uint8 InSourceCurrentMip, uint8 InDestCurrentMip, uint32 InSourceArrayIndex, uint32 InDestArrayIndex) const
@@ -2794,7 +2803,7 @@ void ALandscape::PrepareComponentDataToExtractMaterialLayersCS(const FLandscapeL
 
 			if (ComponentLayerData != nullptr)
 			{
-				if (ComponentLayerData->WeightmapData.Textures.IsValidIndex(InCurrentWeightmapToProcessIndex))
+				if (ComponentLayerData->WeightmapData.Textures.IsValidIndex(InCurrentWeightmapToProcessIndex) && ComponentLayerData->WeightmapData.TextureUsages.IsValidIndex(InCurrentWeightmapToProcessIndex))
 				{
 					UTexture2D* LayerWeightmap = ComponentLayerData->WeightmapData.Textures[InCurrentWeightmapToProcessIndex];
 					check(LayerWeightmap != nullptr);
@@ -3943,7 +3952,7 @@ void ALandscape::InitializeLayers()
 	if (HasLayersContent)
 	{
 		CreateLayersRenderingResource(ComponentCounts);
-		InitializeLayersWeightmapUsage();
+		InitializeLandscapeLayersWeightmapUsage();
 	}
 	else
 	{
diff --git a/Engine/Source/Runtime/Landscape/Private/LandscapeModule.cpp b/Engine/Source/Runtime/Landscape/Private/LandscapeModule.cpp
--- a/Engine/Source/Runtime/Landscape/Private/LandscapeModule.cpp
+++ b/Engine/Source/Runtime/Landscape/Private/LandscapeModule.cpp
@@ -133,6 +133,7 @@ void WorldRenameEventFunction(UWorld* World, const TCHAR* InName, UObject* NewOu
 	UPackage* PersistentLevelPackage = World->PersistentLevel->GetOutermost();
 	for (auto* OldTexOrMat : LandscapeTexturesAndMaterials)
 	{
+		// Now that landscape textures and materials are properly parented, this should not be necessary anymore
 		if (OldTexOrMat && OldTexOrMat->GetOuter() == PersistentLevelPackage)
 		{
 			// The names for these objects are not important, just generate a new name to avoid collisions
@@ -173,32 +174,6 @@ void WorldDuplicateEventFunction(UWorld* World, bool bDuplicateForPIE, TMap<UObj
 	{
 		AddPerWorldLandscapeData(World);
 	}
-
-#if WITH_EDITOR
-	if (!bDuplicateForPIE)
-	{
-		UPackage* WorldPackage = World->GetOutermost();
-
-		// Also duplicate all textures and materials used by landscape components
-		TArray<UObject*> LandscapeTexturesAndMaterials;
-		GetLandscapeTexturesAndMaterials(World->PersistentLevel, LandscapeTexturesAndMaterials);
-		for (auto* OldTexOrMat : LandscapeTexturesAndMaterials)
-		{
-			if (OldTexOrMat && OldTexOrMat->GetOuter() != WorldPackage)
-			{
-				// The names for these objects are not important, just generate a new name to avoid collisions
-				UObject* NewTextureOrMaterial = StaticDuplicateObject(OldTexOrMat, WorldPackage);
-				ReplacementMap.Add(OldTexOrMat, NewTextureOrMaterial);
-
-				// Materials hold references to the textures being moved, so they will need references fixed up as well
-				if (OldTexOrMat->IsA(UMaterialInterface::StaticClass()))
-				{
-					ObjectsToFixReferences.Add(NewTextureOrMaterial);
-				}
-			}
-		}
-	}
-#endif // WITH_EDITOR
 }
 
 void FLandscapeModule::StartupModule()
