Fixed Landscape height sculpting brush not targetting right location in Layer System Mode
#rnx

diff --git a/Engine/Source/Runtime/Landscape/Private/LandscapeEditInterface.cpp b/Engine/Source/Runtime/Landscape/Private/LandscapeEditInterface.cpp
--- a/Engine/Source/Runtime/Landscape/Private/LandscapeEditInterface.cpp
+++ b/Engine/Source/Runtime/Landscape/Private/LandscapeEditInterface.cpp
@@ -422,35 +422,37 @@ void FLandscapeEditDataInterface::SetHeightData(int32 X1, int32 Y1, int32 X2, in
 				}
 			}
 
-			// Update mipmaps
-
-			// Work out how many mips should be calculated directly from one component's data.
-			// The remaining mips are calculated on a per texture basis.
-			// eg if subsection is 7x7 quads, we need one 3 mips total: (8x8, 4x4, 2x2 verts)
-			if (InGenerateMips)
+			if (!GetMutableDefault<UEditorExperimentalSettings>()->bLandscapeLayerSystem)
 			{
-				int32 BaseNumMips = FMath::CeilLogTwo(SubsectionSizeQuads + 1);
-				TArray<FColor*> MipData;
-				MipData.AddUninitialized(BaseNumMips);
-				MipData[0] = HeightmapTextureData;
-				for (int32 MipIdx = 1; MipIdx < BaseNumMips; MipIdx++)
-				{
-					MipData[MipIdx] = (FColor*)TexDataInfo->GetMipData(MipIdx);
-				}
-				Component->GenerateHeightmapMips(MipData, ComponentX1, ComponentY1, ComponentX2, ComponentY2, TexDataInfo);
+				// Update mipmaps
+
+				// Work out how many mips should be calculated directly from one component's data.
+				// The remaining mips are calculated on a per texture basis.
+				// eg if subsection is 7x7 quads, we need one 3 mips total: (8x8, 4x4, 2x2 verts)
+				if (InGenerateMips)
+				{
+					int32 BaseNumMips = FMath::CeilLogTwo(SubsectionSizeQuads + 1);
+					TArray<FColor*> MipData;
+					MipData.AddUninitialized(BaseNumMips);
+					MipData[0] = HeightmapTextureData;
+					for (int32 MipIdx = 1; MipIdx < BaseNumMips; MipIdx++)
+					{
+						MipData[MipIdx] = (FColor*)TexDataInfo->GetMipData(MipIdx);
+					}
+					Component->GenerateHeightmapMips(MipData, ComponentX1, ComponentY1, ComponentX2, ComponentY2, TexDataInfo);
 
-				if (InUpdateCollision)
-				{
-					// Update collision
-					Component->UpdateCollisionHeightData(
-						MipData[Component->CollisionMipLevel],
-						Component->SimpleCollisionMipLevel > Component->CollisionMipLevel ? MipData[Component->SimpleCollisionMipLevel] : nullptr,
-						ComponentX1, ComponentY1, ComponentX2, ComponentY2, bUpdateBoxSphereBounds,
-						XYOffsetMipData);
+					if (InUpdateCollision)
+					{
+						// Update collision
+						Component->UpdateCollisionHeightData(
+							MipData[Component->CollisionMipLevel],
+							Component->SimpleCollisionMipLevel > Component->CollisionMipLevel ? MipData[Component->SimpleCollisionMipLevel] : nullptr,
+							ComponentX1, ComponentY1, ComponentX2, ComponentY2, bUpdateBoxSphereBounds,
+							XYOffsetMipData);
+					}
 				}
 			}
 
-
 			// Update GUID for Platform Data
 			FPlatformMisc::CreateGuid(Component->StateId);
 		}
