- Landscape Layer System: World Composition Support

#rnx

diff --git a/Engine/Source/Editor/LandscapeEditor/Private/LandscapeEdMode.cpp b/Engine/Source/Editor/LandscapeEditor/Private/LandscapeEdMode.cpp
--- a/Engine/Source/Editor/LandscapeEditor/Private/LandscapeEdMode.cpp
+++ b/Engine/Source/Editor/LandscapeEditor/Private/LandscapeEdMode.cpp
@@ -253,7 +253,6 @@ void ALandscape::SplitHeightmap(ULandscapeComponent* Comp, ALandscapeProxy* Targ
 			}
 		});
 
-		Comp->GetGlobalLayersData().TopLeftSectionBase  = Comp->GetSectionBase();
 		Landscape->RequestLayersContentUpdate(ELandscapeLayersContentUpdateFlag::All, true);
 	}
 #endif
diff --git a/Engine/Source/Editor/LandscapeEditor/Private/LandscapeEdModeComponentTools.cpp b/Engine/Source/Editor/LandscapeEditor/Private/LandscapeEdModeComponentTools.cpp
--- a/Engine/Source/Editor/LandscapeEditor/Private/LandscapeEdModeComponentTools.cpp
+++ b/Engine/Source/Editor/LandscapeEditor/Private/LandscapeEdModeComponentTools.cpp
@@ -826,7 +826,7 @@ public:
 			{
 				check(Landscape != nullptr); // Landscape actor is required if layer system is enabled
 
-				Landscape->CreateLayersRenderingResource(Landscape->ComputeComponentCounts());
+				Landscape->RequestLayersInitialization();
 			}
 
 			for (ULandscapeComponent* NewComponent : NewComponents)
diff --git a/Engine/Source/Runtime/Landscape/Classes/Landscape.h b/Engine/Source/Runtime/Landscape/Classes/Landscape.h
--- a/Engine/Source/Runtime/Landscape/Classes/Landscape.h
+++ b/Engine/Source/Runtime/Landscape/Classes/Landscape.h
@@ -248,11 +248,13 @@ public:
 
 	// Layers stuff
 #if WITH_EDITOR
+	LANDSCAPE_API void RequestLayersInitialization() { bLandscapeLayersAreInitialized = false; }
 	LANDSCAPE_API void RequestLayersContentUpdate(ELandscapeLayersContentUpdateFlag InDataFlags, bool InUpdateAllMaterials = false);
 	LANDSCAPE_API bool ReorderLayer(int32 InStartingLayerIndex, int32 InDestinationLayerIndex);
 	LANDSCAPE_API void CreateLayer(FName InName = NAME_None);
-	LANDSCAPE_API void CreateDefaultLayer(const FIntPoint& InComponentCounts);
+	LANDSCAPE_API void CreateDefaultLayer();
 	LANDSCAPE_API void CopyOldDataToDefaultLayer();
+	LANDSCAPE_API void AddLayersToProxy(ALandscapeProxy* InProxy);
 	LANDSCAPE_API TMap<UTexture2D*, TArray<ULandscapeComponent*>> GenerateComponentsPerHeightmaps() const;
 	LANDSCAPE_API FIntPoint ComputeComponentCounts() const;
 	LANDSCAPE_API bool IsLayerNameUnique(const FName& InName) const;
@@ -287,10 +289,10 @@ public:
 	LANDSCAPE_API TArray<int8>& GetBrushesOrderForLayer(int32 InLayerIndex, int32 InTargetType);
 	LANDSCAPE_API class ALandscapeBlueprintCustomBrush* GetBrushForLayer(int32 InLayerIndex, int32 InTargetType, int8 BrushIndex) const;
 	LANDSCAPE_API TArray<class ALandscapeBlueprintCustomBrush*> GetBrushesForLayer(int32 InLayerIndex, int32 InTargetType) const;
-	LANDSCAPE_API void CreateLayersRenderingResource(const FIntPoint& InComponentCounts);
-
+	
 private:
 	void TickLayers(float DeltaTime, ELevelTick TickType, FActorTickFunction& ThisTickFunction);
+	void CreateLayersRenderingResource(const FIntPoint& InComponentCounts);
 	void RegenerateLayersContent();
 	void RegenerateLayersHeightmaps();
 	void RegenerateLayersWeightmaps();
@@ -367,11 +369,10 @@ public:
 	UPROPERTY(Transient)
 	bool PreviousExperimentalLandscapeLayers;
 
+private:
 	UPROPERTY(Transient)
 	bool bLandscapeLayersAreInitialized;
-
-private:
-
+	
 	UPROPERTY(Transient)
 	bool WasCompilingShaders;
 
@@ -409,7 +410,7 @@ public:
 
 private:
 	TWeakObjectPtr<ALandscape> Landscape;
-	const FGuid& LayerGUID;
+	FGuid PreviousLayerGUID;
 	TFunction<void()> CompletionCallback;
 };
 #endif
diff --git a/Engine/Source/Runtime/Landscape/Classes/LandscapeComponent.h b/Engine/Source/Runtime/Landscape/Classes/LandscapeComponent.h
--- a/Engine/Source/Runtime/Landscape/Classes/LandscapeComponent.h
+++ b/Engine/Source/Runtime/Landscape/Classes/LandscapeComponent.h
@@ -290,15 +290,6 @@ struct FLandscapeLayerComponentData
 	bool IsInitialized() const { return HeightmapData.Texture != nullptr || WeightmapData.Textures.Num() > 0;  }
 };
 
-USTRUCT(NotBlueprintable)
-struct FLandscapeLayersGlobalComponentData
-{
-	GENERATED_USTRUCT_BODY()
-
-	UPROPERTY()
-	FIntPoint TopLeftSectionBase;
-};
-
 UCLASS(hidecategories=(Display, Attachment, Physics, Debug, Collision, Movement, Rendering, PrimitiveComponent, Object, Transform, Mobility), showcategories=("Rendering|Material"), MinimalAPI, Within=LandscapeProxy)
 class ULandscapeComponent : public UPrimitiveComponent
 {
@@ -384,10 +375,7 @@ private:
 
 	UPROPERTY()
 	TMap<FGuid, FLandscapeLayerComponentData> LayersData;
-
-	UPROPERTY()
-	FLandscapeLayersGlobalComponentData GlobalLayersData;
-
+		
 	// Final layer data
 	UPROPERTY(Transient)
 	TArray<ULandscapeWeightmapUsage*> WeightmapTexturesUsage;
@@ -570,9 +558,6 @@ public:
 	LANDSCAPE_API TArray<ULandscapeWeightmapUsage*>& GetWeightmapTexturesUsage(bool InReturnEditingWeightmap = false);
 	LANDSCAPE_API const TArray<ULandscapeWeightmapUsage*>& GetWeightmapTexturesUsage(bool InReturnEditingWeightmap = false) const;
 
-	LANDSCAPE_API const FLandscapeLayersGlobalComponentData& GetGlobalLayersData() const;
-	LANDSCAPE_API FLandscapeLayersGlobalComponentData& GetGlobalLayersData();
-
 	LANDSCAPE_API const FLandscapeLayerComponentData* GetLayerData(const FGuid& InLayerGuid) const;
 	LANDSCAPE_API FLandscapeLayerComponentData* GetLayerData(const FGuid& InLayerGuid);
 	LANDSCAPE_API void AddLayerData(const FGuid& InLayerGuid, const FLandscapeLayerComponentData& InData);
diff --git a/Engine/Source/Runtime/Landscape/Private/Landscape.cpp b/Engine/Source/Runtime/Landscape/Private/Landscape.cpp
--- a/Engine/Source/Runtime/Landscape/Private/Landscape.cpp
+++ b/Engine/Source/Runtime/Landscape/Private/Landscape.cpp
@@ -1532,16 +1532,6 @@ FGuid ULandscapeComponent::GetEditingLayerGUID() const
 	return Landscape != nullptr ? Landscape->GetEditingLayer() : FGuid();
 }
 
-const FLandscapeLayersGlobalComponentData& ULandscapeComponent::GetGlobalLayersData() const
-{
-	return GlobalLayersData;
-}
-
-FLandscapeLayersGlobalComponentData& ULandscapeComponent::GetGlobalLayersData() 
-{
-	return GlobalLayersData;
-}
-
 const FLandscapeLayerComponentData* ULandscapeComponent::GetLayerData(const FGuid& InLayerGuid) const 
 {
 	return LayersData.Find(InLayerGuid);
@@ -1572,16 +1562,7 @@ void ULandscapeComponent::AddDefaultLayerData(const FGuid& InLayerGuid, const TA
 	Modify();
 
 	UTexture2D* ComponentHeightmap = GetHeightmap();
-
-	// Compute Global layers data
-	GlobalLayersData.TopLeftSectionBase = FIntPoint(INT_MAX, INT_MAX);
-
-	for (ULandscapeComponent* ComponentUsingHeightmap : InComponentsUsingHeightmap)
-	{
-		GlobalLayersData.TopLeftSectionBase.X = FMath::Min(GlobalLayersData.TopLeftSectionBase.X, ComponentUsingHeightmap->GetSectionBase().X);
-		GlobalLayersData.TopLeftSectionBase.Y = FMath::Min(GlobalLayersData.TopLeftSectionBase.Y, ComponentUsingHeightmap->GetSectionBase().Y);
-	}
-
+		
 	// Compute per layer data
 	FLandscapeLayerComponentData* LayerData = GetLayerData(InLayerGuid);
 
@@ -2763,6 +2744,12 @@ void ULandscapeInfo::RegisterActor(ALandscapeProxy* Proxy, bool bMapCheck)
 		StreamingProxy->FixupSharedData(LandscapeActor.Get());
 	}
 
+	if (LandscapeActor)
+	{
+		// Force update rendering resources
+		LandscapeActor->RequestLayersInitialization();
+	}
+
 	UpdateLayerInfoMap(Proxy);
 	UpdateAllAddCollisions();
 
diff --git a/Engine/Source/Runtime/Landscape/Private/LandscapeEdit.cpp b/Engine/Source/Runtime/Landscape/Private/LandscapeEdit.cpp
--- a/Engine/Source/Runtime/Landscape/Private/LandscapeEdit.cpp
+++ b/Engine/Source/Runtime/Landscape/Private/LandscapeEdit.cpp
@@ -2775,10 +2775,16 @@ LANDSCAPE_API void ALandscapeProxy::Import(
 	{
 		ALandscape* LandscapeActor = GetLandscapeActor();
 		check(LandscapeActor != nullptr);
-
-		// Create the default layer
-		LandscapeActor->CreateDefaultLayer(FIntPoint(NumComponentsX, NumComponentsY));
-		HasLayersContent = true;
+		// Only create Layers on main Landscape
+		if (LandscapeActor == this)
+		{
+			// Create the default layer
+			LandscapeActor->CreateDefaultLayer();
+		} 
+		else
+		{
+			LandscapeActor->AddLayersToProxy(this);
+		}
 	}
 
 	if (GetLevel()->bIsVisible)
diff --git a/Engine/Source/Runtime/Landscape/Private/LandscapeEditLayers.cpp b/Engine/Source/Runtime/Landscape/Private/LandscapeEditLayers.cpp
--- a/Engine/Source/Runtime/Landscape/Private/LandscapeEditLayers.cpp
+++ b/Engine/Source/Runtime/Landscape/Private/LandscapeEditLayers.cpp
@@ -2492,6 +2492,40 @@ void ALandscape::RegenerateLayersHeightmaps()
 			return;
 		}
 
+		// Use to compute TopLeft Component per Heightmap
+		struct FHeightmapTopLeft
+		{
+			FHeightmapTopLeft(UTexture2D* InTexture, FIntPoint InTopLeftSectionBase) : Texture(InTexture), TopLeftSectionBase(InTopLeftSectionBase) {}
+
+			FHeightmapTopLeft(FHeightmapTopLeft&&) = default;
+			FHeightmapTopLeft& operator=(FHeightmapTopLeft&&) = default;
+
+			UTexture2D* Texture;
+			FIntPoint TopLeftSectionBase;
+		};
+
+		// Calculate Top Left Lambda
+		auto GetProxyUniqueHeightmaps = [&](ALandscapeProxy* InProxy, TArray<FHeightmapTopLeft>& OutHeightmaps, const FGuid& InLayerGuid = FGuid())
+		{
+			FScopedSetLandscapeEditingLayer Scope(this, InLayerGuid);
+			for (ULandscapeComponent* Component : InProxy->LandscapeComponents)
+			{
+				UTexture2D* ComponentHeightmap = Component->GetHeightmap(true);
+		
+				int32 Index = OutHeightmaps.IndexOfByPredicate([=](const FHeightmapTopLeft& LayerHeightmap) { return LayerHeightmap.Texture == ComponentHeightmap; });
+
+				if (Index == INDEX_NONE)
+				{
+					OutHeightmaps.Add(FHeightmapTopLeft(ComponentHeightmap, Component->GetSectionBase()));
+				}
+				else
+				{
+					OutHeightmaps[Index].TopLeftSectionBase.X = FMath::Min(OutHeightmaps[Index].TopLeftSectionBase.X, Component->GetSectionBase().X);
+					OutHeightmaps[Index].TopLeftSectionBase.Y = FMath::Min(OutHeightmaps[Index].TopLeftSectionBase.Y, Component->GetSectionBase().Y);
+				}
+			}
+		};
+
 		FLandscapeLayersHeightmapShaderParameters ShaderParams;
 
 		bool FirstLayer = true;
@@ -2526,23 +2560,13 @@ void ALandscape::RegenerateLayersHeightmaps()
 
 			Info->ForAllLandscapeProxies([&](ALandscapeProxy* Proxy)
 			{
-				TArray<UTexture2D*> ProcessedTextures;
+                TArray<FHeightmapTopLeft> LayerHeightmaps;
+				GetProxyUniqueHeightmaps(Proxy, LayerHeightmaps, Layer.Guid);
 
-				for (ULandscapeComponent* Component : Proxy->LandscapeComponents)
+				for (const FHeightmapTopLeft& LayerHeightmap : LayerHeightmaps)
 				{
-					FLandscapeLayerComponentData* ComponentLayerData = Component->GetLayerData(Layer.Guid);
-					FLandscapeLayersGlobalComponentData& GlobalLayersData = Component->GetGlobalLayersData();
-					UTexture2D* LayerHeightmap = ComponentLayerData->HeightmapData.Texture;			
-
-					if (!ProcessedTextures.Contains(LayerHeightmap))
-					{
-						ProcessedTextures.Add(LayerHeightmap);
-
-						FIntPoint TextureTopLeftSectionBase = GlobalLayersData.TopLeftSectionBase - MinExtend;
-
-						CopyLayersTexture(LayerHeightmap, LandscapeScratchRT1, nullptr, TextureTopLeftSectionBase);
-						PrintLayersDebugRT(OutputDebugName ? FString::Printf(TEXT("LS Height: %s Component %s += -> CombinedAtlas %s"), *Layer.Name.ToString(), *LayerHeightmap->GetName(), *LandscapeScratchRT1->GetName()) : TEXT(""), LandscapeScratchRT1);
-					}
+					CopyLayersTexture(LayerHeightmap.Texture, LandscapeScratchRT1, nullptr, LayerHeightmap.TopLeftSectionBase-MinExtend);
+					PrintLayersDebugRT(OutputDebugName ? FString::Printf(TEXT("LS Height: %s Component %s += -> CombinedAtlas %s"), *Layer.Name.ToString(), *LayerHeightmap.Texture->GetName(), *LandscapeScratchRT1->GetName()) : TEXT(""), LandscapeScratchRT1);
 				}
 			});
 
@@ -2622,34 +2646,27 @@ void ALandscape::RegenerateLayersHeightmaps()
 		// Copy back all Mips to original heightmap data
 		Info->ForAllLandscapeProxies([&](ALandscapeProxy* Proxy)
 		{
-			TArray<UTexture2D*> ProcessedTextures;
+			TArray<FHeightmapTopLeft> Heightmaps;
+			GetProxyUniqueHeightmaps(Proxy, Heightmaps);
 
-			for (ULandscapeComponent* Component : Proxy->LandscapeComponents)
+			for (const FHeightmapTopLeft& Heightmap : Heightmaps)
 			{
-				UTexture2D* ComponentHeightmap = Component->GetHeightmap();
-
-				if (!ProcessedTextures.Contains(ComponentHeightmap))
-				{
-					int32 CurrentMip = 0;
+				int32 CurrentMip = 0;
 
-					FLandscapeLayersTexture2DCPUReadBackResource** CPUReadback = Proxy->HeightmapsCPUReadBack.Find(ComponentHeightmap);
-					check(CPUReadback != nullptr);
-
-					ProcessedTextures.Add(ComponentHeightmap);
-					FLandscapeLayersGlobalComponentData& GlobalLayersData = Component->GetGlobalLayersData();
+				FLandscapeLayersTexture2DCPUReadBackResource** CPUReadback = Proxy->HeightmapsCPUReadBack.Find(Heightmap.Texture);
+				check(CPUReadback != nullptr);
 
-					FIntPoint TextureTopLeftSectionBase = GlobalLayersData.TopLeftSectionBase - MinExtend;
+				FIntPoint TextureTopLeftSectionBase = Heightmap.TopLeftSectionBase - MinExtend;
 
-					CopyLayersTexture(CombinedHeightmapAtlasRT, ComponentHeightmap, *CPUReadback, TextureTopLeftSectionBase, CurrentMip, CurrentMip);
-					++CurrentMip;
+				CopyLayersTexture(CombinedHeightmapAtlasRT, Heightmap.Texture, *CPUReadback, TextureTopLeftSectionBase, CurrentMip, CurrentMip);
+				++CurrentMip;
 
-					for (int32 MipRTIndex = EHeightmapRTType::HeightmapRT_Mip1; MipRTIndex < EHeightmapRTType::HeightmapRT_Count; ++MipRTIndex)
+				for (int32 MipRTIndex = EHeightmapRTType::HeightmapRT_Mip1; MipRTIndex < EHeightmapRTType::HeightmapRT_Count; ++MipRTIndex)
+				{
+					if (HeightmapRTList[MipRTIndex] != nullptr)
 					{
-						if (HeightmapRTList[MipRTIndex] != nullptr)
-						{
-							CopyLayersTexture(HeightmapRTList[MipRTIndex], ComponentHeightmap, *CPUReadback, TextureTopLeftSectionBase, CurrentMip, CurrentMip);
-							++CurrentMip;
-						}
+						CopyLayersTexture(HeightmapRTList[MipRTIndex], Heightmap.Texture, *CPUReadback, TextureTopLeftSectionBase, CurrentMip, CurrentMip);
+						++CurrentMip;
 					}
 				}
 			}
@@ -3957,9 +3974,11 @@ void ALandscape::InitializeLayers()
 	else
 	{
 		// Special case to load the existing content to the default layer
-		CreateDefaultLayer(ComponentCounts);
+		CreateDefaultLayer();
+		CreateLayersRenderingResource(ComponentCounts);
 		CopyOldDataToDefaultLayer();
 	}
+	bLandscapeLayersAreInitialized = true;
 }
 
 void ALandscape::TickLayers(float DeltaTime, ELevelTick TickType, FActorTickFunction& ThisTickFunction)
@@ -3982,7 +4001,6 @@ void ALandscape::TickLayers(float DeltaTime, ELevelTick TickType, FActorTickFunc
 			{
 				PreviousExperimentalLandscapeLayers = GetMutableDefault<UEditorExperimentalSettings>()->bLandscapeLayerSystem;
 				InitializeLayers();
-				bLandscapeLayersAreInitialized = true;
 			}
 
 			// If doing editing while shader are compiling or at load of a map, it's possible we will need another update pass after shader are completed to see the correct result
@@ -4212,7 +4230,7 @@ bool ALandscapeProxy::AddLayer(const FGuid& InLayerGuid)
 	{
 		InitializeLayerWithEmptyContent(InLayerGuid);
 	}
-
+	HasLayersContent = true;
 	return bModified;
 }
 
@@ -4264,9 +4282,7 @@ void ALandscapeProxy::InitializeLayerWithEmptyContent(const FGuid& InLayerGuid)
 		TArray<ULandscapeComponent*>& ComponentList = ComponentsPerHeightmaps.FindOrAdd(ComponentHeightmapTexture);
 		ComponentList.Add(Component);
 	}
-
-	const FIntPoint InitTopLeftSectionBase(INT_MAX, INT_MAX);
-
+		
 	// Init layers with valid "empty" data
 	TMap<UTexture2D*, UTexture2D*> CreatedHeightmapTextures; // < Final layer texture, New created texture for layer
 
@@ -4517,7 +4533,7 @@ void ALandscape::ClearLayer(const FGuid& InLayerGuid, bool bInUpdateCollision)
 	}
 
 	Modify();
-	FScopedSetLandscapeEditingLayer Scope(this, Layer ? Layer->Guid : FGuid(), [=] { RequestLayersContentUpdate(ELandscapeLayersContentUpdateFlag::All, true); });
+	FScopedSetLandscapeEditingLayer Scope(this, Layer->Guid, [=] { RequestLayersContentUpdate(ELandscapeLayersContentUpdateFlag::All, true); });
 
 	TArray<uint16> NewHeightData;
 	NewHeightData.AddZeroed(FMath::Square(ComponentSizeQuads + 1));
@@ -4679,20 +4695,20 @@ void ALandscape::UpdateLandscapeSplines(const FGuid& InTargetLayer, bool bUpdate
 
 FScopedSetLandscapeEditingLayer::FScopedSetLandscapeEditingLayer(ALandscape* InLandscape, const FGuid& InLayerGUID, TFunction<void()> InCompletionCallback)
 	: Landscape(InLandscape)
-	, LayerGUID(InLayerGUID)
 	, CompletionCallback(MoveTemp(InCompletionCallback))
 {
-	if (GetMutableDefault<UEditorExperimentalSettings>()->bLandscapeLayerSystem && Landscape.IsValid() && LayerGUID.IsValid())
+	if (GetMutableDefault<UEditorExperimentalSettings>()->bLandscapeLayerSystem && Landscape.IsValid())
 	{
-		Landscape->SetEditingLayer(LayerGUID);
+		PreviousLayerGUID = Landscape->GetEditingLayer();
+		Landscape->SetEditingLayer(InLayerGUID);
 	}
 }
 
 FScopedSetLandscapeEditingLayer::~FScopedSetLandscapeEditingLayer()
 {
-	if (GetMutableDefault<UEditorExperimentalSettings>()->bLandscapeLayerSystem && Landscape.IsValid() && LayerGUID.IsValid())
+	if (GetMutableDefault<UEditorExperimentalSettings>()->bLandscapeLayerSystem && Landscape.IsValid())
 	{
-		Landscape->SetEditingLayer();
+		Landscape->SetEditingLayer(PreviousLayerGUID);
 		if (CompletionCallback)
 		{
 			CompletionCallback();
@@ -4713,7 +4729,7 @@ bool ALandscape::IsEditingLayerReservedForSplines() const
 void ALandscape::SetEditingLayer(const FGuid& InLayerGuid)
 {
 	ensure(GetMutableDefault<UEditorExperimentalSettings>()->bLandscapeLayerSystem);
-
+	
 	ULandscapeInfo* LandscapeInfo = GetLandscapeInfo();
 	if (!LandscapeInfo)
 	{
@@ -4751,7 +4767,7 @@ TMap<UTexture2D*, TArray<ULandscapeComponent*>> ALandscape::GenerateComponentsPe
 	return ComponentsPerHeightmaps;
 }
 
-void ALandscape::CreateDefaultLayer(const FIntPoint& InComponentCounts)
+void ALandscape::CreateDefaultLayer()
 {
 	ULandscapeInfo* LandscapeInfo = GetLandscapeInfo();
 	if (!LandscapeInfo || !GetMutableDefault<UEditorExperimentalSettings>()->bLandscapeLayerSystem)
@@ -4762,7 +4778,9 @@ void ALandscape::CreateDefaultLayer(const FIntPoint& InComponentCounts)
 	check(LandscapeLayers.Num() == 0); // We can only call this function if we have no layers
 
 	CreateLayer(FName(TEXT("Layer")));
-	CreateLayersRenderingResource(InComponentCounts);
+	HasLayersContent = true;
+	// Force update rendering resources
+	RequestLayersInitialization();
 }
 
 
@@ -4786,6 +4804,29 @@ void ALandscape::CreateLayer(FName InName)
 	});
 }
 
+void ALandscape::AddLayersToProxy(ALandscapeProxy* InProxy)
+{
+	ULandscapeInfo* LandscapeInfo = GetLandscapeInfo();
+	if (!LandscapeInfo || !GetMutableDefault<UEditorExperimentalSettings>()->bLandscapeLayerSystem)
+	{
+		return;
+	}
+
+	Modify();
+
+	check(InProxy != this);
+	check(InProxy != nullptr);
+
+	ForEachLayer([&](FLandscapeLayer& Layer)
+	{
+		InProxy->AddLayer(Layer.Guid);
+	});
+	InProxy->HasLayersContent = true;
+
+	// Force update rendering resources
+	RequestLayersInitialization();
+}
+
 bool ALandscape::ReorderLayer(int32 InStartingLayerIndex, int32 InDestinationLayerIndex)
 {
 	if (InStartingLayerIndex != InDestinationLayerIndex &&
