- Fixed Collision used by [REDACTED_PERSON] not properly computed when using Landscape Layer System (now deffered in TickLayers)
- Painting height in Landscape Layer System now only updates edition heightfield collision of dirty landscape components
- Fixed CachedBounds of LandscapeComponent that was sometimes computing bounds based on editing layer instead of the final layer in Landscape Layer System mode.
- Fixed random crash caused by missing FlushRenderingCommands when re-creating CPU Readback resources and using theses resources in the same frame
#rnx

diff --git a/Engine/Source/Editor/LandscapeEditor/Private/LandscapeEdModeComponentTools.cpp b/Engine/Source/Editor/LandscapeEditor/Private/LandscapeEdModeComponentTools.cpp
--- a/Engine/Source/Editor/LandscapeEditor/Private/LandscapeEdModeComponentTools.cpp
+++ b/Engine/Source/Editor/LandscapeEditor/Private/LandscapeEdModeComponentTools.cpp
@@ -929,6 +929,7 @@ public:
 			EdMode->LandscapeRenderAddCollision = nullptr;
 
 			// Add/update "add collision" around the newly added components
+			if (!GetMutableDefault<UEditorExperimentalSettings>()->bLandscapeLayerSystem)
 			{
 				// Top row
 				int32 ComponentIndexY = ComponentIndexY1 - 1;
diff --git a/Engine/Source/Runtime/Landscape/Classes/Landscape.h b/Engine/Source/Runtime/Landscape/Classes/Landscape.h
--- a/Engine/Source/Runtime/Landscape/Classes/Landscape.h
+++ b/Engine/Source/Runtime/Landscape/Classes/Landscape.h
@@ -75,19 +75,21 @@ enum EWeightmapRTType : uint8
 // Internal Update Flags that shouldn't be used alone
 enum EInternalLandscapeLayersContentUpdateFlag : uint32
 {
-	Internal_Heightmap_Render = 0x00000001u,
-	Internal_Heightmap_BoundsAndCollision = 0x00000002u,
-	Internal_Heightmap_ResolveToTexture = 0x00000004u,
-	Internal_Heightmap_ClientUpdate = 0x000000008u,
-
-	Internal_Weightmap_Render = 0x00000100u,
-	Internal_Weightmap_Collision = 0x00000200u,
-	Internal_Weightmap_ResolveToTexture = 0x00000400u,
+	Internal_Heightmap_Render = 1 << 0, 
+	Internal_Heightmap_BoundsAndCollision = 1 << 1,
+	Internal_Heightmap_BoundsAndCollisionPartial = 1 << 2,
+	Internal_Heightmap_ResolveToTexture = 1 << 3,
+	Internal_Heightmap_ClientUpdate = 1 << 4,
+
+	Internal_Weightmap_Render = 1 << 8,
+	Internal_Weightmap_Collision = 1 << 9,
+	Internal_Weightmap_ResolveToTexture = 1 << 10,
 };
 
 enum ELandscapeLayersContentUpdateFlag : uint32
 {
 	Heightmap_Render = Internal_Heightmap_Render,
+	Heightmap_Editing = Internal_Heightmap_Render | Internal_Heightmap_ResolveToTexture | Internal_Heightmap_BoundsAndCollisionPartial,
 	Weightmap_Render = Internal_Weightmap_Render,
 
 	// Combinations
@@ -248,7 +250,7 @@ public:
 
 	// Layers stuff
 #if WITH_EDITOR
-	LANDSCAPE_API void RequestLayersInitialization() { bLandscapeLayersAreInitialized = false; }
+	LANDSCAPE_API void RequestLayersInitialization() { bLandscapeLayersAreInitialized = false; RequestLayersContentUpdate(ELandscapeLayersContentUpdateFlag::All); }
 	LANDSCAPE_API void RequestLayersContentUpdate(ELandscapeLayersContentUpdateFlag InDataFlags, bool InUpdateAllMaterials = false);
 	LANDSCAPE_API bool ReorderLayer(int32 InStartingLayerIndex, int32 InDestinationLayerIndex);
 	LANDSCAPE_API void CreateLayer(FName InName = NAME_None);
@@ -392,7 +394,7 @@ private:
 	class FLandscapeTexture2DResource* WeightmapScratchExtractLayerTextureResource;	
 	
 	// Used in packing the material layer data contained into CombinedLayersWeightmapAllMaterialLayersResource to be set again for each component weightmap (size of the landscape)
-	class FLandscapeTexture2DResource* WeightmapScratchPackLayerTextureResource;	
+	class FLandscapeTexture2DResource* WeightmapScratchPackLayerTextureResource;
 #endif
 
 protected:
diff --git a/Engine/Source/Runtime/Landscape/Classes/LandscapeComponent.h b/Engine/Source/Runtime/Landscape/Classes/LandscapeComponent.h
--- a/Engine/Source/Runtime/Landscape/Classes/LandscapeComponent.h
+++ b/Engine/Source/Runtime/Landscape/Classes/LandscapeComponent.h
@@ -379,6 +379,8 @@ private:
 	// Final layer data
 	UPROPERTY(Transient)
 	TArray<ULandscapeWeightmapUsage*> WeightmapTexturesUsage;
+
+	bool bLayerContentDirty;
 #endif // WITH_EDITORONLY_DATA
 
 	/** Heightmap texture reference */
@@ -766,7 +768,7 @@ public:
 	 * @param XYOffsetTextureMipData: xy-offset map data
 	 * @returns True if CollisionComponent was created in this update.
 	 */
-	bool UpdateCollisionHeightData(const FColor* HeightmapTextureMipData, const FColor* SimpleCollisionHeightmapTextureData, int32 ComponentX1=0, int32 ComponentY1=0, int32 ComponentX2=MAX_int32, int32 ComponentY2=MAX_int32, bool bUpdateBounds=false, const FColor* XYOffsetTextureMipData=nullptr);
+	void UpdateCollisionHeightData(const FColor* HeightmapTextureMipData, const FColor* SimpleCollisionHeightmapTextureData, int32 ComponentX1=0, int32 ComponentY1=0, int32 ComponentX2=MAX_int32, int32 ComponentY2=MAX_int32, bool bUpdateBounds=false, const FColor* XYOffsetTextureMipData=nullptr, bool bOnlyUpdateEditingHeightfield=true);
 
 	/**
 	 * Deletes Collision Component
@@ -774,8 +776,8 @@ public:
 	void DestroyCollisionData();
 
 	/** Updates collision component height data for the entire component, locking and unlocking heightmap textures
-	 * @param: bRebuild: If true, recreates the collision component */
-	void UpdateCollisionData();
+	 * @param: bOnlyUpdateEditingHeightfield: If true, only updates the editing collision heightfield */
+	void UpdateCollisionData(bool bOnlyUpdateEditingHeightfield = false);
 
 	/**
 	 * Update collision component dominant layer data
@@ -838,6 +840,9 @@ public:
 	
 	/** Updates the values of component-level properties exposed by the Landscape Actor */
 	LANDSCAPE_API void UpdatedSharedPropertiesFromActor();
+
+	LANDSCAPE_API bool IsLayerContentDirty() const { return bLayerContentDirty; }
+	LANDSCAPE_API void SetLayerContentDirty(bool InValue) { bLayerContentDirty = InValue; }
 #endif
 
 	friend class FLandscapeComponentSceneProxy;
diff --git a/Engine/Source/Runtime/Landscape/Private/Landscape.cpp b/Engine/Source/Runtime/Landscape/Private/Landscape.cpp
--- a/Engine/Source/Runtime/Landscape/Private/Landscape.cpp
+++ b/Engine/Source/Runtime/Landscape/Private/Landscape.cpp
@@ -140,6 +140,7 @@ FAutoConsoleCommand CmdPrintNumLandscapeShadows(
 
 ULandscapeComponent::ULandscapeComponent(const FObjectInitializer& ObjectInitializer)
 : Super(ObjectInitializer)
+, bLayerContentDirty(false)
 , GrassData(MakeShareable(new FLandscapeComponentGrassData()))
 , ChangeTag(0)
 {
diff --git a/Engine/Source/Runtime/Landscape/Private/LandscapeDataAccess.cpp b/Engine/Source/Runtime/Landscape/Private/LandscapeDataAccess.cpp
--- a/Engine/Source/Runtime/Landscape/Private/LandscapeDataAccess.cpp
+++ b/Engine/Source/Runtime/Landscape/Private/LandscapeDataAccess.cpp
@@ -6,25 +6,27 @@
 #if WITH_EDITOR
 
 
-LANDSCAPE_API FLandscapeComponentDataInterface::FLandscapeComponentDataInterface(ULandscapeComponent* InComponent, int32 InMipLevel) :
+LANDSCAPE_API FLandscapeComponentDataInterface::FLandscapeComponentDataInterface(ULandscapeComponent* InComponent, int32 InMipLevel, bool InWorkOnEditingLayer) :
 	Component(InComponent),
+	bWorkOnEditingLayer(InWorkOnEditingLayer),
 	HeightMipData(NULL),
 	XYOffsetMipData(NULL),
 	MipLevel(InMipLevel)
 {
 	// Offset and stride for this component's data in heightmap texture
-	HeightmapStride = Component->GetHeightmap(true)->Source.GetSizeX() >> MipLevel;
-	HeightmapComponentOffsetX = FMath::RoundToInt((float)(Component->GetHeightmap(true)->Source.GetSizeX() >> MipLevel) * Component->HeightmapScaleBias.Z);
-	HeightmapComponentOffsetY = FMath::RoundToInt((float)(Component->GetHeightmap(true)->Source.GetSizeY() >> MipLevel) * Component->HeightmapScaleBias.W);
+	UTexture2D* HeightMapTexture = Component->GetHeightmap(bWorkOnEditingLayer);
+	HeightmapStride = HeightMapTexture->Source.GetSizeX() >> MipLevel;
+	HeightmapComponentOffsetX = FMath::RoundToInt((float)(HeightMapTexture->Source.GetSizeX() >> MipLevel) * Component->HeightmapScaleBias.Z);
+	HeightmapComponentOffsetY = FMath::RoundToInt((float)(HeightMapTexture->Source.GetSizeY() >> MipLevel) * Component->HeightmapScaleBias.W);
 	HeightmapSubsectionOffset = (Component->SubsectionSizeQuads + 1) >> MipLevel;
 
 	ComponentSizeVerts = (Component->ComponentSizeQuads + 1) >> MipLevel;
 	SubsectionSizeVerts = (Component->SubsectionSizeQuads + 1) >> MipLevel;
 	ComponentNumSubsections = Component->NumSubsections;
 
-	if (MipLevel < Component->GetHeightmap(true)->Source.GetNumMips())
+	if (MipLevel < HeightMapTexture->Source.GetNumMips())
 	{
-		HeightMipData = (FColor*)DataInterface.LockMip(Component->GetHeightmap(true), MipLevel);
+		HeightMipData = (FColor*)DataInterface.LockMip(HeightMapTexture, MipLevel);
 		if (Component->XYOffsetmapTexture)
 		{
 			XYOffsetMipData = (FColor*)DataInterface.LockMip(Component->XYOffsetmapTexture, MipLevel);
@@ -36,7 +38,8 @@ LANDSCAPE_API FLandscapeComponentDataInterface::~FLandscapeComponentDataInterfac
 {
 	if (HeightMipData)
 	{
-		DataInterface.UnlockMip(Component->GetHeightmap(true), MipLevel);
+		UTexture2D* HeightMapTexture = Component->GetHeightmap(bWorkOnEditingLayer);
+		DataInterface.UnlockMip(HeightMapTexture, MipLevel);
 		if (Component->XYOffsetmapTexture)
 		{
 			DataInterface.UnlockMip(Component->XYOffsetmapTexture, MipLevel);
diff --git a/Engine/Source/Runtime/Landscape/Private/LandscapeEdit.cpp b/Engine/Source/Runtime/Landscape/Private/LandscapeEdit.cpp
--- a/Engine/Source/Runtime/Landscape/Private/LandscapeEdit.cpp
+++ b/Engine/Source/Runtime/Landscape/Private/LandscapeEdit.cpp
@@ -98,7 +98,9 @@ void ULandscapeComponent::Init(int32 InBaseX, int32 InBaseY, int32 InComponentSi
 
 void ULandscapeComponent::UpdateCachedBounds()
 {
-	FLandscapeComponentDataInterface CDI(this);
+	const int32 MipLevel = 0;
+	const bool bWorkOnEditingLayer = false; // We never want to compute bounds based on anything else that final landscape layer's height data
+	FLandscapeComponentDataInterface CDI(this, MipLevel, bWorkOnEditingLayer);
 
 	// Update local-space bounding box
 	CachedLocalBox.Init();
@@ -788,7 +790,7 @@ public:
 	int32 SizeVertsSquare = 0;
 };
 
-bool ULandscapeComponent::UpdateCollisionHeightData(const FColor* const HeightmapTextureMipData, const FColor* const SimpleCollisionHeightmapTextureData, int32 ComponentX1/*=0*/, int32 ComponentY1/*=0*/, int32 ComponentX2/*=MAX_int32*/, int32 ComponentY2/*=MAX_int32*/, bool bUpdateBounds/*=false*/, const FColor* XYOffsetTextureMipData/*=nullptr*/)
+void ULandscapeComponent::UpdateCollisionHeightData(const FColor* const HeightmapTextureMipData, const FColor* const SimpleCollisionHeightmapTextureData, int32 ComponentX1/*=0*/, int32 ComponentY1/*=0*/, int32 ComponentX2/*=MAX_int32*/, int32 ComponentY2/*=MAX_int32*/, bool bUpdateBounds/*=false*/, const FColor* XYOffsetTextureMipData/*=nullptr*/, bool bOnlyUpdateEditingHeightfield/*=true*/)
 {
 	ULandscapeInfo* Info = GetLandscapeInfo();
 	ALandscapeProxy* Proxy = GetLandscapeProxy();
@@ -830,7 +832,7 @@ bool ULandscapeComponent::UpdateCollisionHeightData(const FColor* const Heightma
 		if (ComponentX2 < ComponentX1 || ComponentY2 < ComponentY1)
 		{
 			// nothing to do
-			return false;
+			return;
 		}
 
 		if (bUpdateBounds)
@@ -944,7 +946,11 @@ bool ULandscapeComponent::UpdateCollisionHeightData(const FColor* const Heightma
 		CollisionComp->RegisterComponent();
 	}
 
-	return CreatedNew;
+	// Only recreate collision if not in partial update and not newly created component
+	if (!CreatedNew && !bOnlyUpdateEditingHeightfield)
+	{
+		CollisionComp->RecreateCollision();
+	}
 }
 
 void ULandscapeComponent::DestroyCollisionData()
@@ -957,7 +963,7 @@ void ULandscapeComponent::DestroyCollisionData()
 	}
 }
 
-void ULandscapeComponent::UpdateCollisionData()
+void ULandscapeComponent::UpdateCollisionData(bool bOnlyUpdateEditingHeightfield /*= false*/)
 {
 	TArray<uint8> CollisionMipData;
 	TArray<uint8> SimpleCollisionMipData;
@@ -973,21 +979,11 @@ void ULandscapeComponent::UpdateCollisionData()
 		XYOffsetmapTexture->Source.GetMipData(XYOffsetMipData, CollisionMipLevel);
 	}
 
-	bool bCreatedNewComponent = UpdateCollisionHeightData(
+	UpdateCollisionHeightData(
 		(FColor*)CollisionMipData.GetData(),
 		SimpleCollisionMipLevel > CollisionMipLevel ? (FColor*)SimpleCollisionMipData.GetData() : nullptr,
 		0, 0, MAX_int32, MAX_int32, true,
-		XYOffsetmapTexture ? (FColor*)XYOffsetMipData.GetData() : nullptr);
-
-    // If the Collision Component hasn't changed we need to reflect its data to the underlying physics implementation
-	if (!bCreatedNewComponent)
-	{
-		ULandscapeHeightfieldCollisionComponent* CollisionComp = CollisionComponent.Get();
-		if (CollisionComp)
-		{
-			CollisionComp->RecreateCollision();
-		}
-	}
+		XYOffsetmapTexture ? (FColor*)XYOffsetMipData.GetData() : nullptr, bOnlyUpdateEditingHeightfield);
 }
 
 void ULandscapeComponent::RecreateCollisionComponent(bool bUseSimpleCollision)
diff --git a/Engine/Source/Runtime/Landscape/Private/LandscapeEditInterface.cpp b/Engine/Source/Runtime/Landscape/Private/LandscapeEditInterface.cpp
--- a/Engine/Source/Runtime/Landscape/Private/LandscapeEditInterface.cpp
+++ b/Engine/Source/Runtime/Landscape/Private/LandscapeEditInterface.cpp
@@ -452,6 +452,14 @@ void FLandscapeEditDataInterface::SetHeightData(int32 X1, int32 Y1, int32 X2, in
 					}
 				}
 			}
+			else
+			{
+				if (ALandscape* Landscape = Component->GetLandscapeActor())
+				{
+					Component->SetLayerContentDirty(true);
+					Landscape->RequestLayersContentUpdate(ELandscapeLayersContentUpdateFlag::Heightmap_Editing);
+				}
+			}
 
 			// Update GUID for Platform Data
 			FPlatformMisc::CreateGuid(Component->StateId);
diff --git a/Engine/Source/Runtime/Landscape/Private/LandscapeEditLayers.cpp b/Engine/Source/Runtime/Landscape/Private/LandscapeEditLayers.cpp
--- a/Engine/Source/Runtime/Landscape/Private/LandscapeEditLayers.cpp
+++ b/Engine/Source/Runtime/Landscape/Private/LandscapeEditLayers.cpp
@@ -1267,6 +1267,9 @@ void ALandscape::CreateLayersRenderingResource(const FIntPoint& InComponentCount
 		}
 	});
 
+	// Flush because TickLayers can access CPU Readbacks in the same frame (see ResolveLayersTexture)
+	FlushRenderingCommands();
+
 	const int32 TotalVertexCountX = (SubsectionSizeQuads * NumSubsections) * InComponentCounts.X + 1;
 	const int32 TotalVertexCountY = (SubsectionSizeQuads * NumSubsections) * InComponentCounts.Y + 1;
 
@@ -2678,19 +2681,20 @@ void ALandscape::RegenerateLayersHeightmaps()
 		ResolveLayersHeightmapTexture();
 	}
 
-	bool bUpdateBoundsAndCollsion = (LayersContentUpdateFlags & EInternalLandscapeLayersContentUpdateFlag::Internal_Heightmap_BoundsAndCollision) != 0;
+	bool bUpdateBoundsAndCollsionFull = (LayersContentUpdateFlags & EInternalLandscapeLayersContentUpdateFlag::Internal_Heightmap_BoundsAndCollision) != 0;
+	bool bUpdateBoundsAndCollsionPartial = ((LayersContentUpdateFlags & EInternalLandscapeLayersContentUpdateFlag::Internal_Heightmap_BoundsAndCollisionPartial) != 0) && !bUpdateBoundsAndCollsionFull; // Partial only if not full update
+	bool bUpdateBoundsAndCollsion = bUpdateBoundsAndCollsionFull || bUpdateBoundsAndCollsionPartial;
 	bool bUpdateClients = (LayersContentUpdateFlags & EInternalLandscapeLayersContentUpdateFlag::Internal_Heightmap_ClientUpdate) != 0;
 	int32 PostponeUpdateClientFlag = 0;
 	if (bUpdateBoundsAndCollsion || bUpdateClients)
 	{
 		for (ULandscapeComponent* Component : AllLandscapeComponents)
 		{
-			if (bUpdateBoundsAndCollsion)
+			if (bUpdateBoundsAndCollsion && (!bUpdateBoundsAndCollsionPartial || Component->IsLayerContentDirty()))
 			{
 				Component->UpdateCachedBounds();
 				Component->UpdateComponentToWorld();
-
-				Component->UpdateCollisionData();
+				Component->UpdateCollisionData(bUpdateBoundsAndCollsionPartial);
 			}
 
 			if (bUpdateClients)
@@ -2711,9 +2715,14 @@ void ALandscape::RegenerateLayersHeightmaps()
 					}
 				}
 			}
+			Component->SetLayerContentDirty(false);
 		}
+
+		// Ideally, we should only update necessary components and only when inside AddComponnet Tool.
+		Info->UpdateAllAddCollisions();
 	}
 
+	LayersContentUpdateFlags &= ~EInternalLandscapeLayersContentUpdateFlag::Internal_Heightmap_BoundsAndCollisionPartial;
 	LayersContentUpdateFlags &= ~ELandscapeLayersContentUpdateFlag::Heightmap_All;
 	LayersContentUpdateFlags |= PostponeUpdateClientFlag; // if non-zero postpone client update (transacting)
 
diff --git a/Engine/Source/Runtime/Landscape/Public/LandscapeDataAccess.h b/Engine/Source/Runtime/Landscape/Public/LandscapeDataAccess.h
--- a/Engine/Source/Runtime/Landscape/Public/LandscapeDataAccess.h
+++ b/Engine/Source/Runtime/Landscape/Public/LandscapeDataAccess.h
@@ -115,7 +115,7 @@ struct FLandscapeComponentDataInterface
 	friend struct FLandscapeDataInterface;
 
 	// tors
-	LANDSCAPE_API FLandscapeComponentDataInterface(ULandscapeComponent* InComponent, int32 InMipLevel = 0);
+	LANDSCAPE_API FLandscapeComponentDataInterface(ULandscapeComponent* InComponent, int32 InMipLevel = 0, bool InWorkOnEditingLayer = true);
 	LANDSCAPE_API ~FLandscapeComponentDataInterface();
 
 	// Accessors
@@ -323,6 +323,7 @@ struct FLandscapeComponentDataInterface
 private:
 	FLandscapeDataInterface DataInterface;
 	ULandscapeComponent* Component;
+	bool bWorkOnEditingLayer;
 
 public:
 	// offset of this component's data into heightmap texture
