Fixed build break
#rnx

diff --git a/Engine/Source/Runtime/Landscape/Private/Landscape.cpp b/Engine/Source/Runtime/Landscape/Private/Landscape.cpp
--- a/Engine/Source/Runtime/Landscape/Private/Landscape.cpp
+++ b/Engine/Source/Runtime/Landscape/Private/Landscape.cpp
@@ -140,7 +140,9 @@ FAutoConsoleCommand CmdPrintNumLandscapeShadows(
 
 ULandscapeComponent::ULandscapeComponent(const FObjectInitializer& ObjectInitializer)
 : Super(ObjectInitializer)
+#if WITH_EDITORONLY_DATA
 , bLayerContentDirty(false)
+#endif
 , GrassData(MakeShareable(new FLandscapeComponentGrassData()))
 , ChangeTag(0)
 {
