- Lanscape Layer System: Make sure we only do Collision Updates in the TickLayers when using Layers

#rnx

diff --git a/Engine/Source/Editor/LandscapeEditor/Private/LandscapeEdModeComponentTools.cpp b/Engine/Source/Editor/LandscapeEditor/Private/LandscapeEdModeComponentTools.cpp
--- a/Engine/Source/Editor/LandscapeEditor/Private/LandscapeEdModeComponentTools.cpp
+++ b/Engine/Source/Editor/LandscapeEditor/Private/LandscapeEdModeComponentTools.cpp
@@ -1728,9 +1728,7 @@ public:
 	virtual ELandscapeLayersContentUpdateFlag GetBeginToolContentUpdateFlag() const override { return ELandscapeLayersContentUpdateFlag::All_Render; }
 
 	virtual ELandscapeLayersContentUpdateFlag GetTickToolContentUpdateFlag() const override { return ELandscapeLayersContentUpdateFlag::All_Render; }
-
-	virtual ELandscapeLayersContentUpdateFlag GetEndToolContentUpdateFlag() const override { return ELandscapeLayersContentUpdateFlag::All; }
-
+		
 	virtual bool BeginTool(FEditorViewportClient* ViewportClient, const FLandscapeToolTarget& InTarget, const FVector& InHitLocation) override
 	{
 		this->EdMode->GizmoBrush->Tick(ViewportClient, 0.1f);
diff --git a/Engine/Source/Editor/LandscapeEditor/Private/LandscapeEdModeMirrorTool.cpp b/Engine/Source/Editor/LandscapeEditor/Private/LandscapeEdModeMirrorTool.cpp
--- a/Engine/Source/Editor/LandscapeEditor/Private/LandscapeEdModeMirrorTool.cpp
+++ b/Engine/Source/Editor/LandscapeEditor/Private/LandscapeEdModeMirrorTool.cpp
@@ -647,19 +647,22 @@ public:
 		TSet<ULandscapeComponent*> Components;
 		if (LandscapeEdit.GetComponentsInRegion(DestMinX, DestMinY, DestMaxX, DestMaxY, &Components) && Components.Num() > 0)
 		{
-			for (ULandscapeComponent* Component : Components)
+			if (!GetMutableDefault<UEditorExperimentalSettings>()->bLandscapeLayerSystem)
 			{
-				// Recreate collision for modified components and update the navmesh
-				ULandscapeHeightfieldCollisionComponent* CollisionComponent = Component->CollisionComponent.Get();
-				if (CollisionComponent)
+				for (ULandscapeComponent* Component : Components)
 				{
-					CollisionComponent->RecreateCollision();
-					FNavigationSystem::UpdateComponentData(*CollisionComponent);
+					// Recreate collision for modified components and update the navmesh
+					ULandscapeHeightfieldCollisionComponent* CollisionComponent = Component->CollisionComponent.Get();
+					if (CollisionComponent)
+					{
+						CollisionComponent->RecreateCollision();
+						FNavigationSystem::UpdateComponentData(*CollisionComponent);
+					}
 				}
-			}
 
-			// Flush dynamic foliage (grass)
-			ALandscapeProxy::InvalidateGeneratedComponentData(Components);
+				// Flush dynamic foliage (grass)
+				ALandscapeProxy::InvalidateGeneratedComponentData(Components);
+			}
 
 			EdMode->UpdateLayerUsageInformation();
 		}
diff --git a/Engine/Source/Editor/LandscapeEditor/Private/LandscapeEdModeTools.h b/Engine/Source/Editor/LandscapeEditor/Private/LandscapeEdModeTools.h
--- a/Engine/Source/Editor/LandscapeEditor/Private/LandscapeEdModeTools.h
+++ b/Engine/Source/Editor/LandscapeEditor/Private/LandscapeEdModeTools.h
@@ -692,19 +692,22 @@ struct FXYOffsetmapAccessor
 					NewHeights[X - X1 + (Y - Y1) * (X2 - X1 + 1)] = FMath::Clamp<uint16>(Data[(X - X1 + (Y - Y1) * (X2 - X1 + 1))].Z * LANDSCAPE_INV_ZSCALE + 32768.0f, 0, 65535);
 				}
 			}
-
-			// Flush dynamic foliage (grass)
-			ALandscapeProxy::InvalidateGeneratedComponentData(Components);
-
+						
 			// Notify foliage to move any attached instances
 			bool bUpdateFoliage = false;
-			for (ULandscapeComponent* Component : Components)
+			if (!GetMutableDefault<UEditorExperimentalSettings>()->bLandscapeLayerSystem)
 			{
-				ULandscapeHeightfieldCollisionComponent* CollisionComponent = Component->CollisionComponent.Get();
-				if (CollisionComponent && AInstancedFoliageActor::HasFoliageAttached(CollisionComponent))
+				// Flush dynamic foliage (grass)
+				ALandscapeProxy::InvalidateGeneratedComponentData(Components);
+
+				for (ULandscapeComponent* Component : Components)
 				{
-					bUpdateFoliage = true;
-					break;
+					ULandscapeHeightfieldCollisionComponent* CollisionComponent = Component->CollisionComponent.Get();
+					if (CollisionComponent && AInstancedFoliageActor::HasFoliageAttached(CollisionComponent))
+					{
+						bUpdateFoliage = true;
+						break;
+					}
 				}
 			}
 
@@ -836,18 +839,21 @@ struct FFullWeightmapAccessor
 
 	~FFullWeightmapAccessor()
 	{
-		// Recreate collision for modified components to update the physical materials
-		for (ULandscapeComponent* Component : ModifiedComponents)
+		if (!GetMutableDefault<UEditorExperimentalSettings>()->bLandscapeLayerSystem)
 		{
-			ULandscapeHeightfieldCollisionComponent* CollisionComponent = Component->CollisionComponent.Get();
-			if (CollisionComponent)
+			// Recreate collision for modified components to update the physical materials
+			for (ULandscapeComponent* Component : ModifiedComponents)
 			{
-				CollisionComponent->RecreateCollision();
-
-				// We need to trigger navigation mesh build, in case user have painted holes on a landscape
-				if (LandscapeInfo->GetLayerInfoIndex(ALandscapeProxy::VisibilityLayer) != INDEX_NONE)
+				ULandscapeHeightfieldCollisionComponent* CollisionComponent = Component->CollisionComponent.Get();
+				if (CollisionComponent)
 				{
-					FNavigationSystem::UpdateComponentData(*CollisionComponent);
+					CollisionComponent->RecreateCollision();
+
+					// We need to trigger navigation mesh build, in case user have painted holes on a landscape
+					if (LandscapeInfo->GetLayerInfoIndex(ALandscapeProxy::VisibilityLayer) != INDEX_NONE)
+					{
+						FNavigationSystem::UpdateComponentData(*CollisionComponent);
+					}
 				}
 			}
 		}
@@ -870,11 +876,13 @@ struct FFullWeightmapAccessor
 		TSet<ULandscapeComponent*> Components;
 		if (LandscapeEdit.GetComponentsInRegion(X1, Y1, X2, Y2, &Components))
 		{
-			// Flush dynamic foliage (grass)
-			ALandscapeProxy::InvalidateGeneratedComponentData(Components);
-
+			if (!GetMutableDefault<UEditorExperimentalSettings>()->bLandscapeLayerSystem)
+			{
+				// Flush dynamic foliage (grass)
+				ALandscapeProxy::InvalidateGeneratedComponentData(Components);
+				ModifiedComponents.Append(Components);
+			}
 			LandscapeEdit.SetAlphaData(DirtyLayerInfos, X1, Y1, X2, Y2, Data, 0, PaintingRestriction);
-			ModifiedComponents.Append(Components);
 		}
 		DirtyLayerInfos.Empty();
 	}
@@ -1136,8 +1144,7 @@ public:
 
 	virtual ELandscapeLayersContentUpdateFlag GetEndToolContentUpdateFlag() const
 	{
-		bool bUpdateHeightmap = this->EdMode->CurrentToolTarget.TargetType == ELandscapeToolTargetType::Type::Heightmap; 
-		return bUpdateHeightmap ? ELandscapeLayersContentUpdateFlag::Heightmap_All : ELandscapeLayersContentUpdateFlag::Weightmap_All;
+		return ELandscapeLayersContentUpdateFlag::All;
 	}
 
 	virtual bool BeginTool(FEditorViewportClient* ViewportClient, const FLandscapeToolTarget& InTarget, const FVector& InHitLocation) override
diff --git a/Engine/Source/Runtime/Landscape/Classes/Landscape.h b/Engine/Source/Runtime/Landscape/Classes/Landscape.h
--- a/Engine/Source/Runtime/Landscape/Classes/Landscape.h
+++ b/Engine/Source/Runtime/Landscape/Classes/Landscape.h
@@ -79,11 +79,12 @@ enum EInternalLandscapeLayersContentUpdateFlag : uint32
 	Internal_Heightmap_BoundsAndCollision = 1 << 1,
 	Internal_Heightmap_BoundsAndCollisionPartial = 1 << 2,
 	Internal_Heightmap_ResolveToTexture = 1 << 3,
-	Internal_Heightmap_ClientUpdate = 1 << 4,
-
+	
 	Internal_Weightmap_Render = 1 << 8,
 	Internal_Weightmap_Collision = 1 << 9,
 	Internal_Weightmap_ResolveToTexture = 1 << 10,
+
+	Internal_ClientUpdate = 1 << 14
 };
 
 enum ELandscapeLayersContentUpdateFlag : uint32
@@ -93,8 +94,8 @@ enum ELandscapeLayersContentUpdateFlag : uint32
 	Weightmap_Render = Internal_Weightmap_Render,
 
 	// Combinations
-	Heightmap_All = Internal_Heightmap_Render | Internal_Heightmap_BoundsAndCollision | Internal_Heightmap_ResolveToTexture | Internal_Heightmap_ClientUpdate,
-	Weightmap_All = Internal_Weightmap_Render | Internal_Weightmap_Collision | Internal_Weightmap_ResolveToTexture,
+	Heightmap_All = Internal_Heightmap_Render | Internal_Heightmap_BoundsAndCollision | Internal_Heightmap_ResolveToTexture | Internal_ClientUpdate,
+	Weightmap_All = Internal_Weightmap_Render | Internal_Weightmap_Collision | Internal_Weightmap_ResolveToTexture | Internal_ClientUpdate,
 
 	All = Heightmap_All | Weightmap_All,
 	All_Render = Heightmap_Render | Weightmap_Render,
@@ -296,8 +297,9 @@ private:
 	void TickLayers(float DeltaTime, ELevelTick TickType, FActorTickFunction& ThisTickFunction);
 	void CreateLayersRenderingResource(const FIntPoint& InComponentCounts);
 	void RegenerateLayersContent();
-	void RegenerateLayersHeightmaps();
-	void RegenerateLayersWeightmaps();
+	void RegenerateLayersHeightmaps(const TArray<ULandscapeComponent*>& InLandscapeComponents, bool& bRecreateCollision);
+	void RegenerateLayersWeightmaps(const TArray<ULandscapeComponent*>& InLandscapeComponents, bool& bRecreateCollision);
+	bool UpdateClients(const TArray<ULandscapeComponent*>& InLandscapeComponents, bool bCollisionRecreated);
 	void ResolveLayersHeightmapTexture();
 	void ResolveLayersWeightmapTexture();
 	void ResolveLayersTexture(class FLandscapeLayersTexture2DCPUReadBackResource* InCPUReadBackTexture, UTexture2D* InOutputTexture);
diff --git a/Engine/Source/Runtime/Landscape/Classes/LandscapeComponent.h b/Engine/Source/Runtime/Landscape/Classes/LandscapeComponent.h
--- a/Engine/Source/Runtime/Landscape/Classes/LandscapeComponent.h
+++ b/Engine/Source/Runtime/Landscape/Classes/LandscapeComponent.h
@@ -768,7 +768,7 @@ public:
 	 * @param XYOffsetTextureMipData: xy-offset map data
 	 * @returns True if CollisionComponent was created in this update.
 	 */
-	void UpdateCollisionHeightData(const FColor* HeightmapTextureMipData, const FColor* SimpleCollisionHeightmapTextureData, int32 ComponentX1=0, int32 ComponentY1=0, int32 ComponentX2=MAX_int32, int32 ComponentY2=MAX_int32, bool bUpdateBounds=false, const FColor* XYOffsetTextureMipData=nullptr, bool bOnlyUpdateEditingHeightfield=true);
+	void UpdateCollisionHeightData(const FColor* HeightmapTextureMipData, const FColor* SimpleCollisionHeightmapTextureData, int32 ComponentX1=0, int32 ComponentY1=0, int32 ComponentX2=MAX_int32, int32 ComponentY2=MAX_int32, bool bUpdateBounds=false, const FColor* XYOffsetTextureMipData=nullptr);
 
 	/**
 	 * Deletes Collision Component
@@ -776,8 +776,8 @@ public:
 	void DestroyCollisionData();
 
 	/** Updates collision component height data for the entire component, locking and unlocking heightmap textures
-	 * @param: bOnlyUpdateEditingHeightfield: If true, only updates the editing collision heightfield */
-	void UpdateCollisionData(bool bOnlyUpdateEditingHeightfield = false);
+	 */
+	void UpdateCollisionData();
 
 	/**
 	 * Update collision component dominant layer data
diff --git a/Engine/Source/Runtime/Landscape/Classes/LandscapeProxy.h b/Engine/Source/Runtime/Landscape/Classes/LandscapeProxy.h
--- a/Engine/Source/Runtime/Landscape/Classes/LandscapeProxy.h
+++ b/Engine/Source/Runtime/Landscape/Classes/LandscapeProxy.h
@@ -754,6 +754,7 @@ public:
 
 	/* Invalidate the precomputed grass and baked texture data for the specified components */
 	LANDSCAPE_API static void InvalidateGeneratedComponentData(const TSet<ULandscapeComponent*>& Components);
+	LANDSCAPE_API static void InvalidateGeneratedComponentData(const TArray<ULandscapeComponent*>& Components);
 
 	/* Invalidate the precomputed grass and baked texture data on all components */
 	LANDSCAPE_API void InvalidateGeneratedComponentData();
diff --git a/Engine/Source/Runtime/Landscape/Private/Landscape.cpp b/Engine/Source/Runtime/Landscape/Private/Landscape.cpp
--- a/Engine/Source/Runtime/Landscape/Private/Landscape.cpp
+++ b/Engine/Source/Runtime/Landscape/Private/Landscape.cpp
@@ -3504,6 +3504,11 @@ void ALandscapeProxy::InvalidateGeneratedComponentData()
 	InvalidateGeneratedComponentDataImpl(LandscapeComponents);
 }
 
+void ALandscapeProxy::InvalidateGeneratedComponentData(const TArray<ULandscapeComponent*>& Components)
+{
+	InvalidateGeneratedComponentDataImpl(Components);
+}
+
 void ALandscapeProxy::InvalidateGeneratedComponentData(const TSet<ULandscapeComponent*>& Components)
 {
 	InvalidateGeneratedComponentDataImpl(Components);
diff --git a/Engine/Source/Runtime/Landscape/Private/LandscapeCollision.cpp b/Engine/Source/Runtime/Landscape/Private/LandscapeCollision.cpp
--- a/Engine/Source/Runtime/Landscape/Private/LandscapeCollision.cpp
+++ b/Engine/Source/Runtime/Landscape/Private/LandscapeCollision.cpp
@@ -1463,10 +1463,14 @@ void ULandscapeMeshCollisionComponent::Serialize(FArchive& Ar)
 void ULandscapeHeightfieldCollisionComponent::PostEditImport()
 {
 	Super::PostEditImport();
-	// Reinitialize physics after paste
-	if (CollisionSizeQuads > 0)
+
+	if (!GetMutableDefault<UEditorExperimentalSettings>()->bLandscapeLayerSystem)
 	{
-		RecreateCollision();
+		// Reinitialize physics after paste
+		if (CollisionSizeQuads > 0)
+		{
+			RecreateCollision();
+		}
 	}
 }
 
diff --git a/Engine/Source/Runtime/Landscape/Private/LandscapeEdit.cpp b/Engine/Source/Runtime/Landscape/Private/LandscapeEdit.cpp
--- a/Engine/Source/Runtime/Landscape/Private/LandscapeEdit.cpp
+++ b/Engine/Source/Runtime/Landscape/Private/LandscapeEdit.cpp
@@ -573,15 +573,20 @@ void ULandscapeComponent::PostEditUndo()
 		UpdateEditToolRenderData();
 	}
 
-	TSet<ULandscapeComponent*> Components;
-	Components.Add(this);
-	GetLandscapeProxy()->FlushGrassComponents(&Components);
-
-	ALandscape* LandscapeActor = GetLandscapeActor();
-	// Might be a ALandscapeStreamingProxy
-	if (LandscapeActor)
+	if (GetMutableDefault<UEditorExperimentalSettings>()->bLandscapeLayerSystem)
 	{
-		LandscapeActor->RequestLayersContentUpdate(ELandscapeLayersContentUpdateFlag::All, true);
+		ALandscape* LandscapeActor = GetLandscapeActor();
+		// Might be a ALandscapeStreamingProxy
+		if (LandscapeActor)
+		{
+			LandscapeActor->RequestLayersContentUpdate(ELandscapeLayersContentUpdateFlag::All, true);
+		}
+	}
+	else
+	{
+		TSet<ULandscapeComponent*> Components;
+		Components.Add(this);
+		GetLandscapeProxy()->FlushGrassComponents(&Components);
 	}
 }
 
@@ -790,7 +795,7 @@ public:
 	int32 SizeVertsSquare = 0;
 };
 
-void ULandscapeComponent::UpdateCollisionHeightData(const FColor* const HeightmapTextureMipData, const FColor* const SimpleCollisionHeightmapTextureData, int32 ComponentX1/*=0*/, int32 ComponentY1/*=0*/, int32 ComponentX2/*=MAX_int32*/, int32 ComponentY2/*=MAX_int32*/, bool bUpdateBounds/*=false*/, const FColor* XYOffsetTextureMipData/*=nullptr*/, bool bOnlyUpdateEditingHeightfield/*=true*/)
+void ULandscapeComponent::UpdateCollisionHeightData(const FColor* const HeightmapTextureMipData, const FColor* const SimpleCollisionHeightmapTextureData, int32 ComponentX1/*=0*/, int32 ComponentY1/*=0*/, int32 ComponentX2/*=MAX_int32*/, int32 ComponentY2/*=MAX_int32*/, bool bUpdateBounds/*=false*/, const FColor* XYOffsetTextureMipData/*=nullptr*/)
 {
 	ULandscapeInfo* Info = GetLandscapeInfo();
 	ALandscapeProxy* Proxy = GetLandscapeProxy();
@@ -945,12 +950,6 @@ void ULandscapeComponent::UpdateCollisionHeightData(const FColor* const Heightma
 	{
 		CollisionComp->RegisterComponent();
 	}
-
-	// Only recreate collision if not in partial update and not newly created component
-	if (!CreatedNew && !bOnlyUpdateEditingHeightfield)
-	{
-		CollisionComp->RecreateCollision();
-	}
 }
 
 void ULandscapeComponent::DestroyCollisionData()
@@ -963,7 +962,7 @@ void ULandscapeComponent::DestroyCollisionData()
 	}
 }
 
-void ULandscapeComponent::UpdateCollisionData(bool bOnlyUpdateEditingHeightfield /*= false*/)
+void ULandscapeComponent::UpdateCollisionData()
 {
 	TArray<uint8> CollisionMipData;
 	TArray<uint8> SimpleCollisionMipData;
@@ -983,7 +982,7 @@ void ULandscapeComponent::UpdateCollisionData(bool bOnlyUpdateEditingHeightfield
 		(FColor*)CollisionMipData.GetData(),
 		SimpleCollisionMipLevel > CollisionMipLevel ? (FColor*)SimpleCollisionMipData.GetData() : nullptr,
 		0, 0, MAX_int32, MAX_int32, true,
-		XYOffsetmapTexture ? (FColor*)XYOffsetMipData.GetData() : nullptr, bOnlyUpdateEditingHeightfield);
+		XYOffsetmapTexture ? (FColor*)XYOffsetMipData.GetData() : nullptr);
 }
 
 void ULandscapeComponent::RecreateCollisionComponent(bool bUseSimpleCollision)
@@ -1358,24 +1357,19 @@ void ULandscapeComponent::UpdateCollisionLayerData(const FColor* const* const We
 
 void ULandscapeComponent::UpdateCollisionLayerData()
 {
-	TArray<UTexture2D*>& ComponentWeightmapsTexture = GetWeightmapTextures(true);
+	TArray<UTexture2D*>& ComponentWeightmapsTexture = GetWeightmapTextures();
 
 	// Generate the dominant layer data
 	TArray<FColor*> WeightmapTextureMipData;
-	TArray<TArray<uint8>> CachedWeightmapTextureMipData;
-
 	WeightmapTextureMipData.Empty(ComponentWeightmapsTexture.Num());
-	CachedWeightmapTextureMipData.Empty(ComponentWeightmapsTexture.Num());
 	for (int32 WeightmapIdx = 0; WeightmapIdx < ComponentWeightmapsTexture.Num(); ++WeightmapIdx)
 	{
 		TArray<uint8> MipData;
 		ComponentWeightmapsTexture[WeightmapIdx]->Source.GetMipData(MipData, CollisionMipLevel);
 		WeightmapTextureMipData.Add((FColor*)MipData.GetData());
-		CachedWeightmapTextureMipData.Add(MoveTemp(MipData));
 	}
 
 	TArray<FColor*> SimpleCollisionWeightmapMipData;
-	TArray<TArray<uint8>> SimpleCollisionCachedWeightmapTextureMipData;
 	if (SimpleCollisionMipLevel > CollisionMipLevel)
 	{
 		for (int32 WeightmapIdx = 0; WeightmapIdx < ComponentWeightmapsTexture.Num(); ++WeightmapIdx)
@@ -1383,7 +1377,6 @@ void ULandscapeComponent::UpdateCollisionLayerData()
 			TArray<uint8> MipData;
 			ComponentWeightmapsTexture[WeightmapIdx]->Source.GetMipData(MipData, SimpleCollisionMipLevel);
 			SimpleCollisionWeightmapMipData.Add((FColor*)MipData.GetData());
-			SimpleCollisionCachedWeightmapTextureMipData.Add(MoveTemp(MipData));
 		}
 	}
 
@@ -3952,10 +3945,20 @@ void ALandscapeProxy::PostEditChangeProperty(FPropertyChangedEvent& PropertyChan
 	}
 	if (bRemovedAnyLayers)
 	{
-		// Flush dynamic data (e.g. grass)
-		TSet<ULandscapeComponent*> Components;
-		Components.Append(LandscapeComponents);
-		ALandscapeProxy::InvalidateGeneratedComponentData(Components);
+		if (GetMutableDefault<UEditorExperimentalSettings>()->bLandscapeLayerSystem)
+		{
+			if(ALandscape* LandscapeActor = GetLandscapeActor())
+			{
+				LandscapeActor->RequestLayersContentUpdate(ELandscapeLayersContentUpdateFlag::All, true);
+			}
+		}
+		else
+		{
+			// Flush dynamic data (e.g. grass)
+			TSet<ULandscapeComponent*> Components;
+			Components.Append(LandscapeComponents);
+			ALandscapeProxy::InvalidateGeneratedComponentData(Components);
+		}
 	}
 
 	// Must do this *after* correcting the scale or reattaching the landscape components will crash!
diff --git a/Engine/Source/Runtime/Landscape/Private/LandscapeEditInterface.cpp b/Engine/Source/Runtime/Landscape/Private/LandscapeEditInterface.cpp
--- a/Engine/Source/Runtime/Landscape/Private/LandscapeEditInterface.cpp
+++ b/Engine/Source/Runtime/Landscape/Private/LandscapeEditInterface.cpp
@@ -1752,28 +1752,32 @@ void ULandscapeComponent::DeleteLayer(ULandscapeLayerInfoObject* LayerInfo, FLan
 	Component->EditToolRenderData.UpdateDebugColorMaterial(Component);
 	Component->UpdateEditToolRenderData();
 
+	
 	// Update dominant layer info stored in collision component
-	TArray<FColor*> CollisionWeightmapMipData;
-	for (int32 WeightmapIdx = 0; WeightmapIdx < ComponentWeightmapTextures.Num(); WeightmapIdx++)
-	{
-		CollisionWeightmapMipData.Add((FColor*)LandscapeEdit.GetTextureDataInfo(ComponentWeightmapTextures[WeightmapIdx])->GetMipData(Component->CollisionMipLevel));
-	}
-	TArray<FColor*> SimpleCollisionWeightmapMipData;
-	if (Component->SimpleCollisionMipLevel > Component->CollisionMipLevel)
+	if (!GetMutableDefault<UEditorExperimentalSettings>()->bLandscapeLayerSystem)
 	{
+		TArray<FColor*> CollisionWeightmapMipData;
 		for (int32 WeightmapIdx = 0; WeightmapIdx < ComponentWeightmapTextures.Num(); WeightmapIdx++)
 		{
-			SimpleCollisionWeightmapMipData.Add((FColor*)LandscapeEdit.GetTextureDataInfo(ComponentWeightmapTextures[WeightmapIdx])->GetMipData(Component->SimpleCollisionMipLevel));
+			CollisionWeightmapMipData.Add((FColor*)LandscapeEdit.GetTextureDataInfo(ComponentWeightmapTextures[WeightmapIdx])->GetMipData(Component->CollisionMipLevel));
 		}
-	}
-	Component->UpdateCollisionLayerData(
-		CollisionWeightmapMipData.GetData(),
-		Component->SimpleCollisionMipLevel > Component->CollisionMipLevel ? SimpleCollisionWeightmapMipData.GetData() : nullptr);
+		TArray<FColor*> SimpleCollisionWeightmapMipData;
+		if (Component->SimpleCollisionMipLevel > Component->CollisionMipLevel)
+		{
+			for (int32 WeightmapIdx = 0; WeightmapIdx < ComponentWeightmapTextures.Num(); WeightmapIdx++)
+			{
+				SimpleCollisionWeightmapMipData.Add((FColor*)LandscapeEdit.GetTextureDataInfo(ComponentWeightmapTextures[WeightmapIdx])->GetMipData(Component->SimpleCollisionMipLevel));
+			}
+		}
+		Component->UpdateCollisionLayerData(
+			CollisionWeightmapMipData.GetData(),
+			Component->SimpleCollisionMipLevel > Component->CollisionMipLevel ? SimpleCollisionWeightmapMipData.GetData() : nullptr);
 
-	ULandscapeHeightfieldCollisionComponent* LocalCollisionComponent = Component->CollisionComponent.Get();
-	if (LocalCollisionComponent)
-	{
-		LocalCollisionComponent->RecreateCollision();
+		ULandscapeHeightfieldCollisionComponent* LocalCollisionComponent = Component->CollisionComponent.Get();
+		if (LocalCollisionComponent)
+		{
+			LocalCollisionComponent->RecreateCollision();
+		}
 	}
 }
 
@@ -1791,9 +1795,19 @@ void FLandscapeEditDataInterface::DeleteLayer(ULandscapeLayerInfoObject* LayerIn
 	}
 
 	// Flush dynamic data (e.g. grass)
-	TSet<ULandscapeComponent*> Components;
-	Algo::Transform(LandscapeInfo->XYtoComponentMap, Components, &TPair<FIntPoint, ULandscapeComponent*>::Value);
-	ALandscapeProxy::InvalidateGeneratedComponentData(Components);
+	if (GetMutableDefault<UEditorExperimentalSettings>()->bLandscapeLayerSystem)
+	{
+		if (LandscapeInfo->LandscapeActor)
+		{
+			LandscapeInfo->LandscapeActor->RequestLayersContentUpdate(ELandscapeLayersContentUpdateFlag::All, true);
+		}
+	}
+	else
+	{
+		TSet<ULandscapeComponent*> Components;
+		Algo::Transform(LandscapeInfo->XYtoComponentMap, Components, &TPair<FIntPoint, ULandscapeComponent*>::Value);
+		ALandscapeProxy::InvalidateGeneratedComponentData(Components);
+	}
 }
 
 void ULandscapeComponent::FillLayer(ULandscapeLayerInfoObject* LayerInfo, FLandscapeEditDataInterface& LandscapeEdit)
@@ -1840,11 +1854,6 @@ void ULandscapeComponent::FillLayer(ULandscapeLayerInfoObject* LayerInfo, FLands
 		FillLayerIdx = ComponentWeightmapLayerAllocations.Num();
 		ComponentWeightmapLayerAllocations.Add(FWeightmapLayerAllocationInfo(LayerInfo));
 		Component->ReallocateWeightmaps(&LandscapeEdit);
-
-		if (GetMutableDefault<UEditorExperimentalSettings>()->bLandscapeLayerSystem)
-		{
-			GetLandscapeActor()->RequestLayersContentUpdate(ELandscapeLayersContentUpdateFlag::Weightmap_Render);
-		}
 	}
 
 	check(FillLayerIdx != INDEX_NONE);
@@ -1952,55 +1961,50 @@ void ULandscapeComponent::FillLayer(ULandscapeLayerInfoObject* LayerInfo, FLands
 	Component->InvalidateLightingCache();
 
 	// Update dominant layer info stored in collision component
-	TArray<FColor*> CollisionWeightmapMipData;
-	for (int32 WeightmapIdx = 0; WeightmapIdx < ComponentWeightmapTextures.Num(); WeightmapIdx++)
+	if (GetMutableDefault<UEditorExperimentalSettings>()->bLandscapeLayerSystem)
 	{
-		CollisionWeightmapMipData.Add((FColor*)LandscapeEdit.GetTextureDataInfo(ComponentWeightmapTextures[WeightmapIdx])->GetMipData(Component->CollisionMipLevel));
+		GetLandscapeActor()->RequestLayersContentUpdate(ELandscapeLayersContentUpdateFlag::All);
 	}
-	TArray<FColor*> SimpleCollisionWeightmapMipData;
-	if (Component->SimpleCollisionMipLevel > Component->CollisionMipLevel)
+	else
 	{
+		TArray<FColor*> CollisionWeightmapMipData;
 		for (int32 WeightmapIdx = 0; WeightmapIdx < ComponentWeightmapTextures.Num(); WeightmapIdx++)
 		{
-			SimpleCollisionWeightmapMipData.Add((FColor*)LandscapeEdit.GetTextureDataInfo(ComponentWeightmapTextures[WeightmapIdx])->GetMipData(Component->SimpleCollisionMipLevel));
+			CollisionWeightmapMipData.Add((FColor*)LandscapeEdit.GetTextureDataInfo(ComponentWeightmapTextures[WeightmapIdx])->GetMipData(Component->CollisionMipLevel));
 		}
-	}
-	Component->UpdateCollisionLayerData(
-		CollisionWeightmapMipData.GetData(),
-		Component->SimpleCollisionMipLevel > Component->CollisionMipLevel ? SimpleCollisionWeightmapMipData.GetData() : nullptr);
+		TArray<FColor*> SimpleCollisionWeightmapMipData;
+		if (Component->SimpleCollisionMipLevel > Component->CollisionMipLevel)
+		{
+			for (int32 WeightmapIdx = 0; WeightmapIdx < ComponentWeightmapTextures.Num(); WeightmapIdx++)
+			{
+				SimpleCollisionWeightmapMipData.Add((FColor*)LandscapeEdit.GetTextureDataInfo(ComponentWeightmapTextures[WeightmapIdx])->GetMipData(Component->SimpleCollisionMipLevel));
+			}
+		}
+		Component->UpdateCollisionLayerData(
+			CollisionWeightmapMipData.GetData(),
+			Component->SimpleCollisionMipLevel > Component->CollisionMipLevel ? SimpleCollisionWeightmapMipData.GetData() : nullptr);
 
-	ULandscapeHeightfieldCollisionComponent* LocalCollisionComponent = Component->CollisionComponent.Get();
-	if (LocalCollisionComponent)
-	{
-		LocalCollisionComponent->RecreateCollision();
+		ULandscapeHeightfieldCollisionComponent* LocalCollisionComponent = Component->CollisionComponent.Get();
+		if (LocalCollisionComponent)
+		{
+			LocalCollisionComponent->RecreateCollision();
+		}
 	}
 }
 
 void FLandscapeEditDataInterface::FillLayer(ULandscapeLayerInfoObject* LayerInfo)
 {
-	if (!LandscapeInfo)
-	{
-		return;
-	}
-
-	LayerInfo->IsReferencedFromLoadedData = true;
-
-	for (auto& XYComponentPair : LandscapeInfo->XYtoComponentMap)
-	{
-		ULandscapeComponent* Component = XYComponentPair.Value;
-		Component->FillLayer(LayerInfo, *this);
-	}	
-
-	// Flush dynamic data (e.g. grass)
-	TSet<ULandscapeComponent*> Components;
-	for (auto& XYComponentPair : LandscapeInfo->XYtoComponentMap)
-	{
-		Components.Add(XYComponentPair.Value);
-	}
-	ALandscapeProxy::InvalidateGeneratedComponentData(Components);
+	const bool bEmptyLayersOnly = false;
+	FillLayer(LayerInfo, bEmptyLayersOnly);
 }
 
 void FLandscapeEditDataInterface::FillEmptyLayers(ULandscapeLayerInfoObject* LayerInfo)
+{
+	const bool bEmptyLayersOnly = true;
+	FillLayer(LayerInfo, bEmptyLayersOnly);
+}
+
+void FLandscapeEditDataInterface::FillLayer(ULandscapeLayerInfoObject* LayerInfo, bool bEmptyLayersOnly)
 {
 	if (!LandscapeInfo)
 	{
@@ -2016,12 +2020,15 @@ void FLandscapeEditDataInterface::FillEmptyLayers(ULandscapeLayerInfoObject* Lay
 
 		bool LayerNeedFilling = true;
 
-		for (const FWeightmapLayerAllocationInfo& Alloc : ComponentWeightmapLayerAllocations)
+		if (bEmptyLayersOnly)
 		{
-			if (Alloc.LayerInfo != nullptr)
+			for (const FWeightmapLayerAllocationInfo& Alloc : ComponentWeightmapLayerAllocations)
 			{
-				LayerNeedFilling = false;
-				break;
+				if (Alloc.LayerInfo != nullptr)
+				{
+					LayerNeedFilling = false;
+					break;
+				}
 			}
 		}
 
@@ -2031,13 +2038,23 @@ void FLandscapeEditDataInterface::FillEmptyLayers(ULandscapeLayerInfoObject* Lay
 		}
 	}
 
-	// Flush dynamic data (e.g. grass)
-	TSet<ULandscapeComponent*> Components;
-	for (auto& XYComponentPair : LandscapeInfo->XYtoComponentMap)
+	if (GetMutableDefault<UEditorExperimentalSettings>()->bLandscapeLayerSystem)
 	{
-		Components.Add(XYComponentPair.Value);
+		if (LandscapeInfo->LandscapeActor)
+		{
+			LandscapeInfo->LandscapeActor->RequestLayersContentUpdate(ELandscapeLayersContentUpdateFlag::All, true);
+		}
+	}
+	else
+	{
+		// Flush dynamic data (e.g. grass)
+		TSet<ULandscapeComponent*> Components;
+		for (auto& XYComponentPair : LandscapeInfo->XYtoComponentMap)
+		{
+			Components.Add(XYComponentPair.Value);
+		}
+		ALandscapeProxy::InvalidateGeneratedComponentData(Components);
 	}
-	ALandscapeProxy::InvalidateGeneratedComponentData(Components);
 }
 
 void ULandscapeComponent::ReplaceLayer(ULandscapeLayerInfoObject* FromLayerInfo, ULandscapeLayerInfoObject* ToLayerInfo, FLandscapeEditDataInterface& LandscapeEdit)
@@ -2241,28 +2258,29 @@ void FLandscapeEditDataInterface::ReplaceLayer(ULandscapeLayerInfoObject* FromLa
 			ULandscapeComponent* Component = It.Value();
 			Component->ReplaceLayer(FromLayerInfo, ToLayerInfo, *this);
 
-			TArray<UTexture2D*>& ComponentWeightmapTextures = Component->GetWeightmapTextures(true);
-
-			// Update dominant layer info stored in collision component
-			TArray<FColor*> CollisionWeightmapMipData;
-			for (int32 WeightmapIdx = 0; WeightmapIdx < ComponentWeightmapTextures.Num(); WeightmapIdx++)
-			{
-				CollisionWeightmapMipData.Add((FColor*)GetTextureDataInfo(ComponentWeightmapTextures[WeightmapIdx])->GetMipData(Component->CollisionMipLevel));
-			}
-			TArray<FColor*> SimpleCollisionWeightmapMipData;
-			if (Component->SimpleCollisionMipLevel > Component->CollisionMipLevel)
+			if(!GetMutableDefault<UEditorExperimentalSettings>()->bLandscapeLayerSystem)
 			{
+				TArray<UTexture2D*>& ComponentWeightmapTextures = Component->GetWeightmapTextures(true);
+
+				// Update dominant layer info stored in collision component
+				TArray<FColor*> CollisionWeightmapMipData;
 				for (int32 WeightmapIdx = 0; WeightmapIdx < ComponentWeightmapTextures.Num(); WeightmapIdx++)
 				{
-					SimpleCollisionWeightmapMipData.Add((FColor*)GetTextureDataInfo(ComponentWeightmapTextures[WeightmapIdx])->GetMipData(Component->SimpleCollisionMipLevel));
+					CollisionWeightmapMipData.Add((FColor*)GetTextureDataInfo(ComponentWeightmapTextures[WeightmapIdx])->GetMipData(Component->CollisionMipLevel));
 				}
-			}
-			Component->UpdateCollisionLayerData(
-				CollisionWeightmapMipData.GetData(),
-				Component->SimpleCollisionMipLevel > Component->CollisionMipLevel ? SimpleCollisionWeightmapMipData.GetData() : nullptr);
+				TArray<FColor*> SimpleCollisionWeightmapMipData;
+				if (Component->SimpleCollisionMipLevel > Component->CollisionMipLevel)
+				{
+					for (int32 WeightmapIdx = 0; WeightmapIdx < ComponentWeightmapTextures.Num(); WeightmapIdx++)
+					{
+						SimpleCollisionWeightmapMipData.Add((FColor*)GetTextureDataInfo(ComponentWeightmapTextures[WeightmapIdx])->GetMipData(Component->SimpleCollisionMipLevel));
+					}
+				}
+				Component->UpdateCollisionLayerData(
+					CollisionWeightmapMipData.GetData(),
+					Component->SimpleCollisionMipLevel > Component->CollisionMipLevel ? SimpleCollisionWeightmapMipData.GetData() : nullptr);
+
 
-			if (!GetMutableDefault<UEditorExperimentalSettings>()->bLandscapeLayerSystem)
-			{
 				ULandscapeHeightfieldCollisionComponent* CollisionComponent = Component->CollisionComponent.Get();
 				if (CollisionComponent)
 				{
@@ -3046,42 +3064,45 @@ void FLandscapeEditDataInterface::SetAlphaData(ULandscapeLayerInfoObject* const
 				}
 			}
 
-			// Update mipmaps
-			CollisionWeightmapMipData.Reset();
-			CollisionWeightmapMipData.AddUninitialized(ComponentWeightmapTextures.Num());
-			for (int32 WeightmapIdx = 0; WeightmapIdx < ComponentWeightmapTextures.Num(); WeightmapIdx++)
+			if (!GetMutableDefault<UEditorExperimentalSettings>()->bLandscapeLayerSystem)
 			{
-				UTexture2D* const WeightmapTexture = ComponentWeightmapTextures[WeightmapIdx];
-
-				const int32 NumMips = WeightmapTexture->Source.GetNumMips();
-				WeightmapTextureMipData.Reset();
-				WeightmapTextureMipData.AddUninitialized(NumMips);
-				for (int32 MipIdx = 0; MipIdx < NumMips; MipIdx++)
+				// Update mipmaps
+				CollisionWeightmapMipData.Reset();
+				CollisionWeightmapMipData.AddUninitialized(ComponentWeightmapTextures.Num());
+				for (int32 WeightmapIdx = 0; WeightmapIdx < ComponentWeightmapTextures.Num(); WeightmapIdx++)
 				{
-					FColor* const MipData = (FColor*)TexDataInfos[WeightmapIdx]->GetMipData(MipIdx);
-					WeightmapTextureMipData[MipIdx] = MipData;
-				}
-				CollisionWeightmapMipData[WeightmapIdx] = WeightmapTextureMipData[Component->CollisionMipLevel];
+					UTexture2D* const WeightmapTexture = ComponentWeightmapTextures[WeightmapIdx];
 
-				ULandscapeComponent::UpdateWeightmapMips(ComponentNumSubsections, SubsectionSizeQuads, WeightmapTexture, WeightmapTextureMipData, ComponentX1, ComponentY1, ComponentX2, ComponentY2, TexDataInfos[WeightmapIdx]);
-				WeightmapTextureMipData.Reset();
-			}
+					const int32 NumMips = WeightmapTexture->Source.GetNumMips();
+					WeightmapTextureMipData.Reset();
+					WeightmapTextureMipData.AddUninitialized(NumMips);
+					for (int32 MipIdx = 0; MipIdx < NumMips; MipIdx++)
+					{
+						FColor* const MipData = (FColor*)TexDataInfos[WeightmapIdx]->GetMipData(MipIdx);
+						WeightmapTextureMipData[MipIdx] = MipData;
+					}
+					CollisionWeightmapMipData[WeightmapIdx] = WeightmapTextureMipData[Component->CollisionMipLevel];
 
-			if (Component->SimpleCollisionMipLevel > Component->CollisionMipLevel)
-			{
-				for (int32 WeightmapIdx = 0; WeightmapIdx < ComponentWeightmapTextures.Num(); WeightmapIdx++)
+					ULandscapeComponent::UpdateWeightmapMips(ComponentNumSubsections, SubsectionSizeQuads, WeightmapTexture, WeightmapTextureMipData, ComponentX1, ComponentY1, ComponentX2, ComponentY2, TexDataInfos[WeightmapIdx]);
+					WeightmapTextureMipData.Reset();
+				}
+
+				if (Component->SimpleCollisionMipLevel > Component->CollisionMipLevel)
 				{
-					SimpleCollisionWeightmapMipData.Add((FColor*)TexDataInfos[WeightmapIdx]->GetMipData(Component->SimpleCollisionMipLevel));
+					for (int32 WeightmapIdx = 0; WeightmapIdx < ComponentWeightmapTextures.Num(); WeightmapIdx++)
+					{
+						SimpleCollisionWeightmapMipData.Add((FColor*)TexDataInfos[WeightmapIdx]->GetMipData(Component->SimpleCollisionMipLevel));
+					}
 				}
-			}
 
-			// Update dominant layer info stored in collision component
-			Component->UpdateCollisionLayerData(
-				CollisionWeightmapMipData.GetData(),
-				Component->SimpleCollisionMipLevel > Component->CollisionMipLevel ? SimpleCollisionWeightmapMipData.GetData() : nullptr,
-				ComponentX1, ComponentY1, ComponentX2, ComponentY2);
-			CollisionWeightmapMipData.Reset();
-			SimpleCollisionWeightmapMipData.Reset();
+				// Update dominant layer info stored in collision component
+				Component->UpdateCollisionLayerData(
+					CollisionWeightmapMipData.GetData(),
+					Component->SimpleCollisionMipLevel > Component->CollisionMipLevel ? SimpleCollisionWeightmapMipData.GetData() : nullptr,
+					ComponentX1, ComponentY1, ComponentX2, ComponentY2);
+				CollisionWeightmapMipData.Reset();
+				SimpleCollisionWeightmapMipData.Reset();
+			}
 
 			// Check if we need to remove weightmap allocations for layers that were completely painted away
 			bool bRemovedLayer = false;
@@ -3325,42 +3346,45 @@ void FLandscapeEditDataInterface::SetAlphaData(const TSet<ULandscapeLayerInfoObj
 				}
 			}
 
-			// Update mipmaps
-			CollisionWeightmapMipData.Reset();
-			CollisionWeightmapMipData.AddUninitialized(ComponentWeightmapTextures.Num());
-			for (int32 WeightmapIdx = 0; WeightmapIdx < ComponentWeightmapTextures.Num(); WeightmapIdx++)
+			if (!GetMutableDefault<UEditorExperimentalSettings>()->bLandscapeLayerSystem)
 			{
-				UTexture2D* const WeightmapTexture = ComponentWeightmapTextures[WeightmapIdx];
-
-				const int32 NumMips = WeightmapTexture->Source.GetNumMips();
-				WeightmapTextureMipData.Reset();
-				WeightmapTextureMipData.AddUninitialized(NumMips);
-				for (int32 MipIdx = 0; MipIdx < NumMips; MipIdx++)
+				// Update mipmaps
+				CollisionWeightmapMipData.Reset();
+				CollisionWeightmapMipData.AddUninitialized(ComponentWeightmapTextures.Num());
+				for (int32 WeightmapIdx = 0; WeightmapIdx < ComponentWeightmapTextures.Num(); WeightmapIdx++)
 				{
-					FColor* const MipData = (FColor*)TexDataInfos[WeightmapIdx]->GetMipData(MipIdx);
-					WeightmapTextureMipData[MipIdx] = MipData;
-				}
-				CollisionWeightmapMipData[WeightmapIdx] = WeightmapTextureMipData[Component->CollisionMipLevel];
+					UTexture2D* const WeightmapTexture = ComponentWeightmapTextures[WeightmapIdx];
 
-				ULandscapeComponent::UpdateWeightmapMips(ComponentNumSubsections, SubsectionSizeQuads, WeightmapTexture, WeightmapTextureMipData, ComponentX1, ComponentY1, ComponentX2, ComponentY2, TexDataInfos[WeightmapIdx]);
-				WeightmapTextureMipData.Reset();
-			}
+					const int32 NumMips = WeightmapTexture->Source.GetNumMips();
+					WeightmapTextureMipData.Reset();
+					WeightmapTextureMipData.AddUninitialized(NumMips);
+					for (int32 MipIdx = 0; MipIdx < NumMips; MipIdx++)
+					{
+						FColor* const MipData = (FColor*)TexDataInfos[WeightmapIdx]->GetMipData(MipIdx);
+						WeightmapTextureMipData[MipIdx] = MipData;
+					}
+					CollisionWeightmapMipData[WeightmapIdx] = WeightmapTextureMipData[Component->CollisionMipLevel];
 
-			if (Component->SimpleCollisionMipLevel > Component->CollisionMipLevel)
-			{
-				for (int32 WeightmapIdx = 0; WeightmapIdx < ComponentWeightmapTextures.Num(); WeightmapIdx++)
+					ULandscapeComponent::UpdateWeightmapMips(ComponentNumSubsections, SubsectionSizeQuads, WeightmapTexture, WeightmapTextureMipData, ComponentX1, ComponentY1, ComponentX2, ComponentY2, TexDataInfos[WeightmapIdx]);
+					WeightmapTextureMipData.Reset();
+				}
+
+				if (Component->SimpleCollisionMipLevel > Component->CollisionMipLevel)
 				{
-					SimpleCollisionWeightmapMipData.Add((FColor*)TexDataInfos[WeightmapIdx]->GetMipData(Component->SimpleCollisionMipLevel));
+					for (int32 WeightmapIdx = 0; WeightmapIdx < ComponentWeightmapTextures.Num(); WeightmapIdx++)
+					{
+						SimpleCollisionWeightmapMipData.Add((FColor*)TexDataInfos[WeightmapIdx]->GetMipData(Component->SimpleCollisionMipLevel));
+					}
 				}
-			}
 
-			// Update dominant layer info stored in collision component
-			Component->UpdateCollisionLayerData(
-				CollisionWeightmapMipData.GetData(),
-				Component->SimpleCollisionMipLevel > Component->CollisionMipLevel ? SimpleCollisionWeightmapMipData.GetData() : nullptr,
-				ComponentX1, ComponentY1, ComponentX2, ComponentY2);
-			CollisionWeightmapMipData.Reset();
-			SimpleCollisionWeightmapMipData.Reset();
+				// Update dominant layer info stored in collision component
+				Component->UpdateCollisionLayerData(
+					CollisionWeightmapMipData.GetData(),
+					Component->SimpleCollisionMipLevel > Component->CollisionMipLevel ? SimpleCollisionWeightmapMipData.GetData() : nullptr,
+					ComponentX1, ComponentY1, ComponentX2, ComponentY2);
+				CollisionWeightmapMipData.Reset();
+				SimpleCollisionWeightmapMipData.Reset();
+			}
 
 			// Check if we need to remove weightmap allocations for layers that were completely painted away
 			bool bRemovedLayer = false;
diff --git a/Engine/Source/Runtime/Landscape/Private/LandscapeEditLayers.cpp b/Engine/Source/Runtime/Landscape/Private/LandscapeEditLayers.cpp
--- a/Engine/Source/Runtime/Landscape/Private/LandscapeEditLayers.cpp
+++ b/Engine/Source/Runtime/Landscape/Private/LandscapeEditLayers.cpp
@@ -2446,9 +2446,10 @@ bool ALandscape::PrepareLayersHeightmapTextureResources() const
 	return IsReady;
 }
 
-void ALandscape::RegenerateLayersHeightmaps()
+void ALandscape::RegenerateLayersHeightmaps(const TArray<ULandscapeComponent*>& InLandscapeComponents, bool& bOutRecreateCollision)
 {
 	SCOPE_CYCLE_COUNTER(STAT_LandscapeLayersRegenerateHeightmaps);
+	bOutRecreateCollision = false;
 
 	ULandscapeInfo* Info = GetLandscapeInfo();
 
@@ -2457,13 +2458,6 @@ void ALandscape::RegenerateLayersHeightmaps()
 		return;
 	}
 
-	TArray<ULandscapeComponent*> AllLandscapeComponents;
-
-	Info->ForAllLandscapeProxies([&AllLandscapeComponents](ALandscapeProxy* Proxy)
-	{
-		AllLandscapeComponents.Append(Proxy->LandscapeComponents);
-	});
-
 	// Handle missing Heightmap CPUReadback (Undo of a delete landscape component can trigger that case)
 	if (LayersContentUpdateFlags & ELandscapeLayersContentUpdateFlag::Heightmap_All)
 	{
@@ -2483,7 +2477,7 @@ void ALandscape::RegenerateLayersHeightmaps()
 		});
 	}
 
-	if ((LayersContentUpdateFlags & ELandscapeLayersContentUpdateFlag::Heightmap_Render) != 0)
+	if (LayersContentUpdateFlags & ELandscapeLayersContentUpdateFlag::Heightmap_Render)
 	{
 		check(HeightmapRTList.Num() > 0);
 
@@ -2575,11 +2569,11 @@ void ALandscape::RegenerateLayersHeightmaps()
 
 			// NOTE: From this point on, we always work in non atlas, we'll convert back at the end to atlas only
 			DrawHeightmapComponentsToRenderTarget(OutputDebugName ? FString::Printf(TEXT("LS Height: %s += -> NonAtlas %s"), *Layer.Name.ToString(), *LandscapeScratchRT1->GetName(), *LandscapeScratchRT2->GetName()) : TEXT(""),
-												  AllLandscapeComponents, MinExtend, LandscapeScratchRT1, nullptr, LandscapeScratchRT2, ERTDrawingType::RTAtlasToNonAtlas, true, ShaderParams);
+												  InLandscapeComponents, MinExtend, LandscapeScratchRT1, nullptr, LandscapeScratchRT2, ERTDrawingType::RTAtlasToNonAtlas, true, ShaderParams);
 
 			// Combine Current layer with current result
 			DrawHeightmapComponentsToRenderTarget(OutputDebugName ? FString::Printf(TEXT("LS Height: %s += -> CombinedNonAtlas %s"), *Layer.Name.ToString(), *LandscapeScratchRT2->GetName(), *CombinedHeightmapNonAtlasRT->GetName()) : TEXT(""),
-												AllLandscapeComponents, MinExtend, LandscapeScratchRT2, FirstLayer ? nullptr : LandscapeScratchRT3, CombinedHeightmapNonAtlasRT, ERTDrawingType::RTNonAtlas, FirstLayer, ShaderParams);
+												InLandscapeComponents, MinExtend, LandscapeScratchRT2, FirstLayer ? nullptr : LandscapeScratchRT3, CombinedHeightmapNonAtlasRT, ERTDrawingType::RTNonAtlas, FirstLayer, ShaderParams);
 
 			ShaderParams.ApplyLayerModifiers = false;
 
@@ -2637,14 +2631,14 @@ void ALandscape::RegenerateLayersHeightmaps()
 		ShaderParams.GridSize = GetRootComponent()->RelativeScale3D;
 
 		DrawHeightmapComponentsToRenderTarget(OutputDebugName ? FString::Printf(TEXT("LS Height: %s = -> CombinedNonAtlasNormals : %s"), *CombinedHeightmapNonAtlasRT->GetName(), *LandscapeScratchRT1->GetName()) : TEXT(""),
-											  AllLandscapeComponents, MinExtend, CombinedHeightmapNonAtlasRT, nullptr, LandscapeScratchRT1, ERTDrawingType::RTNonAtlas, true, ShaderParams);
+											  InLandscapeComponents, MinExtend, CombinedHeightmapNonAtlasRT, nullptr, LandscapeScratchRT1, ERTDrawingType::RTNonAtlas, true, ShaderParams);
 
 		ShaderParams.GenerateNormals = false;
 
 		DrawHeightmapComponentsToRenderTarget(OutputDebugName ? FString::Printf(TEXT("LS Height: %s = -> CombinedAtlasFinal : %s"), *LandscapeScratchRT1->GetName(), *CombinedHeightmapAtlasRT->GetName()) : TEXT(""),
-											  AllLandscapeComponents, MinExtend, LandscapeScratchRT1, nullptr, CombinedHeightmapAtlasRT, ERTDrawingType::RTNonAtlasToAtlas, true, ShaderParams);
+											  InLandscapeComponents, MinExtend, LandscapeScratchRT1, nullptr, CombinedHeightmapAtlasRT, ERTDrawingType::RTNonAtlasToAtlas, true, ShaderParams);
 
-		DrawHeightmapComponentsToRenderTargetMips(AllLandscapeComponents, MinExtend, CombinedHeightmapAtlasRT, true, ShaderParams);
+		DrawHeightmapComponentsToRenderTargetMips(InLandscapeComponents, MinExtend, CombinedHeightmapAtlasRT, true, ShaderParams);
 
 		// Copy back all Mips to original heightmap data
 		Info->ForAllLandscapeProxies([&](ALandscapeProxy* Proxy)
@@ -2676,7 +2670,7 @@ void ALandscape::RegenerateLayersHeightmaps()
 		});
 	}
 
-	if ((LayersContentUpdateFlags & EInternalLandscapeLayersContentUpdateFlag::Internal_Heightmap_ResolveToTexture) != 0)
+	if (LayersContentUpdateFlags & EInternalLandscapeLayersContentUpdateFlag::Internal_Heightmap_ResolveToTexture)
 	{
 		ResolveLayersHeightmapTexture();
 	}
@@ -2684,48 +2678,30 @@ void ALandscape::RegenerateLayersHeightmaps()
 	bool bUpdateBoundsAndCollsionFull = (LayersContentUpdateFlags & EInternalLandscapeLayersContentUpdateFlag::Internal_Heightmap_BoundsAndCollision) != 0;
 	bool bUpdateBoundsAndCollsionPartial = ((LayersContentUpdateFlags & EInternalLandscapeLayersContentUpdateFlag::Internal_Heightmap_BoundsAndCollisionPartial) != 0) && !bUpdateBoundsAndCollsionFull; // Partial only if not full update
 	bool bUpdateBoundsAndCollsion = bUpdateBoundsAndCollsionFull || bUpdateBoundsAndCollsionPartial;
-	bool bUpdateClients = (LayersContentUpdateFlags & EInternalLandscapeLayersContentUpdateFlag::Internal_Heightmap_ClientUpdate) != 0;
-	int32 PostponeUpdateClientFlag = 0;
-	if (bUpdateBoundsAndCollsion || bUpdateClients)
+
+	if (bUpdateBoundsAndCollsion)
 	{
-		for (ULandscapeComponent* Component : AllLandscapeComponents)
+		for (ULandscapeComponent* Component : InLandscapeComponents)
 		{
 			if (bUpdateBoundsAndCollsion && (!bUpdateBoundsAndCollsionPartial || Component->IsLayerContentDirty()))
 			{
 				Component->UpdateCachedBounds();
 				Component->UpdateComponentToWorld();
-				Component->UpdateCollisionData(bUpdateBoundsAndCollsionPartial);
+				Component->UpdateCollisionData();
 			}
 
-			if (bUpdateClients)
-			{
-				// Post-pone client update to next tick while in transaction
-				if (GUndo)
-				{
-					PostponeUpdateClientFlag |= EInternalLandscapeLayersContentUpdateFlag::Internal_Heightmap_ClientUpdate;
-				}
-				else
-				{
-					// Recreate collision for modified components to update the physical materials
-					ULandscapeHeightfieldCollisionComponent* CollisionComponent = Component->CollisionComponent.Get();
-					if (CollisionComponent)
-					{
-						FNavigationSystem::UpdateComponentData(*CollisionComponent);
-						CollisionComponent->SnapFoliageInstances();
-					}
-				}
-			}
 			Component->SetLayerContentDirty(false);
 		}
 
 		// Ideally, we should only update necessary components and only when inside AddComponnet Tool.
 		Info->UpdateAllAddCollisions();
+	
+		bOutRecreateCollision = !bUpdateBoundsAndCollsionPartial;
 	}
-
+		
 	LayersContentUpdateFlags &= ~EInternalLandscapeLayersContentUpdateFlag::Internal_Heightmap_BoundsAndCollisionPartial;
 	LayersContentUpdateFlags &= ~ELandscapeLayersContentUpdateFlag::Heightmap_All;
-	LayersContentUpdateFlags |= PostponeUpdateClientFlag; // if non-zero postpone client update (transacting)
-
+		
 	// If doing rendering debug, keep doing the render only
 	if (CVarOutputLayersDebugDrawCallName.GetValueOnAnyThread() == 1)
 	{
@@ -2733,6 +2709,33 @@ void ALandscape::RegenerateLayersHeightmaps()
 	}
 }
 
+bool ALandscape::UpdateClients(const TArray<ULandscapeComponent*>& InLandscapeComponents, bool bInCollisionRecreated)
+{
+	// Post-pone client update to next tick while in transaction
+	if (GUndo)
+	{
+		return false;
+	}
+	else
+	{
+		for (ULandscapeComponent* Component : InLandscapeComponents)
+		{
+			// Recreate collision for modified components to update the physical materials
+			ULandscapeHeightfieldCollisionComponent* CollisionComponent = Component->CollisionComponent.Get();
+			if (CollisionComponent)
+			{
+				FNavigationSystem::UpdateComponentData(*CollisionComponent);
+				CollisionComponent->SnapFoliageInstances();
+			}
+		}
+
+		ALandscapeProxy::InvalidateGeneratedComponentData(InLandscapeComponents);
+	}
+	
+	// Re update clients if collision wasn't recreated
+	return bInCollisionRecreated;
+}
+
 void ALandscape::ResolveLayersHeightmapTexture()
 {
 	SCOPE_CYCLE_COUNTER(STAT_LandscapeLayersResolveHeightmaps);
@@ -3370,30 +3373,25 @@ bool ALandscape::PrepareLayersWeightmapTextureResources() const
 	return IsReady;
 }
 
-void ALandscape::RegenerateLayersWeightmaps()
+void ALandscape::RegenerateLayersWeightmaps(const TArray<ULandscapeComponent*>& InLandscapeComponents, bool& bOutRecreateCollision)
 {
 	SCOPE_CYCLE_COUNTER(STAT_LandscapeLayersRegenerateWeightmaps);
-
+	bOutRecreateCollision = false;
+    
 	ULandscapeInfo* Info = GetLandscapeInfo();
 
 	if (LayersContentUpdateFlags == 0 || Info == nullptr || Info->Layers.Num() == 0 || !PrepareLayersWeightmapTextureResources())
 	{
 		return;
 	}
-
-	TArray<ULandscapeComponent*> AllLandscapeComponents;
-
-	Info->ForAllLandscapeProxies([&AllLandscapeComponents](ALandscapeProxy* Proxy)
-	{
-		AllLandscapeComponents.Append(Proxy->LandscapeComponents);
-	});
-
+	
+		
 	TArray<ULandscapeComponent*> ComponentThatNeedMaterialRebuild;
 	TArray<ULandscapeLayerInfoObject*> BrushRequiredAllocations;
 	int32 LayerCount = Info->Layers.Num() + 1; // due to visibility being stored at 0
 	bool ClearFlagsAfterUpdate = true;
 
-	if ((LayersContentUpdateFlags & ELandscapeLayersContentUpdateFlag::Weightmap_Render) != 0)
+	if (LayersContentUpdateFlags & ELandscapeLayersContentUpdateFlag::Weightmap_Render)
 	{
 		FIntPoint MinExtend;
 		FIntPoint MaxExtend;
@@ -3522,7 +3520,7 @@ void ALandscape::RegenerateLayersWeightmaps()
 					PSShaderParams.LayerAlpha = LayerInfoObj == ALandscapeProxy::VisibilityLayer ? 1.0f : Layer.WeightmapAlpha; // visibility can't be affected by weight
 
 					DrawWeightmapComponentsToRenderTarget(OutputDebugName ? FString::Printf(TEXT("LS Weight: %s Paint: %s += -> %s"), *Layer.Name.ToString(), *LayerInfoObj->LayerName.ToString(), *LandscapeScratchRT1->GetName(), *LandscapeScratchRT2->GetName()) : TEXT(""),
-														  AllLandscapeComponents, MinExtend, LandscapeScratchRT1, nullptr, LandscapeScratchRT2, true, PSShaderParams, 0);
+														  InLandscapeComponents, MinExtend, LandscapeScratchRT1, nullptr, LandscapeScratchRT2, true, PSShaderParams, 0);
 
 					PSShaderParams.ApplyLayerModifiers = false;
 
@@ -3546,7 +3544,7 @@ void ALandscape::RegenerateLayersWeightmaps()
 					}
 
 					DrawWeightmapComponentsToRenderTarget(OutputDebugName ? FString::Printf(TEXT("LS Weight: %s PaintLayer: %s, %s += -> Combined %s"), *Layer.Name.ToString(), *LayerInfoObj->LayerName.ToString(), *LandscapeScratchRT2->GetName(), *LandscapeScratchRT3->GetName()) : TEXT(""),
-														  AllLandscapeComponents, MinExtend, LandscapeScratchRT2, FirstLayer ? nullptr : LandscapeScratchRT1, LandscapeScratchRT3, true, PSShaderParams, 0);
+														  InLandscapeComponents, MinExtend, LandscapeScratchRT2, FirstLayer ? nullptr : LandscapeScratchRT1, LandscapeScratchRT3, true, PSShaderParams, 0);
 
 					PSShaderParams.OutputAsSubstractive = false;
 
@@ -3639,7 +3637,7 @@ void ALandscape::RegenerateLayersWeightmaps()
 			while (HasFoundWeightmapToProcess)
 			{
 				TArray<FLandscapeLayerWeightmapPackMaterialLayersComponentData> PackLayersComponentsData;
-				PrepareComponentDataToPackMaterialLayersCS(CurrentWeightmapToProcessIndex, MinExtend, OutputDebugName, AllLandscapeComponents, ProcessedWeightmaps, ProcessedCPUReadbackTextures, PackLayersComponentsData);
+				PrepareComponentDataToPackMaterialLayersCS(CurrentWeightmapToProcessIndex, MinExtend, OutputDebugName, InLandscapeComponents, ProcessedWeightmaps, ProcessedCPUReadbackTextures, PackLayersComponentsData);
 				HasFoundWeightmapToProcess = PackLayersComponentsData.Num() > 0;
 
 				// Perform the compute shader
@@ -3729,23 +3727,25 @@ void ALandscape::RegenerateLayersWeightmaps()
 			}
 		}
 
-		UpdateLayersMaterialInstances(LayersUpdateAllMaterials ? AllLandscapeComponents : ComponentThatNeedMaterialRebuild);
+		UpdateLayersMaterialInstances(LayersUpdateAllMaterials ? InLandscapeComponents : ComponentThatNeedMaterialRebuild);
 		LayersUpdateAllMaterials = false;
 	}
 
-	if ((LayersContentUpdateFlags & EInternalLandscapeLayersContentUpdateFlag::Internal_Weightmap_ResolveToTexture))
+	if (LayersContentUpdateFlags & EInternalLandscapeLayersContentUpdateFlag::Internal_Weightmap_ResolveToTexture)
 	{
 		ResolveLayersWeightmapTexture();
 	}
 
-	if ((LayersContentUpdateFlags & EInternalLandscapeLayersContentUpdateFlag::Internal_Weightmap_Collision) != 0)
+	if (LayersContentUpdateFlags & EInternalLandscapeLayersContentUpdateFlag::Internal_Weightmap_Collision)
 	{
-		for (ULandscapeComponent* Component : AllLandscapeComponents)
+		for (ULandscapeComponent* Component : InLandscapeComponents)
 		{
 			Component->UpdateCollisionLayerData();
 		}
-	}
 
+		bOutRecreateCollision = true;
+	}
+		
 	if (ClearFlagsAfterUpdate)
 	{
 		LayersContentUpdateFlags &= ~ELandscapeLayersContentUpdateFlag::Weightmap_All;
@@ -3965,8 +3965,42 @@ void ALandscape::RequestLayersContentUpdate(ELandscapeLayersContentUpdateFlag In
 
 void ALandscape::RegenerateLayersContent()
 {
-	RegenerateLayersHeightmaps();
-	RegenerateLayersWeightmaps();
+	if (GetLandscapeInfo() == nullptr || LayersContentUpdateFlags == 0)
+	{
+		return;
+	}
+	
+	TArray<ULandscapeComponent*> AllLandscapeComponents;
+	GetLandscapeInfo()->ForAllLandscapeProxies([&AllLandscapeComponents](ALandscapeProxy* Proxy)
+	{
+		AllLandscapeComponents.Append(Proxy->LandscapeComponents);
+	});
+
+	bool bUpdateClients = (LayersContentUpdateFlags & EInternalLandscapeLayersContentUpdateFlag::Internal_ClientUpdate);
+
+	bool bHeightmapRecreateCollision = false;
+	RegenerateLayersHeightmaps(AllLandscapeComponents, bHeightmapRecreateCollision);
+	bool bWeightmapRecreateCollision = false;
+	RegenerateLayersWeightmaps(AllLandscapeComponents, bWeightmapRecreateCollision);
+	
+	if (bHeightmapRecreateCollision || bWeightmapRecreateCollision)
+	{
+		for (ULandscapeComponent* LandscapeComponent : AllLandscapeComponents)
+		{
+			if (ULandscapeHeightfieldCollisionComponent* CollisionComp = LandscapeComponent->CollisionComponent.Get())
+			{
+				CollisionComp->RecreateCollision();
+			}
+		}
+	}
+
+	if (bUpdateClients)
+	{
+		if (!UpdateClients(AllLandscapeComponents, bHeightmapRecreateCollision | bWeightmapRecreateCollision))
+		{
+			LayersContentUpdateFlags |= EInternalLandscapeLayersContentUpdateFlag::Internal_ClientUpdate;
+		}
+	}
 }
 
 void ALandscape::InitializeLayers()
@@ -4413,14 +4447,6 @@ void ALandscape::SetLayerAlpha(int32 InLayerIndex, const float InAlpha, bool bIn
 	Modify();
 	LayerAlpha = InAlpha;
 	RequestLayersContentUpdate(ELandscapeLayersContentUpdateFlag::All, true);
-
-	LandscapeInfo->ForAllLandscapeProxies([](ALandscapeProxy* Proxy)
-	{
-		if (Proxy)
-		{
-			Proxy->InvalidateGeneratedComponentData();
-		}
-	});
 }
 
 void ALandscape::SetLayerVisibility(int32 InLayerIndex, bool bInVisible)
@@ -4435,14 +4461,6 @@ void ALandscape::SetLayerVisibility(int32 InLayerIndex, bool bInVisible)
 	Modify();
 	Layer->bVisible = bInVisible;
 	RequestLayersContentUpdate(ELandscapeLayersContentUpdateFlag::All, true);
-
-	LandscapeInfo->ForAllLandscapeProxies([](ALandscapeProxy* Proxy)
-	{
-		if (Proxy)
-		{
-			Proxy->InvalidateGeneratedComponentData();
-		}
-	});
 }
 
 void ALandscape::SetLayerLocked(int32 InLayerIndex, bool bLocked)
diff --git a/Engine/Source/Runtime/Landscape/Private/LandscapeSplineRaster.cpp b/Engine/Source/Runtime/Landscape/Private/LandscapeSplineRaster.cpp
--- a/Engine/Source/Runtime/Landscape/Private/LandscapeSplineRaster.cpp
+++ b/Engine/Source/Runtime/Landscape/Private/LandscapeSplineRaster.cpp
@@ -625,10 +625,9 @@ bool ULandscapeInfo::ApplySplinesInternal(bool bOnlySelected, ALandscapeProxy* P
 
 	LandscapeEdit.Flush();
 
-	ALandscape* Landscape = Proxy->GetLandscapeActor();
-	bool bSkipCollisionAndNavUpdate = Landscape && Landscape->IsEditingLayerReservedForSplines();
-	if (!bSkipCollisionAndNavUpdate)
+	if (!GetMutableDefault<UEditorExperimentalSettings>()->bLandscapeLayerSystem)
 	{
+		ALandscape* Landscape = Proxy->GetLandscapeActor();
 		for (ULandscapeComponent* Component : ModifiedComponents)
 		{
 			// Recreate collision for modified components and update the navmesh
@@ -639,11 +638,11 @@ bool ULandscapeInfo::ApplySplinesInternal(bool bOnlySelected, ALandscapeProxy* P
 				FNavigationSystem::UpdateComponentData(*CollisionComponent);
 			}
 		}
+		
+		// Invalidate landscape grass of modified landscape components
+		ALandscapeProxy::InvalidateGeneratedComponentData(ModifiedComponents);
 	}
 
-	// Invalidate landscape grass of modified landscape components
-	ALandscapeProxy::InvalidateGeneratedComponentData(ModifiedComponents);
-
 	return true;
 }
 
@@ -724,14 +723,17 @@ namespace LandscapeSplineRaster
 
 		LandscapeEdit.Flush();
 
-		for (ULandscapeComponent* Component : ModifiedComponents)
+		if (!GetMutableDefault<UEditorExperimentalSettings>()->bLandscapeLayerSystem)
 		{
-			// Recreate collision for modified components and update the navmesh
-			ULandscapeHeightfieldCollisionComponent* CollisionComponent = Component->CollisionComponent.Get();
-			if (CollisionComponent)
+			for (ULandscapeComponent* Component : ModifiedComponents)
 			{
-				CollisionComponent->RecreateCollision();
-				FNavigationSystem::UpdateComponentData(*CollisionComponent);
+				// Recreate collision for modified components and update the navmesh
+				ULandscapeHeightfieldCollisionComponent* CollisionComponent = Component->CollisionComponent.Get();
+				if (CollisionComponent)
+				{
+					CollisionComponent->RecreateCollision();
+					FNavigationSystem::UpdateComponentData(*CollisionComponent);
+				}
 			}
 		}
 	}
diff --git a/Engine/Source/Runtime/Landscape/Public/LandscapeEdit.h b/Engine/Source/Runtime/Landscape/Public/LandscapeEdit.h
--- a/Engine/Source/Runtime/Landscape/Public/LandscapeEdit.h
+++ b/Engine/Source/Runtime/Landscape/Public/LandscapeEdit.h
@@ -230,6 +230,8 @@ private:
 
 	ULandscapeInfo* LandscapeInfo;
 
+	void FillLayer(ULandscapeLayerInfoObject* LayerInfo, bool bEmptyLayersOnly);
+
 	// Only for Missing Data interpolation... only internal usage
 	template<typename TData, typename TStoreData, typename FType>
 	FORCEINLINE void CalcMissingValues(const int32& X1, const int32& X2, const int32& Y1, const int32& Y2,
@@ -318,16 +320,16 @@ struct FHeightmapAccessor
 			{
 				Component->InvalidateLightingCache();
 			}
-
-			// Flush dynamic foliage (grass)
-			ALandscapeProxy::InvalidateGeneratedComponentData(Components);
-
+						
 			// Notify foliage to move any attached instances
 			bool bUpdateFoliage = false;
 
             // Landscape Layers are updates are delayed and done in  ALandscape::TickLayers
 			if (!GetMutableDefault<UEditorExperimentalSettings>()->bLandscapeLayerSystem)
 			{
+				// Flush dynamic foliage (grass)
+				ALandscapeProxy::InvalidateGeneratedComponentData(Components);
+
 				for (ULandscapeComponent* Component : Components)
 				{
 					ULandscapeHeightfieldCollisionComponent* CollisionComponent = Component->CollisionComponent.Get();
@@ -441,18 +443,21 @@ struct FAlphamapAccessor
 
 	~FAlphamapAccessor()
 	{
-		// Recreate collision for modified components to update the physical materials
-		for (ULandscapeComponent* Component : ModifiedComponents)
+		if (!GetMutableDefault<UEditorExperimentalSettings>()->bLandscapeLayerSystem)
 		{
-			ULandscapeHeightfieldCollisionComponent* CollisionComponent = Component->CollisionComponent.Get();
-			if (CollisionComponent)
+			// Recreate collision for modified components to update the physical materials
+			for (ULandscapeComponent* Component : ModifiedComponents)
 			{
-				CollisionComponent->RecreateCollision();
-
-				// We need to trigger navigation mesh build, in case user have painted holes on a landscape
-				if (LayerInfo == ALandscapeProxy::VisibilityLayer)
+				ULandscapeHeightfieldCollisionComponent* CollisionComponent = Component->CollisionComponent.Get();
+				if (CollisionComponent)
 				{
-					FNavigationSystem::UpdateComponentData(*CollisionComponent);
+					CollisionComponent->RecreateCollision();
+
+					// We need to trigger navigation mesh build, in case user have painted holes on a landscape
+					if (LayerInfo == ALandscapeProxy::VisibilityLayer)
+					{
+						FNavigationSystem::UpdateComponentData(*CollisionComponent);
+					}
 				}
 			}
 		}
@@ -473,12 +478,15 @@ struct FAlphamapAccessor
 		TSet<ULandscapeComponent*> Components;
 		if (LandscapeEdit.GetComponentsInRegion(X1, Y1, X2, Y2, &Components))
 		{
-			// Flush dynamic foliage (grass)
-			ALandscapeProxy::InvalidateGeneratedComponentData(Components);
+			if (!GetMutableDefault<UEditorExperimentalSettings>()->bLandscapeLayerSystem)
+			{
+				// Flush dynamic foliage (grass)
+				ALandscapeProxy::InvalidateGeneratedComponentData(Components);
+				ModifiedComponents.Append(Components);
+			}
 
 			LandscapeEdit.SetAlphaData(LayerInfo, X1, Y1, X2, Y2, Data, 0, PaintingRestriction, bBlendWeight, bUseTotalNormalize);
 			//LayerInfo->IsReferencedFromLoadedData = true;
-			ModifiedComponents.Append(Components);
 		}
 	}
 
