- Landscape Layer System: Re-add invalidation of component data while editing so that grass dissapears immediatly and not on stroke end.

#rnx

diff --git a/Engine/Source/Editor/LandscapeEditor/Private/LandscapeEdModeMirrorTool.cpp b/Engine/Source/Editor/LandscapeEditor/Private/LandscapeEdModeMirrorTool.cpp
--- a/Engine/Source/Editor/LandscapeEditor/Private/LandscapeEdModeMirrorTool.cpp
+++ b/Engine/Source/Editor/LandscapeEditor/Private/LandscapeEdModeMirrorTool.cpp
@@ -647,6 +647,8 @@ public:
 		TSet<ULandscapeComponent*> Components;
 		if (LandscapeEdit.GetComponentsInRegion(DestMinX, DestMinY, DestMaxX, DestMaxY, &Components) && Components.Num() > 0)
 		{
+			ALandscapeProxy::InvalidateGeneratedComponentData(Components);
+
 			if (!GetMutableDefault<UEditorExperimentalSettings>()->bLandscapeLayerSystem)
 			{
 				for (ULandscapeComponent* Component : Components)
@@ -659,9 +661,6 @@ public:
 						FNavigationSystem::UpdateComponentData(*CollisionComponent);
 					}
 				}
-
-				// Flush dynamic foliage (grass)
-				ALandscapeProxy::InvalidateGeneratedComponentData(Components);
 			}
 
 			EdMode->UpdateLayerUsageInformation();
diff --git a/Engine/Source/Editor/LandscapeEditor/Private/LandscapeEdModeTools.h b/Engine/Source/Editor/LandscapeEditor/Private/LandscapeEdModeTools.h
--- a/Engine/Source/Editor/LandscapeEditor/Private/LandscapeEdModeTools.h
+++ b/Engine/Source/Editor/LandscapeEditor/Private/LandscapeEdModeTools.h
@@ -695,11 +695,11 @@ struct FXYOffsetmapAccessor
 						
 			// Notify foliage to move any attached instances
 			bool bUpdateFoliage = false;
+			
+			ALandscapeProxy::InvalidateGeneratedComponentData(Components);
+
 			if (!GetMutableDefault<UEditorExperimentalSettings>()->bLandscapeLayerSystem)
 			{
-				// Flush dynamic foliage (grass)
-				ALandscapeProxy::InvalidateGeneratedComponentData(Components);
-
 				for (ULandscapeComponent* Component : Components)
 				{
 					ULandscapeHeightfieldCollisionComponent* CollisionComponent = Component->CollisionComponent.Get();
@@ -876,10 +876,10 @@ struct FFullWeightmapAccessor
 		TSet<ULandscapeComponent*> Components;
 		if (LandscapeEdit.GetComponentsInRegion(X1, Y1, X2, Y2, &Components))
 		{
+			ALandscapeProxy::InvalidateGeneratedComponentData(Components);
+
 			if (!GetMutableDefault<UEditorExperimentalSettings>()->bLandscapeLayerSystem)
 			{
-				// Flush dynamic foliage (grass)
-				ALandscapeProxy::InvalidateGeneratedComponentData(Components);
 				ModifiedComponents.Append(Components);
 			}
 			LandscapeEdit.SetAlphaData(DirtyLayerInfos, X1, Y1, X2, Y2, Data, 0, PaintingRestriction);
diff --git a/Engine/Source/Runtime/Landscape/Private/LandscapeEdit.cpp b/Engine/Source/Runtime/Landscape/Private/LandscapeEdit.cpp
--- a/Engine/Source/Runtime/Landscape/Private/LandscapeEdit.cpp
+++ b/Engine/Source/Runtime/Landscape/Private/LandscapeEdit.cpp
@@ -572,9 +572,13 @@ void ULandscapeComponent::PostEditUndo()
 		EditToolRenderData.UpdateSelectionMaterial(EditToolRenderData.SelectedType, this);
 		UpdateEditToolRenderData();
 	}
-
+		
 	if (GetMutableDefault<UEditorExperimentalSettings>()->bLandscapeLayerSystem)
 	{
+		TArray<ULandscapeComponent*> SingleComponent;
+		SingleComponent.Add(this);
+		GetLandscapeProxy()->InvalidateGeneratedComponentData(SingleComponent);
+		
 		ALandscape* LandscapeActor = GetLandscapeActor();
 		// Might be a ALandscapeStreamingProxy
 		if (LandscapeActor)
@@ -3945,6 +3949,8 @@ void ALandscapeProxy::PostEditChangeProperty(FPropertyChangedEvent& PropertyChan
 	}
 	if (bRemovedAnyLayers)
 	{
+		ALandscapeProxy::InvalidateGeneratedComponentData(LandscapeComponents);
+
 		if (GetMutableDefault<UEditorExperimentalSettings>()->bLandscapeLayerSystem)
 		{
 			if(ALandscape* LandscapeActor = GetLandscapeActor())
@@ -3952,13 +3958,7 @@ void ALandscapeProxy::PostEditChangeProperty(FPropertyChangedEvent& PropertyChan
 				LandscapeActor->RequestLayersContentUpdate(ELandscapeLayersContentUpdateFlag::All, true);
 			}
 		}
-		else
-		{
-			// Flush dynamic data (e.g. grass)
-			TSet<ULandscapeComponent*> Components;
-			Components.Append(LandscapeComponents);
-			ALandscapeProxy::InvalidateGeneratedComponentData(Components);
-		}
+		
 	}
 
 	// Must do this *after* correcting the scale or reattaching the landscape components will crash!
diff --git a/Engine/Source/Runtime/Landscape/Private/LandscapeEditInterface.cpp b/Engine/Source/Runtime/Landscape/Private/LandscapeEditInterface.cpp
--- a/Engine/Source/Runtime/Landscape/Private/LandscapeEditInterface.cpp
+++ b/Engine/Source/Runtime/Landscape/Private/LandscapeEditInterface.cpp
@@ -1794,7 +1794,10 @@ void FLandscapeEditDataInterface::DeleteLayer(ULandscapeLayerInfoObject* LayerIn
 		Component->DeleteLayer(LayerInfo, *this);
 	}
 
-	// Flush dynamic data (e.g. grass)
+	TSet<ULandscapeComponent*> Components;
+	Algo::Transform(LandscapeInfo->XYtoComponentMap, Components, &TPair<FIntPoint, ULandscapeComponent*>::Value);
+	ALandscapeProxy::InvalidateGeneratedComponentData(Components);
+
 	if (GetMutableDefault<UEditorExperimentalSettings>()->bLandscapeLayerSystem)
 	{
 		if (LandscapeInfo->LandscapeActor)
@@ -1802,12 +1805,6 @@ void FLandscapeEditDataInterface::DeleteLayer(ULandscapeLayerInfoObject* LayerIn
 			LandscapeInfo->LandscapeActor->RequestLayersContentUpdate(ELandscapeLayersContentUpdateFlag::All, true);
 		}
 	}
-	else
-	{
-		TSet<ULandscapeComponent*> Components;
-		Algo::Transform(LandscapeInfo->XYtoComponentMap, Components, &TPair<FIntPoint, ULandscapeComponent*>::Value);
-		ALandscapeProxy::InvalidateGeneratedComponentData(Components);
-	}
 }
 
 void ULandscapeComponent::FillLayer(ULandscapeLayerInfoObject* LayerInfo, FLandscapeEditDataInterface& LandscapeEdit)
@@ -2038,6 +2035,10 @@ void FLandscapeEditDataInterface::FillLayer(ULandscapeLayerInfoObject* LayerInfo
 		}
 	}
 
+	TSet<ULandscapeComponent*> Components;
+	Algo::Transform(LandscapeInfo->XYtoComponentMap, Components, &TPair<FIntPoint, ULandscapeComponent*>::Value);
+	ALandscapeProxy::InvalidateGeneratedComponentData(Components);
+
 	if (GetMutableDefault<UEditorExperimentalSettings>()->bLandscapeLayerSystem)
 	{
 		if (LandscapeInfo->LandscapeActor)
@@ -2045,16 +2046,6 @@ void FLandscapeEditDataInterface::FillLayer(ULandscapeLayerInfoObject* LayerInfo
 			LandscapeInfo->LandscapeActor->RequestLayersContentUpdate(ELandscapeLayersContentUpdateFlag::All, true);
 		}
 	}
-	else
-	{
-		// Flush dynamic data (e.g. grass)
-		TSet<ULandscapeComponent*> Components;
-		for (auto& XYComponentPair : LandscapeInfo->XYtoComponentMap)
-		{
-			Components.Add(XYComponentPair.Value);
-		}
-		ALandscapeProxy::InvalidateGeneratedComponentData(Components);
-	}
 }
 
 void ULandscapeComponent::ReplaceLayer(ULandscapeLayerInfoObject* FromLayerInfo, ULandscapeLayerInfoObject* ToLayerInfo, FLandscapeEditDataInterface& LandscapeEdit)
diff --git a/Engine/Source/Runtime/Landscape/Private/LandscapeSplineRaster.cpp b/Engine/Source/Runtime/Landscape/Private/LandscapeSplineRaster.cpp
--- a/Engine/Source/Runtime/Landscape/Private/LandscapeSplineRaster.cpp
+++ b/Engine/Source/Runtime/Landscape/Private/LandscapeSplineRaster.cpp
@@ -624,8 +624,10 @@ bool ULandscapeInfo::ApplySplinesInternal(bool bOnlySelected, ALandscapeProxy* P
 	}
 
 	LandscapeEdit.Flush();
-
-	if (!GetMutableDefault<UEditorExperimentalSettings>()->bLandscapeLayerSystem)
+		
+	ALandscapeProxy::InvalidateGeneratedComponentData(ModifiedComponents);
+	
+	if(!GetMutableDefault<UEditorExperimentalSettings>()->bLandscapeLayerSystem)
 	{
 		ALandscape* Landscape = Proxy->GetLandscapeActor();
 		for (ULandscapeComponent* Component : ModifiedComponents)
@@ -638,9 +640,6 @@ bool ULandscapeInfo::ApplySplinesInternal(bool bOnlySelected, ALandscapeProxy* P
 				FNavigationSystem::UpdateComponentData(*CollisionComponent);
 			}
 		}
-		
-		// Invalidate landscape grass of modified landscape components
-		ALandscapeProxy::InvalidateGeneratedComponentData(ModifiedComponents);
 	}
 
 	return true;
diff --git a/Engine/Source/Runtime/Landscape/Public/LandscapeEdit.h b/Engine/Source/Runtime/Landscape/Public/LandscapeEdit.h
--- a/Engine/Source/Runtime/Landscape/Public/LandscapeEdit.h
+++ b/Engine/Source/Runtime/Landscape/Public/LandscapeEdit.h
@@ -324,12 +324,11 @@ struct FHeightmapAccessor
 			// Notify foliage to move any attached instances
 			bool bUpdateFoliage = false;
 
+			ALandscapeProxy::InvalidateGeneratedComponentData(Components);
+
             // Landscape Layers are updates are delayed and done in  ALandscape::TickLayers
 			if (!GetMutableDefault<UEditorExperimentalSettings>()->bLandscapeLayerSystem)
 			{
-				// Flush dynamic foliage (grass)
-				ALandscapeProxy::InvalidateGeneratedComponentData(Components);
-
 				for (ULandscapeComponent* Component : Components)
 				{
 					ULandscapeHeightfieldCollisionComponent* CollisionComponent = Component->CollisionComponent.Get();
@@ -478,10 +477,10 @@ struct FAlphamapAccessor
 		TSet<ULandscapeComponent*> Components;
 		if (LandscapeEdit.GetComponentsInRegion(X1, Y1, X2, Y2, &Components))
 		{
+			ALandscapeProxy::InvalidateGeneratedComponentData(Components);
+			
 			if (!GetMutableDefault<UEditorExperimentalSettings>()->bLandscapeLayerSystem)
 			{
-				// Flush dynamic foliage (grass)
-				ALandscapeProxy::InvalidateGeneratedComponentData(Components);
 				ModifiedComponents.Append(Components);
 			}
 
