Fixed crash in Landscape layer system where async graphic init resources were accessed before their associated graphic command were executed.
#rnx

diff --git a/Engine/Source/Runtime/Landscape/Private/LandscapeEditLayers.cpp b/Engine/Source/Runtime/Landscape/Private/LandscapeEditLayers.cpp
--- a/Engine/Source/Runtime/Landscape/Private/LandscapeEditLayers.cpp
+++ b/Engine/Source/Runtime/Landscape/Private/LandscapeEditLayers.cpp
@@ -2707,31 +2707,31 @@ void ALandscape::RegenerateLayersHeightmaps(const TArray<ULandscapeComponent*>&
 	{
 		LayersContentUpdateFlags |= ELandscapeLayersContentUpdateFlag::Heightmap_Render;
 	}
-}
+			}
 
 bool ALandscape::UpdateClients(const TArray<ULandscapeComponent*>& InLandscapeComponents, bool bInCollisionRecreated)
-{
-	// Post-pone client update to next tick while in transaction
-	if (GUndo)
-	{
+			{
+				// Post-pone client update to next tick while in transaction
+				if (GUndo)
+				{
 		return false;
-	}
-	else
-	{
+				}
+				else
+				{
 		for (ULandscapeComponent* Component : InLandscapeComponents)
 		{
-			// Recreate collision for modified components to update the physical materials
-			ULandscapeHeightfieldCollisionComponent* CollisionComponent = Component->CollisionComponent.Get();
-			if (CollisionComponent)
-			{
-				FNavigationSystem::UpdateComponentData(*CollisionComponent);
-				CollisionComponent->SnapFoliageInstances();
-			}
-		}
+					// Recreate collision for modified components to update the physical materials
+					ULandscapeHeightfieldCollisionComponent* CollisionComponent = Component->CollisionComponent.Get();
+					if (CollisionComponent)
+					{
+						FNavigationSystem::UpdateComponentData(*CollisionComponent);
+						CollisionComponent->SnapFoliageInstances();
+					}
+				}
 
 		ALandscapeProxy::InvalidateGeneratedComponentData(InLandscapeComponents);
 	}
-	
+
 	// Re update clients if collision wasn't recreated
 	return bInCollisionRecreated;
 }
@@ -3377,15 +3377,15 @@ void ALandscape::RegenerateLayersWeightmaps(const TArray<ULandscapeComponent*>&
 {
 	SCOPE_CYCLE_COUNTER(STAT_LandscapeLayersRegenerateWeightmaps);
 	bOutRecreateCollision = false;
-    
+
 	ULandscapeInfo* Info = GetLandscapeInfo();
 
 	if (LayersContentUpdateFlags == 0 || Info == nullptr || Info->Layers.Num() == 0 || !PrepareLayersWeightmapTextureResources())
 	{
 		return;
 	}
-	
-		
+
+
 	TArray<ULandscapeComponent*> ComponentThatNeedMaterialRebuild;
 	TArray<ULandscapeLayerInfoObject*> BrushRequiredAllocations;
 	int32 LayerCount = Info->Layers.Num() + 1; // due to visibility being stored at 0
@@ -3604,6 +3604,7 @@ void ALandscape::RegenerateLayersWeightmaps(const TArray<ULandscapeComponent*>&
 		if (ComputeShaderGeneratedData)
 		{
 			// Will generate CPU read back resource, if required
+			bool bHasPendingInitResource = false;
 			Info->ForAllLandscapeProxies([&](ALandscapeProxy* Proxy)
 			{
 				for (ULandscapeComponent* Component : Proxy->LandscapeComponents)
@@ -3618,13 +3619,19 @@ void ALandscape::RegenerateLayersWeightmaps(const TArray<ULandscapeComponent*>&
 						{
 							FLandscapeLayersTexture2DCPUReadBackResource* NewWeightmapCPUReadBack = new FLandscapeLayersTexture2DCPUReadBackResource(WeightmapTexture->Source.GetSizeX(), WeightmapTexture->Source.GetSizeY(), WeightmapTexture->GetPixelFormat(), WeightmapTexture->Source.GetNumMips());
 							BeginInitResource(NewWeightmapCPUReadBack);
-
 							Component->GetLandscapeProxy()->WeightmapsCPUReadBack.Add(WeightmapTexture, NewWeightmapCPUReadBack);
+							bHasPendingInitResource = true;
 						}
 					}
 				}
 			});
 
+			if (bHasPendingInitResource)
+			{
+				// Flush because TickLayers can access CPU Readbacks in the same frame (see ResolveLayersTexture)
+				FlushRenderingCommands();
+			}
+
 			int8 CurrentWeightmapToProcessIndex = 0;
 			bool HasFoundWeightmapToProcess = true; // try processing at least once
 
@@ -3745,7 +3752,7 @@ void ALandscape::RegenerateLayersWeightmaps(const TArray<ULandscapeComponent*>&
 
 		bOutRecreateCollision = true;
 	}
-		
+
 	if (ClearFlagsAfterUpdate)
 	{
 		LayersContentUpdateFlags &= ~ELandscapeLayersContentUpdateFlag::Weightmap_All;
