Optim: Only rebuild necessary landscape materials when reallocating layers weightmaps in Landscape Layer System
#rnx

diff --git a/Engine/Source/Runtime/Landscape/Classes/Landscape.h b/Engine/Source/Runtime/Landscape/Classes/Landscape.h
--- a/Engine/Source/Runtime/Landscape/Classes/Landscape.h
+++ b/Engine/Source/Runtime/Landscape/Classes/Landscape.h
@@ -383,9 +383,6 @@ private:
 	UPROPERTY(Transient)
 	uint32 LayersContentUpdateFlags;
 
-	UPROPERTY(Transient)
-	bool LayersUpdateAllMaterials;
-
 	// Represent all the resolved paint layer, from all layers blended together (size of the landscape x material layer count)
 	class FLandscapeTexture2DArrayResource* CombinedLayersWeightmapAllMaterialLayersResource;
 	
diff --git a/Engine/Source/Runtime/Landscape/Private/LandscapeEditLayers.cpp b/Engine/Source/Runtime/Landscape/Private/LandscapeEditLayers.cpp
--- a/Engine/Source/Runtime/Landscape/Private/LandscapeEditLayers.cpp
+++ b/Engine/Source/Runtime/Landscape/Private/LandscapeEditLayers.cpp
@@ -3024,32 +3024,26 @@ void ALandscape::ReallocateLayersWeightmaps(const TArray<ULandscapeLayerInfoObje
 		AllLandscapeComponents.Append(Proxy->LandscapeComponents);
 	});
 
-	// Copy Previous Usage, to know which texture need updating
-	TMap<UTexture2D*, ULandscapeWeightmapUsage*> CurrentWeightmapsUsage;
-
-	for (ULandscapeComponent* Component : AllLandscapeComponents)
+	auto HashComponentWeightmaps = [](ULandscapeComponent* Component) -> uint32
 	{
+		uint32 Hash = 0;
 		TArray<UTexture2D*>& ComponentWeightmapTextures = Component->GetWeightmapTextures();
 		TArray<ULandscapeWeightmapUsage*>& ComponentWeightmapTextureUsage = Component->GetWeightmapTexturesUsage();
-
 		for (int32 i = 0; i < ComponentWeightmapTextures.Num(); ++i)
 		{
-			UTexture2D* ComponentWeightmapTexture = ComponentWeightmapTextures[i];
-			ULandscapeWeightmapUsage** CurrentWeightmapTextureUsage = CurrentWeightmapsUsage.Find(ComponentWeightmapTexture);
-
-			if (CurrentWeightmapTextureUsage == nullptr)
+			Hash = PointerHash(ComponentWeightmapTextures[i], Hash);
+			for (int32 j = 0; j < ULandscapeWeightmapUsage::NumChannels; ++j)
 			{
-				ULandscapeWeightmapUsage* ComponentWeightmapUsage = ComponentWeightmapTextureUsage[i];
-				ULandscapeWeightmapUsage* Usage = Component->GetLandscapeProxy()->CreateWeightmapUsage();
-
-				for (int32 j = 0; j < ULandscapeWeightmapUsage::NumChannels; ++j)
-				{
-					Usage->ChannelUsage[j] = ComponentWeightmapUsage->ChannelUsage[j];
-				}
-
-				CurrentWeightmapsUsage.Add(ComponentWeightmapTexture, Usage);
+				Hash = PointerHash(ComponentWeightmapTextureUsage[i]->ChannelUsage[j], Hash);
 			}
 		}
+		return Hash;
+	};
+
+	TMap<ULandscapeComponent*, uint32> BeforeReallocComponentWeightmapsHashes;
+	for (ULandscapeComponent* Component : AllLandscapeComponents)
+	{
+		BeforeReallocComponentWeightmapsHashes.Add(Component, HashComponentWeightmaps(Component));
 	}
 
 	// Clear allocation data
@@ -3166,13 +3160,6 @@ void ALandscape::ReallocateLayersWeightmaps(const TArray<ULandscapeLayerInfoObje
 
 	//GDisableAutomaticTextureMaterialUpdateDependencies = false;
 
-	// Determine which Component need updating
-
-	// BEGIN HACK for now we disable the logic that will determine what need updating as it's buggy, and i prefer a correct behavior than an optimized non working logic.
-	OutComponentThatNeedMaterialRebuild.Reset();
-	OutComponentThatNeedMaterialRebuild.Append(AllLandscapeComponents);
-	// END HACK
-
 	// Clean-up unused weightmap CPUReadback resources
 	Info->ForAllLandscapeProxies([&AllLandscapeComponents](ALandscapeProxy* Proxy)
 	{
@@ -3217,45 +3204,17 @@ void ALandscape::ReallocateLayersWeightmaps(const TArray<ULandscapeLayerInfoObje
 		}
 	});
 
-	/*for (ULandscapeComponent* Component : AllLandscapeComponents)
+	// Determine which Component needs a material rebuild
+	for (ULandscapeComponent* Component : AllLandscapeComponents)
 	{
-		if (OutComponentThatNeedMaterialRebuild.Contains(Component))
-		{
-			continue;
-		}
-
-		TArray<UTexture2D*>& ComponentWeightmapTextures = Component->GetWeightmapTextures();
-		TArray<ULandscapeWeightmapUsage*>& ComponentWeightmapTextureUsage = Component->GetWeightmapTexturesUsage();
-
-		for (int32 i = 0; i < ComponentWeightmapTextures.Num(); ++i)
+		uint32* OldHash = BeforeReallocComponentWeightmapsHashes.Find(Component);
+		check(OldHash);
+		uint32 NewHash = HashComponentWeightmaps(Component);
+		if (*OldHash != NewHash)
 		{
-			UTexture2D* ComponentWeightmapTexture = ComponentWeightmapTextures[i];
-			ULandscapeWeightmapUsage* ComponentWeightmapUsage = ComponentWeightmapTextureUsage[i];
-			ULandscapeWeightmapUsage** CurrentWeightmapTextureUsage = CurrentWeightmapsUsage.Find(ComponentWeightmapTexture);
-
-			if (CurrentWeightmapTextureUsage != nullptr)
-			{
-				for (int32 j = 0; j < 4; ++j)
-				{
-					if (ComponentWeightmapUsage->ChannelUsage[j] != (*CurrentWeightmapTextureUsage)->ChannelUsage[j] && ComponentWeightmapUsage->ChannelUsage[j] != nullptr)
-					{
-						OutComponentThatNeedMaterialRebuild.AddUnique(ComponentWeightmapUsage->ChannelUsage[j]);
-					}
-				}
-			}
-			else if (NewCreatedTextures.Contains(ComponentWeightmapTexture))
-			{
-				for (int32 j = 0; j < 4; ++j)
-				{
-					if (ComponentWeightmapUsage->ChannelUsage[j] != nullptr)
-					{
-						OutComponentThatNeedMaterialRebuild.AddUnique(ComponentWeightmapUsage->ChannelUsage[j]);
-					}
-				}
-			}
+			OutComponentThatNeedMaterialRebuild.Add(Component);
 		}
 	}
-	*/
 }
 
 void ALandscape::InitializeLayersWeightmapResources()
@@ -3734,8 +3693,7 @@ void ALandscape::RegenerateLayersWeightmaps(const TArray<ULandscapeComponent*>&
 			}
 		}
 
-		UpdateLayersMaterialInstances(LayersUpdateAllMaterials ? InLandscapeComponents : ComponentThatNeedMaterialRebuild);
-		LayersUpdateAllMaterials = false;
+		UpdateLayersMaterialInstances(ComponentThatNeedMaterialRebuild);
 	}
 
 	if (LayersContentUpdateFlags & EInternalLandscapeLayersContentUpdateFlag::Internal_Weightmap_ResolveToTexture)
@@ -3967,7 +3925,6 @@ void ALandscape::ResolveLayersWeightmapTexture()
 void ALandscape::RequestLayersContentUpdate(ELandscapeLayersContentUpdateFlag InDataFlags, bool InUpdateAllMaterials)
 {
 	LayersContentUpdateFlags |= InDataFlags;
-	LayersUpdateAllMaterials |= InUpdateAllMaterials;
 }
 
 void ALandscape::RegenerateLayersContent()
