- Landscape Layer System: Component partial update support for Collision and Clients (Foliage, Grass, Nav)

#rnx

diff --git a/Engine/Source/Editor/LandscapeEditor/Private/LandscapeEdMode.cpp b/Engine/Source/Editor/LandscapeEditor/Private/LandscapeEdMode.cpp
--- a/Engine/Source/Editor/LandscapeEditor/Private/LandscapeEdMode.cpp
+++ b/Engine/Source/Editor/LandscapeEditor/Private/LandscapeEdMode.cpp
@@ -253,7 +253,7 @@ void ALandscape::SplitHeightmap(ULandscapeComponent* Comp, ALandscapeProxy* Targ
 			}
 		});
 
-		Landscape->RequestLayersContentUpdate(ELandscapeLayersContentUpdateFlag::All, true);
+		Landscape->RequestLayersContentUpdateForceAll();
 	}
 #endif
 
@@ -504,13 +504,9 @@ void FEdModeLandscape::Enter()
 		ALandscapeProxy* LandscapeProxy = CurrentToolTarget.LandscapeInfo->GetLandscapeProxy();
 		LandscapeProxy->OnMaterialChangedDelegate().AddRaw(this, &FEdModeLandscape::OnLandscapeMaterialChangedDelegate);
 
-		if (GetMutableDefault<UEditorExperimentalSettings>()->bLandscapeLayerSystem)
+		if(ALandscape* Landscape = GetLandscape())
 		{
-			ALandscape* Landscape = GetLandscape();
-			if (Landscape)
-			{
-				Landscape->RequestLayersContentUpdate(ELandscapeLayersContentUpdateFlag::All_Render);
-			}
+			Landscape->RequestLayersContentUpdateForceAll();
 		}
 	}
 
@@ -3334,7 +3330,7 @@ void FEdModeLandscape::ReimportData(const FLandscapeTargetListInfo& TargetInfo)
 	const FString& SourceFilePath = TargetInfo.ReimportFilePath();
 	if (SourceFilePath.Len())
 	{
-		FScopedSetLandscapeEditingLayer Scope(GetLandscape(), GetCurrentLayerGuid(), [&] { RequestLayersContentUpdate(); });
+		FScopedSetLandscapeEditingLayer Scope(GetLandscape(), GetCurrentLayerGuid(), [&] { RequestLayersContentUpdateForceAll(); });
 		ImportData(TargetInfo, SourceFilePath);
 	}
 	else
@@ -3441,7 +3437,7 @@ void FEdModeLandscape::ImportData(const FLandscapeTargetListInfo& TargetInfo, co
 
 			{
 				ALandscape* Landscape = GetLandscape();
-				FScopedSetLandscapeEditingLayer Scope(Landscape, GetCurrentLayerGuid(), [&] { check(Landscape); Landscape->RequestLayersContentUpdate(ELandscapeLayersContentUpdateFlag::Heightmap_All); });
+				FScopedSetLandscapeEditingLayer Scope(Landscape, GetCurrentLayerGuid(), [&] { check(Landscape); Landscape->RequestLayersContentUpdate(ELandscapeLayerUpdateMode::Heightmap_All); });
 
 				TArray<uint16> Data;
 				if (ImportResolution != LandscapeResolution)
@@ -3557,7 +3553,7 @@ void FEdModeLandscape::ImportData(const FLandscapeTargetListInfo& TargetInfo, co
 
 			{
 				ALandscape* Landscape = GetLandscape();
-				FScopedSetLandscapeEditingLayer Scope(Landscape, GetCurrentLayerGuid(), [&] { check(Landscape); Landscape->RequestLayersContentUpdate(ELandscapeLayersContentUpdateFlag::Weightmap_All); });
+				FScopedSetLandscapeEditingLayer Scope(Landscape, GetCurrentLayerGuid(), [&] { check(Landscape); Landscape->RequestLayersContentUpdate(ELandscapeLayerUpdateMode::Weightmap_All); });
 
 				TArray<uint8> Data;
 				if (ImportResolution != LandscapeResolution)
@@ -4225,12 +4221,19 @@ void FEdModeLandscape::SetLayerLocked(int32 InLayerIndex, bool bInLocked)
 	}
 }
 
-void FEdModeLandscape::RequestLayersContentUpdate(bool InUpdateAllMaterials)
+void FEdModeLandscape::RequestLayersContentUpdate(ELandscapeLayerUpdateMode InUpdateMode, bool bInForceUpdateAllComponents)
 {
-	ALandscape* Landscape = GetLandscape();
-	if (Landscape)
+	if (ALandscape* Landscape = GetLandscape())
+	{
+		Landscape->RequestLayersContentUpdate(InUpdateMode, bInForceUpdateAllComponents);
+	}
+}
+
+void FEdModeLandscape::RequestLayersContentUpdateForceAll()
+{
+	if (ALandscape* Landscape = GetLandscape())
 	{
-		Landscape->RequestLayersContentUpdate(ELandscapeLayersContentUpdateFlag::All_Render, InUpdateAllMaterials);
+		Landscape->RequestLayersContentUpdateForceAll();
 	}
 }
 
diff --git a/Engine/Source/Editor/LandscapeEditor/Private/LandscapeEdMode.h b/Engine/Source/Editor/LandscapeEditor/Private/LandscapeEdMode.h
--- a/Engine/Source/Editor/LandscapeEditor/Private/LandscapeEdMode.h
+++ b/Engine/Source/Editor/LandscapeEditor/Private/LandscapeEdMode.h
@@ -523,7 +523,8 @@ public:
 	TArray<class ALandscapeBlueprintCustomBrush*> GetBrushesForCurrentLayer(int32 InTargetType);
 	
 	bool NeedToFillEmptyMaterialLayers() const;
-	void RequestLayersContentUpdate(bool InUpdateAllMaterials = false);
+	void RequestLayersContentUpdate(ELandscapeLayerUpdateMode InUpdateMode, bool bInForceUpdateAllComponents = false);
+	void RequestLayersContentUpdateForceAll();
 
 	void OnLevelActorAdded(AActor* InActor);
 	void OnLevelActorRemoved(AActor* InActor);
diff --git a/Engine/Source/Editor/LandscapeEditor/Private/LandscapeEdModeComponentTools.cpp b/Engine/Source/Editor/LandscapeEditor/Private/LandscapeEdModeComponentTools.cpp
--- a/Engine/Source/Editor/LandscapeEditor/Private/LandscapeEdModeComponentTools.cpp
+++ b/Engine/Source/Editor/LandscapeEditor/Private/LandscapeEdModeComponentTools.cpp
@@ -602,7 +602,7 @@ public:
 							FScopedSetLandscapeEditingLayer Scope(Landscape, LayerGuid);
 							Component->ReallocateWeightmaps(&LandscapeEdit, true, true, false, true, LandscapeProxy);
 						});
-						Landscape->RequestLayersContentUpdate(ELandscapeLayersContentUpdateFlag::All, true);
+						Landscape->RequestLayersContentUpdateForceAll();
 					}
 
 					// Need to Repacking all the Weight map (to make it packed well...)
@@ -1725,11 +1725,11 @@ public:
 		return ELandscapeToolTargetTypeMask::FromType(ToolTarget::TargetType);
 	}
 
-	virtual ELandscapeLayersContentUpdateFlag GetBeginToolContentUpdateFlag() const override { return ELandscapeLayersContentUpdateFlag::All_Render; }
+	virtual ELandscapeLayerUpdateMode GetBeginToolContentUpdateFlag() const override { return ELandscapeLayerUpdateMode::All_Editing; }
 
-	virtual ELandscapeLayersContentUpdateFlag GetTickToolContentUpdateFlag() const override { return ELandscapeLayersContentUpdateFlag::All_Render; }
+	virtual ELandscapeLayerUpdateMode GetTickToolContentUpdateFlag() const override { return GetBeginToolContentUpdateFlag(); }
 	
-	virtual ELandscapeLayersContentUpdateFlag GetEndToolContentUpdateFlag() const override { return ELandscapeLayersContentUpdateFlag::All; }
+	virtual ELandscapeLayerUpdateMode GetEndToolContentUpdateFlag() const override { return ELandscapeLayerUpdateMode::All; }
 		
 	virtual bool BeginTool(FEditorViewportClient* ViewportClient, const FLandscapeToolTarget& InTarget, const FVector& InHitLocation) override
 	{
diff --git a/Engine/Source/Editor/LandscapeEditor/Private/LandscapeEdModeMirrorTool.cpp b/Engine/Source/Editor/LandscapeEditor/Private/LandscapeEdModeMirrorTool.cpp
--- a/Engine/Source/Editor/LandscapeEditor/Private/LandscapeEdModeMirrorTool.cpp
+++ b/Engine/Source/Editor/LandscapeEditor/Private/LandscapeEdModeMirrorTool.cpp
@@ -508,7 +508,7 @@ public:
 			return;
 		}
 		FScopedTransaction Transaction(LOCTEXT("Mirror_Apply", "Landscape Editing: Mirror Landscape"));
-		FScopedSetLandscapeEditingLayer Scope(EdMode->GetLandscape(), EdMode->GetCurrentLayerGuid(), [&] { EdMode->RequestLayersContentUpdate(true); });
+		FScopedSetLandscapeEditingLayer Scope(EdMode->GetLandscape(), EdMode->GetCurrentLayerGuid(), [&] { EdMode->RequestLayersContentUpdateForceAll(); });
 
 		const ULandscapeInfo* const LandscapeInfo = EdMode->CurrentToolTarget.LandscapeInfo.Get();
 		const ALandscapeProxy* const LandscapeProxy = LandscapeInfo->GetLandscapeProxy();
diff --git a/Engine/Source/Editor/LandscapeEditor/Private/LandscapeEdModeRampTool.cpp b/Engine/Source/Editor/LandscapeEditor/Private/LandscapeEdModeRampTool.cpp
--- a/Engine/Source/Editor/LandscapeEditor/Private/LandscapeEdModeRampTool.cpp
+++ b/Engine/Source/Editor/LandscapeEditor/Private/LandscapeEdModeRampTool.cpp
@@ -486,7 +486,7 @@ public:
 
 		FScopedTransaction Transaction(LOCTEXT("Ramp_Apply", "Landscape Editing: Add ramp"));
 		ALandscape* Landscape = EdMode->GetLandscape();
-		FScopedSetLandscapeEditingLayer Scope(Landscape, EdMode->GetCurrentLayerGuid(), [&] { if (Landscape) { Landscape->RequestLayersContentUpdate(ELandscapeLayersContentUpdateFlag::Heightmap_All); } });
+		FScopedSetLandscapeEditingLayer Scope(Landscape, EdMode->GetCurrentLayerGuid(), [&] { if (Landscape) { Landscape->RequestLayersContentUpdate(ELandscapeLayerUpdateMode::Heightmap_All); } });
 
 		const ULandscapeInfo* LandscapeInfo = EdMode->CurrentToolTarget.LandscapeInfo.Get();
 		const ALandscapeProxy* LandscapeProxy = LandscapeInfo->GetLandscapeProxy();
diff --git a/Engine/Source/Editor/LandscapeEditor/Private/LandscapeEdModeTools.h b/Engine/Source/Editor/LandscapeEditor/Private/LandscapeEdModeTools.h
--- a/Engine/Source/Editor/LandscapeEditor/Private/LandscapeEdModeTools.h
+++ b/Engine/Source/Editor/LandscapeEditor/Private/LandscapeEdModeTools.h
@@ -1131,21 +1131,21 @@ public:
 
 	virtual bool ShouldUpdateEditingLayer() const { return GetMutableDefault<UEditorExperimentalSettings>()->bLandscapeLayerSystem; }
 
-	virtual ELandscapeLayersContentUpdateFlag GetBeginToolContentUpdateFlag() const
+	virtual ELandscapeLayerUpdateMode GetBeginToolContentUpdateFlag() const
 	{
 		bool bUpdateHeightmap = this->EdMode->CurrentToolTarget.TargetType == ELandscapeToolTargetType::Type::Heightmap; 
-		return bUpdateHeightmap ? ELandscapeLayersContentUpdateFlag::Heightmap_Render : ELandscapeLayersContentUpdateFlag::Weightmap_Render;
+		return bUpdateHeightmap ? ELandscapeLayerUpdateMode::Heightmap_Editing : ELandscapeLayerUpdateMode::Weightmap_Editing;
 	}
 
-	virtual ELandscapeLayersContentUpdateFlag GetTickToolContentUpdateFlag() const
+	virtual ELandscapeLayerUpdateMode GetTickToolContentUpdateFlag() const
 	{
 		return GetBeginToolContentUpdateFlag();
 	}
 
-	virtual ELandscapeLayersContentUpdateFlag GetEndToolContentUpdateFlag() const
+	virtual ELandscapeLayerUpdateMode GetEndToolContentUpdateFlag() const
 	{
 		bool bUpdateHeightmap = this->EdMode->CurrentToolTarget.TargetType == ELandscapeToolTargetType::Type::Heightmap;
-		return bUpdateHeightmap ? ELandscapeLayersContentUpdateFlag::Heightmap_All : ELandscapeLayersContentUpdateFlag::Weightmap_All;
+		return bUpdateHeightmap ? ELandscapeLayerUpdateMode::Heightmap_All : ELandscapeLayerUpdateMode::Weightmap_All;
 	}
 
 	virtual bool BeginTool(FEditorViewportClient* ViewportClient, const FLandscapeToolTarget& InTarget, const FVector& InHitLocation) override
diff --git a/Engine/Source/Editor/LandscapeEditor/Private/LandscapeEditorDetailCustomization_Layers.cpp b/Engine/Source/Editor/LandscapeEditor/Private/LandscapeEditorDetailCustomization_Layers.cpp
--- a/Engine/Source/Editor/LandscapeEditor/Private/LandscapeEditorDetailCustomization_Layers.cpp
+++ b/Engine/Source/Editor/LandscapeEditor/Private/LandscapeEditorDetailCustomization_Layers.cpp
@@ -693,7 +693,7 @@ FReply FLandscapeEditorCustomNodeBuilder_Layers::HandleAcceptDrop(FDragDropEvent
 			{
 				LandscapeEdMode->SetCurrentLayer(DestinationLayerIndex);
 				LandscapeEdMode->RefreshDetailPanel();
-				LandscapeEdMode->RequestLayersContentUpdate();
+				LandscapeEdMode->RequestLayersContentUpdateForceAll();
 				return FReply::Handled();
 			}
 		}
diff --git a/Engine/Source/Editor/LandscapeEditor/Private/LandscapeEditorDetailCustomization_LayersBrushStack.cpp b/Engine/Source/Editor/LandscapeEditor/Private/LandscapeEditorDetailCustomization_LayersBrushStack.cpp
--- a/Engine/Source/Editor/LandscapeEditor/Private/LandscapeEditorDetailCustomization_LayersBrushStack.cpp
+++ b/Engine/Source/Editor/LandscapeEditor/Private/LandscapeEditorDetailCustomization_LayersBrushStack.cpp
@@ -362,7 +362,7 @@ FReply FLandscapeEditorCustomNodeBuilder_LayersBrushStack::HandleAcceptDrop(FDra
 					BrushOrderStack.Insert(MovingBrushIndex, DestinationLayerIndex);
 
 					LandscapeEdMode->RefreshDetailPanel();
-					LandscapeEdMode->RequestLayersContentUpdate();
+					LandscapeEdMode->RequestLayersContentUpdateForceAll();
 
 					return FReply::Handled();
 				}
diff --git a/Engine/Source/Editor/LandscapeEditor/Private/LandscapeEditorDetailCustomization_TargetLayers.cpp b/Engine/Source/Editor/LandscapeEditor/Private/LandscapeEditorDetailCustomization_TargetLayers.cpp
--- a/Engine/Source/Editor/LandscapeEditor/Private/LandscapeEditorDetailCustomization_TargetLayers.cpp
+++ b/Engine/Source/Editor/LandscapeEditor/Private/LandscapeEditorDetailCustomization_TargetLayers.cpp
@@ -1094,7 +1094,7 @@ void FLandscapeEditorCustomNodeBuilder_TargetLayers::OnFillLayer(const TSharedRe
 		FEdModeLandscape* LandscapeEdMode = GetEditorMode();
 		if (LandscapeEdMode)
 		{
-			FScopedSetLandscapeEditingLayer Scope(LandscapeEdMode->GetLandscape(), LandscapeEdMode->GetCurrentLayerGuid(), [&] { LandscapeEdMode->RequestLayersContentUpdate(true); });
+			FScopedSetLandscapeEditingLayer Scope(LandscapeEdMode->GetLandscape(), LandscapeEdMode->GetCurrentLayerGuid(), [&] { LandscapeEdMode->RequestLayersContentUpdateForceAll(); });
 			LandscapeEdit.FillLayer(Target->LayerInfoObj.Get());
 		}
 	}
@@ -1132,7 +1132,7 @@ void FLandscapeEditorCustomNodeBuilder_TargetLayers::OnClearLayer(const TSharedR
 		FEdModeLandscape* LandscapeEdMode = GetEditorMode();
 		if (LandscapeEdMode)
 		{
-			FScopedSetLandscapeEditingLayer Scope(LandscapeEdMode->GetLandscape(), LandscapeEdMode->GetCurrentLayerGuid(), [&] { LandscapeEdMode->RequestLayersContentUpdate(true); });
+			FScopedSetLandscapeEditingLayer Scope(LandscapeEdMode->GetLandscape(), LandscapeEdMode->GetCurrentLayerGuid(), [&] { LandscapeEdMode->RequestLayersContentUpdateForceAll(); });
 			FLandscapeEditDataInterface LandscapeEdit(Target->LandscapeInfo.Get());
 			LandscapeEdit.DeleteLayer(Target->LayerInfoObj.Get());
 		}
diff --git a/Engine/Source/Runtime/Landscape/Classes/Landscape.h b/Engine/Source/Runtime/Landscape/Classes/Landscape.h
--- a/Engine/Source/Runtime/Landscape/Classes/Landscape.h
+++ b/Engine/Source/Runtime/Landscape/Classes/Landscape.h
@@ -72,34 +72,9 @@ enum EWeightmapRTType : uint8
 	WeightmapRT_Count
 };
 
-// Internal Update Flags that shouldn't be used alone
-enum EInternalLandscapeLayersContentUpdateFlag : uint32
-{
-	Internal_Heightmap_Render = 1 << 0, 
-	Internal_Heightmap_BoundsAndCollision = 1 << 1,
-	Internal_Heightmap_BoundsAndCollisionPartial = 1 << 2,
-	Internal_Heightmap_ResolveToTexture = 1 << 3,
-	
-	Internal_Weightmap_Render = 1 << 8,
-	Internal_Weightmap_Collision = 1 << 9,
-	Internal_Weightmap_ResolveToTexture = 1 << 10,
-
-	Internal_ClientUpdate = 1 << 14
-};
-
-enum ELandscapeLayersContentUpdateFlag : uint32
-{
-	Heightmap_Render = Internal_Heightmap_Render,
-	Heightmap_Editing = Internal_Heightmap_Render | Internal_Heightmap_ResolveToTexture | Internal_Heightmap_BoundsAndCollisionPartial,
-	Weightmap_Render = Internal_Weightmap_Render,
-
-	// Combinations
-	Heightmap_All = Internal_Heightmap_Render | Internal_Heightmap_BoundsAndCollision | Internal_Heightmap_ResolveToTexture | Internal_ClientUpdate,
-	Weightmap_All = Internal_Weightmap_Render | Internal_Weightmap_Collision | Internal_Weightmap_ResolveToTexture | Internal_ClientUpdate,
-
-	All = Heightmap_All | Weightmap_All,
-	All_Render = Heightmap_Render | Weightmap_Render,
-};
+#if WITH_EDITOR
+enum ELandscapeLayerUpdateMode : uint32;
+#endif
 
 USTRUCT()
 struct FLandscapeLayerBrush
@@ -251,8 +226,9 @@ public:
 
 	// Layers stuff
 #if WITH_EDITOR
-	LANDSCAPE_API void RequestLayersInitialization() { bLandscapeLayersAreInitialized = false; RequestLayersContentUpdate(ELandscapeLayersContentUpdateFlag::All); }
-	LANDSCAPE_API void RequestLayersContentUpdate(ELandscapeLayersContentUpdateFlag InDataFlags, bool InUpdateAllMaterials = false);
+	LANDSCAPE_API void RequestLayersInitialization();
+	LANDSCAPE_API void RequestLayersContentUpdateForceAll();
+	LANDSCAPE_API void RequestLayersContentUpdate(ELandscapeLayerUpdateMode InModeMask, bool bInForceUpdateAllComponents = false);
 	LANDSCAPE_API bool ReorderLayer(int32 InStartingLayerIndex, int32 InDestinationLayerIndex);
 	LANDSCAPE_API void CreateLayer(FName InName = NAME_None);
 	LANDSCAPE_API void CreateDefaultLayer();
@@ -297,9 +273,8 @@ private:
 	void TickLayers(float DeltaTime, ELevelTick TickType, FActorTickFunction& ThisTickFunction);
 	void CreateLayersRenderingResource(const FIntPoint& InComponentCounts);
 	void RegenerateLayersContent();
-	void RegenerateLayersHeightmaps(const TArray<ULandscapeComponent*>& InLandscapeComponents, bool& bRecreateCollision);
-	void RegenerateLayersWeightmaps(const TArray<ULandscapeComponent*>& InLandscapeComponents, bool& bRecreateCollision);
-	bool UpdateClients(const TArray<ULandscapeComponent*>& InLandscapeComponents, bool bCollisionRecreated);
+	void RegenerateLayersHeightmaps(const TArray<ULandscapeComponent*>& InLandscapeComponents);
+	void RegenerateLayersWeightmaps(const TArray<ULandscapeComponent*>& InLandscapeComponents);
 	void ResolveLayersHeightmapTexture();
 	void ResolveLayersWeightmapTexture();
 	void ResolveLayersTexture(class FLandscapeLayersTexture2DCPUReadBackResource* InCPUReadBackTexture, UTexture2D* InOutputTexture);
@@ -381,7 +356,10 @@ private:
 	bool WasCompilingShaders;
 
 	UPROPERTY(Transient)
-	uint32 LayersContentUpdateFlags;
+	uint32 LayerContentUpdateModes;
+
+	UPROPERTY(Transient)
+	bool bLayerForceUpdateAllComponents;
 
 	// Represent all the resolved paint layer, from all layers blended together (size of the landscape x material layer count)
 	class FLandscapeTexture2DArrayResource* CombinedLayersWeightmapAllMaterialLayersResource;
diff --git a/Engine/Source/Runtime/Landscape/Classes/LandscapeComponent.h b/Engine/Source/Runtime/Landscape/Classes/LandscapeComponent.h
--- a/Engine/Source/Runtime/Landscape/Classes/LandscapeComponent.h
+++ b/Engine/Source/Runtime/Landscape/Classes/LandscapeComponent.h
@@ -290,6 +290,33 @@ struct FLandscapeLayerComponentData
 	bool IsInitialized() const { return HeightmapData.Texture != nullptr || WeightmapData.Textures.Num() > 0;  }
 };
 
+#if WITH_EDITOR
+enum ELandscapeComponentUpdateFlag : uint32
+{
+	// Will call UpdateCollisionHeightData, UpdateCacheBounds, UpdateComponentToWorld on Component
+	Update_Heightmap_Collision = 1 << 0,
+	// Will call UdateCollisionLayerData on Component
+	Update_Weightmap_Collision = 1 << 1,
+	// Will call RecreateCollision on Component
+	RecreateCollision = 1 << 2,
+	// Will update Component clients: Navigation data, Foliage, Grass, etc.
+	ClientUpdate = 1 << 3
+};
+
+enum ELandscapeLayerUpdateMode : uint32
+{ 
+	Heightmap_All = 1 << 0,
+	Heightmap_Editing = 1 << 1,
+	Weightmap_All = 1 << 2,
+	Weightmap_Editing = 1 << 3,
+	All = Weightmap_All | Heightmap_All,
+	All_Editing = Weightmap_Editing | Heightmap_Editing,
+	// In cases where we couldn't update the clients right away this flag will be set in RegenerateLayersContent
+	DeferredClientUpdate = 1 << 4,
+};
+
+#endif
+
 UCLASS(hidecategories=(Display, Attachment, Physics, Debug, Collision, Movement, Rendering, PrimitiveComponent, Object, Transform, Mobility), showcategories=("Rendering|Material"), MinimalAPI, Within=LandscapeProxy)
 class ULandscapeComponent : public UPrimitiveComponent
 {
@@ -380,7 +407,8 @@ private:
 	UPROPERTY(Transient)
 	TArray<ULandscapeWeightmapUsage*> WeightmapTexturesUsage;
 
-	bool bLayerContentDirty;
+	UPROPERTY(Transient)
+	uint32 LayerUpdateFlagPerMode;
 #endif // WITH_EDITORONLY_DATA
 
 	/** Heightmap texture reference */
@@ -841,8 +869,11 @@ public:
 	/** Updates the values of component-level properties exposed by the Landscape Actor */
 	LANDSCAPE_API void UpdatedSharedPropertiesFromActor();
 
-	LANDSCAPE_API bool IsLayerContentDirty() const { return bLayerContentDirty; }
-	LANDSCAPE_API void SetLayerContentDirty(bool InValue) { bLayerContentDirty = InValue; }
+	LANDSCAPE_API bool IsUpdateFlagEnabledForModes(ELandscapeComponentUpdateFlag InFlag, uint32 InModeMask, bool bForce) const;
+	LANDSCAPE_API void ClearUpdateFlagsForModes(uint32 InModeMask);
+	LANDSCAPE_API void RequestWeightmapUpdate(bool bUpdateAll = false);
+	LANDSCAPE_API void RequestHeightmapUpdate(bool bUpdateAll = false);
+	LANDSCAPE_API void RequestDeferredClientUpdate();
 #endif
 
 	friend class FLandscapeComponentSceneProxy;
diff --git a/Engine/Source/Runtime/Landscape/Private/Landscape.cpp b/Engine/Source/Runtime/Landscape/Private/Landscape.cpp
--- a/Engine/Source/Runtime/Landscape/Private/Landscape.cpp
+++ b/Engine/Source/Runtime/Landscape/Private/Landscape.cpp
@@ -141,7 +141,7 @@ FAutoConsoleCommand CmdPrintNumLandscapeShadows(
 ULandscapeComponent::ULandscapeComponent(const FObjectInitializer& ObjectInitializer)
 : Super(ObjectInitializer)
 #if WITH_EDITORONLY_DATA
-, bLayerContentDirty(false)
+, LayerUpdateFlagPerMode(0)
 #endif
 , GrassData(MakeShareable(new FLandscapeComponentGrassData()))
 , ChangeTag(0)
@@ -1076,7 +1076,8 @@ ALandscape::ALandscape(const FObjectInitializer& ObjectInitializer)
 	bLockLocation = false;
 	PreviousExperimentalLandscapeLayers = false;
 	WasCompilingShaders = false;
-	LayersContentUpdateFlags = 0;
+	LayerContentUpdateModes = 0;
+	bLayerForceUpdateAllComponents = false;
 	CombinedLayersWeightmapAllMaterialLayersResource = nullptr;
 	CurrentLayersWeightmapAllMaterialLayersResource = nullptr;
 	WeightmapScratchExtractLayerTextureResource = nullptr;
diff --git a/Engine/Source/Runtime/Landscape/Private/LandscapeBPCustomBrush.cpp b/Engine/Source/Runtime/Landscape/Private/LandscapeBPCustomBrush.cpp
--- a/Engine/Source/Runtime/Landscape/Private/LandscapeBPCustomBrush.cpp
+++ b/Engine/Source/Runtime/Landscape/Private/LandscapeBPCustomBrush.cpp
@@ -78,7 +78,7 @@ void ALandscapeBlueprintCustomBrush::PostEditMove(bool bFinished)
 
 	if (OwningLandscape != nullptr)
 	{
-		OwningLandscape->RequestLayersContentUpdate(bFinished ? ELandscapeLayersContentUpdateFlag::All : ELandscapeLayersContentUpdateFlag::All_Render);
+		OwningLandscape->RequestLayersContentUpdate(bFinished ? ELandscapeLayerUpdateMode::All : ELandscapeLayerUpdateMode::All_Editing);
 	}
 }
 
@@ -157,7 +157,7 @@ void ALandscapeBlueprintCustomBrush::PostEditChangeProperty(FPropertyChangedEven
 
 	if (OwningLandscape != nullptr)
 	{
-		OwningLandscape->RequestLayersContentUpdate(ELandscapeLayersContentUpdateFlag::All, true); // For now since we dont have brush bounds, we have to force an update of all components, otherwise we will have not right layer being rendered as materials are not updated
+		OwningLandscape->RequestLayersContentUpdateForceAll(); // For now since we dont have brush bounds, we have to force an update of all components, otherwise we will have not right layer being rendered as materials are not updated
 	}
 }
 #endif
diff --git a/Engine/Source/Runtime/Landscape/Private/LandscapeEdit.cpp b/Engine/Source/Runtime/Landscape/Private/LandscapeEdit.cpp
--- a/Engine/Source/Runtime/Landscape/Private/LandscapeEdit.cpp
+++ b/Engine/Source/Runtime/Landscape/Private/LandscapeEdit.cpp
@@ -579,12 +579,9 @@ void ULandscapeComponent::PostEditUndo()
 		SingleComponent.Add(this);
 		GetLandscapeProxy()->InvalidateGeneratedComponentData(SingleComponent);
 		
-		ALandscape* LandscapeActor = GetLandscapeActor();
-		// Might be a ALandscapeStreamingProxy
-		if (LandscapeActor)
-		{
-			LandscapeActor->RequestLayersContentUpdate(ELandscapeLayersContentUpdateFlag::All, true);
-		}
+		const bool bUpdateAll = true;
+		RequestHeightmapUpdate(bUpdateAll);
+		RequestWeightmapUpdate(bUpdateAll);
 	}
 	else
 	{
@@ -3508,7 +3505,7 @@ void ALandscape::PostEditUndo()
 {
 	Super::PostEditUndo();
 
-	RequestLayersContentUpdate(ELandscapeLayersContentUpdateFlag::All, true);
+	RequestLayersContentUpdate(ELandscapeLayerUpdateMode::All);
 }
 
 bool ALandscape::ShouldImport(FString* ActorPropString, bool IsMovingLevel)
@@ -3955,7 +3952,7 @@ void ALandscapeProxy::PostEditChangeProperty(FPropertyChangedEvent& PropertyChan
 		{
 			if(ALandscape* LandscapeActor = GetLandscapeActor())
 			{
-				LandscapeActor->RequestLayersContentUpdate(ELandscapeLayersContentUpdateFlag::All, true);
+				LandscapeActor->RequestLayersContentUpdate(ELandscapeLayerUpdateMode::All);
 			}
 		}
 		
diff --git a/Engine/Source/Runtime/Landscape/Private/LandscapeEditInterface.cpp b/Engine/Source/Runtime/Landscape/Private/LandscapeEditInterface.cpp
--- a/Engine/Source/Runtime/Landscape/Private/LandscapeEditInterface.cpp
+++ b/Engine/Source/Runtime/Landscape/Private/LandscapeEditInterface.cpp
@@ -396,34 +396,35 @@ void FLandscapeEditDataInterface::SetHeightData(int32 X1, int32 Y1, int32 X2, in
 					TexDataInfo->AddMipUpdateRegion(0, TexX1, TexY1, TexX2, TexY2);
 				}
 			}
-
-			// See if we need to adjust the bounds. Note we never shrink the bounding box at this point
-			bool bUpdateBoxSphereBounds = false;
-
-			if (InUpdateBounds)
+												
+			Component->RequestHeightmapUpdate();
+			if (!GetMutableDefault<UEditorExperimentalSettings>()->bLandscapeLayerSystem)
 			{
-				float MinLocalZ = LandscapeDataAccess::GetLocalHeight(MinHeight);
-				float MaxLocalZ = LandscapeDataAccess::GetLocalHeight(MaxHeight);
+				// See if we need to adjust the bounds. Note we never shrink the bounding box at this point
+				bool bUpdateBoxSphereBounds = false;
 
-				if (MinLocalZ < Component->CachedLocalBox.Min.Z)
-				{
-					Component->CachedLocalBox.Min.Z = MinLocalZ;
-					bUpdateBoxSphereBounds = true;
-				}
-				if (MaxLocalZ > Component->CachedLocalBox.Max.Z)
+				if (InUpdateBounds)
 				{
-					Component->CachedLocalBox.Max.Z = MaxLocalZ;
-					bUpdateBoxSphereBounds = true;
-				}
+					float MinLocalZ = LandscapeDataAccess::GetLocalHeight(MinHeight);
+					float MaxLocalZ = LandscapeDataAccess::GetLocalHeight(MaxHeight);
 
-				if (bUpdateBoxSphereBounds)
-				{
-					Component->UpdateComponentToWorld();
+					if (MinLocalZ < Component->CachedLocalBox.Min.Z)
+					{
+						Component->CachedLocalBox.Min.Z = MinLocalZ;
+						bUpdateBoxSphereBounds = true;
+					}
+					if (MaxLocalZ > Component->CachedLocalBox.Max.Z)
+					{
+						Component->CachedLocalBox.Max.Z = MaxLocalZ;
+						bUpdateBoxSphereBounds = true;
+					}
+
+					if (bUpdateBoxSphereBounds)
+					{
+						Component->UpdateComponentToWorld();
+					}
 				}
-			}
 
-			if (!GetMutableDefault<UEditorExperimentalSettings>()->bLandscapeLayerSystem)
-			{
 				// Update mipmaps
 
 				// Work out how many mips should be calculated directly from one component's data.
@@ -452,15 +453,7 @@ void FLandscapeEditDataInterface::SetHeightData(int32 X1, int32 Y1, int32 X2, in
 					}
 				}
 			}
-			else
-			{
-				if (ALandscape* Landscape = Component->GetLandscapeActor())
-				{
-					Component->SetLayerContentDirty(true);
-					Landscape->RequestLayersContentUpdate(ELandscapeLayersContentUpdateFlag::Heightmap_Editing);
-				}
-			}
-
+			
 			// Update GUID for Platform Data
 			FPlatformMisc::CreateGuid(Component->StateId);
 		}
@@ -1751,7 +1744,7 @@ void ULandscapeComponent::DeleteLayer(ULandscapeLayerInfoObject* LayerInfo, FLan
 	Component->UpdateMaterialInstances();
 	Component->EditToolRenderData.UpdateDebugColorMaterial(Component);
 	Component->UpdateEditToolRenderData();
-
+	Component->RequestWeightmapUpdate();
 	
 	// Update dominant layer info stored in collision component
 	if (!GetMutableDefault<UEditorExperimentalSettings>()->bLandscapeLayerSystem)
@@ -1798,12 +1791,9 @@ void FLandscapeEditDataInterface::DeleteLayer(ULandscapeLayerInfoObject* LayerIn
 	Algo::Transform(LandscapeInfo->XYtoComponentMap, Components, &TPair<FIntPoint, ULandscapeComponent*>::Value);
 	ALandscapeProxy::InvalidateGeneratedComponentData(Components);
 
-	if (GetMutableDefault<UEditorExperimentalSettings>()->bLandscapeLayerSystem)
+	if (LandscapeInfo->LandscapeActor)
 	{
-		if (LandscapeInfo->LandscapeActor)
-		{
-			LandscapeInfo->LandscapeActor->RequestLayersContentUpdate(ELandscapeLayersContentUpdateFlag::All, true);
-		}
+		LandscapeInfo->LandscapeActor->RequestLayersContentUpdate(ELandscapeLayerUpdateMode::All, true);
 	}
 }
 
@@ -1956,13 +1946,10 @@ void ULandscapeComponent::FillLayer(ULandscapeLayerInfoObject* LayerInfo, FLands
 	Component->UpdateEditToolRenderData();
 
 	Component->InvalidateLightingCache();
+	Component->RequestWeightmapUpdate();
 
 	// Update dominant layer info stored in collision component
-	if (GetMutableDefault<UEditorExperimentalSettings>()->bLandscapeLayerSystem)
-	{
-		GetLandscapeActor()->RequestLayersContentUpdate(ELandscapeLayersContentUpdateFlag::All);
-	}
-	else
+	if (!GetMutableDefault<UEditorExperimentalSettings>()->bLandscapeLayerSystem)
 	{
 		TArray<FColor*> CollisionWeightmapMipData;
 		for (int32 WeightmapIdx = 0; WeightmapIdx < ComponentWeightmapTextures.Num(); WeightmapIdx++)
@@ -2038,13 +2025,9 @@ void FLandscapeEditDataInterface::FillLayer(ULandscapeLayerInfoObject* LayerInfo
 	TSet<ULandscapeComponent*> Components;
 	Algo::Transform(LandscapeInfo->XYtoComponentMap, Components, &TPair<FIntPoint, ULandscapeComponent*>::Value);
 	ALandscapeProxy::InvalidateGeneratedComponentData(Components);
-
-	if (GetMutableDefault<UEditorExperimentalSettings>()->bLandscapeLayerSystem)
+	if (LandscapeInfo->LandscapeActor)
 	{
-		if (LandscapeInfo->LandscapeActor)
-		{
-			LandscapeInfo->LandscapeActor->RequestLayersContentUpdate(ELandscapeLayersContentUpdateFlag::All, true);
-		}
+		LandscapeInfo->LandscapeActor->RequestLayersContentUpdate(ELandscapeLayerUpdateMode::All, true);
 	}
 }
 
@@ -2236,6 +2219,35 @@ void ULandscapeComponent::ReplaceLayer(ULandscapeLayerInfoObject* FromLayerInfo,
 		EditToolRenderData.UpdateDebugColorMaterial(this);
 		UpdateEditToolRenderData();
 	}
+
+	RequestWeightmapUpdate();
+
+	if (!GetMutableDefault<UEditorExperimentalSettings>()->bLandscapeLayerSystem)
+	{
+		// Update dominant layer info stored in collision component
+		TArray<FColor*> CollisionWeightmapMipData;
+		for (int32 WeightmapIdx = 0; WeightmapIdx < ComponentWeightmapTextures.Num(); WeightmapIdx++)
+		{
+			CollisionWeightmapMipData.Add((FColor*)LandscapeEdit.GetTextureDataInfo(ComponentWeightmapTextures[WeightmapIdx])->GetMipData(CollisionMipLevel));
+		}
+		TArray<FColor*> SimpleCollisionWeightmapMipData;
+		if (SimpleCollisionMipLevel > CollisionMipLevel)
+		{
+			for (int32 WeightmapIdx = 0; WeightmapIdx < ComponentWeightmapTextures.Num(); WeightmapIdx++)
+			{
+				SimpleCollisionWeightmapMipData.Add((FColor*)LandscapeEdit.GetTextureDataInfo(ComponentWeightmapTextures[WeightmapIdx])->GetMipData(SimpleCollisionMipLevel));
+			}
+		}
+		UpdateCollisionLayerData(
+			CollisionWeightmapMipData.GetData(),
+			SimpleCollisionMipLevel > CollisionMipLevel ? SimpleCollisionWeightmapMipData.GetData() : nullptr);
+
+
+		if(ULandscapeHeightfieldCollisionComponent* CollisionComp = CollisionComponent.Get())
+		{
+			CollisionComp->RecreateCollision();
+		}
+	}
 }
 
 void FLandscapeEditDataInterface::ReplaceLayer(ULandscapeLayerInfoObject* FromLayerInfo, ULandscapeLayerInfoObject* ToLayerInfo)
@@ -2248,36 +2260,6 @@ void FLandscapeEditDataInterface::ReplaceLayer(ULandscapeLayerInfoObject* FromLa
 		{
 			ULandscapeComponent* Component = It.Value();
 			Component->ReplaceLayer(FromLayerInfo, ToLayerInfo, *this);
-
-			if(!GetMutableDefault<UEditorExperimentalSettings>()->bLandscapeLayerSystem)
-			{
-				TArray<UTexture2D*>& ComponentWeightmapTextures = Component->GetWeightmapTextures(true);
-
-				// Update dominant layer info stored in collision component
-				TArray<FColor*> CollisionWeightmapMipData;
-				for (int32 WeightmapIdx = 0; WeightmapIdx < ComponentWeightmapTextures.Num(); WeightmapIdx++)
-				{
-					CollisionWeightmapMipData.Add((FColor*)GetTextureDataInfo(ComponentWeightmapTextures[WeightmapIdx])->GetMipData(Component->CollisionMipLevel));
-				}
-				TArray<FColor*> SimpleCollisionWeightmapMipData;
-				if (Component->SimpleCollisionMipLevel > Component->CollisionMipLevel)
-				{
-					for (int32 WeightmapIdx = 0; WeightmapIdx < ComponentWeightmapTextures.Num(); WeightmapIdx++)
-					{
-						SimpleCollisionWeightmapMipData.Add((FColor*)GetTextureDataInfo(ComponentWeightmapTextures[WeightmapIdx])->GetMipData(Component->SimpleCollisionMipLevel));
-					}
-				}
-				Component->UpdateCollisionLayerData(
-					CollisionWeightmapMipData.GetData(),
-					Component->SimpleCollisionMipLevel > Component->CollisionMipLevel ? SimpleCollisionWeightmapMipData.GetData() : nullptr);
-
-
-				ULandscapeHeightfieldCollisionComponent* CollisionComponent = Component->CollisionComponent.Get();
-				if (CollisionComponent)
-				{
-					CollisionComponent->RecreateCollision();
-				}
-			}
 		}
 	};
 
@@ -2297,7 +2279,7 @@ void FLandscapeEditDataInterface::ReplaceLayer(ULandscapeLayerInfoObject* FromLa
 				FScopedSetLandscapeEditingLayer Scope(LandscapeInfo->LandscapeActor.Get(), CurrentLayer.Guid);
 				DoReplace();
 			});
-			LandscapeInfo->LandscapeActor->RequestLayersContentUpdate(ELandscapeLayersContentUpdateFlag::All, true);
+			LandscapeInfo->LandscapeActor->RequestLayersContentUpdate(ELandscapeLayerUpdateMode::Weightmap_All, true);
 		}
 	}
 	else
diff --git a/Engine/Source/Runtime/Landscape/Private/LandscapeEditLayers.cpp b/Engine/Source/Runtime/Landscape/Private/LandscapeEditLayers.cpp
--- a/Engine/Source/Runtime/Landscape/Private/LandscapeEditLayers.cpp
+++ b/Engine/Source/Runtime/Landscape/Private/LandscapeEditLayers.cpp
@@ -2446,20 +2446,22 @@ bool ALandscape::PrepareLayersHeightmapTextureResources() const
 	return IsReady;
 }
 
-void ALandscape::RegenerateLayersHeightmaps(const TArray<ULandscapeComponent*>& InLandscapeComponents, bool& bOutRecreateCollision)
+void ALandscape::RegenerateLayersHeightmaps(const TArray<ULandscapeComponent*>& InLandscapeComponents)
 {
 	SCOPE_CYCLE_COUNTER(STAT_LandscapeLayersRegenerateHeightmaps);
-	bOutRecreateCollision = false;
-
 	ULandscapeInfo* Info = GetLandscapeInfo();
 
-	if (LayersContentUpdateFlags == 0 || Info == nullptr || !PrepareLayersHeightmapTextureResources())
+	const int32 AllHeightmapUpdateModes = (ELandscapeLayerUpdateMode::Heightmap_All |  ELandscapeLayerUpdateMode::Heightmap_Editing);
+	const int32 HeightmapUpdateModes = LayerContentUpdateModes & AllHeightmapUpdateModes;
+	const bool bForceRender = CVarOutputLayersDebugDrawCallName.GetValueOnAnyThread() == 1;
+
+	if ((HeightmapUpdateModes == 0 && !bForceRender) || Info == nullptr || !PrepareLayersHeightmapTextureResources())
 	{
 		return;
 	}
-
+	   
 	// Handle missing Heightmap CPUReadback (Undo of a delete landscape component can trigger that case)
-	if (LayersContentUpdateFlags & ELandscapeLayersContentUpdateFlag::Heightmap_All)
+	if (HeightmapUpdateModes)
 	{
 		Info->ForAllLandscapeProxies([&](ALandscapeProxy* Proxy)
 		{
@@ -2477,7 +2479,7 @@ void ALandscape::RegenerateLayersHeightmaps(const TArray<ULandscapeComponent*>&
 		});
 	}
 
-	if (LayersContentUpdateFlags & ELandscapeLayersContentUpdateFlag::Heightmap_Render)
+	if (HeightmapUpdateModes || bForceRender)
 	{
 		check(HeightmapRTList.Num() > 0);
 
@@ -2670,70 +2672,23 @@ void ALandscape::RegenerateLayersHeightmaps(const TArray<ULandscapeComponent*>&
 		});
 	}
 
-	if (LayersContentUpdateFlags & EInternalLandscapeLayersContentUpdateFlag::Internal_Heightmap_ResolveToTexture)
+	if (HeightmapUpdateModes)
 	{
 		ResolveLayersHeightmapTexture();
-	}
-
-	bool bUpdateBoundsAndCollsionFull = (LayersContentUpdateFlags & EInternalLandscapeLayersContentUpdateFlag::Internal_Heightmap_BoundsAndCollision) != 0;
-	bool bUpdateBoundsAndCollsionPartial = ((LayersContentUpdateFlags & EInternalLandscapeLayersContentUpdateFlag::Internal_Heightmap_BoundsAndCollisionPartial) != 0) && !bUpdateBoundsAndCollsionFull; // Partial only if not full update
-	bool bUpdateBoundsAndCollsion = bUpdateBoundsAndCollsionFull || bUpdateBoundsAndCollsionPartial;
 
-	if (bUpdateBoundsAndCollsion)
-	{
+		// Partial Component Update
 		for (ULandscapeComponent* Component : InLandscapeComponents)
 		{
-			if (bUpdateBoundsAndCollsion && (!bUpdateBoundsAndCollsionPartial || Component->IsLayerContentDirty()))
+			if(Component->IsUpdateFlagEnabledForModes(ELandscapeComponentUpdateFlag::Update_Heightmap_Collision, HeightmapUpdateModes, bLayerForceUpdateAllComponents))
 			{
 				Component->UpdateCachedBounds();
 				Component->UpdateComponentToWorld();
 				Component->UpdateCollisionData();
 			}
-
-			Component->SetLayerContentDirty(false);
 		}
 
-		// Ideally, we should only update necessary components and only when inside AddComponnet Tool.
 		Info->UpdateAllAddCollisions();
-	
-		bOutRecreateCollision = !bUpdateBoundsAndCollsionPartial;
-	}
-		
-	LayersContentUpdateFlags &= ~EInternalLandscapeLayersContentUpdateFlag::Internal_Heightmap_BoundsAndCollisionPartial;
-	LayersContentUpdateFlags &= ~ELandscapeLayersContentUpdateFlag::Heightmap_All;
-		
-	// If doing rendering debug, keep doing the render only
-	if (CVarOutputLayersDebugDrawCallName.GetValueOnAnyThread() == 1)
-	{
-		LayersContentUpdateFlags |= ELandscapeLayersContentUpdateFlag::Heightmap_Render;
-	}
-			}
-
-bool ALandscape::UpdateClients(const TArray<ULandscapeComponent*>& InLandscapeComponents, bool bInCollisionRecreated)
-			{
-				// Post-pone client update to next tick while in transaction
-				if (GUndo)
-				{
-		return false;
-				}
-				else
-				{
-		for (ULandscapeComponent* Component : InLandscapeComponents)
-		{
-					// Recreate collision for modified components to update the physical materials
-					ULandscapeHeightfieldCollisionComponent* CollisionComponent = Component->CollisionComponent.Get();
-					if (CollisionComponent)
-					{
-						FNavigationSystem::UpdateComponentData(*CollisionComponent);
-						CollisionComponent->SnapFoliageInstances();
-					}
-				}
-
-		ALandscapeProxy::InvalidateGeneratedComponentData(InLandscapeComponents);
 	}
-
-	// Re update clients if collision wasn't recreated
-	return bInCollisionRecreated;
 }
 
 void ALandscape::ResolveLayersHeightmapTexture()
@@ -3332,25 +3287,26 @@ bool ALandscape::PrepareLayersWeightmapTextureResources() const
 	return IsReady;
 }
 
-void ALandscape::RegenerateLayersWeightmaps(const TArray<ULandscapeComponent*>& InLandscapeComponents, bool& bOutRecreateCollision)
+void ALandscape::RegenerateLayersWeightmaps(const TArray<ULandscapeComponent*>& InLandscapeComponents)
 {
 	SCOPE_CYCLE_COUNTER(STAT_LandscapeLayersRegenerateWeightmaps);
-	bOutRecreateCollision = false;
-
+	const int32 AllWeightmapUpdateModes = (ELandscapeLayerUpdateMode::Weightmap_All | ELandscapeLayerUpdateMode::Weightmap_Editing);
+	const int32 WeightmapUpdateModes = LayerContentUpdateModes & AllWeightmapUpdateModes;
+	const bool bForceRender = CVarOutputLayersDebugDrawCallName.GetValueOnAnyThread() == 1;
+	
 	ULandscapeInfo* Info = GetLandscapeInfo();
 
-	if (LayersContentUpdateFlags == 0 || Info == nullptr || Info->Layers.Num() == 0 || !PrepareLayersWeightmapTextureResources())
+	if ((WeightmapUpdateModes == 0 && !bForceRender) || Info == nullptr || Info->Layers.Num() == 0 || !PrepareLayersWeightmapTextureResources())
 	{
 		return;
 	}
-
-
+			
 	TArray<ULandscapeComponent*> ComponentThatNeedMaterialRebuild;
 	TArray<ULandscapeLayerInfoObject*> BrushRequiredAllocations;
 	int32 LayerCount = Info->Layers.Num() + 1; // due to visibility being stored at 0
 	bool ClearFlagsAfterUpdate = true;
 
-	if (LayersContentUpdateFlags & ELandscapeLayersContentUpdateFlag::Weightmap_Render)
+	if (WeightmapUpdateModes || bForceRender)
 	{
 		FIntPoint MinExtend;
 		FIntPoint MaxExtend;
@@ -3696,30 +3652,17 @@ void ALandscape::RegenerateLayersWeightmaps(const TArray<ULandscapeComponent*>&
 		UpdateLayersMaterialInstances(ComponentThatNeedMaterialRebuild);
 	}
 
-	if (LayersContentUpdateFlags & EInternalLandscapeLayersContentUpdateFlag::Internal_Weightmap_ResolveToTexture)
+	if (WeightmapUpdateModes)
 	{
 		ResolveLayersWeightmapTexture();
-	}
-
-	if (LayersContentUpdateFlags & EInternalLandscapeLayersContentUpdateFlag::Internal_Weightmap_Collision)
-	{
+	
 		for (ULandscapeComponent* Component : InLandscapeComponents)
 		{
-			Component->UpdateCollisionLayerData();
+			if (Component->IsUpdateFlagEnabledForModes(ELandscapeComponentUpdateFlag::Update_Weightmap_Collision, WeightmapUpdateModes, bLayerForceUpdateAllComponents))
+			{
+				Component->UpdateCollisionLayerData();
+			}
 		}
-
-		bOutRecreateCollision = true;
-	}
-
-	if (ClearFlagsAfterUpdate)
-	{
-		LayersContentUpdateFlags &= ~ELandscapeLayersContentUpdateFlag::Weightmap_All;
-	}
-
-	// If doing rendering debug, keep doing the render only
-	if (CVarOutputLayersDebugDrawCallName.GetValueOnAnyThread() == 1)
-	{
-		LayersContentUpdateFlags |= ELandscapeLayersContentUpdateFlag::Weightmap_Render;
 	}
 }
 
@@ -3922,14 +3865,134 @@ void ALandscape::ResolveLayersWeightmapTexture()
 	});
 }
 
-void ALandscape::RequestLayersContentUpdate(ELandscapeLayersContentUpdateFlag InDataFlags, bool InUpdateAllMaterials)
+
+
+void ALandscape::RequestLayersInitialization()
 {
-	LayersContentUpdateFlags |= InDataFlags;
+	if (!GetMutableDefault<UEditorExperimentalSettings>()->bLandscapeLayerSystem)
+	{
+		return;
+	}
+
+	bLandscapeLayersAreInitialized = false; 
+	RequestLayersContentUpdateForceAll();
+}
+
+void ALandscape::RequestLayersContentUpdate(ELandscapeLayerUpdateMode InUpdateMode, bool bInForceUpdateAllComponents)
+{
+	LayerContentUpdateModes |= InUpdateMode;
+	bLayerForceUpdateAllComponents |= bInForceUpdateAllComponents;
+}
+
+void ALandscape::RequestLayersContentUpdateForceAll()
+{
+	if (!GetMutableDefault<UEditorExperimentalSettings>()->bLandscapeLayerSystem)
+	{
+		return;
+	}
+
+	if (ULandscapeInfo* LandscapeInfo = GetLandscapeInfo())
+	{
+		LandscapeInfo->ForAllLandscapeProxies([](ALandscapeProxy* Proxy)
+		{
+			if (Proxy)
+			{
+				Proxy->InvalidateGeneratedComponentData();
+			}
+		});
+	}
+	
+	const bool bLocalForceUpdateAllComponents = true;
+	RequestLayersContentUpdate(ELandscapeLayerUpdateMode::All, bLocalForceUpdateAllComponents);
+}
+
+bool ULandscapeComponent::IsUpdateFlagEnabledForModes(ELandscapeComponentUpdateFlag InFlag, uint32 InModeMask, bool bForce) const
+{
+	uint32 UpdateMode = bForce? InModeMask : (LayerUpdateFlagPerMode & InModeMask);
+	
+	if (UpdateMode & ELandscapeLayerUpdateMode::Heightmap_All)
+	{
+		const uint32 HeightmapAllFlags = ELandscapeComponentUpdateFlag::Update_Heightmap_Collision | ELandscapeComponentUpdateFlag::RecreateCollision | ELandscapeComponentUpdateFlag::ClientUpdate;
+		if (HeightmapAllFlags & InFlag)
+		{
+			return true;
+		}
+	}
+		
+	if (UpdateMode & ELandscapeLayerUpdateMode::Heightmap_Editing)
+	{
+		const uint32 HeightmapEditingFlags = ELandscapeComponentUpdateFlag::Update_Heightmap_Collision;
+		if (HeightmapEditingFlags & InFlag)
+		{
+			return true;
+		}
+	}
+		
+	if (UpdateMode & ELandscapeLayerUpdateMode::Weightmap_All)
+	{
+		const uint32 WeightmapAllFlags = ELandscapeComponentUpdateFlag::Update_Weightmap_Collision | ELandscapeComponentUpdateFlag::RecreateCollision | ELandscapeComponentUpdateFlag::ClientUpdate;
+		if (WeightmapAllFlags & InFlag)
+		{
+			return true;
+		}
+	}
+
+	if (UpdateMode & ELandscapeLayerUpdateMode::Weightmap_Editing)
+	{
+		const uint32 WeightmapEditingFlags = ELandscapeComponentUpdateFlag::Update_Weightmap_Collision;
+		if (WeightmapEditingFlags & InFlag)
+		{
+			return true;
+		}
+	}
+
+	if (UpdateMode & ELandscapeLayerUpdateMode::DeferredClientUpdate)
+	{
+		const uint32 DeferredClientUpdateFlags = ELandscapeComponentUpdateFlag::ClientUpdate;
+		if (DeferredClientUpdateFlags & InFlag)
+		{
+			return true;
+		}
+	}
+
+	return false;
+}
+
+void ULandscapeComponent::ClearUpdateFlagsForModes(uint32 InModeMask)
+{
+	LayerUpdateFlagPerMode &= ~InModeMask;
+}
+
+void ULandscapeComponent::RequestDeferredClientUpdate()
+{
+	LayerUpdateFlagPerMode |= ELandscapeLayerUpdateMode::DeferredClientUpdate;
+}
+
+void ULandscapeComponent::RequestHeightmapUpdate(bool bUpdateAll)
+{
+	LayerUpdateFlagPerMode |= ELandscapeLayerUpdateMode::Heightmap_Editing;
+	LayerUpdateFlagPerMode |= ELandscapeLayerUpdateMode::Heightmap_All;
+	GetLandscapeActor()->RequestLayersContentUpdate(ELandscapeLayerUpdateMode::Heightmap_Editing);
+	if (bUpdateAll)
+	{
+		GetLandscapeActor()->RequestLayersContentUpdate(Heightmap_All);
+	}
+}
+
+void ULandscapeComponent::RequestWeightmapUpdate(bool bUpdateAll)
+{
+	LayerUpdateFlagPerMode |= ELandscapeLayerUpdateMode::Weightmap_Editing;
+	LayerUpdateFlagPerMode |= ELandscapeLayerUpdateMode::Weightmap_All;
+	GetLandscapeActor()->RequestLayersContentUpdate(ELandscapeLayerUpdateMode::Weightmap_Editing);
+	if (bUpdateAll)
+	{
+		GetLandscapeActor()->RequestLayersContentUpdate(Weightmap_All);
+	}
 }
 
 void ALandscape::RegenerateLayersContent()
 {
-	if (GetLandscapeInfo() == nullptr || LayersContentUpdateFlags == 0)
+	if (GetLandscapeInfo() == nullptr || LayerContentUpdateModes == 0)
 	{
 		return;
 	}
@@ -3940,31 +4003,51 @@ void ALandscape::RegenerateLayersContent()
 		AllLandscapeComponents.Append(Proxy->LandscapeComponents);
 	});
 
-	bool bUpdateClients = (LayersContentUpdateFlags & EInternalLandscapeLayersContentUpdateFlag::Internal_ClientUpdate);
+	RegenerateLayersHeightmaps(AllLandscapeComponents);
+	RegenerateLayersWeightmaps(AllLandscapeComponents);
 
-	bool bHeightmapRecreateCollision = false;
-	RegenerateLayersHeightmaps(AllLandscapeComponents, bHeightmapRecreateCollision);
-	bool bWeightmapRecreateCollision = false;
-	RegenerateLayersWeightmaps(AllLandscapeComponents, bWeightmapRecreateCollision);
-	
-	if (bHeightmapRecreateCollision || bWeightmapRecreateCollision)
+	bool bDeferClientUpdate = false;
+
+	for (ULandscapeComponent* LandscapeComponent : AllLandscapeComponents)
 	{
-		for (ULandscapeComponent* LandscapeComponent : AllLandscapeComponents)
+		if (LandscapeComponent->IsUpdateFlagEnabledForModes(ELandscapeComponentUpdateFlag::RecreateCollision, LayerContentUpdateModes, bLayerForceUpdateAllComponents))
 		{
 			if (ULandscapeHeightfieldCollisionComponent* CollisionComp = LandscapeComponent->CollisionComponent.Get())
 			{
 				CollisionComp->RecreateCollision();
+				
 			}
 		}
-	}
-
-	if (bUpdateClients)
-	{
-		if (!UpdateClients(AllLandscapeComponents, bHeightmapRecreateCollision | bWeightmapRecreateCollision))
+		
+		if (LandscapeComponent->IsUpdateFlagEnabledForModes(ELandscapeComponentUpdateFlag::ClientUpdate, LayerContentUpdateModes, bLayerForceUpdateAllComponents))
+		{
+			if (!GUndo)
+			{
+				if (ULandscapeHeightfieldCollisionComponent* CollisionComp = LandscapeComponent->CollisionComponent.Get())
+				{
+					FNavigationSystem::UpdateComponentData(*CollisionComp);
+					CollisionComp->SnapFoliageInstances();
+				}
+			}
+			else
+			{
+				bDeferClientUpdate = true;
+			}
+		}
+		LandscapeComponent->ClearUpdateFlagsForModes(LayerContentUpdateModes);
+		if (bDeferClientUpdate)
 		{
-			LayersContentUpdateFlags |= EInternalLandscapeLayersContentUpdateFlag::Internal_ClientUpdate;
+			LandscapeComponent->RequestDeferredClientUpdate();
 		}
 	}
+	
+	LayerContentUpdateModes = 0;
+	if (bDeferClientUpdate)
+	{
+		LayerContentUpdateModes |= ELandscapeLayerUpdateMode::DeferredClientUpdate;
+	}
+
+	bLayerForceUpdateAllComponents = false;
 }
 
 void ALandscape::InitializeLayers()
@@ -4020,7 +4103,7 @@ void ALandscape::TickLayers(float DeltaTime, ELevelTick TickType, FActorTickFunc
 			else if (WasCompilingShaders)
 			{
 				WasCompilingShaders = false;
-				RequestLayersContentUpdate(ELandscapeLayersContentUpdateFlag::All);
+				RequestLayersContentUpdateForceAll();
 			}
 
 			RegenerateLayersContent();
@@ -4363,7 +4446,7 @@ bool ALandscape::IsUpToDate() const
 #if WITH_EDITORONLY_DATA
 	if (GetMutableDefault<UEditorExperimentalSettings>()->bLandscapeLayerSystem)
 	{
-		return LayersContentUpdateFlags == 0;
+		return LayerContentUpdateModes == 0;
 	}
 #endif
 
@@ -4410,7 +4493,7 @@ void ALandscape::SetLayerAlpha(int32 InLayerIndex, const float InAlpha, bool bIn
 
 	Modify();
 	LayerAlpha = InAlpha;
-	RequestLayersContentUpdate(ELandscapeLayersContentUpdateFlag::All, true);
+	RequestLayersContentUpdateForceAll();
 }
 
 void ALandscape::SetLayerVisibility(int32 InLayerIndex, bool bInVisible)
@@ -4424,7 +4507,7 @@ void ALandscape::SetLayerVisibility(int32 InLayerIndex, bool bInVisible)
 
 	Modify();
 	Layer->bVisible = bInVisible;
-	RequestLayersContentUpdate(ELandscapeLayersContentUpdateFlag::All, true);
+	RequestLayersContentUpdateForceAll();
 }
 
 void ALandscape::SetLayerLocked(int32 InLayerIndex, bool bLocked)
@@ -4500,7 +4583,7 @@ void ALandscape::DeleteLayer(int32 InLayerIndex)
 	LandscapeLayers.RemoveAt(InLayerIndex);
 
 	// Request Update
-	RequestLayersContentUpdate(ELandscapeLayersContentUpdateFlag::All, true);
+	RequestLayersContentUpdateForceAll();
 }
 
 void ALandscape::ClearLayer(int32 InLayerIndex, bool bInUpdateCollision)
@@ -4524,7 +4607,7 @@ void ALandscape::ClearLayer(const FGuid& InLayerGuid, bool bInUpdateCollision)
 	}
 
 	Modify();
-	FScopedSetLandscapeEditingLayer Scope(this, Layer->Guid, [=] { RequestLayersContentUpdate(ELandscapeLayersContentUpdateFlag::All, true); });
+	FScopedSetLandscapeEditingLayer Scope(this, Layer->Guid, [=] { RequestLayersContentUpdateForceAll(); });
 
 	TArray<uint16> NewHeightData;
 	NewHeightData.AddZeroed(FMath::Square(ComponentSizeQuads + 1));
@@ -4599,7 +4682,7 @@ void ALandscape::ShowOnlySelectedLayer(int32 InLayerIndex)
 		{
 			Layer.bVisible = (&Layer == VisibleLayer);
 		}
-		RequestLayersContentUpdate(ELandscapeLayersContentUpdateFlag::All, true);
+		RequestLayersContentUpdateForceAll();
 	}
 }
 
@@ -4611,7 +4694,7 @@ void ALandscape::ShowAllLayers()
 		{
 			Layer.bVisible = true;
 		}
-		RequestLayersContentUpdate(ELandscapeLayersContentUpdateFlag::All, true);
+		RequestLayersContentUpdateForceAll();
 	}
 }
 
@@ -4670,8 +4753,8 @@ void ALandscape::UpdateLandscapeSplines(const FGuid& InTargetLayer, bool bUpdate
 	const FLandscapeLayer* TargetLayer = GetLayer(TargetLayerGuid);
 	if (LandscapeInfo && TargetLayer)
 	{
-		FScopedSetLandscapeEditingLayer Scope(this, TargetLayerGuid, [=] { this->RequestLayersContentUpdate(ELandscapeLayersContentUpdateFlag::All, true); });
-		// Temporarily disable material instance updates since it will be done once at the end (requested by RequestLayersContentUpdate)
+		FScopedSetLandscapeEditingLayer Scope(this, TargetLayerGuid, [=] { this->RequestLayersContentUpdateForceAll(); });
+		// Temporarily disable material instance updates since it will be done once at the end (requested by RequestLayersContentUpdateForceAll)
 		GDisableUpdateLandscapeMaterialInstances = true;
 		if (LandscapeSplinesTargetLayerGuid.IsValid())
 		{
@@ -4887,7 +4970,8 @@ void ALandscape::SetLayerSubstractiveBlendStatus(int32 InLayerIndex, bool InStat
 		*AllocationBlend = InStatus;
 	}
 
-	RequestLayersContentUpdate(ELandscapeLayersContentUpdateFlag::All, true);
+	const bool bLocalForceUpdateAllComponents = true;
+	RequestLayersContentUpdate(ELandscapeLayerUpdateMode::Weightmap_All, bLocalForceUpdateAllComponents);
 }
 
 void ALandscape::AddBrushToLayer(int32 InLayerIndex, int32 InTargetType, ALandscapeBlueprintCustomBrush* InBrush)
@@ -4912,7 +4996,7 @@ void ALandscape::AddBrushToLayer(int32 InLayerIndex, int32 InTargetType, ALandsc
 
 	InBrush->SetOwningLandscape(this);
 
-	RequestLayersContentUpdate(ELandscapeLayersContentUpdateFlag::All_Render);
+	RequestLayersContentUpdateForceAll();
 }
 
 void ALandscape::RemoveBrushFromLayer(int32 InLayerIndex, int32 InTargetType, ALandscapeBlueprintCustomBrush* InBrush)
@@ -4977,7 +5061,7 @@ void ALandscape::RemoveBrushFromLayer(int32 InLayerIndex, int32 InTargetType, AL
 		InBrush->SetOwningLandscape(nullptr);
 	}
 
-	RequestLayersContentUpdate(ELandscapeLayersContentUpdateFlag::All_Render);
+	RequestLayersContentUpdateForceAll();
 }
 
 bool ALandscape::AreAllBrushesCommitedToLayer(int32 InLayerIndex, int32 InTargetType)
diff --git a/Engine/Source/Runtime/Landscape/Private/LandscapeSplineRaster.cpp b/Engine/Source/Runtime/Landscape/Private/LandscapeSplineRaster.cpp
--- a/Engine/Source/Runtime/Landscape/Private/LandscapeSplineRaster.cpp
+++ b/Engine/Source/Runtime/Landscape/Private/LandscapeSplineRaster.cpp
@@ -465,7 +465,10 @@ bool ULandscapeInfo::ApplySplines(bool bOnlySelected)
 	ALandscape* Landscape = LandscapeActor.Get();
 	const FLandscapeLayer* Layer = Landscape ? Landscape->GetLandscapeSplinesReservedLayer() : nullptr;
 	FGuid SplinesTargetLayerGuid = Layer ? Layer->Guid : FGuid();
-	FScopedSetLandscapeEditingLayer Scope(Landscape, SplinesTargetLayerGuid, [=] { Landscape->RequestLayersContentUpdate(ELandscapeLayersContentUpdateFlag::All); });
+	FScopedSetLandscapeEditingLayer Scope(Landscape, SplinesTargetLayerGuid, [=] 
+	{ 
+		Landscape->RequestLayersContentUpdate(ELandscapeLayerUpdateMode::All);
+	});
 
 	ForAllLandscapeProxies([&bResult, bOnlySelected, this](ALandscapeProxy* Proxy)
 	{
@@ -626,11 +629,15 @@ bool ULandscapeInfo::ApplySplinesInternal(bool bOnlySelected, ALandscapeProxy* P
 	LandscapeEdit.Flush();
 		
 	ALandscapeProxy::InvalidateGeneratedComponentData(ModifiedComponents);
-	
-	if(!GetMutableDefault<UEditorExperimentalSettings>()->bLandscapeLayerSystem)
+		
+	for (ULandscapeComponent* Component : ModifiedComponents)
 	{
-		ALandscape* Landscape = Proxy->GetLandscapeActor();
-		for (ULandscapeComponent* Component : ModifiedComponents)
+		if (GetMutableDefault<UEditorExperimentalSettings>()->bLandscapeLayerSystem)
+		{
+			Component->RequestHeightmapUpdate();
+			Component->RequestWeightmapUpdate();
+		}
+		else
 		{
 			// Recreate collision for modified components and update the navmesh
 			ULandscapeHeightfieldCollisionComponent* CollisionComponent = Component->CollisionComponent.Get();
@@ -641,6 +648,7 @@ bool ULandscapeInfo::ApplySplinesInternal(bool bOnlySelected, ALandscapeProxy* P
 			}
 		}
 	}
+	
 
 	return true;
 }
diff --git a/Engine/Source/Runtime/Landscape/Public/LandscapeEdit.h b/Engine/Source/Runtime/Landscape/Public/LandscapeEdit.h
--- a/Engine/Source/Runtime/Landscape/Public/LandscapeEdit.h
+++ b/Engine/Source/Runtime/Landscape/Public/LandscapeEdit.h
@@ -319,6 +319,7 @@ struct FHeightmapAccessor
 			for (ULandscapeComponent* Component : Components)
 			{
 				Component->InvalidateLightingCache();
+				Component->RequestHeightmapUpdate();
 			}
 						
 			// Notify foliage to move any attached instances
@@ -478,6 +479,11 @@ struct FAlphamapAccessor
 		if (LandscapeEdit.GetComponentsInRegion(X1, Y1, X2, Y2, &Components))
 		{
 			ALandscapeProxy::InvalidateGeneratedComponentData(Components);
+			for (ULandscapeComponent* LandscapeComponent : Components)
+			{
+				// Flag both modes depending on client calling SetData
+				LandscapeComponent->RequestWeightmapUpdate();
+			}
 			
 			if (!GetMutableDefault<UEditorExperimentalSettings>()->bLandscapeLayerSystem)
 			{
