Fixed crash when disabling & re-enabling Landscape Layer System experimental flag
#rnx

diff --git a/Engine/Source/Editor/LandscapeEditor/Private/LandscapeEdMode.cpp b/Engine/Source/Editor/LandscapeEditor/Private/LandscapeEdMode.cpp
--- a/Engine/Source/Editor/LandscapeEditor/Private/LandscapeEdMode.cpp
+++ b/Engine/Source/Editor/LandscapeEditor/Private/LandscapeEdMode.cpp
@@ -3924,7 +3924,7 @@ ALandscape* FEdModeLandscape::ChangeComponentSetting(int32 NumComponentsX, int32
 				ALandscape* OldLandscapeActor = Cast<ALandscape>(OldLandscapeProxy);
 				ALandscape* NewLandscapeActor = Cast<ALandscape>(Landscape);
 
-				NewLandscapeActor->PreviousExperimentalLandscapeLayers = OldLandscapeActor->PreviousExperimentalLandscapeLayers;
+				NewLandscapeActor->bInitializedWithFlagExperimentalLandscapeLayers = OldLandscapeActor->bInitializedWithFlagExperimentalLandscapeLayers;
 
 				// Moves specific config including brushes to new landscape actor
 				for (const FLandscapeLayer& Layer : OldLandscapeActor->LandscapeLayers)
diff --git a/Engine/Source/Editor/LandscapeEditor/Private/LandscapeEditorDetailCustomization_NewLandscape.cpp b/Engine/Source/Editor/LandscapeEditor/Private/LandscapeEditorDetailCustomization_NewLandscape.cpp
--- a/Engine/Source/Editor/LandscapeEditor/Private/LandscapeEditorDetailCustomization_NewLandscape.cpp
+++ b/Engine/Source/Editor/LandscapeEditor/Private/LandscapeEditorDetailCustomization_NewLandscape.cpp
@@ -811,7 +811,7 @@ FReply FLandscapeEditorDetailCustomization_NewLandscape::OnCreateButtonClicked()
 
 		if (GetMutableDefault<UEditorExperimentalSettings>()->bLandscapeLayerSystem)
 		{
-			Landscape->PreviousExperimentalLandscapeLayers = true;
+			Landscape->bInitializedWithFlagExperimentalLandscapeLayers = true;
 		}
 
 		Landscape->Import(FGuid::NewGuid(), 0, 0, SizeX-1, SizeY-1, LandscapeEdMode->UISettings->NewLandscape_SectionsPerComponent, LandscapeEdMode->UISettings->NewLandscape_QuadsPerSection, Data.GetData(),
diff --git a/Engine/Source/Runtime/Landscape/Classes/Landscape.h b/Engine/Source/Runtime/Landscape/Classes/Landscape.h
--- a/Engine/Source/Runtime/Landscape/Classes/Landscape.h
+++ b/Engine/Source/Runtime/Landscape/Classes/Landscape.h
@@ -271,8 +271,10 @@ public:
 	
 private:
 	void TickLayers(float DeltaTime, ELevelTick TickType, FActorTickFunction& ThisTickFunction);
-	void CreateLayersRenderingResource(const FIntPoint& InComponentCounts);
+	void CreateLayersRenderingResource();
+	void ReleaseLayersRenderingResource();
 	void RegenerateLayersContent();
+	void MonitorShaderCompilation();
 	void RegenerateLayersHeightmaps(const TArray<ULandscapeComponent*>& InLandscapeComponents);
 	void RegenerateLayersWeightmaps(const TArray<ULandscapeComponent*>& InLandscapeComponents);
 	void ResolveLayersHeightmapTexture();
@@ -346,7 +348,7 @@ public:
 	TArray<UTextureRenderTarget2D*> WeightmapRTList;
 
 	UPROPERTY(Transient)
-	bool PreviousExperimentalLandscapeLayers;
+	bool bInitializedWithFlagExperimentalLandscapeLayers;
 
 private:
 	UPROPERTY(Transient)
diff --git a/Engine/Source/Runtime/Landscape/Private/Landscape.cpp b/Engine/Source/Runtime/Landscape/Private/Landscape.cpp
--- a/Engine/Source/Runtime/Landscape/Private/Landscape.cpp
+++ b/Engine/Source/Runtime/Landscape/Private/Landscape.cpp
@@ -1074,7 +1074,7 @@ ALandscape::ALandscape(const FObjectInitializer& ObjectInitializer)
 {
 #if WITH_EDITORONLY_DATA
 	bLockLocation = false;
-	PreviousExperimentalLandscapeLayers = false;
+	bInitializedWithFlagExperimentalLandscapeLayers = false;
 	WasCompilingShaders = false;
 	LayerContentUpdateModes = 0;
 	bLayerForceUpdateAllComponents = false;
diff --git a/Engine/Source/Runtime/Landscape/Private/LandscapeEditLayers.cpp b/Engine/Source/Runtime/Landscape/Private/LandscapeEditLayers.cpp
--- a/Engine/Source/Runtime/Landscape/Private/LandscapeEditLayers.cpp
+++ b/Engine/Source/Runtime/Landscape/Private/LandscapeEditLayers.cpp
@@ -1231,10 +1231,9 @@ typedef FLandscapeLayersRender_RenderThread<FLandscapeLayersWeightmapShaderParam
 
 #if WITH_EDITOR
 
-void ALandscape::CreateLayersRenderingResource(const FIntPoint& InComponentCounts)
+void ALandscape::CreateLayersRenderingResource()
 {
 	ULandscapeInfo* Info = GetLandscapeInfo();
-
 	if (Info == nullptr)
 	{
 		return;
@@ -1253,7 +1252,7 @@ void ALandscape::CreateLayersRenderingResource(const FIntPoint& InComponentCount
 
 			if (CPUReadback != nullptr)
 			{
-				ReleaseResourceAndFlush(*CPUReadback);
+				BeginReleaseResource(*CPUReadback);
 				*CPUReadback = nullptr;
 			}
 
@@ -1270,8 +1269,9 @@ void ALandscape::CreateLayersRenderingResource(const FIntPoint& InComponentCount
 	// Flush because TickLayers can access CPU Readbacks in the same frame (see ResolveLayersTexture)
 	FlushRenderingCommands();
 
-	const int32 TotalVertexCountX = (SubsectionSizeQuads * NumSubsections) * InComponentCounts.X + 1;
-	const int32 TotalVertexCountY = (SubsectionSizeQuads * NumSubsections) * InComponentCounts.Y + 1;
+	const FIntPoint ComponentCounts = ComputeComponentCounts();
+	const int32 TotalVertexCountX = (SubsectionSizeQuads * NumSubsections) * ComponentCounts.X + 1;
+	const int32 TotalVertexCountY = (SubsectionSizeQuads * NumSubsections) * ComponentCounts.Y + 1;
 
 	ALandscape* Landscape = GetLandscapeActor();
 	check(Landscape);
@@ -1280,8 +1280,8 @@ void ALandscape::CreateLayersRenderingResource(const FIntPoint& InComponentCount
 	{
 		Landscape->HeightmapRTList.Init(nullptr, EHeightmapRTType::HeightmapRT_Count);
 
-		int32 CurrentMipSizeX = ((SubsectionSizeQuads + 1) * NumSubsections) * InComponentCounts.X;
-		int32 CurrentMipSizeY = ((SubsectionSizeQuads + 1) * NumSubsections) * InComponentCounts.Y;
+		int32 CurrentMipSizeX = ((SubsectionSizeQuads + 1) * NumSubsections) * ComponentCounts.X;
+		int32 CurrentMipSizeY = ((SubsectionSizeQuads + 1) * NumSubsections) * ComponentCounts.Y;
 
 		for (int32 i = 0; i < EHeightmapRTType::HeightmapRT_Count; ++i)
 		{
@@ -1306,7 +1306,7 @@ void ALandscape::CreateLayersRenderingResource(const FIntPoint& InComponentCount
 			Landscape->HeightmapRTList[i]->UpdateResourceImmediate(true);
 
 			// Only generate required mips RT
-			if (CurrentMipSizeX == InComponentCounts.X && CurrentMipSizeY == InComponentCounts.Y)
+			if (CurrentMipSizeX == ComponentCounts.X && CurrentMipSizeY == ComponentCounts.Y)
 			{
 				break;
 			}
@@ -1314,8 +1314,8 @@ void ALandscape::CreateLayersRenderingResource(const FIntPoint& InComponentCount
 	}
 	else // Simply resize the render target
 	{
-		int32 CurrentMipSizeX = ((SubsectionSizeQuads + 1) * NumSubsections) * InComponentCounts.X;
-		int32 CurrentMipSizeY = ((SubsectionSizeQuads + 1) * NumSubsections) * InComponentCounts.Y;
+		int32 CurrentMipSizeX = ((SubsectionSizeQuads + 1) * NumSubsections) * ComponentCounts.X;
+		int32 CurrentMipSizeY = ((SubsectionSizeQuads + 1) * NumSubsections) * ComponentCounts.Y;
 
 		for (int32 i = 0; i < EHeightmapRTType::HeightmapRT_Count; ++i)
 		{
@@ -1331,7 +1331,7 @@ void ALandscape::CreateLayersRenderingResource(const FIntPoint& InComponentCount
 			}
 
 			// Only generate required mips RT
-			if (CurrentMipSizeX == InComponentCounts.X && CurrentMipSizeY == InComponentCounts.Y)
+			if (CurrentMipSizeX == ComponentCounts.X && CurrentMipSizeY == ComponentCounts.Y)
 			{
 				break;
 			}
@@ -1342,8 +1342,8 @@ void ALandscape::CreateLayersRenderingResource(const FIntPoint& InComponentCount
 	{
 		Landscape->WeightmapRTList.Init(nullptr, EWeightmapRTType::WeightmapRT_Count);
 
-		int32 CurrentMipSizeX = ((SubsectionSizeQuads + 1) * NumSubsections) * InComponentCounts.X;
-		int32 CurrentMipSizeY = ((SubsectionSizeQuads + 1) * NumSubsections) * InComponentCounts.Y;
+		int32 CurrentMipSizeX = ((SubsectionSizeQuads + 1) * NumSubsections) * ComponentCounts.X;
+		int32 CurrentMipSizeY = ((SubsectionSizeQuads + 1) * NumSubsections) * ComponentCounts.Y;
 
 		for (int32 i = 0; i < EWeightmapRTType::WeightmapRT_Count; ++i)
 		{
@@ -1371,7 +1371,7 @@ void ALandscape::CreateLayersRenderingResource(const FIntPoint& InComponentCount
 			Landscape->WeightmapRTList[i]->UpdateResourceImmediate(true);
 
 			// Only generate required mips RT
-			if (CurrentMipSizeX == InComponentCounts.X && CurrentMipSizeY == InComponentCounts.Y)
+			if (CurrentMipSizeX == ComponentCounts.X && CurrentMipSizeY == ComponentCounts.Y)
 			{
 				break;
 			}
@@ -1379,8 +1379,8 @@ void ALandscape::CreateLayersRenderingResource(const FIntPoint& InComponentCount
 	}
 	else // Simply resize the render target
 	{
-		int32 CurrentMipSizeX = ((SubsectionSizeQuads + 1) * NumSubsections) * InComponentCounts.X;
-		int32 CurrentMipSizeY = ((SubsectionSizeQuads + 1) * NumSubsections) * InComponentCounts.Y;
+		int32 CurrentMipSizeX = ((SubsectionSizeQuads + 1) * NumSubsections) * ComponentCounts.X;
+		int32 CurrentMipSizeY = ((SubsectionSizeQuads + 1) * NumSubsections) * ComponentCounts.Y;
 
 		for (int32 i = 0; i < EWeightmapRTType::WeightmapRT_Count; ++i)
 		{
@@ -1397,7 +1397,7 @@ void ALandscape::CreateLayersRenderingResource(const FIntPoint& InComponentCount
 			}
 
 			// Only generate required mips RT
-			if (CurrentMipSizeX == InComponentCounts.X && CurrentMipSizeY == InComponentCounts.Y)
+			if (CurrentMipSizeX == ComponentCounts.X && CurrentMipSizeY == ComponentCounts.Y)
 			{
 				break;
 			}
@@ -1407,6 +1407,92 @@ void ALandscape::CreateLayersRenderingResource(const FIntPoint& InComponentCount
 	InitializeLayersWeightmapResources();
 }
 
+void ALandscape::ReleaseLayersRenderingResource()
+{
+	check(!GetMutableDefault<UEditorExperimentalSettings>()->bLandscapeLayerSystem);
+	ULandscapeInfo* Info = GetLandscapeInfo();
+	if (Info == nullptr)
+	{
+		return;
+	}
+
+	Info->ForAllLandscapeProxies([&](ALandscapeProxy* Proxy)
+	{
+		for (auto& ItPair : Proxy->HeightmapsCPUReadBack)
+		{
+			FLandscapeLayersTexture2DCPUReadBackResource* HeightmapCPUReadBack = ItPair.Value;
+
+			if (HeightmapCPUReadBack != nullptr)
+			{
+				BeginReleaseResource(HeightmapCPUReadBack);
+			}
+		}
+
+		for (auto& ItPair : Proxy->WeightmapsCPUReadBack)
+		{
+			FLandscapeLayersTexture2DCPUReadBackResource* WeightmapCPUReadBack = ItPair.Value;
+
+			if (WeightmapCPUReadBack != nullptr)
+			{
+				BeginReleaseResource(WeightmapCPUReadBack);
+			}
+		}
+	});
+
+	if (CombinedLayersWeightmapAllMaterialLayersResource != nullptr)
+	{
+		BeginReleaseResource(CombinedLayersWeightmapAllMaterialLayersResource);
+	}
+
+	if (CurrentLayersWeightmapAllMaterialLayersResource != nullptr)
+	{
+		BeginReleaseResource(CurrentLayersWeightmapAllMaterialLayersResource);
+	}
+
+	if (WeightmapScratchExtractLayerTextureResource != nullptr)
+	{
+		BeginReleaseResource(WeightmapScratchExtractLayerTextureResource);
+	}
+
+	if (WeightmapScratchPackLayerTextureResource != nullptr)
+	{
+		BeginReleaseResource(WeightmapScratchPackLayerTextureResource);
+	}
+
+	FlushRenderingCommands();
+
+	Info->ForAllLandscapeProxies([&](ALandscapeProxy* Proxy)
+	{
+		for (auto& ItPair : Proxy->HeightmapsCPUReadBack)
+		{
+			FLandscapeLayersTexture2DCPUReadBackResource* HeightmapCPUReadBack = ItPair.Value;
+
+			delete HeightmapCPUReadBack;
+			HeightmapCPUReadBack = nullptr;
+		}
+		Proxy->HeightmapsCPUReadBack.Empty();
+
+		for (auto& ItPair : Proxy->WeightmapsCPUReadBack)
+		{
+			FLandscapeLayersTexture2DCPUReadBackResource* WeightmapCPUReadBack = ItPair.Value;
+
+			delete WeightmapCPUReadBack;
+			WeightmapCPUReadBack = nullptr;
+		}
+		Proxy->WeightmapsCPUReadBack.Empty();
+	});
+
+	delete CombinedLayersWeightmapAllMaterialLayersResource;
+	delete CurrentLayersWeightmapAllMaterialLayersResource;
+	delete WeightmapScratchExtractLayerTextureResource;
+	delete WeightmapScratchPackLayerTextureResource;
+
+	CombinedLayersWeightmapAllMaterialLayersResource = nullptr;
+	CurrentLayersWeightmapAllMaterialLayersResource = nullptr;
+	WeightmapScratchExtractLayerTextureResource = nullptr;
+	WeightmapScratchPackLayerTextureResource = nullptr;
+}
+
 FIntPoint ALandscape::ComputeComponentCounts() const
 {
 	ULandscapeInfo* Info = GetLandscapeInfo();
@@ -3990,8 +4076,25 @@ void ULandscapeComponent::RequestWeightmapUpdate(bool bUpdateAll)
 	}
 }
 
+void ALandscape::MonitorShaderCompilation()
+{
+	// If doing editing while shader are compiling or at load of a map, it's possible we will need another update pass after shader are completed to see the correct result
+	const int32 RemainingShadersThisFrame = GShaderCompilingManager->GetNumRemainingJobs();
+	if (!WasCompilingShaders && RemainingShadersThisFrame > 0)
+	{
+		WasCompilingShaders = true;
+	}
+	else if (WasCompilingShaders)
+	{
+		WasCompilingShaders = false;
+		RequestLayersContentUpdateForceAll();
+	}
+}
+
 void ALandscape::RegenerateLayersContent()
 {
+	MonitorShaderCompilation();
+
 	if (GetLandscapeInfo() == nullptr || LayerContentUpdateModes == 0)
 	{
 		return;
@@ -4052,22 +4155,22 @@ void ALandscape::RegenerateLayersContent()
 
 void ALandscape::InitializeLayers()
 {
-	check(!bLandscapeLayersAreInitialized || PreviousExperimentalLandscapeLayers != GetMutableDefault<UEditorExperimentalSettings>()->bLandscapeLayerSystem);
-
-	FIntPoint ComponentCounts = ComputeComponentCounts();
+	check(GetMutableDefault<UEditorExperimentalSettings>()->bLandscapeLayerSystem);
+	check(!bLandscapeLayersAreInitialized || (bInitializedWithFlagExperimentalLandscapeLayers != GetMutableDefault<UEditorExperimentalSettings>()->bLandscapeLayerSystem));
 
 	if (HasLayersContent)
 	{
-		CreateLayersRenderingResource(ComponentCounts);
+		CreateLayersRenderingResource();
 		InitializeLandscapeLayersWeightmapUsage();
 	}
 	else
 	{
 		// Special case to load the existing content to the default layer
 		CreateDefaultLayer();
-		CreateLayersRenderingResource(ComponentCounts);
+		CreateLayersRenderingResource();
 		CopyOldDataToDefaultLayer();
 	}
+
 	bLandscapeLayersAreInitialized = true;
 }
 
@@ -4076,127 +4179,29 @@ void ALandscape::TickLayers(float DeltaTime, ELevelTick TickType, FActorTickFunc
 	check(GIsEditor);
 
 	UWorld* World = GetWorld();
-	if (World && !World->IsPlayInEditor())
+	if (World && !World->IsPlayInEditor() && GetLandscapeInfo())
 	{
-#if WITH_EDITOR
 		if (CVarLandscapeSimulatePhysics.GetValueOnAnyThread() == 1)
 		{
 			World->bShouldSimulatePhysics = true;
 		}
-#endif
 
 		if (GetMutableDefault<UEditorExperimentalSettings>()->bLandscapeLayerSystem)
 		{
-			if (!bLandscapeLayersAreInitialized || PreviousExperimentalLandscapeLayers != GetMutableDefault<UEditorExperimentalSettings>()->bLandscapeLayerSystem)
+			if (!bLandscapeLayersAreInitialized || !bInitializedWithFlagExperimentalLandscapeLayers)
 			{
-				PreviousExperimentalLandscapeLayers = GetMutableDefault<UEditorExperimentalSettings>()->bLandscapeLayerSystem;
 				InitializeLayers();
-			}
-
-			// If doing editing while shader are compiling or at load of a map, it's possible we will need another update pass after shader are completed to see the correct result
-			const int32 RemainingShadersThisFrame = GShaderCompilingManager->GetNumRemainingJobs();
-
-			if (!WasCompilingShaders && RemainingShadersThisFrame > 0)
-			{
-				WasCompilingShaders = true;
-			}
-			else if (WasCompilingShaders)
-			{
-				WasCompilingShaders = false;
-				RequestLayersContentUpdateForceAll();
+				bInitializedWithFlagExperimentalLandscapeLayers = true;
 			}
 
 			RegenerateLayersContent();
 		}
 		else
 		{
-			if (PreviousExperimentalLandscapeLayers != GetMutableDefault<UEditorExperimentalSettings>()->bLandscapeLayerSystem)
+			if (bInitializedWithFlagExperimentalLandscapeLayers)
 			{
-				ULandscapeInfo* Info = GetLandscapeInfo();
-
-				if (Info == nullptr)
-				{
-					return;
-				}
-
-				PreviousExperimentalLandscapeLayers = GetMutableDefault<UEditorExperimentalSettings>()->bLandscapeLayerSystem;
-
-#if WITH_EDITORONLY_DATA
-
-				Info->ForAllLandscapeProxies([&](ALandscapeProxy* Proxy)
-				{
-					for (auto& ItPair : Proxy->HeightmapsCPUReadBack)
-					{
-						FLandscapeLayersTexture2DCPUReadBackResource* HeightmapCPUReadBack = ItPair.Value;
-
-						if (HeightmapCPUReadBack != nullptr)
-						{
-							BeginReleaseResource(HeightmapCPUReadBack);
-						}
-					}
-
-					for (auto& ItPair : Proxy->WeightmapsCPUReadBack)
-					{
-						FLandscapeLayersTexture2DCPUReadBackResource* WeightmapCPUReadBack = ItPair.Value;
-
-						if (WeightmapCPUReadBack != nullptr)
-						{
-							BeginReleaseResource(WeightmapCPUReadBack);
-						}
-					}
-				});
-
-				if (CombinedLayersWeightmapAllMaterialLayersResource != nullptr)
-				{
-					BeginReleaseResource(CombinedLayersWeightmapAllMaterialLayersResource);
-				}
-
-				if (CurrentLayersWeightmapAllMaterialLayersResource != nullptr)
-				{
-					BeginReleaseResource(CurrentLayersWeightmapAllMaterialLayersResource);
-				}
-
-				if (WeightmapScratchExtractLayerTextureResource != nullptr)
-				{
-					BeginReleaseResource(WeightmapScratchExtractLayerTextureResource);
-				}
-
-				if (WeightmapScratchPackLayerTextureResource != nullptr)
-				{
-					BeginReleaseResource(WeightmapScratchPackLayerTextureResource);
-				}
-
-				FlushRenderingCommands();
-
-				Info->ForAllLandscapeProxies([&](ALandscapeProxy* Proxy)
-				{
-					for (auto& ItPair : Proxy->HeightmapsCPUReadBack)
-					{
-						FLandscapeLayersTexture2DCPUReadBackResource* HeightmapCPUReadBack = ItPair.Value;
-
-						delete HeightmapCPUReadBack;
-						HeightmapCPUReadBack = nullptr;
-					}
-
-					for (auto& ItPair : Proxy->WeightmapsCPUReadBack)
-					{
-						FLandscapeLayersTexture2DCPUReadBackResource* WeightmapCPUReadBack = ItPair.Value;
-
-						delete WeightmapCPUReadBack;
-						WeightmapCPUReadBack = nullptr;
-					}
-				});
-
-				delete CombinedLayersWeightmapAllMaterialLayersResource;
-				delete CurrentLayersWeightmapAllMaterialLayersResource;
-				delete WeightmapScratchExtractLayerTextureResource;
-				delete WeightmapScratchPackLayerTextureResource;
-
-				CombinedLayersWeightmapAllMaterialLayersResource = nullptr;
-				CurrentLayersWeightmapAllMaterialLayersResource = nullptr;
-				WeightmapScratchExtractLayerTextureResource = nullptr;
-				WeightmapScratchPackLayerTextureResource = nullptr;
-#endif
+				ReleaseLayersRenderingResource();
+				bInitializedWithFlagExperimentalLandscapeLayers = false;
 			}
 		}
 	}
