- Landscape Layer System: Fix Weightmap Usage Undo/Redo bug

#rnx

diff --git a/Engine/Source/Runtime/Landscape/Classes/Landscape.h b/Engine/Source/Runtime/Landscape/Classes/Landscape.h
--- a/Engine/Source/Runtime/Landscape/Classes/Landscape.h
+++ b/Engine/Source/Runtime/Landscape/Classes/Landscape.h
@@ -284,13 +284,13 @@ private:
 	bool PrepareLayersHeightmapTextureResources() const;
 	bool PrepareLayersWeightmapTextureResources() const;
 
-	void UpdateLayersMaterialInstances(const TArray<ULandscapeComponent*>& InComponentsToUpdate);
+	void UpdateLayersMaterialInstances();
 
 	void PrepareComponentDataToExtractMaterialLayersCS(const FLandscapeLayer& InLayer, int32 InCurrentWeightmapToProcessIndex, const FIntPoint& InLandscapeBase, bool InOutputDebugName, class FLandscapeTexture2DResource* InOutTextureData,
 														  TArray<struct FLandscapeLayerWeightmapExtractMaterialLayersComponentData>& OutComponentData, TMap<ULandscapeLayerInfoObject*, int32>& OutLayerInfoObjects);
 	void PrepareComponentDataToPackMaterialLayersCS(int32 InCurrentWeightmapToProcessIndex, const FIntPoint& InLandscapeBase, bool InOutputDebugName, const TArray<ULandscapeComponent*>& InAllLandscapeComponents, TArray<UTexture2D*>& InOutProcessedWeightmaps,
 													TArray<class FLandscapeLayersTexture2DCPUReadBackResource*>& OutProcessedCPUReadBackTexture, TArray<struct FLandscapeLayerWeightmapPackMaterialLayersComponentData>& OutComponentData);
-	void ReallocateLayersWeightmaps(const TArray<ULandscapeLayerInfoObject*>& InBrushRequiredAllocations, TArray<ULandscapeComponent*>& OutComponentThatNeedMaterialRebuild);
+	void ReallocateLayersWeightmaps(const TArray<ULandscapeLayerInfoObject*>& InBrushRequiredAllocations);
 	void InitializeLayersWeightmapResources();
 	bool GenerateZeroAllocationPerComponents(const TArray<ALandscapeProxy*>& InAllLandscape, const TMap<ULandscapeLayerInfoObject*, bool>& InWeightmapLayersBlendSubstractive);
 
diff --git a/Engine/Source/Runtime/Landscape/Classes/LandscapeComponent.h b/Engine/Source/Runtime/Landscape/Classes/LandscapeComponent.h
--- a/Engine/Source/Runtime/Landscape/Classes/LandscapeComponent.h
+++ b/Engine/Source/Runtime/Landscape/Classes/LandscapeComponent.h
@@ -494,6 +494,9 @@ public:
 	/** Represent the chosen material for each LOD */
 	UPROPERTY(DuplicateTransient)
 	TMap<UMaterialInterface*, int8> MaterialPerLOD;
+
+	/** Represents hash of last weightmap usage update */
+	uint32 WeightmapsHash;
 #endif
 
 	/** For ES2 */
@@ -874,6 +877,7 @@ public:
 	LANDSCAPE_API void RequestWeightmapUpdate(bool bUpdateAll = false);
 	LANDSCAPE_API void RequestHeightmapUpdate(bool bUpdateAll = false);
 	LANDSCAPE_API void RequestDeferredClientUpdate();
+	LANDSCAPE_API uint32 ComputeWeightmapsHash();
 #endif
 
 	friend class FLandscapeComponentSceneProxy;
diff --git a/Engine/Source/Runtime/Landscape/Private/Landscape.cpp b/Engine/Source/Runtime/Landscape/Private/Landscape.cpp
--- a/Engine/Source/Runtime/Landscape/Private/Landscape.cpp
+++ b/Engine/Source/Runtime/Landscape/Private/Landscape.cpp
@@ -142,6 +142,7 @@ ULandscapeComponent::ULandscapeComponent(const FObjectInitializer& ObjectInitial
 : Super(ObjectInitializer)
 #if WITH_EDITORONLY_DATA
 , LayerUpdateFlagPerMode(0)
+, WeightmapsHash(0)
 #endif
 , GrassData(MakeShareable(new FLandscapeComponentGrassData()))
 , ChangeTag(0)
diff --git a/Engine/Source/Runtime/Landscape/Private/LandscapeEditLayers.cpp b/Engine/Source/Runtime/Landscape/Private/LandscapeEditLayers.cpp
--- a/Engine/Source/Runtime/Landscape/Private/LandscapeEditLayers.cpp
+++ b/Engine/Source/Runtime/Landscape/Private/LandscapeEditLayers.cpp
@@ -3047,7 +3047,7 @@ void ALandscape::PrepareComponentDataToPackMaterialLayersCS(int32 InCurrentWeigh
 	}
 }
 
-void ALandscape::ReallocateLayersWeightmaps(const TArray<ULandscapeLayerInfoObject*>& InBrushRequiredAllocations, TArray<ULandscapeComponent*>& OutComponentThatNeedMaterialRebuild)
+void ALandscape::ReallocateLayersWeightmaps(const TArray<ULandscapeLayerInfoObject*>& InBrushRequiredAllocations)
 {
 	SCOPE_CYCLE_COUNTER(STAT_LandscapeLayersReallocateWeightmaps);
 
@@ -3064,29 +3064,7 @@ void ALandscape::ReallocateLayersWeightmaps(const TArray<ULandscapeLayerInfoObje
 	{
 		AllLandscapeComponents.Append(Proxy->LandscapeComponents);
 	});
-
-	auto HashComponentWeightmaps = [](ULandscapeComponent* Component) -> uint32
-	{
-		uint32 Hash = 0;
-		TArray<UTexture2D*>& ComponentWeightmapTextures = Component->GetWeightmapTextures();
-		TArray<ULandscapeWeightmapUsage*>& ComponentWeightmapTextureUsage = Component->GetWeightmapTexturesUsage();
-		for (int32 i = 0; i < ComponentWeightmapTextures.Num(); ++i)
-		{
-			Hash = PointerHash(ComponentWeightmapTextures[i], Hash);
-			for (int32 j = 0; j < ULandscapeWeightmapUsage::NumChannels; ++j)
-			{
-				Hash = PointerHash(ComponentWeightmapTextureUsage[i]->ChannelUsage[j], Hash);
-			}
-		}
-		return Hash;
-	};
-
-	TMap<ULandscapeComponent*, uint32> BeforeReallocComponentWeightmapsHashes;
-	for (ULandscapeComponent* Component : AllLandscapeComponents)
-	{
-		BeforeReallocComponentWeightmapsHashes.Add(Component, HashComponentWeightmaps(Component));
-	}
-
+		
 	// Clear allocation data
 	for (ULandscapeComponent* Component : AllLandscapeComponents)
 	{
@@ -3244,18 +3222,6 @@ void ALandscape::ReallocateLayersWeightmaps(const TArray<ULandscapeLayerInfoObje
 			}
 		}
 	});
-
-	// Determine which Component needs a material rebuild
-	for (ULandscapeComponent* Component : AllLandscapeComponents)
-	{
-		uint32* OldHash = BeforeReallocComponentWeightmapsHashes.Find(Component);
-		check(OldHash);
-		uint32 NewHash = HashComponentWeightmaps(Component);
-		if (*OldHash != NewHash)
-		{
-			OutComponentThatNeedMaterialRebuild.Add(Component);
-		}
-	}
 }
 
 void ALandscape::InitializeLayersWeightmapResources()
@@ -3600,7 +3566,7 @@ void ALandscape::RegenerateLayersWeightmaps(const TArray<ULandscapeComponent*>&
 		}
 
 		// TODO:  if editing a Brush affecting layers, since we don't have any bounds to brush, right now ReallocateLayersWeightmaps wont ask a rebuild of the component affected by Brushes, which mean ComponentThatNeedMaterialRebuild wont contains Brush affected component!
-		ReallocateLayersWeightmaps(BrushRequiredAllocations, ComponentThatNeedMaterialRebuild);
+		ReallocateLayersWeightmaps(BrushRequiredAllocations);
 
 		if (ComputeShaderGeneratedData)
 		{
@@ -3735,7 +3701,7 @@ void ALandscape::RegenerateLayersWeightmaps(const TArray<ULandscapeComponent*>&
 			}
 		}
 
-		UpdateLayersMaterialInstances(ComponentThatNeedMaterialRebuild);
+		UpdateLayersMaterialInstances();
 	}
 
 	if (WeightmapUpdateModes)
@@ -3752,17 +3718,43 @@ void ALandscape::RegenerateLayersWeightmaps(const TArray<ULandscapeComponent*>&
 	}
 }
 
-void ALandscape::UpdateLayersMaterialInstances(const TArray<ULandscapeComponent*>& InComponentsToUpdate)
+uint32 ULandscapeComponent::ComputeWeightmapsHash()
 {
-	if (InComponentsToUpdate.Num() == 0)
+	uint32 Hash = 0;
+	TArray<UTexture2D*>& ComponentWeightmapTextures = GetWeightmapTextures();
+	TArray<ULandscapeWeightmapUsage*>& ComponentWeightmapTextureUsage = GetWeightmapTexturesUsage();
+	for (int32 i = 0; i < ComponentWeightmapTextures.Num(); ++i)
 	{
-		return;
+		Hash = PointerHash(ComponentWeightmapTextures[i], Hash);
+		for (int32 j = 0; j < ULandscapeWeightmapUsage::NumChannels; ++j)
+		{
+			Hash = PointerHash(ComponentWeightmapTextureUsage[i]->ChannelUsage[j], Hash);
+		}
 	}
+	return Hash;
+}
 
+void ALandscape::UpdateLayersMaterialInstances()
+{
+	SCOPE_CYCLE_COUNTER(STAT_LandscapeLayersUpdateMaterialInstance);
 	TArray<ULandscapeComponent*> ComponentsToUpdate;
-	ComponentsToUpdate.Append(InComponentsToUpdate);
 
-	SCOPE_CYCLE_COUNTER(STAT_LandscapeLayersUpdateMaterialInstance);
+	// Compute Weightmap usage changes
+	if (ULandscapeInfo* Info = GetLandscapeInfo())
+	{
+		Info->ForAllLandscapeProxies([&ComponentsToUpdate](ALandscapeProxy* Proxy)
+		{
+			for (ULandscapeComponent* LandscapeComponent : Proxy->LandscapeComponents)
+			{
+				uint32 NewHash = LandscapeComponent->ComputeWeightmapsHash();
+				if (LandscapeComponent->WeightmapsHash != NewHash)
+				{
+					ComponentsToUpdate.Add(LandscapeComponent);
+					LandscapeComponent->WeightmapsHash = NewHash;
+				}
+			}
+		});
+	}
 
 	// we're not having the material update context recreate render states because we will manually do it for only our components
 	TArray<FComponentRecreateRenderStateContext> RecreateRenderStateContexts;
