Fix Crash in Non Layer (introduced by 6318826)

#rnx

diff --git a/Engine/Source/Runtime/Landscape/Private/LandscapeEditLayers.cpp b/Engine/Source/Runtime/Landscape/Private/LandscapeEditLayers.cpp
--- a/Engine/Source/Runtime/Landscape/Private/LandscapeEditLayers.cpp
+++ b/Engine/Source/Runtime/Landscape/Private/LandscapeEditLayers.cpp
@@ -4050,10 +4050,13 @@ void ULandscapeComponent::RequestHeightmapUpdate(bool bUpdateAll)
 {
 	LayerUpdateFlagPerMode |= ELandscapeLayerUpdateMode::Heightmap_Editing;
 	LayerUpdateFlagPerMode |= ELandscapeLayerUpdateMode::Heightmap_All;
-	GetLandscapeActor()->RequestLayersContentUpdate(ELandscapeLayerUpdateMode::Heightmap_Editing);
-	if (bUpdateAll)
+	if (ALandscape* LandscapeActor = GetLandscapeActor())
 	{
-		GetLandscapeActor()->RequestLayersContentUpdate(Heightmap_All);
+		LandscapeActor->RequestLayersContentUpdate(ELandscapeLayerUpdateMode::Heightmap_Editing);
+		if (bUpdateAll)
+		{
+			LandscapeActor->RequestLayersContentUpdate(Heightmap_All);
+		}
 	}
 }
 
@@ -4061,10 +4064,13 @@ void ULandscapeComponent::RequestWeightmapUpdate(bool bUpdateAll)
 {
 	LayerUpdateFlagPerMode |= ELandscapeLayerUpdateMode::Weightmap_Editing;
 	LayerUpdateFlagPerMode |= ELandscapeLayerUpdateMode::Weightmap_All;
-	GetLandscapeActor()->RequestLayersContentUpdate(ELandscapeLayerUpdateMode::Weightmap_Editing);
-	if (bUpdateAll)
+	if (ALandscape* LandscapeActor = GetLandscapeActor())
 	{
-		GetLandscapeActor()->RequestLayersContentUpdate(Weightmap_All);
+		LandscapeActor->RequestLayersContentUpdate(ELandscapeLayerUpdateMode::Weightmap_Editing);
+		if (bUpdateAll)
+		{
+			LandscapeActor->RequestLayersContentUpdate(Weightmap_All);
+		}
 	}
 }
 
