Optimized erasing weightmap data (avoid reallocating weightmaps and updating material instances twice if no data)
#rnx

diff --git a/Engine/Source/Runtime/Landscape/Private/LandscapeEditInterface.cpp b/Engine/Source/Runtime/Landscape/Private/LandscapeEditInterface.cpp
--- a/Engine/Source/Runtime/Landscape/Private/LandscapeEditInterface.cpp
+++ b/Engine/Source/Runtime/Landscape/Private/LandscapeEditInterface.cpp
@@ -2783,6 +2783,60 @@ void FLandscapeEditDataInterface::SetAlphaData(ULandscapeLayerInfoObject* const
 					continue;
 				}
 
+				auto IsDataAllZeroForComponent = [&]() -> bool
+				{
+					// Find coordinates of box that lies inside component
+					const int32 ComponentX1 = FMath::Clamp<int32>(X1 - ComponentIndexX * ComponentSizeQuads, 0, ComponentSizeQuads);
+					const int32 ComponentY1 = FMath::Clamp<int32>(Y1 - ComponentIndexY * ComponentSizeQuads, 0, ComponentSizeQuads);
+					const int32 ComponentX2 = FMath::Clamp<int32>(X2 - ComponentIndexX * ComponentSizeQuads, 0, ComponentSizeQuads);
+					const int32 ComponentY2 = FMath::Clamp<int32>(Y2 - ComponentIndexY * ComponentSizeQuads, 0, ComponentSizeQuads);
+
+					// Find subsection range for this box
+					const int32 SubIndexX1 = FMath::Clamp<int32>((ComponentX1 - 1) / SubsectionSizeQuads, 0, ComponentNumSubsections - 1); // -1 because we need to pick up vertices shared between subsections
+					const int32 SubIndexY1 = FMath::Clamp<int32>((ComponentY1 - 1) / SubsectionSizeQuads, 0, ComponentNumSubsections - 1);
+					const int32 SubIndexX2 = FMath::Clamp<int32>(ComponentX2 / SubsectionSizeQuads, 0, ComponentNumSubsections - 1);
+					const int32 SubIndexY2 = FMath::Clamp<int32>(ComponentY2 / SubsectionSizeQuads, 0, ComponentNumSubsections - 1);
+
+					for (int32 SubIndexY = SubIndexY1; SubIndexY <= SubIndexY2; SubIndexY++)
+					{
+						for (int32 SubIndexX = SubIndexX1; SubIndexX <= SubIndexX2; SubIndexX++)
+						{
+							// Find coordinates of box that lies inside subsection
+							const int32 SubX1 = FMath::Clamp<int32>(ComponentX1 - SubsectionSizeQuads * SubIndexX, 0, SubsectionSizeQuads);
+							const int32 SubY1 = FMath::Clamp<int32>(ComponentY1 - SubsectionSizeQuads * SubIndexY, 0, SubsectionSizeQuads);
+							const int32 SubX2 = FMath::Clamp<int32>(ComponentX2 - SubsectionSizeQuads * SubIndexX, 0, SubsectionSizeQuads);
+							const int32 SubY2 = FMath::Clamp<int32>(ComponentY2 - SubsectionSizeQuads * SubIndexY, 0, SubsectionSizeQuads);
+
+							// Update texture data for the box that lies inside subsection
+							for (int32 SubY = SubY1; SubY <= SubY2; SubY++)
+							{
+								for (int32 SubX = SubX1; SubX <= SubX2; SubX++)
+								{
+									const int32 LandscapeX = SubIndexX * SubsectionSizeQuads + ComponentIndexX * ComponentSizeQuads + SubX;
+									const int32 LandscapeY = SubIndexY * SubsectionSizeQuads + ComponentIndexY * ComponentSizeQuads + SubY;
+									checkSlow(LandscapeX >= X1 && LandscapeX <= X2);
+									checkSlow(LandscapeY >= Y1 && LandscapeY <= Y2);
+
+									// Find the input data corresponding to this vertex
+									const int32 DataIndex = (LandscapeX - X1) + Stride * (LandscapeY - Y1);
+									uint8 NewWeight = Data[DataIndex];
+									if (NewWeight)
+									{
+										return false;
+									}
+								}
+							}
+						}
+					}
+					return true;
+				};
+
+				// Avoid reallocating and update material instances if we're about to write all zeros and we don't have an allocated channel for this LayerInfo
+				if (IsDataAllZeroForComponent())
+				{
+					continue;
+				}
+
 				UpdateLayerIdx = ComponentWeightmapLayerAllocations.Num();
 				new (ComponentWeightmapLayerAllocations) FWeightmapLayerAllocationInfo(LayerInfo);
 				Component->ReallocateWeightmaps(this);
