Boosted grass update frequency when editing landscape for a more responsive feedback. Also sorted landscape components to process based on their minimum distance from camera.

diff --git a/Engine/Source/Runtime/Landscape/Classes/LandscapeProxy.h b/Engine/Source/Runtime/Landscape/Classes/LandscapeProxy.h
--- a/Engine/Source/Runtime/Landscape/Classes/LandscapeProxy.h
+++ b/Engine/Source/Runtime/Landscape/Classes/LandscapeProxy.h
@@ -967,4 +967,7 @@ protected:
 	FLandscapeMaterialChangedDelegate LandscapeMaterialChangedDelegate;
 
 #endif
+private:
+	/** Returns Grass Update interval */
+	int32 GetGrassUpdateInterval() const;
 };
diff --git a/Engine/Source/Runtime/Landscape/Private/LandscapeGrass.cpp b/Engine/Source/Runtime/Landscape/Private/LandscapeGrass.cpp
--- a/Engine/Source/Runtime/Landscape/Private/LandscapeGrass.cpp
+++ b/Engine/Source/Runtime/Landscape/Private/LandscapeGrass.cpp
@@ -1326,11 +1326,24 @@ void FLandscapeComponentGrassData::ConditionalDiscardDataOnLoad()
 // ALandscapeProxy grass-related functions
 //
 
+int32 ALandscapeProxy::GetGrassUpdateInterval() const
+{
+#if WITH_EDITOR
+	// When editing landscape, force update interval to be every frame
+	if (GLandscapeEditModeActive)
+	{
+		return 1;
+	}
+#endif
+	return GGrassUpdateInterval;
+}
+
 void ALandscapeProxy::TickGrass()
 {
-	if (GGrassUpdateInterval > 1)
+	const int32 UpdateInterval = GetGrassUpdateInterval();
+	if (UpdateInterval > 1)
 	{
-		if ((GFrameNumber + FrameOffsetForTickInterval) % uint32(GGrassUpdateInterval))
+		if ((GFrameNumber + FrameOffsetForTickInterval) % uint32(UpdateInterval))
 		{
 			return;
 		}
@@ -2174,33 +2187,57 @@ void ALandscapeProxy::UpdateGrass(const TArray<FVector>& Cameras, bool bForceSyn
 #endif
 			ERHIFeatureLevel::Type FeatureLevel = World->Scene->GetFeatureLevel();
 
-			int32 NumCompsCreated = 0;
-			for (int32 ComponentIndex = 0; ComponentIndex < LandscapeComponents.Num(); ComponentIndex++)
+			struct SortedLandscapeElement
 			{
+				SortedLandscapeElement(ULandscapeComponent* InComponent, float InMinDistance, const FBox& InBoundsBox)
+				: Component(InComponent)
+				, MinDistance(InMinDistance)
+				, BoundsBox(InBoundsBox)
+				{}
+
+				ULandscapeComponent* Component;
+				float MinDistance;
+				FBox BoundsBox;
+			};
 
-				ULandscapeComponent* Component = LandscapeComponents[ComponentIndex];
-
+			TArray<SortedLandscapeElement> SortedLandscapeComponents;
+			SortedLandscapeComponents.Reserve(LandscapeComponents.Num());
+			for (ULandscapeComponent* Component : LandscapeComponents)
+			{
 				// skip if we have no data and no way to generate it
-				if (Component==nullptr || (World->IsGameWorld() && !Component->GrassData->HasData()))
+				if (Component == nullptr || (World->IsGameWorld() && !Component->GrassData->HasData()))
 				{
 					continue;
 				}
-
-				FScopeCycleCounter Context(Component->GetStatID());
-
 				FBoxSphereBounds WorldBounds = Component->CalcBounds(Component->GetComponentTransform());
-				float MinDistanceToComp = Cameras.Num() ? MAX_flt : 0.0f;
-
-				for (auto& Pos : Cameras)
+				float MinSqrDistanceToComponent = Cameras.Num() ? MAX_flt : 0.0f;
+				for (const FVector& CameraPos : Cameras)
 				{
-					MinDistanceToComp = FMath::Min<float>(MinDistanceToComp, WorldBounds.ComputeSquaredDistanceFromBoxToPoint(Pos));
+					MinSqrDistanceToComponent = FMath::Min<float>(MinSqrDistanceToComponent, WorldBounds.ComputeSquaredDistanceFromBoxToPoint(CameraPos));
 				}
+				SortedLandscapeComponents.Emplace(Component, FMath::Sqrt(MinSqrDistanceToComponent), WorldBounds.GetBox());
+			}
+			
+#if WITH_EDITOR
+			// When editing landscape, prioritize components that are closer to camera for a more reactive update
+			if (GLandscapeEditModeActive)
+			{
+				Algo::Sort(SortedLandscapeComponents, [](const SortedLandscapeElement& A, const SortedLandscapeElement& B) { return A.MinDistance < B.MinDistance; });
+			}
+#endif
+
+			int32 NumCompsCreated = 0;
+			for (const SortedLandscapeElement& SortedLandscapeComponent : SortedLandscapeComponents)
+			{
+				ULandscapeComponent* Component = SortedLandscapeComponent.Component;
+				float MinDistanceToComp = SortedLandscapeComponent.MinDistance;
+
 				if (Component->ChangeTag != GGrassExclusionChangeTag)
 				{
 					Component->ActiveExcludedBoxes.Empty();
 					if (GGrassExclusionBoxes.Num() && CVarIgnoreExcludeBoxes.GetValueOnAnyThread() == 0)
 					{
-						FBox WorldBox = WorldBounds.GetBox();
+						const FBox& WorldBox = SortedLandscapeComponent.BoundsBox;
 
 						for (const TPair<FWeakObjectPtr, FBox>& Pair : GGrassExclusionBoxes)
 						{
@@ -2213,8 +2250,6 @@ void ALandscapeProxy::UpdateGrass(const TArray<FVector>& Cameras, bool bForceSyn
 					Component->ChangeTag = GGrassExclusionChangeTag;
 				}
 
-				MinDistanceToComp = FMath::Sqrt(MinDistanceToComp);
-
 				for (auto GrassType : GrassTypes)
 				{
 					if (GrassType)
@@ -2597,7 +2632,7 @@ void ALandscapeProxy::UpdateGrass(const TArray<FVector>& Cameras, bool bForceSyn
 
 		// trim cached items based on time, pending and emptiness
 		double OldestToKeepTime = FPlatformTime::Seconds() - CVarMinTimeToKeepGrass.GetValueOnGameThread();
-		uint32 OldestToKeepFrame = GFrameNumber - CVarMinFramesToKeepGrass.GetValueOnGameThread() * GGrassUpdateInterval;
+		uint32 OldestToKeepFrame = GFrameNumber - CVarMinFramesToKeepGrass.GetValueOnGameThread() * GetGrassUpdateInterval();
 		for (FCachedLandscapeFoliage::TGrassSet::TIterator Iter(FoliageCache.CachedGrassComps); Iter; ++Iter)
 		{
 			const FCachedLandscapeFoliage::FGrassComp& GrassItem = *Iter;
