- Landscape Layer System: Fix Regenerate Flag Clearing when textures not ready

#rnx

diff --git a/Engine/Source/Runtime/Landscape/Classes/Landscape.h b/Engine/Source/Runtime/Landscape/Classes/Landscape.h
--- a/Engine/Source/Runtime/Landscape/Classes/Landscape.h
+++ b/Engine/Source/Runtime/Landscape/Classes/Landscape.h
@@ -275,8 +275,9 @@ private:
 	void ReleaseLayersRenderingResource();
 	void RegenerateLayersContent();
 	void MonitorShaderCompilation();
-	void RegenerateLayersHeightmaps(const TArray<ULandscapeComponent*>& InLandscapeComponents);
-	void RegenerateLayersWeightmaps(const TArray<ULandscapeComponent*>& InLandscapeComponents);
+	int32 RegenerateLayersHeightmaps(const TArray<ULandscapeComponent*>& InLandscapeComponents);
+	int32 RegenerateLayersWeightmaps(const TArray<ULandscapeComponent*>& InLandscapeComponents);
+	static bool UpdateCollisionAndClients(const TArray<ULandscapeComponent*>& InLandscapeComponents, const int32 InContentUpdateModes, const bool bInLayerForceUpdateAllComponents);
 	void ResolveLayersHeightmapTexture();
 	void ResolveLayersWeightmapTexture();
 	void ResolveLayersTexture(class FLandscapeLayersTexture2DCPUReadBackResource* InCPUReadBackTexture, UTexture2D* InOutputTexture);
diff --git a/Engine/Source/Runtime/Landscape/Private/LandscapeEditLayers.cpp b/Engine/Source/Runtime/Landscape/Private/LandscapeEditLayers.cpp
--- a/Engine/Source/Runtime/Landscape/Private/LandscapeEditLayers.cpp
+++ b/Engine/Source/Runtime/Landscape/Private/LandscapeEditLayers.cpp
@@ -2532,7 +2532,7 @@ bool ALandscape::PrepareLayersHeightmapTextureResources() const
 	return IsReady;
 }
 
-void ALandscape::RegenerateLayersHeightmaps(const TArray<ULandscapeComponent*>& InLandscapeComponents)
+int32 ALandscape::RegenerateLayersHeightmaps(const TArray<ULandscapeComponent*>& InLandscapeComponents)
 {
 	SCOPE_CYCLE_COUNTER(STAT_LandscapeLayersRegenerateHeightmaps);
 	ULandscapeInfo* Info = GetLandscapeInfo();
@@ -2543,7 +2543,7 @@ void ALandscape::RegenerateLayersHeightmaps(const TArray<ULandscapeComponent*>&
 
 	if ((HeightmapUpdateModes == 0 && !bForceRender) || Info == nullptr || !PrepareLayersHeightmapTextureResources())
 	{
-		return;
+		return 0;
 	}
 	   
 	// Handle missing Heightmap CPUReadback (Undo of a delete landscape component can trigger that case)
@@ -2574,7 +2574,7 @@ void ALandscape::RegenerateLayersHeightmaps(const TArray<ULandscapeComponent*>&
 
 		if (!Info->GetLandscapeExtent(MinExtend.X, MinExtend.Y, MaxExtend.X, MaxExtend.Y))
 		{
-			return;
+			return 0;
 		}
 
 		// Use to compute TopLeft Component per Heightmap
@@ -2775,6 +2775,8 @@ void ALandscape::RegenerateLayersHeightmaps(const TArray<ULandscapeComponent*>&
 
 		Info->UpdateAllAddCollisions();
 	}
+
+	return HeightmapUpdateModes;
 }
 
 void ALandscape::ResolveLayersHeightmapTexture()
@@ -3339,7 +3341,7 @@ bool ALandscape::PrepareLayersWeightmapTextureResources() const
 	return IsReady;
 }
 
-void ALandscape::RegenerateLayersWeightmaps(const TArray<ULandscapeComponent*>& InLandscapeComponents)
+int32 ALandscape::RegenerateLayersWeightmaps(const TArray<ULandscapeComponent*>& InLandscapeComponents)
 {
 	SCOPE_CYCLE_COUNTER(STAT_LandscapeLayersRegenerateWeightmaps);
 	const int32 AllWeightmapUpdateModes = (ELandscapeLayerUpdateMode::Weightmap_All | ELandscapeLayerUpdateMode::Weightmap_Editing);
@@ -3350,7 +3352,7 @@ void ALandscape::RegenerateLayersWeightmaps(const TArray<ULandscapeComponent*>&
 
 	if ((WeightmapUpdateModes == 0 && !bForceRender) || Info == nullptr || Info->Layers.Num() == 0 || !PrepareLayersWeightmapTextureResources())
 	{
-		return;
+		return 0;
 	}
 			
 	TArray<ULandscapeComponent*> ComponentThatNeedMaterialRebuild;
@@ -3365,7 +3367,7 @@ void ALandscape::RegenerateLayersWeightmaps(const TArray<ULandscapeComponent*>&
 
 		if (!Info->GetLandscapeExtent(MinExtend.X, MinExtend.Y, MaxExtend.X, MaxExtend.Y))
 		{
-			return;
+			return 0;
 		}
 		
 		check(WeightmapRTList.Num() > 0);
@@ -3716,6 +3718,8 @@ void ALandscape::RegenerateLayersWeightmaps(const TArray<ULandscapeComponent*>&
 			}
 		}
 	}
+
+	return WeightmapUpdateModes;
 }
 
 uint32 ULandscapeComponent::ComputeWeightmapsHash()
@@ -4104,23 +4108,39 @@ void ALandscape::RegenerateLayersContent()
 		AllLandscapeComponents.Append(Proxy->LandscapeComponents);
 	});
 
-	RegenerateLayersHeightmaps(AllLandscapeComponents);
-	RegenerateLayersWeightmaps(AllLandscapeComponents);
+	int32 ProcessedModes = (LayerContentUpdateModes & ELandscapeLayerUpdateMode::DeferredClientUpdate);
+	ProcessedModes |= RegenerateLayersHeightmaps(AllLandscapeComponents);
+	ProcessedModes |= RegenerateLayersWeightmaps(AllLandscapeComponents);
+	LayerContentUpdateModes &= ~ProcessedModes;
+
+	if (!ALandscape::UpdateCollisionAndClients(AllLandscapeComponents, ProcessedModes, bLayerForceUpdateAllComponents))
+	{
+		LayerContentUpdateModes |= ELandscapeLayerUpdateMode::DeferredClientUpdate;
+	}
+
+	// Clear flag if we managed to update all modes
+	if (LayerContentUpdateModes == 0)
+	{
+		bLayerForceUpdateAllComponents = false;
+	}
+}
 
-	bool bDeferClientUpdate = false;
+bool ALandscape::UpdateCollisionAndClients(const TArray<ULandscapeComponent*>& InLandscapeComponents, const int32 InContentUpdateModes, const bool bInLayerForceUpdateAllComponents)
+{
+	bool bAllClientsUpdated = true;
 
-	for (ULandscapeComponent* LandscapeComponent : AllLandscapeComponents)
+	for (ULandscapeComponent* LandscapeComponent : InLandscapeComponents)
 	{
-		if (LandscapeComponent->IsUpdateFlagEnabledForModes(ELandscapeComponentUpdateFlag::RecreateCollision, LayerContentUpdateModes, bLayerForceUpdateAllComponents))
+		bool bDeferClientUpdateForComponent = false;
+		if (LandscapeComponent->IsUpdateFlagEnabledForModes(ELandscapeComponentUpdateFlag::RecreateCollision, InContentUpdateModes, bInLayerForceUpdateAllComponents))
 		{
 			if (ULandscapeHeightfieldCollisionComponent* CollisionComp = LandscapeComponent->CollisionComponent.Get())
 			{
 				CollisionComp->RecreateCollision();
-				
 			}
 		}
-		
-		if (LandscapeComponent->IsUpdateFlagEnabledForModes(ELandscapeComponentUpdateFlag::ClientUpdate, LayerContentUpdateModes, bLayerForceUpdateAllComponents))
+
+		if (LandscapeComponent->IsUpdateFlagEnabledForModes(ELandscapeComponentUpdateFlag::ClientUpdate, InContentUpdateModes, bInLayerForceUpdateAllComponents))
 		{
 			if (!GUndo)
 			{
@@ -4132,23 +4152,19 @@ void ALandscape::RegenerateLayersContent()
 			}
 			else
 			{
-				bDeferClientUpdate = true;
+				bDeferClientUpdateForComponent = true;
+				bAllClientsUpdated = false;
 			}
 		}
-		LandscapeComponent->ClearUpdateFlagsForModes(LayerContentUpdateModes);
-		if (bDeferClientUpdate)
+
+		LandscapeComponent->ClearUpdateFlagsForModes(InContentUpdateModes);
+		if (bDeferClientUpdateForComponent)
 		{
 			LandscapeComponent->RequestDeferredClientUpdate();
 		}
 	}
-	
-	LayerContentUpdateModes = 0;
-	if (bDeferClientUpdate)
-	{
-		LayerContentUpdateModes |= ELandscapeLayerUpdateMode::DeferredClientUpdate;
-	}
 
-	bLayerForceUpdateAllComponents = false;
+	return bAllClientsUpdated;
 }
 
 void ALandscape::InitializeLayers()
