Fixed a case where doing a ForceDelete on a folder would cause a MIC to have no parent anymore, resulting in a crash

diff --git a/Engine/Source/Runtime/Landscape/Private/LandscapeRender.cpp b/Engine/Source/Runtime/Landscape/Private/LandscapeRender.cpp
--- a/Engine/Source/Runtime/Landscape/Private/LandscapeRender.cpp
+++ b/Engine/Source/Runtime/Landscape/Private/LandscapeRender.cpp
@@ -797,17 +797,24 @@ FLandscapeComponentSceneProxy::FLandscapeComponentSceneProxy(ULandscapeComponent
 
 	for (UMaterialInterface*& MaterialInterface : AvailableMaterials)
 	{
-		if (MaterialInterface != nullptr)
+		UMaterial* LandscapeMaterial = MaterialInterface != nullptr ? MaterialInterface->GetMaterial() : nullptr;
+		UMaterialInstance* MaterialInstance = Cast<UMaterialInstance>(MaterialInterface);
+
+		if (LandscapeMaterial != nullptr)
 		{
-			MaterialRelevances.Add(MaterialInterface->GetRelevance(FeatureLevel));
+			// In some case it's possible that the Material Instance we have and the Material are not related, for example, in case where content was force deleted, we can have a MIC with no parent, so GetMaterial will fallback to the default material.
+			// and since the MIC is not really valid, dont generate the relevance.
+			if (MaterialInstance == nullptr || MaterialInstance->IsChildOf(LandscapeMaterial))
+			{
+				MaterialRelevances.Add(MaterialInterface->GetRelevance(FeatureLevel));
+			}
+
 			bRequiresAdjacencyInformation |= MaterialSettingsRequireAdjacencyInformation_GameThread(MaterialInterface, XYOffsetmapTexture == nullptr ? &FLandscapeVertexFactory::StaticType : &FLandscapeXYOffsetVertexFactory::StaticType, InComponent->GetWorld()->FeatureLevel);
 
 			bool HasTessellationEnabled = false;
 
 			if (FeatureLevel >= ERHIFeatureLevel::SM4)
 			{
-				UMaterial* LandscapeMaterial = MaterialInterface->GetMaterial();
-
 				HasTessellationEnabled = LandscapeMaterial != nullptr && LandscapeMaterial->D3D11TessellationMode != EMaterialTessellationMode::MTM_NoTessellation;
 			}
 
