Minor cleanup based on Richard comments

diff --git a/Engine/Source/Runtime/Landscape/Private/LandscapeRender.cpp b/Engine/Source/Runtime/Landscape/Private/LandscapeRender.cpp
--- a/Engine/Source/Runtime/Landscape/Private/LandscapeRender.cpp
+++ b/Engine/Source/Runtime/Landscape/Private/LandscapeRender.cpp
@@ -798,10 +798,11 @@ FLandscapeComponentSceneProxy::FLandscapeComponentSceneProxy(ULandscapeComponent
 	for (UMaterialInterface*& MaterialInterface : AvailableMaterials)
 	{
 		UMaterial* LandscapeMaterial = MaterialInterface != nullptr ? MaterialInterface->GetMaterial() : nullptr;
-		UMaterialInstance* MaterialInstance = Cast<UMaterialInstance>(MaterialInterface);
 
 		if (LandscapeMaterial != nullptr)
 		{
+			UMaterialInstance* MaterialInstance = Cast<UMaterialInstance>(MaterialInterface);
+
 			// In some case it's possible that the Material Instance we have and the Material are not related, for example, in case where content was force deleted, we can have a MIC with no parent, so GetMaterial will fallback to the default material.
 			// and since the MIC is not really valid, dont generate the relevance.
 			if (MaterialInstance == nullptr || MaterialInstance->IsChildOf(LandscapeMaterial))
@@ -815,7 +816,7 @@ FLandscapeComponentSceneProxy::FLandscapeComponentSceneProxy(ULandscapeComponent
 
 			if (FeatureLevel >= ERHIFeatureLevel::SM4)
 			{
-				HasTessellationEnabled = LandscapeMaterial != nullptr && LandscapeMaterial->D3D11TessellationMode != EMaterialTessellationMode::MTM_NoTessellation;
+				HasTessellationEnabled = LandscapeMaterial->D3D11TessellationMode != EMaterialTessellationMode::MTM_NoTessellation;
 			}
 
 			MaterialHasTessellationEnabled.Add(HasTessellationEnabled);
