- Update Uniform expression when changing MIC
- When loading do not regenerate content

diff --git a/Engine/Source/Runtime/Landscape/Classes/Landscape.h b/Engine/Source/Runtime/Landscape/Classes/Landscape.h
--- a/Engine/Source/Runtime/Landscape/Classes/Landscape.h
+++ b/Engine/Source/Runtime/Landscape/Classes/Landscape.h
@@ -228,7 +228,7 @@ public:
 
 	// Layers stuff
 #if WITH_EDITOR
-	LANDSCAPE_API void RequestLayersInitialization();
+	LANDSCAPE_API void RequestLayersInitialization(bool bInRequestContentUpdate = true);
 	LANDSCAPE_API void RequestLayersContentUpdateForceAll();
 	LANDSCAPE_API void RequestLayersContentUpdate(ELandscapeLayerUpdateMode InModeMask, bool bInForceUpdateAllComponents = false);
 	LANDSCAPE_API bool ReorderLayer(int32 InStartingLayerIndex, int32 InDestinationLayerIndex);
diff --git a/Engine/Source/Runtime/Landscape/Private/Landscape.cpp b/Engine/Source/Runtime/Landscape/Private/Landscape.cpp
--- a/Engine/Source/Runtime/Landscape/Private/Landscape.cpp
+++ b/Engine/Source/Runtime/Landscape/Private/Landscape.cpp
@@ -2752,7 +2752,8 @@ void ULandscapeInfo::RegisterActor(ALandscapeProxy* Proxy, bool bMapCheck)
 	if (GetMutableDefault<UEditorExperimentalSettings>()->bLandscapeLayerSystem && LandscapeActor)
 	{
 		// Force update rendering resources
-		LandscapeActor->RequestLayersInitialization();
+		const bool bInRequestContentUpdate = false;
+		LandscapeActor->RequestLayersInitialization(bInRequestContentUpdate);		
 	}
 
 	UpdateLayerInfoMap(Proxy);
diff --git a/Engine/Source/Runtime/Landscape/Private/LandscapeEditLayers.cpp b/Engine/Source/Runtime/Landscape/Private/LandscapeEditLayers.cpp
--- a/Engine/Source/Runtime/Landscape/Private/LandscapeEditLayers.cpp
+++ b/Engine/Source/Runtime/Landscape/Private/LandscapeEditLayers.cpp
@@ -3771,6 +3771,8 @@ void ALandscape::UpdateLayersMaterialInstances()
 	TOptional<FMaterialUpdateContext> MaterialUpdateContext;
 	MaterialUpdateContext.Emplace(FMaterialUpdateContext::EOptions::Default & ~FMaterialUpdateContext::EOptions::RecreateRenderStates);
 
+	bool bHasUniformExpressionUpdatePending = false;
+
 	for (ULandscapeComponent* Component : ComponentsToUpdate)
 	{
 		int32 MaxLOD = FMath::CeilLogTwo(SubsectionSizeQuads + 1) - 1;
@@ -3862,6 +3864,11 @@ void ALandscape::UpdateLayersMaterialInstances()
 				{
 					MaterialInstance->PostEditChange();
 				}
+				else
+				{
+					bHasUniformExpressionUpdatePending = true;
+					MaterialInstance->RecacheUniformExpressions(true);
+				}
 
 				/*// Setup material instance with disabled tessellation
 				if (CombinationMaterialInstance->GetMaterial()->D3D11TessellationMode != EMaterialTessellationMode::MTM_NoTessellation)
@@ -3906,6 +3913,15 @@ void ALandscape::UpdateLayersMaterialInstances()
 	// Recreate the render state for our components, needed to update the static drawlist which has cached the MaterialRenderProxies
 	// Must be after the FMaterialUpdateContext is destroyed
 	RecreateRenderStateContexts.Empty();
+
+	if (bHasUniformExpressionUpdatePending)
+	{
+		ENQUEUE_RENDER_COMMAND(UpdateDeferredCachedUniformExpressions)(
+			[](FRHICommandList& RHICmdList)
+		{
+			FMaterialRenderProxy::UpdateDeferredCachedUniformExpressions();
+		});
+	}
 }
 
 void ALandscape::ResolveLayersWeightmapTexture()
@@ -3949,15 +3965,19 @@ void ALandscape::ResolveLayersWeightmapTexture()
 
 
 
-void ALandscape::RequestLayersInitialization()
+void ALandscape::RequestLayersInitialization(bool bInRequestContentUpdate)
 {
 	if (!GetMutableDefault<UEditorExperimentalSettings>()->bLandscapeLayerSystem)
 	{
 		return;
 	}
 
-	bLandscapeLayersAreInitialized = false; 
-	RequestLayersContentUpdateForceAll();
+	bLandscapeLayersAreInitialized = false; 	
+
+	if (bInRequestContentUpdate)
+	{
+		RequestLayersContentUpdateForceAll();
+	}
 }
 
 void ALandscape::RequestLayersContentUpdate(ELandscapeLayerUpdateMode InUpdateMode, bool bInForceUpdateAllComponents)
