- Added Subtractive blend support for Height layers in Landscape Layer System
- Fixed bug with Landscape Layer alpha was applied twice for height maps
#rnx

diff --git a/Engine/Shaders/Private/LandscapeLayersPS.usf b/Engine/Shaders/Private/LandscapeLayersPS.usf
--- a/Engine/Shaders/Private/LandscapeLayersPS.usf
+++ b/Engine/Shaders/Private/LandscapeLayersPS.usf
@@ -46,6 +46,7 @@ float3 SampleTextureXNormal(float2 InTextureCoordinates, float2 InTexelSize, boo
 
 void PSHeightmapMain(float2 InTextureCoordinates : TEXCOORD0, out float4 OutColor : SV_Target0)
 {
+	float LayerVisibility = LayerInfo.y;
 	float4 SourceColor = ReadTexture1.Sample(ReadTexture1Sampler, InTextureCoordinates);
 	float Height = ExtractHeight(SourceColor.rg);
 	
@@ -55,7 +56,8 @@ void PSHeightmapMain(float2 InTextureCoordinates : TEXCOORD0, out float4 OutColo
 	// Output as Layer, so apply Layer info
 	if (OutputConfig.x == 1.0)
 	{
-		Height = round(Height * LayerInfo.x * LayerInfo.y);
+		float LayerAlpha = LayerInfo.x;
+		Height *= LayerAlpha * LayerVisibility;
 	}
 
 	// Output using 2nd heightmap
@@ -64,7 +66,9 @@ void PSHeightmapMain(float2 InTextureCoordinates : TEXCOORD0, out float4 OutColo
 		float2 BaseColor = ReadTexture2.Sample(ReadTexture2Sampler, InTextureCoordinates).rg;
 		float BaseHeight = ExtractHeight(BaseColor.rg);
 		BaseHeight -= 32768.0;
-		if (LayerInfo.z == 1 && LayerInfo.y == 1)
+
+		float LayerBlendMode = LayerInfo.z;
+		if (LayerBlendMode == 1 && LayerVisibility == 1)
 		{
 			// Alpha Blend Mode
 			float HeightAlphaBlend = SourceColor.b;
diff --git a/Engine/Source/Editor/LandscapeEditor/Private/LandscapeEdMode.cpp b/Engine/Source/Editor/LandscapeEditor/Private/LandscapeEdMode.cpp
--- a/Engine/Source/Editor/LandscapeEditor/Private/LandscapeEdMode.cpp
+++ b/Engine/Source/Editor/LandscapeEditor/Private/LandscapeEdMode.cpp
@@ -4246,12 +4246,28 @@ bool FEdModeLandscape::IsLayerAlphaVisible(int32 InLayerIndex) const
 	return (CurrentToolTarget.TargetType == ELandscapeToolTargetType::Heightmap || CurrentToolTarget.TargetType == ELandscapeToolTargetType::Weightmap);
 }
 
+float FEdModeLandscape::GetClampedLayerAlpha(float InLayerAlpha) const
+{
+	ALandscape* Landscape = GetLandscape();
+	if (Landscape)
+	{
+		if (CurrentToolTarget.TargetType == ELandscapeToolTargetType::Heightmap || CurrentToolTarget.TargetType == ELandscapeToolTargetType::Weightmap)
+		{
+			return Landscape->GetClampedLayerAlpha(InLayerAlpha, CurrentToolTarget.TargetType == ELandscapeToolTargetType::Heightmap);
+		}
+	}
+	return InLayerAlpha;
+}
+
 float FEdModeLandscape::GetLayerAlpha(int32 InLayerIndex) const
 {
-	FLandscapeLayer* Layer = GetLayer(InLayerIndex);
-	if (Layer)
+	ALandscape* Landscape = GetLandscape();
+	if (Landscape)
 	{
-		return (CurrentToolTarget.TargetType == ELandscapeToolTargetType::Heightmap) ? Layer->HeightmapAlpha : Layer->WeightmapAlpha;
+		if (CurrentToolTarget.TargetType == ELandscapeToolTargetType::Heightmap || CurrentToolTarget.TargetType == ELandscapeToolTargetType::Weightmap)
+		{
+			return Landscape->GetLayerAlpha(InLayerIndex, CurrentToolTarget.TargetType == ELandscapeToolTargetType::Heightmap);
+		}
 	}
 	return 1.0f;
 }
diff --git a/Engine/Source/Editor/LandscapeEditor/Private/LandscapeEdMode.h b/Engine/Source/Editor/LandscapeEditor/Private/LandscapeEdMode.h
--- a/Engine/Source/Editor/LandscapeEditor/Private/LandscapeEdMode.h
+++ b/Engine/Source/Editor/LandscapeEditor/Private/LandscapeEdMode.h
@@ -500,8 +500,9 @@ public:
 	FName GetLayerName(int32 InLayerIndex) const;
 	void SetLayerName(int32 InLayerIndex, const FName& InName);
 	bool CanRenameLayerTo(int32 InLayerIndex, const FName& InNewName);
-	float GetLayerAlpha(int32 InLayerIndex) const;
 	void SetLayerAlpha(int32 InLayerIndex, float InAlpha);
+	float GetLayerAlpha(int32 InLayerIndex) const;
+	float GetClampedLayerAlpha(float InLayerAlpha) const;
 	void SetLayerVisibility(bool InVisible, int32 InLayerIndex);
 	bool IsLayerVisible(int32 InLayerIndex) const;
 	bool IsLayerLocked(int32 InLayerIndex) const;
diff --git a/Engine/Source/Editor/LandscapeEditor/Private/LandscapeEditorDetailCustomization_Layers.cpp b/Engine/Source/Editor/LandscapeEditor/Private/LandscapeEditorDetailCustomization_Layers.cpp
--- a/Engine/Source/Editor/LandscapeEditor/Private/LandscapeEditorDetailCustomization_Layers.cpp
+++ b/Engine/Source/Editor/LandscapeEditor/Private/LandscapeEditorDetailCustomization_Layers.cpp
@@ -269,9 +269,11 @@ TSharedPtr<SWidget> FLandscapeEditorCustomNodeBuilder_Layers::GenerateRow(int32
 				[
 					SNew(SNumericEntryBox<float>)
 					.AllowSpin(true)
-					.MinValue(0.0f)
-					.MaxValue(100.0f)
-					.MaxSliderValue(100.0f)
+					.MinValue(-1.0f)
+					.MaxValue(1.0f)
+					.MinSliderValue(-1.0f)
+					.MaxSliderValue(1.0f)
+					.Delta(0.01f)
 					.MinDesiredValueWidth(60.0f)
 					.IsEnabled(this, &FLandscapeEditorCustomNodeBuilder_Layers::IsLayerEditionEnabled, InLayerIndex)
 					.Visibility(this, &FLandscapeEditorCustomNodeBuilder_Layers::GetLayerAlphaVisibility, InLayerIndex)
@@ -585,7 +587,7 @@ void FLandscapeEditorCustomNodeBuilder_Layers::SetLayerAlpha(float InAlpha, int3
 	if (LandscapeEdMode)
 	{
 		// We get multiple commits when editing through the text box
-		if (LandscapeEdMode->GetLayerAlpha(InLayerIndex) == InAlpha)
+		if (LandscapeEdMode->GetLayerAlpha(InLayerIndex) == LandscapeEdMode->GetClampedLayerAlpha(InAlpha))
 		{
 			return;
 		}
diff --git a/Engine/Source/Runtime/Landscape/Classes/Landscape.h b/Engine/Source/Runtime/Landscape/Classes/Landscape.h
--- a/Engine/Source/Runtime/Landscape/Classes/Landscape.h
+++ b/Engine/Source/Runtime/Landscape/Classes/Landscape.h
@@ -242,6 +242,8 @@ public:
 	LANDSCAPE_API bool IsLayerNameUnique(const FName& InName) const;
 	LANDSCAPE_API void SetLayerName(int32 InLayerIndex, const FName& InName);
 	LANDSCAPE_API void SetLayerAlpha(int32 InLayerIndex, const float InAlpha, bool bInHeightmap);
+	LANDSCAPE_API float GetLayerAlpha(int32 InLayerIndex, bool bInHeightmap) const;
+	LANDSCAPE_API float GetClampedLayerAlpha(float InAlpha, bool bInHeightmap) const;
 	LANDSCAPE_API void SetLayerVisibility(int32 InLayerIndex, bool bInVisible);
 	LANDSCAPE_API void SetLayerLocked(int32 InLayerIndex, bool bLocked);
 	LANDSCAPE_API uint8 GetLayerCount() const;
diff --git a/Engine/Source/Runtime/Landscape/Private/LandscapeEditLayers.cpp b/Engine/Source/Runtime/Landscape/Private/LandscapeEditLayers.cpp
--- a/Engine/Source/Runtime/Landscape/Private/LandscapeEditLayers.cpp
+++ b/Engine/Source/Runtime/Landscape/Private/LandscapeEditLayers.cpp
@@ -2625,7 +2625,7 @@ int32 ALandscape::RegenerateLayersHeightmaps(const TArray<ULandscapeComponent*>&
 		for (FLandscapeLayer& Layer : LandscapeLayers)
 		{
 			//Draw Layer heightmap to Combined RT Atlas
-			ShaderParams.ApplyLayerModifiers = true;
+			ShaderParams.ApplyLayerModifiers = false;
 			ShaderParams.LayerVisible = Layer.bVisible;
 			ShaderParams.GenerateNormals = false;
 			ShaderParams.LayerBlendMode = Layer.BlendMode;
@@ -2659,6 +2659,8 @@ int32 ALandscape::RegenerateLayersHeightmaps(const TArray<ULandscapeComponent*>&
 			DrawHeightmapComponentsToRenderTarget(OutputDebugName ? FString::Printf(TEXT("LS Height: %s += -> NonAtlas %s"), *Layer.Name.ToString(), *LandscapeScratchRT1->GetName(), *LandscapeScratchRT2->GetName()) : TEXT(""),
 												  InLandscapeComponents, MinExtend, LandscapeScratchRT1, nullptr, LandscapeScratchRT2, ERTDrawingType::RTAtlasToNonAtlas, true, ShaderParams);
 
+			ShaderParams.ApplyLayerModifiers = true;
+
 			// Combine Current layer with current result
 			DrawHeightmapComponentsToRenderTarget(OutputDebugName ? FString::Printf(TEXT("LS Height: %s += -> CombinedNonAtlas %s"), *Layer.Name.ToString(), *LandscapeScratchRT2->GetName(), *CombinedHeightmapNonAtlasRT->GetName()) : TEXT(""),
 												InLandscapeComponents, MinExtend, LandscapeScratchRT2, FirstLayer ? nullptr : LandscapeScratchRT3, CombinedHeightmapNonAtlasRT, ERTDrawingType::RTNonAtlas, FirstLayer, ShaderParams);
@@ -4516,7 +4518,24 @@ void ALandscape::SetLayerName(int32 InLayerIndex, const FName& InName)
 	LandscapeLayers[InLayerIndex].Name = InName;
 }
 
-void ALandscape::SetLayerAlpha(int32 InLayerIndex, const float InAlpha, bool bInHeightmap)
+float ALandscape::GetLayerAlpha(int32 InLayerIndex, bool bInHeightmap) const
+{
+	const FLandscapeLayer* SplinesReservedLayer = GetLandscapeSplinesReservedLayer();
+	const FLandscapeLayer* Layer = GetLayer(InLayerIndex);
+	if (Layer && SplinesReservedLayer != Layer)
+	{
+		return GetClampedLayerAlpha(bInHeightmap ? Layer->HeightmapAlpha : Layer->WeightmapAlpha, bInHeightmap);
+	}
+	return 1.0f;
+}
+
+float ALandscape::GetClampedLayerAlpha(float InAlpha, bool bInHeightmap) const
+{
+	float AlphaClamped = FMath::Clamp<float>(InAlpha, bInHeightmap  ? -1.f : 0.f, 1.f);
+	return AlphaClamped;
+}
+
+void ALandscape::SetLayerAlpha(int32 InLayerIndex, float InAlpha, bool bInHeightmap)
 {
 	FLandscapeLayer* Layer = GetLayer(InLayerIndex);
 	ULandscapeInfo* LandscapeInfo = GetLandscapeInfo();
@@ -4524,14 +4543,15 @@ void ALandscape::SetLayerAlpha(int32 InLayerIndex, const float InAlpha, bool bIn
 	{
 		return;
 	}
+	const float InAlphaClamped = GetClampedLayerAlpha(InAlpha, bInHeightmap);
 	float& LayerAlpha = bInHeightmap ? Layer->HeightmapAlpha : Layer->WeightmapAlpha;
-	if (LayerAlpha == InAlpha)
+	if (LayerAlpha == InAlphaClamped)
 	{
 		return;
 	}
 
 	Modify();
-	LayerAlpha = InAlpha;
+	LayerAlpha = InAlphaClamped;
 	RequestLayersContentUpdateForceAll();
 }
 
