- Landscape Layer System: OnPreSave make sure we have generated final data

#rnx

diff --git a/Engine/Source/Editor/LandscapeEditor/Private/LandscapeEdMode.cpp b/Engine/Source/Editor/LandscapeEditor/Private/LandscapeEdMode.cpp
--- a/Engine/Source/Editor/LandscapeEditor/Private/LandscapeEdMode.cpp
+++ b/Engine/Source/Editor/LandscapeEditor/Private/LandscapeEdMode.cpp
@@ -779,6 +779,25 @@ void FEdModeLandscape::OnVRHoverUpdate(UViewportInteractor* Interactor, FVector&
 	}
 }
 
+void FEdModeLandscape::OnPreSaveWorld(uint32 InSaveFlags, const class UWorld* InWorld)
+{
+	// Avoid doing this during cooking to keep determinism and we don't want to do this on GameWorlds.
+	if (!InWorld->IsGameWorld() && !GIsCookerLoadingPackage)
+	{
+		ULandscapeInfoMap& LandscapeInfoMap = ULandscapeInfoMap::GetLandscapeInfoMap(InWorld);
+		for (const TPair<FGuid, ULandscapeInfo*>& Pair : LandscapeInfoMap.Map)
+		{
+			if (const ULandscapeInfo* LandscapeInfo = Pair.Value)
+			{
+				if (ALandscape* LandscapeActor = LandscapeInfo->LandscapeActor.Get())
+				{
+					LandscapeActor->OnPreSave();
+				}
+			}
+		}
+	}
+}
+
 void FEdModeLandscape::OnVRAction(FEditorViewportClient& ViewportClient, UViewportInteractor* Interactor, const struct FViewportActionKeyInput& Action, bool& bOutIsInputCaptured, bool& bWasHandled)
 {
 	UVREditorMode* VREditorMode = Cast<UVREditorMode>( GEditor->GetEditorWorldExtensionsManager()->GetEditorWorldExtensions( GetWorld() )->FindExtension( UVREditorMode::StaticClass() ) );
diff --git a/Engine/Source/Editor/LandscapeEditor/Private/LandscapeEdMode.h b/Engine/Source/Editor/LandscapeEditor/Private/LandscapeEdMode.h
--- a/Engine/Source/Editor/LandscapeEditor/Private/LandscapeEdMode.h
+++ b/Engine/Source/Editor/LandscapeEditor/Private/LandscapeEdMode.h
@@ -534,6 +534,8 @@ public:
 	DECLARE_EVENT(FEdModeLandscape, FTargetsListUpdated);
 	static FTargetsListUpdated TargetsListUpdated;
 
+	void OnPreSaveWorld(uint32 InSaveFlags, const class UWorld* InWorld);
+
 	/** Called when the user presses a button on their motion controller device */
 	void OnVRAction(FEditorViewportClient& ViewportClient, UViewportInteractor* Interactor, const FViewportActionKeyInput& Action, bool& bOutIsInputCaptured, bool& bWasHandled);
 
diff --git a/Engine/Source/Editor/LandscapeEditor/Private/LandscapeEditorModule.cpp b/Engine/Source/Editor/LandscapeEditor/Private/LandscapeEditorModule.cpp
--- a/Engine/Source/Editor/LandscapeEditor/Private/LandscapeEditorModule.cpp
+++ b/Engine/Source/Editor/LandscapeEditor/Private/LandscapeEditorModule.cpp
@@ -60,6 +60,8 @@ public:
 	{
 		FLandscapeEditorCommands::Register();
 
+		PreSaveWorldHandle = FEditorDelegates::PreSaveWorld.AddRaw(this, &FLandscapeEditorModule::OnPreSaveWorld);
+
 		// register the editor mode
 		FEditorModeRegistry::Get().RegisterMode<FEdModeLandscape>(
 			FBuiltinEditorModes::EM_Landscape,
@@ -116,6 +118,8 @@ public:
 	{
 		FLandscapeEditorCommands::Unregister();
 
+		FEditorDelegates::PreSaveWorld.Remove(PreSaveWorldHandle);
+
 		// unregister the editor mode
 		FEditorModeRegistry::Get().UnregisterMode(FBuiltinEditorModes::EM_Landscape);
 
@@ -219,6 +223,15 @@ public:
 		}
 	}
 
+	void OnPreSaveWorld(uint32 SaveFlags, class UWorld* World)
+	{
+		FEdModeLandscape* EdMode = (FEdModeLandscape*)GLevelEditorModeTools().FindMode(FBuiltinEditorModes::EM_Landscape);
+		if (EdMode)
+		{
+			EdMode->OnPreSaveWorld(SaveFlags, World);
+		}
+	}
+
 	virtual const TCHAR* GetHeightmapImportDialogTypeString() const override;
 	virtual const TCHAR* GetWeightmapImportDialogTypeString() const override;
 
@@ -229,8 +242,9 @@ public:
 	virtual const ILandscapeWeightmapFileFormat* GetWeightmapFormatByExtension(const TCHAR* Extension) const override;
 
 	virtual TSharedPtr<FUICommandList> GetLandscapeLevelViewportCommandList() const override;
-
+		
 protected:
+	FDelegateHandle PreSaveWorldHandle;
 	TSharedPtr<FExtender> ViewportMenuExtender;
 	TSharedPtr<FUICommandList> GlobalUICommandList;
 	TArray<FRegisteredLandscapeHeightmapFileFormat> HeightmapFormats;
diff --git a/Engine/Source/Runtime/Landscape/Classes/Landscape.h b/Engine/Source/Runtime/Landscape/Classes/Landscape.h
--- a/Engine/Source/Runtime/Landscape/Classes/Landscape.h
+++ b/Engine/Source/Runtime/Landscape/Classes/Landscape.h
@@ -275,21 +275,23 @@ public:
 	LANDSCAPE_API class ALandscapeBlueprintCustomBrush* GetBrushForLayer(int32 InLayerIndex, int32 InTargetType, int8 BrushIndex) const;
 	LANDSCAPE_API TArray<class ALandscapeBlueprintCustomBrush*> GetBrushesForLayer(int32 InLayerIndex, int32 InTargetType) const;
 	
+	LANDSCAPE_API void OnPreSave();
+
 private:
 	void TickLayers(float DeltaTime, ELevelTick TickType, FActorTickFunction& ThisTickFunction);
 	void CreateLayersRenderingResource();
 	void ReleaseLayersRenderingResource();
-	void RegenerateLayersContent();
+	void UpdateLayersContent(bool bInWaitForStreaming = false);
 	void MonitorShaderCompilation();
-	int32 RegenerateLayersHeightmaps(const TArray<ULandscapeComponent*>& InLandscapeComponents);
-	int32 RegenerateLayersWeightmaps(const TArray<ULandscapeComponent*>& InLandscapeComponents);
+	int32 RegenerateLayersHeightmaps(const TArray<ULandscapeComponent*>& InLandscapeComponents, bool bInWaitForStreaming);
+	int32 RegenerateLayersWeightmaps(const TArray<ULandscapeComponent*>& InLandscapeComponents, bool bInWaitForStreaming);
 	static bool UpdateCollisionAndClients(const TArray<ULandscapeComponent*>& InLandscapeComponents, const int32 InContentUpdateModes, const bool bInLayerForceUpdateAllComponents);
 	void ResolveLayersHeightmapTexture();
 	void ResolveLayersWeightmapTexture();
 	void ResolveLayersTexture(class FLandscapeLayersTexture2DCPUReadBackResource* InCPUReadBackTexture, UTexture2D* InOutputTexture);
 
-	bool PrepareLayersHeightmapTextureResources() const;
-	bool PrepareLayersWeightmapTextureResources() const;
+	bool PrepareLayersHeightmapTextureResources(bool bInWaitForStreaming) const;
+	bool PrepareLayersWeightmapTextureResources(bool bInWaitForStreaming) const;
 
 	void UpdateLayersMaterialInstances();
 
diff --git a/Engine/Source/Runtime/Landscape/Classes/LandscapeInfoMap.h b/Engine/Source/Runtime/Landscape/Classes/LandscapeInfoMap.h
--- a/Engine/Source/Runtime/Landscape/Classes/LandscapeInfoMap.h
+++ b/Engine/Source/Runtime/Landscape/Classes/LandscapeInfoMap.h
@@ -35,6 +35,6 @@ public:
 	*
 	* @returns Landscape-specific data associated with given world.
 	*/
-	LANDSCAPE_API static ULandscapeInfoMap& GetLandscapeInfoMap(UWorld* World);
+	LANDSCAPE_API static ULandscapeInfoMap& GetLandscapeInfoMap(const UWorld* World);
 #endif // WITH_EDITORONLY_DATA
 };
diff --git a/Engine/Source/Runtime/Landscape/Private/LandscapeEditLayers.cpp b/Engine/Source/Runtime/Landscape/Private/LandscapeEditLayers.cpp
--- a/Engine/Source/Runtime/Landscape/Private/LandscapeEditLayers.cpp
+++ b/Engine/Source/Runtime/Landscape/Private/LandscapeEditLayers.cpp
@@ -1231,6 +1231,19 @@ typedef FLandscapeLayersRender_RenderThread<FLandscapeLayersWeightmapShaderParam
 
 #if WITH_EDITOR
 
+struct FLandscapeIsTextureFullyStreamedIn
+{
+	bool operator()(UTexture2D* InTexture, bool bInWaitForStreaming)
+	{
+		check(InTexture);
+		if (bInWaitForStreaming)
+		{
+			InTexture->WaitForStreaming();
+		}
+		return InTexture->IsFullyStreamedIn();
+	}
+};
+
 void ALandscape::CreateLayersRenderingResource()
 {
 	ULandscapeInfo* Info = GetLandscapeInfo();
@@ -2482,7 +2495,7 @@ void ALandscape::PrintLayersDebugTextureResource(const FString& InContext, FText
 	}
 }
 
-bool ALandscape::PrepareLayersHeightmapTextureResources() const
+bool ALandscape::PrepareLayersHeightmapTextureResources(bool bInWaitForStreaming) const
 {
 	ULandscapeInfo* Info = GetLandscapeInfo();
 
@@ -2495,11 +2508,13 @@ bool ALandscape::PrepareLayersHeightmapTextureResources() const
 
 	Info->ForAllLandscapeProxies([&](ALandscapeProxy* Proxy)
 	{
+		FLandscapeIsTextureFullyStreamedIn IsTextureFullyStreamedIn;
+				
 		for (ULandscapeComponent* Component : Proxy->LandscapeComponents)
 		{
 			UTexture2D* ComponentHeightmap = Component->GetHeightmap();
 
-			IsReady &= ComponentHeightmap->IsFullyStreamedIn();
+			IsReady &= IsTextureFullyStreamedIn(ComponentHeightmap, bInWaitForStreaming);
 
 			for (const FLandscapeLayer& Layer : LandscapeLayers)
 			{
@@ -2523,16 +2538,22 @@ bool ALandscape::PrepareLayersHeightmapTextureResources() const
 						}
 					}
 
-					IsReady &= LayerHeightmap->Resource != nullptr && LayerHeightmap->Resource->IsInitialized() && LayerHeightmap->IsFullyStreamedIn();
+					IsReady &= IsTextureFullyStreamedIn(LayerHeightmap, bInWaitForStreaming);
+					IsReady &= bInWaitForStreaming || (LayerHeightmap->Resource != nullptr && LayerHeightmap->Resource->IsInitialized());
 				}
 			}
 		}
 	});
 
+	if (bInWaitForStreaming)
+	{
+		FlushRenderingCommands();
+	}
+
 	return IsReady;
 }
 
-int32 ALandscape::RegenerateLayersHeightmaps(const TArray<ULandscapeComponent*>& InLandscapeComponents)
+int32 ALandscape::RegenerateLayersHeightmaps(const TArray<ULandscapeComponent*>& InLandscapeComponents, bool bInWaitForStreaming)
 {
 	SCOPE_CYCLE_COUNTER(STAT_LandscapeLayersRegenerateHeightmaps);
 	ULandscapeInfo* Info = GetLandscapeInfo();
@@ -2541,7 +2562,7 @@ int32 ALandscape::RegenerateLayersHeightmaps(const TArray<ULandscapeComponent*>&
 	const int32 HeightmapUpdateModes = LayerContentUpdateModes & AllHeightmapUpdateModes;
 	const bool bForceRender = CVarOutputLayersDebugDrawCallName.GetValueOnAnyThread() == 1;
 
-	if ((HeightmapUpdateModes == 0 && !bForceRender) || Info == nullptr || !PrepareLayersHeightmapTextureResources())
+	if ((HeightmapUpdateModes == 0 && !bForceRender) || Info == nullptr || !PrepareLayersHeightmapTextureResources(bInWaitForStreaming))
 	{
 		return 0;
 	}
@@ -3287,7 +3308,7 @@ void ALandscape::InitializeLayersWeightmapResources()
 	BeginInitResource(WeightmapScratchPackLayerTextureResource);
 }
 
-bool ALandscape::PrepareLayersWeightmapTextureResources() const
+bool ALandscape::PrepareLayersWeightmapTextureResources(bool bInWaitForStreaming) const
 {
 	ULandscapeInfo* Info = GetLandscapeInfo();
 
@@ -3300,13 +3321,14 @@ bool ALandscape::PrepareLayersWeightmapTextureResources() const
 
 	Info->ForAllLandscapeProxies([&](ALandscapeProxy* Proxy)
 	{
+		FLandscapeIsTextureFullyStreamedIn IsTextureFullyStreamedIn;
 		for (const FLandscapeLayer& Layer : LandscapeLayers)
 		{
 			for (ULandscapeComponent* Component : Proxy->LandscapeComponents)
 			{
 				for (UTexture2D* ComponentWeightmap : Component->GetWeightmapTextures())
 				{
-					if (!ComponentWeightmap->IsFullyStreamedIn())
+					if (!IsTextureFullyStreamedIn(ComponentWeightmap, bInWaitForStreaming))
 					{
 						IsReady = false;
 						break;
@@ -3333,17 +3355,23 @@ bool ALandscape::PrepareLayersWeightmapTextureResources() const
 							}
 						}
 
-						IsReady &= LayerWeightmap->Resource != nullptr && LayerWeightmap->Resource->IsInitialized() && LayerWeightmap->IsFullyStreamedIn();
+						IsReady &= IsTextureFullyStreamedIn(LayerWeightmap, bInWaitForStreaming);
+						IsReady &= bInWaitForStreaming || (LayerWeightmap->Resource != nullptr && LayerWeightmap->Resource->IsInitialized());
 					}
 				}
 			}
 		}
 	});
 
+	if (bInWaitForStreaming)
+	{
+		FlushRenderingCommands();
+	}
+
 	return IsReady;
 }
 
-int32 ALandscape::RegenerateLayersWeightmaps(const TArray<ULandscapeComponent*>& InLandscapeComponents)
+int32 ALandscape::RegenerateLayersWeightmaps(const TArray<ULandscapeComponent*>& InLandscapeComponents, bool bInWaitForStreaming)
 {
 	SCOPE_CYCLE_COUNTER(STAT_LandscapeLayersRegenerateWeightmaps);
 	const int32 AllWeightmapUpdateModes = (ELandscapeLayerUpdateMode::Weightmap_All | ELandscapeLayerUpdateMode::Weightmap_Editing);
@@ -3352,7 +3380,7 @@ int32 ALandscape::RegenerateLayersWeightmaps(const TArray<ULandscapeComponent*>&
 	
 	ULandscapeInfo* Info = GetLandscapeInfo();
 
-	if ((WeightmapUpdateModes == 0 && !bForceRender) || Info == nullptr || Info->Layers.Num() == 0 || !PrepareLayersWeightmapTextureResources())
+	if ((WeightmapUpdateModes == 0 && !bForceRender) || Info == nullptr || Info->Layers.Num() == 0 || !PrepareLayersWeightmapTextureResources(bInWaitForStreaming))
 	{
 		return 0;
 	}
@@ -4115,8 +4143,24 @@ void ALandscape::MonitorShaderCompilation()
 	}
 }
 
-void ALandscape::RegenerateLayersContent()
+void ALandscape::UpdateLayersContent(bool bInWaitForStreaming)
 {
+	if (!GetMutableDefault<UEditorExperimentalSettings>()->bLandscapeLayerSystem)
+	{
+		if (bInitializedWithFlagExperimentalLandscapeLayers)
+		{
+			ReleaseLayersRenderingResource();
+			bInitializedWithFlagExperimentalLandscapeLayers = false;
+		}
+		return;
+	}
+
+	if (!bLandscapeLayersAreInitialized || !bInitializedWithFlagExperimentalLandscapeLayers)
+	{
+		InitializeLayers();
+		bInitializedWithFlagExperimentalLandscapeLayers = true;
+	}
+	
 	MonitorShaderCompilation();
 
 	if (GetLandscapeInfo() == nullptr || LayerContentUpdateModes == 0)
@@ -4131,8 +4175,8 @@ void ALandscape::RegenerateLayersContent()
 	});
 
 	int32 ProcessedModes = (LayerContentUpdateModes & ELandscapeLayerUpdateMode::DeferredClientUpdate);
-	ProcessedModes |= RegenerateLayersHeightmaps(AllLandscapeComponents);
-	ProcessedModes |= RegenerateLayersWeightmaps(AllLandscapeComponents);
+	ProcessedModes |= RegenerateLayersHeightmaps(AllLandscapeComponents, bInWaitForStreaming);
+	ProcessedModes |= RegenerateLayersWeightmaps(AllLandscapeComponents, bInWaitForStreaming);
 	LayerContentUpdateModes &= ~ProcessedModes;
 
 	if (!ALandscape::UpdateCollisionAndClients(AllLandscapeComponents, ProcessedModes, bLayerForceUpdateAllComponents))
@@ -4210,6 +4254,12 @@ void ALandscape::InitializeLayers()
 	bLandscapeLayersAreInitialized = true;
 }
 
+void ALandscape::OnPreSave()
+{
+	const bool bWaitForStreaming = true;
+	UpdateLayersContent(bWaitForStreaming);
+}
+
 void ALandscape::TickLayers(float DeltaTime, ELevelTick TickType, FActorTickFunction& ThisTickFunction)
 {
 	check(GIsEditor);
@@ -4222,24 +4272,7 @@ void ALandscape::TickLayers(float DeltaTime, ELevelTick TickType, FActorTickFunc
 			World->bShouldSimulatePhysics = true;
 		}
 
-		if (GetMutableDefault<UEditorExperimentalSettings>()->bLandscapeLayerSystem)
-		{
-			if (!bLandscapeLayersAreInitialized || !bInitializedWithFlagExperimentalLandscapeLayers)
-			{
-				InitializeLayers();
-				bInitializedWithFlagExperimentalLandscapeLayers = true;
-			}
-
-			RegenerateLayersContent();
-		}
-		else
-		{
-			if (bInitializedWithFlagExperimentalLandscapeLayers)
-			{
-				ReleaseLayersRenderingResource();
-				bInitializedWithFlagExperimentalLandscapeLayers = false;
-			}
-		}
+		UpdateLayersContent();
 	}
 }
 
diff --git a/Engine/Source/Runtime/Landscape/Private/LandscapeInfoMap.cpp b/Engine/Source/Runtime/Landscape/Private/LandscapeInfoMap.cpp
--- a/Engine/Source/Runtime/Landscape/Private/LandscapeInfoMap.cpp
+++ b/Engine/Source/Runtime/Landscape/Private/LandscapeInfoMap.cpp
@@ -44,7 +44,7 @@ void ULandscapeInfoMap::AddReferencedObjects(UObject* InThis, FReferenceCollecto
 }
 
 #if WITH_EDITORONLY_DATA
-ULandscapeInfoMap& ULandscapeInfoMap::GetLandscapeInfoMap(UWorld* World)
+ULandscapeInfoMap& ULandscapeInfoMap::GetLandscapeInfoMap(const UWorld* World)
 {
 	ULandscapeInfoMap *FoundObject = nullptr;
 	World->PerModuleDataObjects.FindItemByClass(&FoundObject);
