- Landscape Layer System:
Fix Undo/Redo weightmap rare bug
Fix Undo/Redo of use cases where RequestLayersContentUpdateForceAll was undone.
Force Update now just propagates to components instead of maintaining a flag in the Proxy

#rnx

diff --git a/Engine/Source/Editor/LandscapeEditor/Private/LandscapeEdMode.cpp b/Engine/Source/Editor/LandscapeEditor/Private/LandscapeEdMode.cpp
--- a/Engine/Source/Editor/LandscapeEditor/Private/LandscapeEdMode.cpp
+++ b/Engine/Source/Editor/LandscapeEditor/Private/LandscapeEdMode.cpp
@@ -4333,19 +4333,19 @@ void FEdModeLandscape::SetLayerLocked(int32 InLayerIndex, bool bInLocked)
 	}
 }
 
-void FEdModeLandscape::RequestLayersContentUpdate(ELandscapeLayerUpdateMode InUpdateMode, bool bInForceUpdateAllComponents)
+void FEdModeLandscape::RequestLayersContentUpdate(ELandscapeLayerUpdateMode InUpdateMode)
 {
 	if (ALandscape* Landscape = GetLandscape())
 	{
-		Landscape->RequestLayersContentUpdate(InUpdateMode, bInForceUpdateAllComponents);
+		Landscape->RequestLayersContentUpdate(InUpdateMode);
 	}
 }
 
-void FEdModeLandscape::RequestLayersContentUpdateForceAll()
+void FEdModeLandscape::RequestLayersContentUpdateForceAll(ELandscapeLayerUpdateMode InUpdateMode)
 {
 	if (ALandscape* Landscape = GetLandscape())
 	{
-		Landscape->RequestLayersContentUpdateForceAll();
+		Landscape->RequestLayersContentUpdateForceAll(InUpdateMode);
 	}
 }
 
diff --git a/Engine/Source/Editor/LandscapeEditor/Private/LandscapeEdMode.h b/Engine/Source/Editor/LandscapeEditor/Private/LandscapeEdMode.h
--- a/Engine/Source/Editor/LandscapeEditor/Private/LandscapeEdMode.h
+++ b/Engine/Source/Editor/LandscapeEditor/Private/LandscapeEdMode.h
@@ -525,8 +525,8 @@ public:
 	TArray<class ALandscapeBlueprintCustomBrush*> GetBrushesForCurrentLayer(int32 InTargetType);
 	
 	bool NeedToFillEmptyMaterialLayers() const;
-	void RequestLayersContentUpdate(ELandscapeLayerUpdateMode InUpdateMode, bool bInForceUpdateAllComponents = false);
-	void RequestLayersContentUpdateForceAll();
+	void RequestLayersContentUpdate(ELandscapeLayerUpdateMode InUpdateMode);
+	void RequestLayersContentUpdateForceAll(ELandscapeLayerUpdateMode InUpdateMode = ELandscapeLayerUpdateMode::All);
 
 	void OnLevelActorAdded(AActor* InActor);
 	void OnLevelActorRemoved(AActor* InActor);
diff --git a/Engine/Source/Runtime/Landscape/Classes/Landscape.h b/Engine/Source/Runtime/Landscape/Classes/Landscape.h
--- a/Engine/Source/Runtime/Landscape/Classes/Landscape.h
+++ b/Engine/Source/Runtime/Landscape/Classes/Landscape.h
@@ -229,8 +229,8 @@ public:
 	// Layers stuff
 #if WITH_EDITOR
 	LANDSCAPE_API void RequestLayersInitialization(bool bInRequestContentUpdate = true);
-	LANDSCAPE_API void RequestLayersContentUpdateForceAll();
-	LANDSCAPE_API void RequestLayersContentUpdate(ELandscapeLayerUpdateMode InModeMask, bool bInForceUpdateAllComponents = false);
+	LANDSCAPE_API void RequestLayersContentUpdateForceAll(ELandscapeLayerUpdateMode InModeMask = ELandscapeLayerUpdateMode::All);
+	LANDSCAPE_API void RequestLayersContentUpdate(ELandscapeLayerUpdateMode InModeMask);
 	LANDSCAPE_API bool ReorderLayer(int32 InStartingLayerIndex, int32 InDestinationLayerIndex);
 	LANDSCAPE_API FLandscapeLayer* DuplicateLayer(const FLandscapeLayer& InOtherLayer);
 	LANDSCAPE_API void CreateLayer(FName InName = NAME_None);
@@ -285,7 +285,7 @@ private:
 	void MonitorShaderCompilation();
 	int32 RegenerateLayersHeightmaps(const TArray<ULandscapeComponent*>& InLandscapeComponents, bool bInWaitForStreaming);
 	int32 RegenerateLayersWeightmaps(const TArray<ULandscapeComponent*>& InLandscapeComponents, bool bInWaitForStreaming);
-	static bool UpdateCollisionAndClients(const TArray<ULandscapeComponent*>& InLandscapeComponents, const int32 InContentUpdateModes, const bool bInLayerForceUpdateAllComponents);
+	static bool UpdateCollisionAndClients(const TArray<ULandscapeComponent*>& InLandscapeComponents, const int32 InContentUpdateModes);
 	void ResolveLayersHeightmapTexture();
 	void ResolveLayersWeightmapTexture();
 	void ResolveLayersTexture(class FLandscapeLayersTexture2DCPUReadBackResource* InCPUReadBackTexture, UTexture2D* InOutputTexture);
@@ -368,10 +368,7 @@ private:
 
 	UPROPERTY(Transient)
 	uint32 LayerContentUpdateModes;
-
-	UPROPERTY(Transient)
-	bool bLayerForceUpdateAllComponents;
-
+		
 	// Represent all the resolved paint layer, from all layers blended together (size of the landscape x material layer count)
 	class FLandscapeTexture2DArrayResource* CombinedLayersWeightmapAllMaterialLayersResource;
 	
diff --git a/Engine/Source/Runtime/Landscape/Classes/LandscapeComponent.h b/Engine/Source/Runtime/Landscape/Classes/LandscapeComponent.h
--- a/Engine/Source/Runtime/Landscape/Classes/LandscapeComponent.h
+++ b/Engine/Source/Runtime/Landscape/Classes/LandscapeComponent.h
@@ -195,6 +195,8 @@ struct FWeightmapLayerAllocationInfo
 	}
 	
 	FName GetLayerName() const;
+
+	uint32 GetHash() const;
 };
 
 struct FLandscapeComponentGrassData
@@ -872,7 +874,7 @@ public:
 	/** Updates the values of component-level properties exposed by the Landscape Actor */
 	LANDSCAPE_API void UpdatedSharedPropertiesFromActor();
 
-	LANDSCAPE_API bool IsUpdateFlagEnabledForModes(ELandscapeComponentUpdateFlag InFlag, uint32 InModeMask, bool bForce) const;
+	LANDSCAPE_API bool IsUpdateFlagEnabledForModes(ELandscapeComponentUpdateFlag InFlag, uint32 InModeMask) const;
 	LANDSCAPE_API void ClearUpdateFlagsForModes(uint32 InModeMask);
 	LANDSCAPE_API void RequestWeightmapUpdate(bool bUpdateAll = false);
 	LANDSCAPE_API void RequestHeightmapUpdate(bool bUpdateAll = false);
diff --git a/Engine/Source/Runtime/Landscape/Private/Landscape.cpp b/Engine/Source/Runtime/Landscape/Private/Landscape.cpp
--- a/Engine/Source/Runtime/Landscape/Private/Landscape.cpp
+++ b/Engine/Source/Runtime/Landscape/Private/Landscape.cpp
@@ -1078,7 +1078,6 @@ ALandscape::ALandscape(const FObjectInitializer& ObjectInitializer)
 	bInitializedWithFlagExperimentalLandscapeLayers = GetMutableDefault<UEditorExperimentalSettings>()->bLandscapeLayerSystem;
 	WasCompilingShaders = false;
 	LayerContentUpdateModes = 0;
-	bLayerForceUpdateAllComponents = false;
 	CombinedLayersWeightmapAllMaterialLayersResource = nullptr;
 	CurrentLayersWeightmapAllMaterialLayersResource = nullptr;
 	WeightmapScratchExtractLayerTextureResource = nullptr;
diff --git a/Engine/Source/Runtime/Landscape/Private/LandscapeComponent.cpp b/Engine/Source/Runtime/Landscape/Private/LandscapeComponent.cpp
--- a/Engine/Source/Runtime/Landscape/Private/LandscapeComponent.cpp
+++ b/Engine/Source/Runtime/Landscape/Private/LandscapeComponent.cpp
@@ -15,6 +15,14 @@ FName FWeightmapLayerAllocationInfo::GetLayerName() const
 	return NAME_None;
 }
 
+uint32 FWeightmapLayerAllocationInfo::GetHash() const
+{
+	uint32 Hash = PointerHash(LayerInfo);
+	Hash = HashCombine(GetTypeHash(WeightmapTextureIndex), Hash);
+	Hash = HashCombine(GetTypeHash(WeightmapTextureChannel), Hash);
+	return Hash;
+}
+
 #if WITH_EDITOR
 
 void FLandscapeEditToolRenderData::UpdateDebugColorMaterial(const ULandscapeComponent* const Component)
diff --git a/Engine/Source/Runtime/Landscape/Private/LandscapeEdit.cpp b/Engine/Source/Runtime/Landscape/Private/LandscapeEdit.cpp
--- a/Engine/Source/Runtime/Landscape/Private/LandscapeEdit.cpp
+++ b/Engine/Source/Runtime/Landscape/Private/LandscapeEdit.cpp
@@ -560,17 +560,22 @@ void ULandscapeComponent::PostEditUndo()
 {
 	if (!IsPendingKill())
 	{
-		UpdateMaterialInstances();
+		if (!GetMutableDefault<UEditorExperimentalSettings>()->bLandscapeLayerSystem)
+		{
+			UpdateMaterialInstances();
+		}
 	}
 
 	Super::PostEditUndo();
 
 	if (!IsPendingKill())
 	{
-		EditToolRenderData.UpdateDebugColorMaterial(this);
-
 		EditToolRenderData.UpdateSelectionMaterial(EditToolRenderData.SelectedType, this);
-		UpdateEditToolRenderData();
+		if (!GetMutableDefault<UEditorExperimentalSettings>()->bLandscapeLayerSystem)
+		{
+			EditToolRenderData.UpdateDebugColorMaterial(this);
+            UpdateEditToolRenderData();
+		}	
 	}
 		
 	if (GetMutableDefault<UEditorExperimentalSettings>()->bLandscapeLayerSystem)
diff --git a/Engine/Source/Runtime/Landscape/Private/LandscapeEditInterface.cpp b/Engine/Source/Runtime/Landscape/Private/LandscapeEditInterface.cpp
--- a/Engine/Source/Runtime/Landscape/Private/LandscapeEditInterface.cpp
+++ b/Engine/Source/Runtime/Landscape/Private/LandscapeEditInterface.cpp
@@ -1749,9 +1749,12 @@ void ULandscapeComponent::DeleteLayer(ULandscapeLayerInfoObject* LayerInfo, FLan
 	}
 
 	// Update the shaders for this component
-	Component->UpdateMaterialInstances();
-	Component->EditToolRenderData.UpdateDebugColorMaterial(Component);
-	Component->UpdateEditToolRenderData();
+	if (!GetMutableDefault<UEditorExperimentalSettings>()->bLandscapeLayerSystem)
+	{
+		Component->UpdateMaterialInstances();
+		Component->EditToolRenderData.UpdateDebugColorMaterial(Component);
+		Component->UpdateEditToolRenderData();
+	}
 	Component->RequestWeightmapUpdate();
 	
 	// Update dominant layer info stored in collision component
@@ -1801,7 +1804,7 @@ void FLandscapeEditDataInterface::DeleteLayer(ULandscapeLayerInfoObject* LayerIn
 
 	if (LandscapeInfo->LandscapeActor)
 	{
-		LandscapeInfo->LandscapeActor->RequestLayersContentUpdate(ELandscapeLayerUpdateMode::All, true);
+		LandscapeInfo->LandscapeActor->RequestLayersContentUpdateForceAll();
 	}
 }
 
@@ -1949,10 +1952,12 @@ void ULandscapeComponent::FillLayer(ULandscapeLayerInfoObject* LayerInfo, FLands
 	}
 
 	// Update the shaders for this component
-	Component->UpdateMaterialInstances();
-	Component->EditToolRenderData.UpdateDebugColorMaterial(Component);
-	Component->UpdateEditToolRenderData();
-
+	if (!GetMutableDefault<UEditorExperimentalSettings>()->bLandscapeLayerSystem)
+	{
+		Component->UpdateMaterialInstances();
+		Component->EditToolRenderData.UpdateDebugColorMaterial(Component);
+		Component->UpdateEditToolRenderData();
+	}
 	Component->InvalidateLightingCache();
 	Component->RequestWeightmapUpdate();
 
@@ -2035,7 +2040,7 @@ void FLandscapeEditDataInterface::FillLayer(ULandscapeLayerInfoObject* LayerInfo
 	ALandscapeProxy::InvalidateGeneratedComponentData(Components);
 	if (LandscapeInfo->LandscapeActor)
 	{
-		LandscapeInfo->LandscapeActor->RequestLayersContentUpdate(ELandscapeLayerUpdateMode::All, true);
+		LandscapeInfo->LandscapeActor->RequestLayersContentUpdateForceAll();
 	}
 }
 
@@ -2221,11 +2226,14 @@ void ULandscapeComponent::ReplaceLayer(ULandscapeLayerInfoObject* FromLayerInfo,
 		// Remove the layer
 		ComponentWeightmapLayerAllocations.RemoveAt(FromLayerIdx);
 
-		// Update the shaders for this component
-		UpdateMaterialInstances();
+		if (!GetMutableDefault<UEditorExperimentalSettings>()->bLandscapeLayerSystem)
+		{
+			// Update the shaders for this component
+			UpdateMaterialInstances();
 
-		EditToolRenderData.UpdateDebugColorMaterial(this);
-		UpdateEditToolRenderData();
+			EditToolRenderData.UpdateDebugColorMaterial(this);
+			UpdateEditToolRenderData();
+		}
 	}
 
 	RequestWeightmapUpdate();
@@ -2287,7 +2295,7 @@ void FLandscapeEditDataInterface::ReplaceLayer(ULandscapeLayerInfoObject* FromLa
 				FScopedSetLandscapeEditingLayer Scope(LandscapeInfo->LandscapeActor.Get(), CurrentLayer.Guid);
 				DoReplace();
 			});
-			LandscapeInfo->LandscapeActor->RequestLayersContentUpdate(ELandscapeLayerUpdateMode::Weightmap_All, true);
+			LandscapeInfo->LandscapeActor->RequestLayersContentUpdateForceAll(ELandscapeLayerUpdateMode::Weightmap_All);
 		}
 	}
 	else
@@ -2849,11 +2857,13 @@ void FLandscapeEditDataInterface::SetAlphaData(ULandscapeLayerInfoObject* const
 				new (ComponentWeightmapLayerAllocations) FWeightmapLayerAllocationInfo(LayerInfo);
 				Component->ReallocateWeightmaps(this);
 
-				Component->UpdateMaterialInstances();
-
-				Component->EditToolRenderData.UpdateDebugColorMaterial(Component);
-
-				Component->UpdateEditToolRenderData();
+				if (!GetMutableDefault<UEditorExperimentalSettings>()->bLandscapeLayerSystem)
+				{
+					Component->UpdateMaterialInstances();
+					Component->EditToolRenderData.UpdateDebugColorMaterial(Component);
+					Component->UpdateEditToolRenderData();
+				}
+				Component->RequestWeightmapUpdate();
 			}
 
 			// Lock data for all the weightmaps
@@ -3160,11 +3170,13 @@ void FLandscapeEditDataInterface::SetAlphaData(ULandscapeLayerInfoObject* const
 
 			if (bRemovedLayer)
 			{
-				Component->UpdateMaterialInstances();
-
-				Component->EditToolRenderData.UpdateDebugColorMaterial(Component);
-
-				Component->UpdateEditToolRenderData();
+				if (!GetMutableDefault<UEditorExperimentalSettings>()->bLandscapeLayerSystem)
+				{
+					Component->UpdateMaterialInstances();
+					Component->EditToolRenderData.UpdateDebugColorMaterial(Component);
+					Component->UpdateEditToolRenderData();
+				}
+				Component->RequestWeightmapUpdate();
 			}
 		}
 	}
@@ -3267,11 +3279,14 @@ void FLandscapeEditDataInterface::SetAlphaData(const TSet<ULandscapeLayerInfoObj
 						ComponentWeightmapLayerAllocations.Emplace(LayerInfoNeedingAllocation);
 					}
 					Component->ReallocateWeightmaps(this);
-					Component->UpdateMaterialInstances();
 					
-					Component->EditToolRenderData.UpdateDebugColorMaterial(Component);
-
-					Component->UpdateEditToolRenderData();
+					if (!GetMutableDefault<UEditorExperimentalSettings>()->bLandscapeLayerSystem)
+					{
+						Component->UpdateMaterialInstances();
+						Component->EditToolRenderData.UpdateDebugColorMaterial(Component);
+						Component->UpdateEditToolRenderData();
+					}
+					Component->RequestWeightmapUpdate();
 				}
 			}
 
@@ -3442,11 +3457,13 @@ void FLandscapeEditDataInterface::SetAlphaData(const TSet<ULandscapeLayerInfoObj
 
 			if (bRemovedLayer)
 			{
-				Component->UpdateMaterialInstances();
-
-				Component->EditToolRenderData.UpdateDebugColorMaterial(Component);
-
-				Component->UpdateEditToolRenderData();
+				if (!GetMutableDefault<UEditorExperimentalSettings>()->bLandscapeLayerSystem)
+				{
+					Component->UpdateMaterialInstances();
+					Component->EditToolRenderData.UpdateDebugColorMaterial(Component);
+					Component->UpdateEditToolRenderData();
+				}
+				Component->RequestWeightmapUpdate();
 			}
 		}
 	}
diff --git a/Engine/Source/Runtime/Landscape/Private/LandscapeEditLayers.cpp b/Engine/Source/Runtime/Landscape/Private/LandscapeEditLayers.cpp
--- a/Engine/Source/Runtime/Landscape/Private/LandscapeEditLayers.cpp
+++ b/Engine/Source/Runtime/Landscape/Private/LandscapeEditLayers.cpp
@@ -2788,7 +2788,7 @@ int32 ALandscape::RegenerateLayersHeightmaps(const TArray<ULandscapeComponent*>&
 		// Partial Component Update
 		for (ULandscapeComponent* Component : InLandscapeComponents)
 		{
-			if(Component->IsUpdateFlagEnabledForModes(ELandscapeComponentUpdateFlag::Update_Heightmap_Collision, HeightmapUpdateModes, bLayerForceUpdateAllComponents))
+			if(Component->IsUpdateFlagEnabledForModes(ELandscapeComponentUpdateFlag::Update_Heightmap_Collision, HeightmapUpdateModes))
 			{
 				Component->UpdateCachedBounds();
 				Component->UpdateComponentToWorld();
@@ -3111,9 +3111,7 @@ void ALandscape::ReallocateLayersWeightmaps(const TArray<ULandscapeLayerInfoObje
 			Usage->ClearUsage();
 		}
 	}
-
-	bool NeedMaterialInstanceRebuild = false;
-
+		
 	// Build a map of all the allocation per components
 	TMap<ULandscapeComponent*, TArray<ULandscapeLayerInfoObject*>> LayerAllocsPerComponent;
 
@@ -3742,7 +3740,7 @@ int32 ALandscape::RegenerateLayersWeightmaps(const TArray<ULandscapeComponent*>&
 	
 		for (ULandscapeComponent* Component : InLandscapeComponents)
 		{
-			if (Component->IsUpdateFlagEnabledForModes(ELandscapeComponentUpdateFlag::Update_Weightmap_Collision, WeightmapUpdateModes, bLayerForceUpdateAllComponents))
+			if (Component->IsUpdateFlagEnabledForModes(ELandscapeComponentUpdateFlag::Update_Weightmap_Collision, WeightmapUpdateModes))
 			{
 				Component->UpdateCollisionLayerData();
 			}
@@ -3755,11 +3753,18 @@ int32 ALandscape::RegenerateLayersWeightmaps(const TArray<ULandscapeComponent*>&
 uint32 ULandscapeComponent::ComputeWeightmapsHash()
 {
 	uint32 Hash = 0;
+	TArray<FWeightmapLayerAllocationInfo>& ComponentWeightmapAllocations = GetWeightmapLayerAllocations();
+	for (const FWeightmapLayerAllocationInfo& AllocationInfo : ComponentWeightmapAllocations)
+	{
+		Hash = HashCombine(AllocationInfo.GetHash(), Hash);
+	}
+
 	TArray<UTexture2D*>& ComponentWeightmapTextures = GetWeightmapTextures();
 	TArray<ULandscapeWeightmapUsage*>& ComponentWeightmapTextureUsage = GetWeightmapTexturesUsage();
 	for (int32 i = 0; i < ComponentWeightmapTextures.Num(); ++i)
 	{
 		Hash = PointerHash(ComponentWeightmapTextures[i], Hash);
+		Hash = PointerHash(ComponentWeightmapTextureUsage[i], Hash);
 		for (int32 j = 0; j < ULandscapeWeightmapUsage::NumChannels; ++j)
 		{
 			Hash = PointerHash(ComponentWeightmapTextureUsage[i]->ChannelUsage[j], Hash);
@@ -4010,37 +4015,55 @@ void ALandscape::RequestLayersInitialization(bool bInRequestContentUpdate)
 	}
 }
 
-void ALandscape::RequestLayersContentUpdate(ELandscapeLayerUpdateMode InUpdateMode, bool bInForceUpdateAllComponents)
+void ALandscape::RequestLayersContentUpdate(ELandscapeLayerUpdateMode InUpdateMode)
 {
 	LayerContentUpdateModes |= InUpdateMode;
-	bLayerForceUpdateAllComponents |= bInForceUpdateAllComponents;
 }
 
-void ALandscape::RequestLayersContentUpdateForceAll()
+void ALandscape::RequestLayersContentUpdateForceAll(ELandscapeLayerUpdateMode InModeMask)
 {
 	if (!GetMutableDefault<UEditorExperimentalSettings>()->bLandscapeLayerSystem)
 	{
 		return;
 	}
 
+	const bool bUpdateWeightmap = (InModeMask & (ELandscapeLayerUpdateMode::Weightmap_All | ELandscapeLayerUpdateMode::Weightmap_Editing)) != 0;
+	const bool bUpdateHeightmap = (InModeMask & (ELandscapeLayerUpdateMode::Heightmap_All | ELandscapeLayerUpdateMode::Heightmap_Editing)) != 0;
 	if (ULandscapeInfo* LandscapeInfo = GetLandscapeInfo())
 	{
-		LandscapeInfo->ForAllLandscapeProxies([](ALandscapeProxy* Proxy)
+		LandscapeInfo->ForAllLandscapeProxies([bUpdateHeightmap,bUpdateWeightmap](ALandscapeProxy* Proxy)
 		{
 			if (Proxy)
 			{
 				Proxy->InvalidateGeneratedComponentData();
+
+				for (ULandscapeComponent* Component : Proxy->LandscapeComponents)
+				{
+					if (bUpdateWeightmap || bUpdateHeightmap)
+					{
+						Component->Modify();
+					}
+
+					if (bUpdateHeightmap)
+					{
+						Component->RequestHeightmapUpdate();
+					}
+
+					if (bUpdateWeightmap)
+					{
+						Component->RequestWeightmapUpdate();
+					}
+				}
 			}
 		});
 	}
-	
-	const bool bLocalForceUpdateAllComponents = true;
-	RequestLayersContentUpdate(ELandscapeLayerUpdateMode::All, bLocalForceUpdateAllComponents);
+
+	RequestLayersContentUpdate(InModeMask);
 }
 
-bool ULandscapeComponent::IsUpdateFlagEnabledForModes(ELandscapeComponentUpdateFlag InFlag, uint32 InModeMask, bool bForce) const
+bool ULandscapeComponent::IsUpdateFlagEnabledForModes(ELandscapeComponentUpdateFlag InFlag, uint32 InModeMask) const
 {
-	uint32 UpdateMode = bForce? InModeMask : (LayerUpdateFlagPerMode & InModeMask);
+	uint32 UpdateMode = (LayerUpdateFlagPerMode & InModeMask);
 	
 	if (UpdateMode & ELandscapeLayerUpdateMode::Heightmap_All)
 	{
@@ -4179,26 +4202,20 @@ void ALandscape::UpdateLayersContent(bool bInWaitForStreaming)
 	ProcessedModes |= RegenerateLayersWeightmaps(AllLandscapeComponents, bInWaitForStreaming);
 	LayerContentUpdateModes &= ~ProcessedModes;
 
-	if (!ALandscape::UpdateCollisionAndClients(AllLandscapeComponents, ProcessedModes, bLayerForceUpdateAllComponents))
+	if (!ALandscape::UpdateCollisionAndClients(AllLandscapeComponents, ProcessedModes))
 	{
 		LayerContentUpdateModes |= ELandscapeLayerUpdateMode::DeferredClientUpdate;
 	}
-
-	// Clear flag if we managed to update all modes
-	if (LayerContentUpdateModes == 0)
-	{
-		bLayerForceUpdateAllComponents = false;
-	}
 }
 
-bool ALandscape::UpdateCollisionAndClients(const TArray<ULandscapeComponent*>& InLandscapeComponents, const int32 InContentUpdateModes, const bool bInLayerForceUpdateAllComponents)
+bool ALandscape::UpdateCollisionAndClients(const TArray<ULandscapeComponent*>& InLandscapeComponents, const int32 InContentUpdateModes)
 {
 	bool bAllClientsUpdated = true;
 
 	for (ULandscapeComponent* LandscapeComponent : InLandscapeComponents)
 	{
 		bool bDeferClientUpdateForComponent = false;
-		if (LandscapeComponent->IsUpdateFlagEnabledForModes(ELandscapeComponentUpdateFlag::RecreateCollision, InContentUpdateModes, bInLayerForceUpdateAllComponents))
+		if (LandscapeComponent->IsUpdateFlagEnabledForModes(ELandscapeComponentUpdateFlag::RecreateCollision, InContentUpdateModes))
 		{
 			if (ULandscapeHeightfieldCollisionComponent* CollisionComp = LandscapeComponent->CollisionComponent.Get())
 			{
@@ -4206,7 +4223,7 @@ bool ALandscape::UpdateCollisionAndClients(const TArray<ULandscapeComponent*>& I
 			}
 		}
 
-		if (LandscapeComponent->IsUpdateFlagEnabledForModes(ELandscapeComponentUpdateFlag::ClientUpdate, InContentUpdateModes, bInLayerForceUpdateAllComponents))
+		if (LandscapeComponent->IsUpdateFlagEnabledForModes(ELandscapeComponentUpdateFlag::ClientUpdate, InContentUpdateModes))
 		{
 			if (!GUndo)
 			{
@@ -5102,8 +5119,7 @@ void ALandscape::SetLayerSubstractiveBlendStatus(int32 InLayerIndex, bool InStat
 		*AllocationBlend = InStatus;
 	}
 
-	const bool bLocalForceUpdateAllComponents = true;
-	RequestLayersContentUpdate(ELandscapeLayerUpdateMode::Weightmap_All, bLocalForceUpdateAllComponents);
+	RequestLayersContentUpdateForceAll(ELandscapeLayerUpdateMode::Weightmap_All);
 }
 
 void ALandscape::AddBrushToLayer(int32 InLayerIndex, int32 InTargetType, ALandscapeBlueprintCustomBrush* InBrush)
