Added Landscape Layer contribution debug display
#rnx

diff --git a/Engine/Source/Editor/LandscapeEditor/Private/LandscapeEdMode.cpp b/Engine/Source/Editor/LandscapeEditor/Private/LandscapeEdMode.cpp
--- a/Engine/Source/Editor/LandscapeEditor/Private/LandscapeEdMode.cpp
+++ b/Engine/Source/Editor/LandscapeEditor/Private/LandscapeEdMode.cpp
@@ -306,6 +306,7 @@ FEdModeLandscape::FEdModeLandscape()
 	GSelectionColorMaterial  = LandscapeTool::CreateMaterialInstance(LoadObject<UMaterialInstanceConstant>(nullptr, TEXT("/Engine/EditorLandscapeResources/SelectBrushMaterial_Selected.SelectBrushMaterial_Selected")));
 	GSelectionRegionMaterial = LandscapeTool::CreateMaterialInstance(LoadObject<UMaterialInstanceConstant>(nullptr, TEXT("/Engine/EditorLandscapeResources/SelectBrushMaterial_SelectedRegion.SelectBrushMaterial_SelectedRegion")));
 	GMaskRegionMaterial      = LandscapeTool::CreateMaterialInstance(LoadObject<UMaterialInstanceConstant>(nullptr, TEXT("/Engine/EditorLandscapeResources/MaskBrushMaterial_MaskedRegion.MaskBrushMaterial_MaskedRegion")));
+	GColorMaskRegionMaterial = LandscapeTool::CreateMaterialInstance(LoadObject<UMaterialInstanceConstant>(nullptr, TEXT("/Engine/EditorLandscapeResources/ColorMaskBrushMaterial_MaskedRegion.ColorMaskBrushMaterial_MaskedRegion")));
 	GLandscapeBlackTexture   = LoadObject<UTexture2D>(nullptr, TEXT("/Engine/EngineResources/Black.Black"));
 	GLandscapeLayerUsageMaterial = LandscapeTool::CreateMaterialInstance(LoadObject<UMaterial>(nullptr, TEXT("/Engine/EditorLandscapeResources/LandscapeLayerUsageMaterial.LandscapeLayerUsageMaterial")));
 
@@ -379,6 +380,7 @@ FEdModeLandscape::~FEdModeLandscape()
 	GSelectionColorMaterial = NULL;
 	GSelectionRegionMaterial = NULL;
 	GMaskRegionMaterial = NULL;
+	GColorMaskRegionMaterial = NULL;
 	GLandscapeBlackTexture = NULL;
 	GLandscapeLayerUsageMaterial = NULL;
 
@@ -397,6 +399,7 @@ void FEdModeLandscape::AddReferencedObjects(FReferenceCollector& Collector)
 	Collector.AddReferencedObject(GSelectionColorMaterial);
 	Collector.AddReferencedObject(GSelectionRegionMaterial);
 	Collector.AddReferencedObject(GMaskRegionMaterial);
+	Collector.AddReferencedObject(GColorMaskRegionMaterial);
 	Collector.AddReferencedObject(GLandscapeBlackTexture);
 	Collector.AddReferencedObject(GLandscapeLayerUsageMaterial);
 }
@@ -458,12 +461,47 @@ TSharedRef<FUICommandList> FEdModeLandscape::GetUICommandList() const
 	return Toolkit->GetToolkitCommands();
 }
 
+ELandscapeToolTargetType::Type FEdModeLandscape::GetLandscapeToolTargetType() const
+{
+	if (CurrentToolMode)
+	{
+		if (CurrentToolMode->ToolModeName == "ToolMode_Sculpt")
+		{
+			return ELandscapeToolTargetType::Heightmap;
+		}
+		else if (CurrentToolMode->ToolModeName == "ToolMode_Paint")
+		{
+			return ELandscapeToolTargetType::Weightmap;
+		}
+	}
+	return ELandscapeToolTargetType::Invalid;
+}
+
+FGuid FEdModeLandscape::GetLandscapeSelectedLayer() const
+{
+	return GetCurrentLayerGuid();
+}
+
 /** FEdMode: Called when the mode is entered */
 void FEdModeLandscape::Enter()
 {
 	// Call parent implementation
 	FEdMode::Enter();
 
+	if (UWorld* World = GetWorld())
+	{
+		for (auto It = ULandscapeInfoMap::GetLandscapeInfoMap(World).Map.CreateIterator(); It; ++It)
+		{
+			if (ULandscapeInfo* LandscapeInfo = It.Value())
+			{
+				if (ALandscape* Landscape = !LandscapeInfo->IsPendingKill() ? LandscapeInfo->LandscapeActor.Get() : nullptr)
+				{
+					Landscape->RegisterLandscapeEdMode(this);
+				}
+			}
+		}
+	}
+
 	OnLevelActorDeletedDelegateHandle = GEngine->OnLevelActorDeleted().AddSP(this, &FEdModeLandscape::OnLevelActorRemoved);
 	OnLevelActorAddedDelegateHandle = GEngine->OnLevelActorAdded().AddSP(this, &FEdModeLandscape::OnLevelActorAdded);
 	OnLandscapeLayerSystemFlagChangedDelegateHandle = GetMutableDefault<UEditorExperimentalSettings>()->OnSettingChanged().AddLambda([this](FName PropertyName)
@@ -660,6 +698,20 @@ void FEdModeLandscape::Enter()
 /** FEdMode: Called when the mode is exited */
 void FEdModeLandscape::Exit()
 {
+	if (UWorld* World = GetWorld())
+	{
+		for (auto It = ULandscapeInfoMap::GetLandscapeInfoMap(World).Map.CreateIterator(); It; ++It)
+		{
+			if (ULandscapeInfo* LandscapeInfo = It.Value())
+			{
+				if (ALandscape* Landscape = !LandscapeInfo->IsPendingKill() ? LandscapeInfo->LandscapeActor.Get() : nullptr)
+				{
+					Landscape->UnregisterLandscapeEdMode();
+				}
+			}
+		}
+	}
+
 	// Unregister VR mode from event handlers
 	UViewportWorldInteraction* ViewportWorldInteraction = Cast<UViewportWorldInteraction>(GEditor->GetEditorWorldExtensionsManager()->GetEditorWorldExtensions(GetWorld())->FindExtension(UViewportWorldInteraction::StaticClass()));
 	if (ViewportWorldInteraction != nullptr)
@@ -3498,7 +3550,7 @@ void FEdModeLandscape::ImportData(const FLandscapeTargetListInfo& TargetInfo, co
 
 			{
 				ALandscape* Landscape = GetLandscape();
-				FScopedSetLandscapeEditingLayer Scope(Landscape, GetCurrentLayerGuid(), [&] { check(Landscape); Landscape->RequestLayersContentUpdate(ELandscapeLayerUpdateMode::Heightmap_All); });
+				FScopedSetLandscapeEditingLayer Scope(Landscape, GetCurrentLayerGuid(), [&] { check(Landscape); Landscape->RequestLayersContentUpdate(ELandscapeLayerUpdateMode::Update_Heightmap_All); });
 
 				TArray<uint16> Data;
 				if (ImportResolution != LandscapeResolution)
@@ -3614,7 +3666,7 @@ void FEdModeLandscape::ImportData(const FLandscapeTargetListInfo& TargetInfo, co
 
 			{
 				ALandscape* Landscape = GetLandscape();
-				FScopedSetLandscapeEditingLayer Scope(Landscape, GetCurrentLayerGuid(), [&] { check(Landscape); Landscape->RequestLayersContentUpdate(ELandscapeLayerUpdateMode::Weightmap_All); });
+				FScopedSetLandscapeEditingLayer Scope(Landscape, GetCurrentLayerGuid(), [&] { check(Landscape); Landscape->RequestLayersContentUpdate(ELandscapeLayerUpdateMode::Update_Weightmap_All); });
 
 				TArray<uint8> Data;
 				if (ImportResolution != LandscapeResolution)
@@ -4208,8 +4260,8 @@ void FEdModeLandscape::SetCurrentLayer(int32 InLayerIndex)
 {
 	UISettings->Modify();
 	UISettings->CurrentLayerIndex = InLayerIndex;
-
 	RefreshDetailPanel();
+	RequestLayersContentUpdateForceAll(ELandscapeLayerUpdateMode::Update_Client_Editing);
 }
 
 int32 FEdModeLandscape::GetCurrentLayerIndex() const
diff --git a/Engine/Source/Editor/LandscapeEditor/Private/LandscapeEdMode.h b/Engine/Source/Editor/LandscapeEditor/Private/LandscapeEdMode.h
--- a/Engine/Source/Editor/LandscapeEditor/Private/LandscapeEdMode.h
+++ b/Engine/Source/Editor/LandscapeEditor/Private/LandscapeEdMode.h
@@ -12,6 +12,7 @@
 #include "LandscapeInfo.h"
 #include "LandscapeLayerInfoObject.h"
 #include "LandscapeGizmoActiveActor.h"
+#include "LandscapeEdit.h"
 
 class ALandscape;
 class FCanvas;
@@ -250,7 +251,7 @@ enum class ELandscapeEditingState : uint8
 /**
  * Landscape editor mode
  */
-class FEdModeLandscape : public FEdMode
+class FEdModeLandscape : public FEdMode, public ILandscapeEdModeInterface
 {
 public:
 
@@ -326,6 +327,10 @@ public:
 	/** Destructor */
 	virtual ~FEdModeLandscape();
 
+	/** ILandscapeEdModeInterface */
+	virtual ELandscapeToolTargetType::Type GetLandscapeToolTargetType() const override;
+	virtual FGuid GetLandscapeSelectedLayer() const override;
+
 	/** FGCObject interface */
 	virtual void AddReferencedObjects(FReferenceCollector& Collector) override;
 
@@ -526,7 +531,7 @@ public:
 	
 	bool NeedToFillEmptyMaterialLayers() const;
 	void RequestLayersContentUpdate(ELandscapeLayerUpdateMode InUpdateMode);
-	void RequestLayersContentUpdateForceAll(ELandscapeLayerUpdateMode InUpdateMode = ELandscapeLayerUpdateMode::All);
+	void RequestLayersContentUpdateForceAll(ELandscapeLayerUpdateMode InUpdateMode = ELandscapeLayerUpdateMode::Update_All);
 
 	void OnLevelActorAdded(AActor* InActor);
 	void OnLevelActorRemoved(AActor* InActor);
diff --git a/Engine/Source/Editor/LandscapeEditor/Private/LandscapeEdModeComponentTools.cpp b/Engine/Source/Editor/LandscapeEditor/Private/LandscapeEdModeComponentTools.cpp
--- a/Engine/Source/Editor/LandscapeEditor/Private/LandscapeEdModeComponentTools.cpp
+++ b/Engine/Source/Editor/LandscapeEditor/Private/LandscapeEdModeComponentTools.cpp
@@ -1725,11 +1725,11 @@ public:
 		return ELandscapeToolTargetTypeMask::FromType(ToolTarget::TargetType);
 	}
 
-	virtual ELandscapeLayerUpdateMode GetBeginToolContentUpdateFlag() const override { return ELandscapeLayerUpdateMode::All_Editing; }
+	virtual ELandscapeLayerUpdateMode GetBeginToolContentUpdateFlag() const override { return ELandscapeLayerUpdateMode::Update_All_Editing; }
 
 	virtual ELandscapeLayerUpdateMode GetTickToolContentUpdateFlag() const override { return GetBeginToolContentUpdateFlag(); }
 	
-	virtual ELandscapeLayerUpdateMode GetEndToolContentUpdateFlag() const override { return ELandscapeLayerUpdateMode::All; }
+	virtual ELandscapeLayerUpdateMode GetEndToolContentUpdateFlag() const override { return ELandscapeLayerUpdateMode::Update_All; }
 		
 	virtual bool BeginTool(FEditorViewportClient* ViewportClient, const FLandscapeToolTarget& InTarget, const FVector& InHitLocation) override
 	{
diff --git a/Engine/Source/Editor/LandscapeEditor/Private/LandscapeEdModeRampTool.cpp b/Engine/Source/Editor/LandscapeEditor/Private/LandscapeEdModeRampTool.cpp
--- a/Engine/Source/Editor/LandscapeEditor/Private/LandscapeEdModeRampTool.cpp
+++ b/Engine/Source/Editor/LandscapeEditor/Private/LandscapeEdModeRampTool.cpp
@@ -486,7 +486,7 @@ public:
 
 		FScopedTransaction Transaction(LOCTEXT("Ramp_Apply", "Landscape Editing: Add ramp"));
 		ALandscape* Landscape = EdMode->GetLandscape();
-		FScopedSetLandscapeEditingLayer Scope(Landscape, EdMode->GetCurrentLayerGuid(), [&] { if (Landscape) { Landscape->RequestLayersContentUpdate(ELandscapeLayerUpdateMode::Heightmap_All); } });
+		FScopedSetLandscapeEditingLayer Scope(Landscape, EdMode->GetCurrentLayerGuid(), [&] { if (Landscape) { Landscape->RequestLayersContentUpdate(ELandscapeLayerUpdateMode::Update_Heightmap_All); } });
 
 		const ULandscapeInfo* LandscapeInfo = EdMode->CurrentToolTarget.LandscapeInfo.Get();
 		const ALandscapeProxy* LandscapeProxy = LandscapeInfo->GetLandscapeProxy();
diff --git a/Engine/Source/Editor/LandscapeEditor/Private/LandscapeEdModeTools.h b/Engine/Source/Editor/LandscapeEditor/Private/LandscapeEdModeTools.h
--- a/Engine/Source/Editor/LandscapeEditor/Private/LandscapeEdModeTools.h
+++ b/Engine/Source/Editor/LandscapeEditor/Private/LandscapeEdModeTools.h
@@ -1134,7 +1134,7 @@ public:
 	virtual ELandscapeLayerUpdateMode GetBeginToolContentUpdateFlag() const
 	{
 		bool bUpdateHeightmap = this->EdMode->CurrentToolTarget.TargetType == ELandscapeToolTargetType::Type::Heightmap; 
-		return bUpdateHeightmap ? ELandscapeLayerUpdateMode::Heightmap_Editing : ELandscapeLayerUpdateMode::Weightmap_Editing;
+		return bUpdateHeightmap ? ELandscapeLayerUpdateMode::Update_Heightmap_Editing : ELandscapeLayerUpdateMode::Update_Weightmap_Editing;
 	}
 
 	virtual ELandscapeLayerUpdateMode GetTickToolContentUpdateFlag() const
@@ -1145,7 +1145,7 @@ public:
 	virtual ELandscapeLayerUpdateMode GetEndToolContentUpdateFlag() const
 	{
 		bool bUpdateHeightmap = this->EdMode->CurrentToolTarget.TargetType == ELandscapeToolTargetType::Type::Heightmap;
-		return bUpdateHeightmap ? ELandscapeLayerUpdateMode::Heightmap_All : ELandscapeLayerUpdateMode::Weightmap_All;
+		return bUpdateHeightmap ? ELandscapeLayerUpdateMode::Update_Heightmap_All : ELandscapeLayerUpdateMode::Update_Weightmap_All;
 	}
 
 	virtual bool BeginTool(FEditorViewportClient* ViewportClient, const FLandscapeToolTarget& InTarget, const FVector& InHitLocation) override
diff --git a/Engine/Source/Editor/LandscapeEditor/Private/LandscapeEditorCommands.cpp b/Engine/Source/Editor/LandscapeEditor/Private/LandscapeEditorCommands.cpp
--- a/Engine/Source/Editor/LandscapeEditor/Private/LandscapeEditorCommands.cpp
+++ b/Engine/Source/Editor/LandscapeEditor/Private/LandscapeEditorCommands.cpp
@@ -2,6 +2,7 @@
 
 
 #include "LandscapeEditorCommands.h"
+#include "Settings/EditorExperimentalSettings.h"
 
 #define LOCTEXT_NAMESPACE "LandscapeEditorCommands"
 
@@ -105,6 +106,10 @@ void FLandscapeEditorCommands::RegisterCommands()
 	UI_COMMAND(ViewModeLayerDensity, "Layer Density", "", EUserInterfaceActionType::RadioButton, FInputChord());
 	UI_COMMAND(ViewModeLayerDebug, "Layer Debug", "", EUserInterfaceActionType::RadioButton, FInputChord());
 	UI_COMMAND(ViewModeWireframeOnTop, "Wireframe on Top", "", EUserInterfaceActionType::RadioButton, FInputChord());
+	if (GetMutableDefault<UEditorExperimentalSettings>()->bLandscapeLayerSystem)
+	{
+		UI_COMMAND(ViewModeLayerContribution, "Layer Contribution", "", EUserInterfaceActionType::RadioButton, FInputChord());
+	}
 
 	UI_COMMAND(IncreaseBrushSize, "Increase Brush Size", "Press this key to increase brush size by a fixed increment.", EUserInterfaceActionType::RadioButton, FInputChord(EKeys::RightBracket));
 	UI_COMMAND(DecreaseBrushSize, "Decrease Brush Size", "Press this key to decrease brush size by a fixed increment.", EUserInterfaceActionType::RadioButton, FInputChord(EKeys::LeftBracket));
diff --git a/Engine/Source/Editor/LandscapeEditor/Private/LandscapeEditorCommands.h b/Engine/Source/Editor/LandscapeEditor/Private/LandscapeEditorCommands.h
--- a/Engine/Source/Editor/LandscapeEditor/Private/LandscapeEditorCommands.h
+++ b/Engine/Source/Editor/LandscapeEditor/Private/LandscapeEditorCommands.h
@@ -75,6 +75,7 @@ public:
 	TSharedPtr<FUICommandInfo> ViewModeLayerDebug;
 	TSharedPtr<FUICommandInfo> ViewModeWireframeOnTop;
 	TSharedPtr<FUICommandInfo> ViewModeLayerUsage;
+	TSharedPtr<FUICommandInfo> ViewModeLayerContribution;
 
 	// Adjusting brushes
 	TSharedPtr<FUICommandInfo> IncreaseBrushSize;
diff --git a/Engine/Source/Editor/LandscapeEditor/Private/LandscapeEditorModule.cpp b/Engine/Source/Editor/LandscapeEditor/Private/LandscapeEditorModule.cpp
--- a/Engine/Source/Editor/LandscapeEditor/Private/LandscapeEditorModule.cpp
+++ b/Engine/Source/Editor/LandscapeEditor/Private/LandscapeEditorModule.cpp
@@ -17,6 +17,7 @@
 #include "Classes/ActorFactoryLandscape.h"
 #include "LandscapeFileFormatPng.h"
 #include "LandscapeFileFormatRaw.h"
+#include "Settings/EditorExperimentalSettings.h"
 
 #include "Framework/MultiBox/MultiBoxBuilder.h"
 #include "PropertyEditorModule.h"
@@ -89,6 +90,7 @@ public:
 		GlobalUICommandList->MapAction(LandscapeActions.ViewModeLayerDebug, FExecuteAction::CreateStatic(&ChangeLandscapeViewMode, ELandscapeViewMode::DebugLayer), FCanExecuteAction(), FIsActionChecked::CreateStatic(&IsLandscapeViewModeSelected, ELandscapeViewMode::DebugLayer));
 		GlobalUICommandList->MapAction(LandscapeActions.ViewModeWireframeOnTop, FExecuteAction::CreateStatic(&ChangeLandscapeViewMode, ELandscapeViewMode::WireframeOnTop), FCanExecuteAction(), FIsActionChecked::CreateStatic(&IsLandscapeViewModeSelected, ELandscapeViewMode::WireframeOnTop));
 		GlobalUICommandList->MapAction(LandscapeActions.ViewModeLayerUsage, FExecuteAction::CreateStatic(&ChangeLandscapeViewMode, ELandscapeViewMode::LayerUsage), FCanExecuteAction(), FIsActionChecked::CreateStatic(&IsLandscapeViewModeSelected, ELandscapeViewMode::LayerUsage));
+		GlobalUICommandList->MapAction(LandscapeActions.ViewModeLayerContribution, FExecuteAction::CreateStatic(&ChangeLandscapeViewMode, ELandscapeViewMode::LayerContribution), FCanExecuteAction(), FIsActionChecked::CreateStatic(&IsLandscapeViewModeSelected, ELandscapeViewMode::LayerContribution));
 
 		ViewportMenuExtender = MakeShareable(new FExtender);
 		ViewportMenuExtender->AddMenuExtension("LevelViewportLandscape", EExtensionHook::First, GlobalUICommandList, FMenuExtensionDelegate::CreateStatic(&ConstructLandscapeViewportMenu));
@@ -160,6 +162,10 @@ public:
 					{
 						InMenuBuilder.AddMenuEntry(LandscapeActions.ViewModeLayerUsage, NAME_None, LOCTEXT("LandscapeViewModeLayerUsage", "Layer Usage"));
 						InMenuBuilder.AddMenuEntry(LandscapeActions.ViewModeLayerDebug, NAME_None, LOCTEXT("LandscapeViewModeLayerDebug", "Layer Debug"));
+						if (GetMutableDefault<UEditorExperimentalSettings>()->bLandscapeLayerSystem)
+						{
+							InMenuBuilder.AddMenuEntry(LandscapeActions.ViewModeLayerContribution, NAME_None, LOCTEXT("LandscapeViewModeLayerContribution", "Layer Contribution"));
+						}
 					}
 					InMenuBuilder.AddMenuEntry(LandscapeActions.ViewModeWireframeOnTop, NAME_None, LOCTEXT("LandscapeViewModeWireframeOnTop", "Wireframe on Top"));
 				}
diff --git a/Engine/Source/Editor/LandscapeEditor/Public/LandscapeToolInterface.h b/Engine/Source/Editor/LandscapeEditor/Public/LandscapeToolInterface.h
--- a/Engine/Source/Editor/LandscapeEditor/Public/LandscapeToolInterface.h
+++ b/Engine/Source/Editor/LandscapeEditor/Public/LandscapeToolInterface.h
@@ -7,6 +7,7 @@
 #include "UObject/GCObject.h"
 #include "UnrealWidget.h"
 #include "EdMode.h"
+#include "LandscapeEdit.h"
 
 class FEditorViewportClient;
 class FPrimitiveDrawInterface;
@@ -136,18 +137,6 @@ struct FLandscapeBrushSet
 	}
 };
 
-namespace ELandscapeToolTargetType
-{
-	enum Type : int8
-	{
-		Heightmap  = 0,
-		Weightmap  = 1,
-		Visibility = 2,
-
-		Invalid    = -1, // only valid for LandscapeEdMode->CurrentToolTarget.TargetType
-	};
-}
-
 namespace ELandscapeToolTargetTypeMask
 {
 	enum Type : uint8
diff --git a/Engine/Source/Runtime/Landscape/Classes/Landscape.h b/Engine/Source/Runtime/Landscape/Classes/Landscape.h
--- a/Engine/Source/Runtime/Landscape/Classes/Landscape.h
+++ b/Engine/Source/Runtime/Landscape/Classes/Landscape.h
@@ -10,6 +10,12 @@
 #include "Landscape.generated.h"
 
 class ULandscapeComponent;
+class ILandscapeEdModeInterface;
+
+namespace ELandscapeToolTargetType
+{
+	enum Type : int8;
+};
 
 UENUM()
 enum ELandscapeSetupErrors
@@ -228,8 +234,10 @@ public:
 
 	// Layers stuff
 #if WITH_EDITOR
+	LANDSCAPE_API void RegisterLandscapeEdMode(ILandscapeEdModeInterface* InLandscapeEdMode) { LandscapeEdMode = InLandscapeEdMode; }
+	LANDSCAPE_API void UnregisterLandscapeEdMode() { LandscapeEdMode = nullptr; }
 	LANDSCAPE_API void RequestLayersInitialization(bool bInRequestContentUpdate = true);
-	LANDSCAPE_API void RequestLayersContentUpdateForceAll(ELandscapeLayerUpdateMode InModeMask = ELandscapeLayerUpdateMode::All);
+	LANDSCAPE_API void RequestLayersContentUpdateForceAll(ELandscapeLayerUpdateMode InModeMask = ELandscapeLayerUpdateMode::Update_All);
 	LANDSCAPE_API void RequestLayersContentUpdate(ELandscapeLayerUpdateMode InModeMask);
 	LANDSCAPE_API bool ReorderLayer(int32 InStartingLayerIndex, int32 InDestinationLayerIndex);
 	LANDSCAPE_API FLandscapeLayer* DuplicateLayer(const FLandscapeLayer& InOtherLayer);
@@ -284,9 +292,10 @@ private:
 	void ReleaseLayersRenderingResource();
 	void UpdateLayersContent(bool bInWaitForStreaming = false);
 	void MonitorShaderCompilation();
+	void MonitorLandscapeEdModeChanges();
 	int32 RegenerateLayersHeightmaps(const TArray<ULandscapeComponent*>& InLandscapeComponents, bool bInWaitForStreaming);
 	int32 RegenerateLayersWeightmaps(const TArray<ULandscapeComponent*>& InLandscapeComponents, bool bInWaitForStreaming);
-	static bool UpdateCollisionAndClients(const TArray<ULandscapeComponent*>& InLandscapeComponents, const int32 InContentUpdateModes);
+	bool UpdateCollisionAndClients(const TArray<ULandscapeComponent*>& InLandscapeComponents, const int32 InContentUpdateModes);
 	void ResolveLayersHeightmapTexture();
 	void ResolveLayersWeightmapTexture();
 	void ResolveLayersTexture(class FLandscapeLayersTexture2DCPUReadBackResource* InCPUReadBackTexture, UTexture2D* InOutputTexture);
@@ -361,6 +370,21 @@ public:
 	bool bInitializedWithFlagExperimentalLandscapeLayers;
 
 private:
+	/** Provides information from LandscapeEdMode */
+	ILandscapeEdModeInterface* LandscapeEdMode;
+
+	/** Information provided by LandscapeEdMode */
+	struct FLandscapeEdModeInfo
+	{
+		FLandscapeEdModeInfo();
+
+		int32 ViewMode;
+		FGuid SelectedLayer;
+		ELandscapeToolTargetType::Type ToolTarget;
+	};
+
+	FLandscapeEdModeInfo LandscapeEdModeInfo;
+
 	UPROPERTY(Transient)
 	bool bLandscapeLayersAreInitialized;
 	
diff --git a/Engine/Source/Runtime/Landscape/Classes/LandscapeComponent.h b/Engine/Source/Runtime/Landscape/Classes/LandscapeComponent.h
--- a/Engine/Source/Runtime/Landscape/Classes/LandscapeComponent.h
+++ b/Engine/Source/Runtime/Landscape/Classes/LandscapeComponent.h
@@ -61,7 +61,8 @@ public:
 		DebugChannelR(INDEX_NONE),
 		DebugChannelG(INDEX_NONE),
 		DebugChannelB(INDEX_NONE),
-		DataTexture(NULL)
+		DataTexture(NULL),
+		LayerContributionTexture(NULL)
 	{}
 
 	// Material used to render the tool.
@@ -88,6 +89,9 @@ public:
 	UPROPERTY(NonTransactional)
 	UTexture2D* DataTexture; // Data texture other than height/weight
 
+	UPROPERTY(NonTransactional)
+	UTexture2D* LayerContributionTexture; // Data texture used to represent layer contribution
+
 #if WITH_EDITOR
 	void UpdateDebugColorMaterial(const ULandscapeComponent* const Component);
 	void UpdateSelectionMaterial(int32 InSelectedType, const ULandscapeComponent* const Component);
@@ -296,25 +300,29 @@ struct FLandscapeLayerComponentData
 enum ELandscapeComponentUpdateFlag : uint32
 {
 	// Will call UpdateCollisionHeightData, UpdateCacheBounds, UpdateComponentToWorld on Component
-	Update_Heightmap_Collision = 1 << 0,
+	Component_Update_Heightmap_Collision = 1 << 0,
 	// Will call UdateCollisionLayerData on Component
-	Update_Weightmap_Collision = 1 << 1,
+	Component_Update_Weightmap_Collision = 1 << 1,
 	// Will call RecreateCollision on Component
-	RecreateCollision = 1 << 2,
+	Component_Update_Recreate_Collision = 1 << 2,
 	// Will update Component clients: Navigation data, Foliage, Grass, etc.
-	ClientUpdate = 1 << 3
+	Component_Update_Client = 1 << 3,
+	// Will update Component clients while editing
+	Component_Update_Client_Editing = 1 << 4
 };
 
 enum ELandscapeLayerUpdateMode : uint32
 { 
-	Heightmap_All = 1 << 0,
-	Heightmap_Editing = 1 << 1,
-	Weightmap_All = 1 << 2,
-	Weightmap_Editing = 1 << 3,
-	All = Weightmap_All | Heightmap_All,
-	All_Editing = Weightmap_Editing | Heightmap_Editing,
+	Update_Heightmap_All = 1 << 0,
+	Update_Heightmap_Editing = 1 << 1,
+	Update_Weightmap_All = 1 << 2,
+	Update_Weightmap_Editing = 1 << 3,
+	Update_All = Update_Weightmap_All | Update_Heightmap_All,
+	Update_All_Editing = Update_Weightmap_Editing | Update_Heightmap_Editing,
 	// In cases where we couldn't update the clients right away this flag will be set in RegenerateLayersContent
-	DeferredClientUpdate = 1 << 4,
+	Update_Client_Deferred = 1 << 4,
+	// Update landscape component clients while editing
+	Update_Client_Editing = 1 << 5
 };
 
 #endif
@@ -878,6 +886,7 @@ public:
 	LANDSCAPE_API void ClearUpdateFlagsForModes(uint32 InModeMask);
 	LANDSCAPE_API void RequestWeightmapUpdate(bool bUpdateAll = false);
 	LANDSCAPE_API void RequestHeightmapUpdate(bool bUpdateAll = false);
+	LANDSCAPE_API void RequestEditingClientUpdate();
 	LANDSCAPE_API void RequestDeferredClientUpdate();
 	LANDSCAPE_API uint32 ComputeWeightmapsHash();
 #endif
diff --git a/Engine/Source/Runtime/Landscape/Private/Landscape.cpp b/Engine/Source/Runtime/Landscape/Private/Landscape.cpp
--- a/Engine/Source/Runtime/Landscape/Private/Landscape.cpp
+++ b/Engine/Source/Runtime/Landscape/Private/Landscape.cpp
@@ -52,6 +52,7 @@ Landscape.cpp: Terrain rendering
 #include "LandscapeWeightmapUsage.h"
 
 #if WITH_EDITOR
+#include "LandscapeEdit.h"
 #include "MaterialUtilities.h"
 #include "Settings/EditorExperimentalSettings.h"
 #include "Editor.h"
@@ -1070,6 +1071,14 @@ ALandscapeProxy::ALandscapeProxy(const FObjectInitializer& ObjectInitializer)
 	FrameOffsetForTickInterval = FrameOffsetForTickIntervalInc++;
 }
 
+#if WITH_EDITORONLY_DATA
+ALandscape::FLandscapeEdModeInfo::FLandscapeEdModeInfo()
+	: ViewMode(ELandscapeViewMode::Invalid)
+	, ToolTarget(ELandscapeToolTargetType::Invalid)
+{
+}
+#endif
+
 ALandscape::ALandscape(const FObjectInitializer& ObjectInitializer)
 	: Super(ObjectInitializer)
 {
@@ -1083,6 +1092,7 @@ ALandscape::ALandscape(const FObjectInitializer& ObjectInitializer)
 	WeightmapScratchExtractLayerTextureResource = nullptr;
 	WeightmapScratchPackLayerTextureResource = nullptr;
 	bLandscapeLayersAreInitialized = false;
+	LandscapeEdMode = nullptr;
 #endif // WITH_EDITORONLY_DATA
 }
 
diff --git a/Engine/Source/Runtime/Landscape/Private/LandscapeBPCustomBrush.cpp b/Engine/Source/Runtime/Landscape/Private/LandscapeBPCustomBrush.cpp
--- a/Engine/Source/Runtime/Landscape/Private/LandscapeBPCustomBrush.cpp
+++ b/Engine/Source/Runtime/Landscape/Private/LandscapeBPCustomBrush.cpp
@@ -78,7 +78,7 @@ void ALandscapeBlueprintCustomBrush::PostEditMove(bool bFinished)
 
 	if (OwningLandscape != nullptr)
 	{
-		OwningLandscape->RequestLayersContentUpdate(bFinished ? ELandscapeLayerUpdateMode::All : ELandscapeLayerUpdateMode::All_Editing);
+		OwningLandscape->RequestLayersContentUpdate(bFinished ? ELandscapeLayerUpdateMode::Update_All : ELandscapeLayerUpdateMode::Update_All_Editing);
 	}
 }
 
diff --git a/Engine/Source/Runtime/Landscape/Private/LandscapeEdit.cpp b/Engine/Source/Runtime/Landscape/Private/LandscapeEdit.cpp
--- a/Engine/Source/Runtime/Landscape/Private/LandscapeEdit.cpp
+++ b/Engine/Source/Runtime/Landscape/Private/LandscapeEdit.cpp
@@ -3609,7 +3609,7 @@ void ALandscape::PostEditUndo()
 {
 	Super::PostEditUndo();
 
-	RequestLayersContentUpdate(ELandscapeLayerUpdateMode::All);
+	RequestLayersContentUpdate(ELandscapeLayerUpdateMode::Update_All);
 }
 
 bool ALandscape::ShouldImport(FString* ActorPropString, bool IsMovingLevel)
@@ -4056,7 +4056,7 @@ void ALandscapeProxy::PostEditChangeProperty(FPropertyChangedEvent& PropertyChan
 		{
 			if(ALandscape* LandscapeActor = GetLandscapeActor())
 			{
-				LandscapeActor->RequestLayersContentUpdate(ELandscapeLayerUpdateMode::All);
+				LandscapeActor->RequestLayersContentUpdate(ELandscapeLayerUpdateMode::Update_All);
 			}
 		}
 		
diff --git a/Engine/Source/Runtime/Landscape/Private/LandscapeEditInterface.cpp b/Engine/Source/Runtime/Landscape/Private/LandscapeEditInterface.cpp
--- a/Engine/Source/Runtime/Landscape/Private/LandscapeEditInterface.cpp
+++ b/Engine/Source/Runtime/Landscape/Private/LandscapeEditInterface.cpp
@@ -2295,7 +2295,7 @@ void FLandscapeEditDataInterface::ReplaceLayer(ULandscapeLayerInfoObject* FromLa
 				FScopedSetLandscapeEditingLayer Scope(LandscapeInfo->LandscapeActor.Get(), CurrentLayer.Guid);
 				DoReplace();
 			});
-			LandscapeInfo->LandscapeActor->RequestLayersContentUpdateForceAll(ELandscapeLayerUpdateMode::Weightmap_All);
+			LandscapeInfo->LandscapeActor->RequestLayersContentUpdateForceAll(ELandscapeLayerUpdateMode::Update_Weightmap_All);
 		}
 	}
 	else
@@ -4453,69 +4453,144 @@ bool FLandscapeTextureDataInterface::EqualTextureValue(UTexture2D* Src, FColor V
 	return EqualTextureValueTempl<FColor>(Src, Value);
 }
 
-
 template<typename TStoreData>
 void FLandscapeEditDataInterface::GetSelectDataTempl(const int32 X1, const int32 Y1, const int32 X2, const int32 Y2, TStoreData& StoreData)
+{
+	auto ReturnComponentTexture = [&](ULandscapeComponent* Component) -> UTexture2D*
+	{
+		check(Component);
+		return Component->EditToolRenderData.DataTexture;
+	};
+	GetEditToolTextureData(X1, Y1, X2, Y2, StoreData, ReturnComponentTexture);
+}
+
+void FLandscapeEditDataInterface::GetSelectData(const int32 X1, const int32 Y1, const int32 X2, const int32 Y2, TMap<FIntPoint, uint8>& SparseData)
+{
+	TSparseStoreData<uint8> SparseStoreData(SparseData);
+	GetSelectDataTempl(X1, Y1, X2, Y2, SparseStoreData);
+}
+
+void FLandscapeEditDataInterface::GetSelectData(const int32 X1, const int32 Y1, const int32 X2, const int32 Y2, uint8* Data, int32 Stride)
+{
+	if( Stride==0 )
+	{
+		Stride = (1+X2-X1);
+	}
+	TArrayStoreData<uint8> ArrayStoreData(X1, Y1, Data, Stride);
+	GetSelectDataTempl(X1, Y1, X2, Y2, ArrayStoreData);
+}
+
+void FLandscapeEditDataInterface::SetSelectData(int32 X1, int32 Y1, int32 X2, int32 Y2, const uint8* Data, int32 Stride)
+{
+	auto ReturnComponentTexture = [&](ULandscapeComponent* Component) -> UTexture2D*&
+	{
+		check(Component);
+		return Component->EditToolRenderData.DataTexture;
+	};
+	SetEditToolTextureData(X1, Y1, X2, Y2, Data, Stride, ReturnComponentTexture);
+}
+
+template<typename TStoreData>
+void FLandscapeEditDataInterface::GetLayerContributionDataTempl(const int32 X1, const int32 Y1, const int32 X2, const int32 Y2, TStoreData& StoreData)
+{
+	auto ReturnComponentTexture = [&](const ULandscapeComponent* Component) -> UTexture2D*
+	{
+		check(Component);
+		return Component->EditToolRenderData.LayerContributionTexture;
+	};
+	GetEditToolTextureData(X1, Y1, X2, Y2, StoreData, ReturnComponentTexture);
+}
+
+void FLandscapeEditDataInterface::GetLayerContributionData(const int32 X1, const int32 Y1, const int32 X2, const int32 Y2, TMap<FIntPoint, uint8>& SparseData)
+{
+	TSparseStoreData<uint8> SparseStoreData(SparseData);
+	GetLayerContributionDataTempl(X1, Y1, X2, Y2, SparseStoreData);
+}
+
+void FLandscapeEditDataInterface::GetLayerContributionData(const int32 X1, const int32 Y1, const int32 X2, const int32 Y2, uint8* Data, int32 Stride)
+{
+	if (Stride == 0)
+	{
+		Stride = (1 + X2 - X1);
+	}
+	TArrayStoreData<uint8> ArrayStoreData(X1, Y1, Data, Stride);
+	GetLayerContributionDataTempl(X1, Y1, X2, Y2, ArrayStoreData);
+}
+
+
+void FLandscapeEditDataInterface::SetLayerContributionData(int32 X1, int32 Y1, int32 X2, int32 Y2, const uint8* Data, int32 Stride)
+{
+	auto ReturnComponentTexture = [](ULandscapeComponent* Component) -> UTexture2D*&
+	{
+		check(Component);
+		return Component->EditToolRenderData.LayerContributionTexture;
+	};
+	SetEditToolTextureData(X1, Y1, X2, Y2, Data, Stride, ReturnComponentTexture);
+}
+
+template<typename TStoreData>
+void FLandscapeEditDataInterface::GetEditToolTextureData(const int32 X1, const int32 Y1, const int32 X2, const int32 Y2, TStoreData& StoreData, TFunctionRef<UTexture2D*(ULandscapeComponent*)> GetComponentTexture)
 {
 	int32 ComponentIndexX1, ComponentIndexY1, ComponentIndexX2, ComponentIndexY2;
 	ALandscape::CalcComponentIndicesNoOverlap(X1, Y1, X2, Y2, ComponentSizeQuads, ComponentIndexX1, ComponentIndexY1, ComponentIndexX2, ComponentIndexY2);
 
-	for( int32 ComponentIndexY=ComponentIndexY1;ComponentIndexY<=ComponentIndexY2;ComponentIndexY++ )
+	for (int32 ComponentIndexY = ComponentIndexY1; ComponentIndexY <= ComponentIndexY2; ComponentIndexY++)
 	{
-		for( int32 ComponentIndexX=ComponentIndexX1;ComponentIndexX<=ComponentIndexX2;ComponentIndexX++ )
-		{		
-			ULandscapeComponent* Component = LandscapeInfo->XYtoComponentMap.FindRef(FIntPoint(ComponentIndexX,ComponentIndexY));
+		for (int32 ComponentIndexX = ComponentIndexX1; ComponentIndexX <= ComponentIndexX2; ComponentIndexX++)
+		{
+			ULandscapeComponent* Component = LandscapeInfo->XYtoComponentMap.FindRef(FIntPoint(ComponentIndexX, ComponentIndexY));
 
 			FLandscapeTextureDataInfo* TexDataInfo = NULL;
 			uint8* SelectTextureData = NULL;
-			if( Component && Component->EditToolRenderData.DataTexture )
+			UTexture2D* EditToolTexture = Component ? GetComponentTexture(Component) : nullptr;
+			if (EditToolTexture)
 			{
-				TexDataInfo = GetTextureDataInfo(Component->EditToolRenderData.DataTexture);
+				TexDataInfo = GetTextureDataInfo(EditToolTexture);
 				SelectTextureData = (uint8*)TexDataInfo->GetMipData(0);
 			}
 
 			// Find coordinates of box that lies inside component
-			int32 ComponentX1 = FMath::Clamp<int32>(X1-ComponentIndexX*ComponentSizeQuads, 0, ComponentSizeQuads);
-			int32 ComponentY1 = FMath::Clamp<int32>(Y1-ComponentIndexY*ComponentSizeQuads, 0, ComponentSizeQuads);
-			int32 ComponentX2 = FMath::Clamp<int32>(X2-ComponentIndexX*ComponentSizeQuads, 0, ComponentSizeQuads);
-			int32 ComponentY2 = FMath::Clamp<int32>(Y2-ComponentIndexY*ComponentSizeQuads, 0, ComponentSizeQuads);
+			int32 ComponentX1 = FMath::Clamp<int32>(X1 - ComponentIndexX * ComponentSizeQuads, 0, ComponentSizeQuads);
+			int32 ComponentY1 = FMath::Clamp<int32>(Y1 - ComponentIndexY * ComponentSizeQuads, 0, ComponentSizeQuads);
+			int32 ComponentX2 = FMath::Clamp<int32>(X2 - ComponentIndexX * ComponentSizeQuads, 0, ComponentSizeQuads);
+			int32 ComponentY2 = FMath::Clamp<int32>(Y2 - ComponentIndexY * ComponentSizeQuads, 0, ComponentSizeQuads);
 
 			// Find subsection range for this box
-			int32 SubIndexX1 = FMath::Clamp<int32>((ComponentX1-1) / SubsectionSizeQuads,0,ComponentNumSubsections-1);	// -1 because we need to pick up vertices shared between subsections
-			int32 SubIndexY1 = FMath::Clamp<int32>((ComponentY1-1) / SubsectionSizeQuads,0,ComponentNumSubsections-1);
-			int32 SubIndexX2 = FMath::Clamp<int32>(ComponentX2 / SubsectionSizeQuads,0,ComponentNumSubsections-1);
-			int32 SubIndexY2 = FMath::Clamp<int32>(ComponentY2 / SubsectionSizeQuads,0,ComponentNumSubsections-1);
+			int32 SubIndexX1 = FMath::Clamp<int32>((ComponentX1 - 1) / SubsectionSizeQuads, 0, ComponentNumSubsections - 1);	// -1 because we need to pick up vertices shared between subsections
+			int32 SubIndexY1 = FMath::Clamp<int32>((ComponentY1 - 1) / SubsectionSizeQuads, 0, ComponentNumSubsections - 1);
+			int32 SubIndexX2 = FMath::Clamp<int32>(ComponentX2 / SubsectionSizeQuads, 0, ComponentNumSubsections - 1);
+			int32 SubIndexY2 = FMath::Clamp<int32>(ComponentY2 / SubsectionSizeQuads, 0, ComponentNumSubsections - 1);
 
-			for( int32 SubIndexY=SubIndexY1;SubIndexY<=SubIndexY2;SubIndexY++ )
+			for (int32 SubIndexY = SubIndexY1; SubIndexY <= SubIndexY2; SubIndexY++)
 			{
-				for( int32 SubIndexX=SubIndexX1;SubIndexX<=SubIndexX2;SubIndexX++ )
+				for (int32 SubIndexX = SubIndexX1; SubIndexX <= SubIndexX2; SubIndexX++)
 				{
 					// Find coordinates of box that lies inside subsection
-					int32 SubX1 = FMath::Clamp<int32>(ComponentX1-SubsectionSizeQuads*SubIndexX, 0, SubsectionSizeQuads);
-					int32 SubY1 = FMath::Clamp<int32>(ComponentY1-SubsectionSizeQuads*SubIndexY, 0, SubsectionSizeQuads);
-					int32 SubX2 = FMath::Clamp<int32>(ComponentX2-SubsectionSizeQuads*SubIndexX, 0, SubsectionSizeQuads);
-					int32 SubY2 = FMath::Clamp<int32>(ComponentY2-SubsectionSizeQuads*SubIndexY, 0, SubsectionSizeQuads);
+					int32 SubX1 = FMath::Clamp<int32>(ComponentX1 - SubsectionSizeQuads * SubIndexX, 0, SubsectionSizeQuads);
+					int32 SubY1 = FMath::Clamp<int32>(ComponentY1 - SubsectionSizeQuads * SubIndexY, 0, SubsectionSizeQuads);
+					int32 SubX2 = FMath::Clamp<int32>(ComponentX2 - SubsectionSizeQuads * SubIndexX, 0, SubsectionSizeQuads);
+					int32 SubY2 = FMath::Clamp<int32>(ComponentY2 - SubsectionSizeQuads * SubIndexY, 0, SubsectionSizeQuads);
 
 					// Update texture data for the box that lies inside subsection
-					for( int32 SubY=SubY1;SubY<=SubY2;SubY++ )
+					for (int32 SubY = SubY1; SubY <= SubY2; SubY++)
 					{
-						for( int32 SubX=SubX1;SubX<=SubX2;SubX++ )
+						for (int32 SubX = SubX1; SubX <= SubX2; SubX++)
 						{
-							int32 LandscapeX = SubIndexX*SubsectionSizeQuads + ComponentIndexX*ComponentSizeQuads + SubX;
-							int32 LandscapeY = SubIndexY*SubsectionSizeQuads + ComponentIndexY*ComponentSizeQuads + SubY;
+							int32 LandscapeX = SubIndexX * SubsectionSizeQuads + ComponentIndexX * ComponentSizeQuads + SubX;
+							int32 LandscapeY = SubIndexY * SubsectionSizeQuads + ComponentIndexY * ComponentSizeQuads + SubY;
 
 							// Find the input data corresponding to this vertex
-							if( Component && SelectTextureData )
+							if (Component && SelectTextureData)
 							{
 								// Find the texture data corresponding to this vertex
-								int32 SizeU = Component->EditToolRenderData.DataTexture->Source.GetSizeX();
-								int32 SizeV = Component->EditToolRenderData.DataTexture->Source.GetSizeY();
+								int32 SizeU = EditToolTexture->Source.GetSizeX();
+								int32 SizeV = EditToolTexture->Source.GetSizeY();
 								int32 WeightmapOffsetX = Component->WeightmapScaleBias.Z * (float)SizeU;
 								int32 WeightmapOffsetY = Component->WeightmapScaleBias.W * (float)SizeV;
 
-								int32 TexX = WeightmapOffsetX + (SubsectionSizeQuads+1) * SubIndexX + SubX;
-								int32 TexY = WeightmapOffsetY + (SubsectionSizeQuads+1) * SubIndexY + SubY;
-								uint8& TexData = SelectTextureData[ TexX + TexY * SizeU ];
+								int32 TexX = WeightmapOffsetX + (SubsectionSizeQuads + 1) * SubIndexX + SubX;
+								int32 TexY = WeightmapOffsetY + (SubsectionSizeQuads + 1) * SubIndexY + SubY;
+								uint8& TexData = SelectTextureData[TexX + TexY * SizeU];
 
 								StoreData.Store(LandscapeX, LandscapeY, TexData);
 							}
@@ -4532,23 +4607,7 @@ void FLandscapeEditDataInterface::GetSelectDataTempl(const int32 X1, const int32
 	}
 }
 
-void FLandscapeEditDataInterface::GetSelectData(const int32 X1, const int32 Y1, const int32 X2, const int32 Y2, TMap<FIntPoint, uint8>& SparseData)
-{
-	TSparseStoreData<uint8> SparseStoreData(SparseData);
-	GetSelectDataTempl(X1, Y1, X2, Y2, SparseStoreData);
-}
-
-void FLandscapeEditDataInterface::GetSelectData(const int32 X1, const int32 Y1, const int32 X2, const int32 Y2, uint8* Data, int32 Stride)
-{
-	if( Stride==0 )
-	{
-		Stride = (1+X2-X1);
-	}
-	TArrayStoreData<uint8> ArrayStoreData(X1, Y1, Data, Stride);
-	GetSelectDataTempl(X1, Y1, X2, Y2, ArrayStoreData);
-}
-
-void FLandscapeEditDataInterface::SetSelectData(int32 X1, int32 Y1, int32 X2, int32 Y2, const uint8* Data, int32 Stride)
+void FLandscapeEditDataInterface::SetEditToolTextureData(int32 X1, int32 Y1, int32 X2, int32 Y2, const uint8* Data, int32 Stride, TFunctionRef<UTexture2D*&(ULandscapeComponent*)> GetComponentTexture)
 {
 	if( Stride==0 )
 	{
@@ -4566,49 +4625,44 @@ void FLandscapeEditDataInterface::SetSelectData(int32 X1, int32 Y1, int32 X2, in
 		{	
 			FIntPoint ComponentKey(ComponentIndexX,ComponentIndexY);
 			ULandscapeComponent* Component = LandscapeInfo->XYtoComponentMap.FindRef(ComponentKey);
-
-			UTexture2D* DataTexture = NULL;
-			// if NULL, it was painted away
-			if( Component==NULL)
+			if (Component == NULL)
 			{
 				continue;
 			}
-			else if (Component->EditToolRenderData.DataTexture == NULL)
+
+			UTexture2D*& EditToolTexture = GetComponentTexture(Component);
+			// if NULL, it was painted away
+			if (EditToolTexture == NULL)
 			{
 				//FlushRenderingCommands();
 				// Construct Texture...
 				int32 WeightmapSize = (Component->SubsectionSizeQuads+1) * Component->NumSubsections;
-				DataTexture = Component->GetLandscapeProxy()->CreateLandscapeTexture(WeightmapSize, WeightmapSize, TEXTUREGROUP_Terrain_Weightmap, TSF_G8);
+				EditToolTexture = Component->GetLandscapeProxy()->CreateLandscapeTexture(WeightmapSize, WeightmapSize, TEXTUREGROUP_Terrain_Weightmap, TSF_G8);
 				// Alloc dummy mips
-				ULandscapeComponent::CreateEmptyTextureMips(DataTexture, true);
-				DataTexture->PostEditChange();
+				ULandscapeComponent::CreateEmptyTextureMips(EditToolTexture, true);
+				EditToolTexture->PostEditChange();
 
 				//FlushRenderingCommands();
-				ZeroTexture(DataTexture);
-				FLandscapeTextureDataInfo* TexDataInfo = GetTextureDataInfo(DataTexture);
-				int32 NumMips = DataTexture->Source.GetNumMips();
+				ZeroTexture(EditToolTexture);
+				FLandscapeTextureDataInfo* TexDataInfo = GetTextureDataInfo(EditToolTexture);
+				int32 NumMips = EditToolTexture->Source.GetNumMips();
 				TArray<uint8*> TextureMipData;
 				TextureMipData.AddUninitialized(NumMips);
 				for( int32 MipIdx=0;MipIdx<NumMips;MipIdx++ )
 				{
 					TextureMipData[MipIdx] = (uint8*)TexDataInfo->GetMipData(MipIdx);
 				}
-				ULandscapeComponent::UpdateDataMips(ComponentNumSubsections, SubsectionSizeQuads, DataTexture, TextureMipData, 0, 0, MAX_int32, MAX_int32, TexDataInfo);
-				
-				Component->EditToolRenderData.DataTexture = DataTexture;
+				ULandscapeComponent::UpdateDataMips(ComponentNumSubsections, SubsectionSizeQuads, EditToolTexture, TextureMipData, 0, 0, MAX_int32, MAX_int32, TexDataInfo);
+
 				Component->UpdateEditToolRenderData();
 			}
-			else
-			{
-				DataTexture = Component->EditToolRenderData.DataTexture;
-			}
 
-			FLandscapeTextureDataInfo* TexDataInfo = GetTextureDataInfo(DataTexture);
+			FLandscapeTextureDataInfo* TexDataInfo = GetTextureDataInfo(EditToolTexture);
 			uint8* SelectTextureData = (uint8*)TexDataInfo->GetMipData(0);
 
 			// Find the texture data corresponding to this vertex
-			int32 SizeU = DataTexture->Source.GetSizeX();
-			int32 SizeV = DataTexture->Source.GetSizeY();
+			int32 SizeU = EditToolTexture->Source.GetSizeX();
+			int32 SizeV = EditToolTexture->Source.GetSizeY();
 			int32 WeightmapOffsetX = Component->WeightmapScaleBias.Z * (float)SizeU;
 			int32 WeightmapOffsetY = Component->WeightmapScaleBias.W * (float)SizeV;
 
@@ -4665,14 +4719,14 @@ void FLandscapeEditDataInterface::SetSelectData(int32 X1, int32 Y1, int32 X2, in
 				}
 			}
 			// Update mipmaps
-			int32 NumMips = DataTexture->Source.GetNumMips();
+			int32 NumMips = EditToolTexture->Source.GetNumMips();
 			TArray<uint8*> TextureMipData;
 			TextureMipData.AddUninitialized(NumMips);
 			for( int32 MipIdx=0;MipIdx<NumMips;MipIdx++ )
 			{
 				TextureMipData[MipIdx] = (uint8*)TexDataInfo->GetMipData(MipIdx);
 			}
-			ULandscapeComponent::UpdateDataMips(ComponentNumSubsections, SubsectionSizeQuads, DataTexture, TextureMipData, ComponentX1, ComponentY1, ComponentX2, ComponentY2, TexDataInfo);
+			ULandscapeComponent::UpdateDataMips(ComponentNumSubsections, SubsectionSizeQuads, EditToolTexture, TextureMipData, ComponentX1, ComponentY1, ComponentX2, ComponentY2, TexDataInfo);
 		}
 	}
 }
diff --git a/Engine/Source/Runtime/Landscape/Private/LandscapeEditLayers.cpp b/Engine/Source/Runtime/Landscape/Private/LandscapeEditLayers.cpp
--- a/Engine/Source/Runtime/Landscape/Private/LandscapeEditLayers.cpp
+++ b/Engine/Source/Runtime/Landscape/Private/LandscapeEditLayers.cpp
@@ -31,6 +31,7 @@ LandscapeEditLayers.cpp: Landscape editing layers mode
 #include "ShaderCompiler.h"
 #include "Algo/Count.h"
 #include "LandscapeSettings.h"
+#include "LandscapeRender.h"
 #endif
 
 #define LOCTEXT_NAMESPACE "Landscape"
@@ -2559,7 +2560,7 @@ int32 ALandscape::RegenerateLayersHeightmaps(const TArray<ULandscapeComponent*>&
 	SCOPE_CYCLE_COUNTER(STAT_LandscapeLayersRegenerateHeightmaps);
 	ULandscapeInfo* Info = GetLandscapeInfo();
 
-	const int32 AllHeightmapUpdateModes = (ELandscapeLayerUpdateMode::Heightmap_All |  ELandscapeLayerUpdateMode::Heightmap_Editing);
+	const int32 AllHeightmapUpdateModes = (ELandscapeLayerUpdateMode::Update_Heightmap_All |  ELandscapeLayerUpdateMode::Update_Heightmap_Editing);
 	const int32 HeightmapUpdateModes = LayerContentUpdateModes & AllHeightmapUpdateModes;
 	const bool bForceRender = CVarOutputLayersDebugDrawCallName.GetValueOnAnyThread() == 1;
 
@@ -2789,7 +2790,7 @@ int32 ALandscape::RegenerateLayersHeightmaps(const TArray<ULandscapeComponent*>&
 		// Partial Component Update
 		for (ULandscapeComponent* Component : InLandscapeComponents)
 		{
-			if(Component->IsUpdateFlagEnabledForModes(ELandscapeComponentUpdateFlag::Update_Heightmap_Collision, HeightmapUpdateModes))
+			if(Component->IsUpdateFlagEnabledForModes(ELandscapeComponentUpdateFlag::Component_Update_Heightmap_Collision, HeightmapUpdateModes))
 			{
 				Component->UpdateCachedBounds();
 				Component->UpdateComponentToWorld();
@@ -3373,7 +3374,7 @@ bool ALandscape::PrepareLayersWeightmapTextureResources(bool bInWaitForStreaming
 int32 ALandscape::RegenerateLayersWeightmaps(const TArray<ULandscapeComponent*>& InLandscapeComponents, bool bInWaitForStreaming)
 {
 	SCOPE_CYCLE_COUNTER(STAT_LandscapeLayersRegenerateWeightmaps);
-	const int32 AllWeightmapUpdateModes = (ELandscapeLayerUpdateMode::Weightmap_All | ELandscapeLayerUpdateMode::Weightmap_Editing);
+	const int32 AllWeightmapUpdateModes = (ELandscapeLayerUpdateMode::Update_Weightmap_All | ELandscapeLayerUpdateMode::Update_Weightmap_Editing);
 	const int32 WeightmapUpdateModes = LayerContentUpdateModes & AllWeightmapUpdateModes;
 	const bool bForceRender = CVarOutputLayersDebugDrawCallName.GetValueOnAnyThread() == 1;
 	
@@ -3741,7 +3742,7 @@ int32 ALandscape::RegenerateLayersWeightmaps(const TArray<ULandscapeComponent*>&
 	
 		for (ULandscapeComponent* Component : InLandscapeComponents)
 		{
-			if (Component->IsUpdateFlagEnabledForModes(ELandscapeComponentUpdateFlag::Update_Weightmap_Collision, WeightmapUpdateModes))
+			if (Component->IsUpdateFlagEnabledForModes(ELandscapeComponentUpdateFlag::Component_Update_Weightmap_Collision, WeightmapUpdateModes))
 			{
 				Component->UpdateCollisionLayerData();
 			}
@@ -4028,11 +4029,12 @@ void ALandscape::RequestLayersContentUpdateForceAll(ELandscapeLayerUpdateMode In
 		return;
 	}
 
-	const bool bUpdateWeightmap = (InModeMask & (ELandscapeLayerUpdateMode::Weightmap_All | ELandscapeLayerUpdateMode::Weightmap_Editing)) != 0;
-	const bool bUpdateHeightmap = (InModeMask & (ELandscapeLayerUpdateMode::Heightmap_All | ELandscapeLayerUpdateMode::Heightmap_Editing)) != 0;
+	const bool bUpdateWeightmap = (InModeMask & (ELandscapeLayerUpdateMode::Update_Weightmap_All | ELandscapeLayerUpdateMode::Update_Weightmap_Editing)) != 0;
+	const bool bUpdateHeightmap = (InModeMask & (ELandscapeLayerUpdateMode::Update_Heightmap_All | ELandscapeLayerUpdateMode::Update_Heightmap_Editing)) != 0;
+	const bool bUpdateClientUdpateEditing = (InModeMask & ELandscapeLayerUpdateMode::Update_Client_Editing) != 0;
 	if (ULandscapeInfo* LandscapeInfo = GetLandscapeInfo())
 	{
-		LandscapeInfo->ForAllLandscapeProxies([bUpdateHeightmap,bUpdateWeightmap](ALandscapeProxy* Proxy)
+		LandscapeInfo->ForAllLandscapeProxies([bUpdateHeightmap, bUpdateWeightmap, bUpdateClientUdpateEditing](ALandscapeProxy* Proxy)
 		{
 			if (Proxy)
 			{
@@ -4054,6 +4056,11 @@ void ALandscape::RequestLayersContentUpdateForceAll(ELandscapeLayerUpdateMode In
 					{
 						Component->RequestWeightmapUpdate();
 					}
+
+					if (bUpdateClientUdpateEditing)
+					{
+						Component->RequestEditingClientUpdate();
+					}
 				}
 			}
 		});
@@ -4066,45 +4073,54 @@ bool ULandscapeComponent::IsUpdateFlagEnabledForModes(ELandscapeComponentUpdateF
 {
 	uint32 UpdateMode = (LayerUpdateFlagPerMode & InModeMask);
 	
-	if (UpdateMode & ELandscapeLayerUpdateMode::Heightmap_All)
+	if (UpdateMode & ELandscapeLayerUpdateMode::Update_Heightmap_All)
 	{
-		const uint32 HeightmapAllFlags = ELandscapeComponentUpdateFlag::Update_Heightmap_Collision | ELandscapeComponentUpdateFlag::RecreateCollision | ELandscapeComponentUpdateFlag::ClientUpdate;
+		const uint32 HeightmapAllFlags = ELandscapeComponentUpdateFlag::Component_Update_Heightmap_Collision | ELandscapeComponentUpdateFlag::Component_Update_Recreate_Collision | ELandscapeComponentUpdateFlag::Component_Update_Client;
 		if (HeightmapAllFlags & InFlag)
 		{
 			return true;
 		}
 	}
 		
-	if (UpdateMode & ELandscapeLayerUpdateMode::Heightmap_Editing)
+	if (UpdateMode & ELandscapeLayerUpdateMode::Update_Heightmap_Editing)
 	{
-		const uint32 HeightmapEditingFlags = ELandscapeComponentUpdateFlag::Update_Heightmap_Collision;
+		const uint32 HeightmapEditingFlags = ELandscapeComponentUpdateFlag::Component_Update_Heightmap_Collision | ELandscapeComponentUpdateFlag::Component_Update_Client_Editing;
 		if (HeightmapEditingFlags & InFlag)
 		{
 			return true;
 		}
 	}
 		
-	if (UpdateMode & ELandscapeLayerUpdateMode::Weightmap_All)
+	if (UpdateMode & ELandscapeLayerUpdateMode::Update_Weightmap_All)
 	{
-		const uint32 WeightmapAllFlags = ELandscapeComponentUpdateFlag::Update_Weightmap_Collision | ELandscapeComponentUpdateFlag::RecreateCollision | ELandscapeComponentUpdateFlag::ClientUpdate;
+		const uint32 WeightmapAllFlags = ELandscapeComponentUpdateFlag::Component_Update_Weightmap_Collision | ELandscapeComponentUpdateFlag::Component_Update_Recreate_Collision | ELandscapeComponentUpdateFlag::Component_Update_Client;
 		if (WeightmapAllFlags & InFlag)
 		{
 			return true;
 		}
 	}
 
-	if (UpdateMode & ELandscapeLayerUpdateMode::Weightmap_Editing)
+	if (UpdateMode & ELandscapeLayerUpdateMode::Update_Weightmap_Editing)
+	{
+		const uint32 WeightmapEditingFlags = ELandscapeComponentUpdateFlag::Component_Update_Weightmap_Collision | ELandscapeComponentUpdateFlag::Component_Update_Client_Editing;
+		if (WeightmapEditingFlags & InFlag)
+		{
+			return true;
+		}
+	}
+
+	if (UpdateMode & ELandscapeLayerUpdateMode::Update_Client_Editing)
 	{
-		const uint32 WeightmapEditingFlags = ELandscapeComponentUpdateFlag::Update_Weightmap_Collision;
+		const uint32 WeightmapEditingFlags = ELandscapeComponentUpdateFlag::Component_Update_Client_Editing;
 		if (WeightmapEditingFlags & InFlag)
 		{
 			return true;
 		}
 	}
 
-	if (UpdateMode & ELandscapeLayerUpdateMode::DeferredClientUpdate)
+	if (UpdateMode & ELandscapeLayerUpdateMode::Update_Client_Deferred)
 	{
-		const uint32 DeferredClientUpdateFlags = ELandscapeComponentUpdateFlag::ClientUpdate;
+		const uint32 DeferredClientUpdateFlags = ELandscapeComponentUpdateFlag::Component_Update_Client;
 		if (DeferredClientUpdateFlags & InFlag)
 		{
 			return true;
@@ -4121,37 +4137,75 @@ void ULandscapeComponent::ClearUpdateFlagsForModes(uint32 InModeMask)
 
 void ULandscapeComponent::RequestDeferredClientUpdate()
 {
-	LayerUpdateFlagPerMode |= ELandscapeLayerUpdateMode::DeferredClientUpdate;
+	LayerUpdateFlagPerMode |= ELandscapeLayerUpdateMode::Update_Client_Deferred;
+}
+
+void ULandscapeComponent::RequestEditingClientUpdate()
+{
+	LayerUpdateFlagPerMode |= ELandscapeLayerUpdateMode::Update_Client_Editing;
+	if (ALandscape* LandscapeActor = GetLandscapeActor())
+	{
+		LandscapeActor->RequestLayersContentUpdate(ELandscapeLayerUpdateMode::Update_Client_Editing);
+	}
 }
 
 void ULandscapeComponent::RequestHeightmapUpdate(bool bUpdateAll)
 {
-	LayerUpdateFlagPerMode |= ELandscapeLayerUpdateMode::Heightmap_Editing;
-	LayerUpdateFlagPerMode |= ELandscapeLayerUpdateMode::Heightmap_All;
+	LayerUpdateFlagPerMode |= ELandscapeLayerUpdateMode::Update_Heightmap_Editing;
+	LayerUpdateFlagPerMode |= ELandscapeLayerUpdateMode::Update_Heightmap_All;
 	if (ALandscape* LandscapeActor = GetLandscapeActor())
 	{
-		LandscapeActor->RequestLayersContentUpdate(ELandscapeLayerUpdateMode::Heightmap_Editing);
+		LandscapeActor->RequestLayersContentUpdate(ELandscapeLayerUpdateMode::Update_Heightmap_Editing);
 		if (bUpdateAll)
 		{
-			LandscapeActor->RequestLayersContentUpdate(Heightmap_All);
+			LandscapeActor->RequestLayersContentUpdate(ELandscapeLayerUpdateMode::Update_Heightmap_All);
 		}
 	}
 }
 
 void ULandscapeComponent::RequestWeightmapUpdate(bool bUpdateAll)
 {
-	LayerUpdateFlagPerMode |= ELandscapeLayerUpdateMode::Weightmap_Editing;
-	LayerUpdateFlagPerMode |= ELandscapeLayerUpdateMode::Weightmap_All;
+	LayerUpdateFlagPerMode |= ELandscapeLayerUpdateMode::Update_Weightmap_Editing;
+	LayerUpdateFlagPerMode |= ELandscapeLayerUpdateMode::Update_Weightmap_All;
 	if (ALandscape* LandscapeActor = GetLandscapeActor())
 	{
-		LandscapeActor->RequestLayersContentUpdate(ELandscapeLayerUpdateMode::Weightmap_Editing);
+		LandscapeActor->RequestLayersContentUpdate(ELandscapeLayerUpdateMode::Update_Weightmap_Editing);
 		if (bUpdateAll)
 		{
-			LandscapeActor->RequestLayersContentUpdate(Weightmap_All);
+			LandscapeActor->RequestLayersContentUpdate(ELandscapeLayerUpdateMode::Update_Weightmap_All);
 		}
 	}
 }
 
+void ALandscape::MonitorLandscapeEdModeChanges()
+{
+	bool bRequiredEditingClientFullUpdate = false;
+	if (LandscapeEdModeInfo.ViewMode != GLandscapeViewMode)
+	{
+		LandscapeEdModeInfo.ViewMode = GLandscapeViewMode;
+		bRequiredEditingClientFullUpdate = true;
+	}
+
+	ELandscapeToolTargetType::Type NewValue = LandscapeEdMode ? LandscapeEdMode->GetLandscapeToolTargetType() : ELandscapeToolTargetType::Invalid;
+	if (LandscapeEdModeInfo.ToolTarget != NewValue)
+	{
+		LandscapeEdModeInfo.ToolTarget = NewValue;
+		bRequiredEditingClientFullUpdate = true;
+	}
+
+	FGuid NewSelectedLayer = LandscapeEdMode ? LandscapeEdMode->GetLandscapeSelectedLayer() : FGuid();
+	if (LandscapeEdModeInfo.SelectedLayer != NewSelectedLayer)
+	{
+		LandscapeEdModeInfo.SelectedLayer = NewSelectedLayer;
+		bRequiredEditingClientFullUpdate = true;
+	}
+
+	if (bRequiredEditingClientFullUpdate && (LandscapeEdModeInfo.ViewMode == ELandscapeViewMode::LayerContribution))
+	{
+		RequestLayersContentUpdateForceAll(ELandscapeLayerUpdateMode::Update_Client_Editing);
+	}
+}
+
 void ALandscape::MonitorShaderCompilation()
 {
 	// If doing editing while shader are compiling or at load of a map, it's possible we will need another update pass after shader are completed to see the correct result
@@ -4186,6 +4240,7 @@ void ALandscape::UpdateLayersContent(bool bInWaitForStreaming)
 	}
 	
 	MonitorShaderCompilation();
+	MonitorLandscapeEdModeChanges();
 
 	if (GetLandscapeInfo() == nullptr || LayerContentUpdateModes == 0)
 	{
@@ -4198,14 +4253,16 @@ void ALandscape::UpdateLayersContent(bool bInWaitForStreaming)
 		AllLandscapeComponents.Append(Proxy->LandscapeComponents);
 	});
 
-	int32 ProcessedModes = (LayerContentUpdateModes & ELandscapeLayerUpdateMode::DeferredClientUpdate);
+	int32 ProcessedModes = 0;
 	ProcessedModes |= RegenerateLayersHeightmaps(AllLandscapeComponents, bInWaitForStreaming);
 	ProcessedModes |= RegenerateLayersWeightmaps(AllLandscapeComponents, bInWaitForStreaming);
+	ProcessedModes |= (LayerContentUpdateModes & ELandscapeLayerUpdateMode::Update_Client_Deferred);
+	ProcessedModes |= (LayerContentUpdateModes & ELandscapeLayerUpdateMode::Update_Client_Editing);
 	LayerContentUpdateModes &= ~ProcessedModes;
 
 	if (!ALandscape::UpdateCollisionAndClients(AllLandscapeComponents, ProcessedModes))
 	{
-		LayerContentUpdateModes |= ELandscapeLayerUpdateMode::DeferredClientUpdate;
+		LayerContentUpdateModes |= ELandscapeLayerUpdateMode::Update_Client_Deferred;
 	}
 }
 
@@ -4213,10 +4270,17 @@ bool ALandscape::UpdateCollisionAndClients(const TArray<ULandscapeComponent*>& I
 {
 	bool bAllClientsUpdated = true;
 
+	const uint16 DefaultHeightValue = LandscapeDataAccess::GetTexHeight(0.f);
+	const uint8 DefaultWeightValue = 0;
+	const uint8 LayerContributingValue = 64;
+	TArray<uint16> HeightData;
+	TArray<uint8> WeightData;
+	TArray<uint8> LayerContributionMaskData;
+
 	for (ULandscapeComponent* LandscapeComponent : InLandscapeComponents)
 	{
 		bool bDeferClientUpdateForComponent = false;
-		if (LandscapeComponent->IsUpdateFlagEnabledForModes(ELandscapeComponentUpdateFlag::RecreateCollision, InContentUpdateModes))
+		if (LandscapeComponent->IsUpdateFlagEnabledForModes(ELandscapeComponentUpdateFlag::Component_Update_Recreate_Collision, InContentUpdateModes))
 		{
 			if (ULandscapeHeightfieldCollisionComponent* CollisionComp = LandscapeComponent->CollisionComponent.Get())
 			{
@@ -4224,7 +4288,7 @@ bool ALandscape::UpdateCollisionAndClients(const TArray<ULandscapeComponent*>& I
 			}
 		}
 
-		if (LandscapeComponent->IsUpdateFlagEnabledForModes(ELandscapeComponentUpdateFlag::ClientUpdate, InContentUpdateModes))
+		if (LandscapeComponent->IsUpdateFlagEnabledForModes(ELandscapeComponentUpdateFlag::Component_Update_Client, InContentUpdateModes))
 		{
 			if (!GUndo)
 			{
@@ -4241,6 +4305,68 @@ bool ALandscape::UpdateCollisionAndClients(const TArray<ULandscapeComponent*>& I
 			}
 		}
 
+		if (LandscapeComponent->IsUpdateFlagEnabledForModes(ELandscapeComponentUpdateFlag::Component_Update_Client_Editing, InContentUpdateModes))
+		{
+			if (LandscapeEdModeInfo.ViewMode == ELandscapeViewMode::LayerContribution && LandscapeEdModeInfo.SelectedLayer.IsValid())
+			{
+				ULandscapeInfo* Info = LandscapeComponent->GetLandscapeInfo();
+				check(Info);
+
+				FLandscapeEditDataInterface LandscapeEdit(Info);
+				FScopedSetLandscapeEditingLayer Scope(this, LandscapeEdModeInfo.SelectedLayer);
+
+				check(ComponentSizeQuads == LandscapeComponent->ComponentSizeQuads);
+				const int32 Stride = (1 + ComponentSizeQuads);
+				const int32 ArraySize = Stride * Stride;
+				if (LayerContributionMaskData.Num() != ArraySize)
+				{
+					LayerContributionMaskData.AddZeroed(ArraySize);
+				}
+				uint8* LayerContributionMaskDataPtr = LayerContributionMaskData.GetData();
+				const int32 X1 = LandscapeComponent->GetSectionBase().X;
+				const int32 X2 = X1 + ComponentSizeQuads;
+				const int32 Y1 = LandscapeComponent->GetSectionBase().Y;
+				const int32 Y2 = Y1 + ComponentSizeQuads;
+
+				if (LandscapeEdModeInfo.ToolTarget == ELandscapeToolTargetType::Heightmap)
+				{
+					if (HeightData.Num() != ArraySize)
+					{
+						HeightData.AddZeroed(ArraySize);
+					}
+					LandscapeEdit.GetHeightDataFast(X1, Y1, X2, Y2, HeightData.GetData(), Stride);
+					for (int i = 0; i < ArraySize; ++i)
+					{
+						LayerContributionMaskData[i] = HeightData[i] != DefaultHeightValue ? LayerContributingValue : 0;
+					}
+				}
+				else if (LandscapeEdModeInfo.ToolTarget == ELandscapeToolTargetType::Weightmap)
+				{
+					if (WeightData.Num() != ArraySize)
+					{
+						WeightData.AddZeroed(ArraySize);
+					}
+					FMemory::Memzero(LayerContributionMaskDataPtr, ArraySize);
+					for (const FLandscapeInfoLayerSettings& LayerSettings : Info->Layers)
+					{
+						if (LayerSettings.LayerInfoObj == ALandscapeProxy::VisibilityLayer)
+						{
+							continue;
+						}
+						LandscapeEdit.GetWeightDataFast(LayerSettings.LayerInfoObj, X1, Y1, X2, Y2, WeightData.GetData(), Stride);
+						for (int i = 0; i < ArraySize; ++i)
+						{
+							if (WeightData[i] != DefaultWeightValue)
+							{
+								LayerContributionMaskData[i] = LayerContributingValue;
+							}
+						}
+					}
+				}
+				LandscapeEdit.SetLayerContributionData(X1, Y1, X2, Y2, LayerContributionMaskDataPtr, 0);
+			}
+		}
+
 		LandscapeComponent->ClearUpdateFlagsForModes(InContentUpdateModes);
 		if (bDeferClientUpdateForComponent)
 		{
@@ -5125,7 +5251,7 @@ void ALandscape::SetLayerSubstractiveBlendStatus(int32 InLayerIndex, bool InStat
 		*AllocationBlend = InStatus;
 	}
 
-	RequestLayersContentUpdateForceAll(ELandscapeLayerUpdateMode::Weightmap_All);
+	RequestLayersContentUpdateForceAll(ELandscapeLayerUpdateMode::Update_Weightmap_All);
 }
 
 void ALandscape::AddBrushToLayer(int32 InLayerIndex, int32 InTargetType, ALandscapeBlueprintCustomBrush* InBrush)
diff --git a/Engine/Source/Runtime/Landscape/Private/LandscapeRender.cpp b/Engine/Source/Runtime/Landscape/Private/LandscapeRender.cpp
--- a/Engine/Source/Runtime/Landscape/Private/LandscapeRender.cpp
+++ b/Engine/Source/Runtime/Landscape/Private/LandscapeRender.cpp
@@ -113,7 +113,7 @@ LANDSCAPE_API int32 GLandscapeViewMode = ELandscapeViewMode::Normal;
 FAutoConsoleVariableRef CVarLandscapeDebugViewMode(
 	TEXT("Landscape.DebugViewMode"),
 	GLandscapeViewMode,
-	TEXT("Change the view mode of the landscape rendering. Valid Input: 0 = Normal, 2 = DebugLayer, 3 = LayerDensity, 4 = LayerUsage, 5 = LOD Distribution, 6 = WireframeOnTop"),
+	TEXT("Change the view mode of the landscape rendering. Valid Input: 0 = Normal, 2 = DebugLayer, 3 = LayerDensity, 4 = LayerUsage, 5 = LOD Distribution, 6 = WireframeOnTop, 7 = LayerContribution"),
 	ECVF_Cheat
 );
 #endif
@@ -546,6 +546,7 @@ UMaterialInterface* GLayerDebugColorMaterial = nullptr;
 UMaterialInterface* GSelectionColorMaterial = nullptr;
 UMaterialInterface* GSelectionRegionMaterial = nullptr;
 UMaterialInterface* GMaskRegionMaterial = nullptr;
+UMaterialInterface* GColorMaskRegionMaterial = nullptr;
 UTexture2D* GLandscapeBlackTexture = nullptr;
 UMaterialInterface* GLandscapeLayerUsageMaterial = nullptr;
 #endif
@@ -596,6 +597,7 @@ void ULandscapeComponent::GetUsedMaterials(TArray<UMaterialInterface*>& OutMater
 		OutMaterials.Add(GSelectionColorMaterial);
 		OutMaterials.Add(GSelectionRegionMaterial);
 		OutMaterials.Add(GMaskRegionMaterial);
+		OutMaterials.Add(GColorMaskRegionMaterial);
 		OutMaterials.Add(GLandscapeLayerUsageMaterial);
 	}
 #endif
@@ -1091,6 +1093,12 @@ FPrimitiveViewRelevance FLandscapeComponentSceneProxy::GetViewRelevance(const FS
 				ToolRelevance |= GMaskRegionMaterial->GetRelevance_Concurrent(FeatureLevel);
 			}
 
+			if (GLandscapeViewMode == ELandscapeViewMode::LayerContribution)
+			{
+				Result.bDynamicRelevance = true;
+				ToolRelevance |= GColorMaskRegionMaterial->GetRelevance_Concurrent(FeatureLevel);
+			}
+
 			ToolRelevance.SetPrimitiveViewRelevance(Result);
 		}
 	}
@@ -1116,7 +1124,7 @@ FPrimitiveViewRelevance FLandscapeComponentSceneProxy::GetViewRelevance(const FS
 		View->Family->EngineShowFlags.Wireframe ||
 #if WITH_EDITOR
 		(IsSelected() && !GLandscapeEditModeActive) ||
-		GLandscapeViewMode != ELandscapeViewMode::Normal ||
+		(GLandscapeViewMode != ELandscapeViewMode::Normal) ||
 #else
 		IsSelected() ||
 #endif
@@ -2679,6 +2687,26 @@ void FLandscapeComponentSceneProxy::GetDynamicMeshElements(const TArray<const FS
 			}
 			break;
 
+			case ELandscapeViewMode::LayerContribution:
+			{
+				Mesh.bCanApplyViewModeOverrides = false;
+				Collector.AddMesh(ViewIndex, Mesh);
+				NumPasses++;
+				NumTriangles += Mesh.GetNumPrimitives();
+				NumDrawCalls += Mesh.Elements.Num();
+
+				FMeshBatch& MaskMesh = Collector.AllocateMesh();
+				MaskMesh = MeshTools;
+				auto ColorMaskMaterialInstance = new FLandscapeMaskMaterialRenderProxy(GColorMaskRegionMaterial->GetRenderProxy(), EditToolRenderData.LayerContributionTexture ? EditToolRenderData.LayerContributionTexture : GLandscapeBlackTexture, true);
+				MaskMesh.MaterialRenderProxy = ColorMaskMaterialInstance;
+				Collector.RegisterOneFrameMaterialProxy(ColorMaskMaterialInstance);
+				Collector.AddMesh(ViewIndex, MaskMesh);
+				NumPasses++;
+				NumTriangles += MaskMesh.GetNumPrimitives();
+				NumDrawCalls += MaskMesh.Elements.Num();
+			}
+			break;
+
 			default:
 
 #endif // WITH_EDITOR
diff --git a/Engine/Source/Runtime/Landscape/Private/LandscapeSplineRaster.cpp b/Engine/Source/Runtime/Landscape/Private/LandscapeSplineRaster.cpp
--- a/Engine/Source/Runtime/Landscape/Private/LandscapeSplineRaster.cpp
+++ b/Engine/Source/Runtime/Landscape/Private/LandscapeSplineRaster.cpp
@@ -467,7 +467,7 @@ bool ULandscapeInfo::ApplySplines(bool bOnlySelected)
 	FGuid SplinesTargetLayerGuid = Layer ? Layer->Guid : FGuid();
 	FScopedSetLandscapeEditingLayer Scope(Landscape, SplinesTargetLayerGuid, [=] 
 	{ 
-		Landscape->RequestLayersContentUpdate(ELandscapeLayerUpdateMode::All);
+		Landscape->RequestLayersContentUpdate(ELandscapeLayerUpdateMode::Update_All);
 	});
 
 	ForAllLandscapeProxies([&bResult, bOnlySelected, this](ALandscapeProxy* Proxy)
diff --git a/Engine/Source/Runtime/Landscape/Public/LandscapeEdit.h b/Engine/Source/Runtime/Landscape/Public/LandscapeEdit.h
--- a/Engine/Source/Runtime/Landscape/Public/LandscapeEdit.h
+++ b/Engine/Source/Runtime/Landscape/Public/LandscapeEdit.h
@@ -33,6 +33,25 @@ class ULandscapeLayerInfoObject;
 
 #if WITH_EDITOR
 
+namespace ELandscapeToolTargetType
+{
+	enum Type : int8
+	{
+		Heightmap = 0,
+		Weightmap = 1,
+		Visibility = 2,
+		Invalid = -1, // only valid for LandscapeEdMode->CurrentToolTarget.TargetType
+	};
+}
+
+/** Landscape EdMode Interface (used by ALandscape in Editor mode to access EdMode properties) */
+class ILandscapeEdModeInterface
+{
+public:
+	virtual ELandscapeToolTargetType::Type GetLandscapeToolTargetType() const = 0;
+	virtual FGuid GetLandscapeSelectedLayer() const = 0;
+};
+
 struct FLandscapeTextureDataInfo
 {
 	struct FMipInfo
@@ -187,12 +206,22 @@ struct LANDSCAPE_API FLandscapeEditDataInterface : public FLandscapeTextureDataI
 	// Replace/merge a layer
 	void ReplaceLayer(ULandscapeLayerInfoObject* FromLayerInfo, ULandscapeLayerInfoObject* ToLayerInfo);
 
+	template<typename TStoreData>
+	void GetEditToolTextureData(const int32 X1, const int32 Y1, const int32 X2, const int32 Y2, TStoreData& StoreData, TFunctionRef<UTexture2D*(ULandscapeComponent*)> GetComponentTexture);
+	void SetEditToolTextureData(int32 X1, int32 Y1, int32 X2, int32 Y2, const uint8* Data, int32 Stride, TFunctionRef<UTexture2D*&(ULandscapeComponent*)> GetComponentTexture);
+
 	// Without data interpolation, Select Data 
 	template<typename TStoreData>
 	void GetSelectDataTempl(const int32 X1, const int32 Y1, const int32 X2, const int32 Y2, TStoreData& StoreData);
 	void GetSelectData(const int32 X1, const int32 Y1, const int32 X2, const int32 Y2, uint8* Data, int32 Stride);
 	void GetSelectData(const int32 X1, const int32 Y1, const int32 X2, const int32 Y2, TMap<FIntPoint, uint8>& SparseData);
 	void SetSelectData(const int32 X1, const int32 Y1, const int32 X2, const int32 Y2, const uint8* Data, int32 Stride);
+	
+	template<typename TStoreData>
+	void GetLayerContributionDataTempl(const int32 X1, const int32 Y1, const int32 X2, const int32 Y2, TStoreData& StoreData);
+	void GetLayerContributionData(const int32 X1, const int32 Y1, const int32 X2, const int32 Y2, uint8* Data, int32 Stride);
+	void GetLayerContributionData(const int32 X1, const int32 Y1, const int32 X2, const int32 Y2, TMap<FIntPoint, uint8>& SparseData);
+	void SetLayerContributionData(const int32 X1, const int32 Y1, const int32 X2, const int32 Y2, const uint8* Data, int32 Stride);
 
 	//
 	// XYOffsetmap access
diff --git a/Engine/Source/Runtime/Landscape/Public/LandscapeRender.h b/Engine/Source/Runtime/Landscape/Public/LandscapeRender.h
--- a/Engine/Source/Runtime/Landscape/Public/LandscapeRender.h
+++ b/Engine/Source/Runtime/Landscape/Public/LandscapeRender.h
@@ -50,6 +50,7 @@ namespace ELandscapeViewMode
 		LayerUsage,
 		LOD,
 		WireframeOnTop,
+		LayerContribution
 	};
 }
 
@@ -77,6 +78,7 @@ LANDSCAPE_API extern UMaterialInterface* GLayerDebugColorMaterial;
 LANDSCAPE_API extern UMaterialInterface* GSelectionColorMaterial;
 LANDSCAPE_API extern UMaterialInterface* GSelectionRegionMaterial;
 LANDSCAPE_API extern UMaterialInterface* GMaskRegionMaterial;
+LANDSCAPE_API extern UMaterialInterface* GColorMaskRegionMaterial;
 LANDSCAPE_API extern UTexture2D* GLandscapeBlackTexture;
 LANDSCAPE_API extern UMaterialInterface* GLandscapeLayerUsageMaterial;
 #endif
