- Support Visibility Paint in Landscape Layer contribution Visualizer
- Avoid grass regeneration when selected layer changes

diff --git a/Engine/Source/Editor/LandscapeEditor/Private/LandscapeEdMode.cpp b/Engine/Source/Editor/LandscapeEditor/Private/LandscapeEdMode.cpp
--- a/Engine/Source/Editor/LandscapeEditor/Private/LandscapeEdMode.cpp
+++ b/Engine/Source/Editor/LandscapeEditor/Private/LandscapeEdMode.cpp
@@ -467,7 +467,7 @@ ELandscapeToolTargetType::Type FEdModeLandscape::GetLandscapeToolTargetType() co
 	{
 		if (CurrentToolMode->ToolModeName == "ToolMode_Sculpt")
 		{
-			return ELandscapeToolTargetType::Heightmap;
+			return CurrentToolTarget.TargetType == ELandscapeToolTargetType::Visibility ? ELandscapeToolTargetType::Visibility : ELandscapeToolTargetType::Heightmap;
 		}
 		else if (CurrentToolMode->ToolModeName == "ToolMode_Paint")
 		{
diff --git a/Engine/Source/Runtime/Landscape/Private/LandscapeEditLayers.cpp b/Engine/Source/Runtime/Landscape/Private/LandscapeEditLayers.cpp
--- a/Engine/Source/Runtime/Landscape/Private/LandscapeEditLayers.cpp
+++ b/Engine/Source/Runtime/Landscape/Private/LandscapeEditLayers.cpp
@@ -4038,7 +4038,10 @@ void ALandscape::RequestLayersContentUpdateForceAll(ELandscapeLayerUpdateMode In
 		{
 			if (Proxy)
 			{
-				Proxy->InvalidateGeneratedComponentData();
+				if (bUpdateHeightmap || bUpdateWeightmap)
+				{
+					Proxy->InvalidateGeneratedComponentData();
+				}
 
 				for (ULandscapeComponent* Component : Proxy->LandscapeComponents)
 				{
@@ -4340,25 +4343,37 @@ bool ALandscape::UpdateCollisionAndClients(const TArray<ULandscapeComponent*>& I
 						LayerContributionMaskData[i] = HeightData[i] != DefaultHeightValue ? LayerContributingValue : 0;
 					}
 				}
-				else if (LandscapeEdModeInfo.ToolTarget == ELandscapeToolTargetType::Weightmap)
+				else if (LandscapeEdModeInfo.ToolTarget == ELandscapeToolTargetType::Weightmap ||
+						 LandscapeEdModeInfo.ToolTarget == ELandscapeToolTargetType::Visibility)
 				{
 					if (WeightData.Num() != ArraySize)
 					{
 						WeightData.AddZeroed(ArraySize);
 					}
-					FMemory::Memzero(LayerContributionMaskDataPtr, ArraySize);
-					for (const FLandscapeInfoLayerSettings& LayerSettings : Info->Layers)
+					if (LandscapeEdModeInfo.ToolTarget == ELandscapeToolTargetType::Visibility)
 					{
-						if (LayerSettings.LayerInfoObj == ALandscapeProxy::VisibilityLayer)
+						LandscapeEdit.GetWeightDataFast(ALandscapeProxy::VisibilityLayer, X1, Y1, X2, Y2, WeightData.GetData(), Stride);
+						for (int i = 0; i < ArraySize; ++i)
 						{
-							continue;
+							LayerContributionMaskData[i] = WeightData[i] != DefaultWeightValue ? LayerContributingValue : 0;
 						}
-						LandscapeEdit.GetWeightDataFast(LayerSettings.LayerInfoObj, X1, Y1, X2, Y2, WeightData.GetData(), Stride);
-						for (int i = 0; i < ArraySize; ++i)
+					}
+					else
+					{
+						FMemory::Memzero(LayerContributionMaskDataPtr, ArraySize);
+						for (const FLandscapeInfoLayerSettings& LayerSettings : Info->Layers)
 						{
-							if (WeightData[i] != DefaultWeightValue)
+							if (LayerSettings.LayerInfoObj == ALandscapeProxy::VisibilityLayer)
+							{
+								continue;
+							}
+							LandscapeEdit.GetWeightDataFast(LayerSettings.LayerInfoObj, X1, Y1, X2, Y2, WeightData.GetData(), Stride);
+							for (int i = 0; i < ArraySize; ++i)
 							{
-								LayerContributionMaskData[i] = LayerContributingValue;
+								if (WeightData[i] != DefaultWeightValue)
+								{
+									LayerContributionMaskData[i] = LayerContributingValue;
+								}
 							}
 						}
 					}
