Fixed static analysis
#rnx

diff --git a/Engine/Source/Runtime/Landscape/Private/LandscapeEditInterface.cpp b/Engine/Source/Runtime/Landscape/Private/LandscapeEditInterface.cpp
--- a/Engine/Source/Runtime/Landscape/Private/LandscapeEditInterface.cpp
+++ b/Engine/Source/Runtime/Landscape/Private/LandscapeEditInterface.cpp
@@ -4542,7 +4542,11 @@ void FLandscapeEditDataInterface::GetEditToolTextureData(const int32 X1, const i
 
 			FLandscapeTextureDataInfo* TexDataInfo = NULL;
 			uint8* SelectTextureData = NULL;
-			UTexture2D* EditToolTexture = Component ? GetComponentTexture(Component) : nullptr;
+			UTexture2D* EditToolTexture = nullptr;
+			if (Component)
+			{
+				EditToolTexture = GetComponentTexture(Component);
+			}
 			if (EditToolTexture)
 			{
 				TexDataInfo = GetTextureDataInfo(EditToolTexture);
