- Change the way the Experimental flag was used, we only rely on the flag to allow enabling/disabling the landscape layers feature
- Enabling/Disabling the layer system is per world settings now
- Changed most of the layer code to either rely on the feature being enabled in the level or be data driven, which is we pass through the layer code path if we have layer content

diff --git a/Engine/Source/Editor/DetailCustomizations/Private/WorldSettingsDetails.cpp b/Engine/Source/Editor/DetailCustomizations/Private/WorldSettingsDetails.cpp
--- a/Engine/Source/Editor/DetailCustomizations/Private/WorldSettingsDetails.cpp
+++ b/Engine/Source/Editor/DetailCustomizations/Private/WorldSettingsDetails.cpp
@@ -15,21 +15,48 @@
 #include "Toolkits/AssetEditorManager.h"
 #include "DetailLayoutBuilder.h"
 #include "GameModeInfoCustomizer.h"
+#include "Settings/EditorExperimentalSettings.h"
+#include "GameFramework/WorldSettings.h"
 
 #define LOCTEXT_NAMESPACE "WorldSettingsDetails"
 
 
+FWorldSettingsDetails::~FWorldSettingsDetails()
+{
+	GetMutableDefault<UEditorExperimentalSettings>()->OnSettingChanged().Remove(ExperimentalDelegateHandle);
+}
+
+void FWorldSettingsDetails::OnEditorExperimentalSettingsChanged(FName InPropertyName)
+{
+	if (DetailLayoutBuilder && InPropertyName == GET_MEMBER_NAME_CHECKED(UEditorExperimentalSettings, bLandscapeLayerSystem))
+	{
+		DetailLayoutBuilder->ForceRefreshDetails();
+	}
+}
+
 /* IDetailCustomization overrides
  *****************************************************************************/
 
 void FWorldSettingsDetails::CustomizeDetails( IDetailLayoutBuilder& DetailBuilder )
 {
+	DetailLayoutBuilder = &DetailBuilder;
+
+	if (!ExperimentalDelegateHandle.IsValid())
+	{
+		ExperimentalDelegateHandle = GetMutableDefault<UEditorExperimentalSettings>()->OnSettingChanged().AddSP(this, &FWorldSettingsDetails::OnEditorExperimentalSettingsChanged);
+	}
+
 	IDetailCategoryBuilder& Category = DetailBuilder.EditCategory("GameMode");
 	CustomizeGameInfoProperty("DefaultGameMode", DetailBuilder, Category);
 
 	AddLightmapCustomization(DetailBuilder);
 
 	DetailBuilder.HideProperty(GET_MEMBER_NAME_CHECKED(AActor, bHidden), AActor::StaticClass());
+
+	if (!GetMutableDefault<UEditorExperimentalSettings>()->bLandscapeLayerSystem)
+	{
+		DetailBuilder.HideProperty(GET_MEMBER_NAME_CHECKED(AWorldSettings, bEnableLandscapeLayers), AWorldSettings::StaticClass());
+	}
 }
 
 
diff --git a/Engine/Source/Editor/DetailCustomizations/Private/WorldSettingsDetails.h b/Engine/Source/Editor/DetailCustomizations/Private/WorldSettingsDetails.h
--- a/Engine/Source/Editor/DetailCustomizations/Private/WorldSettingsDetails.h
+++ b/Engine/Source/Editor/DetailCustomizations/Private/WorldSettingsDetails.h
@@ -25,6 +25,10 @@ class FWorldSettingsDetails
 {
 public:
 
+	virtual ~FWorldSettingsDetails();
+
+	void OnEditorExperimentalSettingsChanged(FName InPropertyName);
+
 	// IDetailCustomization interface
 
 	virtual void CustomizeDetails( IDetailLayoutBuilder& DetailBuilder ) override;
@@ -70,6 +74,8 @@ private:
 
 	/** Helper class to customizer GameMode property */
 	TSharedPtr<FGameModeInfoCustomizer>	GameInfoModeCustomizer;
+	FDelegateHandle ExperimentalDelegateHandle;
+	IDetailLayoutBuilder* DetailLayoutBuilder = nullptr;
 };
 
 
diff --git a/Engine/Source/Editor/LandscapeEditor/Private/LandscapeEdMode.cpp b/Engine/Source/Editor/LandscapeEditor/Private/LandscapeEdMode.cpp
--- a/Engine/Source/Editor/LandscapeEditor/Private/LandscapeEdMode.cpp
+++ b/Engine/Source/Editor/LandscapeEditor/Private/LandscapeEdMode.cpp
@@ -191,7 +191,7 @@ void ALandscape::SplitHeightmap(ULandscapeComponent* Comp, ALandscapeProxy* Targ
 	GDisableAutomaticTextureMaterialUpdateDependencies = false;
 
 #if WITH_EDITORONLY_DATA
-	if (GetMutableDefault<UEditorExperimentalSettings>()->bLandscapeLayerSystem)
+	if (Comp->GetLandscapeProxy()->HasLayersContent() && DstProxy->CanHaveLayersContent())
 	{
 		FLandscapeLayersTexture2DCPUReadBackResource* NewCPUReadBackResource = new FLandscapeLayersTexture2DCPUReadBackResource(NewHeightmapTexture->Source.GetSizeX(), NewHeightmapTexture->Source.GetSizeY(), NewHeightmapTexture->GetPixelFormat(), NewHeightmapTexture->Source.GetNumMips());
 		BeginInitResource(NewCPUReadBackResource);
@@ -312,7 +312,7 @@ FEdModeLandscape::FEdModeLandscape()
 
 
 	// Initialize modes
-	InitializeToolModes();
+	UpdateToolModes();
 	CurrentToolMode = nullptr;
 
 	// Initialize tools.
@@ -404,8 +404,10 @@ void FEdModeLandscape::AddReferencedObjects(FReferenceCollector& Collector)
 	Collector.AddReferencedObject(GLandscapeLayerUsageMaterial);
 }
 
-void FEdModeLandscape::InitializeToolModes()
+void FEdModeLandscape::UpdateToolModes()
 {
+	LandscapeToolModes.Reset();
+
 	FLandscapeToolMode* ToolMode_Manage = new(LandscapeToolModes)FLandscapeToolMode(TEXT("ToolMode_Manage"), ELandscapeToolTargetTypeMask::NA);
 	ToolMode_Manage->ValidTools.Add(TEXT("NewLandscape"));
 	ToolMode_Manage->ValidTools.Add(TEXT("Select"));
@@ -427,7 +429,7 @@ void FEdModeLandscape::InitializeToolModes()
 	ToolMode_Sculpt->ValidTools.Add(TEXT("Retopologize"));
 	ToolMode_Sculpt->ValidTools.Add(TEXT("Visibility"));
 
-	if (GetMutableDefault<UEditorExperimentalSettings>()->bLandscapeLayerSystem)
+	if (CanHaveLandscapeLayersContent())
 	{
         ToolMode_Sculpt->ValidTools.Add(TEXT("Erase"));
 		ToolMode_Sculpt->ValidTools.Add(TEXT("BPCustom"));
@@ -444,7 +446,7 @@ void FEdModeLandscape::InitializeToolModes()
 	ToolMode_Paint->ValidTools.Add(TEXT("Noise"));
 	ToolMode_Paint->ValidTools.Add(TEXT("Visibility"));
 
-	if (GetMutableDefault<UEditorExperimentalSettings>()->bLandscapeLayerSystem)
+	if (CanHaveLandscapeLayersContent())
 	{
 		ToolMode_Paint->ValidTools.Add(TEXT("BPCustom"));
 	}
@@ -506,14 +508,34 @@ void FEdModeLandscape::Enter()
 
 	OnLevelActorDeletedDelegateHandle = GEngine->OnLevelActorDeleted().AddSP(this, &FEdModeLandscape::OnLevelActorRemoved);
 	OnLevelActorAddedDelegateHandle = GEngine->OnLevelActorAdded().AddSP(this, &FEdModeLandscape::OnLevelActorAdded);
-	OnLandscapeLayerSystemFlagChangedDelegateHandle = GetMutableDefault<UEditorExperimentalSettings>()->OnSettingChanged().AddLambda([this](FName PropertyName)
-	{ 
-		if (PropertyName == GET_MEMBER_NAME_CHECKED(UEditorExperimentalSettings, bLandscapeLayerSystem))
+
+	AWorldSettings* WorldSettings = GWorld != nullptr ? GWorld->GetWorldSettings() : nullptr;
+	
+	if (WorldSettings)
+	{
+		OnLandscapeLayerSystemFlagChangedDelegateHandle = WorldSettings->OnSettingChanged().AddLambda([this](const FName& InPropertyName)
 		{
-			UpdateLandscapeList();
-			RefreshDetailPanel();
-		}
-	});
+			if (InPropertyName == GET_MEMBER_NAME_CHECKED(AWorldSettings, bEnableLandscapeLayers))
+			{
+				UpdateToolModes();
+
+				if (CurrentToolMode->ToolModeName == TEXT("ToolMode_Paint"))
+				{
+					SetCurrentTool("Paint");
+				}
+				else if (CurrentToolMode->ToolModeName == TEXT("ToolMode_Sculpt"))
+				{
+					SetCurrentTool("Sculpt");
+				}
+
+				RefreshDetailPanel();
+			}
+		});
+	}
+
+	ALandscape::RegisterChangeLandscapeLayersStateDelegate();
+
+	UpdateToolModes();
 
 	ALandscapeProxy* SelectedLandscape = GEditor->GetSelectedActors()->GetTop<ALandscapeProxy>();
 	if (SelectedLandscape)
@@ -553,7 +575,9 @@ void FEdModeLandscape::Enter()
 		ALandscapeProxy* LandscapeProxy = CurrentToolTarget.LandscapeInfo->GetLandscapeProxy();
 		LandscapeProxy->OnMaterialChangedDelegate().AddRaw(this, &FEdModeLandscape::OnLandscapeMaterialChangedDelegate);
 
-		if(ALandscape* Landscape = GetLandscape())
+		ALandscape* Landscape = GetLandscape();
+
+		if(Landscape && Landscape->HasLayersContent())
 		{
 			Landscape->RequestLayersContentUpdateForceAll();
 		}
@@ -653,7 +677,7 @@ void FEdModeLandscape::Enter()
 	}
 	else
 	{
-		if (CurrentToolMode == nullptr || (CurrentToolMode->CurrentToolName == FName("NewLandscape")))
+		if (CurrentToolMode == nullptr || (CurrentToolMode->CurrentToolName == FName("NewLandscape")) || CurrentToolMode->CurrentToolName == NAME_None)
 		{
 			SetCurrentToolMode("ToolMode_Sculpt", false);
 			SetCurrentTool("Sculpt");
@@ -725,6 +749,13 @@ void FEdModeLandscape::Exit()
 	GEngine->OnLevelActorDeleted().Remove(OnLevelActorDeletedDelegateHandle);
 	GEngine->OnLevelActorAdded().Remove(OnLevelActorAddedDelegateHandle);
 	GetMutableDefault<UEditorExperimentalSettings>()->OnSettingChanged().Remove(OnLandscapeLayerSystemFlagChangedDelegateHandle);
+
+	AWorldSettings* WorldSettings = GWorld != nullptr ? GWorld->GetWorldSettings() : nullptr;
+
+	if (WorldSettings)
+	{
+		WorldSettings->OnSettingChanged().Remove(OnLandscapeLayerSystemFlagChangedDelegateHandle);
+	}
 	
 	FEditorSupportDelegates::WorldChange.Remove(OnWorldChangeDelegateHandle);
 	GetWorld()->OnLevelsChanged().Remove(OnLevelsChangedDelegateHandle);
@@ -2508,22 +2539,8 @@ bool FEdModeLandscape::CanEditCurrentTarget(FText* Reason) const
 		return false;
 	}
 
-	if (!GetMutableDefault<UEditorExperimentalSettings>()->bLandscapeLayerSystem)
-	{
-		// If Any Proxy has layer content disable editing when Landscape Layers are disabled
-		bool bHasLayerContent = false;
-		CurrentToolTarget.LandscapeInfo->ForAllLandscapeProxies([&bHasLayerContent](ALandscapeProxy* Proxy) { bHasLayerContent |= Proxy->HasLayersContent; });
-		if (bHasLayerContent)
-		{
-			LocalReason = NSLOCTEXT("UnrealEd", "LandscapeWithLayerContent", "Landscape contains layers and experimental flag is not enabled.");
-			return false;
-		}
-
-		return true;
-	}
-
 	// Landscape Layer Editing not available without a loaded Landscape Actor
-	if (CurrentToolTarget.LandscapeInfo->LandscapeActor == nullptr)
+	if (CanHaveLandscapeLayersContent() && !CurrentToolTarget.LandscapeInfo->LandscapeActor.IsValid())
 	{
 		LocalReason = NSLOCTEXT("UnrealEd", "LandscapeActorNotLoaded", "Landscape actor is not loaded. It is needed to do layer editing.");
 		return false;
@@ -3986,7 +4003,7 @@ ALandscape* FEdModeLandscape::ChangeComponentSetting(int32 NumComponentsX, int32
 					}
 				};
 
-				if (GetMutableDefault<UEditorExperimentalSettings>()->bLandscapeLayerSystem)
+				if (HasLandscapeLayersContent())
 				{
 					int32 HeightCount = 0;
 
@@ -4089,7 +4106,7 @@ ALandscape* FEdModeLandscape::ChangeComponentSetting(int32 NumComponentsX, int32
 			NewLandscape->bUsedForNavigation = OldLandscape->bUsedForNavigation;
 			NewLandscape->MaxPaintedLayersPerComponent = OldLandscape->MaxPaintedLayersPerComponent;
 
-			TArray<FLandscapeLayer>* LandscapeLayers = GetMutableDefault<UEditorExperimentalSettings>()->bLandscapeLayerSystem ? &OldLandscape->LandscapeLayers : nullptr;
+			TArray<FLandscapeLayer>* LandscapeLayers = CanHaveLandscapeLayersContent() ? &OldLandscape->LandscapeLayers : nullptr;
 
 			NewLandscape->Import(FGuid::NewGuid(), NewMinX, NewMinY, NewMaxX, NewMaxY, NumSubsections, SubsectionSizeQuads, HeightDataPerLayers, *OldLandscape->ReimportHeightmapFilePath, ImportMaterialLayerInfosPerLayers, ELandscapeImportAlphamapType::Additive, LandscapeLayers);
 
@@ -4258,6 +4275,18 @@ ELandscapeEditingState FEdModeLandscape::GetEditingState() const
 	return ELandscapeEditingState::Enabled;
 }
 
+bool FEdModeLandscape::CanHaveLandscapeLayersContent() const
+{
+	ALandscape* Landscape = GetLandscape();
+	return Landscape != nullptr ? Landscape->CanHaveLayersContent() : false;
+}
+
+bool FEdModeLandscape::HasLandscapeLayersContent() const
+{
+	ALandscape* Landscape = GetLandscape();	
+	return Landscape != nullptr ? Landscape->HasLayersContent() : false;
+}
+
 int32 FEdModeLandscape::GetLayerCount() const
 {
 	ALandscape* Landscape = GetLandscape();
@@ -4520,7 +4549,7 @@ FLandscapeLayer* FEdModeLandscape::GetCurrentLayer() const
 
 void FEdModeLandscape::AutoUpdateDirtyLandscapeSplines()
 {
-	if (GetMutableDefault<UEditorExperimentalSettings>()->bLandscapeLayerSystem && GEditor->IsTransactionActive())
+	if (HasLandscapeLayersContent() && GEditor->IsTransactionActive())
 	{
 		// Only auto-update if a layer is reserved for landscape splines
 		ALandscape* Landscape = GetLandscape();
@@ -4534,7 +4563,7 @@ void FEdModeLandscape::AutoUpdateDirtyLandscapeSplines()
 
 bool FEdModeLandscape::CanEditLayer(FText* Reason /*=nullptr*/, FLandscapeLayer* InLayer /*= nullptr*/)
 {
-	if (GetMutableDefault<UEditorExperimentalSettings>()->bLandscapeLayerSystem)
+	if (CanHaveLandscapeLayersContent())
 	{
 		ALandscape* Landscape = GetLandscape();
 		FLandscapeLayer* TargetLayer = InLayer ? InLayer : GetCurrentLayer();
@@ -4612,7 +4641,7 @@ bool FEdModeLandscape::CanEditLayer(FText* Reason /*=nullptr*/, FLandscapeLayer*
 
 void FEdModeLandscape::UpdateLandscapeSplines(bool bUpdateOnlySelected /* = false*/)
 {
-	if (GetMutableDefault<UEditorExperimentalSettings>()->bLandscapeLayerSystem)
+	if (HasLandscapeLayersContent())
 	{
 		ALandscape* Landscape = GetLandscape();
 		if (Landscape)
diff --git a/Engine/Source/Editor/LandscapeEditor/Private/LandscapeEdMode.h b/Engine/Source/Editor/LandscapeEditor/Private/LandscapeEdMode.h
--- a/Engine/Source/Editor/LandscapeEditor/Private/LandscapeEdMode.h
+++ b/Engine/Source/Editor/LandscapeEditor/Private/LandscapeEdMode.h
@@ -324,7 +324,7 @@ public:
 	void InitializeTool_Ramp();
 	void InitializeTool_Mirror();
 	void InitializeTool_BPCustom();
-	void InitializeToolModes();
+	void UpdateToolModes();
 
 	/** Destructor */
 	virtual ~FEdModeLandscape();
@@ -499,6 +499,8 @@ public:
 	void RefreshDetailPanel();
 
 	// Layers
+	bool CanHaveLandscapeLayersContent() const;
+	bool HasLandscapeLayersContent() const;
 	int32 GetLayerCount() const;
 	void SetCurrentLayer(int32 InLayerIndex);
 	int32 GetCurrentLayerIndex() const;
diff --git a/Engine/Source/Editor/LandscapeEditor/Private/LandscapeEdModeComponentTools.cpp b/Engine/Source/Editor/LandscapeEditor/Private/LandscapeEdModeComponentTools.cpp
--- a/Engine/Source/Editor/LandscapeEditor/Private/LandscapeEdModeComponentTools.cpp
+++ b/Engine/Source/Editor/LandscapeEditor/Private/LandscapeEdModeComponentTools.cpp
@@ -22,7 +22,6 @@
 #include "PhysicalMaterials/PhysicalMaterial.h"
 #include "Materials/MaterialExpressionLandscapeVisibilityMask.h"
 #include "Algo/Copy.h"
-#include "Settings/EditorExperimentalSettings.h"
 
 #define LOCTEXT_NAMESPACE "Landscape"
 
@@ -215,14 +214,6 @@ public:
 				}
 			}
 
-			ALandscape* Landscape = LandscapeInfo->LandscapeActor.Get();
-
-			if (Landscape != nullptr && Landscape->HasLayersContent && !GetMutableDefault<UEditorExperimentalSettings>()->bLandscapeLayerSystem)
-			{
-				FMessageLog("MapCheck").Warning()->AddToken(FTextToken::Create(LOCTEXT("LandscapeLayers_ChangingDataWithoutSettings", "This map contains landscape layer system content, modifying the landscape data will result in data loss when the map is reopened with Landscape Layer System settings on. Please enable Landscape Layer System settings before modifying the data.")));
-				FMessageLog("MapCheck").Open(EMessageSeverity::Warning);
-			}
-
 			Cache.SetCachedData(X1, Y1, X2, Y2, Data);
 			Cache.Flush();
 		}
@@ -302,14 +293,6 @@ public:
 				}
 			}
 
-			ALandscape* Landscape = LandscapeInfo->LandscapeActor.Get();
-
-			if (Landscape != nullptr && Landscape->HasLayersContent && !GetMutableDefault<UEditorExperimentalSettings>()->bLandscapeLayerSystem)
-			{
-				FMessageLog("MapCheck").Warning()->AddToken(FTextToken::Create(LOCTEXT("LandscapeLayers_ChangingDataWithoutSettings", "This map contains landscape layer system content, modifying the landscape data will result in data loss when the map is reopened with Landscape Layer System settings on. Please enable Landscape Layer System settings before modifying the data.")));
-				FMessageLog("MapCheck").Open(EMessageSeverity::Warning);
-			}
-
 			Cache.SetCachedData(X1, Y1, X2, Y2, Data);
 			Cache.Flush();
 		}
@@ -804,13 +787,6 @@ public:
 				NewComponents[Idx]->RegisterComponent();
 			}
 
-			if (LandscapeInfo->LandscapeActor.IsValid() && LandscapeInfo->LandscapeActor.Get()->HasLayersContent && !GetMutableDefault<UEditorExperimentalSettings>()->bLandscapeLayerSystem)
-			{
-				FMessageLog("MapCheck").Warning()->AddToken(FTextToken::Create(LOCTEXT("LandscapeLayers_ChangingDataWithoutSettings", "This map contains landscape layer system content, modifying the landscape data will result in data loss when the map is reopened with Landscape Layer System settings on. Please enable Landscape Layer System settings before modifying the data.")));
-
-				FMessageLog("MapCheck").Open(EMessageSeverity::Warning);
-			}
-
 			if (bHasXYOffset)
 			{
 				XYOffsetCache.SetCachedData(X1, Y1, X2, Y2, XYOffsetData);
@@ -822,7 +798,9 @@ public:
 
 			ALandscape* Landscape = LandscapeInfo->LandscapeActor.Get();
 
-			if (GetMutableDefault<UEditorExperimentalSettings>()->bLandscapeLayerSystem)
+			bool bHasLandscapeLayersContent = Landscape->HasLayersContent();
+
+			if (bHasLandscapeLayersContent)
 			{
 				check(Landscape != nullptr); // Landscape actor is required if layer system is enabled
 
@@ -831,7 +809,7 @@ public:
 
 			for (ULandscapeComponent* NewComponent : NewComponents)
 			{
-				if (GetMutableDefault<UEditorExperimentalSettings>()->bLandscapeLayerSystem)
+				if (bHasLandscapeLayersContent)
 				{
 					TArray<ULandscapeComponent*> ComponentsUsingHeightmap;
 					ComponentsUsingHeightmap.Add(NewComponent);
@@ -849,7 +827,7 @@ public:
 				NewComponent->UpdateBounds();
 				NewComponent->MarkRenderStateDirty();
 
-				if (!GetMutableDefault<UEditorExperimentalSettings>()->bLandscapeLayerSystem)
+				if (!bHasLandscapeLayersContent)
 				{
 					ULandscapeHeightfieldCollisionComponent* CollisionComp = NewComponent->CollisionComponent.Get();
 					if (CollisionComp && !bHasXYOffset)
@@ -929,7 +907,7 @@ public:
 			EdMode->LandscapeRenderAddCollision = nullptr;
 
 			// Add/update "add collision" around the newly added components
-			if (!GetMutableDefault<UEditorExperimentalSettings>()->bLandscapeLayerSystem)
+			if (!bHasLandscapeLayersContent)
 			{
 				// Top row
 				int32 ComponentIndexY = ComponentIndexY1 - 1;
@@ -1660,14 +1638,6 @@ public:
 				}
 			}
 
-			ALandscape* Landscape = LandscapeInfo->LandscapeActor.Get();
-
-			if (Landscape != nullptr && Landscape->HasLayersContent && !GetMutableDefault<UEditorExperimentalSettings>()->bLandscapeLayerSystem)
-			{
-				FMessageLog("MapCheck").Warning()->AddToken(FTextToken::Create(LOCTEXT("LandscapeLayers_ChangingDataWithoutSettings", "This map contains landscape layer system content, modifying the landscape data will result in data loss when the map is reopened with Landscape Layer System settings on. Please enable Landscape Layer System settings before modifying the data.")));
-				FMessageLog("MapCheck").Open(EMessageSeverity::Warning);
-			}
-
 			if (bApplyToAll)
 			{
 				HeightCache.SetCachedData(X1, Y1, X2, Y2, HeightData);
diff --git a/Engine/Source/Editor/LandscapeEditor/Private/LandscapeEdModeErosionTools.cpp b/Engine/Source/Editor/LandscapeEditor/Private/LandscapeEdModeErosionTools.cpp
--- a/Engine/Source/Editor/LandscapeEditor/Private/LandscapeEdModeErosionTools.cpp
+++ b/Engine/Source/Editor/LandscapeEditor/Private/LandscapeEdModeErosionTools.cpp
@@ -10,7 +10,6 @@
 #include "Logging/TokenizedMessage.h"
 #include "Logging/MessageLog.h"
 #include "Misc/MapErrors.h"
-#include "Settings/EditorExperimentalSettings.h"
 
 #define LOCTEXT_NAMESPACE "LandscapeTools"
 //
@@ -257,14 +256,6 @@ public:
 			}
 		}
 
-		ALandscape* Landscape = LandscapeInfo->LandscapeActor.Get();
-
-		if (Landscape != nullptr && Landscape->HasLayersContent && !GetMutableDefault<UEditorExperimentalSettings>()->bLandscapeLayerSystem)
-		{
-			FMessageLog("MapCheck").Warning()->AddToken(FTextToken::Create(LOCTEXT("LandscapeLayers_ChangingDataWithoutSettings", "This map contains landscape layer system content, modifying the landscape data will result in data loss when the map is reopened with Landscape Layer System settings on. Please enable Landscape Layer System settings before modifying the data.")));
-			FMessageLog("MapCheck").Open(EMessageSeverity::Warning);
-		}
-
 		HeightCache.SetCachedData(X1, Y1, X2, Y2, HeightData);
 		HeightCache.Flush();
 		if (bWeightApplied)
@@ -488,14 +479,6 @@ public:
 			LowPassFilter<uint16>(X1, Y1, X2, Y2, BrushInfo, HeightData, UISettings->HErosionDetailScale, 1.0f);
 		}
 
-		ALandscape* Landscape = LandscapeInfo->LandscapeActor.Get();
-
-		if (Landscape != nullptr && Landscape->HasLayersContent && !GetMutableDefault<UEditorExperimentalSettings>()->bLandscapeLayerSystem)
-		{
-			FMessageLog("MapCheck").Warning()->AddToken(FTextToken::Create(LOCTEXT("LandscapeLayers_ChangingDataWithoutSettings", "This map contains landscape layer system content, modifying the landscape data will result in data loss when the map is reopened with Landscape Layer System settings on. Please enable Landscape Layer System settings before modifying the data.")));
-			FMessageLog("MapCheck").Open(EMessageSeverity::Warning);
-		}
-
 		HeightCache.SetCachedData(X1, Y1, X2, Y2, HeightData);
 		HeightCache.Flush();
 	}
diff --git a/Engine/Source/Editor/LandscapeEditor/Private/LandscapeEdModeMirrorTool.cpp b/Engine/Source/Editor/LandscapeEditor/Private/LandscapeEdModeMirrorTool.cpp
--- a/Engine/Source/Editor/LandscapeEditor/Private/LandscapeEdModeMirrorTool.cpp
+++ b/Engine/Source/Editor/LandscapeEditor/Private/LandscapeEdModeMirrorTool.cpp
@@ -649,7 +649,7 @@ public:
 		{
 			ALandscapeProxy::InvalidateGeneratedComponentData(Components);
 
-			if (!GetMutableDefault<UEditorExperimentalSettings>()->bLandscapeLayerSystem)
+			if (!EdMode->HasLandscapeLayersContent())
 			{
 				for (ULandscapeComponent* Component : Components)
 				{
diff --git a/Engine/Source/Editor/LandscapeEditor/Private/LandscapeEdModePaintTools.cpp b/Engine/Source/Editor/LandscapeEditor/Private/LandscapeEdModePaintTools.cpp
--- a/Engine/Source/Editor/LandscapeEditor/Private/LandscapeEdModePaintTools.cpp
+++ b/Engine/Source/Editor/LandscapeEditor/Private/LandscapeEdModePaintTools.cpp
@@ -13,7 +13,6 @@
 #include "LandscapeEdit.h"
 #include "LandscapeDataAccess.h"
 #include "LandscapeEdModeTools.h"
-#include "Settings/EditorExperimentalSettings.h"
 #include "Landscape.h"
 #include "Logging/TokenizedMessage.h"
 #include "Logging/MessageLog.h"
@@ -255,14 +254,6 @@ public:
 			}
 		}
 
-		ALandscape* Landscape = LandscapeInfo->LandscapeActor.Get();
-
-		if (Landscape != nullptr && Landscape->HasLayersContent && !GetMutableDefault<UEditorExperimentalSettings>()->bLandscapeLayerSystem)
-		{
-			FMessageLog("MapCheck").Warning()->AddToken(FTextToken::Create(LOCTEXT("LandscapeLayers_ChangingDataWithoutSettings", "This map contains landscape layer system content, modifying the landscape data will result in data loss when the map is reopened with Landscape Layer System settings on. Please enable Landscape Layer System settings before modifying the data.")));
-			FMessageLog("MapCheck").Open(EMessageSeverity::Warning);
-		}
-
 		this->Cache.SetCachedData(X1, Y1, X2, Y2, Data, UISettings->PaintingRestriction);
 		this->Cache.Flush();
 	}
@@ -360,14 +351,6 @@ public:
 			}
 		}
 
-		ALandscape* Landscape = this->LandscapeInfo->LandscapeActor.Get();
-
-		if (Landscape != nullptr && Landscape->HasLayersContent && !GetMutableDefault<UEditorExperimentalSettings>()->bLandscapeLayerSystem)
-		{
-			FMessageLog("MapCheck").Warning()->AddToken(FTextToken::Create(LOCTEXT("LandscapeLayers_ChangingDataWithoutSettings", "This map contains landscape layer system content, modifying the landscape data will result in data loss when the map is reopened with Landscape Layer System settings on. Please enable Landscape Layer System settings before modifying the data.")));
-			FMessageLog("MapCheck").Open(EMessageSeverity::Warning);
-		}
-
 		this->Cache.SetCachedData(X1, Y1, X2, Y2, Data, UISettings->PaintingRestriction);
 		this->Cache.Flush();
 	}
@@ -581,14 +564,6 @@ public:
 			}
 		}
 
-		ALandscape* Landscape = LandscapeInfo->LandscapeActor.Get();
-
-		if (Landscape != nullptr && Landscape->HasLayersContent && !GetMutableDefault<UEditorExperimentalSettings>()->bLandscapeLayerSystem)
-		{
-			FMessageLog("MapCheck").Warning()->AddToken(FTextToken::Create(LOCTEXT("LandscapeLayers_ChangingDataWithoutSettings", "This map contains landscape layer system content, modifying the landscape data will result in data loss when the map is reopened with Landscape Layer System settings on. Please enable Landscape Layer System settings before modifying the data.")));
-			FMessageLog("MapCheck").Open(EMessageSeverity::Warning);
-		}
-
 		this->Cache.SetCachedData(X1, Y1, X2, Y2, Data);
 		this->Cache.Flush();
 	}
@@ -727,14 +702,6 @@ public:
 			}
 		}
 
-		ALandscape* Landscape = this->LandscapeInfo->LandscapeActor.Get();
-
-		if (Landscape != nullptr && Landscape->HasLayersContent && !GetMutableDefault<UEditorExperimentalSettings>()->bLandscapeLayerSystem)
-		{
-			FMessageLog("MapCheck").Warning()->AddToken(FTextToken::Create(LOCTEXT("LandscapeLayers_ChangingDataWithoutSettings", "This map contains landscape layer system content, modifying the landscape data will result in data loss when the map is reopened with Landscape Layer System settings on. Please enable Landscape Layer System settings before modifying the data.")));
-			FMessageLog("MapCheck").Open(EMessageSeverity::Warning);
-		}
-
 		this->Cache.SetCachedData(X1, Y1, X2, Y2, Data, UISettings->PaintingRestriction);
 		this->Cache.Flush();
 	}
@@ -969,14 +936,6 @@ public:
 			}
 		}
 
-		ALandscape* Landscape = this->LandscapeInfo->LandscapeActor.Get();
-
-		if (Landscape != nullptr && Landscape->HasLayersContent && !GetMutableDefault<UEditorExperimentalSettings>()->bLandscapeLayerSystem)
-		{
-			FMessageLog("MapCheck").Warning()->AddToken(FTextToken::Create(LOCTEXT("LandscapeLayers_ChangingDataWithoutSettings", "This map contains landscape layer system content, modifying the landscape data will result in data loss when the map is reopened with Landscape Layer System settings on. Please enable Landscape Layer System settings before modifying the data.")));
-			FMessageLog("MapCheck").Open(EMessageSeverity::Warning);
-		}
-
 		this->Cache.SetCachedData(X1, Y1, X2, Y2, Data, UISettings->PaintingRestriction);
 		this->Cache.Flush();
 	}
@@ -1199,14 +1158,6 @@ public:
 			}
 		}
 
-		ALandscape* Landscape = this->LandscapeInfo->LandscapeActor.Get();
-
-		if (Landscape != nullptr && Landscape->HasLayersContent && !GetMutableDefault<UEditorExperimentalSettings>()->bLandscapeLayerSystem)
-		{
-			FMessageLog("MapCheck").Warning()->AddToken(FTextToken::Create(LOCTEXT("LandscapeLayers_ChangingDataWithoutSettings", "This map contains landscape layer system content, modifying the landscape data will result in data loss when the map is reopened with Landscape Layer System settings on. Please enable Landscape Layer System settings before modifying the data.")));
-			FMessageLog("MapCheck").Open(EMessageSeverity::Warning);
-		}
-
 		this->Cache.SetCachedData(X1, Y1, X2, Y2, Data, UISettings->PaintingRestriction);
 		this->Cache.Flush();
 	}
diff --git a/Engine/Source/Editor/LandscapeEditor/Private/LandscapeEdModeRampTool.cpp b/Engine/Source/Editor/LandscapeEditor/Private/LandscapeEdModeRampTool.cpp
--- a/Engine/Source/Editor/LandscapeEditor/Private/LandscapeEdModeRampTool.cpp
+++ b/Engine/Source/Editor/LandscapeEditor/Private/LandscapeEdModeRampTool.cpp
@@ -20,7 +20,6 @@
 #include "LandscapeHeightfieldCollisionComponent.h"
 #include "Raster.h"
 #include "Landscape.h"
-#include "Settings/EditorExperimentalSettings.h"
 #include "Misc/MessageDialog.h"
 
 #define LOCTEXT_NAMESPACE "Landscape"
@@ -582,7 +581,7 @@ public:
 			LandscapeEdit.SetHeightData(MinX, MinY, MaxX, MaxY, Data.GetData(), 0, true);
 			LandscapeEdit.Flush();
 
-			if (!GetMutableDefault<UEditorExperimentalSettings>()->bLandscapeLayerSystem)
+			if (!EdMode->HasLandscapeLayersContent())
 			{
 				TSet<ULandscapeComponent*> Components;
 				if (LandscapeEdit.GetComponentsInRegion(MinX, MinY, MaxX, MaxY, &Components))
diff --git a/Engine/Source/Editor/LandscapeEditor/Private/LandscapeEdModeTools.h b/Engine/Source/Editor/LandscapeEditor/Private/LandscapeEdModeTools.h
--- a/Engine/Source/Editor/LandscapeEditor/Private/LandscapeEdModeTools.h
+++ b/Engine/Source/Editor/LandscapeEditor/Private/LandscapeEdModeTools.h
@@ -15,7 +15,6 @@
 #include "InstancedFoliageActor.h"
 #include "VREditorInteractor.h"
 #include "AI/NavigationSystemBase.h"
-#include "Settings/EditorExperimentalSettings.h"
 #include "Landscape.h"
 
 // VR Editor
@@ -698,7 +697,7 @@ struct FXYOffsetmapAccessor
 			
 			ALandscapeProxy::InvalidateGeneratedComponentData(Components);
 
-			if (!GetMutableDefault<UEditorExperimentalSettings>()->bLandscapeLayerSystem)
+			if (!LandscapeEdit->HasLandscapeLayersContent())
 			{
 				for (ULandscapeComponent* Component : Components)
 				{
@@ -839,7 +838,7 @@ struct FFullWeightmapAccessor
 
 	~FFullWeightmapAccessor()
 	{
-		if (!GetMutableDefault<UEditorExperimentalSettings>()->bLandscapeLayerSystem)
+		if (!LandscapeEdit.HasLandscapeLayersContent())
 		{
 			// Recreate collision for modified components to update the physical materials
 			for (ULandscapeComponent* Component : ModifiedComponents)
@@ -878,7 +877,7 @@ struct FFullWeightmapAccessor
 		{
 			ALandscapeProxy::InvalidateGeneratedComponentData(Components);
 
-			if (!GetMutableDefault<UEditorExperimentalSettings>()->bLandscapeLayerSystem)
+			if (!LandscapeEdit.HasLandscapeLayersContent())
 			{
 				ModifiedComponents.Append(Components);
 			}
@@ -1129,7 +1128,7 @@ public:
 	{
 	}
 
-	virtual bool ShouldUpdateEditingLayer() const { return GetMutableDefault<UEditorExperimentalSettings>()->bLandscapeLayerSystem; }
+	virtual bool ShouldUpdateEditingLayer() const { return EdMode->CanHaveLandscapeLayersContent(); }
 
 	virtual ELandscapeLayerUpdateMode GetBeginToolContentUpdateFlag() const
 	{
diff --git a/Engine/Source/Editor/LandscapeEditor/Private/LandscapeEdModeXYOffset.cpp b/Engine/Source/Editor/LandscapeEditor/Private/LandscapeEdModeXYOffset.cpp
--- a/Engine/Source/Editor/LandscapeEditor/Private/LandscapeEdModeXYOffset.cpp
+++ b/Engine/Source/Editor/LandscapeEditor/Private/LandscapeEdModeXYOffset.cpp
@@ -10,7 +10,6 @@
 #include "Logging/TokenizedMessage.h"
 #include "Logging/MessageLog.h"
 #include "Misc/MapErrors.h"
-#include "Settings/EditorExperimentalSettings.h"
 
 #define LOCTEXT_NAMESPACE "LandscapeTools"
 namespace
@@ -540,15 +539,7 @@ public:
 				//XYOffsetVectorData(Index) = NewXYOffset(Index);
 			}
 		}
-
-		ALandscape* Landscape = LandscapeInfo->LandscapeActor.Get();
-
-		if (Landscape != nullptr && Landscape->HasLayersContent && !GetMutableDefault<UEditorExperimentalSettings>()->bLandscapeLayerSystem)
-		{
-			FMessageLog("MapCheck").Warning()->AddToken(FTextToken::Create(LOCTEXT("LandscapeLayers_ChangingDataWithoutSettings", "This map contains landscape layer system content, modifying the landscape data will result in data loss when the map is reopened with Landscape Layer System settings on. Please enable Landscape Layer System settings before modifying the data.")));
-			FMessageLog("MapCheck").Open(EMessageSeverity::Warning);
-		}
-
+		
 		// Apply to XYOffset Texture map and Height map
 		Cache.SetCachedData(X1, Y1, X2, Y2, XYOffsetVectorData);
 		Cache.Flush();
diff --git a/Engine/Source/Editor/LandscapeEditor/Private/LandscapeEditorDetailCustomization_Layers.cpp b/Engine/Source/Editor/LandscapeEditor/Private/LandscapeEditorDetailCustomization_Layers.cpp
--- a/Engine/Source/Editor/LandscapeEditor/Private/LandscapeEditorDetailCustomization_Layers.cpp
+++ b/Engine/Source/Editor/LandscapeEditor/Private/LandscapeEditorDetailCustomization_Layers.cpp
@@ -48,7 +48,6 @@
 #include "Widgets/Input/SNumericEntryBox.h"
 #include "Widgets/Text/SInlineEditableTextBlock.h"
 #include "LandscapeEditorCommands.h"
-#include "Settings/EditorExperimentalSettings.h"
 
 #define LOCTEXT_NAMESPACE "LandscapeEditor.Layers"
 
diff --git a/Engine/Source/Editor/LandscapeEditor/Private/LandscapeEditorDetailCustomization_NewLandscape.cpp b/Engine/Source/Editor/LandscapeEditor/Private/LandscapeEditorDetailCustomization_NewLandscape.cpp
--- a/Engine/Source/Editor/LandscapeEditor/Private/LandscapeEditorDetailCustomization_NewLandscape.cpp
+++ b/Engine/Source/Editor/LandscapeEditor/Private/LandscapeEditorDetailCustomization_NewLandscape.cpp
@@ -40,7 +40,6 @@
 #include "Framework/Application/SlateApplication.h"
 #include "Widgets/Input/SNumericEntryBox.h"
 #include "LandscapeDataAccess.h"
-#include "Settings/EditorExperimentalSettings.h"
 
 #define LOCTEXT_NAMESPACE "LandscapeEditor.NewLandscape"
 
diff --git a/Engine/Source/Editor/LandscapeEditor/Private/LandscapeEditorDetailCustomization_TargetLayers.cpp b/Engine/Source/Editor/LandscapeEditor/Private/LandscapeEditorDetailCustomization_TargetLayers.cpp
--- a/Engine/Source/Editor/LandscapeEditor/Private/LandscapeEditorDetailCustomization_TargetLayers.cpp
+++ b/Engine/Source/Editor/LandscapeEditor/Private/LandscapeEditorDetailCustomization_TargetLayers.cpp
@@ -42,7 +42,6 @@
 #include "IDetailGroup.h"
 #include "Widgets/SBoxPanel.h"
 #include "Editor/EditorStyle/Private/SlateEditorStyle.h"
-#include "Settings/EditorExperimentalSettings.h"
 
 #define LOCTEXT_NAMESPACE "LandscapeEditor.TargetLayers"
 
@@ -1106,8 +1105,8 @@ void FLandscapeEditorCustomNodeBuilder_TargetLayers::FillEmptyLayers(ULandscapeI
 	if (LandscapeEdMode)
 	{
 		FLandscapeEditDataInterface LandscapeEdit(LandscapeInfo);
-		
-		if (GetMutableDefault<UEditorExperimentalSettings>()->bLandscapeLayerSystem)
+
+		if (LandscapeEdMode->CanHaveLandscapeLayersContent())
 		{
 			if (LandscapeEdMode->NeedToFillEmptyMaterialLayers())
 			{
@@ -1427,7 +1426,8 @@ EVisibility FLandscapeEditorCustomNodeBuilder_TargetLayers::GetDebugModeLayerUsa
 
 EVisibility FLandscapeEditorCustomNodeBuilder_TargetLayers::GetLayersSubstractiveBlendVisibility(const TSharedRef<FLandscapeTargetListInfo> Target)
 {
-	if (GetMutableDefault<UEditorExperimentalSettings>()->bLandscapeLayerSystem && Target->TargetType != ELandscapeToolTargetType::Heightmap && Target->LayerInfoObj.IsValid())
+	FEdModeLandscape* LandscapeEdMode = GetEditorMode();
+	if (LandscapeEdMode != nullptr && LandscapeEdMode->CanHaveLandscapeLayersContent() && Target->TargetType != ELandscapeToolTargetType::Heightmap && Target->LayerInfoObj.IsValid())
 	{
 		return EVisibility::Visible;
 	}
diff --git a/Engine/Source/Editor/LandscapeEditor/Private/LandscapeEditorDetails.cpp b/Engine/Source/Editor/LandscapeEditor/Private/LandscapeEditorDetails.cpp
--- a/Engine/Source/Editor/LandscapeEditor/Private/LandscapeEditorDetails.cpp
+++ b/Engine/Source/Editor/LandscapeEditor/Private/LandscapeEditorDetails.cpp
@@ -21,7 +21,6 @@
 #include "LandscapeEditorCommands.h"
 #include "LandscapeEditorDetailWidgets.h"
 #include "LandscapeEditorDetailCustomization_LayersBrushStack.h"
-#include "Settings/EditorExperimentalSettings.h"
 #include "Landscape.h"
 
 #define LOCTEXT_NAMESPACE "LandscapeEditor"
@@ -134,7 +133,7 @@ void FLandscapeEditorDetails::CustomizeDetails(IDetailLayoutBuilder& DetailBuild
 	Customization_AlphaBrush = MakeShareable(new FLandscapeEditorDetailCustomization_AlphaBrush);
 	Customization_AlphaBrush->CustomizeDetails(DetailBuilder);
 
-	if (GetMutableDefault<UEditorExperimentalSettings>()->bLandscapeLayerSystem)
+	if (LandscapeEdMode->CanHaveLandscapeLayersContent())
 	{
 		// Layers
 		Customization_Layers = MakeShareable(new FLandscapeEditorDetailCustomization_Layers);
@@ -155,6 +154,8 @@ FText FLandscapeEditorDetails::GetLocalizedName(FString Name)
 	static bool bInitialized = false;
 	if (!bInitialized)
 	{
+		FEdModeLandscape* LandscapeEdMode = GetEditorMode();
+
 		bInitialized = true;
 		LOCTEXT("ToolSet_NewLandscape", "New Landscape");
 		LOCTEXT("ToolSet_ResizeLandscape", "Change Component Size");
@@ -170,7 +171,7 @@ FText FLandscapeEditorDetails::GetLocalizedName(FString Name)
 		LOCTEXT("ToolSet_Retopologize", "Retopologize");
 		LOCTEXT("ToolSet_Visibility", "Visibility");
 
-		if (GetMutableDefault<UEditorExperimentalSettings>()->bLandscapeLayerSystem)
+		if (LandscapeEdMode->CanHaveLandscapeLayersContent())
 		{
 			LOCTEXT("ToolSet_BPCustom", "Blueprint Custom");
 		}
@@ -319,7 +320,7 @@ TSharedRef<SWidget> FLandscapeEditorDetails::GetToolSelector()
 			MenuBuilder.BeginSection(NAME_None, LOCTEXT("SculptToolsTitle", "Sculpting Tools"));
 			MenuBuilder.AddToolButton(NameToCommandMap.FindChecked("Tool_Sculpt"), NAME_None, LOCTEXT("Tool.Sculpt", "Sculpt"), LOCTEXT("Tool.Sculpt.Tooltip", "Sculpt height data.\nCtrl+Click to Raise, Ctrl+Shift+Click to lower"));
 			
-			if (GetMutableDefault<UEditorExperimentalSettings>()->bLandscapeLayerSystem)
+			if (LandscapeEdMode->CanHaveLandscapeLayersContent())
 			{
 				MenuBuilder.AddToolButton(NameToCommandMap.FindChecked("Tool_Erase"), NAME_None, LOCTEXT("Tool.Erase", "Erase"), LOCTEXT("Tool.Erase.Tooltip", "Erase height data."));
 			}
@@ -333,7 +334,7 @@ TSharedRef<SWidget> FLandscapeEditorDetails::GetToolSelector()
 			MenuBuilder.AddToolButton(NameToCommandMap.FindChecked("Tool_Retopologize"), NAME_None, LOCTEXT("Tool.Retopologize", "Retopologize"), LOCTEXT("Tool.Retopologize.Tooltip", "Automatically adjusts landscape vertices with an X/Y offset map to improve vertex density on cliffs, reducing texture stretching.\nNote: An X/Y offset map makes the landscape slower to render and paint on with other tools, so only use if needed"));
 			MenuBuilder.AddToolButton(NameToCommandMap.FindChecked("Tool_Visibility"), NAME_None, LOCTEXT("Tool.Visibility", "Visibility"), LOCTEXT("Tool.Visibility.Tooltip", "Mask out individual quads in the landscape, leaving a hole."));
 
-			if (GetMutableDefault<UEditorExperimentalSettings>()->bLandscapeLayerSystem)
+			if (LandscapeEdMode->CanHaveLandscapeLayersContent())
 			{
 				MenuBuilder.AddToolButton(NameToCommandMap.FindChecked("Tool_BPCustom"), NAME_None, LOCTEXT("Tool.SculptBPCustom", "Blueprint Custom"), LOCTEXT("Tool.SculptBPCustom.Tooltip", "Custom sculpting tools created using Blueprint."));
 			}
@@ -355,7 +356,7 @@ TSharedRef<SWidget> FLandscapeEditorDetails::GetToolSelector()
 			MenuBuilder.AddToolButton(NameToCommandMap.FindChecked("Tool_Flatten"), NAME_None, LOCTEXT("Tool.Flatten", "Flatten"), LOCTEXT("Tool.Flatten.Tooltip", "Flattens an area of heightmap or blend layer"));
 			MenuBuilder.AddToolButton(NameToCommandMap.FindChecked("Tool_Noise"), NAME_None, LOCTEXT("Tool.Noise", "Noise"), LOCTEXT("Tool.Noise.Tooltip", "Adds noise to the heightmap or blend layer"));
 
-			if (GetMutableDefault<UEditorExperimentalSettings>()->bLandscapeLayerSystem)
+			if (LandscapeEdMode->CanHaveLandscapeLayersContent())
 			{
 				MenuBuilder.AddToolButton(NameToCommandMap.FindChecked("Tool_BPCustom"), NAME_None, LOCTEXT("Tool.PaintBPCustom", "Blueprint Custom"), LOCTEXT("Tool.PaintBPCustom.Tooltip", "Custom painting tools created using Blueprint."));
 			}
diff --git a/Engine/Source/Editor/LandscapeEditor/Private/LandscapeEditorModule.cpp b/Engine/Source/Editor/LandscapeEditor/Private/LandscapeEditorModule.cpp
--- a/Engine/Source/Editor/LandscapeEditor/Private/LandscapeEditorModule.cpp
+++ b/Engine/Source/Editor/LandscapeEditor/Private/LandscapeEditorModule.cpp
@@ -162,7 +162,10 @@ public:
 					{
 						InMenuBuilder.AddMenuEntry(LandscapeActions.ViewModeLayerUsage, NAME_None, LOCTEXT("LandscapeViewModeLayerUsage", "Layer Usage"));
 						InMenuBuilder.AddMenuEntry(LandscapeActions.ViewModeLayerDebug, NAME_None, LOCTEXT("LandscapeViewModeLayerDebug", "Layer Debug"));
-						if (GetMutableDefault<UEditorExperimentalSettings>()->bLandscapeLayerSystem)
+
+						FEdModeLandscape* LandscapeMode = (FEdModeLandscape*)GLevelEditorModeTools().GetActiveMode(FBuiltinEditorModes::EM_Landscape);
+
+						if (LandscapeMode->CanHaveLandscapeLayersContent())
 						{
 							InMenuBuilder.AddMenuEntry(LandscapeActions.ViewModeLayerContribution, NAME_None, LOCTEXT("LandscapeViewModeLayerContribution", "Layer Contribution"));
 						}
diff --git a/Engine/Source/Editor/LandscapeEditor/Private/SLandscapeEditor.cpp b/Engine/Source/Editor/LandscapeEditor/Private/SLandscapeEditor.cpp
--- a/Engine/Source/Editor/LandscapeEditor/Private/SLandscapeEditor.cpp
+++ b/Engine/Source/Editor/LandscapeEditor/Private/SLandscapeEditor.cpp
@@ -17,7 +17,6 @@
 #include "IDetailsView.h"
 #include "PropertyEditorModule.h"
 #include "IIntroTutorials.h"
-#include "Settings/EditorExperimentalSettings.h"
 
 #define LOCTEXT_NAMESPACE "LandscapeEditor"
 
@@ -345,7 +344,7 @@ void SLandscapeEditor::Construct(const FArguments& InArgs, TSharedRef<FLandscape
 			.Padding(0, 0, 0, 5)
 			[
 				SNew(SMultiLineEditableTextBox)
-				.Visibility_Lambda([]() { return GetMutableDefault<UEditorExperimentalSettings>()->bLandscapeLayerSystem? EVisibility::Visible : EVisibility::Collapsed; })
+				.Visibility_Lambda([=]() { return LandscapeEdMode->CanHaveLandscapeLayersContent() ? EVisibility::Visible : EVisibility::Collapsed; })
 				.IsReadOnly(true)
 				.Font(FCoreStyle::GetDefaultFontStyle("Bold", 10))
 				.Justification(ETextJustify::Center)
diff --git a/Engine/Source/Runtime/Engine/Classes/GameFramework/WorldSettings.h b/Engine/Source/Runtime/Engine/Classes/GameFramework/WorldSettings.h
--- a/Engine/Source/Runtime/Engine/Classes/GameFramework/WorldSettings.h
+++ b/Engine/Source/Runtime/Engine/Classes/GameFramework/WorldSettings.h
@@ -464,6 +464,10 @@ public:
 	UPROPERTY(EditAnywhere, BlueprintReadOnly, Category=World)
 	uint8 bEnableWorldComposition:1;
 
+	/** Enables landscape layers system for this world */
+	UPROPERTY(EditAnywhere, BlueprintReadOnly, Category = World)
+	uint8 bEnableLandscapeLayers : 1;
+
 	/**
 	 * Enables client-side streaming volumes instead of server-side.
 	 * Expected usage scenario: server has all streaming levels always loaded, clients independently stream levels in/out based on streaming volumes.
@@ -714,7 +718,22 @@ public:
 	UPROPERTY()
 	TArray<UAssetUserData*> AssetUserData;
 
+public:
+	/**
+	 * Returns an event delegate that is executed when a setting has changed.
+	 *
+	 * @return The delegate.
+	 */
+	DECLARE_EVENT_OneParam(AWorldSettings, FSettingChangedEvent, FName /*PropertyName*/);
+	FSettingChangedEvent& OnSettingChanged()
+	{
+		return SettingChangedEvent;
+	}
+
 private:
+	// Holds an event delegate that is executed when a setting has changed.
+	FSettingChangedEvent SettingChangedEvent;
+
 	// If paused, PlayerState of person pausing the game.
 	UPROPERTY(transient, replicated)
 	class APlayerState* PauserPlayerState;
diff --git a/Engine/Source/Runtime/Engine/Private/WorldSettings.cpp b/Engine/Source/Runtime/Engine/Private/WorldSettings.cpp
--- a/Engine/Source/Runtime/Engine/Private/WorldSettings.cpp
+++ b/Engine/Source/Runtime/Engine/Private/WorldSettings.cpp
@@ -35,6 +35,8 @@
 #include "HierarchicalLOD.h"
 #include "IMeshMergeUtilities.h"
 #include "MeshMergeModule.h"
+#include "Settings/EditorExperimentalSettings.h"
+#include "Landscape.h"
 #endif 
 
 #define LOCTEXT_NAMESPACE "ErrorChecking"
@@ -71,6 +73,7 @@ PRAGMA_ENABLE_DEPRECATION_WARNINGS
 	NavigationSystemConfig = nullptr;
 	bEnableAISystem = true;
 	bEnableWorldComposition = false;
+	bEnableLandscapeLayers = false;
 	bEnableWorldOriginRebasing = false;
 #if WITH_EDITORONLY_DATA	
 	bEnableHierarchicalLODSystem = false;
@@ -605,6 +608,22 @@ void AWorldSettings::PostEditChangeProperty(FPropertyChangedEvent& PropertyChang
 				bEnableWorldComposition = false;
 			}
 		}
+		else if (PropertyName == GET_MEMBER_NAME_CHECKED(AWorldSettings, bEnableLandscapeLayers))
+		{
+			if (ALandscape::ChangeLandscapeLayersStateEvent.IsBound())
+			{
+				bool WasChanged = ALandscape::ChangeLandscapeLayersStateEvent.Execute(GetWorld(), bEnableLandscapeLayers);
+
+				if (!WasChanged)
+				{
+					bEnableLandscapeLayers = !bEnableLandscapeLayers;
+				}
+			}
+			else
+			{
+				bEnableLandscapeLayers = false;
+			}
+		}
 		else if (PropertyName == GET_MEMBER_NAME_CHECKED(AWorldSettings, NavigationSystemConfig))
 		{
 			UWorld* World = GetWorld();
@@ -685,8 +704,12 @@ void AWorldSettings::PostEditChangeProperty(FPropertyChangedEvent& PropertyChang
 	}
 
 	Super::PostEditChangeProperty(PropertyChangedEvent);
-}
 
+	if (PropertyThatChanged)
+	{
+		SettingChangedEvent.Broadcast(PropertyThatChanged->GetFName());
+	}
+}
 
 void UHierarchicalLODSetup::PostEditChangeProperty(struct FPropertyChangedEvent& PropertyChangedEvent)
 {
diff --git a/Engine/Source/Runtime/Landscape/Classes/Landscape.h b/Engine/Source/Runtime/Landscape/Classes/Landscape.h
--- a/Engine/Source/Runtime/Landscape/Classes/Landscape.h
+++ b/Engine/Source/Runtime/Landscape/Classes/Landscape.h
@@ -6,6 +6,7 @@
 #include "UObject/ObjectMacros.h"
 #include "LandscapeProxy.h"
 #include "LandscapeBPCustomBrush.h"
+#include "Delegates/DelegateCombinations.h"
 
 #include "Landscape.generated.h"
 
@@ -187,6 +188,8 @@ struct FLandscapeLayer
 	TMap<ULandscapeLayerInfoObject*, bool> WeightmapLayerAllocationBlend; // True -> Substractive, False -> Additive
 };
 
+DECLARE_DELEGATE_RetVal_TwoParams(bool, FChangeLandscapeLayersState, UWorld*, bool);
+
 UCLASS(MinimalAPI, showcategories=(Display, Movement, Collision, Lighting, LOD, Input), hidecategories=(Mobility))
 class ALandscape : public ALandscapeProxy
 {
@@ -236,6 +239,7 @@ public:
 #if WITH_EDITOR
 	LANDSCAPE_API void RegisterLandscapeEdMode(ILandscapeEdModeInterface* InLandscapeEdMode) { LandscapeEdMode = InLandscapeEdMode; }
 	LANDSCAPE_API void UnregisterLandscapeEdMode() { LandscapeEdMode = nullptr; }
+	LANDSCAPE_API virtual bool HasLayersContent() const override;
 	LANDSCAPE_API void RequestLayersInitialization(bool bInRequestContentUpdate = true);
 	LANDSCAPE_API void RequestLayersContentUpdateForceAll(ELandscapeLayerUpdateMode InModeMask = ELandscapeLayerUpdateMode::Update_All);
 	LANDSCAPE_API void RequestLayersContentUpdate(ELandscapeLayerUpdateMode InModeMask);
@@ -244,6 +248,7 @@ public:
 	LANDSCAPE_API void CreateLayer(FName InName = NAME_None);
 	LANDSCAPE_API void CreateDefaultLayer();
 	LANDSCAPE_API void CopyOldDataToDefaultLayer();
+	LANDSCAPE_API void CopyOldDataToDefaultLayer(ALandscapeProxy* Proxy);
 	LANDSCAPE_API void AddLayersToProxy(ALandscapeProxy* InProxy);
 	LANDSCAPE_API TMap<UTexture2D*, TArray<ULandscapeComponent*>> GenerateComponentsPerHeightmaps() const;
 	LANDSCAPE_API FIntPoint ComputeComponentCounts() const;
@@ -262,6 +267,7 @@ public:
 	LANDSCAPE_API void ClearLayer(int32 InLayerIndex, bool bInUpdateCollision = true);
 	LANDSCAPE_API void ClearLayer(const FGuid& InLayerGuid, bool bInUpdateCollision = true);
 	LANDSCAPE_API void DeleteLayer(int32 InLayerIndex);
+	LANDSCAPE_API void DeleteLayers();
 	LANDSCAPE_API void SetEditingLayer(const FGuid& InLayerGuid = FGuid());
 	LANDSCAPE_API const FGuid& GetEditingLayer() const;
 	LANDSCAPE_API bool IsMaxLayersReached() const;
@@ -286,10 +292,14 @@ public:
 	
 	LANDSCAPE_API void OnPreSave();
 
+	void ReleaseLayersRenderingResource();
+
+	LANDSCAPE_API static void RegisterChangeLandscapeLayersStateDelegate();
+	LANDSCAPE_API static void UnregisterChangeLandscapeLayersStateDelegate();
+
 private:
 	void TickLayers(float DeltaTime, ELevelTick TickType, FActorTickFunction& ThisTickFunction);
 	void CreateLayersRenderingResource();
-	void ReleaseLayersRenderingResource();
 	void UpdateLayersContent(bool bInWaitForStreaming = false);
 	void MonitorShaderCompilation();
 	void MonitorLandscapeEdModeChanges();
@@ -345,6 +355,8 @@ private:
 	void PrintLayersDebugTextureResource(const FString& InContext, FTextureResource* InTextureResource, uint8 InMipRender = 0, bool InOutputHeight = true, bool InOutputNormals = false) const;
 	void PrintLayersDebugHeightData(const FString& InContext, const TArray<FColor>& InHeightmapData, const FIntPoint& InDataSize, uint8 InMipRender, bool InOutputNormals = false) const;
 	void PrintLayersDebugWeightData(const FString& InContext, const TArray<FColor>& InWeightmapData, const FIntPoint& InDataSize, uint8 InMipRender) const;
+
+	LANDSCAPE_API static bool ChangeLandscapeLayersState(UWorld* InWorld, bool bInEnable);
 #endif
 
 public:
@@ -366,8 +378,8 @@ public:
 	UPROPERTY(Transient)
 	TArray<UTextureRenderTarget2D*> WeightmapRTList;
 
-	UPROPERTY(Transient)
-	bool bInitializedWithFlagExperimentalLandscapeLayers;
+	/** Event that will be triggered if the Landscape Layers settings from the WorldSettings is changed*/
+	LANDSCAPE_API static FChangeLandscapeLayersState ChangeLandscapeLayersStateEvent;
 
 private:
 	/** Provides information from LandscapeEdMode */
diff --git a/Engine/Source/Runtime/Landscape/Classes/LandscapeComponent.h b/Engine/Source/Runtime/Landscape/Classes/LandscapeComponent.h
--- a/Engine/Source/Runtime/Landscape/Classes/LandscapeComponent.h
+++ b/Engine/Source/Runtime/Landscape/Classes/LandscapeComponent.h
@@ -601,6 +601,7 @@ public:
 	LANDSCAPE_API TArray<ULandscapeWeightmapUsage*>& GetWeightmapTexturesUsage(bool InReturnEditingWeightmap = false);
 	LANDSCAPE_API const TArray<ULandscapeWeightmapUsage*>& GetWeightmapTexturesUsage(bool InReturnEditingWeightmap = false) const;
 
+	LANDSCAPE_API bool HasLayersData() const;
 	LANDSCAPE_API const FLandscapeLayerComponentData* GetLayerData(const FGuid& InLayerGuid) const;
 	LANDSCAPE_API FLandscapeLayerComponentData* GetLayerData(const FGuid& InLayerGuid);
 	LANDSCAPE_API void AddLayerData(const FGuid& InLayerGuid, const FLandscapeLayerComponentData& InData);
diff --git a/Engine/Source/Runtime/Landscape/Classes/LandscapeProxy.h b/Engine/Source/Runtime/Landscape/Classes/LandscapeProxy.h
--- a/Engine/Source/Runtime/Landscape/Classes/LandscapeProxy.h
+++ b/Engine/Source/Runtime/Landscape/Classes/LandscapeProxy.h
@@ -596,9 +596,6 @@ public:
 	UPROPERTY()
 	TArray<FLandscapeEditorLayerSettings> EditorLayerSettings;
 
-	UPROPERTY()
-	bool HasLayersContent;
-
 	TMap<UTexture2D*, FLandscapeLayersTexture2DCPUReadBackResource*> HeightmapsCPUReadBack;
 	TMap<UTexture2D*, FLandscapeLayersTexture2DCPUReadBackResource*> WeightmapsCPUReadBack;
 	FRenderCommandFence ReleaseResourceFence;
@@ -634,6 +631,10 @@ public:
 	UPROPERTY(EditAnywhere, Category = HierarchicalLOD)
 	bool bUseLandscapeForCullingInvisibleHLODVertices;
 
+	/** Flag that tell if we have some layers content **/
+	UPROPERTY()
+	bool bHasLayersContent;
+
 public:
 
 #if WITH_EDITOR
@@ -942,6 +943,14 @@ public:
 	DECLARE_EVENT(ALandscape, FLandscapeMaterialChangedDelegate);
 	FLandscapeMaterialChangedDelegate& OnMaterialChangedDelegate() { return LandscapeMaterialChangedDelegate; }
 
+	/** Will tell if the landscape proxy as some content related to the layer system */
+	LANDSCAPE_API virtual bool HasLayersContent() const;
+	
+	/** Will tell if the landscape proxy can have some content related to the layer system */
+	LANDSCAPE_API bool CanHaveLayersContent() const;
+
+	void UpdateCachedHasLayersContent(bool InCheckComponentDataIntegrity = false);
+
 protected:
 	friend class ALandscape;
 
diff --git a/Engine/Source/Runtime/Landscape/Private/Landscape.cpp b/Engine/Source/Runtime/Landscape/Private/Landscape.cpp
--- a/Engine/Source/Runtime/Landscape/Private/Landscape.cpp
+++ b/Engine/Source/Runtime/Landscape/Private/Landscape.cpp
@@ -54,7 +54,6 @@ Landscape.cpp: Terrain rendering
 #if WITH_EDITOR
 #include "LandscapeEdit.h"
 #include "MaterialUtilities.h"
-#include "Settings/EditorExperimentalSettings.h"
 #include "Editor.h"
 #include "Algo/Transform.h"
 #endif
@@ -1033,7 +1032,7 @@ ALandscapeProxy::ALandscapeProxy(const FObjectInitializer& ObjectInitializer)
 	bGenerateOverlapEvents = false;
 #if WITH_EDITORONLY_DATA
 	MaxPaintedLayersPerComponent = 0;
-	HasLayersContent = false;
+	bHasLayersContent = false;
 #endif
 
 #if WITH_EDITOR
@@ -1084,7 +1083,6 @@ ALandscape::ALandscape(const FObjectInitializer& ObjectInitializer)
 {
 #if WITH_EDITORONLY_DATA
 	bLockLocation = false;
-	bInitializedWithFlagExperimentalLandscapeLayers = GetMutableDefault<UEditorExperimentalSettings>()->bLandscapeLayerSystem;
 	WasCompilingShaders = false;
 	LayerContentUpdateModes = 0;
 	CombinedLayersWeightmapAllMaterialLayersResource = nullptr;
@@ -1093,6 +1091,12 @@ ALandscape::ALandscape(const FObjectInitializer& ObjectInitializer)
 	WeightmapScratchPackLayerTextureResource = nullptr;
 	bLandscapeLayersAreInitialized = false;
 	LandscapeEdMode = nullptr;
+
+	if (!HasAnyFlags(RF_ArchetypeObject | RF_ClassDefaultObject))
+	{
+		ALandscape::RegisterChangeLandscapeLayersStateDelegate();
+	}
+	
 #endif // WITH_EDITORONLY_DATA
 }
 
@@ -1444,7 +1448,7 @@ void ULandscapeComponent::OnUnregister()
 UTexture2D* ULandscapeComponent::GetHeightmap(bool InReturnEditingHeightmap) const
 {
 #if WITH_EDITORONLY_DATA
-	if (InReturnEditingHeightmap && GetMutableDefault<UEditorExperimentalSettings>()->bLandscapeLayerSystem)
+	if (InReturnEditingHeightmap && GetLandscapeProxy()->HasLayersContent())
 	{
 		const FLandscapeLayerComponentData* EditingLayer = GetEditingLayer();
 
@@ -1461,7 +1465,7 @@ UTexture2D* ULandscapeComponent::GetHeightmap(bool InReturnEditingHeightmap) con
 const TArray<UTexture2D*>& ULandscapeComponent::GetWeightmapTextures(bool InReturnEditingWeightmap) const
 {
 #if WITH_EDITORONLY_DATA
-	if (InReturnEditingWeightmap && GetMutableDefault<UEditorExperimentalSettings>()->bLandscapeLayerSystem)
+	if (InReturnEditingWeightmap && GetLandscapeProxy()->HasLayersContent())
 	{
 		const FLandscapeLayerComponentData* EditingLayer = GetEditingLayer();
 
@@ -1478,7 +1482,7 @@ const TArray<UTexture2D*>& ULandscapeComponent::GetWeightmapTextures(bool InRetu
 TArray<UTexture2D*>& ULandscapeComponent::GetWeightmapTextures(bool InReturnEditingWeightmap)
 {
 #if WITH_EDITORONLY_DATA
-	if (InReturnEditingWeightmap && GetMutableDefault<UEditorExperimentalSettings>()->bLandscapeLayerSystem)
+	if (InReturnEditingWeightmap && GetLandscapeProxy()->HasLayersContent())
 	{
 		FLandscapeLayerComponentData* EditingLayer = GetEditingLayer();
 
@@ -1495,7 +1499,7 @@ TArray<UTexture2D*>& ULandscapeComponent::GetWeightmapTextures(bool InReturnEdit
 const TArray<FWeightmapLayerAllocationInfo>& ULandscapeComponent::GetWeightmapLayerAllocations(bool InReturnEditingWeightmap) const
 {
 #if WITH_EDITORONLY_DATA
-	if (InReturnEditingWeightmap && GetMutableDefault<UEditorExperimentalSettings>()->bLandscapeLayerSystem)
+	if (InReturnEditingWeightmap && GetLandscapeProxy()->HasLayersContent())
 	{
 		const FLandscapeLayerComponentData* EditingLayer = GetEditingLayer();
 
@@ -1512,7 +1516,7 @@ const TArray<FWeightmapLayerAllocationInfo>& ULandscapeComponent::GetWeightmapLa
 TArray<FWeightmapLayerAllocationInfo>& ULandscapeComponent::GetWeightmapLayerAllocations(bool InReturnEditingWeightmap)
 {
 #if WITH_EDITORONLY_DATA
-	if (InReturnEditingWeightmap && GetMutableDefault<UEditorExperimentalSettings>()->bLandscapeLayerSystem)
+	if (InReturnEditingWeightmap && GetLandscapeProxy()->HasLayersContent())
 	{
 		FLandscapeLayerComponentData* EditingLayer = GetEditingLayer();
 
@@ -1546,6 +1550,11 @@ FGuid ULandscapeComponent::GetEditingLayerGUID() const
 	return Landscape != nullptr ? Landscape->GetEditingLayer() : FGuid();
 }
 
+bool ULandscapeComponent::HasLayersData() const
+{
+	return LayersData.Num() > 0;
+}
+
 const FLandscapeLayerComponentData* ULandscapeComponent::GetLayerData(const FGuid& InLayerGuid) const 
 {
 	return LayersData.Find(InLayerGuid);
@@ -1686,7 +1695,7 @@ void ULandscapeComponent::SetWeightmapLayerAllocations(const TArray<FWeightmapLa
 TArray<ULandscapeWeightmapUsage*>& ULandscapeComponent::GetWeightmapTexturesUsage(bool InReturnEditingWeightmap)
 {
 #if WITH_EDITORONLY_DATA
-	if (InReturnEditingWeightmap && GetMutableDefault<UEditorExperimentalSettings>()->bLandscapeLayerSystem)
+	if (InReturnEditingWeightmap && GetLandscapeProxy()->HasLayersContent())
 	{
 		FLandscapeLayerComponentData* EditingLayer = GetEditingLayer();
 
@@ -1703,7 +1712,7 @@ TArray<ULandscapeWeightmapUsage*>& ULandscapeComponent::GetWeightmapTexturesUsag
 const TArray<ULandscapeWeightmapUsage*>& ULandscapeComponent::GetWeightmapTexturesUsage(bool InReturnEditingWeightmap) const
 {
 #if WITH_EDITORONLY_DATA
-	if (InReturnEditingWeightmap && GetMutableDefault<UEditorExperimentalSettings>()->bLandscapeLayerSystem)
+	if (InReturnEditingWeightmap && GetLandscapeProxy()->HasLayersContent())
 	{
 		const FLandscapeLayerComponentData* EditingLayer = GetEditingLayer();
 
@@ -1899,11 +1908,6 @@ void ALandscapeProxy::PreSave(const class ITargetPlatform* TargetPlatform)
 	{
 		bHasLandscapeGrass = LandscapeComponents.ContainsByPredicate([](ULandscapeComponent* Component) { return Component->MaterialHasGrass(); });
 	}
-
-	if (GetMutableDefault<UEditorExperimentalSettings>()->bLandscapeLayerSystem)
-	{
-		HasLayersContent = true;
-	}
 #endif
 }
 
@@ -2266,17 +2270,25 @@ void ALandscapeProxy::PostLoad()
 
 	if (GIsEditor && !GetWorld()->IsGameWorld())
 	{
+		UpdateCachedHasLayersContent(true);
+
+		// Cache the value at this point as RegisterActor might create/destroy layers content if there was a mismatch between landscape & proxy
+		bool bNeedOldDataMigration = !HasLayersContent() && CanHaveLayersContent();
+
 		ULandscapeInfo* LandscapeInfo = CreateLandscapeInfo();
 		LandscapeInfo->RegisterActor(this, true);
 
 		FixupWeightmaps();
+
+		if (bNeedOldDataMigration && LandscapeInfo->LandscapeActor.IsValid())
+		{
+			LandscapeInfo->LandscapeActor.Get()->CopyOldDataToDefaultLayer(this);
+		}
 	}
 
 	// track feature level change to flush grass cache
 	FOnFeatureLevelChanged::FDelegate FeatureLevelChangedDelegate = FOnFeatureLevelChanged::FDelegate::CreateUObject(this, &ALandscapeProxy::OnFeatureLevelChanged);
 	FeatureLevelChangedDelegateHandle = GetWorld()->AddOnFeatureLevelChangedHandler(FeatureLevelChangedDelegate);
-
-
 #endif
 }
 
@@ -2446,7 +2458,6 @@ void ALandscapeProxy::FixupSharedData(ALandscape* Landscape)
 		bUpdated |= AddLayer(Layer.Guid);
 	}
 
-
 	if (bUpdated)
 	{
 		MarkPackageDirty();
@@ -2758,7 +2769,7 @@ void ULandscapeInfo::RegisterActor(ALandscapeProxy* Proxy, bool bMapCheck)
 		StreamingProxy->FixupSharedData(LandscapeActor.Get());
 	}
 
-	if (GetMutableDefault<UEditorExperimentalSettings>()->bLandscapeLayerSystem && LandscapeActor)
+	if (LandscapeActor && LandscapeActor->HasLayersContent())
 	{
 		// Force update rendering resources
 		const bool bInRequestContentUpdate = false;
diff --git a/Engine/Source/Runtime/Landscape/Private/LandscapeCollision.cpp b/Engine/Source/Runtime/Landscape/Private/LandscapeCollision.cpp
--- a/Engine/Source/Runtime/Landscape/Private/LandscapeCollision.cpp
+++ b/Engine/Source/Runtime/Landscape/Private/LandscapeCollision.cpp
@@ -45,10 +45,6 @@
 #include "Physics/PhysicsInterfaceCore.h"
 #include "Physics/PhysicsInterfaceUtils.h"
 
-#if WITH_EDITOR
-#include "Settings/EditorExperimentalSettings.h"
-#endif
-
 #if WITH_EDITOR && WITH_PHYSX
 	#include "Physics/IPhysXCooking.h"
 #endif
@@ -1464,7 +1460,7 @@ void ULandscapeHeightfieldCollisionComponent::PostEditImport()
 {
 	Super::PostEditImport();
 
-	if (!GetMutableDefault<UEditorExperimentalSettings>()->bLandscapeLayerSystem)
+	if (!GetLandscapeProxy()->HasLayersContent())
 	{
 		// Reinitialize physics after paste
 		if (CollisionSizeQuads > 0)
@@ -1479,7 +1475,7 @@ void ULandscapeHeightfieldCollisionComponent::PostEditUndo()
 	Super::PostEditUndo();
 
     // Landscape Layers are updates are delayed and done in  ALandscape::TickLayers
-	if (!GetMutableDefault<UEditorExperimentalSettings>()->bLandscapeLayerSystem)
+	if (!GetLandscapeProxy()->HasLayersContent())
 	{
 		// Reinitialize physics after undo
 		if (CollisionSizeQuads > 0)
diff --git a/Engine/Source/Runtime/Landscape/Private/LandscapeEdit.cpp b/Engine/Source/Runtime/Landscape/Private/LandscapeEdit.cpp
--- a/Engine/Source/Runtime/Landscape/Private/LandscapeEdit.cpp
+++ b/Engine/Source/Runtime/Landscape/Private/LandscapeEdit.cpp
@@ -55,7 +55,6 @@ LandscapeEdit.cpp: Landscape editing
 #include "Engine/TextureRenderTarget2D.h"
 #include "ScopedTransaction.h"
 #include "Editor.h"
-#include "Settings/EditorExperimentalSettings.h"
 #endif
 #include "Algo/Count.h"
 #include "Serialization/MemoryWriter.h"
@@ -122,7 +121,7 @@ void ULandscapeComponent::UpdateCachedBounds()
 	if (HFCollisionComponent)
 	{
         // In Landscape Layers the Collision Component is slave and doesn't need to be transacted
-		if (!GetMutableDefault<UEditorExperimentalSettings>()->bLandscapeLayerSystem)
+		if (!GetLandscapeProxy()->HasLayersContent())
 		{
 			HFCollisionComponent->Modify();
 		}
@@ -560,7 +559,7 @@ void ULandscapeComponent::PostEditUndo()
 {
 	if (!IsPendingKill())
 	{
-		if (!GetMutableDefault<UEditorExperimentalSettings>()->bLandscapeLayerSystem)
+		if (!GetLandscapeProxy()->HasLayersContent())
 		{
 			UpdateMaterialInstances();
 		}
@@ -571,14 +570,14 @@ void ULandscapeComponent::PostEditUndo()
 	if (!IsPendingKill())
 	{
 		EditToolRenderData.UpdateSelectionMaterial(EditToolRenderData.SelectedType, this);
-		if (!GetMutableDefault<UEditorExperimentalSettings>()->bLandscapeLayerSystem)
+		if (!GetLandscapeProxy()->HasLayersContent())
 		{
 			EditToolRenderData.UpdateDebugColorMaterial(this);
             UpdateEditToolRenderData();
 		}	
 	}
 		
-	if (GetMutableDefault<UEditorExperimentalSettings>()->bLandscapeLayerSystem)
+	if (GetLandscapeProxy()->HasLayersContent())
 	{
 		TArray<ULandscapeComponent*> SingleComponent;
 		SingleComponent.Add(this);
@@ -824,7 +823,7 @@ void ULandscapeComponent::UpdateCollisionHeightData(const FColor* const Heightma
 	bool ChangeType = false;
 
     // In Landscape Layers the Collision Component is slave and doesn't need to be transacted
-	if (!GetMutableDefault<UEditorExperimentalSettings>()->bLandscapeLayerSystem)
+	if (!Proxy->HasLayersContent())
 	{
 		if (CollisionComp)
 		{
@@ -2146,7 +2145,8 @@ LANDSCAPE_API void ALandscapeProxy::Import(const FGuid& InGuid, int32 InMinX, in
 {
 	check(InGuid.IsValid());
 	check(InImportHeightData.Num() == InImportMaterialLayerInfos.Num());
-	check(GetMutableDefault<UEditorExperimentalSettings>()->bLandscapeLayerSystem || InImportLayers == nullptr);
+
+	check(CanHaveLayersContent() || InImportLayers == nullptr);
 
 	GWarn->BeginSlowTask(LOCTEXT("BeingImportingLandscapeTask", "Importing Landscape"), true);
 
@@ -2766,7 +2766,7 @@ LANDSCAPE_API void ALandscapeProxy::Import(const FGuid& InGuid, int32 InMinX, in
 	// Create and initialize landscape info object
 	ULandscapeInfo* LandscapeInfo = CreateLandscapeInfo();
 
-	if (GetMutableDefault<UEditorExperimentalSettings>()->bLandscapeLayerSystem)
+	if (CanHaveLayersContent())
 	{
 		// Components need to be registered to be able to import the layer content and we will remove them if they should have not been visible
 		bool ShouldComponentBeRegistered = GetLevel()->bIsVisible;
@@ -4051,15 +4051,12 @@ void ALandscapeProxy::PostEditChangeProperty(FPropertyChangedEvent& PropertyChan
 	if (bRemovedAnyLayers)
 	{
 		ALandscapeProxy::InvalidateGeneratedComponentData(LandscapeComponents);
+		ALandscape* LandscapeActor = GetLandscapeActor();
 
-		if (GetMutableDefault<UEditorExperimentalSettings>()->bLandscapeLayerSystem)
+		if(LandscapeActor != nullptr && LandscapeActor->HasLayersContent())
 		{
-			if(ALandscape* LandscapeActor = GetLandscapeActor())
-			{
-				LandscapeActor->RequestLayersContentUpdate(ELandscapeLayerUpdateMode::Update_All);
-			}
+			LandscapeActor->RequestLayersContentUpdate(ELandscapeLayerUpdateMode::Update_All);
 		}
-		
 	}
 
 	// Must do this *after* correcting the scale or reattaching the landscape components will crash!
@@ -4617,7 +4614,7 @@ void ULandscapeComponent::ReallocateWeightmaps(FLandscapeEditDataInterface* Data
 	ALandscapeProxy* TargetProxy = InTargetProxy ? InTargetProxy : GetLandscapeProxy();
 
 	FGuid EditingLayerGUID = GetEditingLayerGUID();
-	check(!GetMutableDefault<UEditorExperimentalSettings>()->bLandscapeLayerSystem || !InCanUseEditingWeightmap || EditingLayerGUID.IsValid());
+	check(!TargetProxy->HasLayersContent() || !InCanUseEditingWeightmap || EditingLayerGUID.IsValid());
 	FGuid TargetLayerGuid = InCanUseEditingWeightmap ? EditingLayerGUID : FGuid();
 
 	TArray<FWeightmapLayerAllocationInfo>& ComponentWeightmapLayerAllocations = GetWeightmapLayerAllocations(InCanUseEditingWeightmap);
diff --git a/Engine/Source/Runtime/Landscape/Private/LandscapeEditInterface.cpp b/Engine/Source/Runtime/Landscape/Private/LandscapeEditInterface.cpp
--- a/Engine/Source/Runtime/Landscape/Private/LandscapeEditInterface.cpp
+++ b/Engine/Source/Runtime/Landscape/Private/LandscapeEditInterface.cpp
@@ -22,7 +22,6 @@ LandscapeEditInterface.cpp: Landscape editing interface
 #include "LandscapeRender.h"
 #include "ComponentReregisterContext.h"
 #include "Algo/Transform.h"
-#include "Settings/EditorExperimentalSettings.h"
 
 // Channel remapping
 extern const size_t ChannelOffsets[4] = {STRUCT_OFFSET(FColor,R), STRUCT_OFFSET(FColor,G), STRUCT_OFFSET(FColor,B), STRUCT_OFFSET(FColor,A)};
@@ -406,7 +405,7 @@ void FLandscapeEditDataInterface::SetHeightData(int32 X1, int32 Y1, int32 X2, in
 			}
 												
 			Component->RequestHeightmapUpdate();
-			if (!GetMutableDefault<UEditorExperimentalSettings>()->bLandscapeLayerSystem)
+			if (!Component->GetLandscapeProxy()->HasLayersContent())
 			{
 				// See if we need to adjust the bounds. Note we never shrink the bounding box at this point
 				bool bUpdateBoxSphereBounds = false;
@@ -1749,7 +1748,7 @@ void ULandscapeComponent::DeleteLayer(ULandscapeLayerInfoObject* LayerInfo, FLan
 	}
 
 	// Update the shaders for this component
-	if (!GetMutableDefault<UEditorExperimentalSettings>()->bLandscapeLayerSystem)
+	if (!Component->GetLandscapeProxy()->HasLayersContent())
 	{
 		Component->UpdateMaterialInstances();
 		Component->EditToolRenderData.UpdateDebugColorMaterial(Component);
@@ -1758,7 +1757,7 @@ void ULandscapeComponent::DeleteLayer(ULandscapeLayerInfoObject* LayerInfo, FLan
 	Component->RequestWeightmapUpdate();
 	
 	// Update dominant layer info stored in collision component
-	if (!GetMutableDefault<UEditorExperimentalSettings>()->bLandscapeLayerSystem)
+	if (!Component->GetLandscapeProxy()->HasLayersContent())
 	{
 		TArray<FColor*> CollisionWeightmapMipData;
 		for (int32 WeightmapIdx = 0; WeightmapIdx < ComponentWeightmapTextures.Num(); WeightmapIdx++)
@@ -1952,7 +1951,7 @@ void ULandscapeComponent::FillLayer(ULandscapeLayerInfoObject* LayerInfo, FLands
 	}
 
 	// Update the shaders for this component
-	if (!GetMutableDefault<UEditorExperimentalSettings>()->bLandscapeLayerSystem)
+	if (!Component->GetLandscapeProxy()->HasLayersContent())
 	{
 		Component->UpdateMaterialInstances();
 		Component->EditToolRenderData.UpdateDebugColorMaterial(Component);
@@ -1962,7 +1961,7 @@ void ULandscapeComponent::FillLayer(ULandscapeLayerInfoObject* LayerInfo, FLands
 	Component->RequestWeightmapUpdate();
 
 	// Update dominant layer info stored in collision component
-	if (!GetMutableDefault<UEditorExperimentalSettings>()->bLandscapeLayerSystem)
+	if (!Component->GetLandscapeProxy()->HasLayersContent())
 	{
 		TArray<FColor*> CollisionWeightmapMipData;
 		for (int32 WeightmapIdx = 0; WeightmapIdx < ComponentWeightmapTextures.Num(); WeightmapIdx++)
@@ -2226,7 +2225,7 @@ void ULandscapeComponent::ReplaceLayer(ULandscapeLayerInfoObject* FromLayerInfo,
 		// Remove the layer
 		ComponentWeightmapLayerAllocations.RemoveAt(FromLayerIdx);
 
-		if (!GetMutableDefault<UEditorExperimentalSettings>()->bLandscapeLayerSystem)
+		if (!GetLandscapeProxy()->HasLayersContent())
 		{
 			// Update the shaders for this component
 			UpdateMaterialInstances();
@@ -2238,7 +2237,7 @@ void ULandscapeComponent::ReplaceLayer(ULandscapeLayerInfoObject* FromLayerInfo,
 
 	RequestWeightmapUpdate();
 
-	if (!GetMutableDefault<UEditorExperimentalSettings>()->bLandscapeLayerSystem)
+	if (!LandscapeEdit.HasLandscapeLayersContent())
 	{
 		// Update dominant layer info stored in collision component
 		TArray<FColor*> CollisionWeightmapMipData;
@@ -2279,24 +2278,20 @@ void FLandscapeEditDataInterface::ReplaceLayer(ULandscapeLayerInfoObject* FromLa
 		}
 	};
 
-	if (GetMutableDefault<UEditorExperimentalSettings>()->bLandscapeLayerSystem)
+	if (LandscapeInfo->LandscapeActor.IsValid() && LandscapeInfo->LandscapeActor->HasLayersContent())
 	{
-		if (LandscapeInfo->LandscapeActor)
+		LandscapeInfo->LandscapeActor->Modify();
+		LandscapeInfo->LandscapeActor->ForEachLayer([&](FLandscapeLayer& CurrentLayer)
 		{
-			LandscapeInfo->LandscapeActor->Modify();
-			LandscapeInfo->LandscapeActor->ForEachLayer([&](FLandscapeLayer& CurrentLayer)
+			bool OutValue;
+			if (CurrentLayer.WeightmapLayerAllocationBlend.RemoveAndCopyValue(FromLayerInfo, OutValue))
 			{
-				bool OutValue;
-				if (CurrentLayer.WeightmapLayerAllocationBlend.RemoveAndCopyValue(FromLayerInfo, OutValue))
-				{
-					CurrentLayer.WeightmapLayerAllocationBlend.Add(ToLayerInfo, OutValue);
-				}
-
-				FScopedSetLandscapeEditingLayer Scope(LandscapeInfo->LandscapeActor.Get(), CurrentLayer.Guid);
-				DoReplace();
-			});
-			LandscapeInfo->LandscapeActor->RequestLayersContentUpdateForceAll(ELandscapeLayerUpdateMode::Update_Weightmap_All);
-		}
+				CurrentLayer.WeightmapLayerAllocationBlend.Add(ToLayerInfo, OutValue);
+			}
+			FScopedSetLandscapeEditingLayer Scope(LandscapeInfo->LandscapeActor.Get(), CurrentLayer.Guid);
+			DoReplace();
+		});
+		LandscapeInfo->LandscapeActor->RequestLayersContentUpdateForceAll(ELandscapeLayerUpdateMode::Update_Weightmap_All);
 	}
 	else
 	{
@@ -2857,7 +2852,7 @@ void FLandscapeEditDataInterface::SetAlphaData(ULandscapeLayerInfoObject* const
 				new (ComponentWeightmapLayerAllocations) FWeightmapLayerAllocationInfo(LayerInfo);
 				Component->ReallocateWeightmaps(this);
 
-				if (!GetMutableDefault<UEditorExperimentalSettings>()->bLandscapeLayerSystem)
+				if (!Component->GetLandscapeProxy()->HasLayersContent())
 				{
 					Component->UpdateMaterialInstances();
 					Component->EditToolRenderData.UpdateDebugColorMaterial(Component);
@@ -3109,7 +3104,7 @@ void FLandscapeEditDataInterface::SetAlphaData(ULandscapeLayerInfoObject* const
 				}
 			}
 
-			if (!GetMutableDefault<UEditorExperimentalSettings>()->bLandscapeLayerSystem)
+			if (!Component->GetLandscapeProxy()->HasLayersContent())
 			{
 				// Update mipmaps
 				CollisionWeightmapMipData.Reset();
@@ -3170,7 +3165,7 @@ void FLandscapeEditDataInterface::SetAlphaData(ULandscapeLayerInfoObject* const
 
 			if (bRemovedLayer)
 			{
-				if (!GetMutableDefault<UEditorExperimentalSettings>()->bLandscapeLayerSystem)
+				if (!Component->GetLandscapeProxy()->HasLayersContent())
 				{
 					Component->UpdateMaterialInstances();
 					Component->EditToolRenderData.UpdateDebugColorMaterial(Component);
@@ -3280,7 +3275,7 @@ void FLandscapeEditDataInterface::SetAlphaData(const TSet<ULandscapeLayerInfoObj
 					}
 					Component->ReallocateWeightmaps(this);
 					
-					if (!GetMutableDefault<UEditorExperimentalSettings>()->bLandscapeLayerSystem)
+					if (!Component->GetLandscapeProxy()->HasLayersContent())
 					{
 						Component->UpdateMaterialInstances();
 						Component->EditToolRenderData.UpdateDebugColorMaterial(Component);
@@ -3396,7 +3391,7 @@ void FLandscapeEditDataInterface::SetAlphaData(const TSet<ULandscapeLayerInfoObj
 				}
 			}
 
-			if (!GetMutableDefault<UEditorExperimentalSettings>()->bLandscapeLayerSystem)
+			if (!Component->GetLandscapeProxy()->HasLayersContent())
 			{
 				// Update mipmaps
 				CollisionWeightmapMipData.Reset();
@@ -3457,7 +3452,7 @@ void FLandscapeEditDataInterface::SetAlphaData(const TSet<ULandscapeLayerInfoObj
 
 			if (bRemovedLayer)
 			{
-				if (!GetMutableDefault<UEditorExperimentalSettings>()->bLandscapeLayerSystem)
+				if (!Component->GetLandscapeProxy()->HasLayersContent())
 				{
 					Component->UpdateMaterialInstances();
 					Component->EditToolRenderData.UpdateDebugColorMaterial(Component);
@@ -4542,11 +4537,7 @@ void FLandscapeEditDataInterface::GetEditToolTextureData(const int32 X1, const i
 
 			FLandscapeTextureDataInfo* TexDataInfo = NULL;
 			uint8* SelectTextureData = NULL;
-			UTexture2D* EditToolTexture = nullptr;
-			if (Component)
-			{
-				EditToolTexture = GetComponentTexture(Component);
-			}
+			UTexture2D* EditToolTexture = Component ? GetComponentTexture(Component) : nullptr;
 			if (EditToolTexture)
 			{
 				TexDataInfo = GetTextureDataInfo(EditToolTexture);
@@ -5520,6 +5511,18 @@ const ALandscape* FLandscapeEditDataInterface::GetTargetLandscape() const
 	return LandscapeInfo ? LandscapeInfo->LandscapeActor.Get() : nullptr;
 }
 
+bool FLandscapeEditDataInterface::CanHaveLandscapeLayersContent() const
+{
+	const ALandscape* Landscape = GetTargetLandscape();
+	return Landscape ? Landscape->CanHaveLayersContent() : false;
+}
+
+bool FLandscapeEditDataInterface::HasLandscapeLayersContent() const
+{
+	const ALandscape* Landscape = GetTargetLandscape();
+	return Landscape ? Landscape->HasLayersContent() : false;
+}
+
 void FLandscapeEditDataInterface::GetXYOffsetDataFast(const int32 X1, const int32 Y1, const int32 X2, const int32 Y2, FVector2D* Data, int32 Stride)
 {
 	if( Stride==0 )
diff --git a/Engine/Source/Runtime/Landscape/Private/LandscapeEditLayers.cpp b/Engine/Source/Runtime/Landscape/Private/LandscapeEditLayers.cpp
--- a/Engine/Source/Runtime/Landscape/Private/LandscapeEditLayers.cpp
+++ b/Engine/Source/Runtime/Landscape/Private/LandscapeEditLayers.cpp
@@ -23,7 +23,6 @@ LandscapeEditLayers.cpp: Landscape editing layers mode
 #include "LandscapeEditorModule.h"
 #include "LandscapeToolInterface.h"
 #include "ComponentRecreateRenderStateContext.h"
-#include "Settings/EditorExperimentalSettings.h"
 #include "LandscapeBPCustomBrush.h"
 #include "Materials/MaterialInstanceConstant.h"
 #include "LandscapeMaterialInstanceConstant.h"
@@ -32,6 +31,8 @@ LandscapeEditLayers.cpp: Landscape editing layers mode
 #include "Algo/Count.h"
 #include "LandscapeSettings.h"
 #include "LandscapeRender.h"
+#include "LandscapeInfoMap.h"
+#include "Misc/MessageDialog.h"
 #endif
 
 #define LOCTEXT_NAMESPACE "Landscape"
@@ -1424,7 +1425,7 @@ void ALandscape::CreateLayersRenderingResource()
 
 void ALandscape::ReleaseLayersRenderingResource()
 {
-	check(!GetMutableDefault<UEditorExperimentalSettings>()->bLandscapeLayerSystem);
+	check(!CanHaveLayersContent());
 	ULandscapeInfo* Info = GetLandscapeInfo();
 	if (Info == nullptr)
 	{
@@ -1548,143 +1549,157 @@ void ALandscape::CopyOldDataToDefaultLayer()
 		return;
 	}
 
+	Info->ForAllLandscapeProxies([&](ALandscapeProxy* Proxy)
+	{
+		CopyOldDataToDefaultLayer(Proxy);
+	});
+}
+
+void ALandscape::CopyOldDataToDefaultLayer(ALandscapeProxy* InProxy)
+{
+	ULandscapeInfo* Info = GetLandscapeInfo();
+
+	if (Info == nullptr)
+	{
+		return;
+	}
+
+	InProxy->Modify();
+
 	FLandscapeLayer* DefaultLayer = GetLayer(0);
 	check(DefaultLayer != nullptr);
 
-	Info->ForAllLandscapeProxies([&](ALandscapeProxy* Proxy)
+	struct FWeightmapTextureData
 	{
-		struct FWeightmapTextureData
-		{
-			UTexture2D* Texture;
-			ULandscapeWeightmapUsage* Usage;
-		};
+		UTexture2D* Texture;
+		ULandscapeWeightmapUsage* Usage;
+	};
 
-		TMap<UTexture2D*, FWeightmapTextureData> ProcessedWeightmaps;
-		TArray<UTexture2D*> ProcessedHeightmaps;
-		TArray<ULandscapeComponent*> WeightmapsComponentsToCleanup;
+	TMap<UTexture2D*, FWeightmapTextureData> ProcessedWeightmaps;
+	TArray<UTexture2D*> ProcessedHeightmaps;
+	TArray<ULandscapeComponent*> WeightmapsComponentsToCleanup;
 
-		for (ULandscapeComponent* Component : Proxy->LandscapeComponents)
+	for (ULandscapeComponent* Component : InProxy->LandscapeComponents)
+	{
+		FLandscapeLayerComponentData* LayerData = Component->GetLayerData(DefaultLayer->Guid);
+
+		if (ensure(LayerData != nullptr && LayerData->IsInitialized()))
 		{
-			FLandscapeLayerComponentData* LayerData = Component->GetLayerData(DefaultLayer->Guid);
+			// Heightmap
+			UTexture2D* ComponentHeightmap = Component->GetHeightmap();
 
-			if (ensure(LayerData != nullptr && LayerData->IsInitialized()))
+			if (!ProcessedHeightmaps.Contains(ComponentHeightmap))
 			{
-				// Heightmap
-				UTexture2D* ComponentHeightmap = Component->GetHeightmap();
+				ProcessedHeightmaps.Add(ComponentHeightmap);
 
-				if (!ProcessedHeightmaps.Contains(ComponentHeightmap))
-				{
-					ProcessedHeightmaps.Add(ComponentHeightmap);
+				UTexture* DefaultLayerHeightmap = LayerData->HeightmapData.Texture;
+				check(DefaultLayerHeightmap != nullptr);
 
-					UTexture* DefaultLayerHeightmap = LayerData->HeightmapData.Texture;
-					check(DefaultLayerHeightmap != nullptr);
+				// Only copy Mip0 as other mips will get regenerated
+				TArray<uint8> ExistingMip0Data;
+				ComponentHeightmap->Source.GetMipData(ExistingMip0Data, 0);
 
-					// Only copy Mip0 as other mips will get regenerated
-					TArray<uint8> ExistingMip0Data;
-					ComponentHeightmap->Source.GetMipData(ExistingMip0Data, 0);
+				FColor* Mip0Data = (FColor*)DefaultLayerHeightmap->Source.LockMip(0);
+				FMemory::Memcpy(Mip0Data, ExistingMip0Data.GetData(), ExistingMip0Data.Num());
+				DefaultLayerHeightmap->Source.UnlockMip(0);
 
-					FColor* Mip0Data = (FColor*)DefaultLayerHeightmap->Source.LockMip(0);
-					FMemory::Memcpy(Mip0Data, ExistingMip0Data.GetData(), ExistingMip0Data.Num());
-					DefaultLayerHeightmap->Source.UnlockMip(0);
+				DefaultLayerHeightmap->BeginCachePlatformData();
+				DefaultLayerHeightmap->ClearAllCachedCookedPlatformData();
+			}
 
-					DefaultLayerHeightmap->BeginCachePlatformData();
-					DefaultLayerHeightmap->ClearAllCachedCookedPlatformData();
-				}
+			// Weightmaps
+			WeightmapsComponentsToCleanup.Add(Component);
 
-				// Weightmaps
-				WeightmapsComponentsToCleanup.Add(Component);
+			const TArray<UTexture2D*>& ComponentWeightmapTextures = Component->GetWeightmapTextures();
+			const TArray<FWeightmapLayerAllocationInfo>& ComponentLayerAllocations = Component->GetWeightmapLayerAllocations();
+			TArray<ULandscapeWeightmapUsage*>& ComponentWeightmapTexturesUsage = Component->GetWeightmapTexturesUsage();
 
-				const TArray<UTexture2D*>& ComponentWeightmapTextures = Component->GetWeightmapTextures();
-				const TArray<FWeightmapLayerAllocationInfo>& ComponentLayerAllocations = Component->GetWeightmapLayerAllocations();
-				TArray<ULandscapeWeightmapUsage*>& ComponentWeightmapTexturesUsage = Component->GetWeightmapTexturesUsage();
+			LayerData->WeightmapData.Textures.AddDefaulted(ComponentWeightmapTextures.Num());
+			LayerData->WeightmapData.TextureUsages.AddDefaulted(ComponentWeightmapTexturesUsage.Num());
 
-				LayerData->WeightmapData.Textures.AddDefaulted(ComponentWeightmapTextures.Num());
-				LayerData->WeightmapData.TextureUsages.AddDefaulted(ComponentWeightmapTexturesUsage.Num());
+			for (int32 TextureIndex = 0; TextureIndex < ComponentWeightmapTextures.Num(); ++TextureIndex)
+			{
+				UTexture2D* ComponentWeightmap = ComponentWeightmapTextures[TextureIndex];
+				const FWeightmapTextureData* WeightmapTextureData = ProcessedWeightmaps.Find(ComponentWeightmap);
 
-				for (int32 TextureIndex = 0; TextureIndex < ComponentWeightmapTextures.Num(); ++TextureIndex)
+				if (WeightmapTextureData != nullptr)
 				{
-					UTexture2D* ComponentWeightmap = ComponentWeightmapTextures[TextureIndex];
-					const FWeightmapTextureData* WeightmapTextureData = ProcessedWeightmaps.Find(ComponentWeightmap);
+					LayerData->WeightmapData.Textures[TextureIndex] = WeightmapTextureData->Texture;
+					LayerData->WeightmapData.TextureUsages[TextureIndex] = WeightmapTextureData->Usage;
+					check(WeightmapTextureData->Usage->LayerGuid == DefaultLayer->Guid);
 
-					if (WeightmapTextureData != nullptr)
+					for (int32 ChannelIndex = 0; ChannelIndex < ULandscapeWeightmapUsage::NumChannels; ++ChannelIndex)
 					{
-						LayerData->WeightmapData.Textures[TextureIndex] = WeightmapTextureData->Texture;
-						LayerData->WeightmapData.TextureUsages[TextureIndex] = WeightmapTextureData->Usage;
-						check(WeightmapTextureData->Usage->LayerGuid == DefaultLayer->Guid);
+						const ULandscapeComponent* ChannelLandscapeComponent = LayerData->WeightmapData.TextureUsages[TextureIndex]->ChannelUsage[ChannelIndex];
 
-						for (int32 ChannelIndex = 0; ChannelIndex < ULandscapeWeightmapUsage::NumChannels; ++ChannelIndex)
+						if (ChannelLandscapeComponent != nullptr && ChannelLandscapeComponent == Component)
 						{
-							const ULandscapeComponent* ChannelLandscapeComponent = LayerData->WeightmapData.TextureUsages[TextureIndex]->ChannelUsage[ChannelIndex];
-
-							if (ChannelLandscapeComponent != nullptr && ChannelLandscapeComponent == Component)
+							for (const FWeightmapLayerAllocationInfo& Allocation : ComponentLayerAllocations)
 							{
-								for (const FWeightmapLayerAllocationInfo& Allocation : ComponentLayerAllocations)
+								if (Allocation.WeightmapTextureIndex == TextureIndex)
 								{
-									if (Allocation.WeightmapTextureIndex == TextureIndex)
-									{
-										LayerData->WeightmapData.LayerAllocations.Add(Allocation);
-									}
+									LayerData->WeightmapData.LayerAllocations.Add(Allocation);
 								}
-
-								break;
 							}
+
+							break;
 						}
 					}
-					else
-					{
-						UTexture2D* NewLayerWeightmapTexture = Proxy->CreateLandscapeTexture(ComponentWeightmap->Source.GetSizeX(), ComponentWeightmap->Source.GetSizeY(), TEXTUREGROUP_Terrain_Weightmap, ComponentWeightmap->Source.GetFormat());
+				}
+				else
+				{
+					UTexture2D* NewLayerWeightmapTexture = InProxy->CreateLandscapeTexture(ComponentWeightmap->Source.GetSizeX(), ComponentWeightmap->Source.GetSizeY(), TEXTUREGROUP_Terrain_Weightmap, ComponentWeightmap->Source.GetFormat());
 
-						// Only copy Mip0 as other mips will get regenerated
-						TArray<uint8> ExistingMip0Data;
-						ComponentWeightmap->Source.GetMipData(ExistingMip0Data, 0);
+					// Only copy Mip0 as other mips will get regenerated
+					TArray<uint8> ExistingMip0Data;
+					ComponentWeightmap->Source.GetMipData(ExistingMip0Data, 0);
 
-						FColor* Mip0Data = (FColor*)NewLayerWeightmapTexture->Source.LockMip(0);
-						FMemory::Memcpy(Mip0Data, ExistingMip0Data.GetData(), ExistingMip0Data.Num());
-						NewLayerWeightmapTexture->Source.UnlockMip(0);
+					FColor* Mip0Data = (FColor*)NewLayerWeightmapTexture->Source.LockMip(0);
+					FMemory::Memcpy(Mip0Data, ExistingMip0Data.GetData(), ExistingMip0Data.Num());
+					NewLayerWeightmapTexture->Source.UnlockMip(0);
 
-						LayerData->WeightmapData.Textures[TextureIndex] = NewLayerWeightmapTexture;
-						LayerData->WeightmapData.TextureUsages[TextureIndex] = Proxy->WeightmapUsageMap.Add(NewLayerWeightmapTexture, Proxy->CreateWeightmapUsage());
+					LayerData->WeightmapData.Textures[TextureIndex] = NewLayerWeightmapTexture;
+					LayerData->WeightmapData.TextureUsages[TextureIndex] = InProxy->WeightmapUsageMap.Add(NewLayerWeightmapTexture, InProxy->CreateWeightmapUsage());
 
-						for (int32 ChannelIndex = 0; ChannelIndex < ULandscapeWeightmapUsage::NumChannels; ++ChannelIndex)
-						{
-							LayerData->WeightmapData.TextureUsages[TextureIndex]->ChannelUsage[ChannelIndex] = ComponentWeightmapTexturesUsage[TextureIndex]->ChannelUsage[ChannelIndex];
-						}
+					for (int32 ChannelIndex = 0; ChannelIndex < ULandscapeWeightmapUsage::NumChannels; ++ChannelIndex)
+					{
+						LayerData->WeightmapData.TextureUsages[TextureIndex]->ChannelUsage[ChannelIndex] = ComponentWeightmapTexturesUsage[TextureIndex]->ChannelUsage[ChannelIndex];
+					}
 
-						LayerData->WeightmapData.TextureUsages[TextureIndex]->LayerGuid = DefaultLayer->Guid;
+					LayerData->WeightmapData.TextureUsages[TextureIndex]->LayerGuid = DefaultLayer->Guid;
 
-						// Create new Usage for the "final" layer as the other one will now be used by the Default layer
-						for (const FWeightmapLayerAllocationInfo& Allocation : ComponentLayerAllocations)
+					// Create new Usage for the "final" layer as the other one will now be used by the Default layer
+					for (const FWeightmapLayerAllocationInfo& Allocation : ComponentLayerAllocations)
+					{
+						if (Allocation.WeightmapTextureIndex == TextureIndex)
 						{
-							if (Allocation.WeightmapTextureIndex == TextureIndex)
-							{
-								LayerData->WeightmapData.LayerAllocations.Add(Allocation);
-							}
+							LayerData->WeightmapData.LayerAllocations.Add(Allocation);
 						}
+					}
 
-						FWeightmapTextureData NewTextureData;
-						NewTextureData.Texture = NewLayerWeightmapTexture;
-						NewTextureData.Usage = LayerData->WeightmapData.TextureUsages[TextureIndex];
+					FWeightmapTextureData NewTextureData;
+					NewTextureData.Texture = NewLayerWeightmapTexture;
+					NewTextureData.Usage = LayerData->WeightmapData.TextureUsages[TextureIndex];
 
-						ProcessedWeightmaps.Add(ComponentWeightmap, NewTextureData);
+					ProcessedWeightmaps.Add(ComponentWeightmap, NewTextureData);
 
-						NewLayerWeightmapTexture->BeginCachePlatformData();
-						NewLayerWeightmapTexture->ClearAllCachedCookedPlatformData();
-					}
+					NewLayerWeightmapTexture->BeginCachePlatformData();
+					NewLayerWeightmapTexture->ClearAllCachedCookedPlatformData();
 				}
 			}
 		}
+	}
 
-		for (ULandscapeComponent* Component : WeightmapsComponentsToCleanup)
-		{
-			TArray<FWeightmapLayerAllocationInfo>& ComponentLayerAllocations = Component->GetWeightmapLayerAllocations();
+	for (ULandscapeComponent* Component : WeightmapsComponentsToCleanup)
+	{
+		TArray<FWeightmapLayerAllocationInfo>& ComponentLayerAllocations = Component->GetWeightmapLayerAllocations();
 
-			for (FWeightmapLayerAllocationInfo& Allocation : ComponentLayerAllocations)
-			{
-				Allocation.WeightmapTextureChannel = 255;
-				Allocation.WeightmapTextureIndex = 255;
-			}
+		for (FWeightmapLayerAllocationInfo& Allocation : ComponentLayerAllocations)
+		{
+			Allocation.WeightmapTextureChannel = 255;
+			Allocation.WeightmapTextureIndex = 255;
 		}
-	});
+	}
 }
 
 void ALandscape::InitializeLandscapeLayersWeightmapUsage()
@@ -4000,11 +4015,14 @@ void ALandscape::ResolveLayersWeightmapTexture()
 	});
 }
 
-
+bool ALandscape::HasLayersContent() const
+{
+	return LandscapeLayers.Num() > 0;
+}
 
 void ALandscape::RequestLayersInitialization(bool bInRequestContentUpdate)
 {
-	if (!GetMutableDefault<UEditorExperimentalSettings>()->bLandscapeLayerSystem)
+	if (!CanHaveLayersContent())
 	{
 		return;
 	}
@@ -4024,7 +4042,7 @@ void ALandscape::RequestLayersContentUpdate(ELandscapeLayerUpdateMode InUpdateMo
 
 void ALandscape::RequestLayersContentUpdateForceAll(ELandscapeLayerUpdateMode InModeMask)
 {
-	if (!GetMutableDefault<UEditorExperimentalSettings>()->bLandscapeLayerSystem)
+	if (!CanHaveLayersContent())
 	{
 		return;
 	}
@@ -4226,30 +4244,24 @@ void ALandscape::MonitorShaderCompilation()
 
 void ALandscape::UpdateLayersContent(bool bInWaitForStreaming)
 {
-	if (!GetMutableDefault<UEditorExperimentalSettings>()->bLandscapeLayerSystem)
+	if (GetLandscapeInfo() == nullptr || !CanHaveLayersContent())
 	{
-		if (bInitializedWithFlagExperimentalLandscapeLayers)
-		{
-			ReleaseLayersRenderingResource();
-			bInitializedWithFlagExperimentalLandscapeLayers = false;
-		}
 		return;
 	}
 
-	if (!bLandscapeLayersAreInitialized || !bInitializedWithFlagExperimentalLandscapeLayers)
+	if (!bLandscapeLayersAreInitialized)
 	{
 		InitializeLayers();
-		bInitializedWithFlagExperimentalLandscapeLayers = true;
 	}
-	
-	MonitorShaderCompilation();
+
 	MonitorLandscapeEdModeChanges();
+	MonitorShaderCompilation();
 
-	if (GetLandscapeInfo() == nullptr || LayerContentUpdateModes == 0)
+	if (LayerContentUpdateModes == 0)
 	{
 		return;
 	}
-	
+
 	TArray<ULandscapeComponent*> AllLandscapeComponents;
 	GetLandscapeInfo()->ForAllLandscapeProxies([&AllLandscapeComponents](ALandscapeProxy* Proxy)
 	{
@@ -4392,24 +4404,99 @@ bool ALandscape::UpdateCollisionAndClients(const TArray<ULandscapeComponent*>& I
 	return bAllClientsUpdated;
 }
 
-void ALandscape::InitializeLayers()
+bool ALandscape::ChangeLandscapeLayersState(UWorld* InWorld, bool bInEnable)
 {
-	check(GetMutableDefault<UEditorExperimentalSettings>()->bLandscapeLayerSystem);
-	check(!bLandscapeLayersAreInitialized || (bInitializedWithFlagExperimentalLandscapeLayers != GetMutableDefault<UEditorExperimentalSettings>()->bLandscapeLayerSystem));
-
-	if (HasLayersContent)
+	if (bInEnable)
 	{
-		CreateLayersRenderingResource();
-		InitializeLandscapeLayersWeightmapUsage();
+		auto& LandscapeInfoMap = ULandscapeInfoMap::GetLandscapeInfoMap(InWorld);
+
+		for (auto& ItPair : LandscapeInfoMap.Map)
+		{
+			ULandscapeInfo* Info = ItPair.Value;
+			
+			if (Info->LandscapeActor.IsValid())
+			{
+				ALandscape* Landscape = Info->LandscapeActor.Get();
+
+				if (Landscape->GetLayerCount() == 0)
+				{
+					Landscape->CreateDefaultLayer();
+				}
+
+				Landscape->CopyOldDataToDefaultLayer();
+			}
+		}
+
+		return true;
 	}
 	else
 	{
-		// Special case to load the existing content to the default layer
-		CreateDefaultLayer();
-		CreateLayersRenderingResource();
-		CopyOldDataToDefaultLayer();
+		bool HasHiddenLayers = false;
+
+		auto& LandscapeInfoMap = ULandscapeInfoMap::GetLandscapeInfoMap(InWorld);
+
+		for (auto& ItPair : LandscapeInfoMap.Map)
+		{
+			ULandscapeInfo* Info = ItPair.Value;
+			ALandscape* Landscape = Info->LandscapeActor.Get();
+
+			if (Landscape != nullptr)
+			{
+				for (int32 i = 0; i < Landscape->GetLayerCount(); ++i)
+				{
+					const FLandscapeLayer* Layer = Landscape->GetLayer(i);
+					check(Layer != nullptr);
+
+					if (!Layer->bVisible)
+					{
+						HasHiddenLayers = true;
+						break;
+					}
+				}
+			}
+		}
+
+		FText Reason;
+
+		if (HasHiddenLayers)
+		{
+			Reason = NSLOCTEXT("UnrealEd", "LandscapeDisableLayers_HiddenLayers", "Are you sure you want to disable the layers system?\n\nDoing so, will result in losing the data stored for each layers, but the current visual output will be kept. Be aware that some layers are currently hidden, continuing will result in their data being lost.");
+		}
+		else
+		{
+			Reason = NSLOCTEXT("UnrealEd", "LandscapeDisableLayers", "Are you sure you want to disable the layers system?\n\nDoing so, will result in losing the data stored for each layers, but the current visual output will be kept.");
+		}
+
+		if (FMessageDialog::Open(EAppMsgType::YesNo, Reason) == EAppReturnType::Yes)
+		{
+			for (auto& ItPair : LandscapeInfoMap.Map)
+			{
+				ULandscapeInfo* Info = ItPair.Value;
+				ALandscape* Landscape = Info->LandscapeActor.Get();
+
+				// Only clean up data if the Landscape actor is loaded
+				if (Landscape != nullptr)
+				{
+					Landscape->ReleaseLayersRenderingResource();
+					Landscape->DeleteLayers();
+				}			
+			}
+
+			return true;
+		}
 	}
 
+	return false;
+}
+
+void ALandscape::InitializeLayers()
+{
+	check(HasLayersContent());
+	check(!bLandscapeLayersAreInitialized);
+
+	CreateLayersRenderingResource();
+	InitializeLandscapeLayersWeightmapUsage();
+
 	bLandscapeLayersAreInitialized = true;
 }
 
@@ -4440,7 +4527,7 @@ void ALandscape::TickLayers(float DeltaTime, ELevelTick TickType, FActorTickFunc
 void ALandscapeProxy::BeginDestroy()
 {
 #if WITH_EDITORONLY_DATA
-	if (GetMutableDefault<UEditorExperimentalSettings>()->bLandscapeLayerSystem)
+	if (CanHaveLayersContent())
 	{
 		for (auto& ItPair : HeightmapsCPUReadBack)
 		{
@@ -4474,7 +4561,7 @@ bool ALandscapeProxy::IsReadyForFinishDestroy()
 	bool bReadyForFinishDestroy = Super::IsReadyForFinishDestroy();
 
 #if WITH_EDITORONLY_DATA
-	if (GetMutableDefault<UEditorExperimentalSettings>()->bLandscapeLayerSystem)
+	if (CanHaveLayersContent())
 	{
 		if (bReadyForFinishDestroy)
 		{
@@ -4489,7 +4576,7 @@ bool ALandscapeProxy::IsReadyForFinishDestroy()
 void ALandscapeProxy::FinishDestroy()
 {
 #if WITH_EDITORONLY_DATA
-	if (GetMutableDefault<UEditorExperimentalSettings>()->bLandscapeLayerSystem)
+	if (CanHaveLayersContent())
 	{
 		check(ReleaseResourceFence.IsFenceComplete());
 
@@ -4515,6 +4602,32 @@ void ALandscapeProxy::FinishDestroy()
 }
 
 #if WITH_EDITOR
+bool ALandscapeProxy::CanHaveLayersContent() const
+{
+	UWorld* World = GetWorld();
+	AWorldSettings* WorldSettings = World != nullptr ? World->GetWorldSettings() : nullptr;
+	return WorldSettings != nullptr ? WorldSettings->bEnableLandscapeLayers : false;
+}
+
+bool ALandscapeProxy::HasLayersContent() const
+{
+	return bHasLayersContent;
+}
+
+void ALandscapeProxy::UpdateCachedHasLayersContent(bool InCheckComponentDataIntegrity)
+{
+	// In the case of InCheckComponentDataIntegrity we will loop through all components to make sure they all have the same state and in the other case we will assume that the 1st component represent the state of all the others.
+	bHasLayersContent = LandscapeComponents.Num() > 0 ? LandscapeComponents[0]->HasLayersData() : false;
+
+	if (InCheckComponentDataIntegrity)
+	{
+		for (const ULandscapeComponent* Component : LandscapeComponents)
+		{
+			check(bHasLayersContent == Component->HasLayersData());
+		}
+	}
+}
+
 bool ALandscapeProxy::RemoveObsoleteLayers(const TSet<FGuid>& InExistingLayers)
 {
 	TSet<FGuid> ComponentLayers;
@@ -4547,13 +4660,15 @@ bool ALandscapeProxy::AddLayer(const FGuid& InLayerGuid)
 			Component->AddLayerData(InLayerGuid, FLandscapeLayerComponentData());
 			bModified = true;
 		}
-	}
+	}	
+
+	UpdateCachedHasLayersContent();
 
 	if (bModified)
 	{
 		InitializeLayerWithEmptyContent(InLayerGuid);
 	}
-	HasLayersContent = true;
+
 	return bModified;
 }
 
@@ -4585,6 +4700,8 @@ void ALandscapeProxy::DeleteLayer(const FGuid& InLayerGuid)
 			Component->RemoveLayerData(InLayerGuid);
 		}
 	}
+
+	UpdateCachedHasLayersContent();
 }
 
 void ALandscapeProxy::InitializeLayerWithEmptyContent(const FGuid& InLayerGuid)
@@ -4623,7 +4740,9 @@ void ALandscapeProxy::InitializeLayerWithEmptyContent(const FGuid& InLayerGuid)
 void ALandscape::BeginDestroy()
 {
 #if WITH_EDITOR
-	if (GetMutableDefault<UEditorExperimentalSettings>()->bLandscapeLayerSystem)
+	ALandscape::UnregisterChangeLandscapeLayersStateDelegate();
+
+	if (CanHaveLayersContent())
 	{
 		if (CombinedLayersWeightmapAllMaterialLayersResource != nullptr)
 		{
@@ -4655,7 +4774,7 @@ void ALandscape::BeginDestroy()
 void ALandscape::FinishDestroy()
 {
 #if WITH_EDITORONLY_DATA
-	if (GetMutableDefault<UEditorExperimentalSettings>()->bLandscapeLayerSystem)
+	if (CanHaveLayersContent())
 	{
 		check(ReleaseResourceFence.IsFenceComplete());
 
@@ -4677,7 +4796,7 @@ void ALandscape::FinishDestroy()
 bool ALandscape::IsUpToDate() const
 {
 #if WITH_EDITORONLY_DATA
-	if (GetMutableDefault<UEditorExperimentalSettings>()->bLandscapeLayerSystem)
+	if (CanHaveLayersContent())
 	{
 		return LayerContentUpdateModes == 0;
 	}
@@ -4809,13 +4928,21 @@ void ALandscape::ForEachLayer(TFunctionRef<void(struct FLandscapeLayer&)> Fn)
 	}
 }
 
+void ALandscape::DeleteLayers()
+{
+	for (int32 LayerIndex = LandscapeLayers.Num() - 1; LayerIndex >= 0; --LayerIndex)
+	{
+		DeleteLayer(LayerIndex);
+	}
+}
+
 void ALandscape::DeleteLayer(int32 InLayerIndex)
 {
-	ensure(GetMutableDefault<UEditorExperimentalSettings>()->bLandscapeLayerSystem);
+	ensure(HasLayersContent());
 
 	ULandscapeInfo* LandscapeInfo = GetLandscapeInfo();
 	const FLandscapeLayer* Layer = GetLayer(InLayerIndex);
-	if (!LandscapeInfo || !Layer || LandscapeLayers.Num() <= 1)
+	if (!LandscapeInfo || !Layer)
 	{
 		return;
 	}
@@ -4853,7 +4980,7 @@ void ALandscape::ClearLayer(int32 InLayerIndex, bool bInUpdateCollision)
 
 void ALandscape::ClearLayer(const FGuid& InLayerGuid, bool bInUpdateCollision)
 {
-	ensure(GetMutableDefault<UEditorExperimentalSettings>()->bLandscapeLayerSystem);
+	ensure(HasLayersContent());
 
 	const FLandscapeLayer* Layer = GetLayer(InLayerGuid);
 	ULandscapeInfo* LandscapeInfo = GetLandscapeInfo();
@@ -5003,7 +5130,7 @@ LANDSCAPE_API extern bool GDisableUpdateLandscapeMaterialInstances;
 
 void ALandscape::UpdateLandscapeSplines(const FGuid& InTargetLayer, bool bUpdateOnlySelected)
 {
-	check(GetMutableDefault<UEditorExperimentalSettings>()->bLandscapeLayerSystem);
+	check(CanHaveLayersContent());
 	ULandscapeInfo* LandscapeInfo = GetLandscapeInfo();
 	FGuid TargetLayerGuid = LandscapeSplinesTargetLayerGuid.IsValid() ? LandscapeSplinesTargetLayerGuid : InTargetLayer;
 	const FLandscapeLayer* TargetLayer = GetLayer(TargetLayerGuid);
@@ -5027,7 +5154,7 @@ FScopedSetLandscapeEditingLayer::FScopedSetLandscapeEditingLayer(ALandscape* InL
 	: Landscape(InLandscape)
 	, CompletionCallback(MoveTemp(InCompletionCallback))
 {
-	if (GetMutableDefault<UEditorExperimentalSettings>()->bLandscapeLayerSystem && Landscape.IsValid())
+	if (Landscape.IsValid() && Landscape.Get()->CanHaveLayersContent())
 	{
 		PreviousLayerGUID = Landscape->GetEditingLayer();
 		Landscape->SetEditingLayer(InLayerGUID);
@@ -5036,7 +5163,7 @@ FScopedSetLandscapeEditingLayer::FScopedSetLandscapeEditingLayer(ALandscape* InL
 
 FScopedSetLandscapeEditingLayer::~FScopedSetLandscapeEditingLayer()
 {
-	if (GetMutableDefault<UEditorExperimentalSettings>()->bLandscapeLayerSystem && Landscape.IsValid())
+	if (Landscape.IsValid() && Landscape.Get()->CanHaveLayersContent())
 	{
 		Landscape->SetEditingLayer(PreviousLayerGUID);
 		if (CompletionCallback)
@@ -5048,7 +5175,7 @@ FScopedSetLandscapeEditingLayer::~FScopedSetLandscapeEditingLayer()
 
 bool ALandscape::IsEditingLayerReservedForSplines() const
 {
-	if (GetMutableDefault<UEditorExperimentalSettings>()->bLandscapeLayerSystem)
+	if (CanHaveLayersContent())
 	{
 		const FLandscapeLayer* SplinesReservedLayer = GetLandscapeSplinesReservedLayer();
 		return SplinesReservedLayer && SplinesReservedLayer->Guid == EditingLayer;
@@ -5058,7 +5185,7 @@ bool ALandscape::IsEditingLayerReservedForSplines() const
 
 void ALandscape::SetEditingLayer(const FGuid& InLayerGuid)
 {
-	ensure(GetMutableDefault<UEditorExperimentalSettings>()->bLandscapeLayerSystem);
+	ensure(CanHaveLayersContent());
 	
 	ULandscapeInfo* LandscapeInfo = GetLandscapeInfo();
 	if (!LandscapeInfo)
@@ -5105,7 +5232,7 @@ TMap<UTexture2D*, TArray<ULandscapeComponent*>> ALandscape::GenerateComponentsPe
 void ALandscape::CreateDefaultLayer()
 {
 	ULandscapeInfo* LandscapeInfo = GetLandscapeInfo();
-	if (!LandscapeInfo || !GetMutableDefault<UEditorExperimentalSettings>()->bLandscapeLayerSystem)
+	if (!LandscapeInfo || !CanHaveLayersContent())
 	{
 		return;
 	}
@@ -5113,7 +5240,6 @@ void ALandscape::CreateDefaultLayer()
 	check(LandscapeLayers.Num() == 0); // We can only call this function if we have no layers
 
 	CreateLayer(FName(TEXT("Layer")));
-	HasLayersContent = true;
 	// Force update rendering resources
 	RequestLayersInitialization();
 }
@@ -5121,11 +5247,13 @@ void ALandscape::CreateDefaultLayer()
 FLandscapeLayer* ALandscape::DuplicateLayer(const FLandscapeLayer& InOtherLayer)
 {
 	ULandscapeInfo* LandscapeInfo = GetLandscapeInfo();
-	if (!LandscapeInfo || !GetMutableDefault<UEditorExperimentalSettings>()->bLandscapeLayerSystem)
+	if (!LandscapeInfo || !CanHaveLayersContent())
 	{
 		return nullptr;
 	}
 
+	Modify();
+
 	FLandscapeLayer NewLayer(InOtherLayer);
 	NewLayer.Guid = FGuid::NewGuid();
 
@@ -5157,7 +5285,7 @@ FLandscapeLayer* ALandscape::DuplicateLayer(const FLandscapeLayer& InOtherLayer)
 void ALandscape::CreateLayer(FName InName)
 {
 	ULandscapeInfo* LandscapeInfo = GetLandscapeInfo();
-	if (!LandscapeInfo || IsMaxLayersReached() || !GetMutableDefault<UEditorExperimentalSettings>()->bLandscapeLayerSystem)
+	if (!LandscapeInfo || IsMaxLayersReached() || !CanHaveLayersContent())
 	{
 		return;
 	}
@@ -5177,13 +5305,11 @@ void ALandscape::CreateLayer(FName InName)
 void ALandscape::AddLayersToProxy(ALandscapeProxy* InProxy)
 {
 	ULandscapeInfo* LandscapeInfo = GetLandscapeInfo();
-	if (!LandscapeInfo || !GetMutableDefault<UEditorExperimentalSettings>()->bLandscapeLayerSystem)
+	if (!LandscapeInfo || !CanHaveLayersContent())
 	{
 		return;
 	}
 
-	Modify();
-
 	check(InProxy != this);
 	check(InProxy != nullptr);
 
@@ -5191,7 +5317,6 @@ void ALandscape::AddLayersToProxy(ALandscapeProxy* InProxy)
 	{
 		InProxy->AddLayer(Layer.Guid);
 	});
-	InProxy->HasLayersContent = true;
 
 	// Force update rendering resources
 	RequestLayersInitialization();
@@ -5278,6 +5403,8 @@ void ALandscape::AddBrushToLayer(int32 InLayerIndex, int32 InTargetType, ALandsc
 		return;
 	}
 
+	Modify();
+
 	int32 AddedIndex = Layer->Brushes.Add(FLandscapeLayerBrush(InBrush));
 
 	if (InTargetType == ELandscapeToolTargetType::Type::Heightmap)
@@ -5303,6 +5430,8 @@ void ALandscape::RemoveBrushFromLayer(int32 InLayerIndex, int32 InTargetType, AL
 		return;
 	}
 
+	Modify();
+
 	int32 IndexToRemove = INDEX_NONE;
 	for (int32 i = 0; i < Layer->Brushes.Num(); ++i)
 	{
@@ -5471,6 +5600,18 @@ TArray<ALandscapeBlueprintCustomBrush*> ALandscape::GetBrushesForLayer(int32 InL
 	return Brushes;
 }
 
+FChangeLandscapeLayersState ALandscape::ChangeLandscapeLayersStateEvent;
+
+void ALandscape::RegisterChangeLandscapeLayersStateDelegate()
+{
+	ALandscape::ChangeLandscapeLayersStateEvent.BindStatic(&ALandscape::ChangeLandscapeLayersState);
+}
+
+void ALandscape::UnregisterChangeLandscapeLayersStateDelegate()
+{
+	ALandscape::ChangeLandscapeLayersStateEvent.Unbind();	
+}
+
 #endif // WITH_EDITOR
 
 #undef LOCTEXT_NAMESPACE
\ No newline at end of file
diff --git a/Engine/Source/Runtime/Landscape/Private/LandscapeModule.cpp b/Engine/Source/Runtime/Landscape/Private/LandscapeModule.cpp
--- a/Engine/Source/Runtime/Landscape/Private/LandscapeModule.cpp
+++ b/Engine/Source/Runtime/Landscape/Private/LandscapeModule.cpp
@@ -12,10 +12,6 @@
 #include "LandscapeInfoMap.h"
 #include "Materials/MaterialInstance.h"
 
-#if WITH_EDITOR
-#include "Settings/EditorExperimentalSettings.h"
-#endif
-
 #include "LandscapeProxy.h"
 #include "Landscape.h"
 #include "Engine/Texture2D.h"
diff --git a/Engine/Source/Runtime/Landscape/Private/LandscapeSplineRaster.cpp b/Engine/Source/Runtime/Landscape/Private/LandscapeSplineRaster.cpp
--- a/Engine/Source/Runtime/Landscape/Private/LandscapeSplineRaster.cpp
+++ b/Engine/Source/Runtime/Landscape/Private/LandscapeSplineRaster.cpp
@@ -20,7 +20,6 @@
 #include "ScopedTransaction.h"
 #include "Raster.h"
 #include "Landscape.h"
-#include "Settings/EditorExperimentalSettings.h"
 #endif
 
 #define LOCTEXT_NAMESPACE "Landscape"
@@ -632,7 +631,7 @@ bool ULandscapeInfo::ApplySplinesInternal(bool bOnlySelected, ALandscapeProxy* P
 		
 	for (ULandscapeComponent* Component : ModifiedComponents)
 	{
-		if (GetMutableDefault<UEditorExperimentalSettings>()->bLandscapeLayerSystem)
+		if (Component->GetLandscapeProxy()->HasLayersContent())
 		{
 			Component->RequestHeightmapUpdate();
 			Component->RequestWeightmapUpdate();
@@ -730,7 +729,7 @@ namespace LandscapeSplineRaster
 
 		LandscapeEdit.Flush();
 
-		if (!GetMutableDefault<UEditorExperimentalSettings>()->bLandscapeLayerSystem)
+		if (!LandscapeProxy->HasLayersContent())
 		{
 			for (ULandscapeComponent* Component : ModifiedComponents)
 			{
diff --git a/Engine/Source/Runtime/Landscape/Private/LandscapeSplines.cpp b/Engine/Source/Runtime/Landscape/Private/LandscapeSplines.cpp
--- a/Engine/Source/Runtime/Landscape/Private/LandscapeSplines.cpp
+++ b/Engine/Source/Runtime/Landscape/Private/LandscapeSplines.cpp
@@ -37,7 +37,6 @@
 #include "Logging/TokenizedMessage.h"
 #include "Logging/MessageLog.h"
 #include "Misc/UObjectToken.h"
-#include "Settings/EditorExperimentalSettings.h"
 #include "Landscape.h"
 #endif
 
@@ -1881,14 +1880,11 @@ void ULandscapeSplineControlPoint::PostEditChangeProperty(FPropertyChangedEvent&
 
 	if (PropertyChangedEvent.ChangeType == EPropertyChangeType::ValueSet)
 	{
-		if (GetMutableDefault<UEditorExperimentalSettings>()->bLandscapeLayerSystem)
+		ALandscapeProxy* OuterLandscape = Cast<ALandscapeProxy>(GetOuterULandscapeSplinesComponent()->GetOwner());
+		ALandscape* Landscape = OuterLandscape ? OuterLandscape->GetLandscapeActor() : nullptr;
+		if (Landscape && Landscape->HasLayersContent() && Landscape->GetLandscapeSplinesReservedLayer())
 		{
-			ALandscapeProxy* OuterLandscape = Cast<ALandscapeProxy>(GetOuterULandscapeSplinesComponent()->GetOwner());
-			ALandscape* Landscape = OuterLandscape ? OuterLandscape->GetLandscapeActor() : nullptr;
-			if (Landscape && Landscape->GetLandscapeSplinesReservedLayer())
-			{
-				Landscape->UpdateLandscapeSplines();
-			}
+			Landscape->UpdateLandscapeSplines();
 		}
 	}
 }
@@ -2837,14 +2833,11 @@ void ULandscapeSplineSegment::PostEditChangeProperty(FPropertyChangedEvent& Prop
 
 	if (PropertyChangedEvent.ChangeType == EPropertyChangeType::ValueSet)
 	{
-		if (GetMutableDefault<UEditorExperimentalSettings>()->bLandscapeLayerSystem)
+		ALandscapeProxy* OuterLandscape = Cast<ALandscapeProxy>(GetOuterULandscapeSplinesComponent()->GetOwner());
+		ALandscape* Landscape = OuterLandscape ? OuterLandscape->GetLandscapeActor() : nullptr;
+		if (Landscape && Landscape->HasLayersContent() && Landscape->GetLandscapeSplinesReservedLayer())
 		{
-			ALandscapeProxy* OuterLandscape = Cast<ALandscapeProxy>(GetOuterULandscapeSplinesComponent()->GetOwner());
-			ALandscape* Landscape = OuterLandscape ? OuterLandscape->GetLandscapeActor() : nullptr;
-			if (Landscape && Landscape->GetLandscapeSplinesReservedLayer())
-			{
-				Landscape->UpdateLandscapeSplines();
-			}
+			Landscape->UpdateLandscapeSplines();
 		}
 	}
 }
diff --git a/Engine/Source/Runtime/Landscape/Public/LandscapeEdit.h b/Engine/Source/Runtime/Landscape/Public/LandscapeEdit.h
--- a/Engine/Source/Runtime/Landscape/Public/LandscapeEdit.h
+++ b/Engine/Source/Runtime/Landscape/Public/LandscapeEdit.h
@@ -19,12 +19,11 @@ LandscapeEdit.h: Classes for the editor to access to Landscape data
 #include "LandscapeLayerInfoObject.h"
 
 #if WITH_EDITOR
-
 #include "Containers/ArrayView.h"
-#include "Settings/EditorExperimentalSettings.h"
-
 #endif
 
+#include "Landscape.h"
+
 class ULandscapeComponent;
 class ULandscapeInfo;
 class ULandscapeLayerInfoObject;
@@ -252,6 +251,10 @@ struct LANDSCAPE_API FLandscapeEditDataInterface : public FLandscapeTextureDataI
 	static void ShrinkData(TArray<T>& Data, int32 OldMinX, int32 OldMinY, int32 OldMaxX, int32 OldMaxY, int32 NewMinX, int32 NewMinY, int32 NewMaxX, int32 NewMaxY);
 
 	const ALandscape* GetTargetLandscape() const;
+
+	bool CanHaveLandscapeLayersContent() const;
+	bool HasLandscapeLayersContent() const;
+
 private:
 	int32 ComponentSizeQuads;
 	int32 SubsectionSizeQuads;
@@ -358,7 +361,7 @@ struct FHeightmapAccessor
 			ALandscapeProxy::InvalidateGeneratedComponentData(Components);
 
             // Landscape Layers are updates are delayed and done in  ALandscape::TickLayers
-			if (!GetMutableDefault<UEditorExperimentalSettings>()->bLandscapeLayerSystem)
+			if (!LandscapeEdit->HasLandscapeLayersContent())
 			{
 				for (ULandscapeComponent* Component : Components)
 				{
@@ -414,11 +417,8 @@ struct FHeightmapAccessor
 
 	virtual ~FHeightmapAccessor()
 	{
-		delete LandscapeEdit;
-		LandscapeEdit = NULL;
-
 		// Landscape Layers are updates are delayed and done in  ALandscape::TickLayers
-		if (!GetMutableDefault<UEditorExperimentalSettings>()->bLandscapeLayerSystem)
+		if (!LandscapeEdit->HasLandscapeLayersContent())
 		{
 			// Update the bounds and navmesh for the components we edited
 			for (TSet<ULandscapeComponent*>::TConstIterator It(ChangedComponents); It; ++It)
@@ -435,6 +435,9 @@ struct FHeightmapAccessor
 				}
 			}
 		}
+
+		delete LandscapeEdit;
+		LandscapeEdit = NULL;
 	}
 
 private:
@@ -473,7 +476,7 @@ struct FAlphamapAccessor
 
 	~FAlphamapAccessor()
 	{
-		if (!GetMutableDefault<UEditorExperimentalSettings>()->bLandscapeLayerSystem)
+		if (!LandscapeEdit.HasLandscapeLayersContent())
 		{
 			// Recreate collision for modified components to update the physical materials
 			for (ULandscapeComponent* Component : ModifiedComponents)
@@ -515,7 +518,7 @@ struct FAlphamapAccessor
 				LandscapeComponent->RequestWeightmapUpdate();
 			}
 			
-			if (!GetMutableDefault<UEditorExperimentalSettings>()->bLandscapeLayerSystem)
+			if (!LandscapeEdit.HasLandscapeLayersContent())
 			{
 				ModifiedComponents.Append(Components);
 			}
