- Landscape Layer System: Fix crash when loading a landscape with layer content and enabling landscape layers in world settings

#rnx

diff --git a/Engine/Source/Runtime/Landscape/Private/Landscape.cpp b/Engine/Source/Runtime/Landscape/Private/Landscape.cpp
--- a/Engine/Source/Runtime/Landscape/Private/Landscape.cpp
+++ b/Engine/Source/Runtime/Landscape/Private/Landscape.cpp
@@ -2271,11 +2271,19 @@ void ALandscapeProxy::PostLoad()
 	if (GIsEditor && !GetWorld()->IsGameWorld())
 	{
 		UpdateCachedHasLayersContent(true);
+		
+		ULandscapeInfo* LandscapeInfo = CreateLandscapeInfo();
 
 		// Cache the value at this point as RegisterActor might create/destroy layers content if there was a mismatch between landscape & proxy
-		bool bNeedOldDataMigration = !HasLayersContent() && CanHaveLayersContent();
-
-		ULandscapeInfo* LandscapeInfo = CreateLandscapeInfo();
+		const bool bDeleteLayerContent = HasLayersContent() && !CanHaveLayersContent();
+		if (bDeleteLayerContent && LandscapeInfo->LandscapeActor.IsValid())
+		{
+			LandscapeInfo->LandscapeActor->DeleteLayers();
+			UE_LOG(LogLandscape, Warning, TEXT("Landscape layer data is being deleted because World Settings doesn't have layers enabled."));
+		}
+		
+		const bool bNeedOldDataMigration = !HasLayersContent() && CanHaveLayersContent();
+				
 		LandscapeInfo->RegisterActor(this, true);
 
 		FixupWeightmaps();
