Landscape layer contribution visualizer now only highlights the selected paint target layer instead of combining all weightmaps for this layer.
#rnx

diff --git a/Engine/Source/Editor/LandscapeEditor/Private/LandscapeEdMode.cpp b/Engine/Source/Editor/LandscapeEditor/Private/LandscapeEdMode.cpp
--- a/Engine/Source/Editor/LandscapeEditor/Private/LandscapeEdMode.cpp
+++ b/Engine/Source/Editor/LandscapeEditor/Private/LandscapeEdMode.cpp
@@ -479,9 +479,14 @@ ELandscapeToolTargetType::Type FEdModeLandscape::GetLandscapeToolTargetType() co
 	return ELandscapeToolTargetType::Invalid;
 }
 
-FGuid FEdModeLandscape::GetLandscapeSelectedLayer() const
+const FLandscapeLayer* FEdModeLandscape::GetLandscapeSelectedLayer() const
 {
-	return GetCurrentLayerGuid();
+	return GetCurrentLayer();
+}
+
+ULandscapeLayerInfoObject* FEdModeLandscape::GetSelectedLandscapeLayerInfo() const
+{
+	return CurrentToolTarget.LayerInfo.Get();
 }
 
 /** FEdMode: Called when the mode is entered */
diff --git a/Engine/Source/Editor/LandscapeEditor/Private/LandscapeEdMode.h b/Engine/Source/Editor/LandscapeEditor/Private/LandscapeEdMode.h
--- a/Engine/Source/Editor/LandscapeEditor/Private/LandscapeEdMode.h
+++ b/Engine/Source/Editor/LandscapeEditor/Private/LandscapeEdMode.h
@@ -331,7 +331,8 @@ public:
 
 	/** ILandscapeEdModeInterface */
 	virtual ELandscapeToolTargetType::Type GetLandscapeToolTargetType() const override;
-	virtual FGuid GetLandscapeSelectedLayer() const override;
+	virtual const FLandscapeLayer* GetLandscapeSelectedLayer() const override;
+	virtual ULandscapeLayerInfoObject* GetSelectedLandscapeLayerInfo() const override;
 
 	/** FGCObject interface */
 	virtual void AddReferencedObjects(FReferenceCollector& Collector) override;
diff --git a/Engine/Source/Runtime/Landscape/Classes/Landscape.h b/Engine/Source/Runtime/Landscape/Classes/Landscape.h
--- a/Engine/Source/Runtime/Landscape/Classes/Landscape.h
+++ b/Engine/Source/Runtime/Landscape/Classes/Landscape.h
@@ -392,6 +392,7 @@ private:
 
 		int32 ViewMode;
 		FGuid SelectedLayer;
+		TWeakObjectPtr<ULandscapeLayerInfoObject> SelectedLayerInfoObject;
 		ELandscapeToolTargetType::Type ToolTarget;
 	};
 
diff --git a/Engine/Source/Runtime/Landscape/Private/LandscapeEditLayers.cpp b/Engine/Source/Runtime/Landscape/Private/LandscapeEditLayers.cpp
--- a/Engine/Source/Runtime/Landscape/Private/LandscapeEditLayers.cpp
+++ b/Engine/Source/Runtime/Landscape/Private/LandscapeEditLayers.cpp
@@ -4214,13 +4214,21 @@ void ALandscape::MonitorLandscapeEdModeChanges()
 		bRequiredEditingClientFullUpdate = true;
 	}
 
-	FGuid NewSelectedLayer = LandscapeEdMode ? LandscapeEdMode->GetLandscapeSelectedLayer() : FGuid();
+	const FLandscapeLayer* SelectedLayer = LandscapeEdMode ? LandscapeEdMode->GetLandscapeSelectedLayer() : nullptr;
+	FGuid NewSelectedLayer = SelectedLayer && SelectedLayer->bVisible ? SelectedLayer->Guid : FGuid();
 	if (LandscapeEdModeInfo.SelectedLayer != NewSelectedLayer)
 	{
 		LandscapeEdModeInfo.SelectedLayer = NewSelectedLayer;
 		bRequiredEditingClientFullUpdate = true;
 	}
 
+	TWeakObjectPtr<ULandscapeLayerInfoObject> NewLayerInfoObject = LandscapeEdMode->GetSelectedLandscapeLayerInfo();
+	if (LandscapeEdModeInfo.SelectedLayerInfoObject != NewLayerInfoObject)
+	{
+		LandscapeEdModeInfo.SelectedLayerInfoObject = NewLayerInfoObject;
+		bRequiredEditingClientFullUpdate = true;
+	}
+
 	if (bRequiredEditingClientFullUpdate && (LandscapeEdModeInfo.ViewMode == ELandscapeViewMode::LayerContribution))
 	{
 		RequestLayersContentUpdateForceAll(ELandscapeLayerUpdateMode::Update_Client_Editing);
@@ -4322,14 +4330,8 @@ bool ALandscape::UpdateCollisionAndClients(const TArray<ULandscapeComponent*>& I
 
 		if (LandscapeComponent->IsUpdateFlagEnabledForModes(ELandscapeComponentUpdateFlag::Component_Update_Client_Editing, InContentUpdateModes))
 		{
-			if (LandscapeEdModeInfo.ViewMode == ELandscapeViewMode::LayerContribution && LandscapeEdModeInfo.SelectedLayer.IsValid())
+			if (LandscapeEdModeInfo.ViewMode == ELandscapeViewMode::LayerContribution)
 			{
-				ULandscapeInfo* Info = LandscapeComponent->GetLandscapeInfo();
-				check(Info);
-
-				FLandscapeEditDataInterface LandscapeEdit(Info);
-				FScopedSetLandscapeEditingLayer Scope(this, LandscapeEdModeInfo.SelectedLayer);
-
 				check(ComponentSizeQuads == LandscapeComponent->ComponentSizeQuads);
 				const int32 Stride = (1 + ComponentSizeQuads);
 				const int32 ArraySize = Stride * Stride;
@@ -4342,54 +4344,50 @@ bool ALandscape::UpdateCollisionAndClients(const TArray<ULandscapeComponent*>& I
 				const int32 X2 = X1 + ComponentSizeQuads;
 				const int32 Y1 = LandscapeComponent->GetSectionBase().Y;
 				const int32 Y2 = Y1 + ComponentSizeQuads;
+				bool bLayerContributionWrittenData = false;
 
-				if (LandscapeEdModeInfo.ToolTarget == ELandscapeToolTargetType::Heightmap)
-				{
-					if (HeightData.Num() != ArraySize)
-					{
-						HeightData.AddZeroed(ArraySize);
-					}
-					LandscapeEdit.GetHeightDataFast(X1, Y1, X2, Y2, HeightData.GetData(), Stride);
-					for (int i = 0; i < ArraySize; ++i)
-					{
-						LayerContributionMaskData[i] = HeightData[i] != DefaultHeightValue ? LayerContributingValue : 0;
-					}
-				}
-				else if (LandscapeEdModeInfo.ToolTarget == ELandscapeToolTargetType::Weightmap ||
-						 LandscapeEdModeInfo.ToolTarget == ELandscapeToolTargetType::Visibility)
+				ULandscapeInfo* Info = LandscapeComponent->GetLandscapeInfo();
+				check(Info);
+				FLandscapeEditDataInterface LandscapeEdit(Info);
+
+				if (LandscapeEdModeInfo.SelectedLayer.IsValid())
 				{
-					if (WeightData.Num() != ArraySize)
-					{
-						WeightData.AddZeroed(ArraySize);
-					}
-					if (LandscapeEdModeInfo.ToolTarget == ELandscapeToolTargetType::Visibility)
+					FScopedSetLandscapeEditingLayer Scope(this, LandscapeEdModeInfo.SelectedLayer);
+					if (LandscapeEdModeInfo.ToolTarget == ELandscapeToolTargetType::Heightmap)
 					{
-						LandscapeEdit.GetWeightDataFast(ALandscapeProxy::VisibilityLayer, X1, Y1, X2, Y2, WeightData.GetData(), Stride);
+						if (HeightData.Num() != ArraySize)
+						{
+							HeightData.AddZeroed(ArraySize);
+						}
+						LandscapeEdit.GetHeightDataFast(X1, Y1, X2, Y2, HeightData.GetData(), Stride);
 						for (int i = 0; i < ArraySize; ++i)
 						{
-							LayerContributionMaskData[i] = WeightData[i] != DefaultWeightValue ? LayerContributingValue : 0;
+							LayerContributionMaskData[i] = HeightData[i] != DefaultHeightValue ? LayerContributingValue : 0;
 						}
+						bLayerContributionWrittenData = true;
 					}
-					else
+					else if (LandscapeEdModeInfo.ToolTarget == ELandscapeToolTargetType::Weightmap || LandscapeEdModeInfo.ToolTarget == ELandscapeToolTargetType::Visibility)
 					{
-						FMemory::Memzero(LayerContributionMaskDataPtr, ArraySize);
-						for (const FLandscapeInfoLayerSettings& LayerSettings : Info->Layers)
+						if (WeightData.Num() != ArraySize)
 						{
-							if (LayerSettings.LayerInfoObj == ALandscapeProxy::VisibilityLayer)
-							{
-								continue;
-							}
-							LandscapeEdit.GetWeightDataFast(LayerSettings.LayerInfoObj, X1, Y1, X2, Y2, WeightData.GetData(), Stride);
+							WeightData.AddZeroed(ArraySize);
+						}
+						ULandscapeLayerInfoObject* LayerObject = (LandscapeEdModeInfo.ToolTarget == ELandscapeToolTargetType::Visibility) ? ALandscapeProxy::VisibilityLayer : LandscapeEdModeInfo.SelectedLayerInfoObject.Get();
+						if (LayerObject)
+						{
+							LandscapeEdit.GetWeightDataFast(LayerObject, X1, Y1, X2, Y2, WeightData.GetData(), Stride);
 							for (int i = 0; i < ArraySize; ++i)
 							{
-								if (WeightData[i] != DefaultWeightValue)
-								{
-									LayerContributionMaskData[i] = LayerContributingValue;
-								}
+								LayerContributionMaskData[i] = WeightData[i] != DefaultWeightValue ? LayerContributingValue : 0;
 							}
+							bLayerContributionWrittenData = true;
 						}
 					}
 				}
+				if (!bLayerContributionWrittenData)
+				{
+					FMemory::Memzero(LayerContributionMaskDataPtr, ArraySize);
+				}
 				LandscapeEdit.SetLayerContributionData(X1, Y1, X2, Y2, LayerContributionMaskDataPtr, 0);
 			}
 		}
diff --git a/Engine/Source/Runtime/Landscape/Public/LandscapeEdit.h b/Engine/Source/Runtime/Landscape/Public/LandscapeEdit.h
--- a/Engine/Source/Runtime/Landscape/Public/LandscapeEdit.h
+++ b/Engine/Source/Runtime/Landscape/Public/LandscapeEdit.h
@@ -48,7 +48,8 @@ class ILandscapeEdModeInterface
 {
 public:
 	virtual ELandscapeToolTargetType::Type GetLandscapeToolTargetType() const = 0;
-	virtual FGuid GetLandscapeSelectedLayer() const = 0;
+	virtual const FLandscapeLayer* GetLandscapeSelectedLayer() const = 0;
+	virtual ULandscapeLayerInfoObject* GetSelectedLandscapeLayerInfo() const = 0;
 };
 
 struct FLandscapeTextureDataInfo
