Fix missing include causing incremental errors in UE4/Main
#rnx

diff --git a/Engine/Source/Runtime/Landscape/Private/LandscapeEditLayers.cpp b/Engine/Source/Runtime/Landscape/Private/LandscapeEditLayers.cpp
--- a/Engine/Source/Runtime/Landscape/Private/LandscapeEditLayers.cpp
+++ b/Engine/Source/Runtime/Landscape/Private/LandscapeEditLayers.cpp
@@ -33,6 +33,7 @@ LandscapeEditLayers.cpp: Landscape editing layers mode
 #include "LandscapeRender.h"
 #include "LandscapeInfoMap.h"
 #include "Misc/MessageDialog.h"
+#include "GameFramework/WorldSettings.h"
 #endif
 
 #define LOCTEXT_NAMESPACE "Landscape"
