- Landscape Layer System: Fix Timing issue between Landscape Import and Texture DDC. Add a check in FLandscapeTextureDataInfo to flag any remaining/future cases.

#rnx

diff --git a/Engine/Source/Runtime/Landscape/Private/LandscapeEditInterface.cpp b/Engine/Source/Runtime/Landscape/Private/LandscapeEditInterface.cpp
--- a/Engine/Source/Runtime/Landscape/Private/LandscapeEditInterface.cpp
+++ b/Engine/Source/Runtime/Landscape/Private/LandscapeEditInterface.cpp
@@ -5562,6 +5562,7 @@ void FLandscapeEditDataInterface::GetXYOffsetDataFast(const int32 X1, const int3
 FLandscapeTextureDataInfo::FLandscapeTextureDataInfo(UTexture2D* InTexture)
 :	Texture(InTexture)
 {
+	check(InTexture->IsAsyncCacheComplete());
 	MipInfo.AddZeroed(Texture->Source.GetNumMips());
 	Texture->SetFlags(RF_Transactional);
 	Texture->TemporarilyDisableStreaming();
diff --git a/Engine/Source/Runtime/Landscape/Private/LandscapeEditLayers.cpp b/Engine/Source/Runtime/Landscape/Private/LandscapeEditLayers.cpp
--- a/Engine/Source/Runtime/Landscape/Private/LandscapeEditLayers.cpp
+++ b/Engine/Source/Runtime/Landscape/Private/LandscapeEditLayers.cpp
@@ -2544,7 +2544,7 @@ bool ALandscape::PrepareLayersHeightmapTextureResources(bool bInWaitForStreaming
 				{
 					UTexture2D* LayerHeightmap = ComponentLayerData->HeightmapData.Texture;
 
-					if (LayerHeightmap->IsAsyncCacheComplete() && LayerHeightmap->Resource == nullptr)
+					if ((LayerHeightmap->IsAsyncCacheComplete() || bInWaitForStreaming) && LayerHeightmap->Resource == nullptr)
 					{
 						LayerHeightmap->FinishCachePlatformData();
 
@@ -3359,7 +3359,7 @@ bool ALandscape::PrepareLayersWeightmapTextureResources(bool bInWaitForStreaming
 				{
 					for (UTexture2D* LayerWeightmap : ComponentLayerData->WeightmapData.Textures)
 					{
-						if (LayerWeightmap->IsAsyncCacheComplete() && LayerWeightmap->Resource == nullptr)
+						if ((LayerWeightmap->IsAsyncCacheComplete() || bInWaitForStreaming) && LayerWeightmap->Resource == nullptr)
 						{
 							LayerWeightmap->FinishCachePlatformData();
 
@@ -4737,6 +4737,15 @@ void ALandscapeProxy::InitializeLayerWithEmptyContent(const FGuid& InLayerGuid)
 
 		Component->AddDefaultLayerData(InLayerGuid, *ComponentsUsingHeightmap, CreatedHeightmapTextures);
 	}
+
+	// Finish caching
+	for (TPair<UTexture2D*, UTexture2D*> Pair : CreatedHeightmapTextures)
+	{
+		if (Pair.Value != nullptr && !Pair.Value->IsAsyncCacheComplete())
+		{
+			Pair.Value->FinishCachePlatformData();
+		}
+	}
 }
 #endif
 
