- Landscape Layer System optimizations:
- Optimized update of reserved layer for splines
- Optimized Flatten Tool while Layer Contribution Visualizer is activated
- Optimized fetch of landscape compoment editing layer data
- Removed unnecessary texture copy while regenerating combined layers heightmaps
#rnx

diff --git a/Engine/Source/Editor/LandscapeEditor/Private/LandscapeEdMode.cpp b/Engine/Source/Editor/LandscapeEditor/Private/LandscapeEdMode.cpp
--- a/Engine/Source/Editor/LandscapeEditor/Private/LandscapeEdMode.cpp
+++ b/Engine/Source/Editor/LandscapeEditor/Private/LandscapeEdMode.cpp
@@ -583,6 +583,10 @@ void FEdModeLandscape::Enter()
 
 		if(Landscape && Landscape->HasLayersContent())
 		{
+			if (Landscape->GetLandscapeSplinesReservedLayer())
+			{
+				Landscape->UpdateLandscapeSplines();
+			}
 			Landscape->RequestLayersContentUpdateForceAll();
 		}
 	}
@@ -2419,7 +2423,12 @@ int32 FEdModeLandscape::UpdateLandscapeList()
 		{
 			ULandscapeInfo* LandscapeInfo = It.Value();
 			if (LandscapeInfo && !LandscapeInfo->IsPendingKill())
-			{				
+			{
+				if (ALandscape* Landscape = LandscapeInfo->LandscapeActor.Get())
+				{
+					Landscape->RegisterLandscapeEdMode(this);
+				}
+
 				ALandscapeProxy* LandscapeProxy = LandscapeInfo->GetLandscapeProxy();
 				if (LandscapeProxy)
 				{
diff --git a/Engine/Source/Editor/LandscapeEditor/Private/LandscapeEditorDetailCustomization_Layers.cpp b/Engine/Source/Editor/LandscapeEditor/Private/LandscapeEditorDetailCustomization_Layers.cpp
--- a/Engine/Source/Editor/LandscapeEditor/Private/LandscapeEditorDetailCustomization_Layers.cpp
+++ b/Engine/Source/Editor/LandscapeEditor/Private/LandscapeEditorDetailCustomization_Layers.cpp
@@ -403,6 +403,9 @@ TSharedPtr<SWidget> FLandscapeEditorCustomNodeBuilder_Layers::OnLayerContextMenu
 				{
 					FUIAction RemoveReserveLayerAction = FUIAction(FExecuteAction::CreateLambda([SharedThis, InLayerIndex] { SharedThis->SetLandscapeSplinesReservedLayer(INDEX_NONE); }));
 					MenuBuilder.AddMenuEntry(LOCTEXT("RemoveReserveLayerForSplines", "Remove Reserve for Splines"), LOCTEXT("RemoveReserveLayerForSplinesTooltip", "Remove reservation of Layer for Landscape Splines"), FSlateIcon(), RemoveReserveLayerAction);
+
+					FUIAction ForceUpdateSplinesAction = FUIAction(FExecuteAction::CreateLambda([SharedThis, InLayerIndex] { SharedThis->ForceUpdateSplines(); }));
+					MenuBuilder.AddMenuEntry(LOCTEXT("UpdateSplines", "Update Splines"), LOCTEXT("UpdateSplinesTooltip", "Update Landscape Splines"), FSlateIcon(), ForceUpdateSplinesAction);
 				}
 			}
 		}
@@ -441,6 +444,18 @@ TSharedPtr<SWidget> FLandscapeEditorCustomNodeBuilder_Layers::OnLayerContextMenu
 	return NULL;
 }
 
+void FLandscapeEditorCustomNodeBuilder_Layers::ForceUpdateSplines()
+{
+	FEdModeLandscape* LandscapeEdMode = GetEditorMode();
+	ALandscape* Landscape = LandscapeEdMode ? LandscapeEdMode->GetLandscape() : nullptr;
+	if (Landscape && Landscape->HasLayersContent() && Landscape->GetLandscapeSplinesReservedLayer())
+	{
+		const bool bUpdateOnlySelection = false;
+		const bool bForceUpdate = true;
+		Landscape->UpdateLandscapeSplines(FGuid(), bUpdateOnlySelection, bForceUpdate);
+	}
+}
+
 void FLandscapeEditorCustomNodeBuilder_Layers::SetLandscapeSplinesReservedLayer(int32 InLayerIndex)
 {
 	FEdModeLandscape* LandscapeEdMode = GetEditorMode();
diff --git a/Engine/Source/Editor/LandscapeEditor/Private/LandscapeEditorDetailCustomization_Layers.h b/Engine/Source/Editor/LandscapeEditor/Private/LandscapeEditorDetailCustomization_Layers.h
--- a/Engine/Source/Editor/LandscapeEditor/Private/LandscapeEditorDetailCustomization_Layers.h
+++ b/Engine/Source/Editor/LandscapeEditor/Private/LandscapeEditorDetailCustomization_Layers.h
@@ -86,6 +86,7 @@ protected:
 	FText GetLayerDisplayName(int32 InLayerIndex) const;
 	bool IsLayerEditionEnabled(int32 InLayerIndex) const;
 	EVisibility GetLayerAlphaVisibility(int32 InLayerIndex) const;
+	void ForceUpdateSplines();
 
 	TOptional<float> GetLayerAlpha(int32 InLayerIndex) const;
 	void SetLayerAlpha(float InAlpha, int32 InLayerIndex, bool bCommit);
diff --git a/Engine/Source/Runtime/Landscape/Classes/Landscape.h b/Engine/Source/Runtime/Landscape/Classes/Landscape.h
--- a/Engine/Source/Runtime/Landscape/Classes/Landscape.h
+++ b/Engine/Source/Runtime/Landscape/Classes/Landscape.h
@@ -263,8 +263,8 @@ public:
 	LANDSCAPE_API const struct FLandscapeLayer* GetLayer(int32 InLayerIndex) const;
 	LANDSCAPE_API const struct FLandscapeLayer* GetLayer(const FGuid& InLayerGuid) const;
 	LANDSCAPE_API void ForEachLayer(TFunctionRef<void(struct FLandscapeLayer&)> Fn);
-	LANDSCAPE_API void ClearLayer(int32 InLayerIndex, bool bInUpdateCollision = true);
-	LANDSCAPE_API void ClearLayer(const FGuid& InLayerGuid, bool bInUpdateCollision = true);
+	LANDSCAPE_API void ClearLayer(int32 InLayerIndex, TSet<ULandscapeComponent*>* InComponents = nullptr);
+	LANDSCAPE_API void ClearLayer(const FGuid& InLayerGuid, TSet<ULandscapeComponent*>* InComponents = nullptr);
 	LANDSCAPE_API void DeleteLayer(int32 InLayerIndex);
 	LANDSCAPE_API void DeleteLayers();
 	LANDSCAPE_API void SetEditingLayer(const FGuid& InLayerGuid = FGuid());
@@ -273,7 +273,7 @@ public:
 	LANDSCAPE_API bool IsMaxLayersReached() const;
 	LANDSCAPE_API void ShowOnlySelectedLayer(int32 InLayerIndex);
 	LANDSCAPE_API void ShowAllLayers();
-	LANDSCAPE_API void UpdateLandscapeSplines(const FGuid& InLayerGuid = FGuid(), bool bUpdateOnlySelected = false);
+	LANDSCAPE_API void UpdateLandscapeSplines(const FGuid& InLayerGuid = FGuid(), bool bInUpdateOnlySelected = false, bool bInForceUpdateAllCompoments = false);
 	LANDSCAPE_API void SetLandscapeSplinesReservedLayer(int32 InLayerIndex);
 	LANDSCAPE_API struct FLandscapeLayer* GetLandscapeSplinesReservedLayer();
 	LANDSCAPE_API const struct FLandscapeLayer* GetLandscapeSplinesReservedLayer() const;
@@ -300,7 +300,7 @@ public:
 private:
 	void TickLayers(float DeltaTime, ELevelTick TickType, FActorTickFunction& ThisTickFunction);
 	void CreateLayersRenderingResource();
-	void UpdateLayersContent(bool bInWaitForStreaming = false);
+	void UpdateLayersContent(bool bInWaitForStreaming = false, bool bInSkipMonitorLandscapeEdModeChanges = false);
 	void MonitorShaderCompilation();
 	void MonitorLandscapeEdModeChanges();
 	int32 RegenerateLayersHeightmaps(const TArray<ULandscapeComponent*>& InLandscapeComponents, bool bInWaitForStreaming);
@@ -383,6 +383,11 @@ public:
 	TArray<UTextureRenderTarget2D*> WeightmapRTList;
 
 private:
+
+	/** Components affected by landscape splines (used to partially clear Layer Reserved for Splines) */
+	UPROPERTY(Transient)
+	TSet<ULandscapeComponent*> LandscapeSplinesAffectedComponents;
+
 	/** Provides information from LandscapeEdMode */
 	ILandscapeEdModeInterface* LandscapeEdMode;
 
diff --git a/Engine/Source/Runtime/Landscape/Classes/LandscapeComponent.h b/Engine/Source/Runtime/Landscape/Classes/LandscapeComponent.h
--- a/Engine/Source/Runtime/Landscape/Classes/LandscapeComponent.h
+++ b/Engine/Source/Runtime/Landscape/Classes/LandscapeComponent.h
@@ -412,6 +412,11 @@ private:
 
 	UPROPERTY()
 	TMap<FGuid, FLandscapeLayerComponentData> LayersData;
+
+	/** Compoment's Data for Editing Layer */
+	FGuid LandscapeEditingLayer;
+	mutable FGuid CachedEditingLayer;
+	mutable FLandscapeLayerComponentData* CachedEditingLayerData;
 		
 	// Final layer data
 	UPROPERTY(Transient)
@@ -609,6 +614,7 @@ public:
 	LANDSCAPE_API void RemoveLayerData(const FGuid& InLayerGuid);
 	LANDSCAPE_API void ForEachLayer(TFunctionRef<void(const FGuid&, struct FLandscapeLayerComponentData&)> Fn);
 
+	LANDSCAPE_API void SetEditingLayer(const FGuid& InEditingLayer);
 	FLandscapeLayerComponentData* GetEditingLayer();
 	const FLandscapeLayerComponentData* GetEditingLayer() const;
 	FGuid GetEditingLayerGUID() const;
diff --git a/Engine/Source/Runtime/Landscape/Classes/LandscapeInfo.h b/Engine/Source/Runtime/Landscape/Classes/LandscapeInfo.h
--- a/Engine/Source/Runtime/Landscape/Classes/LandscapeInfo.h
+++ b/Engine/Source/Runtime/Landscape/Classes/LandscapeInfo.h
@@ -149,8 +149,8 @@ public:
 	LANDSCAPE_API bool GetLandscapeExtent(int32& MinX, int32& MinY, int32& MaxX, int32& MaxY) const;
 	LANDSCAPE_API void ExportHeightmap(const FString& Filename);
 	LANDSCAPE_API void ExportLayer(ULandscapeLayerInfoObject* LayerInfo, const FString& Filename);
-	LANDSCAPE_API bool ApplySplines(bool bOnlySelected);
-	bool ApplySplinesInternal(bool bOnlySelected, ALandscapeProxy* Landscape);
+	LANDSCAPE_API bool ApplySplines(bool bOnlySelected, TSet<ULandscapeComponent*>* OutModifiedComponents = nullptr);
+	bool ApplySplinesInternal(bool bOnlySelected, ALandscapeProxy* Proxy, TSet<ULandscapeComponent*>* OutModifiedComponents);
 
 	LANDSCAPE_API bool GetSelectedExtent(int32& MinX, int32& MinY, int32& MaxX, int32& MaxY) const;
 	FVector GetLandscapeCenterPos(float& LengthZ, int32 MinX = MAX_int32, int32 MinY = MAX_int32, int32 MaxX = MIN_int32, int32 MaxY = MIN_int32);
diff --git a/Engine/Source/Runtime/Landscape/Classes/LandscapeSplinesComponent.h b/Engine/Source/Runtime/Landscape/Classes/LandscapeSplinesComponent.h
--- a/Engine/Source/Runtime/Landscape/Classes/LandscapeSplinesComponent.h
+++ b/Engine/Source/Runtime/Landscape/Classes/LandscapeSplinesComponent.h
@@ -221,6 +221,6 @@ public:
 	friend class ULandscapeSplineSegment;
 #if WITH_EDITOR
 	// TODO - move this out of ULandscapeInfo
-	friend bool ULandscapeInfo::ApplySplinesInternal(bool bOnlySelected, ALandscapeProxy* Landscape);
+	friend bool ULandscapeInfo::ApplySplinesInternal(bool bOnlySelected, ALandscapeProxy* Landscape, TSet<ULandscapeComponent*>* OutModifiedComponents);
 #endif
 };
diff --git a/Engine/Source/Runtime/Landscape/Private/Landscape.cpp b/Engine/Source/Runtime/Landscape/Private/Landscape.cpp
--- a/Engine/Source/Runtime/Landscape/Private/Landscape.cpp
+++ b/Engine/Source/Runtime/Landscape/Private/Landscape.cpp
@@ -141,6 +141,7 @@ FAutoConsoleCommand CmdPrintNumLandscapeShadows(
 ULandscapeComponent::ULandscapeComponent(const FObjectInitializer& ObjectInitializer)
 : Super(ObjectInitializer)
 #if WITH_EDITORONLY_DATA
+, CachedEditingLayerData(nullptr)
 , LayerUpdateFlagPerMode(0)
 , WeightmapsHash(0)
 #endif
@@ -1537,16 +1538,29 @@ TArray<FWeightmapLayerAllocationInfo>& ULandscapeComponent::GetWeightmapLayerAll
 
 #if WITH_EDITOR
 
+void ULandscapeComponent::SetEditingLayer(const FGuid& InEditingLayer)
+{
+	LandscapeEditingLayer = InEditingLayer;
+}
+
 FLandscapeLayerComponentData* ULandscapeComponent::GetEditingLayer()
-{	
-	ALandscape* Landscape = GetLandscapeActor();
-	return Landscape != nullptr ? LayersData.Find(Landscape->GetEditingLayer()) : nullptr;
+{
+	if (CachedEditingLayer != LandscapeEditingLayer)
+	{
+		CachedEditingLayer = LandscapeEditingLayer;
+		CachedEditingLayerData = CachedEditingLayer.IsValid() ? LayersData.Find(CachedEditingLayer) : nullptr;
+	}
+	return CachedEditingLayerData;
 }
 
 const FLandscapeLayerComponentData* ULandscapeComponent::GetEditingLayer() const
 {
-	ALandscape* Landscape = GetLandscapeActor();
-	return Landscape != nullptr ? LayersData.Find(Landscape->GetEditingLayer()) : nullptr;
+	if (CachedEditingLayer != LandscapeEditingLayer)
+	{
+		CachedEditingLayer = LandscapeEditingLayer;
+		CachedEditingLayerData = CachedEditingLayer.IsValid() ? const_cast<TMap<FGuid, FLandscapeLayerComponentData>&>(LayersData).Find(CachedEditingLayer) : nullptr;
+	}
+	return CachedEditingLayerData;
 }
 
 FGuid ULandscapeComponent::GetEditingLayerGUID() const
@@ -1581,6 +1595,7 @@ void ULandscapeComponent::ForEachLayer(TFunctionRef<void(const FGuid&, struct FL
 void ULandscapeComponent::AddLayerData(const FGuid& InLayerGuid, const FLandscapeLayerComponentData& InData)
 {
 	Modify();
+	check(!LandscapeEditingLayer.IsValid());
 	FLandscapeLayerComponentData& Data = LayersData.FindOrAdd(InLayerGuid);
 	Data = InData;
 }
@@ -1666,6 +1681,7 @@ void ULandscapeComponent::AddDefaultLayerData(const FGuid& InLayerGuid, const TA
 void ULandscapeComponent::RemoveLayerData(const FGuid& InLayerGuid)
 {
 	Modify();
+	check(!LandscapeEditingLayer.IsValid());
 	LayersData.Remove(InLayerGuid);
 }
 
diff --git a/Engine/Source/Runtime/Landscape/Private/LandscapeEditInterface.cpp b/Engine/Source/Runtime/Landscape/Private/LandscapeEditInterface.cpp
--- a/Engine/Source/Runtime/Landscape/Private/LandscapeEditInterface.cpp
+++ b/Engine/Source/Runtime/Landscape/Private/LandscapeEditInterface.cpp
@@ -1803,7 +1803,7 @@ void FLandscapeEditDataInterface::DeleteLayer(ULandscapeLayerInfoObject* LayerIn
 
 	if (LandscapeInfo->LandscapeActor)
 	{
-		LandscapeInfo->LandscapeActor->RequestLayersContentUpdateForceAll();
+		LandscapeInfo->LandscapeActor->RequestLayersContentUpdate(ELandscapeLayerUpdateMode::Update_All);
 	}
 }
 
diff --git a/Engine/Source/Runtime/Landscape/Private/LandscapeEditLayers.cpp b/Engine/Source/Runtime/Landscape/Private/LandscapeEditLayers.cpp
--- a/Engine/Source/Runtime/Landscape/Private/LandscapeEditLayers.cpp
+++ b/Engine/Source/Runtime/Landscape/Private/LandscapeEditLayers.cpp
@@ -2730,14 +2730,7 @@ int32 ALandscape::RegenerateLayersHeightmaps(const TArray<ULandscapeComponent*>&
 
 			if (Layer.bVisible)
 			{
-				// Draw each Combined RT into a Non Atlas RT format to be use as base for all brush rendering
-				if (Layer.Brushes.Num() > 0)
-				{
-					CopyLayersTexture(CombinedHeightmapNonAtlasRT, LandscapeScratchRT1);
-					PrintLayersDebugRT(OutputDebugName ? FString::Printf(TEXT("LS Height: %s Component %s += -> CombinedNonAtlas %s"), *Layer.Name.ToString(), *CombinedHeightmapNonAtlasRT->GetName(), *LandscapeScratchRT1->GetName()) : TEXT(""), LandscapeScratchRT1);
-				}
-
-				// Draw each brushes				
+				// Draw each brushes
 				for (int32 i = 0; i < Layer.HeightmapBrushOrderIndices.Num(); ++i)
 				{
 					// TODO: handle conversion from float to RG8 by using material params to write correct values
@@ -4050,7 +4043,8 @@ void ALandscape::RequestLayersInitialization(bool bInRequestContentUpdate)
 		return;
 	}
 
-	bLandscapeLayersAreInitialized = false; 	
+	bLandscapeLayersAreInitialized = false;
+	LandscapeSplinesAffectedComponents.Empty();
 
 	if (bInRequestContentUpdate)
 	{
@@ -4277,7 +4271,7 @@ void ALandscape::MonitorShaderCompilation()
 	}
 }
 
-void ALandscape::UpdateLayersContent(bool bInWaitForStreaming)
+void ALandscape::UpdateLayersContent(bool bInWaitForStreaming, bool bInSkipMonitorLandscapeEdModeChanges)
 {
 	if (GetLandscapeInfo() == nullptr || !CanHaveLayersContent())
 	{
@@ -4289,7 +4283,10 @@ void ALandscape::UpdateLayersContent(bool bInWaitForStreaming)
 		InitializeLayers();
 	}
 
-	MonitorLandscapeEdModeChanges();
+	if (!bInSkipMonitorLandscapeEdModeChanges)
+	{
+		MonitorLandscapeEdModeChanges();
+	}
 	MonitorShaderCompilation();
 
 	if (LayerContentUpdateModes == 0)
@@ -4448,7 +4445,8 @@ void ALandscape::OnPreSave()
 void ALandscape::ForceUpdateLayersContent()
 {
 	const bool bWaitForStreaming = true;
-	UpdateLayersContent(bWaitForStreaming);
+	const bool bInSkipMonitorLandscapeEdModeChanges = true;
+	UpdateLayersContent(bWaitForStreaming, bInSkipMonitorLandscapeEdModeChanges);
 }
 
 void ALandscape::TickLayers(float DeltaTime, ELevelTick TickType, FActorTickFunction& ThisTickFunction)
@@ -4939,16 +4937,16 @@ void ALandscape::DeleteLayer(int32 InLayerIndex)
 	RequestLayersContentUpdateForceAll();
 }
 
-void ALandscape::ClearLayer(int32 InLayerIndex, bool bInUpdateCollision)
+void ALandscape::ClearLayer(int32 InLayerIndex, TSet<ULandscapeComponent*>* InComponents)
 {
 	const FLandscapeLayer* Layer = GetLayer(InLayerIndex);
 	if (Layer)
 	{
-		ClearLayer(Layer->Guid, bInUpdateCollision);
+		ClearLayer(Layer->Guid, InComponents);
 	}
 }
 
-void ALandscape::ClearLayer(const FGuid& InLayerGuid, bool bInUpdateCollision)
+void ALandscape::ClearLayer(const FGuid& InLayerGuid, TSet<ULandscapeComponent*>* InComponents)
 {
 	ensure(HasLayersContent());
 
@@ -4960,7 +4958,7 @@ void ALandscape::ClearLayer(const FGuid& InLayerGuid, bool bInUpdateCollision)
 	}
 
 	Modify();
-	FScopedSetLandscapeEditingLayer Scope(this, Layer->Guid, [=] { RequestLayersContentUpdateForceAll(); });
+	FScopedSetLandscapeEditingLayer Scope(this, Layer->Guid, [=] { RequestLayersContentUpdate(ELandscapeLayerUpdateMode::Update_All); });
 
 	TArray<uint16> NewHeightData;
 	NewHeightData.AddZeroed(FMath::Square(ComponentSizeQuads + 1));
@@ -4980,50 +4978,68 @@ void ALandscape::ClearLayer(const FGuid& InLayerGuid, bool bInUpdateCollision)
 		NewHeightFlagsData.AddZeroed(FMath::Square(ComponentSizeQuads + 1));
 	}
 
-	FLandscapeEditDataInterface LandscapeEdit(LandscapeInfo);
-	LandscapeInfo->ForAllLandscapeProxies([&](ALandscapeProxy* Proxy)
+	TArray<ULandscapeComponent*> Components;
+	if (InComponents)
 	{
-		Proxy->Modify();
-		for (ULandscapeComponent* Component : Proxy->LandscapeComponents)
+		Components.Reserve(InComponents->Num());
+		for (ULandscapeComponent* Compoment : *InComponents)
 		{
-			Component->Modify();
-
-			int32 MinX = MAX_int32;
-			int32 MinY = MAX_int32;
-			int32 MaxX = MIN_int32;
-			int32 MaxY = MIN_int32;
-			Component->GetComponentExtent(MinX, MinY, MaxX, MaxY);
-			check(ComponentSizeQuads == (MaxX - MinX));
-			check(ComponentSizeQuads == (MaxY - MinY));
-
-			TArray<uint16> OldHeightData;
-			OldHeightData.AddZeroed((1 + MaxY - MinY) * (1 + MaxX - MinX));
-			LandscapeEdit.GetHeightData(MinX, MinY, MaxX, MaxY, OldHeightData.GetData(), 0);
-
-			TArray<uint8> OldHeightAlphaBlendData;
-			TArray<uint8> OldHeightFlagsData;
-			if (Layer->BlendMode == LSBM_AlphaBlend)
+			if (Compoment)
 			{
-				OldHeightAlphaBlendData.AddZeroed((1 + MaxY - MinY) * (1 + MaxX - MinX));
-				LandscapeEdit.GetHeightAlphaBlendData(MinX, MinY, MaxX, MaxY, OldHeightAlphaBlendData.GetData(), 0);
-				OldHeightFlagsData.AddZeroed((1 + MaxY - MinY) * (1 + MaxX - MinX));
-				LandscapeEdit.GetHeightFlagsData(MinX, MinY, MaxX, MaxY, OldHeightFlagsData.GetData(), 0);
+				Components.Add(Compoment);
 			}
+		}
+	}
+	else
+	{
+		LandscapeInfo->ForAllLandscapeProxies([&](ALandscapeProxy* Proxy)
+		{
+			Components.Append(Proxy->LandscapeComponents);
+		});
+	}
 
-			if ((FMemory::Memcmp(OldHeightData.GetData(), NewHeightData.GetData(), NewHeightData.Num() * NewHeightData.GetTypeSize()) != 0) ||
-				(FMemory::Memcmp(OldHeightAlphaBlendData.GetData(), NewHeightAlphaBlendData.GetData(), NewHeightAlphaBlendData.Num() * NewHeightAlphaBlendData.GetTypeSize()) != 0) ||
-				(FMemory::Memcmp(OldHeightFlagsData.GetData(), NewHeightFlagsData.GetData(), NewHeightFlagsData.Num() * NewHeightFlagsData.GetTypeSize()) != 0))
-			{
-				LandscapeEdit.SetHeightData(MinX, MinY, MaxX, MaxY, NewHeightData.GetData(), 0, false, nullptr, NewHeightAlphaBlendData.GetData(), NewHeightFlagsData.GetData(), false, nullptr, nullptr, true, bInUpdateCollision);
-			}
+	FLandscapeEditDataInterface LandscapeEdit(LandscapeInfo);
+	for (ULandscapeComponent* Component : Components)
+	{
+		ALandscapeProxy* Proxy = Component->GetLandscapeProxy();
+		Proxy->Modify();
+		Component->Modify();
+
+		int32 MinX = MAX_int32;
+		int32 MinY = MAX_int32;
+		int32 MaxX = MIN_int32;
+		int32 MaxY = MIN_int32;
+		Component->GetComponentExtent(MinX, MinY, MaxX, MaxY);
+		check(ComponentSizeQuads == (MaxX - MinX));
+		check(ComponentSizeQuads == (MaxY - MinY));
+
+		TArray<uint16> OldHeightData;
+		OldHeightData.AddZeroed((1 + MaxY - MinY) * (1 + MaxX - MinX));
+		LandscapeEdit.GetHeightData(MinX, MinY, MaxX, MaxY, OldHeightData.GetData(), 0);
+
+		TArray<uint8> OldHeightAlphaBlendData;
+		TArray<uint8> OldHeightFlagsData;
+		if (Layer->BlendMode == LSBM_AlphaBlend)
+		{
+			OldHeightAlphaBlendData.AddZeroed((1 + MaxY - MinY) * (1 + MaxX - MinX));
+			LandscapeEdit.GetHeightAlphaBlendData(MinX, MinY, MaxX, MaxY, OldHeightAlphaBlendData.GetData(), 0);
+			OldHeightFlagsData.AddZeroed((1 + MaxY - MinY) * (1 + MaxX - MinX));
+			LandscapeEdit.GetHeightFlagsData(MinX, MinY, MaxX, MaxY, OldHeightFlagsData.GetData(), 0);
+		}
 
-			// Clear weight maps
-			for (FLandscapeInfoLayerSettings& LayerSettings : LandscapeInfo->Layers)
-			{
-				Component->DeleteLayer(LayerSettings.LayerInfoObj, LandscapeEdit);
-			}
+		if ((FMemory::Memcmp(OldHeightData.GetData(), NewHeightData.GetData(), NewHeightData.Num() * NewHeightData.GetTypeSize()) != 0) ||
+			(FMemory::Memcmp(OldHeightAlphaBlendData.GetData(), NewHeightAlphaBlendData.GetData(), NewHeightAlphaBlendData.Num() * NewHeightAlphaBlendData.GetTypeSize()) != 0) ||
+			(FMemory::Memcmp(OldHeightFlagsData.GetData(), NewHeightFlagsData.GetData(), NewHeightFlagsData.Num() * NewHeightFlagsData.GetTypeSize()) != 0))
+		{
+			LandscapeEdit.SetHeightData(MinX, MinY, MaxX, MaxY, NewHeightData.GetData(), 0, false, nullptr, NewHeightAlphaBlendData.GetData(), NewHeightFlagsData.GetData());
 		}
-	});
+
+		// Clear weight maps
+		for (FLandscapeInfoLayerSettings& LayerSettings : LandscapeInfo->Layers)
+		{
+			Component->DeleteLayer(LayerSettings.LayerInfoObj, LandscapeEdit);
+		}
+	}
 }
 
 void ALandscape::ShowOnlySelectedLayer(int32 InLayerIndex)
@@ -5058,6 +5074,7 @@ void ALandscape::SetLandscapeSplinesReservedLayer(int32 InLayerIndex)
 	FLandscapeLayer* PreviousLayer = GetLandscapeSplinesReservedLayer();
 	if (NewLayer != PreviousLayer)
 	{
+		LandscapeSplinesAffectedComponents.Empty();
 		if (PreviousLayer)
 		{
 			ClearLayer(LandscapeSplinesTargetLayerGuid);
@@ -5098,7 +5115,7 @@ FLandscapeLayer* ALandscape::GetLandscapeSplinesReservedLayer()
 
 LANDSCAPE_API extern bool GDisableUpdateLandscapeMaterialInstances;
 
-void ALandscape::UpdateLandscapeSplines(const FGuid& InTargetLayer, bool bUpdateOnlySelected)
+void ALandscape::UpdateLandscapeSplines(const FGuid& InTargetLayer, bool bInUpdateOnlySelected, bool bInForceUpdateAllCompoments)
 {
 	check(CanHaveLayersContent());
 	ULandscapeInfo* LandscapeInfo = GetLandscapeInfo();
@@ -5106,16 +5123,19 @@ void ALandscape::UpdateLandscapeSplines(const FGuid& InTargetLayer, bool bUpdate
 	const FLandscapeLayer* TargetLayer = GetLayer(TargetLayerGuid);
 	if (LandscapeInfo && TargetLayer)
 	{
-		FScopedSetLandscapeEditingLayer Scope(this, TargetLayerGuid, [=] { this->RequestLayersContentUpdateForceAll(); });
+		FScopedSetLandscapeEditingLayer Scope(this, TargetLayerGuid, [=] { this->RequestLayersContentUpdate(ELandscapeLayerUpdateMode::Update_All); });
 		// Temporarily disable material instance updates since it will be done once at the end (requested by RequestLayersContentUpdateForceAll)
 		GDisableUpdateLandscapeMaterialInstances = true;
+		TSet<ULandscapeComponent*>* ModifiedComponent = nullptr;
 		if (LandscapeSplinesTargetLayerGuid.IsValid())
 		{
-			ClearLayer(LandscapeSplinesTargetLayerGuid, false);
+			ClearLayer(LandscapeSplinesTargetLayerGuid, (!bInForceUpdateAllCompoments && LandscapeSplinesAffectedComponents.Num()) ? &LandscapeSplinesAffectedComponents : nullptr);
+			LandscapeSplinesAffectedComponents.Empty();
+			ModifiedComponent = &LandscapeSplinesAffectedComponents;
 			// For now, in Landscape Layer System Mode with a reserved layer for splines, we always update all the splines since we clear the whole layer first
-			bUpdateOnlySelected = false;
+			bInUpdateOnlySelected = false;
 		}
-		LandscapeInfo->ApplySplines(bUpdateOnlySelected);
+		LandscapeInfo->ApplySplines(bInUpdateOnlySelected, ModifiedComponent);
 		GDisableUpdateLandscapeMaterialInstances = false;
 	}
 }
@@ -5165,6 +5185,15 @@ void ALandscape::SetEditingLayer(const FGuid& InLayerGuid)
 	}
 
 	EditingLayer = InLayerGuid;
+
+	// Propagate Editing Layer to components (will be cached)
+	LandscapeInfo->ForAllLandscapeProxies([&](ALandscapeProxy* Proxy)
+	{
+		for (ULandscapeComponent* Component : Proxy->LandscapeComponents)
+		{
+			Component->SetEditingLayer(InLayerGuid);
+		}
+	});
 }
 
 void ALandscape::SetGrassUpdateEnabled(bool bInGrassUpdateEnabled)
diff --git a/Engine/Source/Runtime/Landscape/Private/LandscapeSplineRaster.cpp b/Engine/Source/Runtime/Landscape/Private/LandscapeSplineRaster.cpp
--- a/Engine/Source/Runtime/Landscape/Private/LandscapeSplineRaster.cpp
+++ b/Engine/Source/Runtime/Landscape/Private/LandscapeSplineRaster.cpp
@@ -457,27 +457,24 @@ void RasterizeSegmentAlpha(int32& MinX, int32& MinY, int32& MaxX, int32& MaxY, F
 	LandscapeEdit.GetComponentsInRegion(MinX, MinY, MaxX, MaxY, &ModifiedComponents);
 }
 
-bool ULandscapeInfo::ApplySplines(bool bOnlySelected)
+bool ULandscapeInfo::ApplySplines(bool bOnlySelected, TSet<ULandscapeComponent*>* OutModifiedComponents)
 {
 	bool bResult = false;
 
 	ALandscape* Landscape = LandscapeActor.Get();
 	const FLandscapeLayer* Layer = Landscape ? Landscape->GetLandscapeSplinesReservedLayer() : nullptr;
 	FGuid SplinesTargetLayerGuid = Layer ? Layer->Guid : Landscape ? Landscape->GetEditingLayer() : FGuid();
-	FScopedSetLandscapeEditingLayer Scope(Landscape, SplinesTargetLayerGuid, [=] 
-	{ 
-		Landscape->RequestLayersContentUpdate(ELandscapeLayerUpdateMode::Update_All);
-	});
+	FScopedSetLandscapeEditingLayer Scope(Landscape, SplinesTargetLayerGuid, [=] { Landscape->RequestLayersContentUpdate(ELandscapeLayerUpdateMode::Update_All); });
 
-	ForAllLandscapeProxies([&bResult, bOnlySelected, this](ALandscapeProxy* Proxy)
+	ForAllLandscapeProxies([&](ALandscapeProxy* Proxy)
 	{
-		bResult |= ApplySplinesInternal(bOnlySelected, Proxy);
+		bResult |= ApplySplinesInternal(bOnlySelected, Proxy, OutModifiedComponents);
 	});
 
 	return bResult;
 }
 
-bool ULandscapeInfo::ApplySplinesInternal(bool bOnlySelected, ALandscapeProxy* Proxy)
+bool ULandscapeInfo::ApplySplinesInternal(bool bOnlySelected, ALandscapeProxy* Proxy, TSet<ULandscapeComponent*>* OutModifiedComponents)
 {
 	if (!Proxy || !Proxy->SplineComponent || Proxy->SplineComponent->ControlPoints.Num() == 0 || Proxy->SplineComponent->Segments.Num() == 0)
 	{
@@ -635,6 +632,10 @@ bool ULandscapeInfo::ApplySplinesInternal(bool bOnlySelected, ALandscapeProxy* P
 		{
 			Component->RequestHeightmapUpdate();
 			Component->RequestWeightmapUpdate();
+			if (OutModifiedComponents)
+			{
+				OutModifiedComponents->Add(Component);
+			}
 		}
 		else
 		{
