- Landscape Layer System: Optimization avoids updating Heightfield Collision Region if RecreateCollision is called in same frame.

#rnx

diff --git a/Engine/Source/Runtime/Landscape/Classes/LandscapeComponent.h b/Engine/Source/Runtime/Landscape/Classes/LandscapeComponent.h
--- a/Engine/Source/Runtime/Landscape/Classes/LandscapeComponent.h
+++ b/Engine/Source/Runtime/Landscape/Classes/LandscapeComponent.h
@@ -816,7 +816,7 @@ public:
 	 * @param XYOffsetTextureMipData: xy-offset map data
 	 * @returns True if CollisionComponent was created in this update.
 	 */
-	void UpdateCollisionHeightData(const FColor* HeightmapTextureMipData, const FColor* SimpleCollisionHeightmapTextureData, int32 ComponentX1=0, int32 ComponentY1=0, int32 ComponentX2=MAX_int32, int32 ComponentY2=MAX_int32, bool bUpdateBounds=false, const FColor* XYOffsetTextureMipData=nullptr);
+	void UpdateCollisionHeightData(const FColor* HeightmapTextureMipData, const FColor* SimpleCollisionHeightmapTextureData, int32 ComponentX1=0, int32 ComponentY1=0, int32 ComponentX2=MAX_int32, int32 ComponentY2=MAX_int32, bool bUpdateBounds=false, const FColor* XYOffsetTextureMipData=nullptr, bool bInUpdateHeightfieldRegion=true);
 
 	/**
 	 * Deletes Collision Component
@@ -825,7 +825,7 @@ public:
 
 	/** Updates collision component height data for the entire component, locking and unlocking heightmap textures
 	 */
-	void UpdateCollisionData();
+	void UpdateCollisionData(bool bInUpdateHeightfieldRegion = true);
 
 	/**
 	 * Update collision component dominant layer data
diff --git a/Engine/Source/Runtime/Landscape/Private/LandscapeEdit.cpp b/Engine/Source/Runtime/Landscape/Private/LandscapeEdit.cpp
--- a/Engine/Source/Runtime/Landscape/Private/LandscapeEdit.cpp
+++ b/Engine/Source/Runtime/Landscape/Private/LandscapeEdit.cpp
@@ -800,7 +800,7 @@ public:
 	int32 SizeVertsSquare = 0;
 };
 
-void ULandscapeComponent::UpdateCollisionHeightData(const FColor* const HeightmapTextureMipData, const FColor* const SimpleCollisionHeightmapTextureData, int32 ComponentX1/*=0*/, int32 ComponentY1/*=0*/, int32 ComponentX2/*=MAX_int32*/, int32 ComponentY2/*=MAX_int32*/, bool bUpdateBounds/*=false*/, const FColor* XYOffsetTextureMipData/*=nullptr*/)
+void ULandscapeComponent::UpdateCollisionHeightData(const FColor* const HeightmapTextureMipData, const FColor* const SimpleCollisionHeightmapTextureData, int32 ComponentX1/*=0*/, int32 ComponentY1/*=0*/, int32 ComponentX2/*=MAX_int32*/, int32 ComponentY2/*=MAX_int32*/, bool bUpdateBounds/*=false*/, const FColor* XYOffsetTextureMipData/*=nullptr*/, bool bInUpdateHeightfieldRegion/*=true*/)
 {
 	ULandscapeInfo* Info = GetLandscapeInfo();
 	ALandscapeProxy* Proxy = GetLandscapeProxy();
@@ -866,7 +866,7 @@ void ULandscapeComponent::UpdateCollisionHeightData(const FColor* const Heightma
 	}
 
 	CollisionHeightData = (uint16*)CollisionComp->CollisionHeightData.Lock(LOCK_READ_WRITE);
-
+	
 	if (XYOffsetmapTexture && MeshCollisionComponent)
 	{
 		CollisionXYOffsetData = (uint16*)MeshCollisionComponent->CollisionXYOffsetData.Lock(LOCK_READ_WRITE);
@@ -911,7 +911,7 @@ void ULandscapeComponent::UpdateCollisionHeightData(const FColor* const Heightma
 	}
 
 	// If we updated an existing component, we need to update the phys x heightfield edit data
-	if (!CreatedNew)
+	if (!CreatedNew && bInUpdateHeightfieldRegion)
 	{
 		if (MeshCollisionComponent)
 		{
@@ -967,7 +967,7 @@ void ULandscapeComponent::DestroyCollisionData()
 	}
 }
 
-void ULandscapeComponent::UpdateCollisionData()
+void ULandscapeComponent::UpdateCollisionData(bool bInUpdateHeightfieldRegion)
 {
 	TArray<uint8> CollisionMipData;
 	TArray<uint8> SimpleCollisionMipData;
@@ -987,7 +987,7 @@ void ULandscapeComponent::UpdateCollisionData()
 		(FColor*)CollisionMipData.GetData(),
 		SimpleCollisionMipLevel > CollisionMipLevel ? (FColor*)SimpleCollisionMipData.GetData() : nullptr,
 		0, 0, MAX_int32, MAX_int32, true,
-		XYOffsetmapTexture ? (FColor*)XYOffsetMipData.GetData() : nullptr);
+		XYOffsetmapTexture ? (FColor*)XYOffsetMipData.GetData() : nullptr, bInUpdateHeightfieldRegion);
 }
 
 void ULandscapeComponent::RecreateCollisionComponent(bool bUseSimpleCollision)
diff --git a/Engine/Source/Runtime/Landscape/Private/LandscapeEditLayers.cpp b/Engine/Source/Runtime/Landscape/Private/LandscapeEditLayers.cpp
--- a/Engine/Source/Runtime/Landscape/Private/LandscapeEditLayers.cpp
+++ b/Engine/Source/Runtime/Landscape/Private/LandscapeEditLayers.cpp
@@ -2825,7 +2825,10 @@ int32 ALandscape::RegenerateLayersHeightmaps(const TArray<ULandscapeComponent*>&
 			{
 				Component->UpdateCachedBounds();
 				Component->UpdateComponentToWorld();
-				Component->UpdateCollisionData();
+
+				// Avoid updating height field if we are going to recreate collision in this update
+				bool bUpdateHeightfieldRegion = !Component->IsUpdateFlagEnabledForModes(Component_Update_Recreate_Collision, HeightmapUpdateModes);
+				Component->UpdateCollisionData(bUpdateHeightfieldRegion);
 			}
 		}
 
@@ -4294,24 +4297,24 @@ void ALandscape::UpdateLayersContent(bool bInWaitForStreaming, bool bInSkipMonit
 		return;
 	}
 
-	TArray<ULandscapeComponent*> AllLandscapeComponents;
-	GetLandscapeInfo()->ForAllLandscapeProxies([&AllLandscapeComponents](ALandscapeProxy* Proxy)
-	{
-		AllLandscapeComponents.Append(Proxy->LandscapeComponents);
-	});
+		TArray<ULandscapeComponent*> AllLandscapeComponents;
+		GetLandscapeInfo()->ForAllLandscapeProxies([&AllLandscapeComponents](ALandscapeProxy* Proxy)
+		{
+			AllLandscapeComponents.Append(Proxy->LandscapeComponents);
+		});
 
-	int32 ProcessedModes = 0;
-	ProcessedModes |= RegenerateLayersHeightmaps(AllLandscapeComponents, bInWaitForStreaming);
-	ProcessedModes |= RegenerateLayersWeightmaps(AllLandscapeComponents, bInWaitForStreaming);
-	ProcessedModes |= (LayerContentUpdateModes & ELandscapeLayerUpdateMode::Update_Client_Deferred);
-	ProcessedModes |= (LayerContentUpdateModes & ELandscapeLayerUpdateMode::Update_Client_Editing);
-	LayerContentUpdateModes &= ~ProcessedModes;
+		int32 ProcessedModes = 0;
+		ProcessedModes |= RegenerateLayersHeightmaps(AllLandscapeComponents, bInWaitForStreaming);
+		ProcessedModes |= RegenerateLayersWeightmaps(AllLandscapeComponents, bInWaitForStreaming);
+		ProcessedModes |= (LayerContentUpdateModes & ELandscapeLayerUpdateMode::Update_Client_Deferred);
+		ProcessedModes |= (LayerContentUpdateModes & ELandscapeLayerUpdateMode::Update_Client_Editing);
+		LayerContentUpdateModes &= ~ProcessedModes;
 
-	if (!ALandscape::UpdateCollisionAndClients(AllLandscapeComponents, ProcessedModes))
-	{
-		LayerContentUpdateModes |= ELandscapeLayerUpdateMode::Update_Client_Deferred;
+		if (!ALandscape::UpdateCollisionAndClients(AllLandscapeComponents, ProcessedModes))
+		{
+			LayerContentUpdateModes |= ELandscapeLayerUpdateMode::Update_Client_Deferred;
+		}
 	}
-}
 
 bool ALandscape::UpdateCollisionAndClients(const TArray<ULandscapeComponent*>& InLandscapeComponents, const int32 InContentUpdateModes)
 {
