- Landscape Layer System : Optimizations of ClearLayers, Reading heightmap data, Landscape splines update
#rnx

diff --git a/Engine/Source/Runtime/Landscape/Private/Landscape.cpp b/Engine/Source/Runtime/Landscape/Private/Landscape.cpp
--- a/Engine/Source/Runtime/Landscape/Private/Landscape.cpp
+++ b/Engine/Source/Runtime/Landscape/Private/Landscape.cpp
@@ -1454,11 +1454,9 @@ void ULandscapeComponent::OnUnregister()
 UTexture2D* ULandscapeComponent::GetHeightmap(bool InReturnEditingHeightmap) const
 {
 #if WITH_EDITORONLY_DATA
-	if (InReturnEditingHeightmap && GetLandscapeProxy()->HasLayersContent())
+	if (InReturnEditingHeightmap)
 	{
-		const FLandscapeLayerComponentData* EditingLayer = GetEditingLayer();
-
-		if (EditingLayer != nullptr)
+		if (const FLandscapeLayerComponentData* EditingLayer = GetEditingLayer())
 		{
 			return EditingLayer->HeightmapData.Texture;
 		}
@@ -1471,11 +1469,9 @@ UTexture2D* ULandscapeComponent::GetHeightmap(bool InReturnEditingHeightmap) con
 const TArray<UTexture2D*>& ULandscapeComponent::GetWeightmapTextures(bool InReturnEditingWeightmap) const
 {
 #if WITH_EDITORONLY_DATA
-	if (InReturnEditingWeightmap && GetLandscapeProxy()->HasLayersContent())
+	if (InReturnEditingWeightmap)
 	{
-		const FLandscapeLayerComponentData* EditingLayer = GetEditingLayer();
-
-		if (EditingLayer != nullptr)
+		if (const FLandscapeLayerComponentData* EditingLayer = GetEditingLayer())
 		{
 			return EditingLayer->WeightmapData.Textures;
 		}
@@ -1488,11 +1484,9 @@ const TArray<UTexture2D*>& ULandscapeComponent::GetWeightmapTextures(bool InRetu
 TArray<UTexture2D*>& ULandscapeComponent::GetWeightmapTextures(bool InReturnEditingWeightmap)
 {
 #if WITH_EDITORONLY_DATA
-	if (InReturnEditingWeightmap && GetLandscapeProxy()->HasLayersContent())
+	if (InReturnEditingWeightmap)
 	{
-		FLandscapeLayerComponentData* EditingLayer = GetEditingLayer();
-
-		if (EditingLayer != nullptr)
+		if (FLandscapeLayerComponentData* EditingLayer = GetEditingLayer())
 		{
 			return EditingLayer->WeightmapData.Textures;
 		}
@@ -1505,11 +1499,9 @@ TArray<UTexture2D*>& ULandscapeComponent::GetWeightmapTextures(bool InReturnEdit
 const TArray<FWeightmapLayerAllocationInfo>& ULandscapeComponent::GetWeightmapLayerAllocations(bool InReturnEditingWeightmap) const
 {
 #if WITH_EDITORONLY_DATA
-	if (InReturnEditingWeightmap && GetLandscapeProxy()->HasLayersContent())
+	if (InReturnEditingWeightmap)
 	{
-		const FLandscapeLayerComponentData* EditingLayer = GetEditingLayer();
-
-		if (EditingLayer != nullptr)
+		if (const FLandscapeLayerComponentData* EditingLayer = GetEditingLayer())
 		{
 			return EditingLayer->WeightmapData.LayerAllocations;
 		}
@@ -1522,11 +1514,9 @@ const TArray<FWeightmapLayerAllocationInfo>& ULandscapeComponent::GetWeightmapLa
 TArray<FWeightmapLayerAllocationInfo>& ULandscapeComponent::GetWeightmapLayerAllocations(bool InReturnEditingWeightmap)
 {
 #if WITH_EDITORONLY_DATA
-	if (InReturnEditingWeightmap && GetLandscapeProxy()->HasLayersContent())
+	if (InReturnEditingWeightmap)
 	{
-		FLandscapeLayerComponentData* EditingLayer = GetEditingLayer();
-
-		if (EditingLayer != nullptr)
+		if (FLandscapeLayerComponentData* EditingLayer = GetEditingLayer())
 		{
 			return EditingLayer->WeightmapData.LayerAllocations;
 		}
@@ -1716,11 +1706,9 @@ void ULandscapeComponent::SetWeightmapLayerAllocations(const TArray<FWeightmapLa
 TArray<ULandscapeWeightmapUsage*>& ULandscapeComponent::GetWeightmapTexturesUsage(bool InReturnEditingWeightmap)
 {
 #if WITH_EDITORONLY_DATA
-	if (InReturnEditingWeightmap && GetLandscapeProxy()->HasLayersContent())
+	if (InReturnEditingWeightmap)
 	{
-		FLandscapeLayerComponentData* EditingLayer = GetEditingLayer();
-
-		if (EditingLayer != nullptr)
+		if (FLandscapeLayerComponentData* EditingLayer = GetEditingLayer())
 		{
 			return EditingLayer->WeightmapData.TextureUsages;
 		}
@@ -1733,11 +1721,9 @@ TArray<ULandscapeWeightmapUsage*>& ULandscapeComponent::GetWeightmapTexturesUsag
 const TArray<ULandscapeWeightmapUsage*>& ULandscapeComponent::GetWeightmapTexturesUsage(bool InReturnEditingWeightmap) const
 {
 #if WITH_EDITORONLY_DATA
-	if (InReturnEditingWeightmap && GetLandscapeProxy()->HasLayersContent())
+	if (InReturnEditingWeightmap)
 	{
-		const FLandscapeLayerComponentData* EditingLayer = GetEditingLayer();
-
-		if (EditingLayer != nullptr)
+		if (const FLandscapeLayerComponentData* EditingLayer = GetEditingLayer())
 		{
 			return EditingLayer->WeightmapData.TextureUsages;
 		}
diff --git a/Engine/Source/Runtime/Landscape/Private/LandscapeEditInterface.cpp b/Engine/Source/Runtime/Landscape/Private/LandscapeEditInterface.cpp
--- a/Engine/Source/Runtime/Landscape/Private/LandscapeEditInterface.cpp
+++ b/Engine/Source/Runtime/Landscape/Private/LandscapeEditInterface.cpp
@@ -952,8 +952,10 @@ FColor& FLandscapeEditDataInterface::GetHeightMapColor(const ULandscapeComponent
 		TextureData = (FColor*)TexDataInfo->GetMipData(0);	
 	}
 
-	int32 SizeU = Component->GetHeightmap(true)->Source.GetSizeX();
-	int32 SizeV = Component->GetHeightmap(true)->Source.GetSizeY();
+	// All Heightmaps of component have the same texture size
+	const FTextureSource& HeightmapTextureSource = Component->GetHeightmap()->Source;
+	int32 SizeU = HeightmapTextureSource.GetSizeX();
+	int32 SizeV = HeightmapTextureSource.GetSizeY();
 	int32 HeightmapOffsetX = Component->HeightmapScaleBias.Z * (float)SizeU;
 	int32 HeightmapOffsetY = Component->HeightmapScaleBias.W * (float)SizeV;
 
@@ -981,8 +983,8 @@ uint8 FLandscapeEditDataInterface::GetHeightMapFlagsData(const ULandscapeCompone
 	return TexData.A;
 }
 
-template<typename TDataAccess>
-void FLandscapeEditDataInterface::GetHeightDataInternal(int32& ValidX1, int32& ValidY1, int32& ValidX2, int32& ValidY2, TDataAccess& StoreData, TFunctionRef<typename TDataAccess::DataType(const ULandscapeComponent*, int32, int32, FColor*)> GetHeightMapDataFunction)
+template<typename TDataAccess, typename TGetHeightMapDataFunction>
+void FLandscapeEditDataInterface::GetHeightDataInternal(int32& ValidX1, int32& ValidY1, int32& ValidX2, int32& ValidY2, TDataAccess& StoreData, TGetHeightMapDataFunction GetHeightMapDataFunction)
 {
 	// Copy variables
 	int32 X1 = ValidX1, X2 = ValidX2, Y1 = ValidY1, Y2 = ValidY2;
diff --git a/Engine/Source/Runtime/Landscape/Private/LandscapeEditLayers.cpp b/Engine/Source/Runtime/Landscape/Private/LandscapeEditLayers.cpp
--- a/Engine/Source/Runtime/Landscape/Private/LandscapeEditLayers.cpp
+++ b/Engine/Source/Runtime/Landscape/Private/LandscapeEditLayers.cpp
@@ -4984,12 +4984,19 @@ void ALandscape::ClearLayer(const FGuid& InLayerGuid, TSet<ULandscapeComponent*>
 	TArray<ULandscapeComponent*> Components;
 	if (InComponents)
 	{
+		TSet<ALandscapeProxy*> Proxies;
 		Components.Reserve(InComponents->Num());
-		for (ULandscapeComponent* Compoment : *InComponents)
+		for (ULandscapeComponent* Component : *InComponents)
 		{
-			if (Compoment)
+			if (Component)
 			{
-				Components.Add(Compoment);
+				Components.Add(Component);
+				ALandscapeProxy* Proxy = Component->GetLandscapeProxy();
+				if (!Proxies.Find(Proxy))
+				{
+					Proxies.Add(Proxy);
+					Proxy->Modify();
+				}
 			}
 		}
 	}
@@ -4997,6 +5004,7 @@ void ALandscape::ClearLayer(const FGuid& InLayerGuid, TSet<ULandscapeComponent*>
 	{
 		LandscapeInfo->ForAllLandscapeProxies([&](ALandscapeProxy* Proxy)
 		{
+			Proxy->Modify();
 			Components.Append(Proxy->LandscapeComponents);
 		});
 	}
@@ -5004,10 +5012,6 @@ void ALandscape::ClearLayer(const FGuid& InLayerGuid, TSet<ULandscapeComponent*>
 	FLandscapeEditDataInterface LandscapeEdit(LandscapeInfo);
 	for (ULandscapeComponent* Component : Components)
 	{
-		ALandscapeProxy* Proxy = Component->GetLandscapeProxy();
-		Proxy->Modify();
-		Component->Modify();
-
 		int32 MinX = MAX_int32;
 		int32 MinY = MAX_int32;
 		int32 MaxX = MIN_int32;
@@ -5015,27 +5019,7 @@ void ALandscape::ClearLayer(const FGuid& InLayerGuid, TSet<ULandscapeComponent*>
 		Component->GetComponentExtent(MinX, MinY, MaxX, MaxY);
 		check(ComponentSizeQuads == (MaxX - MinX));
 		check(ComponentSizeQuads == (MaxY - MinY));
-
-		TArray<uint16> OldHeightData;
-		OldHeightData.AddZeroed((1 + MaxY - MinY) * (1 + MaxX - MinX));
-		LandscapeEdit.GetHeightData(MinX, MinY, MaxX, MaxY, OldHeightData.GetData(), 0);
-
-		TArray<uint8> OldHeightAlphaBlendData;
-		TArray<uint8> OldHeightFlagsData;
-		if (Layer->BlendMode == LSBM_AlphaBlend)
-		{
-			OldHeightAlphaBlendData.AddZeroed((1 + MaxY - MinY) * (1 + MaxX - MinX));
-			LandscapeEdit.GetHeightAlphaBlendData(MinX, MinY, MaxX, MaxY, OldHeightAlphaBlendData.GetData(), 0);
-			OldHeightFlagsData.AddZeroed((1 + MaxY - MinY) * (1 + MaxX - MinX));
-			LandscapeEdit.GetHeightFlagsData(MinX, MinY, MaxX, MaxY, OldHeightFlagsData.GetData(), 0);
-		}
-
-		if ((FMemory::Memcmp(OldHeightData.GetData(), NewHeightData.GetData(), NewHeightData.Num() * NewHeightData.GetTypeSize()) != 0) ||
-			(FMemory::Memcmp(OldHeightAlphaBlendData.GetData(), NewHeightAlphaBlendData.GetData(), NewHeightAlphaBlendData.Num() * NewHeightAlphaBlendData.GetTypeSize()) != 0) ||
-			(FMemory::Memcmp(OldHeightFlagsData.GetData(), NewHeightFlagsData.GetData(), NewHeightFlagsData.Num() * NewHeightFlagsData.GetTypeSize()) != 0))
-		{
-			LandscapeEdit.SetHeightData(MinX, MinY, MaxX, MaxY, NewHeightData.GetData(), 0, false, nullptr, NewHeightAlphaBlendData.GetData(), NewHeightFlagsData.GetData());
-		}
+		LandscapeEdit.SetHeightData(MinX, MinY, MaxX, MaxY, NewHeightData.GetData(), 0, false, nullptr, NewHeightAlphaBlendData.GetData(), NewHeightFlagsData.GetData());
 
 		// Clear weight maps
 		for (FLandscapeInfoLayerSettings& LayerSettings : LandscapeInfo->Layers)
diff --git a/Engine/Source/Runtime/Landscape/Public/LandscapeEdit.h b/Engine/Source/Runtime/Landscape/Public/LandscapeEdit.h
--- a/Engine/Source/Runtime/Landscape/Public/LandscapeEdit.h
+++ b/Engine/Source/Runtime/Landscape/Public/LandscapeEdit.h
@@ -275,8 +275,8 @@ private:
 		TArray<bool>& NoBorderY1, TArray<bool>& NoBorderY2, TArray<bool>& ComponentDataExist, TStoreData& StoreData);
 
 	// Generic Height Data access
-	template<typename TStoreData>
-	void GetHeightDataInternal(int32& ValidX1, int32& ValidY1, int32& ValidX2, int32& ValidY2, TStoreData& StoreData, TFunctionRef<typename TStoreData::DataType(const ULandscapeComponent*, int32, int32, FColor*)> GetHeightDataFunction);
+	template<typename TStoreData, typename TGetHeightMapDataFunction>
+	void GetHeightDataInternal(int32& ValidX1, int32& ValidY1, int32& ValidX2, int32& ValidY2, TStoreData& StoreData, TGetHeightMapDataFunction GetHeightMapDataFunction);
 
 	FORCEINLINE FColor& GetHeightMapColor(const ULandscapeComponent* Component, int32 TexU, int32 TexV, FColor* TextureData);
 
