- Landscape Layer System: Optimization Hash Collision to avoid Recreate when possible

#rnx

diff --git a/Engine/Source/Runtime/Landscape/Classes/LandscapeHeightfieldCollisionComponent.h b/Engine/Source/Runtime/Landscape/Classes/LandscapeHeightfieldCollisionComponent.h
--- a/Engine/Source/Runtime/Landscape/Classes/LandscapeHeightfieldCollisionComponent.h
+++ b/Engine/Source/Runtime/Landscape/Classes/LandscapeHeightfieldCollisionComponent.h
@@ -117,6 +117,10 @@ class ULandscapeHeightfieldCollisionComponent : public UPrimitiveComponent
 	};
 	
 #if WITH_EDITORONLY_DATA
+	friend struct FEnableCollisionHashOptimScope;
+	
+	uint32										CollisionHash = 0;
+
 	/** The collision height values. Stripped from cooked content */
 	FWordBulkData								CollisionHeightData;
 
@@ -166,6 +170,11 @@ class ULandscapeHeightfieldCollisionComponent : public UPrimitiveComponent
 		QF_NoCollision = 128,			// This quad has no collision.
 	};
 
+#if WITH_EDITORONLY_DATA
+private:
+	bool bEnableCollisionHashOptim = false;
+#endif //WITH_EDITORONLY_DATA
+
 	//~ Begin UActorComponent Interface.
 protected:
 	virtual void OnCreatePhysicsState() override;
@@ -223,6 +232,9 @@ public:
 
 	/** Modify a sub-region of the PhysX heightfield. Note that this does not update the physical material */
 	void UpdateHeightfieldRegion(int32 ComponentX1, int32 ComponentY1, int32 ComponentX2, int32 ComponentY2);
+
+	/** Computes a hash of all the data that will impact final collision */
+	virtual uint32 ComputeCollisionHash() const;
 #endif
 
 	/** Creates collision object from a cooked collision data */
@@ -238,7 +250,7 @@ public:
 	LANDSCAPE_API void SetSectionBase(FIntPoint InSectionBase);
 
 	/** Recreate heightfield and restart physics */
-	LANDSCAPE_API virtual void RecreateCollision();
+	LANDSCAPE_API virtual bool RecreateCollision();
 
 #if WITH_EDITORONLY_DATA
 	// Called from editor code to manage foliage instances on landscape.
diff --git a/Engine/Source/Runtime/Landscape/Classes/LandscapeMeshCollisionComponent.h b/Engine/Source/Runtime/Landscape/Classes/LandscapeMeshCollisionComponent.h
--- a/Engine/Source/Runtime/Landscape/Classes/LandscapeMeshCollisionComponent.h
+++ b/Engine/Source/Runtime/Landscape/Classes/LandscapeMeshCollisionComponent.h
@@ -105,7 +105,8 @@ public:
 
 	//~ Begin ULandscapeHeightfieldCollisionComponent Interface
 	virtual void CreateCollisionObject() override;
-	virtual void RecreateCollision() override;
+	virtual bool RecreateCollision() override;
+	virtual uint32 ComputeCollisionHash() const override { return 0; }
 	//~ End ULandscapeHeightfieldCollisionComponent Interface
 };
 
diff --git a/Engine/Source/Runtime/Landscape/Private/LandscapeCollision.cpp b/Engine/Source/Runtime/Landscape/Private/LandscapeCollision.cpp
--- a/Engine/Source/Runtime/Landscape/Private/LandscapeCollision.cpp
+++ b/Engine/Source/Runtime/Landscape/Private/LandscapeCollision.cpp
@@ -1077,6 +1077,28 @@ void ULandscapeMeshCollisionComponent::DestroyComponent(bool bPromoteChildren/*=
 }
 
 #if WITH_EDITOR
+uint32 ULandscapeHeightfieldCollisionComponent::ComputeCollisionHash() const
+{
+	uint32 Hash = 0;
+		
+	Hash = HashCombine(GetTypeHash(SimpleCollisionSizeQuads), Hash);
+	Hash = HashCombine(GetTypeHash(CollisionSizeQuads), Hash);
+	Hash = HashCombine(GetTypeHash(CollisionScale), Hash);
+
+	FTransform ComponentTransform = GetComponentToWorld();
+	Hash = FCrc::MemCrc32(&ComponentTransform, sizeof(ComponentTransform));
+
+	const void* HeightBuffer = CollisionHeightData.LockReadOnly();
+	Hash = FCrc::MemCrc32(HeightBuffer, CollisionHeightData.GetBulkDataSize(), Hash);
+	CollisionHeightData.Unlock();
+
+	const void* DominantBuffer = DominantLayerData.LockReadOnly();
+	Hash = FCrc::MemCrc32(DominantBuffer, DominantLayerData.GetBulkDataSize(), Hash);
+	DominantLayerData.Unlock();
+
+	return Hash;
+}
+
 void ULandscapeHeightfieldCollisionComponent::UpdateHeightfieldRegion(int32 ComponentX1, int32 ComponentY1, int32 ComponentX2, int32 ComponentY2)
 {
 #if WITH_PHYSX
@@ -1213,15 +1235,23 @@ void ULandscapeMeshCollisionComponent::BeginDestroy()
 	Super::BeginDestroy();
 }
 
-void ULandscapeHeightfieldCollisionComponent::RecreateCollision()
+bool ULandscapeHeightfieldCollisionComponent::RecreateCollision()
 {
 	if (!HasAnyFlags(RF_ClassDefaultObject))
 	{
+		uint32 NewHash = ComputeCollisionHash();
+		if (bPhysicsStateCreated && NewHash == CollisionHash && CollisionHash != 0 && bEnableCollisionHashOptim)
+		{
+			return false;
+		}
+		CollisionHash = NewHash;
+
 		HeightfieldRef = NULL;
 		HeightfieldGuid = FGuid();
 
 		RecreatePhysicsState();
 	}
+	return true;
 }
 
 #if WITH_EDITORONLY_DATA
@@ -1350,7 +1380,7 @@ void ULandscapeHeightfieldCollisionComponent::SnapFoliageInstances(const FBox& I
 }
 #endif // WITH_EDITORONLY_DATA
 
-void ULandscapeMeshCollisionComponent::RecreateCollision()
+bool ULandscapeMeshCollisionComponent::RecreateCollision()
 {
 	if (!HasAnyFlags(RF_ClassDefaultObject))
 	{
@@ -1358,7 +1388,7 @@ void ULandscapeMeshCollisionComponent::RecreateCollision()
 		MeshGuid = FGuid();
 	}
 
-	Super::RecreateCollision();
+	return Super::RecreateCollision();
 }
 
 void ULandscapeHeightfieldCollisionComponent::Serialize(FArchive& Ar)
diff --git a/Engine/Source/Runtime/Landscape/Private/LandscapeEditLayers.cpp b/Engine/Source/Runtime/Landscape/Private/LandscapeEditLayers.cpp
--- a/Engine/Source/Runtime/Landscape/Private/LandscapeEditLayers.cpp
+++ b/Engine/Source/Runtime/Landscape/Private/LandscapeEditLayers.cpp
@@ -4316,6 +4316,32 @@ void ALandscape::UpdateLayersContent(bool bInWaitForStreaming, bool bInSkipMonit
 		}
 	}
 
+// not thread safe
+struct FEnableCollisionHashOptimScope
+{
+	FEnableCollisionHashOptimScope(ULandscapeHeightfieldCollisionComponent* InCollisionComponent)
+	{
+		CollisionComponent = InCollisionComponent;
+		if (CollisionComponent)
+		{
+			// not reentrant
+			check(!CollisionComponent->bEnableCollisionHashOptim);
+			CollisionComponent->bEnableCollisionHashOptim = true;
+		}
+	}
+
+	~FEnableCollisionHashOptimScope()
+	{
+		if (CollisionComponent)
+		{
+			CollisionComponent->bEnableCollisionHashOptim = false;
+		}
+	}
+
+private:
+	ULandscapeHeightfieldCollisionComponent* CollisionComponent;
+};
+
 bool ALandscape::UpdateCollisionAndClients(const TArray<ULandscapeComponent*>& InLandscapeComponents, const int32 InContentUpdateModes)
 {
 	bool bAllClientsUpdated = true;
@@ -4330,15 +4356,17 @@ bool ALandscape::UpdateCollisionAndClients(const TArray<ULandscapeComponent*>& I
 	for (ULandscapeComponent* LandscapeComponent : InLandscapeComponents)
 	{
 		bool bDeferClientUpdateForComponent = false;
+		bool bDoUpdateClient = true;
 		if (LandscapeComponent->IsUpdateFlagEnabledForModes(ELandscapeComponentUpdateFlag::Component_Update_Recreate_Collision, InContentUpdateModes))
 		{
 			if (ULandscapeHeightfieldCollisionComponent* CollisionComp = LandscapeComponent->CollisionComponent.Get())
 			{
-				CollisionComp->RecreateCollision();
+				FEnableCollisionHashOptimScope Scope(CollisionComp);
+				bDoUpdateClient = CollisionComp->RecreateCollision();
 			}
 		}
 
-		if (LandscapeComponent->IsUpdateFlagEnabledForModes(ELandscapeComponentUpdateFlag::Component_Update_Client, InContentUpdateModes))
+		if (bDoUpdateClient && LandscapeComponent->IsUpdateFlagEnabledForModes(ELandscapeComponentUpdateFlag::Component_Update_Client, InContentUpdateModes))
 		{
 			if (!GUndo)
 			{
