Fix Build Non-editor build for real

#rnx

diff --git a/Engine/Source/Runtime/Landscape/Private/LandscapeCollision.cpp b/Engine/Source/Runtime/Landscape/Private/LandscapeCollision.cpp
--- a/Engine/Source/Runtime/Landscape/Private/LandscapeCollision.cpp
+++ b/Engine/Source/Runtime/Landscape/Private/LandscapeCollision.cpp
@@ -1239,13 +1239,14 @@ bool ULandscapeHeightfieldCollisionComponent::RecreateCollision()
 {
 	if (!HasAnyFlags(RF_ClassDefaultObject))
 	{
+#if WITH_EDITOR
 		uint32 NewHash = ComputeCollisionHash();
 		if (bPhysicsStateCreated && NewHash == CollisionHash && CollisionHash != 0 && bEnableCollisionHashOptim)
 		{
 			return false;
 		}
 		CollisionHash = NewHash;
-
+#endif
 		HeightfieldRef = NULL;
 		HeightfieldGuid = FGuid();
 
