Fixed Landscape BP Brushes not updating collision/foliage/nav in PostEditMove
#rnx

diff --git a/Engine/Source/Runtime/Landscape/Private/LandscapeBPCustomBrush.cpp b/Engine/Source/Runtime/Landscape/Private/LandscapeBPCustomBrush.cpp
--- a/Engine/Source/Runtime/Landscape/Private/LandscapeBPCustomBrush.cpp
+++ b/Engine/Source/Runtime/Landscape/Private/LandscapeBPCustomBrush.cpp
@@ -78,7 +78,14 @@ void ALandscapeBlueprintCustomBrush::PostEditMove(bool bFinished)
 
 	if (OwningLandscape != nullptr)
 	{
-		OwningLandscape->RequestLayersContentUpdate(bFinished ? ELandscapeLayerUpdateMode::Update_All : ELandscapeLayerUpdateMode::Update_All_Editing);
+		if (bFinished)
+		{
+			OwningLandscape->RequestLayersContentUpdateForceAll();
+		}
+		else
+		{
+			OwningLandscape->RequestLayersContentUpdate(ELandscapeLayerUpdateMode::Update_All_Editing);
+		}
 	}
 }
 
