Fixed importing a heightmap in non layer mode, as we would lock the heightmaps, but due to order of execution we would never unlock before trying to lock again in another mode.

diff --git a/Engine/Source/Runtime/Landscape/Public/LandscapeEdit.h b/Engine/Source/Runtime/Landscape/Public/LandscapeEdit.h
--- a/Engine/Source/Runtime/Landscape/Public/LandscapeEdit.h
+++ b/Engine/Source/Runtime/Landscape/Public/LandscapeEdit.h
@@ -424,6 +424,9 @@ struct FHeightmapAccessor
 
 	virtual ~FHeightmapAccessor()
 	{
+		// Flush here manually so it will release the lock of the textures, as we will re lock for other things afterward
+		Flush();
+
 		// Landscape Layers are updates are delayed and done in  ALandscape::TickLayers
 		if (!LandscapeEdit->HasLandscapeLayersContent())
 		{
