Fixed crash using undo + BP custom brushes
#rnx

diff --git a/Engine/Source/Runtime/Landscape/Private/LandscapeBPCustomBrush.cpp b/Engine/Source/Runtime/Landscape/Private/LandscapeBPCustomBrush.cpp
--- a/Engine/Source/Runtime/Landscape/Private/LandscapeBPCustomBrush.cpp
+++ b/Engine/Source/Runtime/Landscape/Private/LandscapeBPCustomBrush.cpp
@@ -104,66 +104,66 @@ void ALandscapeBlueprintCustomBrush::PreEditChange(UProperty* PropertyThatWillCh
 void ALandscapeBlueprintCustomBrush::PostEditChangeProperty(FPropertyChangedEvent& PropertyChangedEvent)
 {
 	Super::PostEditChangeProperty(PropertyChangedEvent);
-
-	const FName PropertyName = PropertyChangedEvent.Property ? PropertyChangedEvent.Property->GetFName() : NAME_None;
-
-	if (PropertyName == GET_MEMBER_NAME_CHECKED(ALandscapeBlueprintCustomBrush, AffectHeightmap)
-		|| PropertyName == GET_MEMBER_NAME_CHECKED(ALandscapeBlueprintCustomBrush, AffectWeightmap))
+	
+	if (OwningLandscape != nullptr)
 	{
-		for (FLandscapeLayer& Layer : OwningLandscape->LandscapeLayers)
+		const FName PropertyName = PropertyChangedEvent.Property ? PropertyChangedEvent.Property->GetFName() : NAME_None;
+
+		if (PropertyName == GET_MEMBER_NAME_CHECKED(ALandscapeBlueprintCustomBrush, AffectHeightmap)
+			|| PropertyName == GET_MEMBER_NAME_CHECKED(ALandscapeBlueprintCustomBrush, AffectWeightmap))
 		{
-			for (int32 i = 0; i < Layer.Brushes.Num(); ++i)
+			for (FLandscapeLayer& Layer : OwningLandscape->LandscapeLayers)
 			{
-				if (Layer.Brushes[i].BPCustomBrush == this)
+				for (int32 i = 0; i < Layer.Brushes.Num(); ++i)
 				{
-					if (AffectHeightmap && !PreviousAffectHeightmap) // changed to affect
+					if (Layer.Brushes[i].BPCustomBrush == this)
 					{
-						Layer.HeightmapBrushOrderIndices.Add(i); // simply add the brush as the last one
-					}
-					else if (!AffectHeightmap && PreviousAffectHeightmap) // changed to no longer affect
-					{
-						for (int32 j = 0; j < Layer.HeightmapBrushOrderIndices.Num(); ++j)
+						if (AffectHeightmap && !PreviousAffectHeightmap) // changed to affect
 						{
-							if (Layer.HeightmapBrushOrderIndices[j] == i)
+							Layer.HeightmapBrushOrderIndices.Add(i); // simply add the brush as the last one
+						}
+						else if (!AffectHeightmap && PreviousAffectHeightmap) // changed to no longer affect
+						{
+							for (int32 j = 0; j < Layer.HeightmapBrushOrderIndices.Num(); ++j)
 							{
-								Layer.HeightmapBrushOrderIndices.RemoveAt(j);
-								break;
+								if (Layer.HeightmapBrushOrderIndices[j] == i)
+								{
+									Layer.HeightmapBrushOrderIndices.RemoveAt(j);
+									break;
+								}
 							}
 						}
-					}
 
-					if (AffectWeightmap && !PreviousAffectWeightmap) // changed to affect
-					{
-						Layer.WeightmapBrushOrderIndices.Add(i); // simply add the brush as the last one
-					}
-					else if (!AffectWeightmap && PreviousAffectWeightmap) // changed to no longer affect
-					{
-						for (int32 j = 0; j < Layer.WeightmapBrushOrderIndices.Num(); ++j)
+						if (AffectWeightmap && !PreviousAffectWeightmap) // changed to affect
 						{
-							if (Layer.WeightmapBrushOrderIndices[j] == i)
+							Layer.WeightmapBrushOrderIndices.Add(i); // simply add the brush as the last one
+						}
+						else if (!AffectWeightmap && PreviousAffectWeightmap) // changed to no longer affect
+						{
+							for (int32 j = 0; j < Layer.WeightmapBrushOrderIndices.Num(); ++j)
 							{
-								Layer.WeightmapBrushOrderIndices.RemoveAt(j);
-								break;
+								if (Layer.WeightmapBrushOrderIndices[j] == i)
+								{
+									Layer.WeightmapBrushOrderIndices.RemoveAt(j);
+									break;
+								}
 							}
 						}
-					}
 
-					PreviousAffectHeightmap = AffectHeightmap;
-					PreviousAffectWeightmap = AffectWeightmap;
+						PreviousAffectHeightmap = AffectHeightmap;
+						PreviousAffectWeightmap = AffectWeightmap;
 
-					break;
-				}
-			}			
-		}		
+						break;
+					}
+				}			
+			}		
 
-		// Should trigger a rebuild of the UI so the visual is updated with changes made to actor
-		//TODO: find a way to trigger the update of the UI
-		//FEdModeLandscape* EdMode = (FEdModeLandscape*)GLevelEditorModeTools().GetActiveMode(FBuiltinEditorModes::EM_Landscape);
-		//EdMode->RefreshDetailPanel();
-	}
+			// Should trigger a rebuild of the UI so the visual is updated with changes made to actor
+			//TODO: find a way to trigger the update of the UI
+			//FEdModeLandscape* EdMode = (FEdModeLandscape*)GLevelEditorModeTools().GetActiveMode(FBuiltinEditorModes::EM_Landscape);
+			//EdMode->RefreshDetailPanel();
+		}
 
-	if (OwningLandscape != nullptr)
-	{
 		OwningLandscape->RequestLayersContentUpdateForceAll(); // For now since we dont have brush bounds, we have to force an update of all components, otherwise we will have not right layer being rendered as materials are not updated
 	}
 }
