Landscape BP Brushes :
- Unique ordered list height & weight
- Blocked reordering brushes when layer is locked
- Disabled commit/uncommit options on brushes
#rnx

diff --git a/Engine/Source/Editor/LandscapeEditor/Private/LandscapeEdMode.cpp b/Engine/Source/Editor/LandscapeEditor/Private/LandscapeEdMode.cpp
--- a/Engine/Source/Editor/LandscapeEditor/Private/LandscapeEdMode.cpp
+++ b/Engine/Source/Editor/LandscapeEditor/Private/LandscapeEdMode.cpp
@@ -579,15 +579,17 @@ void FEdModeLandscape::Enter()
 		ALandscapeProxy* LandscapeProxy = CurrentToolTarget.LandscapeInfo->GetLandscapeProxy();
 		LandscapeProxy->OnMaterialChangedDelegate().AddRaw(this, &FEdModeLandscape::OnLandscapeMaterialChangedDelegate);
 
-		ALandscape* Landscape = GetLandscape();
-
-		if(Landscape && Landscape->HasLayersContent())
+		if (ALandscape* Landscape = GetLandscape())
 		{
-			if (Landscape->GetLandscapeSplinesReservedLayer())
+			Landscape->OnBPCustomBrushChangedDelegate().AddLambda([this]() { this->RefreshDetailPanel(); });
+			if (Landscape->HasLayersContent())
 			{
-				Landscape->UpdateLandscapeSplines();
+				if (Landscape->GetLandscapeSplinesReservedLayer())
+				{
+					Landscape->UpdateLandscapeSplines();
+				}
+				Landscape->RequestLayersContentUpdateForceAll();
 			}
-			Landscape->RequestLayersContentUpdateForceAll();
 		}
 	}
 
@@ -765,6 +767,10 @@ void FEdModeLandscape::Exit()
 	{
 		ALandscapeProxy* LandscapeProxy = CurrentToolTarget.LandscapeInfo->GetLandscapeProxy();
 		LandscapeProxy->OnMaterialChangedDelegate().RemoveAll(this);
+		if (ALandscape* Landscape = GetLandscape())
+		{
+			Landscape->OnBPCustomBrushChangedDelegate().RemoveAll(this);
+		}
 	}
 
 	// Restore real-time viewport state if we changed it
@@ -2516,6 +2522,10 @@ void FEdModeLandscape::SetTargetLandscape(const TWeakObjectPtr<ULandscapeInfo>&
 	{
 		ALandscapeProxy* LandscapeProxy = CurrentToolTarget.LandscapeInfo->GetLandscapeProxy();
 		LandscapeProxy->OnMaterialChangedDelegate().RemoveAll(this);
+		if (ALandscape* Landscape = GetLandscape())
+		{
+			Landscape->OnBPCustomBrushChangedDelegate().RemoveAll(this);
+		}
 	}
 
 	SetLandscapeInfo(InLandscapeInfo.Get());
@@ -2532,6 +2542,10 @@ void FEdModeLandscape::SetTargetLandscape(const TWeakObjectPtr<ULandscapeInfo>&
 	{
 		ALandscapeProxy* LandscapeProxy = CurrentToolTarget.LandscapeInfo->GetLandscapeProxy();
 		LandscapeProxy->OnMaterialChangedDelegate().AddRaw(this, &FEdModeLandscape::OnLandscapeMaterialChangedDelegate);
+		if (ALandscape* Landscape = GetLandscape())
+		{
+			Landscape->OnBPCustomBrushChangedDelegate().AddLambda([this]() { this->RefreshDetailPanel(); });
+		}
 	}
 
 	UpdateTargetList();
@@ -4452,7 +4466,7 @@ void FEdModeLandscape::RequestLayersContentUpdateForceAll(ELandscapeLayerUpdateM
 	}
 }
 
-void FEdModeLandscape::AddBrushToCurrentLayer(int32 InTargetType, ALandscapeBlueprintCustomBrush* InBrush)
+void FEdModeLandscape::AddBrushToCurrentLayer(ALandscapeBlueprintCustomBrush* InBrush)
 {
 	ALandscape* Landscape = GetLandscape();
 	if (Landscape == nullptr)
@@ -4460,10 +4474,10 @@ void FEdModeLandscape::AddBrushToCurrentLayer(int32 InTargetType, ALandscapeBlue
 		return;
 	}
 
-	Landscape->AddBrushToLayer(GetCurrentLayerIndex(), InTargetType, InBrush);
+	Landscape->AddBrushToLayer(GetCurrentLayerIndex(), InBrush);
 }
 
-void FEdModeLandscape::RemoveBrushFromCurrentLayer(int32 InTargetType, ALandscapeBlueprintCustomBrush* InBrush)
+void FEdModeLandscape::RemoveBrushFromCurrentLayer(ALandscapeBlueprintCustomBrush* InBrush)
 {
 	ALandscape* Landscape = GetLandscape();
 
@@ -4472,65 +4486,25 @@ void FEdModeLandscape::RemoveBrushFromCurrentLayer(int32 InTargetType, ALandscap
 		return;
 	}
 
-	Landscape->RemoveBrushFromLayer(GetCurrentLayerIndex(), InTargetType, InBrush);
-}
-
-bool FEdModeLandscape::AreAllBrushesCommitedToCurrentLayer(int32 InTargetType)
-{
-	ALandscape* Landscape = GetLandscape();
-
-	if (Landscape == nullptr)
-	{
-		return false;
-	}
-
-	return Landscape->AreAllBrushesCommitedToLayer(GetCurrentLayerIndex(), InTargetType);
+	Landscape->RemoveBrushFromLayer(GetCurrentLayerIndex(), InBrush);
 }
 
-void FEdModeLandscape::SetBrushesCommitStateForCurrentLayer(int32 InTargetType, bool InCommited)
+ALandscapeBlueprintCustomBrush* FEdModeLandscape::GetBrushForCurrentLayer(int8 InBrushIndex) const
 {
-	ALandscape* Landscape = GetLandscape();
-
-	if (Landscape == nullptr)
-	{
-		return;
-	}
-
-	Landscape->SetBrushesCommitStateForLayer(GetCurrentLayerIndex(), InTargetType, InCommited);
-	
-	GEngine->BroadcastLevelActorListChanged();
-}
-
-TArray<int8>& FEdModeLandscape::GetBrushesOrderForCurrentLayer(int32 InTargetType) const
-{
-	ALandscape* Landscape = GetLandscape();
-	check(Landscape);
-	return Landscape->GetBrushesOrderForLayer(GetCurrentLayerIndex(), InTargetType);
-}
-
-ALandscapeBlueprintCustomBrush* FEdModeLandscape::GetBrushForCurrentLayer(int32 InTargetType, int8 InBrushIndex) const
-{
-	ALandscape* Landscape = GetLandscape();
-
-	if (Landscape == nullptr)
+	if (ALandscape* Landscape = GetLandscape())
 	{
-		return nullptr;
+		return Landscape->GetBrushForLayer(GetCurrentLayerIndex(), InBrushIndex);
 	}
-
-	return Landscape->GetBrushForLayer(GetCurrentLayerIndex(), InTargetType, InBrushIndex);
+	return nullptr;
 }
 
-TArray<ALandscapeBlueprintCustomBrush*> FEdModeLandscape::GetBrushesForCurrentLayer(int32 InTargetType)
+TArray<ALandscapeBlueprintCustomBrush*> FEdModeLandscape::GetBrushesForCurrentLayer()
 {
-	ALandscape* Landscape = GetLandscape();
 	TArray<ALandscapeBlueprintCustomBrush*> Brushes;
-
-	if (Landscape == nullptr)
+	if (ALandscape* Landscape = GetLandscape())
 	{
-		return Brushes;
+		Brushes = Landscape->GetBrushesForLayer(GetCurrentLayerIndex());
 	}
-
-	Brushes = Landscape->GetBrushesForLayer(GetCurrentLayerIndex(), InTargetType);
 	return Brushes;
 }
 
@@ -4734,7 +4708,7 @@ void FEdModeLandscape::OnLevelActorAdded(AActor* InActor)
 
 	if (Brush != nullptr && Brush->GetTypedOuter<UPackage>() != GetTransientPackage())
 	{
-		AddBrushToCurrentLayer(CurrentToolTarget.TargetType, Brush);
+		AddBrushToCurrentLayer(Brush);
 		RefreshDetailPanel();
 	}
 }
@@ -4750,7 +4724,7 @@ void FEdModeLandscape::OnLevelActorRemoved(AActor* InActor)
 
 	if (Brush != nullptr && Brush->GetTypedOuter<UPackage>() != GetTransientPackage())
 	{
-		RemoveBrushFromCurrentLayer(CurrentToolTarget.TargetType, Brush);
+		RemoveBrushFromCurrentLayer(Brush);
 		RefreshDetailPanel();
 	}
 }
diff --git a/Engine/Source/Editor/LandscapeEditor/Private/LandscapeEdMode.h b/Engine/Source/Editor/LandscapeEditor/Private/LandscapeEdMode.h
--- a/Engine/Source/Editor/LandscapeEditor/Private/LandscapeEdMode.h
+++ b/Engine/Source/Editor/LandscapeEditor/Private/LandscapeEdMode.h
@@ -527,13 +527,10 @@ public:
 	void AutoUpdateDirtyLandscapeSplines();
 	bool CanEditLayer(FText* Reason = nullptr, FLandscapeLayer* InLayer = nullptr);
 
-	void AddBrushToCurrentLayer(int32 InTargetType, class ALandscapeBlueprintCustomBrush* InBrush);
-	void RemoveBrushFromCurrentLayer(int32 InTargetType, class ALandscapeBlueprintCustomBrush* InBrush);
-	bool AreAllBrushesCommitedToCurrentLayer(int32 InTargetType);
-	void SetBrushesCommitStateForCurrentLayer(int32 InTargetType, bool InCommited);
-	TArray<int8>& GetBrushesOrderForCurrentLayer(int32 InTargetType) const;
-	class ALandscapeBlueprintCustomBrush* GetBrushForCurrentLayer(int32 InTargetType, int8 BrushIndex) const;
-	TArray<class ALandscapeBlueprintCustomBrush*> GetBrushesForCurrentLayer(int32 InTargetType);
+	void AddBrushToCurrentLayer(class ALandscapeBlueprintCustomBrush* InBrush);
+	void RemoveBrushFromCurrentLayer(class ALandscapeBlueprintCustomBrush* InBrush);
+	class ALandscapeBlueprintCustomBrush* GetBrushForCurrentLayer(int8 BrushIndex) const;
+	TArray<class ALandscapeBlueprintCustomBrush*> GetBrushesForCurrentLayer();
 	
 	bool NeedToFillEmptyMaterialLayers() const;
 	void RequestLayersContentUpdate(ELandscapeLayerUpdateMode InUpdateMode);
diff --git a/Engine/Source/Editor/LandscapeEditor/Private/LandscapeEditorDetailCustomization_LayersBrushStack.cpp b/Engine/Source/Editor/LandscapeEditor/Private/LandscapeEditorDetailCustomization_LayersBrushStack.cpp
--- a/Engine/Source/Editor/LandscapeEditor/Private/LandscapeEditorDetailCustomization_LayersBrushStack.cpp
+++ b/Engine/Source/Editor/LandscapeEditor/Private/LandscapeEditorDetailCustomization_LayersBrushStack.cpp
@@ -84,18 +84,17 @@ void FLandscapeEditorCustomNodeBuilder_LayersBrushStack::GenerateHeaderRowConten
 	}
 
 	NodeRow.NameWidget
-		[
-			SNew(STextBlock)
-			.Font(IDetailLayoutBuilder::GetDetailFont())
-			.Text(FText::FromString(TEXT("Stack")))
-		];
+	[
+		SNew(STextBlock)
+		.Font(IDetailLayoutBuilder::GetDetailFont())
+		.Text(FText::FromString(TEXT("Stack")))
+	];
 }
 
 BEGIN_SLATE_FUNCTION_BUILD_OPTIMIZATION
 void FLandscapeEditorCustomNodeBuilder_LayersBrushStack::GenerateChildContent(IDetailChildrenBuilder& ChildrenBuilder)
 {
-	FEdModeLandscape* LandscapeEdMode = GetEditorMode();
-	if (LandscapeEdMode != NULL)
+	if (FEdModeLandscape* LandscapeEdMode = GetEditorMode())
 	{
 		TSharedPtr<SDragAndDropVerticalBox> BrushesList = SNew(SDragAndDropVerticalBox)
 			.OnCanAcceptDrop(this, &FLandscapeEditorCustomNodeBuilder_LayersBrushStack::HandleCanAcceptDrop)
@@ -109,7 +108,6 @@ void FLandscapeEditorCustomNodeBuilder_LayersBrushStack::GenerateChildContent(ID
 			.Visibility(EVisibility::Visible)
 			[
 				SNew(SVerticalBox)
-
 				+ SVerticalBox::Slot()
 				.AutoHeight()
 				.VAlign(VAlign_Center)
@@ -117,34 +115,14 @@ void FLandscapeEditorCustomNodeBuilder_LayersBrushStack::GenerateChildContent(ID
 				[
 					BrushesList.ToSharedRef()
 				]
-
-				+ SVerticalBox::Slot()
-				.AutoHeight()
-				.VAlign(VAlign_Center)
-				.Padding(0, 2)
-				[
-					SNew(SHorizontalBox)
-
-					+SHorizontalBox::Slot()
-					.HAlign(HAlign_Right)
-					//.Padding(4, 0)
-					[
-						SNew(SButton)
-						.Text(this, &FLandscapeEditorCustomNodeBuilder_LayersBrushStack::GetCommitBrushesButtonText)
-						.OnClicked(this, &FLandscapeEditorCustomNodeBuilder_LayersBrushStack::ToggleCommitBrushes)
-						.IsEnabled(this, &FLandscapeEditorCustomNodeBuilder_LayersBrushStack::IsCommitBrushesButtonEnabled)
-					]
-				]
 			];
 
-		if (LandscapeEdMode->CurrentToolMode != nullptr)
+		if (LandscapeEdMode->CurrentToolMode)
 		{
-			const TArray<int8>& BrushOrderStack = LandscapeEdMode->GetBrushesOrderForCurrentLayer(LandscapeEdMode->CurrentToolTarget.TargetType);
-
-			for (int32 i = 0; i < BrushOrderStack.Num(); ++i)
+			int32 BrushCount = LandscapeEdMode->GetBrushesForCurrentLayer().Num();
+			for (int32 i = 0; i < BrushCount; ++i)
 			{
 				TSharedPtr<SWidget> GeneratedRowWidget = GenerateRow(i);
-
 				if (GeneratedRowWidget.IsValid())
 				{
 					BrushesList->AddSlot()
@@ -170,8 +148,7 @@ TSharedPtr<SWidget> FLandscapeEditorCustomNodeBuilder_LayersBrushStack::Generate
 		.IsSelected(TAttribute<bool>::Create(TAttribute<bool>::FGetter::CreateSP(this, &FLandscapeEditorCustomNodeBuilder_LayersBrushStack::IsBrushSelected, InBrushIndex)))
 		[	
 			SNew(SHorizontalBox)
-
-			+ SHorizontalBox::Slot()
+			+SHorizontalBox::Slot()
 			.VAlign(VAlign_Center)
 			.Padding(4, 0)
 			[
@@ -184,6 +161,7 @@ TSharedPtr<SWidget> FLandscapeEditorCustomNodeBuilder_LayersBrushStack::Generate
 					SNew(STextBlock)
 					.ColorAndOpacity(TAttribute<FSlateColor>::Create(TAttribute<FSlateColor>::FGetter::CreateSP(this, &FLandscapeEditorCustomNodeBuilder_LayersBrushStack::GetBrushTextColor, InBrushIndex)))
 					.Text(this, &FLandscapeEditorCustomNodeBuilder_LayersBrushStack::GetBrushText, InBrushIndex)
+					.IsEnabled(TAttribute<bool>::Create(TAttribute<bool>::FGetter::CreateSP(this, &FLandscapeEditorCustomNodeBuilder_LayersBrushStack::IsBrushEnabled, InBrushIndex)))
 				]
 			]
 		];
@@ -192,24 +170,22 @@ TSharedPtr<SWidget> FLandscapeEditorCustomNodeBuilder_LayersBrushStack::Generate
 }
 END_SLATE_FUNCTION_BUILD_OPTIMIZATION
 
-bool FLandscapeEditorCustomNodeBuilder_LayersBrushStack::IsBrushSelected(int32 InBrushIndex) const
+bool FLandscapeEditorCustomNodeBuilder_LayersBrushStack::IsBrushEnabled(int32 InBrushIndex) const
 {
 	ALandscapeBlueprintCustomBrush* Brush = GetBrush(InBrushIndex);
+	FEdModeLandscape* LandscapeEdMode = GetEditorMode();
+	return Brush && LandscapeEdMode && ((Brush->IsAffectingHeightmap() && LandscapeEdMode->CurrentToolTarget.TargetType == ELandscapeToolTargetType::Type::Heightmap) || (Brush->IsAffectingWeightmap() && LandscapeEdMode->CurrentToolTarget.TargetType == ELandscapeToolTargetType::Type::Weightmap));
+}
 
-	return Brush != nullptr ? Brush->IsSelected() : false;
+bool FLandscapeEditorCustomNodeBuilder_LayersBrushStack::IsBrushSelected(int32 InBrushIndex) const
+{
+	ALandscapeBlueprintCustomBrush* Brush = GetBrush(InBrushIndex);
+	return Brush && Brush->IsSelected();
 }
 
 void FLandscapeEditorCustomNodeBuilder_LayersBrushStack::OnBrushSelectionChanged(int32 InBrushIndex)
 {
-	FEdModeLandscape* LandscapeEdMode = GetEditorMode();
-
-	if (LandscapeEdMode != nullptr && LandscapeEdMode->AreAllBrushesCommitedToCurrentLayer(LandscapeEdMode->CurrentToolTarget.TargetType))
-	{
-		return;
-	}
-
 	ALandscapeBlueprintCustomBrush* Brush = GetBrush(InBrushIndex);
-
 	if (Brush != nullptr && !Brush->IsCommited())
 	{
 		GEditor->SelectNone(true, true);
@@ -219,9 +195,7 @@ void FLandscapeEditorCustomNodeBuilder_LayersBrushStack::OnBrushSelectionChanged
 
 FText FLandscapeEditorCustomNodeBuilder_LayersBrushStack::GetBrushText(int32 InBrushIndex) const
 {
-	ALandscapeBlueprintCustomBrush* Brush = GetBrush(InBrushIndex);
-
-	if (Brush != nullptr)
+	if (ALandscapeBlueprintCustomBrush* Brush = GetBrush(InBrushIndex))
 	{
 		return FText::FromString(Brush->GetActorLabel());
 	}
@@ -231,9 +205,7 @@ FText FLandscapeEditorCustomNodeBuilder_LayersBrushStack::GetBrushText(int32 InB
 
 FSlateColor FLandscapeEditorCustomNodeBuilder_LayersBrushStack::GetBrushTextColor(int32 InBrushIndex) const
 {
-	ALandscapeBlueprintCustomBrush* Brush = GetBrush(InBrushIndex);
-
-	if (Brush != nullptr)
+	if (ALandscapeBlueprintCustomBrush* Brush = GetBrush(InBrushIndex))
 	{
 		return Brush->IsCommited() ? FSlateColor::UseSubduedForeground() : FSlateColor::UseForeground();
 	}
@@ -243,85 +215,32 @@ FSlateColor FLandscapeEditorCustomNodeBuilder_LayersBrushStack::GetBrushTextColo
 
 ALandscapeBlueprintCustomBrush* FLandscapeEditorCustomNodeBuilder_LayersBrushStack::GetBrush(int32 InBrushIndex) const
 {
-	FEdModeLandscape* LandscapeEdMode = GetEditorMode();
-
-	if (LandscapeEdMode != nullptr)
+	if (FEdModeLandscape* LandscapeEdMode = GetEditorMode())
 	{
-		return LandscapeEdMode->GetBrushForCurrentLayer(LandscapeEdMode->CurrentToolTarget.TargetType, InBrushIndex);
+		return LandscapeEdMode->GetBrushForCurrentLayer(InBrushIndex);
 	}
 
 	return nullptr;
 }
 
-FReply FLandscapeEditorCustomNodeBuilder_LayersBrushStack::ToggleCommitBrushes()
-{
-	FEdModeLandscape* LandscapeEdMode = GetEditorMode();
-
-	if (LandscapeEdMode != nullptr)
-	{
-		bool CommitBrushes = !LandscapeEdMode->AreAllBrushesCommitedToCurrentLayer(LandscapeEdMode->CurrentToolTarget.TargetType);
-
-		if (CommitBrushes)
-		{
-			TArray<ALandscapeBlueprintCustomBrush*> BrushStack = LandscapeEdMode->GetBrushesForCurrentLayer(LandscapeEdMode->CurrentToolTarget.TargetType);
-
-			for (ALandscapeBlueprintCustomBrush* Brush : BrushStack)
-			{
-				GEditor->SelectActor(Brush, false, true);
-			}
-		}
-
-		LandscapeEdMode->SetBrushesCommitStateForCurrentLayer(LandscapeEdMode->CurrentToolTarget.TargetType, CommitBrushes);
-	}
-
-	return FReply::Handled();
-}
-
-bool FLandscapeEditorCustomNodeBuilder_LayersBrushStack::IsCommitBrushesButtonEnabled() const
-{
-	FEdModeLandscape* LandscapeEdMode = GetEditorMode();
-
-	if (LandscapeEdMode != nullptr)
-	{
-		TArray<ALandscapeBlueprintCustomBrush*> BrushStack = LandscapeEdMode->GetBrushesForCurrentLayer(LandscapeEdMode->CurrentToolTarget.TargetType);
-
-		return BrushStack.Num() > 0;
-	}
-
-	return false;
-}
-
-FText FLandscapeEditorCustomNodeBuilder_LayersBrushStack::GetCommitBrushesButtonText() const
-{
-	FEdModeLandscape* LandscapeEdMode = GetEditorMode();
-
-	if (LandscapeEdMode != nullptr)
-	{
-		return LandscapeEdMode->AreAllBrushesCommitedToCurrentLayer(LandscapeEdMode->CurrentToolTarget.TargetType) ? LOCTEXT("UnCommitBrushesText", "Uncommit") : LOCTEXT("CommitBrushesText", "Commit");
-	}
-
-	return FText::FromName(NAME_None);
-}
-
 FReply FLandscapeEditorCustomNodeBuilder_LayersBrushStack::HandleDragDetected(const FGeometry& MyGeometry, const FPointerEvent& MouseEvent, int32 SlotIndex, SVerticalBox::FSlot* Slot)
 {
-	FEdModeLandscape* LandscapeEdMode = GetEditorMode();
-
-	if (LandscapeEdMode != nullptr)
+	if (FEdModeLandscape* LandscapeEdMode = GetEditorMode())
 	{
-		const TArray<int8>& BrushOrderStack = LandscapeEdMode->GetBrushesOrderForCurrentLayer(LandscapeEdMode->CurrentToolTarget.TargetType);
-
-		if (BrushOrderStack.IsValidIndex(SlotIndex))
+		FLandscapeLayer* Layer = LandscapeEdMode->GetCurrentLayer();
+		if (Layer && !Layer->bLocked)
 		{
-			TSharedPtr<SWidget> Row = GenerateRow(SlotIndex);
-
-			if (Row.IsValid())
+			const TArray<ALandscapeBlueprintCustomBrush*>& BrushStack = LandscapeEdMode->GetBrushesForCurrentLayer();
+			if (BrushStack.IsValidIndex(SlotIndex))
 			{
-				return FReply::Handled().BeginDragDrop(FLandscapeListElementDragDropOp::New(SlotIndex, Slot, Row));
+				TSharedPtr<SWidget> Row = GenerateRow(SlotIndex);
+				if (Row.IsValid())
+				{
+					return FReply::Handled().BeginDragDrop(FLandscapeListElementDragDropOp::New(SlotIndex, Slot, Row));
+				}
 			}
 		}
 	}
-
 	return FReply::Unhandled();
 }
 
@@ -350,7 +269,7 @@ FReply FLandscapeEditorCustomNodeBuilder_LayersBrushStack::HandleAcceptDrop(FDra
 			int32 StartingLayerIndex = DragDropOperation->SlotIndexBeingDragged;
 			int32 DestinationLayerIndex = SlotIndex;
 			const FScopedTransaction Transaction(LOCTEXT("Landscape_LayerBrushes_Reorder", "Reorder Layer Brush"));
-			if (Landscape->ReorderLayerBrush(LandscapeEdMode->GetCurrentLayerIndex(), LandscapeEdMode->CurrentToolTarget.TargetType, StartingLayerIndex, DestinationLayerIndex))
+			if (Landscape->ReorderLayerBrush(LandscapeEdMode->GetCurrentLayerIndex(), StartingLayerIndex, DestinationLayerIndex))
 			{
 				LandscapeEdMode->RefreshDetailPanel();
 				return FReply::Handled();
diff --git a/Engine/Source/Editor/LandscapeEditor/Private/LandscapeEditorDetailCustomization_LayersBrushStack.h b/Engine/Source/Editor/LandscapeEditor/Private/LandscapeEditorDetailCustomization_LayersBrushStack.h
--- a/Engine/Source/Editor/LandscapeEditor/Private/LandscapeEditorDetailCustomization_LayersBrushStack.h
+++ b/Engine/Source/Editor/LandscapeEditor/Private/LandscapeEditorDetailCustomization_LayersBrushStack.h
@@ -59,15 +59,12 @@ protected:
 	TSharedPtr<SWidget> GenerateRow(int32 InLayerIndex);
 
 	bool IsBrushSelected(int32 InBrushIndex) const;
+	bool IsBrushEnabled(int32 InBrushIndex) const;
 	void OnBrushSelectionChanged(int32 InBrushIndex);
 	FText GetBrushText(int32 InBrushIndex) const;
 	FSlateColor GetBrushTextColor(int32 InBrushIndex) const;
 	ALandscapeBlueprintCustomBrush* GetBrush(int32 InBrushIndex) const;
 
-	FReply ToggleCommitBrushes();
-	bool IsCommitBrushesButtonEnabled() const;
-	FText GetCommitBrushesButtonText() const;
-
 	// Drag/Drop handling
 	FReply HandleDragDetected(const FGeometry& MyGeometry, const FPointerEvent& MouseEvent, int32 SlotIndex, SVerticalBox::FSlot* Slot);
 	TOptional<SDragAndDropVerticalBox::EItemDropZone> HandleCanAcceptDrop(const FDragDropEvent& DragDropEvent, SDragAndDropVerticalBox::EItemDropZone DropZone, SVerticalBox::FSlot* Slot);
diff --git a/Engine/Source/Runtime/Landscape/Classes/Landscape.h b/Engine/Source/Runtime/Landscape/Classes/Landscape.h
--- a/Engine/Source/Runtime/Landscape/Classes/Landscape.h
+++ b/Engine/Source/Runtime/Landscape/Classes/Landscape.h
@@ -178,12 +178,6 @@ struct FLandscapeLayer
 	UPROPERTY()
 	TArray<FLandscapeLayerBrush> Brushes;
 
-	UPROPERTY()
-	TArray<int8> HeightmapBrushOrderIndices;
-
-	UPROPERTY()
-	TArray<int8> WeightmapBrushOrderIndices;
-
 	UPROPERTY()
 	TMap<ULandscapeLayerInfoObject*, bool> WeightmapLayerAllocationBlend; // True -> Substractive, False -> Additive
 };
@@ -282,14 +276,12 @@ public:
 	LANDSCAPE_API bool IsLayerBlendSubstractive(int32 InLayerIndex, const TWeakObjectPtr<ULandscapeLayerInfoObject>& InLayerInfoObj) const;
 	LANDSCAPE_API void SetLayerSubstractiveBlendStatus(int32 InLayerIndex, bool InStatus, const TWeakObjectPtr<ULandscapeLayerInfoObject>& InLayerInfoObj);
 
-	LANDSCAPE_API void AddBrushToLayer(int32 InLayerIndex, int32 InTargetType, class ALandscapeBlueprintCustomBrush* InBrush);
-	LANDSCAPE_API void RemoveBrushFromLayer(int32 InLayerIndex, int32 InTargetType, class ALandscapeBlueprintCustomBrush* InBrush);
-	LANDSCAPE_API bool ReorderLayerBrush(int32 InLayerIndex, int32 InTargetType, int32 InStartingLayerBrushIndex, int32 InDestinationLayerBrushIndex);
-	LANDSCAPE_API bool AreAllBrushesCommitedToLayer(int32 InLayerIndex, int32 InTargetType);
-	LANDSCAPE_API void SetBrushesCommitStateForLayer(int32 InLayerIndex, int32 InTargetType, bool InCommited);
-	LANDSCAPE_API TArray<int8>& GetBrushesOrderForLayer(int32 InLayerIndex, int32 InTargetType);
-	LANDSCAPE_API class ALandscapeBlueprintCustomBrush* GetBrushForLayer(int32 InLayerIndex, int32 InTargetType, int8 BrushIndex) const;
-	LANDSCAPE_API TArray<class ALandscapeBlueprintCustomBrush*> GetBrushesForLayer(int32 InLayerIndex, int32 InTargetType) const;
+	LANDSCAPE_API void AddBrushToLayer(int32 InLayerIndex, class ALandscapeBlueprintCustomBrush* InBrush);
+	LANDSCAPE_API void RemoveBrushFromLayer(int32 InLayerIndex, class ALandscapeBlueprintCustomBrush* InBrush);
+	LANDSCAPE_API bool ReorderLayerBrush(int32 InLayerIndex, int32 InStartingLayerBrushIndex, int32 InDestinationLayerBrushIndex);
+	LANDSCAPE_API class ALandscapeBlueprintCustomBrush* GetBrushForLayer(int32 InLayerIndex, int8 BrushIndex) const;
+	LANDSCAPE_API TArray<class ALandscapeBlueprintCustomBrush*> GetBrushesForLayer(int32 InLayerIndex) const;
+	LANDSCAPE_API void OnBPCustomBrushChanged();
 	
 	LANDSCAPE_API void OnPreSave();
 
@@ -364,6 +356,9 @@ public:
 	UPROPERTY(EditAnywhere, Category=Experimental, AdvancedDisplay)
 	bool bCanHaveLayersContent = false;
 
+	DECLARE_EVENT(ALandscape, FLandscapeBPCustomBrushChangedDelegate);
+	FLandscapeBPCustomBrushChangedDelegate& OnBPCustomBrushChangedDelegate() { return LandscapeBPCustomBrushChangedDelegate; }
+
 	/** Target Landscape Layer for Landscape Splines */
 	UPROPERTY()
 	FGuid LandscapeSplinesTargetLayerGuid;
@@ -384,6 +379,7 @@ public:
 	TArray<UTextureRenderTarget2D*> WeightmapRTList;
 
 private:
+	FLandscapeBPCustomBrushChangedDelegate LandscapeBPCustomBrushChangedDelegate;
 
 	/** Components affected by landscape splines (used to partially clear Layer Reserved for Splines) */
 	UPROPERTY(Transient)
diff --git a/Engine/Source/Runtime/Landscape/Classes/LandscapeBPCustomBrush.h b/Engine/Source/Runtime/Landscape/Classes/LandscapeBPCustomBrush.h
--- a/Engine/Source/Runtime/Landscape/Classes/LandscapeBPCustomBrush.h
+++ b/Engine/Source/Runtime/Landscape/Classes/LandscapeBPCustomBrush.h
@@ -32,12 +32,6 @@ private:
 
 	UPROPERTY(Transient)
 	bool bIsInitialized;
-
-	UPROPERTY(Transient)
-	bool PreviousAffectHeightmap;
-
-	UPROPERTY(Transient)
-	bool PreviousAffectWeightmap;
 #endif
 public:
 
@@ -65,7 +59,6 @@ public:
 	class ALandscape* GetOwningLandscape() const;
 
 	virtual void PostEditMove(bool bFinished) override;
-	virtual void PreEditChange(UProperty* PropertyThatWillChange) override;
 	virtual void PostEditChangeProperty(FPropertyChangedEvent& PropertyChangedEvent) override;
 #endif
 };
diff --git a/Engine/Source/Runtime/Landscape/Private/LandscapeBPCustomBrush.cpp b/Engine/Source/Runtime/Landscape/Private/LandscapeBPCustomBrush.cpp
--- a/Engine/Source/Runtime/Landscape/Private/LandscapeBPCustomBrush.cpp
+++ b/Engine/Source/Runtime/Landscape/Private/LandscapeBPCustomBrush.cpp
@@ -89,82 +89,22 @@ void ALandscapeBlueprintCustomBrush::PostEditMove(bool bFinished)
 	}
 }
 
-void ALandscapeBlueprintCustomBrush::PreEditChange(UProperty* PropertyThatWillChange)
-{
-	const FName PropertyName = PropertyThatWillChange != nullptr ? PropertyThatWillChange->GetFName() : NAME_None;
-
-	if (PropertyName == GET_MEMBER_NAME_CHECKED(ALandscapeBlueprintCustomBrush, AffectHeightmap)
-		|| PropertyName == GET_MEMBER_NAME_CHECKED(ALandscapeBlueprintCustomBrush, AffectWeightmap))
-	{
-		PreviousAffectHeightmap = AffectHeightmap;
-		PreviousAffectWeightmap = AffectWeightmap;
-	}
-}
-
 void ALandscapeBlueprintCustomBrush::PostEditChangeProperty(FPropertyChangedEvent& PropertyChangedEvent)
 {
 	Super::PostEditChangeProperty(PropertyChangedEvent);
 	
-	if (OwningLandscape != nullptr)
+	if (OwningLandscape)
 	{
 		const FName PropertyName = PropertyChangedEvent.Property ? PropertyChangedEvent.Property->GetFName() : NAME_None;
-
-		if (PropertyName == GET_MEMBER_NAME_CHECKED(ALandscapeBlueprintCustomBrush, AffectHeightmap)
-			|| PropertyName == GET_MEMBER_NAME_CHECKED(ALandscapeBlueprintCustomBrush, AffectWeightmap))
+		if (PropertyName == GET_MEMBER_NAME_CHECKED(ALandscapeBlueprintCustomBrush, AffectHeightmap) || 
+			PropertyName == GET_MEMBER_NAME_CHECKED(ALandscapeBlueprintCustomBrush, AffectWeightmap))
 		{
-			for (FLandscapeLayer& Layer : OwningLandscape->LandscapeLayers)
-			{
-				for (int32 i = 0; i < Layer.Brushes.Num(); ++i)
-				{
-					if (Layer.Brushes[i].BPCustomBrush == this)
-					{
-						if (AffectHeightmap && !PreviousAffectHeightmap) // changed to affect
-						{
-							Layer.HeightmapBrushOrderIndices.Add(i); // simply add the brush as the last one
-						}
-						else if (!AffectHeightmap && PreviousAffectHeightmap) // changed to no longer affect
-						{
-							for (int32 j = 0; j < Layer.HeightmapBrushOrderIndices.Num(); ++j)
-							{
-								if (Layer.HeightmapBrushOrderIndices[j] == i)
-								{
-									Layer.HeightmapBrushOrderIndices.RemoveAt(j);
-									break;
-								}
-							}
-						}
-
-						if (AffectWeightmap && !PreviousAffectWeightmap) // changed to affect
-						{
-							Layer.WeightmapBrushOrderIndices.Add(i); // simply add the brush as the last one
-						}
-						else if (!AffectWeightmap && PreviousAffectWeightmap) // changed to no longer affect
-						{
-							for (int32 j = 0; j < Layer.WeightmapBrushOrderIndices.Num(); ++j)
-							{
-								if (Layer.WeightmapBrushOrderIndices[j] == i)
-								{
-									Layer.WeightmapBrushOrderIndices.RemoveAt(j);
-									break;
-								}
-							}
-						}
-
-						PreviousAffectHeightmap = AffectHeightmap;
-						PreviousAffectWeightmap = AffectWeightmap;
-
-						break;
-					}
-				}			
-			}		
-
 			// Should trigger a rebuild of the UI so the visual is updated with changes made to actor
 			//TODO: find a way to trigger the update of the UI
 			//FEdModeLandscape* EdMode = (FEdModeLandscape*)GLevelEditorModeTools().GetActiveMode(FBuiltinEditorModes::EM_Landscape);
 			//EdMode->RefreshDetailPanel();
+			OwningLandscape->OnBPCustomBrushChanged(); 
 		}
-
-		OwningLandscape->RequestLayersContentUpdateForceAll(); // For now since we dont have brush bounds, we have to force an update of all components, otherwise we will have not right layer being rendered as materials are not updated
 	}
 }
 #endif
diff --git a/Engine/Source/Runtime/Landscape/Private/LandscapeEditLayers.cpp b/Engine/Source/Runtime/Landscape/Private/LandscapeEditLayers.cpp
--- a/Engine/Source/Runtime/Landscape/Private/LandscapeEditLayers.cpp
+++ b/Engine/Source/Runtime/Landscape/Private/LandscapeEditLayers.cpp
@@ -2731,12 +2731,12 @@ int32 ALandscape::RegenerateLayersHeightmaps(const TArray<ULandscapeComponent*>&
 			if (Layer.bVisible)
 			{
 				// Draw each brushes
-				for (int32 i = 0; i < Layer.HeightmapBrushOrderIndices.Num(); ++i)
+				for (int32 i = 0; i < Layer.Brushes.Num(); ++i)
 				{
 					// TODO: handle conversion from float to RG8 by using material params to write correct values
 					// TODO: handle conversion/handling of RT not same size as internal size
 
-					FLandscapeLayerBrush& Brush = Layer.Brushes[Layer.HeightmapBrushOrderIndices[i]];
+					FLandscapeLayerBrush& Brush = Layer.Brushes[i];
 
 					if (Brush.BPCustomBrush == nullptr || !Brush.BPCustomBrush->IsAffectingHeightmap())
 					{
@@ -3469,9 +3469,9 @@ int32 ALandscape::RegenerateLayersWeightmaps(const TArray<ULandscapeComponent*>&
 				{
 					const FLandscapeInfoLayerSettings& InfoLayerSettings = Info->Layers[LayerInfoSettingsIndex];
 
-					for (int32 i = 0; i < Layer.WeightmapBrushOrderIndices.Num(); ++i)
+					for (int32 i = 0; i < Layer.Brushes.Num(); ++i)
 					{
-						FLandscapeLayerBrush& Brush = Layer.Brushes[Layer.WeightmapBrushOrderIndices[i]];
+						FLandscapeLayerBrush& Brush = Layer.Brushes[i];
 
 						if (Brush.BPCustomBrush == nullptr)
 						{
@@ -3587,11 +3587,11 @@ int32 ALandscape::RegenerateLayersWeightmaps(const TArray<ULandscapeComponent*>&
 					if (Layer.bVisible)
 					{
 						// Draw each brushes				
-						for (int32 i = 0; i < Layer.WeightmapBrushOrderIndices.Num(); ++i)
+						for (int32 i = 0; i < Layer.Brushes.Num(); ++i)
 						{
 							// TODO: handle conversion/handling of RT not same size as internal size
 
-							FLandscapeLayerBrush& Brush = Layer.Brushes[Layer.WeightmapBrushOrderIndices[i]];
+							FLandscapeLayerBrush& Brush = Layer.Brushes[i];
 
 							if (Brush.BPCustomBrush == nullptr || !Brush.BPCustomBrush->IsAffectingWeightmap() || !Brush.BPCustomBrush->IsAffectingWeightmapLayer(LayerInfoObj->LayerName))
 							{
@@ -5416,19 +5416,18 @@ void ALandscape::SetLayerSubstractiveBlendStatus(int32 InLayerIndex, bool InStat
 	RequestLayersContentUpdateForceAll(ELandscapeLayerUpdateMode::Update_Weightmap_All);
 }
 
-bool ALandscape::ReorderLayerBrush(int32 InLayerIndex, int32 InTargetType, int32 InStartingLayerBrushIndex, int32 InDestinationLayerBrushIndex)
+bool ALandscape::ReorderLayerBrush(int32 InLayerIndex, int32 InStartingLayerBrushIndex, int32 InDestinationLayerBrushIndex)
 {
-	if (const FLandscapeLayer* Layer = GetLayer(InLayerIndex))
+	if (FLandscapeLayer* Layer = GetLayer(InLayerIndex))
 	{
-		TArray<int8>& BrushOrderStack = GetBrushesOrderForLayer(InLayerIndex, InTargetType);
 		if (InStartingLayerBrushIndex != InDestinationLayerBrushIndex &&
-			BrushOrderStack.IsValidIndex(InStartingLayerBrushIndex) &&
-			BrushOrderStack.IsValidIndex(InDestinationLayerBrushIndex))
+			Layer->Brushes.IsValidIndex(InStartingLayerBrushIndex) &&
+			Layer->Brushes.IsValidIndex(InDestinationLayerBrushIndex))
 		{
 			Modify();
-			int8 MovingBrushIndex = BrushOrderStack[InStartingLayerBrushIndex];
-			BrushOrderStack.RemoveAt(InStartingLayerBrushIndex);
-			BrushOrderStack.Insert(MovingBrushIndex, InDestinationLayerBrushIndex);
+			FLandscapeLayerBrush MovingBrush = Layer->Brushes[InStartingLayerBrushIndex];
+			Layer->Brushes.RemoveAt(InStartingLayerBrushIndex);
+			Layer->Brushes.Insert(MovingBrush, InDestinationLayerBrushIndex);
 			RequestLayersContentUpdateForceAll();
 			return true;
 		}
@@ -5436,209 +5435,66 @@ bool ALandscape::ReorderLayerBrush(int32 InLayerIndex, int32 InTargetType, int32
 	return false;
 }
 
-void ALandscape::AddBrushToLayer(int32 InLayerIndex, int32 InTargetType, ALandscapeBlueprintCustomBrush* InBrush)
+void ALandscape::AddBrushToLayer(int32 InLayerIndex, ALandscapeBlueprintCustomBrush* InBrush)
 {
-	FLandscapeLayer* Layer = GetLayer(InLayerIndex);
-
-	if (Layer == nullptr)
+	if (FLandscapeLayer* Layer = GetLayer(InLayerIndex))
 	{
-		return;
-	}
-
-	Modify();
-
-	int32 AddedIndex = Layer->Brushes.Add(FLandscapeLayerBrush(InBrush));
-
-	if (InTargetType == ELandscapeToolTargetType::Type::Heightmap)
-	{
-		Layer->HeightmapBrushOrderIndices.Add(AddedIndex);
-	}
-	else
-	{
-		Layer->WeightmapBrushOrderIndices.Add(AddedIndex);
+		Modify();
+		Layer->Brushes.Add(FLandscapeLayerBrush(InBrush));
+		InBrush->SetOwningLandscape(this);
+		RequestLayersContentUpdateForceAll();
 	}
-
-	InBrush->SetOwningLandscape(this);
-
-	RequestLayersContentUpdateForceAll();
 }
 
-void ALandscape::RemoveBrushFromLayer(int32 InLayerIndex, int32 InTargetType, ALandscapeBlueprintCustomBrush* InBrush)
+void ALandscape::RemoveBrushFromLayer(int32 InLayerIndex, ALandscapeBlueprintCustomBrush* InBrush)
 {
-	FLandscapeLayer* Layer = GetLayer(InLayerIndex);
-
-	if (Layer == nullptr)
-	{
-		return;
-	}
-
-	Modify();
-
-	int32 IndexToRemove = INDEX_NONE;
-	for (int32 i = 0; i < Layer->Brushes.Num(); ++i)
-	{
-		if (Layer->Brushes[i].BPCustomBrush == InBrush)
-		{
-			IndexToRemove = i;
-			break;
-		}
-	}
-
-	if (IndexToRemove != INDEX_NONE)
+	if (FLandscapeLayer* Layer = GetLayer(InLayerIndex))
 	{
-		Layer->Brushes.RemoveAt(IndexToRemove);
-
-		for (int32 i = 0; i < Layer->HeightmapBrushOrderIndices.Num(); ++i)
-		{
-			if (Layer->HeightmapBrushOrderIndices[i] == IndexToRemove)
-			{
-				// Update the value of the index of all the one after the one we removed, so index still correctly match actual brushes list
-				for (int32 j = 0; j < Layer->HeightmapBrushOrderIndices.Num(); ++j)
-				{
-					if (Layer->HeightmapBrushOrderIndices[j] > IndexToRemove)
-					{
-						--Layer->HeightmapBrushOrderIndices[j];
-					}
-				}
-
-				Layer->HeightmapBrushOrderIndices.RemoveAt(i);
-				break;
-			}
-		}
-
-		for (int32 i = 0; i < Layer->WeightmapBrushOrderIndices.Num(); ++i)
+		for (int32 i = 0; i < Layer->Brushes.Num(); ++i)
 		{
-			if (Layer->WeightmapBrushOrderIndices[i] == IndexToRemove)
+			if (Layer->Brushes[i].BPCustomBrush == InBrush)
 			{
-				// Update the value of the index of all the one after the one we removed, so index still correctly match actual brushes list
-				for (int32 j = 0; j < Layer->WeightmapBrushOrderIndices.Num(); ++j)
-				{
-					if (Layer->WeightmapBrushOrderIndices[j] > IndexToRemove)
-					{
-						--Layer->WeightmapBrushOrderIndices[j];
-					}
-				}
-
-				Layer->WeightmapBrushOrderIndices.RemoveAt(i);
+				Modify();
+				Layer->Brushes.RemoveAt(i);
+				InBrush->SetOwningLandscape(nullptr);
+				RequestLayersContentUpdateForceAll();
 				break;
 			}
 		}
-
-		InBrush->SetOwningLandscape(nullptr);
-	}
-
-	RequestLayersContentUpdateForceAll();
-}
-
-bool ALandscape::AreAllBrushesCommitedToLayer(int32 InLayerIndex, int32 InTargetType)
-{
-	FLandscapeLayer* Layer = GetLayer(InLayerIndex);
-
-	if (Layer == nullptr)
-	{
-		return false;
-	}
-
-	for (FLandscapeLayerBrush& Brush : Layer->Brushes)
-	{
-		if (Brush.BPCustomBrush != nullptr && !Brush.BPCustomBrush->IsCommited()
-			&& ((InTargetType == ELandscapeToolTargetType::Type::Heightmap && Brush.BPCustomBrush->IsAffectingHeightmap()) || (InTargetType == ELandscapeToolTargetType::Type::Weightmap && Brush.BPCustomBrush->IsAffectingWeightmap())))
-		{
-			return false;
-		}
-	}
-
-	return true;
-}
-
-void ALandscape::SetBrushesCommitStateForLayer(int32 InLayerIndex, int32 InTargetType, bool InCommited)
-{
-	FLandscapeLayer* Layer = GetLayer(InLayerIndex);
-
-	if (Layer == nullptr)
-	{
-		return;
-	}
-
-	for (FLandscapeLayerBrush& Brush : Layer->Brushes)
-	{
-		if (Brush.BPCustomBrush != nullptr)
-		{
-			Brush.BPCustomBrush->SetCommitState(InCommited);
-		}
 	}
 }
 
-TArray<int8>& ALandscape::GetBrushesOrderForLayer(int32 InLayerIndex, int32 InTargetType)
+void ALandscape::OnBPCustomBrushChanged()
 {
-	FLandscapeLayer* Layer = GetLayer(InLayerIndex);
-	check(Layer != nullptr);
-
-	if (InTargetType == ELandscapeToolTargetType::Type::Heightmap)
-	{
-		return Layer->HeightmapBrushOrderIndices;
-	}
-
-	return Layer->WeightmapBrushOrderIndices;
+#if WITH_EDITORONLY_DATA
+	LandscapeBPCustomBrushChangedDelegate.Broadcast();
+	RequestLayersContentUpdateForceAll();
+#endif
 }
 
-ALandscapeBlueprintCustomBrush* ALandscape::GetBrushForLayer(int32 InLayerIndex, int32 InTargetType, int8 InBrushIndex) const
+ALandscapeBlueprintCustomBrush* ALandscape::GetBrushForLayer(int32 InLayerIndex, int8 InBrushIndex) const
 {
-	const FLandscapeLayer* Layer = GetLayer(InLayerIndex);
-
-	if (Layer == nullptr)
-	{
-		return nullptr;
-	}
-
-	if (InTargetType == ELandscapeToolTargetType::Type::Heightmap)
-	{
-		if (Layer->HeightmapBrushOrderIndices.IsValidIndex(InBrushIndex))
-		{
-			int8 ActualBrushIndex = Layer->HeightmapBrushOrderIndices[InBrushIndex];
-			if (Layer->Brushes.IsValidIndex(ActualBrushIndex))
-			{
-				return Layer->Brushes[ActualBrushIndex].BPCustomBrush;
-			}
-		}
-	}
-	else
+	if (const FLandscapeLayer* Layer = GetLayer(InLayerIndex))
 	{
-		if (Layer->WeightmapBrushOrderIndices.IsValidIndex(InBrushIndex))
+		if (Layer->Brushes.IsValidIndex(InBrushIndex))
 		{
-			int8 ActualBrushIndex = Layer->WeightmapBrushOrderIndices[InBrushIndex];
-			if (Layer->Brushes.IsValidIndex(ActualBrushIndex))
-			{
-				return Layer->Brushes[ActualBrushIndex].BPCustomBrush;
-			}
+			return Layer->Brushes[InBrushIndex].BPCustomBrush;
 		}
 	}
-
 	return nullptr;
 }
 
-TArray<ALandscapeBlueprintCustomBrush*> ALandscape::GetBrushesForLayer(int32 InLayerIndex, int32 InTargetType) const
+TArray<ALandscapeBlueprintCustomBrush*> ALandscape::GetBrushesForLayer(int32 InLayerIndex) const
 {
 	TArray<ALandscapeBlueprintCustomBrush*> Brushes;
-
-	const FLandscapeLayer* Layer = GetLayer(InLayerIndex);
-
-	if (Layer == nullptr)
-	{
-		return Brushes;
-	}
-
-	Brushes.Reserve(Layer->Brushes.Num());
-
-	for (const FLandscapeLayerBrush& Brush : Layer->Brushes)
+	if (const FLandscapeLayer* Layer = GetLayer(InLayerIndex))
 	{
-		if (Brush.BPCustomBrush != nullptr && ((Brush.BPCustomBrush->IsAffectingHeightmap() && InTargetType == ELandscapeToolTargetType::Type::Heightmap)
-			|| (Brush.BPCustomBrush->IsAffectingWeightmap() && InTargetType == ELandscapeToolTargetType::Type::Weightmap)))
+		Brushes.Reserve(Layer->Brushes.Num());
+		for (const FLandscapeLayerBrush& Brush : Layer->Brushes)
 		{
 			Brushes.Add(Brush.BPCustomBrush);
 		}
 	}
-
 	return Brushes;
 }
 
