Fixed BP Brush changes not triggering an update of the landscape
#rnx

diff --git a/Engine/Source/Runtime/Landscape/Private/LandscapeBPCustomBrush.cpp b/Engine/Source/Runtime/Landscape/Private/LandscapeBPCustomBrush.cpp
--- a/Engine/Source/Runtime/Landscape/Private/LandscapeBPCustomBrush.cpp
+++ b/Engine/Source/Runtime/Landscape/Private/LandscapeBPCustomBrush.cpp
@@ -92,19 +92,9 @@ void ALandscapeBlueprintCustomBrush::PostEditMove(bool bFinished)
 void ALandscapeBlueprintCustomBrush::PostEditChangeProperty(FPropertyChangedEvent& PropertyChangedEvent)
 {
 	Super::PostEditChangeProperty(PropertyChangedEvent);
-	
 	if (OwningLandscape)
 	{
-		const FName PropertyName = PropertyChangedEvent.Property ? PropertyChangedEvent.Property->GetFName() : NAME_None;
-		if (PropertyName == GET_MEMBER_NAME_CHECKED(ALandscapeBlueprintCustomBrush, AffectHeightmap) || 
-			PropertyName == GET_MEMBER_NAME_CHECKED(ALandscapeBlueprintCustomBrush, AffectWeightmap))
-		{
-			// Should trigger a rebuild of the UI so the visual is updated with changes made to actor
-			//TODO: find a way to trigger the update of the UI
-			//FEdModeLandscape* EdMode = (FEdModeLandscape*)GLevelEditorModeTools().GetActiveMode(FBuiltinEditorModes::EM_Landscape);
-			//EdMode->RefreshDetailPanel();
-			OwningLandscape->OnBPCustomBrushChanged(); 
-		}
+		OwningLandscape->OnBPCustomBrushChanged();
 	}
 }
 #endif
