- Landscape Layer System: Blueprint Brush can now request layer update from Blueprint Script -> RequestLandscapeUpdate()

#rnx

diff --git a/Engine/Source/Runtime/Landscape/Classes/LandscapeBPCustomBrush.h b/Engine/Source/Runtime/Landscape/Classes/LandscapeBPCustomBrush.h
--- a/Engine/Source/Runtime/Landscape/Classes/LandscapeBPCustomBrush.h
+++ b/Engine/Source/Runtime/Landscape/Classes/LandscapeBPCustomBrush.h
@@ -51,6 +51,9 @@ public:
 	UFUNCTION(BlueprintImplementableEvent)
 	void Initialize(const FIntPoint& InLandscapeSize, const FIntPoint& InLandscapeRenderTargetSize);
 
+	UFUNCTION(BlueprintCallable, Category="Landscape")
+	void RequestLandscapeUpdate();
+
 #if WITH_EDITOR
 	void SetCommitState(bool InCommited);
 	bool IsCommited() const { return bIsCommited; }
diff --git a/Engine/Source/Runtime/Landscape/Private/LandscapeBPCustomBrush.cpp b/Engine/Source/Runtime/Landscape/Private/LandscapeBPCustomBrush.cpp
--- a/Engine/Source/Runtime/Landscape/Private/LandscapeBPCustomBrush.cpp
+++ b/Engine/Source/Runtime/Landscape/Private/LandscapeBPCustomBrush.cpp
@@ -42,6 +42,14 @@ bool ALandscapeBlueprintCustomBrush::ShouldTickIfViewportsOnly() const
 	return true;
 }
 
+void ALandscapeBlueprintCustomBrush::RequestLandscapeUpdate()
+{
+	if (OwningLandscape)
+	{
+		OwningLandscape->RequestLayersContentUpdateForceAll();
+	}
+}
+
 #if WITH_EDITOR
 
 void ALandscapeBlueprintCustomBrush::SetCommitState(bool InCommited)
