Fix Build

#rnx

diff --git a/Engine/Source/Runtime/Landscape/Private/LandscapeBPCustomBrush.cpp b/Engine/Source/Runtime/Landscape/Private/LandscapeBPCustomBrush.cpp
--- a/Engine/Source/Runtime/Landscape/Private/LandscapeBPCustomBrush.cpp
+++ b/Engine/Source/Runtime/Landscape/Private/LandscapeBPCustomBrush.cpp
@@ -44,10 +44,12 @@ bool ALandscapeBlueprintCustomBrush::ShouldTickIfViewportsOnly() const
 
 void ALandscapeBlueprintCustomBrush::RequestLandscapeUpdate()
 {
+#if WITH_EDITORONLY_DATA
 	if (OwningLandscape)
 	{
 		OwningLandscape->RequestLayersContentUpdateForceAll();
 	}
+#endif
 }
 
 #if WITH_EDITOR
