ALandscapeBlueprintCustomBrush::Initialize now also receives Landscape's Transform
#rnx

diff --git a/Engine/Source/Runtime/Landscape/Classes/Landscape.h b/Engine/Source/Runtime/Landscape/Classes/Landscape.h
--- a/Engine/Source/Runtime/Landscape/Classes/Landscape.h
+++ b/Engine/Source/Runtime/Landscape/Classes/Landscape.h
@@ -90,6 +90,8 @@ struct FLandscapeLayerBrush
 
 	FLandscapeLayerBrush()
 		: BPCustomBrush(nullptr)
+		, LandscapeSize(MAX_int32, MAX_int32)
+		, LandscapeRenderTargetSize(MAX_int32, MAX_int32)
 	{}
 
 	FLandscapeLayerBrush(ALandscapeBlueprintCustomBrush* InBrush)
@@ -97,36 +99,21 @@ struct FLandscapeLayerBrush
 	{}
 
 #if WITH_EDITOR
-	UTextureRenderTarget2D* Render(bool InIsHeightmap, UTextureRenderTarget2D* InCombinedResult)
-	{
-		if (BPCustomBrush != nullptr)
-		{
-			TGuardValue<bool> AutoRestore(GAllowActorScriptExecutionInEditor, true);
-			return BPCustomBrush->Render(InIsHeightmap, InCombinedResult);
-		}
-
-		return nullptr;
-	}
-
-	bool IsInitialized() const 
-	{
-		return BPCustomBrush != nullptr ? BPCustomBrush->IsInitialized() : false;
-	}
-
-	void Initialize(const FIntRect& InBoundRect, const FIntPoint& InLandscapeRenderTargetSize)
-	{
-		if (BPCustomBrush != nullptr)
-		{
-			TGuardValue<bool> AutoRestore(GAllowActorScriptExecutionInEditor, true);
-			FIntPoint LandscapeSize = InBoundRect.Max - InBoundRect.Min;
-			BPCustomBrush->Initialize(LandscapeSize, InLandscapeRenderTargetSize);
-			BPCustomBrush->SetIsInitialized(true);
-		}
-	}
+	UTextureRenderTarget2D* Render(bool InIsHeightmap, const FIntPoint& InLandscapeSize, UTextureRenderTarget2D* InLandscapeRenderTarget);
 #endif
 
 	UPROPERTY()
 	ALandscapeBlueprintCustomBrush* BPCustomBrush;
+
+private:
+
+#if WITH_EDITOR
+	bool Initialize(const FIntPoint& LandscapeSize, UTextureRenderTarget2D* InLandscapeRenderTarget);
+
+	FTransform LandscapeTransform;
+	FIntPoint LandscapeSize;
+	FIntPoint LandscapeRenderTargetSize;
+#endif
 };
 
 UENUM()
diff --git a/Engine/Source/Runtime/Landscape/Classes/LandscapeBPCustomBrush.h b/Engine/Source/Runtime/Landscape/Classes/LandscapeBPCustomBrush.h
--- a/Engine/Source/Runtime/Landscape/Classes/LandscapeBPCustomBrush.h
+++ b/Engine/Source/Runtime/Landscape/Classes/LandscapeBPCustomBrush.h
@@ -30,9 +30,6 @@ private:
 	UPROPERTY(NonTransactional, DuplicateTransient)
 	bool bIsCommited;
 
-	UPROPERTY(Transient)
-	bool bIsInitialized;
-
 	UPROPERTY(Transient)
 	bool bIsVisible;
 #endif
@@ -49,7 +46,7 @@ public:
 	UTextureRenderTarget2D* Render(bool InIsHeightmap, UTextureRenderTarget2D* InCombinedResult);
 
 	UFUNCTION(BlueprintImplementableEvent)
-	void Initialize(const FIntPoint& InLandscapeSize, const FIntPoint& InLandscapeRenderTargetSize);
+	void Initialize(const FTransform& InLandscapeTransform, const FIntPoint& InLandscapeSize, const FIntPoint& InLandscapeRenderTargetSize);
 
 	UFUNCTION(BlueprintCallable, Category="Landscape")
 	void RequestLandscapeUpdate();
@@ -58,9 +55,6 @@ public:
 	void SetCommitState(bool InCommited);
 	bool IsCommited() const { return bIsCommited; }
 
-	bool IsInitialized() const { return bIsInitialized; }
-	void SetIsInitialized(bool InIsInitialized);
-
 	bool IsVisible() const { return bIsVisible; }
 	void SetIsVisible(bool bInIsVisible);
 
diff --git a/Engine/Source/Runtime/Landscape/Private/LandscapeBPCustomBrush.cpp b/Engine/Source/Runtime/Landscape/Private/LandscapeBPCustomBrush.cpp
--- a/Engine/Source/Runtime/Landscape/Private/LandscapeBPCustomBrush.cpp
+++ b/Engine/Source/Runtime/Landscape/Private/LandscapeBPCustomBrush.cpp
@@ -11,7 +11,6 @@ ALandscapeBlueprintCustomBrush::ALandscapeBlueprintCustomBrush(const FObjectInit
 #if WITH_EDITORONLY_DATA
 	: OwningLandscape(nullptr)
 	, bIsCommited(false)
-	, bIsInitialized(false)
 	, bIsVisible(true)
 #endif
 {
@@ -95,11 +94,6 @@ ALandscape* ALandscapeBlueprintCustomBrush::GetOwningLandscape() const
 	return OwningLandscape; 
 }
 
-void ALandscapeBlueprintCustomBrush::SetIsInitialized(bool InIsInitialized)
-{
-	bIsInitialized = InIsInitialized;
-}
-
 bool ALandscapeBlueprintCustomBrush::IsAffectingWeightmapLayer(const FName& InLayerName) const
 {
 	return AffectedWeightmapLayers.Contains(InLayerName);
diff --git a/Engine/Source/Runtime/Landscape/Private/LandscapeEdit.cpp b/Engine/Source/Runtime/Landscape/Private/LandscapeEdit.cpp
--- a/Engine/Source/Runtime/Landscape/Private/LandscapeEdit.cpp
+++ b/Engine/Source/Runtime/Landscape/Private/LandscapeEdit.cpp
@@ -3182,24 +3182,19 @@ bool ALandscapeProxy::ExportToRawMesh(int32 InExportLOD, FMeshDescription& OutRa
 
 FIntRect ALandscapeProxy::GetBoundingRect() const
 {
-	FIntRect Rect(MAX_int32, MAX_int32, MIN_int32, MIN_int32);
-
-	for (int32 CompIdx = 0; CompIdx < LandscapeComponents.Num(); CompIdx++)
-	{
-		Rect.Include(LandscapeComponents[CompIdx]->GetSectionBase());
-	}
-
 	if (LandscapeComponents.Num() > 0)
 	{
+		FIntRect Rect(MAX_int32, MAX_int32, MIN_int32, MIN_int32);
+		for (int32 CompIdx = 0; CompIdx < LandscapeComponents.Num(); CompIdx++)
+		{
+			Rect.Include(LandscapeComponents[CompIdx]->GetSectionBase());
+		}
 		Rect.Max += FIntPoint(ComponentSizeQuads, ComponentSizeQuads);
 		Rect -= LandscapeSectionOffset;
-	}
-	else
-	{
-		Rect = FIntRect();
+		return Rect;
 	}
 
-	return Rect;
+	return FIntRect();
 }
 
 bool ALandscape::HasAllComponent()
diff --git a/Engine/Source/Runtime/Landscape/Private/LandscapeEditLayers.cpp b/Engine/Source/Runtime/Landscape/Private/LandscapeEditLayers.cpp
--- a/Engine/Source/Runtime/Landscape/Private/LandscapeEditLayers.cpp
+++ b/Engine/Source/Runtime/Landscape/Private/LandscapeEditLayers.cpp
@@ -2632,11 +2632,11 @@ int32 ALandscape::RegenerateLayersHeightmaps(const TArray<ULandscapeComponent*>&
 
 		FIntPoint MinExtend;
 		FIntPoint MaxExtend;
-
 		if (!Info->GetLandscapeExtent(MinExtend.X, MinExtend.Y, MaxExtend.X, MaxExtend.Y))
 		{
 			return 0;
 		}
+		const FIntPoint LandscapeSize = MaxExtend - MinExtend;
 
 		// Use to compute TopLeft Component per Heightmap
 		struct FHeightmapTopLeft
@@ -2743,13 +2743,7 @@ int32 ALandscape::RegenerateLayersHeightmaps(const TArray<ULandscapeComponent*>&
 						continue;
 					}
 
-					if (!Brush.IsInitialized())
-					{
-						Brush.Initialize(GetBoundingRect(), FIntPoint(CombinedHeightmapNonAtlasRT->SizeX, CombinedHeightmapNonAtlasRT->SizeY));
-					}
-
-					UTextureRenderTarget2D* BrushOutputNonAtlasRT = Brush.Render(true, CombinedHeightmapNonAtlasRT);
-
+					UTextureRenderTarget2D* BrushOutputNonAtlasRT = Brush.Render(true, LandscapeSize, CombinedHeightmapNonAtlasRT);
 					if (BrushOutputNonAtlasRT == nullptr || BrushOutputNonAtlasRT->SizeX != CombinedHeightmapNonAtlasRT->SizeX || BrushOutputNonAtlasRT->SizeY != CombinedHeightmapNonAtlasRT->SizeY)
 					{
 						continue;
@@ -3433,6 +3427,7 @@ int32 ALandscape::RegenerateLayersWeightmaps(const TArray<ULandscapeComponent*>&
 		{
 			return 0;
 		}
+		const FIntPoint LandscapeSize = MaxExtend - MinExtend;
 		
 		check(WeightmapRTList.Num() > 0);
 
@@ -3600,13 +3595,7 @@ int32 ALandscape::RegenerateLayersWeightmaps(const TArray<ULandscapeComponent*>&
 
 							BrushRequiredAllocations.AddUnique(LayerInfoObj);
 
-							if (!Brush.IsInitialized())
-							{
-								Brush.Initialize(GetBoundingRect(), FIntPoint(LandscapeScratchRT3->SizeX, LandscapeScratchRT3->SizeY));
-							}
-
-							UTextureRenderTarget2D* BrushOutputRT = Brush.Render(false, LandscapeScratchRT3);
-
+							UTextureRenderTarget2D* BrushOutputRT = Brush.Render(false, LandscapeSize, LandscapeScratchRT3);
 							if (BrushOutputRT == nullptr || BrushOutputRT->SizeX != LandscapeScratchRT3->SizeX || BrushOutputRT->SizeY != LandscapeScratchRT3->SizeY)
 							{
 								continue;
@@ -5506,6 +5495,39 @@ TArray<ALandscapeBlueprintCustomBrush*> ALandscape::GetBrushesForLayer(int32 InL
 	return Brushes;
 }
 
+UTextureRenderTarget2D* FLandscapeLayerBrush::Render(bool InIsHeightmap, const FIntPoint& InLandscapeSize, UTextureRenderTarget2D* InLandscapeRenderTarget)
+{
+	if (Initialize(InLandscapeSize, InLandscapeRenderTarget))
+	{
+		TGuardValue<bool> AutoRestore(GAllowActorScriptExecutionInEditor, true);
+		return BPCustomBrush->Render(InIsHeightmap, InLandscapeRenderTarget);
+	}
+	return nullptr;
+}
+
+bool FLandscapeLayerBrush::Initialize(const FIntPoint& InLandscapeSize, UTextureRenderTarget2D* InLandscapeRenderTarget)
+{
+	if (BPCustomBrush && InLandscapeRenderTarget)
+	{
+		if (ALandscape* Landscape = BPCustomBrush->GetOwningLandscape())
+		{
+			const FIntPoint NewLandscapeRenderTargetSize = FIntPoint(InLandscapeRenderTarget->SizeX, InLandscapeRenderTarget->SizeY);
+			const FTransform NewLandscapeTransform = Landscape->GetTransform();
+			if (!LandscapeTransform.Equals(NewLandscapeTransform) || (LandscapeSize != InLandscapeSize) || LandscapeRenderTargetSize != NewLandscapeRenderTargetSize)
+			{
+				LandscapeTransform = NewLandscapeTransform;
+				LandscapeRenderTargetSize = NewLandscapeRenderTargetSize;
+				LandscapeSize = InLandscapeSize;
+				
+				TGuardValue<bool> AutoRestore(GAllowActorScriptExecutionInEditor, true);
+				BPCustomBrush->Initialize(LandscapeTransform, LandscapeSize, LandscapeRenderTargetSize);
+			}
+			return true;
+		}
+	}
+	return false;
+}
+
 #endif // WITH_EDITOR
 
 #undef LOCTEXT_NAMESPACE
\ No newline at end of file
