Fixed non-editor build
#rnx

diff --git a/Engine/Source/Runtime/Landscape/Classes/Landscape.h b/Engine/Source/Runtime/Landscape/Classes/Landscape.h
--- a/Engine/Source/Runtime/Landscape/Classes/Landscape.h
+++ b/Engine/Source/Runtime/Landscape/Classes/Landscape.h
@@ -90,8 +90,10 @@ struct FLandscapeLayerBrush
 
 	FLandscapeLayerBrush()
 		: BPCustomBrush(nullptr)
+#if WITH_EDITOR
 		, LandscapeSize(MAX_int32, MAX_int32)
 		, LandscapeRenderTargetSize(MAX_int32, MAX_int32)
+#endif
 	{}
 
 	FLandscapeLayerBrush(ALandscapeBlueprintCustomBrush* InBrush)
