Fixed Landscape BP Custom Brushes removed from landscape after a re-compilation when outside of Landscape Mode
#rnx

diff --git a/Engine/Source/Runtime/Landscape/Private/LandscapeBPCustomBrush.cpp b/Engine/Source/Runtime/Landscape/Private/LandscapeBPCustomBrush.cpp
--- a/Engine/Source/Runtime/Landscape/Private/LandscapeBPCustomBrush.cpp
+++ b/Engine/Source/Runtime/Landscape/Private/LandscapeBPCustomBrush.cpp
@@ -129,10 +129,11 @@ void ALandscapeBlueprintCustomBrush::Destroyed()
 {
 	Super::Destroyed();
 
-	if (OwningLandscape)
+	if (OwningLandscape && !GIsReinstancing)
 	{
 		OwningLandscape->RemoveBrush(this);
 	}
+	OwningLandscape = nullptr;
 }
 
 #endif
