- Moved ALandscapeBlueprintCustomBrush in an editor module (also renamed as ALandscapeBlueprintBrush).
- Added runtime base class ALandscapeBlueprintBrushBase. This was done so that landscape blueprints brushes can access editor functions.

#rnx

diff --git a/Engine/Config/BaseEngine.ini b/Engine/Config/BaseEngine.ini
--- a/Engine/Config/BaseEngine.ini
+++ b/Engine/Config/BaseEngine.ini
@@ -669,6 +669,9 @@ bUseStreamingPause=false
 +ClassRedirects=(OldName="EditorAutomationActor",NewName="EditorUtilityActor")
 +ClassRedirects=(OldName="EditorAutomationActorComponent",NewName="EditorUtilityActorComponent")
 +ClassRedirects=(OldName="EditorAutomationObject",NewName="EditorUtilityObject")
++ClassRedirects=(OldName="LandscapeBlueprintCustomBrush",NewName="/Script/LandscapeEditor.LandscapeBlueprintBrush")
+
++PropertyRedirects=(OldName="LandscapeLayerBrush.BPCustomBrush",NewName="LandscapeLayerBrush.BlueprintBrush")
 
 +PropertyRedirects=(OldName="StructVariableDescription.bDontEditoOnInstance",NewName="bDontEditOnInstance")
 +PropertyRedirects=(OldName="KismetMathLibrary.DegAtan2.A",NewName="Y")
@@ -2132,11 +2135,11 @@ LLVMWasmBackend=False
 +BrowserLauncher=(BrowserName="FF Nightly(32bit)", BrowserPath=(FilePath="C:/Program Files (x86)/Nightly/firefox.exe"))
 +BrowserLauncher=(BrowserName="FF Nightly",        BrowserPath=(FilePath="C:/Program Files/Nightly/firefox.exe"))
 ; chrome
-+BrowserLauncher=(BrowserName="Chrome",            BrowserPath=(FilePath="C:/Program Files (x86)/Google/Chrome/Application/chrome.exe"))
++BrowserLauncher=(BrowserName="Chrome",         BrowserPath=(FilePath="C:/Program Files (x86)/Google/Chrome/Application/chrome.exe"))
 ; Chrome Canary -- edit path to canary installation location
 +BrowserLauncher=(BrowserName="Canary",            BrowserPath=(FilePath="[REDACTED_USER_HOME]/AppData/Local/Google/Chrome SxS/Application/chrome.exe"))
 ; Edge -- will be autodetected...
-;+BrowserLauncher=(BrowserName="Edge",             BrowserPath=(FilePath="start microsoft-edge:"))
+;+BrowserLauncher=(BrowserName="Edge",           BrowserPath=(FilePath="start microsoft-edge:"))
 
 ; osx
 +BrowserLauncher=(BrowserName="Safari",  BrowserPath=(FilePath="/Applications/Safari.app"))
diff --git a/Engine/Source/Editor/LandscapeEditor/Private/LandscapeBlueprintBrush.cpp b/Engine/Source/Editor/LandscapeEditor/Private/LandscapeBlueprintBrush.cpp
new file mode 100644
--- /dev/null
+++ b/Engine/Source/Editor/LandscapeEditor/Private/LandscapeBlueprintBrush.cpp
@@ -0,0 +1,15 @@
+// Copyright 1998-2019 Epic Games, Inc. All Rights Reserved.
+
+#include "LandscapeBlueprintBrush.h"
+#include "CoreMinimal.h"
+#include "LandscapeProxy.h"
+#include "Landscape.h"
+
+#define LOCTEXT_NAMESPACE "Landscape"
+
+ALandscapeBlueprintBrush::ALandscapeBlueprintBrush(const FObjectInitializer& ObjectInitializer)
+	: Super(ObjectInitializer)
+{
+}
+
+#undef LOCTEXT_NAMESPACE
diff --git a/Engine/Source/Editor/LandscapeEditor/Private/LandscapeEdMode.cpp b/Engine/Source/Editor/LandscapeEditor/Private/LandscapeEdMode.cpp
--- a/Engine/Source/Editor/LandscapeEditor/Private/LandscapeEdMode.cpp
+++ b/Engine/Source/Editor/LandscapeEditor/Private/LandscapeEdMode.cpp
@@ -52,7 +52,7 @@
 #include "Framework/Commands/InputBindingManager.h"
 #include "MouseDeltaTracker.h"
 #include "Interfaces/IMainFrameModule.h"
-#include "LandscapeBPCustomBrush.h"
+#include "LandscapeBlueprintBrushBase.h"
 #include "Engine/TextureRenderTarget2D.h"
 #include "Settings/EditorExperimentalSettings.h"
 #include "ComponentRecreateRenderStateContext.h"
@@ -592,7 +592,7 @@ void FEdModeLandscape::Enter()
 
 		if (ALandscape* Landscape = GetLandscape())
 		{
-			Landscape->OnBPCustomBrushChangedDelegate().AddLambda([this]() { this->RefreshDetailPanel(); });
+			Landscape->OnBlueprintBrushChangedDelegate().AddLambda([this]() { this->RefreshDetailPanel(); });
 			if (Landscape->HasLayersContent())
 			{
 				if (Landscape->GetLandscapeSplinesReservedLayer())
@@ -780,7 +780,7 @@ void FEdModeLandscape::Exit()
 		LandscapeProxy->OnMaterialChangedDelegate().RemoveAll(this);
 		if (ALandscape* Landscape = GetLandscape())
 		{
-			Landscape->OnBPCustomBrushChangedDelegate().RemoveAll(this);
+			Landscape->OnBlueprintBrushChangedDelegate().RemoveAll(this);
 		}
 	}
 
@@ -1622,10 +1622,10 @@ bool FEdModeLandscape::ProcessEditCopy()
 
 		if (!Result)
 		{
-			ALandscapeBlueprintCustomBrush* CurrentlySelectedBPBrush = nullptr;
+			ALandscapeBlueprintBrushBase* CurrentlySelectedBPBrush = nullptr;
 			for (FSelectionIterator It(GEditor->GetSelectedActorIterator()); It; ++It)
 			{
-				CurrentlySelectedBPBrush = Cast<ALandscapeBlueprintCustomBrush>(*It);
+				CurrentlySelectedBPBrush = Cast<ALandscapeBlueprintBrushBase>(*It);
 				if (CurrentlySelectedBPBrush)
 				{
 					break;
@@ -1691,10 +1691,10 @@ bool FEdModeLandscape::ProcessEditPaste()
 
 		if (!Result)
 		{
-			ALandscapeBlueprintCustomBrush* CurrentlySelectedBPBrush = nullptr;
+			ALandscapeBlueprintBrushBase* CurrentlySelectedBPBrush = nullptr;
 			for (FSelectionIterator It(GEditor->GetSelectedActorIterator()); It; ++It)
 			{
-				CurrentlySelectedBPBrush = Cast<ALandscapeBlueprintCustomBrush>(*It);
+				CurrentlySelectedBPBrush = Cast<ALandscapeBlueprintBrushBase>(*It);
 				if (CurrentlySelectedBPBrush)
 				{
 					break;
@@ -2394,7 +2394,7 @@ void FEdModeLandscape::SetCurrentBrush(int32 BrushIndex)
 	}
 }
 
-const TArray<ALandscapeBlueprintCustomBrush*>& FEdModeLandscape::GetBrushList() const
+const TArray<ALandscapeBlueprintBrushBase*>& FEdModeLandscape::GetBrushList() const
 {
 	return BrushList;
 }
@@ -2540,7 +2540,7 @@ void FEdModeLandscape::SetTargetLandscape(const TWeakObjectPtr<ULandscapeInfo>&
 		LandscapeProxy->OnMaterialChangedDelegate().RemoveAll(this);
 		if (ALandscape* Landscape = GetLandscape())
 		{
-			Landscape->OnBPCustomBrushChangedDelegate().RemoveAll(this);
+			Landscape->OnBlueprintBrushChangedDelegate().RemoveAll(this);
 		}
 	}
 
@@ -2560,7 +2560,7 @@ void FEdModeLandscape::SetTargetLandscape(const TWeakObjectPtr<ULandscapeInfo>&
 		LandscapeProxy->OnMaterialChangedDelegate().AddRaw(this, &FEdModeLandscape::OnLandscapeMaterialChangedDelegate);
 		if (ALandscape* Landscape = GetLandscape())
 		{
-			Landscape->OnBPCustomBrushChangedDelegate().AddLambda([this]() { this->RefreshDetailPanel(); });
+			Landscape->OnBlueprintBrushChangedDelegate().AddLambda([this]() { this->RefreshDetailPanel(); });
 		}
 	}
 
@@ -3426,7 +3426,7 @@ bool FEdModeLandscape::IsSelectionAllowed(AActor* InActor, bool bInSelection) co
 	{
 		return true;
 	}
-	else if (InActor->IsA(ALandscapeBlueprintCustomBrush::StaticClass()))
+	else if (InActor->IsA(ALandscapeBlueprintBrushBase::StaticClass()))
 	{
 		return true;
 	}
@@ -4478,7 +4478,7 @@ void FEdModeLandscape::RequestLayersContentUpdateForceAll(ELandscapeLayerUpdateM
 	}
 }
 
-void FEdModeLandscape::AddBrushToCurrentLayer(ALandscapeBlueprintCustomBrush* InBrush)
+void FEdModeLandscape::AddBrushToCurrentLayer(ALandscapeBlueprintBrushBase* InBrush)
 {
 	ALandscape* Landscape = GetLandscape();
 	if (Landscape == nullptr)
@@ -4490,7 +4490,7 @@ void FEdModeLandscape::AddBrushToCurrentLayer(ALandscapeBlueprintCustomBrush* In
 	RefreshDetailPanel();
 }
 
-void FEdModeLandscape::RemoveBrushFromCurrentLayer(ALandscapeBlueprintCustomBrush* InBrush)
+void FEdModeLandscape::RemoveBrushFromCurrentLayer(ALandscapeBlueprintBrushBase* InBrush)
 {
 	ALandscape* Landscape = GetLandscape();
 
@@ -4503,7 +4503,7 @@ void FEdModeLandscape::RemoveBrushFromCurrentLayer(ALandscapeBlueprintCustomBrus
 	RefreshDetailPanel();
 }
 
-ALandscapeBlueprintCustomBrush* FEdModeLandscape::GetBrushForCurrentLayer(int8 InBrushIndex) const
+ALandscapeBlueprintBrushBase* FEdModeLandscape::GetBrushForCurrentLayer(int8 InBrushIndex) const
 {
 	if (ALandscape* Landscape = GetLandscape())
 	{
@@ -4512,9 +4512,9 @@ ALandscapeBlueprintCustomBrush* FEdModeLandscape::GetBrushForCurrentLayer(int8 I
 	return nullptr;
 }
 
-TArray<ALandscapeBlueprintCustomBrush*> FEdModeLandscape::GetBrushesForCurrentLayer()
+TArray<ALandscapeBlueprintBrushBase*> FEdModeLandscape::GetBrushesForCurrentLayer()
 {
-	TArray<ALandscapeBlueprintCustomBrush*> Brushes;
+	TArray<ALandscapeBlueprintBrushBase*> Brushes;
 	if (ALandscape* Landscape = GetLandscape())
 	{
 		Brushes = Landscape->GetBrushesForLayer(GetCurrentLayerIndex());
@@ -4714,9 +4714,9 @@ bool FEdModeLandscape::NeedToFillEmptyMaterialLayers() const
 void FEdModeLandscape::UpdateBrushList()
 {
 	BrushList.Empty();
-	for (TObjectIterator<ALandscapeBlueprintCustomBrush> BrushIt(RF_Transient|RF_ClassDefaultObject|RF_ArchetypeObject, true, EInternalObjectFlags::PendingKill); BrushIt; ++BrushIt)
+	for (TObjectIterator<ALandscapeBlueprintBrushBase> BrushIt(RF_Transient|RF_ClassDefaultObject|RF_ArchetypeObject, true, EInternalObjectFlags::PendingKill); BrushIt; ++BrushIt)
 	{
-		ALandscapeBlueprintCustomBrush* Brush = *BrushIt;
+		ALandscapeBlueprintBrushBase* Brush = *BrushIt;
 		if (Brush->GetTypedOuter<UPackage>() != GetTransientPackage())
 		{
 			BrushList.Add(Brush);
@@ -4732,7 +4732,7 @@ void FEdModeLandscape::OnLevelActorAdded(AActor* InActor)
 		Landscape->RegisterLandscapeEdMode(this);
 	}
 
-	ALandscapeBlueprintCustomBrush* Brush = Cast<ALandscapeBlueprintCustomBrush>(InActor);
+	ALandscapeBlueprintBrushBase* Brush = Cast<ALandscapeBlueprintBrushBase>(InActor);
 	if (Brush && Brush->GetTypedOuter<UPackage>() != GetTransientPackage())
 	{
 		if (!GIsReinstancing)
@@ -4751,7 +4751,7 @@ void FEdModeLandscape::OnLevelActorRemoved(AActor* InActor)
 		Landscape->UnregisterLandscapeEdMode();
 	}
 
-	ALandscapeBlueprintCustomBrush* Brush = Cast<ALandscapeBlueprintCustomBrush>(InActor);
+	ALandscapeBlueprintBrushBase* Brush = Cast<ALandscapeBlueprintBrushBase>(InActor);
 	if (Brush && Brush->GetTypedOuter<UPackage>() != GetTransientPackage())
 	{
 		UpdateBrushList();
diff --git a/Engine/Source/Editor/LandscapeEditor/Private/LandscapeEdMode.h b/Engine/Source/Editor/LandscapeEditor/Private/LandscapeEdMode.h
--- a/Engine/Source/Editor/LandscapeEditor/Private/LandscapeEdMode.h
+++ b/Engine/Source/Editor/LandscapeEditor/Private/LandscapeEdMode.h
@@ -510,7 +510,7 @@ public:
 	void SetCurrentBrush(int32 BrushIndex);
 
 	void UpdateBrushList();
-	const TArray<ALandscapeBlueprintCustomBrush*>& GetBrushList() const;
+	const TArray<ALandscapeBlueprintBrushBase*>& GetBrushList() const;
 
 	const TArray<TSharedRef<FLandscapeTargetListInfo>>& GetTargetList() const;
 	const TArray<FName>* GetTargetDisplayOrderList() const;
@@ -564,10 +564,10 @@ public:
 	void AutoUpdateDirtyLandscapeSplines();
 	bool CanEditLayer(FText* Reason = nullptr, FLandscapeLayer* InLayer = nullptr);
 
-	void AddBrushToCurrentLayer(class ALandscapeBlueprintCustomBrush* InBrush);
-	void RemoveBrushFromCurrentLayer(class ALandscapeBlueprintCustomBrush* InBrush);
-	class ALandscapeBlueprintCustomBrush* GetBrushForCurrentLayer(int8 BrushIndex) const;
-	TArray<class ALandscapeBlueprintCustomBrush*> GetBrushesForCurrentLayer();
+	void AddBrushToCurrentLayer(class ALandscapeBlueprintBrushBase* InBrush);
+	void RemoveBrushFromCurrentLayer(class ALandscapeBlueprintBrushBase* InBrush);
+	class ALandscapeBlueprintBrushBase* GetBrushForCurrentLayer(int8 BrushIndex) const;
+	TArray<class ALandscapeBlueprintBrushBase*> GetBrushesForCurrentLayer();
 	
 	bool NeedToFillEmptyMaterialLayers() const;
 	void RequestLayersContentUpdate(ELandscapeLayerUpdateMode InUpdateMode);
@@ -619,7 +619,7 @@ public:
 private:
 	TArray<TSharedRef<FLandscapeTargetListInfo>> LandscapeTargetList;
 	TArray<FLandscapeListInfo> LandscapeList;
-	TArray<ALandscapeBlueprintCustomBrush*> BrushList;
+	TArray<ALandscapeBlueprintBrushBase*> BrushList;
 	TArray<FName> ShownTargetLayerList;
 	
 	/** Represent the index offset of the target layer in LandscapeTargetList */
diff --git a/Engine/Source/Editor/LandscapeEditor/Private/LandscapeEdModeBPCustomTools.cpp b/Engine/Source/Editor/LandscapeEditor/Private/LandscapeEdModeBPCustomTools.cpp
--- a/Engine/Source/Editor/LandscapeEditor/Private/LandscapeEdModeBPCustomTools.cpp
+++ b/Engine/Source/Editor/LandscapeEditor/Private/LandscapeEdModeBPCustomTools.cpp
@@ -19,7 +19,7 @@
 #include "LandscapeRender.h"
 #include "LandscapeHeightfieldCollisionComponent.h"
 #include "LandscapeEdModeTools.h"
-#include "LandscapeBPCustomBrush.h"
+#include "LandscapeBlueprintBrushBase.h"
 #include "LandscapeInfo.h"
 #include "Landscape.h"
 //#include "LandscapeDataAccess.h"
@@ -72,12 +72,12 @@ public:
 
 	virtual bool BeginTool(FEditorViewportClient* ViewportClient, const FLandscapeToolTarget& Target, const FVector& InHitLocation) override
 	{
-		if (EdMode->UISettings->BlueprintCustomBrush == nullptr)
+		if (EdMode->UISettings->BlueprintBrush == nullptr)
 		{
 			return false;
 		}
 
-		ALandscapeBlueprintCustomBrush* DefaultObject = Cast<ALandscapeBlueprintCustomBrush>(EdMode->UISettings->BlueprintCustomBrush->GetDefaultObject(false));
+		ALandscapeBlueprintBrushBase* DefaultObject = Cast<ALandscapeBlueprintBrushBase>(EdMode->UISettings->BlueprintBrush->GetDefaultObject(false));
 
 		if (DefaultObject == nullptr)
 		{
@@ -97,9 +97,9 @@ public:
 			SpawnInfo.bNoFail = true;
 			SpawnInfo.OverrideLevel = Info->LandscapeActor.Get()->GetTypedOuter<ULevel>(); // always spawn in the same level as the one containing the ALandscape
 
-			FScopedTransaction Transaction(LOCTEXT("LandscapeEdModeBPCustomToolSpawn", "Create landscape brush"));
-			ALandscapeBlueprintCustomBrush* Brush = ViewportClient->GetWorld()->SpawnActor<ALandscapeBlueprintCustomBrush>(EdMode->UISettings->BlueprintCustomBrush, SpawnLocation, FRotator(0.0f), SpawnInfo);
-			EdMode->UISettings->BlueprintCustomBrush = nullptr;
+			FScopedTransaction Transaction(LOCTEXT("LandscapeEdModeBlueprintToolSpawn", "Create landscape brush"));
+			ALandscapeBlueprintBrushBase* Brush = ViewportClient->GetWorld()->SpawnActor<ALandscapeBlueprintBrushBase>(EdMode->UISettings->BlueprintBrush, SpawnLocation, FRotator(0.0f), SpawnInfo);
+			EdMode->UISettings->BlueprintBrush = nullptr;
 
 			GEditor->SelectNone(true, true);
 			GEditor->SelectActor(Brush, true, true);
diff --git a/Engine/Source/Editor/LandscapeEditor/Private/LandscapeEditorDetailCustomization_Layers.cpp b/Engine/Source/Editor/LandscapeEditor/Private/LandscapeEditorDetailCustomization_Layers.cpp
--- a/Engine/Source/Editor/LandscapeEditor/Private/LandscapeEditorDetailCustomization_Layers.cpp
+++ b/Engine/Source/Editor/LandscapeEditor/Private/LandscapeEditorDetailCustomization_Layers.cpp
@@ -348,9 +348,9 @@ void FLandscapeEditorCustomNodeBuilder_Layers::SetLayerName(const FText& InText,
 }
 
 
-void FLandscapeEditorCustomNodeBuilder_Layers::FillAddBrushMenu(FMenuBuilder& MenuBuilder, TArray<ALandscapeBlueprintCustomBrush*> Brushes)
+void FLandscapeEditorCustomNodeBuilder_Layers::FillAddBrushMenu(FMenuBuilder& MenuBuilder, TArray<ALandscapeBlueprintBrushBase*> Brushes)
 {
-	for (ALandscapeBlueprintCustomBrush* Brush : Brushes)
+	for (ALandscapeBlueprintBrushBase* Brush : Brushes)
 	{
 		TSharedRef<FLandscapeEditorCustomNodeBuilder_Layers> SharedThis = AsShared();
 		FUIAction AddAction = FUIAction(FExecuteAction::CreateLambda([SharedThis, Brush]() { SharedThis->AddBrushToCurrentLayer(Brush); }));
@@ -413,7 +413,7 @@ void FLandscapeEditorCustomNodeBuilder_Layers::FillClearLayerMenu(FMenuBuilder&
 	}
 }
 
-void FLandscapeEditorCustomNodeBuilder_Layers::AddBrushToCurrentLayer(ALandscapeBlueprintCustomBrush* Brush)
+void FLandscapeEditorCustomNodeBuilder_Layers::AddBrushToCurrentLayer(ALandscapeBlueprintBrushBase* Brush)
 {
 	const FScopedTransaction Transaction(LOCTEXT("LandscapeBrushAddToCurrentLayerTransaction", "Add brush to current layer"));
 	GetEditorMode()->AddBrushToCurrentLayer(Brush);
@@ -504,8 +504,8 @@ TSharedPtr<SWidget> FLandscapeEditorCustomNodeBuilder_Layers::OnLayerContextMenu
 		}
 		MenuBuilder.EndSection();
 
-		const TArray<ALandscapeBlueprintCustomBrush*>& Brushes = LandscapeEdMode->GetBrushList();
-		TArray<ALandscapeBlueprintCustomBrush*> FilteredBrushes = Brushes.FilterByPredicate([](ALandscapeBlueprintCustomBrush* Brush) { return Brush->GetOwningLandscape() == nullptr; });
+		const TArray<ALandscapeBlueprintBrushBase*>& Brushes = LandscapeEdMode->GetBrushList();
+		TArray<ALandscapeBlueprintBrushBase*> FilteredBrushes = Brushes.FilterByPredicate([](ALandscapeBlueprintBrushBase* Brush) { return Brush->GetOwningLandscape() == nullptr; });
 		if (FilteredBrushes.Num())
 		{
 			MenuBuilder.BeginSection("LandscapeEditorBrushActions", LOCTEXT("LandscapeEditorBrushActions.Heading", "Brushes"));
diff --git a/Engine/Source/Editor/LandscapeEditor/Private/LandscapeEditorDetailCustomization_Layers.h b/Engine/Source/Editor/LandscapeEditor/Private/LandscapeEditorDetailCustomization_Layers.h
--- a/Engine/Source/Editor/LandscapeEditor/Private/LandscapeEditorDetailCustomization_Layers.h
+++ b/Engine/Source/Editor/LandscapeEditor/Private/LandscapeEditorDetailCustomization_Layers.h
@@ -23,7 +23,7 @@ class FDetailWidgetRow;
 class IDetailChildrenBuilder;
 class IDetailLayoutBuilder;
 class SDragAndDropVerticalBox;
-class ALandscapeBlueprintCustomBrush;
+class ALandscapeBlueprintBrushBase;
 class ULandscapeLayerInfoObject;
 
 /**
@@ -100,10 +100,10 @@ protected:
 	FReply OnToggleLock(int32 InLayerIndex);
 	const FSlateBrush* GetLockBrushForLayer(int32 InLayerIndex) const;
 
-	void FillAddBrushMenu(FMenuBuilder& MenuBuilder, TArray<ALandscapeBlueprintCustomBrush*> Brushes);
+	void FillAddBrushMenu(FMenuBuilder& MenuBuilder, TArray<ALandscapeBlueprintBrushBase*> Brushes);
+	void AddBrushToCurrentLayer(ALandscapeBlueprintBrushBase* Brush);
 	void FillClearPaintLayerMenu(FMenuBuilder& MenuBuilder, int32 InLayerIndex, TArray<ULandscapeLayerInfoObject*> InUsedLayerInfos);
 	void FillClearLayerMenu(FMenuBuilder& MenuBuilder, int32 InLayerIndex);
-	void AddBrushToCurrentLayer(ALandscapeBlueprintCustomBrush* Brush);
 
 private:
 
diff --git a/Engine/Source/Editor/LandscapeEditor/Private/LandscapeEditorDetailCustomization_LayersBrushStack.cpp b/Engine/Source/Editor/LandscapeEditor/Private/LandscapeEditorDetailCustomization_LayersBrushStack.cpp
--- a/Engine/Source/Editor/LandscapeEditor/Private/LandscapeEditorDetailCustomization_LayersBrushStack.cpp
+++ b/Engine/Source/Editor/LandscapeEditor/Private/LandscapeEditorDetailCustomization_LayersBrushStack.cpp
@@ -26,7 +26,7 @@
 
 #include "LandscapeEditorDetailCustomization_TargetLayers.h"
 #include "Widgets/Input/SEditableText.h"
-#include "LandscapeBPCustomBrush.h"
+#include "LandscapeBlueprintBrushBase.h"
 
 #define LOCTEXT_NAMESPACE "LandscapeEditor.Layers"
 
@@ -209,7 +209,7 @@ END_SLATE_FUNCTION_BUILD_OPTIMIZATION
 
 TSharedPtr<SWidget> FLandscapeEditorCustomNodeBuilder_LayersBrushStack::OnBrushContextMenuOpening(int32 InBrushIndex)
 {
-	if (ALandscapeBlueprintCustomBrush* CurrentBrush = GetBrush(InBrushIndex))
+	if (ALandscapeBlueprintBrushBase* CurrentBrush = GetBrush(InBrushIndex))
 	{
 		FMenuBuilder MenuBuilder(true, NULL);
 		MenuBuilder.BeginSection("LandscapeEditorBrushActions", LOCTEXT("LandscapeEditorBrushActions.Heading", "Brushes"));
@@ -217,8 +217,8 @@ TSharedPtr<SWidget> FLandscapeEditorCustomNodeBuilder_LayersBrushStack::OnBrushC
 			TSharedRef<FLandscapeEditorCustomNodeBuilder_LayersBrushStack> SharedThis = AsShared();
 
 			// Add Brush
-			const TArray<ALandscapeBlueprintCustomBrush*>& Brushes = GetEditorMode()->GetBrushList();
-			TArray<ALandscapeBlueprintCustomBrush*> FilteredBrushes = Brushes.FilterByPredicate([](ALandscapeBlueprintCustomBrush* Brush) { return Brush->GetOwningLandscape() == nullptr; });
+			const TArray<ALandscapeBlueprintBrushBase*>& Brushes = GetEditorMode()->GetBrushList();
+			TArray<ALandscapeBlueprintBrushBase*> FilteredBrushes = Brushes.FilterByPredicate([](ALandscapeBlueprintBrushBase* Brush) { return Brush->GetOwningLandscape() == nullptr; });
 			if (FilteredBrushes.Num())
 			{
 				MenuBuilder.AddSubMenu(
@@ -245,10 +245,10 @@ TSharedPtr<SWidget> FLandscapeEditorCustomNodeBuilder_LayersBrushStack::OnBrushC
 	return nullptr;
 }
 
-void FLandscapeEditorCustomNodeBuilder_LayersBrushStack::FillAddBrushMenu(FMenuBuilder& MenuBuilder, TArray<ALandscapeBlueprintCustomBrush*> Brushes)
+void FLandscapeEditorCustomNodeBuilder_LayersBrushStack::FillAddBrushMenu(FMenuBuilder& MenuBuilder, TArray<ALandscapeBlueprintBrushBase*> Brushes)
 {
 	TSharedRef<FLandscapeEditorCustomNodeBuilder_LayersBrushStack> SharedThis = AsShared();
-	for (ALandscapeBlueprintCustomBrush* Brush : Brushes)
+	for (ALandscapeBlueprintBrushBase* Brush : Brushes)
 	{
 		FUIAction AddAction = FUIAction(FExecuteAction::CreateLambda([SharedThis, Brush]()
 		{
@@ -261,7 +261,7 @@ void FLandscapeEditorCustomNodeBuilder_LayersBrushStack::FillAddBrushMenu(FMenuB
 
 FReply FLandscapeEditorCustomNodeBuilder_LayersBrushStack::OnToggleAffectsHeightmap(int32 InBrushIndex)
 {
-	if (ALandscapeBlueprintCustomBrush* Brush = GetBrush(InBrushIndex))
+	if (ALandscapeBlueprintBrushBase* Brush = GetBrush(InBrushIndex))
 	{
 		const FScopedTransaction Transaction(LOCTEXT("Landscape_Brush_AffectsHeightmap", "Set Brush Affects Heightmap"));
 		bool bAffectsHeightmap = Brush->IsAffectingHeightmap();
@@ -272,7 +272,7 @@ FReply FLandscapeEditorCustomNodeBuilder_LayersBrushStack::OnToggleAffectsHeight
 
 FReply FLandscapeEditorCustomNodeBuilder_LayersBrushStack::OnToggleAffectsWeightmap(int32 InBrushIndex)
 {
-	if (ALandscapeBlueprintCustomBrush* Brush = GetBrush(InBrushIndex))
+	if (ALandscapeBlueprintBrushBase* Brush = GetBrush(InBrushIndex))
 	{
 		const FScopedTransaction Transaction(LOCTEXT("Landscape_Brush_AffectsWeightmap", "Set Brush Affects Weightmap"));
 		bool bAffectsWeightmap = Brush->IsAffectingWeightmap();
@@ -283,7 +283,7 @@ FReply FLandscapeEditorCustomNodeBuilder_LayersBrushStack::OnToggleAffectsWeight
 
 FReply FLandscapeEditorCustomNodeBuilder_LayersBrushStack::OnToggleVisibility(int32 InBrushIndex)
 {
-	if (ALandscapeBlueprintCustomBrush* Brush = GetBrush(InBrushIndex))
+	if (ALandscapeBlueprintBrushBase* Brush = GetBrush(InBrushIndex))
 	{
 		const FScopedTransaction Transaction(LOCTEXT("Landscape_Brush_SetVisibility", "Set Brush Visibility"));
 		bool bVisible = Brush->IsVisible();
@@ -294,33 +294,33 @@ FReply FLandscapeEditorCustomNodeBuilder_LayersBrushStack::OnToggleVisibility(in
 
 const FSlateBrush* FLandscapeEditorCustomNodeBuilder_LayersBrushStack::GetAffectsWeightmapBrush(int32 InBrushIndex) const
 {
-	ALandscapeBlueprintCustomBrush* Brush = GetBrush(InBrushIndex);
+	ALandscapeBlueprintBrushBase* Brush = GetBrush(InBrushIndex);
 	return Brush && Brush->IsAffectingWeightmap() ? FEditorStyle::GetBrush("LandscapeEditor.Brush.AffectsWeight.Enabled") : FEditorStyle::GetBrush("LandscapeEditor.Brush.AffectsWeight.Disabled");
 }
 
 const FSlateBrush* FLandscapeEditorCustomNodeBuilder_LayersBrushStack::GetAffectsHeightmapBrush(int32 InBrushIndex) const
 {
-	ALandscapeBlueprintCustomBrush* Brush = GetBrush(InBrushIndex);
+	ALandscapeBlueprintBrushBase* Brush = GetBrush(InBrushIndex);
 	return Brush && Brush->IsAffectingHeightmap() ? FEditorStyle::GetBrush("LandscapeEditor.Brush.AffectsHeight.Enabled") : FEditorStyle::GetBrush("LandscapeEditor.Brush.AffectsHeight.Disabled");
 }
 
 const FSlateBrush* FLandscapeEditorCustomNodeBuilder_LayersBrushStack::GetVisibilityBrush(int32 InBrushIndex) const
 {
-	ALandscapeBlueprintCustomBrush* Brush = GetBrush(InBrushIndex);
+	ALandscapeBlueprintBrushBase* Brush = GetBrush(InBrushIndex);
 	bool bIsVisible = Brush && Brush->IsVisible();
 	return bIsVisible ? FEditorStyle::GetBrush("Level.VisibleIcon16x") : FEditorStyle::GetBrush("Level.NotVisibleIcon16x");
 }
 
 bool FLandscapeEditorCustomNodeBuilder_LayersBrushStack::IsBrushEnabled(int32 InBrushIndex) const
 {
-	ALandscapeBlueprintCustomBrush* Brush = GetBrush(InBrushIndex);
+	ALandscapeBlueprintBrushBase* Brush = GetBrush(InBrushIndex);
 	const FEdModeLandscape* LandscapeEdMode = GetEditorMode();
 	return Brush && ((Brush->IsAffectingHeightmap() && LandscapeEdMode->CurrentToolTarget.TargetType == ELandscapeToolTargetType::Type::Heightmap) || (Brush->IsAffectingWeightmap() && LandscapeEdMode->CurrentToolTarget.TargetType == ELandscapeToolTargetType::Type::Weightmap));
 }
 
 bool FLandscapeEditorCustomNodeBuilder_LayersBrushStack::IsBrushSelected(int32 InBrushIndex) const
 {
-	ALandscapeBlueprintCustomBrush* Brush = GetBrush(InBrushIndex);
+	ALandscapeBlueprintBrushBase* Brush = GetBrush(InBrushIndex);
 	return Brush && Brush->IsSelected();
 }
 
@@ -335,8 +335,8 @@ void FLandscapeEditorCustomNodeBuilder_LayersBrushStack::OnBrushSelectionChanged
 	const FName CurrentToolName = LandscapeEdMode->CurrentTool->GetToolName();
 	if (CurrentToolName == TEXT("BlueprintBrush"))
 	{
-		ALandscapeBlueprintCustomBrush* Brush = GetBrush(InBrushIndex);
-		if (Brush != nullptr && !Brush->IsCommited())
+		ALandscapeBlueprintBrushBase* Brush = GetBrush(InBrushIndex);
+		if (Brush != nullptr)
 		{
 			FScopedTransaction Transaction(LOCTEXT("LandscapeBrushSelect", "Brush selection"));
 			GEditor->SelectNone(true, true);
@@ -347,7 +347,7 @@ void FLandscapeEditorCustomNodeBuilder_LayersBrushStack::OnBrushSelectionChanged
 
 FText FLandscapeEditorCustomNodeBuilder_LayersBrushStack::GetBrushText(int32 InBrushIndex) const
 {
-	if (ALandscapeBlueprintCustomBrush* Brush = GetBrush(InBrushIndex))
+	if (ALandscapeBlueprintBrushBase* Brush = GetBrush(InBrushIndex))
 	{
 		return FText::FromString(Brush->GetActorLabel());
 	}
@@ -357,15 +357,15 @@ FText FLandscapeEditorCustomNodeBuilder_LayersBrushStack::GetBrushText(int32 InB
 
 FSlateColor FLandscapeEditorCustomNodeBuilder_LayersBrushStack::GetBrushTextColor(int32 InBrushIndex) const
 {
-	if (ALandscapeBlueprintCustomBrush* Brush = GetBrush(InBrushIndex))
+	if (ALandscapeBlueprintBrushBase* Brush = GetBrush(InBrushIndex))
 	{
-		return Brush->IsCommited() ? FSlateColor::UseSubduedForeground() : FSlateColor::UseForeground();
+		return FSlateColor::UseForeground();
 	}
 
 	return FSlateColor::UseSubduedForeground();
 }
 
-ALandscapeBlueprintCustomBrush* FLandscapeEditorCustomNodeBuilder_LayersBrushStack::GetBrush(int32 InBrushIndex) const
+ALandscapeBlueprintBrushBase* FLandscapeEditorCustomNodeBuilder_LayersBrushStack::GetBrush(int32 InBrushIndex) const
 {
 	if (FEdModeLandscape* LandscapeEdMode = GetEditorMode())
 	{
@@ -382,7 +382,7 @@ FReply FLandscapeEditorCustomNodeBuilder_LayersBrushStack::HandleDragDetected(co
 		FLandscapeLayer* Layer = LandscapeEdMode->GetCurrentLayer();
 		if (Layer && !Layer->bLocked)
 		{
-			const TArray<ALandscapeBlueprintCustomBrush*>& BrushStack = LandscapeEdMode->GetBrushesForCurrentLayer();
+			const TArray<ALandscapeBlueprintBrushBase*>& BrushStack = LandscapeEdMode->GetBrushesForCurrentLayer();
 			if (BrushStack.IsValidIndex(SlotIndex))
 			{
 				TSharedPtr<SWidget> Row = GenerateRow(SlotIndex);
diff --git a/Engine/Source/Editor/LandscapeEditor/Private/LandscapeEditorDetailCustomization_LayersBrushStack.h b/Engine/Source/Editor/LandscapeEditor/Private/LandscapeEditorDetailCustomization_LayersBrushStack.h
--- a/Engine/Source/Editor/LandscapeEditor/Private/LandscapeEditorDetailCustomization_LayersBrushStack.h
+++ b/Engine/Source/Editor/LandscapeEditor/Private/LandscapeEditorDetailCustomization_LayersBrushStack.h
@@ -59,7 +59,7 @@ protected:
 	TSharedPtr<SWidget> GenerateRow(int32 InBrushIndex);
 	TSharedPtr<SWidget> OnBrushContextMenuOpening(int32 InBrushIndex);
 
-	void FillAddBrushMenu(FMenuBuilder& MenuBuilder, TArray<ALandscapeBlueprintCustomBrush*> Brushes);
+	void FillAddBrushMenu(FMenuBuilder& MenuBuilder, TArray<ALandscapeBlueprintBrushBase*> Brushes);
 
 	FReply OnToggleVisibility(int32 InBrushIndex);
 	FReply OnToggleAffectsHeightmap(int32 InBrushIndex);
@@ -73,7 +73,7 @@ protected:
 	void OnBrushSelectionChanged(int32 InBrushIndex);
 	FText GetBrushText(int32 InBrushIndex) const;
 	FSlateColor GetBrushTextColor(int32 InBrushIndex) const;
-	ALandscapeBlueprintCustomBrush* GetBrush(int32 InBrushIndex) const;
+	ALandscapeBlueprintBrushBase* GetBrush(int32 InBrushIndex) const;
 
 	// Drag/Drop handling
 	FReply HandleDragDetected(const FGeometry& MyGeometry, const FPointerEvent& MouseEvent, int32 SlotIndex, SVerticalBox::FSlot* Slot);
diff --git a/Engine/Source/Editor/LandscapeEditor/Public/LandscapeBlueprintBrush.h b/Engine/Source/Editor/LandscapeEditor/Public/LandscapeBlueprintBrush.h
new file mode 100644
--- /dev/null
+++ b/Engine/Source/Editor/LandscapeEditor/Public/LandscapeBlueprintBrush.h
@@ -0,0 +1,17 @@
+// Copyright 1998-2019 Epic Games, Inc. All Rights Reserved.
+
+#pragma once
+
+#include "CoreMinimal.h"
+#include "UObject/ObjectMacros.h"
+#include "GameFramework/Actor.h"
+#include "LandscapeBlueprintBrushBase.h"
+
+#include "LandscapeBlueprintBrush.generated.h"
+
+UCLASS(Abstract, Blueprintable, hidecategories = (Replication, Input, LOD, Actor, Cooking, Rendering))
+class ALandscapeBlueprintBrush : public ALandscapeBlueprintBrushBase
+{
+	GENERATED_UCLASS_BODY()
+};
+
diff --git a/Engine/Source/Editor/LandscapeEditor/Public/LandscapeEditorObject.h b/Engine/Source/Editor/LandscapeEditor/Public/LandscapeEditorObject.h
--- a/Engine/Source/Editor/LandscapeEditor/Public/LandscapeEditorObject.h
+++ b/Engine/Source/Editor/LandscapeEditor/Public/LandscapeEditorObject.h
@@ -8,7 +8,7 @@
 #include "LandscapeProxy.h"
 #include "Editor/LandscapeEditor/Private/LandscapeEdMode.h"
 #include "LandscapeFileFormatInterface.h"
-#include "LandscapeBPCustomBrush.h"
+#include "LandscapeBlueprintBrush.h"
 
 #include "LandscapeEditorObject.generated.h"
 
@@ -454,7 +454,7 @@ class ULandscapeEditorObject : public UObject
 	// Blueprint Brush Tool
 
 	UPROPERTY(Category = "Tool Settings", EditAnywhere, Transient, meta = (DisplayName = "Blueprint Brush", ShowForTools = "BlueprintBrush"))
-	TSubclassOf<ALandscapeBlueprintCustomBrush> BlueprintCustomBrush;
+	TSubclassOf<ALandscapeBlueprintBrush> BlueprintBrush;
 
 	// Resize Landscape Tool
 
diff --git a/Engine/Source/Runtime/Landscape/Classes/Landscape.h b/Engine/Source/Runtime/Landscape/Classes/Landscape.h
--- a/Engine/Source/Runtime/Landscape/Classes/Landscape.h
+++ b/Engine/Source/Runtime/Landscape/Classes/Landscape.h
@@ -5,7 +5,7 @@
 #include "CoreMinimal.h"
 #include "UObject/ObjectMacros.h"
 #include "LandscapeProxy.h"
-#include "LandscapeBPCustomBrush.h"
+#include "LandscapeBlueprintBrushBase.h"
 #include "Delegates/DelegateCombinations.h"
 
 #include "Landscape.generated.h"
@@ -89,28 +89,38 @@ struct FLandscapeLayerBrush
 	GENERATED_USTRUCT_BODY()
 
 	FLandscapeLayerBrush()
-		: BPCustomBrush(nullptr)
-#if WITH_EDITOR
+#if WITH_EDITORONLY_DATA
+		: BlueprintBrush(nullptr)
 		, LandscapeSize(MAX_int32, MAX_int32)
 		, LandscapeRenderTargetSize(MAX_int32, MAX_int32)
 #endif
 	{}
 
-	FLandscapeLayerBrush(ALandscapeBlueprintCustomBrush* InBrush)
-		: BPCustomBrush(InBrush)
+	FLandscapeLayerBrush(ALandscapeBlueprintBrushBase* InBlueprintBrush)
+#if WITH_EDITORONLY_DATA
+		: BlueprintBrush(InBlueprintBrush)
+		, LandscapeSize(MAX_int32, MAX_int32)
+		, LandscapeRenderTargetSize(MAX_int32, MAX_int32)
+#endif
 	{}
 
 #if WITH_EDITOR
-	UTextureRenderTarget2D* Render(bool InIsHeightmap, const FIntPoint& InLandscapeSize, UTextureRenderTarget2D* InLandscapeRenderTarget);
+	UTextureRenderTarget2D* Render(bool InIsHeightmap, const FIntPoint& InLandscapeSize, UTextureRenderTarget2D* InLandscapeRenderTarget, const FName& InWeightmapLayerName = NAME_None);
+	ALandscapeBlueprintBrushBase* GetBrush() const;
+	bool IsAffectingHeightmap() const;
+	bool IsAffectingWeightmapLayer(const FName& InWeightmapLayerName) const;
+	void SetOwner(ALandscape* InOwner);
 #endif
 
-	UPROPERTY()
-	ALandscapeBlueprintCustomBrush* BPCustomBrush;
-
 private:
 
 #if WITH_EDITOR
 	bool Initialize(const FIntPoint& LandscapeSize, UTextureRenderTarget2D* InLandscapeRenderTarget);
+#endif
+
+#if WITH_EDITORONLY_DATA
+	UPROPERTY()
+	ALandscapeBlueprintBrushBase* BlueprintBrush;
 
 	FTransform LandscapeTransform;
 	FIntPoint LandscapeSize;
@@ -259,7 +269,7 @@ public:
 	LANDSCAPE_API void RequestLayersContentUpdateForceAll(ELandscapeLayerUpdateMode InModeMask = ELandscapeLayerUpdateMode::Update_All);
 	LANDSCAPE_API void RequestLayersContentUpdate(ELandscapeLayerUpdateMode InModeMask);
 	LANDSCAPE_API bool ReorderLayer(int32 InStartingLayerIndex, int32 InDestinationLayerIndex);
-	LANDSCAPE_API FLandscapeLayer* DuplicateLayer(const FLandscapeLayer& InOtherLayer);
+	LANDSCAPE_API FLandscapeLayer* DuplicateLayerAndMoveBrushes(const FLandscapeLayer& InOtherLayer);
 	LANDSCAPE_API void CreateLayer(FName InName = NAME_None);
 	LANDSCAPE_API void CreateDefaultLayer();
 	LANDSCAPE_API void CopyOldDataToDefaultLayer();
@@ -302,13 +312,13 @@ public:
 	LANDSCAPE_API bool IsLayerBlendSubstractive(int32 InLayerIndex, const TWeakObjectPtr<ULandscapeLayerInfoObject>& InLayerInfoObj) const;
 	LANDSCAPE_API void SetLayerSubstractiveBlendStatus(int32 InLayerIndex, bool InStatus, const TWeakObjectPtr<ULandscapeLayerInfoObject>& InLayerInfoObj);
 
-	LANDSCAPE_API void AddBrushToLayer(int32 InLayerIndex, class ALandscapeBlueprintCustomBrush* InBrush);
-	LANDSCAPE_API void RemoveBrush(class ALandscapeBlueprintCustomBrush* InBrush);
-	LANDSCAPE_API void RemoveBrushFromLayer(int32 InLayerIndex, class ALandscapeBlueprintCustomBrush* InBrush);
+	LANDSCAPE_API void AddBrushToLayer(int32 InLayerIndex, class ALandscapeBlueprintBrushBase* InBrush);
+	LANDSCAPE_API void RemoveBrush(class ALandscapeBlueprintBrushBase* InBrush);
+	LANDSCAPE_API void RemoveBrushFromLayer(int32 InLayerIndex, class ALandscapeBlueprintBrushBase* InBrush);
 	LANDSCAPE_API bool ReorderLayerBrush(int32 InLayerIndex, int32 InStartingLayerBrushIndex, int32 InDestinationLayerBrushIndex);
-	LANDSCAPE_API class ALandscapeBlueprintCustomBrush* GetBrushForLayer(int32 InLayerIndex, int8 BrushIndex) const;
-	LANDSCAPE_API TArray<class ALandscapeBlueprintCustomBrush*> GetBrushesForLayer(int32 InLayerIndex) const;
-	LANDSCAPE_API void OnBPCustomBrushChanged();
+	LANDSCAPE_API class ALandscapeBlueprintBrushBase* GetBrushForLayer(int32 InLayerIndex, int8 BrushIndex) const;
+	LANDSCAPE_API TArray<class ALandscapeBlueprintBrushBase*> GetBrushesForLayer(int32 InLayerIndex) const;
+	LANDSCAPE_API void OnBlueprintBrushChanged();
 	
 	LANDSCAPE_API void OnPreSave();
 
@@ -390,8 +400,8 @@ public:
 	UPROPERTY(EditAnywhere, Category=Landscape)
 	bool bCanHaveLayersContent = false;
 
-	DECLARE_EVENT(ALandscape, FLandscapeBPCustomBrushChangedDelegate);
-	FLandscapeBPCustomBrushChangedDelegate& OnBPCustomBrushChangedDelegate() { return LandscapeBPCustomBrushChangedDelegate; }
+	DECLARE_EVENT(ALandscape, FLandscapeBlueprintBrushChangedDelegate);
+	FLandscapeBlueprintBrushChangedDelegate& OnBlueprintBrushChangedDelegate() { return LandscapeBlueprintBrushChangedDelegate; }
 
 	/** Target Landscape Layer for Landscape Splines */
 	UPROPERTY()
@@ -413,7 +423,7 @@ public:
 	TArray<UTextureRenderTarget2D*> WeightmapRTList;
 
 private:
-	FLandscapeBPCustomBrushChangedDelegate LandscapeBPCustomBrushChangedDelegate;
+	FLandscapeBlueprintBrushChangedDelegate LandscapeBlueprintBrushChangedDelegate;
 
 	/** Components affected by landscape splines (used to partially clear Layer Reserved for Splines) */
 	UPROPERTY(Transient)
diff --git a/Engine/Source/Runtime/Landscape/Classes/LandscapeBPCustomBrush.h b/Engine/Source/Runtime/Landscape/Classes/LandscapeBPCustomBrush.h
deleted file mode 100644
diff --git a/Engine/Source/Runtime/Landscape/Classes/LandscapeProxy.h b/Engine/Source/Runtime/Landscape/Classes/LandscapeProxy.h
--- a/Engine/Source/Runtime/Landscape/Classes/LandscapeProxy.h
+++ b/Engine/Source/Runtime/Landscape/Classes/LandscapeProxy.h
@@ -12,7 +12,6 @@
 #include "Async/AsyncWork.h"
 #include "Engine/Texture.h"
 #include "PerPlatformProperties.h"
-#include "LandscapeBPCustomBrush.h"
 #include "LandscapeComponent.h"
 #include "LandscapeWeightmapUsage.h"
 #include "VT/RuntimeVirtualTextureEnum.h"
diff --git a/Engine/Source/Runtime/Landscape/Private/LandscapeBPCustomBrush.cpp b/Engine/Source/Runtime/Landscape/Private/LandscapeBPCustomBrush.cpp
deleted file mode 100644
diff --git a/Engine/Source/Runtime/Landscape/Private/LandscapeBlueprintBrushBase.cpp b/Engine/Source/Runtime/Landscape/Private/LandscapeBlueprintBrushBase.cpp
new file mode 100644
--- /dev/null
+++ b/Engine/Source/Runtime/Landscape/Private/LandscapeBlueprintBrushBase.cpp
@@ -0,0 +1,160 @@
+#include "LandscapeBlueprintBrushBase.h"
+#include "CoreMinimal.h"
+#include "LandscapeProxy.h"
+#include "Landscape.h"
+
+#define LOCTEXT_NAMESPACE "Landscape"
+
+ALandscapeBlueprintBrushBase::ALandscapeBlueprintBrushBase(const FObjectInitializer& ObjectInitializer)
+#if WITH_EDITORONLY_DATA
+	: OwningLandscape(nullptr)
+	, AffectHeightmap(false)
+	, AffectWeightmap(false)
+	, bIsVisible(true)
+#endif
+{
+#if WITH_EDITOR
+	USceneComponent* SceneComp = CreateDefaultSubobject<USceneComponent>(TEXT("RootComponent"));
+	RootComponent = SceneComp;
+
+	PrimaryActorTick.bCanEverTick = true;
+	PrimaryActorTick.TickGroup = TG_DuringPhysics;
+	PrimaryActorTick.bStartWithTickEnabled = true;
+	PrimaryActorTick.SetTickFunctionEnable(true);
+	bIsEditorOnlyActor = true;
+#endif // WITH_EDITOR
+}
+
+void ALandscapeBlueprintBrushBase::RequestLandscapeUpdate()
+{
+#if WITH_EDITORONLY_DATA
+	if (OwningLandscape)
+	{
+		OwningLandscape->RequestLayersContentUpdateForceAll();
+	}
+#endif
+}
+
+#if WITH_EDITOR
+void ALandscapeBlueprintBrushBase::Tick(float DeltaSeconds)
+{
+	// Forward the Tick to the instances class of this BP
+	if (GetClass()->HasAnyClassFlags(CLASS_CompiledFromBlueprint))
+	{
+		TGuardValue<bool> AutoRestore(GAllowActorScriptExecutionInEditor, true);
+		ReceiveTick(DeltaSeconds);
+	}
+
+	Super::Tick(DeltaSeconds);
+}
+
+bool ALandscapeBlueprintBrushBase::ShouldTickIfViewportsOnly() const
+{
+	return true;
+}
+
+void ALandscapeBlueprintBrushBase::SetIsVisible(bool bInIsVisible)
+{
+#if WITH_EDITORONLY_DATA
+	Modify();
+	bIsVisible = bInIsVisible;
+	if (OwningLandscape)
+	{
+		OwningLandscape->OnBlueprintBrushChanged();
+	}
+#endif
+}
+
+void ALandscapeBlueprintBrushBase::SetAffectsHeightmap(bool bInAffectsHeightmap)
+{
+#if WITH_EDITORONLY_DATA
+	Modify();
+	AffectHeightmap = bInAffectsHeightmap;
+	if (OwningLandscape)
+	{
+		OwningLandscape->OnBlueprintBrushChanged();
+	}
+#endif
+}
+
+void ALandscapeBlueprintBrushBase::SetAffectsWeightmap(bool bInAffectsWeightmap)
+{
+#if WITH_EDITORONLY_DATA
+	Modify();
+	AffectWeightmap = bInAffectsWeightmap;
+	if (OwningLandscape)
+	{
+		OwningLandscape->OnBlueprintBrushChanged();
+	}
+#endif
+}
+
+bool ALandscapeBlueprintBrushBase::IsAffectingWeightmapLayer(const FName& InLayerName) const
+{
+#if WITH_EDITORONLY_DATA
+	return AffectedWeightmapLayers.Contains(InLayerName);
+#else
+	return false;
+#endif
+}
+
+void ALandscapeBlueprintBrushBase::PostEditMove(bool bFinished)
+{
+	Super::PostEditMove(bFinished);
+#if WITH_EDITORONLY_DATA
+	if (OwningLandscape)
+	{
+		if (bFinished)
+		{
+			OwningLandscape->RequestLayersContentUpdateForceAll();
+		}
+		else
+		{
+			OwningLandscape->RequestLayersContentUpdate(ELandscapeLayerUpdateMode::Update_All_Editing);
+		}
+	}
+#endif
+}
+
+void ALandscapeBlueprintBrushBase::PostEditChangeProperty(FPropertyChangedEvent& PropertyChangedEvent)
+{
+	Super::PostEditChangeProperty(PropertyChangedEvent);
+#if WITH_EDITORONLY_DATA
+	if (OwningLandscape)
+	{
+		OwningLandscape->OnBlueprintBrushChanged();
+	}
+#endif
+}
+
+void ALandscapeBlueprintBrushBase::Destroyed()
+{
+	Super::Destroyed();
+#if WITH_EDITORONLY_DATA
+	if (OwningLandscape && !GIsReinstancing)
+	{
+		OwningLandscape->RemoveBrush(this);
+	}
+	OwningLandscape = nullptr;
+#endif
+}
+
+void ALandscapeBlueprintBrushBase::SetOwningLandscape(ALandscape* InOwningLandscape)
+{
+#if WITH_EDITORONLY_DATA
+	Modify();
+	OwningLandscape = InOwningLandscape;
+#endif
+}
+
+ALandscape* ALandscapeBlueprintBrushBase::GetOwningLandscape() const
+{
+#if WITH_EDITORONLY_DATA
+	return OwningLandscape;
+#else
+	return nullptr;
+#endif
+}
+#endif // WITH_EDITOR
+
+#undef LOCTEXT_NAMESPACE
\ No newline at end of file
diff --git a/Engine/Source/Runtime/Landscape/Private/LandscapeEdit.cpp b/Engine/Source/Runtime/Landscape/Private/LandscapeEdit.cpp
--- a/Engine/Source/Runtime/Landscape/Private/LandscapeEdit.cpp
+++ b/Engine/Source/Runtime/Landscape/Private/LandscapeEdit.cpp
@@ -2803,7 +2803,7 @@ LANDSCAPE_API void ALandscapeProxy::Import(const FGuid& InGuid, int32 InMinX, in
 		{
 			for (const FLandscapeLayer& OldLayer : *InImportLayers)
 			{
-				FLandscapeLayer* NewLayer = LandscapeActor->DuplicateLayer(OldLayer);
+				FLandscapeLayer* NewLayer = LandscapeActor->DuplicateLayerAndMoveBrushes(OldLayer);
 				check(NewLayer != nullptr);
 
 				FLayerImportSettings ImportSettings;
diff --git a/Engine/Source/Runtime/Landscape/Private/LandscapeEditLayers.cpp b/Engine/Source/Runtime/Landscape/Private/LandscapeEditLayers.cpp
--- a/Engine/Source/Runtime/Landscape/Private/LandscapeEditLayers.cpp
+++ b/Engine/Source/Runtime/Landscape/Private/LandscapeEditLayers.cpp
@@ -23,7 +23,7 @@ LandscapeEditLayers.cpp: Landscape editing layers mode
 #include "LandscapeEditorModule.h"
 #include "LandscapeToolInterface.h"
 #include "ComponentRecreateRenderStateContext.h"
-#include "LandscapeBPCustomBrush.h"
+#include "LandscapeBlueprintBrushBase.h"
 #include "Materials/MaterialInstanceConstant.h"
 #include "LandscapeMaterialInstanceConstant.h"
 #include "Materials/MaterialExpressionLandscapeVisibilityMask.h"
@@ -2753,12 +2753,6 @@ int32 ALandscape::RegenerateLayersHeightmaps(const TArray<ULandscapeComponent*>&
 					// TODO: handle conversion/handling of RT not same size as internal size
 
 					FLandscapeLayerBrush& Brush = Layer.Brushes[i];
-
-					if (Brush.BPCustomBrush == nullptr || !Brush.BPCustomBrush->IsAffectingHeightmap() || !Brush.BPCustomBrush->IsVisible())
-					{
-						continue;
-					}
-
 					UTextureRenderTarget2D* BrushOutputNonAtlasRT = Brush.Render(true, LandscapeSize, CombinedHeightmapNonAtlasRT);
 					if (BrushOutputNonAtlasRT == nullptr || BrushOutputNonAtlasRT->SizeX != CombinedHeightmapNonAtlasRT->SizeX || BrushOutputNonAtlasRT->SizeY != CombinedHeightmapNonAtlasRT->SizeY)
 					{
@@ -2767,7 +2761,7 @@ int32 ALandscape::RegenerateLayersHeightmaps(const TArray<ULandscapeComponent*>&
 
 					INC_DWORD_STAT(STAT_LandscapeLayersRegenerateDrawCalls); // Brush Render
 
-					PrintLayersDebugRT(OutputDebugName ? FString::Printf(TEXT("LS Height: %s %s -> BrushNonAtlas %s"), *Layer.Name.ToString(), *Brush.BPCustomBrush->GetName(), *BrushOutputNonAtlasRT->GetName()) : TEXT(""), BrushOutputNonAtlasRT);
+					PrintLayersDebugRT(OutputDebugName ? FString::Printf(TEXT("LS Height: %s %s -> BrushNonAtlas %s"), *Layer.Name.ToString(), *Brush.GetBrush()->GetName(), *BrushOutputNonAtlasRT->GetName()) : TEXT(""), BrushOutputNonAtlasRT);
 
 					// Resolve back to Combined heightmap
 					CopyLayersTexture(BrushOutputNonAtlasRT, CombinedHeightmapNonAtlasRT);
@@ -3489,13 +3483,7 @@ int32 ALandscape::RegenerateLayersWeightmaps(const TArray<ULandscapeComponent*>&
 					for (int32 i = 0; i < Layer.Brushes.Num(); ++i)
 					{
 						FLandscapeLayerBrush& Brush = Layer.Brushes[i];
-
-						if (Brush.BPCustomBrush == nullptr || !Brush.BPCustomBrush->IsAffectingWeightmap() || !Brush.BPCustomBrush->IsVisible())
-						{
-							continue;
-						}
-
-						if (Brush.BPCustomBrush->IsAffectingWeightmapLayer(InfoLayerSettings.GetLayerName()) && !LayerInfoObjects.Contains(InfoLayerSettings.LayerInfoObj))
+						if (Brush.IsAffectingWeightmapLayer(InfoLayerSettings.GetLayerName()) && !LayerInfoObjects.Contains(InfoLayerSettings.LayerInfoObj))
 						{
 							LayerInfoObjects.Add(InfoLayerSettings.LayerInfoObj, LayerInfoSettingsIndex + 1); // due to visibility layer that is at 0
 						}
@@ -3609,23 +3597,17 @@ int32 ALandscape::RegenerateLayersWeightmaps(const TArray<ULandscapeComponent*>&
 							// TODO: handle conversion/handling of RT not same size as internal size
 
 							FLandscapeLayerBrush& Brush = Layer.Brushes[i];
-
-							if (Brush.BPCustomBrush == nullptr || !Brush.BPCustomBrush->IsAffectingWeightmap() || !Brush.BPCustomBrush->IsAffectingWeightmapLayer(LayerInfoObj->LayerName) || !Brush.BPCustomBrush->IsVisible())
+							UTextureRenderTarget2D* BrushOutputRT = Brush.Render(false, LandscapeSize, LandscapeScratchRT3, LayerInfoObj->LayerName);
+							if (BrushOutputRT == nullptr || BrushOutputRT->SizeX != LandscapeScratchRT3->SizeX || BrushOutputRT->SizeY != LandscapeScratchRT3->SizeY)
 							{
 								continue;
 							}
 
 							BrushRequiredAllocations.AddUnique(LayerInfoObj);
 
-							UTextureRenderTarget2D* BrushOutputRT = Brush.Render(false, LandscapeSize, LandscapeScratchRT3);
-							if (BrushOutputRT == nullptr || BrushOutputRT->SizeX != LandscapeScratchRT3->SizeX || BrushOutputRT->SizeY != LandscapeScratchRT3->SizeY)
-							{
-								continue;
-							}
-
 							INC_DWORD_STAT(STAT_LandscapeLayersRegenerateDrawCalls); // Brush Render
 
-							PrintLayersDebugRT(OutputDebugName ? FString::Printf(TEXT("LS Weight: %s %s -> Brush %s"), *Layer.Name.ToString(), *Brush.BPCustomBrush->GetName(), *BrushOutputRT->GetName()) : TEXT(""), BrushOutputRT);
+							PrintLayersDebugRT(OutputDebugName ? FString::Printf(TEXT("LS Weight: %s %s -> Brush %s"), *Layer.Name.ToString(), *Brush.GetBrush()->GetName(), *BrushOutputRT->GetName()) : TEXT(""), BrushOutputRT);
 
 							CopyLayersTexture(OutputDebugName ? FString::Printf(TEXT("Weight: %s PaintLayer: %s Brush: %s"), *Layer.Name.ToString(), *LayerInfoObj->LayerName.ToString(), *BrushOutputRT->GetName()) : TEXT(""), BrushOutputRT->GameThread_GetRenderTargetResource(), *LandscapeScratchRT3->GetName(), LandscapeScratchRT3->GameThread_GetRenderTargetResource());
 							PrintLayersDebugRT(OutputDebugName ? FString::Printf(TEXT("LS Weight: %s Component %s += -> Combined %s"), *Layer.Name.ToString(), *BrushOutputRT->GetName(), *LandscapeScratchRT3->GetName()) : TEXT(""), LandscapeScratchRT3);
@@ -5291,7 +5273,7 @@ void ALandscape::SetEditingLayer(const FGuid& InLayerGuid)
 
 void ALandscape::SetGrassUpdateEnabled(bool bInGrassUpdateEnabled)
 {
-#ifdef WITH_EDITORONLY_DATA
+#if WITH_EDITORONLY_DATA
 	bGrassUpdateEnabled = bInGrassUpdateEnabled;
 #endif
 }
@@ -5343,7 +5325,7 @@ void ALandscape::CreateDefaultLayer()
 	RequestLayersInitialization();
 }
 
-FLandscapeLayer* ALandscape::DuplicateLayer(const FLandscapeLayer& InOtherLayer)
+FLandscapeLayer* ALandscape::DuplicateLayerAndMoveBrushes(const FLandscapeLayer& InOtherLayer)
 {
 	ULandscapeInfo* LandscapeInfo = GetLandscapeInfo();
 	if (!LandscapeInfo || !CanHaveLayersContent())
@@ -5356,18 +5338,10 @@ FLandscapeLayer* ALandscape::DuplicateLayer(const FLandscapeLayer& InOtherLayer)
 	FLandscapeLayer NewLayer(InOtherLayer);
 	NewLayer.Guid = FGuid::NewGuid();
 
-	// Copy Brush and reparent to the new landscape level if required
+	// Update owning landscape and reparent to landscape's level if necessary
 	for (FLandscapeLayerBrush& Brush : NewLayer.Brushes)
 	{
-		if (Brush.BPCustomBrush != nullptr)
-		{
-			if (Brush.BPCustomBrush->GetTypedOuter<ULevel>() != GetTypedOuter<ULevel>())
-			{
-				Brush.BPCustomBrush = DuplicateObject<ALandscapeBlueprintCustomBrush>(Brush.BPCustomBrush, GetTypedOuter<ULevel>());
-			}
-
-			Brush.BPCustomBrush->SetOwningLandscape(this);
-		}
+		Brush.SetOwner(this);
 	}
 
 	int32 AddedIndex = LandscapeLayers.Add(NewLayer);
@@ -5513,7 +5487,7 @@ bool ALandscape::ReorderLayerBrush(int32 InLayerIndex, int32 InStartingLayerBrus
 	return false;
 }
 
-void ALandscape::AddBrushToLayer(int32 InLayerIndex, ALandscapeBlueprintCustomBrush* InBrush)
+void ALandscape::AddBrushToLayer(int32 InLayerIndex, ALandscapeBlueprintBrushBase* InBrush)
 {
 	if (FLandscapeLayer* Layer = GetLayer(InLayerIndex))
 	{
@@ -5524,7 +5498,7 @@ void ALandscape::AddBrushToLayer(int32 InLayerIndex, ALandscapeBlueprintCustomBr
 	}
 }
 
-void ALandscape::RemoveBrush(ALandscapeBlueprintCustomBrush* InBrush)
+void ALandscape::RemoveBrush(ALandscapeBlueprintBrushBase* InBrush)
 {
 	for (int32 LayerIndex = 0; LayerIndex < LandscapeLayers.Num(); ++LayerIndex)
 	{
@@ -5532,13 +5506,13 @@ void ALandscape::RemoveBrush(ALandscapeBlueprintCustomBrush* InBrush)
 	}
 }
 
-void ALandscape::RemoveBrushFromLayer(int32 InLayerIndex, ALandscapeBlueprintCustomBrush* InBrush)
+void ALandscape::RemoveBrushFromLayer(int32 InLayerIndex, ALandscapeBlueprintBrushBase* InBrush)
 {
 	if (FLandscapeLayer* Layer = GetLayer(InLayerIndex))
 	{
 		for (int32 i = 0; i < Layer->Brushes.Num(); ++i)
 		{
-			if (Layer->Brushes[i].BPCustomBrush == InBrush)
+			if (Layer->Brushes[i].GetBrush() == InBrush)
 			{
 				Modify();
 				Layer->Brushes.RemoveAt(i);
@@ -5550,55 +5524,104 @@ void ALandscape::RemoveBrushFromLayer(int32 InLayerIndex, ALandscapeBlueprintCus
 	}
 }
 
-void ALandscape::OnBPCustomBrushChanged()
+void ALandscape::OnBlueprintBrushChanged()
 {
 #if WITH_EDITORONLY_DATA
-	LandscapeBPCustomBrushChangedDelegate.Broadcast();
+	LandscapeBlueprintBrushChangedDelegate.Broadcast();
 	RequestLayersContentUpdateForceAll();
 #endif
 }
 
-ALandscapeBlueprintCustomBrush* ALandscape::GetBrushForLayer(int32 InLayerIndex, int8 InBrushIndex) const
+ALandscapeBlueprintBrushBase* ALandscape::GetBrushForLayer(int32 InLayerIndex, int8 InBrushIndex) const
 {
 	if (const FLandscapeLayer* Layer = GetLayer(InLayerIndex))
 	{
 		if (Layer->Brushes.IsValidIndex(InBrushIndex))
 		{
-			return Layer->Brushes[InBrushIndex].BPCustomBrush;
+			return Layer->Brushes[InBrushIndex].GetBrush();
 		}
 	}
 	return nullptr;
 }
 
-TArray<ALandscapeBlueprintCustomBrush*> ALandscape::GetBrushesForLayer(int32 InLayerIndex) const
+TArray<ALandscapeBlueprintBrushBase*> ALandscape::GetBrushesForLayer(int32 InLayerIndex) const
 {
-	TArray<ALandscapeBlueprintCustomBrush*> Brushes;
+	TArray<ALandscapeBlueprintBrushBase*> Brushes;
 	if (const FLandscapeLayer* Layer = GetLayer(InLayerIndex))
 	{
 		Brushes.Reserve(Layer->Brushes.Num());
 		for (const FLandscapeLayerBrush& Brush : Layer->Brushes)
 		{
-			Brushes.Add(Brush.BPCustomBrush);
+			Brushes.Add(Brush.GetBrush());
 		}
 	}
 	return Brushes;
 }
 
-UTextureRenderTarget2D* FLandscapeLayerBrush::Render(bool InIsHeightmap, const FIntPoint& InLandscapeSize, UTextureRenderTarget2D* InLandscapeRenderTarget)
+ALandscapeBlueprintBrushBase* FLandscapeLayerBrush::GetBrush() const
+{
+#if WITH_EDITORONLY_DATA
+	return BlueprintBrush;
+#else
+	return nullptr;
+#endif
+}
+
+void FLandscapeLayerBrush::SetOwner(ALandscape* InOwner)
 {
+#if WITH_EDITORONLY_DATA
+	if (BlueprintBrush && InOwner)
+	{
+		if (BlueprintBrush->GetTypedOuter<ULevel>() != InOwner->GetTypedOuter<ULevel>())
+		{
+			BlueprintBrush->Rename(nullptr, InOwner->GetTypedOuter<ULevel>());
+		}
+		BlueprintBrush->SetOwningLandscape(InOwner);
+	}
+#endif
+}
+
+bool FLandscapeLayerBrush::IsAffectingHeightmap() const
+{
+#if WITH_EDITORONLY_DATA
+	return BlueprintBrush && BlueprintBrush->IsVisible() && BlueprintBrush->IsAffectingHeightmap();
+#else
+	return false;
+#endif
+}
+
+bool FLandscapeLayerBrush::IsAffectingWeightmapLayer(const FName& InWeightmapLayerName) const
+{
+#if WITH_EDITORONLY_DATA
+	return BlueprintBrush && BlueprintBrush->IsVisible() && BlueprintBrush->IsAffectingWeightmap() && BlueprintBrush->IsAffectingWeightmapLayer(InWeightmapLayerName);
+#else
+	return false;
+#endif
+}
+
+UTextureRenderTarget2D* FLandscapeLayerBrush::Render(bool InIsHeightmap, const FIntPoint& InLandscapeSize, UTextureRenderTarget2D* InLandscapeRenderTarget, const FName& InWeightmapLayerName)
+{
+#if WITH_EDITORONLY_DATA
+	if ((InIsHeightmap && !IsAffectingHeightmap()) ||
+		(!InIsHeightmap && !IsAffectingWeightmapLayer(InWeightmapLayerName)))
+	{
+		return nullptr;
+	}
 	if (Initialize(InLandscapeSize, InLandscapeRenderTarget))
 	{
 		TGuardValue<bool> AutoRestore(GAllowActorScriptExecutionInEditor, true);
-		return BPCustomBrush->Render(InIsHeightmap, InLandscapeRenderTarget);
+		return BlueprintBrush->Render(InIsHeightmap, InLandscapeRenderTarget);
 	}
+#endif
 	return nullptr;
 }
 
 bool FLandscapeLayerBrush::Initialize(const FIntPoint& InLandscapeSize, UTextureRenderTarget2D* InLandscapeRenderTarget)
 {
-	if (BPCustomBrush && InLandscapeRenderTarget)
+#if WITH_EDITORONLY_DATA
+	if (BlueprintBrush && InLandscapeRenderTarget)
 	{
-		if (ALandscape* Landscape = BPCustomBrush->GetOwningLandscape())
+		if (ALandscape* Landscape = BlueprintBrush->GetOwningLandscape())
 		{
 			const FIntPoint NewLandscapeRenderTargetSize = FIntPoint(InLandscapeRenderTarget->SizeX, InLandscapeRenderTarget->SizeY);
 			const FTransform NewLandscapeTransform = Landscape->GetTransform();
@@ -5609,11 +5632,12 @@ bool FLandscapeLayerBrush::Initialize(const FIntPoint& InLandscapeSize, UTexture
 				LandscapeSize = InLandscapeSize;
 				
 				TGuardValue<bool> AutoRestore(GAllowActorScriptExecutionInEditor, true);
-				BPCustomBrush->Initialize(LandscapeTransform, LandscapeSize, LandscapeRenderTargetSize);
+				BlueprintBrush->Initialize(LandscapeTransform, LandscapeSize, LandscapeRenderTargetSize);
 			}
 			return true;
 		}
 	}
+#endif
 	return false;
 }
 
diff --git a/Engine/Source/Runtime/Landscape/Public/LandscapeBlueprintBrushBase.h b/Engine/Source/Runtime/Landscape/Public/LandscapeBlueprintBrushBase.h
new file mode 100644
--- /dev/null
+++ b/Engine/Source/Runtime/Landscape/Public/LandscapeBlueprintBrushBase.h
@@ -0,0 +1,63 @@
+// Copyright 1998-2019 Epic Games, Inc. All Rights Reserved.
+
+#pragma once
+
+#include "CoreMinimal.h"
+#include "UObject/ObjectMacros.h"
+#include "GameFramework/Actor.h"
+
+#include "LandscapeBlueprintBrushBase.generated.h"
+
+UCLASS(Abstract, NotBlueprintable)
+class LANDSCAPE_API ALandscapeBlueprintBrushBase : public AActor
+{
+	GENERATED_UCLASS_BODY()
+
+protected:
+#if WITH_EDITORONLY_DATA
+	UPROPERTY(DuplicateTransient)
+	class ALandscape* OwningLandscape;
+
+	UPROPERTY(Category = "Settings", EditAnywhere)
+	bool AffectHeightmap;
+
+	UPROPERTY(Category = "Settings", EditAnywhere)
+	bool AffectWeightmap;
+
+	UPROPERTY(Category = "Settings", EditAnywhere)
+	TArray<FName> AffectedWeightmapLayers;
+
+	UPROPERTY(Transient)
+	bool bIsVisible;
+#endif
+
+public:
+	UFUNCTION(BlueprintImplementableEvent)
+	UTextureRenderTarget2D* Render(bool InIsHeightmap, UTextureRenderTarget2D* InCombinedResult);
+
+	UFUNCTION(BlueprintImplementableEvent)
+	void Initialize(const FTransform& InLandscapeTransform, const FIntPoint& InLandscapeSize, const FIntPoint& InLandscapeRenderTargetSize);
+
+	UFUNCTION(BlueprintCallable, Category = "Landscape")
+	void RequestLandscapeUpdate();
+
+#if WITH_EDITOR
+	void SetOwningLandscape(class ALandscape* InOwningLandscape);
+	class ALandscape* GetOwningLandscape() const;
+
+	bool IsAffectingHeightmap() const { return AffectHeightmap; }
+	bool IsAffectingWeightmap() const { return AffectWeightmap; }
+	bool IsAffectingWeightmapLayer(const FName& InLayerName) const;
+	bool IsVisible() const { return bIsVisible; }
+	
+	void SetIsVisible(bool bInIsVisible);
+	void SetAffectsHeightmap(bool bInAffectsHeightmap);
+	void SetAffectsWeightmap(bool bInAffectsWeightmap);
+
+	virtual bool ShouldTickIfViewportsOnly() const override;
+	virtual void Tick(float DeltaSeconds) override;
+	virtual void PostEditMove(bool bFinished) override;
+	virtual void PostEditChangeProperty(FPropertyChangedEvent& PropertyChangedEvent) override;
+	virtual void Destroyed() override;
+#endif
+};
\ No newline at end of file
