Remove some debug string computation when we don't need them
Add deferred copy texture request for weightmap generation
Add missing GPU stat tag for all render thread command
When we have no changes, early exit so we don't generate FComponentRecreateRenderStateContext and FMaterialUpdateContext

diff --git a/Engine/Source/Runtime/Landscape/Private/LandscapeEditLayers.cpp b/Engine/Source/Runtime/Landscape/Private/LandscapeEditLayers.cpp
--- a/Engine/Source/Runtime/Landscape/Private/LandscapeEditLayers.cpp
+++ b/Engine/Source/Runtime/Landscape/Private/LandscapeEditLayers.cpp
@@ -38,6 +38,8 @@ LandscapeEditLayers.cpp: Landscape editing layers mode
 
 #define LOCTEXT_NAMESPACE "Landscape"
 
+static const FString GEmptyDebugName(TEXT(""));
+
 ENGINE_API extern bool GDisableAutomaticTextureMaterialUpdateDependencies;
 
 #if WITH_EDITOR
@@ -64,7 +66,6 @@ static TAutoConsoleVariable<int32> CVarLandscapeSimulatePhysics(
 #endif
 
 DECLARE_GPU_STAT_NAMED(LandscapeLayersRender, TEXT("Landscape Layer System Render"));
-DECLARE_GPU_STAT_NAMED(LandscapeLayersCopy, TEXT("Landscape Layer System Copy"));
 
 // Vertex format and vertex buffer
 
@@ -765,6 +766,7 @@ public:
 	{
 		SCOPE_CYCLE_COUNTER(STAT_LandscapeLayersRegenerate_RenderThread);
 		SCOPED_DRAW_EVENTF(InRHICmdList, LandscapeLayersRender, TEXT("ExtractLayers"));
+		SCOPED_GPU_STAT(InRHICmdList, LandscapeLayersRender);
 
 		TShaderMapRef<FLandscapeLayerWeightmapExtractMaterialLayersCS> ComputeShader(GetGlobalShaderMap(GMaxRHIFeatureLevel));
 		InRHICmdList.SetComputeShader(ComputeShader->GetComputeShader());
@@ -965,6 +967,7 @@ public:
 	{
 		SCOPE_CYCLE_COUNTER(STAT_LandscapeLayersRegenerate_RenderThread);
 		SCOPED_DRAW_EVENTF(InRHICmdList, LandscapeLayersRender, TEXT("PackLayers"));
+		SCOPED_GPU_STAT(InRHICmdList, LandscapeLayersRender);
 
 		TShaderMapRef<FLandscapeLayerWeightmapPackMaterialLayersCS> ComputeShader(GetGlobalShaderMap(GMaxRHIFeatureLevel));
 		InRHICmdList.SetComputeShader(ComputeShader->GetComputeShader());
@@ -997,7 +1000,7 @@ public:
 	{
 		SCOPE_CYCLE_COUNTER(STAT_LandscapeLayersRegenerate_RenderThread);
 		SCOPED_DRAW_EVENTF(InRHICmdList, LandscapeLayersCopy, TEXT("LS Copy %s -> %s, Mip (%d -> %d), Array Index (%d -> %d)"), *Params.SourceResourceDebugName, *Params.DestResourceDebugName, Params.SourceMip, Params.DestMip, Params.SourceArrayIndex, Params.DestArrayIndex);
-		SCOPED_GPU_STAT(InRHICmdList, LandscapeLayersCopy);
+		SCOPED_GPU_STAT(InRHICmdList, LandscapeLayersRender);
 
 		FIntPoint SourceSize(Params.SourceResource->GetSizeX(), Params.SourceResource->GetSizeY()); // SourceResource is always proper size, as it's always the good MIP we want to copy from
 		FIntPoint DestSize(Params.DestResource->GetSizeX() >> Params.DestMip, Params.DestResource->GetSizeY() >> Params.DestMip);
@@ -1763,7 +1766,10 @@ void ALandscape::AddDeferredCopyLayersTexture(UTexture* InSourceTexture, UTextur
 {
 	if (InSourceTexture != nullptr && InDestTexture != nullptr)
 	{
-		AddDeferredCopyLayersTexture(InSourceTexture->GetName(), InSourceTexture->Resource, InDestTexture->GetName(), InDestTexture->Resource, InDestCPUResource, InInitialPositionOffset, InSourceCurrentMip, InDestCurrentMip, InSourceArrayIndex, InDestArrayIndex);
+		FString SourceDebugName = CVarOutputLayersDebugDrawCallName.GetValueOnGameThread() == 1 ? InSourceTexture->GetName() : GEmptyDebugName;
+		FString DestDebugName = CVarOutputLayersDebugDrawCallName.GetValueOnGameThread() == 1 ? InDestTexture->GetName() : GEmptyDebugName;
+
+		AddDeferredCopyLayersTexture(SourceDebugName, InSourceTexture->Resource, DestDebugName, InDestTexture->Resource, InDestCPUResource, InInitialPositionOffset, InSourceCurrentMip, InDestCurrentMip, InSourceArrayIndex, InDestArrayIndex);
 	}
 }
 
@@ -1795,7 +1801,10 @@ void ALandscape::CopyLayersTexture(UTexture* InSourceTexture, UTexture* InDestTe
 {
 	if (InSourceTexture != nullptr && InDestTexture != nullptr)
 	{
-		CopyLayersTexture(InSourceTexture->GetName(), InSourceTexture->Resource, InDestTexture->GetName(), InDestTexture->Resource, InDestCPUResource, InFirstComponentSectionBase, InSourceCurrentMip, InDestCurrentMip, InSourceArrayIndex, InDestArrayIndex);
+		FString SourceDebugName = CVarOutputLayersDebugDrawCallName.GetValueOnGameThread() == 1 ? InSourceTexture->GetName() : GEmptyDebugName;
+		FString DestDebugName = CVarOutputLayersDebugDrawCallName.GetValueOnGameThread() == 1 ? InDestTexture->GetName() : GEmptyDebugName;
+
+		CopyLayersTexture(SourceDebugName, InSourceTexture->Resource, DestDebugName, InDestTexture->Resource, InDestCPUResource, InFirstComponentSectionBase, InSourceCurrentMip, InDestCurrentMip, InSourceArrayIndex, InDestArrayIndex);
 	}
 }
 
@@ -1967,7 +1976,7 @@ void ALandscape::DrawWeightmapComponentToRenderTargetMips(const TArray<FVector2D
 
 		if (WriteMipRT != nullptr)
 		{
-			DrawWeightmapComponentsToRenderTarget(OutputDebugName ? FString::Printf(TEXT("LS Weight: %s = -> %s Mips %d"), *ReadMipRT->GetName(), *WriteMipRT->GetName(), CurrentMip) : TEXT(""),
+			DrawWeightmapComponentsToRenderTarget(OutputDebugName ? FString::Printf(TEXT("LS Weight: %s = -> %s Mips %d"), *ReadMipRT->GetName(), *WriteMipRT->GetName(), CurrentMip) : GEmptyDebugName,
 													SectionBaseToDraw, WeightmapScaleBias, ReadMipRT, nullptr, WriteMipRT, InClearRTWrite, InShaderParams, CurrentMip++);
 		}
 
@@ -1998,7 +2007,7 @@ void ALandscape::DrawHeightmapComponentsToRenderTargetMips(const TArray<ULandsca
 
 		if (WriteMipRT != nullptr)
 		{
-			DrawHeightmapComponentsToRenderTarget(OutputDebugName ? FString::Printf(TEXT("LS Height: %s = -> %s CombinedAtlasWithMips %d"), *ReadMipRT->GetName(), *WriteMipRT->GetName(), CurrentMip) : TEXT(""),
+			DrawHeightmapComponentsToRenderTarget(OutputDebugName ? FString::Printf(TEXT("LS Height: %s = -> %s CombinedAtlasWithMips %d"), *ReadMipRT->GetName(), *WriteMipRT->GetName(), CurrentMip) : GEmptyDebugName,
 												  InComponentsToDraw, InLandscapeBase, ReadMipRT, nullptr, WriteMipRT, ERTDrawingType::RTMips, InClearRTWrite, InShaderParams, CurrentMip++);
 		}
 
@@ -2733,13 +2742,13 @@ int32 ALandscape::RegenerateLayersHeightmaps(const TArray<ULandscapeComponent*>&
 			CommitDeferredCopyLayersTexture();
 
 			// NOTE: From this point on, we always work in non atlas, we'll convert back at the end to atlas only
-			DrawHeightmapComponentsToRenderTarget(OutputDebugName ? FString::Printf(TEXT("LS Height: %s += -> NonAtlas %s"), *Layer.Name.ToString(), *LandscapeScratchRT1->GetName(), *LandscapeScratchRT2->GetName()) : TEXT(""),
+			DrawHeightmapComponentsToRenderTarget(OutputDebugName ? FString::Printf(TEXT("LS Height: %s += -> NonAtlas %s"), *Layer.Name.ToString(), *LandscapeScratchRT1->GetName(), *LandscapeScratchRT2->GetName()) : GEmptyDebugName,
 												  InLandscapeComponents, MinExtend, LandscapeScratchRT1, nullptr, LandscapeScratchRT2, ERTDrawingType::RTAtlasToNonAtlas, true, ShaderParams);
 
 			ShaderParams.ApplyLayerModifiers = true;
 
 			// Combine Current layer with current result
-			DrawHeightmapComponentsToRenderTarget(OutputDebugName ? FString::Printf(TEXT("LS Height: %s += -> CombinedNonAtlas %s"), *Layer.Name.ToString(), *LandscapeScratchRT2->GetName(), *CombinedHeightmapNonAtlasRT->GetName()) : TEXT(""),
+			DrawHeightmapComponentsToRenderTarget(OutputDebugName ? FString::Printf(TEXT("LS Height: %s += -> CombinedNonAtlas %s"), *Layer.Name.ToString(), *LandscapeScratchRT2->GetName(), *CombinedHeightmapNonAtlasRT->GetName()) : GEmptyDebugName,
 												InLandscapeComponents, MinExtend, LandscapeScratchRT2, FirstLayer ? nullptr : LandscapeScratchRT3, CombinedHeightmapNonAtlasRT, ERTDrawingType::RTNonAtlas, FirstLayer, ShaderParams);
 
 			ShaderParams.ApplyLayerModifiers = false;
@@ -2761,16 +2770,16 @@ int32 ALandscape::RegenerateLayersHeightmaps(const TArray<ULandscapeComponent*>&
 
 					INC_DWORD_STAT(STAT_LandscapeLayersRegenerateDrawCalls); // Brush Render
 
-					PrintLayersDebugRT(OutputDebugName ? FString::Printf(TEXT("LS Height: %s %s -> BrushNonAtlas %s"), *Layer.Name.ToString(), *Brush.GetBrush()->GetName(), *BrushOutputNonAtlasRT->GetName()) : TEXT(""), BrushOutputNonAtlasRT);
+					PrintLayersDebugRT(OutputDebugName ? FString::Printf(TEXT("LS Height: %s %s -> BrushNonAtlas %s"), *Layer.Name.ToString(), *Brush.GetBrush()->GetName(), *BrushOutputNonAtlasRT->GetName()) : GEmptyDebugName, BrushOutputNonAtlasRT);
 
 					// Resolve back to Combined heightmap
 					CopyLayersTexture(BrushOutputNonAtlasRT, CombinedHeightmapNonAtlasRT);
-					PrintLayersDebugRT(OutputDebugName ? FString::Printf(TEXT("LS Height: %s Component %s += -> CombinedNonAtlas %s"), *Layer.Name.ToString(), *BrushOutputNonAtlasRT->GetName(), *CombinedHeightmapNonAtlasRT->GetName()) : TEXT(""), CombinedHeightmapNonAtlasRT);
+					PrintLayersDebugRT(OutputDebugName ? FString::Printf(TEXT("LS Height: %s Component %s += -> CombinedNonAtlas %s"), *Layer.Name.ToString(), *BrushOutputNonAtlasRT->GetName(), *CombinedHeightmapNonAtlasRT->GetName()) : GEmptyDebugName, CombinedHeightmapNonAtlasRT);
 				}
 			}
 
 			CopyLayersTexture(CombinedHeightmapNonAtlasRT, LandscapeScratchRT3);
-			PrintLayersDebugRT(OutputDebugName ? FString::Printf(TEXT("LS Height: %s Component %s += -> CombinedNonAtlas %s"), *Layer.Name.ToString(), *CombinedHeightmapNonAtlasRT->GetName(), *LandscapeScratchRT3->GetName()) : TEXT(""), LandscapeScratchRT3);
+			PrintLayersDebugRT(OutputDebugName ? FString::Printf(TEXT("LS Height: %s Component %s += -> CombinedNonAtlas %s"), *Layer.Name.ToString(), *CombinedHeightmapNonAtlasRT->GetName(), *LandscapeScratchRT3->GetName()) : GEmptyDebugName, LandscapeScratchRT3);
 
 			FirstLayer = false;
 		}
@@ -2778,12 +2787,12 @@ int32 ALandscape::RegenerateLayersHeightmaps(const TArray<ULandscapeComponent*>&
 		ShaderParams.GenerateNormals = true;
 		ShaderParams.GridSize = GetRootComponent()->RelativeScale3D;
 
-		DrawHeightmapComponentsToRenderTarget(OutputDebugName ? FString::Printf(TEXT("LS Height: %s = -> CombinedNonAtlasNormals : %s"), *CombinedHeightmapNonAtlasRT->GetName(), *LandscapeScratchRT1->GetName()) : TEXT(""),
+		DrawHeightmapComponentsToRenderTarget(OutputDebugName ? FString::Printf(TEXT("LS Height: %s = -> CombinedNonAtlasNormals : %s"), *CombinedHeightmapNonAtlasRT->GetName(), *LandscapeScratchRT1->GetName()) : GEmptyDebugName,
 											  InLandscapeComponents, MinExtend, CombinedHeightmapNonAtlasRT, nullptr, LandscapeScratchRT1, ERTDrawingType::RTNonAtlas, true, ShaderParams);
 
 		ShaderParams.GenerateNormals = false;
 
-		DrawHeightmapComponentsToRenderTarget(OutputDebugName ? FString::Printf(TEXT("LS Height: %s = -> CombinedAtlasFinal : %s"), *LandscapeScratchRT1->GetName(), *CombinedHeightmapAtlasRT->GetName()) : TEXT(""),
+		DrawHeightmapComponentsToRenderTarget(OutputDebugName ? FString::Printf(TEXT("LS Height: %s = -> CombinedAtlasFinal : %s"), *LandscapeScratchRT1->GetName(), *CombinedHeightmapAtlasRT->GetName()) : GEmptyDebugName,
 											  InLandscapeComponents, MinExtend, LandscapeScratchRT1, nullptr, CombinedHeightmapAtlasRT, ERTDrawingType::RTNonAtlasToAtlas, true, ShaderParams);
 
 		DrawHeightmapComponentsToRenderTargetMips(InLandscapeComponents, MinExtend, CombinedHeightmapAtlasRT, true, ShaderParams);
@@ -2954,8 +2963,8 @@ void ALandscape::PrepareComponentDataToExtractMaterialLayersCS(const FLandscapeL
 
 					FIntPoint ComponentSectionBase = Component->GetSectionBase() - InLandscapeBase;
 
-					CopyLayersTexture(*LayerWeightmap->GetName(), LayerWeightmap->Resource, InOutputDebugName ? FString::Printf(TEXT("%s WeightmapScratchTexture"), *InLayer.Name.ToString()) : TEXT(""), InOutTextureData, nullptr, ComponentSectionBase, 0);
-					PrintLayersDebugTextureResource(InOutputDebugName ? FString::Printf(TEXT("LS Weight: %s WeightmapScratchTexture %s"), *InLayer.Name.ToString(), TEXT("WeightmapScratchTextureResource")) : TEXT(""), InOutTextureData, 0, false);
+					AddDeferredCopyLayersTexture(InOutputDebugName ? *LayerWeightmap->GetName() : GEmptyDebugName, LayerWeightmap->Resource, InOutputDebugName ? FString::Printf(TEXT("%s WeightmapScratchTexture"), *InLayer.Name.ToString()) : GEmptyDebugName, InOutTextureData, nullptr, ComponentSectionBase, 0);
+					PrintLayersDebugTextureResource(InOutputDebugName ? FString::Printf(TEXT("LS Weight: %s WeightmapScratchTexture %s"), *InLayer.Name.ToString(), TEXT("WeightmapScratchTextureResource")) : GEmptyDebugName, InOutTextureData, 0, false);
 
 					for (const FWeightmapLayerAllocationInfo& WeightmapLayerAllocation : ComponentLayerData->WeightmapData.LayerAllocations)
 					{
@@ -3008,6 +3017,8 @@ void ALandscape::PrepareComponentDataToExtractMaterialLayersCS(const FLandscapeL
 			}
 		}
 	});
+
+	CommitDeferredCopyLayersTexture();
 }
 
 void ALandscape::PrepareComponentDataToPackMaterialLayersCS(int32 InCurrentWeightmapToProcessIndex, const FIntPoint& InLandscapeBase, bool InOutputDebugName, const TArray<ULandscapeComponent*>& InAllLandscapeComponents, TArray<UTexture2D*>& OutProcessedWeightmaps,
@@ -3453,15 +3464,21 @@ int32 ALandscape::RegenerateLayersWeightmaps(const TArray<ULandscapeComponent*>&
 		UTextureRenderTarget2D* EmptyRT = WeightmapRTList[EWeightmapRTType::WeightmapRT_Scratch_RGBA];
 		FLandscapeLayersWeightmapShaderParameters PSShaderParams;
 		bool OutputDebugName = (CVarOutputLayersDebugDrawCallName.GetValueOnAnyThread() == 1 || CVarOutputLayersRTContent.GetValueOnAnyThread() == 1 || CVarOutputLayersWeightmapsRTContent.GetValueOnAnyThread() == 1) ? true : false;
-
+		FString SourceDebugName = GEmptyDebugName;
+		FString DestDebugName = GEmptyDebugName;
 		ClearLayersWeightmapTextureResource(TEXT("ClearRT RGBA"), EmptyRT->GameThread_GetRenderTargetResource());
 		ClearLayersWeightmapTextureResource(TEXT("ClearRT R"), LandscapeScratchRT1->GameThread_GetRenderTargetResource());
 
 		for (int32 LayerIndex = 0; LayerIndex < LayerCount; ++LayerIndex)
 		{
-			CopyLayersTexture(*LandscapeScratchRT1->GetName(), LandscapeScratchRT1->GameThread_GetRenderTargetResource(), OutputDebugName ? FString::Printf(TEXT("Weight: Clear CombinedProcLayerWeightmapAllLayersResource %d, "), LayerIndex) : TEXT(""), CombinedLayersWeightmapAllMaterialLayersResource, nullptr, FIntPoint(0, 0), 0, 0, 0, LayerIndex);
+			SourceDebugName = OutputDebugName ? LandscapeScratchRT1->GetName() : GEmptyDebugName;
+			DestDebugName = OutputDebugName ? FString::Printf(TEXT("Weight: Clear CombinedProcLayerWeightmapAllLayersResource %d, "), LayerIndex) : GEmptyDebugName;
+
+			AddDeferredCopyLayersTexture(SourceDebugName, LandscapeScratchRT1->GameThread_GetRenderTargetResource(), DestDebugName, CombinedLayersWeightmapAllMaterialLayersResource, nullptr, FIntPoint(0, 0), 0, 0, 0, LayerIndex);
 		}
 
+		CommitDeferredCopyLayersTexture();
+
 		bool ComputeShaderGeneratedData = false;
 		bool FirstLayer = true;
 		TMap<ULandscapeLayerInfoObject*, bool> WeightmapLayersBlendSubstractive;
@@ -3494,7 +3511,10 @@ int32 ALandscape::RegenerateLayersWeightmaps(const TArray<ULandscapeComponent*>&
 			// Loop until there is no more weightmap texture to process
 			while (HasFoundWeightmapToProcess)
 			{
-				CopyLayersTexture(*EmptyRT->GetName(), EmptyRT->GameThread_GetRenderTargetResource(), OutputDebugName ? FString::Printf(TEXT("Weight: %s Clear WeightmapScratchExtractLayerTextureResource"), *Layer.Name.ToString()) : TEXT(""), WeightmapScratchExtractLayerTextureResource);
+				SourceDebugName = OutputDebugName ? EmptyRT->GetName() : GEmptyDebugName;
+				DestDebugName = OutputDebugName ? FString::Printf(TEXT("Weight: %s Clear WeightmapScratchExtractLayerTextureResource"), *Layer.Name.ToString()) : GEmptyDebugName;
+
+				CopyLayersTexture(SourceDebugName, EmptyRT->GameThread_GetRenderTargetResource(), DestDebugName, WeightmapScratchExtractLayerTextureResource);
 
 				// Prepare compute shader data
 				TArray<FLandscapeLayerWeightmapExtractMaterialLayersComponentData> ComponentsData;
@@ -3505,7 +3525,7 @@ int32 ALandscape::RegenerateLayersWeightmaps(const TArray<ULandscapeComponent*>&
 				// Perform the compute shader
 				if (ComponentsData.Num() > 0)
 				{
-					PrintLayersDebugTextureResource(OutputDebugName ? FString::Printf(TEXT("LS Weight: %s WeightmapScratchTexture %s"), *Layer.Name.ToString(), TEXT("WeightmapScratchTextureResource")) : TEXT(""), WeightmapScratchExtractLayerTextureResource, 0, false);
+					PrintLayersDebugTextureResource(OutputDebugName ? FString::Printf(TEXT("LS Weight: %s WeightmapScratchTexture %s"), *Layer.Name.ToString(), TEXT("WeightmapScratchTextureResource")) : GEmptyDebugName, WeightmapScratchExtractLayerTextureResource, 0, false);
 
 					// Clear the current atlas if required
 					if (CurrentWeightmapToProcessIndex == 0)
@@ -3516,8 +3536,14 @@ int32 ALandscape::RegenerateLayersWeightmaps(const TArray<ULandscapeComponent*>&
 						for (auto& ItPair : LayerInfoObjects)
 						{
 							int32 LayerIndex = ItPair.Value;
-							CopyLayersTexture(*LandscapeScratchRT1->GetName(), LandscapeScratchRT1->GameThread_GetRenderTargetResource(), OutputDebugName ? FString::Printf(TEXT("Weight: %s Clear CurrentProcLayerWeightmapAllLayersResource %d, "), *Layer.Name.ToString(), LayerIndex) : TEXT(""), CurrentLayersWeightmapAllMaterialLayersResource, nullptr, FIntPoint(0, 0), 0, 0, 0, LayerIndex);
+
+							SourceDebugName = OutputDebugName ? LandscapeScratchRT1->GetName() : GEmptyDebugName;
+							DestDebugName = OutputDebugName ? FString::Printf(TEXT("Weight: %s Clear CurrentProcLayerWeightmapAllLayersResource %d, "), *Layer.Name.ToString(), LayerIndex) : GEmptyDebugName;
+
+							AddDeferredCopyLayersTexture(SourceDebugName, LandscapeScratchRT1->GameThread_GetRenderTargetResource(), DestDebugName, CurrentLayersWeightmapAllMaterialLayersResource, nullptr, FIntPoint(0, 0), 0, 0, 0, LayerIndex);
 						}
+
+						CommitDeferredCopyLayersTexture();
 					}
 
 					FLandscapeLayerWeightmapExtractMaterialLayersComputeShaderParameters CSExtractLayersShaderParams;
@@ -3550,21 +3576,27 @@ int32 ALandscape::RegenerateLayersWeightmaps(const TArray<ULandscapeComponent*>&
 					ULandscapeLayerInfoObject* LayerInfoObj = LayerInfoObject.Key;
 
 					// Copy the layer we are working on
-					CopyLayersTexture(OutputDebugName ? FString::Printf(TEXT("Weight: %s PaintLayer: %s, CurrentProcLayerWeightmapAllLayersResource"), *Layer.Name.ToString(), *LayerInfoObj->LayerName.ToString()) : TEXT(""), CurrentLayersWeightmapAllMaterialLayersResource, *LandscapeScratchRT1->GetName(), LandscapeScratchRT1->GameThread_GetRenderTargetResource(), nullptr, FIntPoint(0, 0), 0, 0, LayerIndex, 0);
-					PrintLayersDebugRT(OutputDebugName ? FString::Printf(TEXT("LS Weight: %s CurrentProcLayerWeightmapAllLayersResource -> Paint Layer RT %s"), *Layer.Name.ToString(), *LandscapeScratchRT1->GetName()) : TEXT(""), LandscapeScratchRT1, 0, false);
+					SourceDebugName = OutputDebugName ? FString::Printf(TEXT("Weight: %s PaintLayer: %s, CurrentProcLayerWeightmapAllLayersResource"), *Layer.Name.ToString(), *LayerInfoObj->LayerName.ToString()) : GEmptyDebugName;
+					DestDebugName = OutputDebugName ? LandscapeScratchRT1->GetName() : GEmptyDebugName;
+
+					CopyLayersTexture(SourceDebugName, CurrentLayersWeightmapAllMaterialLayersResource, DestDebugName, LandscapeScratchRT1->GameThread_GetRenderTargetResource(), nullptr, FIntPoint(0, 0), 0, 0, LayerIndex, 0);
+					PrintLayersDebugRT(OutputDebugName ? FString::Printf(TEXT("LS Weight: %s CurrentProcLayerWeightmapAllLayersResource -> Paint Layer RT %s"), *Layer.Name.ToString(), *LandscapeScratchRT1->GetName()) : GEmptyDebugName, LandscapeScratchRT1, 0, false);
 
 					PSShaderParams.ApplyLayerModifiers = true;
 					PSShaderParams.LayerVisible = Layer.bVisible;
 					PSShaderParams.LayerAlpha = LayerInfoObj == ALandscapeProxy::VisibilityLayer ? 1.0f : Layer.WeightmapAlpha; // visibility can't be affected by weight
 
-					DrawWeightmapComponentsToRenderTarget(OutputDebugName ? FString::Printf(TEXT("LS Weight: %s Paint: %s += -> %s"), *Layer.Name.ToString(), *LayerInfoObj->LayerName.ToString(), *LandscapeScratchRT1->GetName(), *LandscapeScratchRT2->GetName()) : TEXT(""),
+					DrawWeightmapComponentsToRenderTarget(OutputDebugName ? FString::Printf(TEXT("LS Weight: %s Paint: %s += -> %s"), *Layer.Name.ToString(), *LayerInfoObj->LayerName.ToString(), *LandscapeScratchRT1->GetName(), *LandscapeScratchRT2->GetName()) : GEmptyDebugName,
 														  InLandscapeComponents, MinExtend, LandscapeScratchRT1, nullptr, LandscapeScratchRT2, true, PSShaderParams, 0);
 
 					PSShaderParams.ApplyLayerModifiers = false;
 
 					// Combined Layer data with current stack
-					CopyLayersTexture(OutputDebugName ? FString::Printf(TEXT("Weight: %s PaintLayer: %s CombinedProcLayerWeightmap"), *Layer.Name.ToString(), *LayerInfoObj->LayerName.ToString()) : TEXT(""), CombinedLayersWeightmapAllMaterialLayersResource, *LandscapeScratchRT1->GetName(), LandscapeScratchRT1->GameThread_GetRenderTargetResource(), nullptr, FIntPoint(0, 0), 0, 0, LayerIndex, 0);
-					PrintLayersDebugRT(OutputDebugName ? FString::Printf(TEXT("LS Weight: %s CombinedProcLayerWeightmap -> Paint Layer RT %s"), *Layer.Name.ToString(), *LandscapeScratchRT1->GetName()) : TEXT(""), LandscapeScratchRT1, 0, false);
+					SourceDebugName = OutputDebugName ? FString::Printf(TEXT("Weight: %s PaintLayer: %s CombinedProcLayerWeightmap"), *Layer.Name.ToString(), *LayerInfoObj->LayerName.ToString()) : GEmptyDebugName;
+					DestDebugName = OutputDebugName ? LandscapeScratchRT1->GetName() : GEmptyDebugName;
+
+					CopyLayersTexture(SourceDebugName, CombinedLayersWeightmapAllMaterialLayersResource, DestDebugName, LandscapeScratchRT1->GameThread_GetRenderTargetResource(), nullptr, FIntPoint(0, 0), 0, 0, LayerIndex, 0);
+					PrintLayersDebugRT(OutputDebugName ? FString::Printf(TEXT("LS Weight: %s CombinedProcLayerWeightmap -> Paint Layer RT %s"), *Layer.Name.ToString(), *LandscapeScratchRT1->GetName()) : GEmptyDebugName, LandscapeScratchRT1, 0, false);
 
 					// Combine with current status and copy back to the combined 2d resource array
 					PSShaderParams.OutputAsSubstractive = false;
@@ -3581,12 +3613,15 @@ int32 ALandscape::RegenerateLayersWeightmaps(const TArray<ULandscapeComponent*>&
 						}
 					}
 
-					DrawWeightmapComponentsToRenderTarget(OutputDebugName ? FString::Printf(TEXT("LS Weight: %s PaintLayer: %s, %s += -> Combined %s"), *Layer.Name.ToString(), *LayerInfoObj->LayerName.ToString(), *LandscapeScratchRT2->GetName(), *LandscapeScratchRT3->GetName()) : TEXT(""),
+					DrawWeightmapComponentsToRenderTarget(OutputDebugName ? FString::Printf(TEXT("LS Weight: %s PaintLayer: %s, %s += -> Combined %s"), *Layer.Name.ToString(), *LayerInfoObj->LayerName.ToString(), *LandscapeScratchRT2->GetName(), *LandscapeScratchRT3->GetName()) : GEmptyDebugName,
 														  InLandscapeComponents, MinExtend, LandscapeScratchRT2, FirstLayer ? nullptr : LandscapeScratchRT1, LandscapeScratchRT3, true, PSShaderParams, 0);
 
 					PSShaderParams.OutputAsSubstractive = false;
 
-					CopyLayersTexture(OutputDebugName ? FString::Printf(TEXT("Weight: %s PaintLayer: %s %s"), *Layer.Name.ToString(), *LayerInfoObj->LayerName.ToString(), *LandscapeScratchRT3->GetName()) : TEXT(""), LandscapeScratchRT3->GameThread_GetRenderTargetResource(), TEXT("CombinedProcLayerWeightmap"), CombinedLayersWeightmapAllMaterialLayersResource, nullptr, FIntPoint(0, 0), 0, 0, 0, LayerIndex);
+					SourceDebugName = OutputDebugName ? FString::Printf(TEXT("Weight: %s PaintLayer: %s %s"), *Layer.Name.ToString(), *LayerInfoObj->LayerName.ToString(), *LandscapeScratchRT3->GetName()) : GEmptyDebugName;
+					DestDebugName = OutputDebugName ? TEXT("CombinedProcLayerWeightmap") : GEmptyDebugName;
+
+					CopyLayersTexture(SourceDebugName, LandscapeScratchRT3->GameThread_GetRenderTargetResource(), DestDebugName, CombinedLayersWeightmapAllMaterialLayersResource, nullptr, FIntPoint(0, 0), 0, 0, 0, LayerIndex);
 
 					// Handle brush blending
 					if (Layer.bVisible)
@@ -3607,14 +3642,20 @@ int32 ALandscape::RegenerateLayersWeightmaps(const TArray<ULandscapeComponent*>&
 
 							INC_DWORD_STAT(STAT_LandscapeLayersRegenerateDrawCalls); // Brush Render
 
-							PrintLayersDebugRT(OutputDebugName ? FString::Printf(TEXT("LS Weight: %s %s -> Brush %s"), *Layer.Name.ToString(), *Brush.GetBrush()->GetName(), *BrushOutputRT->GetName()) : TEXT(""), BrushOutputRT);
+							PrintLayersDebugRT(OutputDebugName ? FString::Printf(TEXT("LS Weight: %s %s -> Brush %s"), *Layer.Name.ToString(), *Brush.GetBrush()->GetName(), *BrushOutputRT->GetName()) : GEmptyDebugName, BrushOutputRT);
 
-							CopyLayersTexture(OutputDebugName ? FString::Printf(TEXT("Weight: %s PaintLayer: %s Brush: %s"), *Layer.Name.ToString(), *LayerInfoObj->LayerName.ToString(), *BrushOutputRT->GetName()) : TEXT(""), BrushOutputRT->GameThread_GetRenderTargetResource(), *LandscapeScratchRT3->GetName(), LandscapeScratchRT3->GameThread_GetRenderTargetResource());
-							PrintLayersDebugRT(OutputDebugName ? FString::Printf(TEXT("LS Weight: %s Component %s += -> Combined %s"), *Layer.Name.ToString(), *BrushOutputRT->GetName(), *LandscapeScratchRT3->GetName()) : TEXT(""), LandscapeScratchRT3);
+							SourceDebugName = OutputDebugName ? FString::Printf(TEXT("Weight: %s PaintLayer: %s Brush: %s"), *Layer.Name.ToString(), *LayerInfoObj->LayerName.ToString(), *BrushOutputRT->GetName()) : GEmptyDebugName;
+							DestDebugName = OutputDebugName ? LandscapeScratchRT3->GetName() : GEmptyDebugName;
+
+							CopyLayersTexture(SourceDebugName, BrushOutputRT->GameThread_GetRenderTargetResource(), DestDebugName, LandscapeScratchRT3->GameThread_GetRenderTargetResource());
+							PrintLayersDebugRT(OutputDebugName ? FString::Printf(TEXT("LS Weight: %s Component %s += -> Combined %s"), *Layer.Name.ToString(), *BrushOutputRT->GetName(), *LandscapeScratchRT3->GetName()) : GEmptyDebugName, LandscapeScratchRT3);
 						}
 
-						PrintLayersDebugRT(OutputDebugName ? FString::Printf(TEXT("LS Weight: %s CombinedPostBrushProcLayerWeightmap -> Paint Layer RT %s"), *Layer.Name.ToString(), *LandscapeScratchRT3->GetName()) : TEXT(""), LandscapeScratchRT3, 0, false);
-						CopyLayersTexture(OutputDebugName ? FString::Printf(TEXT("Weight: %s PaintLayer: %s %s"), *Layer.Name.ToString(), *LayerInfoObj->LayerName.ToString(), *LandscapeScratchRT3->GetName()) : TEXT(""), LandscapeScratchRT3->GameThread_GetRenderTargetResource(), TEXT("CombinedProcLayerWeightmap"), CombinedLayersWeightmapAllMaterialLayersResource, nullptr, FIntPoint(0, 0), 0, 0, 0, LayerIndex);
+						PrintLayersDebugRT(OutputDebugName ? FString::Printf(TEXT("LS Weight: %s CombinedPostBrushProcLayerWeightmap -> Paint Layer RT %s"), *Layer.Name.ToString(), *LandscapeScratchRT3->GetName()) : GEmptyDebugName, LandscapeScratchRT3, 0, false);
+
+						SourceDebugName = OutputDebugName ? FString::Printf(TEXT("Weight: %s PaintLayer: %s %s"), *Layer.Name.ToString(), *LayerInfoObj->LayerName.ToString(), *LandscapeScratchRT3->GetName()) : GEmptyDebugName;
+						DestDebugName = OutputDebugName ? TEXT("CombinedProcLayerWeightmap") : GEmptyDebugName;
+						CopyLayersTexture(SourceDebugName, LandscapeScratchRT3->GameThread_GetRenderTargetResource(), DestDebugName, CombinedLayersWeightmapAllMaterialLayersResource, nullptr, FIntPoint(0, 0), 0, 0, 0, LayerIndex);
 					}
 				}
 
@@ -3712,7 +3753,10 @@ int32 ALandscape::RegenerateLayersWeightmaps(const TArray<ULandscapeComponent*>&
 					}
 
 					// Clear Pack texture
-					CopyLayersTexture(*EmptyRT->GetName(), EmptyRT->GameThread_GetRenderTargetResource(), TEXT("Weight: Clear WeightmapScratchPackLayerTextureResource"), WeightmapScratchPackLayerTextureResource);
+					SourceDebugName = OutputDebugName ? *EmptyRT->GetName() : GEmptyDebugName;
+					DestDebugName = OutputDebugName ? TEXT("Weight: Clear WeightmapScratchPackLayerTextureResource") : GEmptyDebugName;
+
+					CopyLayersTexture(SourceDebugName, EmptyRT->GameThread_GetRenderTargetResource(), DestDebugName, WeightmapScratchPackLayerTextureResource);
 
 					FLandscapeLayerWeightmapPackMaterialLayersComputeShaderParameters CSPackLayersShaderParams;
 					CSPackLayersShaderParams.AtlasWeightmapsPerLayer = CombinedLayersWeightmapAllMaterialLayersResource;
@@ -3730,7 +3774,10 @@ int32 ALandscape::RegenerateLayersWeightmaps(const TArray<ULandscapeComponent*>&
 					});
 
 					UTextureRenderTarget2D* CurrentRT = WeightmapRTList[WeightmapRT_Mip0];
-					CopyLayersTexture(TEXT("WeightmapScratchTexture"), WeightmapScratchPackLayerTextureResource, *CurrentRT->GetName(), CurrentRT->GameThread_GetRenderTargetResource());
+
+					SourceDebugName = OutputDebugName ? TEXT("WeightmapScratchTexture") : GEmptyDebugName;
+					DestDebugName = OutputDebugName ? CurrentRT->GetName() : GEmptyDebugName;
+					CopyLayersTexture(SourceDebugName, WeightmapScratchPackLayerTextureResource, DestDebugName, CurrentRT->GameThread_GetRenderTargetResource());
 					DrawWeightmapComponentToRenderTargetMips(WeightmapTextureOutputOffset, CurrentRT, true, PSShaderParams);
 
 					int32 StartTextureIndex = NextTextureIndexToProcess;
@@ -3749,11 +3796,16 @@ int32 ALandscape::RegenerateLayersWeightmaps(const TArray<ULandscapeComponent*>&
 
 							if (CurrentRT != nullptr)
 							{
-								CopyLayersTexture(*CurrentRT->GetName(), CurrentRT->GameThread_GetRenderTargetResource(), OutputDebugName ? FString::Printf(TEXT("Weightmap Mip: %d"), CurrentMip) : TEXT(""), WeightmapTexture->Resource, WeightmapCPUReadBack, TextureTopLeftPositionInAtlas, CurrentMip, CurrentMip);
+								SourceDebugName = OutputDebugName ? CurrentRT->GetName() : GEmptyDebugName;
+								DestDebugName = OutputDebugName ? FString::Printf(TEXT("Weightmap Mip: %d"), CurrentMip) : GEmptyDebugName;
+
+								AddDeferredCopyLayersTexture(SourceDebugName, CurrentRT->GameThread_GetRenderTargetResource(), DestDebugName, WeightmapTexture->Resource, WeightmapCPUReadBack, TextureTopLeftPositionInAtlas, CurrentMip, CurrentMip);
 								++CurrentMip;
 							}
 						}
 					}
+
+					CommitDeferredCopyLayersTexture();
 				}
 
 				++CurrentWeightmapToProcessIndex;
@@ -3824,6 +3876,11 @@ void ALandscape::UpdateLayersMaterialInstances()
 		});
 	}
 
+	if (ComponentsToUpdate.Num() == 0)
+	{
+		return;
+	}
+
 	// we're not having the material update context recreate render states because we will manually do it for only our components
 	TArray<FComponentRecreateRenderStateContext> RecreateRenderStateContexts;
 	RecreateRenderStateContexts.Reserve(ComponentsToUpdate.Num());
