Back out changelist 6930726
#rnx

diff --git a/Engine/Config/BaseEngine.ini b/Engine/Config/BaseEngine.ini
--- a/Engine/Config/BaseEngine.ini
+++ b/Engine/Config/BaseEngine.ini
@@ -669,9 +669,6 @@ bUseStreamingPause=false
 +ClassRedirects=(OldName="EditorAutomationActor",NewName="EditorUtilityActor")
 +ClassRedirects=(OldName="EditorAutomationActorComponent",NewName="EditorUtilityActorComponent")
 +ClassRedirects=(OldName="EditorAutomationObject",NewName="EditorUtilityObject")
-+ClassRedirects=(OldName="LandscapeBlueprintCustomBrush",NewName="/Script/LandscapeEditor.LandscapeBlueprintBrush")
-
-+PropertyRedirects=(OldName="LandscapeLayerBrush.BPCustomBrush",NewName="LandscapeLayerBrush.BlueprintBrush")
 
 +PropertyRedirects=(OldName="StructVariableDescription.bDontEditoOnInstance",NewName="bDontEditOnInstance")
 +PropertyRedirects=(OldName="KismetMathLibrary.DegAtan2.A",NewName="Y")
diff --git a/Engine/Source/Editor/LandscapeEditor/Private/LandscapeBlueprintBrush.cpp b/Engine/Source/Editor/LandscapeEditor/Private/LandscapeBlueprintBrush.cpp
deleted file mode 100644
diff --git a/Engine/Source/Editor/LandscapeEditor/Private/LandscapeEdMode.cpp b/Engine/Source/Editor/LandscapeEditor/Private/LandscapeEdMode.cpp
--- a/Engine/Source/Editor/LandscapeEditor/Private/LandscapeEdMode.cpp
+++ b/Engine/Source/Editor/LandscapeEditor/Private/LandscapeEdMode.cpp
@@ -52,7 +52,7 @@
 #include "Framework/Commands/InputBindingManager.h"
 #include "MouseDeltaTracker.h"
 #include "Interfaces/IMainFrameModule.h"
-#include "LandscapeBlueprintBrushBase.h"
+#include "LandscapeBPCustomBrush.h"
 #include "Engine/TextureRenderTarget2D.h"
 #include "Settings/EditorExperimentalSettings.h"
 #include "ComponentRecreateRenderStateContext.h"
@@ -592,7 +592,7 @@ void FEdModeLandscape::Enter()
 
 		if (ALandscape* Landscape = GetLandscape())
 		{
-			Landscape->OnBlueprintBrushChangedDelegate().AddLambda([this]() { this->RefreshDetailPanel(); });
+			Landscape->OnBPCustomBrushChangedDelegate().AddLambda([this]() { this->RefreshDetailPanel(); });
 			if (Landscape->HasLayersContent())
 			{
 				if (Landscape->GetLandscapeSplinesReservedLayer())
@@ -780,7 +780,7 @@ void FEdModeLandscape::Exit()
 		LandscapeProxy->OnMaterialChangedDelegate().RemoveAll(this);
 		if (ALandscape* Landscape = GetLandscape())
 		{
-			Landscape->OnBlueprintBrushChangedDelegate().RemoveAll(this);
+			Landscape->OnBPCustomBrushChangedDelegate().RemoveAll(this);
 		}
 	}
 
@@ -1622,10 +1622,10 @@ bool FEdModeLandscape::ProcessEditCopy()
 
 		if (!Result)
 		{
-			ALandscapeBlueprintBrushBase* CurrentlySelectedBPBrush = nullptr;
+			ALandscapeBlueprintCustomBrush* CurrentlySelectedBPBrush = nullptr;
 			for (FSelectionIterator It(GEditor->GetSelectedActorIterator()); It; ++It)
 			{
-				CurrentlySelectedBPBrush = Cast<ALandscapeBlueprintBrushBase>(*It);
+				CurrentlySelectedBPBrush = Cast<ALandscapeBlueprintCustomBrush>(*It);
 				if (CurrentlySelectedBPBrush)
 				{
 					break;
@@ -1691,10 +1691,10 @@ bool FEdModeLandscape::ProcessEditPaste()
 
 		if (!Result)
 		{
-			ALandscapeBlueprintBrushBase* CurrentlySelectedBPBrush = nullptr;
+			ALandscapeBlueprintCustomBrush* CurrentlySelectedBPBrush = nullptr;
 			for (FSelectionIterator It(GEditor->GetSelectedActorIterator()); It; ++It)
 			{
-				CurrentlySelectedBPBrush = Cast<ALandscapeBlueprintBrushBase>(*It);
+				CurrentlySelectedBPBrush = Cast<ALandscapeBlueprintCustomBrush>(*It);
 				if (CurrentlySelectedBPBrush)
 				{
 					break;
@@ -2394,7 +2394,7 @@ void FEdModeLandscape::SetCurrentBrush(int32 BrushIndex)
 	}
 }
 
-const TArray<ALandscapeBlueprintBrushBase*>& FEdModeLandscape::GetBrushList() const
+const TArray<ALandscapeBlueprintCustomBrush*>& FEdModeLandscape::GetBrushList() const
 {
 	return BrushList;
 }
@@ -2540,7 +2540,7 @@ void FEdModeLandscape::SetTargetLandscape(const TWeakObjectPtr<ULandscapeInfo>&
 		LandscapeProxy->OnMaterialChangedDelegate().RemoveAll(this);
 		if (ALandscape* Landscape = GetLandscape())
 		{
-			Landscape->OnBlueprintBrushChangedDelegate().RemoveAll(this);
+			Landscape->OnBPCustomBrushChangedDelegate().RemoveAll(this);
 		}
 	}
 
@@ -2560,7 +2560,7 @@ void FEdModeLandscape::SetTargetLandscape(const TWeakObjectPtr<ULandscapeInfo>&
 		LandscapeProxy->OnMaterialChangedDelegate().AddRaw(this, &FEdModeLandscape::OnLandscapeMaterialChangedDelegate);
 		if (ALandscape* Landscape = GetLandscape())
 		{
-			Landscape->OnBlueprintBrushChangedDelegate().AddLambda([this]() { this->RefreshDetailPanel(); });
+			Landscape->OnBPCustomBrushChangedDelegate().AddLambda([this]() { this->RefreshDetailPanel(); });
 		}
 	}
 
@@ -3426,7 +3426,7 @@ bool FEdModeLandscape::IsSelectionAllowed(AActor* InActor, bool bInSelection) co
 	{
 		return true;
 	}
-	else if (InActor->IsA(ALandscapeBlueprintBrushBase::StaticClass()))
+	else if (InActor->IsA(ALandscapeBlueprintCustomBrush::StaticClass()))
 	{
 		return true;
 	}
@@ -4478,7 +4478,7 @@ void FEdModeLandscape::RequestLayersContentUpdateForceAll(ELandscapeLayerUpdateM
 	}
 }
 
-void FEdModeLandscape::AddBrushToCurrentLayer(ALandscapeBlueprintBrushBase* InBrush)
+void FEdModeLandscape::AddBrushToCurrentLayer(ALandscapeBlueprintCustomBrush* InBrush)
 {
 	ALandscape* Landscape = GetLandscape();
 	if (Landscape == nullptr)
@@ -4490,7 +4490,7 @@ void FEdModeLandscape::AddBrushToCurrentLayer(ALandscapeBlueprintBrushBase* InBr
 	RefreshDetailPanel();
 }
 
-void FEdModeLandscape::RemoveBrushFromCurrentLayer(ALandscapeBlueprintBrushBase* InBrush)
+void FEdModeLandscape::RemoveBrushFromCurrentLayer(ALandscapeBlueprintCustomBrush* InBrush)
 {
 	ALandscape* Landscape = GetLandscape();
 
@@ -4503,7 +4503,7 @@ void FEdModeLandscape::RemoveBrushFromCurrentLayer(ALandscapeBlueprintBrushBase*
 	RefreshDetailPanel();
 }
 
-ALandscapeBlueprintBrushBase* FEdModeLandscape::GetBrushForCurrentLayer(int8 InBrushIndex) const
+ALandscapeBlueprintCustomBrush* FEdModeLandscape::GetBrushForCurrentLayer(int8 InBrushIndex) const
 {
 	if (ALandscape* Landscape = GetLandscape())
 	{
@@ -4512,9 +4512,9 @@ ALandscapeBlueprintBrushBase* FEdModeLandscape::GetBrushForCurrentLayer(int8 InB
 	return nullptr;
 }
 
-TArray<ALandscapeBlueprintBrushBase*> FEdModeLandscape::GetBrushesForCurrentLayer()
+TArray<ALandscapeBlueprintCustomBrush*> FEdModeLandscape::GetBrushesForCurrentLayer()
 {
-	TArray<ALandscapeBlueprintBrushBase*> Brushes;
+	TArray<ALandscapeBlueprintCustomBrush*> Brushes;
 	if (ALandscape* Landscape = GetLandscape())
 	{
 		Brushes = Landscape->GetBrushesForLayer(GetCurrentLayerIndex());
@@ -4714,9 +4714,9 @@ bool FEdModeLandscape::NeedToFillEmptyMaterialLayers() const
 void FEdModeLandscape::UpdateBrushList()
 {
 	BrushList.Empty();
-	for (TObjectIterator<ALandscapeBlueprintBrushBase> BrushIt(RF_Transient|RF_ClassDefaultObject|RF_ArchetypeObject, true, EInternalObjectFlags::PendingKill); BrushIt; ++BrushIt)
+	for (TObjectIterator<ALandscapeBlueprintCustomBrush> BrushIt(RF_Transient|RF_ClassDefaultObject|RF_ArchetypeObject, true, EInternalObjectFlags::PendingKill); BrushIt; ++BrushIt)
 	{
-		ALandscapeBlueprintBrushBase* Brush = *BrushIt;
+		ALandscapeBlueprintCustomBrush* Brush = *BrushIt;
 		if (Brush->GetTypedOuter<UPackage>() != GetTransientPackage())
 		{
 			BrushList.Add(Brush);
@@ -4732,7 +4732,7 @@ void FEdModeLandscape::OnLevelActorAdded(AActor* InActor)
 		Landscape->RegisterLandscapeEdMode(this);
 	}
 
-	ALandscapeBlueprintBrushBase* Brush = Cast<ALandscapeBlueprintBrushBase>(InActor);
+	ALandscapeBlueprintCustomBrush* Brush = Cast<ALandscapeBlueprintCustomBrush>(InActor);
 	if (Brush && Brush->GetTypedOuter<UPackage>() != GetTransientPackage())
 	{
 		if (!GIsReinstancing)
@@ -4751,7 +4751,7 @@ void FEdModeLandscape::OnLevelActorRemoved(AActor* InActor)
 		Landscape->UnregisterLandscapeEdMode();
 	}
 
-	ALandscapeBlueprintBrushBase* Brush = Cast<ALandscapeBlueprintBrushBase>(InActor);
+	ALandscapeBlueprintCustomBrush* Brush = Cast<ALandscapeBlueprintCustomBrush>(InActor);
 	if (Brush && Brush->GetTypedOuter<UPackage>() != GetTransientPackage())
 	{
 		UpdateBrushList();
diff --git a/Engine/Source/Editor/LandscapeEditor/Private/LandscapeEdMode.h b/Engine/Source/Editor/LandscapeEditor/Private/LandscapeEdMode.h
--- a/Engine/Source/Editor/LandscapeEditor/Private/LandscapeEdMode.h
+++ b/Engine/Source/Editor/LandscapeEditor/Private/LandscapeEdMode.h
@@ -510,7 +510,7 @@ public:
 	void SetCurrentBrush(int32 BrushIndex);
 
 	void UpdateBrushList();
-	const TArray<ALandscapeBlueprintBrushBase*>& GetBrushList() const;
+	const TArray<ALandscapeBlueprintCustomBrush*>& GetBrushList() const;
 
 	const TArray<TSharedRef<FLandscapeTargetListInfo>>& GetTargetList() const;
 	const TArray<FName>* GetTargetDisplayOrderList() const;
@@ -564,10 +564,10 @@ public:
 	void AutoUpdateDirtyLandscapeSplines();
 	bool CanEditLayer(FText* Reason = nullptr, FLandscapeLayer* InLayer = nullptr);
 
-	void AddBrushToCurrentLayer(class ALandscapeBlueprintBrushBase* InBrush);
-	void RemoveBrushFromCurrentLayer(class ALandscapeBlueprintBrushBase* InBrush);
-	class ALandscapeBlueprintBrushBase* GetBrushForCurrentLayer(int8 BrushIndex) const;
-	TArray<class ALandscapeBlueprintBrushBase*> GetBrushesForCurrentLayer();
+	void AddBrushToCurrentLayer(class ALandscapeBlueprintCustomBrush* InBrush);
+	void RemoveBrushFromCurrentLayer(class ALandscapeBlueprintCustomBrush* InBrush);
+	class ALandscapeBlueprintCustomBrush* GetBrushForCurrentLayer(int8 BrushIndex) const;
+	TArray<class ALandscapeBlueprintCustomBrush*> GetBrushesForCurrentLayer();
 	
 	bool NeedToFillEmptyMaterialLayers() const;
 	void RequestLayersContentUpdate(ELandscapeLayerUpdateMode InUpdateMode);
@@ -619,7 +619,7 @@ public:
 private:
 	TArray<TSharedRef<FLandscapeTargetListInfo>> LandscapeTargetList;
 	TArray<FLandscapeListInfo> LandscapeList;
-	TArray<ALandscapeBlueprintBrushBase*> BrushList;
+	TArray<ALandscapeBlueprintCustomBrush*> BrushList;
 	TArray<FName> ShownTargetLayerList;
 	
 	/** Represent the index offset of the target layer in LandscapeTargetList */
diff --git a/Engine/Source/Editor/LandscapeEditor/Private/LandscapeEdModeBPCustomTools.cpp b/Engine/Source/Editor/LandscapeEditor/Private/LandscapeEdModeBPCustomTools.cpp
--- a/Engine/Source/Editor/LandscapeEditor/Private/LandscapeEdModeBPCustomTools.cpp
+++ b/Engine/Source/Editor/LandscapeEditor/Private/LandscapeEdModeBPCustomTools.cpp
@@ -19,7 +19,7 @@
 #include "LandscapeRender.h"
 #include "LandscapeHeightfieldCollisionComponent.h"
 #include "LandscapeEdModeTools.h"
-#include "LandscapeBlueprintBrushBase.h"
+#include "LandscapeBPCustomBrush.h"
 #include "LandscapeInfo.h"
 #include "Landscape.h"
 //#include "LandscapeDataAccess.h"
@@ -72,12 +72,12 @@ public:
 
 	virtual bool BeginTool(FEditorViewportClient* ViewportClient, const FLandscapeToolTarget& Target, const FVector& InHitLocation) override
 	{
-		if (EdMode->UISettings->BlueprintBrush == nullptr)
+		if (EdMode->UISettings->BlueprintCustomBrush == nullptr)
 		{
 			return false;
 		}
 
-		ALandscapeBlueprintBrushBase* DefaultObject = Cast<ALandscapeBlueprintBrushBase>(EdMode->UISettings->BlueprintBrush->GetDefaultObject(false));
+		ALandscapeBlueprintCustomBrush* DefaultObject = Cast<ALandscapeBlueprintCustomBrush>(EdMode->UISettings->BlueprintCustomBrush->GetDefaultObject(false));
 
 		if (DefaultObject == nullptr)
 		{
@@ -97,9 +97,9 @@ public:
 			SpawnInfo.bNoFail = true;
 			SpawnInfo.OverrideLevel = Info->LandscapeActor.Get()->GetTypedOuter<ULevel>(); // always spawn in the same level as the one containing the ALandscape
 
-			FScopedTransaction Transaction(LOCTEXT("LandscapeEdModeBlueprintToolSpawn", "Create landscape brush"));
-			ALandscapeBlueprintBrushBase* Brush = ViewportClient->GetWorld()->SpawnActor<ALandscapeBlueprintBrushBase>(EdMode->UISettings->BlueprintBrush, SpawnLocation, FRotator(0.0f), SpawnInfo);
-			EdMode->UISettings->BlueprintBrush = nullptr;
+			FScopedTransaction Transaction(LOCTEXT("LandscapeEdModeBPCustomToolSpawn", "Create landscape brush"));
+			ALandscapeBlueprintCustomBrush* Brush = ViewportClient->GetWorld()->SpawnActor<ALandscapeBlueprintCustomBrush>(EdMode->UISettings->BlueprintCustomBrush, SpawnLocation, FRotator(0.0f), SpawnInfo);
+			EdMode->UISettings->BlueprintCustomBrush = nullptr;
 
 			GEditor->SelectNone(true, true);
 			GEditor->SelectActor(Brush, true, true);
diff --git a/Engine/Source/Editor/LandscapeEditor/Private/LandscapeEditorDetailCustomization_Layers.cpp b/Engine/Source/Editor/LandscapeEditor/Private/LandscapeEditorDetailCustomization_Layers.cpp
--- a/Engine/Source/Editor/LandscapeEditor/Private/LandscapeEditorDetailCustomization_Layers.cpp
+++ b/Engine/Source/Editor/LandscapeEditor/Private/LandscapeEditorDetailCustomization_Layers.cpp
@@ -348,9 +348,9 @@ void FLandscapeEditorCustomNodeBuilder_Layers::SetLayerName(const FText& InText,
 }
 
 
-void FLandscapeEditorCustomNodeBuilder_Layers::FillAddBrushMenu(FMenuBuilder& MenuBuilder, TArray<ALandscapeBlueprintBrushBase*> Brushes)
+void FLandscapeEditorCustomNodeBuilder_Layers::FillAddBrushMenu(FMenuBuilder& MenuBuilder, TArray<ALandscapeBlueprintCustomBrush*> Brushes)
 {
-	for (ALandscapeBlueprintBrushBase* Brush : Brushes)
+	for (ALandscapeBlueprintCustomBrush* Brush : Brushes)
 	{
 		TSharedRef<FLandscapeEditorCustomNodeBuilder_Layers> SharedThis = AsShared();
 		FUIAction AddAction = FUIAction(FExecuteAction::CreateLambda([SharedThis, Brush]() { SharedThis->AddBrushToCurrentLayer(Brush); }));
@@ -413,7 +413,7 @@ void FLandscapeEditorCustomNodeBuilder_Layers::FillClearLayerMenu(FMenuBuilder&
 	}
 }
 
-void FLandscapeEditorCustomNodeBuilder_Layers::AddBrushToCurrentLayer(ALandscapeBlueprintBrushBase* Brush)
+void FLandscapeEditorCustomNodeBuilder_Layers::AddBrushToCurrentLayer(ALandscapeBlueprintCustomBrush* Brush)
 {
 	const FScopedTransaction Transaction(LOCTEXT("LandscapeBrushAddToCurrentLayerTransaction", "Add brush to current layer"));
 	GetEditorMode()->AddBrushToCurrentLayer(Brush);
@@ -504,8 +504,8 @@ TSharedPtr<SWidget> FLandscapeEditorCustomNodeBuilder_Layers::OnLayerContextMenu
 		}
 		MenuBuilder.EndSection();
 
-		const TArray<ALandscapeBlueprintBrushBase*>& Brushes = LandscapeEdMode->GetBrushList();
-		TArray<ALandscapeBlueprintBrushBase*> FilteredBrushes = Brushes.FilterByPredicate([](ALandscapeBlueprintBrushBase* Brush) { return Brush->GetOwningLandscape() == nullptr; });
+		const TArray<ALandscapeBlueprintCustomBrush*>& Brushes = LandscapeEdMode->GetBrushList();
+		TArray<ALandscapeBlueprintCustomBrush*> FilteredBrushes = Brushes.FilterByPredicate([](ALandscapeBlueprintCustomBrush* Brush) { return Brush->GetOwningLandscape() == nullptr; });
 		if (FilteredBrushes.Num())
 		{
 			MenuBuilder.BeginSection("LandscapeEditorBrushActions", LOCTEXT("LandscapeEditorBrushActions.Heading", "Brushes"));
diff --git a/Engine/Source/Editor/LandscapeEditor/Private/LandscapeEditorDetailCustomization_Layers.h b/Engine/Source/Editor/LandscapeEditor/Private/LandscapeEditorDetailCustomization_Layers.h
--- a/Engine/Source/Editor/LandscapeEditor/Private/LandscapeEditorDetailCustomization_Layers.h
+++ b/Engine/Source/Editor/LandscapeEditor/Private/LandscapeEditorDetailCustomization_Layers.h
@@ -23,7 +23,7 @@ class FDetailWidgetRow;
 class IDetailChildrenBuilder;
 class IDetailLayoutBuilder;
 class SDragAndDropVerticalBox;
-class ALandscapeBlueprintBrushBase;
+class ALandscapeBlueprintCustomBrush;
 class ULandscapeLayerInfoObject;
 
 /**
@@ -100,10 +100,10 @@ protected:
 	FReply OnToggleLock(int32 InLayerIndex);
 	const FSlateBrush* GetLockBrushForLayer(int32 InLayerIndex) const;
 
-	void FillAddBrushMenu(FMenuBuilder& MenuBuilder, TArray<ALandscapeBlueprintBrushBase*> Brushes);
-	void AddBrushToCurrentLayer(ALandscapeBlueprintBrushBase* Brush);
+	void FillAddBrushMenu(FMenuBuilder& MenuBuilder, TArray<ALandscapeBlueprintCustomBrush*> Brushes);
 	void FillClearPaintLayerMenu(FMenuBuilder& MenuBuilder, int32 InLayerIndex, TArray<ULandscapeLayerInfoObject*> InUsedLayerInfos);
 	void FillClearLayerMenu(FMenuBuilder& MenuBuilder, int32 InLayerIndex);
+	void AddBrushToCurrentLayer(ALandscapeBlueprintCustomBrush* Brush);
 
 private:
 
diff --git a/Engine/Source/Editor/LandscapeEditor/Private/LandscapeEditorDetailCustomization_LayersBrushStack.cpp b/Engine/Source/Editor/LandscapeEditor/Private/LandscapeEditorDetailCustomization_LayersBrushStack.cpp
--- a/Engine/Source/Editor/LandscapeEditor/Private/LandscapeEditorDetailCustomization_LayersBrushStack.cpp
+++ b/Engine/Source/Editor/LandscapeEditor/Private/LandscapeEditorDetailCustomization_LayersBrushStack.cpp
@@ -26,7 +26,7 @@
 
 #include "LandscapeEditorDetailCustomization_TargetLayers.h"
 #include "Widgets/Input/SEditableText.h"
-#include "LandscapeBlueprintBrushBase.h"
+#include "LandscapeBPCustomBrush.h"
 
 #define LOCTEXT_NAMESPACE "LandscapeEditor.Layers"
 
@@ -209,7 +209,7 @@ END_SLATE_FUNCTION_BUILD_OPTIMIZATION
 
 TSharedPtr<SWidget> FLandscapeEditorCustomNodeBuilder_LayersBrushStack::OnBrushContextMenuOpening(int32 InBrushIndex)
 {
-	if (ALandscapeBlueprintBrushBase* CurrentBrush = GetBrush(InBrushIndex))
+	if (ALandscapeBlueprintCustomBrush* CurrentBrush = GetBrush(InBrushIndex))
 	{
 		FMenuBuilder MenuBuilder(true, NULL);
 		MenuBuilder.BeginSection("LandscapeEditorBrushActions", LOCTEXT("LandscapeEditorBrushActions.Heading", "Brushes"));
@@ -217,8 +217,8 @@ TSharedPtr<SWidget> FLandscapeEditorCustomNodeBuilder_LayersBrushStack::OnBrushC
 			TSharedRef<FLandscapeEditorCustomNodeBuilder_LayersBrushStack> SharedThis = AsShared();
 
 			// Add Brush
-			const TArray<ALandscapeBlueprintBrushBase*>& Brushes = GetEditorMode()->GetBrushList();
-			TArray<ALandscapeBlueprintBrushBase*> FilteredBrushes = Brushes.FilterByPredicate([](ALandscapeBlueprintBrushBase* Brush) { return Brush->GetOwningLandscape() == nullptr; });
+			const TArray<ALandscapeBlueprintCustomBrush*>& Brushes = GetEditorMode()->GetBrushList();
+			TArray<ALandscapeBlueprintCustomBrush*> FilteredBrushes = Brushes.FilterByPredicate([](ALandscapeBlueprintCustomBrush* Brush) { return Brush->GetOwningLandscape() == nullptr; });
 			if (FilteredBrushes.Num())
 			{
 				MenuBuilder.AddSubMenu(
@@ -245,10 +245,10 @@ TSharedPtr<SWidget> FLandscapeEditorCustomNodeBuilder_LayersBrushStack::OnBrushC
 	return nullptr;
 }
 
-void FLandscapeEditorCustomNodeBuilder_LayersBrushStack::FillAddBrushMenu(FMenuBuilder& MenuBuilder, TArray<ALandscapeBlueprintBrushBase*> Brushes)
+void FLandscapeEditorCustomNodeBuilder_LayersBrushStack::FillAddBrushMenu(FMenuBuilder& MenuBuilder, TArray<ALandscapeBlueprintCustomBrush*> Brushes)
 {
 	TSharedRef<FLandscapeEditorCustomNodeBuilder_LayersBrushStack> SharedThis = AsShared();
-	for (ALandscapeBlueprintBrushBase* Brush : Brushes)
+	for (ALandscapeBlueprintCustomBrush* Brush : Brushes)
 	{
 		FUIAction AddAction = FUIAction(FExecuteAction::CreateLambda([SharedThis, Brush]()
 		{
@@ -261,7 +261,7 @@ void FLandscapeEditorCustomNodeBuilder_LayersBrushStack::FillAddBrushMenu(FMenuB
 
 FReply FLandscapeEditorCustomNodeBuilder_LayersBrushStack::OnToggleAffectsHeightmap(int32 InBrushIndex)
 {
-	if (ALandscapeBlueprintBrushBase* Brush = GetBrush(InBrushIndex))
+	if (ALandscapeBlueprintCustomBrush* Brush = GetBrush(InBrushIndex))
 	{
 		const FScopedTransaction Transaction(LOCTEXT("Landscape_Brush_AffectsHeightmap", "Set Brush Affects Heightmap"));
 		bool bAffectsHeightmap = Brush->IsAffectingHeightmap();
@@ -272,7 +272,7 @@ FReply FLandscapeEditorCustomNodeBuilder_LayersBrushStack::OnToggleAffectsHeight
 
 FReply FLandscapeEditorCustomNodeBuilder_LayersBrushStack::OnToggleAffectsWeightmap(int32 InBrushIndex)
 {
-	if (ALandscapeBlueprintBrushBase* Brush = GetBrush(InBrushIndex))
+	if (ALandscapeBlueprintCustomBrush* Brush = GetBrush(InBrushIndex))
 	{
 		const FScopedTransaction Transaction(LOCTEXT("Landscape_Brush_AffectsWeightmap", "Set Brush Affects Weightmap"));
 		bool bAffectsWeightmap = Brush->IsAffectingWeightmap();
@@ -283,7 +283,7 @@ FReply FLandscapeEditorCustomNodeBuilder_LayersBrushStack::OnToggleAffectsWeight
 
 FReply FLandscapeEditorCustomNodeBuilder_LayersBrushStack::OnToggleVisibility(int32 InBrushIndex)
 {
-	if (ALandscapeBlueprintBrushBase* Brush = GetBrush(InBrushIndex))
+	if (ALandscapeBlueprintCustomBrush* Brush = GetBrush(InBrushIndex))
 	{
 		const FScopedTransaction Transaction(LOCTEXT("Landscape_Brush_SetVisibility", "Set Brush Visibility"));
 		bool bVisible = Brush->IsVisible();
@@ -294,33 +294,33 @@ FReply FLandscapeEditorCustomNodeBuilder_LayersBrushStack::OnToggleVisibility(in
 
 const FSlateBrush* FLandscapeEditorCustomNodeBuilder_LayersBrushStack::GetAffectsWeightmapBrush(int32 InBrushIndex) const
 {
-	ALandscapeBlueprintBrushBase* Brush = GetBrush(InBrushIndex);
+	ALandscapeBlueprintCustomBrush* Brush = GetBrush(InBrushIndex);
 	return Brush && Brush->IsAffectingWeightmap() ? FEditorStyle::GetBrush("LandscapeEditor.Brush.AffectsWeight.Enabled") : FEditorStyle::GetBrush("LandscapeEditor.Brush.AffectsWeight.Disabled");
 }
 
 const FSlateBrush* FLandscapeEditorCustomNodeBuilder_LayersBrushStack::GetAffectsHeightmapBrush(int32 InBrushIndex) const
 {
-	ALandscapeBlueprintBrushBase* Brush = GetBrush(InBrushIndex);
+	ALandscapeBlueprintCustomBrush* Brush = GetBrush(InBrushIndex);
 	return Brush && Brush->IsAffectingHeightmap() ? FEditorStyle::GetBrush("LandscapeEditor.Brush.AffectsHeight.Enabled") : FEditorStyle::GetBrush("LandscapeEditor.Brush.AffectsHeight.Disabled");
 }
 
 const FSlateBrush* FLandscapeEditorCustomNodeBuilder_LayersBrushStack::GetVisibilityBrush(int32 InBrushIndex) const
 {
-	ALandscapeBlueprintBrushBase* Brush = GetBrush(InBrushIndex);
+	ALandscapeBlueprintCustomBrush* Brush = GetBrush(InBrushIndex);
 	bool bIsVisible = Brush && Brush->IsVisible();
 	return bIsVisible ? FEditorStyle::GetBrush("Level.VisibleIcon16x") : FEditorStyle::GetBrush("Level.NotVisibleIcon16x");
 }
 
 bool FLandscapeEditorCustomNodeBuilder_LayersBrushStack::IsBrushEnabled(int32 InBrushIndex) const
 {
-	ALandscapeBlueprintBrushBase* Brush = GetBrush(InBrushIndex);
+	ALandscapeBlueprintCustomBrush* Brush = GetBrush(InBrushIndex);
 	const FEdModeLandscape* LandscapeEdMode = GetEditorMode();
 	return Brush && ((Brush->IsAffectingHeightmap() && LandscapeEdMode->CurrentToolTarget.TargetType == ELandscapeToolTargetType::Type::Heightmap) || (Brush->IsAffectingWeightmap() && LandscapeEdMode->CurrentToolTarget.TargetType == ELandscapeToolTargetType::Type::Weightmap));
 }
 
 bool FLandscapeEditorCustomNodeBuilder_LayersBrushStack::IsBrushSelected(int32 InBrushIndex) const
 {
-	ALandscapeBlueprintBrushBase* Brush = GetBrush(InBrushIndex);
+	ALandscapeBlueprintCustomBrush* Brush = GetBrush(InBrushIndex);
 	return Brush && Brush->IsSelected();
 }
 
@@ -335,8 +335,8 @@ void FLandscapeEditorCustomNodeBuilder_LayersBrushStack::OnBrushSelectionChanged
 	const FName CurrentToolName = LandscapeEdMode->CurrentTool->GetToolName();
 	if (CurrentToolName == TEXT("BlueprintBrush"))
 	{
-		ALandscapeBlueprintBrushBase* Brush = GetBrush(InBrushIndex);
-		if (Brush != nullptr)
+		ALandscapeBlueprintCustomBrush* Brush = GetBrush(InBrushIndex);
+		if (Brush != nullptr && !Brush->IsCommited())
 		{
 			FScopedTransaction Transaction(LOCTEXT("LandscapeBrushSelect", "Brush selection"));
 			GEditor->SelectNone(true, true);
@@ -347,7 +347,7 @@ void FLandscapeEditorCustomNodeBuilder_LayersBrushStack::OnBrushSelectionChanged
 
 FText FLandscapeEditorCustomNodeBuilder_LayersBrushStack::GetBrushText(int32 InBrushIndex) const
 {
-	if (ALandscapeBlueprintBrushBase* Brush = GetBrush(InBrushIndex))
+	if (ALandscapeBlueprintCustomBrush* Brush = GetBrush(InBrushIndex))
 	{
 		return FText::FromString(Brush->GetActorLabel());
 	}
@@ -357,15 +357,15 @@ FText FLandscapeEditorCustomNodeBuilder_LayersBrushStack::GetBrushText(int32 InB
 
 FSlateColor FLandscapeEditorCustomNodeBuilder_LayersBrushStack::GetBrushTextColor(int32 InBrushIndex) const
 {
-	if (ALandscapeBlueprintBrushBase* Brush = GetBrush(InBrushIndex))
+	if (ALandscapeBlueprintCustomBrush* Brush = GetBrush(InBrushIndex))
 	{
-		return FSlateColor::UseForeground();
+		return Brush->IsCommited() ? FSlateColor::UseSubduedForeground() : FSlateColor::UseForeground();
 	}
 
 	return FSlateColor::UseSubduedForeground();
 }
 
-ALandscapeBlueprintBrushBase* FLandscapeEditorCustomNodeBuilder_LayersBrushStack::GetBrush(int32 InBrushIndex) const
+ALandscapeBlueprintCustomBrush* FLandscapeEditorCustomNodeBuilder_LayersBrushStack::GetBrush(int32 InBrushIndex) const
 {
 	if (FEdModeLandscape* LandscapeEdMode = GetEditorMode())
 	{
@@ -382,7 +382,7 @@ FReply FLandscapeEditorCustomNodeBuilder_LayersBrushStack::HandleDragDetected(co
 		FLandscapeLayer* Layer = LandscapeEdMode->GetCurrentLayer();
 		if (Layer && !Layer->bLocked)
 		{
-			const TArray<ALandscapeBlueprintBrushBase*>& BrushStack = LandscapeEdMode->GetBrushesForCurrentLayer();
+			const TArray<ALandscapeBlueprintCustomBrush*>& BrushStack = LandscapeEdMode->GetBrushesForCurrentLayer();
 			if (BrushStack.IsValidIndex(SlotIndex))
 			{
 				TSharedPtr<SWidget> Row = GenerateRow(SlotIndex);
diff --git a/Engine/Source/Editor/LandscapeEditor/Private/LandscapeEditorDetailCustomization_LayersBrushStack.h b/Engine/Source/Editor/LandscapeEditor/Private/LandscapeEditorDetailCustomization_LayersBrushStack.h
--- a/Engine/Source/Editor/LandscapeEditor/Private/LandscapeEditorDetailCustomization_LayersBrushStack.h
+++ b/Engine/Source/Editor/LandscapeEditor/Private/LandscapeEditorDetailCustomization_LayersBrushStack.h
@@ -59,7 +59,7 @@ protected:
 	TSharedPtr<SWidget> GenerateRow(int32 InBrushIndex);
 	TSharedPtr<SWidget> OnBrushContextMenuOpening(int32 InBrushIndex);
 
-	void FillAddBrushMenu(FMenuBuilder& MenuBuilder, TArray<ALandscapeBlueprintBrushBase*> Brushes);
+	void FillAddBrushMenu(FMenuBuilder& MenuBuilder, TArray<ALandscapeBlueprintCustomBrush*> Brushes);
 
 	FReply OnToggleVisibility(int32 InBrushIndex);
 	FReply OnToggleAffectsHeightmap(int32 InBrushIndex);
@@ -73,7 +73,7 @@ protected:
 	void OnBrushSelectionChanged(int32 InBrushIndex);
 	FText GetBrushText(int32 InBrushIndex) const;
 	FSlateColor GetBrushTextColor(int32 InBrushIndex) const;
-	ALandscapeBlueprintBrushBase* GetBrush(int32 InBrushIndex) const;
+	ALandscapeBlueprintCustomBrush* GetBrush(int32 InBrushIndex) const;
 
 	// Drag/Drop handling
 	FReply HandleDragDetected(const FGeometry& MyGeometry, const FPointerEvent& MouseEvent, int32 SlotIndex, SVerticalBox::FSlot* Slot);
diff --git a/Engine/Source/Editor/LandscapeEditor/Public/LandscapeBlueprintBrush.h b/Engine/Source/Editor/LandscapeEditor/Public/LandscapeBlueprintBrush.h
deleted file mode 100644
diff --git a/Engine/Source/Editor/LandscapeEditor/Public/LandscapeEditorObject.h b/Engine/Source/Editor/LandscapeEditor/Public/LandscapeEditorObject.h
--- a/Engine/Source/Editor/LandscapeEditor/Public/LandscapeEditorObject.h
+++ b/Engine/Source/Editor/LandscapeEditor/Public/LandscapeEditorObject.h
@@ -8,7 +8,7 @@
 #include "LandscapeProxy.h"
 #include "Editor/LandscapeEditor/Private/LandscapeEdMode.h"
 #include "LandscapeFileFormatInterface.h"
-#include "LandscapeBlueprintBrush.h"
+#include "LandscapeBPCustomBrush.h"
 
 #include "LandscapeEditorObject.generated.h"
 
@@ -454,7 +454,7 @@ class ULandscapeEditorObject : public UObject
 	// Blueprint Brush Tool
 
 	UPROPERTY(Category = "Tool Settings", EditAnywhere, Transient, meta = (DisplayName = "Blueprint Brush", ShowForTools = "BlueprintBrush"))
-	TSubclassOf<ALandscapeBlueprintBrush> BlueprintBrush;
+	TSubclassOf<ALandscapeBlueprintCustomBrush> BlueprintCustomBrush;
 
 	// Resize Landscape Tool
 
diff --git a/Engine/Source/Runtime/Landscape/Classes/Landscape.h b/Engine/Source/Runtime/Landscape/Classes/Landscape.h
--- a/Engine/Source/Runtime/Landscape/Classes/Landscape.h
+++ b/Engine/Source/Runtime/Landscape/Classes/Landscape.h
@@ -5,7 +5,7 @@
 #include "CoreMinimal.h"
 #include "UObject/ObjectMacros.h"
 #include "LandscapeProxy.h"
-#include "LandscapeBlueprintBrushBase.h"
+#include "LandscapeBPCustomBrush.h"
 #include "Delegates/DelegateCombinations.h"
 
 #include "Landscape.generated.h"
@@ -89,38 +89,28 @@ struct FLandscapeLayerBrush
 	GENERATED_USTRUCT_BODY()
 
 	FLandscapeLayerBrush()
-#if WITH_EDITORONLY_DATA
-		: BlueprintBrush(nullptr)
+		: BPCustomBrush(nullptr)
+#if WITH_EDITOR
 		, LandscapeSize(MAX_int32, MAX_int32)
 		, LandscapeRenderTargetSize(MAX_int32, MAX_int32)
 #endif
 	{}
 
-	FLandscapeLayerBrush(ALandscapeBlueprintBrushBase* InBlueprintBrush)
-#if WITH_EDITORONLY_DATA
-		: BlueprintBrush(InBlueprintBrush)
-		, LandscapeSize(MAX_int32, MAX_int32)
-		, LandscapeRenderTargetSize(MAX_int32, MAX_int32)
-#endif
+	FLandscapeLayerBrush(ALandscapeBlueprintCustomBrush* InBrush)
+		: BPCustomBrush(InBrush)
 	{}
 
 #if WITH_EDITOR
-	UTextureRenderTarget2D* Render(bool InIsHeightmap, const FIntPoint& InLandscapeSize, UTextureRenderTarget2D* InLandscapeRenderTarget, const FName& InWeightmapLayerName = NAME_None);
-	ALandscapeBlueprintBrushBase* GetBrush() const;
-	bool IsAffectingHeightmap() const;
-	bool IsAffectingWeightmapLayer(const FName& InWeightmapLayerName) const;
-	void SetOwner(ALandscape* InOwner);
+	UTextureRenderTarget2D* Render(bool InIsHeightmap, const FIntPoint& InLandscapeSize, UTextureRenderTarget2D* InLandscapeRenderTarget);
 #endif
 
+	UPROPERTY()
+	ALandscapeBlueprintCustomBrush* BPCustomBrush;
+
 private:
 
 #if WITH_EDITOR
 	bool Initialize(const FIntPoint& LandscapeSize, UTextureRenderTarget2D* InLandscapeRenderTarget);
-#endif
-
-#if WITH_EDITORONLY_DATA
-	UPROPERTY()
-	ALandscapeBlueprintBrushBase* BlueprintBrush;
 
 	FTransform LandscapeTransform;
 	FIntPoint LandscapeSize;
@@ -269,7 +259,7 @@ public:
 	LANDSCAPE_API void RequestLayersContentUpdateForceAll(ELandscapeLayerUpdateMode InModeMask = ELandscapeLayerUpdateMode::Update_All);
 	LANDSCAPE_API void RequestLayersContentUpdate(ELandscapeLayerUpdateMode InModeMask);
 	LANDSCAPE_API bool ReorderLayer(int32 InStartingLayerIndex, int32 InDestinationLayerIndex);
-	LANDSCAPE_API FLandscapeLayer* DuplicateLayerAndMoveBrushes(const FLandscapeLayer& InOtherLayer);
+	LANDSCAPE_API FLandscapeLayer* DuplicateLayer(const FLandscapeLayer& InOtherLayer);
 	LANDSCAPE_API void CreateLayer(FName InName = NAME_None);
 	LANDSCAPE_API void CreateDefaultLayer();
 	LANDSCAPE_API void CopyOldDataToDefaultLayer();
@@ -312,13 +302,13 @@ public:
 	LANDSCAPE_API bool IsLayerBlendSubstractive(int32 InLayerIndex, const TWeakObjectPtr<ULandscapeLayerInfoObject>& InLayerInfoObj) const;
 	LANDSCAPE_API void SetLayerSubstractiveBlendStatus(int32 InLayerIndex, bool InStatus, const TWeakObjectPtr<ULandscapeLayerInfoObject>& InLayerInfoObj);
 
-	LANDSCAPE_API void AddBrushToLayer(int32 InLayerIndex, class ALandscapeBlueprintBrushBase* InBrush);
-	LANDSCAPE_API void RemoveBrush(class ALandscapeBlueprintBrushBase* InBrush);
-	LANDSCAPE_API void RemoveBrushFromLayer(int32 InLayerIndex, class ALandscapeBlueprintBrushBase* InBrush);
+	LANDSCAPE_API void AddBrushToLayer(int32 InLayerIndex, class ALandscapeBlueprintCustomBrush* InBrush);
+	LANDSCAPE_API void RemoveBrush(class ALandscapeBlueprintCustomBrush* InBrush);
+	LANDSCAPE_API void RemoveBrushFromLayer(int32 InLayerIndex, class ALandscapeBlueprintCustomBrush* InBrush);
 	LANDSCAPE_API bool ReorderLayerBrush(int32 InLayerIndex, int32 InStartingLayerBrushIndex, int32 InDestinationLayerBrushIndex);
-	LANDSCAPE_API class ALandscapeBlueprintBrushBase* GetBrushForLayer(int32 InLayerIndex, int8 BrushIndex) const;
-	LANDSCAPE_API TArray<class ALandscapeBlueprintBrushBase*> GetBrushesForLayer(int32 InLayerIndex) const;
-	LANDSCAPE_API void OnBlueprintBrushChanged();
+	LANDSCAPE_API class ALandscapeBlueprintCustomBrush* GetBrushForLayer(int32 InLayerIndex, int8 BrushIndex) const;
+	LANDSCAPE_API TArray<class ALandscapeBlueprintCustomBrush*> GetBrushesForLayer(int32 InLayerIndex) const;
+	LANDSCAPE_API void OnBPCustomBrushChanged();
 	
 	LANDSCAPE_API void OnPreSave();
 
@@ -400,8 +390,8 @@ public:
 	UPROPERTY(EditAnywhere, Category=Landscape)
 	bool bCanHaveLayersContent = false;
 
-	DECLARE_EVENT(ALandscape, FLandscapeBlueprintBrushChangedDelegate);
-	FLandscapeBlueprintBrushChangedDelegate& OnBlueprintBrushChangedDelegate() { return LandscapeBlueprintBrushChangedDelegate; }
+	DECLARE_EVENT(ALandscape, FLandscapeBPCustomBrushChangedDelegate);
+	FLandscapeBPCustomBrushChangedDelegate& OnBPCustomBrushChangedDelegate() { return LandscapeBPCustomBrushChangedDelegate; }
 
 	/** Target Landscape Layer for Landscape Splines */
 	UPROPERTY()
@@ -423,7 +413,7 @@ public:
 	TArray<UTextureRenderTarget2D*> WeightmapRTList;
 
 private:
-	FLandscapeBlueprintBrushChangedDelegate LandscapeBlueprintBrushChangedDelegate;
+	FLandscapeBPCustomBrushChangedDelegate LandscapeBPCustomBrushChangedDelegate;
 
 	/** Components affected by landscape splines (used to partially clear Layer Reserved for Splines) */
 	UPROPERTY(Transient)
diff --git a/Engine/Source/Runtime/Landscape/Classes/LandscapeBPCustomBrush.h b/Engine/Source/Runtime/Landscape/Classes/LandscapeBPCustomBrush.h
new file mode 100644
--- /dev/null
+++ b/Engine/Source/Runtime/Landscape/Classes/LandscapeBPCustomBrush.h
@@ -0,0 +1,73 @@
+// Copyright 1998-2019 Epic Games, Inc. All Rights Reserved.
+
+#pragma once
+
+#include "CoreMinimal.h"
+#include "UObject/ObjectMacros.h"
+#include "GameFramework/Actor.h"
+
+#include "LandscapeBPCustomBrush.generated.h"
+
+UCLASS(Abstract, Blueprintable, hidecategories = (Replication, Input, LOD, Actor, Cooking, Rendering))
+class LANDSCAPE_API ALandscapeBlueprintCustomBrush : public AActor
+{
+	GENERATED_UCLASS_BODY()
+
+private:
+	UPROPERTY(Category= "Settings", EditAnywhere)
+	bool AffectHeightmap;
+
+	UPROPERTY(Category= "Settings", EditAnywhere)
+	bool AffectWeightmap;
+
+	UPROPERTY(Category = "Settings", EditAnywhere)
+	TArray<FName> AffectedWeightmapLayers;
+
+#if WITH_EDITORONLY_DATA
+	UPROPERTY(DuplicateTransient)
+	class ALandscape* OwningLandscape;
+
+	UPROPERTY(NonTransactional, DuplicateTransient)
+	bool bIsCommited;
+
+	UPROPERTY(Transient)
+	bool bIsVisible;
+#endif
+public:
+
+	virtual bool ShouldTickIfViewportsOnly() const override;
+	virtual void Tick(float DeltaSeconds) override;
+
+	bool IsAffectingHeightmap() const { return AffectHeightmap; }
+	bool IsAffectingWeightmap() const { return AffectWeightmap; }
+	bool IsAffectingWeightmapLayer(const FName& InLayerName) const;
+
+	UFUNCTION(BlueprintImplementableEvent)
+	UTextureRenderTarget2D* Render(bool InIsHeightmap, UTextureRenderTarget2D* InCombinedResult);
+
+	UFUNCTION(BlueprintImplementableEvent)
+	void Initialize(const FTransform& InLandscapeTransform, const FIntPoint& InLandscapeSize, const FIntPoint& InLandscapeRenderTargetSize);
+
+	UFUNCTION(BlueprintCallable, Category="Landscape")
+	void RequestLandscapeUpdate();
+
+#if WITH_EDITOR
+	void SetCommitState(bool InCommited);
+	bool IsCommited() const { return bIsCommited; }
+
+	bool IsVisible() const { return bIsVisible; }
+	void SetIsVisible(bool bInIsVisible);
+
+	void SetAffectsHeightmap(bool bInAffectsHeightmap);
+	void SetAffectsWeightmap(bool bInAffectsWeightmap);
+
+	void SetOwningLandscape(class ALandscape* InOwningLandscape);
+	class ALandscape* GetOwningLandscape() const;
+
+	virtual void PostEditMove(bool bFinished) override;
+	virtual void PostEditChangeProperty(FPropertyChangedEvent& PropertyChangedEvent) override;
+
+	virtual void Destroyed() override;
+#endif
+};
+
diff --git a/Engine/Source/Runtime/Landscape/Classes/LandscapeProxy.h b/Engine/Source/Runtime/Landscape/Classes/LandscapeProxy.h
--- a/Engine/Source/Runtime/Landscape/Classes/LandscapeProxy.h
+++ b/Engine/Source/Runtime/Landscape/Classes/LandscapeProxy.h
@@ -12,6 +12,7 @@
 #include "Async/AsyncWork.h"
 #include "Engine/Texture.h"
 #include "PerPlatformProperties.h"
+#include "LandscapeBPCustomBrush.h"
 #include "LandscapeComponent.h"
 #include "LandscapeWeightmapUsage.h"
 #include "VT/RuntimeVirtualTextureEnum.h"
diff --git a/Engine/Source/Runtime/Landscape/Private/LandscapeBPCustomBrush.cpp b/Engine/Source/Runtime/Landscape/Private/LandscapeBPCustomBrush.cpp
new file mode 100644
--- /dev/null
+++ b/Engine/Source/Runtime/Landscape/Private/LandscapeBPCustomBrush.cpp
@@ -0,0 +1,141 @@
+// Copyright 1998-2019 Epic Games, Inc. All Rights Reserved.
+
+#include "LandscapeBPCustomBrush.h"
+#include "CoreMinimal.h"
+#include "LandscapeProxy.h"
+#include "Landscape.h"
+
+#define LOCTEXT_NAMESPACE "Landscape"
+
+ALandscapeBlueprintCustomBrush::ALandscapeBlueprintCustomBrush(const FObjectInitializer& ObjectInitializer)
+#if WITH_EDITORONLY_DATA
+	: OwningLandscape(nullptr)
+	, bIsCommited(false)
+	, bIsVisible(true)
+#endif
+{
+	USceneComponent* SceneComp = CreateDefaultSubobject<USceneComponent>(TEXT("RootComponent"));
+	RootComponent = SceneComp;
+
+	PrimaryActorTick.bCanEverTick = true;
+	PrimaryActorTick.TickGroup = TG_DuringPhysics;
+	PrimaryActorTick.bStartWithTickEnabled = true;
+	PrimaryActorTick.SetTickFunctionEnable(true);
+	bIsEditorOnlyActor = true;
+}
+
+void ALandscapeBlueprintCustomBrush::Tick(float DeltaSeconds)
+{
+	// Forward the Tick to the instances class of this BP
+	if (GetClass()->HasAnyClassFlags(CLASS_CompiledFromBlueprint))
+	{
+		TGuardValue<bool> AutoRestore(GAllowActorScriptExecutionInEditor, true);
+		ReceiveTick(DeltaSeconds);
+	}
+
+	Super::Tick(DeltaSeconds);
+}
+
+bool ALandscapeBlueprintCustomBrush::ShouldTickIfViewportsOnly() const
+{
+	return true;
+}
+
+void ALandscapeBlueprintCustomBrush::RequestLandscapeUpdate()
+{
+#if WITH_EDITORONLY_DATA
+	if (OwningLandscape)
+	{
+		OwningLandscape->RequestLayersContentUpdateForceAll();
+	}
+#endif
+}
+
+#if WITH_EDITOR
+
+void ALandscapeBlueprintCustomBrush::SetCommitState(bool InCommited)
+{
+#if WITH_EDITORONLY_DATA
+	bListedInSceneOutliner = !InCommited;
+	bEditable = !InCommited;
+	bIsCommited = InCommited;
+#endif
+}
+
+void ALandscapeBlueprintCustomBrush::SetIsVisible(bool bInIsVisible)
+{
+	Modify();
+	bIsVisible = bInIsVisible;
+	GetOwningLandscape()->OnBPCustomBrushChanged();
+}
+
+void ALandscapeBlueprintCustomBrush::SetAffectsHeightmap(bool bInAffectsHeightmap)
+{
+	Modify();
+	AffectHeightmap = bInAffectsHeightmap;
+	GetOwningLandscape()->OnBPCustomBrushChanged();
+}
+
+void ALandscapeBlueprintCustomBrush::SetAffectsWeightmap(bool bInAffectsWeightmap)
+{
+	Modify();
+	AffectWeightmap = bInAffectsWeightmap;
+	GetOwningLandscape()->OnBPCustomBrushChanged();
+}
+
+void ALandscapeBlueprintCustomBrush::SetOwningLandscape(ALandscape* InOwningLandscape)
+{
+	Modify();
+	OwningLandscape = InOwningLandscape;
+}
+
+ALandscape* ALandscapeBlueprintCustomBrush::GetOwningLandscape() const
+{ 
+	return OwningLandscape; 
+}
+
+bool ALandscapeBlueprintCustomBrush::IsAffectingWeightmapLayer(const FName& InLayerName) const
+{
+	return AffectedWeightmapLayers.Contains(InLayerName);
+}
+
+void ALandscapeBlueprintCustomBrush::PostEditMove(bool bFinished)
+{
+	Super::PostEditMove(bFinished);
+
+	if (OwningLandscape != nullptr)
+	{
+		if (bFinished)
+		{
+			OwningLandscape->RequestLayersContentUpdateForceAll();
+		}
+		else
+		{
+			OwningLandscape->RequestLayersContentUpdate(ELandscapeLayerUpdateMode::Update_All_Editing);
+		}
+	}
+}
+
+void ALandscapeBlueprintCustomBrush::PostEditChangeProperty(FPropertyChangedEvent& PropertyChangedEvent)
+{
+	Super::PostEditChangeProperty(PropertyChangedEvent);
+	if (OwningLandscape)
+	{
+		OwningLandscape->OnBPCustomBrushChanged();
+	}
+}
+
+void ALandscapeBlueprintCustomBrush::Destroyed()
+{
+	Super::Destroyed();
+
+	if (OwningLandscape && !GIsReinstancing)
+	{
+		OwningLandscape->RemoveBrush(this);
+	}
+	OwningLandscape = nullptr;
+}
+
+#endif
+
+#undef LOCTEXT_NAMESPACE
diff --git a/Engine/Source/Runtime/Landscape/Private/LandscapeBlueprintBrushBase.cpp b/Engine/Source/Runtime/Landscape/Private/LandscapeBlueprintBrushBase.cpp
deleted file mode 100644
diff --git a/Engine/Source/Runtime/Landscape/Private/LandscapeEdit.cpp b/Engine/Source/Runtime/Landscape/Private/LandscapeEdit.cpp
--- a/Engine/Source/Runtime/Landscape/Private/LandscapeEdit.cpp
+++ b/Engine/Source/Runtime/Landscape/Private/LandscapeEdit.cpp
@@ -2803,7 +2803,7 @@ LANDSCAPE_API void ALandscapeProxy::Import(const FGuid& InGuid, int32 InMinX, in
 		{
 			for (const FLandscapeLayer& OldLayer : *InImportLayers)
 			{
-				FLandscapeLayer* NewLayer = LandscapeActor->DuplicateLayerAndMoveBrushes(OldLayer);
+				FLandscapeLayer* NewLayer = LandscapeActor->DuplicateLayer(OldLayer);
 				check(NewLayer != nullptr);
 
 				FLayerImportSettings ImportSettings;
diff --git a/Engine/Source/Runtime/Landscape/Private/LandscapeEditLayers.cpp b/Engine/Source/Runtime/Landscape/Private/LandscapeEditLayers.cpp
--- a/Engine/Source/Runtime/Landscape/Private/LandscapeEditLayers.cpp
+++ b/Engine/Source/Runtime/Landscape/Private/LandscapeEditLayers.cpp
@@ -23,7 +23,7 @@ LandscapeEditLayers.cpp: Landscape editing layers mode
 #include "LandscapeEditorModule.h"
 #include "LandscapeToolInterface.h"
 #include "ComponentRecreateRenderStateContext.h"
-#include "LandscapeBlueprintBrushBase.h"
+#include "LandscapeBPCustomBrush.h"
 #include "Materials/MaterialInstanceConstant.h"
 #include "LandscapeMaterialInstanceConstant.h"
 #include "Materials/MaterialExpressionLandscapeVisibilityMask.h"
@@ -2762,6 +2762,12 @@ int32 ALandscape::RegenerateLayersHeightmaps(const TArray<ULandscapeComponent*>&
 					// TODO: handle conversion/handling of RT not same size as internal size
 
 					FLandscapeLayerBrush& Brush = Layer.Brushes[i];
+
+					if (Brush.BPCustomBrush == nullptr || !Brush.BPCustomBrush->IsAffectingHeightmap() || !Brush.BPCustomBrush->IsVisible())
+					{
+						continue;
+					}
+
 					UTextureRenderTarget2D* BrushOutputNonAtlasRT = Brush.Render(true, LandscapeSize, CombinedHeightmapNonAtlasRT);
 					if (BrushOutputNonAtlasRT == nullptr || BrushOutputNonAtlasRT->SizeX != CombinedHeightmapNonAtlasRT->SizeX || BrushOutputNonAtlasRT->SizeY != CombinedHeightmapNonAtlasRT->SizeY)
 					{
@@ -2770,7 +2776,7 @@ int32 ALandscape::RegenerateLayersHeightmaps(const TArray<ULandscapeComponent*>&
 
 					INC_DWORD_STAT(STAT_LandscapeLayersRegenerateDrawCalls); // Brush Render
 
-					PrintLayersDebugRT(OutputDebugName ? FString::Printf(TEXT("LS Height: %s %s -> BrushNonAtlas %s"), *Layer.Name.ToString(), *Brush.GetBrush()->GetName(), *BrushOutputNonAtlasRT->GetName()) : GEmptyDebugName, BrushOutputNonAtlasRT);
+					PrintLayersDebugRT(OutputDebugName ? FString::Printf(TEXT("LS Height: %s %s -> BrushNonAtlas %s"), *Layer.Name.ToString(), *Brush.BPCustomBrush->GetName(), *BrushOutputNonAtlasRT->GetName()) : GEmptyDebugName, BrushOutputNonAtlasRT);
 
 					// Resolve back to Combined heightmap
 					CopyLayersTexture(BrushOutputNonAtlasRT, CombinedHeightmapNonAtlasRT);
@@ -3500,7 +3506,13 @@ int32 ALandscape::RegenerateLayersWeightmaps(const TArray<ULandscapeComponent*>&
 					for (int32 i = 0; i < Layer.Brushes.Num(); ++i)
 					{
 						FLandscapeLayerBrush& Brush = Layer.Brushes[i];
-						if (Brush.IsAffectingWeightmapLayer(InfoLayerSettings.GetLayerName()) && !LayerInfoObjects.Contains(InfoLayerSettings.LayerInfoObj))
+
+						if (Brush.BPCustomBrush == nullptr || !Brush.BPCustomBrush->IsAffectingWeightmap() || !Brush.BPCustomBrush->IsVisible())
+						{
+							continue;
+						}
+
+						if (Brush.BPCustomBrush->IsAffectingWeightmapLayer(InfoLayerSettings.GetLayerName()) && !LayerInfoObjects.Contains(InfoLayerSettings.LayerInfoObj))
 						{
 							LayerInfoObjects.Add(InfoLayerSettings.LayerInfoObj, LayerInfoSettingsIndex + 1); // due to visibility layer that is at 0
 						}
@@ -3632,17 +3644,23 @@ int32 ALandscape::RegenerateLayersWeightmaps(const TArray<ULandscapeComponent*>&
 							// TODO: handle conversion/handling of RT not same size as internal size
 
 							FLandscapeLayerBrush& Brush = Layer.Brushes[i];
-							UTextureRenderTarget2D* BrushOutputRT = Brush.Render(false, LandscapeSize, LandscapeScratchRT3, LayerInfoObj->LayerName);
-							if (BrushOutputRT == nullptr || BrushOutputRT->SizeX != LandscapeScratchRT3->SizeX || BrushOutputRT->SizeY != LandscapeScratchRT3->SizeY)
+
+							if (Brush.BPCustomBrush == nullptr || !Brush.BPCustomBrush->IsAffectingWeightmap() || !Brush.BPCustomBrush->IsAffectingWeightmapLayer(LayerInfoObj->LayerName) || !Brush.BPCustomBrush->IsVisible())
 							{
 								continue;
 							}
 
 							BrushRequiredAllocations.AddUnique(LayerInfoObj);
 
+							UTextureRenderTarget2D* BrushOutputRT = Brush.Render(false, LandscapeSize, LandscapeScratchRT3);
+							if (BrushOutputRT == nullptr || BrushOutputRT->SizeX != LandscapeScratchRT3->SizeX || BrushOutputRT->SizeY != LandscapeScratchRT3->SizeY)
+							{
+								continue;
+							}
+
 							INC_DWORD_STAT(STAT_LandscapeLayersRegenerateDrawCalls); // Brush Render
 
-							PrintLayersDebugRT(OutputDebugName ? FString::Printf(TEXT("LS Weight: %s %s -> Brush %s"), *Layer.Name.ToString(), *Brush.GetBrush()->GetName(), *BrushOutputRT->GetName()) : GEmptyDebugName, BrushOutputRT);
+							PrintLayersDebugRT(OutputDebugName ? FString::Printf(TEXT("LS Weight: %s %s -> Brush %s"), *Layer.Name.ToString(), *Brush.BPCustomBrush->GetName(), *BrushOutputRT->GetName()) : GEmptyDebugName, BrushOutputRT);
 
 							SourceDebugName = OutputDebugName ? FString::Printf(TEXT("Weight: %s PaintLayer: %s Brush: %s"), *Layer.Name.ToString(), *LayerInfoObj->LayerName.ToString(), *BrushOutputRT->GetName()) : GEmptyDebugName;
 							DestDebugName = OutputDebugName ? LandscapeScratchRT3->GetName() : GEmptyDebugName;
@@ -5330,7 +5348,7 @@ void ALandscape::SetEditingLayer(const FGuid& InLayerGuid)
 
 void ALandscape::SetGrassUpdateEnabled(bool bInGrassUpdateEnabled)
 {
-#if WITH_EDITORONLY_DATA
+#ifdef WITH_EDITORONLY_DATA
 	bGrassUpdateEnabled = bInGrassUpdateEnabled;
 #endif
 }
@@ -5382,7 +5400,7 @@ void ALandscape::CreateDefaultLayer()
 	RequestLayersInitialization();
 }
 
-FLandscapeLayer* ALandscape::DuplicateLayerAndMoveBrushes(const FLandscapeLayer& InOtherLayer)
+FLandscapeLayer* ALandscape::DuplicateLayer(const FLandscapeLayer& InOtherLayer)
 {
 	ULandscapeInfo* LandscapeInfo = GetLandscapeInfo();
 	if (!LandscapeInfo || !CanHaveLayersContent())
@@ -5395,10 +5413,18 @@ FLandscapeLayer* ALandscape::DuplicateLayerAndMoveBrushes(const FLandscapeLayer&
 	FLandscapeLayer NewLayer(InOtherLayer);
 	NewLayer.Guid = FGuid::NewGuid();
 
-	// Update owning landscape and reparent to landscape's level if necessary
+	// Copy Brush and reparent to the new landscape level if required
 	for (FLandscapeLayerBrush& Brush : NewLayer.Brushes)
 	{
-		Brush.SetOwner(this);
+		if (Brush.BPCustomBrush != nullptr)
+		{
+			if (Brush.BPCustomBrush->GetTypedOuter<ULevel>() != GetTypedOuter<ULevel>())
+			{
+				Brush.BPCustomBrush = DuplicateObject<ALandscapeBlueprintCustomBrush>(Brush.BPCustomBrush, GetTypedOuter<ULevel>());
+			}
+
+			Brush.BPCustomBrush->SetOwningLandscape(this);
+		}
 	}
 
 	int32 AddedIndex = LandscapeLayers.Add(NewLayer);
@@ -5544,7 +5570,7 @@ bool ALandscape::ReorderLayerBrush(int32 InLayerIndex, int32 InStartingLayerBrus
 	return false;
 }
 
-void ALandscape::AddBrushToLayer(int32 InLayerIndex, ALandscapeBlueprintBrushBase* InBrush)
+void ALandscape::AddBrushToLayer(int32 InLayerIndex, ALandscapeBlueprintCustomBrush* InBrush)
 {
 	if (FLandscapeLayer* Layer = GetLayer(InLayerIndex))
 	{
@@ -5555,7 +5581,7 @@ void ALandscape::AddBrushToLayer(int32 InLayerIndex, ALandscapeBlueprintBrushBas
 	}
 }
 
-void ALandscape::RemoveBrush(ALandscapeBlueprintBrushBase* InBrush)
+void ALandscape::RemoveBrush(ALandscapeBlueprintCustomBrush* InBrush)
 {
 	for (int32 LayerIndex = 0; LayerIndex < LandscapeLayers.Num(); ++LayerIndex)
 	{
@@ -5563,13 +5589,13 @@ void ALandscape::RemoveBrush(ALandscapeBlueprintBrushBase* InBrush)
 	}
 }
 
-void ALandscape::RemoveBrushFromLayer(int32 InLayerIndex, ALandscapeBlueprintBrushBase* InBrush)
+void ALandscape::RemoveBrushFromLayer(int32 InLayerIndex, ALandscapeBlueprintCustomBrush* InBrush)
 {
 	if (FLandscapeLayer* Layer = GetLayer(InLayerIndex))
 	{
 		for (int32 i = 0; i < Layer->Brushes.Num(); ++i)
 		{
-			if (Layer->Brushes[i].GetBrush() == InBrush)
+			if (Layer->Brushes[i].BPCustomBrush == InBrush)
 			{
 				Modify();
 				Layer->Brushes.RemoveAt(i);
@@ -5581,104 +5607,55 @@ void ALandscape::RemoveBrushFromLayer(int32 InLayerIndex, ALandscapeBlueprintBru
 	}
 }
 
-void ALandscape::OnBlueprintBrushChanged()
+void ALandscape::OnBPCustomBrushChanged()
 {
 #if WITH_EDITORONLY_DATA
-	LandscapeBlueprintBrushChangedDelegate.Broadcast();
+	LandscapeBPCustomBrushChangedDelegate.Broadcast();
 	RequestLayersContentUpdateForceAll();
 #endif
 }
 
-ALandscapeBlueprintBrushBase* ALandscape::GetBrushForLayer(int32 InLayerIndex, int8 InBrushIndex) const
+ALandscapeBlueprintCustomBrush* ALandscape::GetBrushForLayer(int32 InLayerIndex, int8 InBrushIndex) const
 {
 	if (const FLandscapeLayer* Layer = GetLayer(InLayerIndex))
 	{
 		if (Layer->Brushes.IsValidIndex(InBrushIndex))
 		{
-			return Layer->Brushes[InBrushIndex].GetBrush();
+			return Layer->Brushes[InBrushIndex].BPCustomBrush;
 		}
 	}
 	return nullptr;
 }
 
-TArray<ALandscapeBlueprintBrushBase*> ALandscape::GetBrushesForLayer(int32 InLayerIndex) const
+TArray<ALandscapeBlueprintCustomBrush*> ALandscape::GetBrushesForLayer(int32 InLayerIndex) const
 {
-	TArray<ALandscapeBlueprintBrushBase*> Brushes;
+	TArray<ALandscapeBlueprintCustomBrush*> Brushes;
 	if (const FLandscapeLayer* Layer = GetLayer(InLayerIndex))
 	{
 		Brushes.Reserve(Layer->Brushes.Num());
 		for (const FLandscapeLayerBrush& Brush : Layer->Brushes)
 		{
-			Brushes.Add(Brush.GetBrush());
+			Brushes.Add(Brush.BPCustomBrush);
 		}
 	}
 	return Brushes;
 }
 
-ALandscapeBlueprintBrushBase* FLandscapeLayerBrush::GetBrush() const
-{
-#if WITH_EDITORONLY_DATA
-	return BlueprintBrush;
-#else
-	return nullptr;
-#endif
-}
-
-void FLandscapeLayerBrush::SetOwner(ALandscape* InOwner)
+UTextureRenderTarget2D* FLandscapeLayerBrush::Render(bool InIsHeightmap, const FIntPoint& InLandscapeSize, UTextureRenderTarget2D* InLandscapeRenderTarget)
 {
-#if WITH_EDITORONLY_DATA
-	if (BlueprintBrush && InOwner)
-	{
-		if (BlueprintBrush->GetTypedOuter<ULevel>() != InOwner->GetTypedOuter<ULevel>())
-		{
-			BlueprintBrush->Rename(nullptr, InOwner->GetTypedOuter<ULevel>());
-		}
-		BlueprintBrush->SetOwningLandscape(InOwner);
-	}
-#endif
-}
-
-bool FLandscapeLayerBrush::IsAffectingHeightmap() const
-{
-#if WITH_EDITORONLY_DATA
-	return BlueprintBrush && BlueprintBrush->IsVisible() && BlueprintBrush->IsAffectingHeightmap();
-#else
-	return false;
-#endif
-}
-
-bool FLandscapeLayerBrush::IsAffectingWeightmapLayer(const FName& InWeightmapLayerName) const
-{
-#if WITH_EDITORONLY_DATA
-	return BlueprintBrush && BlueprintBrush->IsVisible() && BlueprintBrush->IsAffectingWeightmap() && BlueprintBrush->IsAffectingWeightmapLayer(InWeightmapLayerName);
-#else
-	return false;
-#endif
-}
-
-UTextureRenderTarget2D* FLandscapeLayerBrush::Render(bool InIsHeightmap, const FIntPoint& InLandscapeSize, UTextureRenderTarget2D* InLandscapeRenderTarget, const FName& InWeightmapLayerName)
-{
-#if WITH_EDITORONLY_DATA
-	if ((InIsHeightmap && !IsAffectingHeightmap()) ||
-		(!InIsHeightmap && !IsAffectingWeightmapLayer(InWeightmapLayerName)))
-	{
-		return nullptr;
-	}
 	if (Initialize(InLandscapeSize, InLandscapeRenderTarget))
 	{
 		TGuardValue<bool> AutoRestore(GAllowActorScriptExecutionInEditor, true);
-		return BlueprintBrush->Render(InIsHeightmap, InLandscapeRenderTarget);
+		return BPCustomBrush->Render(InIsHeightmap, InLandscapeRenderTarget);
 	}
-#endif
 	return nullptr;
 }
 
 bool FLandscapeLayerBrush::Initialize(const FIntPoint& InLandscapeSize, UTextureRenderTarget2D* InLandscapeRenderTarget)
 {
-#if WITH_EDITORONLY_DATA
-	if (BlueprintBrush && InLandscapeRenderTarget)
+	if (BPCustomBrush && InLandscapeRenderTarget)
 	{
-		if (ALandscape* Landscape = BlueprintBrush->GetOwningLandscape())
+		if (ALandscape* Landscape = BPCustomBrush->GetOwningLandscape())
 		{
 			const FIntPoint NewLandscapeRenderTargetSize = FIntPoint(InLandscapeRenderTarget->SizeX, InLandscapeRenderTarget->SizeY);
 			const FTransform NewLandscapeTransform = Landscape->GetTransform();
@@ -5689,12 +5666,11 @@ bool FLandscapeLayerBrush::Initialize(const FIntPoint& InLandscapeSize, UTexture
 				LandscapeSize = InLandscapeSize;
 				
 				TGuardValue<bool> AutoRestore(GAllowActorScriptExecutionInEditor, true);
-				BlueprintBrush->Initialize(LandscapeTransform, LandscapeSize, LandscapeRenderTargetSize);
+				BPCustomBrush->Initialize(LandscapeTransform, LandscapeSize, LandscapeRenderTargetSize);
 			}
 			return true;
 		}
 	}
-#endif
 	return false;
 }
 
diff --git a/Engine/Source/Runtime/Landscape/Public/LandscapeBlueprintBrushBase.h b/Engine/Source/Runtime/Landscape/Public/LandscapeBlueprintBrushBase.h
deleted file mode 100644
