Missing copyright
#rnx

diff --git a/Engine/Source/Runtime/Landscape/Private/LandscapeBlueprintBrushBase.cpp b/Engine/Source/Runtime/Landscape/Private/LandscapeBlueprintBrushBase.cpp
--- a/Engine/Source/Runtime/Landscape/Private/LandscapeBlueprintBrushBase.cpp
+++ b/Engine/Source/Runtime/Landscape/Private/LandscapeBlueprintBrushBase.cpp
@@ -1,3 +1,5 @@
+// Copyright 1998-2019 Epic Games, Inc. All Rights Reserved.
+
 #include "LandscapeBlueprintBrushBase.h"
 #include "CoreMinimal.h"
 #include "LandscapeProxy.h"
