Included Landscape compoments offset in landscape transform sent to Landscape BP brushes initialization
#rnx

diff --git a/Engine/Source/Runtime/Landscape/Classes/Landscape.h b/Engine/Source/Runtime/Landscape/Classes/Landscape.h
--- a/Engine/Source/Runtime/Landscape/Classes/Landscape.h
+++ b/Engine/Source/Runtime/Landscape/Classes/Landscape.h
@@ -105,7 +105,7 @@ struct FLandscapeLayerBrush
 	{}
 
 #if WITH_EDITOR
-	UTextureRenderTarget2D* Render(bool InIsHeightmap, const FIntPoint& InLandscapeSize, UTextureRenderTarget2D* InLandscapeRenderTarget, const FName& InWeightmapLayerName = NAME_None);
+	UTextureRenderTarget2D* Render(bool InIsHeightmap, const FIntRect& InLandscapeSize, UTextureRenderTarget2D* InLandscapeRenderTarget, const FName& InWeightmapLayerName = NAME_None);
 	ALandscapeBlueprintBrushBase* GetBrush() const;
 	bool IsAffectingHeightmap() const;
 	bool IsAffectingWeightmapLayer(const FName& InWeightmapLayerName) const;
@@ -115,7 +115,7 @@ struct FLandscapeLayerBrush
 private:
 
 #if WITH_EDITOR
-	bool Initialize(const FIntPoint& LandscapeSize, UTextureRenderTarget2D* InLandscapeRenderTarget);
+	bool Initialize(const FIntRect& InLandscapeExtent, UTextureRenderTarget2D* InLandscapeRenderTarget);
 #endif
 
 #if WITH_EDITORONLY_DATA
diff --git a/Engine/Source/Runtime/Landscape/Private/LandscapeEditLayers.cpp b/Engine/Source/Runtime/Landscape/Private/LandscapeEditLayers.cpp
--- a/Engine/Source/Runtime/Landscape/Private/LandscapeEditLayers.cpp
+++ b/Engine/Source/Runtime/Landscape/Private/LandscapeEditLayers.cpp
@@ -2648,13 +2648,11 @@ int32 ALandscape::RegenerateLayersHeightmaps(const TArray<ULandscapeComponent*>&
 	{
 		check(HeightmapRTList.Num() > 0);
 
-		FIntPoint MinExtend;
-		FIntPoint MaxExtend;
-		if (!Info->GetLandscapeExtent(MinExtend.X, MinExtend.Y, MaxExtend.X, MaxExtend.Y))
+		FIntRect LandscapeExtent;
+		if (!Info->GetLandscapeExtent(LandscapeExtent.Min.X, LandscapeExtent.Min.Y, LandscapeExtent.Max.X, LandscapeExtent.Max.Y))
 		{
 			return 0;
 		}
-		const FIntPoint LandscapeSize = MaxExtend - MinExtend;
 
 		// Use to compute TopLeft Component per Heightmap
 		struct FHeightmapTopLeft
@@ -2735,7 +2733,7 @@ int32 ALandscape::RegenerateLayersHeightmaps(const TArray<ULandscapeComponent*>&
 
 				for (const FHeightmapTopLeft& LayerHeightmap : LayerHeightmaps)
 				{
-					AddDeferredCopyLayersTexture(LayerHeightmap.Texture, LandscapeScratchRT1, nullptr, LayerHeightmap.TopLeftSectionBase - MinExtend);
+					AddDeferredCopyLayersTexture(LayerHeightmap.Texture, LandscapeScratchRT1, nullptr, LayerHeightmap.TopLeftSectionBase - LandscapeExtent.Min);
 				}
 			});
 
@@ -2743,13 +2741,13 @@ int32 ALandscape::RegenerateLayersHeightmaps(const TArray<ULandscapeComponent*>&
 
 			// NOTE: From this point on, we always work in non atlas, we'll convert back at the end to atlas only
 			DrawHeightmapComponentsToRenderTarget(OutputDebugName ? FString::Printf(TEXT("LS Height: %s += -> NonAtlas %s"), *Layer.Name.ToString(), *LandscapeScratchRT1->GetName(), *LandscapeScratchRT2->GetName()) : GEmptyDebugName,
-												  InLandscapeComponents, MinExtend, LandscapeScratchRT1, nullptr, LandscapeScratchRT2, ERTDrawingType::RTAtlasToNonAtlas, true, ShaderParams);
+												  InLandscapeComponents, LandscapeExtent.Min, LandscapeScratchRT1, nullptr, LandscapeScratchRT2, ERTDrawingType::RTAtlasToNonAtlas, true, ShaderParams);
 
 			ShaderParams.ApplyLayerModifiers = true;
 
 			// Combine Current layer with current result
 			DrawHeightmapComponentsToRenderTarget(OutputDebugName ? FString::Printf(TEXT("LS Height: %s += -> CombinedNonAtlas %s"), *Layer.Name.ToString(), *LandscapeScratchRT2->GetName(), *CombinedHeightmapNonAtlasRT->GetName()) : GEmptyDebugName,
-												InLandscapeComponents, MinExtend, LandscapeScratchRT2, FirstLayer ? nullptr : LandscapeScratchRT3, CombinedHeightmapNonAtlasRT, ERTDrawingType::RTNonAtlas, FirstLayer, ShaderParams);
+												InLandscapeComponents, LandscapeExtent.Min, LandscapeScratchRT2, FirstLayer ? nullptr : LandscapeScratchRT3, CombinedHeightmapNonAtlasRT, ERTDrawingType::RTNonAtlas, FirstLayer, ShaderParams);
 
 			ShaderParams.ApplyLayerModifiers = false;
 
@@ -2762,7 +2760,7 @@ int32 ALandscape::RegenerateLayersHeightmaps(const TArray<ULandscapeComponent*>&
 					// TODO: handle conversion/handling of RT not same size as internal size
 
 					FLandscapeLayerBrush& Brush = Layer.Brushes[i];
-					UTextureRenderTarget2D* BrushOutputNonAtlasRT = Brush.Render(true, LandscapeSize, CombinedHeightmapNonAtlasRT);
+					UTextureRenderTarget2D* BrushOutputNonAtlasRT = Brush.Render(true, LandscapeExtent, CombinedHeightmapNonAtlasRT);
 					if (BrushOutputNonAtlasRT == nullptr || BrushOutputNonAtlasRT->SizeX != CombinedHeightmapNonAtlasRT->SizeX || BrushOutputNonAtlasRT->SizeY != CombinedHeightmapNonAtlasRT->SizeY)
 					{
 						continue;
@@ -2788,14 +2786,14 @@ int32 ALandscape::RegenerateLayersHeightmaps(const TArray<ULandscapeComponent*>&
 		ShaderParams.GridSize = GetRootComponent()->RelativeScale3D;
 
 		DrawHeightmapComponentsToRenderTarget(OutputDebugName ? FString::Printf(TEXT("LS Height: %s = -> CombinedNonAtlasNormals : %s"), *CombinedHeightmapNonAtlasRT->GetName(), *LandscapeScratchRT1->GetName()) : GEmptyDebugName,
-											  InLandscapeComponents, MinExtend, CombinedHeightmapNonAtlasRT, nullptr, LandscapeScratchRT1, ERTDrawingType::RTNonAtlas, true, ShaderParams);
+											  InLandscapeComponents, LandscapeExtent.Min, CombinedHeightmapNonAtlasRT, nullptr, LandscapeScratchRT1, ERTDrawingType::RTNonAtlas, true, ShaderParams);
 
 		ShaderParams.GenerateNormals = false;
 
 		DrawHeightmapComponentsToRenderTarget(OutputDebugName ? FString::Printf(TEXT("LS Height: %s = -> CombinedAtlasFinal : %s"), *LandscapeScratchRT1->GetName(), *CombinedHeightmapAtlasRT->GetName()) : GEmptyDebugName,
-											  InLandscapeComponents, MinExtend, LandscapeScratchRT1, nullptr, CombinedHeightmapAtlasRT, ERTDrawingType::RTNonAtlasToAtlas, true, ShaderParams);
+											  InLandscapeComponents, LandscapeExtent.Min, LandscapeScratchRT1, nullptr, CombinedHeightmapAtlasRT, ERTDrawingType::RTNonAtlasToAtlas, true, ShaderParams);
 
-		DrawHeightmapComponentsToRenderTargetMips(InLandscapeComponents, MinExtend, CombinedHeightmapAtlasRT, true, ShaderParams);
+		DrawHeightmapComponentsToRenderTargetMips(InLandscapeComponents, LandscapeExtent.Min, CombinedHeightmapAtlasRT, true, ShaderParams);
 
 		// Copy back all Mips to original heightmap data
 		Info->ForAllLandscapeProxies([&](ALandscapeProxy* Proxy)
@@ -2805,7 +2803,7 @@ int32 ALandscape::RegenerateLayersHeightmaps(const TArray<ULandscapeComponent*>&
 
 			for (const FHeightmapTopLeft& Heightmap : Heightmaps)
 			{
-				FIntPoint TextureTopLeftSectionBase = Heightmap.TopLeftSectionBase - MinExtend;
+				FIntPoint TextureTopLeftSectionBase = Heightmap.TopLeftSectionBase - LandscapeExtent.Min;
 				uint8 MipIndex = 0;
 				
 				check(Heightmap.CPUReadback);
@@ -3447,14 +3445,11 @@ int32 ALandscape::RegenerateLayersWeightmaps(const TArray<ULandscapeComponent*>&
 
 	if (WeightmapUpdateModes || bForceRender)
 	{
-		FIntPoint MinExtend;
-		FIntPoint MaxExtend;
-
-		if (!Info->GetLandscapeExtent(MinExtend.X, MinExtend.Y, MaxExtend.X, MaxExtend.Y))
+		FIntRect LandscapeExtent;
+		if (!Info->GetLandscapeExtent(LandscapeExtent.Min.X, LandscapeExtent.Min.Y, LandscapeExtent.Max.X, LandscapeExtent.Max.Y))
 		{
 			return 0;
 		}
-		const FIntPoint LandscapeSize = MaxExtend - MinExtend;
 		
 		check(WeightmapRTList.Num() > 0);
 
@@ -3518,7 +3513,7 @@ int32 ALandscape::RegenerateLayersWeightmaps(const TArray<ULandscapeComponent*>&
 
 				// Prepare compute shader data
 				TArray<FLandscapeLayerWeightmapExtractMaterialLayersComponentData> ComponentsData;
-				PrepareComponentDataToExtractMaterialLayersCS(Layer, CurrentWeightmapToProcessIndex, MinExtend, OutputDebugName, WeightmapScratchExtractLayerTextureResource, ComponentsData, LayerInfoObjects);
+				PrepareComponentDataToExtractMaterialLayersCS(Layer, CurrentWeightmapToProcessIndex, LandscapeExtent.Min, OutputDebugName, WeightmapScratchExtractLayerTextureResource, ComponentsData, LayerInfoObjects);
 
 				HasFoundWeightmapToProcess = ComponentsData.Num() > 0;
 
@@ -3587,7 +3582,7 @@ int32 ALandscape::RegenerateLayersWeightmaps(const TArray<ULandscapeComponent*>&
 					PSShaderParams.LayerAlpha = LayerInfoObj == ALandscapeProxy::VisibilityLayer ? 1.0f : Layer.WeightmapAlpha; // visibility can't be affected by weight
 
 					DrawWeightmapComponentsToRenderTarget(OutputDebugName ? FString::Printf(TEXT("LS Weight: %s Paint: %s += -> %s"), *Layer.Name.ToString(), *LayerInfoObj->LayerName.ToString(), *LandscapeScratchRT1->GetName(), *LandscapeScratchRT2->GetName()) : GEmptyDebugName,
-														  InLandscapeComponents, MinExtend, LandscapeScratchRT1, nullptr, LandscapeScratchRT2, true, PSShaderParams, 0);
+														  InLandscapeComponents, LandscapeExtent.Min, LandscapeScratchRT1, nullptr, LandscapeScratchRT2, true, PSShaderParams, 0);
 
 					PSShaderParams.ApplyLayerModifiers = false;
 
@@ -3614,7 +3609,7 @@ int32 ALandscape::RegenerateLayersWeightmaps(const TArray<ULandscapeComponent*>&
 					}
 
 					DrawWeightmapComponentsToRenderTarget(OutputDebugName ? FString::Printf(TEXT("LS Weight: %s PaintLayer: %s, %s += -> Combined %s"), *Layer.Name.ToString(), *LayerInfoObj->LayerName.ToString(), *LandscapeScratchRT2->GetName(), *LandscapeScratchRT3->GetName()) : GEmptyDebugName,
-														  InLandscapeComponents, MinExtend, LandscapeScratchRT2, FirstLayer ? nullptr : LandscapeScratchRT1, LandscapeScratchRT3, true, PSShaderParams, 0);
+														  InLandscapeComponents, LandscapeExtent.Min, LandscapeScratchRT2, FirstLayer ? nullptr : LandscapeScratchRT1, LandscapeScratchRT3, true, PSShaderParams, 0);
 
 					PSShaderParams.OutputAsSubstractive = false;
 
@@ -3632,7 +3627,7 @@ int32 ALandscape::RegenerateLayersWeightmaps(const TArray<ULandscapeComponent*>&
 							// TODO: handle conversion/handling of RT not same size as internal size
 
 							FLandscapeLayerBrush& Brush = Layer.Brushes[i];
-							UTextureRenderTarget2D* BrushOutputRT = Brush.Render(false, LandscapeSize, LandscapeScratchRT3, LayerInfoObj->LayerName);
+							UTextureRenderTarget2D* BrushOutputRT = Brush.Render(false, LandscapeExtent, LandscapeScratchRT3, LayerInfoObj->LayerName);
 							if (BrushOutputRT == nullptr || BrushOutputRT->SizeX != LandscapeScratchRT3->SizeX || BrushOutputRT->SizeY != LandscapeScratchRT3->SizeY)
 							{
 								continue;
@@ -3711,7 +3706,7 @@ int32 ALandscape::RegenerateLayersWeightmaps(const TArray<ULandscapeComponent*>&
 			while (HasFoundWeightmapToProcess)
 			{
 				TArray<FLandscapeLayerWeightmapPackMaterialLayersComponentData> PackLayersComponentsData;
-				PrepareComponentDataToPackMaterialLayersCS(CurrentWeightmapToProcessIndex, MinExtend, OutputDebugName, InLandscapeComponents, ProcessedWeightmaps, ProcessedCPUReadbackTextures, PackLayersComponentsData);
+				PrepareComponentDataToPackMaterialLayersCS(CurrentWeightmapToProcessIndex, LandscapeExtent.Min, OutputDebugName, InLandscapeComponents, ProcessedWeightmaps, ProcessedCPUReadbackTextures, PackLayersComponentsData);
 				HasFoundWeightmapToProcess = PackLayersComponentsData.Num() > 0;
 
 				// Perform the compute shader
@@ -5656,7 +5651,7 @@ bool FLandscapeLayerBrush::IsAffectingWeightmapLayer(const FName& InWeightmapLay
 #endif
 }
 
-UTextureRenderTarget2D* FLandscapeLayerBrush::Render(bool InIsHeightmap, const FIntPoint& InLandscapeSize, UTextureRenderTarget2D* InLandscapeRenderTarget, const FName& InWeightmapLayerName)
+UTextureRenderTarget2D* FLandscapeLayerBrush::Render(bool InIsHeightmap, const FIntRect& InLandscapeExtent, UTextureRenderTarget2D* InLandscapeRenderTarget, const FName& InWeightmapLayerName)
 {
 #if WITH_EDITORONLY_DATA
 	if ((InIsHeightmap && !IsAffectingHeightmap()) ||
@@ -5664,7 +5659,7 @@ UTextureRenderTarget2D* FLandscapeLayerBrush::Render(bool InIsHeightmap, const F
 	{
 		return nullptr;
 	}
-	if (Initialize(InLandscapeSize, InLandscapeRenderTarget))
+	if (Initialize(InLandscapeExtent, InLandscapeRenderTarget))
 	{
 		TGuardValue<bool> AutoRestore(GAllowActorScriptExecutionInEditor, true);
 		return BlueprintBrush->Render(InIsHeightmap, InLandscapeRenderTarget);
@@ -5673,7 +5668,7 @@ UTextureRenderTarget2D* FLandscapeLayerBrush::Render(bool InIsHeightmap, const F
 	return nullptr;
 }
 
-bool FLandscapeLayerBrush::Initialize(const FIntPoint& InLandscapeSize, UTextureRenderTarget2D* InLandscapeRenderTarget)
+bool FLandscapeLayerBrush::Initialize(const FIntRect& InLandscapeExtent, UTextureRenderTarget2D* InLandscapeRenderTarget)
 {
 #if WITH_EDITORONLY_DATA
 	if (BlueprintBrush && InLandscapeRenderTarget)
@@ -5681,12 +5676,16 @@ bool FLandscapeLayerBrush::Initialize(const FIntPoint& InLandscapeSize, UTexture
 		if (ALandscape* Landscape = BlueprintBrush->GetOwningLandscape())
 		{
 			const FIntPoint NewLandscapeRenderTargetSize = FIntPoint(InLandscapeRenderTarget->SizeX, InLandscapeRenderTarget->SizeY);
-			const FTransform NewLandscapeTransform = Landscape->GetTransform();
-			if (!LandscapeTransform.Equals(NewLandscapeTransform) || (LandscapeSize != InLandscapeSize) || LandscapeRenderTargetSize != NewLandscapeRenderTargetSize)
+			FTransform NewLandscapeTransform = Landscape->GetTransform();
+			FVector OffsetVector(InLandscapeExtent.Min.X, InLandscapeExtent.Min.Y, 0.f);
+			FVector Translation = NewLandscapeTransform.TransformFVector4(OffsetVector);
+			NewLandscapeTransform.SetTranslation(Translation);
+			FIntPoint NewLandscapeSize = InLandscapeExtent.Max - InLandscapeExtent.Min;
+			if (!LandscapeTransform.Equals(NewLandscapeTransform) || (LandscapeSize != NewLandscapeSize) || LandscapeRenderTargetSize != NewLandscapeRenderTargetSize)
 			{
 				LandscapeTransform = NewLandscapeTransform;
 				LandscapeRenderTargetSize = NewLandscapeRenderTargetSize;
-				LandscapeSize = InLandscapeSize;
+				LandscapeSize = NewLandscapeSize;
 				
 				TGuardValue<bool> AutoRestore(GAllowActorScriptExecutionInEditor, true);
 				BlueprintBrush->Initialize(LandscapeTransform, LandscapeSize, LandscapeRenderTargetSize);
