BPBrush owning landscape is now transient and reset automatically at postload
#rnx

diff --git a/Engine/Source/Runtime/Landscape/Private/Landscape.cpp b/Engine/Source/Runtime/Landscape/Private/Landscape.cpp
--- a/Engine/Source/Runtime/Landscape/Private/Landscape.cpp
+++ b/Engine/Source/Runtime/Landscape/Private/Landscape.cpp
@@ -1898,6 +1898,10 @@ void ALandscape::PostLoad()
 	{
 		// For now, only Layer reserved for Landscape Spline uses AlphaBlend
 		Layer.BlendMode = (Layer.Guid == LandscapeSplinesTargetLayerGuid) ? LSBM_AlphaBlend : LSBM_AdditiveBlend;
+		for (FLandscapeLayerBrush& Brush : Layer.Brushes)
+		{
+			Brush.SetOwner(this);
+		}
 	}
 #endif
 
diff --git a/Engine/Source/Runtime/Landscape/Private/LandscapeBlueprintBrushBase.cpp b/Engine/Source/Runtime/Landscape/Private/LandscapeBlueprintBrushBase.cpp
--- a/Engine/Source/Runtime/Landscape/Private/LandscapeBlueprintBrushBase.cpp
+++ b/Engine/Source/Runtime/Landscape/Private/LandscapeBlueprintBrushBase.cpp
@@ -144,7 +144,8 @@ void ALandscapeBlueprintBrushBase::Destroyed()
 void ALandscapeBlueprintBrushBase::SetOwningLandscape(ALandscape* InOwningLandscape)
 {
 #if WITH_EDITORONLY_DATA
-	Modify();
+	const bool bAlwaysMarkDirty = false;
+	Modify(bAlwaysMarkDirty);
 	OwningLandscape = InOwningLandscape;
 #endif
 }
diff --git a/Engine/Source/Runtime/Landscape/Public/LandscapeBlueprintBrushBase.h b/Engine/Source/Runtime/Landscape/Public/LandscapeBlueprintBrushBase.h
--- a/Engine/Source/Runtime/Landscape/Public/LandscapeBlueprintBrushBase.h
+++ b/Engine/Source/Runtime/Landscape/Public/LandscapeBlueprintBrushBase.h
@@ -15,7 +15,7 @@ class LANDSCAPE_API ALandscapeBlueprintBrushBase : public AActor
 
 protected:
 #if WITH_EDITORONLY_DATA
-	UPROPERTY(DuplicateTransient)
+	UPROPERTY(Transient)
 	class ALandscape* OwningLandscape;
 
 	UPROPERTY(Category = "Settings", EditAnywhere)
