Landscape painted weight value now affects intensity of  layer contribution visualizer.
#rnx

diff --git a/Engine/Source/Runtime/Landscape/Private/LandscapeEditLayers.cpp b/Engine/Source/Runtime/Landscape/Private/LandscapeEditLayers.cpp
--- a/Engine/Source/Runtime/Landscape/Private/LandscapeEditLayers.cpp
+++ b/Engine/Source/Runtime/Landscape/Private/LandscapeEditLayers.cpp
@@ -4399,10 +4399,8 @@ bool ALandscape::UpdateCollisionAndClients(const TArray<ULandscapeComponent*>& I
 	bool bAllClientsUpdated = true;
 
 	const uint16 DefaultHeightValue = LandscapeDataAccess::GetTexHeight(0.f);
-	const uint8 DefaultWeightValue = 0;
-	const uint8 LayerContributingValue = 64;
+	const uint8 LayerContributingValue = UINT8_MAX;
 	TArray<uint16> HeightData;
-	TArray<uint8> WeightData;
 	TArray<uint8> LayerContributionMaskData;
 
 	for (ULandscapeComponent* LandscapeComponent : InLandscapeComponents)
@@ -4475,18 +4473,10 @@ bool ALandscape::UpdateCollisionAndClients(const TArray<ULandscapeComponent*>& I
 					}
 					else if (LandscapeEdModeInfo.ToolTarget == ELandscapeToolTargetType::Weightmap || LandscapeEdModeInfo.ToolTarget == ELandscapeToolTargetType::Visibility)
 					{
-						if (WeightData.Num() != ArraySize)
-						{
-							WeightData.AddZeroed(ArraySize);
-						}
 						ULandscapeLayerInfoObject* LayerObject = (LandscapeEdModeInfo.ToolTarget == ELandscapeToolTargetType::Visibility) ? ALandscapeProxy::VisibilityLayer : LandscapeEdModeInfo.SelectedLayerInfoObject.Get();
 						if (LayerObject)
 						{
-							LandscapeEdit.GetWeightDataFast(LayerObject, X1, Y1, X2, Y2, WeightData.GetData(), Stride);
-							for (int i = 0; i < ArraySize; ++i)
-							{
-								LayerContributionMaskData[i] = WeightData[i] != DefaultWeightValue ? LayerContributingValue : 0;
-							}
+							LandscapeEdit.GetWeightDataFast(LayerObject, X1, Y1, X2, Y2, LayerContributionMaskData.GetData(), Stride);
 							bLayerContributionWrittenData = true;
 						}
 					}
