Landscape component recreates material instances when material is destroyed by delete asset dialog
Integrated from //UE4/Dev-Editor

diff --git a/Engine/Source/Runtime/Landscape/Private/LandscapeEdit.cpp b/Engine/Source/Runtime/Landscape/Private/LandscapeEdit.cpp
--- a/Engine/Source/Runtime/Landscape/Private/LandscapeEdit.cpp
+++ b/Engine/Source/Runtime/Landscape/Private/LandscapeEdit.cpp
@@ -432,7 +432,11 @@ void ULandscapeComponent::UpdateMaterialInstances_Internal(FMaterialUpdateContex
 
 			UMaterialInstanceConstant* MobileCombinationMaterialInstance = GetCombinationMaterial(&Context, MobileWeightmapLayerAllocations, MaterialLOD, true);
 			MobileCombinationMaterialInstances[MobileMaterialIndex] = MobileCombinationMaterialInstance;
-			Context.AddMaterialInstance(MobileCombinationMaterialInstance);
+
+			if (MobileCombinationMaterialInstance != nullptr)
+			{
+				Context.AddMaterialInstance(MobileCombinationMaterialInstance);
+			}
 						
 			++MobileMaterialIndex;
 		}
@@ -4516,7 +4520,7 @@ void ULandscapeComponent::PostEditChangeProperty(FPropertyChangedEvent& Property
 	const FName PropertyName = PropertyChangedEvent.Property ? PropertyChangedEvent.Property->GetFName() : NAME_None;
 	const FName MemberPropertyName = PropertyChangedEvent.MemberProperty ? PropertyChangedEvent.MemberProperty->GetFName() : NAME_None;
 
-	if (PropertyName == FName(TEXT("OverrideMaterial")) || MemberPropertyName == FName(TEXT("OverrideMaterials")))
+	if (PropertyName == FName(TEXT("OverrideMaterial")) || MemberPropertyName == FName(TEXT("OverrideMaterials")) || MemberPropertyName == FName(TEXT("MaterialPerLOD_Key")))
 	{
 		bool RecreateMaterialInstances = true;
 
