-Landscape Layer System: Add GetLayerIndex and a return index to the CreateLayer method.

#rnx

diff --git a/Engine/Source/Runtime/Landscape/Classes/Landscape.h b/Engine/Source/Runtime/Landscape/Classes/Landscape.h
--- a/Engine/Source/Runtime/Landscape/Classes/Landscape.h
+++ b/Engine/Source/Runtime/Landscape/Classes/Landscape.h
@@ -270,7 +270,7 @@ public:
 	LANDSCAPE_API void RequestLayersContentUpdate(ELandscapeLayerUpdateMode InModeMask);
 	LANDSCAPE_API bool ReorderLayer(int32 InStartingLayerIndex, int32 InDestinationLayerIndex);
 	LANDSCAPE_API FLandscapeLayer* DuplicateLayerAndMoveBrushes(const FLandscapeLayer& InOtherLayer);
-	LANDSCAPE_API void CreateLayer(FName InName = NAME_None);
+	LANDSCAPE_API int32 CreateLayer(FName InName = NAME_None);
 	LANDSCAPE_API void CreateDefaultLayer();
 	LANDSCAPE_API void CopyOldDataToDefaultLayer();
 	LANDSCAPE_API void CopyOldDataToDefaultLayer(ALandscapeProxy* Proxy);
@@ -288,6 +288,7 @@ public:
 	LANDSCAPE_API struct FLandscapeLayer* GetLayer(int32 InLayerIndex);
 	LANDSCAPE_API const struct FLandscapeLayer* GetLayer(int32 InLayerIndex) const;
 	LANDSCAPE_API const struct FLandscapeLayer* GetLayer(const FGuid& InLayerGuid) const;
+	LANDSCAPE_API int32 GetLayerIndex(FName InLayerName) const;
 	LANDSCAPE_API void ForEachLayer(TFunctionRef<void(struct FLandscapeLayer&)> Fn);
 	LANDSCAPE_API void GetUsedPaintLayers(int32 InLayerIndex, TArray<ULandscapeLayerInfoObject*>& OutUsedLayerInfos) const;
 	LANDSCAPE_API void GetUsedPaintLayers(const FGuid& InLayerGuid, TArray<ULandscapeLayerInfoObject*>& OutUsedLayerInfos) const;
diff --git a/Engine/Source/Runtime/Landscape/Private/LandscapeEditLayers.cpp b/Engine/Source/Runtime/Landscape/Private/LandscapeEditLayers.cpp
--- a/Engine/Source/Runtime/Landscape/Private/LandscapeEditLayers.cpp
+++ b/Engine/Source/Runtime/Landscape/Private/LandscapeEditLayers.cpp
@@ -4956,6 +4956,11 @@ const FLandscapeLayer* ALandscape::GetLayer(int32 InLayerIndex) const
 	return nullptr;
 }
 
+int32 ALandscape::GetLayerIndex(FName InLayerName) const
+{
+	return LandscapeLayers.IndexOfByPredicate([InLayerName](const FLandscapeLayer& Layer) { return Layer.Name == InLayerName; });
+}
+
 const FLandscapeLayer* ALandscape::GetLayer(const FGuid& InLayerGuid) const
 {
 	return LandscapeLayers.FindByPredicate([&InLayerGuid](const FLandscapeLayer& Other) { return Other.Guid == InLayerGuid; });
@@ -5397,24 +5402,26 @@ FLandscapeLayer* ALandscape::DuplicateLayerAndMoveBrushes(const FLandscapeLayer&
 	return &LandscapeLayers[AddedIndex];
 }
 
-void ALandscape::CreateLayer(FName InName)
+int32 ALandscape::CreateLayer(FName InName)
 {
 	ULandscapeInfo* LandscapeInfo = GetLandscapeInfo();
 	if (!LandscapeInfo || IsMaxLayersReached() || !CanHaveLayersContent())
 	{
-		return;
+		return INDEX_NONE;
 	}
 
 	Modify();
 	FLandscapeLayer NewLayer;
 	NewLayer.Name = GenerateUniqueLayerName(InName);
-	LandscapeLayers.Add(NewLayer);
+	int32 LayerIndex = LandscapeLayers.Add(NewLayer);
 
 	// Create associated layer data in each landscape proxy
 	LandscapeInfo->ForAllLandscapeProxies([&NewLayer](ALandscapeProxy* Proxy)
 	{
 		Proxy->AddLayer(NewLayer.Guid);
 	});
+
+	return LayerIndex;
 }
 
 void ALandscape::AddLayersToProxy(ALandscapeProxy* InProxy)
