Back out changelist 6652492

diff --git a/Engine/Source/Editor/LandscapeEditor/Private/LandscapeEdMode.cpp b/Engine/Source/Editor/LandscapeEditor/Private/LandscapeEdMode.cpp
--- a/Engine/Source/Editor/LandscapeEditor/Private/LandscapeEdMode.cpp
+++ b/Engine/Source/Editor/LandscapeEditor/Private/LandscapeEdMode.cpp
@@ -4164,6 +4164,7 @@ ALandscape* FEdModeLandscape::ChangeComponentSetting(int32 NumComponentsX, int32
 			NewLandscape->PositiveZBoundsExtension = OldLandscape->PositiveZBoundsExtension;
 			NewLandscape->DefaultPhysMaterial = OldLandscape->DefaultPhysMaterial;
 			NewLandscape->StreamingDistanceMultiplier = OldLandscape->StreamingDistanceMultiplier;
+			NewLandscape->LandscapeHoleMaterial = OldLandscape->LandscapeHoleMaterial;
 			NewLandscape->StaticLightingResolution = OldLandscape->StaticLightingResolution;
 			NewLandscape->bCastStaticShadow = OldLandscape->bCastStaticShadow;
 			NewLandscape->bCastShadowAsTwoSided = OldLandscape->bCastShadowAsTwoSided;
diff --git a/Engine/Source/Editor/LandscapeEditor/Private/LandscapeEdModeComponentTools.cpp b/Engine/Source/Editor/LandscapeEditor/Private/LandscapeEdModeComponentTools.cpp
--- a/Engine/Source/Editor/LandscapeEditor/Private/LandscapeEdModeComponentTools.cpp
+++ b/Engine/Source/Editor/LandscapeEditor/Private/LandscapeEdModeComponentTools.cpp
@@ -313,8 +313,11 @@ public:
 	virtual bool BeginTool(FEditorViewportClient* ViewportClient, const FLandscapeToolTarget& InTarget, const FVector& InHitLocation) override
 	{
 		ALandscapeProxy* Proxy = InTarget.LandscapeInfo->GetLandscapeProxy();
-		UMaterialInterface* HoleMaterial = Proxy->GetLandscapeMaterial();
-
+		UMaterialInterface* HoleMaterial = Proxy->GetLandscapeHoleMaterial();
+		if (!HoleMaterial)
+		{
+			HoleMaterial = Proxy->GetLandscapeMaterial();
+		}
 		if (!HoleMaterial->GetMaterial()->HasAnyExpressionsInMaterialAndFunctionsOfType<UMaterialExpressionLandscapeVisibilityMask>())
 		{
 			FMessageDialog::Open(EAppMsgType::Ok,
diff --git a/Engine/Source/Editor/LandscapeEditor/Private/LandscapeEditorDetailCustomization_TargetLayers.cpp b/Engine/Source/Editor/LandscapeEditor/Private/LandscapeEditorDetailCustomization_TargetLayers.cpp
--- a/Engine/Source/Editor/LandscapeEditor/Private/LandscapeEditorDetailCustomization_TargetLayers.cpp
+++ b/Engine/Source/Editor/LandscapeEditor/Private/LandscapeEditorDetailCustomization_TargetLayers.cpp
@@ -138,9 +138,12 @@ bool FLandscapeEditorDetailCustomization_TargetLayers::ShouldShowVisibilityTip()
 		if (LandscapeEdMode->CurrentToolTarget.TargetType == ELandscapeToolTargetType::Visibility)
 		{
 			ALandscapeProxy* Proxy = LandscapeEdMode->CurrentToolTarget.LandscapeInfo->GetLandscapeProxy();
-			UMaterialInterface* HoleMaterial = Proxy->GetLandscapeMaterial();
-
-			if (HoleMaterial && !HoleMaterial->GetMaterial()->HasAnyExpressionsInMaterialAndFunctionsOfType<UMaterialExpressionLandscapeVisibilityMask>())
+			UMaterialInterface* HoleMaterial = Proxy->GetLandscapeHoleMaterial();
+			if (!HoleMaterial)
+			{
+				HoleMaterial = Proxy->GetLandscapeMaterial();
+			}
+			if (!HoleMaterial->GetMaterial()->HasAnyExpressionsInMaterialAndFunctionsOfType<UMaterialExpressionLandscapeVisibilityMask>())
 			{
 				return true;
 			}
diff --git a/Engine/Source/Runtime/Landscape/Classes/LandscapeComponent.h b/Engine/Source/Runtime/Landscape/Classes/LandscapeComponent.h
--- a/Engine/Source/Runtime/Landscape/Classes/LandscapeComponent.h
+++ b/Engine/Source/Runtime/Landscape/Classes/LandscapeComponent.h
@@ -363,8 +363,8 @@ class ULandscapeComponent : public UPrimitiveComponent
 	UPROPERTY(EditAnywhere, BlueprintReadWrite, Category=LandscapeComponent)
 	UMaterialInterface* OverrideMaterial;
 
-	UPROPERTY()
-	UMaterialInterface* OverrideHoleMaterial_DEPRECATED;
+	UPROPERTY(EditAnywhere, BlueprintReadWrite, Category=LandscapeComponent, AdvancedDisplay)
+	UMaterialInterface* OverrideHoleMaterial;
 
 	UPROPERTY(EditAnywhere, Category = LandscapeComponent)
 	TArray<FLandscapeComponentMaterialOverride> OverrideMaterials;
@@ -860,6 +860,9 @@ public:
 	/** Returns the actor's LandscapeMaterial, or the Component's OverrideLandscapeMaterial if set */
 	LANDSCAPE_API UMaterialInterface* GetLandscapeMaterial(int8 InLODIndex = INDEX_NONE) const;
 
+	/** Returns the actor's LandscapeHoleMaterial, or the Component's OverrideLandscapeHoleMaterial if set */
+	LANDSCAPE_API UMaterialInterface* GetLandscapeHoleMaterial() const;
+
 	/** Returns true if this component has visibility painted */
 	LANDSCAPE_API bool ComponentHasVisibilityPainted() const;
 
diff --git a/Engine/Source/Runtime/Landscape/Classes/LandscapeProxy.h b/Engine/Source/Runtime/Landscape/Classes/LandscapeProxy.h
--- a/Engine/Source/Runtime/Landscape/Classes/LandscapeProxy.h
+++ b/Engine/Source/Runtime/Landscape/Classes/LandscapeProxy.h
@@ -483,8 +483,9 @@ public:
 	UPROPERTY(EditAnywhere, BlueprintSetter=EditorSetLandscapeMaterial, Category=Landscape)
 	UMaterialInterface* LandscapeMaterial;
 
-	UPROPERTY()
-	UMaterialInterface* LandscapeHoleMaterial_DEPRECATED;
+	/** Material used to render landscape components with holes. If not set, LandscapeMaterial will be used (blend mode will be overridden to Masked if it is set to Opaque) */
+	UPROPERTY(EditAnywhere, Category=Landscape, AdvancedDisplay)
+	UMaterialInterface* LandscapeHoleMaterial;
 
 	UPROPERTY(EditAnywhere, Category = Landscape)
 	TArray<FLandscapeProxyMaterialOverride> LandscapeMaterialsOverride;
@@ -493,6 +494,9 @@ public:
 	UPROPERTY(Transient)
 	UMaterialInterface* PreEditLandscapeMaterial;
 
+	UPROPERTY(Transient)
+	UMaterialInterface* PreEditLandscapeHoleMaterial;
+
 	UPROPERTY(Transient)
 	TArray<FLandscapeProxyMaterialOverride> PreEditLandscapeMaterialsOverride;
 
@@ -857,6 +861,9 @@ public:
 	// Get Landscape Material assigned to this Landscape
 	virtual UMaterialInterface* GetLandscapeMaterial(int8 InLODIndex = INDEX_NONE) const;
 
+	// Get Hole Landscape Material assigned to this Landscape
+	virtual UMaterialInterface* GetLandscapeHoleMaterial() const;
+
 	// 
 	void FixupWeightmaps();
 
diff --git a/Engine/Source/Runtime/Landscape/Classes/LandscapeStreamingProxy.h b/Engine/Source/Runtime/Landscape/Classes/LandscapeStreamingProxy.h
--- a/Engine/Source/Runtime/Landscape/Classes/LandscapeStreamingProxy.h
+++ b/Engine/Source/Runtime/Landscape/Classes/LandscapeStreamingProxy.h
@@ -33,6 +33,7 @@ public:
 	virtual const ALandscape* GetLandscapeActor() const override;
 #if WITH_EDITOR
 	virtual UMaterialInterface* GetLandscapeMaterial(int8 InLODIndex = INDEX_NONE) const override;
+	virtual UMaterialInterface* GetLandscapeHoleMaterial() const override;
 #endif
 	//~ End ALandscapeBase Interface
 
diff --git a/Engine/Source/Runtime/Landscape/Private/Landscape.cpp b/Engine/Source/Runtime/Landscape/Private/Landscape.cpp
--- a/Engine/Source/Runtime/Landscape/Private/Landscape.cpp
+++ b/Engine/Source/Runtime/Landscape/Private/Landscape.cpp
@@ -60,7 +60,6 @@ Landscape.cpp: Terrain rendering
 #include "LandscapeVersion.h"
 #include "UObject/FortniteMainBranchObjectVersion.h"
 #include "LandscapeDataAccess.h"
-#include "UObject/EditorObjectVersion.h"
 
 /** Landscape stats */
 
@@ -441,24 +440,7 @@ void ULandscapeComponent::Serialize(FArchive& Ar)
 		LegacyComponentData.Data.Emplace(MapBuildDataId, LegacyMapBuildData);
 		GComponentsWithLegacyLightmaps.AddAnnotation(this, LegacyComponentData);
 	}
-	
-	if (Ar.IsLoading() && Ar.CustomVer(FEditorObjectVersion::GUID) < FEditorObjectVersion::RemoveLandscapeHoleMaterial)
-	{
-		if (OverrideHoleMaterial_DEPRECATED != nullptr)
-		{
-			if (OverrideMaterial == nullptr)
-			{
-				OverrideMaterial = OverrideHoleMaterial_DEPRECATED;
-			}
-			else
-			{
-				if (OverrideMaterial->GetBlendMode() != EBlendMode::BLEND_Masked)
-				{
-					UE_LOG(LogLandscape, Warning, TEXT("OverrideHoleMaterial was deprecated, your OverrideMaterial is not currently setup correctly to support visibility painting, please correct your material, otherwise your landscape component \"%s\" won't render as before."), *GetFullName());
-				}
-			}
-		}
-	}
+
 	if (Ar.IsLoading() && Ar.CustomVer(FFortniteMainBranchObjectVersion::GUID) < FFortniteMainBranchObjectVersion::NewLandscapeMaterialPerLOD)
 	{
 		if (MobileMaterialInterface_DEPRECATED != nullptr)
@@ -590,6 +572,20 @@ UMaterialInterface* ULandscapeComponent::GetLandscapeMaterial(int8 InLODIndex) c
 	return UMaterial::GetDefaultMaterial(MD_Surface);
 }
 
+UMaterialInterface* ULandscapeComponent::GetLandscapeHoleMaterial() const
+{
+	if (OverrideHoleMaterial)
+	{
+		return OverrideHoleMaterial;
+	}
+	ALandscapeProxy* Proxy = GetLandscapeProxy();
+	if (Proxy)
+	{
+		return Proxy->GetLandscapeHoleMaterial();
+	}
+	return nullptr;
+}
+
 bool ULandscapeComponent::ComponentHasVisibilityPainted() const
 {
 	for (const FWeightmapLayerAllocationInfo& Allocation : WeightmapLayerAllocations)
@@ -1969,24 +1965,6 @@ void ALandscapeProxy::Serialize(FArchive& Ar)
 			}
 		}
 	}
-	
-	if (Ar.IsLoading() && Ar.CustomVer(FEditorObjectVersion::GUID) < FEditorObjectVersion::RemoveLandscapeHoleMaterial)
-	{
-		if (LandscapeHoleMaterial_DEPRECATED != nullptr)
-		{
-			if (LandscapeMaterial == nullptr)
-			{
-				LandscapeMaterial = LandscapeHoleMaterial_DEPRECATED;
-			}
-			else
-			{
-				if (LandscapeMaterial->GetBlendMode() != EBlendMode::BLEND_Masked)
-				{
-					UE_LOG(LogLandscape, Warning, TEXT("LandscapeHoleMaterial was deprecated, your landscape material is not currently setup correctly to support visibility painting, please correct your material, otherwise your landscape \"%s\" won't render as before."), *GetFullName());
-				}
-			}
-		}
-	}
 }
 
 void ALandscapeProxy::AddReferencedObjects(UObject* InThis, FReferenceCollector& Collector)
@@ -2414,6 +2392,10 @@ void ALandscapeProxy::GetSharedProperties(ALandscapeProxy* Landscape)
 			LandscapeMaterial = Landscape->LandscapeMaterial;
 			LandscapeMaterialsOverride = Landscape->LandscapeMaterialsOverride;
 		}
+		if (!LandscapeHoleMaterial)
+		{
+			LandscapeHoleMaterial = Landscape->LandscapeHoleMaterial;
+		}
 		if (LandscapeMaterial == Landscape->LandscapeMaterial)
 		{
 			EditorLayerSettings = Landscape->EditorLayerSettings;
@@ -2608,6 +2590,15 @@ UMaterialInterface* ALandscapeProxy::GetLandscapeMaterial(int8 InLODIndex) const
 	return LandscapeMaterial != nullptr ? LandscapeMaterial : UMaterial::GetDefaultMaterial(MD_Surface);
 }
 
+UMaterialInterface* ALandscapeProxy::GetLandscapeHoleMaterial() const
+{
+	if (LandscapeHoleMaterial)
+	{
+		return LandscapeHoleMaterial;
+	}
+	return nullptr;
+}
+
 UMaterialInterface* ALandscapeStreamingProxy::GetLandscapeMaterial(int8 InLODIndex) const
 {
 	if (InLODIndex != INDEX_NONE)
@@ -2644,6 +2635,19 @@ UMaterialInterface* ALandscapeStreamingProxy::GetLandscapeMaterial(int8 InLODInd
 	return UMaterial::GetDefaultMaterial(MD_Surface);
 }
 
+UMaterialInterface* ALandscapeStreamingProxy::GetLandscapeHoleMaterial() const
+{
+	if (LandscapeHoleMaterial)
+	{
+		return LandscapeHoleMaterial;
+	}
+	else if (ALandscape* Landscape = LandscapeActor.Get())
+	{
+		return Landscape->GetLandscapeHoleMaterial();
+	}
+	return nullptr;
+}
+
 void ALandscape::PreSave(const class ITargetPlatform* TargetPlatform)
 {
 	Super::PreSave(TargetPlatform);
diff --git a/Engine/Source/Runtime/Landscape/Private/LandscapeEdit.cpp b/Engine/Source/Runtime/Landscape/Private/LandscapeEdit.cpp
--- a/Engine/Source/Runtime/Landscape/Private/LandscapeEdit.cpp
+++ b/Engine/Source/Runtime/Landscape/Private/LandscapeEdit.cpp
@@ -229,8 +229,9 @@ UMaterialInstanceConstant* ULandscapeComponent::GetCombinationMaterial(FMaterial
 
 	const bool bComponentHasHoles = ComponentHasVisibilityPainted();
 	UMaterialInterface* const LandscapeMaterial = GetLandscapeMaterial(InLODIndex);
-	UMaterialInterface* const MaterialToUse = LandscapeMaterial;
-	bool bOverrideBlendMode = bComponentHasHoles && LandscapeMaterial->GetBlendMode() == BLEND_Opaque;
+	UMaterialInterface* const HoleMaterial = bComponentHasHoles ? GetLandscapeHoleMaterial() : nullptr;
+	UMaterialInterface* const MaterialToUse = bComponentHasHoles && HoleMaterial ? HoleMaterial : LandscapeMaterial;
+	bool bOverrideBlendMode = bComponentHasHoles && !HoleMaterial && LandscapeMaterial->GetBlendMode() == BLEND_Opaque;
 
 	if (bOverrideBlendMode)
 	{
@@ -4118,7 +4119,7 @@ void ALandscapeStreamingProxy::PostEditChangeProperty(FPropertyChangedEvent& Pro
 			LandscapeActor = nullptr;
 		}
 	}
-	else if (PropertyName == FName(TEXT("LandscapeMaterial")) || PropertyName == FName(TEXT("LandscapeMaterialsOverride")))
+	else if (PropertyName == FName(TEXT("LandscapeMaterial")) || PropertyName == FName(TEXT("LandscapeHoleMaterial")) || PropertyName == FName(TEXT("LandscapeMaterialsOverride")))
 	{
 		bool RecreateMaterialInstances = true;
 
@@ -4170,6 +4171,7 @@ void ALandscapeStreamingProxy::PostEditChangeProperty(FPropertyChangedEvent& Pro
 void ALandscape::PreEditChange(UProperty* PropertyThatWillChange)
 {
 	PreEditLandscapeMaterial = LandscapeMaterial;
+	PreEditLandscapeHoleMaterial = LandscapeHoleMaterial;
 	PreEditLandscapeMaterialsOverride = LandscapeMaterialsOverride;
 
 	Super::PreEditChange(PropertyThatWillChange);
@@ -4188,14 +4190,14 @@ void ALandscape::PostEditChangeProperty(FPropertyChangedEvent& PropertyChangedEv
 
 	ULandscapeInfo* Info = GetLandscapeInfo();
 
-	if ((PropertyName == GET_MEMBER_NAME_CHECKED(ALandscapeProxy, LandscapeMaterial) || MemberPropertyName == GET_MEMBER_NAME_CHECKED(ALandscapeProxy, LandscapeMaterialsOverride))
+	if ((PropertyName == GET_MEMBER_NAME_CHECKED(ALandscapeProxy, LandscapeMaterial) || PropertyName == GET_MEMBER_NAME_CHECKED(ALandscapeProxy, LandscapeHoleMaterial) || MemberPropertyName == GET_MEMBER_NAME_CHECKED(ALandscapeProxy, LandscapeMaterialsOverride))
 		&& PropertyChangedEvent.ChangeType != EPropertyChangeType::ArrayAdd)
 	{
 		bool HasMaterialChanged = false;
 
 		if (PropertyChangedEvent.ChangeType != EPropertyChangeType::Interactive)
 		{
-			if (PreEditLandscapeMaterial != LandscapeMaterial || PreEditLandscapeMaterialsOverride.Num() != LandscapeMaterialsOverride.Num() || bIsPerformingInteractiveActionOnLandscapeMaterialOverride)
+			if (PreEditLandscapeMaterial != LandscapeMaterial || PreEditLandscapeHoleMaterial != LandscapeHoleMaterial || PreEditLandscapeMaterialsOverride.Num() != LandscapeMaterialsOverride.Num() || bIsPerformingInteractiveActionOnLandscapeMaterialOverride)
 			{
 				HasMaterialChanged = true;
 			}
@@ -4419,6 +4421,7 @@ void ALandscape::PostEditChangeProperty(FPropertyChangedEvent& PropertyChangedEv
 	}
 
 	PreEditLandscapeMaterial = nullptr;
+	PreEditLandscapeHoleMaterial = nullptr;
 	PreEditLandscapeMaterialsOverride.Empty();
 }
 
diff --git a/Engine/Source/Runtime/Landscape/Private/LandscapeRender.cpp b/Engine/Source/Runtime/Landscape/Private/LandscapeRender.cpp
--- a/Engine/Source/Runtime/Landscape/Private/LandscapeRender.cpp
+++ b/Engine/Source/Runtime/Landscape/Private/LandscapeRender.cpp
@@ -590,6 +590,11 @@ void ULandscapeComponent::GetUsedMaterials(TArray<UMaterialInterface*>& OutMater
 		OutMaterials.Add(OverrideMaterial);
 	}
 
+	if (OverrideHoleMaterial)
+	{
+		OutMaterials.Add(OverrideHoleMaterial);
+	}
+
 	OutMaterials.Append(MobileMaterialInterfaces);
 
 #if WITH_EDITORONLY_DATA
