Cherry-pick 7064593: Fix crash when ALandscapeStreamingProxy is used with ray tracing

diff --git a/Engine/Source/Runtime/Landscape/Private/LandscapeRender.cpp b/Engine/Source/Runtime/Landscape/Private/LandscapeRender.cpp
--- a/Engine/Source/Runtime/Landscape/Private/LandscapeRender.cpp
+++ b/Engine/Source/Runtime/Landscape/Private/LandscapeRender.cpp
@@ -2346,9 +2346,13 @@ float FLandscapeComponentSceneProxy::GetNeighborLOD(const FSceneView& InView, fl
 
 		if (Neighbor != nullptr)
 		{
-			NeighborSceneProxy = (FLandscapeComponentSceneProxy*)Neighbor;
-			NeighborBounds = NeighborSceneProxy->GetBounds();
-			NeighborMaxExtends = NeighborSceneProxy->SubsectionSizeQuads * FMath::Max(NeighborSceneProxy->GetLocalToWorld().GetScaleVector().X, NeighborSceneProxy->GetLocalToWorld().GetScaleVector().Y);
+			const ULandscapeComponent* NeighborComponent = Neighbor->GetLandscapeComponent();
+			if (NeighborComponent != nullptr && NeighborComponent->SceneProxy != nullptr)
+			{
+				NeighborSceneProxy = (FLandscapeComponentSceneProxy*)NeighborComponent->SceneProxy;
+				NeighborBounds = NeighborSceneProxy->GetBounds();
+				NeighborMaxExtends = NeighborSceneProxy->SubsectionSizeQuads * FMath::Max(NeighborSceneProxy->GetLocalToWorld().GetScaleVector().X, NeighborSceneProxy->GetLocalToWorld().GetScaleVector().Y);
+			}
 		}
 
 		if (NumSubsections > 1)
