Fixed build

diff --git a/Engine/Source/Runtime/Landscape/Private/Landscape.cpp b/Engine/Source/Runtime/Landscape/Private/Landscape.cpp
--- a/Engine/Source/Runtime/Landscape/Private/Landscape.cpp
+++ b/Engine/Source/Runtime/Landscape/Private/Landscape.cpp
@@ -60,6 +60,7 @@ Landscape.cpp: Terrain rendering
 #include "LandscapeVersion.h"
 #include "UObject/FortniteMainBranchObjectVersion.h"
 #include "LandscapeDataAccess.h"
+#include "UObject/EditorObjectVersion.h"
 
 /** Landscape stats */
 
