Fix crash in landscape spline when saving current level as new level
Integrating from //UE4/Dev-Rendering
#rnx

diff --git a/Engine/Source/Runtime/Landscape/Private/LandscapeSplines.cpp b/Engine/Source/Runtime/Landscape/Private/LandscapeSplines.cpp
--- a/Engine/Source/Runtime/Landscape/Private/LandscapeSplines.cpp
+++ b/Engine/Source/Runtime/Landscape/Private/LandscapeSplines.cpp
@@ -879,7 +879,8 @@ ULandscapeSplinesComponent* ULandscapeSplinesComponent::GetStreamingSplinesCompo
 	if (OuterLandscape &&
 		// when copy/pasting this can get called with a null guid on the parent landscape
 		// this is fine, we won't have any cross-level meshes in this case anyway
-		OuterLandscape->GetLandscapeGuid().IsValid())
+		OuterLandscape->GetLandscapeGuid().IsValid() &&
+		OuterLandscape->GetLandscapeInfo())
 	{
 		FVector LandscapeLocalLocation = GetComponentTransform().GetRelativeTransform(OuterLandscape->LandscapeActorToWorld()).TransformPosition(LocalLocation);
 		const int32 ComponentIndexX = (LandscapeLocalLocation.X >= 0.0f) ? FMath::FloorToInt(LandscapeLocalLocation.X / OuterLandscape->ComponentSizeQuads) : FMath::CeilToInt(LandscapeLocalLocation.X / OuterLandscape->ComponentSizeQuads);
