Fix Crash: Build -> Lighting only (with Landscape Grass)

diff --git a/Engine/Source/Runtime/Landscape/Private/LandscapeLight.cpp b/Engine/Source/Runtime/Landscape/Private/LandscapeLight.cpp
--- a/Engine/Source/Runtime/Landscape/Private/LandscapeLight.cpp
+++ b/Engine/Source/Runtime/Landscape/Private/LandscapeLight.cpp
@@ -94,6 +94,14 @@ void FLandscapeStaticLightingTextureMapping::Apply(FQuantizedLightmapData* Quant
 		MeshBuildData.ShadowMap = NULL;
 	}
 
+	// Flush Grass 
+	if (ALandscapeProxy* Proxy = Cast<ALandscapeProxy>(LandscapeComponent->GetOuter()))
+	{
+		TSet<ULandscapeComponent*> Components;
+		Components.Add(LandscapeComponent);
+		Proxy->FlushGrassComponents(&Components, false);
+	}
+	
 	// Build the list of statically irrelevant lights.
 	// TODO: This should be stored per LOD.
 	for (int32 LightIndex = 0; LightIndex < Mesh->RelevantLights.Num(); LightIndex++)
