Optimized landscape updates when Landscape BP Brushes are modified (went from ~6 fps to ~25 fps when editing landscape BP brushes)
#rxn

diff --git a/Engine/Source/Runtime/Landscape/Classes/LandscapeComponent.h b/Engine/Source/Runtime/Landscape/Classes/LandscapeComponent.h
--- a/Engine/Source/Runtime/Landscape/Classes/LandscapeComponent.h
+++ b/Engine/Source/Runtime/Landscape/Classes/LandscapeComponent.h
@@ -308,17 +308,22 @@ enum ELandscapeComponentUpdateFlag : uint32
 	// Will update Component clients: Navigation data, Foliage, Grass, etc.
 	Component_Update_Client = 1 << 3,
 	// Will update Component clients while editing
-	Component_Update_Client_Editing = 1 << 4
+	Component_Update_Client_Editing = 1 << 4,
+	// Will compute component approximated bounds
+	Component_Update_Approximated_Bounds = 1 << 5
 };
 
 enum ELandscapeLayerUpdateMode : uint32
 { 
 	Update_Heightmap_All = 1 << 0,
 	Update_Heightmap_Editing = 1 << 1,
-	Update_Weightmap_All = 1 << 2,
-	Update_Weightmap_Editing = 1 << 3,
+	Update_Heightmap_Editing_NoCollision = 1 << 2,
+	Update_Weightmap_All = 1 << 3,
+	Update_Weightmap_Editing = 1 << 4,
+	Update_Weightmap_Editing_NoCollision = 1 << 5,
 	Update_All = Update_Weightmap_All | Update_Heightmap_All,
 	Update_All_Editing = Update_Weightmap_Editing | Update_Heightmap_Editing,
+	Update_All_Editing_NoCollision = Update_Weightmap_Editing_NoCollision | Update_Heightmap_Editing_NoCollision,
 	// In cases where we couldn't update the clients right away this flag will be set in RegenerateLayersContent
 	Update_Client_Deferred = 1 << 4,
 	// Update landscape component clients while editing
@@ -761,7 +766,7 @@ public:
 	/**
 	 * Recalculate cached bounds using height values.
 	 */
-	LANDSCAPE_API void UpdateCachedBounds();
+	LANDSCAPE_API void UpdateCachedBounds(bool bInApproximateBounds = false);
 
 	/**
 	 * Update the MaterialInstance parameters to match the layer and weightmaps for this component
@@ -903,8 +908,8 @@ public:
 
 	LANDSCAPE_API bool IsUpdateFlagEnabledForModes(ELandscapeComponentUpdateFlag InFlag, uint32 InModeMask) const;
 	LANDSCAPE_API void ClearUpdateFlagsForModes(uint32 InModeMask);
-	LANDSCAPE_API void RequestWeightmapUpdate(bool bUpdateAll = false);
-	LANDSCAPE_API void RequestHeightmapUpdate(bool bUpdateAll = false);
+	LANDSCAPE_API void RequestWeightmapUpdate(bool bUpdateAll = false, bool bUpdateCollision = true);
+	LANDSCAPE_API void RequestHeightmapUpdate(bool bUpdateAll = false, bool bUpdateCollision = true);
 	LANDSCAPE_API void RequestEditingClientUpdate();
 	LANDSCAPE_API void RequestDeferredClientUpdate();
 	LANDSCAPE_API uint32 ComputeWeightmapsHash();
diff --git a/Engine/Source/Runtime/Landscape/Private/LandscapeBlueprintBrushBase.cpp b/Engine/Source/Runtime/Landscape/Private/LandscapeBlueprintBrushBase.cpp
--- a/Engine/Source/Runtime/Landscape/Private/LandscapeBlueprintBrushBase.cpp
+++ b/Engine/Source/Runtime/Landscape/Private/LandscapeBlueprintBrushBase.cpp
@@ -13,6 +13,7 @@ ALandscapeBlueprintBrushBase::ALandscapeBlueprintBrushBase(const FObjectInitiali
 	, AffectHeightmap(false)
 	, AffectWeightmap(false)
 	, bIsVisible(true)
+	, LastRequestLayersContentUpdateFrameNumber(0)
 #endif
 {
 #if WITH_EDITOR
@@ -32,7 +33,20 @@ void ALandscapeBlueprintBrushBase::RequestLandscapeUpdate()
 #if WITH_EDITORONLY_DATA
 	if (OwningLandscape)
 	{
-		OwningLandscape->RequestLayersContentUpdateForceAll();
+		uint32 ModeMask = 0;
+		if (AffectHeightmap)
+		{
+			ModeMask |= ELandscapeLayerUpdateMode::Update_Heightmap_Editing_NoCollision;
+		}
+		if (AffectWeightmap)
+		{
+			ModeMask |= ELandscapeLayerUpdateMode::Update_Weightmap_Editing_NoCollision;
+		}
+		if (ModeMask)
+		{
+			OwningLandscape->RequestLayersContentUpdateForceAll((ELandscapeLayerUpdateMode)ModeMask);
+			LastRequestLayersContentUpdateFrameNumber = GFrameNumber;
+		}
 	}
 #endif
 }
@@ -40,6 +54,27 @@ void ALandscapeBlueprintBrushBase::RequestLandscapeUpdate()
 #if WITH_EDITOR
 void ALandscapeBlueprintBrushBase::Tick(float DeltaSeconds)
 {
+#if WITH_EDITORONLY_DATA
+	// Avoid computing collision and client updates every frame
+	// Wait until we didn't trigger any more landscape update requests (padding of a couple of frames)
+	if (LastRequestLayersContentUpdateFrameNumber + 5 == GFrameNumber)
+	{
+		uint32 ModeMask = 0;
+		if (AffectHeightmap)
+		{
+			ModeMask |= ELandscapeLayerUpdateMode::Update_Heightmap_All;
+		}
+		if (AffectWeightmap)
+		{
+			ModeMask |= ELandscapeLayerUpdateMode::Update_Weightmap_All;
+		}
+		if (ModeMask)
+		{
+			OwningLandscape->RequestLayersContentUpdateForceAll((ELandscapeLayerUpdateMode)ModeMask);
+		}
+	}
+#endif
+
 	// Forward the Tick to the instances class of this BP
 	if (GetClass()->HasAnyClassFlags(CLASS_CompiledFromBlueprint))
 	{
@@ -104,17 +139,7 @@ void ALandscapeBlueprintBrushBase::PostEditMove(bool bFinished)
 {
 	Super::PostEditMove(bFinished);
 #if WITH_EDITORONLY_DATA
-	if (OwningLandscape)
-	{
-		if (bFinished)
-		{
-			OwningLandscape->RequestLayersContentUpdateForceAll();
-		}
-		else
-		{
-			OwningLandscape->RequestLayersContentUpdate(ELandscapeLayerUpdateMode::Update_All_Editing);
-		}
-	}
+	RequestLandscapeUpdate();
 #endif
 }
 
diff --git a/Engine/Source/Runtime/Landscape/Private/LandscapeEdit.cpp b/Engine/Source/Runtime/Landscape/Private/LandscapeEdit.cpp
--- a/Engine/Source/Runtime/Landscape/Private/LandscapeEdit.cpp
+++ b/Engine/Source/Runtime/Landscape/Private/LandscapeEdit.cpp
@@ -96,19 +96,28 @@ void ULandscapeComponent::Init(int32 InBaseX, int32 InBaseY, int32 InComponentSi
 	ULandscapeInfo* Info = GetLandscapeInfo();
 }
 
-void ULandscapeComponent::UpdateCachedBounds()
+void ULandscapeComponent::UpdateCachedBounds(bool bInApproximateBounds)
 {
-	const int32 MipLevel = 0;
-	const bool bWorkOnEditingLayer = false; // We never want to compute bounds based on anything else that final landscape layer's height data
-	FLandscapeComponentDataInterface CDI(this, MipLevel, bWorkOnEditingLayer);
-
 	// Update local-space bounding box
 	CachedLocalBox.Init();
-	for (int32 y = 0; y < ComponentSizeQuads + 1; y++)
+	if (bInApproximateBounds && GetLandscapeProxy()->HasLayersContent())
+	{
+		FVector MinBox(0, 0, LandscapeDataAccess::GetLocalHeight(0));
+		FVector MaxBox(ComponentSizeQuads + 1, ComponentSizeQuads + 1, LandscapeDataAccess::GetLocalHeight(UINT16_MAX));
+		CachedLocalBox = FBox(MinBox, MaxBox);
+	}
+	else
 	{
-		for (int32 x = 0; x < ComponentSizeQuads + 1; x++)
+		const int32 MipLevel = 0;
+		const bool bWorkOnEditingLayer = false; // We never want to compute bounds based on anything else that final landscape layer's height data
+		FLandscapeComponentDataInterface CDI(this, MipLevel, bWorkOnEditingLayer);
+
+		for (int32 y = 0; y < ComponentSizeQuads + 1; y++)
 		{
-			CachedLocalBox += CDI.GetLocalVertex(x, y);
+			for (int32 x = 0; x < ComponentSizeQuads + 1; x++)
+			{
+				CachedLocalBox += CDI.GetLocalVertex(x, y);
+			}
 		}
 	}
 	if (CachedLocalBox.GetExtent().Z == 0)
diff --git a/Engine/Source/Runtime/Landscape/Private/LandscapeEditLayers.cpp b/Engine/Source/Runtime/Landscape/Private/LandscapeEditLayers.cpp
--- a/Engine/Source/Runtime/Landscape/Private/LandscapeEditLayers.cpp
+++ b/Engine/Source/Runtime/Landscape/Private/LandscapeEditLayers.cpp
@@ -2616,7 +2616,7 @@ int32 ALandscape::RegenerateLayersHeightmaps(const TArray<ULandscapeComponent*>&
 	SCOPE_CYCLE_COUNTER(STAT_LandscapeLayersRegenerateHeightmaps);
 	ULandscapeInfo* Info = GetLandscapeInfo();
 
-	const int32 AllHeightmapUpdateModes = (ELandscapeLayerUpdateMode::Update_Heightmap_All |  ELandscapeLayerUpdateMode::Update_Heightmap_Editing);
+	const int32 AllHeightmapUpdateModes = (ELandscapeLayerUpdateMode::Update_Heightmap_All | ELandscapeLayerUpdateMode::Update_Heightmap_Editing | ELandscapeLayerUpdateMode::Update_Heightmap_Editing_NoCollision);
 	const int32 HeightmapUpdateModes = LayerContentUpdateModes & AllHeightmapUpdateModes;
 	const bool bForceRender = CVarOutputLayersDebugDrawCallName.GetValueOnAnyThread() == 1;
 
@@ -2842,9 +2842,16 @@ int32 ALandscape::RegenerateLayersHeightmaps(const TArray<ULandscapeComponent*>&
 				Component->UpdateComponentToWorld();
 
 				// Avoid updating height field if we are going to recreate collision in this update
-				bool bUpdateHeightfieldRegion = !Component->IsUpdateFlagEnabledForModes(Component_Update_Recreate_Collision, HeightmapUpdateModes);
+				bool bUpdateHeightfieldRegion = !Component->IsUpdateFlagEnabledForModes(ELandscapeComponentUpdateFlag::Component_Update_Recreate_Collision, HeightmapUpdateModes);
 				Component->UpdateCollisionData(bUpdateHeightfieldRegion);
 			}
+			else if (Component->IsUpdateFlagEnabledForModes(ELandscapeComponentUpdateFlag::Component_Update_Approximated_Bounds, HeightmapUpdateModes))
+			{
+				// Update bounds with an approximated value (real computation will be done anyways when computing collision)
+				const bool bInApproximateBounds = true;
+				Component->UpdateCachedBounds(bInApproximateBounds);
+				Component->UpdateComponentToWorld();
+			}
 		}
 
 		Info->UpdateAllAddCollisions();
@@ -3427,7 +3434,7 @@ bool ALandscape::PrepareLayersWeightmapTextureResources(bool bInWaitForStreaming
 int32 ALandscape::RegenerateLayersWeightmaps(const TArray<ULandscapeComponent*>& InLandscapeComponents, bool bInWaitForStreaming)
 {
 	SCOPE_CYCLE_COUNTER(STAT_LandscapeLayersRegenerateWeightmaps);
-	const int32 AllWeightmapUpdateModes = (ELandscapeLayerUpdateMode::Update_Weightmap_All | ELandscapeLayerUpdateMode::Update_Weightmap_Editing);
+	const int32 AllWeightmapUpdateModes = (ELandscapeLayerUpdateMode::Update_Weightmap_All | ELandscapeLayerUpdateMode::Update_Weightmap_Editing | ELandscapeLayerUpdateMode::Update_Weightmap_Editing_NoCollision);
 	const int32 WeightmapUpdateModes = LayerContentUpdateModes & AllWeightmapUpdateModes;
 	const bool bForceRender = CVarOutputLayersDebugDrawCallName.GetValueOnAnyThread() == 1;
 	
@@ -4112,12 +4119,14 @@ void ALandscape::RequestLayersContentUpdateForceAll(ELandscapeLayerUpdateMode In
 		return;
 	}
 
-	const bool bUpdateWeightmap = (InModeMask & (ELandscapeLayerUpdateMode::Update_Weightmap_All | ELandscapeLayerUpdateMode::Update_Weightmap_Editing)) != 0;
-	const bool bUpdateHeightmap = (InModeMask & (ELandscapeLayerUpdateMode::Update_Heightmap_All | ELandscapeLayerUpdateMode::Update_Heightmap_Editing)) != 0;
+	const bool bUpdateWeightmap = (InModeMask & (ELandscapeLayerUpdateMode::Update_Weightmap_All | ELandscapeLayerUpdateMode::Update_Weightmap_Editing | ELandscapeLayerUpdateMode::Update_Weightmap_Editing_NoCollision)) != 0;
+	const bool bUpdateHeightmap = (InModeMask & (ELandscapeLayerUpdateMode::Update_Heightmap_All | ELandscapeLayerUpdateMode::Update_Heightmap_Editing | ELandscapeLayerUpdateMode::Update_Heightmap_Editing_NoCollision)) != 0;
+	const bool bUpdateWeightCollision = (InModeMask & (ELandscapeLayerUpdateMode::Update_Weightmap_All | ELandscapeLayerUpdateMode::Update_Weightmap_Editing)) != 0;
+	const bool bUpdateHeightCollision = (InModeMask & (ELandscapeLayerUpdateMode::Update_Heightmap_All | ELandscapeLayerUpdateMode::Update_Heightmap_Editing)) != 0;
 	const bool bUpdateClientUdpateEditing = (InModeMask & ELandscapeLayerUpdateMode::Update_Client_Editing) != 0;
 	if (ULandscapeInfo* LandscapeInfo = GetLandscapeInfo())
 	{
-		LandscapeInfo->ForAllLandscapeProxies([bUpdateHeightmap, bUpdateWeightmap, bUpdateClientUdpateEditing](ALandscapeProxy* Proxy)
+		LandscapeInfo->ForAllLandscapeProxies([bUpdateHeightmap, bUpdateWeightmap, bUpdateHeightCollision, bUpdateWeightCollision, bUpdateClientUdpateEditing](ALandscapeProxy* Proxy)
 		{
 			if (Proxy)
 			{
@@ -4135,12 +4144,14 @@ void ALandscape::RequestLayersContentUpdateForceAll(ELandscapeLayerUpdateMode In
 
 					if (bUpdateHeightmap)
 					{
-						Component->RequestHeightmapUpdate();
+						const bool bUpdateAll = false;
+						Component->RequestHeightmapUpdate(bUpdateAll, bUpdateHeightCollision);
 					}
 
 					if (bUpdateWeightmap)
 					{
-						Component->RequestWeightmapUpdate();
+						const bool bUpdateAll = false;
+						Component->RequestWeightmapUpdate(bUpdateAll, bUpdateWeightCollision);
 					}
 
 					if (bUpdateClientUdpateEditing)
@@ -4176,7 +4187,7 @@ bool ULandscapeComponent::IsUpdateFlagEnabledForModes(ELandscapeComponentUpdateF
 			return true;
 		}
 	}
-		
+
 	if (UpdateMode & ELandscapeLayerUpdateMode::Update_Weightmap_All)
 	{
 		const uint32 WeightmapAllFlags = ELandscapeComponentUpdateFlag::Component_Update_Weightmap_Collision | ELandscapeComponentUpdateFlag::Component_Update_Recreate_Collision | ELandscapeComponentUpdateFlag::Component_Update_Client;
@@ -4213,6 +4224,15 @@ bool ULandscapeComponent::IsUpdateFlagEnabledForModes(ELandscapeComponentUpdateF
 		}
 	}
 
+	if (UpdateMode & (ELandscapeLayerUpdateMode::Update_Heightmap_Editing_NoCollision | ELandscapeLayerUpdateMode::Update_Weightmap_Editing_NoCollision))
+	{
+		const uint32 EditingNoCollisionFlags = ELandscapeComponentUpdateFlag::Component_Update_Approximated_Bounds;
+		if (EditingNoCollisionFlags & InFlag)
+		{
+			return true;
+		}
+	}
+
 	return false;
 }
 
@@ -4235,28 +4255,42 @@ void ULandscapeComponent::RequestEditingClientUpdate()
 	}
 }
 
-void ULandscapeComponent::RequestHeightmapUpdate(bool bUpdateAll)
+void ULandscapeComponent::RequestHeightmapUpdate(bool bUpdateAll, bool bUpdateCollision)
 {
-	LayerUpdateFlagPerMode |= ELandscapeLayerUpdateMode::Update_Heightmap_Editing;
-	LayerUpdateFlagPerMode |= ELandscapeLayerUpdateMode::Update_Heightmap_All;
+	if (bUpdateAll || bUpdateCollision)
+	{
+		LayerUpdateFlagPerMode |= ELandscapeLayerUpdateMode::Update_Heightmap_Editing;
+		LayerUpdateFlagPerMode |= ELandscapeLayerUpdateMode::Update_Heightmap_All;
+	}
+	else
+	{
+		LayerUpdateFlagPerMode |= ELandscapeLayerUpdateMode::Update_Heightmap_Editing_NoCollision;
+	}
 	if (ALandscape* LandscapeActor = GetLandscapeActor())
 	{
-		LandscapeActor->RequestLayersContentUpdate(ELandscapeLayerUpdateMode::Update_Heightmap_Editing);
-		if (bUpdateAll)
+		LandscapeActor->RequestLayersContentUpdate(bUpdateCollision ? ELandscapeLayerUpdateMode::Update_Heightmap_Editing : ELandscapeLayerUpdateMode::Update_Heightmap_Editing_NoCollision);
+		if (bUpdateAll || bUpdateCollision)
 		{
 			LandscapeActor->RequestLayersContentUpdate(ELandscapeLayerUpdateMode::Update_Heightmap_All);
 		}
 	}
 }
 
-void ULandscapeComponent::RequestWeightmapUpdate(bool bUpdateAll)
+void ULandscapeComponent::RequestWeightmapUpdate(bool bUpdateAll, bool bUpdateCollision)
 {
-	LayerUpdateFlagPerMode |= ELandscapeLayerUpdateMode::Update_Weightmap_Editing;
-	LayerUpdateFlagPerMode |= ELandscapeLayerUpdateMode::Update_Weightmap_All;
+	if (bUpdateAll || bUpdateCollision)
+	{
+		LayerUpdateFlagPerMode |= ELandscapeLayerUpdateMode::Update_Weightmap_Editing;
+		LayerUpdateFlagPerMode |= ELandscapeLayerUpdateMode::Update_Weightmap_All;
+	}
+	else
+	{
+		LayerUpdateFlagPerMode |= ELandscapeLayerUpdateMode::Update_Weightmap_Editing_NoCollision;
+	}
 	if (ALandscape* LandscapeActor = GetLandscapeActor())
 	{
-		LandscapeActor->RequestLayersContentUpdate(ELandscapeLayerUpdateMode::Update_Weightmap_Editing);
-		if (bUpdateAll)
+		LandscapeActor->RequestLayersContentUpdate(bUpdateCollision ? ELandscapeLayerUpdateMode::Update_Weightmap_Editing : ELandscapeLayerUpdateMode::Update_Weightmap_Editing_NoCollision);
+		if (bUpdateAll || bUpdateCollision)
 		{
 			LandscapeActor->RequestLayersContentUpdate(ELandscapeLayerUpdateMode::Update_Weightmap_All);
 		}
diff --git a/Engine/Source/Runtime/Landscape/Public/LandscapeBlueprintBrushBase.h b/Engine/Source/Runtime/Landscape/Public/LandscapeBlueprintBrushBase.h
--- a/Engine/Source/Runtime/Landscape/Public/LandscapeBlueprintBrushBase.h
+++ b/Engine/Source/Runtime/Landscape/Public/LandscapeBlueprintBrushBase.h
@@ -29,6 +29,8 @@ protected:
 
 	UPROPERTY(Transient)
 	bool bIsVisible;
+
+	uint32 LastRequestLayersContentUpdateFrameNumber;
 #endif
 
 public:
