Handle orthographic camera case for landscape LOD calculation
Integrated from //UE4/Dev-Rendering

diff --git a/Engine/Source/Runtime/Landscape/Private/LandscapeRender.cpp b/Engine/Source/Runtime/Landscape/Private/LandscapeRender.cpp
--- a/Engine/Source/Runtime/Landscape/Private/LandscapeRender.cpp
+++ b/Engine/Source/Runtime/Landscape/Private/LandscapeRender.cpp
@@ -1399,18 +1399,29 @@ void FLandscapeComponentSceneProxy::OnTransformChanged()
 
 float FLandscapeComponentSceneProxy::GetComponentScreenSize(const FSceneView* View, const FVector& Origin, float MaxExtend, float ElementRadius) const
 {
-	FVector CameraOrigin = View->ViewMatrices.GetViewOrigin();
-	FMatrix ProjMatrix = View->ViewMatrices.GetProjectionMatrix();
+	float SquaredScreenRadius = 0.f;
 
-	const FVector OriginToCamera = (CameraOrigin - Origin).GetAbs();
-	const FVector ClosestPoint = OriginToCamera.ComponentMin(FVector(MaxExtend));
-	const float DistSquared = (OriginToCamera - ClosestPoint).SizeSquared();
+	if (!View->ViewMatrices.IsPerspectiveProjection())
+	{
+		FMatrix ProjMatrix = View->ViewMatrices.GetProjectionMatrix();
+		const float ScreenMultiple = FMath::Max(0.5f * ProjMatrix.M[0][0], 0.5f * ProjMatrix.M[1][1]);
+		SquaredScreenRadius = FMath::Square(ScreenMultiple * ElementRadius);
+	}
+	else
+	{
+		FVector CameraOrigin = View->ViewMatrices.GetViewOrigin();
+		FMatrix ProjMatrix = View->ViewMatrices.GetProjectionMatrix();
+
+		const FVector OriginToCamera = (CameraOrigin - Origin).GetAbs();
+		const FVector ClosestPoint = OriginToCamera.ComponentMin(FVector(MaxExtend));
+		const float DistSquared = (OriginToCamera - ClosestPoint).SizeSquared();
 
-	// Get projection multiple accounting for view scaling.
-	const float ScreenMultiple = FMath::Max(0.5f * ProjMatrix.M[0][0], 0.5f * ProjMatrix.M[1][1]);
+		// Get projection multiple accounting for view scaling.
+		const float ScreenMultiple = FMath::Max(0.5f * ProjMatrix.M[0][0], 0.5f * ProjMatrix.M[1][1]);
 
-	// Calculate screen-space projected radius
-	float SquaredScreenRadius = FMath::Square(ScreenMultiple * ElementRadius) / FMath::Max(1.0f, DistSquared);
+		// Calculate screen-space projected radius
+		SquaredScreenRadius = FMath::Square(ScreenMultiple * ElementRadius) / FMath::Max(1.0f, DistSquared);
+	}
 
 	return SquaredScreenRadius * 2.0f;
 }
