#rnx

diff --git a/Engine/Source/Runtime/Landscape/Private/LandscapeRender.cpp b/Engine/Source/Runtime/Landscape/Private/LandscapeRender.cpp
--- a/Engine/Source/Runtime/Landscape/Private/LandscapeRender.cpp
+++ b/Engine/Source/Runtime/Landscape/Private/LandscapeRender.cpp
@@ -889,7 +889,7 @@ FLandscapeComponentSceneProxy::FLandscapeComponentSceneProxy(ULandscapeComponent
 		FWeightmapLayerAllocationInfo& Allocation = InComponent->WeightmapLayerAllocations[Idx];
 		if (Allocation.LayerInfo == ALandscapeProxy::VisibilityLayer)
 		{
-			VisibilityWeightmapTexture = WeightmapTextures[Idx];
+			VisibilityWeightmapTexture = WeightmapTextures[Allocation.WeightmapTextureIndex];
 			VisibilityWeightmapChannel = Allocation.WeightmapTextureChannel;
 			break;
 		}
