Fix up issue opening a level with landscape with feature level preview enabled

diff --git a/Engine/Source/Editor/UnrealEd/Classes/Editor/EditorEngine.h b/Engine/Source/Editor/UnrealEd/Classes/Editor/EditorEngine.h
--- a/Engine/Source/Editor/UnrealEd/Classes/Editor/EditorEngine.h
+++ b/Engine/Source/Editor/UnrealEd/Classes/Editor/EditorEngine.h
@@ -918,6 +918,7 @@ public:
 	virtual bool GetMapBuildCancelled() const override { return false; }
 	virtual void SetMapBuildCancelled(bool InCancelled) override { /* Intentionally empty. */ }
 	virtual void HandleNetworkFailure(UWorld *World, UNetDriver *NetDriver, ENetworkFailure::Type FailureType, const FString& ErrorString) override;
+	virtual ERHIFeatureLevel::Type GetDefaultWorldFeatureLevel() const override { return DefaultWorldFeatureLevel; }
 
 	FString GetPlayOnTargetPlatformName() const;
 protected:
diff --git a/Engine/Source/Runtime/Engine/Classes/Engine/Engine.h b/Engine/Source/Runtime/Engine/Classes/Engine/Engine.h
--- a/Engine/Source/Runtime/Engine/Classes/Engine/Engine.h
+++ b/Engine/Source/Runtime/Engine/Classes/Engine/Engine.h
@@ -16,7 +16,7 @@
 #include "Misc/FrameRate.h"
 #include "Subsystems/SubsystemCollection.h"
 #include "Subsystems/EngineSubsystem.h"
-
+#include "RHI.h"
 #include "Engine.generated.h"
 
 #define WITH_DYNAMIC_RESOLUTION (!UE_SERVER)
@@ -1828,6 +1828,9 @@ public:
 
 	virtual bool IsInitialized() const { return bIsInitialized; }
 
+	/** The feature used to create new worlds, by default. Overridden for feature level preview in the editor */
+	virtual ERHIFeatureLevel::Type GetDefaultWorldFeatureLevel() const { return GMaxRHIFeatureLevel;  }
+
 #if WITH_EDITOR
 
 	/** Editor-only event triggered when the actor list of the world has changed */
diff --git a/Engine/Source/Runtime/Landscape/Private/Landscape.cpp b/Engine/Source/Runtime/Landscape/Private/Landscape.cpp
--- a/Engine/Source/Runtime/Landscape/Private/Landscape.cpp
+++ b/Engine/Source/Runtime/Landscape/Private/Landscape.cpp
@@ -978,8 +978,8 @@ void ULandscapeComponent::PostLoad()
 	{
 		UWorld* World = GetWorld();
 
-		// If we're loading on a platform that doesn't require cooked data, but *only* supports OpenGL ES, generate or preload data from the DDC
-		if (!FPlatformProperties::RequiresCookedData() && ((GMaxRHIFeatureLevel <= ERHIFeatureLevel::ES3_1) || (World && (World->FeatureLevel <= ERHIFeatureLevel::ES3_1))))
+		// If we're loading on a platform that doesn't require cooked data, but defaults to a mobile feature level, generate or preload data from the DDC
+		if (!FPlatformProperties::RequiresCookedData() && ((GEngine->GetDefaultWorldFeatureLevel() <= ERHIFeatureLevel::ES3_1) || (World && (World->FeatureLevel <= ERHIFeatureLevel::ES3_1))))
 		{
 			CheckGenerateLandscapePlatformData(false, nullptr);
 		}
@@ -1375,13 +1375,7 @@ FPrimitiveSceneProxy* ULandscapeComponent::CreateSceneProxy()
 	}
 	else // i.e. (FeatureLevel <= ERHIFeatureLevel::ES3_1)
 	{
-		if (!PlatformData.HasValidRuntimeData())
-		{
-#if WITH_EDITOR
-			CheckGenerateLandscapePlatformData(false, nullptr);
-#endif
-		}
-		else
+		if (PlatformData.HasValidRuntimeData())
 		{
 			Proxy = new FLandscapeComponentSceneProxyMobile(this);
 		}
