Add Weightmap Layer Name Parameter to Blueprint Brush Render Method

#rnx

diff --git a/Engine/Source/Runtime/Landscape/Private/LandscapeEditLayers.cpp b/Engine/Source/Runtime/Landscape/Private/LandscapeEditLayers.cpp
--- a/Engine/Source/Runtime/Landscape/Private/LandscapeEditLayers.cpp
+++ b/Engine/Source/Runtime/Landscape/Private/LandscapeEditLayers.cpp
@@ -5693,7 +5693,7 @@ UTextureRenderTarget2D* FLandscapeLayerBrush::Render(bool InIsHeightmap, const F
 	if (Initialize(InLandscapeExtent, InLandscapeRenderTarget))
 	{
 		TGuardValue<bool> AutoRestore(GAllowActorScriptExecutionInEditor, true);
-		return BlueprintBrush->Render(InIsHeightmap, InLandscapeRenderTarget);
+		return BlueprintBrush->Render(InIsHeightmap, InLandscapeRenderTarget, InWeightmapLayerName);
 	}
 #endif
 	return nullptr;
diff --git a/Engine/Source/Runtime/Landscape/Public/LandscapeBlueprintBrushBase.h b/Engine/Source/Runtime/Landscape/Public/LandscapeBlueprintBrushBase.h
--- a/Engine/Source/Runtime/Landscape/Public/LandscapeBlueprintBrushBase.h
+++ b/Engine/Source/Runtime/Landscape/Public/LandscapeBlueprintBrushBase.h
@@ -35,7 +35,7 @@ protected:
 
 public:
 	UFUNCTION(BlueprintImplementableEvent)
-	UTextureRenderTarget2D* Render(bool InIsHeightmap, UTextureRenderTarget2D* InCombinedResult);
+	UTextureRenderTarget2D* Render(bool InIsHeightmap, UTextureRenderTarget2D* InCombinedResult, const FName& InWeightmapLayerName);
 
 	UFUNCTION(BlueprintImplementableEvent)
 	void Initialize(const FTransform& InLandscapeTransform, const FIntPoint& InLandscapeSize, const FIntPoint& InLandscapeRenderTargetSize);
