[NO NON-PII COMMIT MESSAGE]

diff --git a/Engine/Source/Runtime/Landscape/Private/LandscapeLight.cpp b/Engine/Source/Runtime/Landscape/Private/LandscapeLight.cpp
--- a/Engine/Source/Runtime/Landscape/Private/LandscapeLight.cpp
+++ b/Engine/Source/Runtime/Landscape/Private/LandscapeLight.cpp
@@ -662,7 +662,7 @@ FLightRayIntersection FLandscapeStaticLightingMesh::IntersectLightRay(const FVec
 
 void ULandscapeComponent::GetStaticLightingInfo(FStaticLightingPrimitiveInfo& OutPrimitiveInfo,const TArray<ULightComponent*>& InRelevantLights,const FLightingBuildOptions& Options)
 {
-	if( HasStaticLighting() )
+	if( HasStaticLighting() && GetLandscapeInfo() != nullptr)
 	{
 		const float LightMapRes = StaticLightingResolution > 0.f ? StaticLightingResolution : GetLandscapeProxy()->StaticLightingResolution;
 		const int32 LightingLOD = GetLandscapeProxy()->StaticLightingLOD;
