[NO NON-PII COMMIT MESSAGE]

diff --git a/Engine/Source/Runtime/Landscape/Private/LandscapeRender.cpp b/Engine/Source/Runtime/Landscape/Private/LandscapeRender.cpp
--- a/Engine/Source/Runtime/Landscape/Private/LandscapeRender.cpp
+++ b/Engine/Source/Runtime/Landscape/Private/LandscapeRender.cpp
@@ -717,7 +717,7 @@ FLandscapeComponentSceneProxy::FLandscapeComponentSceneProxy(ULandscapeComponent
 		HeightmapSubsectionOffsetV = ((float)(InComponent->SubsectionSizeQuads + 1) / (float)HeightmapTexture->GetSizeY());
 	}
 
-	float ScreenSizeRatioDivider = InComponent->GetLandscapeProxy()->LOD0DistributionSetting * GLandscapeLOD0DistributionScale;
+	float ScreenSizeRatioDivider = FMath::Max(InComponent->GetLandscapeProxy()->LOD0DistributionSetting * GLandscapeLOD0DistributionScale, 1.01f);
 	float CurrentScreenSizeRatio = 1.0f;
 
 	LODScreenRatioSquared.AddUninitialized(MaxLOD + 1);
@@ -725,7 +725,7 @@ FLandscapeComponentSceneProxy::FLandscapeComponentSceneProxy(ULandscapeComponent
 	// LOD 0 handling
 	LODScreenRatioSquared[0] = FMath::Square(CurrentScreenSizeRatio);
 	CurrentScreenSizeRatio /= ScreenSizeRatioDivider;
-	ScreenSizeRatioDivider = InComponent->GetLandscapeProxy()->LODDistributionSetting * GLandscapeLODDistributionScale;
+	ScreenSizeRatioDivider = FMath::Max(InComponent->GetLandscapeProxy()->LODDistributionSetting * GLandscapeLODDistributionScale, 1.01f);
 
 	// Other LODs
 	for (int32 LODIndex = 1; LODIndex <= MaxLOD; ++LODIndex) // This should ALWAYS be calculated from the component size, not user MaxLOD override
