[NO NON-PII COMMIT MESSAGE]

diff --git a/Engine/Source/Runtime/Landscape/Private/Landscape.cpp b/Engine/Source/Runtime/Landscape/Private/Landscape.cpp
--- a/Engine/Source/Runtime/Landscape/Private/Landscape.cpp
+++ b/Engine/Source/Runtime/Landscape/Private/Landscape.cpp
@@ -1051,6 +1051,14 @@ void ULandscapeComponent::GetGeneratedTexturesAndMaterialInstances(TArray<UObjec
 			}
 		}
 	}
+
+	for (UMaterialInstanceConstant* MaterialInstance : MobileCombinationMaterialInstances)
+	{
+		for (ULandscapeMaterialInstanceConstant* CurrentMIC = Cast<ULandscapeMaterialInstanceConstant>(MaterialInstance); CurrentMIC; CurrentMIC = Cast<ULandscapeMaterialInstanceConstant>(CurrentMIC->Parent))
+		{
+			OutTexturesAndMaterials.Add(CurrentMIC);
+		}
+	}	
 }
 #endif
 
