Deprecate FRawMesh, replace by MeshDescription - Landscape

diff --git a/Engine/Plugins/Experimental/ProxyLODPlugin/Source/ProxyLOD/Private/ProxyLODMeshTypes.cpp b/Engine/Plugins/Experimental/ProxyLODPlugin/Source/ProxyLOD/Private/ProxyLODMeshTypes.cpp
--- a/Engine/Plugins/Experimental/ProxyLODPlugin/Source/ProxyLOD/Private/ProxyLODMeshTypes.cpp
+++ b/Engine/Plugins/Experimental/ProxyLODPlugin/Source/ProxyLOD/Private/ProxyLODMeshTypes.cpp
@@ -53,9 +53,12 @@ FMeshDescription* GetMeshDescriptionFromMeshMergeData(const FMeshMergeData* Merg
 	FMeshDescription *MeshDescription = new FMeshDescription();
 	UStaticMesh::RegisterMeshAttributes(*MeshDescription);
 	TMap<int32, FName> MaterialMap;
-	for (int32 MaterialIndex = 0; MaterialIndex < MergeData->SourceStaticMesh->StaticMaterials.Num(); ++MaterialIndex)
+	if (MergeData->SourceStaticMesh)
 	{
-		MaterialMap.Add(MaterialIndex, MergeData->SourceStaticMesh->StaticMaterials[MaterialIndex].ImportedMaterialSlotName);
+		for (int32 MaterialIndex = 0; MaterialIndex < MergeData->SourceStaticMesh->StaticMaterials.Num(); ++MaterialIndex)
+		{
+			MaterialMap.Add(MaterialIndex, MergeData->SourceStaticMesh->StaticMaterials[MaterialIndex].ImportedMaterialSlotName);
+		}
 	}
 	FMeshDescriptionOperations::ConvertFromRawMesh(*(MergeData->RawMesh), *MeshDescription, MaterialMap);
 	return MeshDescription;
diff --git a/Engine/Source/Developer/MeshMergeUtilities/MeshMergeUtilities.Build.cs b/Engine/Source/Developer/MeshMergeUtilities/MeshMergeUtilities.Build.cs
--- a/Engine/Source/Developer/MeshMergeUtilities/MeshMergeUtilities.Build.cs
+++ b/Engine/Source/Developer/MeshMergeUtilities/MeshMergeUtilities.Build.cs
@@ -31,6 +31,8 @@ public class MeshMergeUtilities : ModuleRules
         PublicDependencyModuleNames.AddRange(
 			new string [] {
 				"RawMesh",
+                "MeshDescription",
+                "MeshDescriptionOperations"
             }
 		);
 
diff --git a/Engine/Source/Developer/MeshMergeUtilities/Private/MeshMergeHelpers.cpp b/Engine/Source/Developer/MeshMergeUtilities/Private/MeshMergeHelpers.cpp
--- a/Engine/Source/Developer/MeshMergeUtilities/Private/MeshMergeHelpers.cpp
+++ b/Engine/Source/Developer/MeshMergeUtilities/Private/MeshMergeHelpers.cpp
@@ -8,6 +8,11 @@
 #include "MaterialOptions.h"
 #include "RawMesh.h"
 
+#include "MeshDescription.h"
+#include "MeshAttributes.h"
+#include "MeshAttributeArray.h"
+#include "MeshDescriptionOperations.h"
+
 #include "Misc/PackageName.h"
 #include "MaterialUtilities.h"
 #include "Components/SkeletalMeshComponent.h"
@@ -799,11 +804,21 @@ void FMeshMergeHelpers::RetrieveCullingLandscapeAndVolumes(UWorld* InWorld, cons
 	for (ALandscapeProxy* Landscape : LandscapeActors)
 	{
 		// Export the landscape to raw mesh format
-		FRawMesh* LandscapeRawMesh = new FRawMesh();
+		FMeshDescription MeshDescription;
+		UStaticMesh::RegisterMeshAttributes(MeshDescription);
 		FBoxSphereBounds LandscapeBounds = EstimatedMeshProxyBounds;
-		Landscape->ExportToRawMesh(LandscapeExportLOD, *LandscapeRawMesh, LandscapeBounds);
-		if (LandscapeRawMesh->VertexPositions.Num())
+		Landscape->ExportToRawMesh(LandscapeExportLOD, MeshDescription, LandscapeBounds);
+		if (MeshDescription.Vertices().Num())
 		{
+			TMap<FName, int32> MaterialMap;
+			TPolygonGroupAttributesConstRef<FName> PolygonGroupNames = MeshDescription.PolygonGroupAttributes().GetAttributesRef<FName>(MeshAttribute::PolygonGroup::ImportedMaterialSlotName);
+			int32 MaterialIndex = 0;
+			for (const FPolygonGroupID PolygonGroupID : MeshDescription.PolygonGroups().GetElementIDs())
+			{
+				MaterialMap.Add(PolygonGroupNames[PolygonGroupID], MaterialIndex++);
+			}
+			FRawMesh* LandscapeRawMesh = new FRawMesh();
+			FMeshDescriptionOperations::ConvertToRawMesh(MeshDescription, *LandscapeRawMesh, MaterialMap);
 			CullingRawMeshes.Add(LandscapeRawMesh);
 		}
 	}
diff --git a/Engine/Source/Editor/WorldBrowser/Private/Tiles/WorldTileCollectionModel.cpp b/Engine/Source/Editor/WorldBrowser/Private/Tiles/WorldTileCollectionModel.cpp
--- a/Engine/Source/Editor/WorldBrowser/Private/Tiles/WorldTileCollectionModel.cpp
+++ b/Engine/Source/Editor/WorldBrowser/Private/Tiles/WorldTileCollectionModel.cpp
@@ -22,7 +22,11 @@
 #include "IDetailsView.h"
 #include "AssetRegistryModule.h"
 #include "MeshUtilities.h"
-#include "RawMesh.h"
+
+#include "MeshDescription.h"
+#include "MeshAttributes.h"
+#include "MeshAttributeArray.h"
+
 #include "Materials/Material.h"
 #include "MaterialUtilities.h"
 #include "LandscapeLayerInfoObject.h"
@@ -2064,22 +2068,14 @@ bool FWorldTileCollectionModel::GenerateLODLevels(FLevelModelList InLevelList, i
 		{
 			GWarn->StatusUpdate(LandscapeActorIndex, LandscapeActors.Num(), LOCTEXT("ExportingLandscape", "Exporting Landscape Actors"));
 
-			FRawMesh LandscapeRawMesh;
 			FFlattenMaterial LandscapeFlattenMaterial;
 			FVector LandscapeWorldLocation = Landscape->GetActorLocation();
-			
+
 			int32 LandscapeLOD = SimplificationDetails.LandscapeExportLOD;
 			if (!SimplificationDetails.bOverrideLandscapeExportLOD)
 			{
 				LandscapeLOD = Landscape->MaxLODLevel >= 0 ? Landscape->MaxLODLevel : FMath::CeilLogTwo(Landscape->SubsectionSizeQuads + 1) - 1;
 			}
-		
-			Landscape->ExportToRawMesh(LandscapeLOD, LandscapeRawMesh);
-		
-			for (FVector& VertexPos : LandscapeRawMesh.VertexPositions)
-			{
-				VertexPos-= LandscapeWorldLocation;
-			}
 
 			// Filter out primitives for landscape texture flattening
 			TSet<FPrimitiveComponentId> PrimitivesToHide;
@@ -2138,10 +2134,8 @@ bool FWorldTileCollectionModel::GenerateLODLevels(FLevelModelList InLevelList, i
 			}
 
 			auto StaticMesh = NewObject<UStaticMesh>(MeshOuter, *LandscapeMeshAssetName, RF_Public | RF_Standalone);
-			StaticMesh->InitResources();
 			{
-				FString OutputPath = StaticMesh->GetPathName();
-
+				StaticMesh->InitResources();
 				// make sure it has a new lighting guid
 				StaticMesh->LightingGuid = FGuid::NewGuid();
 
@@ -2156,14 +2150,26 @@ bool FWorldTileCollectionModel::GenerateLODLevels(FLevelModelList InLevelList, i
 				SrcModel.BuildSettings.bRemoveDegenerates = false;
 				SrcModel.BuildSettings.bUseHighPrecisionTangentBasis = false;
 				SrcModel.BuildSettings.bUseFullPrecisionUVs = false;
-				SrcModel.SaveRawMesh(LandscapeRawMesh);
 
 				//Assign the proxy material to the static mesh
 				StaticMesh->StaticMaterials.Add(FStaticMaterial(StaticLandscapeMaterial));
 
 				//Set the Imported version before calling the build
 				StaticMesh->ImportVersion = EImportStaticMeshVersion::LastVersion;
+			}
+			FMeshDescription& LandscapeRawMesh = *(StaticMesh->CreateOriginalMeshDescription(0));
+		
+			Landscape->ExportToRawMesh(LandscapeLOD, LandscapeRawMesh);
+		
+			TVertexAttributesRef<FVector> VertexPositions = LandscapeRawMesh.VertexAttributes().GetAttributesRef<FVector>(MeshAttribute::Vertex::Position);
+			for (const FVertexID& VertexID : LandscapeRawMesh.Vertices().GetElementIDs())
+			{
+				VertexPositions[VertexID] -= LandscapeWorldLocation;
+			}
 
+			//Commit raw mesh and build the staticmesh
+			{
+				StaticMesh->CommitOriginalMeshDescription(0);
 				StaticMesh->Build();
 				StaticMesh->PostEditChange();
 			}
diff --git a/Engine/Source/Editor/WorldBrowser/WorldBrowser.Build.cs b/Engine/Source/Editor/WorldBrowser/WorldBrowser.Build.cs
--- a/Engine/Source/Editor/WorldBrowser/WorldBrowser.Build.cs
+++ b/Engine/Source/Editor/WorldBrowser/WorldBrowser.Build.cs
@@ -43,7 +43,7 @@ public class WorldBrowser : ModuleRules
                 "MainFrame",
                 "SourceControl",
 				"SourceControlWindows",
-                "RawMesh",
+                "MeshDescription",
                 "LandscapeEditor",
                 "FoliageEdit",
                 "ImageWrapper",
diff --git a/Engine/Source/Runtime/Engine/Private/StaticMesh.cpp b/Engine/Source/Runtime/Engine/Private/StaticMesh.cpp
--- a/Engine/Source/Runtime/Engine/Private/StaticMesh.cpp
+++ b/Engine/Source/Runtime/Engine/Private/StaticMesh.cpp
@@ -2731,8 +2731,8 @@ void UStaticMesh::LoadMeshDescriptions()
 				// If the MeshDescriptions are out of sync with the SourceModels RawMesh, perform a conversion here.
 				// @todo: once all tools are ported, we can replace this with a check() instead.
 				FMeshDescription* MeshDescription = MeshDescriptions->Create(LodIndex);
-				SourceModels[LodIndex].OriginalMeshDescription = MeshDescription;
 				RegisterMeshAttributes(*MeshDescription);
+				SourceModels[LodIndex].OriginalMeshDescription = MeshDescription;
 
 				FRawMesh LodRawMesh;
 				SourceModels[LodIndex].LoadRawMesh(LodRawMesh);
@@ -2843,8 +2843,9 @@ FMeshDescription* UStaticMesh::GetOriginalMeshDescription(int32 LodIndex)
 			// If the MeshDescriptions are out of sync with the SourceModels RawMesh, perform a conversion here.
 			// @todo: once all tools are ported, we can replace this with the disabled check() above.
 			FMeshDescription* MeshDescription = MeshDescriptions->Create(LodIndex);
-			SourceModels[LodIndex].OriginalMeshDescription = MeshDescription;
 			RegisterMeshAttributes(*MeshDescription);
+			SourceModels[LodIndex].OriginalMeshDescription = MeshDescription;
+			
 
 			FRawMesh LodRawMesh;
 			SourceModels[LodIndex].LoadRawMesh(LodRawMesh);
@@ -2871,6 +2872,7 @@ FMeshDescription* UStaticMesh::CreateOriginalMeshDescription(int32 LodIndex)
 
 		check(MeshDescriptions->Num() == SourceModels.Num());
 		FMeshDescription* MeshDescription = MeshDescriptions->Create(LodIndex);
+		UStaticMesh::RegisterMeshAttributes(*MeshDescription);
 		SourceModels[LodIndex].OriginalMeshDescription = MeshDescription;
 		return MeshDescription;
 	}
diff --git a/Engine/Source/Runtime/Landscape/Classes/LandscapeProxy.h b/Engine/Source/Runtime/Landscape/Classes/LandscapeProxy.h
--- a/Engine/Source/Runtime/Landscape/Classes/LandscapeProxy.h
+++ b/Engine/Source/Runtime/Landscape/Classes/LandscapeProxy.h
@@ -30,7 +30,7 @@ class USplineComponent;
 class UTexture2D;
 struct FAsyncGrassBuilder;
 struct FLandscapeInfoLayerSettings;
-struct FRawMesh;
+struct FMeshDescription;
 enum class ENavDataGatheringMode : uint8;
 
 /** Structure storing channel usage for weightmap textures */
@@ -798,7 +798,7 @@ public:
 	 * @param OutRawMesh - Resulting raw mesh
 	 * @return true if successful
 	 */
-	LANDSCAPE_API bool ExportToRawMesh(int32 InExportLOD, FRawMesh& OutRawMesh) const;
+	LANDSCAPE_API bool ExportToRawMesh(int32 InExportLOD, FMeshDescription& OutRawMesh) const;
 
 
 	/**
@@ -809,7 +809,7 @@ public:
 	* @param InBounds - Box/Sphere bounds which limit the geometry exported out into OutRawMesh
 	* @return true if successful
 	*/
-	LANDSCAPE_API bool ExportToRawMesh(int32 InExportLOD, FRawMesh& OutRawMesh, const FBoxSphereBounds& InBounds ) const;
+	LANDSCAPE_API bool ExportToRawMesh(int32 InExportLOD, FMeshDescription& OutRawMesh, const FBoxSphereBounds& InBounds, bool bIgnoreBounds = false) const;
 
 	/** Generate platform data if it's missing or outdated */
 	LANDSCAPE_API void CheckGenerateLandscapePlatformData(bool bIsCooking, const ITargetPlatform* TargetPlatform);
diff --git a/Engine/Source/Runtime/Landscape/Landscape.Build.cs b/Engine/Source/Runtime/Landscape/Landscape.Build.cs
--- a/Engine/Source/Runtime/Landscape/Landscape.Build.cs
+++ b/Engine/Source/Runtime/Landscape/Landscape.Build.cs
@@ -46,7 +46,8 @@ public class Landscape : ModuleRules
 		{
 			PrivateDependencyModuleNames.AddRange(
 				new string[] {
-					"RawMesh"
+					"MeshDescription",
+                    "MeshUtilitiesCommon"
 				}
 			);
 		}
diff --git a/Engine/Source/Runtime/Landscape/Private/LandscapeEdit.cpp b/Engine/Source/Runtime/Landscape/Private/LandscapeEdit.cpp
--- a/Engine/Source/Runtime/Landscape/Private/LandscapeEdit.cpp
+++ b/Engine/Source/Runtime/Landscape/Private/LandscapeEdit.cpp
@@ -40,7 +40,11 @@ LandscapeEdit.cpp: Landscape editing
 #include "LandscapeSplinesComponent.h"
 #include "Serialization/MemoryWriter.h"
 #if WITH_EDITOR
-#include "RawMesh.h"
+#include "MeshDescription.h"
+#include "MeshAttributes.h"
+#include "MeshAttributeArray.h"
+#include "MeshUtilitiesCommon.h"
+
 #include "EngineUtils.h"
 #include "Framework/Notifications/NotificationManager.h"
 #include "Widgets/Notifications/SNotificationList.h"
@@ -2615,7 +2619,13 @@ LANDSCAPE_API void ALandscapeProxy::Import(
 	GWarn->EndSlowTask();
 }
 
-bool ALandscapeProxy::ExportToRawMesh(int32 InExportLOD, FRawMesh& OutRawMesh) const
+bool ALandscapeProxy::ExportToRawMesh(int32 InExportLOD, FMeshDescription& OutRawMesh) const
+{
+	FBoxSphereBounds GarbageBounds;
+	return ExportToRawMesh(InExportLOD, OutRawMesh, GarbageBounds, true);
+}
+
+bool ALandscapeProxy::ExportToRawMesh(int32 InExportLOD, FMeshDescription& OutRawMesh, const FBoxSphereBounds& InBounds, bool bIgnoreBounds /*= false*/) const
 {
 	TInlineComponentArray<ULandscapeComponent*> RegisteredLandscapeComponents;
 	GetComponents<ULandscapeComponent>(RegisteredLandscapeComponents);
@@ -2623,6 +2633,21 @@ bool ALandscapeProxy::ExportToRawMesh(int32 InExportLOD, FRawMesh& OutRawMesh) c
 	const FIntRect LandscapeSectionRect = GetBoundingRect();
 	const FVector2D LandscapeUVScale = FVector2D(1.0f, 1.0f) / FVector2D(LandscapeSectionRect.Size());
 
+	TVertexAttributesRef<FVector> VertexPositions = OutRawMesh.VertexAttributes().GetAttributesRef<FVector>(MeshAttribute::Vertex::Position);
+	TEdgeAttributesRef<bool> EdgeHardnesses = OutRawMesh.EdgeAttributes().GetAttributesRef<bool>(MeshAttribute::Edge::IsHard);
+	TEdgeAttributesRef<float> EdgeCreaseSharpnesses = OutRawMesh.EdgeAttributes().GetAttributesRef<float>(MeshAttribute::Edge::CreaseSharpness);
+	TPolygonGroupAttributesRef<FName> PolygonGroupImportedMaterialSlotNames = OutRawMesh.PolygonGroupAttributes().GetAttributesRef<FName>(MeshAttribute::PolygonGroup::ImportedMaterialSlotName);
+	TVertexInstanceAttributesRef<FVector> VertexInstanceNormals = OutRawMesh.VertexInstanceAttributes().GetAttributesRef<FVector>(MeshAttribute::VertexInstance::Normal);
+	TVertexInstanceAttributesRef<FVector> VertexInstanceTangents = OutRawMesh.VertexInstanceAttributes().GetAttributesRef<FVector>(MeshAttribute::VertexInstance::Tangent);
+	TVertexInstanceAttributesRef<float> VertexInstanceBinormalSigns = OutRawMesh.VertexInstanceAttributes().GetAttributesRef<float>(MeshAttribute::VertexInstance::BinormalSign);
+	TVertexInstanceAttributesRef<FVector4> VertexInstanceColors = OutRawMesh.VertexInstanceAttributes().GetAttributesRef<FVector4>(MeshAttribute::VertexInstance::Color);
+	TVertexInstanceAttributesRef<FVector2D> VertexInstanceUVs = OutRawMesh.VertexInstanceAttributes().GetAttributesRef<FVector2D>(MeshAttribute::VertexInstance::TextureCoordinate);
+
+	if (VertexInstanceUVs.GetNumIndices() < 2)
+	{
+		VertexInstanceUVs.SetNumIndices(2);
+	}
+
 	// User specified LOD to export
 	int32 LandscapeLODToExport = ExportLOD;
 	if (InExportLOD != INDEX_NONE)
@@ -2634,6 +2659,13 @@ bool ALandscapeProxy::ExportToRawMesh(int32 InExportLOD, FRawMesh& OutRawMesh) c
 	for (auto It = RegisteredLandscapeComponents.CreateConstIterator(); It; ++It)
 	{
 		ULandscapeComponent* Component = (*It);
+
+		// Early out if the Landscape bounds and given bounds do not overlap at all
+		if (!bIgnoreBounds && !FBoxSphereBounds::SpheresIntersect(Component->Bounds, InBounds))
+		{
+			continue;
+		}
+
 		FLandscapeComponentDataInterface CDI(Component, LandscapeLODToExport);
 		const int32 ComponentSizeQuadsLOD = ((Component->ComponentSizeQuads + 1) >> LandscapeLODToExport) - 1;
 		const int32 SubsectionSizeQuadsLOD = ((Component->SubsectionSizeQuads + 1) >> LandscapeLODToExport) - 1;
@@ -2643,22 +2675,25 @@ bool ALandscapeProxy::ExportToRawMesh(int32 InExportLOD, FRawMesh& OutRawMesh) c
 
 		const int32 NumFaces = FMath::Square(ComponentSizeQuadsLOD) * 2;
 		const int32 NumVertices = NumFaces * 3;
-		const int32 VerticesOffset = OutRawMesh.VertexPositions.Num();
-		const int32 IndicesOffset = OutRawMesh.WedgeIndices.Num();
-
-		//
-		OutRawMesh.FaceMaterialIndices.AddZeroed(NumFaces);
-		OutRawMesh.FaceSmoothingMasks.AddZeroed(NumFaces);
 
-		OutRawMesh.VertexPositions.AddZeroed(NumVertices);
-		OutRawMesh.WedgeIndices.AddZeroed(NumVertices);
-		OutRawMesh.WedgeTangentX.AddZeroed(NumVertices);
-		OutRawMesh.WedgeTangentY.AddZeroed(NumVertices);
-		OutRawMesh.WedgeTangentZ.AddZeroed(NumVertices);
-		OutRawMesh.WedgeTexCoords[0].AddZeroed(NumVertices);
+		OutRawMesh.ReserveNewVertices(NumVertices);
+		OutRawMesh.ReserveNewPolygons(NumFaces);
+		OutRawMesh.ReserveNewVertexInstances(NumVertices);
+		OutRawMesh.ReserveNewEdges(NumVertices);
 
+		FPolygonGroupID PolygonGroupID = FPolygonGroupID::Invalid;
+		if (OutRawMesh.PolygonGroups().Num() < 1)
+		{
+			PolygonGroupID = OutRawMesh.CreatePolygonGroup();
+			PolygonGroupImportedMaterialSlotNames[PolygonGroupID] = FName(TEXT("LandscapeMat_0"));
+		}
+		else
+		{
+			PolygonGroupID = OutRawMesh.PolygonGroups().GetFirstValidID();
+		}
 
 		// Check if there are any holes
+		const int32 VisThreshold = 170;
 		TArray<uint8> VisDataMap;
 
 		for (int32 AllocIdx = 0; AllocIdx < Component->WeightmapLayerAllocations.Num(); AllocIdx++)
@@ -2682,156 +2717,71 @@ bool ALandscapeProxy::ExportToRawMesh(int32 InExportLOD, FRawMesh& OutRawMesh) c
 			FIntPoint(1, 0),
 		};
 
-		const int32 VisThreshold = 170;
 		const int32 WeightMapSize = (SubsectionSizeQuadsLOD + 1) * Component->NumSubsections;
-		uint32* Faces = OutRawMesh.WedgeIndices.GetData() + IndicesOffset;
 
-		// Export verts
-		int32 VertexIdx = VerticesOffset;
+		const float SquaredSphereRadius = FMath::Square(InBounds.SphereRadius);
+
+		//We need to not duplicate the vertex position, so we use the FIndexAndZ to achieve fast result
+		TArray<FIndexAndZ> VertIndexAndZ;
+		VertIndexAndZ.Reserve(ComponentSizeQuadsLOD*ComponentSizeQuadsLOD*ARRAY_COUNT(QuadPattern));
+		int32 CurrentIndex = 0;
+		TMap<int32, FVector> IndexToPosition;
+		IndexToPosition.Reserve(ComponentSizeQuadsLOD*ComponentSizeQuadsLOD*ARRAY_COUNT(QuadPattern));
 		for (int32 y = 0; y < ComponentSizeQuadsLOD; y++)
 		{
 			for (int32 x = 0; x < ComponentSizeQuadsLOD; x++)
 			{
-				// Fill indices
-				{
-
-					// Whether this vertex is in hole
-					bool bInvisible = false;
-					if (VisDataMap.Num())
-					{
-						int32 TexelX, TexelY;
-						CDI.VertexXYToTexelXY(x, y, TexelX, TexelY);
-						bInvisible = (VisDataMap[CDI.TexelXYToIndex(TexelX, TexelY)] >= VisThreshold);
-					}
-
-					// triangulation matches FLandscapeIndexBuffer constructor
-					Faces[0] = VertexIdx;
-					Faces[1] = bInvisible ? Faces[0] : VertexIdx + 1;
-					Faces[2] = bInvisible ? Faces[0] : VertexIdx + 2;
-					Faces += 3;
-
-					Faces[0] = VertexIdx + 3;
-					Faces[1] = bInvisible ? Faces[0] : VertexIdx + 4;
-					Faces[2] = bInvisible ? Faces[0] : VertexIdx + 5;
-					Faces += 3;
-				}
-
-				// Fill vertices
 				for (int32 i = 0; i < ARRAY_COUNT(QuadPattern); i++)
 				{
 					int32 VertexX = x + QuadPattern[i].X;
 					int32 VertexY = y + QuadPattern[i].Y;
-					FVector LocalVertexPos = CDI.GetWorldVertex(VertexX, VertexY);
-
-					FVector LocalTangentX, LocalTangentY, LocalTangentZ;
-					CDI.GetLocalTangentVectors(VertexX, VertexY, LocalTangentX, LocalTangentY, LocalTangentZ);
+					FVector Position = CDI.GetWorldVertex(VertexX, VertexY);
 
-					OutRawMesh.VertexPositions[VertexIdx] = LocalVertexPos;
-					OutRawMesh.WedgeTangentX[VertexIdx] = LocalTangentX;
-					OutRawMesh.WedgeTangentY[VertexIdx] = LocalTangentY;
-					OutRawMesh.WedgeTangentZ[VertexIdx] = LocalTangentZ;
-
-					OutRawMesh.WedgeTexCoords[0][VertexIdx] = (ComponentUVOffsetLOD + FVector2D(VertexX, VertexY))*ComponentUVScaleLOD;
-
-					VertexIdx++;
+					// If at least one vertex is within the given bounds we should process the quad  
+					new(VertIndexAndZ)FIndexAndZ(CurrentIndex, Position);
+					IndexToPosition.Add(CurrentIndex, Position);
+					CurrentIndex++;
 				}
 			}
 		}
-	}
-
-	// Add lightmap UVs
-	OutRawMesh.WedgeTexCoords[1].Append(OutRawMesh.WedgeTexCoords[0]);
-
-	return true;
-}
-
-bool ALandscapeProxy::ExportToRawMesh(int32 InExportLOD, FRawMesh& OutRawMesh, const FBoxSphereBounds& InBounds) const
-{
-	TInlineComponentArray<ULandscapeComponent*> RegisteredLandscapeComponents;
-	GetComponents<ULandscapeComponent>(RegisteredLandscapeComponents);
-
-	const FIntRect LandscapeSectionRect = GetBoundingRect();
-	const FVector2D LandscapeUVScale = FVector2D(1.0f, 1.0f) / FVector2D(LandscapeSectionRect.Size());
-
-	// User specified LOD to export
-	int32 LandscapeLODToExport = ExportLOD;
-	if (InExportLOD != INDEX_NONE)
-	{
-		LandscapeLODToExport = FMath::Clamp<int32>(InExportLOD, 0, FMath::CeilLogTwo(SubsectionSizeQuads + 1) - 1);
-	}
-
-	// Export data for each component
-	for (auto It = RegisteredLandscapeComponents.CreateConstIterator(); It; ++It)
-	{
-		ULandscapeComponent* Component = (*It);
-
-		// Early out if the Landscape bounds and given bounds do not overlap at all
-		if (!FBoxSphereBounds::SpheresIntersect(Component->Bounds, InBounds))
-		{
-			continue;
-		}
-
-		FLandscapeComponentDataInterface CDI(Component, LandscapeLODToExport);
-		const int32 ComponentSizeQuadsLOD = ((Component->ComponentSizeQuads + 1) >> LandscapeLODToExport) - 1;
-		const int32 SubsectionSizeQuadsLOD = ((Component->SubsectionSizeQuads + 1) >> LandscapeLODToExport) - 1;
-		const FIntPoint ComponentOffsetQuads = Component->GetSectionBase() - LandscapeSectionOffset - LandscapeSectionRect.Min;
-		const FVector2D ComponentUVOffsetLOD = FVector2D(ComponentOffsetQuads)*((float)ComponentSizeQuadsLOD / ComponentSizeQuads);
-		const FVector2D ComponentUVScaleLOD = LandscapeUVScale*((float)ComponentSizeQuads / ComponentSizeQuadsLOD);
-
-		const int32 NumFaces = FMath::Square(ComponentSizeQuadsLOD) * 2;
-		const int32 TotalCurrentFaces = OutRawMesh.FaceSmoothingMasks.Num() + NumFaces;
-		const int32 NumVertices = NumFaces * 3;
-		const int32 TotalCurrentVertices = TotalCurrentFaces * 3;
-		const int32 VerticesOffset = OutRawMesh.VertexPositions.Num();
-		const int32 IndicesOffset = OutRawMesh.WedgeIndices.Num();
+		// Sort the vertices by z value
+		VertIndexAndZ.Sort(FCompareIndexAndZ());
 
-		OutRawMesh.FaceMaterialIndices.Reserve(TotalCurrentFaces);
-		OutRawMesh.FaceSmoothingMasks.Reserve(TotalCurrentFaces);
-
-		OutRawMesh.VertexPositions.Reserve(TotalCurrentVertices);
-		OutRawMesh.WedgeIndices.Reserve(TotalCurrentVertices);
-		OutRawMesh.WedgeTangentX.Reserve(TotalCurrentVertices);
-		OutRawMesh.WedgeTangentY.Reserve(TotalCurrentVertices);
-		OutRawMesh.WedgeTangentZ.Reserve(TotalCurrentVertices);
-		OutRawMesh.WedgeTexCoords[0].Reserve(TotalCurrentVertices);
-
-		// Check if there are any holes
-		const int32 VisThreshold = 170;
-		TArray<uint8> VisDataMap;
-
-		for (int32 AllocIdx = 0; AllocIdx < Component->WeightmapLayerAllocations.Num(); AllocIdx++)
+		auto FindPreviousIndex = [&VertIndexAndZ, &IndexToPosition](int32 Index)->int32
 		{
-			FWeightmapLayerAllocationInfo& AllocInfo = Component->WeightmapLayerAllocations[AllocIdx];
-			if (AllocInfo.LayerInfo == ALandscapeProxy::VisibilityLayer)
+			const FVector& PositionA = IndexToPosition[Index];
+			FIndexAndZ CompressPosition(0, PositionA);
+			// Search for lowest index duplicates
+			int32 BestIndex = MAX_int32;
+			for (int32 i = 0; i < IndexToPosition.Num(); i++)
 			{
-				CDI.GetWeightmapTextureData(AllocInfo.LayerInfo, VisDataMap);
+				if (CompressPosition.Z > (VertIndexAndZ[i].Z + SMALL_NUMBER))
+				{
+					//We will not find anything there is no point searching more
+					break;
+				}
+				const FVector& PositionB = IndexToPosition[VertIndexAndZ[i].Index];
+				if (PointsEqual(PositionA, PositionB, SMALL_NUMBER))
+				{
+					if (VertIndexAndZ[i].Index < BestIndex)
+					{
+						BestIndex = VertIndexAndZ[i].Index;
+					}
+				}
 			}
-		}
-
-		const FIntPoint QuadPattern[6] =
-		{
-			//face 1
-			FIntPoint(0, 0),
-			FIntPoint(0, 1),
-			FIntPoint(1, 1),
-			//face 2
-			FIntPoint(0, 0),
-			FIntPoint(1, 1),
-			FIntPoint(1, 0),
+			return BestIndex < MAX_int32 ? BestIndex : Index;
 		};
-				
-		const int32 WeightMapSize = (SubsectionSizeQuadsLOD + 1) * Component->NumSubsections;
 
-		const float SquaredSphereRadius = FMath::Square(InBounds.SphereRadius);
-
-		// Export verts
-		int32 VertexIdx = VerticesOffset;
+		// Export to MeshDescription
+		TMap<int32, FVertexID> IndexToVertexID;
+		IndexToVertexID.Reserve(CurrentIndex);
+		CurrentIndex = 0;
 		for (int32 y = 0; y < ComponentSizeQuadsLOD; y++)
 		{
 			for (int32 x = 0; x < ComponentSizeQuadsLOD; x++)
 			{
 				FVector Positions[ARRAY_COUNT(QuadPattern)];
-				bool bProcess = false;
+				bool bProcess = bIgnoreBounds;
 
 				// Fill positions
 				for (int32 i = 0; i < ARRAY_COUNT(QuadPattern); i++)
@@ -2849,7 +2799,31 @@ bool ALandscapeProxy::ExportToRawMesh(int32 InExportLOD, FRawMesh& OutRawMesh, c
 
 				if (bProcess)
 				{
-					// Fill indices
+					//Fill the vertexID we need
+					TArray<FVertexID> VertexIDs;
+					VertexIDs.Reserve(ARRAY_COUNT(QuadPattern));
+					TArray<FVertexInstanceID> VertexInstanceIDs;
+					VertexInstanceIDs.Reserve(ARRAY_COUNT(QuadPattern));
+					// Fill positions
+					for (int32 i = 0; i < ARRAY_COUNT(QuadPattern); i++)
+					{
+						int32 DuplicateLowestIndex = FindPreviousIndex(CurrentIndex);
+						FVertexID VertexID;
+						if (DuplicateLowestIndex < CurrentIndex)
+						{
+							VertexID = IndexToVertexID[DuplicateLowestIndex];
+						}
+						else
+						{
+							VertexID = OutRawMesh.CreateVertex();
+							VertexPositions[VertexID] = Positions[i];
+						}
+						IndexToVertexID.Add(CurrentIndex, VertexID);
+						VertexIDs.Add(VertexID);
+						CurrentIndex++;
+					}
+
+					// Create triangle
 					{
 						// Whether this vertex is in hole
 						bool bInvisible = false;
@@ -2859,47 +2833,85 @@ bool ALandscapeProxy::ExportToRawMesh(int32 InExportLOD, FRawMesh& OutRawMesh, c
 							CDI.VertexXYToTexelXY(x, y, TexelX, TexelY);
 							bInvisible = (VisDataMap[CDI.TexelXYToIndex(TexelX, TexelY)] >= VisThreshold);
 						}
+						//Add vertexInstance and polygon only if we are visible
+						if (!bInvisible)
+						{
+							VertexInstanceIDs.Add(OutRawMesh.CreateVertexInstance(VertexIDs[0]));
+							VertexInstanceIDs.Add(OutRawMesh.CreateVertexInstance(VertexIDs[1]));
+							VertexInstanceIDs.Add(OutRawMesh.CreateVertexInstance(VertexIDs[2]));
 
-						// triangulation matches FLandscapeIndexBuffer constructor
-						OutRawMesh.WedgeIndices.Add(VertexIdx);
-						OutRawMesh.WedgeIndices.Add(bInvisible ? VertexIdx : VertexIdx + 1);
-						OutRawMesh.WedgeIndices.Add(bInvisible ? VertexIdx : VertexIdx + 2);
-						
-						OutRawMesh.WedgeIndices.Add(VertexIdx + 3);
-						OutRawMesh.WedgeIndices.Add(bInvisible ? VertexIdx : VertexIdx + 4);
-						OutRawMesh.WedgeIndices.Add(bInvisible ? VertexIdx : VertexIdx + 5);
-
-						OutRawMesh.FaceMaterialIndices.AddZeroed(2);
-						OutRawMesh.FaceSmoothingMasks.AddZeroed(2);
-					}
-
-					// Fill other vertex data
-					for (int32 i = 0; i < ARRAY_COUNT(QuadPattern); i++)
-					{
-						int32 VertexX = x + QuadPattern[i].X;
-						int32 VertexY = y + QuadPattern[i].Y;
+							VertexInstanceIDs.Add(OutRawMesh.CreateVertexInstance(VertexIDs[3]));
+							VertexInstanceIDs.Add(OutRawMesh.CreateVertexInstance(VertexIDs[4]));
+							VertexInstanceIDs.Add(OutRawMesh.CreateVertexInstance(VertexIDs[5]));
 
-						FVector LocalTangentX, LocalTangentY, LocalTangentZ;
-						CDI.GetLocalTangentVectors(VertexX, VertexY, LocalTangentX, LocalTangentY, LocalTangentZ);
+							// Fill other vertex data
+							for (int32 i = 0; i < ARRAY_COUNT(QuadPattern); i++)
+							{
+								int32 VertexX = x + QuadPattern[i].X;
+								int32 VertexY = y + QuadPattern[i].Y;
 
-						OutRawMesh.VertexPositions.Add(Positions[i]);
-						OutRawMesh.WedgeTangentX.Add(LocalTangentX);
-						OutRawMesh.WedgeTangentY.Add(LocalTangentY);
-						OutRawMesh.WedgeTangentZ.Add(LocalTangentZ);
+								FVector LocalTangentX, LocalTangentY, LocalTangentZ;
+								CDI.GetLocalTangentVectors(VertexX, VertexY, LocalTangentX, LocalTangentY, LocalTangentZ);
 
-						OutRawMesh.WedgeTexCoords[0].Add((ComponentUVOffsetLOD + FVector2D(VertexX, VertexY))*ComponentUVScaleLOD);
+								VertexInstanceTangents[VertexInstanceIDs[i]] = LocalTangentX;
+								VertexInstanceBinormalSigns[VertexInstanceIDs[i]] = GetBasisDeterminantSign(LocalTangentX, LocalTangentY, LocalTangentZ);
+								VertexInstanceNormals[VertexInstanceIDs[i]] = LocalTangentZ;
 
-						VertexIdx++;
+								FVector2D UV = (ComponentUVOffsetLOD + FVector2D(VertexX, VertexY))*ComponentUVScaleLOD;
+								VertexInstanceUVs.Set(VertexInstanceIDs[i], 0, UV);
+								// Add lightmap UVs
+								VertexInstanceUVs.Set(VertexInstanceIDs[i], 1, UV);
+							}
+							auto AddTriangle = [&OutRawMesh, &EdgeHardnesses, &EdgeCreaseSharpnesses, &PolygonGroupID, &VertexIDs, &VertexInstanceIDs](int32 BaseIndex)
+							{
+								//Create a polygon from this triangle
+								TArray<FMeshDescription::FContourPoint> Contours;
+								for (int32 Corner = 0; Corner < 3; ++Corner)
+								{
+									int32 ContourPointIndex = Contours.AddDefaulted();
+									FMeshDescription::FContourPoint& ContourPoint = Contours[ContourPointIndex];
+									//Find the matching edge ID
+									uint32 CornerIndices[2];
+									CornerIndices[0] = BaseIndex + ((Corner + 0) % 3);
+									CornerIndices[1] = BaseIndex + ((Corner + 1) % 3);
+
+									FVertexID EdgeVertexIDs[2];
+									EdgeVertexIDs[0] = VertexIDs[CornerIndices[0]];
+									EdgeVertexIDs[1] = VertexIDs[CornerIndices[1]];
+
+									FEdgeID MatchEdgeId = OutRawMesh.GetVertexPairEdge(EdgeVertexIDs[0], EdgeVertexIDs[1]);
+									if (MatchEdgeId == FEdgeID::Invalid)
+									{
+										MatchEdgeId = OutRawMesh.CreateEdge(EdgeVertexIDs[0], EdgeVertexIDs[1]);
+										EdgeHardnesses[MatchEdgeId] = false;
+										EdgeCreaseSharpnesses[MatchEdgeId] = 0.0f;
+									}
+									ContourPoint.EdgeID = MatchEdgeId;
+									ContourPoint.VertexInstanceID = VertexInstanceIDs[CornerIndices[0]];
+								}
+								// Insert a polygon into the mesh
+								const FPolygonID NewPolygonID = OutRawMesh.CreatePolygon(PolygonGroupID, Contours);
+								//Triangulate the polygon
+								FMeshPolygon& Polygon = OutRawMesh.GetPolygon(NewPolygonID);
+								OutRawMesh.ComputePolygonTriangulation(NewPolygonID, Polygon.Triangles);
+							};
+							AddTriangle(0);
+							AddTriangle(3);
+						}
 					}
 				}
+				else
+				{
+					CurrentIndex += ARRAY_COUNT(QuadPattern);
+				}
 			}
 		}
 	}
 
-	// Add lightmap UVs
-	OutRawMesh.WedgeTexCoords[1].Append(OutRawMesh.WedgeTexCoords[0]);
-
-	return OutRawMesh.VertexPositions.Num() != 0;
+	//Compact the MeshDescription, if there was visibility mask or some bounding box clip, it need to be compacted so the sparse array are from 0 to n with no invalid data in between. 
+	FElementIDRemappings ElementIDRemappings;
+	OutRawMesh.Compact(ElementIDRemappings);
+	return OutRawMesh.Polygons().Num() > 0;
 }
 
 
