[NO NON-PII COMMIT MESSAGE]

diff --git a/Engine/Source/Runtime/Engine/Classes/Engine/Texture.h b/Engine/Source/Runtime/Engine/Classes/Engine/Texture.h
--- a/Engine/Source/Runtime/Engine/Classes/Engine/Texture.h
+++ b/Engine/Source/Runtime/Engine/Classes/Engine/Texture.h
@@ -700,12 +700,12 @@ public:
 	/**
 	 * Begins caching platform data in the background.
 	 */
-	void BeginCachePlatformData();
+	ENGINE_API void BeginCachePlatformData();
 
 	/**
 	 * Returns true if all async caching has completed.
 	 */
-	bool IsAsyncCacheComplete();
+	ENGINE_API bool IsAsyncCacheComplete();
 
 	/**
 	 * Blocks on async cache tasks and prepares platform data for use.
diff --git a/Engine/Source/Runtime/Landscape/Private/LandscapeEdit.cpp b/Engine/Source/Runtime/Landscape/Private/LandscapeEdit.cpp
--- a/Engine/Source/Runtime/Landscape/Private/LandscapeEdit.cpp
+++ b/Engine/Source/Runtime/Landscape/Private/LandscapeEdit.cpp
@@ -2619,6 +2619,8 @@ LANDSCAPE_API void ALandscapeProxy::Import(
 		}
 	}
 
+	TArray<UTexture2D*> PendingTexturePlatformDataCreation;
+
 	// Unlock the weightmaps' base mips
 	for (int32 AllocationIndex = 0; AllocationIndex < TextureAllocations.Num(); AllocationIndex++)
 	{
@@ -2629,7 +2631,10 @@ LANDSCAPE_API void ALandscapeProxy::Import(
 		ULandscapeComponent::GenerateWeightmapMips(NumSubsections, SubsectionSizeQuads, WeightmapTexture, BaseMipData);
 
 		WeightmapTexture->Source.UnlockMip(0);
-		WeightmapTexture->PostEditChange();
+
+		WeightmapTexture->BeginCachePlatformData();
+		WeightmapTexture->ClearAllCachedCookedPlatformData();
+		PendingTexturePlatformDataCreation.Add(WeightmapTexture);
 	}
 
 	// Generate mipmaps for the components, and create the collision components
@@ -2691,7 +2696,16 @@ LANDSCAPE_API void ALandscapeProxy::Import(
 		{
 			HeightmapInfo.HeightmapTexture->Source.UnlockMip(i);
 		}
-		HeightmapInfo.HeightmapTexture->PostEditChange();
+
+		HeightmapInfo.HeightmapTexture->BeginCachePlatformData();
+		HeightmapInfo.HeightmapTexture->ClearAllCachedCookedPlatformData();
+		PendingTexturePlatformDataCreation.Add(HeightmapInfo.HeightmapTexture);
+	}
+
+	for (UTexture2D* Texture : PendingTexturePlatformDataCreation)
+	{
+		Texture->FinishCachePlatformData();
+		Texture->PostEditChange();
 	}
 
 	// Update MaterialInstances (must be done after textures are fully initialized)
