Allow FlushGrassComponents to work without a Scene

diff --git a/Engine/Source/Runtime/Landscape/Private/LandscapeGrass.cpp b/Engine/Source/Runtime/Landscape/Private/LandscapeGrass.cpp
--- a/Engine/Source/Runtime/Landscape/Private/LandscapeGrass.cpp
+++ b/Engine/Source/Runtime/Landscape/Private/LandscapeGrass.cpp
@@ -1900,7 +1900,7 @@ void ALandscapeProxy::FlushGrassComponents(const TSet<ULandscapeComponent*>* Onl
 			}
 		}
 #if WITH_EDITOR
-		if (GIsEditor && bFlushGrassMaps && GetWorld() && GetWorld()->Scene && GetWorld()->Scene->GetFeatureLevel() >= ERHIFeatureLevel::SM4)
+		if (GIsEditor && bFlushGrassMaps && GetWorld() && GetWorld()->FeatureLevel >= ERHIFeatureLevel::SM4)
 		{
 			for (ULandscapeComponent* Component : *OnlyForComponents)
 			{
