[NO NON-PII COMMIT MESSAGE]

diff --git a/Engine/Source/Runtime/Landscape/Private/Landscape.cpp b/Engine/Source/Runtime/Landscape/Private/Landscape.cpp
--- a/Engine/Source/Runtime/Landscape/Private/Landscape.cpp
+++ b/Engine/Source/Runtime/Landscape/Private/Landscape.cpp
@@ -2262,7 +2262,7 @@ void ULandscapeInfo::RegisterActor(ALandscapeProxy* Proxy, bool bMapCheck)
 		checkf(!LandscapeActor || LandscapeActor == Landscape, TEXT("Multiple landscapes with the same GUID detected: %s vs %s"), *LandscapeActor->GetPathName(), *Landscape->GetPathName());
 		LandscapeActor = Landscape;
 		// In world composition user is not allowed to move landscape in editor, only through WorldBrowser 
-		LandscapeActor->bLockLocation = (OwningWorld->WorldComposition != nullptr);
+		LandscapeActor->bLockLocation = OwningWorld != nullptr ? OwningWorld->WorldComposition != nullptr : false;
 
 		// update proxies reference actor
 		for (ALandscapeStreamingProxy* StreamingProxy : Proxies)
