[NO NON-PII COMMIT MESSAGE]

diff --git a/Engine/Source/Runtime/Landscape/Private/LandscapeRender.cpp b/Engine/Source/Runtime/Landscape/Private/LandscapeRender.cpp
--- a/Engine/Source/Runtime/Landscape/Private/LandscapeRender.cpp
+++ b/Engine/Source/Runtime/Landscape/Private/LandscapeRender.cpp
@@ -806,9 +806,13 @@ FLandscapeComponentSceneProxy::FLandscapeComponentSceneProxy(ULandscapeComponent
 
 	MaterialRelevances.Reserve(AvailableMaterials.Num());
 
-	for (UMaterialInterface* Material : AvailableMaterials)
+	for (UMaterialInterface*& MaterialInterface : AvailableMaterials)
 	{
-		MaterialRelevances.Add(Material->GetRelevance(FeatureLevel));
+		if (MaterialInterface != nullptr)
+		{
+			MaterialRelevances.Add(MaterialInterface->GetRelevance(FeatureLevel));
+			bRequiresAdjacencyInformation |= MaterialSettingsRequireAdjacencyInformation_GameThread(MaterialInterface, XYOffsetmapTexture == nullptr ? &FLandscapeVertexFactory::StaticType : &FLandscapeXYOffsetVertexFactory::StaticType, InComponent->GetWorld()->FeatureLevel);
+		}
 	}
 
 #if !(UE_BUILD_SHIPPING || UE_BUILD_TEST) || (UE_BUILD_SHIPPING && WITH_EDITOR)
@@ -828,11 +832,6 @@ FLandscapeComponentSceneProxy::FLandscapeComponentSceneProxy(ULandscapeComponent
 	}
 #endif
 
-	for (UMaterialInterface*& MaterialInterface : AvailableMaterials)
-	{
-		bRequiresAdjacencyInformation |= MaterialSettingsRequireAdjacencyInformation_GameThread(MaterialInterface, XYOffsetmapTexture == nullptr ? &FLandscapeVertexFactory::StaticType : &FLandscapeXYOffsetVertexFactory::StaticType, InComponent->GetWorld()->FeatureLevel);
-	}
-
 	const int8 SubsectionSizeLog2 = FMath::CeilLogTwo(InComponent->SubsectionSizeQuads + 1);
 	SharedBuffersKey = (SubsectionSizeLog2 & 0xf) | ((NumSubsections & 0xf) << 4) |
 		(FeatureLevel <= ERHIFeatureLevel::ES3_1 ? 0 : 1 << 30) | (XYOffsetmapTexture == nullptr ? 0 : 1 << 31);
