[NO NON-PII COMMIT MESSAGE]

diff --git a/Engine/Source/Runtime/Landscape/Private/LandscapeRender.cpp b/Engine/Source/Runtime/Landscape/Private/LandscapeRender.cpp
--- a/Engine/Source/Runtime/Landscape/Private/LandscapeRender.cpp
+++ b/Engine/Source/Runtime/Landscape/Private/LandscapeRender.cpp
@@ -1612,7 +1612,7 @@ void FLandscapeComponentSceneProxy::DrawStaticElements(FStaticPrimitiveDrawInter
 			continue;
 		}
 
-		bool HasTessellationEnabled = (FeatureLevel >= ERHIFeatureLevel::SM4) ? MaterialInstance->GetMaterial()->D3D11TessellationMode != EMaterialTessellationMode::MTM_NoTessellation : false;
+		bool HasTessellationEnabled = (FeatureLevel >= ERHIFeatureLevel::SM4) ? MaterialInstance->GetMaterial()->D3D11TessellationMode != EMaterialTessellationMode::MTM_NoTessellation && MaterialIndexToDisabledTessellationMaterial[i] != INDEX_NONE : false;
 
 		// Only add normal batch index (for each material)
 		MaterialIndexToStaticMeshBatchLOD[i] = CurrentLODIndex;
@@ -2526,7 +2526,7 @@ void FLandscapeComponentSceneProxy::GetDynamicMeshElements(const TArray<const FS
 					}
 				}
 
-				HasTessellationEnabled = MaterialHasTessellationEnabled[LODIndexToMaterialIndex[CurrentLODIndex]] && (LandscapeMIC != nullptr && !LandscapeMIC->bDisableTessellation);
+				HasTessellationEnabled = MaterialHasTessellationEnabled[LODIndexToMaterialIndex[CurrentLODIndex]] && (LandscapeMIC != nullptr && !LandscapeMIC->bDisableTessellation) && AvailableMaterials.Num() > 1;
 			}
 
 			bool DisableTessellation = false;
