[NO NON-PII COMMIT MESSAGE]

diff --git a/Engine/Source/Runtime/Landscape/Private/LandscapeRender.cpp b/Engine/Source/Runtime/Landscape/Private/LandscapeRender.cpp
--- a/Engine/Source/Runtime/Landscape/Private/LandscapeRender.cpp
+++ b/Engine/Source/Runtime/Landscape/Private/LandscapeRender.cpp
@@ -685,19 +685,6 @@ FLandscapeComponentSceneProxy::FLandscapeComponentSceneProxy(ULandscapeComponent
 	LODIndexToMaterialIndex = InComponent->LODIndexToMaterialIndex;
 	check(LODIndexToMaterialIndex.Num() == MaxLOD+1);
 
-	for (int32 i = 0; i < AvailableMaterials.Num(); ++i)
-	{
-		bool HasTessellationEnabled = false;
-
-		if (FeatureLevel >= ERHIFeatureLevel::SM4)
-		{
-			UMaterialInstance* LandscapeMI = Cast<UMaterialInstance>(AvailableMaterials[i]);
-			HasTessellationEnabled = LandscapeMI->GetMaterial()->D3D11TessellationMode != EMaterialTessellationMode::MTM_NoTessellation;
-		}
-
-		MaterialHasTessellationEnabled.Add(HasTessellationEnabled);
-	}
-
 	if (!IsComponentLevelVisible())
 	{
 		bNeedsLevelAddedToWorldNotification = true;
@@ -812,6 +799,17 @@ FLandscapeComponentSceneProxy::FLandscapeComponentSceneProxy(ULandscapeComponent
 		{
 			MaterialRelevances.Add(MaterialInterface->GetRelevance(FeatureLevel));
 			bRequiresAdjacencyInformation |= MaterialSettingsRequireAdjacencyInformation_GameThread(MaterialInterface, XYOffsetmapTexture == nullptr ? &FLandscapeVertexFactory::StaticType : &FLandscapeXYOffsetVertexFactory::StaticType, InComponent->GetWorld()->FeatureLevel);
+
+			bool HasTessellationEnabled = false;
+
+			if (FeatureLevel >= ERHIFeatureLevel::SM4)
+			{
+				UMaterial* LandscapeMaterial = MaterialInterface->GetMaterial();
+
+				HasTessellationEnabled = LandscapeMaterial != nullptr && LandscapeMaterial->D3D11TessellationMode != EMaterialTessellationMode::MTM_NoTessellation;
+			}
+
+			MaterialHasTessellationEnabled.Add(HasTessellationEnabled);
 		}
 	}
 
@@ -1605,14 +1603,14 @@ void FLandscapeComponentSceneProxy::DrawStaticElements(FStaticPrimitiveDrawInter
 
 	for (int32 i = 0; i < MaterialCount; ++i)
 	{
-		UMaterialInstance* MaterialInstance = Cast<UMaterialInstance>(AvailableMaterials[i]);
-
-		if (MaterialInstance == nullptr)
+		if (AvailableMaterials[i] == nullptr)
 		{
 			continue;
 		}
 
-		bool HasTessellationEnabled = (FeatureLevel >= ERHIFeatureLevel::SM4) ? MaterialInstance->GetMaterial()->D3D11TessellationMode != EMaterialTessellationMode::MTM_NoTessellation && MaterialIndexToDisabledTessellationMaterial[i] != INDEX_NONE : false;
+		UMaterialInstance* MaterialInstance = Cast<UMaterialInstance>(AvailableMaterials[i]);
+
+		bool HasTessellationEnabled = (FeatureLevel >= ERHIFeatureLevel::SM4) ? MaterialInstance != nullptr && MaterialInstance->GetMaterial()->D3D11TessellationMode != EMaterialTessellationMode::MTM_NoTessellation && MaterialIndexToDisabledTessellationMaterial[i] != INDEX_NONE : false;
 
 		// Only add normal batch index (for each material)
 		MaterialIndexToStaticMeshBatchLOD[i] = CurrentLODIndex;
@@ -1620,7 +1618,7 @@ void FLandscapeComponentSceneProxy::DrawStaticElements(FStaticPrimitiveDrawInter
 		// Default Batch, tessellated if enabled
 		FMeshBatch NormalMeshBatch;
 
-		if (GetMeshElement(false, false, HasTessellationEnabled, CurrentLODIndex++, MaterialInstance, NormalMeshBatch, StaticBatchParamArray))
+		if (GetMeshElement(false, false, HasTessellationEnabled, CurrentLODIndex++, AvailableMaterials[i], NormalMeshBatch, StaticBatchParamArray))
 		{
 			PDI->DrawMesh(NormalMeshBatch, FLT_MAX);
 		}
