Cleanup unrequired data for Adjacency information
Fixed a Memory leak where Adjacency index buffer would not be freed, when the owner was freed (Shared Buffer)

diff --git a/Engine/Source/Runtime/Landscape/Private/LandscapeRender.cpp b/Engine/Source/Runtime/Landscape/Private/LandscapeRender.cpp
--- a/Engine/Source/Runtime/Landscape/Private/LandscapeRender.cpp
+++ b/Engine/Source/Runtime/Landscape/Private/LandscapeRender.cpp
@@ -602,7 +602,6 @@ void ULandscapeComponent::GetUsedMaterials(TArray<UMaterialInterface*>& OutMater
 // FLandscapeComponentSceneProxy
 //
 TMap<uint32, FLandscapeSharedBuffers*>FLandscapeComponentSceneProxy::SharedBuffersMap;
-TMap<uint32, FLandscapeSharedAdjacencyIndexBuffer*>FLandscapeComponentSceneProxy::SharedAdjacencyIndexBufferMap;
 TMap<FLandscapeNeighborInfo::FLandscapeKey, TMap<FIntPoint, const FLandscapeNeighborInfo*> > FLandscapeNeighborInfo::SharedSceneProxyMap;
 
 const static FName NAME_LandscapeResourceNameForDebugging(TEXT("Landscape"));
@@ -904,9 +903,7 @@ void FLandscapeComponentSceneProxy::CreateRenderThreadResources()
 			}
 
 			SharedBuffers->AdjacencyIndexBuffers = new FLandscapeSharedAdjacencyIndexBuffer(SharedBuffers);
-			FLandscapeComponentSceneProxy::SharedAdjacencyIndexBufferMap.Add(SharedBuffersKey, SharedBuffers->AdjacencyIndexBuffers);
 		}
-		SharedBuffers->AdjacencyIndexBuffers->AddRef();
 
 		// Delayed Initialize for IndexBuffers
 		for (int32 i = 0; i < SharedBuffers->NumIndexBuffers; i++)
@@ -3259,15 +3256,7 @@ FLandscapeSharedBuffers::~FLandscapeSharedBuffers()
 	}
 #endif
 
-	if (AdjacencyIndexBuffers)
-	{
-		if (AdjacencyIndexBuffers->Release() == 0)
-		{
-			FLandscapeComponentSceneProxy::SharedAdjacencyIndexBufferMap.Remove(SharedBuffersKey);
-		}
-		AdjacencyIndexBuffers = nullptr;
-	}
-
+	delete AdjacencyIndexBuffers;
 	delete VertexFactory;
 
 	if (OccluderIndicesSP.IsValid())
diff --git a/Engine/Source/Runtime/Landscape/Public/LandscapeRender.h b/Engine/Source/Runtime/Landscape/Public/LandscapeRender.h
--- a/Engine/Source/Runtime/Landscape/Public/LandscapeRender.h
+++ b/Engine/Source/Runtime/Landscape/Public/LandscapeRender.h
@@ -272,7 +272,7 @@ public:
 //
 // FLandscapeSharedAdjacencyIndexBuffer
 //
-class FLandscapeSharedAdjacencyIndexBuffer : public FRefCountedObject
+class FLandscapeSharedAdjacencyIndexBuffer
 {
 public:
 	FLandscapeSharedAdjacencyIndexBuffer(class FLandscapeSharedBuffers* SharedBuffer);
@@ -603,7 +603,6 @@ protected:
 	// Reference counted vertex and index buffer shared among all landscape scene proxies of the same component size
 	// Key is the component size and number of subsections.
 	static TMap<uint32, FLandscapeSharedBuffers*> SharedBuffersMap;
-	static TMap<uint32, FLandscapeSharedAdjacencyIndexBuffer*> SharedAdjacencyIndexBufferMap;
 
 #if WITH_EDITORONLY_DATA
 	FLandscapeEditToolRenderData EditToolRenderData;
