[NO NON-PII COMMIT MESSAGE]

diff --git a/Engine/Source/Runtime/Engine/Classes/Materials/MaterialInstance.h b/Engine/Source/Runtime/Engine/Classes/Materials/MaterialInstance.h
--- a/Engine/Source/Runtime/Engine/Classes/Materials/MaterialInstance.h
+++ b/Engine/Source/Runtime/Engine/Classes/Materials/MaterialInstance.h
@@ -499,12 +499,12 @@ public:
 	 * Can be passed either a minimal parameter set (overridden parameters only) or the entire set generated by GetStaticParameterValues().
 	 * Can also trigger recompile based on new set of FMaterialInstanceBasePropertyOverrides 
 	 */
-	ENGINE_API void UpdateStaticPermutation(const FStaticParameterSet& NewParameters, FMaterialInstanceBasePropertyOverrides& NewBasePropertyOverrides, const bool bForceStaticPermutationUpdate = false);
+	ENGINE_API void UpdateStaticPermutation(const FStaticParameterSet& NewParameters, FMaterialInstanceBasePropertyOverrides& NewBasePropertyOverrides, const bool bForceStaticPermutationUpdate = false, FMaterialUpdateContext* MaterialUpdateContext = nullptr);
 	/**
 	* Sets new static parameter overrides on the instance and recompiles the static permutation resources if needed.
 	* Can be passed either a minimal parameter set (overridden parameters only) or the entire set generated by GetStaticParameterValues().
 	*/
-	ENGINE_API void UpdateStaticPermutation(const FStaticParameterSet& NewParameters);
+	ENGINE_API void UpdateStaticPermutation(const FStaticParameterSet& NewParameters, FMaterialUpdateContext* MaterialUpdateContext = nullptr);
 	/** Ensure's static permutations for current parameters and overrides are upto date. */
 	ENGINE_API void UpdateStaticPermutation();
 
diff --git a/Engine/Source/Runtime/Engine/Private/Materials/MaterialInstance.cpp b/Engine/Source/Runtime/Engine/Private/Materials/MaterialInstance.cpp
--- a/Engine/Source/Runtime/Engine/Private/Materials/MaterialInstance.cpp
+++ b/Engine/Source/Runtime/Engine/Private/Materials/MaterialInstance.cpp
@@ -3515,7 +3515,7 @@ void UMaterialInstance::ClearParameterValuesInternal(const bool bAllParameters)
 }
 
 #if WITH_EDITOR
-void UMaterialInstance::UpdateStaticPermutation(const FStaticParameterSet& NewParameters, FMaterialInstanceBasePropertyOverrides& NewBasePropertyOverrides, const bool bForceStaticPermutationUpdate /*= false*/)
+void UMaterialInstance::UpdateStaticPermutation(const FStaticParameterSet& NewParameters, FMaterialInstanceBasePropertyOverrides& NewBasePropertyOverrides, const bool bForceStaticPermutationUpdate /*= false*/, FMaterialUpdateContext* MaterialUpdateContext)
 {
 	FStaticParameterSet CompareParameters = NewParameters;
 
@@ -3538,20 +3538,28 @@ void UMaterialInstance::UpdateStaticPermutation(const FStaticParameterSet& NewPa
 
 	if (bHasStaticPermutationResource != bWantsStaticPermutationResource || bParamsHaveChanged || (bBasePropertyOverridesHaveChanged && bWantsStaticPermutationResource) || bForceStaticPermutationUpdate)
 	{
-		// This will flush the rendering thread which is necessary before changing bHasStaticPermutationResource, since the RT is reading from that directly
-		// The update context will also make sure any dependent MI's with static parameters get recompiled
-		FMaterialUpdateContext MaterialUpdateContext;
-		MaterialUpdateContext.AddMaterialInstance(this);
 		bHasStaticPermutationResource = bWantsStaticPermutationResource;
 		StaticParameters = CompareParameters;
 
 		CacheResourceShadersForRendering();
+
+		if (MaterialUpdateContext != nullptr)
+		{
+			MaterialUpdateContext->AddMaterialInstance(this);
+		}
+		else
+		{
+			// This will flush the rendering thread which is necessary before changing bHasStaticPermutationResource, since the RT is reading from that directly
+			// The update context will also make sure any dependent MI's with static parameters get recompiled
+			FMaterialUpdateContext LocalMaterialUpdateContext;
+			LocalMaterialUpdateContext.AddMaterialInstance(this);
+		}
 	}
 }
 
-void UMaterialInstance::UpdateStaticPermutation(const FStaticParameterSet& NewParameters)
+void UMaterialInstance::UpdateStaticPermutation(const FStaticParameterSet& NewParameters, FMaterialUpdateContext* MaterialUpdateContext)
 {
-	UpdateStaticPermutation(NewParameters, BasePropertyOverrides);
+	UpdateStaticPermutation(NewParameters, BasePropertyOverrides, false, MaterialUpdateContext);
 }
 
 void UMaterialInstance::UpdateStaticPermutation()
diff --git a/Engine/Source/Runtime/Landscape/Classes/LandscapeComponent.h b/Engine/Source/Runtime/Landscape/Classes/LandscapeComponent.h
--- a/Engine/Source/Runtime/Landscape/Classes/LandscapeComponent.h
+++ b/Engine/Source/Runtime/Landscape/Classes/LandscapeComponent.h
@@ -645,7 +645,7 @@ public:
 	void UpdateMaterialInstances_Internal(FMaterialUpdateContext& Context);
 
 	/** Helper function for UpdateMaterialInstance to get Material without set parameters */
-	UMaterialInstanceConstant* GetCombinationMaterial(const TArray<FWeightmapLayerAllocationInfo>& Allocations, int8 InLODIndex, bool bMobile = false) const;
+	UMaterialInstanceConstant* GetCombinationMaterial(FMaterialUpdateContext* InMaterialUpdateContext, const TArray<FWeightmapLayerAllocationInfo>& Allocations, int8 InLODIndex, bool bMobile = false) const;
 	/**
 	 * Generate mipmaps for height and tangent data.
 	 * @param HeightmapTextureMipData - array of pointers to the locked mip data.
diff --git a/Engine/Source/Runtime/Landscape/Private/LandscapeEdit.cpp b/Engine/Source/Runtime/Landscape/Private/LandscapeEdit.cpp
--- a/Engine/Source/Runtime/Landscape/Private/LandscapeEdit.cpp
+++ b/Engine/Source/Runtime/Landscape/Private/LandscapeEdit.cpp
@@ -208,7 +208,7 @@ FString ULandscapeComponent::GetLayerAllocationKey(const TArray<FWeightmapLayerA
 	return Result;
 }
 
-UMaterialInstanceConstant* ULandscapeComponent::GetCombinationMaterial(const TArray<FWeightmapLayerAllocationInfo>& Allocations, int8 InLODIndex, bool bMobile /*= false*/) const
+UMaterialInstanceConstant* ULandscapeComponent::GetCombinationMaterial(FMaterialUpdateContext* InMaterialUpdateContext, const TArray<FWeightmapLayerAllocationInfo>& Allocations, int8 InLODIndex, bool bMobile /*= false*/) const
 {
 	check(GIsEditor);
 
@@ -272,7 +272,7 @@ UMaterialInstanceConstant* ULandscapeComponent::GetCombinationMaterial(const TAr
 					StaticParameters.TerrainLayerWeightParameters.Add(FStaticTerrainLayerWeightParameter(LayerParameter, Allocation.WeightmapTextureIndex, true, FGuid(), !Allocation.LayerInfo->bNoWeightBlend));
 				}
 			}
-			CombinationMaterialInstance->UpdateStaticPermutation(StaticParameters);
+			CombinationMaterialInstance->UpdateStaticPermutation(StaticParameters, InMaterialUpdateContext);
 
 			CombinationMaterialInstance->PostEditChange();
 		}
@@ -331,7 +331,7 @@ void ULandscapeComponent::UpdateMaterialInstances_Internal(FMaterialUpdateContex
 		const int8 MaterialLOD = It.Value();
 
 		// Find or set a matching MIC in the Landscape's map.
-		UMaterialInstanceConstant* CombinationMaterialInstance = GetCombinationMaterial(WeightmapLayerAllocations, MaterialLOD, false);
+		UMaterialInstanceConstant* CombinationMaterialInstance = GetCombinationMaterial(&Context, WeightmapLayerAllocations, MaterialLOD, false);
 
 		if (CombinationMaterialInstance != nullptr)
 		{
@@ -412,7 +412,7 @@ void ULandscapeComponent::UpdateMaterialInstances_Internal(FMaterialUpdateContex
 		{
 			const int8 MaterialLOD = It.Value();
 
-			MobileCombinationMaterialInstances[MobileMaterialIndex] = GetCombinationMaterial(MobileWeightmapLayerAllocations, MaterialLOD, true);
+			MobileCombinationMaterialInstances[MobileMaterialIndex] = GetCombinationMaterial(&Context, MobileWeightmapLayerAllocations, MaterialLOD, true);
 			++MobileMaterialIndex;
 		}
 	}
@@ -5353,7 +5353,7 @@ void ULandscapeComponent::GeneratePlatformPixelData()
 			const int8 MaterialLOD = It.Value();
 
 			// Find or set a matching MIC in the Landscape's map.
-			MobileCombinationMaterialInstances[MaterialIndex] = GetCombinationMaterial(MobileWeightmapLayerAllocations, MaterialLOD, true);
+			MobileCombinationMaterialInstances[MaterialIndex] = GetCombinationMaterial(nullptr, MobileWeightmapLayerAllocations, MaterialLOD, true);
 			check(MobileCombinationMaterialInstances[MaterialIndex] != nullptr);
 
 			UMaterialInstanceConstant* NewMobileMaterialInstance = NewObject<ULandscapeMaterialInstanceConstant>(GetOutermost());
