Fix landscape runtime crash from bad cast

diff --git a/Engine/Source/Runtime/Landscape/Landscape.Build.cs b/Engine/Source/Runtime/Landscape/Landscape.Build.cs
--- a/Engine/Source/Runtime/Landscape/Landscape.Build.cs
+++ b/Engine/Source/Runtime/Landscape/Landscape.Build.cs
@@ -19,6 +19,7 @@ public class Landscape : ModuleRules
 				"TargetPlatform",
 				"DerivedDataCache",
 				"Foliage",
+				"Renderer",
 			}
 		);
 
diff --git a/Engine/Source/Runtime/Landscape/Private/LandscapeRender.cpp b/Engine/Source/Runtime/Landscape/Private/LandscapeRender.cpp
--- a/Engine/Source/Runtime/Landscape/Private/LandscapeRender.cpp
+++ b/Engine/Source/Runtime/Landscape/Private/LandscapeRender.cpp
@@ -33,6 +33,7 @@ LandscapeRender.cpp: New terrain rendering
 #include "DrawDebugHelpers.h"
 #include "PrimitiveSceneInfo.h"
 #include "SceneView.h"
+#include "Runtime/Renderer/Private/SceneCore.h"
 #include "LandscapeProxy.h"
 #include "MeshMaterialShader.h"
 
@@ -2263,8 +2264,8 @@ bool FLandscapeComponentSceneProxy::IsUsingCustomWholeSceneShadowLODRules() cons
 
 bool FLandscapeComponentSceneProxy::CanUseMeshBatchForShadowCascade(int8 InLODIndex, float InShadowMapTextureResolution, float InShadowMapCascadeSize) const
 {
-	const FMeshBatch* MeshBatch = nullptr;
-	const TIndirectArray<FMeshBatch>& PrimitiveStaticMeshes = (const TIndirectArray<FMeshBatch>&)GetPrimitiveSceneInfo()->StaticMeshes;
+	const FStaticMesh* MeshBatch = nullptr;
+	const TArray<FStaticMesh>& PrimitiveStaticMeshes = (const TArray<FStaticMesh>&)GetPrimitiveSceneInfo()->StaticMeshes;
 
 	check(PrimitiveStaticMeshes.IsValidIndex(InLODIndex));
 	check(PrimitiveStaticMeshes[InLODIndex].CastShadow);
