FLandscapeComponentSceneProxy - removed redundant StaticMeshes array cast.

diff --git a/Engine/Source/Runtime/Landscape/Private/LandscapeRender.cpp b/Engine/Source/Runtime/Landscape/Private/LandscapeRender.cpp
--- a/Engine/Source/Runtime/Landscape/Private/LandscapeRender.cpp
+++ b/Engine/Source/Runtime/Landscape/Private/LandscapeRender.cpp
@@ -2265,7 +2265,7 @@ bool FLandscapeComponentSceneProxy::IsUsingCustomWholeSceneShadowLODRules() cons
 bool FLandscapeComponentSceneProxy::CanUseMeshBatchForShadowCascade(int8 InLODIndex, float InShadowMapTextureResolution, float InShadowMapCascadeSize) const
 {
 	const FStaticMesh* MeshBatch = nullptr;
-	const TArray<FStaticMesh>& PrimitiveStaticMeshes = (const TArray<FStaticMesh>&)GetPrimitiveSceneInfo()->StaticMeshes;
+	const TArray<FStaticMesh>& PrimitiveStaticMeshes = GetPrimitiveSceneInfo()->StaticMeshes;
 
 	check(PrimitiveStaticMeshes.IsValidIndex(InLODIndex));
 	check(PrimitiveStaticMeshes[InLODIndex].CastShadow);
