Landscape - disable depth pass caching when a separate MeshBatch is used for it.

diff --git a/Engine/Source/Runtime/Landscape/Private/LandscapeRender.cpp b/Engine/Source/Runtime/Landscape/Private/LandscapeRender.cpp
--- a/Engine/Source/Runtime/Landscape/Private/LandscapeRender.cpp
+++ b/Engine/Source/Runtime/Landscape/Private/LandscapeRender.cpp
@@ -1455,6 +1455,7 @@ bool FLandscapeComponentSceneProxy::GetMeshElement(bool UseSeperateBatchForShado
 	OutMeshBatch.LCI = ComponentLightInfo.Get();
 	OutMeshBatch.ReverseCulling = IsLocalToWorldDeterminantNegative();
 	OutMeshBatch.CastShadow = UseSeperateBatchForShadow ? ShadowOnly : true;
+	OutMeshBatch.bUseForDepthPass = (UseSeperateBatchForShadow ? !ShadowOnly : true);
 	OutMeshBatch.bUseAsOccluder = (UseSeperateBatchForShadow ? !ShadowOnly : true) && ShouldUseAsOccluder() && GetScene().GetShadingPath() == EShadingPath::Deferred && !IsMovable();
 	OutMeshBatch.bUseForMaterial = (UseSeperateBatchForShadow ? !ShadowOnly : true);
 	OutMeshBatch.Type = bCurrentRequiresAdjacencyInformation ? PT_12_ControlPointPatchList : PT_TriangleList;
