Added missing shader bindings for mobile landscape

diff --git a/Engine/Source/Runtime/Landscape/Private/LandscapeRenderMobile.cpp b/Engine/Source/Runtime/Landscape/Private/LandscapeRenderMobile.cpp
--- a/Engine/Source/Runtime/Landscape/Private/LandscapeRenderMobile.cpp
+++ b/Engine/Source/Runtime/Landscape/Private/LandscapeRenderMobile.cpp
@@ -194,7 +194,37 @@ public:
 		FLandscapeVertexFactoryPixelShaderParameters::Serialize(Ar);
 		Ar << BlendableLayerMaskParameter;
 	}
-	
+
+	virtual void GetElementShaderBindings(
+		const class FSceneInterface* Scene,
+		const FSceneView* InView,
+		const class FMeshMaterialShader* Shader,
+		bool bShaderRequiresPositionOnlyStream,
+		ERHIFeatureLevel::Type FeatureLevel,
+		const FVertexFactory* VertexFactory,
+		const FMeshBatchElement& BatchElement,
+		class FMeshDrawSingleShaderBindings& ShaderBindings,
+		FVertexInputStreamArray& VertexStreams
+	) const override final
+	{
+		SCOPE_CYCLE_COUNTER(STAT_LandscapeVFDrawTimePS);
+		
+		FLandscapeVertexFactoryPixelShaderParameters::GetElementShaderBindings(Scene, InView, Shader, bShaderRequiresPositionOnlyStream, FeatureLevel, VertexFactory, BatchElement, ShaderBindings, VertexStreams);
+
+		if (BlendableLayerMaskParameter.IsBound())
+		{
+			const FLandscapeBatchElementParams* BatchElementParams = (const FLandscapeBatchElementParams*)BatchElement.UserData;
+			check(BatchElementParams);
+			const FLandscapeComponentSceneProxyMobile* SceneProxy = (const FLandscapeComponentSceneProxyMobile*)BatchElementParams->SceneProxy;
+			
+			FVector MaskVector;
+			MaskVector[0] = (SceneProxy->BlendableLayerMask & (1 << 0)) ? 1 : 0;
+			MaskVector[1] = (SceneProxy->BlendableLayerMask & (1 << 1)) ? 1 : 0;
+			MaskVector[2] = (SceneProxy->BlendableLayerMask & (1 << 2)) ? 1 : 0;
+			ShaderBindings.Add(BlendableLayerMaskParameter, MaskVector);
+		}
+	}
+
 protected:
 	FShaderParameter BlendableLayerMaskParameter;
 };
diff --git a/Engine/Source/Runtime/Landscape/Public/LandscapeRender.h b/Engine/Source/Runtime/Landscape/Public/LandscapeRender.h
--- a/Engine/Source/Runtime/Landscape/Public/LandscapeRender.h
+++ b/Engine/Source/Runtime/Landscape/Public/LandscapeRender.h
@@ -139,7 +139,7 @@ public:
 		const FMeshBatchElement& BatchElement,
 		class FMeshDrawSingleShaderBindings& ShaderBindings,
 		FVertexInputStreamArray& VertexStreams
-	) const override final;
+	) const override;
 
 	virtual uint32 GetSize() const override
 	{
