Changed cook check about landscape verts into a warning.

diff --git a/Engine/Source/Runtime/Landscape/Private/LandscapeEdit.cpp b/Engine/Source/Runtime/Landscape/Private/LandscapeEdit.cpp
--- a/Engine/Source/Runtime/Landscape/Private/LandscapeEdit.cpp
+++ b/Engine/Source/Runtime/Landscape/Private/LandscapeEdit.cpp
@@ -5463,7 +5463,12 @@ void ULandscapeComponent::GeneratePlatformVertexData(const ITargetPlatform* Targ
 			}
 		}
 	}
-	check(VertexOrder.Num() == FMath::Square(SubsectionSizeVerts) * FMath::Square(NumSubsections));
+
+	if (VertexOrder.Num() != FMath::Square(SubsectionSizeVerts) * FMath::Square(NumSubsections)) 
+	{
+		UE_LOG(LogLandscape, Warning, TEXT("VertexOrder count of %d did not match expected size of %d"), 
+			VertexOrder.Num(), FMath::Square(SubsectionSizeVerts) * FMath::Square(NumSubsections));
+	}
 
 	int32 NumMobileVerices = FMath::Square(SubsectionSizeVerts * NumSubsections);
 	TArray<FLandscapeMobileVertex> MobileVertices;
