* Fix for landscape flicker issue UE-69785
* Problem was caused by a race condition where custom data was being created in the shadow collection at the same time as it was being referenced in the main pass mesh processing
* This caused the depth and material draw calls in the main pass to calculate different LOD settings
* There were a few options (see Jira comments) but the simplest/safest fix is to flag the shadow only custom data not to be used for neighbor LOD calcuation

diff --git a/Engine/Source/Runtime/Engine/Public/PrimitiveSceneProxy.h b/Engine/Source/Runtime/Engine/Public/PrimitiveSceneProxy.h
--- a/Engine/Source/Runtime/Engine/Public/PrimitiveSceneProxy.h
+++ b/Engine/Source/Runtime/Engine/Public/PrimitiveSceneProxy.h
@@ -679,11 +679,12 @@ public:
 	 * @param InView - Current View
  	 * @param InViewLODScale - View LOD scale
   	 * @param InCustomDataMemStack - MemStack to allocate the custom data
-   	 * @param InIsStaticRelevant - Tell us if it was called in a static of dynamic relevancy context
+	 * @param InIsStaticRelevant - Tell us if it was called in a static of dynamic relevancy context
+	 * @param InIsShadowOnly - Tell us if we are creating in the shadow context
    	 * @param InVisiblePrimitiveLODMask - Calculated LODMask for visibile primitive in static relevancy
    	 * @param InMeshScreenSizeSquared - Computed mesh batch screen size, passed to prevent recalculation
 	 */
-	ENGINE_API virtual void* InitViewCustomData(const FSceneView& InView, float InViewLODScale, FMemStackBase& InCustomDataMemStack, bool InIsStaticRelevant = false, const struct FLODMask* InVisiblePrimitiveLODMask = nullptr, float InMeshScreenSizeSquared = -1.0f) { return nullptr; }
+	ENGINE_API virtual void* InitViewCustomData(const FSceneView& InView, float InViewLODScale, FMemStackBase& InCustomDataMemStack, bool InIsStaticRelevant, bool InIsShadowOnly, const struct FLODMask* InVisiblePrimitiveLODMask = nullptr, float InMeshScreenSizeSquared = -1.0f) { return nullptr; }
 	
 	/**
 	 * Called during post visibility and shadow setup, just before the frame is rendered. It can be used to update custom data that had a dependency between them.
diff --git a/Engine/Source/Runtime/Landscape/Private/LandscapeRender.cpp b/Engine/Source/Runtime/Landscape/Private/LandscapeRender.cpp
--- a/Engine/Source/Runtime/Landscape/Private/LandscapeRender.cpp
+++ b/Engine/Source/Runtime/Landscape/Private/LandscapeRender.cpp
@@ -1954,7 +1954,7 @@ int8 FLandscapeComponentSceneProxy::GetLODFromScreenSize(float InScreenSizeSquar
 	return INDEX_NONE;
 }
 
-void* FLandscapeComponentSceneProxy::InitViewCustomData(const FSceneView& InView, float InViewLODScale, FMemStackBase& InCustomDataMemStack, bool InIsStaticRelevant, const FLODMask* InVisiblePrimitiveLODMask, float InMeshScreenSizeSquared)
+void* FLandscapeComponentSceneProxy::InitViewCustomData(const FSceneView& InView, float InViewLODScale, FMemStackBase& InCustomDataMemStack, bool InIsStaticRelevant, bool InIsShadowOnly, const FLODMask* InVisiblePrimitiveLODMask, float InMeshScreenSizeSquared)
 {
 	SCOPE_CYCLE_COUNTER(STAT_LandscapeInitViewCustomData);
 
@@ -1994,6 +1994,8 @@ void* FLandscapeComponentSceneProxy::InitViewCustomData(const FSceneView& InView
 		}
 	}
 
+	LODData->IsShadowOnly = InIsShadowOnly;
+
 	// Mobile use a different way of calculating the Bias
 	if (GetScene().GetFeatureLevel() >= ERHIFeatureLevel::SM4)
 	{
@@ -2138,9 +2140,14 @@ float FLandscapeComponentSceneProxy::GetNeighborLOD(const FSceneView& InView, fl
 
 	if (CustomData != nullptr)
 	{
-		ComputeNeighborCustomDataLOD = false;
 		const FViewCustomDataLOD* LODData = (const FViewCustomDataLOD*)CustomData;
-		NeighborLOD = FMath::Max(LODData->SubSections[DesiredSubSectionIndex].fBatchElementCurrentLOD, InBatchElementCurrentLOD);
+		// Don't use the custom data for neighbor calculation when it is marked shadow only (ie it is not visible in the view)
+		// See UE-69785 for more information
+		if (!LODData->IsShadowOnly)
+		{
+			ComputeNeighborCustomDataLOD = false;
+			NeighborLOD = FMath::Max(LODData->SubSections[DesiredSubSectionIndex].fBatchElementCurrentLOD, InBatchElementCurrentLOD);
+		}
 	}
 
 	if (ComputeNeighborCustomDataLOD)
diff --git a/Engine/Source/Runtime/Landscape/Public/LandscapeRender.h b/Engine/Source/Runtime/Landscape/Public/LandscapeRender.h
--- a/Engine/Source/Runtime/Landscape/Public/LandscapeRender.h
+++ b/Engine/Source/Runtime/Landscape/Public/LandscapeRender.h
@@ -501,6 +501,7 @@ public:
 		FViewCustomDataLOD()
 			: StaticMeshBatchLOD(INDEX_NONE)
 			, UseCombinedMeshBatch(true)
+			, IsShadowOnly(false)
 			, ComponentScreenSize(0.0f)
 			, ShaderCurrentLOD(ForceInitToZero)
 			, LodBias(ForceInitToZero)
@@ -509,6 +510,7 @@ public:
 
 		int8 StaticMeshBatchLOD;
 		bool UseCombinedMeshBatch;
+		bool IsShadowOnly;
 		float ComponentScreenSize;
 		TStaticArray<FViewCustomDataSubSectionLOD, MAX_SUBSECTION_COUNT> SubSections; // We always have at least 1 subsections
 
@@ -677,7 +679,7 @@ public:
 	virtual void OnTransformChanged() override;
 	virtual void CreateRenderThreadResources() override;
 	virtual void OnLevelAddedToWorld() override;
-	virtual void* InitViewCustomData(const FSceneView& InView, float InViewLODScale, FMemStackBase& InCustomDataMemStack, bool InIsStaticRelevant = false, const FLODMask* InVisiblePrimitiveLODMask = nullptr, float InMeshScreenSizeSquared = -1.0f) override;
+	virtual void* InitViewCustomData(const FSceneView& InView, float InViewLODScale, FMemStackBase& InCustomDataMemStack, bool InIsStaticRelevant, bool InIsShadowOnly, const FLODMask* InVisiblePrimitiveLODMask = nullptr, float InMeshScreenSizeSquared = -1.0f) override;
 	virtual void PostInitViewCustomData(const FSceneView& InView, void* InViewCustomData) const override;
 	virtual bool IsUsingCustomLODRules() const override;
 	virtual FLODMask GetCustomLOD(const FSceneView& View, float InViewLODScale, int32 InForcedLODLevel, float& OutScreenSizeSquared) const override;
diff --git a/Engine/Source/Runtime/Renderer/Private/SceneVisibility.cpp b/Engine/Source/Runtime/Renderer/Private/SceneVisibility.cpp
--- a/Engine/Source/Runtime/Renderer/Private/SceneVisibility.cpp
+++ b/Engine/Source/Runtime/Renderer/Private/SceneVisibility.cpp
@@ -2048,7 +2048,7 @@ struct FRelevancePacket
 
 			if (OutHasViewCustomDataMasks[PrimitiveIndex] != 0) // Has a relevance for this view
 			{
-				UserViewCustomData = PrimitiveSceneInfo->Proxy->InitViewCustomData(View, View.LODDistanceFactor, PrimitiveCustomDataMemStack, true, &LODToRender, MeshScreenSizeSquared);
+				UserViewCustomData = PrimitiveSceneInfo->Proxy->InitViewCustomData(View, View.LODDistanceFactor, PrimitiveCustomDataMemStack, true, false, &LODToRender, MeshScreenSizeSquared);
 
 				if (UserViewCustomData != nullptr)
 				{
@@ -2528,7 +2528,7 @@ static void SetDynamicMeshElementViewCustomData(TArray<FViewInfo>& InViews, cons
 
 			if (InHasViewCustomDataMasks[PrimitiveIndex] & (1 << ViewIndex) && ViewInfo.GetCustomData(InPrimitiveSceneInfo->GetIndex()) == nullptr)
 			{
-				ViewInfo.SetCustomData(InPrimitiveSceneInfo, InPrimitiveSceneInfo->Proxy->InitViewCustomData(ViewInfo, ViewInfo.LODDistanceFactor, ViewInfo.GetCustomDataGlobalMemStack()));
+				ViewInfo.SetCustomData(InPrimitiveSceneInfo, InPrimitiveSceneInfo->Proxy->InitViewCustomData(ViewInfo, ViewInfo.LODDistanceFactor, ViewInfo.GetCustomDataGlobalMemStack(), false, false));
 			}
 		}
 	}
diff --git a/Engine/Source/Runtime/Renderer/Private/ShadowSetup.cpp b/Engine/Source/Runtime/Renderer/Private/ShadowSetup.cpp
--- a/Engine/Source/Runtime/Renderer/Private/ShadowSetup.cpp
+++ b/Engine/Source/Runtime/Renderer/Private/ShadowSetup.cpp
@@ -877,7 +877,7 @@ bool FProjectedShadowInfo::ShouldDrawStaticMeshes(FViewInfo& InCurrentView, bool
 			{
 				if (InCurrentView.GetCustomData(InPrimitiveSceneInfo->GetIndex()) == nullptr)
 				{
-					InCurrentView.SetCustomData(InPrimitiveSceneInfo, InPrimitiveSceneInfo->Proxy->InitViewCustomData(InCurrentView, InCurrentView.LODDistanceFactor, InCurrentView.GetCustomDataGlobalMemStack(), true, &ViewLODToRender, MeshScreenSizeSquared));
+					InCurrentView.SetCustomData(InPrimitiveSceneInfo, InPrimitiveSceneInfo->Proxy->InitViewCustomData(InCurrentView, InCurrentView.LODDistanceFactor, InCurrentView.GetCustomDataGlobalMemStack(), true, true, &ViewLODToRender, MeshScreenSizeSquared));
 				}
 			}
 		}
@@ -1472,7 +1472,7 @@ void FProjectedShadowInfo::GatherDynamicMeshElementsArray(
 			{
 				if (DependentView->GetCustomData(PrimitiveSceneInfo->GetIndex()) == nullptr)
 				{
-					void* CustomData = PrimitiveSceneInfo->Proxy->InitViewCustomData(*DependentView, DependentView->LODDistanceFactor, DependentView->GetCustomDataGlobalMemStack());
+					void* CustomData = PrimitiveSceneInfo->Proxy->InitViewCustomData(*DependentView, DependentView->LODDistanceFactor, DependentView->GetCustomDataGlobalMemStack(), false, true);
 					DependentView->SetCustomData(PrimitiveSceneInfo, CustomData);
 
 					// This is required as GetDynamicMeshElements will received ReusedViewsArray which contains only FoundView
