When changing component size, and we have a layer where we have 0 total weight, don't assume 255, otherwise it will fill the weightmaps of the component with wrong data. Resulting in a difference in behavior between import of a layer and Change Component Size, which is basically another reimport option.

diff --git a/Engine/Source/Runtime/Landscape/Private/LandscapeEdit.cpp b/Engine/Source/Runtime/Landscape/Private/LandscapeEdit.cpp
--- a/Engine/Source/Runtime/Landscape/Private/LandscapeEdit.cpp
+++ b/Engine/Source/Runtime/Landscape/Private/LandscapeEdit.cpp
@@ -2397,14 +2397,7 @@ LANDSCAPE_API void ALandscapeProxy::Import(
 						}
 					}
 
-					if (TotalWeight == 0)
-					{
-						if (MaxLayerIdx >= 0)
-						{
-							WeightValues[MaxLayerIdx][Idx] = 255;
-						}
-					}
-					else if (TotalWeight != 255)
+					if (TotalWeight > 0 && TotalWeight != 255)
 					{
 						// normalization...
 						float Factor = 255.0f / TotalWeight;
diff --git a/Engine/Source/Runtime/Landscape/Private/LandscapeEditInterface.cpp b/Engine/Source/Runtime/Landscape/Private/LandscapeEditInterface.cpp
--- a/Engine/Source/Runtime/Landscape/Private/LandscapeEditInterface.cpp
+++ b/Engine/Source/Runtime/Landscape/Private/LandscapeEditInterface.cpp
@@ -2766,7 +2766,7 @@ void FLandscapeEditDataInterface::SetAlphaData(ULandscapeLayerInfoObject* const
 										}
 									}
 
-									if (OtherLayerWeightSum != 255)
+									if (OtherLayerWeightSum > 0 && OtherLayerWeightSum != 255)
 									{
 										const float Factor = 255.0f / OtherLayerWeightSum;
 										OtherLayerWeightSum = 0;
@@ -3909,9 +3909,8 @@ void FLandscapeEditDataInterface::GetWeightDataTempl(ULandscapeLayerInfoObject*
 									}
 
 									Weight = FinalValue;
+									StoreData.Store(LandscapeX, LandscapeY, Weight);
 								}
-
-								StoreData.Store(LandscapeX, LandscapeY, Weight);
 							}
 							else // Whole weight map case... no interpolation now...
 							{
