When Calling GetWeightmapData() it will try interpolating the missing component data, but there was an error and would also try interpolating missing/not allocated allocation. For example, if in a component you had no "grass" painted, but got some painted in the whole landscape. It would try to set some grass even if there was none before. We only want to perform this behavior on missing component, not existing one.

diff --git a/Engine/Source/Runtime/Landscape/Private/LandscapeEditInterface.cpp b/Engine/Source/Runtime/Landscape/Private/LandscapeEditInterface.cpp
--- a/Engine/Source/Runtime/Landscape/Private/LandscapeEditInterface.cpp
+++ b/Engine/Source/Runtime/Landscape/Private/LandscapeEditInterface.cpp
@@ -3459,7 +3459,8 @@ void FLandscapeEditDataInterface::GetWeightDataTempl(ULandscapeLayerInfoObject*
 			uint8 CornerSet = 0;
 			bool ExistLeft = ComponentIndexXX > 0 && ComponentDataExist[ ComponentIndexXX-1 + ComponentIndexYY * ComponentSizeX ];
 			bool ExistUp = ComponentIndexYY > 0 && ComponentDataExist[ ComponentIndexXX + (ComponentIndexYY-1) * ComponentSizeX ];
-
+			bool IsBorderPosition = false;
+		
 			if( Component )
 			{
 				if (LayerInfo != NULL)
@@ -3506,6 +3507,14 @@ void FLandscapeEditDataInterface::GetWeightDataTempl(ULandscapeLayerInfoObject*
 					bHasMissingValue = true;
 				}
 
+				if (LandscapeInfo->XYtoComponentMap.FindRef(FIntPoint(ComponentIndexX - 1, ComponentIndexY)) == nullptr
+					|| LandscapeInfo->XYtoComponentMap.FindRef(FIntPoint(ComponentIndexX, ComponentIndexY - 1)) == nullptr
+					|| LandscapeInfo->XYtoComponentMap.FindRef(FIntPoint(ComponentIndexX + 1, ComponentIndexY)) == nullptr
+					|| LandscapeInfo->XYtoComponentMap.FindRef(FIntPoint(ComponentIndexX, ComponentIndexY + 1)) == nullptr)
+				{
+					IsBorderPosition = true;
+				}
+
 				// Search for neighbor component for interpolation
 				bool bShouldSearchX = (BorderComponent[1] && BorderComponent[1]->GetSectionBase().X / ComponentSizeQuads <= ComponentIndexX);
 				bool bShouldSearchY = (BorderComponentY2[ComponentIndexXX] && BorderComponentY2[ComponentIndexXX]->GetSectionBase().Y / ComponentSizeQuads <= ComponentIndexY);
@@ -3792,14 +3801,14 @@ void FLandscapeEditDataInterface::GetWeightDataTempl(ULandscapeLayerInfoObject*
 							if (LayerInfo != NULL)
 							{
 								// Find the input data corresponding to this vertex
-								uint8 Weight;
+								uint8 Weight = 0;
 								if( WeightmapTexture )
 								{
 									// Find the texture data corresponding to this vertex
 									Weight = GetWeightMapData(Component, LayerInfo, (SubsectionSizeQuads+1) * SubIndexX + SubX, (SubsectionSizeQuads+1) * SubIndexY + SubY, WeightmapChannelOffset, WeightmapTexture, WeightmapTextureData );
-									StoreData.Store(LandscapeX, LandscapeY, Weight);
 								}
-								else
+								// only try generating between component if we have no component, otherwise we might end up trying to add some layer info to some component that do not originally had the allocation.
+								else if (Component == nullptr && !IsBorderPosition)
 								{
 									// Find the texture data corresponding to this vertex
 									uint8 Value[4] = {0, 0, 0, 0};
@@ -3909,8 +3918,9 @@ void FLandscapeEditDataInterface::GetWeightDataTempl(ULandscapeLayerInfoObject*
 									}
 
 									Weight = FinalValue;
-									StoreData.Store(LandscapeX, LandscapeY, Weight);
 								}
+
+								StoreData.Store(LandscapeX, LandscapeY, Weight);
 							}
 							else // Whole weight map case... no interpolation now...
 							{
