Fixed compilation error

diff --git a/Engine/Source/Runtime/Landscape/Private/LandscapeEdit.cpp b/Engine/Source/Runtime/Landscape/Private/LandscapeEdit.cpp
--- a/Engine/Source/Runtime/Landscape/Private/LandscapeEdit.cpp
+++ b/Engine/Source/Runtime/Landscape/Private/LandscapeEdit.cpp
@@ -2773,6 +2773,7 @@ LANDSCAPE_API void ALandscapeProxy::Import(
 
 		// Update MaterialInstances (must be done after textures are fully initialized)
 		UpdateAllComponentMaterialInstances(UpdateContext, RecreateRenderStateContexts);
+	}
 
 	// Recreate the render state for this component, needed to update the static drawlist which has cached the MaterialRenderProxies
 	// Must be after the FMaterialUpdateContext is destroyed
