Fixed Mac/Linux compilation error

diff --git a/Engine/Source/Runtime/Landscape/Classes/Landscape.h b/Engine/Source/Runtime/Landscape/Classes/Landscape.h
--- a/Engine/Source/Runtime/Landscape/Classes/Landscape.h
+++ b/Engine/Source/Runtime/Landscape/Classes/Landscape.h
@@ -173,7 +173,7 @@ public:
 	LANDSCAPE_API static void CalcComponentIndicesNoOverlap(const int32 X1, const int32 Y1, const int32 X2, const int32 Y2, const int32 ComponentSizeQuads,
 		int32& ComponentIndexX1, int32& ComponentIndexY1, int32& ComponentIndexX2, int32& ComponentIndexY2);
 
-	static void SplitHeightmap(ULandscapeComponent* Comp, bool bMoveToCurrentLevel = false, FMaterialUpdateContext* InOutUpdateContext = nullptr, TArray<FComponentRecreateRenderStateContext>* InOutRecreateRenderStateContext = nullptr, bool InReregisterComponent = true);
+	static void SplitHeightmap(ULandscapeComponent* Comp, bool bMoveToCurrentLevel = false, class FMaterialUpdateContext* InOutUpdateContext = nullptr, TArray<class FComponentRecreateRenderStateContext>* InOutRecreateRenderStateContext = nullptr, bool InReregisterComponent = true);
 	
 	//~ Begin UObject Interface.
 	virtual void PreSave(const class ITargetPlatform* TargetPlatform) override;
diff --git a/Engine/Source/Runtime/Landscape/Classes/LandscapeComponent.h b/Engine/Source/Runtime/Landscape/Classes/LandscapeComponent.h
--- a/Engine/Source/Runtime/Landscape/Classes/LandscapeComponent.h
+++ b/Engine/Source/Runtime/Landscape/Classes/LandscapeComponent.h
@@ -640,7 +640,7 @@ public:
 	 * Creates the MaterialInstance if it doesn't exist.
 	 */
 	LANDSCAPE_API void UpdateMaterialInstances();
-	LANDSCAPE_API void UpdateMaterialInstances(FMaterialUpdateContext& InOutMaterialContext, TArray<FComponentRecreateRenderStateContext>& InOutRecreateRenderStateContext);
+	LANDSCAPE_API void UpdateMaterialInstances(FMaterialUpdateContext& InOutMaterialContext, TArray<class FComponentRecreateRenderStateContext>& InOutRecreateRenderStateContext);
 
 	// Internal implementation of UpdateMaterialInstances, not safe to call directly
 	void UpdateMaterialInstances_Internal(FMaterialUpdateContext& Context);
diff --git a/Engine/Source/Runtime/Landscape/Classes/LandscapeProxy.h b/Engine/Source/Runtime/Landscape/Classes/LandscapeProxy.h
--- a/Engine/Source/Runtime/Landscape/Classes/LandscapeProxy.h
+++ b/Engine/Source/Runtime/Landscape/Classes/LandscapeProxy.h
@@ -891,7 +891,7 @@ public:
 
 	/** Update the material instances for all the landscape components */
 	LANDSCAPE_API void UpdateAllComponentMaterialInstances();
-	LANDSCAPE_API void UpdateAllComponentMaterialInstances(FMaterialUpdateContext& InOutMaterialContext, TArray<FComponentRecreateRenderStateContext>& InOutRecreateRenderStateContext);
+	LANDSCAPE_API void UpdateAllComponentMaterialInstances(FMaterialUpdateContext& InOutMaterialContext, TArray<class FComponentRecreateRenderStateContext>& InOutRecreateRenderStateContext);
 
 	/** Create a thumbnail material for a given layer */
 	LANDSCAPE_API static ULandscapeMaterialInstanceConstant* GetLayerThumbnailMIC(UMaterialInterface* LandscapeMaterial, FName LayerName, UTexture2D* ThumbnailWeightmap, UTexture2D* ThumbnailHeightmap, ALandscapeProxy* Proxy);
