Removed unrequired code about material count
Added guard to make sure indices are valid before using them

diff --git a/Engine/Source/Runtime/Landscape/Private/LandscapeRender.cpp b/Engine/Source/Runtime/Landscape/Private/LandscapeRender.cpp
--- a/Engine/Source/Runtime/Landscape/Private/LandscapeRender.cpp
+++ b/Engine/Source/Runtime/Landscape/Private/LandscapeRender.cpp
@@ -1283,36 +1283,21 @@ void FLandscapeComponentSceneProxy::BuildDynamicMeshElement(const FViewCustomDat
 	}
 
 	const int8 CurrentLODIndex = InPrimitiveCustomData->SubSections[0].BatchElementCurrentLOD;
-	UMaterialInterface* SelectedMaterial = AvailableMaterials[LODIndexToMaterialIndex[CurrentLODIndex]];
+	int8 MaterialIndex = LODIndexToMaterialIndex.IsValidIndex(CurrentLODIndex) ? LODIndexToMaterialIndex[CurrentLODIndex] : INDEX_NONE;
+	UMaterialInterface* SelectedMaterial = MaterialIndex != INDEX_NONE ? AvailableMaterials[MaterialIndex] : nullptr;
 
-	if (InHasTessellation)
-	{
-		int32 MaterialCount = AvailableMaterials.Num();
-
-		for (int32 i = 0; i < AvailableMaterials.Num(); ++i)
+	if (InHasTessellation && MaterialIndex != INDEX_NONE)
+	{		
+		if (InDisableTessellation && MaterialIndexToDisabledTessellationMaterial.IsValidIndex(MaterialIndex))
 		{
-			ULandscapeMaterialInstanceConstant* LandscapeMIC = Cast<ULandscapeMaterialInstanceConstant>(AvailableMaterials[i]);
-
-			if (LandscapeMIC == nullptr)
-			{
-				UMaterialInstanceDynamic* LandscapeMID = Cast<UMaterialInstanceDynamic>(AvailableMaterials[i]);
-
-				if (LandscapeMID != nullptr)
-				{
-					LandscapeMIC = Cast<ULandscapeMaterialInstanceConstant>(LandscapeMID->Parent);
-				}
-			}
-
-			bool HasTessellationEnabled = MaterialHasTessellationEnabled[i];
-			
-			// Remove material with disabled tessellation as they were generated
-			if (HasTessellationEnabled && LandscapeMIC != nullptr && LandscapeMIC->bDisableTessellation)
-			{
-				--MaterialCount;
-			}			
+			SelectedMaterial = AvailableMaterials[MaterialIndexToDisabledTessellationMaterial[MaterialIndex]];
 		}
+	}
 
-		SelectedMaterial = AvailableMaterials[InDisableTessellation ? MaterialIndexToDisabledTessellationMaterial[LODIndexToMaterialIndex[CurrentLODIndex]] : LODIndexToMaterialIndex[CurrentLODIndex]];
+	// this is really not normal that we have no material at this point, so do not continue
+	if (SelectedMaterial == nullptr)
+	{
+		return;
 	}
 
 	// Could be different from bRequiresAdjacencyInformation during shader compilation
