Call GetLODView to correctly use the right view to compute LOD in VR.

diff --git a/Engine/Source/Runtime/Landscape/Private/LandscapeRender.cpp b/Engine/Source/Runtime/Landscape/Private/LandscapeRender.cpp
--- a/Engine/Source/Runtime/Landscape/Private/LandscapeRender.cpp
+++ b/Engine/Source/Runtime/Landscape/Private/LandscapeRender.cpp
@@ -1717,13 +1717,14 @@ uint64 FLandscapeComponentSceneProxy::GetStaticBatchElementVisibility(const FSce
 	uint64 BatchesToRenderMask = 0;
 
 	SCOPE_CYCLE_COUNTER(STAT_LandscapeStaticDrawLODTime);
+	const FSceneView& View = GetLODView(InView);
 
-	const void* ViewCustomData = InViewCustomData != nullptr ? InViewCustomData : InView.GetCustomData(GetPrimitiveSceneInfo()->GetIndex());
+	const void* ViewCustomData = InViewCustomData != nullptr ? InViewCustomData : View.GetCustomData(GetPrimitiveSceneInfo()->GetIndex());
 
 	if (ViewCustomData != nullptr)
 	{
 		const FViewCustomDataLOD* CurrentLODData = (const FViewCustomDataLOD*)ViewCustomData;
-		const auto FeatureLevel = InView.GetFeatureLevel();
+		const auto FeatureLevel = View.GetFeatureLevel();
 
 		if (FeatureLevel >= ERHIFeatureLevel::SM4)
 		{
@@ -1792,7 +1793,6 @@ void FLandscapeComponentSceneProxy::CalculateBatchElementLOD(const FSceneView& I
 
 		float SubSectionMaxExtend = ComponentMaxExtend / 2.0f;
 		float SubSectionRadius = LandscapeComponent->Bounds.SphereRadius / 2.0f;
-		const FSceneView& View = GetLODView(InView);
 		float CombinedScreenRatio = 0.0f;
 		bool AllSubSectionHaveSameScreenSize = true;
 
@@ -1804,11 +1804,11 @@ void FLandscapeComponentSceneProxy::CalculateBatchElementLOD(const FSceneView& I
 				int32 SubSectionIndex = SubX + SubY * NumSubsections;
 				FViewCustomDataSubSectionLOD& SubSectionLODData = InOutLODData.SubSections[SubSectionIndex];
 
-				SubSectionLODData.ScreenSizeSquared = GetComponentScreenSize(&View, SubSectionScreenSizeTestingPosition[SubSectionIndex], SubSectionMaxExtend, SubSectionRadius);
+				SubSectionLODData.ScreenSizeSquared = GetComponentScreenSize(&InView, SubSectionScreenSizeTestingPosition[SubSectionIndex], SubSectionMaxExtend, SubSectionRadius);
 
 				check(SubSectionLODData.ScreenSizeSquared > 0.0f);
 
-				CalculateLODFromScreenSize(View, SubSectionLODData.ScreenSizeSquared, InViewLODScale, SubSectionIndex, InOutLODData);
+				CalculateLODFromScreenSize(InView, SubSectionLODData.ScreenSizeSquared, InViewLODScale, SubSectionIndex, InOutLODData);
 				check(SubSectionLODData.fBatchElementCurrentLOD != -1.0f);
 
 				InOutLODData.ShaderCurrentLOD.Component(SubSectionIndex) = SubSectionLODData.fBatchElementCurrentLOD;
@@ -2156,7 +2156,7 @@ float FLandscapeComponentSceneProxy::GetNeighborLOD(const FSceneView& InView, fl
 			FVector ComponentTopLeftCorner = LandscapeComponentOrigin - FVector(SubSectionMaxExtend, SubSectionMaxExtend, 0.0f);
 
 			FVector SubSectionOrigin = ComponentTopLeftCorner + FVector(LandscapeComponentMaxExtends * DesiredSubSectionX, LandscapeComponentMaxExtends * DesiredSubSectionY, 0.0f);
-			float MeshBatchScreenSizeSquared = GetComponentScreenSize(&GetLODView(InView), SubSectionOrigin, SubSectionMaxExtend, LandscapeComponent->Bounds.SphereRadius / 2.0f);
+			float MeshBatchScreenSizeSquared = GetComponentScreenSize(&InView, SubSectionOrigin, SubSectionMaxExtend, LandscapeComponent->Bounds.SphereRadius / 2.0f);
 
 			FViewCustomDataLOD NeighborLODData;
 			CalculateLODFromScreenSize(InView, MeshBatchScreenSizeSquared, InView.LODDistanceFactor, DesiredSubSectionIndex, NeighborLODData);
@@ -2170,7 +2170,7 @@ float FLandscapeComponentSceneProxy::GetNeighborLOD(const FSceneView& InView, fl
 		}
 		else
 		{
-			float MeshBatchScreenSizeSquared = GetComponentScreenSize(&GetLODView(InView), LandscapeComponentOrigin, LandscapeComponentMaxExtends, LandscapeComponent->Bounds.SphereRadius);
+			float MeshBatchScreenSizeSquared = GetComponentScreenSize(&InView, LandscapeComponentOrigin, LandscapeComponentMaxExtends, LandscapeComponent->Bounds.SphereRadius);
 
 			FViewCustomDataLOD NeighborLODData;
 			CalculateLODFromScreenSize(InView, MeshBatchScreenSizeSquared, InView.LODDistanceFactor, 0, NeighborLODData);
@@ -2204,13 +2204,15 @@ void FLandscapeComponentSceneProxy::PostInitViewCustomData(const FSceneView& InV
 	FViewCustomDataLOD* CurrentLODData = (FViewCustomDataLOD*)InViewCustomData;
 	check(CurrentLODData != nullptr);
 
+	const FSceneView& View = GetLODView(InView);
+
 	for (int8 SubY = 0; SubY < NumSubsections; ++SubY)
 	{
 		for (int8 SubX = 0; SubX < NumSubsections; ++SubX)
 		{
 			int8 SubSectionIndex = SubX + SubY * NumSubsections;
 			FViewCustomDataSubSectionLOD& SubSectionLODData = CurrentLODData->SubSections[SubSectionIndex];
-			GetShaderCurrentNeighborLOD(InView, SubSectionLODData.fBatchElementCurrentLOD, NumSubsections > 1 ? SubX : INDEX_NONE, NumSubsections > 1 ? SubY : INDEX_NONE, SubSectionIndex, SubSectionLODData.ShaderCurrentNeighborLOD);
+			GetShaderCurrentNeighborLOD(View, SubSectionLODData.fBatchElementCurrentLOD, NumSubsections > 1 ? SubX : INDEX_NONE, NumSubsections > 1 ? SubY : INDEX_NONE, SubSectionIndex, SubSectionLODData.ShaderCurrentNeighborLOD);
 		}
 	}
 
@@ -2296,7 +2298,7 @@ FLODMask FLandscapeComponentSceneProxy::GetCustomLOD(const FSceneView& View, flo
 			float SubSectionRadius = LandscapeComponent->Bounds.SphereRadius / 2.0f;
 
 			// Compute screen size of each sub section to determine if we should use the combined logic or the individual logic
-			float ScreenSizeSquared = GetComponentScreenSize(&View, SubSectionScreenSizeTestingPosition[0], SubSectionMaxExtend, SubSectionRadius);
+			float ScreenSizeSquared = GetComponentScreenSize(&LODView, SubSectionScreenSizeTestingPosition[0], SubSectionMaxExtend, SubSectionRadius);
 			PotentialLOD = GetLODFromScreenSize(ScreenSizeSquared, InViewLODScale);
 		}
 		else
@@ -2304,7 +2306,7 @@ FLODMask FLandscapeComponentSceneProxy::GetCustomLOD(const FSceneView& View, flo
 			PotentialLOD = GetLODFromScreenSize(OutScreenSizeSquared, InViewLODScale);
 		}
 
-		const auto FeatureLevel = View.GetFeatureLevel();
+		const auto FeatureLevel = LODView.GetFeatureLevel();
 		bool HasTessellationEnabled = false;
 
 		if (FeatureLevel >= ERHIFeatureLevel::SM4)
@@ -2394,7 +2396,7 @@ FLODMask FLandscapeComponentSceneProxy::GetCustomWholeSceneShadowLOD(const FScen
 			PotentialLOD = PrimitiveCustomData->SubSections[0].BatchElementCurrentLOD;
 		}
 
-		const auto FeatureLevel = InView.GetFeatureLevel();
+		const auto FeatureLevel = LODView.GetFeatureLevel();
 		bool HasTessellationEnabled = false;
 
 		if (FeatureLevel >= ERHIFeatureLevel::SM4)
@@ -3431,7 +3433,7 @@ public:
 
 		if (LODData != nullptr)
 		{
-			SceneProxy->PostInitViewCustomData(*InView, LODData);
+			SceneProxy->PostInitViewCustomData(View, LODData);
 
 			if (LodTessellationParameter.IsBound())
 			{
@@ -3491,7 +3493,7 @@ public:
 			float ComponentScreenSize = SceneProxy->GetComponentScreenSize(&View, SceneProxy->LandscapeComponent->Bounds.Origin, SceneProxy->ComponentMaxExtend, SceneProxy->LandscapeComponent->Bounds.SphereRadius);
 			
 			FLandscapeComponentSceneProxy::FViewCustomDataLOD CurrentLODData;
-			SceneProxy->CalculateBatchElementLOD(*InView, ComponentScreenSize, View.LODDistanceFactor, CurrentLODData, true);
+			SceneProxy->CalculateBatchElementLOD(View, ComponentScreenSize, View.LODDistanceFactor, CurrentLODData, true);
 			check(CurrentLODData.UseCombinedMeshBatch);
 
 			if (LodBiasParameter.IsBound())
