Remove all call GetLODView as both eye would use the same custom data, that could result in a race condition.

diff --git a/Engine/Source/Runtime/Landscape/Private/LandscapeRender.cpp b/Engine/Source/Runtime/Landscape/Private/LandscapeRender.cpp
--- a/Engine/Source/Runtime/Landscape/Private/LandscapeRender.cpp
+++ b/Engine/Source/Runtime/Landscape/Private/LandscapeRender.cpp
@@ -1792,7 +1792,6 @@ void FLandscapeComponentSceneProxy::CalculateBatchElementLOD(const FSceneView& I
 
 		float SubSectionMaxExtend = ComponentMaxExtend / 2.0f;
 		float SubSectionRadius = LandscapeComponent->Bounds.SphereRadius / 2.0f;
-		const FSceneView& View = GetLODView(InView);
 		float CombinedScreenRatio = 0.0f;
 		bool AllSubSectionHaveSameScreenSize = true;
 
@@ -1804,11 +1803,11 @@ void FLandscapeComponentSceneProxy::CalculateBatchElementLOD(const FSceneView& I
 				int32 SubSectionIndex = SubX + SubY * NumSubsections;
 				FViewCustomDataSubSectionLOD& SubSectionLODData = InOutLODData.SubSections[SubSectionIndex];
 
-				SubSectionLODData.ScreenSizeSquared = GetComponentScreenSize(&View, SubSectionScreenSizeTestingPosition[SubSectionIndex], SubSectionMaxExtend, SubSectionRadius);
+				SubSectionLODData.ScreenSizeSquared = GetComponentScreenSize(&InView, SubSectionScreenSizeTestingPosition[SubSectionIndex], SubSectionMaxExtend, SubSectionRadius);
 
 				check(SubSectionLODData.ScreenSizeSquared > 0.0f);
 
-				CalculateLODFromScreenSize(View, SubSectionLODData.ScreenSizeSquared, InViewLODScale, SubSectionIndex, InOutLODData);
+				CalculateLODFromScreenSize(InView, SubSectionLODData.ScreenSizeSquared, InViewLODScale, SubSectionIndex, InOutLODData);
 				check(SubSectionLODData.fBatchElementCurrentLOD != -1.0f);
 
 				InOutLODData.ShaderCurrentLOD.Component(SubSectionIndex) = SubSectionLODData.fBatchElementCurrentLOD;
@@ -1945,7 +1944,6 @@ void* FLandscapeComponentSceneProxy::InitViewCustomData(const FSceneView& InView
 	// NOTE: we can't access other proxy here as this can be run in parallel we need to wait for the PostInitViewCustomData which is run in synchronous	
 
 	PrimitiveCustomDataIndex = GetPrimitiveSceneInfo()->GetIndex();
-	const FSceneView& View = GetLODView(InView);
 
 	FViewCustomDataLOD* LODData = (FViewCustomDataLOD*)new(InCustomDataMemStack) FViewCustomDataLOD();
 
@@ -1955,10 +1953,10 @@ void* FLandscapeComponentSceneProxy::InitViewCustomData(const FSceneView& InView
 	// If a valid screen size was provided, we use it instead of recomputing it
 	if (InMeshScreenSizeSquared < 0.0f)
 	{
-		LODData->ComponentScreenSize = GetComponentScreenSize(&View, LandscapeComponent->Bounds.Origin, ComponentMaxExtend, LandscapeComponent->Bounds.SphereRadius);
+		LODData->ComponentScreenSize = GetComponentScreenSize(&InView, LandscapeComponent->Bounds.Origin, ComponentMaxExtend, LandscapeComponent->Bounds.SphereRadius);
 	}
 
-	CalculateBatchElementLOD(View, LODData->ComponentScreenSize, InViewLODScale, *LODData, false);
+	CalculateBatchElementLOD(InView, LODData->ComponentScreenSize, InViewLODScale, *LODData, false);
 
 	if (InIsStaticRelevant)
 	{
@@ -1986,7 +1984,7 @@ void* FLandscapeComponentSceneProxy::InitViewCustomData(const FSceneView& InView
 		LODData->LodBias = GetShaderLODBias();
 	}
 
-	ComputeTessellationFalloffShaderValues(*LODData, View.ViewMatrices.GetProjectionMatrix(), LODData->LodTessellationParams.X, LODData->LodTessellationParams.Y);
+	ComputeTessellationFalloffShaderValues(*LODData, InView.ViewMatrices.GetProjectionMatrix(), LODData->LodTessellationParams.X, LODData->LodTessellationParams.Y);
 
 	return LODData;
 }
@@ -2156,7 +2154,7 @@ float FLandscapeComponentSceneProxy::GetNeighborLOD(const FSceneView& InView, fl
 			FVector ComponentTopLeftCorner = LandscapeComponentOrigin - FVector(SubSectionMaxExtend, SubSectionMaxExtend, 0.0f);
 
 			FVector SubSectionOrigin = ComponentTopLeftCorner + FVector(LandscapeComponentMaxExtends * DesiredSubSectionX, LandscapeComponentMaxExtends * DesiredSubSectionY, 0.0f);
-			float MeshBatchScreenSizeSquared = GetComponentScreenSize(&GetLODView(InView), SubSectionOrigin, SubSectionMaxExtend, LandscapeComponent->Bounds.SphereRadius / 2.0f);
+			float MeshBatchScreenSizeSquared = GetComponentScreenSize(&InView, SubSectionOrigin, SubSectionMaxExtend, LandscapeComponent->Bounds.SphereRadius / 2.0f);
 
 			FViewCustomDataLOD NeighborLODData;
 			CalculateLODFromScreenSize(InView, MeshBatchScreenSizeSquared, InView.LODDistanceFactor, DesiredSubSectionIndex, NeighborLODData);
@@ -2170,7 +2168,7 @@ float FLandscapeComponentSceneProxy::GetNeighborLOD(const FSceneView& InView, fl
 		}
 		else
 		{
-			float MeshBatchScreenSizeSquared = GetComponentScreenSize(&GetLODView(InView), LandscapeComponentOrigin, LandscapeComponentMaxExtends, LandscapeComponent->Bounds.SphereRadius);
+			float MeshBatchScreenSizeSquared = GetComponentScreenSize(&InView, LandscapeComponentOrigin, LandscapeComponentMaxExtends, LandscapeComponent->Bounds.SphereRadius);
 
 			FViewCustomDataLOD NeighborLODData;
 			CalculateLODFromScreenSize(InView, MeshBatchScreenSizeSquared, InView.LODDistanceFactor, 0, NeighborLODData);
@@ -2269,11 +2267,10 @@ bool FLandscapeComponentSceneProxy::CanUseMeshBatchForShadowCascade(int8 InLODIn
 	return MeshBatch->TessellationDisablingShadowMapMeshSize >= WorldUnitsForOneTexel * (GShadowMapWorldUnitsToTexelFactor != -1.0f ? GShadowMapWorldUnitsToTexelFactor : 1.0f);
 }
 
-FLODMask FLandscapeComponentSceneProxy::GetCustomLOD(const FSceneView& View, float InViewLODScale, int32 InForcedLODLevel, float& OutScreenSizeSquared) const
+FLODMask FLandscapeComponentSceneProxy::GetCustomLOD(const FSceneView& InView, float InViewLODScale, int32 InForcedLODLevel, float& OutScreenSizeSquared) const
 {
 	SCOPE_CYCLE_COUNTER(STAT_LandscapeComputeCustomMeshBatchLOD);
 	FLODMask LODToRender;
-	const FSceneView& LODView = GetLODView(View);
 	OutScreenSizeSquared = 0.0f;
 
 	// Handle forced LOD level first
@@ -2285,10 +2282,10 @@ FLODMask FLandscapeComponentSceneProxy::GetCustomLOD(const FSceneView& View, flo
 
 		LODToRender.SetLOD(FMath::Clamp<int8>(InForcedLODLevel, MinMeshLOD, MaxMeshLOD));
 	}
-	else if (LODView.Family->EngineShowFlags.LOD)
+	else if (InView.Family->EngineShowFlags.LOD)
 	{
 		int8 PotentialLOD = INDEX_NONE;
-		OutScreenSizeSquared = GetComponentScreenSize(&LODView, LandscapeComponent->Bounds.Origin, ComponentMaxExtend, LandscapeComponent->Bounds.SphereRadius);
+		OutScreenSizeSquared = GetComponentScreenSize(&InView, LandscapeComponent->Bounds.Origin, ComponentMaxExtend, LandscapeComponent->Bounds.SphereRadius);
 
 		if (NumSubsections > 1)
 		{
@@ -2296,7 +2293,7 @@ FLODMask FLandscapeComponentSceneProxy::GetCustomLOD(const FSceneView& View, flo
 			float SubSectionRadius = LandscapeComponent->Bounds.SphereRadius / 2.0f;
 
 			// Compute screen size of each sub section to determine if we should use the combined logic or the individual logic
-			float ScreenSizeSquared = GetComponentScreenSize(&View, SubSectionScreenSizeTestingPosition[0], SubSectionMaxExtend, SubSectionRadius);
+			float ScreenSizeSquared = GetComponentScreenSize(&InView, SubSectionScreenSizeTestingPosition[0], SubSectionMaxExtend, SubSectionRadius);
 			PotentialLOD = GetLODFromScreenSize(ScreenSizeSquared, InViewLODScale);
 		}
 		else
@@ -2304,7 +2301,7 @@ FLODMask FLandscapeComponentSceneProxy::GetCustomLOD(const FSceneView& View, flo
 			PotentialLOD = GetLODFromScreenSize(OutScreenSizeSquared, InViewLODScale);
 		}
 
-		const auto FeatureLevel = View.GetFeatureLevel();
+		const auto FeatureLevel = InView.GetFeatureLevel();
 		bool HasTessellationEnabled = false;
 
 		if (FeatureLevel >= ERHIFeatureLevel::SM4)
@@ -2365,14 +2362,13 @@ FLODMask FLandscapeComponentSceneProxy::GetCustomWholeSceneShadowLOD(const FScen
 	}
 	else
 	{
-		const FSceneView& LODView = GetLODView(InView);
-		const FViewCustomDataLOD* PrimitiveCustomData = (const FViewCustomDataLOD*)LODView.GetCustomData(GetPrimitiveSceneInfo()->GetIndex());
+		const FViewCustomDataLOD* PrimitiveCustomData = (const FViewCustomDataLOD*)InView.GetCustomData(GetPrimitiveSceneInfo()->GetIndex());
 		int8 PotentialLOD = INDEX_NONE;
 		float ScreenSizeSquared = 0.0f;
 
 		if (PrimitiveCustomData == nullptr)
 		{
-			ScreenSizeSquared = GetComponentScreenSize(&LODView, LandscapeComponent->Bounds.Origin, ComponentMaxExtend, LandscapeComponent->Bounds.SphereRadius);
+			ScreenSizeSquared = GetComponentScreenSize(&InView, LandscapeComponent->Bounds.Origin, ComponentMaxExtend, LandscapeComponent->Bounds.SphereRadius);
 
 			if (NumSubsections > 1)
 			{
@@ -2380,7 +2376,7 @@ FLODMask FLandscapeComponentSceneProxy::GetCustomWholeSceneShadowLOD(const FScen
 				float SubSectionRadius = LandscapeComponent->Bounds.SphereRadius / 2.0f;
 
 				// Compute screen size of each sub section to determine if we should use the combined logic or the individual logic
-				float SubSectionScreenSizeSquared = GetComponentScreenSize(&LODView, SubSectionScreenSizeTestingPosition[0], SubSectionMaxExtend, SubSectionRadius);
+				float SubSectionScreenSizeSquared = GetComponentScreenSize(&InView, SubSectionScreenSizeTestingPosition[0], SubSectionMaxExtend, SubSectionRadius);
 				PotentialLOD = GetLODFromScreenSize(SubSectionScreenSizeSquared, InViewLODScale);
 			}
 			else
@@ -3423,11 +3419,10 @@ public:
 			ShaderBindings.Add(LodValuesParameter, SceneProxy->GetShaderLODValues(BatchElementParams->CurrentLOD));
 		}
 
-		const FSceneView& View = GetLODView(*InView);
 		int32 SubSectionIndex = BatchElementParams->SubX + BatchElementParams->SubY * SceneProxy->NumSubsections;
 
 		// If we have no custom data for this primitive we will compute of the fly the proper values, this will happen if the shader is not used for normal landscape rendering(i.e grass rendering)
-		FLandscapeComponentSceneProxy::FViewCustomDataLOD* LODData = (FLandscapeComponentSceneProxy::FViewCustomDataLOD*)View.GetCustomData(SceneProxy->GetPrimitiveSceneInfo()->GetIndex());
+		FLandscapeComponentSceneProxy::FViewCustomDataLOD* LODData = (FLandscapeComponentSceneProxy::FViewCustomDataLOD*)InView->GetCustomData(SceneProxy->GetPrimitiveSceneInfo()->GetIndex());
 
 		if (LODData != nullptr)
 		{
@@ -3488,10 +3483,10 @@ public:
 		}
 		else
 		{
-			float ComponentScreenSize = SceneProxy->GetComponentScreenSize(&View, SceneProxy->LandscapeComponent->Bounds.Origin, SceneProxy->ComponentMaxExtend, SceneProxy->LandscapeComponent->Bounds.SphereRadius);
+			float ComponentScreenSize = SceneProxy->GetComponentScreenSize(InView, SceneProxy->LandscapeComponent->Bounds.Origin, SceneProxy->ComponentMaxExtend, SceneProxy->LandscapeComponent->Bounds.SphereRadius);
 			
 			FLandscapeComponentSceneProxy::FViewCustomDataLOD CurrentLODData;
-			SceneProxy->CalculateBatchElementLOD(*InView, ComponentScreenSize, View.LODDistanceFactor, CurrentLODData, true);
+			SceneProxy->CalculateBatchElementLOD(*InView, ComponentScreenSize, InView->LODDistanceFactor, CurrentLODData, true);
 			check(CurrentLODData.UseCombinedMeshBatch);
 
 			if (LodBiasParameter.IsBound())
@@ -3502,7 +3497,7 @@ public:
 			if (LodTessellationParameter.IsBound())
 			{
 				FVector4 LodTessellationParams(ForceInitToZero);
-				SceneProxy->ComputeTessellationFalloffShaderValues(CurrentLODData, View.ViewMatrices.GetProjectionMatrix(), LodTessellationParams.X, LodTessellationParams.Y);
+				SceneProxy->ComputeTessellationFalloffShaderValues(CurrentLODData, InView->ViewMatrices.GetProjectionMatrix(), LodTessellationParams.X, LodTessellationParams.Y);
 
 				ShaderBindings.Add(LodTessellationParameter, LodTessellationParams);
 			}
@@ -3521,7 +3516,7 @@ public:
 					for (int32 SubX = 0; SubX < SceneProxy->NumSubsections; SubX++)
 					{
 						int32 NeighborSubSectionIndex = SubX + SubY * SceneProxy->NumSubsections;
-						SceneProxy->GetShaderCurrentNeighborLOD(View, CurrentLODData.SubSections[NeighborSubSectionIndex].fBatchElementCurrentLOD, SceneProxy->NumSubsections > 1 ? SubX : INDEX_NONE, SceneProxy->NumSubsections > 1 ? SubY : INDEX_NONE, NeighborSubSectionIndex, CurrentNeighborLOD[NeighborSubSectionIndex]);
+						SceneProxy->GetShaderCurrentNeighborLOD(*InView, CurrentLODData.SubSections[NeighborSubSectionIndex].fBatchElementCurrentLOD, SceneProxy->NumSubsections > 1 ? SubX : INDEX_NONE, SceneProxy->NumSubsections > 1 ? SubY : INDEX_NONE, NeighborSubSectionIndex, CurrentNeighborLOD[NeighborSubSectionIndex]);
 						check(CurrentNeighborLOD[NeighborSubSectionIndex].X != -1.0f); // they should all match so only check the 1st one for simplicity
 					}
 				}
diff --git a/Engine/Source/Runtime/Landscape/Public/LandscapeRender.h b/Engine/Source/Runtime/Landscape/Public/LandscapeRender.h
--- a/Engine/Source/Runtime/Landscape/Public/LandscapeRender.h
+++ b/Engine/Source/Runtime/Landscape/Public/LandscapeRender.h
@@ -682,7 +682,7 @@ public:
 	virtual void* InitViewCustomData(const FSceneView& InView, float InViewLODScale, FMemStackBase& InCustomDataMemStack, bool InIsStaticRelevant, bool InIsShadowOnly, const FLODMask* InVisiblePrimitiveLODMask = nullptr, float InMeshScreenSizeSquared = -1.0f) override;
 	virtual void PostInitViewCustomData(const FSceneView& InView, void* InViewCustomData) const override;
 	virtual bool IsUsingCustomLODRules() const override;
-	virtual FLODMask GetCustomLOD(const FSceneView& View, float InViewLODScale, int32 InForcedLODLevel, float& OutScreenSizeSquared) const override;
+	virtual FLODMask GetCustomLOD(const FSceneView& InView, float InViewLODScale, int32 InForcedLODLevel, float& OutScreenSizeSquared) const override;
 	virtual bool IsUsingCustomWholeSceneShadowLODRules() const override;
 	virtual FLODMask GetCustomWholeSceneShadowLOD(const FSceneView& InView, float InViewLODScale, int32 InForcedLODLevel, const struct FLODMask& InVisibilePrimitiveLODMask, float InShadowMapTextureResolution, float InShadowMapCascadeSize, int8 InShadowCascadeId, bool InHasSelfShadow) const override;
 	
