-  Landscape Layer System : Fixed ApplySplines not working when no layer was reserved for landscape splines
#rnx

diff --git a/Engine/Source/Runtime/Landscape/Private/LandscapeSplineRaster.cpp b/Engine/Source/Runtime/Landscape/Private/LandscapeSplineRaster.cpp
--- a/Engine/Source/Runtime/Landscape/Private/LandscapeSplineRaster.cpp
+++ b/Engine/Source/Runtime/Landscape/Private/LandscapeSplineRaster.cpp
@@ -463,7 +463,7 @@ bool ULandscapeInfo::ApplySplines(bool bOnlySelected)
 
 	ALandscape* Landscape = LandscapeActor.Get();
 	const FLandscapeLayer* Layer = Landscape ? Landscape->GetLandscapeSplinesReservedLayer() : nullptr;
-	FGuid SplinesTargetLayerGuid = Layer ? Layer->Guid : FGuid();
+	FGuid SplinesTargetLayerGuid = Layer ? Layer->Guid : Landscape ? Landscape->GetEditingLayer() : FGuid();
 	FScopedSetLandscapeEditingLayer Scope(Landscape, SplinesTargetLayerGuid, [=] 
 	{ 
 		Landscape->RequestLayersContentUpdate(ELandscapeLayerUpdateMode::Update_All);
