Fix SimpleCollision Mip Level crash

#rnx

diff --git a/Engine/Source/Runtime/Landscape/Private/LandscapeEdit.cpp b/Engine/Source/Runtime/Landscape/Private/LandscapeEdit.cpp
--- a/Engine/Source/Runtime/Landscape/Private/LandscapeEdit.cpp
+++ b/Engine/Source/Runtime/Landscape/Private/LandscapeEdit.cpp
@@ -1177,7 +1177,7 @@ void ULandscapeComponent::UpdateDominantLayerBuffer(int32 InComponentX1, int32 I
 {
 	const int32 MipSizeU = InWeightmapSizeU >> InCollisionMipLevel;
 
-	FCollisionSize CollisionSize = FCollisionSize::Create(NumSubsections, SubsectionSizeQuads, CollisionMipLevel);
+	FCollisionSize CollisionSize = FCollisionSize::Create(NumSubsections, SubsectionSizeQuads, InCollisionMipLevel);
 	
 	// Ratio to convert update region coordinate to collision mip coordinates
 	const float CollisionQuadRatio = (float)CollisionSize.SubsectionSizeQuads / (float)SubsectionSizeQuads;
