- Landscape Layer System: Optimized update of material instances
#rnx

diff --git a/Engine/Source/Runtime/Landscape/Private/LandscapeEditLayers.cpp b/Engine/Source/Runtime/Landscape/Private/LandscapeEditLayers.cpp
--- a/Engine/Source/Runtime/Landscape/Private/LandscapeEditLayers.cpp
+++ b/Engine/Source/Runtime/Landscape/Private/LandscapeEditLayers.cpp
@@ -3900,7 +3900,7 @@ void ALandscape::UpdateLayersMaterialInstances()
 			const int8 MaterialLOD = ItPair.Value;
 
 			// Find or set a matching MIC in the Landscape's map.
-			UMaterialInstanceConstant* CombinationMaterialInstance = Component->GetCombinationMaterial(nullptr, WeightmapBaseLayerAllocation, MaterialLOD, false);
+			UMaterialInstanceConstant* CombinationMaterialInstance = Component->GetCombinationMaterial(&MaterialUpdateContext.GetValue(), WeightmapBaseLayerAllocation, MaterialLOD, false);
 
 			if (CombinationMaterialInstance != nullptr)
 			{
