Resolve robomerge conflicts

diff --git a/Engine/Source/Runtime/Landscape/Private/LandscapeRender.cpp b/Engine/Source/Runtime/Landscape/Private/LandscapeRender.cpp
--- a/Engine/Source/Runtime/Landscape/Private/LandscapeRender.cpp
+++ b/Engine/Source/Runtime/Landscape/Private/LandscapeRender.cpp
@@ -4094,7 +4094,7 @@ bool ULandscapeMaterialInstanceConstant::HasOverridenBaseProperties() const
 
 //////////////////////////////////////////////////////////////////////////
 
-void ULandscapeComponent::GetStreamingTextureInfo(FStreamingTextureLevelContext& LevelContext, TArray<FStreamingTexturePrimitiveInfo>& OutStreamingTextures) const
+void ULandscapeComponent::GetStreamingRenderAssetInfo(FStreamingTextureLevelContext& LevelContext, TArray<FStreamingRenderAssetPrimitiveInfo>& OutStreamingRenderAssets) const
 {
 	ALandscapeProxy* Proxy = Cast<ALandscapeProxy>(GetOuter());
 	FSphere BoundingSphere = Bounds.GetSphere();
@@ -4126,10 +4126,10 @@ void ULandscapeComponent::GetStreamingTextureInfo(FStreamingTextureLevelContext&
 				UTexture2D* Texture2D = Cast<UTexture2D>(Textures[TextureIndex]);
 				if (!Texture2D) continue;
 
-				FStreamingTexturePrimitiveInfo& StreamingTexture = *new(OutStreamingTextures)FStreamingTexturePrimitiveInfo;
+				FStreamingRenderAssetPrimitiveInfo& StreamingTexture = *new(OutStreamingRenderAssets)FStreamingRenderAssetPrimitiveInfo;
 				StreamingTexture.Bounds = BoundingSphere;
 				StreamingTexture.TexelFactor = TexelFactor;
-				StreamingTexture.Texture = Texture2D;
+				StreamingTexture.RenderAsset = Texture2D;
 			}
 
 			const UMaterial* Material = MaterialInterface->GetMaterial();
@@ -4162,10 +4162,10 @@ void ULandscapeComponent::GetStreamingTextureInfo(FStreamingTextureLevelContext&
 
 						if (TextureCoordinate || TerrainTextureCoordinate)
 						{
-							for (int32 i = 0; i < OutStreamingTextures.Num(); ++i)
+							for (int32 i = 0; i < OutStreamingRenderAssets.Num(); ++i)
 							{
-								FStreamingTexturePrimitiveInfo& StreamingTexture = OutStreamingTextures[i];
-								if (StreamingTexture.Texture == TextureSample->Texture)
+								FStreamingRenderAssetPrimitiveInfo& StreamingTexture = OutStreamingRenderAssets[i];
+								if (StreamingTexture.RenderAsset == TextureSample->Texture)
 								{
 									if (TextureCoordinate)
 									{
@@ -4194,9 +4194,9 @@ void ULandscapeComponent::GetStreamingTextureInfo(FStreamingTextureLevelContext&
 				if (Scale.X > SMALL_NUMBER && Scale.Y > SMALL_NUMBER)
 				{
 					const float LightmapTexelFactor = TexelFactor / FMath::Min(Scale.X, Scale.Y);
-					new (OutStreamingTextures) FStreamingTexturePrimitiveInfo(Lightmap->GetTexture(LightmapIndex), Bounds, LightmapTexelFactor);
-					new (OutStreamingTextures) FStreamingTexturePrimitiveInfo(Lightmap->GetAOMaterialMaskTexture(), Bounds, LightmapTexelFactor);
-					new (OutStreamingTextures) FStreamingTexturePrimitiveInfo(Lightmap->GetSkyOcclusionTexture(), Bounds, LightmapTexelFactor);
+					new (OutStreamingRenderAssets) FStreamingRenderAssetPrimitiveInfo(Lightmap->GetTexture(LightmapIndex), Bounds, LightmapTexelFactor);
+					new (OutStreamingRenderAssets) FStreamingRenderAssetPrimitiveInfo(Lightmap->GetAOMaterialMaskTexture(), Bounds, LightmapTexelFactor);
+					new (OutStreamingRenderAssets) FStreamingRenderAssetPrimitiveInfo(Lightmap->GetSkyOcclusionTexture(), Bounds, LightmapTexelFactor);
 				}
 			}
 
@@ -4208,7 +4208,7 @@ void ULandscapeComponent::GetStreamingTextureInfo(FStreamingTextureLevelContext&
 				if (Scale.X > SMALL_NUMBER && Scale.Y > SMALL_NUMBER)
 				{
 					const float ShadowmapTexelFactor = TexelFactor / FMath::Min(Scale.X, Scale.Y);
-					new (OutStreamingTextures) FStreamingTexturePrimitiveInfo(Shadowmap->GetTexture(), Bounds, ShadowmapTexelFactor);
+					new (OutStreamingRenderAssets) FStreamingRenderAssetPrimitiveInfo(Shadowmap->GetTexture(), Bounds, ShadowmapTexelFactor);
 				}
 			}
 		}
@@ -4217,39 +4217,39 @@ void ULandscapeComponent::GetStreamingTextureInfo(FStreamingTextureLevelContext&
 	// Weightmap
 	for (int32 TextureIndex = 0; TextureIndex < WeightmapTextures.Num(); TextureIndex++)
 	{
-		FStreamingTexturePrimitiveInfo& StreamingWeightmap = *new(OutStreamingTextures)FStreamingTexturePrimitiveInfo;
+		FStreamingRenderAssetPrimitiveInfo& StreamingWeightmap = *new(OutStreamingRenderAssets)FStreamingRenderAssetPrimitiveInfo;
 		StreamingWeightmap.Bounds = BoundingSphere;
 		StreamingWeightmap.TexelFactor = TexelFactor;
-		StreamingWeightmap.Texture = WeightmapTextures[TextureIndex];
+		StreamingWeightmap.RenderAsset = WeightmapTextures[TextureIndex];
 	}
 
 	// Heightmap
 	if (HeightmapTexture)
 	{
-		FStreamingTexturePrimitiveInfo& StreamingHeightmap = *new(OutStreamingTextures)FStreamingTexturePrimitiveInfo;
+		FStreamingRenderAssetPrimitiveInfo& StreamingHeightmap = *new(OutStreamingRenderAssets)FStreamingRenderAssetPrimitiveInfo;
 		StreamingHeightmap.Bounds = BoundingSphere;
 
 		float HeightmapTexelFactor = TexelFactor * (static_cast<float>(HeightmapTexture->GetSizeY()) / (ComponentSizeQuads + 1));
 		StreamingHeightmap.TexelFactor = ForcedLOD >= 0 ? -(1 << (13 - ForcedLOD)) : HeightmapTexelFactor; // Minus Value indicate forced resolution (Mip 13 for 8k texture)
-		StreamingHeightmap.Texture = HeightmapTexture;
+		StreamingHeightmap.RenderAsset = HeightmapTexture;
 	}
 
 	// XYOffset
 	if (XYOffsetmapTexture)
 	{
-		FStreamingTexturePrimitiveInfo& StreamingXYOffset = *new(OutStreamingTextures)FStreamingTexturePrimitiveInfo;
+		FStreamingRenderAssetPrimitiveInfo& StreamingXYOffset = *new(OutStreamingRenderAssets)FStreamingRenderAssetPrimitiveInfo;
 		StreamingXYOffset.Bounds = BoundingSphere;
 		StreamingXYOffset.TexelFactor = TexelFactor;
-		StreamingXYOffset.Texture = XYOffsetmapTexture;
+		StreamingXYOffset.RenderAsset = XYOffsetmapTexture;
 	}
 
 #if WITH_EDITOR
 	if (GIsEditor && EditToolRenderData.DataTexture)
 	{
-		FStreamingTexturePrimitiveInfo& StreamingDatamap = *new(OutStreamingTextures)FStreamingTexturePrimitiveInfo;
+		FStreamingRenderAssetPrimitiveInfo& StreamingDatamap = *new(OutStreamingRenderAssets)FStreamingRenderAssetPrimitiveInfo;
 		StreamingDatamap.Bounds = BoundingSphere;
 		StreamingDatamap.TexelFactor = TexelFactor;
-		StreamingDatamap.Texture = EditToolRenderData.DataTexture;
+		StreamingDatamap.RenderAsset = EditToolRenderData.DataTexture;
 	}
 #endif
 }
diff --git a/Engine/Source/Runtime/OpenGLDrv/Private/OpenGLUAV.cpp b/Engine/Source/Runtime/OpenGLDrv/Private/OpenGLUAV.cpp
--- a/Engine/Source/Runtime/OpenGLDrv/Private/OpenGLUAV.cpp
+++ b/Engine/Source/Runtime/OpenGLDrv/Private/OpenGLUAV.cpp
@@ -28,10 +28,13 @@ FShaderResourceViewRHIRef FOpenGLDynamicRHI::RHICreateShaderResourceView(FVertex
 
 			const FOpenGLTextureFormat& GLFormat = GOpenGLTextureFormats[Format];
 			FOpenGL::GenTextures(1, &TextureID);
-
+			
 			// Use a texture stage that's not likely to be used for draws, to avoid waiting
-			OGLRHI->CachedSetupTextureStage(OGLRHI->GetContextStateForCurrentContext(), FOpenGL::GetMaxCombinedTextureImageUnits() - 1, GL_TEXTURE_BUFFER, TextureID, -1, 1);
-			FOpenGL::TexBuffer(GL_TEXTURE_BUFFER, GLFormat.InternalFormat[0], VertexBuffer->Resource);
+			if (VertexBuffer)
+			{
+				OGLRHI->CachedSetupTextureStage(OGLRHI->GetContextStateForCurrentContext(), FOpenGL::GetMaxCombinedTextureImageUnits() - 1, GL_TEXTURE_BUFFER, TextureID, -1, 1);
+				FOpenGL::TexBuffer(GL_TEXTURE_BUFFER, GLFormat.InternalFormat[0], VertexBuffer->Resource);
+			}
 		}
 
 		// No need to restore texture stage; leave it like this,
@@ -44,6 +47,36 @@ FShaderResourceViewRHIRef FOpenGLDynamicRHI::RHICreateShaderResourceView(FVertex
 	return Result;
 }
 
+void FOpenGLDynamicRHI::RHIUpdateShaderResourceView(FShaderResourceViewRHIParamRef SRV, FVertexBufferRHIParamRef VertexBuffer, uint32 Stride, uint8 Format)
+{
+	if (!FOpenGL::SupportsResourceView())
+	{
+		return;
+	}
+	
+	FOpenGLShaderResourceView* SRVGL = FOpenGLDynamicRHI::ResourceCast(SRV);
+	FOpenGLVertexBuffer* VBGL = FOpenGLDynamicRHI::ResourceCast(VertexBuffer);
+	const FOpenGLTextureFormat& GLFormat = GOpenGLTextureFormats[Format];
+	
+	check(SRVGL);
+	GLuint TextureID = SRVGL->Resource;
+	CachedSetupTextureStage(GetContextStateForCurrentContext(), FOpenGL::GetMaxCombinedTextureImageUnits() - 1, GL_TEXTURE_BUFFER, TextureID, -1, 1);
+	
+	if (!VBGL)
+	{
+		FOpenGL::TexBuffer(GL_TEXTURE_BUFFER, GLFormat.InternalFormat[0], 0);
+		SRVGL->VertexBuffer = nullptr;
+		SRVGL->ModificationVersion = 0;
+	}
+	else
+	{
+		check(SRVGL->Format == Format && SRVGL->Target == GL_TEXTURE_BUFFER);
+		FOpenGL::TexBuffer(GL_TEXTURE_BUFFER, GLFormat.InternalFormat[0], VBGL->Resource);
+		SRVGL->VertexBuffer = VertexBuffer;
+		SRVGL->ModificationVersion = VBGL->ModificationCount;
+	}
+}
+
 FShaderResourceViewRHIRef FOpenGLDynamicRHI::RHICreateShaderResourceView(FIndexBufferRHIParamRef BufferRHI)
 {
 	return new FOpenGLShaderResourceViewProxy([=](FShaderResourceViewRHIParamRef OwnerRHI)
diff --git a/Engine/Source/Runtime/OpenGLDrv/Public/OpenGLDrv.h b/Engine/Source/Runtime/OpenGLDrv/Public/OpenGLDrv.h
--- a/Engine/Source/Runtime/OpenGLDrv/Public/OpenGLDrv.h
+++ b/Engine/Source/Runtime/OpenGLDrv/Public/OpenGLDrv.h
@@ -380,6 +380,7 @@ public:
 	virtual FShaderResourceViewRHIRef RHICreateShaderResourceView(FStructuredBufferRHIParamRef StructuredBuffer) final override;
 	virtual FShaderResourceViewRHIRef RHICreateShaderResourceView(FVertexBufferRHIParamRef VertexBuffer, uint32 Stride, uint8 Format) final override;
 	virtual FShaderResourceViewRHIRef RHICreateShaderResourceView(FIndexBufferRHIParamRef Buffer) final override;
+	virtual void RHIUpdateShaderResourceView(FShaderResourceViewRHIParamRef SRV, FVertexBufferRHIParamRef VertexBuffer, uint32 Stride, uint8 Format) final override;
 	virtual uint64 RHICalcTexture2DPlatformSize(uint32 SizeX, uint32 SizeY, uint8 Format, uint32 NumMips, uint32 NumSamples, uint32 Flags, uint32& OutAlign) final override;
 	virtual uint64 RHICalcTexture3DPlatformSize(uint32 SizeX, uint32 SizeY, uint32 SizeZ, uint8 Format, uint32 NumMips, uint32 Flags, uint32& OutAlign) final override;
 	virtual uint64 RHICalcTextureCubePlatformSize(uint32 Size, uint8 Format, uint32 NumMips, uint32 Flags, uint32& OutAlign) final override;
diff --git a/Engine/Source/Runtime/OpenGLDrv/Public/OpenGLResources.h b/Engine/Source/Runtime/OpenGLDrv/Public/OpenGLResources.h
--- a/Engine/Source/Runtime/OpenGLDrv/Public/OpenGLResources.h
+++ b/Engine/Source/Runtime/OpenGLDrv/Public/OpenGLResources.h
@@ -1880,9 +1880,11 @@ public:
 	,	OpenGLRHI(InOpenGLRHI)
 	,	OwnsResource(true)
 	{
-		check(IsValidRef(VertexBuffer));
-		FOpenGLVertexBuffer* VB = (FOpenGLVertexBuffer*)VertexBuffer.GetReference();
-		ModificationVersion = VB->ModificationCount;
+		if (VertexBuffer)
+		{
+			FOpenGLVertexBuffer* VB = (FOpenGLVertexBuffer*)VertexBuffer.GetReference();
+			ModificationVersion = VB->ModificationCount;
+		}
 	}
 
 	FOpenGLShaderResourceView( FOpenGLDynamicRHI* InOpenGLRHI, GLuint InResource, GLenum InTarget, GLuint Mip, bool InOwnsResource )
diff --git a/Engine/Source/Runtime/RHI/Private/RHICommandList.cpp b/Engine/Source/Runtime/RHI/Private/RHICommandList.cpp
--- a/Engine/Source/Runtime/RHI/Private/RHICommandList.cpp
+++ b/Engine/Source/Runtime/RHI/Private/RHICommandList.cpp
@@ -2497,6 +2497,28 @@ void FRHICommandListImmediate::UpdateTextureReference(FTextureReferenceRHIParamR
 	}
 }
 
+void FRHICommandListImmediate::UpdateRHIResources(FRHIResourceUpdateInfo* UpdateInfos, int32 Num)
+{
+	if (this->Bypass())
+	{
+		FRHICommandUpdateRHIResources Cmd(UpdateInfos, Num);
+		Cmd.Execute(*this);
+	}
+	else
+	{
+		const SIZE_T NumBytes = sizeof(FRHIResourceUpdateInfo) * Num;
+		FRHIResourceUpdateInfo* LocalUpdateInfos = reinterpret_cast<FRHIResourceUpdateInfo*>(this->Alloc(NumBytes, alignof(FRHIResourceUpdateInfo)));
+		FMemory::Memcpy(LocalUpdateInfos, UpdateInfos, NumBytes);
+		new (AllocCommand<FRHICommandUpdateRHIResources>()) FRHICommandUpdateRHIResources(LocalUpdateInfos, Num);
+		RHIThreadFence(true);
+		if (GetUsedMemory() > 256 * 1024)
+		{
+			// we could be loading a level or something, lets get this stuff going
+			ImmediateFlush(EImmediateFlushType::DispatchToRHIThread);
+		}
+	}
+}
+
 void FDynamicRHI::RHICopySubTextureRegion_RenderThread(class FRHICommandListImmediate& RHICmdList, FTexture2DRHIParamRef SourceTexture, FTexture2DRHIParamRef DestinationTexture, FBox2D SourceBox, FBox2D DestinationBox)
 {
 	CSV_SCOPED_TIMING_STAT(RHITStalls, RHICopySubTextureRegion_RenderThread);
diff --git a/Engine/Source/Runtime/RHI/Public/DynamicRHI.h b/Engine/Source/Runtime/RHI/Public/DynamicRHI.h
--- a/Engine/Source/Runtime/RHI/Public/DynamicRHI.h
+++ b/Engine/Source/Runtime/RHI/Public/DynamicRHI.h
@@ -424,6 +424,10 @@ public:
 	// FlushType: Wait RHI Thread
 	virtual FShaderResourceViewRHIRef RHICreateShaderResourceView(FIndexBufferRHIParamRef Buffer) = 0;
 
+	// Must be called on RHI thread timeline
+	// Make sure to call RHIThreadFence(true) afterwards so that parallel translation doesn't refer old resources
+	virtual void RHIUpdateShaderResourceView(FShaderResourceViewRHIParamRef SRV, FVertexBufferRHIParamRef VertexBuffer, uint32 Stride, uint8 Format) /*= 0;*/ {}
+
 	/**
 	* Computes the total size of a 2D texture with the specified parameters.
 	*
diff --git a/Engine/Source/Runtime/RHI/Public/RHICommandList.h b/Engine/Source/Runtime/RHI/Public/RHICommandList.h
--- a/Engine/Source/Runtime/RHI/Public/RHICommandList.h
+++ b/Engine/Source/Runtime/RHI/Public/RHICommandList.h
@@ -1795,6 +1795,41 @@ struct FRHICommandUpdateTextureReference final : public FRHICommand<FRHICommandU
 	RHI_API void Execute(FRHICommandListBase& CmdList);
 };
 
+struct FRHIShaderResourceViewUpdateInfo_VB
+{
+	FShaderResourceViewRHIParamRef SRV;
+	FVertexBufferRHIParamRef VertexBuffer;
+	uint32 Stride;
+	uint8 Format;
+};
+
+struct FRHIResourceUpdateInfo
+{
+	enum EUpdateType
+	{
+		UT_VertexBufferSRV,
+		UT_Num
+	};
+
+	EUpdateType Type;
+	union
+	{
+		FRHIShaderResourceViewUpdateInfo_VB VertexBufferSRV;
+	};
+};
+
+struct FRHICommandUpdateRHIResources final : public FRHICommand<FRHICommandUpdateRHIResources>
+{
+	FRHIResourceUpdateInfo* UpdateInfos;
+	int32 Num;
+
+	FORCEINLINE_DEBUGGABLE FRHICommandUpdateRHIResources(FRHIResourceUpdateInfo* InUpdateInfos, int32 InNum)
+		: UpdateInfos(InUpdateInfos)
+		, Num(InNum)
+	{}
+	RHI_API void Execute(FRHICommandListBase& CmdList);
+};
+
 #if RHI_RAYTRACING
 struct FRHICommandBuildAccelerationStructure final : public FRHICommand<FRHICommandBuildAccelerationStructure>
 {
@@ -4246,6 +4281,7 @@ public:
 	}
 	void UpdateTextureReference(FTextureReferenceRHIParamRef TextureRef, FTextureRHIParamRef NewTexture);
 
+	void UpdateRHIResources(FRHIResourceUpdateInfo* UpdateInfos, int32 Num);
 };
 
  struct FScopedGPUMask
@@ -4498,6 +4534,11 @@ FORCEINLINE FIndexBufferRHIRef RHICreateIndexBuffer(uint32 Stride, uint32 Size,
 	return FRHICommandListExecutor::GetImmediateCommandList().CreateIndexBuffer(Stride, Size, InUsage, CreateInfo);
 }
 
+FORCEINLINE FIndexBufferRHIRef RHIAsyncCreateIndexBuffer(uint32 Stride, uint32 Size, uint32 InUsage, FRHIResourceCreateInfo& CreateInfo)
+{
+	return GDynamicRHI->RHICreateIndexBuffer(Stride, Size, InUsage, CreateInfo);
+}
+
 FORCEINLINE void* RHILockIndexBuffer(FIndexBufferRHIParamRef IndexBuffer, uint32 Offset, uint32 Size, EResourceLockMode LockMode)
 {
 	return FRHICommandListExecutor::GetImmediateCommandList().LockIndexBuffer(IndexBuffer, Offset, Size, LockMode);
@@ -4518,6 +4559,11 @@ FORCEINLINE FVertexBufferRHIRef RHICreateVertexBuffer(uint32 Size, uint32 InUsag
 	return FRHICommandListExecutor::GetImmediateCommandList().CreateVertexBuffer(Size, InUsage, CreateInfo);
 }
 
+FORCEINLINE FVertexBufferRHIRef RHIAsyncCreateVertexBuffer(uint32 Size, uint32 InUsage, FRHIResourceCreateInfo& CreateInfo)
+{
+	return GDynamicRHI->RHICreateVertexBuffer(Size, InUsage, CreateInfo);
+}
+
 FORCEINLINE void* RHILockVertexBuffer(FVertexBufferRHIParamRef VertexBuffer, uint32 Offset, uint32 SizeRHI, EResourceLockMode LockMode)
 {
 	return FRHICommandListExecutor::GetImmediateCommandList().LockVertexBuffer(VertexBuffer, Offset, SizeRHI, LockMode);
@@ -4578,6 +4624,11 @@ FORCEINLINE FShaderResourceViewRHIRef RHICreateShaderResourceView(FIndexBufferRH
 	return FRHICommandListExecutor::GetImmediateCommandList().CreateShaderResourceView(Buffer);
 }
 
+FORCEINLINE void RHIUpdateRHIResources(FRHIResourceUpdateInfo* UpdateInfos, int32 Num)
+{
+	return FRHICommandListExecutor::GetImmediateCommandList().UpdateRHIResources(UpdateInfos, Num);
+}
+
 FORCEINLINE FTextureReferenceRHIRef RHICreateTextureReference(FLastRenderTimeContainer* LastRenderTime)
 {
 	return FRHICommandListExecutor::GetImmediateCommandList().CreateTextureReference(LastRenderTime);
@@ -4813,6 +4864,43 @@ FORCEINLINE void RHIUnlockStagingBuffer(FStagingBufferRHIParamRef StagingBuffer)
 	 FRHICommandListExecutor::GetImmediateCommandList().UnlockStagingBuffer(StagingBuffer);
 }
 
+template <int32 MaxNumUpdates>
+struct TRHIResourceUpdateBatcher
+{
+	FRHIResourceUpdateInfo UpdateInfos[MaxNumUpdates];
+	int32 NumBatched;
+
+	TRHIResourceUpdateBatcher()
+		: NumBatched(0)
+	{}
+
+	~TRHIResourceUpdateBatcher()
+	{
+		Flush();
+	}
+
+	void Flush()
+	{
+		if (NumBatched > 0)
+		{
+			RHIUpdateRHIResources(UpdateInfos, NumBatched);
+			NumBatched = 0;
+		}
+	}
+
+	void QueueUpdateRequest(FShaderResourceViewRHIParamRef SRV, FVertexBufferRHIParamRef VertexBuffer, uint32 Stride, uint8 Format)
+	{
+		check(NumBatched <= MaxNumUpdates);
+		if (NumBatched == MaxNumUpdates)
+		{
+			Flush();
+		}
+		const int32 Idx = NumBatched++;
+		UpdateInfos[Idx].Type = FRHIResourceUpdateInfo::UT_VertexBufferSRV;
+		UpdateInfos[Idx].VertexBufferSRV = { SRV, VertexBuffer, Stride, Format };
+	}
+};
+
 #undef RHICOMMAND_CALLSTACK
 
 #include "RHICommandList.inl"
diff --git a/Engine/Source/Runtime/RHI/Public/RHICommandListCommandExecutes.inl b/Engine/Source/Runtime/RHI/Public/RHICommandListCommandExecutes.inl
--- a/Engine/Source/Runtime/RHI/Public/RHICommandListCommandExecutes.inl
+++ b/Engine/Source/Runtime/RHI/Public/RHICommandListCommandExecutes.inl
@@ -610,6 +610,28 @@ void FRHICommandUpdateTextureReference::Execute(FRHICommandListBase& CmdList)
 	INTERNAL_DECORATOR(RHIUpdateTextureReference)(TextureRef, NewTexture);
 }
 
+void FRHICommandUpdateRHIResources::Execute(FRHICommandListBase& CmdList)
+{
+	RHISTAT(UpdateRHIResources);
+	for (int32 Idx = 0; Idx < Num; ++Idx)
+	{
+		FRHIResourceUpdateInfo& Info = UpdateInfos[Idx];
+		switch (Info.Type)
+		{
+		case FRHIResourceUpdateInfo::UT_VertexBufferSRV:
+			GDynamicRHI->RHIUpdateShaderResourceView(
+				Info.VertexBufferSRV.SRV,
+				Info.VertexBufferSRV.VertexBuffer,
+				Info.VertexBufferSRV.Stride,
+				Info.VertexBufferSRV.Format);
+			break;
+		default:
+			// Unrecognized type, do nothing
+			break;
+		}
+	}
+}
+
 void FRHICommandBeginScene::Execute(FRHICommandListBase& CmdList)
 {
 	RHISTAT(BeginScene);
diff --git a/Engine/Source/Runtime/RHI/Public/RHIDefinitions.h b/Engine/Source/Runtime/RHI/Public/RHIDefinitions.h
--- a/Engine/Source/Runtime/RHI/Public/RHIDefinitions.h
+++ b/Engine/Source/Runtime/RHI/Public/RHIDefinitions.h
@@ -98,6 +98,9 @@ enum ERenderQueryType
 /** Maximum number of miplevels in a texture. */
 enum { MAX_TEXTURE_MIP_COUNT = 14 };
 
+/** Maximum number of static/skeletal mesh LODs */
+enum { MAX_MESH_LOD_COUNT = 8 };
+
 /** Maximum number of immutable samplers in a PSO. */
 enum
 {
diff --git a/Engine/Source/Runtime/Renderer/Private/SceneVisibility.cpp b/Engine/Source/Runtime/Renderer/Private/SceneVisibility.cpp
--- a/Engine/Source/Runtime/Renderer/Private/SceneVisibility.cpp
+++ b/Engine/Source/Runtime/Renderer/Private/SceneVisibility.cpp
@@ -2025,16 +2025,19 @@ struct FRelevancePacket
 			const FPrimitiveViewRelevance& ViewRelevance = View.PrimitiveViewRelevanceMap[PrimitiveIndex];
 			const bool bIsPrimitiveDistanceCullFading = View.PrimitiveFadeUniformBufferMap[PrimitiveIndex];
 
+			const int8 CurFirstLODIdx = PrimitiveSceneInfo->Proxy->GetCurrentFirstLODIdx_RenderThread();
+			check(CurFirstLODIdx >= 0);
 			float MeshScreenSizeSquared = 0;
 			FLODMask LODToRender;
 
 			if (PrimitiveSceneInfo->bIsUsingCustomLODRules)
 			{
 				LODToRender = PrimitiveSceneInfo->Proxy->GetCustomLOD(View, View.LODDistanceFactor, ViewData.ForcedLODLevel, MeshScreenSizeSquared);
+				LODToRender.ClampToFirstLOD(CurFirstLODIdx);
 			}
 			else
 			{
-				LODToRender = ComputeLODForMeshes(PrimitiveSceneInfo->StaticMeshRelevances, View, Bounds.BoxSphereBounds.Origin, Bounds.BoxSphereBounds.SphereRadius, ViewData.ForcedLODLevel, MeshScreenSizeSquared, ViewData.LODScale);
+				LODToRender = ComputeLODForMeshes(PrimitiveSceneInfo->StaticMeshRelevances, View, Bounds.BoxSphereBounds.Origin, Bounds.BoxSphereBounds.SphereRadius, ViewData.ForcedLODLevel, MeshScreenSizeSquared, CurFirstLODIdx, ViewData.LODScale);
 			}
 
 			PrimitivesLODMask.AddPrim(FRelevancePacket::FPrimitiveLODMask(PrimitiveIndex, LODToRender));
diff --git a/Engine/Source/Runtime/Renderer/Private/ShadowSetup.cpp b/Engine/Source/Runtime/Renderer/Private/ShadowSetup.cpp
--- a/Engine/Source/Runtime/Renderer/Private/ShadowSetup.cpp
+++ b/Engine/Source/Runtime/Renderer/Private/ShadowSetup.cpp
@@ -857,15 +857,17 @@ bool FProjectedShadowInfo::ShouldDrawStaticMeshes(FViewInfo& InCurrentView, bool
 		{
 			FLODMask ViewLODToRender;
 			float MeshScreenSizeSquared = 0;
+			const int8 CurFirstLODIdx = InPrimitiveSceneInfo->Proxy->GetCurrentFirstLODIdx_RenderThread();
 
 			if (InPrimitiveSceneInfo->bIsUsingCustomLODRules)
 			{
 				ViewLODToRender = InPrimitiveSceneInfo->Proxy->GetCustomLOD(InCurrentView, InCurrentView.LODDistanceFactor, ForcedLOD, MeshScreenSizeSquared);
+				ViewLODToRender.ClampToFirstLOD(CurFirstLODIdx);
 			}
 			else
 			{
 				const FBoxSphereBounds& Bounds = InPrimitiveSceneInfo->Proxy->GetBounds();
-				ViewLODToRender = ComputeLODForMeshes(InPrimitiveSceneInfo->StaticMeshRelevances, InCurrentView, Bounds.Origin, Bounds.SphereRadius, ForcedLOD, MeshScreenSizeSquared, InCurrentView.LODDistanceFactor);
+				ViewLODToRender = ComputeLODForMeshes(InPrimitiveSceneInfo->StaticMeshRelevances, InCurrentView, Bounds.Origin, Bounds.SphereRadius, ForcedLOD, MeshScreenSizeSquared, CurFirstLODIdx, InCurrentView.LODDistanceFactor);
 			}	
 
 			InCurrentView.PrimitivesLODMask[PrimitiveId] = ViewLODToRender;
diff --git a/Engine/Source/Runtime/VulkanRHI/Private/VulkanUAV.cpp b/Engine/Source/Runtime/VulkanRHI/Private/VulkanUAV.cpp
--- a/Engine/Source/Runtime/VulkanRHI/Private/VulkanUAV.cpp
+++ b/Engine/Source/Runtime/VulkanRHI/Private/VulkanUAV.cpp
@@ -23,12 +23,41 @@ FVulkanShaderResourceView::FVulkanShaderResourceView(FVulkanDevice* Device, FRHI
 
 FVulkanShaderResourceView::~FVulkanShaderResourceView()
 {
-	TextureView.Destroy(*Device);
-	BufferViews.Empty(0);
+	Clear();
+}
+
+void FVulkanShaderResourceView::Clear()
+{
+	SourceRHIBuffer = nullptr;
+	SourceBuffer = nullptr;
+	BufferViews.Empty();
+	SourceStructuredBuffer = nullptr;
+	if (Device)
+	{
+		TextureView.Destroy(*Device);
+	}
 	SourceTexture = nullptr;
 	Device = nullptr;
 }
 
+void FVulkanShaderResourceView::Rename(FRHIResource* InRHIBuffer, FVulkanResourceMultiBuffer* InSourceBuffer, uint32 InSize, EPixelFormat InFormat)
+{
+	check(Device);
+	BufferViewFormat = InFormat;
+	SourceTexture = nullptr;
+	TextureView.Destroy(*Device);
+	SourceStructuredBuffer = nullptr;
+	MipLevel = 0;
+	NumMips = -1;
+	BufferViews.Reset();
+	BufferViews.AddZeroed(InSourceBuffer->IsVolatile() ? 1 : InSourceBuffer->GetNumBuffers());
+	BufferIndex = 0;
+	Size = InSize;
+	SourceBuffer = InSourceBuffer;
+	SourceRHIBuffer = InRHIBuffer;
+	VolatileLockCounter = MAX_uint32;
+}
+
 void FVulkanShaderResourceView::UpdateView()
 {
 #if VULKAN_ENABLE_AGGRESSIVE_STATS
@@ -238,9 +267,12 @@ FShaderResourceViewRHIRef FVulkanDynamicRHI::RHICreateShaderResourceView(FStruct
 
 FShaderResourceViewRHIRef FVulkanDynamicRHI::RHICreateShaderResourceView(FVertexBufferRHIParamRef VertexBufferRHI, uint32 Stride, uint8 Format)
 {	
+	if (!VertexBufferRHI)
+	{
+		return new FVulkanShaderResourceView(Device, nullptr, nullptr, 0, (EPixelFormat)Format);
+	}
 	FVulkanVertexBuffer* VertexBuffer = ResourceCast(VertexBufferRHI);
-	FVulkanShaderResourceView* SRV = new FVulkanShaderResourceView(Device, VertexBufferRHI, VertexBuffer, VertexBuffer->GetSize(), (EPixelFormat)Format);
-	return SRV;
+	return new FVulkanShaderResourceView(Device, VertexBufferRHI, VertexBuffer, VertexBuffer->GetSize(), (EPixelFormat)Format);
 }
 
 FShaderResourceViewRHIRef FVulkanDynamicRHI::RHICreateShaderResourceView(FTexture2DRHIParamRef Texture2DRHI, uint8 MipLevel)
@@ -282,6 +314,21 @@ FShaderResourceViewRHIRef FVulkanDynamicRHI::RHICreateShaderResourceView(FIndexB
 	return SRV;
 }
 
+void FVulkanDynamicRHI::RHIUpdateShaderResourceView(FShaderResourceViewRHIParamRef SRV, FVertexBufferRHIParamRef VertexBuffer, uint32 Stride, uint8 Format)
+{
+	FVulkanShaderResourceView* SRVVk = ResourceCast(SRV);
+	check(SRVVk && SRVVk->GetParent() == Device);
+	if (!VertexBuffer)
+	{
+		SRVVk->Clear();
+	}
+	else if (SRVVk->SourceRHIBuffer.GetReference() != VertexBuffer)
+	{
+		FVulkanVertexBuffer* VertexBufferVk = ResourceCast(VertexBuffer);
+		SRVVk->Rename(VertexBuffer, VertexBufferVk, VertexBufferVk->GetSize(), (EPixelFormat)Format);
+	}
+}
+
 void FVulkanCommandListContext::RHIClearTinyUAV(FUnorderedAccessViewRHIParamRef UnorderedAccessViewRHI, const uint32* Values)
 {
 	FVulkanUnorderedAccessView* UnorderedAccessView = ResourceCast(UnorderedAccessViewRHI);
diff --git a/Engine/Source/Runtime/VulkanRHI/Public/VulkanDynamicRHI.h b/Engine/Source/Runtime/VulkanRHI/Public/VulkanDynamicRHI.h
--- a/Engine/Source/Runtime/VulkanRHI/Public/VulkanDynamicRHI.h
+++ b/Engine/Source/Runtime/VulkanRHI/Public/VulkanDynamicRHI.h
@@ -76,6 +76,7 @@ public:
 	virtual FShaderResourceViewRHIRef RHICreateShaderResourceView(FStructuredBufferRHIParamRef StructuredBuffer) final override;
 	virtual FShaderResourceViewRHIRef RHICreateShaderResourceView(FVertexBufferRHIParamRef VertexBuffer, uint32 Stride, uint8 Format) final override;
 	virtual FShaderResourceViewRHIRef RHICreateShaderResourceView(FIndexBufferRHIParamRef Buffer) final override;
+	virtual void RHIUpdateShaderResourceView(FShaderResourceViewRHIParamRef SRV, FVertexBufferRHIParamRef VertexBuffer, uint32 Stride, uint8 Format) final override;
 	virtual uint64 RHICalcTexture2DPlatformSize(uint32 SizeX, uint32 SizeY, uint8 Format, uint32 NumMips, uint32 NumSamples, uint32 Flags, uint32& OutAlign) final override;
 	virtual uint64 RHICalcTexture3DPlatformSize(uint32 SizeX, uint32 SizeY, uint32 SizeZ, uint8 Format, uint32 NumMips, uint32 Flags, uint32& OutAlign) final override;
 	virtual uint64 RHICalcTextureCubePlatformSize(uint32 Size, uint8 Format, uint32 NumMips, uint32 Flags, uint32& OutAlign) final override;
diff --git a/Engine/Source/Runtime/VulkanRHI/Public/VulkanResources.h b/Engine/Source/Runtime/VulkanRHI/Public/VulkanResources.h
--- a/Engine/Source/Runtime/VulkanRHI/Public/VulkanResources.h
+++ b/Engine/Source/Runtime/VulkanRHI/Public/VulkanResources.h
@@ -1363,6 +1363,9 @@ public:
 	{
 	}
 
+	void Clear();
+
+	void Rename(FRHIResource* InRHIBuffer, FVulkanResourceMultiBuffer* InSourceBuffer, uint32 InSize, EPixelFormat InFormat);
 
 	void UpdateView();
 
diff --git a/Engine/Source/Runtime/Windows/D3D11RHI/Private/D3D11RHIPrivate.h b/Engine/Source/Runtime/Windows/D3D11RHI/Private/D3D11RHIPrivate.h
--- a/Engine/Source/Runtime/Windows/D3D11RHI/Private/D3D11RHIPrivate.h
+++ b/Engine/Source/Runtime/Windows/D3D11RHI/Private/D3D11RHIPrivate.h
@@ -419,6 +419,7 @@ public:
 	virtual FShaderResourceViewRHIRef RHICreateShaderResourceView(FStructuredBufferRHIParamRef StructuredBuffer) final override;
 	virtual FShaderResourceViewRHIRef RHICreateShaderResourceView(FVertexBufferRHIParamRef VertexBuffer, uint32 Stride, uint8 Format) final override;
 	virtual FShaderResourceViewRHIRef RHICreateShaderResourceView(FIndexBufferRHIParamRef Buffer) final override;
+	virtual void RHIUpdateShaderResourceView(FShaderResourceViewRHIParamRef SRV, FVertexBufferRHIParamRef VertexBuffer, uint32 Stride, uint8 Format) final override;
 	virtual uint64 RHICalcTexture2DPlatformSize(uint32 SizeX, uint32 SizeY, uint8 Format, uint32 NumMips, uint32 NumSamples, uint32 Flags, uint32& OutAlign) final override;
 	virtual uint64 RHICalcTexture3DPlatformSize(uint32 SizeX, uint32 SizeY, uint32 SizeZ, uint8 Format, uint32 NumMips, uint32 Flags, uint32& OutAlign) final override;
 	virtual uint64 RHICalcTextureCubePlatformSize(uint32 Size, uint8 Format, uint32 NumMips, uint32 Flags, uint32& OutAlign) final override;
diff --git a/Engine/Source/Runtime/Windows/D3D11RHI/Private/D3D11UAV.cpp b/Engine/Source/Runtime/Windows/D3D11RHI/Private/D3D11UAV.cpp
--- a/Engine/Source/Runtime/Windows/D3D11RHI/Private/D3D11UAV.cpp
+++ b/Engine/Source/Runtime/Windows/D3D11RHI/Private/D3D11UAV.cpp
@@ -188,14 +188,10 @@ FShaderResourceViewRHIRef FD3D11DynamicRHI::RHICreateShaderResourceView(FStructu
 	return new FD3D11ShaderResourceView(ShaderResourceView,StructuredBuffer);
 }
 
-FShaderResourceViewRHIRef FD3D11DynamicRHI::RHICreateShaderResourceView(FVertexBufferRHIParamRef VertexBufferRHI, uint32 Stride, uint8 Format)
+static void CreateD3D11ShaderResourceViewOnVertexBuffer(ID3D11Device* Direct3DDevice, ID3D11Buffer* Buffer, uint32 Stride, uint8 Format, ID3D11ShaderResourceView** OutSRV)
 {
-	FD3D11VertexBuffer* VertexBuffer = ResourceCast(VertexBufferRHI);
-	check(VertexBuffer);
-	check(VertexBuffer->Resource);
-
 	D3D11_BUFFER_DESC BufferDesc;
-	VertexBuffer->Resource->GetDesc(&BufferDesc);
+	Buffer->GetDesc(&BufferDesc);
 
 	// Create a Shader Resource View
 	D3D11_SHADER_RESOURCE_VIEW_DESC SRVDesc;
@@ -203,28 +199,61 @@ FShaderResourceViewRHIRef FD3D11DynamicRHI::RHICreateShaderResourceView(FVertexB
 	SRVDesc.ViewDimension = D3D11_SRV_DIMENSION_BUFFER;
 	SRVDesc.Buffer.FirstElement = 0;
 
-	SRVDesc.Format = FindShaderResourceDXGIFormat((DXGI_FORMAT)GPixelFormats[Format].PlatformFormat,false);
+	SRVDesc.Format = FindShaderResourceDXGIFormat((DXGI_FORMAT)GPixelFormats[Format].PlatformFormat, false);
 	SRVDesc.Buffer.NumElements = BufferDesc.ByteWidth / Stride;
-	TRefCountPtr<ID3D11ShaderResourceView> ShaderResourceView;
-	
-	HRESULT hr = Direct3DDevice->CreateShaderResourceView(VertexBuffer->Resource, &SRVDesc, (ID3D11ShaderResourceView**)ShaderResourceView.GetInitReference());
+
+	HRESULT hr = Direct3DDevice->CreateShaderResourceView(Buffer, &SRVDesc, OutSRV);
 	if (FAILED(hr))
 	{
 		if (hr == E_OUTOFMEMORY)
 		{
 			// There appears to be a driver bug that causes SRV creation to fail with an OOM error and then succeed on the next call.
-			hr = Direct3DDevice->CreateShaderResourceView(VertexBuffer->Resource, &SRVDesc, (ID3D11ShaderResourceView**)ShaderResourceView.GetInitReference());
+			hr = Direct3DDevice->CreateShaderResourceView(Buffer, &SRVDesc, OutSRV);
 		}
 		if (FAILED(hr))
 		{
-			UE_LOG(LogD3D11RHI,Error,TEXT("Failed to create shader resource view for vertex buffer: ByteWidth=%d NumElements=%d Format=%s"),BufferDesc.ByteWidth,BufferDesc.ByteWidth / Stride, GPixelFormats[Format].Name);
-			VerifyD3D11Result(hr,"Direct3DDevice->CreateShaderResourceView",__FILE__,__LINE__,Direct3DDevice);
+			UE_LOG(LogD3D11RHI, Error, TEXT("Failed to create shader resource view for vertex buffer: ByteWidth=%d NumElements=%d Format=%s"), BufferDesc.ByteWidth, BufferDesc.ByteWidth / Stride, GPixelFormats[Format].Name);
+			VerifyD3D11Result(hr, "Direct3DDevice->CreateShaderResourceView", __FILE__, __LINE__, Direct3DDevice);
 		}
 	}
+}
+
+FShaderResourceViewRHIRef FD3D11DynamicRHI::RHICreateShaderResourceView(FVertexBufferRHIParamRef VertexBufferRHI, uint32 Stride, uint8 Format)
+{
+	if (!VertexBufferRHI)
+	{
+		return new FD3D11ShaderResourceView(nullptr, nullptr);
+	}
+	FD3D11VertexBuffer* VertexBuffer = ResourceCast(VertexBufferRHI);
+	check(VertexBuffer);
+	check(VertexBuffer->Resource);
+
+	TRefCountPtr<ID3D11ShaderResourceView> ShaderResourceView;
+	CreateD3D11ShaderResourceViewOnVertexBuffer(Direct3DDevice, VertexBuffer->Resource, Stride, Format, ShaderResourceView.GetInitReference());
 
 	return new FD3D11ShaderResourceView(ShaderResourceView,VertexBuffer);
 }
 
+void FD3D11DynamicRHI::RHIUpdateShaderResourceView(FShaderResourceViewRHIParamRef SRV, FVertexBufferRHIParamRef VertexBufferRHI, uint32 Stride, uint8 Format)
+{
+	check(SRV);
+	FD3D11ShaderResourceView* SRVD3D11 = ResourceCast(SRV);
+	if (!VertexBufferRHI)
+	{
+		SRVD3D11->Rename(nullptr, nullptr);
+	}
+	else
+	{
+		FD3D11VertexBuffer* VertexBuffer = ResourceCast(VertexBufferRHI);
+		check(VertexBuffer->Resource);
+
+		TRefCountPtr<ID3D11ShaderResourceView> ShaderResourceView;
+		CreateD3D11ShaderResourceViewOnVertexBuffer(Direct3DDevice, VertexBuffer->Resource, Stride, Format, ShaderResourceView.GetInitReference());
+
+		SRVD3D11->Rename(ShaderResourceView, VertexBuffer);
+	}
+}
+
 FShaderResourceViewRHIRef FD3D11DynamicRHI::RHICreateShaderResourceView(FIndexBufferRHIParamRef BufferRHI)
 {
 	FD3D11IndexBuffer* Buffer = ResourceCast(BufferRHI);
diff --git a/Engine/Source/Runtime/Windows/D3D11RHI/Public/D3D11Resources.h b/Engine/Source/Runtime/Windows/D3D11RHI/Public/D3D11Resources.h
--- a/Engine/Source/Runtime/Windows/D3D11RHI/Public/D3D11Resources.h
+++ b/Engine/Source/Runtime/Windows/D3D11RHI/Public/D3D11Resources.h
@@ -816,6 +816,12 @@ public:
 	: View(InView)
 	, Resource(InResource)
 	{}
+
+	void Rename(ID3D11ShaderResourceView* InView, FD3D11BaseShaderResource* InResource)
+	{
+		View = InView;
+		Resource = InResource;
+	}
 };
 
 /** Unordered access view class. */
