Batched some CopyTexture & ReadSurface to reduce the amount of render command, doing so helped gain ~10ms in RegenerateLayersHeightmaps in map composed of 127x127 section size with 16x16 (256 components) with an overall resolution of 2033x2033 and each component being in their own heightmap.

diff --git a/Engine/Source/Runtime/Landscape/Classes/Landscape.h b/Engine/Source/Runtime/Landscape/Classes/Landscape.h
--- a/Engine/Source/Runtime/Landscape/Classes/Landscape.h
+++ b/Engine/Source/Runtime/Landscape/Classes/Landscape.h
@@ -171,6 +171,38 @@ struct FLandscapeLayer
 	TMap<ULandscapeLayerInfoObject*, bool> WeightmapLayerAllocationBlend; // True -> Substractive, False -> Additive
 };
 
+struct FLandscapeLayersCopyTextureParams
+{
+	FLandscapeLayersCopyTextureParams(const FString& InSourceResourceDebugName, FTextureResource* InSourceResource, const FString& InDestResourceDebugName, FTextureResource* InDestResource, FTextureResource* InDestCPUResource,
+		const FIntPoint& InInitialPositionOffset, int32 InSubSectionSizeQuad, int32 InNumSubSections, uint8 InSourceCurrentMip, uint8 InDestCurrentMip, uint32 InSourceArrayIndex, uint32 InDestArrayIndex)
+		: SourceResourceDebugName(InSourceResourceDebugName)
+		, SourceResource(InSourceResource)
+		, DestResourceDebugName(InDestResourceDebugName)
+		, DestResource(InDestResource)
+		, DestCPUResource(InDestCPUResource)
+		, InitialPositionOffset(InInitialPositionOffset)
+		, SubSectionSizeQuad(InSubSectionSizeQuad)
+		, NumSubSections(InNumSubSections)
+		, SourceMip(InSourceCurrentMip)
+		, DestMip(InDestCurrentMip)
+		, SourceArrayIndex(InSourceArrayIndex)
+		, DestArrayIndex(InDestArrayIndex)
+	{}
+
+	FString SourceResourceDebugName;
+	FTextureResource* SourceResource;
+	FString DestResourceDebugName;
+	FTextureResource* DestResource;
+	FTextureResource* DestCPUResource;
+	FIntPoint InitialPositionOffset;
+	int32 SubSectionSizeQuad;
+	int32 NumSubSections;
+	uint8 SourceMip;
+	uint8 DestMip;
+	uint32 SourceArrayIndex;
+	uint32 DestArrayIndex;
+};
+
 UCLASS(MinimalAPI, showcategories=(Display, Movement, Collision, Lighting, LOD, Input), hidecategories=(Mobility))
 class ALandscape : public ALandscapeProxy
 {
@@ -328,9 +360,16 @@ private:
 	void DrawWeightmapComponentToRenderTargetMips(const TArray<FVector2D>& InTexturePositionsToDraw, UTexture* InReadWeightmap, bool InClearRTWrite, struct FLandscapeLayersWeightmapShaderParameters& InShaderParams) const;
 
 	void CopyLayersTexture(UTexture* InSourceTexture, UTexture* InDestTexture, FTextureResource* InDestCPUResource = nullptr, const FIntPoint& InInitialPositionOffset = FIntPoint(0, 0), uint8 InSourceCurrentMip = 0, uint8 InDestCurrentMip = 0,
-							   uint32 InSourceArrayIndex = 0, uint32 InDestArrayIndex = 0) const;
+						   uint32 InSourceArrayIndex = 0, uint32 InDestArrayIndex = 0) const;
 	void CopyLayersTexture(const FString& InSourceDebugName, FTextureResource* InSourceResource, const FString& InDestDebugName, FTextureResource* InDestResource, FTextureResource* InDestCPUResource = nullptr, const FIntPoint& InInitialPositionOffset = FIntPoint(0, 0),
-							   uint8 InSourceCurrentMip = 0, uint8 InDestCurrentMip = 0, uint32 InSourceArrayIndex = 0, uint32 uInDestArrayIndex = 0) const;
+						   uint8 InSourceCurrentMip = 0, uint8 InDestCurrentMip = 0, uint32 InSourceArrayIndex = 0, uint32 InDestArrayIndex = 0) const;
+
+	void AddDeferredCopyLayersTexture(UTexture* InSourceTexture, UTexture* InDestTexture, FTextureResource* InDestCPUResource = nullptr, const FIntPoint& InInitialPositionOffset = FIntPoint(0, 0), uint8 InSourceCurrentMip = 0, uint8 InDestCurrentMip = 0,
+									  uint32 InSourceArrayIndex = 0, uint32 InDestArrayIndex = 0);
+	void AddDeferredCopyLayersTexture(const FString& InSourceDebugName, FTextureResource* InSourceResource, const FString& InDestDebugName, FTextureResource* InDestResource, FTextureResource* InDestCPUResource = nullptr, const FIntPoint& InInitialPositionOffset = FIntPoint(0, 0),
+									  uint8 InSourceCurrentMip = 0, uint8 InDestCurrentMip = 0, uint32 InSourceArrayIndex = 0, uint32 InDestArrayIndex = 0);
+
+	void CommitDeferredCopyLayersTexture();
 
 	void InitializeLayers();
 	void InitializeLandscapeLayersWeightmapUsage();
@@ -412,6 +451,8 @@ private:
 	
 	// Used in packing the material layer data contained into CombinedLayersWeightmapAllMaterialLayersResource to be set again for each component weightmap (size of the landscape)
 	class FLandscapeTexture2DResource* WeightmapScratchPackLayerTextureResource;
+
+	TArray<FLandscapeLayersCopyTextureParams> PendingCopyTextures;
 #endif
 
 protected:
diff --git a/Engine/Source/Runtime/Landscape/Private/Landscape.cpp b/Engine/Source/Runtime/Landscape/Private/Landscape.cpp
--- a/Engine/Source/Runtime/Landscape/Private/Landscape.cpp
+++ b/Engine/Source/Runtime/Landscape/Private/Landscape.cpp
@@ -84,6 +84,7 @@ DEFINE_STAT(STAT_LandscapeLayersRegenerateDrawCalls);
 
 DEFINE_STAT(STAT_LandscapeLayersRegenerateHeightmaps);
 DEFINE_STAT(STAT_LandscapeLayersResolveHeightmaps);
+DEFINE_STAT(STAT_LandscapeLayersResolveTexture);
 
 DEFINE_STAT(STAT_LandscapeLayersUpdateMaterialInstance);
 DEFINE_STAT(STAT_LandscapeLayersReallocateWeightmaps);
diff --git a/Engine/Source/Runtime/Landscape/Private/LandscapeEditLayers.cpp b/Engine/Source/Runtime/Landscape/Private/LandscapeEditLayers.cpp
--- a/Engine/Source/Runtime/Landscape/Private/LandscapeEditLayers.cpp
+++ b/Engine/Source/Runtime/Landscape/Private/LandscapeEditLayers.cpp
@@ -989,89 +989,66 @@ private:
 class FLandscapeLayersCopyTexture_RenderThread
 {
 public:
-	FLandscapeLayersCopyTexture_RenderThread(const FString& InSourceResourceDebugName, FTextureResource* InSourceResource, const FString& InDestResourceDebugName, FTextureResource* InDestResource, FTextureResource* InDestCPUResource,
-											 const FIntPoint& InInitialPositionOffset, int32 InSubSectionSizeQuad, int32 InNumSubSections, uint8 InSourceCurrentMip, uint8 InDestCurrentMip, uint32 InSourceArrayIndex, uint32 InDestArrayIndex)
-		: SourceResource(InSourceResource)
-		, DestResource(InDestResource)
-		, DestCPUResource(InDestCPUResource)
-		, SourceMip(InSourceCurrentMip)
-		, DestMip(InDestCurrentMip)
-		, SourceArrayIndex(InSourceArrayIndex)
-		, DestArrayIndex(InDestArrayIndex)
-		, InitialPositionOffset(InInitialPositionOffset)
-		, SubSectionSizeQuad(InSubSectionSizeQuad)
-		, NumSubSections(InNumSubSections)
-		, SourceDebugName(InSourceResourceDebugName)
-		, DestResourceDebugName(InDestResourceDebugName)
-	{}
+	FLandscapeLayersCopyTexture_RenderThread(const FLandscapeLayersCopyTextureParams& InParams)
+		: Params(InParams)
+	{}	
 
 	void Copy(FRHICommandListImmediate& InRHICmdList)
 	{
 		SCOPE_CYCLE_COUNTER(STAT_LandscapeLayersRegenerate_RenderThread);
-		SCOPED_DRAW_EVENTF(InRHICmdList, LandscapeLayersCopy, TEXT("LS Copy %s -> %s, Mip (%d -> %d), Array Index (%d -> %d)"), *SourceDebugName, *DestResourceDebugName, SourceMip, DestMip, SourceArrayIndex, DestArrayIndex);
+		SCOPED_DRAW_EVENTF(InRHICmdList, LandscapeLayersCopy, TEXT("LS Copy %s -> %s, Mip (%d -> %d), Array Index (%d -> %d)"), *Params.SourceResourceDebugName, *Params.DestResourceDebugName, Params.SourceMip, Params.DestMip, Params.SourceArrayIndex, Params.DestArrayIndex);
 		SCOPED_GPU_STAT(InRHICmdList, LandscapeLayersCopy);
 
-		FIntPoint SourceSize(SourceResource->GetSizeX(), SourceResource->GetSizeY()); // SourceResource is always proper size, as it's always the good MIP we want to copy from
-		FIntPoint DestSize(DestResource->GetSizeX() >> DestMip, DestResource->GetSizeY() >> DestMip);
+		FIntPoint SourceSize(Params.SourceResource->GetSizeX(), Params.SourceResource->GetSizeY()); // SourceResource is always proper size, as it's always the good MIP we want to copy from
+		FIntPoint DestSize(Params.DestResource->GetSizeX() >> Params.DestMip, Params.DestResource->GetSizeY() >> Params.DestMip);
 
-		int32 LocalComponentSizeQuad = SubSectionSizeQuad * NumSubSections;
-		FVector2D PositionOffset(FMath::RoundToInt(InitialPositionOffset.X / LocalComponentSizeQuad), FMath::RoundToInt(InitialPositionOffset.Y / LocalComponentSizeQuad));
+		int32 LocalComponentSizeQuad = Params.SubSectionSizeQuad * Params.NumSubSections;
+		FVector2D PositionOffset(FMath::RoundToInt(Params.InitialPositionOffset.X / LocalComponentSizeQuad), FMath::RoundToInt(Params.InitialPositionOffset.Y / LocalComponentSizeQuad));
 
-		FRHICopyTextureInfo Params;
-		Params.NumSlices = 1;
-		Params.Size.Z = 1;
-		Params.SourceSliceIndex = SourceArrayIndex;
-		Params.DestSliceIndex = DestArrayIndex;
-		Params.SourceMipIndex = 0; // In my case, always assume we copy from mip 0 to something else as in my case each mips will be stored into individual texture/RT
-		Params.DestMipIndex = DestMip;
+		FRHICopyTextureInfo Info;
+		Info.NumSlices = 1;
+		Info.Size.Z = 1;
+		Info.SourceSliceIndex = Params.SourceArrayIndex;
+		Info.DestSliceIndex = Params.DestArrayIndex;
+		Info.SourceMipIndex = 0; // In my case, always assume we copy from mip 0 to something else as in my case each mips will be stored into individual texture/RT
+		Info.DestMipIndex = Params.DestMip;
 
 		if (SourceSize.X <= DestSize.X)
 		{
-			Params.SourcePosition.X = 0;
-			Params.Size.X = SourceSize.X;
-			Params.DestPosition.X = FMath::RoundToInt(PositionOffset.X * (((SubSectionSizeQuad + 1) * NumSubSections) >> DestMip));
+			Info.SourcePosition.X = 0;
+			Info.Size.X = SourceSize.X;
+			Info.DestPosition.X = FMath::RoundToInt(PositionOffset.X * (((Params.SubSectionSizeQuad + 1) * Params.NumSubSections) >> Params.DestMip));
 		}
 		else
 		{
-			Params.SourcePosition.X = FMath::RoundToInt(PositionOffset.X * (((SubSectionSizeQuad + 1) * NumSubSections) >> SourceMip));
-			Params.Size.X = DestSize.X;
-			Params.DestPosition.X = 0;
+			Info.SourcePosition.X = FMath::RoundToInt(PositionOffset.X * (((Params.SubSectionSizeQuad + 1) *Params.NumSubSections) >> Params.SourceMip));
+			Info.Size.X = DestSize.X;
+			Info.DestPosition.X = 0;
 		}
 
 		if (SourceSize.Y <= DestSize.Y)
 		{
-			Params.SourcePosition.Y = 0;
-			Params.Size.Y = SourceSize.Y;
-			Params.DestPosition.Y = FMath::RoundToInt(PositionOffset.Y * (((SubSectionSizeQuad + 1) * NumSubSections) >> DestMip));
+			Info.SourcePosition.Y = 0;
+			Info.Size.Y = SourceSize.Y;
+			Info.DestPosition.Y = FMath::RoundToInt(PositionOffset.Y * (((Params.SubSectionSizeQuad + 1) * Params.NumSubSections) >> Params.DestMip));
 		}
 		else
 		{
-			Params.SourcePosition.Y = FMath::RoundToInt(PositionOffset.Y * (((SubSectionSizeQuad + 1) * NumSubSections) >> SourceMip));
-			Params.Size.Y = DestSize.Y;
-			Params.DestPosition.Y = 0;
+			Info.SourcePosition.Y = FMath::RoundToInt(PositionOffset.Y * (((Params.SubSectionSizeQuad + 1) * Params.NumSubSections) >> Params.SourceMip));
+			Info.Size.Y = DestSize.Y;
+			Info.DestPosition.Y = 0;
 		}
 
-		InRHICmdList.CopyTexture(SourceResource->TextureRHI, DestResource->TextureRHI, Params);
+		InRHICmdList.CopyTexture(Params.SourceResource->TextureRHI, Params.DestResource->TextureRHI, Info);
 
-		if (DestCPUResource != nullptr)
+		if (Params.DestCPUResource != nullptr)
 		{
-			InRHICmdList.CopyTexture(SourceResource->TextureRHI, DestCPUResource->TextureRHI, Params);
+			InRHICmdList.CopyTexture(Params.SourceResource->TextureRHI, Params.DestCPUResource->TextureRHI, Info);
 		}
 	}
 
 private:
-	FTextureResource* SourceResource;
-	FTextureResource* DestResource;
-	FTextureResource* DestCPUResource;
-	uint8 SourceMip;
-	uint8 DestMip;
-	uint32 SourceArrayIndex;
-	uint32 DestArrayIndex;
-	FIntPoint InitialPositionOffset;
-	int32 SubSectionSizeQuad;
-	int32 NumSubSections;
-	FString SourceDebugName;
-	FString DestResourceDebugName;
+	FLandscapeLayersCopyTextureParams Params;
 };
 
 // Clear command
@@ -1782,6 +1759,38 @@ void ALandscapeProxy::InitializeProxyLayersWeightmapUsage()
 	}
 }
 
+void ALandscape::AddDeferredCopyLayersTexture(UTexture* InSourceTexture, UTexture* InDestTexture, FTextureResource* InDestCPUResource, const FIntPoint& InInitialPositionOffset, uint8 InSourceCurrentMip, uint8 InDestCurrentMip, uint32 InSourceArrayIndex, uint32 InDestArrayIndex)
+{
+	if (InSourceTexture != nullptr && InDestTexture != nullptr)
+	{
+		AddDeferredCopyLayersTexture(InSourceTexture->GetName(), InSourceTexture->Resource, InDestTexture->GetName(), InDestTexture->Resource, InDestCPUResource, InInitialPositionOffset, InSourceCurrentMip, InDestCurrentMip, InSourceArrayIndex, InDestArrayIndex);
+	}
+}
+
+void ALandscape::AddDeferredCopyLayersTexture(const FString& InSourceDebugName, FTextureResource* InSourceResource, const FString& InDestDebugName, FTextureResource* InDestResource, FTextureResource* InDestCPUResource, const FIntPoint& InInitialPositionOffset,
+											  uint8 InSourceCurrentMip, uint8 InDestCurrentMip, uint32 InSourceArrayIndex, uint32 InDestArrayIndex)
+{
+	check(InSourceResource != nullptr);
+	check(InDestResource != nullptr);
+
+	PendingCopyTextures.Add(FLandscapeLayersCopyTextureParams(InSourceDebugName, InSourceResource, InDestDebugName, InDestResource, InDestCPUResource, InInitialPositionOffset, SubsectionSizeQuads, NumSubsections, InSourceCurrentMip, InDestCurrentMip, InSourceArrayIndex, InDestArrayIndex));
+}
+
+void ALandscape::CommitDeferredCopyLayersTexture()
+{
+	TArray<FLandscapeLayersCopyTextureParams> LocalParams = MoveTemp(PendingCopyTextures);
+
+	ENQUEUE_RENDER_COMMAND(FLandscapeLayersCopyAsyncCommand)(
+		[LocalParams](FRHICommandListImmediate& RHICmdList) mutable
+	{
+		for (const FLandscapeLayersCopyTextureParams& Params : LocalParams)
+		{
+			FLandscapeLayersCopyTexture_RenderThread CopyTexture(Params);
+			CopyTexture.Copy(RHICmdList);
+		}
+	});
+}
+
 void ALandscape::CopyLayersTexture(UTexture* InSourceTexture, UTexture* InDestTexture, FTextureResource* InDestCPUResource, const FIntPoint& InFirstComponentSectionBase, uint8 InSourceCurrentMip, uint8 InDestCurrentMip, uint32 InSourceArrayIndex, uint32 InDestArrayIndex) const
 {
 	if (InSourceTexture != nullptr && InDestTexture != nullptr)
@@ -1796,7 +1805,7 @@ void ALandscape::CopyLayersTexture(const FString& InSourceDebugName, FTextureRes
 	check(InSourceResource != nullptr);
 	check(InDestResource != nullptr);
 
-	FLandscapeLayersCopyTexture_RenderThread CopyTexture(InSourceDebugName, InSourceResource, InDestDebugName, InDestResource, InDestCPUResource, InInitialPositionOffset, SubsectionSizeQuads, NumSubsections, InSourceCurrentMip, InDestCurrentMip, InSourceArrayIndex, InDestArrayIndex);
+	FLandscapeLayersCopyTexture_RenderThread CopyTexture(FLandscapeLayersCopyTextureParams(InSourceDebugName, InSourceResource, InDestDebugName, InDestResource, InDestCPUResource, InInitialPositionOffset, SubsectionSizeQuads, NumSubsections, InSourceCurrentMip, InDestCurrentMip, InSourceArrayIndex, InDestArrayIndex));
 
 	ENQUEUE_RENDER_COMMAND(FLandscapeLayersCopyCommand)(
 		[CopyTexture](FRHICommandListImmediate& RHICmdList) mutable
@@ -2641,13 +2650,18 @@ int32 ALandscape::RegenerateLayersHeightmaps(const TArray<ULandscapeComponent*>&
 		// Use to compute TopLeft Component per Heightmap
 		struct FHeightmapTopLeft
 		{
-			FHeightmapTopLeft(UTexture2D* InTexture, FIntPoint InTopLeftSectionBase) : Texture(InTexture), TopLeftSectionBase(InTopLeftSectionBase) {}
+			FHeightmapTopLeft(UTexture2D* InTexture, FIntPoint InTopLeftSectionBase, FLandscapeLayersTexture2DCPUReadBackResource* InCPUReadback = nullptr) 
+			: Texture(InTexture)
+			, TopLeftSectionBase(InTopLeftSectionBase)
+			, CPUReadback(InCPUReadback)
+			{}
 
 			FHeightmapTopLeft(FHeightmapTopLeft&&) = default;
 			FHeightmapTopLeft& operator=(FHeightmapTopLeft&&) = default;
 
 			UTexture2D* Texture;
 			FIntPoint TopLeftSectionBase;
+			FLandscapeLayersTexture2DCPUReadBackResource* CPUReadback;
 		};
 
 		// Calculate Top Left Lambda
@@ -2662,7 +2676,8 @@ int32 ALandscape::RegenerateLayersHeightmaps(const TArray<ULandscapeComponent*>&
 
 				if (Index == INDEX_NONE)
 				{
-					OutHeightmaps.Add(FHeightmapTopLeft(ComponentHeightmap, Component->GetSectionBase()));
+					FLandscapeLayersTexture2DCPUReadBackResource** CPUReadback = InProxy->HeightmapsCPUReadBack.Find(ComponentHeightmap);
+					OutHeightmaps.Add(FHeightmapTopLeft(ComponentHeightmap, Component->GetSectionBase(), CPUReadback != nullptr ? *CPUReadback : nullptr));
 				}
 				else
 				{
@@ -2711,11 +2726,12 @@ int32 ALandscape::RegenerateLayersHeightmaps(const TArray<ULandscapeComponent*>&
 
 				for (const FHeightmapTopLeft& LayerHeightmap : LayerHeightmaps)
 				{
-					CopyLayersTexture(LayerHeightmap.Texture, LandscapeScratchRT1, nullptr, LayerHeightmap.TopLeftSectionBase-MinExtend);
-					PrintLayersDebugRT(OutputDebugName ? FString::Printf(TEXT("LS Height: %s Component %s += -> CombinedAtlas %s"), *Layer.Name.ToString(), *LayerHeightmap.Texture->GetName(), *LandscapeScratchRT1->GetName()) : TEXT(""), LandscapeScratchRT1);
+					AddDeferredCopyLayersTexture(LayerHeightmap.Texture, LandscapeScratchRT1, nullptr, LayerHeightmap.TopLeftSectionBase - MinExtend);
 				}
 			});
 
+			CommitDeferredCopyLayersTexture();
+
 			// NOTE: From this point on, we always work in non atlas, we'll convert back at the end to atlas only
 			DrawHeightmapComponentsToRenderTarget(OutputDebugName ? FString::Printf(TEXT("LS Height: %s += -> NonAtlas %s"), *Layer.Name.ToString(), *LandscapeScratchRT1->GetName(), *LandscapeScratchRT2->GetName()) : TEXT(""),
 												  InLandscapeComponents, MinExtend, LandscapeScratchRT1, nullptr, LandscapeScratchRT2, ERTDrawingType::RTAtlasToNonAtlas, true, ShaderParams);
@@ -2786,26 +2802,30 @@ int32 ALandscape::RegenerateLayersHeightmaps(const TArray<ULandscapeComponent*>&
 
 			for (const FHeightmapTopLeft& Heightmap : Heightmaps)
 			{
-				int32 CurrentMip = 0;
-
-				FLandscapeLayersTexture2DCPUReadBackResource** CPUReadback = Proxy->HeightmapsCPUReadBack.Find(Heightmap.Texture);
-				check(CPUReadback != nullptr);
-
 				FIntPoint TextureTopLeftSectionBase = Heightmap.TopLeftSectionBase - MinExtend;
+				uint8 MipIndex = 0;
+				
+				check(Heightmap.CPUReadback);
 
-				CopyLayersTexture(CombinedHeightmapAtlasRT, Heightmap.Texture, *CPUReadback, TextureTopLeftSectionBase, CurrentMip, CurrentMip);
-				++CurrentMip;
+				// Mip 0
+				AddDeferredCopyLayersTexture(CombinedHeightmapAtlasRT, Heightmap.Texture, Heightmap.CPUReadback, TextureTopLeftSectionBase, MipIndex, MipIndex);
+				++MipIndex;
 
+				// Other Mips
 				for (int32 MipRTIndex = EHeightmapRTType::HeightmapRT_Mip1; MipRTIndex < EHeightmapRTType::HeightmapRT_Count; ++MipRTIndex)
 				{
-					if (HeightmapRTList[MipRTIndex] != nullptr)
+					UTextureRenderTarget2D* RenderTargetMip = HeightmapRTList[MipRTIndex];
+
+					if (RenderTargetMip != nullptr)
 					{
-						CopyLayersTexture(HeightmapRTList[MipRTIndex], Heightmap.Texture, *CPUReadback, TextureTopLeftSectionBase, CurrentMip, CurrentMip);
-						++CurrentMip;
+						AddDeferredCopyLayersTexture(RenderTargetMip, Heightmap.Texture, Heightmap.CPUReadback, TextureTopLeftSectionBase, MipIndex, MipIndex);
+						++MipIndex;
 					}
 				}
 			}
 		});
+
+		CommitDeferredCopyLayersTexture();
 	}
 
 	if (HeightmapUpdateModes)
@@ -2870,44 +2890,46 @@ void ALandscape::ResolveLayersHeightmapTexture()
 
 void ALandscape::ResolveLayersTexture(FLandscapeLayersTexture2DCPUReadBackResource* InCPUReadBackTexture, UTexture2D* InOutputTexture)
 {
-	TArray<TArray<FColor>> MipData;
-	MipData.AddDefaulted(InCPUReadBackTexture->TextureRHI->GetNumMips());
+	SCOPE_CYCLE_COUNTER(STAT_LandscapeLayersResolveTexture);
 
-	int32 MipSizeU = InCPUReadBackTexture->GetSizeX();
-	int32 MipSizeV = InCPUReadBackTexture->GetSizeY();
-	int32 MipIndex = 0;
+	TArray<TArray<FColor>> OutMipsData;
+	OutMipsData.AddDefaulted(InCPUReadBackTexture->TextureRHI->GetNumMips());
 
-	while (MipSizeU >= 1 && MipSizeV >= 1)
+	ENQUEUE_RENDER_COMMAND(FLandscapeLayersReadSurfaceCommand)(
+		[InCPUReadBackTexture, OutMipsData](FRHICommandListImmediate& RHICmdList) mutable
 	{
-		MipData[MipIndex].Reset();
+		int32 MipSizeU = InCPUReadBackTexture->GetSizeX();
+		int32 MipSizeV = InCPUReadBackTexture->GetSizeY();
+		int32 MipIndex = 0;
 
-		FReadSurfaceDataFlags Flags(RCM_UNorm, CubeFace_MAX);
-		Flags.SetMip(MipIndex);
-		FIntRect Rect(0, 0, MipSizeU, MipSizeV);
-
-		ENQUEUE_RENDER_COMMAND(FLandscapeLayersReadSurfaceCommand)(
-			[SourceTextureRHI = InCPUReadBackTexture->TextureRHI, Rect = Rect, OutData = &MipData[MipIndex], ReadFlags = Flags](FRHICommandListImmediate& RHICmdList) mutable
+		while (MipSizeU >= 1 && MipSizeV >= 1)
 		{
-			RHICmdList.ReadSurfaceData(SourceTextureRHI, Rect, *OutData, ReadFlags);
-		});
+			OutMipsData[MipIndex].Reset();
 
-		MipSizeU >>= 1;
-		MipSizeV >>= 1;
-		++MipIndex;
-	}
+			FReadSurfaceDataFlags Flags(RCM_UNorm, CubeFace_MAX);
+			Flags.SetMip(MipIndex);
+			FIntRect Rect(0, 0, MipSizeU, MipSizeV);
+
+			RHICmdList.ReadSurfaceData(InCPUReadBackTexture->TextureRHI, Rect, OutMipsData[MipIndex], Flags);
+
+			MipSizeU >>= 1;
+			MipSizeV >>= 1;
+			++MipIndex;
+		}
+	});
 
 	// TODO: find a way to NOT have to flush the rendering command as this create hic up of ~10-15ms
 	FlushRenderingCommands();
 
-	for (MipIndex = 0; MipIndex < MipData.Num(); ++MipIndex)
+	for (int8 MipIndex = 0; MipIndex < OutMipsData.Num(); ++MipIndex)
 	{
-		if (MipData[MipIndex].Num() > 0)
+		if (OutMipsData[MipIndex].Num() > 0)
 		{
 			FColor* TextureData = (FColor*)InOutputTexture->Source.LockMip(MipIndex);
-			FMemory::Memcpy(TextureData, MipData[MipIndex].GetData(), MipData[MipIndex].Num() * sizeof(FColor));
+			FMemory::Memcpy(TextureData, OutMipsData[MipIndex].GetData(), OutMipsData[MipIndex].Num() * sizeof(FColor));
 			InOutputTexture->Source.UnlockMip(MipIndex);
 		}
-	}
+	}	
 }
 
 void ALandscape::PrepareComponentDataToExtractMaterialLayersCS(const FLandscapeLayer& InLayer, int32 InCurrentWeightmapToProcessIndex, const FIntPoint& InLandscapeBase, bool InOutputDebugName, FLandscapeTexture2DResource* InOutTextureData,
@@ -4281,29 +4303,31 @@ void ALandscape::UpdateLayersContent(bool bInWaitForStreaming, bool bInSkipMonit
 	}
 	MonitorShaderCompilation();
 
-	if (LayerContentUpdateModes == 0)
+	const bool bForceRender = CVarOutputLayersDebugDrawCallName.GetValueOnAnyThread() == 1;
+
+	if (LayerContentUpdateModes == 0 && !bForceRender)
 	{
 		return;
 	}
 
-		TArray<ULandscapeComponent*> AllLandscapeComponents;
-		GetLandscapeInfo()->ForAllLandscapeProxies([&AllLandscapeComponents](ALandscapeProxy* Proxy)
-		{
-			AllLandscapeComponents.Append(Proxy->LandscapeComponents);
-		});
+	TArray<ULandscapeComponent*> AllLandscapeComponents;
+	GetLandscapeInfo()->ForAllLandscapeProxies([&AllLandscapeComponents](ALandscapeProxy* Proxy)
+	{
+		AllLandscapeComponents.Append(Proxy->LandscapeComponents);
+	});
 
-		int32 ProcessedModes = 0;
-		ProcessedModes |= RegenerateLayersHeightmaps(AllLandscapeComponents, bInWaitForStreaming);
-		ProcessedModes |= RegenerateLayersWeightmaps(AllLandscapeComponents, bInWaitForStreaming);
-		ProcessedModes |= (LayerContentUpdateModes & ELandscapeLayerUpdateMode::Update_Client_Deferred);
-		ProcessedModes |= (LayerContentUpdateModes & ELandscapeLayerUpdateMode::Update_Client_Editing);
-		LayerContentUpdateModes &= ~ProcessedModes;
+	int32 ProcessedModes = 0;
+	ProcessedModes |= RegenerateLayersHeightmaps(AllLandscapeComponents, bInWaitForStreaming);
+	ProcessedModes |= RegenerateLayersWeightmaps(AllLandscapeComponents, bInWaitForStreaming);
+	ProcessedModes |= (LayerContentUpdateModes & ELandscapeLayerUpdateMode::Update_Client_Deferred);
+	ProcessedModes |= (LayerContentUpdateModes & ELandscapeLayerUpdateMode::Update_Client_Editing);
+	LayerContentUpdateModes &= ~ProcessedModes;
 
-		if (!ALandscape::UpdateCollisionAndClients(AllLandscapeComponents, ProcessedModes))
-		{
-			LayerContentUpdateModes |= ELandscapeLayerUpdateMode::Update_Client_Deferred;
-		}
+	if (!ALandscape::UpdateCollisionAndClients(AllLandscapeComponents, ProcessedModes))
+	{
+		LayerContentUpdateModes |= ELandscapeLayerUpdateMode::Update_Client_Deferred;
 	}
+}
 
 // not thread safe
 struct FEnableCollisionHashOptimScope
@@ -5530,4 +5554,4 @@ bool FLandscapeLayerBrush::Initialize(const FIntPoint& InLandscapeSize, UTexture
 
 #endif // WITH_EDITOR
 
-#undef LOCTEXT_NAMESPACE
\ No newline at end of file
+#undef LOCTEXT_NAMESPACE
diff --git a/Engine/Source/Runtime/Landscape/Private/LandscapePrivate.h b/Engine/Source/Runtime/Landscape/Private/LandscapePrivate.h
--- a/Engine/Source/Runtime/Landscape/Private/LandscapePrivate.h
+++ b/Engine/Source/Runtime/Landscape/Private/LandscapePrivate.h
@@ -25,6 +25,7 @@ DECLARE_CYCLE_STAT_EXTERN(TEXT("Regenerate Layers (RenderThread)"), STAT_Landsca
 DECLARE_CYCLE_STAT_EXTERN(TEXT("Regenerate Layers Heightmap (GameThread)"), STAT_LandscapeLayersRegenerateHeightmaps, STATGROUP_Landscape, );
 DECLARE_CYCLE_STAT_EXTERN(TEXT("Resolve Layers Heightmaps"), STAT_LandscapeLayersResolveHeightmaps, STATGROUP_Landscape, );
 
+DECLARE_CYCLE_STAT_EXTERN(TEXT("Resolve Layers Texture"), STAT_LandscapeLayersResolveTexture, STATGROUP_Landscape, );
 DECLARE_CYCLE_STAT_EXTERN(TEXT("Regenerate Layers Weightmap (GameThread)"), STAT_LandscapeLayersRegenerateWeightmaps, STATGROUP_Landscape, );
 
 DECLARE_CYCLE_STAT_EXTERN(TEXT("Resolve Layers Weightmaps"), STAT_LandscapeLayersResolveWeightmaps, STATGROUP_Landscape, );
