Removed default c++ class ALandscapeBlueprintCustomBrush from BP Brush classes in Landscape
Revmoed ALandscapeBlueprintCustomSimulationBrush class since not used
#rnx

diff --git a/Engine/Source/Runtime/Landscape/Classes/LandscapeBPCustomBrush.h b/Engine/Source/Runtime/Landscape/Classes/LandscapeBPCustomBrush.h
--- a/Engine/Source/Runtime/Landscape/Classes/LandscapeBPCustomBrush.h
+++ b/Engine/Source/Runtime/Landscape/Classes/LandscapeBPCustomBrush.h
@@ -8,7 +8,7 @@
 
 #include "LandscapeBPCustomBrush.generated.h"
 
-UCLASS(Blueprintable, hidecategories = (Replication, Input, LOD, Actor, Cooking, Rendering))
+UCLASS(Abstract, Blueprintable, hidecategories = (Replication, Input, LOD, Actor, Cooking, Rendering))
 class LANDSCAPE_API ALandscapeBlueprintCustomBrush : public AActor
 {
 	GENERATED_UCLASS_BODY()
@@ -71,13 +71,3 @@ public:
 #endif
 };
 
-UCLASS(Blueprintable, hidecategories = (Replication, Input, LOD, Actor, Cooking, Rendering))
-class LANDSCAPE_API ALandscapeBlueprintCustomSimulationBrush : public ALandscapeBlueprintCustomBrush
-{
-	GENERATED_UCLASS_BODY()
-
-public:
-	// TODO: To Implement
-};
-
-
diff --git a/Engine/Source/Runtime/Landscape/Private/LandscapeBPCustomBrush.cpp b/Engine/Source/Runtime/Landscape/Private/LandscapeBPCustomBrush.cpp
--- a/Engine/Source/Runtime/Landscape/Private/LandscapeBPCustomBrush.cpp
+++ b/Engine/Source/Runtime/Landscape/Private/LandscapeBPCustomBrush.cpp
@@ -138,8 +138,4 @@ void ALandscapeBlueprintCustomBrush::Destroyed()
 
 #endif
 
-ALandscapeBlueprintCustomSimulationBrush::ALandscapeBlueprintCustomSimulationBrush(const FObjectInitializer& ObjectInitializer)
-{
-}
-
 #undef LOCTEXT_NAMESPACE
