Fix crash when creating new Landscape (with Layers) and specifying any LayerInfo (doesn't have import weight data in this code path)

#rnx

diff --git a/Engine/Source/Runtime/Landscape/Private/LandscapeEdit.cpp b/Engine/Source/Runtime/Landscape/Private/LandscapeEdit.cpp
--- a/Engine/Source/Runtime/Landscape/Private/LandscapeEdit.cpp
+++ b/Engine/Source/Runtime/Landscape/Private/LandscapeEdit.cpp
@@ -2867,7 +2867,7 @@ LANDSCAPE_API void ALandscapeProxy::Import(const FGuid& InGuid, int32 InMinX, in
 			{
 				for (const FLandscapeImportLayerInfo& MaterialLayerInfo : *ImportWeightData)
 				{
-					if (MaterialLayerInfo.LayerInfo != nullptr)
+					if (MaterialLayerInfo.LayerInfo != nullptr && MaterialLayerInfo.LayerData.Num() > 0)
 					{
 						LandscapeEdit.SetAlphaData(MaterialLayerInfo.LayerInfo, InMinX, InMinY, InMaxX, InMaxY, MaterialLayerInfo.LayerData.GetData(), 0, ELandscapeLayerPaintingRestriction::None, true, false);
 					}
diff --git a/Engine/Source/Runtime/Landscape/Private/LandscapeEditInterface.cpp b/Engine/Source/Runtime/Landscape/Private/LandscapeEditInterface.cpp
--- a/Engine/Source/Runtime/Landscape/Private/LandscapeEditInterface.cpp
+++ b/Engine/Source/Runtime/Landscape/Private/LandscapeEditInterface.cpp
@@ -2733,7 +2733,8 @@ const ULandscapeLayerInfoObject* FLandscapeEditDataInterface::ChooseReplacementL
 
 void FLandscapeEditDataInterface::SetAlphaData(ULandscapeLayerInfoObject* const LayerInfo, const int32 X1, const int32 Y1, const int32 X2, const int32 Y2, const uint8* Data, int32 Stride, ELandscapeLayerPaintingRestriction PaintingRestriction /*= None*/, bool bWeightAdjust /*= true*/, bool bTotalWeightAdjust /*= false*/)
 {
-	check(LayerInfo != NULL);
+	check(LayerInfo != nullptr);
+	check(Data != nullptr);
 	if (LayerInfo->bNoWeightBlend)
 	{
 		bWeightAdjust = false;
