Initialize bRequiresAdjacencyInformation in FLandscapeComponentSceneProxy

diff --git a/Engine/Source/Runtime/Landscape/Private/LandscapeRender.cpp b/Engine/Source/Runtime/Landscape/Private/LandscapeRender.cpp
--- a/Engine/Source/Runtime/Landscape/Private/LandscapeRender.cpp
+++ b/Engine/Source/Runtime/Landscape/Private/LandscapeRender.cpp
@@ -635,6 +635,7 @@ FLandscapeComponentSceneProxy::FLandscapeComponentSceneProxy(ULandscapeComponent
 	, FLandscapeNeighborInfo(InComponent->GetWorld(), InComponent->GetLandscapeProxy()->GetLandscapeGuid(), InComponent->GetSectionBase() / InComponent->ComponentSizeQuads, InComponent->GetHeightmap(), InComponent->ForcedLOD, InComponent->LODBias)
 	, MaxLOD(FMath::CeilLogTwo(InComponent->SubsectionSizeQuads + 1) - 1)
 	, UseTessellationComponentScreenSizeFalloff(InComponent->GetLandscapeProxy()->UseTessellationComponentScreenSizeFalloff)
+	, bRequiresAdjacencyInformation(false)
 	, NumWeightmapLayerAllocations(InComponent->GetWeightmapLayerAllocations().Num())
 	, StaticLightingLOD(InComponent->GetLandscapeProxy()->StaticLightingLOD)
 	, WeightmapSubsectionOffset(InComponent->WeightmapSubsectionOffset)
