Fix Blueprint Brush Weightmap bug where painting was needed in order for the brush to affect weightmap

#rnx

diff --git a/Engine/Source/Runtime/Landscape/Private/LandscapeEditLayers.cpp b/Engine/Source/Runtime/Landscape/Private/LandscapeEditLayers.cpp
--- a/Engine/Source/Runtime/Landscape/Private/LandscapeEditLayers.cpp
+++ b/Engine/Source/Runtime/Landscape/Private/LandscapeEditLayers.cpp
@@ -3481,8 +3481,8 @@ int32 ALandscape::RegenerateLayersWeightmaps(const TArray<ULandscapeComponent*>&
 
 		CommitDeferredCopyLayersTexture();
 
-		bool ComputeShaderGeneratedData = false;
-		bool FirstLayer = true;
+		bool bHasWeightmapData = false;
+		bool bFirstLayer = true;
 		TMap<ULandscapeLayerInfoObject*, bool> WeightmapLayersBlendSubstractive;
 
 		for (FLandscapeLayer& Layer : LandscapeLayers)
@@ -3505,6 +3505,7 @@ int32 ALandscape::RegenerateLayersWeightmaps(const TArray<ULandscapeComponent*>&
 						if (Brush.IsAffectingWeightmapLayer(InfoLayerSettings.GetLayerName()) && !LayerInfoObjects.Contains(InfoLayerSettings.LayerInfoObj))
 						{
 							LayerInfoObjects.Add(InfoLayerSettings.LayerInfoObj, LayerInfoSettingsIndex + 1); // due to visibility layer that is at 0
+							bHasWeightmapData = true;
 						}
 					}
 				}
@@ -3565,11 +3566,11 @@ int32 ALandscape::RegenerateLayersWeightmaps(const TArray<ULandscapeComponent*>&
 					});
 
 					++CurrentWeightmapToProcessIndex;
-					ComputeShaderGeneratedData = true; // at least 1 CS was executed, so we can continue the processing
+					bHasWeightmapData = true; // at least 1 CS was executed, so we can continue the processing
 				}
 			}
 
-			// If we did process at least one compute shader
+			// If we have data in at least one weight map layer
 			if (LayerInfoObjects.Num() > 0)
 			{
 				for (auto& LayerInfoObject : LayerInfoObjects)
@@ -3603,7 +3604,7 @@ int32 ALandscape::RegenerateLayersWeightmaps(const TArray<ULandscapeComponent*>&
 					// Combine with current status and copy back to the combined 2d resource array
 					PSShaderParams.OutputAsSubstractive = false;
 
-					if (!FirstLayer)
+					if (!bFirstLayer)
 					{
 						const bool* BlendSubstractive = Layer.WeightmapLayerAllocationBlend.Find(LayerInfoObj);
 						PSShaderParams.OutputAsSubstractive = BlendSubstractive != nullptr ? *BlendSubstractive : false;
@@ -3616,7 +3617,7 @@ int32 ALandscape::RegenerateLayersWeightmaps(const TArray<ULandscapeComponent*>&
 					}
 
 					DrawWeightmapComponentsToRenderTarget(OutputDebugName ? FString::Printf(TEXT("LS Weight: %s PaintLayer: %s, %s += -> Combined %s"), *Layer.Name.ToString(), *LayerInfoObj->LayerName.ToString(), *LandscapeScratchRT2->GetName(), *LandscapeScratchRT3->GetName()) : GEmptyDebugName,
-														  InLandscapeComponents, LandscapeExtent.Min, LandscapeScratchRT2, FirstLayer ? nullptr : LandscapeScratchRT1, LandscapeScratchRT3, true, PSShaderParams, 0);
+														  InLandscapeComponents, LandscapeExtent.Min, LandscapeScratchRT2, bFirstLayer ? nullptr : LandscapeScratchRT1, LandscapeScratchRT3, true, PSShaderParams, 0);
 
 					PSShaderParams.OutputAsSubstractive = false;
 
@@ -3664,13 +3665,12 @@ int32 ALandscape::RegenerateLayersWeightmaps(const TArray<ULandscapeComponent*>&
 				PSShaderParams.ApplyLayerModifiers = false;
 			}
 
-			FirstLayer = false;
+			bFirstLayer = false;
 		}
 
-		// TODO:  if editing a Brush affecting layers, since we don't have any bounds to brush, right now ReallocateLayersWeightmaps wont ask a rebuild of the component affected by Brushes, which mean ComponentThatNeedMaterialRebuild wont contains Brush affected component!
 		ReallocateLayersWeightmaps(BrushRequiredAllocations);
 
-		if (ComputeShaderGeneratedData)
+		if (bHasWeightmapData)
 		{
 			// Will generate CPU read back resource, if required
 			bool bHasPendingInitResource = false;
