Fix Undo/Redo crash relating to Cached editing layer in LandscapeComponent

#rnx

diff --git a/Engine/Source/Runtime/Landscape/Private/LandscapeEdit.cpp b/Engine/Source/Runtime/Landscape/Private/LandscapeEdit.cpp
--- a/Engine/Source/Runtime/Landscape/Private/LandscapeEdit.cpp
+++ b/Engine/Source/Runtime/Landscape/Private/LandscapeEdit.cpp
@@ -599,6 +599,10 @@ void ULandscapeComponent::PostEditUndo()
 		const bool bUpdateAll = true;
 		RequestHeightmapUpdate(bUpdateAll);
 		RequestWeightmapUpdate(bUpdateAll);
+
+		// Clear Cached Editing Data
+		CachedEditingLayer.Invalidate();
+		CachedEditingLayerData = nullptr;
 	}
 	else
 	{
