Fix LandscapeEdMode Detail Panel not refreshing Brush ownership changes
Propagate Owning Landscape changes to Target Landscape Water Brush property

#rnx

diff --git a/Engine/Source/Runtime/Landscape/Private/LandscapeBlueprintBrushBase.cpp b/Engine/Source/Runtime/Landscape/Private/LandscapeBlueprintBrushBase.cpp
--- a/Engine/Source/Runtime/Landscape/Private/LandscapeBlueprintBrushBase.cpp
+++ b/Engine/Source/Runtime/Landscape/Private/LandscapeBlueprintBrushBase.cpp
@@ -169,9 +169,25 @@ void ALandscapeBlueprintBrushBase::Destroyed()
 void ALandscapeBlueprintBrushBase::SetOwningLandscape(ALandscape* InOwningLandscape)
 {
 #if WITH_EDITORONLY_DATA
+	if (OwningLandscape == InOwningLandscape)
+	{
+		return;
+	}
+
 	const bool bAlwaysMarkDirty = false;
 	Modify(bAlwaysMarkDirty);
+
+	if (OwningLandscape)
+	{
+		OwningLandscape->OnBlueprintBrushChanged();
+	}
+
 	OwningLandscape = InOwningLandscape;
+
+	if (OwningLandscape)
+	{
+		OwningLandscape->OnBlueprintBrushChanged();
+	}
 #endif
 }
 
diff --git a/Engine/Source/Runtime/Landscape/Private/LandscapeEditLayers.cpp b/Engine/Source/Runtime/Landscape/Private/LandscapeEditLayers.cpp
--- a/Engine/Source/Runtime/Landscape/Private/LandscapeEditLayers.cpp
+++ b/Engine/Source/Runtime/Landscape/Private/LandscapeEditLayers.cpp
@@ -5496,6 +5496,12 @@ bool ALandscape::ReorderLayer(int32 InStartingLayerIndex, int32 InDestinationLay
 
 FName ALandscape::GenerateUniqueLayerName(FName InName) const
 {
+	// If we are receiving a unique name, use it.
+	if (InName != NAME_None && !LandscapeLayers.ContainsByPredicate([InName](const FLandscapeLayer& Layer) { return Layer.Name == InName; }))
+	{
+		return InName;
+	}
+
 	FString BaseName = InName == NAME_None ? "Layer" : InName.ToString();
 	FName NewName;
 	int32 LayerIndex = 0;
diff --git a/Engine/Source/Runtime/Landscape/Public/LandscapeBlueprintBrushBase.h b/Engine/Source/Runtime/Landscape/Public/LandscapeBlueprintBrushBase.h
--- a/Engine/Source/Runtime/Landscape/Public/LandscapeBlueprintBrushBase.h
+++ b/Engine/Source/Runtime/Landscape/Public/LandscapeBlueprintBrushBase.h
@@ -44,7 +44,7 @@ public:
 	void RequestLandscapeUpdate();
 
 #if WITH_EDITOR
-	void SetOwningLandscape(class ALandscape* InOwningLandscape);
+	virtual void SetOwningLandscape(class ALandscape* InOwningLandscape);
 	class ALandscape* GetOwningLandscape() const;
 
 	bool IsAffectingHeightmap() const { return AffectHeightmap; }
