Fix Current Layer Weightmap not being cleared all the time
Make Shader Params const so that we don't have side effects (safety fix only)
Remove duplicate CopyTexture

#rnx

diff --git a/Engine/Source/Runtime/Landscape/Classes/Landscape.h b/Engine/Source/Runtime/Landscape/Classes/Landscape.h
--- a/Engine/Source/Runtime/Landscape/Classes/Landscape.h
+++ b/Engine/Source/Runtime/Landscape/Classes/Landscape.h
@@ -363,16 +363,16 @@ private:
 
 	void ClearLayersWeightmapTextureResource(const FString& InDebugName, FTextureRenderTargetResource* InTextureResourceToClear) const;
 	void DrawHeightmapComponentsToRenderTarget(const FString& InDebugName, const TArray<ULandscapeComponent*>& InComponentsToDraw, const FIntPoint& InLandscapeBase, UTexture* InHeightmapRTRead, UTextureRenderTarget2D* InOptionalHeightmapRTRead2, UTextureRenderTarget2D* InHeightmapRTWrite, ERTDrawingType InDrawType,
-											   bool InClearRTWrite, struct FLandscapeLayersHeightmapShaderParameters& InShaderParams, uint8 InMipRender = 0) const;
+											   bool InClearRTWrite, const struct FLandscapeLayersHeightmapShaderParameters& InShaderParams, uint8 InMipRender = 0) const;
 
 	void DrawWeightmapComponentsToRenderTarget(const FString& InDebugName, const TArray<FIntPoint>& InSectionBaseList, const FVector2D& InScaleBias, UTexture* InWeightmapRTRead, UTextureRenderTarget2D* InOptionalWeightmapRTRead2, UTextureRenderTarget2D* InWeightmapRTWrite,
-											   bool InClearRTWrite, struct FLandscapeLayersWeightmapShaderParameters& InShaderParams, uint8 InMipRender) const;
+											   bool InClearRTWrite, const struct FLandscapeLayersWeightmapShaderParameters& InShaderParams, uint8 InMipRender) const;
 
 	void DrawWeightmapComponentsToRenderTarget(const FString& InDebugName, const TArray<ULandscapeComponent*>& InComponentsToDraw, const FIntPoint& InLandscapeBase, UTexture* InWeightmapRTRead, UTextureRenderTarget2D* InOptionalWeightmapRTRead2, UTextureRenderTarget2D* InWeightmapRTWrite,
-											   bool InClearRTWrite, struct FLandscapeLayersWeightmapShaderParameters& InShaderParams, uint8 InMipRender) const;
+											   bool InClearRTWrite, const struct FLandscapeLayersWeightmapShaderParameters& InShaderParams, uint8 InMipRender) const;
 
-	void DrawHeightmapComponentsToRenderTargetMips(const TArray<ULandscapeComponent*>& InComponentsToDraw, const FIntPoint& InLandscapeBase, UTexture* InReadHeightmap, bool InClearRTWrite, struct FLandscapeLayersHeightmapShaderParameters& InShaderParams) const;
-	void DrawWeightmapComponentToRenderTargetMips(const TArray<FVector2D>& InTexturePositionsToDraw, UTexture* InReadWeightmap, bool InClearRTWrite, struct FLandscapeLayersWeightmapShaderParameters& InShaderParams) const;
+	void DrawHeightmapComponentsToRenderTargetMips(const TArray<ULandscapeComponent*>& InComponentsToDraw, const FIntPoint& InLandscapeBase, UTexture* InReadHeightmap, bool InClearRTWrite, const struct FLandscapeLayersHeightmapShaderParameters& InShaderParams) const;
+	void DrawWeightmapComponentToRenderTargetMips(const TArray<FVector2D>& InTexturePositionsToDraw, UTexture* InReadWeightmap, bool InClearRTWrite, const struct FLandscapeLayersWeightmapShaderParameters& InShaderParams) const;
 
 	void CopyLayersTexture(UTexture* InSourceTexture, UTexture* InDestTexture, FTextureResource* InDestCPUResource = nullptr, const FIntPoint& InInitialPositionOffset = FIntPoint(0, 0), uint8 InSourceCurrentMip = 0, uint8 InDestCurrentMip = 0,
 						   uint32 InSourceArrayIndex = 0, uint32 InDestArrayIndex = 0) const;
diff --git a/Engine/Source/Runtime/Landscape/Private/LandscapeEditLayers.cpp b/Engine/Source/Runtime/Landscape/Private/LandscapeEditLayers.cpp
--- a/Engine/Source/Runtime/Landscape/Private/LandscapeEditLayers.cpp
+++ b/Engine/Source/Runtime/Landscape/Private/LandscapeEditLayers.cpp
@@ -1824,7 +1824,7 @@ void ALandscape::CopyLayersTexture(const FString& InSourceDebugName, FTextureRes
 }
 
 void ALandscape::DrawWeightmapComponentsToRenderTarget(const FString& InDebugName, const TArray<FIntPoint>& InSectionBaseList, const FVector2D& InScaleBias, UTexture* InWeightmapRTRead, UTextureRenderTarget2D* InOptionalWeightmapRTRead2, UTextureRenderTarget2D* InWeightmapRTWrite,
-														bool InClearRTWrite, FLandscapeLayersWeightmapShaderParameters& InShaderParams, uint8 InMipRender) const
+														bool InClearRTWrite, const FLandscapeLayersWeightmapShaderParameters& InShaderParams, uint8 InMipRender) const
 {
 	check(InWeightmapRTRead != nullptr);
 	check(InWeightmapRTWrite != nullptr);
@@ -1858,20 +1858,21 @@ void ALandscape::DrawWeightmapComponentsToRenderTarget(const FString& InDebugNam
 		}
 	}
 
-	InShaderParams.ReadWeightmap1 = InWeightmapRTRead;
-	InShaderParams.ReadWeightmap2 = InOptionalWeightmapRTRead2;
-	InShaderParams.CurrentMipComponentVertexCount = ((SubsectionSizeQuads + 1) >> InMipRender);
+	FLandscapeLayersWeightmapShaderParameters CopyParams(InShaderParams);
+	CopyParams.ReadWeightmap1 = InWeightmapRTRead;
+	CopyParams.ReadWeightmap2 = InOptionalWeightmapRTRead2;
+	CopyParams.CurrentMipComponentVertexCount = ((SubsectionSizeQuads + 1) >> InMipRender);
 
 	if (InMipRender > 0)
 	{
-		InShaderParams.CurrentMipSize = WeightmapWriteTextureSize;
-		InShaderParams.ParentMipSize = WeightmapReadTextureSize;
+		CopyParams.CurrentMipSize = WeightmapWriteTextureSize;
+		CopyParams.ParentMipSize = WeightmapReadTextureSize;
 	}
 
 	FMatrix ProjectionMatrix = AdjustProjectionMatrixForRHI(FTranslationMatrix(FVector(0, 0, 0)) *
 		FMatrix(FPlane(1.0f / (FMath::Max<uint32>(WeightmapWriteTextureSize.X, 1.f) / 2.0f), 0.0, 0.0f, 0.0f), FPlane(0.0f, -1.0f / (FMath::Max<uint32>(WeightmapWriteTextureSize.Y, 1.f) / 2.0f), 0.0f, 0.0f), FPlane(0.0f, 0.0f, 1.0f, 0.0f), FPlane(-1.0f, 1.0f, 0.0f, 1.0f)));
 
-	FLandscapeLayersWeightmapRender_RenderThread LayersRender(InDebugName, InWeightmapRTWrite, WeightmapWriteTextureSize, WeightmapReadTextureSize, ProjectionMatrix, InShaderParams, InMipRender, TriangleList);
+	FLandscapeLayersWeightmapRender_RenderThread LayersRender(InDebugName, InWeightmapRTWrite, WeightmapWriteTextureSize, WeightmapReadTextureSize, ProjectionMatrix, CopyParams, InMipRender, TriangleList);
 
 	ENQUEUE_RENDER_COMMAND(FDrawLandscapeLayersWeightmapCommand)(
 		[LayersRender, InClearRTWrite](FRHICommandListImmediate& RHICmdList) mutable
@@ -1883,7 +1884,7 @@ void ALandscape::DrawWeightmapComponentsToRenderTarget(const FString& InDebugNam
 }
 
 void ALandscape::DrawWeightmapComponentsToRenderTarget(const FString& InDebugName, const TArray<ULandscapeComponent*>& InComponentsToDraw, const FIntPoint& InLandscapeBase, UTexture* InWeightmapRTRead, UTextureRenderTarget2D* InOptionalWeightmapRTRead2, UTextureRenderTarget2D* InWeightmapRTWrite,
-														bool InClearRTWrite, FLandscapeLayersWeightmapShaderParameters& InShaderParams, uint8 InMipRender) const
+														bool InClearRTWrite, const FLandscapeLayersWeightmapShaderParameters& InShaderParams, uint8 InMipRender) const
 {
 	check(InWeightmapRTRead != nullptr);
 	check(InWeightmapRTWrite != nullptr);
@@ -1925,21 +1926,23 @@ void ALandscape::DrawWeightmapComponentsToRenderTarget(const FString& InDebugNam
 		}
 	}
 
-	InShaderParams.ReadWeightmap1 = InWeightmapRTRead;
-	InShaderParams.ReadWeightmap2 = InOptionalWeightmapRTRead2;
-	InShaderParams.CurrentMipComponentVertexCount = ((SubsectionSizeQuads + 1) >> InMipRender);
+	FLandscapeLayersWeightmapShaderParameters CopyParams(InShaderParams);
+	CopyParams.ReadWeightmap1 = InWeightmapRTRead;
+	CopyParams.ReadWeightmap2 = InOptionalWeightmapRTRead2;
+	CopyParams.CurrentMipComponentVertexCount = ((SubsectionSizeQuads + 1) >> InMipRender);
 
 	if (InMipRender > 0)
 	{
-		InShaderParams.CurrentMipSize = WeightmapWriteTextureSize;
-		InShaderParams.ParentMipSize = WeightmapReadTextureSize;
+		CopyParams.CurrentMipSize = WeightmapWriteTextureSize;
+		CopyParams.ParentMipSize = WeightmapReadTextureSize;
 	}
 
 	FMatrix ProjectionMatrix = AdjustProjectionMatrixForRHI(FTranslationMatrix(FVector(0, 0, 0)) *
 															FMatrix(FPlane(1.0f / (FMath::Max<uint32>(WeightmapWriteTextureSize.X, 1.f) / 2.0f), 0.0, 0.0f, 0.0f), FPlane(0.0f, -1.0f / (FMath::Max<uint32>(WeightmapWriteTextureSize.Y, 1.f) / 2.0f), 0.0f, 0.0f), FPlane(0.0f, 0.0f, 1.0f, 0.0f), FPlane(-1.0f, 1.0f, 0.0f, 1.0f)));
 
-	FLandscapeLayersWeightmapRender_RenderThread LayersRender(InDebugName, InWeightmapRTWrite, WeightmapWriteTextureSize, WeightmapReadTextureSize, ProjectionMatrix, InShaderParams, InMipRender, TriangleList);
+	FLandscapeLayersWeightmapRender_RenderThread LayersRender(InDebugName, InWeightmapRTWrite, WeightmapWriteTextureSize, WeightmapReadTextureSize, ProjectionMatrix, CopyParams, InMipRender, TriangleList);
 
+		
 	ENQUEUE_RENDER_COMMAND(FDrawLandscapeLayersWeightmapCommand)(
 		[LayersRender, InClearRTWrite](FRHICommandListImmediate& RHICmdList) mutable
 	{
@@ -1949,7 +1952,7 @@ void ALandscape::DrawWeightmapComponentsToRenderTarget(const FString& InDebugNam
 	PrintLayersDebugRT(InDebugName, InWeightmapRTWrite, InMipRender, false);
 }
 
-void ALandscape::DrawWeightmapComponentToRenderTargetMips(const TArray<FVector2D>& InTexturePositionsToDraw, UTexture* InReadWeightmap, bool InClearRTWrite, struct FLandscapeLayersWeightmapShaderParameters& InShaderParams) const
+void ALandscape::DrawWeightmapComponentToRenderTargetMips(const TArray<FVector2D>& InTexturePositionsToDraw, UTexture* InReadWeightmap, bool InClearRTWrite, const struct FLandscapeLayersWeightmapShaderParameters& InShaderParams) const
 {
 	bool OutputDebugName = CVarOutputLayersDebugDrawCallName.GetValueOnAnyThread() == 1 || CVarOutputLayersRTContent.GetValueOnAnyThread() == 1 ? true : false;
 	int32 CurrentMip = 1;
@@ -1995,7 +1998,7 @@ void ALandscape::ClearLayersWeightmapTextureResource(const FString& InDebugName,
 	});
 }
 
-void ALandscape::DrawHeightmapComponentsToRenderTargetMips(const TArray<ULandscapeComponent*>& InComponentsToDraw, const FIntPoint& InLandscapeBase, UTexture* InReadHeightmap, bool InClearRTWrite, struct FLandscapeLayersHeightmapShaderParameters& InShaderParams) const
+void ALandscape::DrawHeightmapComponentsToRenderTargetMips(const TArray<ULandscapeComponent*>& InComponentsToDraw, const FIntPoint& InLandscapeBase, UTexture* InReadHeightmap, bool InClearRTWrite, const struct FLandscapeLayersHeightmapShaderParameters& InShaderParams) const
 {
 	bool OutputDebugName = CVarOutputLayersDebugDrawCallName.GetValueOnAnyThread() == 1 || CVarOutputLayersRTContent.GetValueOnAnyThread() == 1 ? true : false;
 	int32 CurrentMip = 1;
@@ -2016,7 +2019,7 @@ void ALandscape::DrawHeightmapComponentsToRenderTargetMips(const TArray<ULandsca
 }
 
 void ALandscape::DrawHeightmapComponentsToRenderTarget(const FString& InDebugName, const TArray<ULandscapeComponent*>& InComponentsToDraw, const FIntPoint& InLandscapeBase, UTexture* InHeightmapRTRead, UTextureRenderTarget2D* InOptionalHeightmapRTRead2, UTextureRenderTarget2D* InHeightmapRTWrite,
-													  ERTDrawingType InDrawType, bool InClearRTWrite, FLandscapeLayersHeightmapShaderParameters& InShaderParams, uint8 InMipRender) const
+													  ERTDrawingType InDrawType, bool InClearRTWrite, const struct FLandscapeLayersHeightmapShaderParameters& InShaderParams, uint8 InMipRender) const
 {
 	check(InHeightmapRTRead != nullptr);
 	check(InHeightmapRTWrite != nullptr);
@@ -2075,21 +2078,22 @@ void ALandscape::DrawHeightmapComponentsToRenderTarget(const FString& InDebugNam
 		}
 	}
 
-	InShaderParams.ReadHeightmap1 = InHeightmapRTRead;
-	InShaderParams.ReadHeightmap2 = InOptionalHeightmapRTRead2;
-	InShaderParams.HeightmapSize = HeightmapReadTextureSize;
-	InShaderParams.CurrentMipComponentVertexCount = ((SubsectionSizeQuads + 1) >> InMipRender);
+	FLandscapeLayersHeightmapShaderParameters CopyParams(InShaderParams);
+	CopyParams.ReadHeightmap1 = InHeightmapRTRead;
+	CopyParams.ReadHeightmap2 = InOptionalHeightmapRTRead2;
+	CopyParams.HeightmapSize = HeightmapReadTextureSize;
+	CopyParams.CurrentMipComponentVertexCount = ((SubsectionSizeQuads + 1) >> InMipRender);
 
 	if (InMipRender > 0)
 	{
-		InShaderParams.CurrentMipSize = HeightmapWriteTextureSize;
-		InShaderParams.ParentMipSize = HeightmapReadTextureSize;
+		CopyParams.CurrentMipSize = HeightmapWriteTextureSize;
+		CopyParams.ParentMipSize = HeightmapReadTextureSize;
 	}
 
 	FMatrix ProjectionMatrix = AdjustProjectionMatrixForRHI(FTranslationMatrix(FVector(0, 0, 0)) *
 															FMatrix(FPlane(1.0f / (FMath::Max<uint32>(HeightmapWriteTextureSize.X, 1.f) / 2.0f), 0.0, 0.0f, 0.0f), FPlane(0.0f, -1.0f / (FMath::Max<uint32>(HeightmapWriteTextureSize.Y, 1.f) / 2.0f), 0.0f, 0.0f), FPlane(0.0f, 0.0f, 1.0f, 0.0f), FPlane(-1.0f, 1.0f, 0.0f, 1.0f)));
 
-	FLandscapeLayersHeightmapRender_RenderThread LayersRender(InDebugName, InHeightmapRTWrite, HeightmapWriteTextureSize, HeightmapReadTextureSize, ProjectionMatrix, InShaderParams, InMipRender, TriangleList);
+	FLandscapeLayersHeightmapRender_RenderThread LayersRender(InDebugName, InHeightmapRTWrite, HeightmapWriteTextureSize, HeightmapReadTextureSize, ProjectionMatrix, CopyParams, InMipRender, TriangleList);
 
 	ENQUEUE_RENDER_COMMAND(FDrawLandscapeLayersHeightmapCommand)(
 		[LayersRender, InClearRTWrite](FRHICommandListImmediate& RHICmdList) mutable
@@ -3511,6 +3515,22 @@ int32 ALandscape::RegenerateLayersWeightmaps(const TArray<ULandscapeComponent*>&
 				}
 			}
 
+			// Clear the current atlas if required
+			ClearLayersWeightmapTextureResource(TEXT("ClearRT"), LandscapeScratchRT1->GameThread_GetRenderTargetResource());
+
+			// Important: for performance reason we only clear the layer we will write to, the other one might contain data but they will not be read during the blend phase
+			for (auto& ItPair : LayerInfoObjects)
+			{
+				int32 LayerIndex = ItPair.Value;
+				
+				SourceDebugName = OutputDebugName ? LandscapeScratchRT1->GetName() : GEmptyDebugName;
+				DestDebugName = OutputDebugName ? FString::Printf(TEXT("Weight: %s Clear CurrentProcLayerWeightmapAllLayersResource %d, "), *Layer.Name.ToString(), LayerIndex) : GEmptyDebugName;
+
+				AddDeferredCopyLayersTexture(SourceDebugName, LandscapeScratchRT1->GameThread_GetRenderTargetResource(), DestDebugName, CurrentLayersWeightmapAllMaterialLayersResource, nullptr, FIntPoint(0, 0), 0, 0, 0, LayerIndex);
+			}
+
+			CommitDeferredCopyLayersTexture();
+			
 			// Loop until there is no more weightmap texture to process
 			while (HasFoundWeightmapToProcess)
 			{
@@ -3530,25 +3550,6 @@ int32 ALandscape::RegenerateLayersWeightmaps(const TArray<ULandscapeComponent*>&
 				{
 					PrintLayersDebugTextureResource(OutputDebugName ? FString::Printf(TEXT("LS Weight: %s WeightmapScratchTexture %s"), *Layer.Name.ToString(), TEXT("WeightmapScratchTextureResource")) : GEmptyDebugName, WeightmapScratchExtractLayerTextureResource, 0, false);
 
-					// Clear the current atlas if required
-					if (CurrentWeightmapToProcessIndex == 0)
-					{
-						ClearLayersWeightmapTextureResource(TEXT("ClearRT"), LandscapeScratchRT1->GameThread_GetRenderTargetResource());
-
-						// Important: for performance reason we only clear the layer we will write to, the other one might contain data but they will not be read during the blend phase
-						for (auto& ItPair : LayerInfoObjects)
-						{
-							int32 LayerIndex = ItPair.Value;
-
-							SourceDebugName = OutputDebugName ? LandscapeScratchRT1->GetName() : GEmptyDebugName;
-							DestDebugName = OutputDebugName ? FString::Printf(TEXT("Weight: %s Clear CurrentProcLayerWeightmapAllLayersResource %d, "), *Layer.Name.ToString(), LayerIndex) : GEmptyDebugName;
-
-							AddDeferredCopyLayersTexture(SourceDebugName, LandscapeScratchRT1->GameThread_GetRenderTargetResource(), DestDebugName, CurrentLayersWeightmapAllMaterialLayersResource, nullptr, FIntPoint(0, 0), 0, 0, 0, LayerIndex);
-						}
-
-						CommitDeferredCopyLayersTexture();
-					}
-
 					FLandscapeLayerWeightmapExtractMaterialLayersComputeShaderParameters CSExtractLayersShaderParams;
 					CSExtractLayersShaderParams.AtlasWeightmapsPerLayer = CurrentLayersWeightmapAllMaterialLayersResource;
 					CSExtractLayersShaderParams.ComponentWeightmapResource = WeightmapScratchExtractLayerTextureResource;
@@ -3623,9 +3624,7 @@ int32 ALandscape::RegenerateLayersWeightmaps(const TArray<ULandscapeComponent*>&
 
 					SourceDebugName = OutputDebugName ? FString::Printf(TEXT("Weight: %s PaintLayer: %s %s"), *Layer.Name.ToString(), *LayerInfoObj->LayerName.ToString(), *LandscapeScratchRT3->GetName()) : GEmptyDebugName;
 					DestDebugName = OutputDebugName ? TEXT("CombinedProcLayerWeightmap") : GEmptyDebugName;
-
-					CopyLayersTexture(SourceDebugName, LandscapeScratchRT3->GameThread_GetRenderTargetResource(), DestDebugName, CombinedLayersWeightmapAllMaterialLayersResource, nullptr, FIntPoint(0, 0), 0, 0, 0, LayerIndex);
-
+										
 					// Handle brush blending
 					if (Layer.bVisible)
 					{
@@ -3658,8 +3657,9 @@ int32 ALandscape::RegenerateLayersWeightmaps(const TArray<ULandscapeComponent*>&
 
 						SourceDebugName = OutputDebugName ? FString::Printf(TEXT("Weight: %s PaintLayer: %s %s"), *Layer.Name.ToString(), *LayerInfoObj->LayerName.ToString(), *LandscapeScratchRT3->GetName()) : GEmptyDebugName;
 						DestDebugName = OutputDebugName ? TEXT("CombinedProcLayerWeightmap") : GEmptyDebugName;
-						CopyLayersTexture(SourceDebugName, LandscapeScratchRT3->GameThread_GetRenderTargetResource(), DestDebugName, CombinedLayersWeightmapAllMaterialLayersResource, nullptr, FIntPoint(0, 0), 0, 0, 0, LayerIndex);
 					}
+
+					CopyLayersTexture(SourceDebugName, LandscapeScratchRT3->GameThread_GetRenderTargetResource(), DestDebugName, CombinedLayersWeightmapAllMaterialLayersResource, nullptr, FIntPoint(0, 0), 0, 0, 0, LayerIndex);
 				}
 
 				PSShaderParams.ApplyLayerModifiers = false;
@@ -5656,6 +5656,8 @@ ALandscapeBlueprintBrushBase* FLandscapeLayerBrush::GetBrush() const
 #endif
 }
 
+PRAGMA_DISABLE_OPTIMIZATION
+
 void FLandscapeLayerBrush::SetOwner(ALandscape* InOwner)
 {
 #if WITH_EDITORONLY_DATA
@@ -5670,6 +5672,8 @@ void FLandscapeLayerBrush::SetOwner(ALandscape* InOwner)
 #endif
 }
 
+PRAGMA_ENABLE_OPTIMIZATION
+
 bool FLandscapeLayerBrush::IsAffectingHeightmap() const
 {
 #if WITH_EDITORONLY_DATA
@@ -5736,4 +5740,6 @@ bool FLandscapeLayerBrush::Initialize(const FIntRect& InLandscapeExtent, UTextur
 
 #endif // WITH_EDITOR
 
+PRAGMA_ENABLE_OPTIMIZATION
+
 #undef LOCTEXT_NAMESPACE
