Fix LandscapeGrass invalid Height when using large scale (fix from a 2016 answer on answers.unrealengine.com)

diff --git a/Engine/Source/Runtime/Landscape/Private/LandscapeGrass.cpp b/Engine/Source/Runtime/Landscape/Private/LandscapeGrass.cpp
--- a/Engine/Source/Runtime/Landscape/Private/LandscapeGrass.cpp
+++ b/Engine/Source/Runtime/Landscape/Private/LandscapeGrass.cpp
@@ -2013,11 +2013,25 @@ struct FAsyncGrassBuilder : public FGrassBuilderBase
 
 			OutLocation->X = InLocation.X - DrawScale.X * float(LandscapeSectionOffset.X);
 			OutLocation->Y = InLocation.Y - DrawScale.Y * float(LandscapeSectionOffset.Y);
-			// Bilinear interpolate
-			OutLocation->Z = DrawScale.Z * FMath::Lerp(
-				FMath::Lerp(Sample11, Sample21, LerpX),
-				FMath::Lerp(Sample12, Sample22, LerpX),
-				LerpY);
+			
+			FVector P11 = { (float)X1, (float)Y1, Sample11 };
+			FVector P12 = { (float)X1, (float)Y2, Sample12 };
+			FVector P21 = { (float)X2, (float)Y1, Sample21 };
+			FVector P22 = { (float)X2, (float)Y2, Sample22 };
+
+			float Z;
+			if (LerpX > LerpY)
+			{
+				FVector N = FVector::CrossProduct(P22 - P21, P11 - P21);
+				Z = P21.Z - ((TestX - P21.X) * N.X + (TestY - P21.Y) * N.Y) / N.Z;
+			}
+			else
+			{
+				FVector N = FVector::CrossProduct(P22 - P12, P11 - P12);
+				Z = P12.Z - ((TestX - P12.X) * N.X + (TestY - P12.Y) * N.Y) / N.Z;
+			}
+
+			OutLocation->Z = DrawScale.Z * Z;
 		}
 		return Result;
 	}
