Remove submitted PRAGMA_*_OPTIMIZATION

#rnx

diff --git a/Engine/Source/Runtime/Landscape/Private/LandscapeEditLayers.cpp b/Engine/Source/Runtime/Landscape/Private/LandscapeEditLayers.cpp
--- a/Engine/Source/Runtime/Landscape/Private/LandscapeEditLayers.cpp
+++ b/Engine/Source/Runtime/Landscape/Private/LandscapeEditLayers.cpp
@@ -5656,8 +5656,6 @@ ALandscapeBlueprintBrushBase* FLandscapeLayerBrush::GetBrush() const
 #endif
 }
 
-PRAGMA_DISABLE_OPTIMIZATION
-
 void FLandscapeLayerBrush::SetOwner(ALandscape* InOwner)
 {
 #if WITH_EDITORONLY_DATA
@@ -5672,8 +5670,6 @@ void FLandscapeLayerBrush::SetOwner(ALandscape* InOwner)
 #endif
 }
 
-PRAGMA_ENABLE_OPTIMIZATION
-
 bool FLandscapeLayerBrush::IsAffectingHeightmap() const
 {
 #if WITH_EDITORONLY_DATA
@@ -5740,6 +5736,4 @@ bool FLandscapeLayerBrush::Initialize(const FIntRect& InLandscapeExtent, UTextur
 
 #endif // WITH_EDITOR
 
-PRAGMA_ENABLE_OPTIMIZATION
-
 #undef LOCTEXT_NAMESPACE
