Fix for crash when landscape layers is enabled in World Settings on Mac.  Clone of Dev-Rendering CL 7093397 and CL 7196314.

diff --git a/Engine/Shaders/Public/Platform/Metal/MetalCommon.ush b/Engine/Shaders/Public/Platform/Metal/MetalCommon.ush
--- a/Engine/Shaders/Public/Platform/Metal/MetalCommon.ush
+++ b/Engine/Shaders/Public/Platform/Metal/MetalCommon.ush
@@ -5,8 +5,8 @@
 =============================================================================*/
 
 // Update this GUID to improve shader recompilation for Metal only shaders
-// GUID = CBA97EA3-E604-4B76-AD49-26D838AE8F3B
-  
+// GUID = cda708c4-2f08-4eb2-b045-42a269d14ed4
+
 #pragma once
 
 #if VERTEXSHADER && HULLSHADER
diff --git a/Engine/Source/Developer/Apple/MetalShaderFormat/Private/MetalBackend.cpp b/Engine/Source/Developer/Apple/MetalShaderFormat/Private/MetalBackend.cpp
--- a/Engine/Source/Developer/Apple/MetalShaderFormat/Private/MetalBackend.cpp
+++ b/Engine/Source/Developer/Apple/MetalShaderFormat/Private/MetalBackend.cpp
@@ -1945,7 +1945,14 @@ protected:
 				}
 				break;
 			case ir_txs:
-				ralloc_asprintf_append(buffer, "int2((int)");
+				if(tex->sampler->type->sampler_array)
+				{
+					ralloc_asprintf_append(buffer, "int3((int)");
+				}
+				else
+				{
+					ralloc_asprintf_append(buffer, "int2((int)");
+				}
 				break;
 			default:
 				break;
@@ -2237,8 +2244,58 @@ protected:
 			{
 				tex->sampler->accept(this);
 				ralloc_asprintf_append(buffer, ".read(");
-				tex->coordinate->accept(this);
 				
+				if (tex->sampler->type->sampler_array)
+				{
+					// Need to split the coordinate
+					char const* CoordSwizzle = "";
+					char const* IndexSwizzle = "y";
+					switch(tex->sampler->type->sampler_dimensionality)
+					{
+						case GLSL_SAMPLER_DIM_1D:
+						{
+							break;
+						}
+						case GLSL_SAMPLER_DIM_2D:
+						case GLSL_SAMPLER_DIM_RECT:
+						{
+							CoordSwizzle = "y";
+							IndexSwizzle = "z";
+							break;
+						}
+						case GLSL_SAMPLER_DIM_3D:
+						{
+							CoordSwizzle = "yz";
+							IndexSwizzle = "w";
+							break;
+						}
+						case GLSL_SAMPLER_DIM_CUBE:
+						{
+							CoordSwizzle = "yz";
+							IndexSwizzle = "w";
+							break;
+						}
+						case GLSL_SAMPLER_DIM_BUF:
+						case GLSL_SAMPLER_DIM_EXTERNAL:
+						default:
+						{
+							check(0);
+							break;
+						}
+					}
+					
+					ralloc_asprintf_append(buffer, "(");
+					tex->coordinate->accept(this);
+					
+					ralloc_asprintf_append(buffer, ").x%s, (uint)(", CoordSwizzle);
+					tex->coordinate->accept(this);
+					ralloc_asprintf_append(buffer, ").%s", IndexSwizzle);
+				}
+				else
+				{
+					tex->coordinate->accept(this);
+				}
+
 				if (tex->sampler->type->sampler_ms)
 				{
 					ralloc_asprintf_append(buffer, ",");
@@ -2429,6 +2486,13 @@ protected:
 				tex->lod_info.lod->accept(this);
 			}
 			ralloc_asprintf_append(buffer, ")");
+			
+			if(tex->sampler->type->sampler_array)
+			{
+				ralloc_asprintf_append(buffer, ", (int)");
+				tex->sampler->accept(this);
+				ralloc_asprintf_append(buffer, ".get_array_size()");
+			}
 		}
 			break;
 
diff --git a/Engine/Source/Developer/Apple/MetalShaderFormat/Private/MetalUtils.cpp b/Engine/Source/Developer/Apple/MetalShaderFormat/Private/MetalUtils.cpp
--- a/Engine/Source/Developer/Apple/MetalShaderFormat/Private/MetalUtils.cpp
+++ b/Engine/Source/Developer/Apple/MetalShaderFormat/Private/MetalUtils.cpp
@@ -3460,7 +3460,30 @@ void FMetalCodeBackend::RemovePackedVarReferences(exec_list* ir, _mesa_glsl_pars
 		}
 		else
 		{
-			OldVar->insert_after(NewVar);
+			// New variable declaration including its initialization must be inserted *after* the original variable has been initialized.
+			// So look ahead until we find the <ir_assignment> node that initializes "OldVar".
+			ir_assignment* OldVarInit = nullptr;
+			for (exec_node* Node = OldVar->next; Node && !OldVarInit; Node = Node->next)
+			{
+				ir_instruction* Instr = static_cast<ir_instruction*>(Node);
+				if (Instr->ir_type == ir_type_assignment)
+				{
+					ir_assignment* AssignStmt = static_cast<ir_assignment*>(Instr);
+					if (AssignStmt->lhs->variable_referenced() == OldVar)
+					{
+						OldVarInit = AssignStmt;
+					}
+				}
+			}
+			
+			if (OldVarInit != nullptr)
+			{
+				OldVarInit->insert_after(NewVar);
+			}
+			else
+			{
+				OldVar->insert_after(NewVar);
+			}
 		}
 	}
 
diff --git a/Engine/Source/Runtime/Apple/MetalRHI/Private/MetalBuffer.h b/Engine/Source/Runtime/Apple/MetalRHI/Private/MetalBuffer.h
--- a/Engine/Source/Runtime/Apple/MetalRHI/Private/MetalBuffer.h
+++ b/Engine/Source/Runtime/Apple/MetalRHI/Private/MetalBuffer.h
@@ -13,11 +13,12 @@ struct FMetalPooledBufferArgs
 {
     FMetalPooledBufferArgs() : Device(nil), Size(0), Flags(0), Storage(mtlpp::StorageMode::Shared) {}
 	
-    FMetalPooledBufferArgs(mtlpp::Device InDevice, uint32 InSize, uint32 InFlags, mtlpp::StorageMode InStorage)
+    FMetalPooledBufferArgs(mtlpp::Device InDevice, uint32 InSize, uint32 InFlags, mtlpp::StorageMode InStorage, mtlpp::CpuCacheMode InCpuCacheMode = mtlpp::CpuCacheMode::DefaultCache)
 	: Device(InDevice)
 	, Size(InSize)
     , Flags(InFlags)
 	, Storage(InStorage)
+	, CpuCacheMode(InCpuCacheMode)
 	{
 	}
 	
@@ -25,6 +26,7 @@ struct FMetalPooledBufferArgs
 	uint32 Size;
     uint32 Flags;
 	mtlpp::StorageMode Storage;
+	mtlpp::CpuCacheMode CpuCacheMode;
 };
 
 class FMetalSubBufferHeap
diff --git a/Engine/Source/Runtime/Apple/MetalRHI/Private/MetalContext.cpp b/Engine/Source/Runtime/Apple/MetalRHI/Private/MetalContext.cpp
--- a/Engine/Source/Runtime/Apple/MetalRHI/Private/MetalContext.cpp
+++ b/Engine/Source/Runtime/Apple/MetalRHI/Private/MetalContext.cpp
@@ -779,7 +779,8 @@ FMetalTexture FMetalDeviceContext::CreateTexture(FMetalSurface* Surface, mtlpp::
 
 FMetalBuffer FMetalDeviceContext::CreatePooledBuffer(FMetalPooledBufferArgs const& Args)
 {
-    FMetalBuffer Buffer = Heap.CreateBuffer(Args.Size, BufferOffsetAlignment, Args.Flags, FMetalCommandQueue::GetCompatibleResourceOptions((mtlpp::ResourceOptions)(BUFFER_CACHE_MODE | mtlpp::ResourceOptions::HazardTrackingModeUntracked | ((NSUInteger)Args.Storage << mtlpp::ResourceStorageModeShift))));
+	NSUInteger CpuResourceOption = ((NSUInteger)Args.CpuCacheMode) << mtlpp::ResourceCpuCacheModeShift;
+    FMetalBuffer Buffer = Heap.CreateBuffer(Args.Size, BufferOffsetAlignment, Args.Flags, FMetalCommandQueue::GetCompatibleResourceOptions((mtlpp::ResourceOptions)(CpuResourceOption | mtlpp::ResourceOptions::HazardTrackingModeUntracked | ((NSUInteger)Args.Storage << mtlpp::ResourceStorageModeShift))));
 	check(Buffer && Buffer.GetPtr());
 #if METAL_DEBUG_OPTIONS
 	if (GMetalResourcePurgeOnDelete && !Buffer.GetHeap())
diff --git a/Engine/Source/Runtime/Apple/MetalRHI/Private/MetalStateCache.cpp b/Engine/Source/Runtime/Apple/MetalRHI/Private/MetalStateCache.cpp
--- a/Engine/Source/Runtime/Apple/MetalRHI/Private/MetalStateCache.cpp
+++ b/Engine/Source/Runtime/Apple/MetalRHI/Private/MetalStateCache.cpp
@@ -1649,7 +1649,8 @@ void FMetalStateCache::SetShaderUnorderedAccessView(EMetalShaderStages ShaderSta
 				FPlatformAtomics::InterlockedExchange(&Source->Written, 1);
 				
 				SetShaderTexture(ShaderStage, Surface->Texture, BindIndex, mtlpp::ResourceUsage(mtlpp::ResourceUsage::Read | mtlpp::ResourceUsage::Write));
-				if ((Source->Flags & (TexCreate_UAV|TexCreate_NoTiling)) == (TexCreate_UAV|TexCreate_NoTiling))
+				
+				if ((Source->Flags & (TexCreate_UAV|TexCreate_NoTiling)) == (TexCreate_UAV|TexCreate_NoTiling) && Surface->Texture.GetBuffer())
 				{
 					FMetalBuffer Buffer(Surface->Texture.GetBuffer(), false);
 					SetShaderBuffer(ShaderStage, Buffer, nullptr, 0, Surface->Texture.GetBuffer().GetLength(), BindIndex, mtlpp::ResourceUsage(mtlpp::ResourceUsage::Read | mtlpp::ResourceUsage::Write), PF_MAX);
diff --git a/Engine/Source/Runtime/Apple/MetalRHI/Private/MetalTexture.cpp b/Engine/Source/Runtime/Apple/MetalRHI/Private/MetalTexture.cpp
--- a/Engine/Source/Runtime/Apple/MetalRHI/Private/MetalTexture.cpp
+++ b/Engine/Source/Runtime/Apple/MetalRHI/Private/MetalTexture.cpp
@@ -867,7 +867,9 @@ FMetalSurface::FMetalSurface(ERHIResourceType ResourceType, EPixelFormat Format,
 	}
 	else
 	{
-		if ((Flags & (TexCreate_UAV|TexCreate_NoTiling)) != (TexCreate_UAV|TexCreate_NoTiling))
+		const bool bBufferCompatibleOption = (Desc.GetTextureType() == mtlpp::TextureType::Texture2D || Desc.GetTextureType() == mtlpp::TextureType::TextureBuffer) && NumMips == 1;
+
+		if (!bBufferCompatibleOption || ((Flags & (TexCreate_UAV|TexCreate_NoTiling)) != (TexCreate_UAV|TexCreate_NoTiling)))
 		{
 			Texture = GetMetalDeviceContext().CreateTexture(this, Desc);
 		}
@@ -875,9 +877,12 @@ FMetalSurface::FMetalSurface(ERHIResourceType ResourceType, EPixelFormat Format,
 		{
 			mtlpp::Device Device = GetMetalDeviceContext().GetDevice();
 			mtlpp::SizeAndAlign SizeAlign = Device.HeapTextureSizeAndAlign(Desc);
-			FMetalPooledBufferArgs Args(Device, SizeAlign.Size, BUF_Dynamic, mtlpp::StorageMode::Private);
+
+			// Backing buffer resource options must match the texture we are going to create from it
+			FMetalPooledBufferArgs Args(Device, SizeAlign.Size, BUF_Dynamic, mtlpp::StorageMode::Private, Desc.GetCpuCacheMode());
 			FMetalBuffer Buffer = GetMetalDeviceContext().CreatePooledBuffer(Args);
-			Texture = Buffer.NewTexture(Desc, 0, SizeAlign.Size);
+
+			Texture = Buffer.NewTexture(Desc, 0, Align(SizeAlign.Size / SizeY, 256));
 		}
 		
 		METAL_FATAL_ASSERT(Texture, TEXT("Failed to create texture, desc %s"), *FString([Desc description]));
diff --git a/Engine/Source/Runtime/Landscape/Private/LandscapeEditLayers.cpp b/Engine/Source/Runtime/Landscape/Private/LandscapeEditLayers.cpp
--- a/Engine/Source/Runtime/Landscape/Private/LandscapeEditLayers.cpp
+++ b/Engine/Source/Runtime/Landscape/Private/LandscapeEditLayers.cpp
@@ -146,7 +146,7 @@ public:
 
 	static bool ShouldCompilePermutation(const FGlobalShaderPermutationParameters& Parameters)
 	{
-		return IsFeatureLevelSupported(Parameters.Platform, ERHIFeatureLevel::SM5) && !IsConsolePlatform(Parameters.Platform) && !IsMetalPlatform(Parameters.Platform);
+		return IsFeatureLevelSupported(Parameters.Platform, ERHIFeatureLevel::SM5) && !IsConsolePlatform(Parameters.Platform) && !IsMetalMobilePlatform(Parameters.Platform);
 	}
 
 	static void ModifyCompilationEnvironment(const FGlobalShaderPermutationParameters& Parameters, FShaderCompilerEnvironment& OutEnvironment)
@@ -218,7 +218,7 @@ public:
 
 	static bool ShouldCompilePermutation(const FGlobalShaderPermutationParameters& Parameters)
 	{
-		return IsFeatureLevelSupported(Parameters.Platform, ERHIFeatureLevel::SM5) && !IsConsolePlatform(Parameters.Platform) && !IsMetalPlatform(Parameters.Platform);
+		return IsFeatureLevelSupported(Parameters.Platform, ERHIFeatureLevel::SM5) && !IsConsolePlatform(Parameters.Platform) && !IsMetalMobilePlatform(Parameters.Platform);
 	}
 
 	static void ModifyCompilationEnvironment(const FGlobalShaderPermutationParameters& Parameters, FShaderCompilerEnvironment& OutEnvironment)
@@ -295,7 +295,7 @@ public:
 
 	static bool ShouldCompilePermutation(const FGlobalShaderPermutationParameters& Parameters)
 	{
-		return IsFeatureLevelSupported(Parameters.Platform, ERHIFeatureLevel::SM5) && !IsConsolePlatform(Parameters.Platform) && !IsMetalPlatform(Parameters.Platform);
+		return IsFeatureLevelSupported(Parameters.Platform, ERHIFeatureLevel::SM5) && !IsConsolePlatform(Parameters.Platform) && !IsMetalMobilePlatform(Parameters.Platform);
 	}
 
 	static void ModifyCompilationEnvironment(const FGlobalShaderPermutationParameters& Parameters, FShaderCompilerEnvironment& OutEnvironment)
@@ -382,7 +382,7 @@ public:
 
 	static bool ShouldCompilePermutation(const FGlobalShaderPermutationParameters& Parameters)
 	{
-		return IsFeatureLevelSupported(Parameters.Platform, ERHIFeatureLevel::SM5) && !IsConsolePlatform(Parameters.Platform) && !IsMetalPlatform(Parameters.Platform);
+		return IsFeatureLevelSupported(Parameters.Platform, ERHIFeatureLevel::SM5) && !IsConsolePlatform(Parameters.Platform) && !IsMetalMobilePlatform(Parameters.Platform);
 	}
 
 	static void ModifyCompilationEnvironment(const FGlobalShaderPermutationParameters& Parameters, FShaderCompilerEnvironment& OutEnvironment)
@@ -449,7 +449,7 @@ public:
 
 	static bool ShouldCompilePermutation(const FGlobalShaderPermutationParameters& Parameters)
 	{
-		return IsFeatureLevelSupported(Parameters.Platform, ERHIFeatureLevel::SM5) && !IsConsolePlatform(Parameters.Platform) && !IsMetalPlatform(Parameters.Platform);
+		return IsFeatureLevelSupported(Parameters.Platform, ERHIFeatureLevel::SM5) && !IsConsolePlatform(Parameters.Platform) && !IsMetalMobilePlatform(Parameters.Platform);
 	}
 
 	static void ModifyCompilationEnvironment(const FGlobalShaderPermutationParameters& Parameters, FShaderCompilerEnvironment& OutEnvironment)
@@ -701,7 +701,7 @@ public:
 
 	static bool ShouldCompilePermutation(const FGlobalShaderPermutationParameters& Parameters)
 	{
-		return IsFeatureLevelSupported(Parameters.Platform, ERHIFeatureLevel::SM5) && !IsConsolePlatform(Parameters.Platform) && !IsMetalPlatform(Parameters.Platform);
+		return IsFeatureLevelSupported(Parameters.Platform, ERHIFeatureLevel::SM5) && !IsConsolePlatform(Parameters.Platform) && !IsMetalMobilePlatform(Parameters.Platform);
 	}
 
 	static void ModifyCompilationEnvironment(const FGlobalShaderPermutationParameters& Parameters, FShaderCompilerEnvironment& OutEnvironment)
@@ -895,7 +895,7 @@ public:
 
 	static bool ShouldCompilePermutation(const FGlobalShaderPermutationParameters& Parameters)
 	{
-		return IsFeatureLevelSupported(Parameters.Platform, ERHIFeatureLevel::SM5) && !IsConsolePlatform(Parameters.Platform) && !IsMetalPlatform(Parameters.Platform);
+		return IsFeatureLevelSupported(Parameters.Platform, ERHIFeatureLevel::SM5) && !IsConsolePlatform(Parameters.Platform) && !IsMetalMobilePlatform(Parameters.Platform);
 	}
 
 	static void ModifyCompilationEnvironment(const FGlobalShaderPermutationParameters& Parameters, FShaderCompilerEnvironment& OutEnvironment)
