Increased precision of alpha blending value used by landscape splines when using Landscape Layer system (from 8-bits to 14-bits). This solves stippled effect on spline side falloffs at higher locations.
#rnx

diff --git a/Engine/Shaders/Private/LandscapeLayersPS.usf b/Engine/Shaders/Private/LandscapeLayersPS.usf
--- a/Engine/Shaders/Private/LandscapeLayersPS.usf
+++ b/Engine/Shaders/Private/LandscapeLayersPS.usf
@@ -19,6 +19,13 @@ float ExtractHeight(float2 InPackedHeight)
 	return float(((int)round(InPackedHeight.r * 255.0) << 8) | (int)round(InPackedHeight.g * 255.0));
 }
 
+float ExtractAlpha(float2 InPackedAlpha)
+{
+	int AlphaWithFlags = ((int)round(InPackedAlpha.r * 255.0) << 8) | (int)round(InPackedAlpha.g * 255.0);
+	int Alpha = AlphaWithFlags & 0xFFFC; // Last 2 bits are used to store lower/raise flags
+	return float(Alpha) / 65532.0; // 0xFFFC in float
+}
+
 float2 PackHeight(float InExtractedHeight)
 {
 	InExtractedHeight = clamp(InExtractedHeight, 0.0f, 65535.0f);
@@ -71,7 +78,7 @@ void PSHeightmapMain(float2 InTextureCoordinates : TEXCOORD0, out float4 OutColo
 		if (LayerBlendMode == 1 && LayerVisibility == 1)
 		{
 			// Alpha Blend Mode
-			float HeightAlphaBlend = SourceColor.b;
+			float HeightAlphaBlend = ExtractAlpha(SourceColor.ba);
 			float NewHeight = Height + (HeightAlphaBlend * BaseHeight);
 			uint RaiseLower = floor(SourceColor.a * 255.f + 0.5f);
 			bool bLowerTerrain = (RaiseLower & 1);
diff --git a/Engine/Source/Runtime/Landscape/Private/LandscapeEditInterface.cpp b/Engine/Source/Runtime/Landscape/Private/LandscapeEditInterface.cpp
--- a/Engine/Source/Runtime/Landscape/Private/LandscapeEditInterface.cpp
+++ b/Engine/Source/Runtime/Landscape/Private/LandscapeEditInterface.cpp
@@ -218,7 +218,7 @@ bool FLandscapeEditDataInterface::GetComponentsInRegion(int32 X1, int32 Y1, int3
 	return bNotLocked;
 }
 
-void FLandscapeEditDataInterface::SetHeightData(int32 X1, int32 Y1, int32 X2, int32 Y2, const uint16* InData, int32 InStride, bool InCalcNormals, const uint16* InNormalData, const uint8* InHeightAlphaBlendData, const uint8* InHeightFlagsData, bool InCreateComponents, UTexture2D* InHeightmap, UTexture2D* InXYOffsetmapTexture,
+void FLandscapeEditDataInterface::SetHeightData(int32 X1, int32 Y1, int32 X2, int32 Y2, const uint16* InData, int32 InStride, bool InCalcNormals, const uint16* InNormalData, const uint16* InHeightAlphaBlendData, const uint8* InHeightFlagsData, bool InCreateComponents, UTexture2D* InHeightmap, UTexture2D* InXYOffsetmapTexture,
 											   bool InUpdateBounds, bool InUpdateCollision, bool InGenerateMips)
 {
 	const int32 NumVertsX = 1 + X2 - X1;
@@ -388,9 +388,10 @@ void FLandscapeEditDataInterface::SetHeightData(int32 X1, int32 Y1, int32 X2, in
 							}
 							else if (InHeightAlphaBlendData)
 							{
-								const uint8& HeightAlphaBlend = InHeightAlphaBlendData[DataIndex];
-								TexData.B = HeightAlphaBlend;
-								TexData.A = InHeightFlagsData ? InHeightFlagsData[DataIndex] : 3;
+								const uint16& HeightAlphaBlend = InHeightAlphaBlendData[DataIndex];
+								TexData.B = HeightAlphaBlend >> 8;
+								TexData.A = HeightAlphaBlend & 0xFC;
+								TexData.A |= (InHeightFlagsData ? InHeightFlagsData[DataIndex] : 3) & 0x3;
 							}
 						}
 					}
@@ -971,16 +972,16 @@ uint16 FLandscapeEditDataInterface::GetHeightMapData(const ULandscapeComponent*
 	return ((((uint16)TexData.R) << 8) | TexData.G);
 }
 
-uint8 FLandscapeEditDataInterface::GetHeightMapAlphaBlendData(const ULandscapeComponent* Component, int32 TexU, int32 TexV, FColor* TextureData /*= NULL*/)
+uint16 FLandscapeEditDataInterface::GetHeightMapAlphaBlendData(const ULandscapeComponent* Component, int32 TexU, int32 TexV, FColor* TextureData /*= NULL*/)
 {
 	const FColor& TexData = GetHeightMapColor(Component, TexU, TexV, TextureData);
-	return TexData.B;
+	return ((((uint16)TexData.B) << 8) | TexData.A) & 0xFFFC;
 }
 
 uint8 FLandscapeEditDataInterface::GetHeightMapFlagsData(const ULandscapeComponent* Component, int32 TexU, int32 TexV, FColor* TextureData /*= NULL*/)
 {
 	const FColor& TexData = GetHeightMapColor(Component, TexU, TexV, TextureData);
-	return TexData.A;
+	return TexData.A & 0x3;
 }
 
 template<typename TDataAccess, typename TGetHeightMapDataFunction>
@@ -1486,7 +1487,7 @@ void FLandscapeEditDataInterface::GetHeightDataTempl(int32& ValidX1, int32& Vali
 template<typename TDataAccess>
 void FLandscapeEditDataInterface::GetHeightAlphaBlendDataTempl(int32& ValidX1, int32& ValidY1, int32& ValidX2, int32& ValidY2, TDataAccess& StoreData)
 {
-	GetHeightDataInternal(ValidX1, ValidY1, ValidX2, ValidY2, StoreData, [&](const ULandscapeComponent* Component, int32 TexU, int32 TexV, FColor* TextureData)->uint8 { return GetHeightMapAlphaBlendData(Component, TexU, TexV, TextureData); });
+	GetHeightDataInternal(ValidX1, ValidY1, ValidX2, ValidY2, StoreData, [&](const ULandscapeComponent* Component, int32 TexU, int32 TexV, FColor* TextureData)->uint16 { return GetHeightMapAlphaBlendData(Component, TexU, TexV, TextureData); });
 }
 
 template<typename TDataAccess>
@@ -1506,14 +1507,14 @@ void FLandscapeEditDataInterface::GetHeightData(int32& X1, int32& Y1, int32& X2,
 	GetHeightDataTempl(X1, Y1, X2, Y2, ArrayStoreData);
 }
 
-void FLandscapeEditDataInterface::GetHeightAlphaBlendData(int32& X1, int32& Y1, int32& X2, int32& Y2, uint8* Data, int32 Stride)
+void FLandscapeEditDataInterface::GetHeightAlphaBlendData(int32& X1, int32& Y1, int32& X2, int32& Y2, uint16* Data, int32 Stride)
 {
 	if (Stride == 0)
 	{
 		Stride = (1 + X2 - X1);
 	}
 
-	TArrayDataAccess<uint8> ArrayStoreData(X1, Y1, Data, Stride);
+	TArrayDataAccess<uint16> ArrayStoreData(X1, Y1, Data, Stride);
 	GetHeightAlphaBlendDataTempl(X1, Y1, X2, Y2, ArrayStoreData);
 }
 
diff --git a/Engine/Source/Runtime/Landscape/Private/LandscapeEditLayers.cpp b/Engine/Source/Runtime/Landscape/Private/LandscapeEditLayers.cpp
--- a/Engine/Source/Runtime/Landscape/Private/LandscapeEditLayers.cpp
+++ b/Engine/Source/Runtime/Landscape/Private/LandscapeEditLayers.cpp
@@ -5131,15 +5131,14 @@ void ALandscape::ClearLayer(const FGuid& InLayerGuid, TSet<ULandscapeComponent*>
 		NewHeightDataValue = ZeroValue;
 	}
 
-	TArray<uint8> NewHeightAlphaBlendData;
+	TArray<uint16> NewHeightAlphaBlendData;
 	TArray<uint8> NewHeightFlagsData;
 
 	if (InClearMode & ELandscapeClearMode::Clear_Heightmap)
 	{
 		if (Layer->BlendMode == LSBM_AlphaBlend)
 		{
-			NewHeightAlphaBlendData.AddZeroed(FMath::Square(ComponentSizeQuads + 1));
-			FMemory::Memset(NewHeightAlphaBlendData.GetData(), 255, NewHeightAlphaBlendData.Num());
+			NewHeightAlphaBlendData.Init(MAX_uint16, FMath::Square(ComponentSizeQuads + 1));
 			NewHeightFlagsData.AddZeroed(FMath::Square(ComponentSizeQuads + 1));
 		}
 	}
diff --git a/Engine/Source/Runtime/Landscape/Private/LandscapeSplineRaster.cpp b/Engine/Source/Runtime/Landscape/Private/LandscapeSplineRaster.cpp
--- a/Engine/Source/Runtime/Landscape/Private/LandscapeSplineRaster.cpp
+++ b/Engine/Source/Runtime/Landscape/Private/LandscapeSplineRaster.cpp
@@ -36,7 +36,7 @@ public:
 	typedef FVector InterpolantType;
 
 	/** Initialization constructor. */
-	FLandscapeSplineHeightsRasterPolicy(TArray<uint16>& InHeightData, int32 InMinX, int32 InMinY, int32 InMaxX, int32 InMaxY, bool InbRaiseTerrain, bool InbLowerTerrain, TArray<uint8>* InHeightAlphaBlendData = nullptr, TArray<uint8>* InHeightFlagsData = nullptr) :
+	FLandscapeSplineHeightsRasterPolicy(TArray<uint16>& InHeightData, int32 InMinX, int32 InMinY, int32 InMaxX, int32 InMaxY, bool InbRaiseTerrain, bool InbLowerTerrain, TArray<uint16>* InHeightAlphaBlendData = nullptr, TArray<uint8>* InHeightFlagsData = nullptr) :
 		HeightData(InHeightData),
 		HeightAlphaBlendData(InHeightAlphaBlendData),
 		HeightFlagsData(InHeightFlagsData),
@@ -77,9 +77,9 @@ protected:
 			float InterpValue = (NewHeight * Alpha) + (Dest * (1.f - Alpha));
 			Dest = (uint16)FMath::Clamp<float>(InterpValue, 0, (float)LandscapeDataAccess::MaxValue);
 
-			uint8& DestAlphaValue = (*HeightAlphaBlendData)[DataIndex];
+			uint16& DestAlphaValue = (*HeightAlphaBlendData)[DataIndex];
 			float InterpAlphaValue = DestAlphaValue * (1.f - Alpha);
-			DestAlphaValue = (uint8)FMath::Clamp<float>(InterpAlphaValue, 0.f, 255.f);
+			DestAlphaValue = (uint16)FMath::Clamp<float>(InterpAlphaValue, 0.f, 65535.f);
 
 			if (HeightFlagsData)
 			{
@@ -106,7 +106,7 @@ protected:
 
 private:
 	TArray<uint16>& HeightData;
-	TArray<uint8>* HeightAlphaBlendData;
+	TArray<uint16>* HeightAlphaBlendData;
 	TArray<uint8>* HeightFlagsData;
 	int32 MinX, MinY, MaxX, MaxY;
 	uint32 bRaiseTerrain : 1, bLowerTerrain : 1;
@@ -190,9 +190,9 @@ void RasterizeHeight(int32& MinX, int32& MinY, int32& MaxX, int32& MaxY, FLandsc
 	bool bIsEditingLayerReservedForSplines = Landscape && Landscape->IsEditingLayerReservedForSplines();
 	check(!bIsEditingLayerReservedForSplines || (Landscape->GetLandscapeSplinesReservedLayer()->BlendMode == LSBM_AlphaBlend));
 
-	TArray<uint8> HeightAlphaBlendData;
+	TArray<uint16> HeightAlphaBlendData;
 	TArray<uint8> HeightFlagsData;
-	TArray<uint8>* HeightAlphaBlendDataPtr = nullptr;
+	TArray<uint16>* HeightAlphaBlendDataPtr = nullptr;
 	TArray<uint8>* HeightFlagsDataPtr = nullptr;
 
 	bool bCalculateNormals = true;
diff --git a/Engine/Source/Runtime/Landscape/Public/LandscapeEdit.h b/Engine/Source/Runtime/Landscape/Public/LandscapeEdit.h
--- a/Engine/Source/Runtime/Landscape/Public/LandscapeEdit.h
+++ b/Engine/Source/Runtime/Landscape/Public/LandscapeEdit.h
@@ -144,13 +144,13 @@ struct LANDSCAPE_API FLandscapeEditDataInterface : public FLandscapeTextureDataI
 	//
 	// Heightmap access
 	//
-	void SetHeightData(int32 X1, int32 Y1, int32 X2, int32 Y2, const uint16* InData, int32 InStride, bool InCalcNormals, const uint16* InNormalData = nullptr, const uint8* InHeightAlphaBlendData = nullptr, const uint8* InHeightRaiseLowerData = nullptr, bool InCreateComponents = false, UTexture2D* InHeightmap = nullptr, UTexture2D* InXYOffsetmapTexture = nullptr,
+	void SetHeightData(int32 X1, int32 Y1, int32 X2, int32 Y2, const uint16* InData, int32 InStride, bool InCalcNormals, const uint16* InNormalData = nullptr, const uint16* InHeightAlphaBlendData = nullptr, const uint8* InHeightRaiseLowerData = nullptr, bool InCreateComponents = false, UTexture2D* InHeightmap = nullptr, UTexture2D* InXYOffsetmapTexture = nullptr,
 					   bool InUpdateBounds = true, bool InUpdateCollision = true, bool InGenerateMips = true);
 
 	// Helper accessor
 	FORCEINLINE uint16 GetHeightMapData(const ULandscapeComponent* Component, int32 TexU, int32 TexV, FColor* TextureData = NULL);
 	// Helper accessor
-	FORCEINLINE uint8 GetHeightMapAlphaBlendData(const ULandscapeComponent* Component, int32 TexU, int32 TexV, FColor* TextureData = NULL);
+	FORCEINLINE uint16 GetHeightMapAlphaBlendData(const ULandscapeComponent* Component, int32 TexU, int32 TexV, FColor* TextureData = NULL);
 	// Helper accessor
 	FORCEINLINE uint8 GetHeightMapFlagsData(const ULandscapeComponent* Component, int32 TexU, int32 TexV, FColor* TextureData = NULL);
 	// Generic
@@ -166,7 +166,7 @@ struct LANDSCAPE_API FLandscapeEditDataInterface : public FLandscapeTextureDataI
 	// Implementation for fixed array
 	void GetHeightData(int32& X1, int32& Y1, int32& X2, int32& Y2, uint16* Data, int32 Stride);
 	void GetHeightDataFast(const int32 X1, const int32 Y1, const int32 X2, const int32 Y2, uint16* Data, int32 Stride, uint16* NormalData = NULL, UTexture2D* InHeightmap = nullptr);
-	void GetHeightAlphaBlendData(int32& X1, int32& Y1, int32& X2, int32& Y2, uint8* Data, int32 Stride);
+	void GetHeightAlphaBlendData(int32& X1, int32& Y1, int32& X2, int32& Y2, uint16* Data, int32 Stride);
 	void GetHeightFlagsData(int32& X1, int32& Y1, int32& X2, int32& Y2, uint8* Data, int32 Stride);
 	// Implementation for sparse array
 	void GetHeightData(int32& X1, int32& Y1, int32& X2, int32& Y2, TMap<FIntPoint, uint16>& SparseData);
