Fixed regression in performance when painting (caused by CL 7150452). Avoid recreating collision when painting.
#rnx

diff --git a/Engine/Source/Runtime/Landscape/Private/LandscapeEditLayers.cpp b/Engine/Source/Runtime/Landscape/Private/LandscapeEditLayers.cpp
--- a/Engine/Source/Runtime/Landscape/Private/LandscapeEditLayers.cpp
+++ b/Engine/Source/Runtime/Landscape/Private/LandscapeEditLayers.cpp
@@ -4123,10 +4123,12 @@ void ALandscape::RequestLayersContentUpdateForceAll(ELandscapeLayerUpdateMode In
 	const bool bUpdateHeightmap = (InModeMask & (ELandscapeLayerUpdateMode::Update_Heightmap_All | ELandscapeLayerUpdateMode::Update_Heightmap_Editing | ELandscapeLayerUpdateMode::Update_Heightmap_Editing_NoCollision)) != 0;
 	const bool bUpdateWeightCollision = (InModeMask & (ELandscapeLayerUpdateMode::Update_Weightmap_All | ELandscapeLayerUpdateMode::Update_Weightmap_Editing)) != 0;
 	const bool bUpdateHeightCollision = (InModeMask & (ELandscapeLayerUpdateMode::Update_Heightmap_All | ELandscapeLayerUpdateMode::Update_Heightmap_Editing)) != 0;
+	const bool bUpdateAllHeightmap = (InModeMask & ELandscapeLayerUpdateMode::Update_Heightmap_All) != 0;
+	const bool bUpdateAllWeightmap = (InModeMask & ELandscapeLayerUpdateMode::Update_Weightmap_All) != 0;
 	const bool bUpdateClientUdpateEditing = (InModeMask & ELandscapeLayerUpdateMode::Update_Client_Editing) != 0;
 	if (ULandscapeInfo* LandscapeInfo = GetLandscapeInfo())
 	{
-		LandscapeInfo->ForAllLandscapeProxies([bUpdateHeightmap, bUpdateWeightmap, bUpdateHeightCollision, bUpdateWeightCollision, bUpdateClientUdpateEditing](ALandscapeProxy* Proxy)
+		LandscapeInfo->ForAllLandscapeProxies([bUpdateHeightmap, bUpdateWeightmap, bUpdateAllHeightmap, bUpdateAllWeightmap, bUpdateHeightCollision, bUpdateWeightCollision, bUpdateClientUdpateEditing](ALandscapeProxy* Proxy)
 		{
 			if (Proxy)
 			{
@@ -4144,14 +4146,12 @@ void ALandscape::RequestLayersContentUpdateForceAll(ELandscapeLayerUpdateMode In
 
 					if (bUpdateHeightmap)
 					{
-						const bool bUpdateAll = false;
-						Component->RequestHeightmapUpdate(bUpdateAll, bUpdateHeightCollision);
+						Component->RequestHeightmapUpdate(bUpdateAllHeightmap, bUpdateHeightCollision);
 					}
 
 					if (bUpdateWeightmap)
 					{
-						const bool bUpdateAll = false;
-						Component->RequestWeightmapUpdate(bUpdateAll, bUpdateWeightCollision);
+						Component->RequestWeightmapUpdate(bUpdateAllWeightmap, bUpdateWeightCollision);
 					}
 
 					if (bUpdateClientUdpateEditing)
@@ -4269,7 +4269,7 @@ void ULandscapeComponent::RequestHeightmapUpdate(bool bUpdateAll, bool bUpdateCo
 	if (ALandscape* LandscapeActor = GetLandscapeActor())
 	{
 		LandscapeActor->RequestLayersContentUpdate(bUpdateCollision ? ELandscapeLayerUpdateMode::Update_Heightmap_Editing : ELandscapeLayerUpdateMode::Update_Heightmap_Editing_NoCollision);
-		if (bUpdateAll || bUpdateCollision)
+		if (bUpdateAll)
 		{
 			LandscapeActor->RequestLayersContentUpdate(ELandscapeLayerUpdateMode::Update_Heightmap_All);
 		}
@@ -4290,7 +4290,7 @@ void ULandscapeComponent::RequestWeightmapUpdate(bool bUpdateAll, bool bUpdateCo
 	if (ALandscape* LandscapeActor = GetLandscapeActor())
 	{
 		LandscapeActor->RequestLayersContentUpdate(bUpdateCollision ? ELandscapeLayerUpdateMode::Update_Weightmap_Editing : ELandscapeLayerUpdateMode::Update_Weightmap_Editing_NoCollision);
-		if (bUpdateAll || bUpdateCollision)
+		if (bUpdateAll)
 		{
 			LandscapeActor->RequestLayersContentUpdate(ELandscapeLayerUpdateMode::Update_Weightmap_All);
 		}
