Optimized update of collision Heightfield region when painting.
Greatly improves performance when painting in landscape with big compoments.
#rnx

diff --git a/Engine/Source/Runtime/Landscape/Classes/LandscapeComponent.h b/Engine/Source/Runtime/Landscape/Classes/LandscapeComponent.h
--- a/Engine/Source/Runtime/Landscape/Classes/LandscapeComponent.h
+++ b/Engine/Source/Runtime/Landscape/Classes/LandscapeComponent.h
@@ -437,6 +437,9 @@ private:
 
 	UPROPERTY(Transient)
 	uint32 LayerUpdateFlagPerMode;
+
+	/** Dirtied collision height region when painting (only used by Landscape Layer System) */
+	FIntRect LayerDirtyCollisionHeightData;
 #endif // WITH_EDITORONLY_DATA
 
 	/** Heightmap texture reference */
@@ -844,6 +847,12 @@ public:
 	 */
 	void UpdateCollisionData(bool bInUpdateHeightfieldRegion = true);
 
+	/** Cumulates component's dirtied collision region that will need to be updated (used by Layer System)*/
+	void UpdateDirtyCollisionHeightData(FIntRect Region);
+
+	/** Clears component's dirtied collision region (used by Layer System)*/
+	void ClearDirtyCollisionHeightData();
+
 	/**
 	 * Update collision component dominant layer data
 	 * @param WeightmapTextureMipData: weightmap data
diff --git a/Engine/Source/Runtime/Landscape/Private/LandscapeEdit.cpp b/Engine/Source/Runtime/Landscape/Private/LandscapeEdit.cpp
--- a/Engine/Source/Runtime/Landscape/Private/LandscapeEdit.cpp
+++ b/Engine/Source/Runtime/Landscape/Private/LandscapeEdit.cpp
@@ -824,6 +824,26 @@ public:
 	int32 SizeVertsSquare = 0;
 };
 
+void ULandscapeComponent::UpdateDirtyCollisionHeightData(FIntRect Region)
+{
+	// Take first value as is
+	if (LayerDirtyCollisionHeightData.IsEmpty())
+	{
+		LayerDirtyCollisionHeightData = Region;
+	}
+	else
+	{
+		// Merge min/max region
+		LayerDirtyCollisionHeightData.Include(Region.Min);
+		LayerDirtyCollisionHeightData.Include(Region.Max);
+	}
+}
+
+void ULandscapeComponent::ClearDirtyCollisionHeightData()
+{
+	LayerDirtyCollisionHeightData = FIntRect();
+}
+
 void ULandscapeComponent::UpdateCollisionHeightData(const FColor* const HeightmapTextureMipData, const FColor* const SimpleCollisionHeightmapTextureData, int32 ComponentX1/*=0*/, int32 ComponentY1/*=0*/, int32 ComponentX2/*=MAX_int32*/, int32 ComponentY2/*=MAX_int32*/, bool bUpdateBounds/*=false*/, const FColor* XYOffsetTextureMipData/*=nullptr*/, bool bInUpdateHeightfieldRegion/*=true*/)
 {
 	ULandscapeInfo* Info = GetLandscapeInfo();
@@ -854,6 +874,18 @@ void ULandscapeComponent::UpdateCollisionHeightData(const FColor* const Heightma
 			CollisionComp->Modify();
 		}
 	}
+	else
+	{
+		// In Landscape Layers, only update dirtied collision height data
+		if (bInUpdateHeightfieldRegion && ComponentX1 == 0 && ComponentY1 == 0 && ComponentX2 == MAX_int32 && ComponentY2 == MAX_int32 && !LayerDirtyCollisionHeightData.IsEmpty())
+		{
+			ComponentX1 = LayerDirtyCollisionHeightData.Min.X;
+			ComponentY1 = LayerDirtyCollisionHeightData.Min.Y;
+			ComponentX2 = LayerDirtyCollisionHeightData.Max.X;
+			ComponentY2 = LayerDirtyCollisionHeightData.Max.Y;
+		}
+		ClearDirtyCollisionHeightData();
+	}
 
 	// Existing collision component is same type with collision
 	if (CollisionComp && ((XYOffsetmapTexture == nullptr) == (MeshCollisionComponent == nullptr)))
diff --git a/Engine/Source/Runtime/Landscape/Private/LandscapeEditInterface.cpp b/Engine/Source/Runtime/Landscape/Private/LandscapeEditInterface.cpp
--- a/Engine/Source/Runtime/Landscape/Private/LandscapeEditInterface.cpp
+++ b/Engine/Source/Runtime/Landscape/Private/LandscapeEditInterface.cpp
@@ -461,6 +461,11 @@ void FLandscapeEditDataInterface::SetHeightData(int32 X1, int32 Y1, int32 X2, in
 					}
 				}
 			}
+			else
+			{
+				// In Layer, cumulate dirty collision region (will be used next time UpdateCollisionHeightData is called)
+				Component->UpdateDirtyCollisionHeightData(FIntRect(ComponentX1, ComponentY1, ComponentX2, ComponentY2));
+			}
 			
 			// Update GUID for Platform Data
 			FPlatformMisc::CreateGuid(Component->StateId);
