Dont create raytracing vertex buffers when raytracing is not enabled.Saves a lot of memory on scenes that have large landscapes

diff --git a/Engine/Source/Runtime/Landscape/Private/LandscapeRender.cpp b/Engine/Source/Runtime/Landscape/Private/LandscapeRender.cpp
--- a/Engine/Source/Runtime/Landscape/Private/LandscapeRender.cpp
+++ b/Engine/Source/Runtime/Landscape/Private/LandscapeRender.cpp
@@ -1036,29 +1036,32 @@ void FLandscapeComponentSceneProxy::CreateRenderThreadResources()
 #endif
 
 #if RHI_RAYTRACING
-	for (int32 SubY = 0; SubY < NumSubsections; SubY++)
+	if (IsRayTracingEnabled())
 	{
-		for (int32 SubX = 0; SubX < NumSubsections; SubX++)
+		for (int32 SubY = 0; SubY < NumSubsections; SubY++)
 		{
-			const int8 SubSectionIdx = SubX + SubY * NumSubsections;
+			for (int32 SubX = 0; SubX < NumSubsections; SubX++)
+			{
+				const int8 SubSectionIdx = SubX + SubY * NumSubsections;
 
-			int32 LodSubsectionSizeVerts = (SubsectionSizeVerts >> 0);
-			uint32 NumPrimitives = FMath::Square((LodSubsectionSizeVerts - 1)) * 2;
-
-			FRayTracingGeometryInitializer Initializer;
-			FRHIResourceCreateInfo CreateInfo;
-			Initializer.PositionVertexBuffer = RHICreateVertexBuffer(sizeof(FVector4) * NumPrimitives * 3, BUF_UnorderedAccess | BUF_ShaderResource, CreateInfo);
-			Initializer.IndexBuffer = nullptr;
-			Initializer.BaseVertexIndex = 0;
-			Initializer.VertexBufferStride = sizeof(FVector);
-			Initializer.TotalPrimitiveCount = NumPrimitives;
-			Initializer.VertexBufferElementType = VET_Float3;
-			Initializer.GeometryType = RTGT_Triangles;
-			Initializer.bFastBuild = true;
-			Initializer.bAllowUpdate = true;
-
-			SectionRayTracingStates[SubSectionIdx].Geometry.SetInitializer(Initializer);
-			SectionRayTracingStates[SubSectionIdx].Geometry.InitResource();
+				int32 LodSubsectionSizeVerts = (SubsectionSizeVerts >> 0);
+				uint32 NumPrimitives = FMath::Square((LodSubsectionSizeVerts - 1)) * 2;
+
+				FRayTracingGeometryInitializer Initializer;
+				FRHIResourceCreateInfo CreateInfo;
+				Initializer.PositionVertexBuffer = RHICreateVertexBuffer(sizeof(FVector4) * NumPrimitives * 3, BUF_UnorderedAccess | BUF_ShaderResource, CreateInfo);
+				Initializer.IndexBuffer = nullptr;
+				Initializer.BaseVertexIndex = 0;
+				Initializer.VertexBufferStride = sizeof(FVector);
+				Initializer.TotalPrimitiveCount = NumPrimitives;
+				Initializer.VertexBufferElementType = VET_Float3;
+				Initializer.GeometryType = RTGT_Triangles;
+				Initializer.bFastBuild = true;
+				Initializer.bAllowUpdate = true;
+
+				SectionRayTracingStates[SubSectionIdx].Geometry.SetInitializer(Initializer);
+				SectionRayTracingStates[SubSectionIdx].Geometry.InitResource();
+			}
 		}
 	}
 #endif
