Landscape Layer System: Multiple Fix
- Fix Brush jagged edges
- Fix some components sometimes not being paintable
- Fix Crashes in some setups (writing out of texture bounds)
- Fix some Heightmap compnents not showing

#rnx

diff --git a/Engine/Shaders/Private/LandscapeLayersCS.usf b/Engine/Shaders/Private/LandscapeLayersCS.usf
--- a/Engine/Shaders/Private/LandscapeLayersCS.usf
+++ b/Engine/Shaders/Private/LandscapeLayersCS.usf
@@ -77,7 +77,7 @@ void PackPaintLayerToWeightmap(uint3 DispatchThreadId : SV_DispatchThreadID)
 	{
 		FLandscapeLayerWeightmapPackMaterialLayersComponentData ComponentData = InPackLayersComponentsData[DispatchThreadId.z];
 
-		float4 Channels;
+		float4 Channels = 0;
 
 		for (int i = 0; i < 4; ++i)
 		{
diff --git a/Engine/Source/Runtime/Landscape/Classes/Landscape.h b/Engine/Source/Runtime/Landscape/Classes/Landscape.h
--- a/Engine/Source/Runtime/Landscape/Classes/Landscape.h
+++ b/Engine/Source/Runtime/Landscape/Classes/Landscape.h
@@ -365,11 +365,11 @@ private:
 	void DrawHeightmapComponentsToRenderTarget(const FString& InDebugName, const TArray<ULandscapeComponent*>& InComponentsToDraw, const FIntPoint& InLandscapeBase, UTexture* InHeightmapRTRead, UTextureRenderTarget2D* InOptionalHeightmapRTRead2, UTextureRenderTarget2D* InHeightmapRTWrite, ERTDrawingType InDrawType,
 											   bool InClearRTWrite, struct FLandscapeLayersHeightmapShaderParameters& InShaderParams, uint8 InMipRender = 0) const;
 
-	void DrawWeightmapComponentsToRenderTarget(const FString& InDebugName, const TArray<FIntPoint>& InSectionBaseList, const FVector2D& InScaleBias, UTexture* InWeightmapRTRead, UTextureRenderTarget2D* InOptionalWeightmapRTRead2, UTextureRenderTarget2D* InWeightmapRTWrite,
-											   bool InClearRTWrite, struct FLandscapeLayersWeightmapShaderParameters& InShaderParams, uint8 InMipRender) const;
+	void DrawWeightmapComponentsToRenderTarget(const FString& InDebugName, const TArray<FIntPoint>& InSectionBaseList, const FVector2D& InScaleBias, TArray<FVector2D>* InScaleBiasPerSection, UTexture* InWeightmapRTRead, UTextureRenderTarget2D* InOptionalWeightmapRTRead2, UTextureRenderTarget2D* InWeightmapRTWrite, ERTDrawingType InDrawType,
+												bool InClearRTWrite, struct FLandscapeLayersWeightmapShaderParameters& InShaderParams, uint8 InMipRender) const;
 
-	void DrawWeightmapComponentsToRenderTarget(const FString& InDebugName, const TArray<ULandscapeComponent*>& InComponentsToDraw, const FIntPoint& InLandscapeBase, UTexture* InWeightmapRTRead, UTextureRenderTarget2D* InOptionalWeightmapRTRead2, UTextureRenderTarget2D* InWeightmapRTWrite,
-											   bool InClearRTWrite, struct FLandscapeLayersWeightmapShaderParameters& InShaderParams, uint8 InMipRender) const;
+	void DrawWeightmapComponentsToRenderTarget(const FString& InDebugName, const TArray<ULandscapeComponent*>& InComponentsToDraw, const FIntPoint& InLandscapeBase, UTexture* InWeightmapRTRead, UTextureRenderTarget2D* InOptionalWeightmapRTRead2, UTextureRenderTarget2D* InWeightmapRTWrite, ERTDrawingType InDrawType,
+												bool InClearRTWrite, struct FLandscapeLayersWeightmapShaderParameters& InShaderParams, uint8 InMipRender) const;
 
 	void DrawHeightmapComponentsToRenderTargetMips(const TArray<ULandscapeComponent*>& InComponentsToDraw, const FIntPoint& InLandscapeBase, UTexture* InReadHeightmap, bool InClearRTWrite, struct FLandscapeLayersHeightmapShaderParameters& InShaderParams) const;
 	void DrawWeightmapComponentToRenderTargetMips(const TArray<FVector2D>& InTexturePositionsToDraw, UTexture* InReadWeightmap, bool InClearRTWrite, struct FLandscapeLayersWeightmapShaderParameters& InShaderParams) const;
diff --git a/Engine/Source/Runtime/Landscape/Private/LandscapeEditLayers.cpp b/Engine/Source/Runtime/Landscape/Private/LandscapeEditLayers.cpp
--- a/Engine/Source/Runtime/Landscape/Private/LandscapeEditLayers.cpp
+++ b/Engine/Source/Runtime/Landscape/Private/LandscapeEditLayers.cpp
@@ -1005,9 +1005,6 @@ public:
 		FIntPoint SourceSize(Params.SourceResource->GetSizeX(), Params.SourceResource->GetSizeY()); // SourceResource is always proper size, as it's always the good MIP we want to copy from
 		FIntPoint DestSize(Params.DestResource->GetSizeX() >> Params.DestMip, Params.DestResource->GetSizeY() >> Params.DestMip);
 
-		int32 LocalComponentSizeQuad = Params.SubSectionSizeQuad * Params.NumSubSections;
-		FVector2D PositionOffset(FMath::RoundToInt(Params.InitialPositionOffset.X / LocalComponentSizeQuad), FMath::RoundToInt(Params.InitialPositionOffset.Y / LocalComponentSizeQuad));
-
 		FRHICopyTextureInfo Info;
 		Info.NumSlices = 1;
 		Info.Size.Z = 1;
@@ -1020,11 +1017,12 @@ public:
 		{
 			Info.SourcePosition.X = 0;
 			Info.Size.X = SourceSize.X;
-			Info.DestPosition.X = FMath::RoundToInt(PositionOffset.X * (((Params.SubSectionSizeQuad + 1) * Params.NumSubSections) >> Params.DestMip));
+			Info.DestPosition.X = Params.InitialPositionOffset.X >> Params.DestMip;
+
 		}
 		else
 		{
-			Info.SourcePosition.X = FMath::RoundToInt(PositionOffset.X * (((Params.SubSectionSizeQuad + 1) *Params.NumSubSections) >> Params.SourceMip));
+			Info.SourcePosition.X = Params.InitialPositionOffset.X >> Params.SourceMip;
 			Info.Size.X = DestSize.X;
 			Info.DestPosition.X = 0;
 		}
@@ -1033,11 +1031,11 @@ public:
 		{
 			Info.SourcePosition.Y = 0;
 			Info.Size.Y = SourceSize.Y;
-			Info.DestPosition.Y = FMath::RoundToInt(PositionOffset.Y * (((Params.SubSectionSizeQuad + 1) * Params.NumSubSections) >> Params.DestMip));
+			Info.DestPosition.Y = Params.InitialPositionOffset.Y >> Params.DestMip;
 		}
 		else
 		{
-			Info.SourcePosition.Y = FMath::RoundToInt(PositionOffset.Y * (((Params.SubSectionSizeQuad + 1) * Params.NumSubSections) >> Params.SourceMip));
+			Info.SourcePosition.Y = Params.InitialPositionOffset.Y >> Params.SourceMip;
 			Info.Size.Y = DestSize.Y;
 			Info.DestPosition.Y = 0;
 		}
@@ -1267,8 +1265,6 @@ void ALandscape::CreateLayersRenderingResource()
 	FlushRenderingCommands();
 
 	const FIntPoint ComponentCounts = ComputeComponentCounts();
-	const int32 TotalVertexCountX = (SubsectionSizeQuads * NumSubsections) * ComponentCounts.X + 1;
-	const int32 TotalVertexCountY = (SubsectionSizeQuads * NumSubsections) * ComponentCounts.Y + 1;
 
 	ALandscape* Landscape = GetLandscapeActor();
 	check(Landscape);
@@ -1291,7 +1287,7 @@ void ALandscape::CreateLayersRenderingResource()
 
 			if (i < EHeightmapRTType::HeightmapRT_Mip1) // Landscape size RT
 			{
-				Landscape->HeightmapRTList[i]->InitAutoFormat(FMath::RoundUpToPowerOfTwo(TotalVertexCountX), FMath::RoundUpToPowerOfTwo(TotalVertexCountY));
+				Landscape->HeightmapRTList[i]->InitAutoFormat(FMath::RoundUpToPowerOfTwo(CurrentMipSizeX), FMath::RoundUpToPowerOfTwo(CurrentMipSizeY));
 			}
 			else // Mips
 			{
@@ -1318,7 +1314,7 @@ void ALandscape::CreateLayersRenderingResource()
 		{
 			if (i < EHeightmapRTType::HeightmapRT_Mip1) // Landscape size RT
 			{
-				Landscape->HeightmapRTList[i]->ResizeTarget(FMath::RoundUpToPowerOfTwo(TotalVertexCountX), FMath::RoundUpToPowerOfTwo(TotalVertexCountY));
+				Landscape->HeightmapRTList[i]->ResizeTarget(FMath::RoundUpToPowerOfTwo(CurrentMipSizeX), FMath::RoundUpToPowerOfTwo(CurrentMipSizeY));
 			}
 			else // Mips
 			{
@@ -1355,7 +1351,7 @@ void ALandscape::CreateLayersRenderingResource()
 			if (i < EWeightmapRTType::WeightmapRT_Mip0) // Landscape size RT, only create the number of layer we have
 			{
 				Landscape->WeightmapRTList[i]->RenderTargetFormat = i == WeightmapRT_Scratch_RGBA ? RTF_RGBA8 : RTF_R8;
-				Landscape->WeightmapRTList[i]->InitAutoFormat(FMath::RoundUpToPowerOfTwo(TotalVertexCountX), FMath::RoundUpToPowerOfTwo(TotalVertexCountY));
+				Landscape->WeightmapRTList[i]->InitAutoFormat(FMath::RoundUpToPowerOfTwo(CurrentMipSizeX), FMath::RoundUpToPowerOfTwo(CurrentMipSizeY));
 			}
 			else // Mips
 			{
@@ -1368,7 +1364,7 @@ void ALandscape::CreateLayersRenderingResource()
 			Landscape->WeightmapRTList[i]->UpdateResourceImmediate(true);
 
 			// Only generate required mips RT
-			if (CurrentMipSizeX == ComponentCounts.X && CurrentMipSizeY == ComponentCounts.Y)
+			if (CurrentMipSizeX < ComponentCounts.X && CurrentMipSizeY < ComponentCounts.Y)
 			{
 				break;
 			}
@@ -1383,7 +1379,7 @@ void ALandscape::CreateLayersRenderingResource()
 		{
 			if (i < EWeightmapRTType::WeightmapRT_Mip0) // Landscape size RT, only create the number of layer we have
 			{
-				Landscape->WeightmapRTList[i]->ResizeTarget(FMath::RoundUpToPowerOfTwo(TotalVertexCountX), FMath::RoundUpToPowerOfTwo(TotalVertexCountY));
+				Landscape->WeightmapRTList[i]->ResizeTarget(FMath::RoundUpToPowerOfTwo(CurrentMipSizeX), FMath::RoundUpToPowerOfTwo(CurrentMipSizeY));
 			}
 			else // Mips
 			{
@@ -1394,7 +1390,7 @@ void ALandscape::CreateLayersRenderingResource()
 			}
 
 			// Only generate required mips RT
-			if (CurrentMipSizeX == ComponentCounts.X && CurrentMipSizeY == ComponentCounts.Y)
+			if (CurrentMipSizeX < ComponentCounts.X && CurrentMipSizeY < ComponentCounts.Y)
 			{
 				break;
 			}
@@ -1823,11 +1819,12 @@ void ALandscape::CopyLayersTexture(const FString& InSourceDebugName, FTextureRes
 	});
 }
 
-void ALandscape::DrawWeightmapComponentsToRenderTarget(const FString& InDebugName, const TArray<FIntPoint>& InSectionBaseList, const FVector2D& InScaleBias, UTexture* InWeightmapRTRead, UTextureRenderTarget2D* InOptionalWeightmapRTRead2, UTextureRenderTarget2D* InWeightmapRTWrite,
-														bool InClearRTWrite, FLandscapeLayersWeightmapShaderParameters& InShaderParams, uint8 InMipRender) const
+void ALandscape::DrawWeightmapComponentsToRenderTarget(const FString& InDebugName, const TArray<FIntPoint>& InSectionBaseList, const FVector2D& InScaleBias, TArray<FVector2D>* InScaleBiasPerSection, UTexture* InWeightmapRTRead, UTextureRenderTarget2D* InOptionalWeightmapRTRead2, UTextureRenderTarget2D* InWeightmapRTWrite,
+														ERTDrawingType InDrawType, bool InClearRTWrite, FLandscapeLayersWeightmapShaderParameters& InShaderParams, uint8 InMipRender) const
 {
 	check(InWeightmapRTRead != nullptr);
 	check(InWeightmapRTWrite != nullptr);
+	check(InScaleBiasPerSection == nullptr || InScaleBiasPerSection->Num() == InSectionBaseList.Num());
 
 	FIntPoint WeightmapWriteTextureSize(InWeightmapRTWrite->SizeX, InWeightmapRTWrite->SizeY);
 	FIntPoint WeightmapReadTextureSize(InWeightmapRTRead->Source.GetSizeX(), InWeightmapRTRead->Source.GetSizeY());
@@ -1843,18 +1840,41 @@ void ALandscape::DrawWeightmapComponentsToRenderTarget(const FString& InDebugNam
 	TArray<FLandscapeLayersTriangle> TriangleList;
 	TriangleList.Reserve(InSectionBaseList.Num() * 2 * NumSubsections);
 
-	if (InMipRender == 0)
+	for (int i = 0; i < InSectionBaseList.Num(); ++i)
 	{
-		for (const FIntPoint& SectionBase : InSectionBaseList)
+		const FVector2D& WeightmapScaleBias = InScaleBiasPerSection != nullptr ? (*InScaleBiasPerSection)[i] : InScaleBias;
+		switch (InDrawType)
 		{
-			GenerateLayersRenderQuadsAtlas(SectionBase, InScaleBias, SubsectionSizeQuads, WeightmapReadTextureSize, WeightmapWriteTextureSize, TriangleList);
-		}
-	}
-	else
-	{
-		for (const FIntPoint& SectionBase : InSectionBaseList)
+		case ERTDrawingType::RTAtlas:
+		{
+			GenerateLayersRenderQuadsAtlas(InSectionBaseList[i], WeightmapScaleBias, SubsectionSizeQuads, WeightmapReadTextureSize, WeightmapWriteTextureSize, TriangleList);
+		} break;
+
+		case ERTDrawingType::RTAtlasToNonAtlas:
+		{
+			GenerateLayersRenderQuadsAtlasToNonAtlas(InSectionBaseList[i], WeightmapScaleBias, SubsectionSizeQuads, WeightmapReadTextureSize, WeightmapWriteTextureSize, TriangleList);
+		} break;
+
+		case ERTDrawingType::RTNonAtlas:
+		{
+			GenerateLayersRenderQuadsNonAtlas(InSectionBaseList[i], WeightmapScaleBias, SubsectionSizeQuads, WeightmapReadTextureSize, WeightmapWriteTextureSize, TriangleList);
+		} break;
+
+		case ERTDrawingType::RTNonAtlasToAtlas:
 		{
-			GenerateLayersRenderQuadsMip(SectionBase, InScaleBias, SubsectionSizeQuads, WeightmapReadTextureSize, WeightmapWriteTextureSize, InMipRender, TriangleList);
+			GenerateLayersRenderQuadsNonAtlasToAtlas(InSectionBaseList[i], WeightmapScaleBias, SubsectionSizeQuads, WeightmapReadTextureSize, WeightmapWriteTextureSize, TriangleList);
+		} break;
+
+		case ERTDrawingType::RTMips:
+		{
+			GenerateLayersRenderQuadsMip(InSectionBaseList[i], WeightmapScaleBias, SubsectionSizeQuads, WeightmapReadTextureSize, WeightmapWriteTextureSize, InMipRender, TriangleList);
+		} break;
+
+		default:
+		{
+			check(false);
+			return;
+		}
 		}
 	}
 
@@ -1883,68 +1903,24 @@ void ALandscape::DrawWeightmapComponentsToRenderTarget(const FString& InDebugNam
 }
 
 void ALandscape::DrawWeightmapComponentsToRenderTarget(const FString& InDebugName, const TArray<ULandscapeComponent*>& InComponentsToDraw, const FIntPoint& InLandscapeBase, UTexture* InWeightmapRTRead, UTextureRenderTarget2D* InOptionalWeightmapRTRead2, UTextureRenderTarget2D* InWeightmapRTWrite,
-														bool InClearRTWrite, FLandscapeLayersWeightmapShaderParameters& InShaderParams, uint8 InMipRender) const
+														ERTDrawingType InDrawType, bool InClearRTWrite, FLandscapeLayersWeightmapShaderParameters& InShaderParams, uint8 InMipRender) const
 {
-	check(InWeightmapRTRead != nullptr);
-	check(InWeightmapRTWrite != nullptr);
-
-	FIntPoint WeightmapWriteTextureSize(InWeightmapRTWrite->SizeX, InWeightmapRTWrite->SizeY);
-	FIntPoint WeightmapReadTextureSize(InWeightmapRTRead->Source.GetSizeX(), InWeightmapRTRead->Source.GetSizeY());
-	UTextureRenderTarget2D* WeightmapRTRead = Cast<UTextureRenderTarget2D>(InWeightmapRTRead);
-
-	if (WeightmapRTRead != nullptr)
-	{
-		WeightmapReadTextureSize.X = WeightmapRTRead->SizeX;
-		WeightmapReadTextureSize.Y = WeightmapRTRead->SizeY;
-	}
-
-	// Quad Setup
-	TArray<FLandscapeLayersTriangle> TriangleList;
-	TriangleList.Reserve(InComponentsToDraw.Num() * 2 * NumSubsections);
 
-	if (InMipRender == 0)
-	{
-		for (ULandscapeComponent* Component : InComponentsToDraw)
-		{
-			// TODO: check what to do with WeightmapSubsectionOffset
-			FVector2D WeightmapScaleBias(Component->WeightmapScaleBias.Z, Component->WeightmapScaleBias.W);
-			FIntPoint ComponentSectionBase = Component->GetSectionBase() - InLandscapeBase;
+	TArray<FIntPoint> SectionBaseList;
+	SectionBaseList.Reserve(InComponentsToDraw.Num());
+	TArray<FVector2D> WeightmapScaleBiasList;
+	WeightmapScaleBiasList.Reserve(InComponentsToDraw.Num());
 
-			GenerateLayersRenderQuadsAtlas(ComponentSectionBase, WeightmapScaleBias, SubsectionSizeQuads, WeightmapReadTextureSize, WeightmapWriteTextureSize, TriangleList);
-		}
-	}
-	else
+	for (ULandscapeComponent* Component : InComponentsToDraw)
 	{
-		for (ULandscapeComponent* Component : InComponentsToDraw)
-		{
-			// TODO: check what to do with WeightmapSubsectionOffset
-			FVector2D WeightmapScaleBias(Component->WeightmapScaleBias.Z, Component->WeightmapScaleBias.W);
-			FIntPoint ComponentSectionBase = Component->GetSectionBase() - InLandscapeBase;
+		FVector2D WeightmapScaleBias(Component->WeightmapScaleBias.Z, Component->WeightmapScaleBias.W);
+		WeightmapScaleBiasList.Add(WeightmapScaleBias);
 
-			GenerateLayersRenderQuadsMip(ComponentSectionBase, WeightmapScaleBias, SubsectionSizeQuads, WeightmapReadTextureSize, WeightmapWriteTextureSize, InMipRender, TriangleList);
-		}
-	}
-
-	InShaderParams.ReadWeightmap1 = InWeightmapRTRead;
-	InShaderParams.ReadWeightmap2 = InOptionalWeightmapRTRead2;
-	InShaderParams.CurrentMipComponentVertexCount = ((SubsectionSizeQuads + 1) >> InMipRender);
-
-	if (InMipRender > 0)
-	{
-		InShaderParams.CurrentMipSize = WeightmapWriteTextureSize;
-		InShaderParams.ParentMipSize = WeightmapReadTextureSize;
+		FIntPoint ComponentSectionBase = Component->GetSectionBase() - InLandscapeBase;
+		SectionBaseList.Add(ComponentSectionBase);
 	}
 
-	FMatrix ProjectionMatrix = AdjustProjectionMatrixForRHI(FTranslationMatrix(FVector(0, 0, 0)) *
-															FMatrix(FPlane(1.0f / (FMath::Max<uint32>(WeightmapWriteTextureSize.X, 1.f) / 2.0f), 0.0, 0.0f, 0.0f), FPlane(0.0f, -1.0f / (FMath::Max<uint32>(WeightmapWriteTextureSize.Y, 1.f) / 2.0f), 0.0f, 0.0f), FPlane(0.0f, 0.0f, 1.0f, 0.0f), FPlane(-1.0f, 1.0f, 0.0f, 1.0f)));
-
-	FLandscapeLayersWeightmapRender_RenderThread LayersRender(InDebugName, InWeightmapRTWrite, WeightmapWriteTextureSize, WeightmapReadTextureSize, ProjectionMatrix, InShaderParams, InMipRender, TriangleList);
-
-	ENQUEUE_RENDER_COMMAND(FDrawLandscapeLayersWeightmapCommand)(
-		[LayersRender, InClearRTWrite](FRHICommandListImmediate& RHICmdList) mutable
-	{
-		LayersRender.Render(RHICmdList, InClearRTWrite);
-	});
+	DrawWeightmapComponentsToRenderTarget(InDebugName, SectionBaseList, FVector2D::ZeroVector, &WeightmapScaleBiasList, InWeightmapRTRead, InOptionalWeightmapRTRead2, InWeightmapRTWrite, InDrawType, InClearRTWrite, InShaderParams, InMipRender);
 
 	PrintLayersDebugRT(InDebugName, InWeightmapRTWrite, InMipRender, false);
 }
@@ -1977,7 +1953,7 @@ void ALandscape::DrawWeightmapComponentToRenderTargetMips(const TArray<FVector2D
 		if (WriteMipRT != nullptr)
 		{
 			DrawWeightmapComponentsToRenderTarget(OutputDebugName ? FString::Printf(TEXT("LS Weight: %s = -> %s Mips %d"), *ReadMipRT->GetName(), *WriteMipRT->GetName(), CurrentMip) : GEmptyDebugName,
-													SectionBaseToDraw, WeightmapScaleBias, ReadMipRT, nullptr, WriteMipRT, InClearRTWrite, InShaderParams, CurrentMip++);
+													SectionBaseToDraw, WeightmapScaleBias, nullptr, ReadMipRT, nullptr, WriteMipRT, ERTDrawingType::RTMips, InClearRTWrite, InShaderParams, CurrentMip++);
 		}
 
 		ReadMipRT = WeightmapRTList[MipRTIndex];
@@ -2672,24 +2648,31 @@ int32 ALandscape::RegenerateLayersHeightmaps(const TArray<ULandscapeComponent*>&
 		};
 
 		// Calculate Top Left Lambda
-		auto GetProxyUniqueHeightmaps = [&](ALandscapeProxy* InProxy, TArray<FHeightmapTopLeft>& OutHeightmaps, const FGuid& InLayerGuid = FGuid())
+		auto GetProxyUniqueHeightmaps = [&](ALandscapeProxy* InProxy, TArray<FHeightmapTopLeft>& OutHeightmaps, const FIntPoint& LandscapeBaseQuads, const FGuid& InLayerGuid = FGuid())
 		{
 			FScopedSetLandscapeEditingLayer Scope(this, InLayerGuid);
+			
+			const int32 ComponentSizeQuad = SubsectionSizeQuads * NumSubsections;
+			const int32 ComponentSizeVerts = (SubsectionSizeQuads + 1) * NumSubsections;
 			for (ULandscapeComponent* Component : InProxy->LandscapeComponents)
 			{
 				UTexture2D* ComponentHeightmap = Component->GetHeightmap(true);
 		
 				int32 Index = OutHeightmaps.IndexOfByPredicate([=](const FHeightmapTopLeft& LayerHeightmap) { return LayerHeightmap.Texture == ComponentHeightmap; });
 
+				FIntPoint ComponentSectionBase = Component->GetSectionBase() - LandscapeBaseQuads;
+				FVector2D SourcePositionOffset(FMath::RoundToInt(ComponentSectionBase.X / ComponentSizeQuad), FMath::RoundToInt(ComponentSectionBase.Y / ComponentSizeQuad));
+				FIntPoint ComponentVertexPosition = FIntPoint(SourcePositionOffset.X * ComponentSizeVerts, SourcePositionOffset.Y * ComponentSizeVerts);
+
 				if (Index == INDEX_NONE)
 				{
 					FLandscapeLayersTexture2DCPUReadBackResource** CPUReadback = InProxy->HeightmapsCPUReadBack.Find(ComponentHeightmap);
-					OutHeightmaps.Add(FHeightmapTopLeft(ComponentHeightmap, Component->GetSectionBase(), CPUReadback != nullptr ? *CPUReadback : nullptr));
+					OutHeightmaps.Add(FHeightmapTopLeft(ComponentHeightmap, ComponentVertexPosition, CPUReadback != nullptr ? *CPUReadback : nullptr));
 				}
 				else
 				{
-					OutHeightmaps[Index].TopLeftSectionBase.X = FMath::Min(OutHeightmaps[Index].TopLeftSectionBase.X, Component->GetSectionBase().X);
-					OutHeightmaps[Index].TopLeftSectionBase.Y = FMath::Min(OutHeightmaps[Index].TopLeftSectionBase.Y, Component->GetSectionBase().Y);
+					OutHeightmaps[Index].TopLeftSectionBase.X = FMath::Min(OutHeightmaps[Index].TopLeftSectionBase.X, ComponentVertexPosition.X);
+					OutHeightmaps[Index].TopLeftSectionBase.Y = FMath::Min(OutHeightmaps[Index].TopLeftSectionBase.Y, ComponentVertexPosition.Y);
 				}
 			}
 		};
@@ -2729,11 +2712,11 @@ int32 ALandscape::RegenerateLayersHeightmaps(const TArray<ULandscapeComponent*>&
 			Info->ForAllLandscapeProxies([&](ALandscapeProxy* Proxy)
 			{
                 TArray<FHeightmapTopLeft> LayerHeightmaps;
-				GetProxyUniqueHeightmaps(Proxy, LayerHeightmaps, Layer.Guid);
+				GetProxyUniqueHeightmaps(Proxy, LayerHeightmaps, LandscapeExtent.Min, Layer.Guid);
 
 				for (const FHeightmapTopLeft& LayerHeightmap : LayerHeightmaps)
 				{
-					AddDeferredCopyLayersTexture(LayerHeightmap.Texture, LandscapeScratchRT1, nullptr, LayerHeightmap.TopLeftSectionBase - LandscapeExtent.Min);
+					AddDeferredCopyLayersTexture(LayerHeightmap.Texture, LandscapeScratchRT1, nullptr, LayerHeightmap.TopLeftSectionBase);
 				}
 			});
 
@@ -2799,11 +2782,11 @@ int32 ALandscape::RegenerateLayersHeightmaps(const TArray<ULandscapeComponent*>&
 		Info->ForAllLandscapeProxies([&](ALandscapeProxy* Proxy)
 		{
 			TArray<FHeightmapTopLeft> Heightmaps;
-			GetProxyUniqueHeightmaps(Proxy, Heightmaps);
+			GetProxyUniqueHeightmaps(Proxy, Heightmaps, LandscapeExtent.Min);
 
 			for (const FHeightmapTopLeft& Heightmap : Heightmaps)
 			{
-				FIntPoint TextureTopLeftSectionBase = Heightmap.TopLeftSectionBase - LandscapeExtent.Min;
+				FIntPoint TextureTopLeftSectionBase = Heightmap.TopLeftSectionBase;
 				uint8 MipIndex = 0;
 				
 				check(Heightmap.CPUReadback);
@@ -2952,6 +2935,8 @@ void ALandscape::PrepareComponentDataToExtractMaterialLayersCS(const FLandscapeL
 
 	Info->ForAllLandscapeProxies([&](ALandscapeProxy* Proxy)
 	{
+		int32 LocalComponentSizeQuad = Proxy->SubsectionSizeQuads * NumSubsections;
+		int32 LocalComponentSizeVerts = (Proxy->SubsectionSizeQuads + 1) * NumSubsections;
 		for (ULandscapeComponent* Component : Proxy->LandscapeComponents)
 		{
 			FLandscapeLayerComponentData* ComponentLayerData = Component->GetLayerData(InLayer.Guid);
@@ -2967,8 +2952,10 @@ void ALandscape::PrepareComponentDataToExtractMaterialLayersCS(const FLandscapeL
 					check(LayerWeightmapUsage != nullptr);
 
 					FIntPoint ComponentSectionBase = Component->GetSectionBase() - InLandscapeBase;
+					FVector2D SourcePositionOffset(FMath::RoundToInt(ComponentSectionBase.X / LocalComponentSizeQuad), FMath::RoundToInt(ComponentSectionBase.Y / LocalComponentSizeQuad));
+					FIntPoint SourceComponentVertexPosition = FIntPoint(SourcePositionOffset.X * LocalComponentSizeVerts, SourcePositionOffset.Y * LocalComponentSizeVerts);
 
-					AddDeferredCopyLayersTexture(InOutputDebugName ? *LayerWeightmap->GetName() : GEmptyDebugName, LayerWeightmap->Resource, InOutputDebugName ? FString::Printf(TEXT("%s WeightmapScratchTexture"), *InLayer.Name.ToString()) : GEmptyDebugName, InOutTextureData, nullptr, ComponentSectionBase, 0);
+                    AddDeferredCopyLayersTexture(InOutputDebugName ? *LayerWeightmap->GetName() : GEmptyDebugName, LayerWeightmap->Resource, InOutputDebugName ? FString::Printf(TEXT("%s WeightmapScratchTexture"), *InLayer.Name.ToString()) : GEmptyDebugName, InOutTextureData, nullptr, SourceComponentVertexPosition, 0);
 					PrintLayersDebugTextureResource(InOutputDebugName ? FString::Printf(TEXT("LS Weight: %s WeightmapScratchTexture %s"), *InLayer.Name.ToString(), TEXT("WeightmapScratchTextureResource")) : GEmptyDebugName, InOutTextureData, 0, false);
 
 					for (const FWeightmapLayerAllocationInfo& WeightmapLayerAllocation : ComponentLayerData->WeightmapData.LayerAllocations)
@@ -2983,12 +2970,9 @@ void ALandscape::PrepareComponentDataToExtractMaterialLayersCS(const FLandscapeL
 							FIntPoint DestComponentSectionBase = DestComponent->GetSectionBase() - InLandscapeBase;
 
 							// Compute component top left vertex position from section base info
-							int32 LocalComponentSizeQuad = Component->SubsectionSizeQuads * NumSubsections;
-							int32 LocalComponentSizeVerts = (Component->SubsectionSizeQuads + 1) * NumSubsections;
-							FVector2D SourcePositionOffset(FMath::RoundToInt(ComponentSectionBase.X / LocalComponentSizeQuad), FMath::RoundToInt(ComponentSectionBase.Y / LocalComponentSizeQuad));
 							FVector2D DestPositionOffset(FMath::RoundToInt(DestComponentSectionBase.X / LocalComponentSizeQuad), FMath::RoundToInt(DestComponentSectionBase.Y / LocalComponentSizeQuad));
 
-							Data.ComponentVertexPosition = FIntPoint(SourcePositionOffset.X * LocalComponentSizeVerts, SourcePositionOffset.Y * LocalComponentSizeVerts);
+							Data.ComponentVertexPosition = SourceComponentVertexPosition;
 							Data.AtlasTexturePositionOutput = FIntPoint(DestPositionOffset.X * LocalComponentSizeVerts, DestPositionOffset.Y * LocalComponentSizeVerts);
 							Data.WeightmapChannelToProcess = WeightmapLayerAllocation.WeightmapTextureChannel;
 
@@ -3590,7 +3574,7 @@ int32 ALandscape::RegenerateLayersWeightmaps(const TArray<ULandscapeComponent*>&
 					PSShaderParams.LayerAlpha = LayerInfoObj == ALandscapeProxy::VisibilityLayer ? 1.0f : Layer.WeightmapAlpha; // visibility can't be affected by weight
 
 					DrawWeightmapComponentsToRenderTarget(OutputDebugName ? FString::Printf(TEXT("LS Weight: %s Paint: %s += -> %s"), *Layer.Name.ToString(), *LayerInfoObj->LayerName.ToString(), *LandscapeScratchRT1->GetName(), *LandscapeScratchRT2->GetName()) : GEmptyDebugName,
-														  InLandscapeComponents, LandscapeExtent.Min, LandscapeScratchRT1, nullptr, LandscapeScratchRT2, true, PSShaderParams, 0);
+														  InLandscapeComponents, LandscapeExtent.Min, LandscapeScratchRT1, nullptr, LandscapeScratchRT2, ERTDrawingType::RTAtlas, true, PSShaderParams, 0);
 
 					PSShaderParams.ApplyLayerModifiers = false;
 
@@ -3617,15 +3601,13 @@ int32 ALandscape::RegenerateLayersWeightmaps(const TArray<ULandscapeComponent*>&
 					}
 
 					DrawWeightmapComponentsToRenderTarget(OutputDebugName ? FString::Printf(TEXT("LS Weight: %s PaintLayer: %s, %s += -> Combined %s"), *Layer.Name.ToString(), *LayerInfoObj->LayerName.ToString(), *LandscapeScratchRT2->GetName(), *LandscapeScratchRT3->GetName()) : GEmptyDebugName,
-														  InLandscapeComponents, LandscapeExtent.Min, LandscapeScratchRT2, bFirstLayer ? nullptr : LandscapeScratchRT1, LandscapeScratchRT3, true, PSShaderParams, 0);
+														  InLandscapeComponents, LandscapeExtent.Min, LandscapeScratchRT2, bFirstLayer ? nullptr : LandscapeScratchRT1, LandscapeScratchRT3, ERTDrawingType::RTAtlasToNonAtlas, true, PSShaderParams, 0);
 
 					PSShaderParams.OutputAsSubstractive = false;
 
 					SourceDebugName = OutputDebugName ? FString::Printf(TEXT("Weight: %s PaintLayer: %s %s"), *Layer.Name.ToString(), *LayerInfoObj->LayerName.ToString(), *LandscapeScratchRT3->GetName()) : GEmptyDebugName;
 					DestDebugName = OutputDebugName ? TEXT("CombinedProcLayerWeightmap") : GEmptyDebugName;
 
-					CopyLayersTexture(SourceDebugName, LandscapeScratchRT3->GameThread_GetRenderTargetResource(), DestDebugName, CombinedLayersWeightmapAllMaterialLayersResource, nullptr, FIntPoint(0, 0), 0, 0, 0, LayerIndex);
-
 					// Handle brush blending
 					if (Layer.bVisible)
 					{
@@ -3660,6 +3642,12 @@ int32 ALandscape::RegenerateLayersWeightmaps(const TArray<ULandscapeComponent*>&
 						DestDebugName = OutputDebugName ? TEXT("CombinedProcLayerWeightmap") : GEmptyDebugName;
 						CopyLayersTexture(SourceDebugName, LandscapeScratchRT3->GameThread_GetRenderTargetResource(), DestDebugName, CombinedLayersWeightmapAllMaterialLayersResource, nullptr, FIntPoint(0, 0), 0, 0, 0, LayerIndex);
 					}
+
+					DrawWeightmapComponentsToRenderTarget(OutputDebugName ? FString::Printf(TEXT("LS Weight: %s Combined Scratch No Border to %s Combined Scratch with Border"), *LandscapeScratchRT3->GetName(), *LandscapeScratchRT1->GetName()) : GEmptyDebugName,
+						InLandscapeComponents, LandscapeExtent.Min, LandscapeScratchRT3, nullptr, LandscapeScratchRT1, ERTDrawingType::RTNonAtlasToAtlas, true, PSShaderParams, 0);
+
+
+					CopyLayersTexture(SourceDebugName, LandscapeScratchRT1->GameThread_GetRenderTargetResource(), DestDebugName, CombinedLayersWeightmapAllMaterialLayersResource, nullptr, FIntPoint(0, 0), 0, 0, 0, LayerIndex);
 				}
 
 				PSShaderParams.ApplyLayerModifiers = false;
@@ -3742,9 +3730,9 @@ int32 ALandscape::RegenerateLayersWeightmaps(const TArray<ULandscapeComponent*>&
 
 					for (int32 i = 0; i < PackLayersComponentsData.Num(); ++i)
 					{
-						check(ComponentY < WeightmapScratchPackLayerTextureResource->GetSizeY()); // This should never happen as it would be a bug in the algo
+						check(ComponentY+ComponentSize <= WeightmapScratchPackLayerTextureResource->GetSizeY()); // This should never happen as it would be a bug in the algo
 
-						if (ComponentX >= WeightmapScratchPackLayerTextureResource->GetSizeX())
+						if (ComponentX+ComponentSize > WeightmapScratchPackLayerTextureResource->GetSizeX())
 						{
 							ComponentY += ComponentSize;
 							ComponentX = 0;
