Expose "effects heightmap", "effects weightmap", "affectedweightmaplayers" to on landscape blueprint brush to blueprints

diff --git a/Engine/Source/Runtime/Landscape/Public/LandscapeBlueprintBrushBase.h b/Engine/Source/Runtime/Landscape/Public/LandscapeBlueprintBrushBase.h
--- a/Engine/Source/Runtime/Landscape/Public/LandscapeBlueprintBrushBase.h
+++ b/Engine/Source/Runtime/Landscape/Public/LandscapeBlueprintBrushBase.h
@@ -18,13 +18,13 @@ protected:
 	UPROPERTY(Transient)
 	class ALandscape* OwningLandscape;
 
-	UPROPERTY(Category = "Settings", EditAnywhere)
+	UPROPERTY(Category = "Settings", EditAnywhere, BlueprintReadWrite)
 	bool AffectHeightmap;
 
-	UPROPERTY(Category = "Settings", EditAnywhere)
+	UPROPERTY(Category = "Settings", EditAnywhere, BlueprintReadWrite)
 	bool AffectWeightmap;
 
-	UPROPERTY(Category = "Settings", EditAnywhere)
+	UPROPERTY(Category = "Settings", EditAnywhere, BlueprintReadWrite)
 	TArray<FName> AffectedWeightmapLayers;
 
 	UPROPERTY(Transient)
