Fix Current Layer not being cleared properly when no paint data exists, causing Brush weightmap fill.

#rnx

diff --git a/Engine/Source/Runtime/Landscape/Private/LandscapeEditLayers.cpp b/Engine/Source/Runtime/Landscape/Private/LandscapeEditLayers.cpp
--- a/Engine/Source/Runtime/Landscape/Private/LandscapeEditLayers.cpp
+++ b/Engine/Source/Runtime/Landscape/Private/LandscapeEditLayers.cpp
@@ -3509,30 +3509,30 @@ int32 ALandscape::RegenerateLayersWeightmaps(const TArray<ULandscapeComponent*>&
 
 				HasFoundWeightmapToProcess = ComponentsData.Num() > 0;
 
-				// Perform the compute shader
-				if (ComponentsData.Num() > 0)
+				// Clear the current atlas if required
+				if (CurrentWeightmapToProcessIndex == 0)
 				{
-					PrintLayersDebugTextureResource(OutputDebugName ? FString::Printf(TEXT("LS Weight: %s WeightmapScratchTexture %s"), *Layer.Name.ToString(), TEXT("WeightmapScratchTextureResource")) : GEmptyDebugName, WeightmapScratchExtractLayerTextureResource, 0, false);
+					ClearLayersWeightmapTextureResource(TEXT("ClearRT"), LandscapeScratchRT1->GameThread_GetRenderTargetResource());
 
-					// Clear the current atlas if required
-					if (CurrentWeightmapToProcessIndex == 0)
+					// Important: for performance reason we only clear the layer we will write to, the other one might contain data but they will not be read during the blend phase
+					for (auto& ItPair : LayerInfoObjects)
 					{
-						ClearLayersWeightmapTextureResource(TEXT("ClearRT"), LandscapeScratchRT1->GameThread_GetRenderTargetResource());
+						int32 LayerIndex = ItPair.Value;
 
-						// Important: for performance reason we only clear the layer we will write to, the other one might contain data but they will not be read during the blend phase
-						for (auto& ItPair : LayerInfoObjects)
-						{
-							int32 LayerIndex = ItPair.Value;
-
-							SourceDebugName = OutputDebugName ? LandscapeScratchRT1->GetName() : GEmptyDebugName;
-							DestDebugName = OutputDebugName ? FString::Printf(TEXT("Weight: %s Clear CurrentProcLayerWeightmapAllLayersResource %d, "), *Layer.Name.ToString(), LayerIndex) : GEmptyDebugName;
-
-							AddDeferredCopyLayersTexture(SourceDebugName, LandscapeScratchRT1->GameThread_GetRenderTargetResource(), DestDebugName, CurrentLayersWeightmapAllMaterialLayersResource, nullptr, FIntPoint(0, 0), 0, 0, 0, LayerIndex);
-						}
+						SourceDebugName = OutputDebugName ? LandscapeScratchRT1->GetName() : GEmptyDebugName;
+						DestDebugName = OutputDebugName ? FString::Printf(TEXT("Weight: %s Clear CurrentProcLayerWeightmapAllLayersResource %d, "), *Layer.Name.ToString(), LayerIndex) : GEmptyDebugName;
 
-						CommitDeferredCopyLayersTexture();
+						AddDeferredCopyLayersTexture(SourceDebugName, LandscapeScratchRT1->GameThread_GetRenderTargetResource(), DestDebugName, CurrentLayersWeightmapAllMaterialLayersResource, nullptr, FIntPoint(0, 0), 0, 0, 0, LayerIndex);
 					}
 
+					CommitDeferredCopyLayersTexture();
+				}
+
+				// Perform the compute shader
+				if (ComponentsData.Num() > 0)
+				{
+					PrintLayersDebugTextureResource(OutputDebugName ? FString::Printf(TEXT("LS Weight: %s WeightmapScratchTexture %s"), *Layer.Name.ToString(), TEXT("WeightmapScratchTextureResource")) : GEmptyDebugName, WeightmapScratchExtractLayerTextureResource, 0, false);
+										
 					FLandscapeLayerWeightmapExtractMaterialLayersComputeShaderParameters CSExtractLayersShaderParams;
 					CSExtractLayersShaderParams.AtlasWeightmapsPerLayer = CurrentLayersWeightmapAllMaterialLayersResource;
 					CSExtractLayersShaderParams.ComponentWeightmapResource = WeightmapScratchExtractLayerTextureResource;
