Landscape Layers: Fix Crash and bad behavior with invisible proxies

#rnx

diff --git a/Engine/Source/Runtime/Landscape/Classes/Landscape.h b/Engine/Source/Runtime/Landscape/Classes/Landscape.h
--- a/Engine/Source/Runtime/Landscape/Classes/Landscape.h
+++ b/Engine/Source/Runtime/Landscape/Classes/Landscape.h
@@ -337,16 +337,16 @@ private:
 	int32 RegenerateLayersHeightmaps(const TArray<ULandscapeComponent*>& InLandscapeComponents, bool bInWaitForStreaming);
 	int32 RegenerateLayersWeightmaps(const TArray<ULandscapeComponent*>& InLandscapeComponents, bool bInWaitForStreaming);
 	bool UpdateCollisionAndClients(const TArray<ULandscapeComponent*>& InLandscapeComponents, const int32 InContentUpdateModes);
-	void ResolveLayersHeightmapTexture();
-	void ResolveLayersWeightmapTexture();
+	void ResolveLayersHeightmapTexture(const TArray<ULandscapeComponent*>& InLandscapeComponents);
+	void ResolveLayersWeightmapTexture(const TArray<ULandscapeComponent*>& InLandscapeComponents);
 	void ResolveLayersTexture(class FLandscapeLayersTexture2DCPUReadBackResource* InCPUReadBackTexture, UTexture2D* InOutputTexture);
 
 	bool PrepareLayersHeightmapTextureResources(bool bInWaitForStreaming) const;
 	bool PrepareLayersWeightmapTextureResources(bool bInWaitForStreaming) const;
 
-	void UpdateLayersMaterialInstances();
+	void UpdateLayersMaterialInstances(const TArray<ULandscapeComponent*>& InLandscapeComponents);
 
-	void PrepareComponentDataToExtractMaterialLayersCS(const FLandscapeLayer& InLayer, int32 InCurrentWeightmapToProcessIndex, const FIntPoint& InLandscapeBase, bool InOutputDebugName, class FLandscapeTexture2DResource* InOutTextureData,
+	void PrepareComponentDataToExtractMaterialLayersCS(const TArray<ULandscapeComponent*>& InLandscapeComponents, const FLandscapeLayer& InLayer, int32 InCurrentWeightmapToProcessIndex, const FIntPoint& InLandscapeBase, bool InOutputDebugName, class FLandscapeTexture2DResource* InOutTextureData,
 														  TArray<struct FLandscapeLayerWeightmapExtractMaterialLayersComponentData>& OutComponentData, TMap<ULandscapeLayerInfoObject*, int32>& OutLayerInfoObjects);
 	void PrepareComponentDataToPackMaterialLayersCS(int32 InCurrentWeightmapToProcessIndex, const FIntPoint& InLandscapeBase, bool InOutputDebugName, const TArray<ULandscapeComponent*>& InAllLandscapeComponents, TArray<UTexture2D*>& InOutProcessedWeightmaps,
 													TArray<class FLandscapeLayersTexture2DCPUReadBackResource*>& OutProcessedCPUReadBackTexture, TArray<struct FLandscapeLayerWeightmapPackMaterialLayersComponentData>& OutComponentData);
diff --git a/Engine/Source/Runtime/Landscape/Classes/LandscapeInfo.h b/Engine/Source/Runtime/Landscape/Classes/LandscapeInfo.h
--- a/Engine/Source/Runtime/Landscape/Classes/LandscapeInfo.h
+++ b/Engine/Source/Runtime/Landscape/Classes/LandscapeInfo.h
@@ -147,6 +147,7 @@ public:
 	// all below.
 	LANDSCAPE_API void GetComponentsInRegion(int32 X1, int32 Y1, int32 X2, int32 Y2, TSet<ULandscapeComponent*>& OutComponents, bool bOverlap = true) const;
 	LANDSCAPE_API bool GetLandscapeExtent(int32& MinX, int32& MinY, int32& MaxX, int32& MaxY) const;
+	LANDSCAPE_API void ForAllLandscapeComponents(TFunctionRef<void(ULandscapeComponent*)> Fn) const;
 	LANDSCAPE_API void ExportHeightmap(const FString& Filename);
 	LANDSCAPE_API void ExportLayer(ULandscapeLayerInfoObject* LayerInfo, const FString& Filename);
 	LANDSCAPE_API bool ApplySplines(bool bOnlySelected, TSet<ULandscapeComponent*>* OutModifiedComponents = nullptr);
diff --git a/Engine/Source/Runtime/Landscape/Private/Landscape.cpp b/Engine/Source/Runtime/Landscape/Private/Landscape.cpp
--- a/Engine/Source/Runtime/Landscape/Private/Landscape.cpp
+++ b/Engine/Source/Runtime/Landscape/Private/Landscape.cpp
@@ -2810,9 +2810,7 @@ void ULandscapeInfo::RegisterActor(ALandscapeProxy* Proxy, bool bMapCheck)
 
 	if (LandscapeActor && LandscapeActor->HasLayersContent())
 	{
-		// Force update rendering resources
-		const bool bInRequestContentUpdate = false;
-		LandscapeActor->RequestLayersInitialization(bInRequestContentUpdate);		
+		LandscapeActor->RequestLayersInitialization();		
 	}
 
 	UpdateLayerInfoMap(Proxy);
diff --git a/Engine/Source/Runtime/Landscape/Private/LandscapeEdit.cpp b/Engine/Source/Runtime/Landscape/Private/LandscapeEdit.cpp
--- a/Engine/Source/Runtime/Landscape/Private/LandscapeEdit.cpp
+++ b/Engine/Source/Runtime/Landscape/Private/LandscapeEdit.cpp
@@ -3281,6 +3281,14 @@ bool ULandscapeInfo::GetLandscapeExtent(int32& MinX, int32& MinY, int32& MaxX, i
 	return (MinX != MAX_int32);
 }
 
+LANDSCAPE_API void ULandscapeInfo::ForAllLandscapeComponents(TFunctionRef<void(ULandscapeComponent*)> Fn) const
+{
+	for (auto& XYComponentPair : XYtoComponentMap)
+	{
+		Fn(XYComponentPair.Value);
+	}
+}
+
 bool ULandscapeInfo::GetSelectedExtent(int32& MinX, int32& MinY, int32& MaxX, int32& MaxY) const
 {
 	MinX = MinY = MAX_int32;
diff --git a/Engine/Source/Runtime/Landscape/Private/LandscapeEditLayers.cpp b/Engine/Source/Runtime/Landscape/Private/LandscapeEditLayers.cpp
--- a/Engine/Source/Runtime/Landscape/Private/LandscapeEditLayers.cpp
+++ b/Engine/Source/Runtime/Landscape/Private/LandscapeEditLayers.cpp
@@ -1018,13 +1018,16 @@ public:
 			Info.SourcePosition.X = 0;
 			Info.Size.X = SourceSize.X;
 			Info.DestPosition.X = Params.InitialPositionOffset.X >> Params.DestMip;
-
+			check(Info.DestPosition.X + Info.Size.X <= DestSize.X);
 		}
 		else
 		{
 			Info.SourcePosition.X = Params.InitialPositionOffset.X >> Params.SourceMip;
 			Info.Size.X = DestSize.X;
 			Info.DestPosition.X = 0;
+            check(Info.SourcePosition.X >= 0);
+			check(Info.SourcePosition.X + Info.Size.X <= SourceSize.X);
+			check(Info.DestPosition.X + Info.Size.X <= DestSize.X);
 		}
 
 		if (SourceSize.Y <= DestSize.Y)
@@ -1032,12 +1035,16 @@ public:
 			Info.SourcePosition.Y = 0;
 			Info.Size.Y = SourceSize.Y;
 			Info.DestPosition.Y = Params.InitialPositionOffset.Y >> Params.DestMip;
+			check(Info.DestPosition.Y + Info.Size.Y <= DestSize.Y);
 		}
 		else
 		{
 			Info.SourcePosition.Y = Params.InitialPositionOffset.Y >> Params.SourceMip;
 			Info.Size.Y = DestSize.Y;
 			Info.DestPosition.Y = 0;
+check(Info.SourcePosition.Y >= 0);
+			check(Info.SourcePosition.Y + Info.Size.Y <= SourceSize.Y);
+			check(Info.DestPosition.Y + Info.Size.Y <= DestSize.Y);
 		}
 
 		InRHICmdList.CopyTexture(Params.SourceResource->TextureRHI, Params.DestResource->TextureRHI, Info);
@@ -2648,13 +2655,13 @@ int32 ALandscape::RegenerateLayersHeightmaps(const TArray<ULandscapeComponent*>&
 		};
 
 		// Calculate Top Left Lambda
-		auto GetProxyUniqueHeightmaps = [&](ALandscapeProxy* InProxy, TArray<FHeightmapTopLeft>& OutHeightmaps, const FIntPoint& LandscapeBaseQuads, const FGuid& InLayerGuid = FGuid())
+		auto GetUniqueHeightmaps = [&](const TArray<ULandscapeComponent*>& InLandscapeComponent, TArray<FHeightmapTopLeft>& OutHeightmaps, const FIntPoint& LandscapeBaseQuads, const FGuid& InLayerGuid = FGuid())
 		{
 			FScopedSetLandscapeEditingLayer Scope(this, InLayerGuid);
 			
 			const int32 ComponentSizeQuad = SubsectionSizeQuads * NumSubsections;
 			const int32 ComponentSizeVerts = (SubsectionSizeQuads + 1) * NumSubsections;
-			for (ULandscapeComponent* Component : InProxy->LandscapeComponents)
+			for (ULandscapeComponent* Component : InLandscapeComponents)
 			{
 				UTexture2D* ComponentHeightmap = Component->GetHeightmap(true);
 		
@@ -2663,10 +2670,11 @@ int32 ALandscape::RegenerateLayersHeightmaps(const TArray<ULandscapeComponent*>&
 				FIntPoint ComponentSectionBase = Component->GetSectionBase() - LandscapeBaseQuads;
 				FVector2D SourcePositionOffset(FMath::RoundToInt(ComponentSectionBase.X / ComponentSizeQuad), FMath::RoundToInt(ComponentSectionBase.Y / ComponentSizeQuad));
 				FIntPoint ComponentVertexPosition = FIntPoint(SourcePositionOffset.X * ComponentSizeVerts, SourcePositionOffset.Y * ComponentSizeVerts);
+				ALandscapeProxy* Proxy = Component->GetLandscapeProxy();
 
 				if (Index == INDEX_NONE)
 				{
-					FLandscapeLayersTexture2DCPUReadBackResource** CPUReadback = InProxy->HeightmapsCPUReadBack.Find(ComponentHeightmap);
+					FLandscapeLayersTexture2DCPUReadBackResource** CPUReadback = Proxy->HeightmapsCPUReadBack.Find(ComponentHeightmap);
 					OutHeightmaps.Add(FHeightmapTopLeft(ComponentHeightmap, ComponentVertexPosition, CPUReadback != nullptr ? *CPUReadback : nullptr));
 				}
 				else
@@ -2709,17 +2717,14 @@ int32 ALandscape::RegenerateLayersHeightmaps(const TArray<ULandscapeComponent*>&
 				ShaderParams.LayerAlpha = Layer.HeightmapAlpha;
 			}
 
-			Info->ForAllLandscapeProxies([&](ALandscapeProxy* Proxy)
 			{
-                TArray<FHeightmapTopLeft> LayerHeightmaps;
-				GetProxyUniqueHeightmaps(Proxy, LayerHeightmaps, LandscapeExtent.Min, Layer.Guid);
-
+				TArray<FHeightmapTopLeft> LayerHeightmaps;
+				GetUniqueHeightmaps(InLandscapeComponents, LayerHeightmaps, LandscapeExtent.Min, Layer.Guid);
 				for (const FHeightmapTopLeft& LayerHeightmap : LayerHeightmaps)
 				{
 					AddDeferredCopyLayersTexture(LayerHeightmap.Texture, LandscapeScratchRT1, nullptr, LayerHeightmap.TopLeftSectionBase);
 				}
-			});
-
+			}
 			CommitDeferredCopyLayersTexture();
 
 			// NOTE: From this point on, we always work in non atlas, we'll convert back at the end to atlas only
@@ -2779,16 +2784,14 @@ int32 ALandscape::RegenerateLayersHeightmaps(const TArray<ULandscapeComponent*>&
 		DrawHeightmapComponentsToRenderTargetMips(InLandscapeComponents, LandscapeExtent.Min, CombinedHeightmapAtlasRT, true, ShaderParams);
 
 		// Copy back all Mips to original heightmap data
-		Info->ForAllLandscapeProxies([&](ALandscapeProxy* Proxy)
 		{
 			TArray<FHeightmapTopLeft> Heightmaps;
-			GetProxyUniqueHeightmaps(Proxy, Heightmaps, LandscapeExtent.Min);
-
+			GetUniqueHeightmaps(InLandscapeComponents, Heightmaps, LandscapeExtent.Min);
 			for (const FHeightmapTopLeft& Heightmap : Heightmaps)
 			{
 				FIntPoint TextureTopLeftSectionBase = Heightmap.TopLeftSectionBase;
 				uint8 MipIndex = 0;
-				
+
 				check(Heightmap.CPUReadback);
 
 				// Mip 0
@@ -2807,14 +2810,13 @@ int32 ALandscape::RegenerateLayersHeightmaps(const TArray<ULandscapeComponent*>&
 					}
 				}
 			}
-		});
-
+		}
 		CommitDeferredCopyLayersTexture();
 	}
 
 	if (HeightmapUpdateModes)
 	{
-		ResolveLayersHeightmapTexture();
+		ResolveLayersHeightmapTexture(InLandscapeComponents);
 
 		// Partial Component Update
 		for (ULandscapeComponent* Component : InLandscapeComponents)
@@ -2843,7 +2845,7 @@ int32 ALandscape::RegenerateLayersHeightmaps(const TArray<ULandscapeComponent*>&
 	return HeightmapUpdateModes;
 }
 
-void ALandscape::ResolveLayersHeightmapTexture()
+void ALandscape::ResolveLayersHeightmapTexture(const TArray<ULandscapeComponent*>& InLandscapeComponents)
 {
 	SCOPE_CYCLE_COUNTER(STAT_LandscapeLayersResolveHeightmaps);
 
@@ -2854,29 +2856,26 @@ void ALandscape::ResolveLayersHeightmapTexture()
 		return;
 	}
 
-	Info->ForAllLandscapeProxies([&](ALandscapeProxy* Proxy)
+	TArray<UTexture*> ProcessedTexture;
+
+	for (ULandscapeComponent* Component : InLandscapeComponents)
 	{
-		TArray<UTexture*> ProcessedTexture;
+		UTexture2D* ComponentHeightmap = Component->GetHeightmap();
 
-		for (ULandscapeComponent* Component : Proxy->LandscapeComponents)
+		if (!ProcessedTexture.Contains(ComponentHeightmap))
 		{
-			UTexture2D* ComponentHeightmap = Component->GetHeightmap();
+			FLandscapeLayersTexture2DCPUReadBackResource** CPUReadback = Component->GetLandscapeProxy()->HeightmapsCPUReadBack.Find(ComponentHeightmap);
 
-			if (!ProcessedTexture.Contains(ComponentHeightmap))
+			if (CPUReadback == nullptr)
 			{
-				FLandscapeLayersTexture2DCPUReadBackResource** CPUReadback = Proxy->HeightmapsCPUReadBack.Find(ComponentHeightmap);
-
-				if (CPUReadback == nullptr)
-				{
-					continue;
-				}
+				continue;
+			}
 
-				ProcessedTexture.Add(ComponentHeightmap);
+			ProcessedTexture.Add(ComponentHeightmap);
 
-				ResolveLayersTexture(*CPUReadback, ComponentHeightmap);
-			}
+			ResolveLayersTexture(*CPUReadback, ComponentHeightmap);
 		}
-	});
+	}
 }
 
 void ALandscape::ResolveLayersTexture(FLandscapeLayersTexture2DCPUReadBackResource* InCPUReadBackTexture, UTexture2D* InOutputTexture)
@@ -2923,7 +2922,7 @@ void ALandscape::ResolveLayersTexture(FLandscapeLayersTexture2DCPUReadBackResour
 	}	
 }
 
-void ALandscape::PrepareComponentDataToExtractMaterialLayersCS(const FLandscapeLayer& InLayer, int32 InCurrentWeightmapToProcessIndex, const FIntPoint& InLandscapeBase, bool InOutputDebugName, FLandscapeTexture2DResource* InOutTextureData,
+void ALandscape::PrepareComponentDataToExtractMaterialLayersCS(const TArray<ULandscapeComponent*>& InLandscapeComponents, const FLandscapeLayer& InLayer, int32 InCurrentWeightmapToProcessIndex, const FIntPoint& InLandscapeBase, bool InOutputDebugName, FLandscapeTexture2DResource* InOutTextureData,
 																TArray<FLandscapeLayerWeightmapExtractMaterialLayersComponentData>& OutComponentData, TMap<ULandscapeLayerInfoObject*, int32>& OutLayerInfoObjects)
 {
 	ULandscapeInfo* Info = GetLandscapeInfo();
@@ -2933,80 +2932,77 @@ void ALandscape::PrepareComponentDataToExtractMaterialLayersCS(const FLandscapeL
 		return;
 	}
 
-	Info->ForAllLandscapeProxies([&](ALandscapeProxy* Proxy)
+	const int32 LocalComponentSizeQuad = SubsectionSizeQuads * NumSubsections;
+	const int32 LocalComponentSizeVerts = (SubsectionSizeQuads + 1) * NumSubsections;
+	for (ULandscapeComponent* Component : InLandscapeComponents)
 	{
-		int32 LocalComponentSizeQuad = Proxy->SubsectionSizeQuads * NumSubsections;
-		int32 LocalComponentSizeVerts = (Proxy->SubsectionSizeQuads + 1) * NumSubsections;
-		for (ULandscapeComponent* Component : Proxy->LandscapeComponents)
-		{
-			FLandscapeLayerComponentData* ComponentLayerData = Component->GetLayerData(InLayer.Guid);
+		FLandscapeLayerComponentData* ComponentLayerData = Component->GetLayerData(InLayer.Guid);
 
-			if (ComponentLayerData != nullptr)
+		if (ComponentLayerData != nullptr)
+		{
+			if (ComponentLayerData->WeightmapData.Textures.IsValidIndex(InCurrentWeightmapToProcessIndex) && ComponentLayerData->WeightmapData.TextureUsages.IsValidIndex(InCurrentWeightmapToProcessIndex))
 			{
-				if (ComponentLayerData->WeightmapData.Textures.IsValidIndex(InCurrentWeightmapToProcessIndex) && ComponentLayerData->WeightmapData.TextureUsages.IsValidIndex(InCurrentWeightmapToProcessIndex))
-				{
-					UTexture2D* LayerWeightmap = ComponentLayerData->WeightmapData.Textures[InCurrentWeightmapToProcessIndex];
-					check(LayerWeightmap != nullptr);
+				UTexture2D* LayerWeightmap = ComponentLayerData->WeightmapData.Textures[InCurrentWeightmapToProcessIndex];
+				check(LayerWeightmap != nullptr);
 
-					const ULandscapeWeightmapUsage* LayerWeightmapUsage = ComponentLayerData->WeightmapData.TextureUsages[InCurrentWeightmapToProcessIndex];
-					check(LayerWeightmapUsage != nullptr);
+				const ULandscapeWeightmapUsage* LayerWeightmapUsage = ComponentLayerData->WeightmapData.TextureUsages[InCurrentWeightmapToProcessIndex];
+				check(LayerWeightmapUsage != nullptr);
 
-					FIntPoint ComponentSectionBase = Component->GetSectionBase() - InLandscapeBase;
-					FVector2D SourcePositionOffset(FMath::RoundToInt(ComponentSectionBase.X / LocalComponentSizeQuad), FMath::RoundToInt(ComponentSectionBase.Y / LocalComponentSizeQuad));
-					FIntPoint SourceComponentVertexPosition = FIntPoint(SourcePositionOffset.X * LocalComponentSizeVerts, SourcePositionOffset.Y * LocalComponentSizeVerts);
+				FIntPoint ComponentSectionBase = Component->GetSectionBase() - InLandscapeBase;
+				FVector2D SourcePositionOffset(FMath::RoundToInt(ComponentSectionBase.X / LocalComponentSizeQuad), FMath::RoundToInt(ComponentSectionBase.Y / LocalComponentSizeQuad));
+				FIntPoint SourceComponentVertexPosition = FIntPoint(SourcePositionOffset.X * LocalComponentSizeVerts, SourcePositionOffset.Y * LocalComponentSizeVerts);
 
-                    AddDeferredCopyLayersTexture(InOutputDebugName ? *LayerWeightmap->GetName() : GEmptyDebugName, LayerWeightmap->Resource, InOutputDebugName ? FString::Printf(TEXT("%s WeightmapScratchTexture"), *InLayer.Name.ToString()) : GEmptyDebugName, InOutTextureData, nullptr, SourceComponentVertexPosition, 0);
-					PrintLayersDebugTextureResource(InOutputDebugName ? FString::Printf(TEXT("LS Weight: %s WeightmapScratchTexture %s"), *InLayer.Name.ToString(), TEXT("WeightmapScratchTextureResource")) : GEmptyDebugName, InOutTextureData, 0, false);
+				AddDeferredCopyLayersTexture(InOutputDebugName ? *LayerWeightmap->GetName() : GEmptyDebugName, LayerWeightmap->Resource, InOutputDebugName ? FString::Printf(TEXT("%s WeightmapScratchTexture"), *InLayer.Name.ToString()) : GEmptyDebugName, InOutTextureData, nullptr, SourceComponentVertexPosition, 0);
+				PrintLayersDebugTextureResource(InOutputDebugName ? FString::Printf(TEXT("LS Weight: %s WeightmapScratchTexture %s"), *InLayer.Name.ToString(), TEXT("WeightmapScratchTextureResource")) : GEmptyDebugName, InOutTextureData, 0, false);
 
-					for (const FWeightmapLayerAllocationInfo& WeightmapLayerAllocation : ComponentLayerData->WeightmapData.LayerAllocations)
+				for (const FWeightmapLayerAllocationInfo& WeightmapLayerAllocation : ComponentLayerData->WeightmapData.LayerAllocations)
+				{
+					if (WeightmapLayerAllocation.LayerInfo != nullptr && WeightmapLayerAllocation.WeightmapTextureIndex != 255 && ComponentLayerData->WeightmapData.Textures[WeightmapLayerAllocation.WeightmapTextureIndex] == LayerWeightmap)
 					{
-						if (WeightmapLayerAllocation.LayerInfo != nullptr && WeightmapLayerAllocation.WeightmapTextureIndex != 255 && ComponentLayerData->WeightmapData.Textures[WeightmapLayerAllocation.WeightmapTextureIndex] == LayerWeightmap)
-						{
-							FLandscapeLayerWeightmapExtractMaterialLayersComponentData Data;
+						FLandscapeLayerWeightmapExtractMaterialLayersComponentData Data;
 
-							const ULandscapeComponent* DestComponent = LayerWeightmapUsage->ChannelUsage[WeightmapLayerAllocation.WeightmapTextureChannel];
-							check(DestComponent);
+						const ULandscapeComponent* DestComponent = LayerWeightmapUsage->ChannelUsage[WeightmapLayerAllocation.WeightmapTextureChannel];
+						check(DestComponent);
 
-							FIntPoint DestComponentSectionBase = DestComponent->GetSectionBase() - InLandscapeBase;
+						FIntPoint DestComponentSectionBase = DestComponent->GetSectionBase() - InLandscapeBase;
 
-							// Compute component top left vertex position from section base info
-							FVector2D DestPositionOffset(FMath::RoundToInt(DestComponentSectionBase.X / LocalComponentSizeQuad), FMath::RoundToInt(DestComponentSectionBase.Y / LocalComponentSizeQuad));
+						// Compute component top left vertex position from section base info
+						FVector2D DestPositionOffset(FMath::RoundToInt(DestComponentSectionBase.X / LocalComponentSizeQuad), FMath::RoundToInt(DestComponentSectionBase.Y / LocalComponentSizeQuad));
 
-							Data.ComponentVertexPosition = SourceComponentVertexPosition;
-							Data.AtlasTexturePositionOutput = FIntPoint(DestPositionOffset.X * LocalComponentSizeVerts, DestPositionOffset.Y * LocalComponentSizeVerts);
-							Data.WeightmapChannelToProcess = WeightmapLayerAllocation.WeightmapTextureChannel;
+						Data.ComponentVertexPosition = SourceComponentVertexPosition;
+						Data.AtlasTexturePositionOutput = FIntPoint(DestPositionOffset.X * LocalComponentSizeVerts, DestPositionOffset.Y * LocalComponentSizeVerts);
+						Data.WeightmapChannelToProcess = WeightmapLayerAllocation.WeightmapTextureChannel;
 
-							if (WeightmapLayerAllocation.LayerInfo == ALandscapeProxy::VisibilityLayer)
-							{
-								Data.DestinationPaintLayerIndex = 0;
-								int32& NewLayerinfoObjectIndex = OutLayerInfoObjects.FindOrAdd(ALandscapeProxy::VisibilityLayer);
-								NewLayerinfoObjectIndex = 0;
-							}
-							else
+						if (WeightmapLayerAllocation.LayerInfo == ALandscapeProxy::VisibilityLayer)
+						{
+							Data.DestinationPaintLayerIndex = 0;
+							int32& NewLayerinfoObjectIndex = OutLayerInfoObjects.FindOrAdd(ALandscapeProxy::VisibilityLayer);
+							NewLayerinfoObjectIndex = 0;
+						}
+						else
+						{
+							for (int32 LayerInfoSettingsIndex = 0; LayerInfoSettingsIndex < Info->Layers.Num(); ++LayerInfoSettingsIndex)
 							{
-								for (int32 LayerInfoSettingsIndex = 0; LayerInfoSettingsIndex < Info->Layers.Num(); ++LayerInfoSettingsIndex)
-								{
-									const FLandscapeInfoLayerSettings& InfoLayerSettings = Info->Layers[LayerInfoSettingsIndex];
+								const FLandscapeInfoLayerSettings& InfoLayerSettings = Info->Layers[LayerInfoSettingsIndex];
 
-									if (InfoLayerSettings.LayerInfoObj == WeightmapLayerAllocation.LayerInfo)
-									{
-										Data.DestinationPaintLayerIndex = LayerInfoSettingsIndex + 1; // due to visibility layer that is at 0
-										int32& NewLayerinfoObjectIndex = OutLayerInfoObjects.FindOrAdd(WeightmapLayerAllocation.LayerInfo);
-										NewLayerinfoObjectIndex = LayerInfoSettingsIndex + 1;
+								if (InfoLayerSettings.LayerInfoObj == WeightmapLayerAllocation.LayerInfo)
+								{
+									Data.DestinationPaintLayerIndex = LayerInfoSettingsIndex + 1; // due to visibility layer that is at 0
+									int32& NewLayerinfoObjectIndex = OutLayerInfoObjects.FindOrAdd(WeightmapLayerAllocation.LayerInfo);
+									NewLayerinfoObjectIndex = LayerInfoSettingsIndex + 1;
 
-										break;
-									}
+									break;
 								}
 							}
-
-							OutComponentData.Add(Data);
 						}
+
+						OutComponentData.Add(Data);
 					}
 				}
 			}
 		}
-	});
-
+	}
+	
 	CommitDeferredCopyLayersTexture();
 }
 
@@ -3505,7 +3501,7 @@ int32 ALandscape::RegenerateLayersWeightmaps(const TArray<ULandscapeComponent*>&
 
 				// Prepare compute shader data
 				TArray<FLandscapeLayerWeightmapExtractMaterialLayersComponentData> ComponentsData;
-				PrepareComponentDataToExtractMaterialLayersCS(Layer, CurrentWeightmapToProcessIndex, LandscapeExtent.Min, OutputDebugName, WeightmapScratchExtractLayerTextureResource, ComponentsData, LayerInfoObjects);
+				PrepareComponentDataToExtractMaterialLayersCS(InLandscapeComponents, Layer, CurrentWeightmapToProcessIndex, LandscapeExtent.Min, OutputDebugName, WeightmapScratchExtractLayerTextureResource, ComponentsData, LayerInfoObjects);
 
 				HasFoundWeightmapToProcess = ComponentsData.Num() > 0;
 
@@ -3802,12 +3798,12 @@ int32 ALandscape::RegenerateLayersWeightmaps(const TArray<ULandscapeComponent*>&
 			}
 		}
 
-		UpdateLayersMaterialInstances();
+		UpdateLayersMaterialInstances(InLandscapeComponents);
 	}
 
 	if (WeightmapUpdateModes)
 	{
-		ResolveLayersWeightmapTexture();
+		ResolveLayersWeightmapTexture(InLandscapeComponents);
 	
 		for (ULandscapeComponent* Component : InLandscapeComponents)
 		{
@@ -3844,7 +3840,7 @@ uint32 ULandscapeComponent::ComputeWeightmapsHash()
 	return Hash;
 }
 
-void ALandscape::UpdateLayersMaterialInstances()
+void ALandscape::UpdateLayersMaterialInstances(const TArray<ULandscapeComponent*>& InLandscapeComponents)
 {
 	SCOPE_CYCLE_COUNTER(STAT_LandscapeLayersUpdateMaterialInstance);
 	TArray<ULandscapeComponent*> ComponentsToUpdate;
@@ -3852,18 +3848,15 @@ void ALandscape::UpdateLayersMaterialInstances()
 	// Compute Weightmap usage changes
 	if (ULandscapeInfo* Info = GetLandscapeInfo())
 	{
-		Info->ForAllLandscapeProxies([&ComponentsToUpdate](ALandscapeProxy* Proxy)
+		for (ULandscapeComponent* LandscapeComponent : InLandscapeComponents)
 		{
-			for (ULandscapeComponent* LandscapeComponent : Proxy->LandscapeComponents)
+			uint32 NewHash = LandscapeComponent->ComputeWeightmapsHash();
+			if (LandscapeComponent->WeightmapsHash != NewHash)
 			{
-				uint32 NewHash = LandscapeComponent->ComputeWeightmapsHash();
-				if (LandscapeComponent->WeightmapsHash != NewHash)
-				{
-					ComponentsToUpdate.Add(LandscapeComponent);
-					LandscapeComponent->WeightmapsHash = NewHash;
-				}
+				ComponentsToUpdate.Add(LandscapeComponent);
+				LandscapeComponent->WeightmapsHash = NewHash;
 			}
-		});
+		}
 	}
 
 	if (ComponentsToUpdate.Num() == 0)
@@ -4035,7 +4028,7 @@ void ALandscape::UpdateLayersMaterialInstances()
 	}
 }
 
-void ALandscape::ResolveLayersWeightmapTexture()
+void ALandscape::ResolveLayersWeightmapTexture(const TArray<ULandscapeComponent*>& InLandscapeComponents)
 {
 	SCOPE_CYCLE_COUNTER(STAT_LandscapeLayersResolveWeightmaps);
 
@@ -4046,32 +4039,29 @@ void ALandscape::ResolveLayersWeightmapTexture()
 		return;
 	}
 
-	Info->ForAllLandscapeProxies([&](ALandscapeProxy* Proxy)
+	TArray<UTexture2D*> ProcessedWeightmaps;
+
+	for (ULandscapeComponent* Component : InLandscapeComponents)
 	{
-		TArray<UTexture2D*> ProcessedWeightmaps;
+		TArray<UTexture2D*>& ComponentWeightmaps = Component->GetWeightmapTextures();
 
-		for (ULandscapeComponent* Component : Proxy->LandscapeComponents)
+		for (UTexture2D* ComponentLayerWeightmap : ComponentWeightmaps)
 		{
-			TArray<UTexture2D*>& ComponentWeightmaps = Component->GetWeightmapTextures();
-
-			for (UTexture2D* ComponentLayerWeightmap : ComponentWeightmaps)
+			if (!ProcessedWeightmaps.Contains(ComponentLayerWeightmap))
 			{
-				if (!ProcessedWeightmaps.Contains(ComponentLayerWeightmap))
-				{
-					FLandscapeLayersTexture2DCPUReadBackResource** CPUReadback = Proxy->WeightmapsCPUReadBack.Find(ComponentLayerWeightmap);
+				FLandscapeLayersTexture2DCPUReadBackResource** CPUReadback = Component->GetLandscapeProxy()->WeightmapsCPUReadBack.Find(ComponentLayerWeightmap);
 
-					if (CPUReadback == nullptr)
-					{
-						continue;
-					}
+				if (CPUReadback == nullptr)
+				{
+					continue;
+				}
 
-					ProcessedWeightmaps.Add(ComponentLayerWeightmap);
+				ProcessedWeightmaps.Add(ComponentLayerWeightmap);
 
-					ResolveLayersTexture(*CPUReadback, ComponentLayerWeightmap);
-				}
+				ResolveLayersTexture(*CPUReadback, ComponentLayerWeightmap);
 			}
 		}
-	});
+	}
 }
 
 bool ALandscape::HasLayersContent() const
@@ -4367,9 +4357,9 @@ void ALandscape::UpdateLayersContent(bool bInWaitForStreaming, bool bInSkipMonit
 	}
 
 	TArray<ULandscapeComponent*> AllLandscapeComponents;
-	GetLandscapeInfo()->ForAllLandscapeProxies([&AllLandscapeComponents](ALandscapeProxy* Proxy)
+	GetLandscapeInfo()->ForAllLandscapeComponents([&AllLandscapeComponents](ULandscapeComponent* InLandscapeComponent)
 	{
-		AllLandscapeComponents.Append(Proxy->LandscapeComponents);
+		AllLandscapeComponents.Add(InLandscapeComponent);
 	});
 
 	int32 ProcessedModes = 0;
diff --git a/Engine/Source/Runtime/Landscape/Private/LandscapeSplineRaster.cpp b/Engine/Source/Runtime/Landscape/Private/LandscapeSplineRaster.cpp
--- a/Engine/Source/Runtime/Landscape/Private/LandscapeSplineRaster.cpp
+++ b/Engine/Source/Runtime/Landscape/Private/LandscapeSplineRaster.cpp
@@ -476,7 +476,7 @@ bool ULandscapeInfo::ApplySplines(bool bOnlySelected, TSet<ULandscapeComponent*>
 
 bool ULandscapeInfo::ApplySplinesInternal(bool bOnlySelected, ALandscapeProxy* Proxy, TSet<ULandscapeComponent*>* OutModifiedComponents)
 {
-	if (!Proxy || !Proxy->SplineComponent || Proxy->SplineComponent->ControlPoints.Num() == 0 || Proxy->SplineComponent->Segments.Num() == 0)
+	if (!Proxy || !Proxy->SplineComponent || !Proxy->SplineComponent->IsRegistered() || Proxy->SplineComponent->ControlPoints.Num() == 0 || Proxy->SplineComponent->Segments.Num() == 0)
 	{
 		return false;
 	}
