Landscape Layers: Windows 7 / DX11.0  UAV Fix

#rnx

diff --git a/Engine/Shaders/Private/LandscapeLayersCS.usf b/Engine/Shaders/Private/LandscapeLayersCS.usf
--- a/Engine/Shaders/Private/LandscapeLayersCS.usf
+++ b/Engine/Shaders/Private/LandscapeLayersCS.usf
@@ -41,7 +41,7 @@ StructuredBuffer<FLandscapeLayerWeightmapPackMaterialLayersComponentData> InPack
 Buffer<float> InWeightmapWeightBlendMode;
 Buffer<float2> InWeightmapTextureOutputOffset;
 Texture2DArray<float> InAtlasPaintLayers;
-RWTexture2D<UNORM float4> OutComponentWeightMaps;
+RWTexture2D<uint> OutComponentWeightMaps;
 
 float ComputeWeightmapChannel(int ChannelIndexToProcess, uint2 TexelPosition, uint SourcePaintLayerIndex)
 {
@@ -77,27 +77,28 @@ void PackPaintLayerToWeightmap(uint3 DispatchThreadId : SV_DispatchThreadID)
 	{
 		FLandscapeLayerWeightmapPackMaterialLayersComponentData ComponentData = InPackLayersComponentsData[DispatchThreadId.z];
 
-		float4 Channels = 0;
+		uint Channels = 0;
 
 		for (int i = 0; i < 4; ++i)
 		{
 			float ChannelResult = ComputeWeightmapChannel(ComponentData.WeightmapChannelToProcess[i], DispatchThreadId.xy + int2(ComponentData.ComponenVertexPositionX[i], ComponentData.ComponenVertexPositionY[i]), ComponentData.SourcePaintLayerIndex[i]);
 
+			// BGRA to RGBA
 			if (ComponentData.WeightmapChannelToProcess[i] == 0)
 			{ 
-				Channels.x = ChannelResult;
+				Channels |= uint(ChannelResult * 255.0) << 8;
 			}
 			else if (ComponentData.WeightmapChannelToProcess[i] == 1)
 			{
-				Channels.y = ChannelResult;
+				Channels |= uint(ChannelResult * 255.0) << 16;
 			}
 			else if (ComponentData.WeightmapChannelToProcess[i] == 2)
 			{
-				Channels.z = ChannelResult;
+				Channels |= uint(ChannelResult * 255.0) << 24;
 			}
 			else if (ComponentData.WeightmapChannelToProcess[i] == 3)
 			{
-				Channels.w = ChannelResult;
+				Channels |= uint(ChannelResult * 255.0);
 			}
 		}
 
diff --git a/Engine/Shaders/Private/LandscapeLayersPS.usf b/Engine/Shaders/Private/LandscapeLayersPS.usf
--- a/Engine/Shaders/Private/LandscapeLayersPS.usf
+++ b/Engine/Shaders/Private/LandscapeLayersPS.usf
@@ -347,6 +347,15 @@ void PSHeightmapMainMips(float2 InTextureCoordinates : TEXCOORD0, out float4 Out
 
 float4 WeightmapLayerWeightBlend; // each xyzw == 1 texture channel
 
+void PSWeightmapConvertFormat(float2 InTextureCoordinates : TEXCOORD0, out float4 OutColor : SV_Target0)
+{
+	uint SourceColor = asuint(ReadTexture1.SampleLevel(ReadTexture1Sampler, InTextureCoordinates, 0).r);
+	OutColor.b = float((SourceColor >> 24) & 0x000000FF) / 255.0;
+	OutColor.g = float((SourceColor >> 16) & 0x000000FF) / 255.0;
+	OutColor.r = float((SourceColor >> 8)  & 0x000000FF) / 255.0;
+	OutColor.a = float(SourceColor & 0x000000FF) / 255.0;
+}
+
 void PSWeightmapMain(float2 InTextureCoordinates : TEXCOORD0, out float4 OutColor : SV_Target0)
 {
 	float4 SourceColor = ReadTexture1.Sample(ReadTexture1Sampler, InTextureCoordinates);
diff --git a/Engine/Source/Runtime/Landscape/Classes/Landscape.h b/Engine/Source/Runtime/Landscape/Classes/Landscape.h
--- a/Engine/Source/Runtime/Landscape/Classes/Landscape.h
+++ b/Engine/Source/Runtime/Landscape/Classes/Landscape.h
@@ -393,6 +393,9 @@ private:
 	void PrintLayersDebugTextureResource(const FString& InContext, FTextureResource* InTextureResource, uint8 InMipRender = 0, bool InOutputHeight = true, bool InOutputNormals = false) const;
 	void PrintLayersDebugHeightData(const FString& InContext, const TArray<FColor>& InHeightmapData, const FIntPoint& InDataSize, uint8 InMipRender, bool InOutputNormals = false) const;
 	void PrintLayersDebugWeightData(const FString& InContext, const TArray<FColor>& InWeightmapData, const FIntPoint& InDataSize, uint8 InMipRender) const;
+
+	// Needed because of UAVs on DX11.0 not supporting BGRA8
+	void ConvertR32ToBGRA8(const TArray<FVector2D>& InWeightmapTextureOutputOffset, UTextureRenderTarget2D* InWeightmapRTWrite, const struct FLandscapeLayersWeightmapConvertFormatShaderParameter& InShaderParams);
 #endif
 
 public:
diff --git a/Engine/Source/Runtime/Landscape/Private/LandscapeEditLayers.cpp b/Engine/Source/Runtime/Landscape/Private/LandscapeEditLayers.cpp
--- a/Engine/Source/Runtime/Landscape/Private/LandscapeEditLayers.cpp
+++ b/Engine/Source/Runtime/Landscape/Private/LandscapeEditLayers.cpp
@@ -67,6 +67,126 @@
 
 DECLARE_GPU_STAT_NAMED(LandscapeLayersRender, TEXT("Landscape Layer System Render"));
 
+// Custom Resources
+
+class FLandscapeTexture2DResource : public FTextureResource
+{
+public:
+	FLandscapeTexture2DResource(uint32 InSizeX, uint32 InSizeY, EPixelFormat InFormat, uint32 InNumMips, bool InNeedUAV)
+		: SizeX(InSizeX)
+		, SizeY(InSizeY)
+		, Format(InFormat)
+		, NumMips(InNumMips)
+		, CreateUAV(InNeedUAV)
+	{}
+
+	virtual uint32 GetSizeX() const override
+	{
+		return SizeX;
+	}
+
+	virtual uint32 GetSizeY() const override
+	{
+		return SizeY;
+	}
+
+	/** Called when the resource is initialized. This is only called by the rendering thread. */
+	virtual void InitRHI() override
+	{
+		FTextureResource::InitRHI();
+
+		FRHIResourceCreateInfo CreateInfo;
+		uint32 Flags = TexCreate_NoTiling | TexCreate_OfflineProcessed;
+
+		if (CreateUAV)
+		{
+			Flags |= TexCreate_UAV;
+		}
+
+		TextureRHI = RHICreateTexture2D(SizeX, SizeY, Format, NumMips, 1, Flags, CreateInfo);
+
+		if (CreateUAV)
+		{
+			TextureUAV = RHICreateUnorderedAccessView(TextureRHI, 0);
+		}
+	}
+
+	FUnorderedAccessViewRHIRef TextureUAV;
+
+private:
+	uint32 SizeX;
+	uint32 SizeY;
+	EPixelFormat Format;
+	uint32 NumMips;
+	bool CreateUAV;
+};
+
+class FLandscapeTexture2DArrayResource : public FTextureResource
+{
+public:
+	FLandscapeTexture2DArrayResource(uint32 InSizeX, uint32 InSizeY, uint32 InSizeZ, EPixelFormat InFormat, uint32 InNumMips, bool InNeedUAV)
+		: SizeX(InSizeX)
+		, SizeY(InSizeY)
+		, SizeZ(InSizeZ)
+		, Format(InFormat)
+		, NumMips(InNumMips)
+		, CreateUAV(InNeedUAV)
+	{}
+
+	virtual uint32 GetSizeX() const override
+	{
+		return SizeX;
+	}
+
+	virtual uint32 GetSizeY() const override
+	{
+		return SizeY;
+	}
+
+	virtual uint32 GetSizeZ() const
+	{
+		return SizeZ;
+	}
+
+	/** Called when the resource is initialized. This is only called by the rendering thread. */
+	virtual void InitRHI() override
+	{
+		FTextureResource::InitRHI();
+
+		FRHIResourceCreateInfo CreateInfo;
+		uint32 Flags = TexCreate_NoTiling | TexCreate_OfflineProcessed;
+
+		if (CreateUAV)
+		{
+			Flags |= TexCreate_UAV;
+		}
+
+		TextureRHI = RHICreateTexture2DArray(SizeX, SizeY, SizeZ, Format, NumMips, 1, Flags, CreateInfo);
+
+		if (CreateUAV)
+		{
+			TextureUAV = RHICreateUnorderedAccessView(TextureRHI, 0);
+		}
+	}
+
+	virtual void ReleaseRHI() override
+	{
+		FTextureResource::ReleaseRHI();
+
+		TextureUAV.SafeRelease();
+	}
+
+	FUnorderedAccessViewRHIRef TextureUAV;
+
+private:
+	uint32 SizeX;
+	uint32 SizeY;
+	uint32 SizeZ;
+	EPixelFormat Format;
+	uint32 NumMips;
+	bool CreateUAV;
+};
+
 // Vertex format and vertex buffer
 
 struct FLandscapeLayersVertex
@@ -500,126 +620,60 @@
 
 IMPLEMENT_GLOBAL_SHADER(FLandscapeLayersWeightmapMipsPS, "/Engine/Private/LandscapeLayersPS.usf", "PSWeightmapMainMips", SF_Pixel);
 
-// Custom Resources
-
-class FLandscapeTexture2DResource : public FTextureResource
+struct FLandscapeLayersWeightmapConvertFormatShaderParameter
 {
-public:
-	FLandscapeTexture2DResource(uint32 InSizeX, uint32 InSizeY, EPixelFormat InFormat, uint32 InNumMips, bool InNeedUAV)
-		: SizeX(InSizeX)
-		, SizeY(InSizeY)
-		, Format(InFormat)
-		, NumMips(InNumMips)
-		, CreateUAV(InNeedUAV)
+	FLandscapeLayersWeightmapConvertFormatShaderParameter()
+		: ReadWeightmap(nullptr)
 	{}
 
-	virtual uint32 GetSizeX() const override
+	FLandscapeTexture2DResource* ReadWeightmap;
+};
+
+// This Shader exists only to support Win7 because DX11.0 doesn't support UAVs on BGRA8 Textures.
+// The Shader will convert from R32 Format to BGRA8 Format
+class FLandscapeLayersWeightmapConvertFormatPS : public FGlobalShader
+{
+	DECLARE_GLOBAL_SHADER(FLandscapeLayersWeightmapConvertFormatPS);
+public:
+
+	static bool ShouldCompilePermutation(const FGlobalShaderPermutationParameters& Parameters)
 	{
-		return SizeX;
+		return IsFeatureLevelSupported(Parameters.Platform, ERHIFeatureLevel::SM5) && !IsConsolePlatform(Parameters.Platform) && !IsMetalPlatform(Parameters.Platform);
 	}
 
-	virtual uint32 GetSizeY() const override
+	static void ModifyCompilationEnvironment(const FGlobalShaderPermutationParameters& Parameters, FShaderCompilerEnvironment& OutEnvironment)
 	{
-		return SizeY;
+	}
+	FLandscapeLayersWeightmapConvertFormatPS(const ShaderMetaType::CompiledShaderInitializerType& Initializer)
+		: FGlobalShader(Initializer)
+	{
+		ReadTexture1Param.Bind(Initializer.ParameterMap, TEXT("ReadTexture1"));
+		ReadTexture1SamplerParam.Bind(Initializer.ParameterMap, TEXT("ReadTexture1Sampler"));
 	}
 
-	/** Called when the resource is initialized. This is only called by the rendering thread. */
-	virtual void InitRHI() override
+	FLandscapeLayersWeightmapConvertFormatPS()
+	{}
+
+	void SetParameters(FRHICommandList& RHICmdList, const FLandscapeLayersWeightmapConvertFormatShaderParameter& InParams)
 	{
-		FTextureResource::InitRHI();
-
-		FRHIResourceCreateInfo CreateInfo;
-		uint32 Flags = TexCreate_NoTiling | TexCreate_OfflineProcessed;
-
-		if (CreateUAV)
-		{
-			Flags |= TexCreate_UAV;
-		}
-
-		TextureRHI = RHICreateTexture2D(SizeX, SizeY, Format, NumMips, 1, Flags, CreateInfo);
-
-		if (CreateUAV)
-		{
-			TextureUAV = RHICreateUnorderedAccessView(TextureRHI, 0);
-		}
+		SetTextureParameter(RHICmdList, GetPixelShader(), ReadTexture1Param, ReadTexture1SamplerParam, TStaticSamplerState<SF_Point, AM_Clamp, AM_Clamp, AM_Clamp>::GetRHI(), InParams.ReadWeightmap->TextureRHI);
 	}
 
-	FUnorderedAccessViewRHIRef TextureUAV;
+	virtual bool Serialize(FArchive& Ar) override
+	{
+		bool bShaderHasOutdatedParameters = FShader::Serialize(Ar);
+		Ar << ReadTexture1Param;
+		Ar << ReadTexture1SamplerParam;
 
-private:
-	uint32 SizeX;
-	uint32 SizeY;
-	EPixelFormat Format;
-	uint32 NumMips;
-	bool CreateUAV;
+		return bShaderHasOutdatedParameters;
+	}
+
+	FShaderResourceParameter ReadTexture1Param;
+	FShaderResourceParameter ReadTexture1SamplerParam;
 };
 
-class FLandscapeTexture2DArrayResource : public FTextureResource
-{
-public:
-	FLandscapeTexture2DArrayResource(uint32 InSizeX, uint32 InSizeY, uint32 InSizeZ, EPixelFormat InFormat, uint32 InNumMips, bool InNeedUAV)
-		: SizeX(InSizeX)
-		, SizeY(InSizeY)
-		, SizeZ(InSizeZ)
-		, Format(InFormat)
-		, NumMips(InNumMips)
-		, CreateUAV(InNeedUAV)
-	{}
+IMPLEMENT_GLOBAL_SHADER(FLandscapeLayersWeightmapConvertFormatPS, "/Engine/Private/LandscapeLayersPS.usf", "PSWeightmapConvertFormat", SF_Pixel);
 
-	virtual uint32 GetSizeX() const override
-	{
-		return SizeX;
-	}
-
-	virtual uint32 GetSizeY() const override
-	{
-		return SizeY;
-	}
-
-	virtual uint32 GetSizeZ() const
-	{
-		return SizeZ;
-	}
-
-	/** Called when the resource is initialized. This is only called by the rendering thread. */
-	virtual void InitRHI() override
-	{
-		FTextureResource::InitRHI();
-
-		FRHIResourceCreateInfo CreateInfo;
-		uint32 Flags = TexCreate_NoTiling | TexCreate_OfflineProcessed;
-
-		if (CreateUAV)
-		{
-			Flags |= TexCreate_UAV;
-		}
-
-		TextureRHI = RHICreateTexture2DArray(SizeX, SizeY, SizeZ, Format, NumMips, 1, Flags, CreateInfo);
-
-		if (CreateUAV)
-		{
-			TextureUAV = RHICreateUnorderedAccessView(TextureRHI, 0);
-		}
-	}
-
-	virtual void ReleaseRHI() override
-	{
-		FTextureResource::ReleaseRHI();
-
-		TextureUAV.SafeRelease();
-	}
-
-	FUnorderedAccessViewRHIRef TextureUAV;
-
-private:
-	uint32 SizeX;
-	uint32 SizeY;
-	uint32 SizeZ;
-	EPixelFormat Format;
-	uint32 NumMips;
-	bool CreateUAV;
-};
-
 // Compute shaders data
 
 int32 GLandscapeLayerWeightmapThreadGroupSizeX = 16;
@@ -1217,6 +1271,8 @@
 
 typedef FLandscapeLayersRender_RenderThread<FLandscapeLayersHeightmapShaderParameters, FLandscapeLayersHeightmapPS, FLandscapeLayersHeightmapMipsPS> FLandscapeLayersHeightmapRender_RenderThread;
 typedef FLandscapeLayersRender_RenderThread<FLandscapeLayersWeightmapShaderParameters, FLandscapeLayersWeightmapPS, FLandscapeLayersWeightmapMipsPS> FLandscapeLayersWeightmapRender_RenderThread;
+// The Mips Param will not be used
+typedef FLandscapeLayersRender_RenderThread<FLandscapeLayersWeightmapConvertFormatShaderParameter, FLandscapeLayersWeightmapConvertFormatPS, FLandscapeLayersWeightmapConvertFormatPS> FLandscapeLayersWeightmapConvertRender_RenderThread;
 
 #if WITH_EDITOR
 
@@ -1967,6 +2023,38 @@
 	}
 }
 
+void ALandscape::ConvertR32ToBGRA8(const TArray<FVector2D>& InWeightmapTextureOutputOffset, UTextureRenderTarget2D* InWeightmapRTWrite, const FLandscapeLayersWeightmapConvertFormatShaderParameter& InShaderParams)
+{
+	check(InShaderParams.ReadWeightmap != nullptr);
+	check(InWeightmapRTWrite != nullptr);
+
+	FIntPoint WeightmapWriteTextureSize(InWeightmapRTWrite->SizeX, InWeightmapRTWrite->SizeY);
+	FIntPoint WeightmapReadTextureSize(InShaderParams.ReadWeightmap->GetSizeX(), InShaderParams.ReadWeightmap->GetSizeY());
+
+	// Quad Setup
+	const int32 LocalComponentSizeQuad = SubsectionSizeQuads * NumSubsections;
+	const int32 LocalComponentSizeVerts = (SubsectionSizeQuads + 1) * NumSubsections;
+
+	TArray<FLandscapeLayersTriangle> TriangleList;
+	TriangleList.Reserve(InWeightmapTextureOutputOffset.Num() * 2 * NumSubsections);
+	for (const FVector2D& TexturePosition : InWeightmapTextureOutputOffset)
+	{
+		FVector2D PositionOffset(FMath::RoundToInt(TexturePosition.X / LocalComponentSizeVerts), FMath::RoundToInt(TexturePosition.Y / LocalComponentSizeVerts));
+		FIntPoint SectionBase(PositionOffset.X * LocalComponentSizeQuad, PositionOffset.Y * LocalComponentSizeQuad);
+		GenerateLayersRenderQuadsAtlas(SectionBase, FVector2D::ZeroVector, SubsectionSizeQuads, WeightmapReadTextureSize, WeightmapWriteTextureSize, TriangleList);
+	}
+
+	FMatrix ProjectionMatrix = AdjustProjectionMatrixForRHI(FTranslationMatrix(FVector(0, 0, 0)) *
+		FMatrix(FPlane(1.0f / (FMath::Max<uint32>(WeightmapWriteTextureSize.X, 1.f) / 2.0f), 0.0, 0.0f, 0.0f), FPlane(0.0f, -1.0f / (FMath::Max<uint32>(WeightmapWriteTextureSize.Y, 1.f) / 2.0f), 0.0f, 0.0f), FPlane(0.0f, 0.0f, 1.0f, 0.0f), FPlane(-1.0f, 1.0f, 0.0f, 1.0f)));
+
+	FLandscapeLayersWeightmapConvertRender_RenderThread LayersRender(FString(TEXT("ConvertR32ToBGRA8")), InWeightmapRTWrite, WeightmapWriteTextureSize, WeightmapReadTextureSize, ProjectionMatrix, InShaderParams, 0, TriangleList);
+	ENQUEUE_RENDER_COMMAND(FConvertR32ToBGRA8Command)(
+		[LayersRender](FRHICommandListImmediate& RHICmdList) mutable
+	{
+		LayersRender.Render(RHICmdList, true);
+	});
+}
+
 void ALandscape::ClearLayersWeightmapTextureResource(const FString& InDebugName, FTextureRenderTargetResource* InTextureResourceToClear) const
 {
 	LandscapeLayersWeightmapClear_RenderThread LayersClear(InDebugName, InTextureResourceToClear);
@@ -3344,7 +3432,7 @@
 		}
 	}
 
-	WeightmapScratchPackLayerTextureResource = new FLandscapeTexture2DResource(FirstWeightmapRT->SizeX, FirstWeightmapRT->SizeY, PF_B8G8R8A8, MipCount, true);
+	WeightmapScratchPackLayerTextureResource = new FLandscapeTexture2DResource(FirstWeightmapRT->SizeX, FirstWeightmapRT->SizeY, PF_R32_UINT, MipCount, true);
 	BeginInitResource(WeightmapScratchPackLayerTextureResource);
 }
 
@@ -3763,7 +3851,12 @@
 
 					SourceDebugName = OutputDebugName ? TEXT("WeightmapScratchTexture") : GEmptyDebugName;
 					DestDebugName = OutputDebugName ? CurrentRT->GetName() : GEmptyDebugName;
-					CopyLayersTexture(SourceDebugName, WeightmapScratchPackLayerTextureResource, DestDebugName, CurrentRT->GameThread_GetRenderTargetResource());
+					
+					// Convert Back to BGRA8(Win7)
+					FLandscapeLayersWeightmapConvertFormatShaderParameter ShaderParams;
+					ShaderParams.ReadWeightmap = WeightmapScratchPackLayerTextureResource;
+					ConvertR32ToBGRA8(WeightmapTextureOutputOffset, CurrentRT, ShaderParams);
+					
 					DrawWeightmapComponentToRenderTargetMips(WeightmapTextureOutputOffset, CurrentRT, true, PSShaderParams);
 
 					int32 StartTextureIndex = NextTextureIndexToProcess;
