Fix Nan Bounds when using Neighbor Proxy.
Proxy was used instead of Component in CL 6596037 but it seems it is possible to access a neighbor proxy thats transform hasn't been initialized.

#rnx

diff --git a/Engine/Source/Runtime/Landscape/Private/LandscapeRender.cpp b/Engine/Source/Runtime/Landscape/Private/LandscapeRender.cpp
--- a/Engine/Source/Runtime/Landscape/Private/LandscapeRender.cpp
+++ b/Engine/Source/Runtime/Landscape/Private/LandscapeRender.cpp
@@ -2365,8 +2365,8 @@ float FLandscapeComponentSceneProxy::GetNeighborLOD(const FSceneView& InView, fl
 			if (NeighborComponent != nullptr && NeighborComponent->SceneProxy != nullptr)
 			{
 				NeighborSceneProxy = (FLandscapeComponentSceneProxy*)NeighborComponent->SceneProxy;
-				NeighborBounds = NeighborSceneProxy->GetBounds();
-				NeighborMaxExtends = NeighborSceneProxy->SubsectionSizeQuads * FMath::Max(NeighborSceneProxy->GetLocalToWorld().GetScaleVector().X, NeighborSceneProxy->GetLocalToWorld().GetScaleVector().Y);
+				NeighborBounds = NeighborComponent->Bounds;
+				NeighborMaxExtends = NeighborSceneProxy->SubsectionSizeQuads * FMath::Max(NeighborComponent->GetRenderMatrix().GetScaleVector().X, NeighborComponent->GetRenderMatrix().GetScaleVector().Y);
 			}
 		}
 
