Landscape Layers: Release UAV RHI Resource in ReleaseRHI call

#rnx

diff --git a/Engine/Source/Runtime/Landscape/Private/LandscapeEditLayers.cpp b/Engine/Source/Runtime/Landscape/Private/LandscapeEditLayers.cpp
--- a/Engine/Source/Runtime/Landscape/Private/LandscapeEditLayers.cpp
+++ b/Engine/Source/Runtime/Landscape/Private/LandscapeEditLayers.cpp
@@ -111,6 +111,16 @@ public:
 		}
 	}
 
+	virtual void ReleaseRHI() override
+	{
+		if (CreateUAV)
+		{
+			TextureUAV.SafeRelease();
+		}
+		
+		FTextureResource::ReleaseRHI();
+	}
+
 	FUnorderedAccessViewRHIRef TextureUAV;
 
 private:
