Modified Landscape Layer Contribution for heightmap to show a gradient value instead of a constant.

diff --git a/Engine/Source/Runtime/Landscape/Private/LandscapeEditLayers.cpp b/Engine/Source/Runtime/Landscape/Private/LandscapeEditLayers.cpp
--- a/Engine/Source/Runtime/Landscape/Private/LandscapeEditLayers.cpp
+++ b/Engine/Source/Runtime/Landscape/Private/LandscapeEditLayers.cpp
@@ -4514,7 +4514,8 @@ bool ALandscape::UpdateCollisionAndClients(const TArray<ULandscapeComponent*>& I
 	bool bAllClientsUpdated = true;
 
 	const uint16 DefaultHeightValue = LandscapeDataAccess::GetTexHeight(0.f);
-	const uint8 LayerContributingValue = UINT8_MAX;
+	const uint8 MaxLayerContributingValue = UINT8_MAX;
+	const float HeightValueNormalizationFactor = 1.f / (0.5f * UINT16_MAX);
 	TArray<uint16> HeightData;
 	TArray<uint8> LayerContributionMaskData;
 
@@ -4582,7 +4583,7 @@ bool ALandscape::UpdateCollisionAndClients(const TArray<ULandscapeComponent*>& I
 						LandscapeEdit.GetHeightDataFast(X1, Y1, X2, Y2, HeightData.GetData(), Stride);
 						for (int i = 0; i < ArraySize; ++i)
 						{
-							LayerContributionMaskData[i] = HeightData[i] != DefaultHeightValue ? LayerContributingValue : 0;
+							LayerContributionMaskData[i] = HeightData[i] != DefaultHeightValue ? (uint8)(FMath::Pow(FMath::Clamp((HeightValueNormalizationFactor*FMath::Abs(HeightData[i] - DefaultHeightValue)), 0.f, (float)1.f), 0.25f)*MaxLayerContributingValue) : 0;
 						}
 						bLayerContributionWrittenData = true;
 					}
