Added LeftSideFalloffFactor & RightSideFalloffFactor on Landscape spline control points to allow different side falloffs

diff --git a/Engine/Source/Editor/LandscapeEditor/Private/LandscapeEdModeSplineTools.cpp b/Engine/Source/Editor/LandscapeEditor/Private/LandscapeEdModeSplineTools.cpp
--- a/Engine/Source/Editor/LandscapeEditor/Private/LandscapeEdModeSplineTools.cpp
+++ b/Engine/Source/Editor/LandscapeEditor/Private/LandscapeEdModeSplineTools.cpp
@@ -463,6 +463,8 @@ public:
 
 			NewControlPoint->Width       = FirstPoint->Width;
 			NewControlPoint->SideFalloff = FirstPoint->SideFalloff;
+			NewControlPoint->LeftSideFalloffFactor = FirstPoint->LeftSideFalloffFactor;
+			NewControlPoint->RightSideFalloffFactor = FirstPoint->RightSideFalloffFactor;
 			NewControlPoint->EndFalloff  = FirstPoint->EndFalloff;
 
 			if (bCopyMeshToNewControlPoint)
@@ -669,6 +671,8 @@ public:
 		NewControlPoint->Width = FMath::Lerp(Segment->Connections[0].ControlPoint->Width, Segment->Connections[1].ControlPoint->Width, t);
 		NewControlPoint->SideFalloff = FMath::Lerp(Segment->Connections[0].ControlPoint->SideFalloff, Segment->Connections[1].ControlPoint->SideFalloff, t);
 		NewControlPoint->EndFalloff = FMath::Lerp(Segment->Connections[0].ControlPoint->EndFalloff, Segment->Connections[1].ControlPoint->EndFalloff, t);
+		NewControlPoint->LeftSideFalloffFactor = FMath::Clamp(FMath::Lerp(Segment->Connections[0].ControlPoint->LeftSideFalloffFactor, Segment->Connections[1].ControlPoint->LeftSideFalloffFactor, t), 0.f, 1.f);
+		NewControlPoint->RightSideFalloffFactor = FMath::Clamp(FMath::Lerp(Segment->Connections[0].ControlPoint->RightSideFalloffFactor, Segment->Connections[1].ControlPoint->RightSideFalloffFactor, t), 0.f, 1.f);
 
 		ULandscapeSplineSegment* NewSegment = NewObject<ULandscapeSplineSegment>(SplinesComponent, NAME_None, RF_Transactional);
 		SplinesComponent->Segments.Add(NewSegment);
@@ -758,6 +762,8 @@ public:
 		//ControlPoint->Rotation.Roll = FMath::Lerp(UseSegment->Connections[0].ControlPoint->Rotation.Roll, UseSegment->Connections[1].ControlPoint->Rotation.Roll, tseg);
 		ControlPoint->Width = FMath::Lerp(UseSegment->Connections[0].ControlPoint->Width, UseSegment->Connections[1].ControlPoint->Width, tseg);
 		ControlPoint->SideFalloff = FMath::Lerp(UseSegment->Connections[0].ControlPoint->SideFalloff, UseSegment->Connections[1].ControlPoint->SideFalloff, tseg);
+		ControlPoint->LeftSideFalloffFactor = FMath::Clamp(FMath::Lerp(UseSegment->Connections[0].ControlPoint->LeftSideFalloffFactor, UseSegment->Connections[1].ControlPoint->LeftSideFalloffFactor, tseg), 0.f, 1.f);
+		ControlPoint->RightSideFalloffFactor = FMath::Clamp(FMath::Lerp(UseSegment->Connections[0].ControlPoint->RightSideFalloffFactor, UseSegment->Connections[1].ControlPoint->RightSideFalloffFactor, tseg), 0.f, 1.f);
 		ControlPoint->EndFalloff = FMath::Lerp(UseSegment->Connections[0].ControlPoint->EndFalloff, UseSegment->Connections[1].ControlPoint->EndFalloff, tseg);
 
 		Segment->Connections[0].TangentLen = DuplicateCacheSplitSegmentTangentLenStart * t;
diff --git a/Engine/Source/Runtime/Landscape/Classes/LandscapeSplineControlPoint.h b/Engine/Source/Runtime/Landscape/Classes/LandscapeSplineControlPoint.h
--- a/Engine/Source/Runtime/Landscape/Classes/LandscapeSplineControlPoint.h
+++ b/Engine/Source/Runtime/Landscape/Classes/LandscapeSplineControlPoint.h
@@ -68,6 +68,12 @@ class ULandscapeSplineControlPoint : public UObject
 	UPROPERTY(EditAnywhere, Category=LandscapeSpline)
 	float SideFalloff;
 
+	UPROPERTY(EditAnywhere, Category = LandscapeSpline, meta=(UIMin = 0, ClampMin = 0, UIMax = 1, ClampMax = 1))
+	float LeftSideFalloffFactor = 1.f;
+
+	UPROPERTY(EditAnywhere, Category = LandscapeSpline, meta = (UIMin = 0, ClampMin = 0, UIMax = 1, ClampMax = 1))
+	float RightSideFalloffFactor = 1.f;
+
 	/** Falloff at the start/end of the spline (if this point is a start or end point, otherwise ignored). */
 	UPROPERTY(EditAnywhere, Category=LandscapeSpline)
 	float EndFalloff;
diff --git a/Engine/Source/Runtime/Landscape/Private/LandscapeBlueprintSupport.cpp b/Engine/Source/Runtime/Landscape/Private/LandscapeBlueprintSupport.cpp
--- a/Engine/Source/Runtime/Landscape/Private/LandscapeBlueprintSupport.cpp
+++ b/Engine/Source/Runtime/Landscape/Private/LandscapeBlueprintSupport.cpp
@@ -20,7 +20,7 @@ void ALandscapeProxy::EditorApplySpline(USplineComponent* InSplineComponent, flo
 	if (InSplineComponent && !GetWorld()->IsGameWorld())
 	{
 		TArray<FLandscapeSplineInterpPoint> Points;
-		LandscapeSplineRaster::Pointify(InSplineComponent->SplineCurves.Position, Points, NumSubdivisions, 0.0f, 0.0f, StartWidth, EndWidth, StartSideFalloff, EndSideFalloff, StartRoll, EndRoll);
+		LandscapeSplineRaster::Pointify(InSplineComponent->SplineCurves.Position, Points, NumSubdivisions, 0.0f, 0.0f, StartWidth, EndWidth, StartSideFalloff, EndSideFalloff, StartSideFalloff, EndSideFalloff, StartRoll, EndRoll);
 
 		FTransform SplineToWorld = InSplineComponent->GetComponentTransform();
 		LandscapeSplineRaster::RasterizeSegmentPoints(GetLandscapeInfo(), MoveTemp(Points), SplineToWorld, bRaiseHeights, bLowerHeights, PaintLayer);
diff --git a/Engine/Source/Runtime/Landscape/Private/LandscapeSplineRaster.cpp b/Engine/Source/Runtime/Landscape/Private/LandscapeSplineRaster.cpp
--- a/Engine/Source/Runtime/Landscape/Private/LandscapeSplineRaster.cpp
+++ b/Engine/Source/Runtime/Landscape/Private/LandscapeSplineRaster.cpp
@@ -869,7 +869,8 @@ namespace LandscapeSplineRaster
 	void Pointify(const FInterpCurveVector& SplineInfo, TArray<FLandscapeSplineInterpPoint>& Points, int32 NumSubdivisions,
 		float StartFalloffFraction, float EndFalloffFraction,
 		const float StartWidth, const float EndWidth,
-		const float StartSideFalloff, const float EndSideFalloff,
+		const float StartLeftSideFalloff, const float EndLeftSideFalloff,
+		const float StartRightSideFalloff, const float EndRightSideFalloff,
 		const float StartRollDegrees, const float EndRollDegrees)
 	{
 		// Stop the start and end fall-off overlapping
@@ -889,15 +890,16 @@ namespace LandscapeSplineRaster
 			const float NewKeyTime = SplineInfo.Points[i].InVal;
 			const float NewKeyCosInterp = 0.5f - 0.5f * FMath::Cos(NewKeyTime * PI);
 			const float NewKeyWidth = FMath::Lerp(StartWidth, EndWidth, NewKeyCosInterp);
-			const float NewKeyFalloff = FMath::Lerp(StartSideFalloff, EndSideFalloff, NewKeyCosInterp);
+			const float NewKeyLeftFalloff = FMath::Lerp(StartLeftSideFalloff, EndLeftSideFalloff, NewKeyCosInterp);
+			const float NewKeyRightFalloff = FMath::Lerp(StartRightSideFalloff, EndRightSideFalloff, NewKeyCosInterp);
 			const float NewKeyRoll = FMath::Lerp(StartRoll, EndRoll, NewKeyCosInterp);
 			const FVector NewKeyPos = SplineInfo.Eval(NewKeyTime, FVector::ZeroVector);
 			const FVector NewKeyTangent = SplineInfo.EvalDerivative(NewKeyTime, FVector::ZeroVector).GetSafeNormal();
 			const FVector NewKeyBiNormal = FQuat(NewKeyTangent, -NewKeyRoll).RotateVector((NewKeyTangent ^ FVector(0, 0, -1)).GetSafeNormal());
 			const FVector NewKeyLeftPos = NewKeyPos - NewKeyBiNormal * NewKeyWidth;
 			const FVector NewKeyRightPos = NewKeyPos + NewKeyBiNormal * NewKeyWidth;
-			const FVector NewKeyFalloffLeftPos = NewKeyPos - NewKeyBiNormal * (NewKeyWidth + NewKeyFalloff);
-			const FVector NewKeyFalloffRightPos = NewKeyPos + NewKeyBiNormal * (NewKeyWidth + NewKeyFalloff);
+			const FVector NewKeyFalloffLeftPos = NewKeyPos - NewKeyBiNormal * (NewKeyWidth + NewKeyLeftFalloff);
+			const FVector NewKeyFalloffRightPos = NewKeyPos + NewKeyBiNormal * (NewKeyWidth + NewKeyRightFalloff);
 			const float NewKeyStartEndFalloff = FMath::Min((StartFalloffFraction > 0 ? NewKeyTime / StartFalloffFraction : 1.0f), (EndFalloffFraction > 0 ? (1 - NewKeyTime) / EndFalloffFraction : 1.0f));
 
 			// If not the first keypoint, interp from the last keypoint.
@@ -912,15 +914,16 @@ namespace LandscapeSplineRaster
 					const float NewTime = OldKeyTime + j*DrawSubstep;
 					const float NewCosInterp = 0.5f - 0.5f * FMath::Cos(NewTime * PI);
 					const float NewWidth = FMath::Lerp(StartWidth, EndWidth, NewCosInterp);
-					const float NewFalloff = FMath::Lerp(StartSideFalloff, EndSideFalloff, NewCosInterp);
+					const float NewLeftFalloff = FMath::Lerp(StartLeftSideFalloff, EndLeftSideFalloff, NewCosInterp);
+					const float NewRightFalloff = FMath::Lerp(StartRightSideFalloff, EndRightSideFalloff, NewCosInterp);
 					const float NewRoll = FMath::Lerp(StartRoll, EndRoll, NewCosInterp);
 					const FVector NewPos = SplineInfo.Eval(NewTime, FVector::ZeroVector);
 					const FVector NewTangent = SplineInfo.EvalDerivative(NewTime, FVector::ZeroVector).GetSafeNormal();
 					const FVector NewBiNormal = FQuat(NewTangent, -NewRoll).RotateVector((NewTangent ^ FVector(0, 0, -1)).GetSafeNormal());
 					const FVector NewLeftPos = NewPos - NewBiNormal * NewWidth;
 					const FVector NewRightPos = NewPos + NewBiNormal * NewWidth;
-					const FVector NewFalloffLeftPos = NewPos - NewBiNormal * (NewWidth + NewFalloff);
-					const FVector NewFalloffRightPos = NewPos + NewBiNormal * (NewWidth + NewFalloff);
+					const FVector NewFalloffLeftPos = NewPos - NewBiNormal * (NewWidth + NewLeftFalloff);
+					const FVector NewFalloffRightPos = NewPos + NewBiNormal * (NewWidth + NewRightFalloff);
 					const float NewStartEndFalloff = FMath::Min((StartFalloffFraction > 0 ? NewTime / StartFalloffFraction : 1.0f), (EndFalloffFraction > 0 ? (1 - NewTime) / EndFalloffFraction : 1.0f));
 
 					Points.Emplace(NewPos, NewLeftPos, NewRightPos, NewFalloffLeftPos, NewFalloffRightPos, NewStartEndFalloff);
diff --git a/Engine/Source/Runtime/Landscape/Private/LandscapeSplineRaster.h b/Engine/Source/Runtime/Landscape/Private/LandscapeSplineRaster.h
--- a/Engine/Source/Runtime/Landscape/Private/LandscapeSplineRaster.h
+++ b/Engine/Source/Runtime/Landscape/Private/LandscapeSplineRaster.h
@@ -16,7 +16,8 @@ namespace LandscapeSplineRaster
 	void Pointify(const FInterpCurveVector& SplineInfo, TArray<FLandscapeSplineInterpPoint>& OutPoints, int32 NumSubdivisions,
 		float StartFalloffFraction, float EndFalloffFraction,
 		const float StartWidth, const float EndWidth,
-		const float StartSideFalloff, const float EndSideFalloff,
+		const float StartLeftSideFalloff, const float EndLeftSideFalloff, 
+		const float StartRightSideFalloff, const float EndRightSideFalloff,
 		const float StartRollDegrees, const float EndRollDegrees);
 
 	void RasterizeSegmentPoints(ULandscapeInfo* LandscapeInfo, TArray<FLandscapeSplineInterpPoint> Points, const FTransform& SplineToWorld, bool bRaiseTerrain, bool bLowerTerrain, ULandscapeLayerInfoObject* LayerInfo);
diff --git a/Engine/Source/Runtime/Landscape/Private/LandscapeSplines.cpp b/Engine/Source/Runtime/Landscape/Private/LandscapeSplines.cpp
--- a/Engine/Source/Runtime/Landscape/Private/LandscapeSplines.cpp
+++ b/Engine/Source/Runtime/Landscape/Private/LandscapeSplines.cpp
@@ -1725,6 +1725,9 @@ void ULandscapeSplineControlPoint::UpdateSplinePoints(bool bUpdateCollision, boo
 		AutoSetConnections(false);
 	}
 
+	const float LeftSideFalloff = LeftSideFalloffFactor * SideFalloff;
+	const float RightSideFalloff = RightSideFalloffFactor * SideFalloff;
+
 	// Update "Points" array
 	if (Mesh != nullptr)
 	{
@@ -1740,8 +1743,8 @@ void ULandscapeSplineControlPoint::UpdateSplinePoints(bool bUpdateCollision, boo
 			const FVector BiNormal = FQuat(Tangent, -Roll).RotateVector((Tangent ^ FVector(0, 0, -1)).GetSafeNormal());
 			const FVector LeftPos = StartLocation - BiNormal * Width;
 			const FVector RightPos = StartLocation + BiNormal * Width;
-			const FVector FalloffLeftPos = StartLocation - BiNormal * (Width + SideFalloff);
-			const FVector FalloffRightPos = StartLocation + BiNormal * (Width + SideFalloff);
+			const FVector FalloffLeftPos = StartLocation - BiNormal * (Width + LeftSideFalloff);
+			const FVector FalloffRightPos = StartLocation + BiNormal * (Width + RightSideFalloff);
 
 			Points.Emplace(StartLocation, LeftPos, RightPos, FalloffLeftPos, FalloffRightPos, 1.0f);
 		}
@@ -1761,8 +1764,8 @@ void ULandscapeSplineControlPoint::UpdateSplinePoints(bool bUpdateCollision, boo
 		const FVector BiNormal = FQuat(Tangent, -Roll).RotateVector((Tangent ^ FVector(0, 0, -1)).GetSafeNormal());
 		const FVector LeftPos = StartLocation - BiNormal * Width;
 		const FVector RightPos = StartLocation + BiNormal * Width;
-		const FVector FalloffLeftPos = StartLocation - BiNormal * (Width + SideFalloff);
-		const FVector FalloffRightPos = StartLocation + BiNormal * (Width + SideFalloff);
+		const FVector FalloffLeftPos = StartLocation - BiNormal * (Width + LeftSideFalloff);
+		const FVector FalloffRightPos = StartLocation + BiNormal * (Width + RightSideFalloff);
 
 		Points.Emplace(StartLocation, LeftPos, RightPos, FalloffLeftPos, FalloffRightPos, 1.0f);
 	}
@@ -1868,6 +1871,8 @@ void ULandscapeSplineControlPoint::PostEditChangeProperty(FPropertyChangedEvent&
 
 	Width = FMath::Max(Width, 0.001f);
 	SideFalloff = FMath::Max(SideFalloff, 0.0f);
+	LeftSideFalloffFactor = FMath::Clamp(LeftSideFalloffFactor, 0.0f, 1.0f);
+	RightSideFalloffFactor = FMath::Clamp(RightSideFalloffFactor, 0.0f, 1.0f);
 	EndFalloff = FMath::Max(EndFalloff, 0.0f);
 
 	// Don't update splines when undoing, not only is it unnecessary and expensive,
@@ -2229,8 +2234,10 @@ void ULandscapeSplineSegment::UpdateSplinePoints(bool bUpdateCollision)
 	const float EndFalloffFraction = ((Connections[1].ControlPoint->ConnectedSegments.Num() > 1) ? 0 : (Connections[1].ControlPoint->EndFalloff / SplineLength));
 	const float StartWidth = Connections[0].ControlPoint->Width;
 	const float EndWidth = Connections[1].ControlPoint->Width;
-	const float StartSideFalloff = Connections[0].ControlPoint->SideFalloff;
-	const float EndSideFalloff = Connections[1].ControlPoint->SideFalloff;
+	const float StartLeftSideFalloff = Connections[0].ControlPoint->LeftSideFalloffFactor * Connections[0].ControlPoint->SideFalloff;
+	const float EndLeftSideFalloff = Connections[1].ControlPoint->LeftSideFalloffFactor * Connections[1].ControlPoint->SideFalloff;
+	const float StartRightSideFalloff = Connections[0].ControlPoint->RightSideFalloffFactor * Connections[0].ControlPoint->SideFalloff;
+	const float EndRightSideFalloff = Connections[1].ControlPoint->RightSideFalloffFactor * Connections[1].ControlPoint->SideFalloff;
 	const float StartRollDegrees = StartRotation.Roll * (Connections[0].TangentLen > 0 ? 1 : -1);
 	const float EndRollDegrees = EndRotation.Roll * (Connections[1].TangentLen > 0 ? -1 : 1);
 	const float StartRoll = FMath::DegreesToRadians(StartRollDegrees);
@@ -2241,7 +2248,7 @@ void ULandscapeSplineSegment::UpdateSplinePoints(bool bUpdateCollision)
 	int32 NumPoints = FMath::CeilToInt(SplineLength / OuterSplines->SplineResolution);
 	NumPoints = FMath::Clamp(NumPoints, 1, 1000);
 
-	LandscapeSplineRaster::Pointify(SplineInfo, Points, NumPoints, StartFalloffFraction, EndFalloffFraction, StartWidth, EndWidth, StartSideFalloff, EndSideFalloff, StartRollDegrees, EndRollDegrees);
+	LandscapeSplineRaster::Pointify(SplineInfo, Points, NumPoints, StartFalloffFraction, EndFalloffFraction, StartWidth, EndWidth, StartLeftSideFalloff, EndLeftSideFalloff, StartRightSideFalloff, EndRightSideFalloff, StartRollDegrees, EndRollDegrees);
 
 	// Update Bounds
 	Bounds = FBox(ForceInit);
