Fix Package Dirty on Load because of Landscape Layers/Brushes

#rnx

diff --git a/Engine/Source/Runtime/Engine/Private/Components/SplineMeshComponent.cpp b/Engine/Source/Runtime/Engine/Private/Components/SplineMeshComponent.cpp
--- a/Engine/Source/Runtime/Engine/Private/Components/SplineMeshComponent.cpp
+++ b/Engine/Source/Runtime/Engine/Private/Components/SplineMeshComponent.cpp
@@ -1018,7 +1018,8 @@ void USplineMeshComponent::RecreateCollision()
 		}
 		else
 		{
-			BodySetup->Modify();
+			const bool bDirtyPackage = false;
+			BodySetup->Modify(bDirtyPackage);
 			BodySetup->InvalidatePhysicsData();
 			BodySetup->CopyBodyPropertiesFrom(GetStaticMesh()->BodySetup);
 			BodySetup->CollisionTraceFlag = GetStaticMesh()->BodySetup->CollisionTraceFlag;
diff --git a/Engine/Source/Runtime/Landscape/Private/Landscape.cpp b/Engine/Source/Runtime/Landscape/Private/Landscape.cpp
--- a/Engine/Source/Runtime/Landscape/Private/Landscape.cpp
+++ b/Engine/Source/Runtime/Landscape/Private/Landscape.cpp
@@ -2810,7 +2810,8 @@ void ULandscapeInfo::RegisterActor(ALandscapeProxy* Proxy, bool bMapCheck)
 
 	if (LandscapeActor && LandscapeActor->HasLayersContent())
 	{
-		LandscapeActor->RequestLayersInitialization();		
+		const bool bInRequestContentUpdate = false;
+		LandscapeActor->RequestLayersInitialization(bInRequestContentUpdate);
 	}
 
 	UpdateLayerInfoMap(Proxy);
diff --git a/Engine/Source/Runtime/Landscape/Private/LandscapeEdit.cpp b/Engine/Source/Runtime/Landscape/Private/LandscapeEdit.cpp
--- a/Engine/Source/Runtime/Landscape/Private/LandscapeEdit.cpp
+++ b/Engine/Source/Runtime/Landscape/Private/LandscapeEdit.cpp
@@ -1316,7 +1316,10 @@ void ULandscapeComponent::UpdateCollisionLayerData(const FColor* const* const We
 
 	if (CollisionComp)
 	{
-		CollisionComp->Modify();
+		if (!Proxy->HasLayersContent())
+		{
+			CollisionComp->Modify();
+		}
 
 		// Simple collision is not currently supported with mesh collision components
 		const bool bUsingSimpleCollision = (SimpleCollisionMipLevel > CollisionMipLevel && SimpleCollisionWeightmapTextureMipData && !XYOffsetmapTexture);
diff --git a/Engine/Source/Runtime/Landscape/Private/LandscapeEditLayers.cpp b/Engine/Source/Runtime/Landscape/Private/LandscapeEditLayers.cpp
--- a/Engine/Source/Runtime/Landscape/Private/LandscapeEditLayers.cpp
+++ b/Engine/Source/Runtime/Landscape/Private/LandscapeEditLayers.cpp
@@ -34,6 +34,7 @@ LandscapeEditLayers.cpp: Landscape editing layers mode
 #include "LandscapeInfoMap.h"
 #include "Misc/MessageDialog.h"
 #include "GameFramework/WorldSettings.h"
+#include "UObject/UObjectThreadContext.h"
 #endif
 
 #define LOCTEXT_NAMESPACE "Landscape"
@@ -4195,6 +4196,12 @@ void ALandscape::RequestLayersContentUpdate(ELandscapeLayerUpdateMode InUpdateMo
 
 void ALandscape::RequestLayersContentUpdateForceAll(ELandscapeLayerUpdateMode InModeMask)
 {
+	// Ignore Update requests while in PostLoad (to avoid dirtying package on load)
+	if (FUObjectThreadContext::Get().IsRoutingPostLoad)
+	{
+		return;
+	}
+
 	if (!CanHaveLayersContent())
 	{
 		return;
