Cherrypick CL 7827812

#rnx

diff --git a/Engine/Source/Runtime/Landscape/Private/LandscapeEdit.cpp b/Engine/Source/Runtime/Landscape/Private/LandscapeEdit.cpp
--- a/Engine/Source/Runtime/Landscape/Private/LandscapeEdit.cpp
+++ b/Engine/Source/Runtime/Landscape/Private/LandscapeEdit.cpp
@@ -1424,27 +1424,32 @@ void ULandscapeComponent::UpdateCollisionLayerData()
 	TArray<UTexture2D*>& ComponentWeightmapsTexture = GetWeightmapTextures();
 
 	// Generate the dominant layer data
-	TArray<FColor*> WeightmapTextureMipData;
-	WeightmapTextureMipData.Empty(ComponentWeightmapsTexture.Num());
+	TArray<TArray64<uint8>> WeightmapTextureMipData;
+	TArray<FColor*> WeightmapTextureMipDataParam;
+	WeightmapTextureMipData.Reserve(ComponentWeightmapsTexture.Num());
+	WeightmapTextureMipDataParam.Reserve(ComponentWeightmapsTexture.Num());
 	for (int32 WeightmapIdx = 0; WeightmapIdx < ComponentWeightmapsTexture.Num(); ++WeightmapIdx)
 	{
-		TArray64<uint8> MipData;
+		TArray64<uint8>& MipData = WeightmapTextureMipData.AddDefaulted_GetRef();
 		ComponentWeightmapsTexture[WeightmapIdx]->Source.GetMipData(MipData, CollisionMipLevel);
-		WeightmapTextureMipData.Add((FColor*)MipData.GetData());
+		WeightmapTextureMipDataParam.Add((FColor*)MipData.GetData());
 	}
 
-	TArray<FColor*> SimpleCollisionWeightmapMipData;
+	TArray<TArray64<uint8>> SimpleCollisionWeightmapMipData;
+	TArray<FColor*> SimpleCollisionWeightmapMipDataParam;
 	if (SimpleCollisionMipLevel > CollisionMipLevel)
 	{
+		SimpleCollisionWeightmapMipData.Reserve(ComponentWeightmapsTexture.Num());
+		SimpleCollisionWeightmapMipDataParam.Reserve(ComponentWeightmapsTexture.Num());
 		for (int32 WeightmapIdx = 0; WeightmapIdx < ComponentWeightmapsTexture.Num(); ++WeightmapIdx)
 		{
-			TArray64<uint8> MipData;
+			TArray64<uint8>& MipData = SimpleCollisionWeightmapMipData.AddDefaulted_GetRef();
 			ComponentWeightmapsTexture[WeightmapIdx]->Source.GetMipData(MipData, SimpleCollisionMipLevel);
-			SimpleCollisionWeightmapMipData.Add((FColor*)MipData.GetData());
+			SimpleCollisionWeightmapMipDataParam.Add((FColor*)MipData.GetData());
 		}
 	}
 
-	UpdateCollisionLayerData(WeightmapTextureMipData.GetData(), SimpleCollisionWeightmapMipData.GetData());
+	UpdateCollisionLayerData(WeightmapTextureMipDataParam.GetData(), SimpleCollisionWeightmapMipDataParam.GetData());
 }
 
 
