Landscape Layers: Memory reporting fix (RenderTargets for Layer blending were taken into account, they are now outered to the TransientPackage)

#rnx

diff --git a/Engine/Source/Runtime/Landscape/Private/LandscapeEditLayers.cpp b/Engine/Source/Runtime/Landscape/Private/LandscapeEditLayers.cpp
--- a/Engine/Source/Runtime/Landscape/Private/LandscapeEditLayers.cpp
+++ b/Engine/Source/Runtime/Landscape/Private/LandscapeEditLayers.cpp
@@ -1352,7 +1352,7 @@ void ALandscape::CreateLayersRenderingResource()
 
 		for (int32 i = 0; i < EHeightmapRTType::HeightmapRT_Count; ++i)
 		{
-			Landscape->HeightmapRTList[i] = NewObject<UTextureRenderTarget2D>(Landscape);
+			Landscape->HeightmapRTList[i] = NewObject<UTextureRenderTarget2D>(GetTransientPackage(), NAME_None, RF_Transient);
 			check(Landscape->HeightmapRTList[i]);
 			Landscape->HeightmapRTList[i]->RenderTargetFormat = RTF_RGBA8;
 			Landscape->HeightmapRTList[i]->AddressX = TextureAddress::TA_Clamp;
@@ -1414,7 +1414,7 @@ void ALandscape::CreateLayersRenderingResource()
 
 		for (int32 i = 0; i < EWeightmapRTType::WeightmapRT_Count; ++i)
 		{
-			Landscape->WeightmapRTList[i] = NewObject<UTextureRenderTarget2D>(Landscape);
+			Landscape->WeightmapRTList[i] = NewObject<UTextureRenderTarget2D>(GetTransientPackage(), NAME_None, RF_Transient);
 
 			check(Landscape->WeightmapRTList[i]);
 			Landscape->WeightmapRTList[i]->AddressX = TextureAddress::TA_Clamp;
