Landscape Layers:
- Add Hashing support so that Splines & Brushes do not need to dirty all the levels when they run (we still have issues with the WaterBrushManager not outputting deterministic results for Weightmaps)
- Invalidation of Grass is now based on the resulting hash changing per component.

#rnx

diff --git a/Engine/Source/Runtime/Landscape/Classes/Landscape.h b/Engine/Source/Runtime/Landscape/Classes/Landscape.h
--- a/Engine/Source/Runtime/Landscape/Classes/Landscape.h
+++ b/Engine/Source/Runtime/Landscape/Classes/Landscape.h
@@ -295,7 +295,7 @@ public:
 	LANDSCAPE_API void ClearPaintLayer(int32 InLayerIndex, ULandscapeLayerInfoObject* InLayerInfo);
 	LANDSCAPE_API void ClearPaintLayer(const FGuid& InLayerGuid, ULandscapeLayerInfoObject* InLayerInfo);
 	LANDSCAPE_API void ClearLayer(int32 InLayerIndex, TSet<ULandscapeComponent*>* InComponents = nullptr, ELandscapeClearMode InClearMode = ELandscapeClearMode::Clear_All);
-	LANDSCAPE_API void ClearLayer(const FGuid& InLayerGuid, TSet<ULandscapeComponent*>* InComponents = nullptr, ELandscapeClearMode InClearMode = ELandscapeClearMode::Clear_All);
+	LANDSCAPE_API void ClearLayer(const FGuid& InLayerGuid, TSet<ULandscapeComponent*>* InComponents = nullptr, ELandscapeClearMode InClearMode = ELandscapeClearMode::Clear_All, bool bMarkPackageDirty = true);
 	LANDSCAPE_API void DeleteLayer(int32 InLayerIndex);
 	LANDSCAPE_API void DeleteLayers();
 	LANDSCAPE_API void SetEditingLayer(const FGuid& InLayerGuid = FGuid());
@@ -339,7 +339,7 @@ private:
 	bool UpdateCollisionAndClients(const TArray<ULandscapeComponent*>& InLandscapeComponents, const int32 InContentUpdateModes);
 	void ResolveLayersHeightmapTexture(const TArray<ULandscapeComponent*>& InLandscapeComponents);
 	void ResolveLayersWeightmapTexture(const TArray<ULandscapeComponent*>& InLandscapeComponents);
-	void ResolveLayersTexture(class FLandscapeLayersTexture2DCPUReadBackResource* InCPUReadBackTexture, UTexture2D* InOutputTexture);
+	bool ResolveLayersTexture(class FLandscapeLayersTexture2DCPUReadBackResource* InCPUReadBackTexture, UTexture2D* InOutputTexture);
 
 	bool PrepareLayersHeightmapTextureResources(bool bInWaitForStreaming) const;
 	bool PrepareLayersWeightmapTextureResources(bool bInWaitForStreaming) const;
diff --git a/Engine/Source/Runtime/Landscape/Classes/LandscapeComponent.h b/Engine/Source/Runtime/Landscape/Classes/LandscapeComponent.h
--- a/Engine/Source/Runtime/Landscape/Classes/LandscapeComponent.h
+++ b/Engine/Source/Runtime/Landscape/Classes/LandscapeComponent.h
@@ -330,6 +330,8 @@ enum ELandscapeLayerUpdateMode : uint32
 	Update_Client_Editing = 1 << 5
 };
 
+static const uint32 DefaultSplineHash = 0xFFFFFFFF;
+
 #endif
 
 UENUM()
@@ -528,6 +530,9 @@ public:
 
 	/** Represents hash of last weightmap usage update */
 	uint32 WeightmapsHash;
+
+	UPROPERTY()
+	uint32 SplineHash;
 #endif
 
 	/** For ES2 */
@@ -617,6 +622,8 @@ public:
 	LANDSCAPE_API const TArray<FWeightmapLayerAllocationInfo>& GetWeightmapLayerAllocations(const FGuid& InLayerGuid) const;
 
 #if WITH_EDITOR
+	LANDSCAPE_API uint32 ComputeLayerHash() const;
+
 	LANDSCAPE_API void SetHeightmap(UTexture2D* NewHeightmap);
 
 	LANDSCAPE_API void SetWeightmapTextures(const TArray<UTexture2D*>& InNewWeightmapTextures, bool InApplyToEditingWeightmap = false);
diff --git a/Engine/Source/Runtime/Landscape/Classes/LandscapeInfo.h b/Engine/Source/Runtime/Landscape/Classes/LandscapeInfo.h
--- a/Engine/Source/Runtime/Landscape/Classes/LandscapeInfo.h
+++ b/Engine/Source/Runtime/Landscape/Classes/LandscapeInfo.h
@@ -128,7 +128,7 @@ public:
 #endif
 
 	UPROPERTY()
-	TSet<ALandscapeStreamingProxy*> Proxies;
+	TArray<ALandscapeStreamingProxy*> Proxies;
 
 private:
 	TSet<ULandscapeComponent*> SelectedComponents;
@@ -150,8 +150,8 @@ public:
 	LANDSCAPE_API void ForAllLandscapeComponents(TFunctionRef<void(ULandscapeComponent*)> Fn) const;
 	LANDSCAPE_API void ExportHeightmap(const FString& Filename);
 	LANDSCAPE_API void ExportLayer(ULandscapeLayerInfoObject* LayerInfo, const FString& Filename);
-	LANDSCAPE_API bool ApplySplines(bool bOnlySelected, TSet<ULandscapeComponent*>* OutModifiedComponents = nullptr);
-	bool ApplySplinesInternal(bool bOnlySelected, ALandscapeProxy* Proxy, TSet<ULandscapeComponent*>* OutModifiedComponents);
+	LANDSCAPE_API bool ApplySplines(bool bOnlySelected, TSet<ULandscapeComponent*>* OutModifiedComponents = nullptr, bool bMarkPackageDirty = true);
+	bool ApplySplinesInternal(bool bOnlySelected, ALandscapeProxy* Proxy, TSet<ULandscapeComponent*>* OutModifiedComponents, bool bMarkPackageDirty);
 
 	LANDSCAPE_API bool GetSelectedExtent(int32& MinX, int32& MinY, int32& MaxX, int32& MaxY) const;
 	FVector GetLandscapeCenterPos(float& LengthZ, int32 MinX = MAX_int32, int32 MinY = MAX_int32, int32 MaxX = MIN_int32, int32 MaxY = MIN_int32);
diff --git a/Engine/Source/Runtime/Landscape/Classes/LandscapeProxy.h b/Engine/Source/Runtime/Landscape/Classes/LandscapeProxy.h
--- a/Engine/Source/Runtime/Landscape/Classes/LandscapeProxy.h
+++ b/Engine/Source/Runtime/Landscape/Classes/LandscapeProxy.h
@@ -356,6 +356,7 @@ public:
 		, SizeY(InSizeY)
 		, Format(InFormat)
 		, NumMips(InNumMips)
+		, Hash(0)
 	{}
 
 	virtual uint32 GetSizeX() const override
@@ -377,11 +378,22 @@ public:
 		TextureRHI = RHICreateTexture2D(SizeX, SizeY, Format, NumMips, 1, TexCreate_CPUReadback, CreateInfo);
 	}
 
+	bool UpdateHashFromTextureSource(const uint8* MipData)
+	{
+		uint32 LocalHash = FCrc::MemCrc32(MipData, SizeX * SizeY * sizeof(FColor));
+		bool bChanged = (LocalHash != Hash);
+		Hash = LocalHash;
+		return bChanged;
+	}
+	   
+	uint32 GetHash() const { return Hash; }
+
 private:
 	uint32 SizeX;
 	uint32 SizeY;
 	EPixelFormat Format;
 	uint32 NumMips;
+	uint32 Hash;
 };
 
 UCLASS(Abstract, MinimalAPI, NotBlueprintable, hidecategories=(Display, Attachment, Physics, Debug, Lighting, LOD), showcategories=(Lighting, Rendering, "Utilities|Transformation"), hidecategories=(Mobility))
@@ -790,7 +802,6 @@ public:
 	virtual void PostEditMove(bool bFinished) override;
 	virtual bool ShouldImport(FString* ActorPropString, bool IsMovingLevel) override;
 	virtual bool ShouldExport() override;
-	virtual bool Modify(bool bAlwaysMarkDirty = true) override;
 	//~ End AActor Interface
 #endif	//WITH_EDITOR
 
diff --git a/Engine/Source/Runtime/Landscape/Classes/LandscapeSplinesComponent.h b/Engine/Source/Runtime/Landscape/Classes/LandscapeSplinesComponent.h
--- a/Engine/Source/Runtime/Landscape/Classes/LandscapeSplinesComponent.h
+++ b/Engine/Source/Runtime/Landscape/Classes/LandscapeSplinesComponent.h
@@ -221,6 +221,6 @@ public:
 	friend class ULandscapeSplineSegment;
 #if WITH_EDITOR
 	// TODO - move this out of ULandscapeInfo
-	friend bool ULandscapeInfo::ApplySplinesInternal(bool bOnlySelected, ALandscapeProxy* Landscape, TSet<ULandscapeComponent*>* OutModifiedComponents);
+	friend bool ULandscapeInfo::ApplySplinesInternal(bool bOnlySelected, ALandscapeProxy* Landscape, TSet<ULandscapeComponent*>* OutModifiedComponents, bool bMarkPackageDirty);
 #endif
 };
diff --git a/Engine/Source/Runtime/Landscape/Private/Landscape.cpp b/Engine/Source/Runtime/Landscape/Private/Landscape.cpp
--- a/Engine/Source/Runtime/Landscape/Private/Landscape.cpp
+++ b/Engine/Source/Runtime/Landscape/Private/Landscape.cpp
@@ -61,6 +61,7 @@ Landscape.cpp: Terrain rendering
 #include "UObject/FortniteMainBranchObjectVersion.h"
 #include "LandscapeDataAccess.h"
 #include "UObject/EditorObjectVersion.h"
+#include "Algo/BinarySearch.h"
 
 /** Landscape stats */
 
@@ -146,6 +147,7 @@ ULandscapeComponent::ULandscapeComponent(const FObjectInitializer& ObjectInitial
 , CachedEditingLayerData(nullptr)
 , LayerUpdateFlagPerMode(0)
 , WeightmapsHash(0)
+, SplineHash(0)
 #endif
 , GrassData(MakeShareable(new FLandscapeComponentGrassData()))
 , ChangeTag(0)
@@ -2801,9 +2803,25 @@ void ULandscapeInfo::RegisterActor(ALandscapeProxy* Proxy, bool bMapCheck)
 	}
 	else
 	{
+		// Insert Proxies in a sorted fashion for generating determisnitic results in the Layer system
 		ALandscapeStreamingProxy* StreamingProxy = CastChecked<ALandscapeStreamingProxy>(Proxy);
+		if (!Proxies.Contains(Proxy))
+		{
+			uint32 InsertIndex = Algo::LowerBound(Proxies, Proxy, [](ALandscapeProxy* A, ALandscapeProxy* B)
+			{
+				FIntPoint SectionBaseA = A->GetSectionBaseOffset();
+				FIntPoint SectionBaseB = B->GetSectionBaseOffset();
+
+				if (SectionBaseA.X != SectionBaseB.X)
+				{
+					return SectionBaseA.X < SectionBaseB.X;
+				}
 
-		Proxies.Add(StreamingProxy);
+				return SectionBaseA.Y < SectionBaseB.Y;
+			});
+			
+			Proxies.Insert(StreamingProxy, InsertIndex);
+		}
 		StreamingProxy->LandscapeActor = LandscapeActor;
 		StreamingProxy->FixupSharedData(LandscapeActor.Get());
 	}
diff --git a/Engine/Source/Runtime/Landscape/Private/LandscapeEdit.cpp b/Engine/Source/Runtime/Landscape/Private/LandscapeEdit.cpp
--- a/Engine/Source/Runtime/Landscape/Private/LandscapeEdit.cpp
+++ b/Engine/Source/Runtime/Landscape/Private/LandscapeEdit.cpp
@@ -592,10 +592,6 @@ void ULandscapeComponent::PostEditUndo()
 		
 	if (GetLandscapeProxy()->HasLayersContent())
 	{
-		TArray<ULandscapeComponent*> SingleComponent;
-		SingleComponent.Add(this);
-		GetLandscapeProxy()->InvalidateGeneratedComponentData(SingleComponent);
-		
 		const bool bUpdateAll = true;
 		RequestHeightmapUpdate(bUpdateAll);
 		RequestWeightmapUpdate(bUpdateAll);
@@ -1331,8 +1327,8 @@ void ULandscapeComponent::UpdateCollisionLayerData(const FColor* const* const We
 		bool bExistingLayerMismatch = false;
 		int32 DataLayerIdx = INDEX_NONE;
 
-		TArray<FWeightmapLayerAllocationInfo>& ComponentWeightmapLayerAllocations = GetWeightmapLayerAllocations(true);
-		TArray<UTexture2D*>& ComponentWeightmapsTexture = GetWeightmapTextures(true);
+		TArray<FWeightmapLayerAllocationInfo>& ComponentWeightmapLayerAllocations = GetWeightmapLayerAllocations(false);
+		TArray<UTexture2D*>& ComponentWeightmapsTexture = GetWeightmapTextures(false);
 
 		// Find the layers we're interested in
 		for (int32 AllocIdx = 0; AllocIdx < ComponentWeightmapLayerAllocations.Num(); AllocIdx++)
@@ -3291,10 +3287,16 @@ bool ULandscapeInfo::GetLandscapeExtent(int32& MinX, int32& MinY, int32& MaxX, i
 
 LANDSCAPE_API void ULandscapeInfo::ForAllLandscapeComponents(TFunctionRef<void(ULandscapeComponent*)> Fn) const
 {
-	for (auto& XYComponentPair : XYtoComponentMap)
+	ForAllLandscapeProxies([&](ALandscapeProxy* Proxy)
 	{
-		Fn(XYComponentPair.Value);
-	}
+		for (ULandscapeComponent* Component : Proxy->LandscapeComponents)
+		{
+			if (Component->IsRegistered())
+			{
+				Fn(Component);
+			}
+		}
+	});
 }
 
 bool ULandscapeInfo::GetSelectedExtent(int32& MinX, int32& MinY, int32& MaxX, int32& MaxY) const
@@ -4150,13 +4152,16 @@ void ALandscapeProxy::PostEditChangeProperty(FPropertyChangedEvent& PropertyChan
 	}
 	if (bRemovedAnyLayers)
 	{
-		ALandscapeProxy::InvalidateGeneratedComponentData(LandscapeComponents);
 		ALandscape* LandscapeActor = GetLandscapeActor();
 
 		if(LandscapeActor != nullptr && LandscapeActor->HasLayersContent())
 		{
 			LandscapeActor->RequestLayersContentUpdate(ELandscapeLayerUpdateMode::Update_All);
 		}
+		else
+		{
+			ALandscapeProxy::InvalidateGeneratedComponentData(LandscapeComponents);
+		}
 	}
 
 	// Must do this *after* correcting the scale or reattaching the landscape components will crash!
@@ -4783,10 +4788,11 @@ void ULandscapeComponent::ReallocateWeightmaps(FLandscapeEditDataInterface* Data
 		}
 	}
 
+	bool bMarkPackageDirty = DataInterface == nullptr ? true : DataInterface->GetShouldDirtyPackage();
 	if (InSaveToTransactionBuffer)
 	{
-		Modify();
-		TargetProxy->Modify();
+		Modify(bMarkPackageDirty);
+		TargetProxy->Modify(bMarkPackageDirty);
 	}
 
 	if (!InForceReallocate)
@@ -4848,7 +4854,7 @@ void ULandscapeComponent::ReallocateWeightmaps(FLandscapeEditDataInterface* Data
 
 						if (InSaveToTransactionBuffer)
 						{
-							Usage->Modify();
+							Usage->Modify(bMarkPackageDirty);
 						}
 						Usage->ChannelUsage[ChanIdx] = this;
 
@@ -4954,7 +4960,7 @@ void ULandscapeComponent::ReallocateWeightmaps(FLandscapeEditDataInterface* Data
 			CurrentWeightmapUsage = TargetProxy->WeightmapUsageMap.Add(CurrentWeightmapTexture, TargetProxy->CreateWeightmapUsage());
 			if (InSaveToTransactionBuffer)
 			{
-				CurrentWeightmapUsage->Modify();
+				CurrentWeightmapUsage->Modify(bMarkPackageDirty);
 			}
 
 			CurrentWeightmapUsage->LayerGuid = TargetLayerGuid;
@@ -4998,7 +5004,7 @@ void ULandscapeComponent::ReallocateWeightmaps(FLandscapeEditDataInterface* Data
 					ULandscapeWeightmapUsage* OldWeightmapUsage = ComponentWeightmapTexturesUsage[AllocInfo.WeightmapTextureIndex];
 					if (InSaveToTransactionBuffer)
 					{
-						OldWeightmapUsage->Modify();
+						OldWeightmapUsage->Modify(bMarkPackageDirty);
 					}
 					OldWeightmapUsage->ChannelUsage[AllocInfo.WeightmapTextureChannel] = nullptr;
 				}
@@ -5006,7 +5012,7 @@ void ULandscapeComponent::ReallocateWeightmaps(FLandscapeEditDataInterface* Data
 				// Assign the new allocation
 				if (InSaveToTransactionBuffer)
 				{
-					CurrentWeightmapUsage->Modify();
+					CurrentWeightmapUsage->Modify(bMarkPackageDirty);
 				}
 				CurrentWeightmapUsage->ChannelUsage[ChanIdx] = this;
 				AllocInfo.WeightmapTextureIndex = NewWeightmapTextures.Num() - 1;
@@ -5271,7 +5277,8 @@ void ULandscapeComponent::InitWeightmapData(TArray<ULandscapeLayerInfoObject*>&
 	{
 		UTexture2D* WeightmapTexture = WeightmapTextures[Idx];
 		{
-			FLandscapeTextureDataInfo WeightmapDataInfo(WeightmapTexture);
+			const bool bShouldDirtyPackage = true;
+			FLandscapeTextureDataInfo WeightmapDataInfo(WeightmapTexture, bShouldDirtyPackage);
 
 			int32 NumMips = WeightmapTexture->Source.GetNumMips();
 			TArray<FColor*> WeightmapTextureMipData;
@@ -5301,19 +5308,6 @@ void ULandscapeComponent::InitWeightmapData(TArray<ULandscapeLayerInfoObject*>&
 #define MAX_LANDSCAPE_EXPORT_COMPONENTS_NUM		16
 #define MAX_LANDSCAPE_PROP_TEXT_LENGTH			1024*1024*16
 
-
-bool ALandscapeProxy::Modify(bool bAlwaysMarkDirty)
-{
-	for (auto& Pair : WeightmapUsageMap)
-	{
-		if (Pair.Value)
-		{
-			Pair.Value->Modify(bAlwaysMarkDirty);
-		}
-	}
-	return Super::Modify(bAlwaysMarkDirty);
-}
-
 bool ALandscapeProxy::ShouldExport()
 {
 	if (!bIsMovingToLevel && LandscapeComponents.Num() > MAX_LANDSCAPE_EXPORT_COMPONENTS_NUM)
diff --git a/Engine/Source/Runtime/Landscape/Private/LandscapeEditInterface.cpp b/Engine/Source/Runtime/Landscape/Private/LandscapeEditInterface.cpp
--- a/Engine/Source/Runtime/Landscape/Private/LandscapeEditInterface.cpp
+++ b/Engine/Source/Runtime/Landscape/Private/LandscapeEditInterface.cpp
@@ -43,7 +43,7 @@ FLandscapeEditDataInterface::FLandscapeEditDataInterface(ULandscapeInfo* InLands
 }
 
 FLandscapeTextureDataInterface::FLandscapeTextureDataInterface(bool bInUploadTextureChangesToGPU)
-	: bUploadTextureChangesToGPU(bInUploadTextureChangesToGPU)
+	: bUploadTextureChangesToGPU(bInUploadTextureChangesToGPU), bShouldDirtyPackage(true)
 {}
 
 FLandscapeTextureDataInterface::~FLandscapeTextureDataInterface()
@@ -295,7 +295,7 @@ void FLandscapeEditDataInterface::SetHeightData(int32 X1, int32 Y1, int32 X2, in
 			UTexture2D* Heightmap = InHeightmap != nullptr ? InHeightmap : Component->GetHeightmap(true);
 			UTexture2D* XYOffsetmapTexture = InXYOffsetmapTexture != nullptr ? InXYOffsetmapTexture : Component->XYOffsetmapTexture;
 
-			Component->Modify();
+			Component->Modify(GetShouldDirtyPackage());
 
 			FLandscapeTextureDataInfo* TexDataInfo = GetTextureDataInfo(Heightmap);
 			FColor* HeightmapTextureData = (FColor*)TexDataInfo->GetMipData(0);
@@ -1718,13 +1718,13 @@ void ULandscapeComponent::DeleteLayer(ULandscapeLayerInfoObject* LayerInfo, FLan
 
 	// Mark the channel as unallocated, so we can reuse it later
 	ALandscapeProxy* Proxy = Component->GetLandscapeProxy();
-	Component->Modify();
-	Proxy->Modify();
+	Component->Modify(LandscapeEdit.GetShouldDirtyPackage());
+	Proxy->Modify(LandscapeEdit.GetShouldDirtyPackage());
 
 	ULandscapeWeightmapUsage* Usage = ComponentWeightmapTexturesUsage.IsValidIndex(DeleteLayerAllocation.WeightmapTextureIndex) ? ComponentWeightmapTexturesUsage[DeleteLayerAllocation.WeightmapTextureIndex] : nullptr;
 	if (Usage) // can be null if WeightmapUsageMap hasn't been built yet
 	{
-		Usage->Modify();
+		Usage->Modify(LandscapeEdit.GetShouldDirtyPackage());
 		Usage->ChannelUsage[DeleteLayerAllocation.WeightmapTextureChannel] = nullptr;
 	}
 
@@ -1735,8 +1735,7 @@ void ULandscapeComponent::DeleteLayer(ULandscapeLayerInfoObject* LayerInfo, FLan
 	if (bCanRemoveLayerTexture)
 	{
 		ComponentWeightmapTextures[DeleteLayerWeightmapTextureIndex]->SetFlags(RF_Transactional);
-		ComponentWeightmapTextures[DeleteLayerWeightmapTextureIndex]->Modify();
-		ComponentWeightmapTextures[DeleteLayerWeightmapTextureIndex]->MarkPackageDirty();
+		ComponentWeightmapTextures[DeleteLayerWeightmapTextureIndex]->Modify(LandscapeEdit.GetShouldDirtyPackage());
 		ComponentWeightmapTextures[DeleteLayerWeightmapTextureIndex]->ClearFlags(RF_Standalone);
 
 		ComponentWeightmapTextures.RemoveAt(DeleteLayerWeightmapTextureIndex);
@@ -1805,14 +1804,16 @@ void FLandscapeEditDataInterface::DeleteLayer(ULandscapeLayerInfoObject* LayerIn
 		Component->DeleteLayer(LayerInfo, *this);
 	}
 
-	TSet<ULandscapeComponent*> Components;
-	Algo::Transform(LandscapeInfo->XYtoComponentMap, Components, &TPair<FIntPoint, ULandscapeComponent*>::Value);
-	ALandscapeProxy::InvalidateGeneratedComponentData(Components);
-
-	if (LandscapeInfo->LandscapeActor)
+	if (LandscapeInfo->LandscapeActor && LandscapeInfo->LandscapeActor->HasLayersContent())
 	{
 		LandscapeInfo->LandscapeActor->RequestLayersContentUpdate(ELandscapeLayerUpdateMode::Update_All);
 	}
+	else
+	{
+		TSet<ULandscapeComponent*> Components;
+		Algo::Transform(LandscapeInfo->XYtoComponentMap, Components, &TPair<FIntPoint, ULandscapeComponent*>::Value);
+		ALandscapeProxy::InvalidateGeneratedComponentData(Components);
+	}
 }
 
 void ULandscapeComponent::FillLayer(ULandscapeLayerInfoObject* LayerInfo, FLandscapeEditDataInterface& LandscapeEdit)
@@ -1822,8 +1823,8 @@ void ULandscapeComponent::FillLayer(ULandscapeLayerInfoObject* LayerInfo, FLands
 	ULandscapeComponent* Component = this;
 
 	ALandscapeProxy* Proxy = Component->GetLandscapeProxy();
-	Component->Modify();
-	Proxy->Modify();
+	Component->Modify(LandscapeEdit.GetShouldDirtyPackage());
+	Proxy->Modify(LandscapeEdit.GetShouldDirtyPackage());
 
 	const bool bFillLayerIsNoWeightBlend = LayerInfo->bNoWeightBlend;
 	bool bClearOtherWeightBlendLayers = !bFillLayerIsNoWeightBlend;
@@ -1903,7 +1904,7 @@ void ULandscapeComponent::FillLayer(ULandscapeLayerInfoObject* LayerInfo, FLands
 			ULandscapeWeightmapUsage* Usage = ComponentWeightmapTexturesUsage[Allocation.WeightmapTextureIndex];
 			if (Usage) // can be null if WeightmapUsageMap hasn't been built yet
 			{
-				Usage->Modify();
+				Usage->Modify(LandscapeEdit.GetShouldDirtyPackage());
 				Usage->ChannelUsage[Allocation.WeightmapTextureChannel] = nullptr;
 			}
 
@@ -1919,7 +1920,7 @@ void ULandscapeComponent::FillLayer(ULandscapeLayerInfoObject* LayerInfo, FLands
 			if (!ComponentWeightmapLayerAllocations.ContainsByPredicate(
 				[TextureIdx](const FWeightmapLayerAllocationInfo& Allocation) { return Allocation.WeightmapTextureIndex == TextureIdx; }))
 			{
-				ComponentWeightmapTextures[TextureIdx]->Modify();
+				ComponentWeightmapTextures[TextureIdx]->Modify(LandscapeEdit.GetShouldDirtyPackage());
 
 				for (int32 LayerIdx = 0; LayerIdx < ComponentWeightmapLayerAllocations.Num(); ++LayerIdx)
 				{
@@ -2042,13 +2043,16 @@ void FLandscapeEditDataInterface::FillLayer(ULandscapeLayerInfoObject* LayerInfo
 		}
 	}
 
-	TSet<ULandscapeComponent*> Components;
-	Algo::Transform(LandscapeInfo->XYtoComponentMap, Components, &TPair<FIntPoint, ULandscapeComponent*>::Value);
-	ALandscapeProxy::InvalidateGeneratedComponentData(Components);
-	if (LandscapeInfo->LandscapeActor)
+	if (LandscapeInfo->LandscapeActor && LandscapeInfo->LandscapeActor->HasLayersContent())
 	{
 		LandscapeInfo->LandscapeActor->RequestLayersContentUpdateForceAll();
 	}
+	else
+	{
+		TSet<ULandscapeComponent*> Components;
+		Algo::Transform(LandscapeInfo->XYtoComponentMap, Components, &TPair<FIntPoint, ULandscapeComponent*>::Value);
+		ALandscapeProxy::InvalidateGeneratedComponentData(Components);
+	}
 }
 
 void ULandscapeComponent::ReplaceLayer(ULandscapeLayerInfoObject* FromLayerInfo, ULandscapeLayerInfoObject* ToLayerInfo, FLandscapeEditDataInterface& LandscapeEdit)
@@ -2197,7 +2201,7 @@ void ULandscapeComponent::ReplaceLayer(ULandscapeLayerInfoObject* FromLayerInfo,
 		//check(Usage);
 		if (Usage)
 		{
-			Usage->Modify();
+			Usage->Modify(LandscapeEdit.GetShouldDirtyPackage());
 			Usage->ChannelUsage[FromLayerAllocation.WeightmapTextureChannel] = nullptr;
 		}
 
@@ -2205,8 +2209,7 @@ void ULandscapeComponent::ReplaceLayer(ULandscapeLayerInfoObject* FromLayerInfo,
 		if (bCanRemoveLayerTexture)
 		{
 			ComponentWeightmapTextures[FromLayerAllocation.WeightmapTextureIndex]->SetFlags(RF_Transactional);
-			ComponentWeightmapTextures[FromLayerAllocation.WeightmapTextureIndex]->Modify();
-			ComponentWeightmapTextures[FromLayerAllocation.WeightmapTextureIndex]->MarkPackageDirty();
+			ComponentWeightmapTextures[FromLayerAllocation.WeightmapTextureIndex]->Modify(LandscapeEdit.GetShouldDirtyPackage());
 			ComponentWeightmapTextures[FromLayerAllocation.WeightmapTextureIndex]->ClearFlags(RF_Standalone);
 
 			ComponentWeightmapTextures.RemoveAt(FromLayerAllocation.WeightmapTextureIndex);
@@ -2288,7 +2291,7 @@ void FLandscapeEditDataInterface::ReplaceLayer(ULandscapeLayerInfoObject* FromLa
 
 	if (LandscapeInfo->LandscapeActor.IsValid() && LandscapeInfo->LandscapeActor->HasLayersContent())
 	{
-		LandscapeInfo->LandscapeActor->Modify();
+		LandscapeInfo->LandscapeActor->Modify(GetShouldDirtyPackage());
 		LandscapeInfo->LandscapeActor->ForEachLayer([&](FLandscapeLayer& CurrentLayer)
 		{
 			bool OutValue;
@@ -2452,7 +2455,7 @@ namespace
 	}
 };
 
-bool DeleteLayerIfAllZero(ULandscapeComponent* const Component, const uint8* const TexDataPtr, int32 TexSize, int32 LayerIdx)
+bool DeleteLayerIfAllZero(ULandscapeComponent* const Component, const uint8* const TexDataPtr, int32 TexSize, int32 LayerIdx, bool bShouldDirtyPackage)
 {
 	// Check the data for the entire component and to see if it's all zero
 	for (int32 TexY = 0; TexY < TexSize; TexY++)
@@ -2471,8 +2474,8 @@ bool DeleteLayerIfAllZero(ULandscapeComponent* const Component, const uint8* con
 	}
 
 	ALandscapeProxy* Proxy = Component->GetLandscapeProxy();
-	Component->Modify();
-	Proxy->Modify();
+	Component->Modify(bShouldDirtyPackage);
+	Proxy->Modify(bShouldDirtyPackage);
 
 	TArray<FWeightmapLayerAllocationInfo>& ComponentWeightmapLayerAllocations = Component->GetWeightmapLayerAllocations(true);
 	TArray<UTexture2D*>& ComponentWeightmapTextures = Component->GetWeightmapTextures(true);
@@ -2481,7 +2484,7 @@ bool DeleteLayerIfAllZero(ULandscapeComponent* const Component, const uint8* con
 	// Mark the channel as unallocated, so we can reuse it later
 	const int32 DeleteLayerWeightmapTextureIndex = ComponentWeightmapLayerAllocations[LayerIdx].WeightmapTextureIndex;
 	ULandscapeWeightmapUsage& Usage = *ComponentWeightmapTexturesUsage[DeleteLayerWeightmapTextureIndex];
-	Usage.Modify();
+	Usage.Modify(bShouldDirtyPackage);
 	Usage.ChannelUsage[ComponentWeightmapLayerAllocations[LayerIdx].WeightmapTextureChannel] = nullptr;
 
 	// Remove the layer as it's totally painted away.
@@ -2491,7 +2494,6 @@ bool DeleteLayerIfAllZero(ULandscapeComponent* const Component, const uint8* con
 	bool bCanRemoveLayerTexture = !ComponentWeightmapLayerAllocations.ContainsByPredicate([DeleteLayerWeightmapTextureIndex](const FWeightmapLayerAllocationInfo& Allocation){ return Allocation.WeightmapTextureIndex == DeleteLayerWeightmapTextureIndex; });
 	if (bCanRemoveLayerTexture)
 	{
-		ComponentWeightmapTextures[DeleteLayerWeightmapTextureIndex]->MarkPackageDirty();
 		ComponentWeightmapTextures[DeleteLayerWeightmapTextureIndex]->ClearFlags(RF_Standalone);
 		ComponentWeightmapTextures.RemoveAt(DeleteLayerWeightmapTextureIndex);
 
@@ -2786,7 +2788,7 @@ void FLandscapeEditDataInterface::SetAlphaData(ULandscapeLayerInfoObject* const
 				continue;
 			}
 
-			Component->Modify();
+			Component->Modify(GetShouldDirtyPackage());
 
 			TArray<FWeightmapLayerAllocationInfo>& ComponentWeightmapLayerAllocations = Component->GetWeightmapLayerAllocations(true);
 			int32 UpdateLayerIdx = ComponentWeightmapLayerAllocations.IndexOfByPredicate([LayerInfo](const FWeightmapLayerAllocationInfo& Allocation){ return Allocation.LayerInfo == LayerInfo; });
@@ -3185,7 +3187,7 @@ void FLandscapeEditDataInterface::SetAlphaData(ULandscapeLayerInfoObject* const
 			{
 				if (LayerEditDataAllZero[LayerIdx])
 				{
-					bool bLayerDeleted = DeleteLayerIfAllZero(Component, LayerDataPtrs[LayerIdx], TexSize, LayerIdx);
+					bool bLayerDeleted = DeleteLayerIfAllZero(Component, LayerDataPtrs[LayerIdx], TexSize, LayerIdx, GetShouldDirtyPackage());
 
 					if (bLayerDeleted)
 					{
@@ -3298,7 +3300,7 @@ void FLandscapeEditDataInterface::SetAlphaData(const TSet<ULandscapeLayerInfoObj
 				}
 				if (PaintingRestriction != ELandscapeLayerPaintingRestriction::ExistingOnly)
 				{
-					Component->Modify();
+					Component->Modify(GetShouldDirtyPackage());
 					for (ULandscapeLayerInfoObject* LayerInfoNeedingAllocation : NeedAllocationInfos)
 					{
 						if (PaintingRestriction == ELandscapeLayerPaintingRestriction::UseMaxLayers &&
@@ -3472,7 +3474,7 @@ void FLandscapeEditDataInterface::SetAlphaData(const TSet<ULandscapeLayerInfoObj
 			{
 				if (LayerEditDataAllZero[LayerIdx])
 				{
-					bool bLayerDeleted = DeleteLayerIfAllZero(Component, LayerDataInfos[LayerIdx].TexDataPtr, TexSize, LayerIdx);
+					bool bLayerDeleted = DeleteLayerIfAllZero(Component, LayerDataInfos[LayerIdx].TexDataPtr, TexSize, LayerIdx, GetShouldDirtyPackage());
 
 					if (bLayerDeleted)
 					{
@@ -4333,7 +4335,7 @@ FLandscapeTextureDataInfo* FLandscapeTextureDataInterface::GetTextureDataInfo(UT
 	FLandscapeTextureDataInfo* Result = TextureDataMap.FindRef(Texture);
 	if( !Result )
 	{
-		Result = TextureDataMap.Add(Texture, new FLandscapeTextureDataInfo(Texture));
+		Result = TextureDataMap.Add(Texture, new FLandscapeTextureDataInfo(Texture, GetShouldDirtyPackage()));
 	}
 	return Result;
 }
@@ -4796,7 +4798,7 @@ void FLandscapeEditDataInterface::SetXYOffsetDataTempl(int32 X1, int32 Y1, int32
 			{
 				if (Component->XYOffsetmapTexture == NULL)
 				{
-					Component->Modify();
+					Component->Modify(GetShouldDirtyPackage());
 					//FlushRenderingCommands();
 					// Construct Texture...
 					int32 WeightmapSize = (Component->SubsectionSizeQuads+1) * Component->NumSubsections;
@@ -5597,14 +5599,14 @@ void FLandscapeEditDataInterface::GetXYOffsetDataFast(const int32 X1, const int3
 // FLandscapeTextureDataInfo
 //
 
-FLandscapeTextureDataInfo::FLandscapeTextureDataInfo(UTexture2D* InTexture)
+FLandscapeTextureDataInfo::FLandscapeTextureDataInfo(UTexture2D* InTexture, bool bShouldDirtyPackage)
 :	Texture(InTexture)
 {
 	check(InTexture->IsAsyncCacheComplete());
 	MipInfo.AddZeroed(Texture->Source.GetNumMips());
 	Texture->SetFlags(RF_Transactional);
 	Texture->TemporarilyDisableStreaming();
-	Texture->Modify();
+	Texture->Modify(bShouldDirtyPackage);
 }
 
 bool FLandscapeTextureDataInfo::UpdateTextureData()
diff --git a/Engine/Source/Runtime/Landscape/Private/LandscapeEditLayers.cpp b/Engine/Source/Runtime/Landscape/Private/LandscapeEditLayers.cpp
--- a/Engine/Source/Runtime/Landscape/Private/LandscapeEditLayers.cpp
+++ b/Engine/Source/Runtime/Landscape/Private/LandscapeEditLayers.cpp
@@ -41,6 +41,9 @@ LandscapeEditLayers.cpp: Landscape editing layers mode
 
 static const FString GEmptyDebugName(TEXT(""));
 
+// Channel remapping
+extern const size_t ChannelOffsets[4];
+
 ENGINE_API extern bool GDisableAutomaticTextureMaterialUpdateDependencies;
 
 #if WITH_EDITOR
@@ -2707,23 +2710,31 @@ int32 ALandscape::RegenerateLayersHeightmaps(const TArray<ULandscapeComponent*>&
 		return 0;
 	}
 	   
-	// Handle missing Heightmap CPUReadback (Undo of a delete landscape component can trigger that case)
+	// Init CPU Readbacks
 	if (HeightmapUpdateModes)
 	{
-		Info->ForAllLandscapeProxies([&](ALandscapeProxy* Proxy)
+		for(ULandscapeComponent* Component : InLandscapeComponents)
 		{
-			for (ULandscapeComponent* Component : Proxy->LandscapeComponents)
+			UTexture2D* ComponentHeightmap = Component->GetHeightmap(false);
+			ALandscapeProxy* Proxy = Component->GetLandscapeProxy();
+			FLandscapeLayersTexture2DCPUReadBackResource** CPUReadback = Proxy->HeightmapsCPUReadBack.Find(ComponentHeightmap);
+			if (CPUReadback == nullptr || *CPUReadback == nullptr)
 			{
-				UTexture2D* ComponentHeightmap = Component->GetHeightmap();
-				FLandscapeLayersTexture2DCPUReadBackResource** CPUReadback = Proxy->HeightmapsCPUReadBack.Find(ComponentHeightmap);
-				if (CPUReadback == nullptr)
-				{
-					FLandscapeLayersTexture2DCPUReadBackResource* NewCPUReadBackResource = new FLandscapeLayersTexture2DCPUReadBackResource(ComponentHeightmap->Source.GetSizeX(), ComponentHeightmap->Source.GetSizeY(), ComponentHeightmap->GetPixelFormat(), ComponentHeightmap->Source.GetNumMips());
-					BeginInitResource(NewCPUReadBackResource);
-					Proxy->HeightmapsCPUReadBack.Add(ComponentHeightmap, NewCPUReadBackResource);
-				}
+				FLandscapeLayersTexture2DCPUReadBackResource* NewCPUReadBackResource = new FLandscapeLayersTexture2DCPUReadBackResource(ComponentHeightmap->Source.GetSizeX(), ComponentHeightmap->Source.GetSizeY(), ComponentHeightmap->GetPixelFormat(), ComponentHeightmap->Source.GetNumMips());
+				BeginInitResource(NewCPUReadBackResource);
+
+				const uint8* LockedMip = ComponentHeightmap->Source.LockMip(0);
+				NewCPUReadBackResource->UpdateHashFromTextureSource(LockedMip);
+				ComponentHeightmap->Source.UnlockMip(0);
+				Proxy->HeightmapsCPUReadBack.Add(ComponentHeightmap, NewCPUReadBackResource);
 			}
-		});
+			else if ((*CPUReadback)->GetHash() == 0)
+			{
+				const uint8* LockedMip = ComponentHeightmap->Source.LockMip(0);
+				(*CPUReadback)->UpdateHashFromTextureSource(LockedMip);
+				ComponentHeightmap->Source.UnlockMip(0);
+			}
+		}
 	}
 
 	if (HeightmapUpdateModes || bForceRender)
@@ -2955,13 +2966,15 @@ void ALandscape::ResolveLayersHeightmapTexture(const TArray<ULandscapeComponent*
 		return;
 	}
 
-	TArray<UTexture*> ProcessedTexture;
+	TMap<UTexture*, bool> ProcessedTexture;
+	TArray<ULandscapeComponent*> ChangedComponents;
 
 	for (ULandscapeComponent* Component : InLandscapeComponents)
 	{
-		UTexture2D* ComponentHeightmap = Component->GetHeightmap();
-
-		if (!ProcessedTexture.Contains(ComponentHeightmap))
+		UTexture2D* ComponentHeightmap = Component->GetHeightmap(false);
+		bool* bValue = ProcessedTexture.Find(ComponentHeightmap);
+		bool bChanged = false;
+		if (bValue == nullptr)
 		{
 			FLandscapeLayersTexture2DCPUReadBackResource** CPUReadback = Component->GetLandscapeProxy()->HeightmapsCPUReadBack.Find(ComponentHeightmap);
 
@@ -2969,15 +2982,24 @@ void ALandscape::ResolveLayersHeightmapTexture(const TArray<ULandscapeComponent*
 			{
 				continue;
 			}
-
-			ProcessedTexture.Add(ComponentHeightmap);
-
-			ResolveLayersTexture(*CPUReadback, ComponentHeightmap);
+						
+			bChanged = ResolveLayersTexture(*CPUReadback, ComponentHeightmap);
+			ProcessedTexture.Add(ComponentHeightmap, bChanged);
+			bValue = &bChanged;
+		}
+		
+		if (*bValue)
+		{
+			Component->MarkPackageDirty();
+			ChangedComponents.Add(Component);
 		}
 	}
+
+	InvalidateGeneratedComponentData(ChangedComponents);
+
 }
 
-void ALandscape::ResolveLayersTexture(FLandscapeLayersTexture2DCPUReadBackResource* InCPUReadBackTexture, UTexture2D* InOutputTexture)
+bool ALandscape::ResolveLayersTexture(FLandscapeLayersTexture2DCPUReadBackResource* InCPUReadBackTexture, UTexture2D* InOutputTexture)
 {
 	SCOPE_CYCLE_COUNTER(STAT_LandscapeLayersResolveTexture);
 
@@ -3010,15 +3032,23 @@ void ALandscape::ResolveLayersTexture(FLandscapeLayersTexture2DCPUReadBackResour
 	// TODO: find a way to NOT have to flush the rendering command as this create hic up of ~10-15ms
 	FlushRenderingCommands();
 
+	bool bChanged = false;
 	for (int8 MipIndex = 0; MipIndex < OutMipsData.Num(); ++MipIndex)
 	{
 		if (OutMipsData[MipIndex].Num() > 0)
 		{
-			FColor* TextureData = (FColor*)InOutputTexture->Source.LockMip(MipIndex);
+			uint8* TextureData = InOutputTexture->Source.LockMip(MipIndex);
+			if (MipIndex == 0)
+			{
+				bChanged = InCPUReadBackTexture->UpdateHashFromTextureSource((uint8*)OutMipsData[MipIndex].GetData());
+			}
+
 			FMemory::Memcpy(TextureData, OutMipsData[MipIndex].GetData(), OutMipsData[MipIndex].Num() * sizeof(FColor));
 			InOutputTexture->Source.UnlockMip(MipIndex);
 		}
 	}	
+
+	return bChanged;
 }
 
 void ALandscape::PrepareComponentDataToExtractMaterialLayersCS(const TArray<ULandscapeComponent*>& InLandscapeComponents, const FLandscapeLayer& InLayer, int32 InCurrentWeightmapToProcessIndex, const FIntPoint& InLandscapeBase, bool InOutputDebugName, FLandscapeTexture2DResource* InOutTextureData,
@@ -3757,26 +3787,27 @@ int32 ALandscape::RegenerateLayersWeightmaps(const TArray<ULandscapeComponent*>&
 		{
 			// Will generate CPU read back resource, if required
 			bool bHasPendingInitResource = false;
-			Info->ForAllLandscapeProxies([&](ALandscapeProxy* Proxy)
+			for (ULandscapeComponent* Component : InLandscapeComponents)
 			{
-				for (ULandscapeComponent* Component : Proxy->LandscapeComponents)
+				const TArray<UTexture2D*>& ComponentWeightmapTextures = Component->GetWeightmapTextures();
+
+				for (UTexture2D* WeightmapTexture : ComponentWeightmapTextures)
 				{
-					const TArray<UTexture2D*>& ComponentWeightmapTextures = Component->GetWeightmapTextures();
+					ALandscapeProxy* Proxy = Component->GetLandscapeProxy();
+					FLandscapeLayersTexture2DCPUReadBackResource** CPUReadback = Proxy->WeightmapsCPUReadBack.Find(WeightmapTexture);
 
-					for (UTexture2D* WeightmapTexture : ComponentWeightmapTextures)
+					if (CPUReadback == nullptr)
 					{
-						FLandscapeLayersTexture2DCPUReadBackResource** CPUReadback = Component->GetLandscapeProxy()->WeightmapsCPUReadBack.Find(WeightmapTexture);
-
-						if (CPUReadback == nullptr)
-						{
-							FLandscapeLayersTexture2DCPUReadBackResource* NewWeightmapCPUReadBack = new FLandscapeLayersTexture2DCPUReadBackResource(WeightmapTexture->Source.GetSizeX(), WeightmapTexture->Source.GetSizeY(), WeightmapTexture->GetPixelFormat(), WeightmapTexture->Source.GetNumMips());
-							BeginInitResource(NewWeightmapCPUReadBack);
-							Component->GetLandscapeProxy()->WeightmapsCPUReadBack.Add(WeightmapTexture, NewWeightmapCPUReadBack);
-							bHasPendingInitResource = true;
-						}
+						FLandscapeLayersTexture2DCPUReadBackResource* NewWeightmapCPUReadBack = new FLandscapeLayersTexture2DCPUReadBackResource(WeightmapTexture->Source.GetSizeX(), WeightmapTexture->Source.GetSizeY(), WeightmapTexture->GetPixelFormat(), WeightmapTexture->Source.GetNumMips());
+						const uint8* LockedMip = WeightmapTexture->Source.LockMip(0);
+						NewWeightmapCPUReadBack->UpdateHashFromTextureSource(LockedMip);
+						WeightmapTexture->Source.UnlockMip(0);
+						BeginInitResource(NewWeightmapCPUReadBack);
+						Proxy->WeightmapsCPUReadBack.Add(WeightmapTexture, NewWeightmapCPUReadBack);
+						bHasPendingInitResource = true;
 					}
 				}
-			});
+			}
 
 			if (bHasPendingInitResource)
 			{
@@ -4143,15 +4174,18 @@ void ALandscape::ResolveLayersWeightmapTexture(const TArray<ULandscapeComponent*
 		return;
 	}
 
-	TArray<UTexture2D*> ProcessedWeightmaps;
+	TMap<UTexture2D*,bool> ProcessedWeightmaps;
+	TArray<ULandscapeComponent*> ChangedComponents;
 
 	for (ULandscapeComponent* Component : InLandscapeComponents)
 	{
-		TArray<UTexture2D*>& ComponentWeightmaps = Component->GetWeightmapTextures();
+		TArray<UTexture2D*>& ComponentWeightmaps = Component->GetWeightmapTextures(false);
+		bool bChanged = false;
 
 		for (UTexture2D* ComponentLayerWeightmap : ComponentWeightmaps)
 		{
-			if (!ProcessedWeightmaps.Contains(ComponentLayerWeightmap))
+			bool* bValue = ProcessedWeightmaps.Find(ComponentLayerWeightmap);
+			if (bValue == nullptr)
 			{
 				FLandscapeLayersTexture2DCPUReadBackResource** CPUReadback = Component->GetLandscapeProxy()->WeightmapsCPUReadBack.Find(ComponentLayerWeightmap);
 
@@ -4159,13 +4193,25 @@ void ALandscape::ResolveLayersWeightmapTexture(const TArray<ULandscapeComponent*
 				{
 					continue;
 				}
-
-				ProcessedWeightmaps.Add(ComponentLayerWeightmap);
-
-				ResolveLayersTexture(*CPUReadback, ComponentLayerWeightmap);
+								
+				bool bWeightmapChanged = ResolveLayersTexture(*CPUReadback, ComponentLayerWeightmap);
+				ProcessedWeightmaps.Add(ComponentLayerWeightmap, bWeightmapChanged);
+				bChanged |= bWeightmapChanged;
 			}
+			else
+			{
+				bChanged |= *bValue;
+			}
+		}
+
+		if (bChanged)
+		{
+			Component->MarkPackageDirty();
+			ChangedComponents.Add(Component);
 		}
 	}
+
+	InvalidateGeneratedComponentData(ChangedComponents);
 }
 
 bool ALandscape::HasLayersContent() const
@@ -4220,18 +4266,8 @@ void ALandscape::RequestLayersContentUpdateForceAll(ELandscapeLayerUpdateMode In
 		{
 			if (Proxy)
 			{
-				if (bUpdateHeightmap || bUpdateWeightmap)
-				{
-					Proxy->InvalidateGeneratedComponentData();
-				}
-
 				for (ULandscapeComponent* Component : Proxy->LandscapeComponents)
 				{
-					if (bUpdateWeightmap || bUpdateHeightmap)
-					{
-						Component->Modify();
-					}
-
 					if (bUpdateHeightmap)
 					{
 						Component->RequestHeightmapUpdate(bUpdateAllHeightmap, bUpdateHeightCollision);
@@ -4467,9 +4503,9 @@ void ALandscape::UpdateLayersContent(bool bInWaitForStreaming, bool bInSkipMonit
 	}
 
 	TArray<ULandscapeComponent*> AllLandscapeComponents;
-	GetLandscapeInfo()->ForAllLandscapeComponents([&AllLandscapeComponents](ULandscapeComponent* InLandscapeComponent)
+	GetLandscapeInfo()->ForAllLandscapeComponents([&AllLandscapeComponents](ULandscapeComponent* Component)
 	{
-		AllLandscapeComponents.Add(InLandscapeComponent);
+		AllLandscapeComponents.Add(Component);
 	});
 
 	int32 ProcessedModes = 0;
@@ -4546,7 +4582,7 @@ bool ALandscape::UpdateCollisionAndClients(const TArray<ULandscapeComponent*>& I
 				if (ULandscapeHeightfieldCollisionComponent* CollisionComp = LandscapeComponent->CollisionComponent.Get())
 				{
 					FNavigationSystem::UpdateComponentData(*CollisionComp);
-					CollisionComp->SnapFoliageInstances();
+					//CollisionComp->SnapFoliageInstances();
 				}
 			}
 			else
@@ -5198,7 +5234,7 @@ void ALandscape::ClearLayer(int32 InLayerIndex, TSet<ULandscapeComponent*>* InCo
 	}
 }
 
-void ALandscape::ClearLayer(const FGuid& InLayerGuid, TSet<ULandscapeComponent*>* InComponents, ELandscapeClearMode InClearMode)
+void ALandscape::ClearLayer(const FGuid& InLayerGuid, TSet<ULandscapeComponent*>* InComponents, ELandscapeClearMode InClearMode, bool bMarkPackageDirty)
 {
 	ensure(HasLayersContent());
 
@@ -5209,7 +5245,7 @@ void ALandscape::ClearLayer(const FGuid& InLayerGuid, TSet<ULandscapeComponent*>
 		return;
 	}
 
-	Modify();
+	Modify(bMarkPackageDirty);
 	FScopedSetLandscapeEditingLayer Scope(this, Layer->Guid, [=] { RequestLayersContentUpdate(ELandscapeLayerUpdateMode::Update_All); });
 
 	TArray<uint16> NewHeightData;
@@ -5246,7 +5282,7 @@ void ALandscape::ClearLayer(const FGuid& InLayerGuid, TSet<ULandscapeComponent*>
 				if (!Proxies.Find(Proxy))
 				{
 					Proxies.Add(Proxy);
-					Proxy->Modify();
+					Proxy->Modify(bMarkPackageDirty);
 				}
 			}
 		}
@@ -5255,12 +5291,13 @@ void ALandscape::ClearLayer(const FGuid& InLayerGuid, TSet<ULandscapeComponent*>
 	{
 		LandscapeInfo->ForAllLandscapeProxies([&](ALandscapeProxy* Proxy)
 		{
-			Proxy->Modify();
+			Proxy->Modify(bMarkPackageDirty);
 			Components.Append(Proxy->LandscapeComponents);
 		});
 	}
 
 	FLandscapeEditDataInterface LandscapeEdit(LandscapeInfo);
+	FLandscapeDoNotDirtyScope DoNotDirtyScope(LandscapeEdit, !bMarkPackageDirty);
 	for (ULandscapeComponent* Component : Components)
 	{
 		if (InClearMode & ELandscapeClearMode::Clear_Heightmap)
@@ -5359,6 +5396,54 @@ FLandscapeLayer* ALandscape::GetLandscapeSplinesReservedLayer()
 
 LANDSCAPE_API extern bool GDisableUpdateLandscapeMaterialInstances;
 
+uint32 ULandscapeComponent::ComputeLayerHash() const
+{
+	UTexture2D* Heightmap = GetHeightmap(true);
+	const uint8* MipData = Heightmap->Source.LockMip(0);
+	uint32 Hash = FCrc::MemCrc32(MipData, Heightmap->GetSizeX()*Heightmap->GetSizeY() * sizeof(FColor));
+	Heightmap->Source.UnlockMip(0);
+
+	// Copy to sort
+	const TArray<UTexture2D*>& Weightmaps = GetWeightmapTextures(true);
+	TArray<FWeightmapLayerAllocationInfo> AllocationInfos = GetWeightmapLayerAllocations(true);
+	
+    // Sort allocations infos by LayerInfo Path so the Weightmaps hahses get ordered properly
+    AllocationInfos.Sort([](const FWeightmapLayerAllocationInfo& A, const FWeightmapLayerAllocationInfo& B)
+	{
+		FString PathA(A.LayerInfo ? A.LayerInfo->GetPathName() : FString());
+		FString PathB(B.LayerInfo ? B.LayerInfo->GetPathName() : FString());
+
+		return PathA < PathB;
+	});
+
+	for (const FWeightmapLayerAllocationInfo& AllocationInfo : AllocationInfos)
+	{
+		if (AllocationInfo.WeightmapTextureChannel != 255 && AllocationInfo.WeightmapTextureIndex != 255)
+		{
+            // Compute hash of actual data of the texture that is owned by the component (per Texture Channel)
+			UTexture2D* Weightmap = Weightmaps[AllocationInfo.WeightmapTextureIndex];
+			MipData = Weightmap->Source.LockMip(0) + ChannelOffsets[AllocationInfo.WeightmapTextureChannel];
+			TArray<uint8> ChannelData;
+			ChannelData.AddDefaulted(Weightmap->GetSizeX()*Weightmap->GetSizeY());
+			int32 TexSize = (SubsectionSizeQuads + 1)*NumSubsections;
+			for (int32 TexY = 0; TexY < TexSize; TexY++)
+			{
+				for (int32 TexX = 0; TexX < TexSize; TexX++)
+				{
+					const int32 Index = (TexX + TexY * TexSize);
+				
+					ChannelData.GetData()[Index] = MipData[4 * Index];
+				}
+			}
+
+			Hash = FCrc::MemCrc32(ChannelData.GetData(), Weightmap->GetSizeX()*Weightmap->GetSizeY(), Hash);
+			Weightmap->Source.UnlockMip(0);
+		}
+	}
+
+	return Hash;
+}
+
 void ALandscape::UpdateLandscapeSplines(const FGuid& InTargetLayer, bool bInUpdateOnlySelected, bool bInForceUpdateAllCompoments)
 {
 	check(CanHaveLayersContent());
@@ -5373,13 +5458,57 @@ void ALandscape::UpdateLandscapeSplines(const FGuid& InTargetLayer, bool bInUpda
 		TSet<ULandscapeComponent*>* ModifiedComponent = nullptr;
 		if (LandscapeSplinesTargetLayerGuid.IsValid())
 		{
-			ClearLayer(LandscapeSplinesTargetLayerGuid, (!bInForceUpdateAllCompoments && LandscapeSplinesAffectedComponents.Num()) ? &LandscapeSplinesAffectedComponents : nullptr);
+			TMap<ULandscapeComponent*, uint32> PreviousHashes;
+			{
+				FLandscapeEditDataInterface LandscapeEdit(LandscapeInfo);
+				
+				LandscapeInfo->ForAllLandscapeComponents([&](ULandscapeComponent* Component)
+				{
+					// Was never computed
+					if (Component->SplineHash == 0)
+					{
+						Component->Modify(); // mark package dirty
+						Component->SplineHash = DefaultSplineHash;
+					}
+
+					PreviousHashes.Add(Component, Component->SplineHash);
+					Component->Modify(false);
+					Component->SplineHash = DefaultSplineHash;
+				});
+			}
+
+			// Clear layers without affecting weightmap allocations
+			const bool bMarkPackageDirty = false;
+			ClearLayer(LandscapeSplinesTargetLayerGuid, (!bInForceUpdateAllCompoments && LandscapeSplinesAffectedComponents.Num()) ? &LandscapeSplinesAffectedComponents : nullptr, Clear_All, bMarkPackageDirty);
 			LandscapeSplinesAffectedComponents.Empty();
 			ModifiedComponent = &LandscapeSplinesAffectedComponents;
 			// For now, in Landscape Layer System Mode with a reserved layer for splines, we always update all the splines since we clear the whole layer first
 			bInUpdateOnlySelected = false;
+
+			// Apply splines without clearing up weightmap allocations
+			LandscapeInfo->ApplySplines(bInUpdateOnlySelected, ModifiedComponent, bMarkPackageDirty);		
+		
+			for (const TPair<ULandscapeComponent*, uint32>& Pair : PreviousHashes)
+			{
+				if (LandscapeSplinesAffectedComponents.Contains(Pair.Key))
+				{
+					uint32 NewHash = Pair.Key->ComputeLayerHash();
+					if (NewHash != Pair.Value)
+					{
+						Pair.Key->MarkPackageDirty();
+					}
+					Pair.Key->SplineHash = NewHash;
+				}
+				else if (Pair.Key->SplineHash == DefaultSplineHash && Pair.Value != DefaultSplineHash)
+				{
+					Pair.Key->MarkPackageDirty();
+				}
+			}
+		}
+		else
+		{
+			LandscapeInfo->ApplySplines(bInUpdateOnlySelected, ModifiedComponent);
 		}
-		LandscapeInfo->ApplySplines(bInUpdateOnlySelected, ModifiedComponent);
 		GDisableUpdateLandscapeMaterialInstances = false;
 	}
 }
diff --git a/Engine/Source/Runtime/Landscape/Private/LandscapeSplineRaster.cpp b/Engine/Source/Runtime/Landscape/Private/LandscapeSplineRaster.cpp
--- a/Engine/Source/Runtime/Landscape/Private/LandscapeSplineRaster.cpp
+++ b/Engine/Source/Runtime/Landscape/Private/LandscapeSplineRaster.cpp
@@ -457,7 +457,7 @@ void RasterizeSegmentAlpha(int32& MinX, int32& MinY, int32& MaxX, int32& MaxY, F
 	LandscapeEdit.GetComponentsInRegion(MinX, MinY, MaxX, MaxY, &ModifiedComponents);
 }
 
-bool ULandscapeInfo::ApplySplines(bool bOnlySelected, TSet<ULandscapeComponent*>* OutModifiedComponents)
+bool ULandscapeInfo::ApplySplines(bool bOnlySelected, TSet<ULandscapeComponent*>* OutModifiedComponents, bool bMarkPackageDirty)
 {
 	bool bResult = false;
 
@@ -468,13 +468,13 @@ bool ULandscapeInfo::ApplySplines(bool bOnlySelected, TSet<ULandscapeComponent*>
 
 	ForAllLandscapeProxies([&](ALandscapeProxy* Proxy)
 	{
-		bResult |= ApplySplinesInternal(bOnlySelected, Proxy, OutModifiedComponents);
+		bResult |= ApplySplinesInternal(bOnlySelected, Proxy, OutModifiedComponents, bMarkPackageDirty);
 	});
 
 	return bResult;
 }
 
-bool ULandscapeInfo::ApplySplinesInternal(bool bOnlySelected, ALandscapeProxy* Proxy, TSet<ULandscapeComponent*>* OutModifiedComponents)
+bool ULandscapeInfo::ApplySplinesInternal(bool bOnlySelected, ALandscapeProxy* Proxy, TSet<ULandscapeComponent*>* OutModifiedComponents, bool bMarkPackageDirty)
 {
 	if (!Proxy || !Proxy->SplineComponent || !Proxy->SplineComponent->IsRegistered() || Proxy->SplineComponent->ControlPoints.Num() == 0 || Proxy->SplineComponent->Segments.Num() == 0)
 	{
@@ -486,6 +486,7 @@ bool ULandscapeInfo::ApplySplinesInternal(bool bOnlySelected, ALandscapeProxy* P
 	const FTransform SplineToLandscape = Proxy->SplineComponent->GetComponentTransform().GetRelativeTransform(Proxy->LandscapeActorToWorld());
 
 	FLandscapeEditDataInterface LandscapeEdit(this);
+	FLandscapeDoNotDirtyScope DoNotDirtyScope(LandscapeEdit, !bMarkPackageDirty);
 	TSet<ULandscapeComponent*> ModifiedComponents;
 
 	// I'd dearly love to use FIntRect in this code, but Landscape works with "Inclusive Max" and FIntRect is "Exclusive Max"
@@ -624,10 +625,13 @@ bool ULandscapeInfo::ApplySplinesInternal(bool bOnlySelected, ALandscapeProxy* P
 
 	LandscapeEdit.Flush();
 		
-	ALandscapeProxy::InvalidateGeneratedComponentData(ModifiedComponents);
-		
-	for (ULandscapeComponent* Component : ModifiedComponents)
+	if (!Proxy->HasLayersContent())
 	{
+		ALandscapeProxy::InvalidateGeneratedComponentData(ModifiedComponents);
+	}
+
+	for (ULandscapeComponent* Component : ModifiedComponents)
+	{	
 		if (Component->GetLandscapeProxy()->HasLayersContent())
 		{
 			Component->RequestHeightmapUpdate();
diff --git a/Engine/Source/Runtime/Landscape/Public/LandscapeEdit.h b/Engine/Source/Runtime/Landscape/Public/LandscapeEdit.h
--- a/Engine/Source/Runtime/Landscape/Public/LandscapeEdit.h
+++ b/Engine/Source/Runtime/Landscape/Public/LandscapeEdit.h
@@ -62,7 +62,7 @@ struct FLandscapeTextureDataInfo
 		TArray<FUpdateTextureRegion2D> MipUpdateRegions;
 	};
 
-	FLandscapeTextureDataInfo(UTexture2D* InTexture);
+	FLandscapeTextureDataInfo(UTexture2D* InTexture, bool bShouldDirtyPackage);
 	virtual ~FLandscapeTextureDataInfo();
 
 	// returns true if we need to block on the render thread before unlocking the mip data
@@ -105,6 +105,9 @@ struct LANDSCAPE_API FLandscapeTextureDataInterface
 {
 	FLandscapeTextureDataInterface(bool bInUploadTextureChangesToGPU = true);
 	virtual ~FLandscapeTextureDataInterface();
+		
+	void SetShouldDirtyPackage(bool bValue) { bShouldDirtyPackage = bValue; }
+	bool GetShouldDirtyPackage() const { return bShouldDirtyPackage; }
 
 	// Texture data access
 	FLandscapeTextureDataInfo* GetTextureDataInfo(UTexture2D* Texture);
@@ -130,9 +133,9 @@ struct LANDSCAPE_API FLandscapeTextureDataInterface
 private:
 	TMap<UTexture2D*, FLandscapeTextureDataInfo*> TextureDataMap;
 	bool bUploadTextureChangesToGPU;
+	bool bShouldDirtyPackage;
 };
 
-
 struct LANDSCAPE_API FLandscapeEditDataInterface : public FLandscapeTextureDataInterface
 {
 	// tor
@@ -293,6 +296,31 @@ private:
 	const ULandscapeLayerInfoObject* ChooseReplacementLayer(const ULandscapeLayerInfoObject* LayerInfo, int32 ComponentIndexX, int32 SubIndexX, int32 SubX, int32 ComponentIndexY, int32 SubIndexY, int32 SubY, TMap<FIntPoint, TMap<const ULandscapeLayerInfoObject*, uint32>>& LayerInfluenceCache, TArrayView<const uint8* const> LayerDataPtrs);
 };
 
+struct LANDSCAPE_API FLandscapeDoNotDirtyScope
+{
+	FLandscapeDoNotDirtyScope(FLandscapeEditDataInterface& InEditInterface, bool bInScopeEnabled = true)
+		: EditInterface(InEditInterface), bScopeEnabled(bInScopeEnabled)
+	{
+		if (bScopeEnabled)
+		{
+			bPreviousValue = EditInterface.GetShouldDirtyPackage();
+			EditInterface.SetShouldDirtyPackage(false);
+		}
+	}
+
+	~FLandscapeDoNotDirtyScope()
+	{
+		if (bScopeEnabled)
+		{
+			EditInterface.SetShouldDirtyPackage(bPreviousValue);
+		}
+	}
+
+	FLandscapeEditDataInterface& EditInterface;
+	bool bPreviousValue;
+	bool bScopeEnabled;
+};
+
 template<typename T>
 void FLandscapeEditDataInterface::ShrinkData(TArray<T>& Data, int32 OldMinX, int32 OldMinY, int32 OldMaxX, int32 OldMaxY, int32 NewMinX, int32 NewMinY, int32 NewMaxX, int32 NewMaxY)
 {
@@ -367,11 +395,10 @@ struct FHeightmapAccessor
 			bool bUpdateFoliage = false;
 			bool bUpdateNormals = false;
 
-			ALandscapeProxy::InvalidateGeneratedComponentData(Components);
-
             // Landscape Layers are updates are delayed and done in  ALandscape::TickLayers
 			if (!LandscapeEdit->HasLandscapeLayersContent())
 			{
+				ALandscapeProxy::InvalidateGeneratedComponentData(Components);
 				bUpdateNormals = true;
 				for (ULandscapeComponent* Component : Components)
 				{
@@ -529,7 +556,6 @@ struct FAlphamapAccessor
 		TSet<ULandscapeComponent*> Components;
 		if (LandscapeEdit.GetComponentsInRegion(X1, Y1, X2, Y2, &Components))
 		{
-			ALandscapeProxy::InvalidateGeneratedComponentData(Components);
 			for (ULandscapeComponent* LandscapeComponent : Components)
 			{
 				// Flag both modes depending on client calling SetData
@@ -538,6 +564,7 @@ struct FAlphamapAccessor
 			
 			if (!LandscapeEdit.HasLandscapeLayersContent())
 			{
+				ALandscapeProxy::InvalidateGeneratedComponentData(Components);
 				ModifiedComponents.Append(Components);
 			}
 
