Fixed snap foliage to landscape not working when using Landscape Layer system

diff --git a/Engine/Source/Runtime/Landscape/Private/LandscapeEditLayers.cpp b/Engine/Source/Runtime/Landscape/Private/LandscapeEditLayers.cpp
--- a/Engine/Source/Runtime/Landscape/Private/LandscapeEditLayers.cpp
+++ b/Engine/Source/Runtime/Landscape/Private/LandscapeEditLayers.cpp
@@ -4582,7 +4582,7 @@ bool ALandscape::UpdateCollisionAndClients(const TArray<ULandscapeComponent*>& I
 				if (ULandscapeHeightfieldCollisionComponent* CollisionComp = LandscapeComponent->CollisionComponent.Get())
 				{
 					FNavigationSystem::UpdateComponentData(*CollisionComp);
-					//CollisionComp->SnapFoliageInstances();
+					CollisionComp->SnapFoliageInstances();
 				}
 			}
 			else
