Reparent Landscape mayterial instance materials so they won't get name collisions when exporting to text during copy/paste operations.
Fix Landscape collision data computations when simple collisions are enabled when exporting to text during copy/paste operations.
#rnx

diff --git a/Engine/Source/Runtime/Landscape/Private/LandscapeCollision.cpp b/Engine/Source/Runtime/Landscape/Private/LandscapeCollision.cpp
--- a/Engine/Source/Runtime/Landscape/Private/LandscapeCollision.cpp
+++ b/Engine/Source/Runtime/Landscape/Private/LandscapeCollision.cpp
@@ -747,11 +747,13 @@ bool ULandscapeMeshCollisionComponent::CookCollisionData(const FName& Format, bo
 	TArray<uint16>			MaterialIndices;
 
 	const int32 CollisionSizeVerts = CollisionSizeQuads + 1;
+	const int32 SimpleCollisionSizeVerts = SimpleCollisionSizeQuads > 0 ? SimpleCollisionSizeQuads + 1 : 0;
 	const int32 NumVerts = FMath::Square(CollisionSizeVerts);
+	const int32 NumSimpleVerts = FMath::Square(SimpleCollisionSizeVerts);
 
 	const uint16* Heights = (const uint16*)CollisionHeightData.LockReadOnly();
 	const uint16* XYOffsets = (const uint16*)CollisionXYOffsetData.LockReadOnly();
-	check(CollisionHeightData.GetElementCount() == NumVerts);
+	check(CollisionHeightData.GetElementCount() == NumVerts + NumSimpleVerts);
 	check(CollisionXYOffsetData.GetElementCount() == NumVerts * 2);
 
 	const uint8* DominantLayers = nullptr;
@@ -1927,10 +1929,13 @@ void ULandscapeHeightfieldCollisionComponent::ExportCustomProperties(FOutputDevi
 		return;
 	}
 
-	uint16* Heights = (uint16*)CollisionHeightData.Lock(LOCK_READ_ONLY);
-	int32 NumHeights = FMath::Square(CollisionSizeQuads + 1);
+	int32 CollisionSizeVerts = CollisionSizeQuads + 1;
+	int32 SimpleCollisionSizeVerts = SimpleCollisionSizeQuads > 0 ? SimpleCollisionSizeQuads + 1 : 0;
+	int32 NumHeights = FMath::Square(CollisionSizeVerts) + FMath::Square(SimpleCollisionSizeVerts);
 	check(CollisionHeightData.GetElementCount() == NumHeights);
 
+	uint16* Heights = (uint16*)CollisionHeightData.Lock(LOCK_READ_ONLY);
+
 	Out.Logf(TEXT("%sCustomProperties CollisionHeightData "), FCString::Spc(Indent));
 	for (int32 i = 0; i < NumHeights; i++)
 	{
@@ -1962,7 +1967,9 @@ void ULandscapeHeightfieldCollisionComponent::ImportCustomProperties(const TCHAR
 {
 	if (FParse::Command(&SourceText, TEXT("CollisionHeightData")))
 	{
-		int32 NumHeights = FMath::Square(CollisionSizeQuads + 1);
+		int32 CollisionSizeVerts = CollisionSizeQuads + 1;
+		int32 SimpleCollisionSizeVerts = SimpleCollisionSizeQuads > 0 ? SimpleCollisionSizeQuads + 1 : 0;
+		int32 NumHeights = FMath::Square(CollisionSizeVerts) + FMath::Square(SimpleCollisionSizeVerts);
 
 		CollisionHeightData.Lock(LOCK_READ_WRITE);
 		uint16* Heights = (uint16*)CollisionHeightData.Realloc(NumHeights);
@@ -2046,7 +2053,9 @@ void ULandscapeMeshCollisionComponent::ImportCustomProperties(const TCHAR* Sourc
 {
 	if (FParse::Command(&SourceText, TEXT("CollisionHeightData")))
 	{
-		int32 NumHeights = FMath::Square(CollisionSizeQuads + 1);
+		int32 CollisionSizeVerts = CollisionSizeQuads + 1;
+		int32 SimpleCollisionSizeVerts = SimpleCollisionSizeQuads > 0 ? SimpleCollisionSizeQuads + 1 : 0;
+		int32 NumHeights = FMath::Square(CollisionSizeVerts) + FMath::Square(SimpleCollisionSizeVerts);
 
 		CollisionHeightData.Lock(LOCK_READ_WRITE);
 		uint16* Heights = (uint16*)CollisionHeightData.Realloc(NumHeights);
diff --git a/Engine/Source/Runtime/Landscape/Private/LandscapeEdit.cpp b/Engine/Source/Runtime/Landscape/Private/LandscapeEdit.cpp
--- a/Engine/Source/Runtime/Landscape/Private/LandscapeEdit.cpp
+++ b/Engine/Source/Runtime/Landscape/Private/LandscapeEdit.cpp
@@ -363,7 +363,7 @@ void ULandscapeComponent::UpdateMaterialInstances_Internal(FMaterialUpdateContex
 		if (CombinationMaterialInstance != nullptr)
 		{
 			// Create the instance for this component, that will use the layer combination instance.
-			UMaterialInstanceConstant* MaterialInstance = NewObject<ULandscapeMaterialInstanceConstant>(this);
+			UMaterialInstanceConstant* MaterialInstance = NewObject<ULandscapeMaterialInstanceConstant>(GetOuter());
 			MaterialInstances[MaterialIndex] = MaterialInstance;
 
 			// Material Instances don't support Undo/Redo (the shader map goes out of sync and crashes happen)
