Don't read bounds from ULandscapeComponent in render thread
Use Bounds cached in the render proxy
Think this may be one source of NaNs in GetComponentScreenSize() which in turn is triggering asserts elsewhere
Added a NaN check in GetComponentScreenSize() to help find any other bad data sources

diff --git a/Engine/Source/Runtime/Landscape/Private/LandscapeRender.cpp b/Engine/Source/Runtime/Landscape/Private/LandscapeRender.cpp
--- a/Engine/Source/Runtime/Landscape/Private/LandscapeRender.cpp
+++ b/Engine/Source/Runtime/Landscape/Private/LandscapeRender.cpp
@@ -1427,6 +1427,8 @@ float FLandscapeComponentSceneProxy::GetComponentScreenSize(const FSceneView* Vi
 		SquaredScreenRadius = FMath::Square(ScreenMultiple * ElementRadius) / FMath::Max(1.0f, DistSquared);
 	}
 
+	// If we hit NaN check see UE-64538
+	check(FMath::IsFinite(SquaredScreenRadius));
 	return SquaredScreenRadius * 2.0f;
 }
 
@@ -2011,7 +2013,7 @@ void FLandscapeComponentSceneProxy::CalculateBatchElementLOD(const FSceneView& I
 		InOutLODData.UseCombinedMeshBatch = false; // default to individual batch render
 
 		float SubSectionMaxExtend = ComponentMaxExtend / 2.0f;
-		float SubSectionRadius = LandscapeComponent->Bounds.SphereRadius / 2.0f;
+		float SubSectionRadius = GetBounds().SphereRadius / 2.0f;
 		float CombinedScreenRatio = 0.0f;
 		bool AllSubSectionHaveSameScreenSize = true;
 
@@ -2173,7 +2175,7 @@ void* FLandscapeComponentSceneProxy::InitViewCustomData(const FSceneView& InView
 	// If a valid screen size was provided, we use it instead of recomputing it
 	if (InMeshScreenSizeSquared < 0.0f)
 	{
-		LODData->ComponentScreenSize = GetComponentScreenSize(&InView, LandscapeComponent->Bounds.Origin, ComponentMaxExtend, LandscapeComponent->Bounds.SphereRadius);
+		LODData->ComponentScreenSize = GetComponentScreenSize(&InView, GetBounds().Origin, ComponentMaxExtend, GetBounds().SphereRadius);
 	}
 
 	CalculateBatchElementLOD(InView, LODData->ComponentScreenSize, InViewLODScale, *LODData, false);
@@ -2240,8 +2242,8 @@ void FLandscapeComponentSceneProxy::ComputeTessellationFalloffShaderValues(const
 	{
 		if (UseTessellationComponentScreenSizeFalloff)
 		{
-			float MaxTesselationDistance = ComputeBoundsDrawDistance(FMath::Sqrt(TessellationComponentSquaredScreenSize), LandscapeComponent->Bounds.SphereRadius / 2.0f, InViewProjectionMatrix);
-			float FallOffStartingDistance = FMath::Min(ComputeBoundsDrawDistance(FMath::Sqrt(FMath::Min(FMath::Square(TessellationComponentScreenSizeFalloff), TessellationComponentSquaredScreenSize)), LandscapeComponent->Bounds.SphereRadius / 2.0f, InViewProjectionMatrix) - MaxTesselationDistance, MaxTesselationDistance);
+			float MaxTesselationDistance = ComputeBoundsDrawDistance(FMath::Sqrt(TessellationComponentSquaredScreenSize), GetBounds().SphereRadius / 2.0f, InViewProjectionMatrix);
+			float FallOffStartingDistance = FMath::Min(ComputeBoundsDrawDistance(FMath::Sqrt(FMath::Min(FMath::Square(TessellationComponentScreenSizeFalloff), TessellationComponentSquaredScreenSize)), GetBounds().SphereRadius / 2.0f, InViewProjectionMatrix) - MaxTesselationDistance, MaxTesselationDistance);
 
 			// Calculate the falloff using a = C - K * d by sending C & K into the shader
 			OutC = MaxTesselationDistance / (MaxTesselationDistance - FallOffStartingDistance);
@@ -2355,7 +2357,7 @@ float FLandscapeComponentSceneProxy::GetNeighborLOD(const FSceneView& InView, fl
 
 	if (ComputeNeighborCustomDataLOD)
 	{
-		FBoxSphereBounds NeighborBounds = LandscapeComponent->Bounds;
+		FBoxSphereBounds NeighborBounds = GetBounds();
 		float NeighborMaxExtends = ComponentMaxExtend;
 		FLandscapeComponentSceneProxy* NeighborSceneProxy = nullptr;
 
@@ -2524,12 +2526,12 @@ FLODMask FLandscapeComponentSceneProxy::GetCustomLOD(const FSceneView& InView, f
 	else if (InView.Family->EngineShowFlags.LOD)
 	{
 		int8 PotentialLOD = INDEX_NONE;
-		OutScreenSizeSquared = GetComponentScreenSize(&InView, LandscapeComponent->Bounds.Origin, ComponentMaxExtend, LandscapeComponent->Bounds.SphereRadius);
+		OutScreenSizeSquared = GetComponentScreenSize(&InView, GetBounds().Origin, ComponentMaxExtend, GetBounds().SphereRadius);
 
 		if (NumSubsections > 1)
 		{
 			float SubSectionMaxExtend = ComponentMaxExtend / 2.0f;
-			float SubSectionRadius = LandscapeComponent->Bounds.SphereRadius / 2.0f;
+			float SubSectionRadius = GetBounds().SphereRadius / 2.0f;
 
 			// Compute screen size of each sub section to determine if we should use the combined logic or the individual logic
 			float ScreenSizeSquared = GetComponentScreenSize(&InView, SubSectionScreenSizeTestingPosition[0], SubSectionMaxExtend, SubSectionRadius);
@@ -2607,12 +2609,12 @@ FLODMask FLandscapeComponentSceneProxy::GetCustomWholeSceneShadowLOD(const FScen
 
 		if (PrimitiveCustomData == nullptr)
 		{
-			ScreenSizeSquared = GetComponentScreenSize(&InView, LandscapeComponent->Bounds.Origin, ComponentMaxExtend, LandscapeComponent->Bounds.SphereRadius);
+			ScreenSizeSquared = GetComponentScreenSize(&InView, GetBounds().Origin, ComponentMaxExtend, GetBounds().SphereRadius);
 
 			if (NumSubsections > 1)
 			{
 				float SubSectionMaxExtend = ComponentMaxExtend / 2.0f;
-				float SubSectionRadius = LandscapeComponent->Bounds.SphereRadius / 2.0f;
+				float SubSectionRadius = GetBounds().SphereRadius / 2.0f;
 
 				// Compute screen size of each sub section to determine if we should use the combined logic or the individual logic
 				float SubSectionScreenSizeSquared = GetComponentScreenSize(&InView, SubSectionScreenSizeTestingPosition[0], SubSectionMaxExtend, SubSectionRadius);
@@ -2761,7 +2763,7 @@ void FLandscapeComponentSceneProxy::GetDynamicMeshElements(const TArray<const FS
 			if (HasTessellationEnabled)
 			{
 				check(AvailableMaterials.Num() > 1);
-				float ScreenSizeSquared = GetComponentScreenSize(View, LandscapeComponent->Bounds.Origin, ComponentMaxExtend, LandscapeComponent->Bounds.SphereRadius);
+				float ScreenSizeSquared = GetComponentScreenSize(View, GetBounds().Origin, ComponentMaxExtend, GetBounds().SphereRadius);
 				DisableTessellation = (ScreenSizeSquared < TessellationComponentSquaredScreenSize * View->LODDistanceFactor) ? true : false;
 
 				if (!DisableTessellation)
@@ -3188,7 +3190,7 @@ void FLandscapeComponentSceneProxy::GetDynamicRayTracingInstances(FRayTracingMat
 					SectionRayTracingStates[SubSectionIdx].LODBias = GetShaderLODBias();
 				}
 
-				float ComponentScreenSize = GetComponentScreenSize(Context.ReferenceView, LandscapeComponent->Bounds.Origin, ComponentMaxExtend, LandscapeComponent->Bounds.SphereRadius);
+				float ComponentScreenSize = GetComponentScreenSize(Context.ReferenceView, GetBounds().Origin, ComponentMaxExtend, GetBounds().SphereRadius);
 
 				FViewCustomDataLOD CurrentLODData;
 				CalculateBatchElementLOD(*Context.ReferenceView, ComponentScreenSize, Context.ReferenceView->LODDistanceFactor, CurrentLODData, true);
@@ -3960,7 +3962,7 @@ public:
 		}
 		else
 		{
-			float ComponentScreenSize = SceneProxy->GetComponentScreenSize(InView, SceneProxy->LandscapeComponent->Bounds.Origin, SceneProxy->ComponentMaxExtend, SceneProxy->LandscapeComponent->Bounds.SphereRadius);
+			float ComponentScreenSize = SceneProxy->GetComponentScreenSize(InView, SceneProxy->GetBounds().Origin, SceneProxy->ComponentMaxExtend, SceneProxy->GetBounds().SphereRadius);
 			
 			FLandscapeComponentSceneProxy::FViewCustomDataLOD CurrentLODData;
 			SceneProxy->CalculateBatchElementLOD(*InView, ComponentScreenSize, InView->LODDistanceFactor, CurrentLODData, true);
