Exposed getting landscape layer weight value to Blueprint (editor-only)
#rnx

diff --git a/Engine/Source/Runtime/Landscape/Classes/LandscapeComponent.h b/Engine/Source/Runtime/Landscape/Classes/LandscapeComponent.h
--- a/Engine/Source/Runtime/Landscape/Classes/LandscapeComponent.h
+++ b/Engine/Source/Runtime/Landscape/Classes/LandscapeComponent.h
@@ -723,6 +723,14 @@ public:
 	UFUNCTION(BlueprintCallable, Category = "Landscape|Runtime|Material")
 	class UMaterialInstanceDynamic* GetMaterialInstanceDynamic(int32 InIndex) const;
 
+	/** Gets the landscape paint layer weight value at the given position using LandscapeLayerInfo . Returns 0 in case it fails. */
+	UFUNCTION(BlueprintCallable, Category = "Landscape|Editor")
+	LANDSCAPE_API float EditorGetPaintLayerWeightAtLocation(const FVector& InLocation, ULandscapeLayerInfoObject* PaintLayer);
+
+	/** Gets the landscape paint layer weight value at the given position using layer name. Returns 0 in case it fails. */
+	UFUNCTION(BlueprintCallable, Category = "Landscape|Editor")
+	LANDSCAPE_API float EditorGetPaintLayerWeightByNameAtLocation(const FVector& InLocation, const FName InPaintLayerName);
+		
 	/** Get the landscape actor associated with this component. */
 	ALandscape* GetLandscapeActor() const;
 
diff --git a/Engine/Source/Runtime/Landscape/Classes/LandscapeHeightfieldCollisionComponent.h b/Engine/Source/Runtime/Landscape/Classes/LandscapeHeightfieldCollisionComponent.h
--- a/Engine/Source/Runtime/Landscape/Classes/LandscapeHeightfieldCollisionComponent.h
+++ b/Engine/Source/Runtime/Landscape/Classes/LandscapeHeightfieldCollisionComponent.h
@@ -75,6 +75,10 @@ class ULandscapeHeightfieldCollisionComponent : public UPrimitiveComponent
 	UPROPERTY()
 	TLazyObjectPtr<ULandscapeComponent> RenderComponent;
 
+	/** Returns associated landscape component */
+	UFUNCTION(BlueprintCallable, Category = "Landscape")
+	ULandscapeComponent* GetRenderComponent() const;
+
 	struct FPhysXHeightfieldRef : public FRefCountedObject
 	{
 		FGuid Guid;
diff --git a/Engine/Source/Runtime/Landscape/Private/LandscapeBlueprintSupport.cpp b/Engine/Source/Runtime/Landscape/Private/LandscapeBlueprintSupport.cpp
--- a/Engine/Source/Runtime/Landscape/Private/LandscapeBlueprintSupport.cpp
+++ b/Engine/Source/Runtime/Landscape/Private/LandscapeBlueprintSupport.cpp
@@ -89,4 +89,26 @@ void ALandscapeProxy::EditorSetLandscapeMaterial(UMaterialInterface* NewLandscap
 		PostEditChangeProperty(PropertyChangedEvent);
 	}
 #endif
+}
+
+float ULandscapeComponent::EditorGetPaintLayerWeightByNameAtLocation(const FVector& InLocation, const FName InPaintLayerName)
+{
+#if WITH_EDITOR
+	ULandscapeInfo* LandscapeInfo = GetLandscapeInfo();
+	ULandscapeLayerInfoObject* PaintLayer = LandscapeInfo ? LandscapeInfo->GetLayerInfoByName(InPaintLayerName) : nullptr;
+	return GetLayerWeightAtLocation(InLocation, PaintLayer);
+#else
+	return 0.f;
+#endif
+}
+
+float ULandscapeComponent::EditorGetPaintLayerWeightAtLocation(const FVector& InLocation, ULandscapeLayerInfoObject* PaintLayer)
+{
+#if WITH_EDITOR
+	if (PaintLayer)
+	{
+		return GetLayerWeightAtLocation(InLocation, PaintLayer);
+	}
+#endif
+	return 0.f;
 }
\ No newline at end of file
diff --git a/Engine/Source/Runtime/Landscape/Private/LandscapeCollision.cpp b/Engine/Source/Runtime/Landscape/Private/LandscapeCollision.cpp
--- a/Engine/Source/Runtime/Landscape/Private/LandscapeCollision.cpp
+++ b/Engine/Source/Runtime/Landscape/Private/LandscapeCollision.cpp
@@ -2184,3 +2184,8 @@ ULandscapeHeightfieldCollisionComponent::ULandscapeHeightfieldCollisionComponent
 	// landscape collision components should be deterministically created and therefor are addressable over the network
 	SetNetAddressable();
 }
+
+ULandscapeComponent* ULandscapeHeightfieldCollisionComponent::GetRenderComponent() const
+{
+	return RenderComponent.Get();
+}
\ No newline at end of file
