Fix up connections when a spline point mesh is removed.

diff --git a/Engine/Source/Runtime/Landscape/Private/LandscapeSplines.cpp b/Engine/Source/Runtime/Landscape/Private/LandscapeSplines.cpp
--- a/Engine/Source/Runtime/Landscape/Private/LandscapeSplines.cpp
+++ b/Engine/Source/Runtime/Landscape/Private/LandscapeSplines.cpp
@@ -1694,6 +1694,8 @@ void ULandscapeSplineControlPoint::UpdateSplinePoints(bool bUpdateCollision, boo
 	{
 		MeshComponent = nullptr;
 		ForeignWorld = nullptr;
+
+		AutoSetConnections(false);
 	}
 
 	// Destroy any unused components
