Fix Crash in LandscapeSpline 'Move to current level' code when no ForeignMesh data existed for a Segment

diff --git a/Engine/Source/Runtime/Landscape/Private/LandscapeSplines.cpp b/Engine/Source/Runtime/Landscape/Private/LandscapeSplines.cpp
--- a/Engine/Source/Runtime/Landscape/Private/LandscapeSplines.cpp
+++ b/Engine/Source/Runtime/Landscape/Private/LandscapeSplines.cpp
@@ -1070,14 +1070,17 @@ void ULandscapeSplinesComponent::RemoveAllForeignMeshComponents(ULandscapeSpline
 	if (ForeignWorldSplineData)
 	{
 		auto* ForeignSplineSegmentData = ForeignWorldSplineData->FindSegmentData(Owner);
-
-		for (auto* MeshComponent : ForeignSplineSegmentData->MeshComponents)
+		if (ForeignSplineSegmentData)
 		{
-			checkSlow(MeshComponentForeignOwnersMap.FindRef(MeshComponent) == Owner);
-			verifySlow(MeshComponentForeignOwnersMap.Remove(MeshComponent) == 1);
+			for (auto* MeshComponent : ForeignSplineSegmentData->MeshComponents)
+			{
+				checkSlow(MeshComponentForeignOwnersMap.FindRef(MeshComponent) == Owner);
+				verifySlow(MeshComponentForeignOwnersMap.Remove(MeshComponent) == 1);
+			}
+			ForeignSplineSegmentData->MeshComponents.Empty();
+			verifySlow(ForeignWorldSplineData->ForeignSplineSegmentData.RemoveSingle(*ForeignSplineSegmentData) == 1);
 		}
-		ForeignSplineSegmentData->MeshComponents.Empty();
-		verifySlow(ForeignWorldSplineData->ForeignSplineSegmentData.RemoveSingle(*ForeignSplineSegmentData) == 1);
+
 		if (ForeignWorldSplineData->IsEmpty())
 		{
 			verifySlow(ForeignWorldSplineDataMap.Remove(OwnerWorld) == 1);
