Landscape Layers: Optimize Paint & Scuplting (Partial Update when CVar Landscape.Optim == 1 (defaults to 1))
Optional optim: Disable Brush Render while painting/sculpting (CVar Landscape.BrushOptim == 1 (defaults to 0))

diff --git a/Engine/Source/Runtime/Landscape/Classes/Landscape.h b/Engine/Source/Runtime/Landscape/Classes/Landscape.h
--- a/Engine/Source/Runtime/Landscape/Classes/Landscape.h
+++ b/Engine/Source/Runtime/Landscape/Classes/Landscape.h
@@ -331,11 +331,13 @@ public:
 private:
 	void TickLayers(float DeltaTime, ELevelTick TickType, FActorTickFunction& ThisTickFunction);
 	void CreateLayersRenderingResource();
+	void GetLandscapeComponentNeighborsToRender(ULandscapeComponent* LandscapeComponent, TSet<ULandscapeComponent*>& NeighborComponents) const;
+	void GetLandscapeComponentWeightmapsToRender(ULandscapeComponent* LandscapeComponent, TSet<ULandscapeComponent*>& WeightmapComponents) const;
 	void UpdateLayersContent(bool bInWaitForStreaming = false, bool bInSkipMonitorLandscapeEdModeChanges = false);
 	void MonitorShaderCompilation();
 	void MonitorLandscapeEdModeChanges();
-	int32 RegenerateLayersHeightmaps(const TArray<ULandscapeComponent*>& InLandscapeComponents, bool bInWaitForStreaming);
-	int32 RegenerateLayersWeightmaps(const TArray<ULandscapeComponent*>& InLandscapeComponents, bool bInWaitForStreaming);
+	int32 RegenerateLayersHeightmaps(const TArray<ULandscapeComponent*>& InLandscapeComponents, const TArray<ULandscapeComponent*>& InLandscapeComponentsToResolve, bool bInWaitForStreaming);
+	int32 RegenerateLayersWeightmaps(const TArray<ULandscapeComponent*>& InLandscapeComponents, const TArray<ULandscapeComponent*>& InLandscapeComponentsToResolve, bool bInWaitForStreaming);
 	bool UpdateCollisionAndClients(const TArray<ULandscapeComponent*>& InLandscapeComponents, const int32 InContentUpdateModes);
 	void ResolveLayersHeightmapTexture(const TArray<ULandscapeComponent*>& InLandscapeComponents);
 	void ResolveLayersWeightmapTexture(const TArray<ULandscapeComponent*>& InLandscapeComponents);
@@ -350,7 +352,7 @@ private:
 														  TArray<struct FLandscapeLayerWeightmapExtractMaterialLayersComponentData>& OutComponentData, TMap<ULandscapeLayerInfoObject*, int32>& OutLayerInfoObjects);
 	void PrepareComponentDataToPackMaterialLayersCS(int32 InCurrentWeightmapToProcessIndex, const FIntPoint& InLandscapeBase, bool InOutputDebugName, const TArray<ULandscapeComponent*>& InAllLandscapeComponents, TArray<UTexture2D*>& InOutProcessedWeightmaps,
 													TArray<class FLandscapeLayersTexture2DCPUReadBackResource*>& OutProcessedCPUReadBackTexture, TArray<struct FLandscapeLayerWeightmapPackMaterialLayersComponentData>& OutComponentData);
-	void ReallocateLayersWeightmaps(const TArray<ULandscapeLayerInfoObject*>& InBrushRequiredAllocations);
+	void ReallocateLayersWeightmaps(const TArray<ULandscapeComponent*>& InLandscapeComponents, const TArray<ULandscapeLayerInfoObject*>& InBrushRequiredAllocations);
 	void InitializeLayersWeightmapResources();
 	bool GenerateZeroAllocationPerComponents(const TArray<ALandscapeProxy*>& InAllLandscape, const TMap<ULandscapeLayerInfoObject*, bool>& InWeightmapLayersBlendSubstractive);
 
diff --git a/Engine/Source/Runtime/Landscape/Classes/LandscapeComponent.h b/Engine/Source/Runtime/Landscape/Classes/LandscapeComponent.h
--- a/Engine/Source/Runtime/Landscape/Classes/LandscapeComponent.h
+++ b/Engine/Source/Runtime/Landscape/Classes/LandscapeComponent.h
@@ -936,6 +936,7 @@ public:
 	LANDSCAPE_API void RequestHeightmapUpdate(bool bUpdateAll = false, bool bUpdateCollision = true);
 	LANDSCAPE_API void RequestEditingClientUpdate();
 	LANDSCAPE_API void RequestDeferredClientUpdate();
+	LANDSCAPE_API uint32 GetLayerUpdateFlagPerMode() const { return LayerUpdateFlagPerMode; }
 	LANDSCAPE_API uint32 ComputeWeightmapsHash();
 #endif
 
diff --git a/Engine/Source/Runtime/Landscape/Classes/LandscapeWeightmapUsage.h b/Engine/Source/Runtime/Landscape/Classes/LandscapeWeightmapUsage.h
--- a/Engine/Source/Runtime/Landscape/Classes/LandscapeWeightmapUsage.h
+++ b/Engine/Source/Runtime/Landscape/Classes/LandscapeWeightmapUsage.h
@@ -43,6 +43,17 @@ public:
 		}
 	}
 
+	void ClearUsage(ULandscapeComponent* Component)
+	{
+		for (int8 i = 0; i < NumChannels; ++i)
+		{
+			if (ChannelUsage[i] == Component)
+			{
+				ChannelUsage[i] = nullptr;
+			}
+		}
+	}
+
 	bool IsEmpty() const
 	{
 		return FreeChannelCount() == NumChannels;
diff --git a/Engine/Source/Runtime/Landscape/Private/LandscapeEditLayers.cpp b/Engine/Source/Runtime/Landscape/Private/LandscapeEditLayers.cpp
--- a/Engine/Source/Runtime/Landscape/Private/LandscapeEditLayers.cpp
+++ b/Engine/Source/Runtime/Landscape/Private/LandscapeEditLayers.cpp
@@ -67,6 +67,16 @@ static TAutoConsoleVariable<int32> CVarLandscapeSimulatePhysics(
 	TEXT("landscape.SimulatePhysics"),
 	0,
 	TEXT("This will enable physic simulation on worlds containing landscape."));
+
+static TAutoConsoleVariable<int32> CVarLandscapeLayerOptim(
+	TEXT("landscape.Optim"),
+	1,
+	TEXT("This will enable landscape layers optim."));
+
+static TAutoConsoleVariable<int32> CVarLandscapeLayerBrushOptim(
+	TEXT("landscape.BrushOptim"),
+	0,
+	TEXT("This will enable landscape layers optim."));
 #endif
 
 DECLARE_GPU_STAT_NAMED(LandscapeLayersRender, TEXT("Landscape Layer System Render"));
@@ -2696,7 +2706,7 @@ bool ALandscape::PrepareLayersHeightmapTextureResources(bool bInWaitForStreaming
 	return IsReady;
 }
 
-int32 ALandscape::RegenerateLayersHeightmaps(const TArray<ULandscapeComponent*>& InLandscapeComponents, bool bInWaitForStreaming)
+int32 ALandscape::RegenerateLayersHeightmaps(const TArray<ULandscapeComponent*>& InLandscapeComponentsToRender, const TArray<ULandscapeComponent*>& InLandscapeComponentsToResolve, bool bInWaitForStreaming)
 {
 	SCOPE_CYCLE_COUNTER(STAT_LandscapeLayersRegenerateHeightmaps);
 	ULandscapeInfo* Info = GetLandscapeInfo();
@@ -2704,16 +2714,17 @@ int32 ALandscape::RegenerateLayersHeightmaps(const TArray<ULandscapeComponent*>&
 	const int32 AllHeightmapUpdateModes = (ELandscapeLayerUpdateMode::Update_Heightmap_All | ELandscapeLayerUpdateMode::Update_Heightmap_Editing | ELandscapeLayerUpdateMode::Update_Heightmap_Editing_NoCollision);
 	const int32 HeightmapUpdateModes = LayerContentUpdateModes & AllHeightmapUpdateModes;
 	const bool bForceRender = CVarOutputLayersDebugDrawCallName.GetValueOnAnyThread() == 1;
+	const bool bSkipBrush = CVarLandscapeLayerBrushOptim.GetValueOnAnyThread() == 1 && ((HeightmapUpdateModes & AllHeightmapUpdateModes) == ELandscapeLayerUpdateMode::Update_Heightmap_Editing);
 
 	if ((HeightmapUpdateModes == 0 && !bForceRender) || Info == nullptr || !PrepareLayersHeightmapTextureResources(bInWaitForStreaming))
 	{
 		return 0;
 	}
-	   
+
 	// Init CPU Readbacks
 	if (HeightmapUpdateModes)
 	{
-		for(ULandscapeComponent* Component : InLandscapeComponents)
+		for(ULandscapeComponent* Component : InLandscapeComponentsToRender)
 		{
 			UTexture2D* ComponentHeightmap = Component->GetHeightmap(false);
 			ALandscapeProxy* Proxy = Component->GetLandscapeProxy();
@@ -2765,7 +2776,7 @@ int32 ALandscape::RegenerateLayersHeightmaps(const TArray<ULandscapeComponent*>&
 		};
 
 		// Calculate Top Left Lambda
-		auto GetUniqueHeightmaps = [&](const TArray<ULandscapeComponent*>& InLandscapeComponent, TArray<FHeightmapTopLeft>& OutHeightmaps, const FIntPoint& LandscapeBaseQuads, const FGuid& InLayerGuid = FGuid())
+		auto GetUniqueHeightmaps = [&](const TArray<ULandscapeComponent*>& InLandscapeComponents, TArray<FHeightmapTopLeft>& OutHeightmaps, const FIntPoint& LandscapeBaseQuads, const FGuid& InLayerGuid = FGuid())
 		{
 			FScopedSetLandscapeEditingLayer Scope(this, InLayerGuid);
 			
@@ -2829,7 +2840,7 @@ int32 ALandscape::RegenerateLayersHeightmaps(const TArray<ULandscapeComponent*>&
 
 			{
 				TArray<FHeightmapTopLeft> LayerHeightmaps;
-				GetUniqueHeightmaps(InLandscapeComponents, LayerHeightmaps, LandscapeExtent.Min, Layer.Guid);
+				GetUniqueHeightmaps(InLandscapeComponentsToRender, LayerHeightmaps, LandscapeExtent.Min, Layer.Guid);
 				for (const FHeightmapTopLeft& LayerHeightmap : LayerHeightmaps)
 				{
 					AddDeferredCopyLayersTexture(LayerHeightmap.Texture, LandscapeScratchRT1, nullptr, LayerHeightmap.TopLeftSectionBase);
@@ -2839,17 +2850,17 @@ int32 ALandscape::RegenerateLayersHeightmaps(const TArray<ULandscapeComponent*>&
 
 			// NOTE: From this point on, we always work in non atlas, we'll convert back at the end to atlas only
 			DrawHeightmapComponentsToRenderTarget(OutputDebugName ? FString::Printf(TEXT("LS Height: %s += -> NonAtlas %s"), *Layer.Name.ToString(), *LandscapeScratchRT1->GetName(), *LandscapeScratchRT2->GetName()) : GEmptyDebugName,
-												  InLandscapeComponents, LandscapeExtent.Min, LandscapeScratchRT1, nullptr, LandscapeScratchRT2, ERTDrawingType::RTAtlasToNonAtlas, true, ShaderParams);
+												  InLandscapeComponentsToRender, LandscapeExtent.Min, LandscapeScratchRT1, nullptr, LandscapeScratchRT2, ERTDrawingType::RTAtlasToNonAtlas, true, ShaderParams);
 
 			ShaderParams.ApplyLayerModifiers = true;
 
 			// Combine Current layer with current result
 			DrawHeightmapComponentsToRenderTarget(OutputDebugName ? FString::Printf(TEXT("LS Height: %s += -> CombinedNonAtlas %s"), *Layer.Name.ToString(), *LandscapeScratchRT2->GetName(), *CombinedHeightmapNonAtlasRT->GetName()) : GEmptyDebugName,
-												InLandscapeComponents, LandscapeExtent.Min, LandscapeScratchRT2, FirstLayer ? nullptr : LandscapeScratchRT3, CombinedHeightmapNonAtlasRT, ERTDrawingType::RTNonAtlas, FirstLayer, ShaderParams);
+												InLandscapeComponentsToRender, LandscapeExtent.Min, LandscapeScratchRT2, FirstLayer ? nullptr : LandscapeScratchRT3, CombinedHeightmapNonAtlasRT, ERTDrawingType::RTNonAtlas, FirstLayer, ShaderParams);
 
 			ShaderParams.ApplyLayerModifiers = false;
 
-			if (Layer.bVisible)
+			if (Layer.bVisible && !bSkipBrush)
 			{
 				// Draw each brushes
 				for (int32 i = 0; i < Layer.Brushes.Num(); ++i)
@@ -2884,19 +2895,19 @@ int32 ALandscape::RegenerateLayersHeightmaps(const TArray<ULandscapeComponent*>&
 		ShaderParams.GridSize = GetRootComponent()->RelativeScale3D;
 
 		DrawHeightmapComponentsToRenderTarget(OutputDebugName ? FString::Printf(TEXT("LS Height: %s = -> CombinedNonAtlasNormals : %s"), *CombinedHeightmapNonAtlasRT->GetName(), *LandscapeScratchRT1->GetName()) : GEmptyDebugName,
-											  InLandscapeComponents, LandscapeExtent.Min, CombinedHeightmapNonAtlasRT, nullptr, LandscapeScratchRT1, ERTDrawingType::RTNonAtlas, true, ShaderParams);
+											  InLandscapeComponentsToRender, LandscapeExtent.Min, CombinedHeightmapNonAtlasRT, nullptr, LandscapeScratchRT1, ERTDrawingType::RTNonAtlas, true, ShaderParams);
 
 		ShaderParams.GenerateNormals = false;
 
 		DrawHeightmapComponentsToRenderTarget(OutputDebugName ? FString::Printf(TEXT("LS Height: %s = -> CombinedAtlasFinal : %s"), *LandscapeScratchRT1->GetName(), *CombinedHeightmapAtlasRT->GetName()) : GEmptyDebugName,
-											  InLandscapeComponents, LandscapeExtent.Min, LandscapeScratchRT1, nullptr, CombinedHeightmapAtlasRT, ERTDrawingType::RTNonAtlasToAtlas, true, ShaderParams);
+											  InLandscapeComponentsToRender, LandscapeExtent.Min, LandscapeScratchRT1, nullptr, CombinedHeightmapAtlasRT, ERTDrawingType::RTNonAtlasToAtlas, true, ShaderParams);
 
-		DrawHeightmapComponentsToRenderTargetMips(InLandscapeComponents, LandscapeExtent.Min, CombinedHeightmapAtlasRT, true, ShaderParams);
+		DrawHeightmapComponentsToRenderTargetMips(InLandscapeComponentsToRender, LandscapeExtent.Min, CombinedHeightmapAtlasRT, true, ShaderParams);
 
 		// Copy back all Mips to original heightmap data
 		{
 			TArray<FHeightmapTopLeft> Heightmaps;
-			GetUniqueHeightmaps(InLandscapeComponents, Heightmaps, LandscapeExtent.Min);
+			GetUniqueHeightmaps(InLandscapeComponentsToRender, Heightmaps, LandscapeExtent.Min);
 			for (const FHeightmapTopLeft& Heightmap : Heightmaps)
 			{
 				FIntPoint TextureTopLeftSectionBase = Heightmap.TopLeftSectionBase;
@@ -2926,10 +2937,10 @@ int32 ALandscape::RegenerateLayersHeightmaps(const TArray<ULandscapeComponent*>&
 
 	if (HeightmapUpdateModes)
 	{
-		ResolveLayersHeightmapTexture(InLandscapeComponents);
+		ResolveLayersHeightmapTexture(InLandscapeComponentsToResolve);
 
 		// Partial Component Update
-		for (ULandscapeComponent* Component : InLandscapeComponents)
+		for (ULandscapeComponent* Component : InLandscapeComponentsToResolve)
 		{
 			if(Component->IsUpdateFlagEnabledForModes(ELandscapeComponentUpdateFlag::Component_Update_Heightmap_Collision, HeightmapUpdateModes))
 			{
@@ -3243,7 +3254,7 @@ void ALandscape::PrepareComponentDataToPackMaterialLayersCS(int32 InCurrentWeigh
 	}
 }
 
-void ALandscape::ReallocateLayersWeightmaps(const TArray<ULandscapeLayerInfoObject*>& InBrushRequiredAllocations)
+void ALandscape::ReallocateLayersWeightmaps(const TArray<ULandscapeComponent*>& InLandscapeComponents, const TArray<ULandscapeLayerInfoObject*>& InBrushRequiredAllocations)
 {
 	SCOPE_CYCLE_COUNTER(STAT_LandscapeLayersReallocateWeightmaps);
 
@@ -3254,15 +3265,8 @@ void ALandscape::ReallocateLayersWeightmaps(const TArray<ULandscapeLayerInfoObje
 		return;
 	}
 
-	TArray<ULandscapeComponent*> AllLandscapeComponents;
-
-	Info->ForAllLandscapeProxies([&AllLandscapeComponents](ALandscapeProxy* Proxy)
-	{
-		AllLandscapeComponents.Append(Proxy->LandscapeComponents);
-	});
-		
 	// Clear allocation data
-	for (ULandscapeComponent* Component : AllLandscapeComponents)
+	for (ULandscapeComponent* Component : InLandscapeComponents)
 	{
 		TArray<FWeightmapLayerAllocationInfo>& BaseLayerAllocations = Component->GetWeightmapLayerAllocations();
 
@@ -3279,48 +3283,45 @@ void ALandscape::ReallocateLayersWeightmaps(const TArray<ULandscapeLayerInfoObje
 			ULandscapeWeightmapUsage* Usage = WeightmapTexturesUsage[i];
 			check(Usage != nullptr);
 
-			Usage->ClearUsage();
+			Usage->ClearUsage(Component);
 		}
 	}
 		
 	// Build a map of all the allocation per components
 	TMap<ULandscapeComponent*, TArray<ULandscapeLayerInfoObject*>> LayerAllocsPerComponent;
 
-	Info->ForAllLandscapeProxies([&](ALandscapeProxy* Proxy)
+	for (FLandscapeLayer& Layer : LandscapeLayers)
 	{
-		for (FLandscapeLayer& Layer : LandscapeLayers)
+		for (ULandscapeComponent* Component : InLandscapeComponents)
 		{
-			for (ULandscapeComponent* Component : Proxy->LandscapeComponents)
-			{
-				TArray<ULandscapeLayerInfoObject*>* ComponentLayerAlloc = LayerAllocsPerComponent.Find(Component);
+			TArray<ULandscapeLayerInfoObject*>* ComponentLayerAlloc = LayerAllocsPerComponent.Find(Component);
 
-				if (ComponentLayerAlloc == nullptr)
-				{
-					TArray<ULandscapeLayerInfoObject*> NewLayerAllocs;
-					ComponentLayerAlloc = &LayerAllocsPerComponent.Add(Component, NewLayerAllocs);
-				}
+			if (ComponentLayerAlloc == nullptr)
+			{
+				TArray<ULandscapeLayerInfoObject*> NewLayerAllocs;
+				ComponentLayerAlloc = &LayerAllocsPerComponent.Add(Component, NewLayerAllocs);
+			}
 
-				FLandscapeLayerComponentData* LayerComponentData = Component->GetLayerData(Layer.Guid);
+			FLandscapeLayerComponentData* LayerComponentData = Component->GetLayerData(Layer.Guid);
 
-				for (const FWeightmapLayerAllocationInfo& LayerWeightmapAllocation : LayerComponentData->WeightmapData.LayerAllocations)
+			for (const FWeightmapLayerAllocationInfo& LayerWeightmapAllocation : LayerComponentData->WeightmapData.LayerAllocations)
+			{
+				if (LayerWeightmapAllocation.LayerInfo != nullptr)
 				{
-					if (LayerWeightmapAllocation.LayerInfo != nullptr)
-					{
-						ComponentLayerAlloc->AddUnique(LayerWeightmapAllocation.LayerInfo);
-					}
+					ComponentLayerAlloc->AddUnique(LayerWeightmapAllocation.LayerInfo);
 				}
+			}
 
-				// Add the brush alloc also
-				for (ULandscapeLayerInfoObject* BrushLayerInfo : InBrushRequiredAllocations)
+			// Add the brush alloc also
+			for (ULandscapeLayerInfoObject* BrushLayerInfo : InBrushRequiredAllocations)
+			{
+				if (BrushLayerInfo != nullptr)
 				{
-					if (BrushLayerInfo != nullptr)
-					{
-						ComponentLayerAlloc->AddUnique(BrushLayerInfo);
-					}
+					ComponentLayerAlloc->AddUnique(BrushLayerInfo);
 				}
 			}
 		}
-	});
+	}
 
 	// Determine if the Final layer need to add/remove some alloc
 	for (auto& ItPair : LayerAllocsPerComponent)
@@ -3355,7 +3356,7 @@ void ALandscape::ReallocateLayersWeightmaps(const TArray<ULandscapeLayerInfoObje
 	// Realloc the weightmap so it will create proper texture (if needed) and will set the allocations information
 	TArray<UTexture2D*> NewCreatedTextures;
 
-	for (ULandscapeComponent* Component : AllLandscapeComponents)
+	for (ULandscapeComponent* Component : InLandscapeComponents)
 	{
 		Component->ReallocateWeightmaps(nullptr, false, false, true, false, nullptr, &NewCreatedTextures);
 	}
@@ -3374,14 +3375,14 @@ void ALandscape::ReallocateLayersWeightmaps(const TArray<ULandscapeLayerInfoObje
 	//GDisableAutomaticTextureMaterialUpdateDependencies = false;
 
 	// Clean-up unused weightmap CPUReadback resources
-	Info->ForAllLandscapeProxies([&AllLandscapeComponents](ALandscapeProxy* Proxy)
+	Info->ForAllLandscapeProxies([](ALandscapeProxy* Proxy)
 	{
 		TArray<UTexture2D*, TInlineAllocator<64>> EntriesToRemoveFromMap;
 		for (auto& Pair : Proxy->WeightmapsCPUReadBack)
 		{
 			UTexture2D* WeightmapTextureKey = Pair.Key;
 			bool IsTextureReferenced = false;
-			for (ULandscapeComponent* Component : AllLandscapeComponents)
+			for (ULandscapeComponent* Component : Proxy->LandscapeComponents)
 			{
 				for (UTexture2D* WeightmapTexture : Component->GetWeightmapTextures(false))
 				{
@@ -3540,11 +3541,12 @@ bool ALandscape::PrepareLayersWeightmapTextureResources(bool bInWaitForStreaming
 	return IsReady;
 }
 
-int32 ALandscape::RegenerateLayersWeightmaps(const TArray<ULandscapeComponent*>& InLandscapeComponents, bool bInWaitForStreaming)
+int32 ALandscape::RegenerateLayersWeightmaps(const TArray<ULandscapeComponent*>& InLandscapeComponentsToRender, const TArray<ULandscapeComponent*>& InLandscapeComponentsToResolve, bool bInWaitForStreaming)
 {
 	SCOPE_CYCLE_COUNTER(STAT_LandscapeLayersRegenerateWeightmaps);
 	const int32 AllWeightmapUpdateModes = (ELandscapeLayerUpdateMode::Update_Weightmap_All | ELandscapeLayerUpdateMode::Update_Weightmap_Editing | ELandscapeLayerUpdateMode::Update_Weightmap_Editing_NoCollision);
 	const int32 WeightmapUpdateModes = LayerContentUpdateModes & AllWeightmapUpdateModes;
+	const bool bSkipBrush = CVarLandscapeLayerBrushOptim.GetValueOnAnyThread() == 1 && ((WeightmapUpdateModes & AllWeightmapUpdateModes) == ELandscapeLayerUpdateMode::Update_Heightmap_Editing);
 	const bool bForceRender = CVarOutputLayersDebugDrawCallName.GetValueOnAnyThread() == 1;
 	
 	ULandscapeInfo* Info = GetLandscapeInfo();
@@ -3602,7 +3604,7 @@ int32 ALandscape::RegenerateLayersWeightmaps(const TArray<ULandscapeComponent*>&
 			TMap<ULandscapeLayerInfoObject*, int32> LayerInfoObjects; // <LayerInfoObj, LayerIndex>
 
 			// Determine if some brush want to write to layer that we have currently no data on
-			if (Layer.bVisible)
+			if (Layer.bVisible && !bSkipBrush)
 			{
 				for (int32 LayerInfoSettingsIndex = 0; LayerInfoSettingsIndex < Info->Layers.Num(); ++LayerInfoSettingsIndex)
 				{
@@ -3630,7 +3632,7 @@ int32 ALandscape::RegenerateLayersWeightmaps(const TArray<ULandscapeComponent*>&
 
 				// Prepare compute shader data
 				TArray<FLandscapeLayerWeightmapExtractMaterialLayersComponentData> ComponentsData;
-				PrepareComponentDataToExtractMaterialLayersCS(InLandscapeComponents, Layer, CurrentWeightmapToProcessIndex, LandscapeExtent.Min, OutputDebugName, WeightmapScratchExtractLayerTextureResource, ComponentsData, LayerInfoObjects);
+				PrepareComponentDataToExtractMaterialLayersCS(InLandscapeComponentsToRender, Layer, CurrentWeightmapToProcessIndex, LandscapeExtent.Min, OutputDebugName, WeightmapScratchExtractLayerTextureResource, ComponentsData, LayerInfoObjects);
 
 				HasFoundWeightmapToProcess = ComponentsData.Num() > 0;
 
@@ -3699,7 +3701,7 @@ int32 ALandscape::RegenerateLayersWeightmaps(const TArray<ULandscapeComponent*>&
 					PSShaderParams.LayerAlpha = LayerInfoObj == ALandscapeProxy::VisibilityLayer ? 1.0f : Layer.WeightmapAlpha; // visibility can't be affected by weight
 
 					DrawWeightmapComponentsToRenderTarget(OutputDebugName ? FString::Printf(TEXT("LS Weight: %s Paint: %s += -> %s"), *Layer.Name.ToString(), *LayerInfoObj->LayerName.ToString(), *LandscapeScratchRT1->GetName(), *LandscapeScratchRT2->GetName()) : GEmptyDebugName,
-														  InLandscapeComponents, LandscapeExtent.Min, LandscapeScratchRT1, nullptr, LandscapeScratchRT2, ERTDrawingType::RTAtlas, true, PSShaderParams, 0);
+														  InLandscapeComponentsToRender, LandscapeExtent.Min, LandscapeScratchRT1, nullptr, LandscapeScratchRT2, ERTDrawingType::RTAtlas, true, PSShaderParams, 0);
 
 					PSShaderParams.ApplyLayerModifiers = false;
 
@@ -3726,7 +3728,7 @@ int32 ALandscape::RegenerateLayersWeightmaps(const TArray<ULandscapeComponent*>&
 					}
 
 					DrawWeightmapComponentsToRenderTarget(OutputDebugName ? FString::Printf(TEXT("LS Weight: %s PaintLayer: %s, %s += -> Combined %s"), *Layer.Name.ToString(), *LayerInfoObj->LayerName.ToString(), *LandscapeScratchRT2->GetName(), *LandscapeScratchRT3->GetName()) : GEmptyDebugName,
-														  InLandscapeComponents, LandscapeExtent.Min, LandscapeScratchRT2, bFirstLayer ? nullptr : LandscapeScratchRT1, LandscapeScratchRT3, ERTDrawingType::RTAtlasToNonAtlas, true, PSShaderParams, 0);
+														  InLandscapeComponentsToRender, LandscapeExtent.Min, LandscapeScratchRT2, bFirstLayer ? nullptr : LandscapeScratchRT1, LandscapeScratchRT3, ERTDrawingType::RTAtlasToNonAtlas, true, PSShaderParams, 0);
 
 					PSShaderParams.OutputAsSubstractive = false;
 
@@ -3734,7 +3736,7 @@ int32 ALandscape::RegenerateLayersWeightmaps(const TArray<ULandscapeComponent*>&
 					DestDebugName = OutputDebugName ? TEXT("CombinedProcLayerWeightmap") : GEmptyDebugName;
 
 					// Handle brush blending
-					if (Layer.bVisible)
+					if (Layer.bVisible && !bSkipBrush)
 					{
 						// Draw each brushes				
 						for (int32 i = 0; i < Layer.Brushes.Num(); ++i)
@@ -3769,7 +3771,7 @@ int32 ALandscape::RegenerateLayersWeightmaps(const TArray<ULandscapeComponent*>&
 					}
 
 					DrawWeightmapComponentsToRenderTarget(OutputDebugName ? FString::Printf(TEXT("LS Weight: %s Combined Scratch No Border to %s Combined Scratch with Border"), *LandscapeScratchRT3->GetName(), *LandscapeScratchRT1->GetName()) : GEmptyDebugName,
-						InLandscapeComponents, LandscapeExtent.Min, LandscapeScratchRT3, nullptr, LandscapeScratchRT1, ERTDrawingType::RTNonAtlasToAtlas, true, PSShaderParams, 0);
+						InLandscapeComponentsToRender, LandscapeExtent.Min, LandscapeScratchRT3, nullptr, LandscapeScratchRT1, ERTDrawingType::RTNonAtlasToAtlas, true, PSShaderParams, 0);
 
 
 					CopyLayersTexture(SourceDebugName, LandscapeScratchRT1->GameThread_GetRenderTargetResource(), DestDebugName, CombinedLayersWeightmapAllMaterialLayersResource, nullptr, FIntPoint(0, 0), 0, 0, 0, LayerIndex);
@@ -3781,13 +3783,22 @@ int32 ALandscape::RegenerateLayersWeightmaps(const TArray<ULandscapeComponent*>&
 			bFirstLayer = false;
 		}
 
-		ReallocateLayersWeightmaps(BrushRequiredAllocations);
+		ReallocateLayersWeightmaps(InLandscapeComponentsToResolve, BrushRequiredAllocations);
+		TSet<UTexture2D*> ToResolve;
+		for (ULandscapeComponent* LandscapeComponentToResolve : InLandscapeComponentsToResolve)
+		{
+			const TArray<UTexture2D*>& ComponentWeightmapTextures = LandscapeComponentToResolve->GetWeightmapTextures();
+			for (UTexture2D* WeightmapTexture : ComponentWeightmapTextures)
+			{
+				ToResolve.Add(WeightmapTexture);
+			}
+		}
 
 		if (bHasWeightmapData)
 		{
 			// Will generate CPU read back resource, if required
 			bool bHasPendingInitResource = false;
-			for (ULandscapeComponent* Component : InLandscapeComponents)
+			for (ULandscapeComponent* Component : InLandscapeComponentsToRender)
 			{
 				const TArray<UTexture2D*>& ComponentWeightmapTextures = Component->GetWeightmapTextures();
 
@@ -3827,7 +3838,7 @@ int32 ALandscape::RegenerateLayersWeightmaps(const TArray<ULandscapeComponent*>&
 			while (HasFoundWeightmapToProcess)
 			{
 				TArray<FLandscapeLayerWeightmapPackMaterialLayersComponentData> PackLayersComponentsData;
-				PrepareComponentDataToPackMaterialLayersCS(CurrentWeightmapToProcessIndex, LandscapeExtent.Min, OutputDebugName, InLandscapeComponents, ProcessedWeightmaps, ProcessedCPUReadbackTextures, PackLayersComponentsData);
+				PrepareComponentDataToPackMaterialLayersCS(CurrentWeightmapToProcessIndex, LandscapeExtent.Min, OutputDebugName, InLandscapeComponentsToRender, ProcessedWeightmaps, ProcessedCPUReadbackTextures, PackLayersComponentsData);
 				HasFoundWeightmapToProcess = PackLayersComponentsData.Num() > 0;
 
 				// Perform the compute shader
@@ -3906,6 +3917,11 @@ int32 ALandscape::RegenerateLayersWeightmaps(const TArray<ULandscapeComponent*>&
 					for (; NextTextureIndexToProcess < ProcessedWeightmaps.Num(); ++NextTextureIndexToProcess)
 					{
 						UTexture2D* WeightmapTexture = ProcessedWeightmaps[NextTextureIndexToProcess];
+						if (!ToResolve.Contains(WeightmapTexture))
+						{
+							continue;
+						}
+
 						FLandscapeLayersTexture2DCPUReadBackResource* WeightmapCPUReadBack = ProcessedCPUReadbackTextures[NextTextureIndexToProcess];
 						FIntPoint TextureTopLeftPositionInAtlas(WeightmapTextureOutputOffset[NextTextureIndexToProcess - StartTextureIndex].X, WeightmapTextureOutputOffset[NextTextureIndexToProcess - StartTextureIndex].Y);
 
@@ -3933,14 +3949,14 @@ int32 ALandscape::RegenerateLayersWeightmaps(const TArray<ULandscapeComponent*>&
 			}
 		}
 
-		UpdateLayersMaterialInstances(InLandscapeComponents);
+		UpdateLayersMaterialInstances(InLandscapeComponentsToResolve);
 	}
 
 	if (WeightmapUpdateModes)
 	{
-		ResolveLayersWeightmapTexture(InLandscapeComponents);
+		ResolveLayersWeightmapTexture(InLandscapeComponentsToResolve);
 	
-		for (ULandscapeComponent* Component : InLandscapeComponents)
+		for (ULandscapeComponent* Component : InLandscapeComponentsToResolve)
 		{
 			if (Component->IsUpdateFlagEnabledForModes(ELandscapeComponentUpdateFlag::Component_Update_Weightmap_Collision, WeightmapUpdateModes))
 			{
@@ -4477,6 +4493,52 @@ void ALandscape::MonitorShaderCompilation()
 	}
 }
 
+void ALandscape::GetLandscapeComponentNeighborsToRender(ULandscapeComponent* LandscapeComponent, TSet<ULandscapeComponent*>& NeighborComponents) const
+{
+	ULandscapeInfo* LandscapeInfo = GetLandscapeInfo();
+	FIntPoint ComponentKey = LandscapeComponent->GetSectionBase() / ComponentSizeQuads;
+	
+	for (int32 IndexX = ComponentKey.X-1; IndexX <= ComponentKey.X+1; ++IndexX)
+	{
+		for (int32 IndexY = ComponentKey.Y-1; IndexY <= ComponentKey.Y+1; ++IndexY)
+		{
+			ULandscapeComponent* Result = LandscapeInfo->XYtoComponentMap.FindRef(FIntPoint(IndexX, IndexY));
+			if (Result && Result != LandscapeComponent)
+			{
+				NeighborComponents.Add(Result);
+			}
+		}
+	}
+}
+
+void ALandscape::GetLandscapeComponentWeightmapsToRender(ULandscapeComponent* LandscapeComponent, TSet<ULandscapeComponent*>& WeightmapComponents) const
+{
+	// Fill with Components that share the same weightmaps so that the Resolve of Weightmap Texture doesn't resolve null data.
+	for (ULandscapeWeightmapUsage* Usage : LandscapeComponent->GetWeightmapTexturesUsage(false))
+	{
+		for (int32 Channel = 0; Channel < ULandscapeWeightmapUsage::NumChannels; ++Channel)
+		{
+			if (Usage != nullptr && Usage->ChannelUsage[Channel] != nullptr)
+			{
+				ULandscapeComponent* Component = Usage->ChannelUsage[Channel];
+				Component->ForEachLayer([&](const FGuid& LayerGuid, FLandscapeLayerComponentData& LayerData)
+				{
+					for (ULandscapeWeightmapUsage* Usage : LayerData.WeightmapData.TextureUsages)
+					{
+						for (int32 Channel = 0; Channel < ULandscapeWeightmapUsage::NumChannels; ++Channel)
+						{
+							if (Usage != nullptr && Usage->ChannelUsage[Channel] != nullptr)
+							{
+								WeightmapComponents.Add(Usage->ChannelUsage[Channel]);
+							}
+						}
+					}
+				});
+			}
+		}
+	}
+}
+
 void ALandscape::UpdateLayersContent(bool bInWaitForStreaming, bool bInSkipMonitorLandscapeEdModeChanges)
 {
 	if (GetLandscapeInfo() == nullptr || !CanHaveLayersContent())
@@ -4502,15 +4564,73 @@ void ALandscape::UpdateLayersContent(bool bInWaitForStreaming, bool bInSkipMonit
 		return;
 	}
 
+	bool bUpdateAll = LayerContentUpdateModes & Update_All;
+
+	bool bPartialUpdate = !bForceRender && !bUpdateAll && CVarLandscapeLayerOptim.GetValueOnAnyThread() == 1;
+	TSet<UTexture*> Heightmaps;
+	TSet<ULandscapeComponent*> NeighborsComponents;
+	TSet<ULandscapeComponent*> WeightmapsComponents;
+	TSet<ULandscapeWeightmapUsage*> WeightmapUsagesToResolve;
+	TArray<ULandscapeComponent*> LandscapeComponentsHeightmapsToRender;
+	TArray<ULandscapeComponent*> LandscapeComponentsHeightmapsToResolve;
+	TArray<ULandscapeComponent*> LandscapeComponentsWeightmapsToRender;
+	TArray<ULandscapeComponent*> LandscapeComponentsWeightmapsToResolve;
 	TArray<ULandscapeComponent*> AllLandscapeComponents;
-	GetLandscapeInfo()->ForAllLandscapeComponents([&AllLandscapeComponents](ULandscapeComponent* Component)
+
+	TArray<ULandscapeComponent*> SkippedComponents;
+	GetLandscapeInfo()->ForAllLandscapeComponents([&](ULandscapeComponent* Component)
 	{
-		AllLandscapeComponents.Add(Component);
+		if (!bPartialUpdate || Component->GetLayerUpdateFlagPerMode() != 0)
+		{
+			AllLandscapeComponents.Add(Component);
+		
+			if (bPartialUpdate)
+			{
+				// Gather Neighbors (Neighbors need to be Rendered but not resolved so that the resolved Components have valid normals on edges)
+				GetLandscapeComponentNeighborsToRender(Component, NeighborsComponents);
+				// Gather Heightmaps (All Components sharing Heightmap textures need to be rendered and resolved)
+				Heightmaps.Add(Component->GetHeightmap(false));
+				// Gather WeightmapUsages (Components sharing weightmap usages with the resolved Components need to be rendered so that the resolving is valid)
+				GetLandscapeComponentWeightmapsToRender(Component, WeightmapsComponents);
+			}
+		}
+		else
+		{
+			SkippedComponents.Add(Component);
+		}
 	});
 
+	// Copy first list into others
+	LandscapeComponentsHeightmapsToResolve.Append(AllLandscapeComponents);
+	LandscapeComponentsHeightmapsToRender.Append(AllLandscapeComponents);
+	LandscapeComponentsWeightmapsToResolve.Append(AllLandscapeComponents);
+	LandscapeComponentsWeightmapsToRender.Append(AllLandscapeComponents);
+		
+	if (bPartialUpdate)
+	{
+		for (ULandscapeComponent* Component : SkippedComponents)
+		{
+			if (Heightmaps.Contains(Component->GetHeightmap(false)))
+			{
+				AllLandscapeComponents.Add(Component);
+				LandscapeComponentsHeightmapsToRender.Add(Component);
+				LandscapeComponentsHeightmapsToResolve.Add(Component);
+			}
+			else if (NeighborsComponents.Contains(Component))
+			{
+				LandscapeComponentsHeightmapsToRender.Add(Component);
+			}
+
+			if (WeightmapsComponents.Contains(Component))
+			{
+				LandscapeComponentsWeightmapsToRender.Add(Component);
+			}
+		}
+	}
+
 	int32 ProcessedModes = 0;
-	ProcessedModes |= RegenerateLayersHeightmaps(AllLandscapeComponents, bInWaitForStreaming);
-	ProcessedModes |= RegenerateLayersWeightmaps(AllLandscapeComponents, bInWaitForStreaming);
+	ProcessedModes |= RegenerateLayersHeightmaps(LandscapeComponentsHeightmapsToRender, LandscapeComponentsHeightmapsToResolve, bInWaitForStreaming);
+	ProcessedModes |= RegenerateLayersWeightmaps(LandscapeComponentsWeightmapsToRender, LandscapeComponentsWeightmapsToResolve, bInWaitForStreaming);
 	ProcessedModes |= (LayerContentUpdateModes & ELandscapeLayerUpdateMode::Update_Client_Deferred);
 	ProcessedModes |= (LayerContentUpdateModes & ELandscapeLayerUpdateMode::Update_Client_Editing);
 	LayerContentUpdateModes &= ~ProcessedModes;
@@ -5954,3 +6074,5 @@ bool FLandscapeLayerBrush::Initialize(const FIntRect& InLandscapeExtent, UTextur
 #endif // WITH_EDITOR
 
 #undef LOCTEXT_NAMESPACE
+
+
