Landscape Layers: Fix Grass in PIE

diff --git a/Engine/Source/Runtime/Landscape/Private/LandscapeEditLayers.cpp b/Engine/Source/Runtime/Landscape/Private/LandscapeEditLayers.cpp
--- a/Engine/Source/Runtime/Landscape/Private/LandscapeEditLayers.cpp
+++ b/Engine/Source/Runtime/Landscape/Private/LandscapeEditLayers.cpp
@@ -5116,7 +5116,7 @@ void ALandscape::FinishDestroy()
 bool ALandscape::IsUpToDate() const
 {
 #if WITH_EDITORONLY_DATA
-	if (CanHaveLayersContent())
+	if (CanHaveLayersContent() && GetWorld() != nullptr && !GetWorld()->IsGameWorld())
 	{
 		return LayerContentUpdateModes == 0;
 	}
