- Fix Move to current level not dealing properly with SplineMesh and ControlPointMesh
- Add new Update Spline Mesh Levels button to clean messed up LandscapeSpline Mesh data

diff --git a/Engine/Source/Editor/LandscapeEditor/Private/LandscapeEdMode.h b/Engine/Source/Editor/LandscapeEditor/Private/LandscapeEdMode.h
--- a/Engine/Source/Editor/LandscapeEditor/Private/LandscapeEdMode.h
+++ b/Engine/Source/Editor/LandscapeEditor/Private/LandscapeEdMode.h
@@ -320,6 +320,7 @@ public:
 	virtual void SelectAllConnectedSplineControlPoints();
 	virtual void SelectAllConnectedSplineSegments();
 	virtual void SplineMoveToCurrentLevel();
+	virtual void UpdateSplineMeshLevels();
 	void SetbUseAutoRotateOnJoin(bool InbAutoRotateOnJoin);
 	bool GetbUseAutoRotateOnJoin();
 
diff --git a/Engine/Source/Editor/LandscapeEditor/Private/LandscapeEdModeSplineTools.cpp b/Engine/Source/Editor/LandscapeEditor/Private/LandscapeEdModeSplineTools.cpp
--- a/Engine/Source/Editor/LandscapeEditor/Private/LandscapeEdModeSplineTools.cpp
+++ b/Engine/Source/Editor/LandscapeEditor/Private/LandscapeEdModeSplineTools.cpp
@@ -828,6 +828,24 @@ public:
 		}
 	}
 
+	void UpdateSplineMeshLevels()
+	{
+		for (ULandscapeSplineControlPoint* ControlPoint : SelectedSplineControlPoints)
+		{
+			const bool bUpdateCollision = true;
+			const bool bUpdateSegments = false;
+			const bool bUpdateMeshLevel = true;
+			ControlPoint->UpdateSplinePoints(bUpdateCollision, bUpdateSegments, bUpdateMeshLevel);
+		}
+
+		for (ULandscapeSplineSegment* Segment : SelectedSplineSegments)
+		{
+			const bool bUpdateCollision = true;
+			const bool bUpdateMeshLevel = true;
+			Segment->UpdateSplinePoints(bUpdateCollision, bUpdateMeshLevel);
+		}
+	}
+
 	void MoveSelectedToLevel()
 	{
 		TSet<ALandscapeProxy*> FromProxies;
@@ -872,45 +890,6 @@ public:
 						FromProxy->SplineComponent->MarkRenderStateDirty();
 					}
 
-					// Handle control point mesh
-					if (ControlPoint->bPlaceSplineMeshesInStreamingLevels)
-					{
-						// Mark previously local component as Foreign
-						if (ControlPoint->LocalMeshComponent)
-						{
-							auto* MeshComponent = ControlPoint->LocalMeshComponent;
-							verifySlow(FromProxy->SplineComponent->MeshComponentLocalOwnersMap.Remove(MeshComponent) == 1);
-							FromProxy->SplineComponent->AddForeignMeshComponent(ControlPoint, MeshComponent);
-						}
-						ControlPoint->LocalMeshComponent = nullptr;
-
-						// Mark previously foreign component as local
-						auto* MeshComponent = ToLandscape->SplineComponent->GetForeignMeshComponent(ControlPoint);
-						if (MeshComponent)
-						{
-							ToLandscape->SplineComponent->RemoveForeignMeshComponent(ControlPoint, MeshComponent);
-							ToLandscape->SplineComponent->MeshComponentLocalOwnersMap.Add(MeshComponent, ControlPoint);
-						}
-						ControlPoint->LocalMeshComponent = MeshComponent;
-					}
-					else
-					{
-						// non-streaming case
-						if (ControlPoint->LocalMeshComponent)
-						{
-							UControlPointMeshComponent* MeshComponent = ControlPoint->LocalMeshComponent;
-							MeshComponent->Modify();
-							MeshComponent->UnregisterComponent();
-							MeshComponent->DetachFromComponent(FDetachmentTransformRules::KeepWorldTransform);
-							MeshComponent->InvalidateLightingCache();
-							MeshComponent->Rename(nullptr, ToLandscape);
-							MeshComponent->AttachToComponent(ToLandscape->SplineComponent, FAttachmentTransformRules::KeepWorldTransform);
-
-							verifySlow(FromProxy->SplineComponent->MeshComponentLocalOwnersMap.Remove(MeshComponent) == 1);
-							ToLandscape->SplineComponent->MeshComponentLocalOwnersMap.Add(MeshComponent, ControlPoint);
-						}
-					}
-
 					// Move control point to new level
 					FromProxy->SplineComponent->ControlPoints.Remove(ControlPoint);
 					ControlPoint->Rename(nullptr, ToLandscape->SplineComponent);
@@ -918,7 +897,10 @@ public:
 
 					ControlPoint->Location = OldToNewTransform.TransformPosition(ControlPoint->Location);
 
-					ControlPoint->UpdateSplinePoints(true, false);
+					const bool bUpdateCollision = true; // default value
+					const bool bUpdateSegments = false; // done in next loop
+					const bool bUpdateMeshLevel = true;
+					ControlPoint->UpdateSplinePoints(bUpdateCollision, bUpdateSegments, bUpdateMeshLevel);
 				}
 			}
 		}
@@ -959,49 +941,14 @@ public:
 						FromProxy->SplineComponent->MarkRenderStateDirty();
 					}
 
-					// Handle spline meshes
-					if (Segment->bPlaceSplineMeshesInStreamingLevels)
-					{
-						// Mark previously local components as Foreign
-						for (auto* MeshComponent : Segment->LocalMeshComponents)
-						{
-							verifySlow(FromProxy->SplineComponent->MeshComponentLocalOwnersMap.Remove(MeshComponent) == 1);
-							FromProxy->SplineComponent->AddForeignMeshComponent(Segment, MeshComponent);
-						}
-						Segment->LocalMeshComponents.Empty();
-
-						// Mark previously foreign components as local
-						TArray<USplineMeshComponent*> MeshComponents = ToLandscape->SplineComponent->GetForeignMeshComponents(Segment);
-						ToLandscape->SplineComponent->RemoveAllForeignMeshComponents(Segment);
-						for (auto* MeshComponent : MeshComponents)
-						{
-							ToLandscape->SplineComponent->MeshComponentLocalOwnersMap.Add(MeshComponent, Segment);
-						}
-						Segment->LocalMeshComponents = MoveTemp(MeshComponents);
-					}
-					else
-					{
-						// non-streaming case
-						for (auto* MeshComponent : Segment->LocalMeshComponents)
-						{
-							MeshComponent->Modify();
-							MeshComponent->UnregisterComponent();
-							MeshComponent->DetachFromComponent(FDetachmentTransformRules::KeepWorldTransform);
-							MeshComponent->InvalidateLightingCache();
-							MeshComponent->Rename(nullptr, ToLandscape);
-							MeshComponent->AttachToComponent(ToLandscape->SplineComponent, FAttachmentTransformRules::KeepWorldTransform);
-
-							verifySlow(FromProxy->SplineComponent->MeshComponentLocalOwnersMap.Remove(MeshComponent) == 1);
-							ToLandscape->SplineComponent->MeshComponentLocalOwnersMap.Add(MeshComponent, Segment);
-						}
-					}
-
 					// Move segment to new level
 					FromProxy->SplineComponent->Segments.Remove(Segment);
 					Segment->Rename(nullptr, ToLandscape->SplineComponent);
 					ToLandscape->SplineComponent->Segments.Add(Segment);
 
-					Segment->UpdateSplinePoints();
+					const bool bUpdateCollision = true; // default value
+					const bool bUpdateMeshLevel = true;
+					Segment->UpdateSplinePoints(bUpdateCollision, bUpdateMeshLevel);
 				}
 			}
 		}
@@ -2401,6 +2348,18 @@ void FEdModeLandscape::SplineMoveToCurrentLevel()
 	}
 }
 
+void FEdModeLandscape::UpdateSplineMeshLevels()
+{
+	FScopedTransaction Transaction(LOCTEXT("LandscapeSpline_UpdateSplineMeshes", "Update Spline Meshes Level"));
+
+	if(SplinesTool)
+	{
+		SelectAllConnectedSplineSegments();
+
+		SplinesTool->UpdateSplineMeshLevels();
+	}
+}
+
 void FEdModeLandscape::SetbUseAutoRotateOnJoin(bool InbAutoRotateOnJoin)
 {
 	if (SplinesTool /*&& SplinesTool == CurrentTool*/)
diff --git a/Engine/Source/Editor/LandscapeEditor/Private/LandscapeSplineDetails.cpp b/Engine/Source/Editor/LandscapeEditor/Private/LandscapeSplineDetails.cpp
--- a/Engine/Source/Editor/LandscapeEditor/Private/LandscapeSplineDetails.cpp
+++ b/Engine/Source/Editor/LandscapeEditor/Private/LandscapeSplineDetails.cpp
@@ -68,6 +68,21 @@ void FLandscapeSplineDetails::CustomizeDetails(IDetailLayoutBuilder& DetailBuild
 			.IsEnabled(this, &FLandscapeSplineDetails::IsMoveToCurrentLevelButtonEnabled)
 		]
 	];
+	LandscapeSplineCategory.AddCustomRow(FText::GetEmpty())
+	[
+		SNew(SHorizontalBox)
+		+ SHorizontalBox::Slot()
+		.Padding(0, 0, 2, 0)
+		.VAlign(VAlign_Center)
+		.FillWidth(1)
+		[
+			SNew(SButton)
+			.Text(LOCTEXT("Move Spline Mesh Components to Proper level", "Update Spline Mesh Levels"))
+			.HAlign(HAlign_Center)
+			.OnClicked(this, &FLandscapeSplineDetails::OnUpdateSplineMeshLevelsButtonClicked)
+			.IsEnabled(this, &FLandscapeSplineDetails::IsUpdateSplineMeshLevelsButtonEnabled)
+		]
+	];
 }
 
 class FEdModeLandscape* FLandscapeSplineDetails::GetEditorMode() const
@@ -114,4 +129,21 @@ bool FLandscapeSplineDetails::IsMoveToCurrentLevelButtonEnabled() const
 	return (LandscapeEdMode && LandscapeEdMode->CurrentToolTarget.LandscapeInfo.IsValid() && LandscapeEdMode->CurrentToolTarget.LandscapeInfo->GetCurrentLevelLandscapeProxy(true));
 }
 
+FReply FLandscapeSplineDetails::OnUpdateSplineMeshLevelsButtonClicked()
+{
+	FEdModeLandscape* LandscapeEdMode = GetEditorMode();
+	if (LandscapeEdMode && LandscapeEdMode->CurrentToolTarget.LandscapeInfo.IsValid())
+	{
+		LandscapeEdMode->UpdateSplineMeshLevels();
+	}
+	
+	return FReply::Handled();
+}
+
+bool FLandscapeSplineDetails::IsUpdateSplineMeshLevelsButtonEnabled() const
+{
+	FEdModeLandscape* LandscapeEdMode = GetEditorMode();
+	return (LandscapeEdMode && LandscapeEdMode->CurrentToolTarget.LandscapeInfo.IsValid());
+}
+
 #undef LOCTEXT_NAMESPACE
diff --git a/Engine/Source/Editor/LandscapeEditor/Private/LandscapeSplineDetails.h b/Engine/Source/Editor/LandscapeEditor/Private/LandscapeSplineDetails.h
--- a/Engine/Source/Editor/LandscapeEditor/Private/LandscapeSplineDetails.h
+++ b/Engine/Source/Editor/LandscapeEditor/Private/LandscapeSplineDetails.h
@@ -24,4 +24,7 @@ private:
 	FReply OnSelectConnectedSegmentsButtonClicked();
 	FReply OnMoveToCurrentLevelButtonClicked();
 	bool IsMoveToCurrentLevelButtonEnabled() const;
+
+	FReply OnUpdateSplineMeshLevelsButtonClicked();
+	bool IsUpdateSplineMeshLevelsButtonEnabled() const;
 };
diff --git a/Engine/Source/Runtime/Landscape/Classes/LandscapeSplineControlPoint.h b/Engine/Source/Runtime/Landscape/Classes/LandscapeSplineControlPoint.h
--- a/Engine/Source/Runtime/Landscape/Classes/LandscapeSplineControlPoint.h
+++ b/Engine/Source/Runtime/Landscape/Classes/LandscapeSplineControlPoint.h
@@ -212,7 +212,7 @@ public:
 	TMap<ULandscapeSplinesComponent*, UControlPointMeshComponent*> GetForeignMeshComponents();
 
 	/** Update spline points */
-	virtual void UpdateSplinePoints(bool bUpdateCollision = true, bool bUpdateAttachedSegments = true);
+	virtual void UpdateSplinePoints(bool bUpdateCollision = true, bool bUpdateAttachedSegments = true, bool bUpdateMeshLevel = false);
 
 	/** Delete spline points */
 	virtual void DeleteSplinePoints();
diff --git a/Engine/Source/Runtime/Landscape/Classes/LandscapeSplineSegment.h b/Engine/Source/Runtime/Landscape/Classes/LandscapeSplineSegment.h
--- a/Engine/Source/Runtime/Landscape/Classes/LandscapeSplineSegment.h
+++ b/Engine/Source/Runtime/Landscape/Classes/LandscapeSplineSegment.h
@@ -307,7 +307,7 @@ public:
 	TMap<ULandscapeSplinesComponent*, TArray<USplineMeshComponent*>> GetForeignMeshComponents();
 	TArray<USplineMeshComponent*> GetLocalMeshComponents() const;
 	
-	virtual void UpdateSplinePoints(bool bUpdateCollision = true);
+	virtual void UpdateSplinePoints(bool bUpdateCollision = true, bool bUpdateMeshLevel = false);
 
 	void UpdateSplineEditorMesh();
 	virtual void DeleteSplinePoints();
diff --git a/Engine/Source/Runtime/Landscape/Classes/LandscapeSplinesComponent.h b/Engine/Source/Runtime/Landscape/Classes/LandscapeSplinesComponent.h
--- a/Engine/Source/Runtime/Landscape/Classes/LandscapeSplinesComponent.h
+++ b/Engine/Source/Runtime/Landscape/Classes/LandscapeSplinesComponent.h
@@ -163,7 +163,7 @@ public:
 	// returns a suitable ULandscapeSplinesComponent to place streaming meshes into, given a location
 	// falls back to "this" if it can't find another suitable, so never returns nullptr
 	// @param bCreate whether to create a component if a suitable actor is found but it has no splines component yet
-	ULandscapeSplinesComponent* GetStreamingSplinesComponentByLocation(const FVector& LocalLocation, bool bCreate = true);
+	LANDSCAPE_API ULandscapeSplinesComponent* GetStreamingSplinesComponentByLocation(const FVector& LocalLocation, bool bCreate = true);
 
 	// returns the matching ULandscapeSplinesComponent for a given level, *can return null*
 	// @param bCreate whether to create a component if a suitable actor is found but it has no splines component yet
diff --git a/Engine/Source/Runtime/Landscape/Private/LandscapeSplines.cpp b/Engine/Source/Runtime/Landscape/Private/LandscapeSplines.cpp
--- a/Engine/Source/Runtime/Landscape/Private/LandscapeSplines.cpp
+++ b/Engine/Source/Runtime/Landscape/Private/LandscapeSplines.cpp
@@ -1559,7 +1559,7 @@ TMap<ULandscapeSplinesComponent*, UControlPointMeshComponent*> ULandscapeSplineC
 	return ForeignMeshComponentsMap;
 }
 
-void ULandscapeSplineControlPoint::UpdateSplinePoints(bool bUpdateCollision, bool bUpdateAttachedSegments)
+void ULandscapeSplineControlPoint::UpdateSplinePoints(bool bUpdateCollision, bool bUpdateAttachedSegments, bool bUpdateMeshLevel)
 {
 	Modify();
 
@@ -1571,7 +1571,7 @@ void ULandscapeSplineControlPoint::UpdateSplinePoints(bool bUpdateCollision, boo
 
 	UControlPointMeshComponent* MeshComponent = LocalMeshComponent;
 	ULandscapeSplinesComponent* MeshComponentOuterSplines = OuterSplines;
-
+		
 	if (Mesh != nullptr)
 	{
 		// Attempt to place mesh components into the appropriate landscape streaming levels based on the components under the spline
@@ -1592,6 +1592,7 @@ void ULandscapeSplineControlPoint::UpdateSplinePoints(bool bUpdateCollision, boo
 
 		// Create mesh component if needed
 		bool bComponentNeedsRegistering = false;
+		bool bUpdateLocalForeign = false;
 		if (MeshComponent == nullptr)
 		{
 			AActor* MeshComponentOuterActor = MeshComponentOuterSplines->GetOwner();
@@ -1601,10 +1602,39 @@ void ULandscapeSplineControlPoint::UpdateSplinePoints(bool bUpdateCollision, boo
 			MeshComponent->bSelected = bSelected;
 			MeshComponent->AttachToComponent(MeshComponentOuterSplines, FAttachmentTransformRules::KeepRelativeTransform);
 			bComponentNeedsRegistering = true;
+			bUpdateLocalForeign = true;
+		}
+		else if(bUpdateMeshLevel)// Update Foreign/Local if necessary
+		{
+			ALandscapeProxy* CurrentMeshComponentOuterActor = MeshComponent->GetTypedOuter<ALandscapeProxy>();
+			ULandscapeSplinesComponent* CurrentLandscapeSplineComponent = CurrentMeshComponentOuterActor->SplineComponent;
+
+			ALandscapeProxy* MeshComponentOuterActor = Cast<ALandscapeProxy>(MeshComponentOuterSplines->GetOwner());
+			
+			// Needs updating
+			if (MeshComponentOuterActor != CurrentMeshComponentOuterActor)
+			{
+				MeshComponentOuterActor->Modify();
+				CurrentMeshComponentOuterActor->Modify();			
+				
+				MeshComponent->Modify();
+				MeshComponent->UnregisterComponent();
+				MeshComponent->DetachFromComponent(FDetachmentTransformRules::KeepWorldTransform);
+				MeshComponent->InvalidateLightingCache();
+				MeshComponent->Rename(nullptr, MeshComponentOuterActor);
+				MeshComponent->AttachToComponent(MeshComponentOuterActor->SplineComponent, FAttachmentTransformRules::KeepWorldTransform);
+				bComponentNeedsRegistering = true;
+				bUpdateLocalForeign = true;
+			}
+		}
 
+		if (bUpdateLocalForeign)
+		{
 			if (MeshComponentOuterSplines == OuterSplines)
 			{
-				MeshComponentOuterSplines->MeshComponentLocalOwnersMap.Add(MeshComponent, this);
+				UObject*& ValueRef = MeshComponentOuterSplines->MeshComponentLocalOwnersMap.FindOrAdd(MeshComponent);
+				ValueRef = this;
+
 				LocalMeshComponent = MeshComponent;
 			}
 			else
@@ -1613,7 +1643,7 @@ void ULandscapeSplineControlPoint::UpdateSplinePoints(bool bUpdateCollision, boo
 				ForeignWorld = MeshComponentOuterSplines->GetTypedOuter<UWorld>();
 			}
 		}
-
+				
 		FVector MeshLocation = Location;
 		FRotator MeshRotation = Rotation;
 		if (MeshComponentOuterSplines != OuterSplines)
@@ -1702,7 +1732,6 @@ void ULandscapeSplineControlPoint::UpdateSplinePoints(bool bUpdateCollision, boo
 	}
 
 	// Destroy any unused components
-	bool bDestroyedAnyComponents = false;
 	if (LocalMeshComponent && LocalMeshComponent != MeshComponent)
 	{
 		OuterSplines->Modify();
@@ -1725,11 +1754,7 @@ void ULandscapeSplineControlPoint::UpdateSplinePoints(bool bUpdateCollision, boo
 		}
 	}
 	ForeignMeshComponentsMap.Empty();
-	if (bDestroyedAnyComponents)
-	{
-		AutoSetConnections(false);
-	}
-
+	
 	const float LeftSideFalloff = LeftSideFalloffFactor * SideFalloff;
 	const float RightSideFalloff = RightSideFalloffFactor * SideFalloff;
 
@@ -2207,7 +2232,7 @@ TArray<USplineMeshComponent*> ULandscapeSplineSegment::GetLocalMeshComponents()
 	return LocalMeshComponents;
 }
 
-void ULandscapeSplineSegment::UpdateSplinePoints(bool bUpdateCollision)
+void ULandscapeSplineSegment::UpdateSplinePoints(bool bUpdateCollision, bool bUpdateMeshLevel)
 {
 	Modify();
 
@@ -2300,11 +2325,13 @@ void ULandscapeSplineSegment::UpdateSplinePoints(bool bUpdateCollision)
 	TArray<USplineMeshComponent*> OldLocalMeshComponents = MoveTemp(LocalMeshComponents);
 	LocalMeshComponents.Reserve(20);
 
-	auto ForeignMeshComponentsMap = GetForeignMeshComponents();
-
-	// Unregister components
+	// Gather all ForeignMesh associated with this Segment
+	TMap<ULandscapeSplinesComponent*, TArray<USplineMeshComponent*>> ForeignMeshComponentsMap = GetForeignMeshComponents();
+	
+	// Unregister components, Remove Foreign/Local Associations
 	for (auto* LocalMeshComponent : OldLocalMeshComponents)
 	{
+		OuterSplines->MeshComponentLocalOwnersMap.Remove(LocalMeshComponent);
 		LocalMeshComponent->Modify();
 		LocalMeshComponent->UnregisterComponent();
 	}
@@ -2312,6 +2339,7 @@ void ULandscapeSplineSegment::UpdateSplinePoints(bool bUpdateCollision)
 	{
 		ForeignMeshComponentsPair.Key->Modify();
 		ForeignMeshComponentsPair.Key->GetOwner()->Modify();
+		ForeignMeshComponentsPair.Key->RemoveAllForeignMeshComponents(this);
 		for (auto* ForeignMeshComponent : ForeignMeshComponentsPair.Value)
 		{
 			ForeignMeshComponent->Modify();
@@ -2390,7 +2418,6 @@ void ULandscapeSplineSegment::UpdateSplinePoints(bool bUpdateCollision)
 				if (OldLocalMeshComponents.Num() > 0)
 				{
 					MeshComponent = OldLocalMeshComponents.Pop(false);
-					LocalMeshComponents.Add(MeshComponent);
 				}
 			}
 			else
@@ -2400,29 +2427,50 @@ void ULandscapeSplineSegment::UpdateSplinePoints(bool bUpdateCollision)
 				{
 					MeshComponentOuterSplines->UpdateModificationKey(this);
 					MeshComponent = ForeignMeshComponents->Pop(false);
-					ForeignWorlds.AddUnique(MeshComponentOuterSplines->GetTypedOuter<UWorld>());
 				}
 			}
 
-			if (!MeshComponent)
+			if (MeshComponent == nullptr)
 			{
 				AActor* MeshComponentOuterActor = MeshComponentOuterSplines->GetOwner();
 				MeshComponentOuterActor->Modify();
 				MeshComponent = NewObject<USplineMeshComponent>(MeshComponentOuterActor, NAME_None, RF_Transactional | RF_TextExportTransient);
 				MeshComponent->bSelected = bSelected;
 				MeshComponent->AttachToComponent(MeshComponentOuterSplines, FAttachmentTransformRules::KeepRelativeTransform);
-				if (MeshComponentOuterSplines == OuterSplines)
-				{
-					LocalMeshComponents.Add(MeshComponent);
-					MeshComponentOuterSplines->MeshComponentLocalOwnersMap.Add(MeshComponent, this);
-				}
-				else
+			}
+			else if (bUpdateMeshLevel)// Update Foreign/Local if necessary
+			{
+				ALandscapeProxy* CurrentMeshComponentOuterActor = MeshComponent->GetTypedOuter<ALandscapeProxy>();
+				ULandscapeSplinesComponent* CurrentLandscapeSplineComponent = CurrentMeshComponentOuterActor->SplineComponent;
+
+				ALandscapeProxy* MeshComponentOuterActor = Cast<ALandscapeProxy>(MeshComponentOuterSplines->GetOwner());
+
+				// Needs updating
+				if (MeshComponentOuterActor != CurrentMeshComponentOuterActor)
 				{
-					MeshComponentOuterSplines->AddForeignMeshComponent(this, MeshComponent);
-					ForeignWorlds.AddUnique(MeshComponentOuterSplines->GetTypedOuter<UWorld>());
+					MeshComponentOuterActor->Modify();
+					CurrentMeshComponentOuterActor->Modify();
+							
+					MeshComponent->Modify();
+					MeshComponent->DetachFromComponent(FDetachmentTransformRules::KeepWorldTransform);
+					MeshComponent->InvalidateLightingCache();
+					MeshComponent->Rename(nullptr, MeshComponentOuterActor);
+					MeshComponent->AttachToComponent(MeshComponentOuterActor->SplineComponent, FAttachmentTransformRules::KeepWorldTransform);
 				}
 			}
 
+			// New Foreign/Local Mapping
+			if (MeshComponentOuterSplines == OuterSplines)
+			{
+				LocalMeshComponents.Add(MeshComponent);
+				MeshComponentOuterSplines->MeshComponentLocalOwnersMap.Add(MeshComponent, this);
+			}
+			else
+			{
+				MeshComponentOuterSplines->AddForeignMeshComponent(this, MeshComponent);
+				ForeignWorlds.AddUnique(MeshComponentOuterSplines->GetTypedOuter<UWorld>());
+			}
+
 			MeshComponents.Add(MeshComponent);
 
 			MeshComponent->SetStaticMesh(Mesh);
@@ -2444,26 +2492,6 @@ void ULandscapeSplineSegment::UpdateSplinePoints(bool bUpdateCollision)
 		// Add terminating key
 		MeshSettings.Add(FMeshSettings(T, nullptr));
 
-		// Destroy old unwanted components now
-		for (UMeshComponent* LocalMeshComponent : OldLocalMeshComponents)
-		{
-			checkSlow(OuterSplines->MeshComponentLocalOwnersMap.FindRef(LocalMeshComponent) == this);
-			verifySlow(OuterSplines->MeshComponentLocalOwnersMap.Remove(LocalMeshComponent) == 1);
-			LocalMeshComponent->DestroyComponent();
-		}
-		OldLocalMeshComponents.Empty();
-
-		for (auto& ForeignMeshComponentsPair : ForeignMeshComponentsMap)
-		{
-			ULandscapeSplinesComponent* MeshComponentOuterSplines = ForeignMeshComponentsPair.Key;
-			for (auto* ForeignMeshComponent : ForeignMeshComponentsPair.Value)
-			{
-				MeshComponentOuterSplines->RemoveForeignMeshComponent(this, ForeignMeshComponent);
-				ForeignMeshComponent->DestroyComponent();
-			}
-		}
-		ForeignMeshComponentsMap.Empty();
-
 		// Second pass:
 		// Rescale components to fit a whole number to the spline, set up final parameters
 		const float Rescale = 1.0f / T;
@@ -2662,30 +2690,23 @@ void ULandscapeSplineSegment::UpdateSplinePoints(bool bUpdateCollision)
 			MeshComponent->RegisterComponent();
 		}
 	}
-	else
+	
+	// Clean up unused components
+	for (auto* LocalMeshComponent : OldLocalMeshComponents)
 	{
-		// Spline needs no mesh components (0 length or no meshes to use) so destroy any we have
-		for (auto* LocalMeshComponent : OldLocalMeshComponents)
-		{
-			checkSlow(OuterSplines->MeshComponentLocalOwnersMap.FindRef(LocalMeshComponent) == this);
-			verifySlow(OuterSplines->MeshComponentLocalOwnersMap.Remove(LocalMeshComponent) == 1);
-			LocalMeshComponent->DestroyComponent();
-		}
-		OldLocalMeshComponents.Empty();
-		for (auto& ForeignMeshComponentsPair : ForeignMeshComponentsMap)
+		checkSlow(OuterSplines->MeshComponentLocalOwnersMap.FindRef(LocalMeshComponent) == this);
+		verifySlow(OuterSplines->MeshComponentLocalOwnersMap.Remove(LocalMeshComponent) == 1);
+		LocalMeshComponent->DestroyComponent();
+	}
+	OldLocalMeshComponents.Empty();
+	for (auto& ForeignMeshComponentsPair : ForeignMeshComponentsMap)
+	{
+		for (USplineMeshComponent* MeshComponent : ForeignMeshComponentsPair.Value)
 		{
-			for (USplineMeshComponent* MeshComponent : ForeignMeshComponentsPair.Value)
-			{
-				ULandscapeSplinesComponent* MeshComponentOuterSplines = dynamic_cast<ULandscapeSplinesComponent*>(MeshComponent->GetAttachParent());
-				if (MeshComponentOuterSplines)
-				{
-					MeshComponentOuterSplines->RemoveForeignMeshComponent(this, MeshComponent);
-				}
-				MeshComponent->DestroyComponent();
-			}
+			MeshComponent->DestroyComponent();
 		}
-		ForeignMeshComponentsMap.Empty();
 	}
+	ForeignMeshComponentsMap.Empty();
 }
 
 void ULandscapeSplineSegment::UpdateSplineEditorMesh()
