Add API to query landscape height at a world location.

diff --git a/Engine/Source/Runtime/Landscape/Classes/LandscapeHeightfieldCollisionComponent.h b/Engine/Source/Runtime/Landscape/Classes/LandscapeHeightfieldCollisionComponent.h
--- a/Engine/Source/Runtime/Landscape/Classes/LandscapeHeightfieldCollisionComponent.h
+++ b/Engine/Source/Runtime/Landscape/Classes/LandscapeHeightfieldCollisionComponent.h
@@ -183,6 +183,9 @@ private:
 protected:
 	virtual void OnCreatePhysicsState() override;
 public:
+
+	TOptional<float> GetHeight(float X, float Y);
+
 	virtual void ApplyWorldOffset(const FVector& InOffset, bool bWorldShift) override;
 	//~ End UActorComponent Interface.
 
diff --git a/Engine/Source/Runtime/Landscape/Classes/LandscapeInfoMap.h b/Engine/Source/Runtime/Landscape/Classes/LandscapeInfoMap.h
--- a/Engine/Source/Runtime/Landscape/Classes/LandscapeInfoMap.h
+++ b/Engine/Source/Runtime/Landscape/Classes/LandscapeInfoMap.h
@@ -27,7 +27,6 @@ public:
 	UWorld* World;
 
 
-#if WITH_EDITORONLY_DATA
 	/**
 	* Gets landscape-specific data for given world.
 	*
@@ -36,5 +35,4 @@ public:
 	* @returns Landscape-specific data associated with given world.
 	*/
 	LANDSCAPE_API static ULandscapeInfoMap& GetLandscapeInfoMap(const UWorld* World);
-#endif // WITH_EDITORONLY_DATA
 };
diff --git a/Engine/Source/Runtime/Landscape/Classes/LandscapeProxy.h b/Engine/Source/Runtime/Landscape/Classes/LandscapeProxy.h
--- a/Engine/Source/Runtime/Landscape/Classes/LandscapeProxy.h
+++ b/Engine/Source/Runtime/Landscape/Classes/LandscapeProxy.h
@@ -415,6 +415,8 @@ protected:
 	FGuid LandscapeGuid;
 
 public:
+	LANDSCAPE_API TOptional<float> GetHeightAtLocation(FVector Location) const;
+
 	/** Offset in quads from global components grid origin (in quads) **/
 	UPROPERTY()
 	FIntPoint LandscapeSectionOffset;
diff --git a/Engine/Source/Runtime/Landscape/Private/Landscape.cpp b/Engine/Source/Runtime/Landscape/Private/Landscape.cpp
--- a/Engine/Source/Runtime/Landscape/Private/Landscape.cpp
+++ b/Engine/Source/Runtime/Landscape/Private/Landscape.cpp
@@ -1144,7 +1144,6 @@ ULandscapeInfo* ALandscapeProxy::CreateLandscapeInfo()
 	check(LandscapeGuid.IsValid());
 	UWorld* OwningWorld = GetWorld();
 	check(OwningWorld);
-	check(!OwningWorld->IsGameWorld());
 	
 	auto& LandscapeInfoMap = ULandscapeInfoMap::GetLandscapeInfoMap(OwningWorld);
 	LandscapeInfo = LandscapeInfoMap.Map.FindRef(LandscapeGuid);
@@ -1170,7 +1169,7 @@ ULandscapeInfo* ALandscapeProxy::GetLandscapeInfo() const
 	check(LandscapeGuid.IsValid());
 	UWorld* OwningWorld = GetWorld();
 
-	if (OwningWorld != nullptr && !OwningWorld->IsGameWorld())
+	if (OwningWorld != nullptr)
 	{
 		auto& LandscapeInfoMap = ULandscapeInfoMap::GetLandscapeInfoMap(OwningWorld);
 		LandscapeInfo = LandscapeInfoMap.Map.FindRef(LandscapeGuid);
@@ -1428,10 +1427,9 @@ void ULandscapeComponent::OnRegister()
 			}
 		}
 
-#if WITH_EDITOR
 		// AActor::GetWorld checks for Unreachable and BeginDestroyed
 		UWorld* World = GetLandscapeProxy()->GetWorld();
-		if (World && !World->IsGameWorld())
+		if (World)
 		{
 			ULandscapeInfo* Info = GetLandscapeInfo();
 			if (Info)
@@ -1439,7 +1437,6 @@ void ULandscapeComponent::OnRegister()
 				Info->RegisterActorComponent(this);
 			}
 		}
-#endif
 	}
 }
 
@@ -1455,7 +1452,6 @@ void ULandscapeComponent::OnUnregister()
 			MaterialInstancesDynamic.Empty();
 		}
 
-#if WITH_EDITOR
 		// AActor::GetWorld checks for Unreachable and BeginDestroyed
 		UWorld* World = GetLandscapeProxy()->GetWorld();
 
@@ -1468,7 +1464,6 @@ void ULandscapeComponent::OnUnregister()
 				Info->UnregisterActorComponent(this);
 			}
 		}
-#endif
 	}
 }
 
@@ -1791,16 +1786,21 @@ void ULandscapeComponent::SetWeightmapTexturesUsage(const TArray<ULandscapeWeigh
 void ALandscapeProxy::PostRegisterAllComponents()
 {
 	Super::PostRegisterAllComponents();
-
+	ULandscapeInfo* LandscapeInfo = nullptr;
+	if (!IsPendingKillPending())
+	{
+		// Duplicated Landscapes don't have a valid guid until PostEditImport is called, we'll register then
+		if (LandscapeGuid.IsValid())
+		{
+			LandscapeInfo = CreateLandscapeInfo();
+		}
+	}
 #if WITH_EDITOR
 	// Game worlds don't have landscape infos
 	if (!GetWorld()->IsGameWorld() && !IsPendingKillPending())
 	{
-		// Duplicated Landscapes don't have a valid guid until PostEditImport is called, we'll register then
 		if (LandscapeGuid.IsValid())
 		{
-			ULandscapeInfo* LandscapeInfo = CreateLandscapeInfo();
-
 			LandscapeInfo->FixupProxiesTransform();
 		}
 	}
@@ -2277,6 +2277,7 @@ void ALandscapeProxy::PostLoad()
 		BodyInstance.FixupData(this);
 	}
 
+	CreateLandscapeInfo();
 #if WITH_EDITOR
 	if (GIsEditor && !GetWorld()->IsGameWorld())
 	{
@@ -2285,7 +2286,6 @@ void ALandscapeProxy::PostLoad()
 			LandscapeComponents.ContainsByPredicate([](ULandscapeComponent* Comp) { return ((Comp != nullptr) && !Comp->CollisionComponent.IsValid()); }))
 		{
 			// Need to clean up invalid collision components
-			CreateLandscapeInfo();
 			RecreateCollisionComponents();
 		}
 	}
diff --git a/Engine/Source/Runtime/Landscape/Private/LandscapeCollision.cpp b/Engine/Source/Runtime/Landscape/Private/LandscapeCollision.cpp
--- a/Engine/Source/Runtime/Landscape/Private/LandscapeCollision.cpp
+++ b/Engine/Source/Runtime/Landscape/Private/LandscapeCollision.cpp
@@ -1790,6 +1790,19 @@ void ULandscapeInfo::UpdateAllAddCollisions()
 	}
 }
 
+TOptional<float> ULandscapeHeightfieldCollisionComponent::GetHeight(float X, float Y)
+{
+	TOptional<float> Height;
+#if WITH_PHYSX
+	const float ZScale = GetComponentTransform().GetScale3D().Z * LANDSCAPE_ZSCALE;
+	if (IsValidRef(HeightfieldRef) && HeightfieldRef->RBHeightfield != nullptr)
+	{
+		Height = HeightfieldRef->RBHeightfield->getHeight(HeightfieldRef->RBHeightfield->getNbRows() - 1 - X, Y) * ZScale;
+	}
+#endif
+	return Height;
+}
+
 void ULandscapeInfo::UpdateAddCollision(FIntPoint LandscapeKey)
 {
 	FLandscapeAddCollision& AddCollision = XYtoAddCollisionMap.FindOrAdd(LandscapeKey);
@@ -2188,4 +2201,26 @@ ULandscapeHeightfieldCollisionComponent::ULandscapeHeightfieldCollisionComponent
 ULandscapeComponent* ULandscapeHeightfieldCollisionComponent::GetRenderComponent() const
 {
 	return RenderComponent.Get();
+}
+
+LANDSCAPE_API TOptional<float> ALandscapeProxy::GetHeightAtLocation(FVector Location) const
+{
+	TOptional<float> Height;
+	ULandscapeInfo* Info = GetLandscapeInfo();
+	const FVector ActorSpaceLocation = GetActorTransform().InverseTransformPosition(Location);
+	const FIntPoint Key = FIntPoint(ActorSpaceLocation.X, ActorSpaceLocation.Y) / ComponentSizeQuads;
+	ULandscapeComponent* Component = Info->XYtoComponentMap.FindRef(Key);
+	if (Component)
+	{
+		if (ULandscapeHeightfieldCollisionComponent* CollisionComp = Component->CollisionComponent.Get())
+		{
+			const FVector ComponentSpaceLocation = CollisionComp->GetComponentToWorld().InverseTransformPosition(Location);
+			const TOptional<float> LocalHeight = CollisionComp->GetHeight(ComponentSpaceLocation.X, ComponentSpaceLocation.Y);
+			if (LocalHeight.IsSet())
+			{
+				Height = CollisionComp->GetComponentToWorld().TransformPositionNoScale(FVector(0, 0, LocalHeight.GetValue())).Z;
+			}
+		}
+	}
+	return Height;
 }
\ No newline at end of file
diff --git a/Engine/Source/Runtime/Landscape/Private/LandscapeInfoMap.cpp b/Engine/Source/Runtime/Landscape/Private/LandscapeInfoMap.cpp
--- a/Engine/Source/Runtime/Landscape/Private/LandscapeInfoMap.cpp
+++ b/Engine/Source/Runtime/Landscape/Private/LandscapeInfoMap.cpp
@@ -43,7 +43,6 @@ void ULandscapeInfoMap::AddReferencedObjects(UObject* InThis, FReferenceCollecto
 	Collector.AddReferencedObjects(This->Map, This);
 }
 
-#if WITH_EDITORONLY_DATA
 ULandscapeInfoMap& ULandscapeInfoMap::GetLandscapeInfoMap(const UWorld* World)
 {
 	ULandscapeInfoMap *FoundObject = nullptr;
@@ -53,4 +52,3 @@ ULandscapeInfoMap& ULandscapeInfoMap::GetLandscapeInfoMap(const UWorld* World)
 
 	return *FoundObject;
 }
-#endif // WITH_EDITORONLY_DATA
