Fix compile errors.

diff --git a/Engine/Source/Runtime/Landscape/Private/Landscape.cpp b/Engine/Source/Runtime/Landscape/Private/Landscape.cpp
--- a/Engine/Source/Runtime/Landscape/Private/Landscape.cpp
+++ b/Engine/Source/Runtime/Landscape/Private/Landscape.cpp
@@ -1135,7 +1135,6 @@
 	return LandscapeActor.Get();
 }
 
-#if WITH_EDITOR
 ULandscapeInfo* ALandscapeProxy::CreateLandscapeInfo()
 {
 	ULandscapeInfo* LandscapeInfo = nullptr;
@@ -1176,7 +1175,6 @@
 	}
 	return LandscapeInfo;
 }
-#endif
 
 ALandscape* ULandscapeComponent::GetLandscapeActor() const
 {
@@ -2885,79 +2883,6 @@
 	UpdateAllAddCollisions();
 }
 
-void ULandscapeInfo::RegisterActorComponent(ULandscapeComponent* Component, bool bMapCheck)
-{
-	// Do not register components which are not part of the world
-	if (Component == nullptr ||
-		Component->IsRegistered() == false)
-	{
-		return;
-	}
-
-	check(Component);
-
-	FIntPoint ComponentKey = Component->GetSectionBase() / Component->ComponentSizeQuads;
-	auto RegisteredComponent = XYtoComponentMap.FindRef(ComponentKey);
-
-	if (RegisteredComponent != Component)
-	{
-		if (RegisteredComponent == nullptr)
-		{
-			XYtoComponentMap.Add(ComponentKey, Component);
-		}
-		else if (bMapCheck)
-		{
-			ALandscapeProxy* OurProxy = Component->GetLandscapeProxy();
-			ALandscapeProxy* ExistingProxy = RegisteredComponent->GetLandscapeProxy();
-			FFormatNamedArguments Arguments;
-			Arguments.Add(TEXT("ProxyName1"), FText::FromString(OurProxy->GetName()));
-			Arguments.Add(TEXT("LevelName1"), FText::FromString(OurProxy->GetLevel()->GetOutermost()->GetName()));
-			Arguments.Add(TEXT("ProxyName2"), FText::FromString(ExistingProxy->GetName()));
-			Arguments.Add(TEXT("LevelName2"), FText::FromString(ExistingProxy->GetLevel()->GetOutermost()->GetName()));
-			Arguments.Add(TEXT("XLocation"), Component->GetSectionBase().X);
-			Arguments.Add(TEXT("YLocation"), Component->GetSectionBase().Y);
-			FMessageLog("MapCheck").Warning()
-				->AddToken(FUObjectToken::Create(OurProxy))
-				->AddToken(FTextToken::Create(FText::Format(LOCTEXT("MapCheck_Message_LandscapeComponentPostLoad_Warning", "Landscape {ProxyName1} of {LevelName1} has overlapping render components with {ProxyName2} of {LevelName2} at location ({XLocation}, {YLocation})."), Arguments)))
-				->AddToken(FActionToken::Create(LOCTEXT("MapCheck_RemoveDuplicateLandscapeComponent", "Delete Duplicate"), LOCTEXT("MapCheck_RemoveDuplicateLandscapeComponentDesc", "Deletes the duplicate landscape component."), FOnActionTokenExecuted::CreateUObject(OurProxy, &ALandscapeProxy::RemoveOverlappingComponent, Component), true))
-				->AddToken(FMapErrorToken::Create(FMapErrors::LandscapeComponentPostLoad_Warning));
-
-			// Show MapCheck window
-			FMessageLog("MapCheck").Open(EMessageSeverity::Warning);
-		}
-	}
-
-	// Update Selected Components/Regions
-	if (Component->EditToolRenderData.SelectedType)
-	{
-		if (Component->EditToolRenderData.SelectedType & FLandscapeEditToolRenderData::ST_COMPONENT)
-		{
-			SelectedComponents.Add(Component);
-		}
-		else if (Component->EditToolRenderData.SelectedType & FLandscapeEditToolRenderData::ST_REGION)
-		{
-			SelectedRegionComponents.Add(Component);
-		}
-	}
-}
-
-void ULandscapeInfo::UnregisterActorComponent(ULandscapeComponent* Component)
-{
-	if (ensure(Component))
-	{
-	FIntPoint ComponentKey = Component->GetSectionBase() / Component->ComponentSizeQuads;
-	auto RegisteredComponent = XYtoComponentMap.FindRef(ComponentKey);
-
-	if (RegisteredComponent == Component)
-	{
-		XYtoComponentMap.Remove(ComponentKey);
-	}
-
-	SelectedComponents.Remove(Component);
-	SelectedRegionComponents.Remove(Component);
-}
-}
-
 void ULandscapeInfo::Reset()
 {
 	LandscapeActor.Reset();
@@ -3096,6 +3021,79 @@
 
 #endif
 
+void ULandscapeInfo::RegisterActorComponent(ULandscapeComponent* Component, bool bMapCheck)
+{
+	// Do not register components which are not part of the world
+	if (Component == nullptr ||
+		Component->IsRegistered() == false)
+	{
+		return;
+	}
+
+	check(Component);
+
+	FIntPoint ComponentKey = Component->GetSectionBase() / Component->ComponentSizeQuads;
+	auto RegisteredComponent = XYtoComponentMap.FindRef(ComponentKey);
+
+	if (RegisteredComponent != Component)
+	{
+		if (RegisteredComponent == nullptr)
+		{
+			XYtoComponentMap.Add(ComponentKey, Component);
+		}
+		else if (bMapCheck)
+		{
+			ALandscapeProxy* OurProxy = Component->GetLandscapeProxy();
+			ALandscapeProxy* ExistingProxy = RegisteredComponent->GetLandscapeProxy();
+			FFormatNamedArguments Arguments;
+			Arguments.Add(TEXT("ProxyName1"), FText::FromString(OurProxy->GetName()));
+			Arguments.Add(TEXT("LevelName1"), FText::FromString(OurProxy->GetLevel()->GetOutermost()->GetName()));
+			Arguments.Add(TEXT("ProxyName2"), FText::FromString(ExistingProxy->GetName()));
+			Arguments.Add(TEXT("LevelName2"), FText::FromString(ExistingProxy->GetLevel()->GetOutermost()->GetName()));
+			Arguments.Add(TEXT("XLocation"), Component->GetSectionBase().X);
+			Arguments.Add(TEXT("YLocation"), Component->GetSectionBase().Y);
+			FMessageLog("MapCheck").Warning()
+				->AddToken(FUObjectToken::Create(OurProxy))
+				->AddToken(FTextToken::Create(FText::Format(LOCTEXT("MapCheck_Message_LandscapeComponentPostLoad_Warning", "Landscape {ProxyName1} of {LevelName1} has overlapping render components with {ProxyName2} of {LevelName2} at location ({XLocation}, {YLocation})."), Arguments)))
+				->AddToken(FActionToken::Create(LOCTEXT("MapCheck_RemoveDuplicateLandscapeComponent", "Delete Duplicate"), LOCTEXT("MapCheck_RemoveDuplicateLandscapeComponentDesc", "Deletes the duplicate landscape component."), FOnActionTokenExecuted::CreateUObject(OurProxy, &ALandscapeProxy::RemoveOverlappingComponent, Component), true))
+				->AddToken(FMapErrorToken::Create(FMapErrors::LandscapeComponentPostLoad_Warning));
+
+			// Show MapCheck window
+			FMessageLog("MapCheck").Open(EMessageSeverity::Warning);
+		}
+	}
+
+	// Update Selected Components/Regions
+	if (Component->EditToolRenderData.SelectedType)
+	{
+		if (Component->EditToolRenderData.SelectedType & FLandscapeEditToolRenderData::ST_COMPONENT)
+		{
+			SelectedComponents.Add(Component);
+		}
+		else if (Component->EditToolRenderData.SelectedType & FLandscapeEditToolRenderData::ST_REGION)
+		{
+			SelectedRegionComponents.Add(Component);
+		}
+	}
+}
+
+void ULandscapeInfo::UnregisterActorComponent(ULandscapeComponent* Component)
+{
+	if (ensure(Component))
+	{
+		FIntPoint ComponentKey = Component->GetSectionBase() / Component->ComponentSizeQuads;
+		auto RegisteredComponent = XYtoComponentMap.FindRef(ComponentKey);
+
+		if (RegisteredComponent == Component)
+		{
+			XYtoComponentMap.Remove(ComponentKey);
+		}
+
+		SelectedComponents.Remove(Component);
+		SelectedRegionComponents.Remove(Component);
+	}
+}
+
 void ULandscapeComponent::PostInitProperties()
 {
 	Super::PostInitProperties();
