Fix non-editor builds

diff --git a/Engine/Source/Runtime/Landscape/Private/Landscape.cpp b/Engine/Source/Runtime/Landscape/Private/Landscape.cpp
--- a/Engine/Source/Runtime/Landscape/Private/Landscape.cpp
+++ b/Engine/Source/Runtime/Landscape/Private/Landscape.cpp
@@ -1319,10 +1319,12 @@ ERuntimeVirtualTextureMainPassType ULandscapeComponent::GetVirtualTextureRenderP
 	return GetLandscapeProxy()->VirtualTextureRenderPassType;
 }
 
+#if WITH_EDITOR
 ULandscapeInfo* ULandscapeComponent::GetLandscapeInfo() const
 {
 	return GetLandscapeProxy()->GetLandscapeInfo();
 }
+#endif // WITH_EDITOR
 
 void ULandscapeComponent::BeginDestroy()
 {
@@ -3063,6 +3065,7 @@ void ULandscapeInfo::RegisterActorComponent(ULandscapeComponent* Component, bool
 		}
 	}
 
+#if WITH_EDITORONLY_DATA
 	// Update Selected Components/Regions
 	if (Component->EditToolRenderData.SelectedType)
 	{
@@ -3075,6 +3078,7 @@ void ULandscapeInfo::RegisterActorComponent(ULandscapeComponent* Component, bool
 			SelectedRegionComponents.Add(Component);
 		}
 	}
+#endif // WITH_EDITORONLY_DATA
 }
 
 void ULandscapeInfo::UnregisterActorComponent(ULandscapeComponent* Component)
