Landscape Splines: Left/Right Falloff Factor for Layer (split from Height falloff)

diff --git a/Engine/Source/Editor/LandscapeEditor/Private/LandscapeEdModeSplineTools.cpp b/Engine/Source/Editor/LandscapeEditor/Private/LandscapeEdModeSplineTools.cpp
--- a/Engine/Source/Editor/LandscapeEditor/Private/LandscapeEdModeSplineTools.cpp
+++ b/Engine/Source/Editor/LandscapeEditor/Private/LandscapeEdModeSplineTools.cpp
@@ -465,6 +465,8 @@ public:
 			NewControlPoint->SideFalloff = FirstPoint->SideFalloff;
 			NewControlPoint->LeftSideFalloffFactor = FirstPoint->LeftSideFalloffFactor;
 			NewControlPoint->RightSideFalloffFactor = FirstPoint->RightSideFalloffFactor;
+			NewControlPoint->LeftSideLayerFalloffFactor = FirstPoint->LeftSideLayerFalloffFactor;
+			NewControlPoint->RightSideLayerFalloffFactor = FirstPoint->RightSideLayerFalloffFactor;
 			NewControlPoint->EndFalloff  = FirstPoint->EndFalloff;
 
 			if (bCopyMeshToNewControlPoint)
@@ -673,6 +675,8 @@ public:
 		NewControlPoint->EndFalloff = FMath::Lerp(Segment->Connections[0].ControlPoint->EndFalloff, Segment->Connections[1].ControlPoint->EndFalloff, t);
 		NewControlPoint->LeftSideFalloffFactor = FMath::Clamp(FMath::Lerp(Segment->Connections[0].ControlPoint->LeftSideFalloffFactor, Segment->Connections[1].ControlPoint->LeftSideFalloffFactor, t), 0.f, 1.f);
 		NewControlPoint->RightSideFalloffFactor = FMath::Clamp(FMath::Lerp(Segment->Connections[0].ControlPoint->RightSideFalloffFactor, Segment->Connections[1].ControlPoint->RightSideFalloffFactor, t), 0.f, 1.f);
+		NewControlPoint->LeftSideLayerFalloffFactor = FMath::Clamp(FMath::Lerp(Segment->Connections[0].ControlPoint->LeftSideLayerFalloffFactor, Segment->Connections[1].ControlPoint->LeftSideLayerFalloffFactor, t), 0.f, 1.f);
+		NewControlPoint->RightSideLayerFalloffFactor = FMath::Clamp(FMath::Lerp(Segment->Connections[0].ControlPoint->RightSideLayerFalloffFactor, Segment->Connections[1].ControlPoint->RightSideLayerFalloffFactor, t), 0.f, 1.f);
 
 		ULandscapeSplineSegment* NewSegment = NewObject<ULandscapeSplineSegment>(SplinesComponent, NAME_None, RF_Transactional);
 		SplinesComponent->Segments.Add(NewSegment);
@@ -764,6 +768,8 @@ public:
 		ControlPoint->SideFalloff = FMath::Lerp(UseSegment->Connections[0].ControlPoint->SideFalloff, UseSegment->Connections[1].ControlPoint->SideFalloff, tseg);
 		ControlPoint->LeftSideFalloffFactor = FMath::Clamp(FMath::Lerp(UseSegment->Connections[0].ControlPoint->LeftSideFalloffFactor, UseSegment->Connections[1].ControlPoint->LeftSideFalloffFactor, tseg), 0.f, 1.f);
 		ControlPoint->RightSideFalloffFactor = FMath::Clamp(FMath::Lerp(UseSegment->Connections[0].ControlPoint->RightSideFalloffFactor, UseSegment->Connections[1].ControlPoint->RightSideFalloffFactor, tseg), 0.f, 1.f);
+		ControlPoint->LeftSideLayerFalloffFactor = FMath::Clamp(FMath::Lerp(UseSegment->Connections[0].ControlPoint->LeftSideLayerFalloffFactor, UseSegment->Connections[1].ControlPoint->LeftSideLayerFalloffFactor, tseg), 0.f, 1.f);
+		ControlPoint->RightSideLayerFalloffFactor = FMath::Clamp(FMath::Lerp(UseSegment->Connections[0].ControlPoint->RightSideLayerFalloffFactor, UseSegment->Connections[1].ControlPoint->RightSideLayerFalloffFactor, tseg), 0.f, 1.f);
 		ControlPoint->EndFalloff = FMath::Lerp(UseSegment->Connections[0].ControlPoint->EndFalloff, UseSegment->Connections[1].ControlPoint->EndFalloff, tseg);
 
 		Segment->Connections[0].TangentLen = DuplicateCacheSplitSegmentTangentLenStart * t;
diff --git a/Engine/Source/Runtime/Landscape/Classes/LandscapeSplineControlPoint.h b/Engine/Source/Runtime/Landscape/Classes/LandscapeSplineControlPoint.h
--- a/Engine/Source/Runtime/Landscape/Classes/LandscapeSplineControlPoint.h
+++ b/Engine/Source/Runtime/Landscape/Classes/LandscapeSplineControlPoint.h
@@ -74,6 +74,12 @@ class ULandscapeSplineControlPoint : public UObject
 	UPROPERTY(EditAnywhere, Category = LandscapeSpline, meta = (UIMin = 0, ClampMin = 0, UIMax = 1, ClampMax = 1))
 	float RightSideFalloffFactor = 1.f;
 
+	UPROPERTY(EditAnywhere, Category = LandscapeSpline, meta = (UIMin = 0, ClampMin = 0, UIMax = 1, ClampMax = 1))
+	float LeftSideLayerFalloffFactor = 0.5f;
+
+	UPROPERTY(EditAnywhere, Category = LandscapeSpline, meta = (UIMin = 0, ClampMin = 0, UIMax = 1, ClampMax = 1))
+	float RightSideLayerFalloffFactor = 0.5f;
+
 	/** Falloff at the start/end of the spline (if this point is a start or end point, otherwise ignored). */
 	UPROPERTY(EditAnywhere, Category=LandscapeSpline)
 	float EndFalloff;
diff --git a/Engine/Source/Runtime/Landscape/Classes/LandscapeSplineSegment.h b/Engine/Source/Runtime/Landscape/Classes/LandscapeSplineSegment.h
--- a/Engine/Source/Runtime/Landscape/Classes/LandscapeSplineSegment.h
+++ b/Engine/Source/Runtime/Landscape/Classes/LandscapeSplineSegment.h
@@ -42,6 +42,15 @@ struct FLandscapeSplineInterpPoint
 	UPROPERTY()
 	FVector FalloffRight;
 
+	/** Left Layer Falloff Point */
+	UPROPERTY()
+	FVector LayerFalloffLeft;
+
+	/** Right Layer FalloffPoint */
+	UPROPERTY()
+	FVector LayerFalloffRight;
+
+
 	/** Start/End Falloff fraction */
 	UPROPERTY()
 	float StartEndFalloff;
@@ -52,16 +61,20 @@ struct FLandscapeSplineInterpPoint
 		, Right(ForceInitToZero)
 		, FalloffLeft(ForceInitToZero)
 		, FalloffRight(ForceInitToZero)
+		, LayerFalloffLeft(ForceInitToZero)
+		, LayerFalloffRight(ForceInitToZero)
 		, StartEndFalloff(0.0f)
 	{
 	}
 
-	FLandscapeSplineInterpPoint(FVector InCenter, FVector InLeft, FVector InRight, FVector InFalloffLeft, FVector InFalloffRight, float InStartEndFalloff) :
+	FLandscapeSplineInterpPoint(FVector InCenter, FVector InLeft, FVector InRight, FVector InFalloffLeft, FVector InFalloffRight, FVector InLayerFalloffLeft, FVector InLayerFalloffRight, float InStartEndFalloff) :
 		Center(InCenter),
 		Left(InLeft),
 		Right(InRight),
 		FalloffLeft(InFalloffLeft),
 		FalloffRight(InFalloffRight),
+		LayerFalloffLeft(InLayerFalloffLeft),
+		LayerFalloffRight(InLayerFalloffRight),
 		StartEndFalloff(InStartEndFalloff)
 	{
 	}
diff --git a/Engine/Source/Runtime/Landscape/Private/LandscapeBlueprintSupport.cpp b/Engine/Source/Runtime/Landscape/Private/LandscapeBlueprintSupport.cpp
--- a/Engine/Source/Runtime/Landscape/Private/LandscapeBlueprintSupport.cpp
+++ b/Engine/Source/Runtime/Landscape/Private/LandscapeBlueprintSupport.cpp
@@ -20,7 +20,8 @@ void ALandscapeProxy::EditorApplySpline(USplineComponent* InSplineComponent, flo
 	if (InSplineComponent && !GetWorld()->IsGameWorld())
 	{
 		TArray<FLandscapeSplineInterpPoint> Points;
-		LandscapeSplineRaster::Pointify(InSplineComponent->SplineCurves.Position, Points, NumSubdivisions, 0.0f, 0.0f, StartWidth, EndWidth, StartSideFalloff, EndSideFalloff, StartSideFalloff, EndSideFalloff, StartRoll, EndRoll);
+		LandscapeSplineRaster::FPointifyFalloffs Falloffs(StartSideFalloff, EndSideFalloff);
+		LandscapeSplineRaster::Pointify(InSplineComponent->SplineCurves.Position, Points, NumSubdivisions, 0.0f, 0.0f, StartWidth, EndWidth, Falloffs, StartRoll, EndRoll);
 
 		FTransform SplineToWorld = InSplineComponent->GetComponentTransform();
 		LandscapeSplineRaster::RasterizeSegmentPoints(GetLandscapeInfo(), MoveTemp(Points), SplineToWorld, bRaiseHeights, bLowerHeights, PaintLayer);
diff --git a/Engine/Source/Runtime/Landscape/Private/LandscapeSplineRaster.cpp b/Engine/Source/Runtime/Landscape/Private/LandscapeSplineRaster.cpp
--- a/Engine/Source/Runtime/Landscape/Private/LandscapeSplineRaster.cpp
+++ b/Engine/Source/Runtime/Landscape/Private/LandscapeSplineRaster.cpp
@@ -326,8 +326,8 @@ void RasterizeControlPointAlpha(int32& MinX, int32& MinY, int32& MaxX, int32& Ma
 		Rasterizer.DrawTriangle(Center, Left1, Right1, CenterPos, Left1Pos, Right1Pos, false);
 
 		// Falloff
-		FVector2D FalloffRight0Pos = FVector2D(Points[i].FalloffRight);
-		FVector2D FalloffLeft1Pos = FVector2D(Points[j].FalloffLeft);
+		FVector2D FalloffRight0Pos = FVector2D(Points[i].LayerFalloffRight);
+		FVector2D FalloffLeft1Pos = FVector2D(Points[j].LayerFalloffLeft);
 		FVector FalloffRight0 = FVector(0.0f, Points[i].StartEndFalloff, BlendValue);
 		FVector FalloffLeft1 = FVector(0.0f, Points[j].StartEndFalloff, BlendValue);
 		Rasterizer.DrawTriangle(Right0, FalloffRight0, Left1, Right0Pos, FalloffRight0Pos, Left1Pos, false);
@@ -436,16 +436,16 @@ void RasterizeSegmentAlpha(int32& MinX, int32& MinY, int32& MaxX, int32& MaxY, F
 		Rasterizer.DrawTriangle(Right0, Left1, Right1, Right0Pos, Left1Pos, Right1Pos, false);
 
 		// Left Falloff
-		FVector2D FalloffLeft0Pos = FVector2D(Points[j - 1].FalloffLeft);
-		FVector2D FalloffLeft1Pos = FVector2D(Points[j].FalloffLeft);
+		FVector2D FalloffLeft0Pos = FVector2D(Points[j - 1].LayerFalloffLeft);
+		FVector2D FalloffLeft1Pos = FVector2D(Points[j].LayerFalloffLeft);
 		FVector FalloffLeft0 = FVector(0.0f, Points[j - 1].StartEndFalloff, BlendValue);
 		FVector FalloffLeft1 = FVector(0.0f, Points[j].StartEndFalloff, BlendValue);
 		Rasterizer.DrawTriangle(FalloffLeft0, Left0, FalloffLeft1, FalloffLeft0Pos, Left0Pos, FalloffLeft1Pos, false);
 		Rasterizer.DrawTriangle(Left0, FalloffLeft1, Left1, Left0Pos, FalloffLeft1Pos, Left1Pos, false);
 
 		// Right Falloff
-		FVector2D FalloffRight0Pos = FVector2D(Points[j - 1].FalloffRight);
-		FVector2D FalloffRight1Pos = FVector2D(Points[j].FalloffRight);
+		FVector2D FalloffRight0Pos = FVector2D(Points[j - 1].LayerFalloffRight);
+		FVector2D FalloffRight1Pos = FVector2D(Points[j].LayerFalloffRight);
 		FVector FalloffRight0 = FVector(0.0f, Points[j - 1].StartEndFalloff, BlendValue);
 		FVector FalloffRight1 = FVector(0.0f, Points[j].StartEndFalloff, BlendValue);
 		Rasterizer.DrawTriangle(Right0, FalloffRight0, Right1, Right0Pos, FalloffRight0Pos, Right1Pos, false);
@@ -535,6 +535,8 @@ bool ULandscapeInfo::ApplySplinesInternal(bool bOnlySelected, ALandscapeProxy* P
 			Points[j].Right = SplineToLandscape.TransformPosition(Points[j].Right);
 			Points[j].FalloffLeft = SplineToLandscape.TransformPosition(Points[j].FalloffLeft);
 			Points[j].FalloffRight = SplineToLandscape.TransformPosition(Points[j].FalloffRight);
+			Points[j].LayerFalloffLeft = SplineToLandscape.TransformPosition(Points[j].LayerFalloffLeft);
+			Points[j].LayerFalloffRight = SplineToLandscape.TransformPosition(Points[j].LayerFalloffRight);
 
 			// local-heights to texture value heights
 			Points[j].Left.Z = Points[j].Left.Z * LANDSCAPE_INV_ZSCALE + LandscapeDataAccess::MidValue;
@@ -595,12 +597,16 @@ bool ULandscapeInfo::ApplySplinesInternal(bool bOnlySelected, ALandscapeProxy* P
 			Points[j].Right = SplineToLandscape.TransformPosition(Points[j].Right);
 			Points[j].FalloffLeft = SplineToLandscape.TransformPosition(Points[j].FalloffLeft);
 			Points[j].FalloffRight = SplineToLandscape.TransformPosition(Points[j].FalloffRight);
+			Points[j].LayerFalloffLeft = SplineToLandscape.TransformPosition(Points[j].LayerFalloffLeft);
+			Points[j].LayerFalloffRight = SplineToLandscape.TransformPosition(Points[j].LayerFalloffRight);
 
 			// local-heights to texture value heights
 			Points[j].Left.Z = Points[j].Left.Z * LANDSCAPE_INV_ZSCALE + LandscapeDataAccess::MidValue;
 			Points[j].Right.Z = Points[j].Right.Z * LANDSCAPE_INV_ZSCALE + LandscapeDataAccess::MidValue;
 			Points[j].FalloffLeft.Z = Points[j].FalloffLeft.Z * LANDSCAPE_INV_ZSCALE + LandscapeDataAccess::MidValue;
 			Points[j].FalloffRight.Z = Points[j].FalloffRight.Z * LANDSCAPE_INV_ZSCALE + LandscapeDataAccess::MidValue;
+			Points[j].LayerFalloffLeft.Z = Points[j].LayerFalloffLeft.Z * LANDSCAPE_INV_ZSCALE + LandscapeDataAccess::MidValue;
+			Points[j].LayerFalloffRight.Z = Points[j].LayerFalloffRight.Z * LANDSCAPE_INV_ZSCALE + LandscapeDataAccess::MidValue;
 		}
 
 		// Heights raster
@@ -873,8 +879,7 @@ namespace LandscapeSplineRaster
 	void Pointify(const FInterpCurveVector& SplineInfo, TArray<FLandscapeSplineInterpPoint>& Points, int32 NumSubdivisions,
 		float StartFalloffFraction, float EndFalloffFraction,
 		const float StartWidth, const float EndWidth,
-		const float StartLeftSideFalloff, const float EndLeftSideFalloff,
-		const float StartRightSideFalloff, const float EndRightSideFalloff,
+		const FPointifyFalloffs& Falloffs,
 		const float StartRollDegrees, const float EndRollDegrees)
 	{
 		// Stop the start and end fall-off overlapping
@@ -894,8 +899,10 @@ namespace LandscapeSplineRaster
 			const float NewKeyTime = SplineInfo.Points[i].InVal;
 			const float NewKeyCosInterp = 0.5f - 0.5f * FMath::Cos(NewKeyTime * PI);
 			const float NewKeyWidth = FMath::Lerp(StartWidth, EndWidth, NewKeyCosInterp);
-			const float NewKeyLeftFalloff = FMath::Lerp(StartLeftSideFalloff, EndLeftSideFalloff, NewKeyCosInterp);
-			const float NewKeyRightFalloff = FMath::Lerp(StartRightSideFalloff, EndRightSideFalloff, NewKeyCosInterp);
+			const float NewKeyLeftFalloff = FMath::Lerp(Falloffs.StartLeftSide, Falloffs.EndLeftSide, NewKeyCosInterp);
+			const float NewKeyRightFalloff = FMath::Lerp(Falloffs.StartRightSide, Falloffs.EndRightSide, NewKeyCosInterp);
+			const float NewKeyLeftLayerFalloff = FMath::Lerp(Falloffs.StartLeftSideLayer, Falloffs.EndLeftSideLayer, NewKeyCosInterp);
+			const float NewKeyRightLayerFalloff = FMath::Lerp(Falloffs.StartRightSideLayer, Falloffs.EndRightSideLayer, NewKeyCosInterp);
 			const float NewKeyRoll = FMath::Lerp(StartRoll, EndRoll, NewKeyCosInterp);
 			const FVector NewKeyPos = SplineInfo.Eval(NewKeyTime, FVector::ZeroVector);
 			const FVector NewKeyTangent = SplineInfo.EvalDerivative(NewKeyTime, FVector::ZeroVector).GetSafeNormal();
@@ -904,6 +911,8 @@ namespace LandscapeSplineRaster
 			const FVector NewKeyRightPos = NewKeyPos + NewKeyBiNormal * NewKeyWidth;
 			const FVector NewKeyFalloffLeftPos = NewKeyPos - NewKeyBiNormal * (NewKeyWidth + NewKeyLeftFalloff);
 			const FVector NewKeyFalloffRightPos = NewKeyPos + NewKeyBiNormal * (NewKeyWidth + NewKeyRightFalloff);
+			const FVector NewKeyLayerFalloffLeftPos = NewKeyPos - NewKeyBiNormal * (NewKeyWidth + NewKeyLeftLayerFalloff);
+			const FVector NewKeyLayerFalloffRightPos = NewKeyPos + NewKeyBiNormal * (NewKeyWidth + NewKeyRightLayerFalloff);
 			const float NewKeyStartEndFalloff = FMath::Min((StartFalloffFraction > 0 ? NewKeyTime / StartFalloffFraction : 1.0f), (EndFalloffFraction > 0 ? (1 - NewKeyTime) / EndFalloffFraction : 1.0f));
 
 			// If not the first keypoint, interp from the last keypoint.
@@ -918,8 +927,10 @@ namespace LandscapeSplineRaster
 					const float NewTime = OldKeyTime + j*DrawSubstep;
 					const float NewCosInterp = 0.5f - 0.5f * FMath::Cos(NewTime * PI);
 					const float NewWidth = FMath::Lerp(StartWidth, EndWidth, NewCosInterp);
-					const float NewLeftFalloff = FMath::Lerp(StartLeftSideFalloff, EndLeftSideFalloff, NewCosInterp);
-					const float NewRightFalloff = FMath::Lerp(StartRightSideFalloff, EndRightSideFalloff, NewCosInterp);
+					const float NewLeftFalloff = FMath::Lerp(Falloffs.StartLeftSide, Falloffs.EndLeftSide, NewCosInterp);
+					const float NewRightFalloff = FMath::Lerp(Falloffs.StartRightSide, Falloffs.EndRightSide, NewCosInterp);
+					const float NewLeftLayerFalloff = FMath::Lerp(Falloffs.StartLeftSideLayer, Falloffs.EndLeftSideLayer, NewCosInterp);
+					const float NewRightLayerFalloff = FMath::Lerp(Falloffs.StartRightSideLayer, Falloffs.EndRightSideLayer, NewCosInterp);
 					const float NewRoll = FMath::Lerp(StartRoll, EndRoll, NewCosInterp);
 					const FVector NewPos = SplineInfo.Eval(NewTime, FVector::ZeroVector);
 					const FVector NewTangent = SplineInfo.EvalDerivative(NewTime, FVector::ZeroVector).GetSafeNormal();
@@ -928,13 +939,15 @@ namespace LandscapeSplineRaster
 					const FVector NewRightPos = NewPos + NewBiNormal * NewWidth;
 					const FVector NewFalloffLeftPos = NewPos - NewBiNormal * (NewWidth + NewLeftFalloff);
 					const FVector NewFalloffRightPos = NewPos + NewBiNormal * (NewWidth + NewRightFalloff);
+					const FVector NewLayerFalloffLeftPos = NewPos - NewBiNormal * (NewWidth + NewLeftLayerFalloff);
+					const FVector NewLayerFalloffRightPos = NewPos + NewBiNormal * (NewWidth + NewRightLayerFalloff);
 					const float NewStartEndFalloff = FMath::Min((StartFalloffFraction > 0 ? NewTime / StartFalloffFraction : 1.0f), (EndFalloffFraction > 0 ? (1 - NewTime) / EndFalloffFraction : 1.0f));
 
-					Points.Emplace(NewPos, NewLeftPos, NewRightPos, NewFalloffLeftPos, NewFalloffRightPos, NewStartEndFalloff);
+					Points.Emplace(NewPos, NewLeftPos, NewRightPos, NewFalloffLeftPos, NewFalloffRightPos, NewLayerFalloffLeftPos, NewLayerFalloffRightPos, NewStartEndFalloff);
 				}
 			}
 
-			Points.Emplace(NewKeyPos, NewKeyLeftPos, NewKeyRightPos, NewKeyFalloffLeftPos, NewKeyFalloffRightPos, NewKeyStartEndFalloff);
+			Points.Emplace(NewKeyPos, NewKeyLeftPos, NewKeyRightPos, NewKeyFalloffLeftPos, NewKeyFalloffRightPos, NewKeyLayerFalloffLeftPos, NewKeyLayerFalloffRightPos, NewKeyStartEndFalloff);
 
 			OldKeyTime = NewKeyTime;
 		}
@@ -944,6 +957,9 @@ namespace LandscapeSplineRaster
 		FixSelfIntersection(Points, &FLandscapeSplineInterpPoint::Right);
 		FixSelfIntersection(Points, &FLandscapeSplineInterpPoint::FalloffLeft);
 		FixSelfIntersection(Points, &FLandscapeSplineInterpPoint::FalloffRight);
+		FixSelfIntersection(Points, &FLandscapeSplineInterpPoint::LayerFalloffLeft);
+		FixSelfIntersection(Points, &FLandscapeSplineInterpPoint::LayerFalloffRight);
+
 	}
 }
 #endif
diff --git a/Engine/Source/Runtime/Landscape/Private/LandscapeSplineRaster.h b/Engine/Source/Runtime/Landscape/Private/LandscapeSplineRaster.h
--- a/Engine/Source/Runtime/Landscape/Private/LandscapeSplineRaster.h
+++ b/Engine/Source/Runtime/Landscape/Private/LandscapeSplineRaster.h
@@ -13,11 +13,40 @@ namespace LandscapeSplineRaster
 #if WITH_EDITOR
 	bool FixSelfIntersection(TArray<FLandscapeSplineInterpPoint>& Points, FVector FLandscapeSplineInterpPoint::* Side);
 
+	struct FPointifyFalloffs
+	{
+		FPointifyFalloffs(float StartFalloff, float EndFalloff)
+			: StartLeftSide(StartFalloff)
+			, EndLeftSide(EndFalloff)
+			, StartRightSide(StartFalloff)
+			, EndRightSide(EndFalloff)
+			, StartLeftSideLayer(StartFalloff)
+			, EndLeftSideLayer(EndFalloff)
+			, StartRightSideLayer(StartFalloff)
+			, EndRightSideLayer(EndFalloff)
+		{
+		}
+
+		FPointifyFalloffs()
+			: FPointifyFalloffs(0.0f, 0.0f)
+		{
+
+		}
+
+		float StartLeftSide;
+		float EndLeftSide;
+		float StartRightSide;
+		float EndRightSide;
+		float StartLeftSideLayer;
+		float EndLeftSideLayer;
+		float StartRightSideLayer;
+		float EndRightSideLayer;
+	};
+
 	void Pointify(const FInterpCurveVector& SplineInfo, TArray<FLandscapeSplineInterpPoint>& OutPoints, int32 NumSubdivisions,
 		float StartFalloffFraction, float EndFalloffFraction,
 		const float StartWidth, const float EndWidth,
-		const float StartLeftSideFalloff, const float EndLeftSideFalloff, 
-		const float StartRightSideFalloff, const float EndRightSideFalloff,
+		const FPointifyFalloffs& Falloffs,
 		const float StartRollDegrees, const float EndRollDegrees);
 
 	void RasterizeSegmentPoints(ULandscapeInfo* LandscapeInfo, TArray<FLandscapeSplineInterpPoint> Points, const FTransform& SplineToWorld, bool bRaiseTerrain, bool bLowerTerrain, ULandscapeLayerInfoObject* LayerInfo);
diff --git a/Engine/Source/Runtime/Landscape/Private/LandscapeSplines.cpp b/Engine/Source/Runtime/Landscape/Private/LandscapeSplines.cpp
--- a/Engine/Source/Runtime/Landscape/Private/LandscapeSplines.cpp
+++ b/Engine/Source/Runtime/Landscape/Private/LandscapeSplines.cpp
@@ -1293,6 +1293,18 @@ void ULandscapeSplineControlPoint::Serialize(FArchive& Ar)
 	{
 		BodyInstance.SetCollisionProfileName(bEnableCollision_DEPRECATED ? UCollisionProfile::BlockAll_ProfileName : UCollisionProfile::NoCollision_ProfileName);
 	}
+
+	if(Ar.IsLoading() && Ar.CustomVer(FLandscapeCustomVersion::GUID) < FLandscapeCustomVersion::AddSplineLayerFalloff)
+	{
+		LeftSideLayerFalloffFactor = LeftSideFalloffFactor;
+		RightSideLayerFalloffFactor = RightSideFalloffFactor;
+
+		for (FLandscapeSplineInterpPoint& Point : Points)
+		{
+			Point.LayerFalloffLeft = Point.FalloffLeft;
+			Point.LayerFalloffRight = Point.FalloffRight;
+		}
+	}
 #endif
 }
 
@@ -1757,6 +1769,8 @@ void ULandscapeSplineControlPoint::UpdateSplinePoints(bool bUpdateCollision, boo
 	
 	const float LeftSideFalloff = LeftSideFalloffFactor * SideFalloff;
 	const float RightSideFalloff = RightSideFalloffFactor * SideFalloff;
+	const float LeftSideLayerFalloff = LeftSideLayerFalloffFactor * SideFalloff;
+	const float RightSideLayerFalloff = RightSideLayerFalloffFactor * SideFalloff;
 
 	// Update "Points" array
 	if (Mesh != nullptr)
@@ -1775,8 +1789,10 @@ void ULandscapeSplineControlPoint::UpdateSplinePoints(bool bUpdateCollision, boo
 			const FVector RightPos = StartLocation + BiNormal * Width;
 			const FVector FalloffLeftPos = StartLocation - BiNormal * (Width + LeftSideFalloff);
 			const FVector FalloffRightPos = StartLocation + BiNormal * (Width + RightSideFalloff);
+			const FVector LayerFalloffLeftPos = StartLocation - BiNormal * (Width + LeftSideLayerFalloff);
+			const FVector LayerFalloffRightPos = StartLocation + BiNormal * (Width + RightSideLayerFalloff);
 
-			Points.Emplace(StartLocation, LeftPos, RightPos, FalloffLeftPos, FalloffRightPos, 1.0f);
+			Points.Emplace(StartLocation, LeftPos, RightPos, FalloffLeftPos, FalloffRightPos, LayerFalloffLeftPos, LayerFalloffRightPos, 1.0f);
 		}
 
 		const FVector CPLocation = Location;
@@ -1796,8 +1812,11 @@ void ULandscapeSplineControlPoint::UpdateSplinePoints(bool bUpdateCollision, boo
 		const FVector RightPos = StartLocation + BiNormal * Width;
 		const FVector FalloffLeftPos = StartLocation - BiNormal * (Width + LeftSideFalloff);
 		const FVector FalloffRightPos = StartLocation + BiNormal * (Width + RightSideFalloff);
+		const FVector LayerFalloffLeftPos = StartLocation - BiNormal * (Width + LeftSideLayerFalloff);
+		const FVector LayerFalloffRightPos = StartLocation + BiNormal * (Width + RightSideLayerFalloff);
+
 
-		Points.Emplace(StartLocation, LeftPos, RightPos, FalloffLeftPos, FalloffRightPos, 1.0f);
+		Points.Emplace(StartLocation, LeftPos, RightPos, FalloffLeftPos, FalloffRightPos, LayerFalloffLeftPos, LayerFalloffRightPos, 1.0f);
 	}
 
 	// Update bounds
@@ -1903,6 +1922,8 @@ void ULandscapeSplineControlPoint::PostEditChangeProperty(FPropertyChangedEvent&
 	SideFalloff = FMath::Max(SideFalloff, 0.0f);
 	LeftSideFalloffFactor = FMath::Clamp(LeftSideFalloffFactor, 0.0f, 1.0f);
 	RightSideFalloffFactor = FMath::Clamp(RightSideFalloffFactor, 0.0f, 1.0f);
+	LeftSideLayerFalloffFactor = FMath::Clamp(LeftSideLayerFalloffFactor, 0.0f, 1.0f);
+	RightSideLayerFalloffFactor = FMath::Clamp(RightSideLayerFalloffFactor, 0.0f, 1.0f);
 	EndFalloff = FMath::Max(EndFalloff, 0.0f);
 
 	// Don't update splines when undoing, not only is it unnecessary and expensive,
@@ -2009,6 +2030,15 @@ void ULandscapeSplineSegment::Serialize(FArchive& Ar)
 	{
 		BodyInstance.SetCollisionProfileName(bEnableCollision_DEPRECATED ? UCollisionProfile::BlockAll_ProfileName : UCollisionProfile::NoCollision_ProfileName);
 	}
+
+	if (Ar.IsLoading() && Ar.CustomVer(FLandscapeCustomVersion::GUID) < FLandscapeCustomVersion::AddSplineLayerFalloff)
+	{
+		for (FLandscapeSplineInterpPoint& Point : Points)
+		{
+			Point.LayerFalloffLeft = Point.FalloffLeft;
+			Point.LayerFalloffRight = Point.FalloffRight;
+		}
+	}
 #endif
 }
 
@@ -2264,10 +2294,16 @@ void ULandscapeSplineSegment::UpdateSplinePoints(bool bUpdateCollision, bool bUp
 	const float EndFalloffFraction = ((Connections[1].ControlPoint->ConnectedSegments.Num() > 1) ? 0 : (Connections[1].ControlPoint->EndFalloff / SplineLength));
 	const float StartWidth = Connections[0].ControlPoint->Width;
 	const float EndWidth = Connections[1].ControlPoint->Width;
-	const float StartLeftSideFalloff = Connections[0].ControlPoint->LeftSideFalloffFactor * Connections[0].ControlPoint->SideFalloff;
-	const float EndLeftSideFalloff = Connections[1].ControlPoint->LeftSideFalloffFactor * Connections[1].ControlPoint->SideFalloff;
-	const float StartRightSideFalloff = Connections[0].ControlPoint->RightSideFalloffFactor * Connections[0].ControlPoint->SideFalloff;
-	const float EndRightSideFalloff = Connections[1].ControlPoint->RightSideFalloffFactor * Connections[1].ControlPoint->SideFalloff;
+
+	LandscapeSplineRaster::FPointifyFalloffs Falloffs;
+	Falloffs.StartLeftSide = Connections[0].ControlPoint->LeftSideFalloffFactor * Connections[0].ControlPoint->SideFalloff;
+	Falloffs.EndLeftSide = Connections[1].ControlPoint->LeftSideFalloffFactor * Connections[1].ControlPoint->SideFalloff;
+	Falloffs.StartRightSide = Connections[0].ControlPoint->RightSideFalloffFactor * Connections[0].ControlPoint->SideFalloff;
+	Falloffs.EndRightSide = Connections[1].ControlPoint->RightSideFalloffFactor * Connections[1].ControlPoint->SideFalloff;
+	Falloffs.StartLeftSideLayer = Connections[0].ControlPoint->LeftSideLayerFalloffFactor * Connections[0].ControlPoint->SideFalloff;
+	Falloffs.EndLeftSideLayer = Connections[1].ControlPoint->LeftSideLayerFalloffFactor * Connections[1].ControlPoint->SideFalloff;
+	Falloffs.StartRightSideLayer = Connections[0].ControlPoint->RightSideLayerFalloffFactor * Connections[0].ControlPoint->SideFalloff;
+	Falloffs.EndRightSideLayer = Connections[1].ControlPoint->RightSideLayerFalloffFactor * Connections[1].ControlPoint->SideFalloff;
 	const float StartRollDegrees = StartRotation.Roll * (Connections[0].TangentLen > 0 ? 1 : -1);
 	const float EndRollDegrees = EndRotation.Roll * (Connections[1].TangentLen > 0 ? -1 : 1);
 	const float StartRoll = FMath::DegreesToRadians(StartRollDegrees);
@@ -2278,7 +2314,7 @@ void ULandscapeSplineSegment::UpdateSplinePoints(bool bUpdateCollision, bool bUp
 	int32 NumPoints = FMath::CeilToInt(SplineLength / OuterSplines->SplineResolution);
 	NumPoints = FMath::Clamp(NumPoints, 1, 1000);
 
-	LandscapeSplineRaster::Pointify(SplineInfo, Points, NumPoints, StartFalloffFraction, EndFalloffFraction, StartWidth, EndWidth, StartLeftSideFalloff, EndLeftSideFalloff, StartRightSideFalloff, EndRightSideFalloff, StartRollDegrees, EndRollDegrees);
+	LandscapeSplineRaster::Pointify(SplineInfo, Points, NumPoints, StartFalloffFraction, EndFalloffFraction, StartWidth, EndWidth, Falloffs, StartRollDegrees, EndRollDegrees);
 
 	// Update Bounds
 	Bounds = FBox(ForceInit);
diff --git a/Engine/Source/Runtime/Landscape/Private/LandscapeVersion.h b/Engine/Source/Runtime/Landscape/Private/LandscapeVersion.h
--- a/Engine/Source/Runtime/Landscape/Private/LandscapeVersion.h
+++ b/Engine/Source/Runtime/Landscape/Private/LandscapeVersion.h
@@ -28,6 +28,8 @@ namespace FLandscapeCustomVersion
 		MigrateOldPropertiesToNewRenderingProperties,
 		// Migration of old EnabledCollision for spline elements to the full body instance exposed
 		AddingBodyInstanceToSplinesElements,
+		// Spline Layer Falloff
+		AddSplineLayerFalloff,
 		// -----<new versions can be added above this line>-------------------------------------------------
 		VersionPlusOne,
 		LatestVersion = VersionPlusOne - 1
