Compile fixes for landscape - final set (hopefully).

diff --git a/Engine/Source/Runtime/Landscape/Classes/LandscapeHeightfieldCollisionComponent.h b/Engine/Source/Runtime/Landscape/Classes/LandscapeHeightfieldCollisionComponent.h
--- a/Engine/Source/Runtime/Landscape/Classes/LandscapeHeightfieldCollisionComponent.h
+++ b/Engine/Source/Runtime/Landscape/Classes/LandscapeHeightfieldCollisionComponent.h
@@ -184,8 +184,6 @@ protected:
 	virtual void OnCreatePhysicsState() override;
 public:
 
-	TOptional<float> GetHeight(float X, float Y);
-
 	virtual void ApplyWorldOffset(const FVector& InOffset, bool bWorldShift) override;
 	//~ End UActorComponent Interface.
 
@@ -265,6 +263,8 @@ public:
 
 	LANDSCAPE_API void SnapFoliageInstances();
 #endif
+	public:
+	TOptional<float> GetHeight(float X, float Y);
 };
 
 
diff --git a/Engine/Source/Runtime/Landscape/Private/Landscape.cpp b/Engine/Source/Runtime/Landscape/Private/Landscape.cpp
--- a/Engine/Source/Runtime/Landscape/Private/Landscape.cpp
+++ b/Engine/Source/Runtime/Landscape/Private/Landscape.cpp
@@ -1319,12 +1319,10 @@
 	return GetLandscapeProxy()->VirtualTextureRenderPassType;
 }
 
-#if WITH_EDITOR
 ULandscapeInfo* ULandscapeComponent::GetLandscapeInfo() const
 {
 	return GetLandscapeProxy()->GetLandscapeInfo();
 }
-#endif // WITH_EDITOR
 
 void ULandscapeComponent::BeginDestroy()
 {
@@ -2753,136 +2751,6 @@
 	}
 }
 
-void ULandscapeInfo::RegisterActor(ALandscapeProxy* Proxy, bool bMapCheck)
-{
-	// do not pass here invalid actors
-	checkSlow(Proxy);
-	check(Proxy->GetLandscapeGuid().IsValid());
-	UWorld* OwningWorld = Proxy->GetWorld();
-
-	// in case this Info object is not initialized yet
-	// initialized it with properties from passed actor
-	if (LandscapeGuid.IsValid() == false ||
-		(GetLandscapeProxy() == nullptr && ensure(LandscapeGuid == Proxy->GetLandscapeGuid())))
-	{
-		LandscapeGuid = Proxy->GetLandscapeGuid();
-		ComponentSizeQuads = Proxy->ComponentSizeQuads;
-		ComponentNumSubsections = Proxy->NumSubsections;
-		SubsectionSizeQuads = Proxy->SubsectionSizeQuads;
-		DrawScale = Proxy->GetRootComponent() != nullptr ? Proxy->GetRootComponent()->RelativeScale3D : FVector(100.0f);
-	}
-
-	// check that passed actor matches all shared parameters
-	check(LandscapeGuid == Proxy->GetLandscapeGuid());
-	check(ComponentSizeQuads == Proxy->ComponentSizeQuads);
-	check(ComponentNumSubsections == Proxy->NumSubsections);
-	check(SubsectionSizeQuads == Proxy->SubsectionSizeQuads);
-
-	if (Proxy->GetRootComponent() != nullptr && !DrawScale.Equals(Proxy->GetRootComponent()->RelativeScale3D))
-	{
-		UE_LOG(LogLandscape, Warning, TEXT("Landscape proxy (%s) scale (%s) does not match to main actor scale (%s)."),
-			*Proxy->GetName(), *Proxy->GetRootComponent()->RelativeScale3D.ToCompactString(), *DrawScale.ToCompactString());
-	}
-
-	// register
-	if (ALandscape* Landscape = Cast<ALandscape>(Proxy))
-	{
-		checkf(!LandscapeActor || LandscapeActor == Landscape, TEXT("Multiple landscapes with the same GUID detected: %s vs %s"), *LandscapeActor->GetPathName(), *Landscape->GetPathName());
-		LandscapeActor = Landscape;
-		// In world composition user is not allowed to move landscape in editor, only through WorldBrowser 
-		LandscapeActor->bLockLocation = OwningWorld != nullptr ? OwningWorld->WorldComposition != nullptr : false;
-
-		// update proxies reference actor
-		for (ALandscapeStreamingProxy* StreamingProxy : Proxies)
-		{
-			StreamingProxy->LandscapeActor = LandscapeActor;
-			StreamingProxy->FixupSharedData(Landscape);
-		}
-	}
-	else
-	{
-		// Insert Proxies in a sorted fashion for generating determisnitic results in the Layer system
-		ALandscapeStreamingProxy* StreamingProxy = CastChecked<ALandscapeStreamingProxy>(Proxy);
-		if (!Proxies.Contains(Proxy))
-		{
-			uint32 InsertIndex = Algo::LowerBound(Proxies, Proxy, [](ALandscapeProxy* A, ALandscapeProxy* B)
-			{
-				FIntPoint SectionBaseA = A->GetSectionBaseOffset();
-				FIntPoint SectionBaseB = B->GetSectionBaseOffset();
-
-				if (SectionBaseA.X != SectionBaseB.X)
-				{
-					return SectionBaseA.X < SectionBaseB.X;
-				}
-
-				return SectionBaseA.Y < SectionBaseB.Y;
-			});
-			
-			Proxies.Insert(StreamingProxy, InsertIndex);
-		}
-		StreamingProxy->LandscapeActor = LandscapeActor;
-		StreamingProxy->FixupSharedData(LandscapeActor.Get());
-	}
-
-	if (LandscapeActor && LandscapeActor->HasLayersContent())
-	{
-		const bool bInRequestContentUpdate = false;
-		LandscapeActor->RequestLayersInitialization(bInRequestContentUpdate);
-	}
-
-	UpdateLayerInfoMap(Proxy);
-	UpdateAllAddCollisions();
-
-	//
-	// add proxy components to the XY map
-	//
-	for (int32 CompIdx = 0; CompIdx < Proxy->LandscapeComponents.Num(); ++CompIdx)
-	{
-		RegisterActorComponent(Proxy->LandscapeComponents[CompIdx], bMapCheck);
-	}
-}
-
-void ULandscapeInfo::UnregisterActor(ALandscapeProxy* Proxy)
-{
-	if (ALandscape* Landscape = Cast<ALandscape>(Proxy))
-	{
-		// Note: UnregisterActor sometimes gets triggered twice, e.g. it has been observed to happen during redo
-		// Note: In some cases LandscapeActor could be updated to a new landscape actor before the old landscape is unregistered/destroyed
-		// e.g. this has been observed when merging levels in the editor
-
-		if (LandscapeActor.Get() == Landscape)
-		{
-		LandscapeActor = nullptr;
-		}
-
-		// update proxies reference to landscape actor
-		for (ALandscapeStreamingProxy* StreamingProxy : Proxies)
-		{
-			StreamingProxy->LandscapeActor = LandscapeActor;
-		}
-	}
-	else
-	{
-		ALandscapeStreamingProxy* StreamingProxy = CastChecked<ALandscapeStreamingProxy>(Proxy);
-		Proxies.Remove(StreamingProxy);
-		StreamingProxy->LandscapeActor = nullptr;
-	}
-
-	// remove proxy components from the XY map
-	for (int32 CompIdx = 0; CompIdx < Proxy->LandscapeComponents.Num(); ++CompIdx)
-	{
-		ULandscapeComponent* Component = Proxy->LandscapeComponents[CompIdx];
-		if (Component) // When a landscape actor is being GC'd it's possible the components were already GC'd and are null
-		{
-			UnregisterActorComponent(Component);
-		}
-	}
-	XYtoComponentMap.Compact();
-
-	UpdateLayerInfoMap();
-	UpdateAllAddCollisions();
-}
-
 void ULandscapeInfo::Reset()
 {
 	LandscapeActor.Reset();
@@ -3021,6 +2889,142 @@
 
 #endif
 
+void ULandscapeInfo::RegisterActor(ALandscapeProxy* Proxy, bool bMapCheck)
+{
+#if WITH_EDITOR
+	// do not pass here invalid actors
+	checkSlow(Proxy);
+	check(Proxy->GetLandscapeGuid().IsValid());
+	UWorld* OwningWorld = Proxy->GetWorld();
+
+	// in case this Info object is not initialized yet
+	// initialized it with properties from passed actor
+	if (LandscapeGuid.IsValid() == false ||
+		(GetLandscapeProxy() == nullptr && ensure(LandscapeGuid == Proxy->GetLandscapeGuid())))
+	{
+		LandscapeGuid = Proxy->GetLandscapeGuid();
+		ComponentSizeQuads = Proxy->ComponentSizeQuads;
+		ComponentNumSubsections = Proxy->NumSubsections;
+		SubsectionSizeQuads = Proxy->SubsectionSizeQuads;
+		DrawScale = Proxy->GetRootComponent() != nullptr ? Proxy->GetRootComponent()->RelativeScale3D : FVector(100.0f);
+	}
+
+	// check that passed actor matches all shared parameters
+	check(LandscapeGuid == Proxy->GetLandscapeGuid());
+	check(ComponentSizeQuads == Proxy->ComponentSizeQuads);
+	check(ComponentNumSubsections == Proxy->NumSubsections);
+	check(SubsectionSizeQuads == Proxy->SubsectionSizeQuads);
+
+	if (Proxy->GetRootComponent() != nullptr && !DrawScale.Equals(Proxy->GetRootComponent()->RelativeScale3D))
+	{
+		UE_LOG(LogLandscape, Warning, TEXT("Landscape proxy (%s) scale (%s) does not match to main actor scale (%s)."),
+			*Proxy->GetName(), *Proxy->GetRootComponent()->RelativeScale3D.ToCompactString(), *DrawScale.ToCompactString());
+	}
+
+	// register
+	if (ALandscape* Landscape = Cast<ALandscape>(Proxy))
+	{
+		checkf(!LandscapeActor || LandscapeActor == Landscape, TEXT("Multiple landscapes with the same GUID detected: %s vs %s"), *LandscapeActor->GetPathName(), *Landscape->GetPathName());
+		LandscapeActor = Landscape;
+		// In world composition user is not allowed to move landscape in editor, only through WorldBrowser 
+		LandscapeActor->bLockLocation = OwningWorld != nullptr ? OwningWorld->WorldComposition != nullptr : false;
+
+		// update proxies reference actor
+		for (ALandscapeStreamingProxy* StreamingProxy : Proxies)
+		{
+			StreamingProxy->LandscapeActor = LandscapeActor;
+			StreamingProxy->FixupSharedData(Landscape);
+		}
+	}
+	else
+	{
+		// Insert Proxies in a sorted fashion for generating determisnitic results in the Layer system
+		ALandscapeStreamingProxy* StreamingProxy = CastChecked<ALandscapeStreamingProxy>(Proxy);
+		if (!Proxies.Contains(Proxy))
+		{
+			uint32 InsertIndex = Algo::LowerBound(Proxies, Proxy, [](ALandscapeProxy* A, ALandscapeProxy* B)
+			{
+				FIntPoint SectionBaseA = A->GetSectionBaseOffset();
+				FIntPoint SectionBaseB = B->GetSectionBaseOffset();
+
+				if (SectionBaseA.X != SectionBaseB.X)
+				{
+					return SectionBaseA.X < SectionBaseB.X;
+				}
+
+				return SectionBaseA.Y < SectionBaseB.Y;
+			});
+
+			Proxies.Insert(StreamingProxy, InsertIndex);
+		}
+		StreamingProxy->LandscapeActor = LandscapeActor;
+		StreamingProxy->FixupSharedData(LandscapeActor.Get());
+	}
+
+	if (LandscapeActor && LandscapeActor->HasLayersContent())
+	{
+		const bool bInRequestContentUpdate = false;
+		LandscapeActor->RequestLayersInitialization(bInRequestContentUpdate);
+	}
+
+	UpdateLayerInfoMap(Proxy);
+	UpdateAllAddCollisions();
+
+#endif
+	//
+	// add proxy components to the XY map
+	//
+	for (int32 CompIdx = 0; CompIdx < Proxy->LandscapeComponents.Num(); ++CompIdx)
+	{
+		RegisterActorComponent(Proxy->LandscapeComponents[CompIdx], bMapCheck);
+	}
+}
+
+void ULandscapeInfo::UnregisterActor(ALandscapeProxy* Proxy)
+{
+#if WITH_EDITOR
+	if (ALandscape* Landscape = Cast<ALandscape>(Proxy))
+	{
+		// Note: UnregisterActor sometimes gets triggered twice, e.g. it has been observed to happen during redo
+		// Note: In some cases LandscapeActor could be updated to a new landscape actor before the old landscape is unregistered/destroyed
+		// e.g. this has been observed when merging levels in the editor
+
+		if (LandscapeActor.Get() == Landscape)
+		{
+			LandscapeActor = nullptr;
+		}
+
+		// update proxies reference to landscape actor
+		for (ALandscapeStreamingProxy* StreamingProxy : Proxies)
+		{
+			StreamingProxy->LandscapeActor = LandscapeActor;
+		}
+	}
+	else
+	{
+		ALandscapeStreamingProxy* StreamingProxy = CastChecked<ALandscapeStreamingProxy>(Proxy);
+		Proxies.Remove(StreamingProxy);
+		StreamingProxy->LandscapeActor = nullptr;
+	}
+#endif
+
+	// remove proxy components from the XY map
+	for (int32 CompIdx = 0; CompIdx < Proxy->LandscapeComponents.Num(); ++CompIdx)
+	{
+		ULandscapeComponent* Component = Proxy->LandscapeComponents[CompIdx];
+		if (Component) // When a landscape actor is being GC'd it's possible the components were already GC'd and are null
+		{
+			UnregisterActorComponent(Component);
+		}
+	}
+	XYtoComponentMap.Compact();
+
+#if WITH_EDITOR
+	UpdateLayerInfoMap();
+	UpdateAllAddCollisions();
+#endif
+}
+
 void ULandscapeInfo::RegisterActorComponent(ULandscapeComponent* Component, bool bMapCheck)
 {
 	// Do not register components which are not part of the world
@@ -3065,7 +3069,8 @@
 		}
 	}
 
-#if WITH_EDITORONLY_DATA
+
+#if WITH_EDITOR
 	// Update Selected Components/Regions
 	if (Component->EditToolRenderData.SelectedType)
 	{
@@ -3078,7 +3083,7 @@
 			SelectedRegionComponents.Add(Component);
 		}
 	}
-#endif // WITH_EDITORONLY_DATA
+#endif
 }
 
 void ULandscapeInfo::UnregisterActorComponent(ULandscapeComponent* Component)
