Remove editor assert in landscapeinfo code.

diff --git a/Engine/Source/Runtime/Landscape/Private/Landscape.cpp b/Engine/Source/Runtime/Landscape/Private/Landscape.cpp
--- a/Engine/Source/Runtime/Landscape/Private/Landscape.cpp
+++ b/Engine/Source/Runtime/Landscape/Private/Landscape.cpp
@@ -1139,7 +1139,6 @@ ULandscapeInfo* ALandscapeProxy::CreateLandscapeInfo()
 {
 	ULandscapeInfo* LandscapeInfo = nullptr;
 
-	check(GIsEditor);
 	check(LandscapeGuid.IsValid());
 	UWorld* OwningWorld = GetWorld();
 	check(OwningWorld);
