Add back conditions for creating landscape info.

diff --git a/Engine/Source/Runtime/Landscape/Private/Landscape.cpp b/Engine/Source/Runtime/Landscape/Private/Landscape.cpp
--- a/Engine/Source/Runtime/Landscape/Private/Landscape.cpp
+++ b/Engine/Source/Runtime/Landscape/Private/Landscape.cpp
@@ -2272,7 +2272,12 @@ void ALandscapeProxy::PostLoad()
 		BodyInstance.FixupData(this);
 	}
 
-	CreateLandscapeInfo();
+	if ((GetLinker() && (GetLinker()->UE4Ver() < VER_UE4_LANDSCAPE_COMPONENT_LAZY_REFERENCES)) ||
+		LandscapeComponents.Num() != CollisionComponents.Num() ||
+		LandscapeComponents.ContainsByPredicate([](ULandscapeComponent* Comp) { return ((Comp != nullptr) && !Comp->CollisionComponent.IsValid()); }))
+	{
+		CreateLandscapeInfo();
+	}
 #if WITH_EDITOR
 	if (GIsEditor && !GetWorld()->IsGameWorld())
 	{
