Fix landscape non-editor.

diff --git a/Engine/Source/Runtime/Landscape/Private/Landscape.cpp b/Engine/Source/Runtime/Landscape/Private/Landscape.cpp
--- a/Engine/Source/Runtime/Landscape/Private/Landscape.cpp
+++ b/Engine/Source/Runtime/Landscape/Private/Landscape.cpp
@@ -1781,22 +1781,26 @@ void ULandscapeComponent::SetWeightmapTexturesUsage(const TArray<ULandscapeWeigh
 void ALandscapeProxy::PostRegisterAllComponents()
 {
 	Super::PostRegisterAllComponents();
-	ULandscapeInfo* LandscapeInfo = nullptr;
-	if (!IsPendingKillPending())
+
+#if WITH_EDITOR
+	// Game worlds don't have landscape infos
+	if (!GetWorld()->IsGameWorld() && !IsPendingKillPending())
 	{
 		// Duplicated Landscapes don't have a valid guid until PostEditImport is called, we'll register then
 		if (LandscapeGuid.IsValid())
 		{
-			LandscapeInfo = CreateLandscapeInfo();
+			ULandscapeInfo* LandscapeInfo = CreateLandscapeInfo();
+
+			LandscapeInfo->FixupProxiesTransform();
 		}
 	}
-#if WITH_EDITOR
-	// Game worlds don't have landscape infos
-	if (!GetWorld()->IsGameWorld() && !IsPendingKillPending())
+#else
+	if (!IsPendingKillPending())
 	{
+		// Duplicated Landscapes don't have a valid guid until PostEditImport is called, we'll register then
 		if (LandscapeGuid.IsValid())
 		{
-			LandscapeInfo->FixupProxiesTransform();
+			CreateLandscapeInfo();
 		}
 	}
 #endif
@@ -2272,12 +2276,6 @@ void ALandscapeProxy::PostLoad()
 		BodyInstance.FixupData(this);
 	}
 
-	if ((GetLinker() && (GetLinker()->UE4Ver() < VER_UE4_LANDSCAPE_COMPONENT_LAZY_REFERENCES)) ||
-		LandscapeComponents.Num() != CollisionComponents.Num() ||
-		LandscapeComponents.ContainsByPredicate([](ULandscapeComponent* Comp) { return ((Comp != nullptr) && !Comp->CollisionComponent.IsValid()); }))
-	{
-		CreateLandscapeInfo();
-	}
 #if WITH_EDITOR
 	if (GIsEditor && !GetWorld()->IsGameWorld())
 	{
@@ -2286,6 +2284,7 @@ void ALandscapeProxy::PostLoad()
 			LandscapeComponents.ContainsByPredicate([](ULandscapeComponent* Comp) { return ((Comp != nullptr) && !Comp->CollisionComponent.IsValid()); }))
 		{
 			// Need to clean up invalid collision components
+			CreateLandscapeInfo();
 			RecreateCollisionComponents();
 		}
 	}
@@ -2323,6 +2322,13 @@ void ALandscapeProxy::PostLoad()
 	// track feature level change to flush grass cache
 	FOnFeatureLevelChanged::FDelegate FeatureLevelChangedDelegate = FOnFeatureLevelChanged::FDelegate::CreateUObject(this, &ALandscapeProxy::OnFeatureLevelChanged);
 	FeatureLevelChangedDelegateHandle = GetWorld()->AddOnFeatureLevelChangedHandler(FeatureLevelChangedDelegate);
+#else
+	if ((GetLinker() && (GetLinker()->UE4Ver() < VER_UE4_LANDSCAPE_COMPONENT_LAZY_REFERENCES)) ||
+		LandscapeComponents.Num() != CollisionComponents.Num() ||
+		LandscapeComponents.ContainsByPredicate([](ULandscapeComponent* Comp) { return ((Comp != nullptr) && !Comp->CollisionComponent.IsValid()); }))
+	{
+		CreateLandscapeInfo();
+	}
 #endif
 }
 
