Fix landscape tests.

diff --git a/Engine/Source/Runtime/Landscape/Private/Landscape.cpp b/Engine/Source/Runtime/Landscape/Private/Landscape.cpp
--- a/Engine/Source/Runtime/Landscape/Private/Landscape.cpp
+++ b/Engine/Source/Runtime/Landscape/Private/Landscape.cpp
@@ -1780,26 +1780,22 @@ void ULandscapeComponent::SetWeightmapTexturesUsage(const TArray<ULandscapeWeigh
 void ALandscapeProxy::PostRegisterAllComponents()
 {
 	Super::PostRegisterAllComponents();
-
-#if WITH_EDITOR
-	// Game worlds don't have landscape infos
-	if (!GetWorld()->IsGameWorld() && !IsPendingKillPending())
+	ULandscapeInfo* LandscapeInfo = nullptr;
+	if (!IsPendingKillPending())
 	{
 		// Duplicated Landscapes don't have a valid guid until PostEditImport is called, we'll register then
 		if (LandscapeGuid.IsValid())
 		{
-			ULandscapeInfo* LandscapeInfo = CreateLandscapeInfo();
-
-			LandscapeInfo->FixupProxiesTransform();
+			LandscapeInfo = CreateLandscapeInfo();
 		}
 	}
-#else
-	if (!IsPendingKillPending())
+#if WITH_EDITOR
+	// Game worlds don't have landscape infos
+	if (!GetWorld()->IsGameWorld() && !IsPendingKillPending())
 	{
-		// Duplicated Landscapes don't have a valid guid until PostEditImport is called, we'll register then
 		if (LandscapeGuid.IsValid())
 		{
-			CreateLandscapeInfo();
+			LandscapeInfo->FixupProxiesTransform();
 		}
 	}
 #endif
@@ -2275,6 +2271,12 @@ void ALandscapeProxy::PostLoad()
 		BodyInstance.FixupData(this);
 	}
 
+	if ((GetLinker() && (GetLinker()->UE4Ver() < VER_UE4_LANDSCAPE_COMPONENT_LAZY_REFERENCES)) ||
+		LandscapeComponents.Num() != CollisionComponents.Num() ||
+		LandscapeComponents.ContainsByPredicate([](ULandscapeComponent* Comp) { return ((Comp != nullptr) && !Comp->CollisionComponent.IsValid()); }))
+	{
+		CreateLandscapeInfo();
+	}
 #if WITH_EDITOR
 	if (GIsEditor && !GetWorld()->IsGameWorld())
 	{
@@ -2283,7 +2285,6 @@ void ALandscapeProxy::PostLoad()
 			LandscapeComponents.ContainsByPredicate([](ULandscapeComponent* Comp) { return ((Comp != nullptr) && !Comp->CollisionComponent.IsValid()); }))
 		{
 			// Need to clean up invalid collision components
-			CreateLandscapeInfo();
 			RecreateCollisionComponents();
 		}
 	}
@@ -2321,13 +2322,6 @@ void ALandscapeProxy::PostLoad()
 	// track feature level change to flush grass cache
 	FOnFeatureLevelChanged::FDelegate FeatureLevelChangedDelegate = FOnFeatureLevelChanged::FDelegate::CreateUObject(this, &ALandscapeProxy::OnFeatureLevelChanged);
 	FeatureLevelChangedDelegateHandle = GetWorld()->AddOnFeatureLevelChangedHandler(FeatureLevelChangedDelegate);
-#else
-	if ((GetLinker() && (GetLinker()->UE4Ver() < VER_UE4_LANDSCAPE_COMPONENT_LAZY_REFERENCES)) ||
-		LandscapeComponents.Num() != CollisionComponents.Num() ||
-		LandscapeComponents.ContainsByPredicate([](ULandscapeComponent* Comp) { return ((Comp != nullptr) && !Comp->CollisionComponent.IsValid()); }))
-	{
-		CreateLandscapeInfo();
-	}
 #endif
 }
 
@@ -2901,84 +2895,86 @@ void ULandscapeInfo::RecreateLandscapeInfo(UWorld* InWorld, bool bMapCheck)
 void ULandscapeInfo::RegisterActor(ALandscapeProxy* Proxy, bool bMapCheck)
 {
 #if WITH_EDITOR
-	// do not pass here invalid actors
-	checkSlow(Proxy);
-	check(Proxy->GetLandscapeGuid().IsValid());
-	UWorld* OwningWorld = Proxy->GetWorld();
-
-	// in case this Info object is not initialized yet
-	// initialized it with properties from passed actor
-	if (LandscapeGuid.IsValid() == false ||
-		(GetLandscapeProxy() == nullptr && ensure(LandscapeGuid == Proxy->GetLandscapeGuid())))
+	if (!Proxy->GetWorld()->IsGameWorld())
 	{
-		LandscapeGuid = Proxy->GetLandscapeGuid();
-		ComponentSizeQuads = Proxy->ComponentSizeQuads;
-		ComponentNumSubsections = Proxy->NumSubsections;
-		SubsectionSizeQuads = Proxy->SubsectionSizeQuads;
-		DrawScale = Proxy->GetRootComponent() != nullptr ? Proxy->GetRootComponent()->RelativeScale3D : FVector(100.0f);
-	}
+		// do not pass here invalid actors
+		checkSlow(Proxy);
+		check(Proxy->GetLandscapeGuid().IsValid());
+		UWorld* OwningWorld = Proxy->GetWorld();
 
-	// check that passed actor matches all shared parameters
-	check(LandscapeGuid == Proxy->GetLandscapeGuid());
-	check(ComponentSizeQuads == Proxy->ComponentSizeQuads);
-	check(ComponentNumSubsections == Proxy->NumSubsections);
-	check(SubsectionSizeQuads == Proxy->SubsectionSizeQuads);
+		// in case this Info object is not initialized yet
+		// initialized it with properties from passed actor
+		if (LandscapeGuid.IsValid() == false ||
+			(GetLandscapeProxy() == nullptr && ensure(LandscapeGuid == Proxy->GetLandscapeGuid())))
+		{
+			LandscapeGuid = Proxy->GetLandscapeGuid();
+			ComponentSizeQuads = Proxy->ComponentSizeQuads;
+			ComponentNumSubsections = Proxy->NumSubsections;
+			SubsectionSizeQuads = Proxy->SubsectionSizeQuads;
+			DrawScale = Proxy->GetRootComponent() != nullptr ? Proxy->GetRootComponent()->RelativeScale3D : FVector(100.0f);
+		}
 
-	if (Proxy->GetRootComponent() != nullptr && !DrawScale.Equals(Proxy->GetRootComponent()->RelativeScale3D))
-	{
-		UE_LOG(LogLandscape, Warning, TEXT("Landscape proxy (%s) scale (%s) does not match to main actor scale (%s)."),
-			*Proxy->GetName(), *Proxy->GetRootComponent()->RelativeScale3D.ToCompactString(), *DrawScale.ToCompactString());
-	}
+		// check that passed actor matches all shared parameters
+		check(LandscapeGuid == Proxy->GetLandscapeGuid());
+		check(ComponentSizeQuads == Proxy->ComponentSizeQuads);
+		check(ComponentNumSubsections == Proxy->NumSubsections);
+		check(SubsectionSizeQuads == Proxy->SubsectionSizeQuads);
 
-	// register
-	if (ALandscape* Landscape = Cast<ALandscape>(Proxy))
-	{
-		checkf(!LandscapeActor || LandscapeActor == Landscape, TEXT("Multiple landscapes with the same GUID detected: %s vs %s"), *LandscapeActor->GetPathName(), *Landscape->GetPathName());
-		LandscapeActor = Landscape;
-		// In world composition user is not allowed to move landscape in editor, only through WorldBrowser 
-		LandscapeActor->bLockLocation = OwningWorld != nullptr ? OwningWorld->WorldComposition != nullptr : false;
+		if (Proxy->GetRootComponent() != nullptr && !DrawScale.Equals(Proxy->GetRootComponent()->RelativeScale3D))
+		{
+			UE_LOG(LogLandscape, Warning, TEXT("Landscape proxy (%s) scale (%s) does not match to main actor scale (%s)."),
+				*Proxy->GetName(), *Proxy->GetRootComponent()->RelativeScale3D.ToCompactString(), *DrawScale.ToCompactString());
+		}
 
-		// update proxies reference actor
-		for (ALandscapeStreamingProxy* StreamingProxy : Proxies)
+		// register
+		if (ALandscape* Landscape = Cast<ALandscape>(Proxy))
 		{
-			StreamingProxy->LandscapeActor = LandscapeActor;
-			StreamingProxy->FixupSharedData(Landscape);
+			checkf(!LandscapeActor || LandscapeActor == Landscape, TEXT("Multiple landscapes with the same GUID detected: %s vs %s"), *LandscapeActor->GetPathName(), *Landscape->GetPathName());
+			LandscapeActor = Landscape;
+			// In world composition user is not allowed to move landscape in editor, only through WorldBrowser 
+			LandscapeActor->bLockLocation = OwningWorld != nullptr ? OwningWorld->WorldComposition != nullptr : false;
+
+			// update proxies reference actor
+			for (ALandscapeStreamingProxy* StreamingProxy : Proxies)
+			{
+				StreamingProxy->LandscapeActor = LandscapeActor;
+				StreamingProxy->FixupSharedData(Landscape);
+			}
 		}
-	}
-	else
-	{
-		// Insert Proxies in a sorted fashion for generating determisnitic results in the Layer system
-		ALandscapeStreamingProxy* StreamingProxy = CastChecked<ALandscapeStreamingProxy>(Proxy);
-		if (!Proxies.Contains(Proxy))
+		else
 		{
-			uint32 InsertIndex = Algo::LowerBound(Proxies, Proxy, [](ALandscapeProxy* A, ALandscapeProxy* B)
+			// Insert Proxies in a sorted fashion for generating determisnitic results in the Layer system
+			ALandscapeStreamingProxy* StreamingProxy = CastChecked<ALandscapeStreamingProxy>(Proxy);
+			if (!Proxies.Contains(Proxy))
 			{
-				FIntPoint SectionBaseA = A->GetSectionBaseOffset();
-				FIntPoint SectionBaseB = B->GetSectionBaseOffset();
-
-				if (SectionBaseA.X != SectionBaseB.X)
+				uint32 InsertIndex = Algo::LowerBound(Proxies, Proxy, [](ALandscapeProxy* A, ALandscapeProxy* B)
 				{
-					return SectionBaseA.X < SectionBaseB.X;
-				}
+					FIntPoint SectionBaseA = A->GetSectionBaseOffset();
+					FIntPoint SectionBaseB = B->GetSectionBaseOffset();
 
-				return SectionBaseA.Y < SectionBaseB.Y;
-			});
+					if (SectionBaseA.X != SectionBaseB.X)
+					{
+						return SectionBaseA.X < SectionBaseB.X;
+					}
 
-			Proxies.Insert(StreamingProxy, InsertIndex);
-		}
-		StreamingProxy->LandscapeActor = LandscapeActor;
-		StreamingProxy->FixupSharedData(LandscapeActor.Get());
-	}
+					return SectionBaseA.Y < SectionBaseB.Y;
+				});
 
-	if (LandscapeActor && LandscapeActor->HasLayersContent())
-	{
-		const bool bInRequestContentUpdate = false;
-		LandscapeActor->RequestLayersInitialization(bInRequestContentUpdate);
-	}
+				Proxies.Insert(StreamingProxy, InsertIndex);
+			}
+			StreamingProxy->LandscapeActor = LandscapeActor;
+			StreamingProxy->FixupSharedData(LandscapeActor.Get());
+		}
 
-	UpdateLayerInfoMap(Proxy);
-	UpdateAllAddCollisions();
+		if (LandscapeActor && LandscapeActor->HasLayersContent())
+		{
+			const bool bInRequestContentUpdate = false;
+			LandscapeActor->RequestLayersInitialization(bInRequestContentUpdate);
+		}
 
+		UpdateLayerInfoMap(Proxy);
+		UpdateAllAddCollisions();
+	}
 #endif
 	//
 	// add proxy components to the XY map
diff --git a/Engine/Source/Runtime/Landscape/Private/LandscapeModule.cpp b/Engine/Source/Runtime/Landscape/Private/LandscapeModule.cpp
--- a/Engine/Source/Runtime/Landscape/Private/LandscapeModule.cpp
+++ b/Engine/Source/Runtime/Landscape/Private/LandscapeModule.cpp
@@ -185,14 +185,12 @@ void FLandscapeModule::StartupModule()
 		)
 	);
 
-#if WITH_EDITORONLY_DATA
 	FWorldDelegates::OnPostWorldCreation.AddStatic(
 		&WorldCreationEventFunction
 	);
 	FWorldDelegates::OnPreWorldFinishDestroy.AddStatic(
 		&WorldDestroyEventFunction
 	);
-#endif // WITH_EDITORONLY_DATA
 
 #if WITH_EDITOR
 	FWorldDelegates::OnPreWorldRename.AddStatic(
