Fix enum typo make condition always false.

#rnx

diff --git a/Engine/Source/Runtime/Landscape/Private/LandscapeEditLayers.cpp b/Engine/Source/Runtime/Landscape/Private/LandscapeEditLayers.cpp
--- a/Engine/Source/Runtime/Landscape/Private/LandscapeEditLayers.cpp
+++ b/Engine/Source/Runtime/Landscape/Private/LandscapeEditLayers.cpp
@@ -3546,7 +3546,7 @@ int32 ALandscape::RegenerateLayersWeightmaps(const TArray<ULandscapeComponent*>&
 	SCOPE_CYCLE_COUNTER(STAT_LandscapeLayersRegenerateWeightmaps);
 	const int32 AllWeightmapUpdateModes = (ELandscapeLayerUpdateMode::Update_Weightmap_All | ELandscapeLayerUpdateMode::Update_Weightmap_Editing | ELandscapeLayerUpdateMode::Update_Weightmap_Editing_NoCollision);
 	const int32 WeightmapUpdateModes = LayerContentUpdateModes & AllWeightmapUpdateModes;
-	const bool bSkipBrush = CVarLandscapeLayerBrushOptim.GetValueOnAnyThread() == 1 && ((WeightmapUpdateModes & AllWeightmapUpdateModes) == ELandscapeLayerUpdateMode::Update_Heightmap_Editing);
+	const bool bSkipBrush = CVarLandscapeLayerBrushOptim.GetValueOnAnyThread() == 1 && ((WeightmapUpdateModes & AllWeightmapUpdateModes) == ELandscapeLayerUpdateMode::Update_Weightmap_Editing);
 	const bool bForceRender = CVarOutputLayersDebugDrawCallName.GetValueOnAnyThread() == 1;
 	
 	ULandscapeInfo* Info = GetLandscapeInfo();
