Landscape Layers: Fix LandscapeOptim when components are sharing Heightmaps

#rnx

diff --git a/Engine/Source/Runtime/Landscape/Private/LandscapeEditLayers.cpp b/Engine/Source/Runtime/Landscape/Private/LandscapeEditLayers.cpp
--- a/Engine/Source/Runtime/Landscape/Private/LandscapeEditLayers.cpp
+++ b/Engine/Source/Runtime/Landscape/Private/LandscapeEditLayers.cpp
@@ -4568,6 +4568,7 @@ void ALandscape::UpdateLayersContent(bool bInWaitForStreaming, bool bInSkipMonit
 
 	bool bPartialUpdate = !bForceRender && !bUpdateAll && CVarLandscapeLayerOptim.GetValueOnAnyThread() == 1;
 	TSet<UTexture*> Heightmaps;
+	TSet<UTexture*> HeightmapsToRender;
 	TSet<ULandscapeComponent*> NeighborsComponents;
 	TSet<ULandscapeComponent*> WeightmapsComponents;
 	TSet<ULandscapeWeightmapUsage*> WeightmapUsagesToResolve;
@@ -4590,6 +4591,10 @@ void ALandscape::UpdateLayersContent(bool bInWaitForStreaming, bool bInSkipMonit
 				GetLandscapeComponentNeighborsToRender(Component, NeighborsComponents);
 				// Gather Heightmaps (All Components sharing Heightmap textures need to be rendered and resolved)
 				Heightmaps.Add(Component->GetHeightmap(false));
+				Component->ForEachLayer([&](const FGuid&, FLandscapeLayerComponentData& LayerData)
+				{
+					HeightmapsToRender.Add(LayerData.HeightmapData.Texture);
+				});
 				// Gather WeightmapUsages (Components sharing weightmap usages with the resolved Components need to be rendered so that the resolving is valid)
 				GetLandscapeComponentWeightmapsToRender(Component, WeightmapsComponents);
 			}
@@ -4600,6 +4605,15 @@ void ALandscape::UpdateLayersContent(bool bInWaitForStreaming, bool bInSkipMonit
 		}
 	});
 
+	// Because of Heightmap Sharing anytime we render a heightmap we need to render all the components that use it
+	for (ULandscapeComponent* NeighborsComponent : NeighborsComponents)
+	{
+		NeighborsComponent->ForEachLayer([&](const FGuid&, FLandscapeLayerComponentData& LayerData)
+		{
+			HeightmapsToRender.Add(LayerData.HeightmapData.Texture);
+		});
+	}
+
 	// Copy first list into others
 	LandscapeComponentsHeightmapsToResolve.Append(AllLandscapeComponents);
 	LandscapeComponentsHeightmapsToRender.Append(AllLandscapeComponents);
@@ -4620,6 +4634,21 @@ void ALandscape::UpdateLayersContent(bool bInWaitForStreaming, bool bInSkipMonit
 			{
 				LandscapeComponentsHeightmapsToRender.Add(Component);
 			}
+			else
+			{
+				bool bAdd = false;
+				Component->ForEachLayer([&](const FGuid&, FLandscapeLayerComponentData& LayerData)
+				{
+					if (HeightmapsToRender.Contains(LayerData.HeightmapData.Texture))
+					{
+						bAdd = true;
+					}
+				});
+				if (bAdd)
+				{
+					LandscapeComponentsHeightmapsToRender.Add(Component);
+				}
+			}
 
 			if (WeightmapsComponents.Contains(Component))
 			{
