Add r.RayTracing.Landscape cvar to control if landscapes should be included in ray tracing scene

(ushell-p4-cherrypick of 8511781 by [REDACTED_PERSON])

diff --git a/Engine/Source/Runtime/Landscape/Private/LandscapeRender.cpp b/Engine/Source/Runtime/Landscape/Private/LandscapeRender.cpp
--- a/Engine/Source/Runtime/Landscape/Private/LandscapeRender.cpp
+++ b/Engine/Source/Runtime/Landscape/Private/LandscapeRender.cpp
@@ -122,6 +122,11 @@ FAutoConsoleVariableRef CVarLandscapeDebugViewMode(
 #endif
 
 #if RHI_RAYTRACING
+static TAutoConsoleVariable<int32> CVarRayTracingLandscape(
+	TEXT("r.RayTracing.Landscape"),
+	1,
+	TEXT("Include landscapes in ray tracing effects (default = 1 (landscape enabled in ray tracing))"));
+
 int32 GLandscapeRayTracingGeometryLODsThatUpdateEveryFrame = 0;
 static FAutoConsoleVariableRef CVarLandscapeRayTracingGeometryLODsThatUpdateEveryFrame(
 	TEXT("r.RayTracing.Landscape.LODsUpdateEveryFrame"),
@@ -3091,6 +3096,11 @@ void FLandscapeComponentSceneProxy::GetDynamicMeshElements(const TArray<const FS
 #if RHI_RAYTRACING
 void FLandscapeComponentSceneProxy::GetDynamicRayTracingInstances(FRayTracingMaterialGatheringContext& Context, TArray<FRayTracingInstance>& OutRayTracingInstances)
 {
+	if (!CVarRayTracingLandscape.GetValueOnRenderThread())
+	{
+		return;
+	}
+
 	FMemStackBase& PrimitiveCustomDataMemStack = FMemStack::Get();
 
 	float MeshScreenSizeSquared = 0;
