Landscape Splines:
- Add Noise Texture Modulation to the Painting
- Add Width ratio fo the Painting
- Optimize: Defer Spline Layer Update to avoid doing to much updates when changing multiple properties at the same time

diff --git a/Engine/Source/Editor/LandscapeEditor/Private/LandscapeEdMode.cpp b/Engine/Source/Editor/LandscapeEditor/Private/LandscapeEdMode.cpp
--- a/Engine/Source/Editor/LandscapeEditor/Private/LandscapeEdMode.cpp
+++ b/Engine/Source/Editor/LandscapeEditor/Private/LandscapeEdMode.cpp
@@ -1061,6 +1061,16 @@ void FEdModeLandscape::Tick(FEditorViewportClient* ViewportClient, float DeltaTi
 			}
 		}
 	}
+
+	static int32 LastLandscapeSplineFalloffModulationValue = CVarLandscapeSplineFalloffModulation.GetValueOnAnyThread();
+	if (LastLandscapeSplineFalloffModulationValue != CVarLandscapeSplineFalloffModulation.GetValueOnAnyThread())
+	{
+		if (ALandscape* Landscape = GetLandscape())
+		{
+			Landscape->RequestSplineLayerUpdate();
+		}
+		LastLandscapeSplineFalloffModulationValue = CVarLandscapeSplineFalloffModulation.GetValueOnAnyThread();
+	}
 }
 
 
@@ -4607,7 +4617,7 @@ void FEdModeLandscape::AutoUpdateDirtyLandscapeSplines()
 		if (Landscape && Landscape->GetLandscapeSplinesReservedLayer())
 		{
 			// TODO : Only update dirty regions
-			UpdateLandscapeSplines();
+			Landscape->RequestSplineLayerUpdate();
 		}
 	}
 }
diff --git a/Engine/Source/Editor/LandscapeEditor/Private/LandscapeEdModeSplineTools.cpp b/Engine/Source/Editor/LandscapeEditor/Private/LandscapeEdModeSplineTools.cpp
--- a/Engine/Source/Editor/LandscapeEditor/Private/LandscapeEdModeSplineTools.cpp
+++ b/Engine/Source/Editor/LandscapeEditor/Private/LandscapeEdModeSplineTools.cpp
@@ -462,6 +462,7 @@ public:
 			}
 
 			NewControlPoint->Width       = FirstPoint->Width;
+			NewControlPoint->LayerWidthRatio = FirstPoint->LayerWidthRatio;
 			NewControlPoint->SideFalloff = FirstPoint->SideFalloff;
 			NewControlPoint->LeftSideFalloffFactor = FirstPoint->LeftSideFalloffFactor;
 			NewControlPoint->RightSideFalloffFactor = FirstPoint->RightSideFalloffFactor;
@@ -671,6 +672,7 @@ public:
 		}
 
 		NewControlPoint->Width = FMath::Lerp(Segment->Connections[0].ControlPoint->Width, Segment->Connections[1].ControlPoint->Width, t);
+		NewControlPoint->LayerWidthRatio = FMath::Lerp(Segment->Connections[0].ControlPoint->LayerWidthRatio, Segment->Connections[1].ControlPoint->LayerWidthRatio, t);
 		NewControlPoint->SideFalloff = FMath::Lerp(Segment->Connections[0].ControlPoint->SideFalloff, Segment->Connections[1].ControlPoint->SideFalloff, t);
 		NewControlPoint->EndFalloff = FMath::Lerp(Segment->Connections[0].ControlPoint->EndFalloff, Segment->Connections[1].ControlPoint->EndFalloff, t);
 		NewControlPoint->LeftSideFalloffFactor = FMath::Clamp(FMath::Lerp(Segment->Connections[0].ControlPoint->LeftSideFalloffFactor, Segment->Connections[1].ControlPoint->LeftSideFalloffFactor, t), 0.f, 1.f);
@@ -765,6 +767,7 @@ public:
 		// Do not update rotation during alt-drag.
 		//ControlPoint->Rotation.Roll = FMath::Lerp(UseSegment->Connections[0].ControlPoint->Rotation.Roll, UseSegment->Connections[1].ControlPoint->Rotation.Roll, tseg);
 		ControlPoint->Width = FMath::Lerp(UseSegment->Connections[0].ControlPoint->Width, UseSegment->Connections[1].ControlPoint->Width, tseg);
+		ControlPoint->LayerWidthRatio = FMath::Lerp(UseSegment->Connections[0].ControlPoint->LayerWidthRatio, UseSegment->Connections[1].ControlPoint->LayerWidthRatio, tseg);
 		ControlPoint->SideFalloff = FMath::Lerp(UseSegment->Connections[0].ControlPoint->SideFalloff, UseSegment->Connections[1].ControlPoint->SideFalloff, tseg);
 		ControlPoint->LeftSideFalloffFactor = FMath::Clamp(FMath::Lerp(UseSegment->Connections[0].ControlPoint->LeftSideFalloffFactor, UseSegment->Connections[1].ControlPoint->LeftSideFalloffFactor, tseg), 0.f, 1.f);
 		ControlPoint->RightSideFalloffFactor = FMath::Clamp(FMath::Lerp(UseSegment->Connections[0].ControlPoint->RightSideFalloffFactor, UseSegment->Connections[1].ControlPoint->RightSideFalloffFactor, tseg), 0.f, 1.f);
diff --git a/Engine/Source/Runtime/Landscape/Classes/Landscape.h b/Engine/Source/Runtime/Landscape/Classes/Landscape.h
--- a/Engine/Source/Runtime/Landscape/Classes/Landscape.h
+++ b/Engine/Source/Runtime/Landscape/Classes/Landscape.h
@@ -18,6 +18,10 @@ namespace ELandscapeToolTargetType
 	enum Type : int8;
 };
 
+#if WITH_EDITOR
+extern LANDSCAPE_API TAutoConsoleVariable<int32> CVarLandscapeSplineFalloffModulation;
+#endif
+
 UENUM()
 enum ELandscapeSetupErrors
 {
@@ -265,6 +269,7 @@ public:
 	LANDSCAPE_API void RegisterLandscapeEdMode(ILandscapeEdModeInterface* InLandscapeEdMode) { LandscapeEdMode = InLandscapeEdMode; }
 	LANDSCAPE_API void UnregisterLandscapeEdMode() { LandscapeEdMode = nullptr; }
 	LANDSCAPE_API virtual bool HasLayersContent() const override;
+	LANDSCAPE_API void RequestSplineLayerUpdate();
 	LANDSCAPE_API void RequestLayersInitialization(bool bInRequestContentUpdate = true);
 	LANDSCAPE_API void RequestLayersContentUpdateForceAll(ELandscapeLayerUpdateMode InModeMask = ELandscapeLayerUpdateMode::Update_All);
 	LANDSCAPE_API void RequestLayersContentUpdate(ELandscapeLayerUpdateMode InModeMask);
@@ -320,7 +325,7 @@ public:
 	LANDSCAPE_API class ALandscapeBlueprintBrushBase* GetBrushForLayer(int32 InLayerIndex, int8 BrushIndex) const;
 	LANDSCAPE_API TArray<class ALandscapeBlueprintBrushBase*> GetBrushesForLayer(int32 InLayerIndex) const;
 	LANDSCAPE_API void OnBlueprintBrushChanged();
-	
+	LANDSCAPE_API void OnLayerInfoSplineFalloffModulationChanged(ULandscapeLayerInfoObject* InLayerInfo);
 	LANDSCAPE_API void OnPreSave();
 
 	void ReleaseLayersRenderingResource();
@@ -463,6 +468,9 @@ private:
 	UPROPERTY(Transient)
 	uint32 LayerContentUpdateModes;
 		
+	UPROPERTY(Transient)
+	bool bSplineLayerUpdateRequested;
+
 	// Represent all the resolved paint layer, from all layers blended together (size of the landscape x material layer count)
 	class FLandscapeTexture2DArrayResource* CombinedLayersWeightmapAllMaterialLayersResource;
 	
diff --git a/Engine/Source/Runtime/Landscape/Classes/LandscapeLayerInfoObject.h b/Engine/Source/Runtime/Landscape/Classes/LandscapeLayerInfoObject.h
--- a/Engine/Source/Runtime/Landscape/Classes/LandscapeLayerInfoObject.h
+++ b/Engine/Source/Runtime/Landscape/Classes/LandscapeLayerInfoObject.h
@@ -10,6 +10,15 @@
 class UPhysicalMaterial;
 struct FPropertyChangedEvent;
 
+UENUM()
+enum class ESplineModulationColorMask : uint8
+{
+	Red,
+	Green,
+	Blue,
+	Alpha
+};
+
 UCLASS(MinimalAPI, BlueprintType)
 class ULandscapeLayerInfoObject : public UObject
 {
@@ -28,6 +37,22 @@ class ULandscapeLayerInfoObject : public UObject
 	UPROPERTY(EditAnywhere, Category=LandscapeLayerInfoObject)
 	uint32 bNoWeightBlend:1;
 
+	/** Texture to modulate the Splines Falloff Layer Alpha */
+	UPROPERTY(EditAnywhere, Category = SplineFalloffModulation, Meta = (DisplayName="Texture"))
+	UTexture2D* SplineFalloffModulationTexture;
+
+	UPROPERTY(EditAnywhere, Category = SplineFalloffModulation, Meta = (DisplayName = "Color Mask"))
+	ESplineModulationColorMask SplineFalloffModulationColorMask;
+
+	UPROPERTY(EditAnywhere, Category = SplineFalloffModulation, Meta = (DisplayName = "Tiling", ClampMin="0.01"))
+	float SplineFalloffModulationTiling;
+
+	UPROPERTY(EditAnywhere, Category = SplineFalloffModulation, Meta = (DisplayName = "Bias", ClampMin="0"))
+	float SplineFalloffModulationBias;
+
+	UPROPERTY(EditAnywhere, Category = SplineFalloffModulation, Meta = (DisplayName = "Scale", ClampMin="0"))
+	float SplineFalloffModulationScale;
+				
 	UPROPERTY(NonTransactional, Transient)
 	bool IsReferencedFromLoadedData;
 #endif // WITH_EDITORONLY_DATA
diff --git a/Engine/Source/Runtime/Landscape/Classes/LandscapeSplineControlPoint.h b/Engine/Source/Runtime/Landscape/Classes/LandscapeSplineControlPoint.h
--- a/Engine/Source/Runtime/Landscape/Classes/LandscapeSplineControlPoint.h
+++ b/Engine/Source/Runtime/Landscape/Classes/LandscapeSplineControlPoint.h
@@ -64,6 +64,10 @@ class ULandscapeSplineControlPoint : public UObject
 	UPROPERTY(EditAnywhere, Category=LandscapeSpline)
 	float Width;
 
+	/** Layer Width ratio of the spline at this point. */
+	UPROPERTY(EditAnywhere, Category = LandscapeSpline)
+	float LayerWidthRatio = 1.f;
+
 	/** Falloff at the sides of the spline at this point. */
 	UPROPERTY(EditAnywhere, Category=LandscapeSpline)
 	float SideFalloff;
diff --git a/Engine/Source/Runtime/Landscape/Classes/LandscapeSplineSegment.h b/Engine/Source/Runtime/Landscape/Classes/LandscapeSplineSegment.h
--- a/Engine/Source/Runtime/Landscape/Classes/LandscapeSplineSegment.h
+++ b/Engine/Source/Runtime/Landscape/Classes/LandscapeSplineSegment.h
@@ -42,6 +42,14 @@ struct FLandscapeSplineInterpPoint
 	UPROPERTY()
 	FVector FalloffRight;
 
+	/** Layer Left Point */
+	UPROPERTY()
+	FVector LayerLeft;
+
+	/** Layer Right Point */
+	UPROPERTY()
+	FVector LayerRight;
+
 	/** Left Layer Falloff Point */
 	UPROPERTY()
 	FVector LayerFalloffLeft;
@@ -61,18 +69,22 @@ struct FLandscapeSplineInterpPoint
 		, Right(ForceInitToZero)
 		, FalloffLeft(ForceInitToZero)
 		, FalloffRight(ForceInitToZero)
+		, LayerLeft(ForceInitToZero)
+		, LayerRight(ForceInitToZero)
 		, LayerFalloffLeft(ForceInitToZero)
 		, LayerFalloffRight(ForceInitToZero)
 		, StartEndFalloff(0.0f)
 	{
 	}
 
-	FLandscapeSplineInterpPoint(FVector InCenter, FVector InLeft, FVector InRight, FVector InFalloffLeft, FVector InFalloffRight, FVector InLayerFalloffLeft, FVector InLayerFalloffRight, float InStartEndFalloff) :
+	FLandscapeSplineInterpPoint(FVector InCenter, FVector InLeft, FVector InRight, FVector InFalloffLeft, FVector InFalloffRight, FVector InLayerLeft, FVector InLayerRight, FVector InLayerFalloffLeft, FVector InLayerFalloffRight, float InStartEndFalloff) :
 		Center(InCenter),
 		Left(InLeft),
 		Right(InRight),
 		FalloffLeft(InFalloffLeft),
 		FalloffRight(InFalloffRight),
+		LayerLeft(InLayerLeft),
+		LayerRight(InLayerRight),
 		LayerFalloffLeft(InLayerFalloffLeft),
 		LayerFalloffRight(InLayerFalloffRight),
 		StartEndFalloff(InStartEndFalloff)
diff --git a/Engine/Source/Runtime/Landscape/Classes/LandscapeSplinesComponent.h b/Engine/Source/Runtime/Landscape/Classes/LandscapeSplinesComponent.h
--- a/Engine/Source/Runtime/Landscape/Classes/LandscapeSplinesComponent.h
+++ b/Engine/Source/Runtime/Landscape/Classes/LandscapeSplinesComponent.h
@@ -188,6 +188,8 @@ public:
 	void AutoFixMeshComponentErrors(UWorld* OtherWorld);
 
 	bool IsUsingEditorMesh(const USplineMeshComponent* SplineMeshComponent) const;
+
+	bool IsUsingLayerInfo(const ULandscapeLayerInfoObject* LayerInfo) const;
 #endif
 
 	//~ Begin UObject Interface
diff --git a/Engine/Source/Runtime/Landscape/Private/Landscape.cpp b/Engine/Source/Runtime/Landscape/Private/Landscape.cpp
--- a/Engine/Source/Runtime/Landscape/Private/Landscape.cpp
+++ b/Engine/Source/Runtime/Landscape/Private/Landscape.cpp
@@ -1097,6 +1097,7 @@ ALandscape::ALandscape(const FObjectInitializer& ObjectInitializer)
 	bLockLocation = false;
 	WasCompilingShaders = false;
 	LayerContentUpdateModes = 0;
+	bSplineLayerUpdateRequested = false;
 	CombinedLayersWeightmapAllMaterialLayersResource = nullptr;
 	CurrentLayersWeightmapAllMaterialLayersResource = nullptr;
 	WeightmapScratchExtractLayerTextureResource = nullptr;
diff --git a/Engine/Source/Runtime/Landscape/Private/LandscapeBlueprintSupport.cpp b/Engine/Source/Runtime/Landscape/Private/LandscapeBlueprintSupport.cpp
--- a/Engine/Source/Runtime/Landscape/Private/LandscapeBlueprintSupport.cpp
+++ b/Engine/Source/Runtime/Landscape/Private/LandscapeBlueprintSupport.cpp
@@ -21,7 +21,7 @@ void ALandscapeProxy::EditorApplySpline(USplineComponent* InSplineComponent, flo
 	{
 		TArray<FLandscapeSplineInterpPoint> Points;
 		LandscapeSplineRaster::FPointifyFalloffs Falloffs(StartSideFalloff, EndSideFalloff);
-		LandscapeSplineRaster::Pointify(InSplineComponent->SplineCurves.Position, Points, NumSubdivisions, 0.0f, 0.0f, StartWidth, EndWidth, Falloffs, StartRoll, EndRoll);
+		LandscapeSplineRaster::Pointify(InSplineComponent->SplineCurves.Position, Points, NumSubdivisions, 0.0f, 0.0f, StartWidth, EndWidth, StartWidth, EndWidth, Falloffs, StartRoll, EndRoll);
 
 		FTransform SplineToWorld = InSplineComponent->GetComponentTransform();
 		LandscapeSplineRaster::RasterizeSegmentPoints(GetLandscapeInfo(), MoveTemp(Points), SplineToWorld, bRaiseHeights, bLowerHeights, PaintLayer);
diff --git a/Engine/Source/Runtime/Landscape/Private/LandscapeEdit.cpp b/Engine/Source/Runtime/Landscape/Private/LandscapeEdit.cpp
--- a/Engine/Source/Runtime/Landscape/Private/LandscapeEdit.cpp
+++ b/Engine/Source/Runtime/Landscape/Private/LandscapeEdit.cpp
@@ -3733,6 +3733,11 @@ ULandscapeLayerInfoObject::ULandscapeLayerInfoObject(const FObjectInitializer& O
 	Hardness = 0.5f;
 #if WITH_EDITORONLY_DATA
 	bNoWeightBlend = false;
+	SplineFalloffModulationTexture = nullptr;
+	SplineFalloffModulationColorMask = ESplineModulationColorMask::Red;
+	SplineFalloffModulationTiling = 1.0f;
+	SplineFalloffModulationBias = 0.5;
+	SplineFalloffModulationScale = 1.0f;
 #endif // WITH_EDITORONLY_DATA
 
 	// Assign initial LayerUsageDebugColor
@@ -3800,6 +3805,20 @@ void ULandscapeLayerInfoObject::PostEditChangeProperty(FPropertyChangedEvent& Pr
 				}
 			}
 		}
+		else if (PropertyName == GET_MEMBER_NAME_CHECKED(ULandscapeLayerInfoObject, SplineFalloffModulationTexture) ||
+				PropertyName == GET_MEMBER_NAME_CHECKED(ULandscapeLayerInfoObject, SplineFalloffModulationColorMask) ||
+				PropertyName == GET_MEMBER_NAME_CHECKED(ULandscapeLayerInfoObject, SplineFalloffModulationBias) ||
+				PropertyName == GET_MEMBER_NAME_CHECKED(ULandscapeLayerInfoObject, SplineFalloffModulationScale) ||
+				PropertyName == GET_MEMBER_NAME_CHECKED(ULandscapeLayerInfoObject, SplineFalloffModulationTiling))
+		{
+			for (TObjectIterator<ULandscapeInfo> It; It; ++It)
+			{
+				if(ALandscape* Landscape = It->LandscapeActor.Get())
+				{
+					Landscape->OnLayerInfoSplineFalloffModulationChanged(this);
+				}
+			}
+		}
 	}
 }
 
diff --git a/Engine/Source/Runtime/Landscape/Private/LandscapeEditLayers.cpp b/Engine/Source/Runtime/Landscape/Private/LandscapeEditLayers.cpp
--- a/Engine/Source/Runtime/Landscape/Private/LandscapeEditLayers.cpp
+++ b/Engine/Source/Runtime/Landscape/Private/LandscapeEditLayers.cpp
@@ -35,6 +35,7 @@ LandscapeEditLayers.cpp: Landscape editing layers mode
 #include "Misc/MessageDialog.h"
 #include "GameFramework/WorldSettings.h"
 #include "UObject/UObjectThreadContext.h"
+#include "LandscapeSplinesComponent.h"
 #endif
 
 #define LOCTEXT_NAMESPACE "Landscape"
@@ -4263,6 +4264,14 @@ void ALandscape::RequestLayersInitialization(bool bInRequestContentUpdate)
 	}
 }
 
+void ALandscape::RequestSplineLayerUpdate()
+{
+	if (HasLayersContent() && GetLandscapeSplinesReservedLayer() != nullptr)
+	{
+		bSplineLayerUpdateRequested = true;
+	}
+}
+
 void ALandscape::RequestLayersContentUpdate(ELandscapeLayerUpdateMode InUpdateMode)
 {
 	LayerContentUpdateModes |= InUpdateMode;
@@ -4569,6 +4578,12 @@ void ALandscape::UpdateLayersContent(bool bInWaitForStreaming, bool bInSkipMonit
 	}
 	MonitorShaderCompilation();
 
+	if (bSplineLayerUpdateRequested)
+	{
+		UpdateLandscapeSplines();
+		bSplineLayerUpdateRequested = false;
+	}
+
 	const bool bForceRender = CVarOutputLayersDebugDrawCallName.GetValueOnAnyThread() == 1;
 
 	if (LayerContentUpdateModes == 0 && !bForceRender)
@@ -6002,6 +6017,33 @@ void ALandscape::OnBlueprintBrushChanged()
 #endif
 }
 
+void ALandscape::OnLayerInfoSplineFalloffModulationChanged(ULandscapeLayerInfoObject* InLayerInfo)
+{
+	ULandscapeInfo* LandscapeInfo = GetLandscapeInfo();
+	
+	if (!LandscapeInfo)
+	{
+		return;
+	}
+	
+	ALandscape* Landscape = LandscapeInfo->LandscapeActor.Get();
+	if (!Landscape || !Landscape->HasLayersContent())
+	{
+		return;
+	}
+
+	bool bUsedForSplines = false;
+	LandscapeInfo->ForAllLandscapeProxies([&](ALandscapeProxy* Proxy)
+	{
+		bUsedForSplines |= (Proxy->SplineComponent && Proxy->SplineComponent->IsUsingLayerInfo(InLayerInfo));
+	});
+
+	if (bUsedForSplines)
+	{
+		Landscape->RequestSplineLayerUpdate();
+	}
+}
+
 ALandscapeBlueprintBrushBase* ALandscape::GetBrushForLayer(int32 InLayerIndex, int8 InBrushIndex) const
 {
 	if (const FLandscapeLayer* Layer = GetLayer(InLayerIndex))
diff --git a/Engine/Source/Runtime/Landscape/Private/LandscapeSplineRaster.cpp b/Engine/Source/Runtime/Landscape/Private/LandscapeSplineRaster.cpp
--- a/Engine/Source/Runtime/Landscape/Private/LandscapeSplineRaster.cpp
+++ b/Engine/Source/Runtime/Landscape/Private/LandscapeSplineRaster.cpp
@@ -29,6 +29,13 @@
 //////////////////////////////////////////////////////////////////////////
 
 #if WITH_EDITOR
+TAutoConsoleVariable<int32> CVarLandscapeSplineFalloffModulation(
+	TEXT("landscape.SplineFalloffModulation"),
+	1,
+	TEXT("Enable Texture Modulation fo Spline Layer Falloff."));
+
+using FModulateAlphaFunc = TFunction<float(float InValue, int32 X, int32 Y)>;
+
 class FLandscapeSplineHeightsRasterPolicy
 {
 public:
@@ -112,6 +119,77 @@ private:
 	uint32 bRaiseTerrain : 1, bLowerTerrain : 1;
 };
 
+extern const size_t ChannelOffsets[4];
+
+class FModulateAlpha
+{
+public:
+	static TSharedPtr<FModulateAlpha> CreateFromLayerInfo(ULandscapeLayerInfoObject* InLayerInfo, int32 InLandscapeMinX, int32 InLandscapeMinY)
+	{
+		if (CVarLandscapeSplineFalloffModulation.GetValueOnAnyThread() == 0 || InLayerInfo == nullptr || InLayerInfo->SplineFalloffModulationTexture == nullptr)
+		{
+			return nullptr;
+		}
+
+		return MakeShareable(new FModulateAlpha(InLayerInfo, InLandscapeMinX, InLandscapeMinY));
+	}
+
+	float Modulate(float InValue, int32 InX, int32 InY) const
+	{
+		float X = FMath::Frac((float)(InX - LandscapeMinX) / (TextureWidth * Tiling)) * TextureWidth;
+		float Y = FMath::Frac((float)(InY - LandscapeMinY) / (TextureHeight * Tiling)) * TextureHeight;
+
+		int32 X0 = FMath::FloorToInt(X);
+		int32 X1 = FMath::CeilToInt(X);
+		int32 Y0 = FMath::FloorToInt(Y);
+		int32 Y1 = FMath::CeilToInt(Y);
+
+		const uint8* Data = MipData.GetData() + ChannelOffset;
+		uint8 SampleX0Y0 = *(Data + ((Y0*TextureWidth + X0) * sizeof(FColor)));
+		uint8 SampleX1Y0 = *(Data + ((Y0*TextureWidth + X1) * sizeof(FColor)));
+		uint8 SampleX0Y1 = *(Data + ((Y1*TextureWidth + X0) * sizeof(FColor)));
+		uint8 SampleX1Y1 = *(Data + ((Y1*TextureWidth + X1) * sizeof(FColor)));
+
+		uint8 LerpY0 = FMath::Lerp(SampleX0Y0, SampleX1Y0, X - X0);
+		uint8 LerpY1 = FMath::Lerp(SampleX0Y1, SampleX1Y1, X - X0);
+		uint8 FinalLerp = FMath::Lerp(LerpY0, LerpY1, Y - Y0);
+
+		InValue = InValue + (((FinalLerp / 255.0f) - Bias) * Scale);
+
+		return FMath::Clamp(InValue, 0.0f, 1.0f);
+	}
+
+private:
+	FModulateAlpha(ULandscapeLayerInfoObject* InLayerInfo, int32 InLandscapeMinX, int32 InLandscapeMinY)
+	{
+		check(InLayerInfo && InLayerInfo->SplineFalloffModulationTexture);
+
+		InLayerInfo->SplineFalloffModulationTexture->Source.GetMipData(MipData, 0);
+		TextureWidth = InLayerInfo->SplineFalloffModulationTexture->Source.GetSizeX();
+		TextureHeight = InLayerInfo->SplineFalloffModulationTexture->Source.GetSizeY();
+
+		Tiling = 1.0f / InLayerInfo->SplineFalloffModulationTiling;
+		Bias = InLayerInfo->SplineFalloffModulationBias;
+		Scale = InLayerInfo->SplineFalloffModulationScale;
+
+		LandscapeMinX = InLandscapeMinX;
+		LandscapeMinY = InLandscapeMinY;
+
+		ChannelOffset = ChannelOffsets[(int32)InLayerInfo->SplineFalloffModulationColorMask];
+	}
+
+private:
+	float Tiling;
+	float Bias;
+	float Scale;
+	TArray<uint8> MipData;
+	int32 TextureWidth;
+	int32 TextureHeight;
+	int32 LandscapeMinX;
+	int32 LandscapeMinY;
+	int32 ChannelOffset;
+};
+
 class FLandscapeSplineBlendmaskRasterPolicy
 {
 public:
@@ -119,15 +197,15 @@ public:
 	typedef FVector InterpolantType;
 
 	/** Initialization constructor. */
-	FLandscapeSplineBlendmaskRasterPolicy(TArray<uint8>& InData, int32 InMinX, int32 InMinY, int32 InMaxX, int32 InMaxY) :
+	FLandscapeSplineBlendmaskRasterPolicy(TArray<uint8>& InData, int32 InMinX, int32 InMinY, int32 InMaxX, int32 InMaxY, const TSharedPtr<FModulateAlpha>& InModulateAlpha = nullptr) :
 		Data(InData),
 		MinX(InMinX),
 		MinY(InMinY),
 		MaxX(InMaxX),
-		MaxY(InMaxY)
+		MaxY(InMaxY),
+		ModulateAlpha(InModulateAlpha)
 	{
 	}
-
 protected:
 
 	// FTriangleRasterizer policy interface.
@@ -139,9 +217,21 @@ protected:
 
 	inline void ProcessPixel(int32 X, int32 Y, const InterpolantType& Interpolant, bool BackFacing)
 	{
-		const float CosInterpX = (Interpolant.X >= 1 ? 1 : 0.5f - 0.5f * FMath::Cos(Interpolant.X * PI));
-		const float CosInterpY = (Interpolant.Y >= 1 ? 1 : 0.5f - 0.5f * FMath::Cos(Interpolant.Y * PI));
-		const float Alpha = CosInterpX * CosInterpY;
+		float Alpha = 0.0f;
+		
+		if (ModulateAlpha == nullptr)
+		{
+			const float CosInterpX = (Interpolant.X >= 1 ? 1 : 0.5f - 0.5f * FMath::Cos(Interpolant.X * PI));
+			const float CosInterpY = (Interpolant.Y >= 1 ? 1 : 0.5f - 0.5f * FMath::Cos(Interpolant.Y * PI));
+			Alpha = CosInterpX * CosInterpY;
+		}
+		else
+		{
+			const float InterpX = FMath::Clamp<float>(Interpolant.X, 0.0f, 1.0f);
+			const float InterpY = FMath::Clamp<float>(Interpolant.Y, 0.0f, 1.0f);
+			Alpha = ModulateAlpha->Modulate(InterpX * InterpY, X, Y);
+		}
+
 		uint8& Dest = Data[(Y - MinY)*(1 + MaxX - MinX) + X - MinX];
 		float Value = FMath::Lerp((float)Dest, Interpolant.Z, Alpha);
 		Dest = (uint32)FMath::Clamp<float>(Value, 0, LandscapeDataAccess::MaxValue);
@@ -150,6 +240,7 @@ protected:
 private:
 	TArray<uint8>& Data;
 	int32 MinX, MinY, MaxX, MaxY;
+	const TSharedPtr<FModulateAlpha>& ModulateAlpha;
 };
 
 void RasterizeHeight(int32& MinX, int32& MinY, int32& MaxX, int32& MaxY, FLandscapeEditDataInterface& LandscapeEdit, bool bRaiseTerrain, bool bLowerTerrain, TSet<ULandscapeComponent*>& ModifiedComponents, TFunctionRef<void(FTriangleRasterizer<FLandscapeSplineHeightsRasterPolicy>&)> RasterizerFunction)
@@ -265,7 +356,7 @@ void RasterizeControlPointHeights(int32& MinX, int32& MinY, int32& MaxX, int32&
 	});
 }
 
-void RasterizeControlPointAlpha(int32& MinX, int32& MinY, int32& MaxX, int32& MaxY, FLandscapeEditDataInterface& LandscapeEdit, FVector ControlPointLocation, const TArray<FLandscapeSplineInterpPoint>& Points, ULandscapeLayerInfoObject* LayerInfo, TSet<ULandscapeComponent*>& ModifiedComponents)
+void RasterizeControlPointAlpha(int32& MinX, int32& MinY, int32& MaxX, int32& MaxY, FLandscapeEditDataInterface& LandscapeEdit, FVector ControlPointLocation, const TArray<FLandscapeSplineInterpPoint>& Points, ULandscapeLayerInfoObject* LayerInfo, TSet<ULandscapeComponent*>& ModifiedComponents, const TSharedPtr<FModulateAlpha>& ModulateAlpha = nullptr)
 {
 	if (LayerInfo == nullptr)
 	{
@@ -305,7 +396,7 @@ void RasterizeControlPointAlpha(int32& MinX, int32& MinY, int32& MaxX, int32& Ma
 	MaxY = ValidMaxY;
 
 	FTriangleRasterizer<FLandscapeSplineBlendmaskRasterPolicy> Rasterizer(
-		FLandscapeSplineBlendmaskRasterPolicy(Data, MinX, MinY, MaxX, MaxY));
+		FLandscapeSplineBlendmaskRasterPolicy(Data, MinX, MinY, MaxX, MaxY, ModulateAlpha));
 
 	const float BlendValue = 255;
 
@@ -315,9 +406,9 @@ void RasterizeControlPointAlpha(int32& MinX, int32& MinY, int32& MaxX, int32& Ma
 	for (int32 i = Points.Num() - 1, j = 0; j < Points.Num(); i = j++)
 	{
 		// Solid center
-		const FVector2D Right0Pos = FVector2D(Points[i].Right);
-		const FVector2D Left1Pos = FVector2D(Points[j].Left);
-		const FVector2D Right1Pos = FVector2D(Points[j].Right);
+		const FVector2D Right0Pos = FVector2D(Points[i].LayerRight);
+		const FVector2D Left1Pos = FVector2D(Points[j].LayerLeft);
+		const FVector2D Right1Pos = FVector2D(Points[j].LayerRight);
 		const FVector Right0 = FVector(1.0f, Points[i].StartEndFalloff, BlendValue);
 		const FVector Left1 = FVector(1.0f, Points[j].StartEndFalloff, BlendValue);
 		const FVector Right1 = FVector(1.0f, Points[j].StartEndFalloff, BlendValue);
@@ -377,7 +468,7 @@ void RasterizeSegmentHeight(int32& MinX, int32& MinY, int32& MaxX, int32& MaxY,
 }
 
 
-void RasterizeSegmentAlpha(int32& MinX, int32& MinY, int32& MaxX, int32& MaxY, FLandscapeEditDataInterface& LandscapeEdit, const TArray<FLandscapeSplineInterpPoint>& Points, ULandscapeLayerInfoObject* LayerInfo, TSet<ULandscapeComponent*>& ModifiedComponents)
+void RasterizeSegmentAlpha(int32& MinX, int32& MinY, int32& MaxX, int32& MaxY, FLandscapeEditDataInterface& LandscapeEdit, const TArray<FLandscapeSplineInterpPoint>& Points, ULandscapeLayerInfoObject* LayerInfo, TSet<ULandscapeComponent*>& ModifiedComponents, const TSharedPtr<FModulateAlpha>& ModulateAlpha = nullptr)
 {
 	if (LayerInfo == nullptr)
 	{
@@ -417,17 +508,17 @@ void RasterizeSegmentAlpha(int32& MinX, int32& MinY, int32& MaxX, int32& MaxY, F
 	MaxY = ValidMaxY;
 
 	FTriangleRasterizer<FLandscapeSplineBlendmaskRasterPolicy> Rasterizer(
-		FLandscapeSplineBlendmaskRasterPolicy(Data, MinX, MinY, MaxX, MaxY));
+		FLandscapeSplineBlendmaskRasterPolicy(Data, MinX, MinY, MaxX, MaxY, ModulateAlpha));
 
 	const float BlendValue = 255;
 
 	for (int32 j = 1; j < Points.Num(); j++)
 	{
 		// Middle
-		FVector2D Left0Pos = FVector2D(Points[j - 1].Left);
-		FVector2D Right0Pos = FVector2D(Points[j - 1].Right);
-		FVector2D Left1Pos = FVector2D(Points[j].Left);
-		FVector2D Right1Pos = FVector2D(Points[j].Right);
+		FVector2D Left0Pos = FVector2D(Points[j - 1].LayerLeft);
+		FVector2D Right0Pos = FVector2D(Points[j - 1].LayerRight);
+		FVector2D Left1Pos = FVector2D(Points[j].LayerLeft);
+		FVector2D Right1Pos = FVector2D(Points[j].LayerRight);
 		FVector Left0 = FVector(1.0f, Points[j - 1].StartEndFalloff, BlendValue);
 		FVector Right0 = FVector(1.0f, Points[j - 1].StartEndFalloff, BlendValue);
 		FVector Left1 = FVector(1.0f, Points[j].StartEndFalloff, BlendValue);
@@ -496,6 +587,21 @@ bool ULandscapeInfo::ApplySplinesInternal(bool bOnlySelected, ALandscapeProxy* P
 		return false;
 	}
 
+	TMap<ULandscapeLayerInfoObject*, TSharedPtr<FModulateAlpha>> ModulatePerLayerInfo;
+	
+	auto GetOrCreateModulate = [&](ULandscapeLayerInfoObject* LayerInfo) -> TSharedPtr<FModulateAlpha>
+	{
+		if (const TSharedPtr<FModulateAlpha>* SharedPtr = ModulatePerLayerInfo.Find(LayerInfo))
+		{
+			return *SharedPtr;
+		}
+		
+		TSharedPtr<FModulateAlpha> SharedPtr = FModulateAlpha::CreateFromLayerInfo(LayerInfo, LandscapeMinX, LandscapeMinY);
+		ModulatePerLayerInfo.Add(LayerInfo, SharedPtr);
+
+		return SharedPtr;
+	};
+
 	for (const ULandscapeSplineControlPoint* ControlPoint : Proxy->SplineComponent->ControlPoints)
 	{
 		if (bOnlySelected && !ControlPoint->IsSplineSelected())
@@ -535,6 +641,9 @@ bool ULandscapeInfo::ApplySplinesInternal(bool bOnlySelected, ALandscapeProxy* P
 			Points[j].Right = SplineToLandscape.TransformPosition(Points[j].Right);
 			Points[j].FalloffLeft = SplineToLandscape.TransformPosition(Points[j].FalloffLeft);
 			Points[j].FalloffRight = SplineToLandscape.TransformPosition(Points[j].FalloffRight);
+
+			Points[j].LayerLeft = SplineToLandscape.TransformPosition(Points[j].LayerLeft);
+			Points[j].LayerRight = SplineToLandscape.TransformPosition(Points[j].LayerRight);
 			Points[j].LayerFalloffLeft = SplineToLandscape.TransformPosition(Points[j].LayerFalloffLeft);
 			Points[j].LayerFalloffRight = SplineToLandscape.TransformPosition(Points[j].LayerFalloffRight);
 
@@ -543,6 +652,11 @@ bool ULandscapeInfo::ApplySplinesInternal(bool bOnlySelected, ALandscapeProxy* P
 			Points[j].Right.Z = Points[j].Right.Z * LANDSCAPE_INV_ZSCALE + LandscapeDataAccess::MidValue;
 			Points[j].FalloffLeft.Z = Points[j].FalloffLeft.Z * LANDSCAPE_INV_ZSCALE + LandscapeDataAccess::MidValue;
 			Points[j].FalloffRight.Z = Points[j].FalloffRight.Z * LANDSCAPE_INV_ZSCALE + LandscapeDataAccess::MidValue;
+
+			Points[j].LayerLeft.Z = Points[j].LayerLeft.Z * LANDSCAPE_INV_ZSCALE + LandscapeDataAccess::MidValue;
+			Points[j].LayerRight.Z = Points[j].LayerRight.Z * LANDSCAPE_INV_ZSCALE + LandscapeDataAccess::MidValue;
+			Points[j].LayerFalloffLeft.Z = Points[j].LayerFalloffLeft.Z * LANDSCAPE_INV_ZSCALE + LandscapeDataAccess::MidValue;
+			Points[j].LayerFalloffRight.Z = Points[j].LayerFalloffRight.Z * LANDSCAPE_INV_ZSCALE + LandscapeDataAccess::MidValue;
 		}
 
 		// Heights raster
@@ -559,7 +673,8 @@ bool ULandscapeInfo::ApplySplinesInternal(bool bOnlySelected, ALandscapeProxy* P
 		{
 			const FVector Center3D = SplineToLandscape.TransformPosition(ControlPoint->Location);
 
-			RasterizeControlPointAlpha(MinX, MinY, MaxX, MaxY, LandscapeEdit, Center3D, Points, LayerInfo, ModifiedComponents);
+			TSharedPtr<FModulateAlpha> ModulateAlpha = GetOrCreateModulate(LayerInfo);
+			RasterizeControlPointAlpha(MinX, MinY, MaxX, MaxY, LandscapeEdit, Center3D, Points, LayerInfo, ModifiedComponents, ModulateAlpha);
 		}
 	}
 
@@ -597,6 +712,9 @@ bool ULandscapeInfo::ApplySplinesInternal(bool bOnlySelected, ALandscapeProxy* P
 			Points[j].Right = SplineToLandscape.TransformPosition(Points[j].Right);
 			Points[j].FalloffLeft = SplineToLandscape.TransformPosition(Points[j].FalloffLeft);
 			Points[j].FalloffRight = SplineToLandscape.TransformPosition(Points[j].FalloffRight);
+
+			Points[j].LayerLeft = SplineToLandscape.TransformPosition(Points[j].LayerLeft);
+			Points[j].LayerRight = SplineToLandscape.TransformPosition(Points[j].LayerRight);
 			Points[j].LayerFalloffLeft = SplineToLandscape.TransformPosition(Points[j].LayerFalloffLeft);
 			Points[j].LayerFalloffRight = SplineToLandscape.TransformPosition(Points[j].LayerFalloffRight);
 
@@ -605,6 +723,9 @@ bool ULandscapeInfo::ApplySplinesInternal(bool bOnlySelected, ALandscapeProxy* P
 			Points[j].Right.Z = Points[j].Right.Z * LANDSCAPE_INV_ZSCALE + LandscapeDataAccess::MidValue;
 			Points[j].FalloffLeft.Z = Points[j].FalloffLeft.Z * LANDSCAPE_INV_ZSCALE + LandscapeDataAccess::MidValue;
 			Points[j].FalloffRight.Z = Points[j].FalloffRight.Z * LANDSCAPE_INV_ZSCALE + LandscapeDataAccess::MidValue;
+
+			Points[j].LayerLeft.Z = Points[j].LayerLeft.Z * LANDSCAPE_INV_ZSCALE + LandscapeDataAccess::MidValue;
+			Points[j].LayerRight.Z = Points[j].LayerRight.Z * LANDSCAPE_INV_ZSCALE + LandscapeDataAccess::MidValue;
 			Points[j].LayerFalloffLeft.Z = Points[j].LayerFalloffLeft.Z * LANDSCAPE_INV_ZSCALE + LandscapeDataAccess::MidValue;
 			Points[j].LayerFalloffRight.Z = Points[j].LayerFalloffRight.Z * LANDSCAPE_INV_ZSCALE + LandscapeDataAccess::MidValue;
 		}
@@ -625,7 +746,8 @@ bool ULandscapeInfo::ApplySplinesInternal(bool bOnlySelected, ALandscapeProxy* P
 		ULandscapeLayerInfoObject* LayerInfo = GetLayerInfoByName(Segment->LayerName);
 		if (Segment->LayerName != NAME_None && LayerInfo != NULL)
 		{
-			RasterizeSegmentAlpha(MinX, MinY, MaxX, MaxY, LandscapeEdit, Points, LayerInfo, ModifiedComponents);
+			TSharedPtr<FModulateAlpha> ModulateAlpha = GetOrCreateModulate(LayerInfo);
+			RasterizeSegmentAlpha(MinX, MinY, MaxX, MaxY, LandscapeEdit, Points, LayerInfo, ModifiedComponents, ModulateAlpha);
 		}
 	}
 
@@ -713,11 +835,21 @@ namespace LandscapeSplineRaster
 			Points[j].FalloffLeft = SplineToLandscape.TransformPosition(Points[j].FalloffLeft);
 			Points[j].FalloffRight = SplineToLandscape.TransformPosition(Points[j].FalloffRight);
 
+			Points[j].LayerLeft = SplineToLandscape.TransformPosition(Points[j].LayerLeft);
+			Points[j].LayerRight = SplineToLandscape.TransformPosition(Points[j].LayerRight);
+			Points[j].LayerFalloffLeft = SplineToLandscape.TransformPosition(Points[j].LayerFalloffLeft);
+			Points[j].LayerFalloffRight = SplineToLandscape.TransformPosition(Points[j].LayerFalloffRight);
+
 			// local-heights to texture value heights
 			Points[j].Left.Z = Points[j].Left.Z * LANDSCAPE_INV_ZSCALE + LandscapeDataAccess::MidValue;
 			Points[j].Right.Z = Points[j].Right.Z * LANDSCAPE_INV_ZSCALE + LandscapeDataAccess::MidValue;
 			Points[j].FalloffLeft.Z = Points[j].FalloffLeft.Z * LANDSCAPE_INV_ZSCALE + LandscapeDataAccess::MidValue;
 			Points[j].FalloffRight.Z = Points[j].FalloffRight.Z * LANDSCAPE_INV_ZSCALE + LandscapeDataAccess::MidValue;
+
+			Points[j].LayerLeft.Z = Points[j].LayerLeft.Z * LANDSCAPE_INV_ZSCALE + LandscapeDataAccess::MidValue;
+			Points[j].LayerRight.Z = Points[j].LayerRight.Z * LANDSCAPE_INV_ZSCALE + LandscapeDataAccess::MidValue;
+			Points[j].LayerFalloffLeft.Z = Points[j].LayerFalloffLeft.Z * LANDSCAPE_INV_ZSCALE + LandscapeDataAccess::MidValue;
+			Points[j].LayerFalloffRight.Z = Points[j].LayerFalloffRight.Z * LANDSCAPE_INV_ZSCALE + LandscapeDataAccess::MidValue;
 		}
 
 		// Heights raster
@@ -879,6 +1011,7 @@ namespace LandscapeSplineRaster
 	void Pointify(const FInterpCurveVector& SplineInfo, TArray<FLandscapeSplineInterpPoint>& Points, int32 NumSubdivisions,
 		float StartFalloffFraction, float EndFalloffFraction,
 		const float StartWidth, const float EndWidth,
+		const float StartLayerWidth, const float EndLayerWidth,
 		const FPointifyFalloffs& Falloffs,
 		const float StartRollDegrees, const float EndRollDegrees)
 	{
@@ -899,6 +1032,7 @@ namespace LandscapeSplineRaster
 			const float NewKeyTime = SplineInfo.Points[i].InVal;
 			const float NewKeyCosInterp = 0.5f - 0.5f * FMath::Cos(NewKeyTime * PI);
 			const float NewKeyWidth = FMath::Lerp(StartWidth, EndWidth, NewKeyCosInterp);
+			const float NewKeyLayerWidth = FMath::Lerp(StartLayerWidth, EndLayerWidth, NewKeyCosInterp);
 			const float NewKeyLeftFalloff = FMath::Lerp(Falloffs.StartLeftSide, Falloffs.EndLeftSide, NewKeyCosInterp);
 			const float NewKeyRightFalloff = FMath::Lerp(Falloffs.StartRightSide, Falloffs.EndRightSide, NewKeyCosInterp);
 			const float NewKeyLeftLayerFalloff = FMath::Lerp(Falloffs.StartLeftSideLayer, Falloffs.EndLeftSideLayer, NewKeyCosInterp);
@@ -911,8 +1045,11 @@ namespace LandscapeSplineRaster
 			const FVector NewKeyRightPos = NewKeyPos + NewKeyBiNormal * NewKeyWidth;
 			const FVector NewKeyFalloffLeftPos = NewKeyPos - NewKeyBiNormal * (NewKeyWidth + NewKeyLeftFalloff);
 			const FVector NewKeyFalloffRightPos = NewKeyPos + NewKeyBiNormal * (NewKeyWidth + NewKeyRightFalloff);
-			const FVector NewKeyLayerFalloffLeftPos = NewKeyPos - NewKeyBiNormal * (NewKeyWidth + NewKeyLeftLayerFalloff);
-			const FVector NewKeyLayerFalloffRightPos = NewKeyPos + NewKeyBiNormal * (NewKeyWidth + NewKeyRightLayerFalloff);
+			
+			const FVector NewKeyLayerLeftPos = NewKeyPos - NewKeyBiNormal * NewKeyLayerWidth;
+			const FVector NewKeyLayerRightPos = NewKeyPos + NewKeyBiNormal * NewKeyLayerWidth;
+			const FVector NewKeyLayerFalloffLeftPos = NewKeyPos - NewKeyBiNormal * (NewKeyLayerWidth + NewKeyLeftLayerFalloff);
+			const FVector NewKeyLayerFalloffRightPos = NewKeyPos + NewKeyBiNormal * (NewKeyLayerWidth + NewKeyRightLayerFalloff);
 			const float NewKeyStartEndFalloff = FMath::Min((StartFalloffFraction > 0 ? NewKeyTime / StartFalloffFraction : 1.0f), (EndFalloffFraction > 0 ? (1 - NewKeyTime) / EndFalloffFraction : 1.0f));
 
 			// If not the first keypoint, interp from the last keypoint.
@@ -927,6 +1064,7 @@ namespace LandscapeSplineRaster
 					const float NewTime = OldKeyTime + j*DrawSubstep;
 					const float NewCosInterp = 0.5f - 0.5f * FMath::Cos(NewTime * PI);
 					const float NewWidth = FMath::Lerp(StartWidth, EndWidth, NewCosInterp);
+					const float NewLayerWidth = FMath::Lerp(StartLayerWidth, EndLayerWidth, NewCosInterp);
 					const float NewLeftFalloff = FMath::Lerp(Falloffs.StartLeftSide, Falloffs.EndLeftSide, NewCosInterp);
 					const float NewRightFalloff = FMath::Lerp(Falloffs.StartRightSide, Falloffs.EndRightSide, NewCosInterp);
 					const float NewLeftLayerFalloff = FMath::Lerp(Falloffs.StartLeftSideLayer, Falloffs.EndLeftSideLayer, NewCosInterp);
@@ -939,15 +1077,18 @@ namespace LandscapeSplineRaster
 					const FVector NewRightPos = NewPos + NewBiNormal * NewWidth;
 					const FVector NewFalloffLeftPos = NewPos - NewBiNormal * (NewWidth + NewLeftFalloff);
 					const FVector NewFalloffRightPos = NewPos + NewBiNormal * (NewWidth + NewRightFalloff);
-					const FVector NewLayerFalloffLeftPos = NewPos - NewBiNormal * (NewWidth + NewLeftLayerFalloff);
-					const FVector NewLayerFalloffRightPos = NewPos + NewBiNormal * (NewWidth + NewRightLayerFalloff);
+
+					const FVector NewLayerLeftPos = NewPos - NewBiNormal * NewLayerWidth;
+					const FVector NewLayerRightPos = NewPos + NewBiNormal * NewLayerWidth;
+					const FVector NewLayerFalloffLeftPos = NewPos - NewBiNormal * (NewLayerWidth + NewLeftLayerFalloff);
+					const FVector NewLayerFalloffRightPos = NewPos + NewBiNormal * (NewLayerWidth + NewRightLayerFalloff);
 					const float NewStartEndFalloff = FMath::Min((StartFalloffFraction > 0 ? NewTime / StartFalloffFraction : 1.0f), (EndFalloffFraction > 0 ? (1 - NewTime) / EndFalloffFraction : 1.0f));
 
-					Points.Emplace(NewPos, NewLeftPos, NewRightPos, NewFalloffLeftPos, NewFalloffRightPos, NewLayerFalloffLeftPos, NewLayerFalloffRightPos, NewStartEndFalloff);
+					Points.Emplace(NewPos, NewLeftPos, NewRightPos, NewFalloffLeftPos, NewFalloffRightPos, NewLayerLeftPos, NewLayerRightPos, NewLayerFalloffLeftPos, NewLayerFalloffRightPos, NewStartEndFalloff);
 				}
 			}
 
-			Points.Emplace(NewKeyPos, NewKeyLeftPos, NewKeyRightPos, NewKeyFalloffLeftPos, NewKeyFalloffRightPos, NewKeyLayerFalloffLeftPos, NewKeyLayerFalloffRightPos, NewKeyStartEndFalloff);
+			Points.Emplace(NewKeyPos, NewKeyLeftPos, NewKeyRightPos, NewKeyFalloffLeftPos, NewKeyFalloffRightPos, NewKeyLayerLeftPos, NewKeyLayerRightPos, NewKeyLayerFalloffLeftPos, NewKeyLayerFalloffRightPos, NewKeyStartEndFalloff);
 
 			OldKeyTime = NewKeyTime;
 		}
@@ -957,6 +1098,8 @@ namespace LandscapeSplineRaster
 		FixSelfIntersection(Points, &FLandscapeSplineInterpPoint::Right);
 		FixSelfIntersection(Points, &FLandscapeSplineInterpPoint::FalloffLeft);
 		FixSelfIntersection(Points, &FLandscapeSplineInterpPoint::FalloffRight);
+		FixSelfIntersection(Points, &FLandscapeSplineInterpPoint::LayerLeft);
+		FixSelfIntersection(Points, &FLandscapeSplineInterpPoint::LayerRight);
 		FixSelfIntersection(Points, &FLandscapeSplineInterpPoint::LayerFalloffLeft);
 		FixSelfIntersection(Points, &FLandscapeSplineInterpPoint::LayerFalloffRight);
 
diff --git a/Engine/Source/Runtime/Landscape/Private/LandscapeSplineRaster.h b/Engine/Source/Runtime/Landscape/Private/LandscapeSplineRaster.h
--- a/Engine/Source/Runtime/Landscape/Private/LandscapeSplineRaster.h
+++ b/Engine/Source/Runtime/Landscape/Private/LandscapeSplineRaster.h
@@ -46,6 +46,7 @@ namespace LandscapeSplineRaster
 	void Pointify(const FInterpCurveVector& SplineInfo, TArray<FLandscapeSplineInterpPoint>& OutPoints, int32 NumSubdivisions,
 		float StartFalloffFraction, float EndFalloffFraction,
 		const float StartWidth, const float EndWidth,
+		const float StartLayerWidth, const float EndLayerWidth,
 		const FPointifyFalloffs& Falloffs,
 		const float StartRollDegrees, const float EndRollDegrees);
 
diff --git a/Engine/Source/Runtime/Landscape/Private/LandscapeSplines.cpp b/Engine/Source/Runtime/Landscape/Private/LandscapeSplines.cpp
--- a/Engine/Source/Runtime/Landscape/Private/LandscapeSplines.cpp
+++ b/Engine/Source/Runtime/Landscape/Private/LandscapeSplines.cpp
@@ -38,6 +38,7 @@
 #include "Logging/MessageLog.h"
 #include "Misc/UObjectToken.h"
 #include "Landscape.h"
+#include "LandscapeLayerInfoObject.h"
 #endif
 
 IMPLEMENT_HIT_PROXY(HLandscapeSplineProxy, HHitProxy);
@@ -645,6 +646,27 @@ bool ULandscapeSplinesComponent::IsUsingEditorMesh(const USplineMeshComponent* S
 	return SplineMeshComponent->GetStaticMesh() == SplineEditorMesh && SplineMeshComponent->bHiddenInGame;
 }
 
+bool ULandscapeSplinesComponent::IsUsingLayerInfo(const ULandscapeLayerInfoObject* LayerInfo) const
+{
+	for (ULandscapeSplineControlPoint* ControlPoint : ControlPoints)
+	{
+		if (ControlPoint->LayerName == LayerInfo->LayerName)
+		{
+			return true;
+		}
+	}
+
+	for (ULandscapeSplineSegment* Segment : Segments)
+	{
+		if (Segment->LayerName == LayerInfo->LayerName)
+		{
+			return true;
+		}
+	}
+
+	return false;
+}
+
 void ULandscapeSplinesComponent::CheckForErrors()
 {
 	Super::CheckForErrors();
@@ -1305,6 +1327,15 @@ void ULandscapeSplineControlPoint::Serialize(FArchive& Ar)
 			Point.LayerFalloffRight = Point.FalloffRight;
 		}
 	}
+
+	if (Ar.IsLoading() && Ar.CustomVer(FLandscapeCustomVersion::GUID) << FLandscapeCustomVersion::AddSplineLayerWidth)
+	{
+		for (FLandscapeSplineInterpPoint& Point : Points)
+		{
+			Point.LayerLeft = Point.Left;
+			Point.LayerRight = Point.Right;
+		}
+	}
 #endif
 }
 
@@ -1771,6 +1802,7 @@ void ULandscapeSplineControlPoint::UpdateSplinePoints(bool bUpdateCollision, boo
 	const float RightSideFalloff = RightSideFalloffFactor * SideFalloff;
 	const float LeftSideLayerFalloff = LeftSideLayerFalloffFactor * SideFalloff;
 	const float RightSideLayerFalloff = RightSideLayerFalloffFactor * SideFalloff;
+	const float LayerWidth = Width * LayerWidthRatio;
 
 	// Update "Points" array
 	if (Mesh != nullptr)
@@ -1789,10 +1821,13 @@ void ULandscapeSplineControlPoint::UpdateSplinePoints(bool bUpdateCollision, boo
 			const FVector RightPos = StartLocation + BiNormal * Width;
 			const FVector FalloffLeftPos = StartLocation - BiNormal * (Width + LeftSideFalloff);
 			const FVector FalloffRightPos = StartLocation + BiNormal * (Width + RightSideFalloff);
-			const FVector LayerFalloffLeftPos = StartLocation - BiNormal * (Width + LeftSideLayerFalloff);
-			const FVector LayerFalloffRightPos = StartLocation + BiNormal * (Width + RightSideLayerFalloff);
+			
+			const FVector LayerLeftPos = StartLocation - BiNormal * LayerWidth;
+			const FVector LayerRightPos = StartLocation + BiNormal * LayerWidth;
+			const FVector LayerFalloffLeftPos = StartLocation - BiNormal * (LayerWidth + LeftSideLayerFalloff);
+			const FVector LayerFalloffRightPos = StartLocation + BiNormal * (LayerWidth + RightSideLayerFalloff);
 
-			Points.Emplace(StartLocation, LeftPos, RightPos, FalloffLeftPos, FalloffRightPos, LayerFalloffLeftPos, LayerFalloffRightPos, 1.0f);
+			Points.Emplace(StartLocation, LeftPos, RightPos, FalloffLeftPos, FalloffRightPos, LayerLeftPos, LayerRightPos, LayerFalloffLeftPos, LayerFalloffRightPos, 1.0f);
 		}
 
 		const FVector CPLocation = Location;
@@ -1812,11 +1847,13 @@ void ULandscapeSplineControlPoint::UpdateSplinePoints(bool bUpdateCollision, boo
 		const FVector RightPos = StartLocation + BiNormal * Width;
 		const FVector FalloffLeftPos = StartLocation - BiNormal * (Width + LeftSideFalloff);
 		const FVector FalloffRightPos = StartLocation + BiNormal * (Width + RightSideFalloff);
-		const FVector LayerFalloffLeftPos = StartLocation - BiNormal * (Width + LeftSideLayerFalloff);
-		const FVector LayerFalloffRightPos = StartLocation + BiNormal * (Width + RightSideLayerFalloff);
 
+		const FVector LayerLeftPos = StartLocation - BiNormal * LayerWidth;
+		const FVector LayerRightPos = StartLocation + BiNormal * LayerWidth;
+		const FVector LayerFalloffLeftPos = StartLocation - BiNormal * (LayerWidth + LeftSideLayerFalloff);
+		const FVector LayerFalloffRightPos = StartLocation + BiNormal * (LayerWidth + RightSideLayerFalloff);
 
-		Points.Emplace(StartLocation, LeftPos, RightPos, FalloffLeftPos, FalloffRightPos, LayerFalloffLeftPos, LayerFalloffRightPos, 1.0f);
+		Points.Emplace(StartLocation, LeftPos, RightPos, FalloffLeftPos, FalloffRightPos, LayerLeftPos, LayerRightPos, LayerFalloffLeftPos, LayerFalloffRightPos, 1.0f);
 	}
 
 	// Update bounds
@@ -1919,6 +1956,7 @@ void ULandscapeSplineControlPoint::PostEditChangeProperty(FPropertyChangedEvent&
 	Super::PostEditChangeProperty(PropertyChangedEvent);
 
 	Width = FMath::Max(Width, 0.001f);
+	LayerWidthRatio = FMath::Max(LayerWidthRatio, 0.01f);
 	SideFalloff = FMath::Max(SideFalloff, 0.0f);
 	LeftSideFalloffFactor = FMath::Clamp(LeftSideFalloffFactor, 0.0f, 1.0f);
 	RightSideFalloffFactor = FMath::Clamp(RightSideFalloffFactor, 0.0f, 1.0f);
@@ -1939,9 +1977,9 @@ void ULandscapeSplineControlPoint::PostEditChangeProperty(FPropertyChangedEvent&
 	{
 		ALandscapeProxy* OuterLandscape = Cast<ALandscapeProxy>(GetOuterULandscapeSplinesComponent()->GetOwner());
 		ALandscape* Landscape = OuterLandscape ? OuterLandscape->GetLandscapeActor() : nullptr;
-		if (Landscape && Landscape->HasLayersContent() && Landscape->GetLandscapeSplinesReservedLayer())
+		if (Landscape)
 		{
-			Landscape->UpdateLandscapeSplines();
+			Landscape->RequestSplineLayerUpdate();
 		}
 	}
 }
@@ -2039,6 +2077,15 @@ void ULandscapeSplineSegment::Serialize(FArchive& Ar)
 			Point.LayerFalloffRight = Point.FalloffRight;
 		}
 	}
+
+	if (Ar.IsLoading() && Ar.CustomVer(FLandscapeCustomVersion::GUID) << FLandscapeCustomVersion::AddSplineLayerWidth)
+	{
+		for (FLandscapeSplineInterpPoint& Point : Points)
+		{
+			Point.LayerLeft = Point.Left;
+			Point.LayerRight = Point.Right;
+		}
+	}
 #endif
 }
 
@@ -2294,6 +2341,8 @@ void ULandscapeSplineSegment::UpdateSplinePoints(bool bUpdateCollision, bool bUp
 	const float EndFalloffFraction = ((Connections[1].ControlPoint->ConnectedSegments.Num() > 1) ? 0 : (Connections[1].ControlPoint->EndFalloff / SplineLength));
 	const float StartWidth = Connections[0].ControlPoint->Width;
 	const float EndWidth = Connections[1].ControlPoint->Width;
+	const float StartLayerWidth = StartWidth * Connections[0].ControlPoint->LayerWidthRatio;
+	const float EndLayerWidth = EndWidth * Connections[1].ControlPoint->LayerWidthRatio;
 
 	LandscapeSplineRaster::FPointifyFalloffs Falloffs;
 	Falloffs.StartLeftSide = Connections[0].ControlPoint->LeftSideFalloffFactor * Connections[0].ControlPoint->SideFalloff;
@@ -2314,7 +2363,7 @@ void ULandscapeSplineSegment::UpdateSplinePoints(bool bUpdateCollision, bool bUp
 	int32 NumPoints = FMath::CeilToInt(SplineLength / OuterSplines->SplineResolution);
 	NumPoints = FMath::Clamp(NumPoints, 1, 1000);
 
-	LandscapeSplineRaster::Pointify(SplineInfo, Points, NumPoints, StartFalloffFraction, EndFalloffFraction, StartWidth, EndWidth, Falloffs, StartRollDegrees, EndRollDegrees);
+	LandscapeSplineRaster::Pointify(SplineInfo, Points, NumPoints, StartFalloffFraction, EndFalloffFraction, StartWidth, EndWidth, StartLayerWidth, EndLayerWidth, Falloffs, StartRollDegrees, EndRollDegrees);
 
 	// Update Bounds
 	Bounds = FBox(ForceInit);
@@ -2894,9 +2943,9 @@ void ULandscapeSplineSegment::PostEditChangeProperty(FPropertyChangedEvent& Prop
 	{
 		ALandscapeProxy* OuterLandscape = Cast<ALandscapeProxy>(GetOuterULandscapeSplinesComponent()->GetOwner());
 		ALandscape* Landscape = OuterLandscape ? OuterLandscape->GetLandscapeActor() : nullptr;
-		if (Landscape && Landscape->HasLayersContent() && Landscape->GetLandscapeSplinesReservedLayer())
+		if (Landscape)
 		{
-			Landscape->UpdateLandscapeSplines();
+			Landscape->RequestSplineLayerUpdate();
 		}
 	}
 }
diff --git a/Engine/Source/Runtime/Landscape/Private/LandscapeVersion.h b/Engine/Source/Runtime/Landscape/Private/LandscapeVersion.h
--- a/Engine/Source/Runtime/Landscape/Private/LandscapeVersion.h
+++ b/Engine/Source/Runtime/Landscape/Private/LandscapeVersion.h
@@ -30,6 +30,8 @@ namespace FLandscapeCustomVersion
 		AddingBodyInstanceToSplinesElements,
 		// Spline Layer Falloff
 		AddSplineLayerFalloff,
+		// Spline Layer Width
+		AddSplineLayerWidth,
 		// -----<new versions can be added above this line>-------------------------------------------------
 		VersionPlusOne,
 		LatestVersion = VersionPlusOne - 1
