Support rendering to RVT of landscape spline control points

diff --git a/Engine/Source/Runtime/Landscape/Classes/LandscapeSplineControlPoint.h b/Engine/Source/Runtime/Landscape/Classes/LandscapeSplineControlPoint.h
--- a/Engine/Source/Runtime/Landscape/Classes/LandscapeSplineControlPoint.h
+++ b/Engine/Source/Runtime/Landscape/Classes/LandscapeSplineControlPoint.h
@@ -154,6 +154,29 @@ class ULandscapeSplineControlPoint : public UObject
 	UPROPERTY(EditAnywhere, Category=Mesh, AdvancedDisplay)
 	uint32 bPlaceSplineMeshesInStreamingLevels : 1;
 
+	/** 
+	 * Array of runtime virtual textures into which we render the spline segment. 
+	 * The material also needs to be set up to output to a virtual texture. 
+	 */
+	UPROPERTY(EditAnywhere, BlueprintReadOnly, Category = VirtualTexture, meta = (DisplayName = "Render to Virtual Textures"))
+	TArray<URuntimeVirtualTexture*> RuntimeVirtualTextures;
+
+	/** Lod bias for rendering to runtime virtual texture. */
+	UPROPERTY(EditAnywhere, AdvancedDisplay, BlueprintReadOnly, Category = VirtualTexture, meta = (DisplayName = "Virtual Texture LOD Bias", UIMin = "0", UIMax = "7"))
+	int32 VirtualTextureLodBias = 0;
+
+	/**
+	 * Number of lower mips in the runtime virtual texture to skip for rendering this primitive.
+	 * Larger values reduce the effective draw distance in the runtime virtual texture.
+	 * This culling method doesn't take into account primitive size or virtual texture size.
+	 */
+	UPROPERTY(EditAnywhere, AdvancedDisplay, BlueprintReadOnly, Category = VirtualTexture, meta = (DisplayName = "Virtual Texture Skip Mips", UIMin = "0", UIMax = "15"))
+	int32 VirtualTextureCullMips = 0;
+
+	/** Render to the main pass based on the virtual texture settings. */
+	UPROPERTY(EditAnywhere, BlueprintReadOnly, Category = VirtualTexture, meta = (DisplayName = "Virtual Texture Pass Type"))
+	ERuntimeVirtualTextureMainPassType VirtualTextureRenderPassType = ERuntimeVirtualTextureMainPassType::Exclusive;
+
 	/** Mesh Collision Settings */
 	UPROPERTY(EditAnywhere, BlueprintReadOnly, Category = Collision, meta = (ShowOnlyInnerProperties))
 	FBodyInstance BodyInstance;
diff --git a/Engine/Source/Runtime/Landscape/Private/LandscapeSplines.cpp b/Engine/Source/Runtime/Landscape/Private/LandscapeSplines.cpp
--- a/Engine/Source/Runtime/Landscape/Private/LandscapeSplines.cpp
+++ b/Engine/Source/Runtime/Landscape/Private/LandscapeSplines.cpp
@@ -1761,6 +1761,34 @@ void ULandscapeSplineControlPoint::UpdateSplinePoints(bool bUpdateCollision, boo
 			MeshComponent->SetCastShadow(bCastShadow);
 		}
 
+		if (MeshComponent->RuntimeVirtualTextures != RuntimeVirtualTextures)
+		{
+			MeshComponent->Modify();
+			MeshComponent->RuntimeVirtualTextures = RuntimeVirtualTextures;
+			MeshComponent->MarkRenderStateDirty();
+		}
+
+		if (MeshComponent->VirtualTextureLodBias != VirtualTextureLodBias)
+		{
+			MeshComponent->Modify();
+			MeshComponent->VirtualTextureLodBias = VirtualTextureLodBias;
+			MeshComponent->MarkRenderStateDirty();
+		}
+		
+		if (MeshComponent->VirtualTextureCullMips != VirtualTextureCullMips)
+		{
+			MeshComponent->Modify();
+			MeshComponent->VirtualTextureCullMips = VirtualTextureCullMips;
+			MeshComponent->MarkRenderStateDirty();
+		}
+
+		if (MeshComponent->VirtualTextureRenderPassType != VirtualTextureRenderPassType)
+		{
+			MeshComponent->Modify();
+			MeshComponent->VirtualTextureRenderPassType = VirtualTextureRenderPassType;
+			MeshComponent->MarkRenderStateDirty();
+		}
+
 		if (bComponentNeedsRegistering)
 		{
 			MeshComponent->RegisterComponent();
