Fix non-unity build

#rnx

diff --git a/Engine/Source/Runtime/Landscape/Classes/LandscapeLayerInfoObject.h b/Engine/Source/Runtime/Landscape/Classes/LandscapeLayerInfoObject.h
--- a/Engine/Source/Runtime/Landscape/Classes/LandscapeLayerInfoObject.h
+++ b/Engine/Source/Runtime/Landscape/Classes/LandscapeLayerInfoObject.h
@@ -8,6 +8,7 @@
 #include "LandscapeLayerInfoObject.generated.h"
 
 class UPhysicalMaterial;
+class UTexture2D;
 struct FPropertyChangedEvent;
 
 UENUM()
