MN - Fix compiler error in LandScape due to Robomerge

diff --git a/Engine/Source/Runtime/Landscape/Private/LandscapeSplineRaster.cpp b/Engine/Source/Runtime/Landscape/Private/LandscapeSplineRaster.cpp
--- a/Engine/Source/Runtime/Landscape/Private/LandscapeSplineRaster.cpp
+++ b/Engine/Source/Runtime/Landscape/Private/LandscapeSplineRaster.cpp
@@ -182,7 +182,7 @@ private:
 	float Tiling;
 	float Bias;
 	float Scale;
-	TArray<uint8> MipData;
+	TArray64<uint8> MipData;
 	int32 TextureWidth;
 	int32 TextureHeight;
 	int32 LandscapeMinX;
