Prevent copy paste of landscape streaming proxies
We can't support this since we end up with overlapping components in the same landscape (breaking neighbor searches and more)
Issue was raised as part of investigation to jira UE-70460

diff --git a/Engine/Source/Runtime/Landscape/Classes/LandscapeStreamingProxy.h b/Engine/Source/Runtime/Landscape/Classes/LandscapeStreamingProxy.h
--- a/Engine/Source/Runtime/Landscape/Classes/LandscapeStreamingProxy.h
+++ b/Engine/Source/Runtime/Landscape/Classes/LandscapeStreamingProxy.h
@@ -23,7 +23,8 @@ public:
 
 	//~ Begin UObject Interface
 #if WITH_EDITOR
-	void PostEditChangeProperty(FPropertyChangedEvent& PropertyChangedEvent) override;
+	virtual bool ShouldExport() override { return false;  }
+	virtual void PostEditChangeProperty(FPropertyChangedEvent& PropertyChangedEvent) override;
 #endif
 	//~ End UObject Interface
 
