Expose ULandscapeLayerInfoObject to Blueprint (for EditorApplySpline() BP function)
#rnx

diff --git a/Engine/Source/Runtime/Landscape/Classes/LandscapeLayerInfoObject.h b/Engine/Source/Runtime/Landscape/Classes/LandscapeLayerInfoObject.h
--- a/Engine/Source/Runtime/Landscape/Classes/LandscapeLayerInfoObject.h
+++ b/Engine/Source/Runtime/Landscape/Classes/LandscapeLayerInfoObject.h
@@ -10,7 +10,7 @@
 class UPhysicalMaterial;
 struct FPropertyChangedEvent;
 
-UCLASS(MinimalAPI)
+UCLASS(MinimalAPI, BlueprintType)
 class ULandscapeLayerInfoObject : public UObject
 {
 	GENERATED_UCLASS_BODY()
