Make sure we don't access an invalid LayerInfo. (Complete CL 3593582)

diff --git a/Engine/Source/Runtime/Landscape/Private/LandscapeEditInterface.cpp b/Engine/Source/Runtime/Landscape/Private/LandscapeEditInterface.cpp
--- a/Engine/Source/Runtime/Landscape/Private/LandscapeEditInterface.cpp
+++ b/Engine/Source/Runtime/Landscape/Private/LandscapeEditInterface.cpp
@@ -2462,7 +2462,7 @@ inline TMap<const ULandscapeLayerInfoObject*, uint32> FLandscapeEditDataInterfac
 	for (int32 LayerIdx = 0; LayerIdx < Component->WeightmapLayerAllocations.Num(); LayerIdx++)
 	{
 		const FWeightmapLayerAllocationInfo& Allocation = Component->WeightmapLayerAllocations[LayerIdx];
-		if (Allocation.LayerInfo->bNoWeightBlend)
+		if (Allocation.LayerInfo == nullptr || Allocation.LayerInfo->bNoWeightBlend)
 		{
 			continue;
 		}
@@ -2673,6 +2673,12 @@ void FLandscapeEditDataInterface::SetAlphaData(ULandscapeLayerInfoObject* const
 					LayerNoWeightBlends[LayerIdx] = Allocation.LayerInfo->bNoWeightBlend;
 					LayerEditDataAllZero[LayerIdx] = true;
 				}
+				else
+				{
+					LayerDataPtrs[LayerIdx] = nullptr;
+					LayerNoWeightBlends[LayerIdx] = true;
+					LayerEditDataAllZero[LayerIdx] = false;
+				}
 			}
 
 			// Find the texture data corresponding to this vertex
@@ -2748,6 +2754,11 @@ void FLandscapeEditDataInterface::SetAlphaData(ULandscapeLayerInfoObject* const
 									// Adjust other layers' weights accordingly
 									for (int32 LayerIdx = 0; LayerIdx < Component->WeightmapLayerAllocations.Num(); LayerIdx++)
 									{
+										if (LayerDataPtrs[LayerIdx] == nullptr)
+										{
+											continue;
+										}
+
 										uint8& ExistingWeight = LayerDataPtrs[LayerIdx][TexDataIndex];
 
 										if (LayerIdx == UpdateLayerIdx)
@@ -2774,6 +2785,11 @@ void FLandscapeEditDataInterface::SetAlphaData(ULandscapeLayerInfoObject* const
 										// Normalize
 										for (int32 LayerIdx = 0; LayerIdx<Component->WeightmapLayerAllocations.Num(); LayerIdx++)
 										{
+											if (LayerDataPtrs[LayerIdx] == nullptr)
+											{
+												continue;
+											}
+
 											uint8& ExistingWeight = LayerDataPtrs[LayerIdx][TexDataIndex];
 
 											if (LayerNoWeightBlends[LayerIdx] == false)
@@ -2791,6 +2807,7 @@ void FLandscapeEditDataInterface::SetAlphaData(ULandscapeLayerInfoObject* const
 
 										if ((255 - OtherLayerWeightSum) && MaxLayerIdx >= 0)
 										{
+											// No need to check for nullptr here because MaxLayerIdx can only be set to a valid layer
 											LayerDataPtrs[MaxLayerIdx][TexDataIndex] += 255 - OtherLayerWeightSum;
 										}
 									}
@@ -2800,9 +2817,10 @@ void FLandscapeEditDataInterface::SetAlphaData(ULandscapeLayerInfoObject* const
 									// Adjust other layers' weights accordingly
 									for (int32 LayerIdx = 0; LayerIdx < Component->WeightmapLayerAllocations.Num(); LayerIdx++)
 									{
-										// Exclude bNoWeightBlend layers
+										// Exclude bNoWeightBlend layers 
 										if (LayerIdx != UpdateLayerIdx && LayerNoWeightBlends[LayerIdx] == false)
 										{
+											// No need to check for nullptr here because invalid layers have LayerNoWeightBlends set to true.
 											OtherLayerWeightSum += LayerDataPtrs.IsValidIndex(LayerIdx) ? LayerDataPtrs[LayerIdx][TexDataIndex] : 0;
 										}
 									}
@@ -2817,6 +2835,7 @@ void FLandscapeEditDataInterface::SetAlphaData(ULandscapeLayerInfoObject* const
 											{
 												const int32 ReplacementLayerIndex = Component->WeightmapLayerAllocations.IndexOfByPredicate([&](const FWeightmapLayerAllocationInfo& AllocationInfo) { return AllocationInfo.LayerInfo == ReplacementLayer; });
 
+												// No need to check for nullptr here because ChooseReplacementLayer can't return an invalid layer.
 												LayerDataPtrs[ReplacementLayerIndex][TexDataIndex] = 255 - NewWeight;
 												LayerEditDataAllZero[ReplacementLayerIndex] = false;
 											}
@@ -2832,12 +2851,18 @@ void FLandscapeEditDataInterface::SetAlphaData(ULandscapeLayerInfoObject* const
 											NewWeight = 255;
 										}
 
+										// No need to check for nullptr here because UpdateLayerIdx is always valid
 										LayerDataPtrs[UpdateLayerIdx][TexDataIndex] = NewWeight;
 									}
 									else
 									{
 										for (int32 LayerIdx = 0; LayerIdx < Component->WeightmapLayerAllocations.Num(); LayerIdx++)
 										{
+											if (LayerDataPtrs[LayerIdx] == nullptr)
+											{
+												continue;
+											}
+
 											uint8& Weight = LayerDataPtrs[LayerIdx][TexDataIndex];
 
 											if (LayerIdx == UpdateLayerIdx)
@@ -2863,6 +2888,7 @@ void FLandscapeEditDataInterface::SetAlphaData(ULandscapeLayerInfoObject* const
 							}
 							else
 							{
+								// No need to check for nullptr here because UpdateLayerIdx is always valid
 								// Weight value set without adjusting other layers' weights
 								uint8& Weight = LayerDataPtrs[UpdateLayerIdx][TexDataIndex];
 								Weight = NewWeight;
