[NO NON-PII COMMIT MESSAGE]

diff --git a/Engine/Source/Runtime/Landscape/Classes/LandscapeComponent.h b/Engine/Source/Runtime/Landscape/Classes/LandscapeComponent.h
--- a/Engine/Source/Runtime/Landscape/Classes/LandscapeComponent.h
+++ b/Engine/Source/Runtime/Landscape/Classes/LandscapeComponent.h
@@ -244,7 +244,7 @@ struct FLandscapeComponentMaterialOverride
 {
 	GENERATED_USTRUCT_BODY()
 
-	UPROPERTY(EditAnywhere, Category = LandscapeComponent)
+	UPROPERTY(EditAnywhere, Category = LandscapeComponent, meta=(UIMin=0, UIMax=8, ClampMin=0, ClampMax=8))
 	FPerPlatformInt LODIndex;
 
 	UPROPERTY(EditAnywhere, Category = LandscapeComponent)
diff --git a/Engine/Source/Runtime/Landscape/Classes/LandscapeProxy.h b/Engine/Source/Runtime/Landscape/Classes/LandscapeProxy.h
--- a/Engine/Source/Runtime/Landscape/Classes/LandscapeProxy.h
+++ b/Engine/Source/Runtime/Landscape/Classes/LandscapeProxy.h
@@ -339,7 +339,7 @@ struct FLandscapeProxyMaterialOverride
 {
 	GENERATED_USTRUCT_BODY()
 
-	UPROPERTY(EditAnywhere, Category = Landscape)
+	UPROPERTY(EditAnywhere, Category = LandscapeComponent, meta = (UIMin = 0, UIMax = 8, ClampMin = 0, ClampMax = 8))
 	FPerPlatformInt LODIndex;
 
 	UPROPERTY(EditAnywhere, Category = Landscape)
