Fix the category, didn't mean to change that.

diff --git a/Engine/Source/Runtime/Landscape/Classes/LandscapeProxy.h b/Engine/Source/Runtime/Landscape/Classes/LandscapeProxy.h
--- a/Engine/Source/Runtime/Landscape/Classes/LandscapeProxy.h
+++ b/Engine/Source/Runtime/Landscape/Classes/LandscapeProxy.h
@@ -339,7 +339,7 @@ struct FLandscapeProxyMaterialOverride
 {
 	GENERATED_USTRUCT_BODY()
 
-	UPROPERTY(EditAnywhere, Category = LandscapeComponent, meta = (UIMin = 0, UIMax = 8, ClampMin = 0, ClampMax = 8))
+	UPROPERTY(EditAnywhere, Category = Landscape, meta = (UIMin = 0, UIMax = 8, ClampMin = 0, ClampMax = 8))
 	FPerPlatformInt LODIndex;
 
 	UPROPERTY(EditAnywhere, Category = Landscape)
