Landscape Optim: Skip Texture Resolving when update flag is Editing_NoCollision (Makes brush live editing a lot more responsive)

diff --git a/Engine/Source/Runtime/Landscape/Private/LandscapeBlueprintBrushBase.cpp b/Engine/Source/Runtime/Landscape/Private/LandscapeBlueprintBrushBase.cpp
--- a/Engine/Source/Runtime/Landscape/Private/LandscapeBlueprintBrushBase.cpp
+++ b/Engine/Source/Runtime/Landscape/Private/LandscapeBlueprintBrushBase.cpp
@@ -7,7 +7,12 @@
 
 #define LOCTEXT_NAMESPACE "Landscape"
 
-static int FramePaddingBeforeLayerUpdate = 5;
+#if WITH_EDITOR
+static TAutoConsoleVariable<int32> CVarLandscapeBrushPadding(
+	TEXT("landscape.BrushFramePadding"),
+	5,
+	TEXT("The number of frames to wait before pushing a full Landscape update when a brush is calling RequestLandscapeUpdate"));
+#endif
 
 ALandscapeBlueprintBrushBase::ALandscapeBlueprintBrushBase(const FObjectInitializer& ObjectInitializer)
 #if WITH_EDITORONLY_DATA
@@ -59,7 +64,7 @@ void ALandscapeBlueprintBrushBase::Tick(float DeltaSeconds)
 #if WITH_EDITORONLY_DATA
 	// Avoid computing collision and client updates every frame
 	// Wait until we didn't trigger any more landscape update requests (padding of a couple of frames)
-	if (LastRequestLayersContentUpdateFrameNumber + FramePaddingBeforeLayerUpdate == GFrameNumber)
+	if (LastRequestLayersContentUpdateFrameNumber + CVarLandscapeBrushPadding.GetValueOnAnyThread() == GFrameNumber)
 	{
 		uint32 ModeMask = 0;
 		if (AffectHeightmap)
@@ -94,7 +99,7 @@ bool ALandscapeBlueprintBrushBase::ShouldTickIfViewportsOnly() const
 
 bool ALandscapeBlueprintBrushBase::IsLayerUpdatePending() const
 {
-	return GFrameNumber < LastRequestLayersContentUpdateFrameNumber + FramePaddingBeforeLayerUpdate;
+	return GFrameNumber < LastRequestLayersContentUpdateFrameNumber + CVarLandscapeBrushPadding.GetValueOnAnyThread();
 }
 
 void ALandscapeBlueprintBrushBase::SetIsVisible(bool bInIsVisible)
diff --git a/Engine/Source/Runtime/Landscape/Private/LandscapeEditLayers.cpp b/Engine/Source/Runtime/Landscape/Private/LandscapeEditLayers.cpp
--- a/Engine/Source/Runtime/Landscape/Private/LandscapeEditLayers.cpp
+++ b/Engine/Source/Runtime/Landscape/Private/LandscapeEditLayers.cpp
@@ -2944,13 +2944,19 @@ int32 ALandscape::RegenerateLayersHeightmaps(const TArray<ULandscapeComponent*>&
 
 	if (HeightmapUpdateModes)
 	{
-		ResolveLayersHeightmapTexture(InLandscapeComponentsToResolve);
+		bool bNeedsResolving = (HeightmapUpdateModes & AllHeightmapUpdateModes) != Update_Heightmap_Editing_NoCollision;
+		// We can skip Resolving if we don't need collision updating at all. (This relies on the fact that at some point we will resolve with a full update)
+		if (bNeedsResolving)
+		{
+			ResolveLayersHeightmapTexture(InLandscapeComponentsToResolve);
+		}
 
 		// Partial Component Update
 		for (ULandscapeComponent* Component : InLandscapeComponentsToResolve)
 		{
 			if(Component->IsUpdateFlagEnabledForModes(ELandscapeComponentUpdateFlag::Component_Update_Heightmap_Collision, HeightmapUpdateModes))
 			{
+				check(bNeedsResolving);
 				Component->UpdateCachedBounds();
 				Component->UpdateComponentToWorld();
 
@@ -3967,12 +3973,18 @@ int32 ALandscape::RegenerateLayersWeightmaps(const TArray<ULandscapeComponent*>&
 
 	if (WeightmapUpdateModes)
 	{
-		ResolveLayersWeightmapTexture(InLandscapeComponentsToResolve);
-	
+		bool bNeedsResolving = (WeightmapUpdateModes & AllWeightmapUpdateModes) != Update_Weightmap_Editing_NoCollision;
+		// We can skip Resolving if we don't need collision updating at all. (This relies on the fact that at some point we will resolve with a full update)
+		if(bNeedsResolving)
+		{
+			ResolveLayersWeightmapTexture(InLandscapeComponentsToResolve);
+		}
+				
 		for (ULandscapeComponent* Component : InLandscapeComponentsToResolve)
 		{
 			if (Component->IsUpdateFlagEnabledForModes(ELandscapeComponentUpdateFlag::Component_Update_Weightmap_Collision, WeightmapUpdateModes))
 			{
+				check(bNeedsResolving);
 				Component->UpdateCollisionLayerData();
 			}
 		}
