Undo changelist 8609516

Needs more work (breaks flatten and smooth)

#rnx

diff --git a/Engine/Source/Runtime/Landscape/Private/LandscapeEditLayers.cpp b/Engine/Source/Runtime/Landscape/Private/LandscapeEditLayers.cpp
--- a/Engine/Source/Runtime/Landscape/Private/LandscapeEditLayers.cpp
+++ b/Engine/Source/Runtime/Landscape/Private/LandscapeEditLayers.cpp
@@ -2944,19 +2944,13 @@ int32 ALandscape::RegenerateLayersHeightmaps(const TArray<ULandscapeComponent*>&
 
 	if (HeightmapUpdateModes)
 	{
-		bool bNeedsResolving = (HeightmapUpdateModes & AllHeightmapUpdateModes) != Update_Heightmap_Editing_NoCollision;
-		// We can skip Resolving if we don't need collision updating at all. (This relies on the fact that at some point we will resolve with a full update)
-		if (bNeedsResolving)
-		{
-			ResolveLayersHeightmapTexture(InLandscapeComponentsToResolve);
-		}
+		ResolveLayersHeightmapTexture(InLandscapeComponentsToResolve);
 
 		// Partial Component Update
 		for (ULandscapeComponent* Component : InLandscapeComponentsToResolve)
 		{
 			if(Component->IsUpdateFlagEnabledForModes(ELandscapeComponentUpdateFlag::Component_Update_Heightmap_Collision, HeightmapUpdateModes))
 			{
-				check(bNeedsResolving);
 				Component->UpdateCachedBounds();
 				Component->UpdateComponentToWorld();
 
@@ -3973,18 +3967,12 @@ int32 ALandscape::RegenerateLayersWeightmaps(const TArray<ULandscapeComponent*>&
 
 	if (WeightmapUpdateModes)
 	{
-		bool bNeedsResolving = (WeightmapUpdateModes & AllWeightmapUpdateModes) != Update_Weightmap_Editing_NoCollision;
-		// We can skip Resolving if we don't need collision updating at all. (This relies on the fact that at some point we will resolve with a full update)
-		if(bNeedsResolving)
-		{
-			ResolveLayersWeightmapTexture(InLandscapeComponentsToResolve);
-		}
-				
+		ResolveLayersWeightmapTexture(InLandscapeComponentsToResolve);
+	
 		for (ULandscapeComponent* Component : InLandscapeComponentsToResolve)
 		{
 			if (Component->IsUpdateFlagEnabledForModes(ELandscapeComponentUpdateFlag::Component_Update_Weightmap_Collision, WeightmapUpdateModes))
 			{
-				check(bNeedsResolving);
 				Component->UpdateCollisionLayerData();
 			}
 		}
