LandscapeSpline: Remove the call to CheckForErrors in the PostLoad as it can depend on other maps not currently loaded.

#rnx

diff --git a/Engine/Source/Runtime/Landscape/Private/LandscapeSplines.cpp b/Engine/Source/Runtime/Landscape/Private/LandscapeSplines.cpp
--- a/Engine/Source/Runtime/Landscape/Private/LandscapeSplines.cpp
+++ b/Engine/Source/Runtime/Landscape/Private/LandscapeSplines.cpp
@@ -806,13 +806,6 @@ void ULandscapeSplinesComponent::PostLoad()
 #endif
 
 	CheckSplinesValid();
-
-#if WITH_EDITOR
-	if (GIsEditor && GetWorld()->WorldType == EWorldType::Editor)
-	{
-		CheckForErrors();
-	}
-#endif
 }
 
 #if WITH_EDITOR
