Landscape Optim: Skip Texture Resolving when update flag is Editing_NoCollision (Makes brush live editing a lot more responsive) (take 2)

diff --git a/Engine/Source/Runtime/Landscape/Private/LandscapeEditLayers.cpp b/Engine/Source/Runtime/Landscape/Private/LandscapeEditLayers.cpp
--- a/Engine/Source/Runtime/Landscape/Private/LandscapeEditLayers.cpp
+++ b/Engine/Source/Runtime/Landscape/Private/LandscapeEditLayers.cpp
@@ -2944,13 +2944,19 @@ int32 ALandscape::RegenerateLayersHeightmaps(const TArray<ULandscapeComponent*>&
 
 	if (HeightmapUpdateModes)
 	{
-		ResolveLayersHeightmapTexture(InLandscapeComponentsToResolve);
+		bool bNeedsResolving = bIntermediateRender || (HeightmapUpdateModes & AllHeightmapUpdateModes) != Update_Heightmap_Editing_NoCollision;
+		// We can skip Resolving if we don't need collision updating at all and that we aren't doing an intermediate render for some Landscape Tool (Flatten, Smooth). (This relies on the fact that at some point we will resolve with a full update)
+		if (bNeedsResolving)
+		{
+			ResolveLayersHeightmapTexture(InLandscapeComponentsToResolve);
+		}
 
 		// Partial Component Update
 		for (ULandscapeComponent* Component : InLandscapeComponentsToResolve)
 		{
 			if(Component->IsUpdateFlagEnabledForModes(ELandscapeComponentUpdateFlag::Component_Update_Heightmap_Collision, HeightmapUpdateModes))
 			{
+				check(bNeedsResolving);
 				Component->UpdateCachedBounds();
 				Component->UpdateComponentToWorld();
 
@@ -3967,12 +3973,18 @@ int32 ALandscape::RegenerateLayersWeightmaps(const TArray<ULandscapeComponent*>&
 
 	if (WeightmapUpdateModes)
 	{
-		ResolveLayersWeightmapTexture(InLandscapeComponentsToResolve);
-	
+		bool bNeedsResolving = bIntermediateRender || (WeightmapUpdateModes & AllWeightmapUpdateModes) != Update_Weightmap_Editing_NoCollision;
+		// We can skip Resolving if we don't need collision updating at all and that we aren't doing an intermediate render for some Landscape Tool (Flatten, Smooth). (This relies on the fact that at some point we will resolve with a full update)
+		if(bNeedsResolving)
+		{
+			ResolveLayersWeightmapTexture(InLandscapeComponentsToResolve);
+		}
+				
 		for (ULandscapeComponent* Component : InLandscapeComponentsToResolve)
 		{
 			if (Component->IsUpdateFlagEnabledForModes(ELandscapeComponentUpdateFlag::Component_Update_Weightmap_Collision, WeightmapUpdateModes))
 			{
+				check(bNeedsResolving);
 				Component->UpdateCollisionLayerData();
 			}
 		}
