New Full Heightmap Render Event

#rnx

diff --git a/Engine/Source/Runtime/Landscape/Classes/Landscape.h b/Engine/Source/Runtime/Landscape/Classes/Landscape.h
--- a/Engine/Source/Runtime/Landscape/Classes/Landscape.h
+++ b/Engine/Source/Runtime/Landscape/Classes/Landscape.h
@@ -414,6 +414,9 @@ public:
 	DECLARE_EVENT(ALandscape, FLandscapeBlueprintBrushChangedDelegate);
 	FLandscapeBlueprintBrushChangedDelegate& OnBlueprintBrushChangedDelegate() { return LandscapeBlueprintBrushChangedDelegate; }
 
+	DECLARE_EVENT_OneParam(ALandscape, FLandscapeFullHeightmapRenderDoneDelegate, UTextureRenderTarget2D*);
+	FLandscapeFullHeightmapRenderDoneDelegate& OnFullHeightmapRenderDoneDelegate() { return LandscapeFullHeightmapRenderDoneDelegate; }
+
 	/** Target Landscape Layer for Landscape Splines */
 	UPROPERTY()
 	FGuid LandscapeSplinesTargetLayerGuid;
@@ -435,6 +438,7 @@ public:
 
 private:
 	FLandscapeBlueprintBrushChangedDelegate LandscapeBlueprintBrushChangedDelegate;
+	FLandscapeFullHeightmapRenderDoneDelegate LandscapeFullHeightmapRenderDoneDelegate;
 
 	/** Components affected by landscape splines (used to partially clear Layer Reserved for Splines) */
 	UPROPERTY(Transient)
diff --git a/Engine/Source/Runtime/Landscape/Private/LandscapeEditLayers.cpp b/Engine/Source/Runtime/Landscape/Private/LandscapeEditLayers.cpp
--- a/Engine/Source/Runtime/Landscape/Private/LandscapeEditLayers.cpp
+++ b/Engine/Source/Runtime/Landscape/Private/LandscapeEditLayers.cpp
@@ -2901,6 +2901,12 @@ int32 ALandscape::RegenerateLayersHeightmaps(const TArray<ULandscapeComponent*>&
 		ShaderParams.GenerateNormals = true;
 		ShaderParams.GridSize = GetRootComponent()->RelativeScale3D;
 
+		// Broadcast Event of the Full Render
+		if ((HeightmapUpdateModes & Update_Heightmap_All) == Update_Heightmap_All)
+		{
+			LandscapeFullHeightmapRenderDoneDelegate.Broadcast(LandscapeScratchRT3);
+		}
+
 		DrawHeightmapComponentsToRenderTarget(OutputDebugName ? FString::Printf(TEXT("LS Height: %s = -> CombinedNonAtlasNormals : %s"), *CombinedHeightmapNonAtlasRT->GetName(), *LandscapeScratchRT1->GetName()) : GEmptyDebugName,
 											  InLandscapeComponentsToRender, LandscapeExtent.Min, CombinedHeightmapNonAtlasRT, nullptr, LandscapeScratchRT1, ERTDrawingType::RTNonAtlas, true, ShaderParams);
 
