Use R8G8B8A8 texture for weightmap accumulation
This is enough to satisfy D3D11.0 base requirement for UAV creation and fixes non supported copy on D3D12 RHI

diff --git a/Engine/Shaders/Private/LandscapeLayersCS.usf b/Engine/Shaders/Private/LandscapeLayersCS.usf
--- a/Engine/Shaders/Private/LandscapeLayersCS.usf
+++ b/Engine/Shaders/Private/LandscapeLayersCS.usf
@@ -41,7 +41,7 @@ StructuredBuffer<FLandscapeLayerWeightmapPackMaterialLayersComponentData> InPack
 Buffer<float> InWeightmapWeightBlendMode;
 Buffer<float2> InWeightmapTextureOutputOffset;
 Texture2DArray<float> InAtlasPaintLayers;
-RWTexture2D<uint> OutComponentWeightMaps;
+RWTexture2D<UNORM float4> OutComponentWeightMaps;
 
 float ComputeWeightmapChannel(int ChannelIndexToProcess, uint2 TexelPosition, uint SourcePaintLayerIndex)
 {
@@ -77,28 +77,27 @@ void PackPaintLayerToWeightmap(uint3 DispatchThreadId : SV_DispatchThreadID)
 	{
 		FLandscapeLayerWeightmapPackMaterialLayersComponentData ComponentData = InPackLayersComponentsData[DispatchThreadId.z];
 
-		uint Channels = 0;
+		float4 Channels = 0;
 
 		for (int i = 0; i < 4; ++i)
 		{
 			float ChannelResult = ComputeWeightmapChannel(ComponentData.WeightmapChannelToProcess[i], DispatchThreadId.xy + int2(ComponentData.ComponenVertexPositionX[i], ComponentData.ComponenVertexPositionY[i]), ComponentData.SourcePaintLayerIndex[i]);
 
-			// BGRA to RGBA
 			if (ComponentData.WeightmapChannelToProcess[i] == 0)
 			{ 
-				Channels |= uint(ChannelResult * 255.0) << 8;
+				Channels.x = ChannelResult;
 			}
 			else if (ComponentData.WeightmapChannelToProcess[i] == 1)
 			{
-				Channels |= uint(ChannelResult * 255.0) << 16;
+				Channels.y = ChannelResult;
 			}
 			else if (ComponentData.WeightmapChannelToProcess[i] == 2)
 			{
-				Channels |= uint(ChannelResult * 255.0) << 24;
+				Channels.z = ChannelResult;
 			}
 			else if (ComponentData.WeightmapChannelToProcess[i] == 3)
 			{
-				Channels |= uint(ChannelResult * 255.0);
+				Channels.w = ChannelResult;
 			}
 		}
 
diff --git a/Engine/Shaders/Private/LandscapeLayersPS.usf b/Engine/Shaders/Private/LandscapeLayersPS.usf
--- a/Engine/Shaders/Private/LandscapeLayersPS.usf
+++ b/Engine/Shaders/Private/LandscapeLayersPS.usf
@@ -347,15 +347,6 @@ void PSHeightmapMainMips(float2 InTextureCoordinates : TEXCOORD0, out float4 Out
 
 float4 WeightmapLayerWeightBlend; // each xyzw == 1 texture channel
 
-void PSWeightmapConvertFormat(float2 InTextureCoordinates : TEXCOORD0, out float4 OutColor : SV_Target0)
-{
-	uint SourceColor = asuint(ReadTexture1.SampleLevel(ReadTexture1Sampler, InTextureCoordinates, 0).r);
-	OutColor.b = float((SourceColor >> 24) & 0x000000FF) / 255.0;
-	OutColor.g = float((SourceColor >> 16) & 0x000000FF) / 255.0;
-	OutColor.r = float((SourceColor >> 8)  & 0x000000FF) / 255.0;
-	OutColor.a = float(SourceColor & 0x000000FF) / 255.0;
-}
-
 void PSWeightmapMain(float2 InTextureCoordinates : TEXCOORD0, out float4 OutColor : SV_Target0)
 {
 	float4 SourceColor = ReadTexture1.Sample(ReadTexture1Sampler, InTextureCoordinates);
@@ -439,4 +430,16 @@ void PSWeightmapMainMips(float2 InTextureCoordinates : TEXCOORD0, out float4 Out
 
 		OutColor = IsMinBorderTexelY ? lerp(SourceSample, RightSample, Alpha) : (IsMaxBorderTexelY ? lerp(DownSample, DownRightSample, Alpha) : lerp(lerp(SourceSample, RightSample, Alpha), lerp(DownSample, DownRightSample, Alpha), Alpha));
 	}
-}
\ No newline at end of file
+}
+
+void CopyTextureVS(in uint VertexId : SV_VertexID, out float4 OutPosition : SV_POSITION, out float2 OutTexCoord : TEXCOORD0)
+{
+	OutTexCoord = float2(((VertexId + 1) / 3) & 1, VertexId & 1);
+	OutPosition.xy = float2(OutTexCoord.x, 1.f - OutTexCoord.y) * 2.f - 1.f;
+	OutPosition.zw = float2(0, 1);
+}
+
+void CopyTexturePS(in float4 InPosition : SV_POSITION, in noperspective float2 InTexCoord : TEXCOORD0, out float4 OutColor : SV_Target0)
+{
+	OutColor = ReadTexture1.SampleLevel(ReadTexture1Sampler, InTexCoord, 0);
+}
diff --git a/Engine/Source/Runtime/Landscape/Classes/Landscape.h b/Engine/Source/Runtime/Landscape/Classes/Landscape.h
--- a/Engine/Source/Runtime/Landscape/Classes/Landscape.h
+++ b/Engine/Source/Runtime/Landscape/Classes/Landscape.h
@@ -381,6 +381,8 @@ private:
 	void DrawHeightmapComponentsToRenderTargetMips(const TArray<ULandscapeComponent*>& InComponentsToDraw, const FIntPoint& InLandscapeBase, UTexture* InReadHeightmap, bool InClearRTWrite, struct FLandscapeLayersHeightmapShaderParameters& InShaderParams) const;
 	void DrawWeightmapComponentToRenderTargetMips(const TArray<FVector2D>& InTexturePositionsToDraw, UTexture* InReadWeightmap, bool InClearRTWrite, struct FLandscapeLayersWeightmapShaderParameters& InShaderParams) const;
 
+	void CopyTexturePS(const FString& InSourceDebugName, FTextureResource* InSourceResource, const FString& InDestDebugName, FTextureResource* InDestResource) const;
+
 	void CopyLayersTexture(UTexture* InSourceTexture, UTexture* InDestTexture, FTextureResource* InDestCPUResource = nullptr, const FIntPoint& InInitialPositionOffset = FIntPoint(0, 0), uint8 InSourceCurrentMip = 0, uint8 InDestCurrentMip = 0,
 						   uint32 InSourceArrayIndex = 0, uint32 InDestArrayIndex = 0) const;
 	void CopyLayersTexture(const FString& InSourceDebugName, FTextureResource* InSourceResource, const FString& InDestDebugName, FTextureResource* InDestResource, FTextureResource* InDestCPUResource = nullptr, const FIntPoint& InInitialPositionOffset = FIntPoint(0, 0),
@@ -400,9 +402,6 @@ private:
 	void PrintLayersDebugTextureResource(const FString& InContext, FTextureResource* InTextureResource, uint8 InMipRender = 0, bool InOutputHeight = true, bool InOutputNormals = false) const;
 	void PrintLayersDebugHeightData(const FString& InContext, const TArray<FColor>& InHeightmapData, const FIntPoint& InDataSize, uint8 InMipRender, bool InOutputNormals = false) const;
 	void PrintLayersDebugWeightData(const FString& InContext, const TArray<FColor>& InWeightmapData, const FIntPoint& InDataSize, uint8 InMipRender) const;
-
-	// Needed because of UAVs on DX11.0 not supporting BGRA8
-	void ConvertR32ToBGRA8(const TArray<FVector2D>& InWeightmapTextureOutputOffset, UTextureRenderTarget2D* InWeightmapRTWrite, const struct FLandscapeLayersWeightmapConvertFormatShaderParameter& InShaderParams);
 #endif
 
 public:
diff --git a/Engine/Source/Runtime/Landscape/Private/LandscapeEditLayers.cpp b/Engine/Source/Runtime/Landscape/Private/LandscapeEditLayers.cpp
--- a/Engine/Source/Runtime/Landscape/Private/LandscapeEditLayers.cpp
+++ b/Engine/Source/Runtime/Landscape/Private/LandscapeEditLayers.cpp
@@ -645,43 +645,51 @@ private:
 
 IMPLEMENT_GLOBAL_SHADER(FLandscapeLayersWeightmapMipsPS, "/Engine/Private/LandscapeLayersPS.usf", "PSWeightmapMainMips", SF_Pixel);
 
-struct FLandscapeLayersWeightmapConvertFormatShaderParameter
+class FLandscapeCopyTextureVS : public FGlobalShader
 {
-	FLandscapeLayersWeightmapConvertFormatShaderParameter()
-		: ReadWeightmap(nullptr)
-	{}
+public:
+	DECLARE_GLOBAL_SHADER(FLandscapeCopyTextureVS);
 
-	FLandscapeTexture2DResource* ReadWeightmap;
+	static bool ShouldCompilePermutation(const FGlobalShaderPermutationParameters&)
+	{
+		return true;
+	}
+
+	FLandscapeCopyTextureVS() 
+	{};
+	
+	FLandscapeCopyTextureVS(const ShaderMetaType::CompiledShaderInitializerType& Initializer)
+		: FGlobalShader(Initializer)
+	{}
 };
 
-// This Shader exists only to support Win7 because DX11.0 doesn't support UAVs on BGRA8 Textures.
-// The Shader will convert from R32 Format to BGRA8 Format
-class FLandscapeLayersWeightmapConvertFormatPS : public FGlobalShader
+class FLandscapeCopyTexturePS : public FGlobalShader
 {
-	DECLARE_GLOBAL_SHADER(FLandscapeLayersWeightmapConvertFormatPS);
+	DECLARE_GLOBAL_SHADER(FLandscapeCopyTexturePS);
 public:
 
 	static bool ShouldCompilePermutation(const FGlobalShaderPermutationParameters& Parameters)
 	{
-		return IsFeatureLevelSupported(Parameters.Platform, ERHIFeatureLevel::SM5) && !IsConsolePlatform(Parameters.Platform) && !IsMetalPlatform(Parameters.Platform);
+		return IsFeatureLevelSupported(Parameters.Platform, ERHIFeatureLevel::SM5) && !IsConsolePlatform(Parameters.Platform) && !IsMetalMobilePlatform(Parameters.Platform);
 	}
 
 	static void ModifyCompilationEnvironment(const FGlobalShaderPermutationParameters& Parameters, FShaderCompilerEnvironment& OutEnvironment)
 	{
 	}
-	FLandscapeLayersWeightmapConvertFormatPS(const ShaderMetaType::CompiledShaderInitializerType& Initializer)
+
+	FLandscapeCopyTexturePS(const ShaderMetaType::CompiledShaderInitializerType& Initializer)
 		: FGlobalShader(Initializer)
 	{
 		ReadTexture1Param.Bind(Initializer.ParameterMap, TEXT("ReadTexture1"));
 		ReadTexture1SamplerParam.Bind(Initializer.ParameterMap, TEXT("ReadTexture1Sampler"));
 	}
 
-	FLandscapeLayersWeightmapConvertFormatPS()
+	FLandscapeCopyTexturePS()
 	{}
 
-	void SetParameters(FRHICommandList& RHICmdList, const FLandscapeLayersWeightmapConvertFormatShaderParameter& InParams)
+	void SetParameters(FRHICommandList& RHICmdList, FRHITexture* TextureRHI)
 	{
-		SetTextureParameter(RHICmdList, GetPixelShader(), ReadTexture1Param, ReadTexture1SamplerParam, TStaticSamplerState<SF_Point, AM_Clamp, AM_Clamp, AM_Clamp>::GetRHI(), InParams.ReadWeightmap->TextureRHI);
+		SetTextureParameter(RHICmdList, GetPixelShader(), ReadTexture1Param, ReadTexture1SamplerParam, TStaticSamplerState<SF_Point, AM_Clamp, AM_Clamp, AM_Clamp>::GetRHI(), TextureRHI);
 	}
 
 	virtual bool Serialize(FArchive& Ar) override
@@ -689,15 +697,16 @@ public:
 		bool bShaderHasOutdatedParameters = FShader::Serialize(Ar);
 		Ar << ReadTexture1Param;
 		Ar << ReadTexture1SamplerParam;
-
 		return bShaderHasOutdatedParameters;
 	}
 
+private:
 	FShaderResourceParameter ReadTexture1Param;
 	FShaderResourceParameter ReadTexture1SamplerParam;
 };
 
-IMPLEMENT_GLOBAL_SHADER(FLandscapeLayersWeightmapConvertFormatPS, "/Engine/Private/LandscapeLayersPS.usf", "PSWeightmapConvertFormat", SF_Pixel);
+IMPLEMENT_GLOBAL_SHADER(FLandscapeCopyTextureVS, "/Engine/Private/LandscapeLayersPS.usf", "CopyTextureVS", SF_Vertex);
+IMPLEMENT_GLOBAL_SHADER(FLandscapeCopyTexturePS, "/Engine/Private/LandscapeLayersPS.usf", "CopyTexturePS", SF_Pixel);
 
 // Compute shaders data
 
@@ -1296,8 +1305,6 @@ private:
 
 typedef FLandscapeLayersRender_RenderThread<FLandscapeLayersHeightmapShaderParameters, FLandscapeLayersHeightmapPS, FLandscapeLayersHeightmapMipsPS> FLandscapeLayersHeightmapRender_RenderThread;
 typedef FLandscapeLayersRender_RenderThread<FLandscapeLayersWeightmapShaderParameters, FLandscapeLayersWeightmapPS, FLandscapeLayersWeightmapMipsPS> FLandscapeLayersWeightmapRender_RenderThread;
-// The Mips Param will not be used
-typedef FLandscapeLayersRender_RenderThread<FLandscapeLayersWeightmapConvertFormatShaderParameter, FLandscapeLayersWeightmapConvertFormatPS, FLandscapeLayersWeightmapConvertFormatPS> FLandscapeLayersWeightmapConvertRender_RenderThread;
 
 #if WITH_EDITOR
 
@@ -1881,6 +1888,43 @@ void ALandscape::CommitDeferredCopyLayersTexture()
 	});
 }
 
+void ALandscape::CopyTexturePS(const FString& InSourceDebugName, FTextureResource* InSourceResource, const FString& InDestDebugName, FTextureResource* InDestResource) const
+{
+	check(InSourceResource != nullptr);
+	check(InDestResource != nullptr);
+	check(InSourceResource->GetSizeX() == InDestResource->GetSizeX());
+	check(InSourceResource->GetSizeY() == InDestResource->GetSizeY());
+
+	ENQUEUE_RENDER_COMMAND(CopyPSCommand)(
+		[InSourceResource, InDestResource](FRHICommandListImmediate& RHICmdList)
+	{
+		FRHIRenderPassInfo RPInfo(InDestResource->TextureRHI, ERenderTargetActions::DontLoad_Store);
+		RHICmdList.BeginRenderPass(RPInfo, TEXT("CopyTexture"));
+
+		TShaderMap<FGlobalShaderType>* GlobalShaderMap = GetGlobalShaderMap(GMaxRHIFeatureLevel);
+		TShaderMapRef< FLandscapeCopyTextureVS > VertexShader(GlobalShaderMap);
+		TShaderMapRef< FLandscapeCopyTexturePS > PixelShader(GlobalShaderMap);
+
+		FGraphicsPipelineStateInitializer GraphicsPSOInit;
+		RHICmdList.ApplyCachedRenderTargets(GraphicsPSOInit);
+		GraphicsPSOInit.BlendState = TStaticBlendState<>::GetRHI();
+		GraphicsPSOInit.RasterizerState = TStaticRasterizerState<>::GetRHI();
+		GraphicsPSOInit.DepthStencilState = TStaticDepthStencilState<false, CF_Always>::GetRHI();
+		GraphicsPSOInit.PrimitiveType = PT_TriangleList;
+		GraphicsPSOInit.BoundShaderState.VertexDeclarationRHI = GetVertexDeclarationFVector4();
+		GraphicsPSOInit.BoundShaderState.VertexShaderRHI = GETSAFERHISHADER_VERTEX(*VertexShader);
+		GraphicsPSOInit.BoundShaderState.PixelShaderRHI = GETSAFERHISHADER_PIXEL(*PixelShader);
+		SetGraphicsPipelineState(RHICmdList, GraphicsPSOInit);
+
+		PixelShader->SetParameters(RHICmdList, InSourceResource->TextureRHI);
+
+		RHICmdList.SetViewport(0.0f, 0.0f, 0.0f, InDestResource->GetSizeX(), InDestResource->GetSizeY(), 1.0f);
+		RHICmdList.DrawIndexedPrimitive(GTwoTrianglesIndexBuffer.IndexBufferRHI, 0, 0, 4, 0, 2, 1);
+
+		RHICmdList.EndRenderPass();
+	});
+}
+
 void ALandscape::CopyLayersTexture(UTexture* InSourceTexture, UTexture* InDestTexture, FTextureResource* InDestCPUResource, const FIntPoint& InFirstComponentSectionBase, uint8 InSourceCurrentMip, uint8 InDestCurrentMip, uint32 InSourceArrayIndex, uint32 InDestArrayIndex) const
 {
 	if (InSourceTexture != nullptr && InDestTexture != nullptr)
@@ -2048,38 +2092,6 @@ void ALandscape::DrawWeightmapComponentToRenderTargetMips(const TArray<FVector2D
 	}
 }
 
-void ALandscape::ConvertR32ToBGRA8(const TArray<FVector2D>& InWeightmapTextureOutputOffset, UTextureRenderTarget2D* InWeightmapRTWrite, const FLandscapeLayersWeightmapConvertFormatShaderParameter& InShaderParams)
-{
-	check(InShaderParams.ReadWeightmap != nullptr);
-	check(InWeightmapRTWrite != nullptr);
-
-	FIntPoint WeightmapWriteTextureSize(InWeightmapRTWrite->SizeX, InWeightmapRTWrite->SizeY);
-	FIntPoint WeightmapReadTextureSize(InShaderParams.ReadWeightmap->GetSizeX(), InShaderParams.ReadWeightmap->GetSizeY());
-
-	// Quad Setup
-	const int32 LocalComponentSizeQuad = SubsectionSizeQuads * NumSubsections;
-	const int32 LocalComponentSizeVerts = (SubsectionSizeQuads + 1) * NumSubsections;
-
-	TArray<FLandscapeLayersTriangle> TriangleList;
-	TriangleList.Reserve(InWeightmapTextureOutputOffset.Num() * 2 * NumSubsections);
-	for (const FVector2D& TexturePosition : InWeightmapTextureOutputOffset)
-	{
-		FVector2D PositionOffset(FMath::RoundToInt(TexturePosition.X / LocalComponentSizeVerts), FMath::RoundToInt(TexturePosition.Y / LocalComponentSizeVerts));
-		FIntPoint SectionBase(PositionOffset.X * LocalComponentSizeQuad, PositionOffset.Y * LocalComponentSizeQuad);
-		GenerateLayersRenderQuadsAtlas(SectionBase, FVector2D::ZeroVector, SubsectionSizeQuads, WeightmapReadTextureSize, WeightmapWriteTextureSize, TriangleList);
-	}
-
-	FMatrix ProjectionMatrix = AdjustProjectionMatrixForRHI(FTranslationMatrix(FVector(0, 0, 0)) *
-		FMatrix(FPlane(1.0f / (FMath::Max<uint32>(WeightmapWriteTextureSize.X, 1.f) / 2.0f), 0.0, 0.0f, 0.0f), FPlane(0.0f, -1.0f / (FMath::Max<uint32>(WeightmapWriteTextureSize.Y, 1.f) / 2.0f), 0.0f, 0.0f), FPlane(0.0f, 0.0f, 1.0f, 0.0f), FPlane(-1.0f, 1.0f, 0.0f, 1.0f)));
-
-	FLandscapeLayersWeightmapConvertRender_RenderThread LayersRender(FString(TEXT("ConvertR32ToBGRA8")), InWeightmapRTWrite, WeightmapWriteTextureSize, WeightmapReadTextureSize, ProjectionMatrix, InShaderParams, 0, TriangleList);
-	ENQUEUE_RENDER_COMMAND(FConvertR32ToBGRA8Command)(
-		[LayersRender](FRHICommandListImmediate& RHICmdList) mutable
-	{
-		LayersRender.Render(RHICmdList, true);
-	});
-}
-
 void ALandscape::ClearLayersWeightmapTextureResource(const FString& InDebugName, FTextureRenderTargetResource* InTextureResourceToClear) const
 {
 	LandscapeLayersWeightmapClear_RenderThread LayersClear(InDebugName, InTextureResourceToClear);
@@ -3494,7 +3506,9 @@ void ALandscape::InitializeLayersWeightmapResources()
 		}
 	}
 
-	WeightmapScratchPackLayerTextureResource = new FLandscapeTexture2DResource(FirstWeightmapRT->SizeX, FirstWeightmapRT->SizeY, PF_R32_UINT, MipCount, true);
+	// Format for UAV can't be PF_B8G8R8A8 on Windows 7 so use PF_R8G8B8A8
+	// We make the final copy out of this to a PF_R8G8B8A8 target with CopyTexturePS() instead of CopyLayersTexture() because a pixel shader will automatically handle the channel swizzling (where a RHICopyTexture won't)
+	WeightmapScratchPackLayerTextureResource = new FLandscapeTexture2DResource(FirstWeightmapRT->SizeX, FirstWeightmapRT->SizeY, PF_R8G8B8A8, MipCount, true);
 	BeginInitResource(WeightmapScratchPackLayerTextureResource);
 }
 
@@ -3930,11 +3944,7 @@ int32 ALandscape::RegenerateLayersWeightmaps(const TArray<ULandscapeComponent*>&
 					SourceDebugName = OutputDebugName ? TEXT("WeightmapScratchTexture") : GEmptyDebugName;
 					DestDebugName = OutputDebugName ? CurrentRT->GetName() : GEmptyDebugName;
 					
-					// Convert Back to BGRA8(Win7)
-					FLandscapeLayersWeightmapConvertFormatShaderParameter ShaderParams;
-					ShaderParams.ReadWeightmap = WeightmapScratchPackLayerTextureResource;
-					ConvertR32ToBGRA8(WeightmapTextureOutputOffset, CurrentRT, ShaderParams);
-					
+					CopyTexturePS(SourceDebugName, WeightmapScratchPackLayerTextureResource, DestDebugName, CurrentRT->GameThread_GetRenderTargetResource());
 					DrawWeightmapComponentToRenderTargetMips(WeightmapTextureOutputOffset, CurrentRT, true, PSShaderParams);
 
 					int32 StartTextureIndex = NextTextureIndexToProcess;
