A further fix to avoid D3D12 RHI failures in the landscape layer editing
We can now see further vsisual problems that look to be related to fences/transitions/readback or similar

diff --git a/Engine/Source/Runtime/Landscape/Private/LandscapeEditLayers.cpp b/Engine/Source/Runtime/Landscape/Private/LandscapeEditLayers.cpp
--- a/Engine/Source/Runtime/Landscape/Private/LandscapeEditLayers.cpp
+++ b/Engine/Source/Runtime/Landscape/Private/LandscapeEditLayers.cpp
@@ -111,11 +111,11 @@ public:
 		FTextureResource::InitRHI();
 
 		FRHIResourceCreateInfo CreateInfo;
-		uint32 Flags = TexCreate_NoTiling | TexCreate_OfflineProcessed;
+		uint32 Flags = 0;
 
 		if (CreateUAV)
 		{
-			Flags |= TexCreate_UAV;
+			Flags |= TexCreate_RenderTargetable|TexCreate_UAV;
 		}
 
 		TextureRHI = RHICreateTexture2D(SizeX, SizeY, Format, NumMips, 1, Flags, CreateInfo);
@@ -3922,7 +3922,7 @@ int32 ALandscape::RegenerateLayersWeightmaps(const TArray<ULandscapeComponent*>&
 					SourceDebugName = OutputDebugName ? *EmptyRT->GetName() : GEmptyDebugName;
 					DestDebugName = OutputDebugName ? TEXT("Weight: Clear WeightmapScratchPackLayerTextureResource") : GEmptyDebugName;
 
-					CopyLayersTexture(SourceDebugName, EmptyRT->GameThread_GetRenderTargetResource(), DestDebugName, WeightmapScratchPackLayerTextureResource);
+					CopyTexturePS(SourceDebugName, EmptyRT->GameThread_GetRenderTargetResource(), DestDebugName, WeightmapScratchPackLayerTextureResource);
 
 					FLandscapeLayerWeightmapPackMaterialLayersComputeShaderParameters CSPackLayersShaderParams;
 					CSPackLayersShaderParams.AtlasWeightmapsPerLayer = CombinedLayersWeightmapAllMaterialLayersResource;
