Landscape Spline: new map check allows fixing some unreferenced segment spline mesh.

#rnx

diff --git a/Engine/Source/Runtime/Landscape/Classes/LandscapeSplinesComponent.h b/Engine/Source/Runtime/Landscape/Classes/LandscapeSplinesComponent.h
--- a/Engine/Source/Runtime/Landscape/Classes/LandscapeSplinesComponent.h
+++ b/Engine/Source/Runtime/Landscape/Classes/LandscapeSplinesComponent.h
@@ -180,6 +180,10 @@ public:
 	virtual void AddForeignMeshComponent(ULandscapeSplineControlPoint* Owner, UControlPointMeshComponent* Component);
 	virtual void RemoveForeignMeshComponent(ULandscapeSplineControlPoint* Owner, UControlPointMeshComponent* Component);
 	virtual void DestroyOrphanedForeignMeshComponents(UWorld* OwnerWorld);
+
+	void DestroyUnreferencedForeignMeshComponents();
+	void ForEachUnreferencedForeignMeshComponent(TFunctionRef<bool(ULandscapeSplineSegment*, USplineMeshComponent*)> Func);
+	
 	virtual UControlPointMeshComponent*   GetForeignMeshComponent(ULandscapeSplineControlPoint* Owner);
 	virtual TArray<USplineMeshComponent*> GetForeignMeshComponents(ULandscapeSplineSegment* Owner);
 
diff --git a/Engine/Source/Runtime/Landscape/Private/LandscapeSplines.cpp b/Engine/Source/Runtime/Landscape/Private/LandscapeSplines.cpp
--- a/Engine/Source/Runtime/Landscape/Private/LandscapeSplines.cpp
+++ b/Engine/Source/Runtime/Landscape/Private/LandscapeSplines.cpp
@@ -773,6 +773,21 @@ void ULandscapeSplinesComponent::CheckForErrors()
 			}
 		}
 	}
+
+	// Find Unreferenced Foreign Mesh Components
+	ForEachUnreferencedForeignMeshComponent([this, ThisOuterWorld](ULandscapeSplineSegment* Segment, USplineMeshComponent* SplineMeshComponent)
+	{
+		FFormatNamedArguments Arguments;
+		Arguments.Add(TEXT("MeshMap"), FText::FromName(ThisOuterWorld->GetFName()));
+
+		FMessageLog("MapCheck").Error()
+			->AddToken(FUObjectToken::Create(GetOwner()))
+			->AddToken(FTextToken::Create(FText::Format(LOCTEXT("MapCheck_Message_UnreferencedMeshes", "{MeshMap} contains meshes with no owners"), Arguments)))
+			->AddToken(FActionToken::Create(LOCTEXT("MapCheck_ActionName_UnreferencedMeshes", "Clean up meshes that are not referenced by their owner segments"), FText(),
+				FOnActionTokenExecuted::CreateUObject(this, &ULandscapeSplinesComponent::DestroyUnreferencedForeignMeshComponents), true));
+
+		return true;
+	});
 }
 #endif
 
@@ -1156,6 +1171,76 @@ void ULandscapeSplinesComponent::RemoveForeignMeshComponent(ULandscapeSplineCont
 	}
 }
 
+void ULandscapeSplinesComponent::ForEachUnreferencedForeignMeshComponent(TFunctionRef<bool(ULandscapeSplineSegment*, USplineMeshComponent *)> Func)
+{
+	TArray<UObject*> ObjectsWithOuter;
+	GetObjectsWithOuter(GetOwner(), ObjectsWithOuter);
+	for (UObject* Obj : ObjectsWithOuter)
+	{
+		if (USplineMeshComponent* SplineMeshComponent = Cast<USplineMeshComponent>(Obj))
+		{
+            // Find Mesh in the Spline Component local map
+			if (TLazyObjectPtr<UObject>* ForeignOwner = MeshComponentForeignOwnersMap.Find(SplineMeshComponent))
+			{
+				if (ULandscapeSplineSegment* OwnerForMesh = Cast<ULandscapeSplineSegment>(ForeignOwner->Get()))
+				{
+                    // Try to find the mesh through the global map from the owner we found in local map
+					TArray<USplineMeshComponent*> SplineMeshComponents = GetForeignMeshComponents(OwnerForMesh);     
+					if (!SplineMeshComponents.Contains(SplineMeshComponent))
+					{
+						// It wasn't found. Link is broken, cleanup is needed.
+						if (Func(OwnerForMesh, SplineMeshComponent))
+						{
+							break;
+						}
+					}
+				}
+			}
+		}
+	}
+}
+
+void ULandscapeSplinesComponent::DestroyUnreferencedForeignMeshComponents()
+{
+	ForEachUnreferencedForeignMeshComponent([this](ULandscapeSplineSegment* Segment, USplineMeshComponent* SplineMeshComponent)
+	{
+		SplineMeshComponent->DestroyComponent();
+		// We should be removing something here
+		verify(MeshComponentForeignOwnersMap.Remove(SplineMeshComponent));
+		TArray<UWorld*> EmptyWorlds;
+		for (auto Pair : ForeignWorldSplineDataMap)
+		{
+			for (int32 i = Pair.Value.ForeignSplineSegmentData.Num() - 1; i >= 0; --i)
+			{
+				FForeignSplineSegmentData& SegmentData = Pair.Value.ForeignSplineSegmentData[i];
+				if (SegmentData.Identifier == Segment)
+				{
+					int32 RemoveCount = SegmentData.MeshComponents.Remove(SplineMeshComponent);
+					// if remove count is not 0, then we are expecting worlds not to match.
+					check(RemoveCount == 0 || Segment->GetTypedOuter<UWorld>() != Pair.Key.Get());
+				}
+
+				if (SegmentData.MeshComponents.Num() == 0)
+				{
+					Pair.Value.ForeignSplineSegmentData.RemoveSingle(SegmentData);
+				}
+
+				if (Pair.Value.ForeignSplineSegmentData.Num() == 0 && Pair.Value.ForeignControlPointData.Num() == 0)
+				{
+					EmptyWorlds.Add(Pair.Key.Get());
+				}
+			}
+		}
+
+		for (UWorld* EmptyWorld : EmptyWorlds)
+		{
+			ForeignWorldSplineDataMap.Remove(EmptyWorld);
+		}
+
+		return false;
+	});
+}
+
 void ULandscapeSplinesComponent::DestroyOrphanedForeignMeshComponents(UWorld* OwnerWorld)
 {
 	auto* ForeignWorldSplineData = ForeignWorldSplineDataMap.Find(OwnerWorld);
