Add support to control spline max draw distance in the main pass for RVT splines.
This allows us to cull RVT splines earlier knowing that just the landscape with RVT is good enough to show roads etc.
We already have the possiblity not to render the splines in the main pass.
Buit the reason to optionally keep rendering splines in the main pass at near distance is that they can keep the impression of a more flattened surface than landscape.

diff --git a/Engine/Source/Runtime/Engine/Classes/Components/PrimitiveComponent.h b/Engine/Source/Runtime/Engine/Classes/Components/PrimitiveComponent.h
--- a/Engine/Source/Runtime/Engine/Classes/Components/PrimitiveComponent.h
+++ b/Engine/Source/Runtime/Engine/Classes/Components/PrimitiveComponent.h
@@ -599,6 +599,8 @@ public:
 	virtual TArray<URuntimeVirtualTexture*> const& GetRuntimeVirtualTextures() const { return RuntimeVirtualTextures; }
 	/** Get the runtime virtual texture pass settings. */
 	virtual ERuntimeVirtualTextureMainPassType GetVirtualTextureRenderPassType() const { return VirtualTextureRenderPassType; }
+	/** Get the max draw distance to use in the main pass when also rendering to a runtime virtual texture. This is combined with the other max draw distance settings. */
+	virtual float GetVirtualTextureMainPassMaxDrawDistance() const { return 0.f; }
 
 	/** Used by the renderer, to identify a component across re-registers. */
 	FPrimitiveComponentId ComponentId;
diff --git a/Engine/Source/Runtime/Engine/Classes/Components/SplineMeshComponent.h b/Engine/Source/Runtime/Engine/Classes/Components/SplineMeshComponent.h
--- a/Engine/Source/Runtime/Engine/Classes/Components/SplineMeshComponent.h
+++ b/Engine/Source/Runtime/Engine/Classes/Components/SplineMeshComponent.h
@@ -141,6 +141,13 @@ class ENGINE_API USplineMeshComponent : public UStaticMeshComponent, public IInt
 	UPROPERTY(EditAnywhere, Category=SplineMesh)
 	TEnumAsByte<ESplineMeshAxis::Type> ForwardAxis;
 
+	/** 
+	 * The max draw distance to use in the main pass when also rendering to a runtime virtual texture. 
+	 * This is only exposed to the user through the same setting on ULandscapeSplineSegment. 
+	 */
+	UPROPERTY()
+	float VirtualTextureMainPassMaxDrawDistance = 0.f;
+
 #if WITH_EDITORONLY_DATA
 	UPROPERTY(transient)
 	uint8 bSelected:1;
@@ -178,6 +185,7 @@ public:
 	}
 #endif
 	virtual bool DoCustomNavigableGeometryExport(FNavigableGeometryExport& GeomExport) const override;
+	virtual float GetVirtualTextureMainPassMaxDrawDistance() const override { return VirtualTextureMainPassMaxDrawDistance; }
 	//End UPrimitiveComponent Interface
 
 	//Begin UStaticMeshComponent Interface
diff --git a/Engine/Source/Runtime/Engine/Private/PrimitiveSceneProxy.cpp b/Engine/Source/Runtime/Engine/Private/PrimitiveSceneProxy.cpp
--- a/Engine/Source/Runtime/Engine/Private/PrimitiveSceneProxy.cpp
+++ b/Engine/Source/Runtime/Engine/Private/PrimitiveSceneProxy.cpp
@@ -247,6 +247,12 @@ FPrimitiveSceneProxy::FPrimitiveSceneProxy(const UPrimitiveComponent* InComponen
 		bRenderInMainPass = false;
 	}
 
+	// Modify max draw distance for main pass if we are using virtual texturing
+	if (bUseVirtualTexture && bRenderInMainPass && InComponent->GetVirtualTextureMainPassMaxDrawDistance() > 0.f)
+	{
+		MaxDrawDistance = FMath::Min(MaxDrawDistance, InComponent->GetVirtualTextureMainPassMaxDrawDistance());
+	}
+
 #if WITH_EDITOR
 	const bool bGetDebugMaterials = true;
 	InComponent->GetUsedMaterials(UsedMaterialsForVerification, bGetDebugMaterials);
diff --git a/Engine/Source/Runtime/Landscape/Classes/ControlPointMeshComponent.h b/Engine/Source/Runtime/Landscape/Classes/ControlPointMeshComponent.h
--- a/Engine/Source/Runtime/Landscape/Classes/ControlPointMeshComponent.h
+++ b/Engine/Source/Runtime/Landscape/Classes/ControlPointMeshComponent.h
@@ -17,6 +17,13 @@ class UControlPointMeshComponent : public UStaticMeshComponent
 	uint32 bSelected:1;
 #endif
 
+	/** 
+	 * The max draw distance to use in the main pass when also rendering to a runtime virtual texture. 
+	 * This is only exposed to the user through the same setting on ULandscapeSplineControlPoint. 
+	 */
+	UPROPERTY()
+	float VirtualTextureMainPassMaxDrawDistance = 0.f;
+
 	//Begin UPrimitiveComponent Interface
 #if WITH_EDITOR
 	virtual bool ShouldRenderSelected() const override
@@ -24,5 +31,6 @@ class UControlPointMeshComponent : public UStaticMeshComponent
 		return Super::ShouldRenderSelected() || bSelected;
 	}
 #endif
+	virtual float GetVirtualTextureMainPassMaxDrawDistance() const override { return VirtualTextureMainPassMaxDrawDistance; }
 	//End UPrimitiveComponent Interface
 };
diff --git a/Engine/Source/Runtime/Landscape/Classes/LandscapeSplineControlPoint.h b/Engine/Source/Runtime/Landscape/Classes/LandscapeSplineControlPoint.h
--- a/Engine/Source/Runtime/Landscape/Classes/LandscapeSplineControlPoint.h
+++ b/Engine/Source/Runtime/Landscape/Classes/LandscapeSplineControlPoint.h
@@ -173,6 +173,10 @@ class ULandscapeSplineControlPoint : public UObject
 	UPROPERTY(EditAnywhere, AdvancedDisplay, BlueprintReadOnly, Category = VirtualTexture, meta = (DisplayName = "Virtual Texture Skip Mips", UIMin = "0", UIMax = "15"))
 	int32 VirtualTextureCullMips = 0;
 
+	/** Desired cull distance in the main pass if we are rendering to both the virtual texture AND the main pass. A value of 0 has no effect. */
+	UPROPERTY(EditAnywhere, AdvancedDisplay, BlueprintReadOnly, Category = VirtualTexture, meta = (DisplayName = "Max Draw Distance in Main Pass"))
+	float VirtualTextureMainPassMaxDrawDistance = 0.f;
+
 	/** Render to the main pass based on the virtual texture settings. */
 	UPROPERTY(EditAnywhere, BlueprintReadOnly, Category = VirtualTexture, meta = (DisplayName = "Virtual Texture Pass Type"))
 	ERuntimeVirtualTextureMainPassType VirtualTextureRenderPassType = ERuntimeVirtualTextureMainPassType::Exclusive;
diff --git a/Engine/Source/Runtime/Landscape/Classes/LandscapeSplineSegment.h b/Engine/Source/Runtime/Landscape/Classes/LandscapeSplineSegment.h
--- a/Engine/Source/Runtime/Landscape/Classes/LandscapeSplineSegment.h
+++ b/Engine/Source/Runtime/Landscape/Classes/LandscapeSplineSegment.h
@@ -275,6 +275,10 @@ class ULandscapeSplineSegment : public UObject
 	UPROPERTY(EditAnywhere, AdvancedDisplay, BlueprintReadOnly, Category = VirtualTexture, meta = (DisplayName = "Virtual Texture Skip Mips", UIMin = "0", UIMax = "15"))
 	int32 VirtualTextureCullMips = 0;
 
+	/** Desired cull distance in the main pass if we are rendering to both the virtual texture AND the main pass. A value of 0 has no effect. */
+	UPROPERTY(EditAnywhere, AdvancedDisplay, BlueprintReadOnly, Category = VirtualTexture, meta = (DisplayName = "Max Draw Distance in Main Pass"))
+	float VirtualTextureMainPassMaxDrawDistance = 0.f;
+
 	/** Render to the main pass based on the virtual texture settings. */
 	UPROPERTY(EditAnywhere, BlueprintReadOnly, Category = VirtualTexture, meta = (DisplayName = "Virtual Texture Pass Type"))
 	ERuntimeVirtualTextureMainPassType VirtualTextureRenderPassType = ERuntimeVirtualTextureMainPassType::Exclusive;
diff --git a/Engine/Source/Runtime/Landscape/Private/LandscapeSplines.cpp b/Engine/Source/Runtime/Landscape/Private/LandscapeSplines.cpp
--- a/Engine/Source/Runtime/Landscape/Private/LandscapeSplines.cpp
+++ b/Engine/Source/Runtime/Landscape/Private/LandscapeSplines.cpp
@@ -1860,6 +1860,13 @@ void ULandscapeSplineControlPoint::UpdateSplinePoints(bool bUpdateCollision, boo
 			MeshComponent->MarkRenderStateDirty();
 		}
 
+		if (MeshComponent->VirtualTextureMainPassMaxDrawDistance != VirtualTextureMainPassMaxDrawDistance)
+		{
+			MeshComponent->Modify();
+			MeshComponent->VirtualTextureMainPassMaxDrawDistance = VirtualTextureMainPassMaxDrawDistance;
+			MeshComponent->MarkRenderStateDirty();
+		}
+
 		if (MeshComponent->VirtualTextureRenderPassType != VirtualTextureRenderPassType)
 		{
 			MeshComponent->Modify();
@@ -2846,6 +2853,7 @@ void ULandscapeSplineSegment::UpdateSplinePoints(bool bUpdateCollision, bool bUp
 			MeshComponent->RuntimeVirtualTextures = RuntimeVirtualTextures;
 			MeshComponent->VirtualTextureLodBias = VirtualTextureLodBias;
 			MeshComponent->VirtualTextureCullMips = VirtualTextureCullMips;
+			MeshComponent->VirtualTextureMainPassMaxDrawDistance = VirtualTextureMainPassMaxDrawDistance;
 			MeshComponent->VirtualTextureRenderPassType = VirtualTextureRenderPassType;
 
 			MeshComponent->SetCastShadow(bCastShadow);
