Landscape Spline: Add more Map Checks for Orphaned/Unreferenced ControlPoint MeshComponents

diff --git a/Engine/Source/Runtime/Landscape/Classes/LandscapeSplinesComponent.h b/Engine/Source/Runtime/Landscape/Classes/LandscapeSplinesComponent.h
--- a/Engine/Source/Runtime/Landscape/Classes/LandscapeSplinesComponent.h
+++ b/Engine/Source/Runtime/Landscape/Classes/LandscapeSplinesComponent.h
@@ -179,10 +179,11 @@ public:
 	virtual void RemoveAllForeignMeshComponents(ULandscapeSplineSegment* Owner);
 	virtual void AddForeignMeshComponent(ULandscapeSplineControlPoint* Owner, UControlPointMeshComponent* Component);
 	virtual void RemoveForeignMeshComponent(ULandscapeSplineControlPoint* Owner, UControlPointMeshComponent* Component);
-	virtual void DestroyOrphanedForeignMeshComponents(UWorld* OwnerWorld);
+	virtual void DestroyOrphanedForeignSplineMeshComponents(UWorld* OwnerWorld);
+	virtual void DestroyOrphanedForeignControlPointMeshComponents(UWorld* OwnerWorld);
 
 	void DestroyUnreferencedForeignMeshComponents();
-	void ForEachUnreferencedForeignMeshComponent(TFunctionRef<bool(ULandscapeSplineSegment*, USplineMeshComponent*)> Func);
+	void ForEachUnreferencedForeignMeshComponent(TFunctionRef<bool(ULandscapeSplineSegment*, USplineMeshComponent*, ULandscapeSplineControlPoint*, UControlPointMeshComponent*)> Func);
 	
 	virtual UControlPointMeshComponent*   GetForeignMeshComponent(ULandscapeSplineControlPoint* Owner);
 	virtual TArray<USplineMeshComponent*> GetForeignMeshComponents(ULandscapeSplineSegment* Owner);
diff --git a/Engine/Source/Runtime/Landscape/Private/LandscapeSplines.cpp b/Engine/Source/Runtime/Landscape/Private/LandscapeSplines.cpp
--- a/Engine/Source/Runtime/Landscape/Private/LandscapeSplines.cpp
+++ b/Engine/Source/Runtime/Landscape/Private/LandscapeSplines.cpp
@@ -637,7 +637,8 @@ void ULandscapeSplinesComponent::AutoFixMeshComponentErrors(UWorld* OtherWorld)
 
 	if (StreamingSplinesComponent)
 	{
-		StreamingSplinesComponent->DestroyOrphanedForeignMeshComponents(ThisOuterWorld);
+		StreamingSplinesComponent->DestroyOrphanedForeignSplineMeshComponents(ThisOuterWorld);
+		StreamingSplinesComponent->DestroyOrphanedForeignControlPointMeshComponents(ThisOuterWorld);
 	}
 }
 
@@ -765,9 +766,30 @@ void ULandscapeSplinesComponent::CheckForErrors()
 
 				FMessageLog("MapCheck").Error()
 					->AddToken(FUObjectToken::Create(GetOwner()))
-					->AddToken(FTextToken::Create(FText::Format(LOCTEXT("MapCheck_Message_OrphanedMeshes", "{MeshMap} contains orphaned meshes due to mismatch with landscape splines in {SplineMap}"), Arguments)))
-					->AddToken(FActionToken::Create(LOCTEXT("MapCheck_ActionName_OrphanedMeshes", "Clean up orphaned meshes"), FText(),
-					FOnActionTokenExecuted::CreateUObject(this, &ULandscapeSplinesComponent::DestroyOrphanedForeignMeshComponents, ForeignWorld), true));
+					->AddToken(FTextToken::Create(FText::Format(LOCTEXT("MapCheck_Message_OrphanedSplineMeshes", "{MeshMap} contains orphaned spline meshes due to mismatch with landscape splines in {SplineMap}"), Arguments)))
+					->AddToken(FActionToken::Create(LOCTEXT("MapCheck_ActionName_OrphanedSplineMeshes", "Clean up orphaned spline meshes"), FText(),
+						FOnActionTokenExecuted::CreateUObject(this, &ULandscapeSplinesComponent::DestroyOrphanedForeignSplineMeshComponents, ForeignWorld), true));
+
+				break;
+			}
+		}
+
+		for (auto& ForeignControlPointData : ForeignWorldSplineData.ForeignControlPointData)
+		{
+			const ULandscapeSplineControlPoint* ForeignControlPoint = ForeignControlPointData.Identifier.Get();
+
+			// No such control point
+			if (!ForeignControlPoint)
+			{
+				FFormatNamedArguments Arguments;
+				Arguments.Add(TEXT("MeshMap"), FText::FromName(ThisOuterWorld->GetFName()));
+				Arguments.Add(TEXT("SplineMap"), FText::FromName(ForeignWorld->GetFName()));
+
+				FMessageLog("MapCheck").Error()
+					->AddToken(FUObjectToken::Create(GetOwner()))
+					->AddToken(FTextToken::Create(FText::Format(LOCTEXT("MapCheck_Message_OrphanedControlPointMeshes", "{MeshMap} contains orphaned control point meshes due to mismatch with landscape control points in {SplineMap}"), Arguments)))
+					->AddToken(FActionToken::Create(LOCTEXT("MapCheck_ActionName_OrphanedControlPointMeshes", "Clean up orphaned control point meshes"), FText(),
+						FOnActionTokenExecuted::CreateUObject(this, &ULandscapeSplinesComponent::DestroyOrphanedForeignControlPointMeshComponents, ForeignWorld), true));
 
 				break;
 			}
@@ -775,7 +797,7 @@ void ULandscapeSplinesComponent::CheckForErrors()
 	}
 
 	// Find Unreferenced Foreign Mesh Components
-	ForEachUnreferencedForeignMeshComponent([this, ThisOuterWorld](ULandscapeSplineSegment* Segment, USplineMeshComponent* SplineMeshComponent)
+	ForEachUnreferencedForeignMeshComponent([this, ThisOuterWorld](ULandscapeSplineSegment*, USplineMeshComponent*, ULandscapeSplineControlPoint*, UControlPointMeshComponent*)
 	{
 		FFormatNamedArguments Arguments;
 		Arguments.Add(TEXT("MeshMap"), FText::FromName(ThisOuterWorld->GetFName()));
@@ -783,7 +805,7 @@ void ULandscapeSplinesComponent::CheckForErrors()
 		FMessageLog("MapCheck").Error()
 			->AddToken(FUObjectToken::Create(GetOwner()))
 			->AddToken(FTextToken::Create(FText::Format(LOCTEXT("MapCheck_Message_UnreferencedMeshes", "{MeshMap} contains meshes with no owners"), Arguments)))
-			->AddToken(FActionToken::Create(LOCTEXT("MapCheck_ActionName_UnreferencedMeshes", "Clean up meshes that are not referenced by their owner segments"), FText(),
+			->AddToken(FActionToken::Create(LOCTEXT("MapCheck_ActionName_UnreferencedMeshes", "Clean up meshes that are not referenced by their owner segments/control points"), FText(),
 				FOnActionTokenExecuted::CreateUObject(this, &ULandscapeSplinesComponent::DestroyUnreferencedForeignMeshComponents), true));
 
 		return true;
@@ -1171,7 +1193,7 @@ void ULandscapeSplinesComponent::RemoveForeignMeshComponent(ULandscapeSplineCont
 	}
 }
 
-void ULandscapeSplinesComponent::ForEachUnreferencedForeignMeshComponent(TFunctionRef<bool(ULandscapeSplineSegment*, USplineMeshComponent *)> Func)
+void ULandscapeSplinesComponent::ForEachUnreferencedForeignMeshComponent(TFunctionRef<bool(ULandscapeSplineSegment*, USplineMeshComponent*, ULandscapeSplineControlPoint*, UControlPointMeshComponent*)> Func)
 {
 	TArray<UObject*> ObjectsWithOuter;
 	GetObjectsWithOuter(GetOwner(), ObjectsWithOuter);
@@ -1189,7 +1211,24 @@ void ULandscapeSplinesComponent::ForEachUnreferencedForeignMeshComponent(TFuncti
 					if (!SplineMeshComponents.Contains(SplineMeshComponent))
 					{
 						// It wasn't found. Link is broken, cleanup is needed.
-						if (Func(OwnerForMesh, SplineMeshComponent))
+						if (Func(OwnerForMesh, SplineMeshComponent, nullptr, nullptr))
+						{
+							break;
+						}
+					}
+				}
+			} 
+		}
+		else if (UControlPointMeshComponent* ControlPointMeshComponent = Cast<UControlPointMeshComponent>(Obj))
+		{
+			if (TLazyObjectPtr<UObject>* ForeignOwner = MeshComponentForeignOwnersMap.Find(SplineMeshComponent))
+			{
+				if (ULandscapeSplineControlPoint* OwnerForMesh = Cast<ULandscapeSplineControlPoint>(ForeignOwner->Get()))
+				{
+					UControlPointMeshComponent* ForeignControlPointMeshComponent = GetForeignMeshComponent(OwnerForMesh);
+					if (ControlPointMeshComponent != ForeignControlPointMeshComponent)
+					{
+						if (Func(nullptr, nullptr, OwnerForMesh, ControlPointMeshComponent))
 						{
 							break;
 						}
@@ -1202,32 +1241,64 @@ void ULandscapeSplinesComponent::ForEachUnreferencedForeignMeshComponent(TFuncti
 
 void ULandscapeSplinesComponent::DestroyUnreferencedForeignMeshComponents()
 {
-	ForEachUnreferencedForeignMeshComponent([this](ULandscapeSplineSegment* Segment, USplineMeshComponent* SplineMeshComponent)
+	ForEachUnreferencedForeignMeshComponent([this](ULandscapeSplineSegment* Segment, USplineMeshComponent* SplineMeshComponent, ULandscapeSplineControlPoint* ControlPoint, UControlPointMeshComponent* ControlPointMeshComponent)
 	{
-		SplineMeshComponent->DestroyComponent();
-		// We should be removing something here
-		verify(MeshComponentForeignOwnersMap.Remove(SplineMeshComponent));
 		TArray<UWorld*> EmptyWorlds;
-		for (auto Pair : ForeignWorldSplineDataMap)
+
+		if (SplineMeshComponent != nullptr)
 		{
-			for (int32 i = Pair.Value.ForeignSplineSegmentData.Num() - 1; i >= 0; --i)
+			SplineMeshComponent->DestroyComponent();
+			
+			// We should be removing something here
+			verify(MeshComponentForeignOwnersMap.Remove(SplineMeshComponent));
+
+			for (auto Pair : ForeignWorldSplineDataMap)
 			{
-				FForeignSplineSegmentData& SegmentData = Pair.Value.ForeignSplineSegmentData[i];
-				if (SegmentData.Identifier == Segment)
+				for (int32 i = Pair.Value.ForeignSplineSegmentData.Num() - 1; i >= 0; --i)
 				{
-					int32 RemoveCount = SegmentData.MeshComponents.Remove(SplineMeshComponent);
-					// if remove count is not 0, then we are expecting worlds not to match.
-					check(RemoveCount == 0 || Segment->GetTypedOuter<UWorld>() != Pair.Key.Get());
-				}
+					FForeignSplineSegmentData& SegmentData = Pair.Value.ForeignSplineSegmentData[i];
+					if (SegmentData.Identifier == Segment)
+					{
+						int32 RemoveCount = SegmentData.MeshComponents.Remove(SplineMeshComponent);
+						// if remove count is not 0, then we are expecting worlds not to match.
+						check(RemoveCount == 0 || Segment->GetTypedOuter<UWorld>() != Pair.Key.Get());
+					}
 
-				if (SegmentData.MeshComponents.Num() == 0)
-				{
-					Pair.Value.ForeignSplineSegmentData.RemoveSingle(SegmentData);
+					if (SegmentData.MeshComponents.Num() == 0)
+					{
+						Pair.Value.ForeignSplineSegmentData.RemoveSingle(SegmentData);
+					}
+
+					if (Pair.Value.IsEmpty())
+					{
+						EmptyWorlds.Add(Pair.Key.Get());
+					}
 				}
+			}
+		}
+
+		if (ControlPointMeshComponent != nullptr)
+		{
+			ControlPointMeshComponent->DestroyComponent();
+			
+			// We should be removing something here
+			verify(MeshComponentForeignOwnersMap.Remove(ControlPointMeshComponent));
 
-				if (Pair.Value.ForeignSplineSegmentData.Num() == 0 && Pair.Value.ForeignControlPointData.Num() == 0)
+			for (auto Pair : ForeignWorldSplineDataMap)
+			{
+				for (int32 i = Pair.Value.ForeignControlPointData.Num() - 1; i >= 0; --i)
 				{
-					EmptyWorlds.Add(Pair.Key.Get());
+					FForeignControlPointData& ControlPointData = Pair.Value.ForeignControlPointData[i];
+					if (ControlPointData.Identifier == ControlPoint && ControlPointData.MeshComponent == ControlPointMeshComponent)
+					{
+						check(ControlPoint->GetTypedOuter<UWorld>() != Pair.Key.Get());
+						Pair.Value.ForeignControlPointData.RemoveSingle(ControlPointData);
+					}
+
+					if (Pair.Value.IsEmpty())
+					{
+						EmptyWorlds.Add(Pair.Key.Get());
+					}
 				}
 			}
 		}
@@ -1241,7 +1312,7 @@ void ULandscapeSplinesComponent::DestroyUnreferencedForeignMeshComponents()
 	});
 }
 
-void ULandscapeSplinesComponent::DestroyOrphanedForeignMeshComponents(UWorld* OwnerWorld)
+void ULandscapeSplinesComponent::DestroyOrphanedForeignSplineMeshComponents(UWorld* OwnerWorld)
 {
 	auto* ForeignWorldSplineData = ForeignWorldSplineDataMap.Find(OwnerWorld);
 
@@ -1273,6 +1344,31 @@ void ULandscapeSplinesComponent::DestroyOrphanedForeignMeshComponents(UWorld* Ow
 	}
 }
 
+void ULandscapeSplinesComponent::DestroyOrphanedForeignControlPointMeshComponents(UWorld* OwnerWorld)
+{
+	auto* ForeignWorldSplineData = ForeignWorldSplineDataMap.Find(OwnerWorld);
+
+	if (ForeignWorldSplineData)
+	{
+		for (int32 i = ForeignWorldSplineData->ForeignControlPointData.Num() - 1; i >= 0; --i)
+		{
+			FForeignControlPointData& ControlPointData = ForeignWorldSplineData->ForeignControlPointData[i];
+			const auto& ForeignControlPoint = ControlPointData.Identifier;
+
+			if (!ForeignControlPoint)
+			{
+				ControlPointData.MeshComponent->DestroyComponent();
+				ForeignWorldSplineData->ForeignControlPointData.RemoveSingle(ControlPointData);
+			}
+		}
+
+		if (ForeignWorldSplineData->IsEmpty())
+		{
+			verifySlow(ForeignWorldSplineDataMap.Remove(OwnerWorld) == 1);
+		}
+	}
+}
+
 UControlPointMeshComponent* ULandscapeSplinesComponent::GetForeignMeshComponent(ULandscapeSplineControlPoint* Owner)
 {
 	UWorld* OwnerWorld = Owner->GetTypedOuter<UWorld>();
