Landscape Splines: Fix Undo/Redo

diff --git a/Engine/Source/Editor/LandscapeEditor/Private/LandscapeEdMode.cpp b/Engine/Source/Editor/LandscapeEditor/Private/LandscapeEdMode.cpp
--- a/Engine/Source/Editor/LandscapeEditor/Private/LandscapeEdMode.cpp
+++ b/Engine/Source/Editor/LandscapeEditor/Private/LandscapeEdMode.cpp
@@ -634,7 +634,7 @@ void FEdModeLandscape::Enter()
 			{
 				if (Landscape->GetLandscapeSplinesReservedLayer())
 				{
-					Landscape->UpdateLandscapeSplines();
+					Landscape->RequestSplineLayerUpdate();
 				}
 				Landscape->RequestLayersContentUpdateForceAll();
 			}
diff --git a/Engine/Source/Editor/LandscapeEditor/Private/LandscapeEditorDetailCustomization_MiscTools.cpp b/Engine/Source/Editor/LandscapeEditor/Private/LandscapeEditorDetailCustomization_MiscTools.cpp
--- a/Engine/Source/Editor/LandscapeEditor/Private/LandscapeEditorDetailCustomization_MiscTools.cpp
+++ b/Engine/Source/Editor/LandscapeEditor/Private/LandscapeEditorDetailCustomization_MiscTools.cpp
@@ -305,6 +305,7 @@ FReply FLandscapeEditorDetailCustomization_MiscTools::OnApplyAllSplinesButtonCli
 	FEdModeLandscape* LandscapeEdMode = GetEditorMode();
 	if (LandscapeEdMode)
 	{
+		FScopedTransaction Transaction(LOCTEXT("LandscapeSpline_ApplyAllSplines", "Apply All Splines to Landscape"));
 		LandscapeEdMode->UpdateLandscapeSplines(false);
 	}
 	return FReply::Handled();
@@ -315,6 +316,7 @@ FReply FLandscapeEditorDetailCustomization_MiscTools::OnApplySelectedSplinesButt
 	FEdModeLandscape* LandscapeEdMode = GetEditorMode();
 	if (LandscapeEdMode)
 	{
+		FScopedTransaction Transaction(LOCTEXT("LandscapeSpline_ApplySelectedSplines", "Apply Selected Splines to Landscape"));
 		LandscapeEdMode->UpdateLandscapeSplines(true);
 	}
 
diff --git a/Engine/Source/Runtime/Landscape/Private/LandscapeSplineRaster.cpp b/Engine/Source/Runtime/Landscape/Private/LandscapeSplineRaster.cpp
--- a/Engine/Source/Runtime/Landscape/Private/LandscapeSplineRaster.cpp
+++ b/Engine/Source/Runtime/Landscape/Private/LandscapeSplineRaster.cpp
@@ -572,8 +572,6 @@ bool ULandscapeInfo::ApplySplinesInternal(bool bOnlySelected, ALandscapeProxy* P
 		return false;
 	}
 
-	FScopedTransaction Transaction(LOCTEXT("LandscapeSpline_ApplySplines", "Apply Splines to Landscape"));
-
 	const FTransform SplineToLandscape = Proxy->SplineComponent->GetComponentTransform().GetRelativeTransform(Proxy->LandscapeActorToWorld());
 
 	FLandscapeEditDataInterface LandscapeEdit(this);
diff --git a/Engine/Source/Runtime/Landscape/Private/LandscapeSplines.cpp b/Engine/Source/Runtime/Landscape/Private/LandscapeSplines.cpp
--- a/Engine/Source/Runtime/Landscape/Private/LandscapeSplines.cpp
+++ b/Engine/Source/Runtime/Landscape/Private/LandscapeSplines.cpp
@@ -2130,7 +2130,15 @@ void ULandscapeSplineControlPoint::PostEditUndo()
 	Super::PostEditUndo();
 	bHackIsUndoingSplines = false;
 
-	GetOuterULandscapeSplinesComponent()->MarkRenderStateDirty();
+	ULandscapeSplinesComponent* SplineComponent = GetOuterULandscapeSplinesComponent();
+	SplineComponent->MarkRenderStateDirty();
+
+	ALandscapeProxy* OuterLandscape = Cast<ALandscapeProxy>(SplineComponent->GetOwner());
+	ALandscape* Landscape = OuterLandscape ? OuterLandscape->GetLandscapeActor() : nullptr;
+	if (Landscape)
+	{
+		Landscape->RequestSplineLayerUpdate();
+	}
 }
 
 void ULandscapeSplineControlPoint::PostDuplicate(bool bDuplicateForPIE)
@@ -3088,7 +3096,15 @@ void ULandscapeSplineSegment::PostEditUndo()
 	Super::PostEditUndo();
 	bHackIsUndoingSplines = false;
 
-	GetOuterULandscapeSplinesComponent()->MarkRenderStateDirty();
+	ULandscapeSplinesComponent* SplineComponent = GetOuterULandscapeSplinesComponent();
+	SplineComponent->MarkRenderStateDirty();
+
+	ALandscapeProxy* OuterLandscape = Cast<ALandscapeProxy>(SplineComponent->GetOwner());
+	ALandscape* Landscape = OuterLandscape ? OuterLandscape->GetLandscapeActor() : nullptr;
+	if (Landscape)
+	{
+		Landscape->RequestSplineLayerUpdate();
+	}
 }
 
 void ULandscapeSplineSegment::PostDuplicate(bool bDuplicateForPIE)
