Landscape Spline: Add missing map cleanup

diff --git a/Engine/Source/Runtime/Landscape/Private/LandscapeSplines.cpp b/Engine/Source/Runtime/Landscape/Private/LandscapeSplines.cpp
--- a/Engine/Source/Runtime/Landscape/Private/LandscapeSplines.cpp
+++ b/Engine/Source/Runtime/Landscape/Private/LandscapeSplines.cpp
@@ -1359,6 +1359,7 @@ void ULandscapeSplinesComponent::DestroyOrphanedForeignControlPointMeshComponent
 			{
 				ControlPointData.MeshComponent->DestroyComponent();
 				ForeignWorldSplineData->ForeignControlPointData.RemoveSingle(ControlPointData);
+				MeshComponentForeignOwnersMap.Remove(ControlPointData.MeshComponent);
 			}
 		}
 
