Lanscape: Remove some useless map dirtyness caused by [REDACTED_PERSON] when Snapping Foliage to Landscape

#rnx

diff --git a/Engine/Source/Runtime/Landscape/Private/LandscapeCollision.cpp b/Engine/Source/Runtime/Landscape/Private/LandscapeCollision.cpp
--- a/Engine/Source/Runtime/Landscape/Private/LandscapeCollision.cpp
+++ b/Engine/Source/Runtime/Landscape/Private/LandscapeCollision.cpp
@@ -1325,7 +1325,6 @@ void ULandscapeHeightfieldCollisionComponent::SnapFoliageInstances(const FBox& I
 				float TraceExtentSize = Bounds.SphereRadius * 2.f + 10.f; // extend a little
 				FVector TraceVector = GetOwner()->GetRootComponent()->GetComponentTransform().GetUnitAxis(EAxis::Z) * TraceExtentSize;
 
-				bool bFirst = true;
 				TArray<int32> InstancesToRemove;
 				TSet<UHierarchicalInstancedStaticMeshComponent*> AffectedFoliageComponents;
 
@@ -1340,12 +1339,6 @@ void ULandscapeHeightfieldCollisionComponent::SnapFoliageInstances(const FBox& I
 
 					if (InInstanceBox.IsInside(TestLocation))
 					{
-						if (bFirst)
-						{
-							bFirst = false;
-							Modify();
-						}
-
 						FVector Start = TestLocation + TraceVector;
 						FVector End = TestLocation - TraceVector;
 
