Landscape Splines: Safeguard null MeshComponents

#rnx

diff --git a/Engine/Source/Runtime/Landscape/Private/LandscapeSplines.cpp b/Engine/Source/Runtime/Landscape/Private/LandscapeSplines.cpp
--- a/Engine/Source/Runtime/Landscape/Private/LandscapeSplines.cpp
+++ b/Engine/Source/Runtime/Landscape/Private/LandscapeSplines.cpp
@@ -1327,9 +1327,12 @@ void ULandscapeSplinesComponent::DestroyOrphanedForeignSplineMeshComponents(UWor
 			{
 				for (auto* MeshComponent : SegmentData.MeshComponents)
 				{
-					checkSlow(!MeshComponentForeignOwnersMap.FindRef(MeshComponent).IsValid());
-					verifySlow(MeshComponentForeignOwnersMap.Remove(MeshComponent) == 1);
-					MeshComponent->DestroyComponent();
+					if (MeshComponent)
+					{
+						checkSlow(!MeshComponentForeignOwnersMap.FindRef(MeshComponent).IsValid());
+						verifySlow(MeshComponentForeignOwnersMap.Remove(MeshComponent) == 1);
+						MeshComponent->DestroyComponent();
+					}
 				}
 				SegmentData.MeshComponents.Empty();
 
@@ -1355,7 +1358,7 @@ void ULandscapeSplinesComponent::DestroyOrphanedForeignControlPointMeshComponent
 			FForeignControlPointData& ControlPointData = ForeignWorldSplineData->ForeignControlPointData[i];
 			const auto& ForeignControlPoint = ControlPointData.Identifier;
 
-			if (!ForeignControlPoint)
+			if (!ForeignControlPoint && ControlPointData.MeshComponent)
 			{
 				ControlPointData.MeshComponent->DestroyComponent();
 				ForeignWorldSplineData->ForeignControlPointData.RemoveSingle(ControlPointData);
