Landscape Layers: Dirty map on Load Fixes (will need a full map save)
- Remove Texture2D Landscape check and instead use the bForceMiplevelsToBeResident flag.
- Add Blueprint Brush Render Dependency support to make sure those are also streamed in before calling Render

diff --git a/Engine/Source/Runtime/Engine/Private/Texture2D.cpp b/Engine/Source/Runtime/Engine/Private/Texture2D.cpp
--- a/Engine/Source/Runtime/Engine/Private/Texture2D.cpp
+++ b/Engine/Source/Runtime/Engine/Private/Texture2D.cpp
@@ -38,11 +38,6 @@
 #include "VT/UploadingVirtualTexture.h"
 #include "VT/VirtualTexturePoolConfig.h"
 
-#if WITH_EDITOR
-#include "Settings/EditorExperimentalSettings.h"
-#include "Landscape/Classes/LandscapeProxy.h"
-#endif
-
 UTexture2D::UTexture2D(const FObjectInitializer& ObjectInitializer)
 	: Super(ObjectInitializer)
 {
@@ -1255,14 +1250,6 @@ bool UTexture2D::ShouldMipLevelsBeForcedResident() const
 		return true;
 	}
 
-#if WITH_EDITOR
-	if (GIsEditor && (LODGroup == TEXTUREGROUP_Terrain_Heightmap || LODGroup == TEXTUREGROUP_Terrain_Weightmap))
-	{
-		ALandscapeProxy* Proxy = Cast<ALandscapeProxy>(GetOuter());
-		return Proxy && Proxy->HasLayersContent();
-	}
-#endif
-
 	return false;
 }
 
diff --git a/Engine/Source/Runtime/Landscape/Classes/Landscape.h b/Engine/Source/Runtime/Landscape/Classes/Landscape.h
--- a/Engine/Source/Runtime/Landscape/Classes/Landscape.h
+++ b/Engine/Source/Runtime/Landscape/Classes/Landscape.h
@@ -348,6 +348,7 @@ private:
 	void ResolveLayersWeightmapTexture(const TArray<ULandscapeComponent*>& InLandscapeComponents);
 	bool ResolveLayersTexture(class FLandscapeLayersTexture2DCPUReadBackResource* InCPUReadBackTexture, UTexture2D* InOutputTexture);
 
+	bool PrepareLayersBrushTextureResources(bool bInWaitForStreaming, bool bHeightmap) const;
 	bool PrepareLayersHeightmapTextureResources(bool bInWaitForStreaming) const;
 	bool PrepareLayersWeightmapTextureResources(bool bInWaitForStreaming) const;
 
diff --git a/Engine/Source/Runtime/Landscape/Private/LandscapeBlueprintBrushBase.cpp b/Engine/Source/Runtime/Landscape/Private/LandscapeBlueprintBrushBase.cpp
--- a/Engine/Source/Runtime/Landscape/Private/LandscapeBlueprintBrushBase.cpp
+++ b/Engine/Source/Runtime/Landscape/Private/LandscapeBlueprintBrushBase.cpp
@@ -178,6 +178,14 @@ void ALandscapeBlueprintBrushBase::Destroyed()
 #endif
 }
 
+void ALandscapeBlueprintBrushBase::GetRenderDependencies(TSet<UTexture2D *>& OutTextures)
+{
+	TArray<UTexture2D*> BPDependencies;
+	GetBlueprintRenderDependencies(BPDependencies);
+
+	OutTextures.Append(BPDependencies);
+}
+
 void ALandscapeBlueprintBrushBase::SetOwningLandscape(ALandscape* InOwningLandscape)
 {
 #if WITH_EDITORONLY_DATA
diff --git a/Engine/Source/Runtime/Landscape/Private/LandscapeEditLayers.cpp b/Engine/Source/Runtime/Landscape/Private/LandscapeEditLayers.cpp
--- a/Engine/Source/Runtime/Landscape/Private/LandscapeEditLayers.cpp
+++ b/Engine/Source/Runtime/Landscape/Private/LandscapeEditLayers.cpp
@@ -1313,6 +1313,7 @@ struct FLandscapeIsTextureFullyStreamedIn
 	bool operator()(UTexture2D* InTexture, bool bInWaitForStreaming)
 	{
 		check(InTexture);
+		InTexture->bForceMiplevelsToBeResident = true;
 		if (bInWaitForStreaming)
 		{
 			InTexture->WaitForStreaming();
@@ -2661,6 +2662,33 @@ void ALandscape::PrintLayersDebugTextureResource(const FString& InContext, FText
 	}
 }
 
+bool ALandscape::PrepareLayersBrushTextureResources(bool bInWaitForStreaming, bool bHeightmap) const
+{
+	TSet<UTexture2D*> StreamableTextures;
+	for (const FLandscapeLayer& Layer : LandscapeLayers)
+	{
+		for (const FLandscapeLayerBrush& Brush : Layer.Brushes)
+		{
+			if (ALandscapeBlueprintBrushBase* LandscapeBrush = Brush.GetBrush())
+			{
+				if ((LandscapeBrush->IsAffectingWeightmap() && !bHeightmap) || (LandscapeBrush->IsAffectingHeightmap() && bHeightmap))
+				{
+					LandscapeBrush->GetRenderDependencies(StreamableTextures);
+				}
+			}
+		}
+	}
+
+	bool bReady = true;
+	FLandscapeIsTextureFullyStreamedIn IsTextureFullyStreamedIn;
+	for (UTexture2D* Texture : StreamableTextures)
+	{
+		bReady &= IsTextureFullyStreamedIn(Texture, bInWaitForStreaming);
+	}
+
+	return bReady;
+}
+
 bool ALandscape::PrepareLayersHeightmapTextureResources(bool bInWaitForStreaming) const
 {
 	ULandscapeInfo* Info = GetLandscapeInfo();
@@ -2711,11 +2739,6 @@ bool ALandscape::PrepareLayersHeightmapTextureResources(bool bInWaitForStreaming
 		}
 	});
 
-	if (bInWaitForStreaming)
-	{
-		FlushRenderingCommands();
-	}
-
 	return IsReady;
 }
 
@@ -2729,7 +2752,16 @@ int32 ALandscape::RegenerateLayersHeightmaps(const TArray<ULandscapeComponent*>&
 	const bool bForceRender = CVarOutputLayersDebugDrawCallName.GetValueOnAnyThread() == 1;
 	const bool bSkipBrush = CVarLandscapeLayerBrushOptim.GetValueOnAnyThread() == 1 && ((HeightmapUpdateModes & AllHeightmapUpdateModes) == ELandscapeLayerUpdateMode::Update_Heightmap_Editing);
 
-	if ((HeightmapUpdateModes == 0 && !bForceRender) || Info == nullptr || !PrepareLayersHeightmapTextureResources(bInWaitForStreaming))
+	if ((HeightmapUpdateModes == 0 && !bForceRender) || Info == nullptr)
+	{
+		return 0;
+	}
+
+	const bool bHeightmapTexturesReady = PrepareLayersHeightmapTextureResources(bInWaitForStreaming);
+	const bool bHeightmap = true;
+	const bool bBrushTexturesReady = PrepareLayersBrushTextureResources(bInWaitForStreaming, bHeightmap);
+
+	if (!(bHeightmapTexturesReady && bBrushTexturesReady))
 	{
 		return 0;
 	}
@@ -3567,11 +3599,6 @@ bool ALandscape::PrepareLayersWeightmapTextureResources(bool bInWaitForStreaming
 		}
 	});
 
-	if (bInWaitForStreaming)
-	{
-		FlushRenderingCommands();
-	}
-
 	return IsReady;
 }
 
@@ -3585,7 +3612,16 @@ int32 ALandscape::RegenerateLayersWeightmaps(const TArray<ULandscapeComponent*>&
 	
 	ULandscapeInfo* Info = GetLandscapeInfo();
 
-	if ((WeightmapUpdateModes == 0 && !bForceRender) || Info == nullptr || Info->Layers.Num() == 0 || !PrepareLayersWeightmapTextureResources(bInWaitForStreaming))
+	if ((WeightmapUpdateModes == 0 && !bForceRender) || Info == nullptr || Info->Layers.Num() == 0)
+	{
+		return 0;
+	}
+
+	const bool bWeightmapTexturesReady = PrepareLayersWeightmapTextureResources(bInWaitForStreaming);
+	const bool bHeightmap = false;
+	const bool bBrushTexturesReady = PrepareLayersBrushTextureResources(bInWaitForStreaming, bHeightmap);
+
+	if (!(bWeightmapTexturesReady && bBrushTexturesReady))
 	{
 		return 0;
 	}
diff --git a/Engine/Source/Runtime/Landscape/Public/LandscapeBlueprintBrushBase.h b/Engine/Source/Runtime/Landscape/Public/LandscapeBlueprintBrushBase.h
--- a/Engine/Source/Runtime/Landscape/Public/LandscapeBlueprintBrushBase.h
+++ b/Engine/Source/Runtime/Landscape/Public/LandscapeBlueprintBrushBase.h
@@ -43,7 +43,12 @@ public:
 	UFUNCTION(BlueprintCallable, Category = "Landscape")
 	void RequestLandscapeUpdate();
 
+	UFUNCTION(BlueprintImplementableEvent)
+	void GetBlueprintRenderDependencies(TArray<UTexture2D*>& OutStreamableAssets);
+
 #if WITH_EDITOR
+	virtual void GetRenderDependencies(TSet<UTexture2D*>& OutStreamableAssets);
+
 	virtual void SetOwningLandscape(class ALandscape* InOwningLandscape);
 	class ALandscape* GetOwningLandscape() const;
 
