Landscape Layers: Add Profiling Scopes

diff --git a/Engine/Source/Runtime/Landscape/Private/LandscapeEditLayers.cpp b/Engine/Source/Runtime/Landscape/Private/LandscapeEditLayers.cpp
--- a/Engine/Source/Runtime/Landscape/Private/LandscapeEditLayers.cpp
+++ b/Engine/Source/Runtime/Landscape/Private/LandscapeEditLayers.cpp
@@ -2664,6 +2664,7 @@ void ALandscape::PrintLayersDebugTextureResource(const FString& InContext, FText
 
 bool ALandscape::PrepareLayersBrushTextureResources(bool bInWaitForStreaming, bool bHeightmap) const
 {
+	TRACE_CPUPROFILER_EVENT_SCOPE("PrepareLayersBrushTextureResources");
 	TSet<UTexture2D*> StreamableTextures;
 	for (const FLandscapeLayer& Layer : LandscapeLayers)
 	{
@@ -2691,6 +2692,7 @@ bool ALandscape::PrepareLayersBrushTextureResources(bool bInWaitForStreaming, bo
 
 bool ALandscape::PrepareLayersHeightmapTextureResources(bool bInWaitForStreaming) const
 {
+	TRACE_CPUPROFILER_EVENT_SCOPE("PrepareLayersHeightmapTextureResources");
 	ULandscapeInfo* Info = GetLandscapeInfo();
 
 	if (Info == nullptr)
@@ -2744,6 +2746,7 @@ bool ALandscape::PrepareLayersHeightmapTextureResources(bool bInWaitForStreaming
 
 int32 ALandscape::RegenerateLayersHeightmaps(const TArray<ULandscapeComponent*>& InLandscapeComponentsToRender, const TArray<ULandscapeComponent*>& InLandscapeComponentsToResolve, bool bInWaitForStreaming)
 {
+	TRACE_CPUPROFILER_EVENT_SCOPE("RegenerateLayersHeightmaps");
 	SCOPE_CYCLE_COUNTER(STAT_LandscapeLayersRegenerateHeightmaps);
 	ULandscapeInfo* Info = GetLandscapeInfo();
 
@@ -2998,9 +3001,15 @@ int32 ALandscape::RegenerateLayersHeightmaps(const TArray<ULandscapeComponent*>&
 		// We can skip Resolving if we don't need collision updating at all and that we aren't doing an intermediate render for some Landscape Tool (Flatten, Smooth). (This relies on the fact that at some point we will resolve with a full update)
 		if (bNeedsResolving)
 		{
+			if(InLandscapeComponentsToResolve.Num() > 0)
+			{
+				TRACE_CPUPROFILER_EVENT_SCOPE("FlushRenderingCommands_ForHeightmapResolve");
+				FlushRenderingCommands();
+			}
 			ResolveLayersHeightmapTexture(InLandscapeComponentsToResolve);
 		}
 
+		TRACE_CPUPROFILER_EVENT_SCOPE("ResolveLayersHeightmapTexture_PostResolve");
 		// Partial Component Update
 		for (ULandscapeComponent* Component : InLandscapeComponentsToResolve)
 		{
@@ -3031,6 +3040,7 @@ int32 ALandscape::RegenerateLayersHeightmaps(const TArray<ULandscapeComponent*>&
 
 void ALandscape::ResolveLayersHeightmapTexture(const TArray<ULandscapeComponent*>& InLandscapeComponents)
 {
+	TRACE_CPUPROFILER_EVENT_SCOPE("ResolveLayersHeightmapTexture");
 	SCOPE_CYCLE_COUNTER(STAT_LandscapeLayersResolveHeightmaps);
 
 	ULandscapeInfo* Info = GetLandscapeInfo();
@@ -3075,6 +3085,7 @@ void ALandscape::ResolveLayersHeightmapTexture(const TArray<ULandscapeComponent*
 
 bool ALandscape::ResolveLayersTexture(FLandscapeLayersTexture2DCPUReadBackResource* InCPUReadBackTexture, UTexture2D* InOutputTexture)
 {
+	TRACE_CPUPROFILER_EVENT_SCOPE("ResolveLayersTexture");
 	SCOPE_CYCLE_COUNTER(STAT_LandscapeLayersResolveTexture);
 
 	TArray<TArray<FColor>> OutMipsData;
@@ -4029,9 +4040,15 @@ int32 ALandscape::RegenerateLayersWeightmaps(const TArray<ULandscapeComponent*>&
 		// We can skip Resolving if we don't need collision updating at all and that we aren't doing an intermediate render for some Landscape Tool (Flatten, Smooth). (This relies on the fact that at some point we will resolve with a full update)
 		if(bNeedsResolving)
 		{
+			if (InLandscapeComponentsToResolve.Num() > 0)
+			{
+				TRACE_CPUPROFILER_EVENT_SCOPE("FlushRenderingCommands_ForWeightmapResolve");
+				FlushRenderingCommands();
+			}
 			ResolveLayersWeightmapTexture(InLandscapeComponentsToResolve);
 		}
 				
+		TRACE_CPUPROFILER_EVENT_SCOPE("ResolveLayersWeightmapTexture_PostResolve");
 		for (ULandscapeComponent* Component : InLandscapeComponentsToResolve)
 		{
 			if (Component->IsUpdateFlagEnabledForModes(ELandscapeComponentUpdateFlag::Component_Update_Weightmap_Collision, WeightmapUpdateModes))
@@ -4258,6 +4275,7 @@ void ALandscape::UpdateLayersMaterialInstances(const TArray<ULandscapeComponent*
 
 void ALandscape::ResolveLayersWeightmapTexture(const TArray<ULandscapeComponent*>& InLandscapeComponents)
 {
+	TRACE_CPUPROFILER_EVENT_SCOPE("ResolveLayersWeightmapTexture");
 	SCOPE_CYCLE_COUNTER(STAT_LandscapeLayersResolveWeightmaps);
 
 	ULandscapeInfo* Info = GetLandscapeInfo();
@@ -4626,6 +4644,7 @@ void ALandscape::GetLandscapeComponentWeightmapsToRender(ULandscapeComponent* La
 
 void ALandscape::UpdateLayersContent(bool bInWaitForStreaming, bool bInSkipMonitorLandscapeEdModeChanges)
 {
+	TRACE_CPUPROFILER_EVENT_SCOPE("UpdateLayersContent");
 	// Remove this command line switch after fixes for D3D12 RHI
 	if (FParse::Param(FCommandLine::Get(), TEXT("nolandscapelayerupdate")))
 	{
@@ -5695,6 +5714,7 @@ uint32 ULandscapeComponent::ComputeLayerHash() const
 
 void ALandscape::UpdateLandscapeSplines(const FGuid& InTargetLayer, bool bInUpdateOnlySelected, bool bInForceUpdateAllCompoments)
 {
+	TRACE_CPUPROFILER_EVENT_SCOPE("UpdateLandscapeSplines");
 	check(CanHaveLayersContent());
 	ULandscapeInfo* LandscapeInfo = GetLandscapeInfo();
 	FGuid TargetLayerGuid = LandscapeSplinesTargetLayerGuid.IsValid() ? LandscapeSplinesTargetLayerGuid : InTargetLayer;
@@ -6184,6 +6204,7 @@ bool FLandscapeLayerBrush::IsAffectingWeightmapLayer(const FName& InWeightmapLay
 UTextureRenderTarget2D* FLandscapeLayerBrush::Render(bool InIsHeightmap, const FIntRect& InLandscapeExtent, UTextureRenderTarget2D* InLandscapeRenderTarget, const FName& InWeightmapLayerName)
 {
 #if WITH_EDITORONLY_DATA
+	TRACE_CPUPROFILER_EVENT_SCOPE("FLandscapeLayerBrush_Render");
 	if ((InIsHeightmap && !IsAffectingHeightmap()) ||
 		(!InIsHeightmap && !IsAffectingWeightmapLayer(InWeightmapLayerName)))
 	{
