Add support for Landscape Discrete LOD & Far Shadow LOD Bias
* This invalidates all shaders and iterative cooks must be cleared
* Check 'Use Discrete LOD' on landscape actors to enable. Recommended lod 0 distribution is 3.0
* Introduced r.Shadow.FarShadowStaticMeshLODBias to control lod bias in far cascades. Recommended value is 1
* Introduced r.Shadow.FarShadowDistanceOverride.  Recommended value is 160,000. When set to 0 it is not effective
* Recommended value for num far shadow cascades is 1

diff --git a/Engine/Shaders/Private/LandscapeVertexFactory.ush b/Engine/Shaders/Private/LandscapeVertexFactory.ush
--- a/Engine/Shaders/Private/LandscapeVertexFactory.ush
+++ b/Engine/Shaders/Private/LandscapeVertexFactory.ush
@@ -514,17 +514,27 @@ void GetShadowMapCoordinate(FVertexFactoryInterpolantsVSToPS Interpolants, out f
 }
 #endif
 
+float4 GetSectionLods()
+{
+	return min(SectionLods + View.FarShadowStaticMeshLODBias.xxxx, LandscapeParameters.LastLOD.xxxx);
+}
+
+float4 GetNeighborSectionLod(int NeighborIndex)
+{
+	return min(NeighborSectionLod[NeighborIndex] + View.FarShadowStaticMeshLODBias.xxxx, LandscapeParameters.LastLOD.xxxx);
+}
+
 float CalcLOD_ES2(float2 xy, float2 Subsection)
 {
 	// LOD value based on Barycentric coordinate
 	float4 L0 = float4(xy.y, xy.x, (1 - xy.x), (1 - xy.y)) * 2;
 
 	float4 LODCalculated4 = lerp(
-								(lerp(L0 * SectionLods.x + (float4(1, 1, 1, 1) - L0) * NeighborSectionLod[0],
-									  L0 * SectionLods.y + (float4(1, 1, 1, 1) - L0) * NeighborSectionLod[1],
+								(lerp(L0 * GetSectionLods() + (float4(1, 1, 1, 1) - L0) * GetNeighborSectionLod(0),
+									  L0 * GetSectionLods() + (float4(1, 1, 1, 1) - L0) * GetNeighborSectionLod(1),
 									  Subsection.x)),
-								(lerp(L0 * SectionLods.z + (float4(1, 1, 1, 1) - L0) * NeighborSectionLod[2],
-									  L0 * SectionLods.w + (float4(1, 1, 1, 1) - L0) * NeighborSectionLod[3],
+								(lerp(L0 * GetSectionLods() + (float4(1, 1, 1, 1) - L0) * GetNeighborSectionLod(2),
+									  L0 * GetSectionLods() + (float4(1, 1, 1, 1) - L0) * GetNeighborSectionLod(3),
 									  Subsection.x)),
 								Subsection.y);
 
@@ -555,26 +565,26 @@ float CalcLOD(float2 xy, float2 Subsection)
 	{
 		if (Subsection.x > 0.5f)
 		{
-			float SectionLod = SectionLods.w;
-			LODCalculated4 = L0 * SectionLod + (float4(1, 1, 1, 1) - L0) * NeighborSectionLod[3];
+			float SectionLod = GetSectionLods().w;
+			LODCalculated4 = L0 * SectionLod + (float4(1, 1, 1, 1) - L0) * GetNeighborSectionLod(3);
 		}
 		else
 		{
-			float SectionLod = SectionLods.z;
-			LODCalculated4 = L0 * SectionLod + (float4(1, 1, 1, 1) - L0) * NeighborSectionLod[2];
+			float SectionLod = GetSectionLods().z;
+			LODCalculated4 = L0 * SectionLod + (float4(1, 1, 1, 1) - L0) * GetNeighborSectionLod(2);
 		}
 	}
 	else
 	{
 		if (Subsection.x > 0.5f)
 		{
-			float SectionLod = SectionLods.y;
-			LODCalculated4 = L0 * SectionLod + (float4(1, 1, 1, 1) - L0) * NeighborSectionLod[1];
+			float SectionLod = GetSectionLods().y;
+			LODCalculated4 = L0 * SectionLod + (float4(1, 1, 1, 1) - L0) * GetNeighborSectionLod(1);
 		}
 		else
 		{
-			float SectionLod = SectionLods.x;
-			LODCalculated4 = L0 * SectionLod + (float4(1, 1, 1, 1) - L0) * NeighborSectionLod[0];
+			float SectionLod = GetSectionLods().x;
+			LODCalculated4 = L0 * SectionLod + (float4(1, 1, 1, 1) - L0) * GetNeighborSectionLod(0);
 		}
 	}
 
@@ -676,11 +686,19 @@ FVertexFactoryIntermediates GetVertexFactoryIntermediates(FVertexFactoryInput In
 	// Get the height and normal XY for current and next LOD out of the textures
 #if FEATURE_LEVEL >= FEATURE_LEVEL_SM4
 	float2 SampleCoords = InputPositionLODAdjusted * LandscapeParameters.HeightmapUVScaleBias.xy + LandscapeParameters.HeightmapUVScaleBias.zw + 0.5*LandscapeParameters.HeightmapUVScaleBias.xy + Intermediates.InputPosition.zw * LandscapeParameters.SubsectionOffsetParams.xy;
+#if FIXED_GRID
+	float4 SampleValue = Texture2DSampleLevel(LandscapeParameters.HeightmapTexture, LandscapeParameters.HeightmapTextureSampler, SampleCoords, 0);
+#else
 	float4 SampleValue = Texture2DSampleLevel(LandscapeParameters.HeightmapTexture, LandscapeParameters.HeightmapTextureSampler, SampleCoords, LodValue-LodBias.z);
+#endif
 	float Height = DecodePackedHeight(SampleValue.xy);
 
 	float2 SampleCoordsNextLOD = InputPositionNextLOD * LandscapeParameters.HeightmapUVScaleBias.xy + LandscapeParameters.HeightmapUVScaleBias.zw + 0.5*LandscapeParameters.HeightmapUVScaleBias.xy + Intermediates.InputPosition.zw * LandscapeParameters.SubsectionOffsetParams.xy;
+#if FIXED_GRID
+	float4 SampleValueNextLOD = Texture2DSampleLevel(LandscapeParameters.HeightmapTexture, LandscapeParameters.HeightmapTextureSampler, SampleCoordsNextLOD, 1);
+#else
 	float4 SampleValueNextLOD = Texture2DSampleLevel(LandscapeParameters.HeightmapTexture, LandscapeParameters.HeightmapTextureSampler, SampleCoordsNextLOD, LodValue+1-LodBias.z );
+#endif
 	float HeightNextLOD = DecodePackedHeight(SampleValueNextLOD.xy);
 #endif
 
@@ -722,7 +740,11 @@ float2 LandscapeVertexFactorySampleHeight(FVertexFactoryInput Input, FVertexFact
 	*/
 
 	float2 xy = Intermediates.InputPosition.xy * LodValues.w;
+#if FIXED_GRID
+	float LODCalculated = floor(LodValues.x);
+#else
 	float LODCalculated = CalcLOD(xy, Intermediates.InputPosition.zw);
+#endif
 	float LodValue = floor(LODCalculated);
 	float CoordTranslate_x = (LandscapeParameters.SubsectionSizeVertsLayerUVPan.x * pow(2, -LodValue) - 1) * LandscapeParameters.SubsectionSizeVertsLayerUVPan.y;
 
@@ -754,7 +776,11 @@ float4 VertexFactoryGetRasterizedWorldPosition(FVertexFactoryInput Input, FVerte
 {
 #if RENDERING_LANDSCAPE_GRASS
 	float2 xy = Intermediates.InputPosition.xy * LodValues.w;
+#if FIXED_GRID
+	float LODCalculated = floor(LodValues.x);
+#else
 	float LODCalculated = CalcLOD(xy, Intermediates.InputPosition.zw);
+#endif
 	float LodValue = floor(LODCalculated);
 	float CoordTranslate_x = (LandscapeParameters.SubsectionSizeVertsLayerUVPan.x * pow(2, -LodValue) - 1) * LandscapeParameters.SubsectionSizeVertsLayerUVPan.y;
 
@@ -859,7 +885,12 @@ FVertexFactoryInterpolantsVSToDS VertexFactoryGetInterpolantsVSToDS(FVertexFacto
 	Interpolants.TangentToWorld2 = float4(TangentToWorld[2], 1);
 
 	float2 xy = Intermediates.InputPosition.xy * LodValues.w;
-	Interpolants.CalculatedLOD = CalcLOD(xy, Intermediates.InputPosition.zw);
+#if FIXED_GRID
+	float LODCalculated = floor(LodValues.x);
+#else
+	float LODCalculated = CalcLOD(xy, Intermediates.InputPosition.zw);
+#endif
+	Interpolants.CalculatedLOD = LODCalculated;
 #endif
 	
 	return Interpolants;
diff --git a/Engine/Source/Editor/LandscapeEditor/Private/LandscapeEdMode.cpp b/Engine/Source/Editor/LandscapeEditor/Private/LandscapeEdMode.cpp
--- a/Engine/Source/Editor/LandscapeEditor/Private/LandscapeEdMode.cpp
+++ b/Engine/Source/Editor/LandscapeEditor/Private/LandscapeEdMode.cpp
@@ -274,10 +274,11 @@ void FLandscapeTool::SetEditRenderType()
 
 namespace LandscapeTool
 {
-	UMaterialInstance* CreateMaterialInstance(UMaterialInterface* BaseMaterial)
+	UMaterialInstance* CreateMaterialInstance(UMaterialInterface* BaseMaterial, bool bUseDiscreteLOD)
 	{
 		ULandscapeMaterialInstanceConstant* MaterialInstance = NewObject<ULandscapeMaterialInstanceConstant>(GetTransientPackage());
 		MaterialInstance->bEditorToolUsage = true;
+		MaterialInstance->bUseDiscreteLOD = bUseDiscreteLOD;
 		MaterialInstance->SetParentEditorOnly(BaseMaterial);
 		MaterialInstance->PostEditChange();
 		return MaterialInstance;
@@ -312,13 +313,19 @@ FEdModeLandscape::FEdModeLandscape()
 	, bNeedsUpdateShownLayerList(false)
 	, bUpdatingLandscapeInfo(false)
 {
-	GLayerDebugColorMaterial = LandscapeTool::CreateMaterialInstance(LoadObject<UMaterial>(nullptr, TEXT("/Engine/EditorLandscapeResources/LayerVisMaterial.LayerVisMaterial")));
-	GSelectionColorMaterial  = LandscapeTool::CreateMaterialInstance(LoadObject<UMaterialInstanceConstant>(nullptr, TEXT("/Engine/EditorLandscapeResources/SelectBrushMaterial_Selected.SelectBrushMaterial_Selected")));
-	GSelectionRegionMaterial = LandscapeTool::CreateMaterialInstance(LoadObject<UMaterialInstanceConstant>(nullptr, TEXT("/Engine/EditorLandscapeResources/SelectBrushMaterial_SelectedRegion.SelectBrushMaterial_SelectedRegion")));
-	GMaskRegionMaterial      = LandscapeTool::CreateMaterialInstance(LoadObject<UMaterialInstanceConstant>(nullptr, TEXT("/Engine/EditorLandscapeResources/MaskBrushMaterial_MaskedRegion.MaskBrushMaterial_MaskedRegion")));
-	GColorMaskRegionMaterial = LandscapeTool::CreateMaterialInstance(LoadObject<UMaterialInstanceConstant>(nullptr, TEXT("/Engine/EditorLandscapeResources/ColorMaskBrushMaterial_MaskedRegion.ColorMaskBrushMaterial_MaskedRegion")));
+	GLayerDebugColorMaterial     = LandscapeTool::CreateMaterialInstance(LoadObject<UMaterial>(nullptr, TEXT("/Engine/EditorLandscapeResources/LayerVisMaterial.LayerVisMaterial")),                                                       false);
+	GSelectionColorMaterial      = LandscapeTool::CreateMaterialInstance(LoadObject<UMaterialInstanceConstant>(nullptr, TEXT("/Engine/EditorLandscapeResources/SelectBrushMaterial_Selected.SelectBrushMaterial_Selected")),               false);
+	GSelectionRegionMaterial     = LandscapeTool::CreateMaterialInstance(LoadObject<UMaterialInstanceConstant>(nullptr, TEXT("/Engine/EditorLandscapeResources/SelectBrushMaterial_SelectedRegion.SelectBrushMaterial_SelectedRegion")),   false);
+	GMaskRegionMaterial          = LandscapeTool::CreateMaterialInstance(LoadObject<UMaterialInstanceConstant>(nullptr, TEXT("/Engine/EditorLandscapeResources/MaskBrushMaterial_MaskedRegion.MaskBrushMaterial_MaskedRegion")),           false);
+	GColorMaskRegionMaterial     = LandscapeTool::CreateMaterialInstance(LoadObject<UMaterialInstanceConstant>(nullptr, TEXT("/Engine/EditorLandscapeResources/ColorMaskBrushMaterial_MaskedRegion.ColorMaskBrushMaterial_MaskedRegion")), false);
+	GLandscapeLayerUsageMaterial = LandscapeTool::CreateMaterialInstance(LoadObject<UMaterial>(nullptr, TEXT("/Engine/EditorLandscapeResources/LandscapeLayerUsageMaterial.LandscapeLayerUsageMaterial")),                                 false);
+	GLayerDebugColorMaterial_DiscreteLOD     = LandscapeTool::CreateMaterialInstance(LoadObject<UMaterial>(nullptr, TEXT("/Engine/EditorLandscapeResources/LayerVisMaterial.LayerVisMaterial")),                                                       true);
+	GSelectionColorMaterial_DiscreteLOD      = LandscapeTool::CreateMaterialInstance(LoadObject<UMaterialInstanceConstant>(nullptr, TEXT("/Engine/EditorLandscapeResources/SelectBrushMaterial_Selected.SelectBrushMaterial_Selected")),               true);
+	GSelectionRegionMaterial_DiscreteLOD     = LandscapeTool::CreateMaterialInstance(LoadObject<UMaterialInstanceConstant>(nullptr, TEXT("/Engine/EditorLandscapeResources/SelectBrushMaterial_SelectedRegion.SelectBrushMaterial_SelectedRegion")),   true);
+	GMaskRegionMaterial_DiscreteLOD          = LandscapeTool::CreateMaterialInstance(LoadObject<UMaterialInstanceConstant>(nullptr, TEXT("/Engine/EditorLandscapeResources/MaskBrushMaterial_MaskedRegion.MaskBrushMaterial_MaskedRegion")),           true);
+	GColorMaskRegionMaterial_DiscreteLOD     = LandscapeTool::CreateMaterialInstance(LoadObject<UMaterialInstanceConstant>(nullptr, TEXT("/Engine/EditorLandscapeResources/ColorMaskBrushMaterial_MaskedRegion.ColorMaskBrushMaterial_MaskedRegion")), true);
+	GLandscapeLayerUsageMaterial_DiscreteLOD = LandscapeTool::CreateMaterialInstance(LoadObject<UMaterial>(nullptr, TEXT("/Engine/EditorLandscapeResources/LandscapeLayerUsageMaterial.LandscapeLayerUsageMaterial")),                                 true);
 	GLandscapeBlackTexture   = LoadObject<UTexture2D>(nullptr, TEXT("/Engine/EngineResources/Black.Black"));
-	GLandscapeLayerUsageMaterial = LandscapeTool::CreateMaterialInstance(LoadObject<UMaterial>(nullptr, TEXT("/Engine/EditorLandscapeResources/LandscapeLayerUsageMaterial.LandscapeLayerUsageMaterial")));
 	
 	// Initialize modes
 	UpdateToolModes();
@@ -390,6 +397,13 @@ FEdModeLandscape::~FEdModeLandscape()
 	GLandscapeBlackTexture = NULL;
 	GLandscapeLayerUsageMaterial = NULL;
 
+	GLayerDebugColorMaterial_DiscreteLOD = NULL;
+	GSelectionColorMaterial_DiscreteLOD = NULL;
+	GSelectionRegionMaterial_DiscreteLOD = NULL;
+	GMaskRegionMaterial_DiscreteLOD = NULL;
+	GColorMaskRegionMaterial_DiscreteLOD = NULL;
+	GLandscapeLayerUsageMaterial_DiscreteLOD = NULL;
+
 	InteractorPainting = nullptr;
 }
 
@@ -406,8 +420,16 @@ void FEdModeLandscape::AddReferencedObjects(FReferenceCollector& Collector)
 	Collector.AddReferencedObject(GSelectionRegionMaterial);
 	Collector.AddReferencedObject(GMaskRegionMaterial);
 	Collector.AddReferencedObject(GColorMaskRegionMaterial);
-	Collector.AddReferencedObject(GLandscapeBlackTexture);
 	Collector.AddReferencedObject(GLandscapeLayerUsageMaterial);
+
+	Collector.AddReferencedObject(GLayerDebugColorMaterial_DiscreteLOD);
+	Collector.AddReferencedObject(GSelectionColorMaterial_DiscreteLOD);
+	Collector.AddReferencedObject(GSelectionRegionMaterial_DiscreteLOD);
+	Collector.AddReferencedObject(GMaskRegionMaterial_DiscreteLOD);
+	Collector.AddReferencedObject(GColorMaskRegionMaterial_DiscreteLOD);
+	Collector.AddReferencedObject(GLandscapeLayerUsageMaterial_DiscreteLOD);
+
+	Collector.AddReferencedObject(GLandscapeBlackTexture);
 }
 
 void FEdModeLandscape::UpdateToolModes()
diff --git a/Engine/Source/Editor/LandscapeEditor/Private/LandscapeEdModeBrushes.cpp b/Engine/Source/Editor/LandscapeEditor/Private/LandscapeEdModeBrushes.cpp
--- a/Engine/Source/Editor/LandscapeEditor/Private/LandscapeEdModeBrushes.cpp
+++ b/Engine/Source/Editor/LandscapeEditor/Private/LandscapeEdModeBrushes.cpp
@@ -12,7 +12,7 @@
 #include "LandscapeGizmoActor.h"
 #include "LandscapeEdMode.h"
 #include "LandscapeEditorObject.h"
-
+#include "LandscapeMaterialInstanceConstant.h"
 #include "LandscapeRender.h"
 
 #include "LevelUtils.h"
@@ -52,6 +52,7 @@ class FLandscapeBrushCircle : public FLandscapeBrush
 protected:
 	FVector2D LastMousePosition;
 	UMaterialInterface* BrushMaterial;
+	UMaterialInterface* BrushMaterial_DiscreteLOD;
 	TMap<ULandscapeComponent*, UMaterialInstanceDynamic*> BrushMaterialInstanceMap;
 
 	virtual float CalculateFalloff(float Distance, float Radius, float Falloff) = 0;
@@ -59,7 +60,8 @@ protected:
 	/** Protected so that only subclasses can create instances of this class. */
 	FLandscapeBrushCircle(FEdModeLandscape* InEdMode, UMaterialInterface* InBrushMaterial)
 		: LastMousePosition(0, 0)
-		, BrushMaterial(LandscapeTool::CreateMaterialInstance(InBrushMaterial))
+		, BrushMaterial(LandscapeTool::CreateMaterialInstance(InBrushMaterial, false))
+		, BrushMaterial_DiscreteLOD(LandscapeTool::CreateMaterialInstance(InBrushMaterial, true))
 		, EdMode(InEdMode)
 	{
 	}
@@ -71,6 +73,7 @@ public:
 	virtual void AddReferencedObjects(FReferenceCollector& Collector) override
 	{
 		Collector.AddReferencedObject(BrushMaterial);
+		Collector.AddReferencedObject(BrushMaterial_DiscreteLOD);
 
 		// Allow any currently unused material instances to be GC'd
 		BrushMaterialFreeInstances.Empty();
@@ -160,14 +163,26 @@ public:
 		for (ULandscapeComponent* AddedComponent : AddedComponents)
 		{
 			UMaterialInstanceDynamic* BrushMaterialInstance = nullptr;
+
 			if (BrushMaterialFreeInstances.Num() > 0)
 			{
-				BrushMaterialInstance = BrushMaterialFreeInstances.Pop();
+				for (auto Iter = BrushMaterialFreeInstances.CreateIterator(); Iter; Iter++)
+				{
+					ULandscapeMaterialInstanceConstant* LandscapeMIC = Cast<ULandscapeMaterialInstanceConstant>((*Iter)->Parent);
+					if (LandscapeMIC->bUseDiscreteLOD == AddedComponent->GetLandscapeProxy()->bUseDiscreteLOD)
+					{
+						BrushMaterialInstance = *Iter;
+						Iter.RemoveCurrent();
+						break;
+					}
+				}
 			}
-			else
+			
+			if (BrushMaterialInstance == nullptr)
 			{
-				BrushMaterialInstance = UMaterialInstanceDynamic::Create(BrushMaterial, nullptr);
+				BrushMaterialInstance = UMaterialInstanceDynamic::Create(AddedComponent->GetLandscapeProxy()->bUseDiscreteLOD ? BrushMaterial_DiscreteLOD : BrushMaterial, nullptr);
 			}
+
 			BrushMaterialInstanceMap.Add(AddedComponent, BrushMaterialInstance);
 			AddedComponent->EditToolRenderData.ToolMaterial = BrushMaterialInstance;
 			AddedComponent->UpdateEditToolRenderData();
@@ -358,15 +373,18 @@ class FLandscapeBrushComponent : public FLandscapeBrush
 protected:
 	FVector2D LastMousePosition;
 	UMaterialInterface* BrushMaterial;
+	UMaterialInterface* BrushMaterial_DiscreteLOD;
 public:
 	FEdModeLandscape* EdMode;
 
 	FLandscapeBrushComponent(FEdModeLandscape* InEdMode)
 		: BrushMaterial(nullptr)
+		, BrushMaterial_DiscreteLOD(nullptr)
 		, EdMode(InEdMode)
 	{
 		UMaterial* BaseBrushMaterial = LoadObject<UMaterial>(nullptr, TEXT("/Engine/EditorLandscapeResources/SelectBrushMaterial.SelectBrushMaterial"));
-		BrushMaterial = LandscapeTool::CreateMaterialInstance(BaseBrushMaterial);
+		BrushMaterial = LandscapeTool::CreateMaterialInstance(BaseBrushMaterial, false);
+		BrushMaterial_DiscreteLOD = LandscapeTool::CreateMaterialInstance(BaseBrushMaterial, true);
 	}
 
 	// FGCObject interface
@@ -374,6 +392,7 @@ public:
 	{
 		Collector.AddReferencedObjects(BrushMaterialComponents);
 		Collector.AddReferencedObject(BrushMaterial);
+		Collector.AddReferencedObject(BrushMaterial_DiscreteLOD);
 	}
 
 	virtual ELandscapeBrushType GetBrushType() override { return ELandscapeBrushType::Component; }
@@ -440,7 +459,7 @@ public:
 				// Set brush material for components in new region
 				for (ULandscapeComponent* NewComponent : NewComponents)
 				{
-					NewComponent->EditToolRenderData.ToolMaterial = BrushMaterial;
+					NewComponent->EditToolRenderData.ToolMaterial = NewComponent->GetLandscapeProxy()->bUseDiscreteLOD ? BrushMaterial_DiscreteLOD : BrushMaterial;
 					NewComponent->UpdateEditToolRenderData();
 				}
 			}
@@ -554,15 +573,18 @@ class FLandscapeBrushGizmo : public FLandscapeBrush
 
 protected:
 	UMaterialInstanceDynamic* BrushMaterial;
+	UMaterialInstanceDynamic* BrushMaterial_DiscreteLOD;
 public:
 	FEdModeLandscape* EdMode;
 
 	FLandscapeBrushGizmo(FEdModeLandscape* InEdMode)
 		: BrushMaterial(nullptr)
+		, BrushMaterial_DiscreteLOD(nullptr)
 		, EdMode(InEdMode)
 	{
 		UMaterialInterface* GizmoMaterial = LoadObject<UMaterialInstanceConstant>(nullptr, TEXT("/Engine/EditorLandscapeResources/MaskBrushMaterial_Gizmo.MaskBrushMaterial_Gizmo"));
-		BrushMaterial = UMaterialInstanceDynamic::Create(LandscapeTool::CreateMaterialInstance(GizmoMaterial), nullptr);
+		BrushMaterial = UMaterialInstanceDynamic::Create(LandscapeTool::CreateMaterialInstance(GizmoMaterial, false), nullptr);
+		BrushMaterial_DiscreteLOD = UMaterialInstanceDynamic::Create(LandscapeTool::CreateMaterialInstance(GizmoMaterial, true), nullptr);
 	}
 
 	// FGCObject interface
@@ -570,6 +592,7 @@ public:
 	{
 		Collector.AddReferencedObjects(BrushMaterialComponents);
 		Collector.AddReferencedObject(BrushMaterial);
+		Collector.AddReferencedObject(BrushMaterial_DiscreteLOD);
 	}
 
 	virtual ELandscapeBrushType GetBrushType() override { return ELandscapeBrushType::Gizmo; }
@@ -639,16 +662,19 @@ public:
 						Gizmo->TextureScale.Y
 						);
 					BrushMaterial->SetVectorParameterValue(FName(TEXT("AlphaScaleBias")), AlphaScaleBias);
+					BrushMaterial_DiscreteLOD->SetVectorParameterValue(FName(TEXT("AlphaScaleBias")), AlphaScaleBias);
 
 					float Angle = (-EdMode->CurrentGizmoActor->GetActorRotation().Euler().Z) * PI / 180.0f;
 					FLinearColor LandscapeLocation(EdMode->CurrentGizmoActor->GetActorLocation().X, EdMode->CurrentGizmoActor->GetActorLocation().Y, EdMode->CurrentGizmoActor->GetActorLocation().Z, Angle);
 					BrushMaterial->SetVectorParameterValue(FName(TEXT("LandscapeLocation")), LandscapeLocation);
+					BrushMaterial_DiscreteLOD->SetVectorParameterValue(FName(TEXT("LandscapeLocation")), LandscapeLocation);
 					BrushMaterial->SetTextureParameterValue(FName(TEXT("AlphaTexture")), DataTexture);
+					BrushMaterial_DiscreteLOD->SetTextureParameterValue(FName(TEXT("AlphaTexture")), DataTexture);
 
 					// Set brush material for components in new region
 					for (ULandscapeComponent* NewComponent : NewComponents)
 					{
-						NewComponent->EditToolRenderData.GizmoMaterial = ((Gizmo->DataType != LGT_None) && (GLandscapeEditRenderMode & ELandscapeEditRenderMode::Gizmo) ? BrushMaterial : nullptr);
+						NewComponent->EditToolRenderData.GizmoMaterial = ((Gizmo->DataType != LGT_None) && (GLandscapeEditRenderMode & ELandscapeEditRenderMode::Gizmo) ? (NewComponent->GetLandscapeProxy()->bUseDiscreteLOD ? BrushMaterial_DiscreteLOD : BrushMaterial) : nullptr);
 						NewComponent->UpdateEditToolRenderData();
 					}
 
diff --git a/Engine/Source/Editor/LandscapeEditor/Public/LandscapeToolInterface.h b/Engine/Source/Editor/LandscapeEditor/Public/LandscapeToolInterface.h
--- a/Engine/Source/Editor/LandscapeEditor/Public/LandscapeToolInterface.h
+++ b/Engine/Source/Editor/LandscapeEditor/Public/LandscapeToolInterface.h
@@ -257,5 +257,5 @@ public:
 
 namespace LandscapeTool
 {
-	UMaterialInstance* CreateMaterialInstance(UMaterialInterface* BaseMaterial);
+	UMaterialInstance* CreateMaterialInstance(UMaterialInterface* BaseMaterial, bool bUseDiscreteLOD);
 }
diff --git a/Engine/Source/Runtime/Engine/Private/Components/DirectionalLightComponent.cpp b/Engine/Source/Runtime/Engine/Private/Components/DirectionalLightComponent.cpp
--- a/Engine/Source/Runtime/Engine/Private/Components/DirectionalLightComponent.cpp
+++ b/Engine/Source/Runtime/Engine/Private/Components/DirectionalLightComponent.cpp
@@ -51,6 +51,20 @@ static TAutoConsoleVariable<int32> CVarMaxNumFarShadowCascades(
 	TEXT("Max number of far shadow cascades that can be cast from a directional light"),
 	ECVF_RenderThreadSafe | ECVF_Scalability );
 
+ENGINE_API int32 GFarShadowStaticMeshLODBias = 0;
+FAutoConsoleVariableRef CVarFarShadowStaticMeshLODBias(
+	TEXT("r.Shadow.FarShadowStaticMeshLODBias"),
+	GFarShadowStaticMeshLODBias,
+	TEXT("Notice: only selected geometry types (static meshes and landscapes) respect this value."),
+	ECVF_Scalability | ECVF_RenderThreadSafe
+);
+
+static TAutoConsoleVariable<float> CVarFarShadowDistanceOverride(
+	TEXT("r.Shadow.FarShadowDistanceOverride"),
+	0.0f,
+	TEXT("Overriding far shadow distance for all directional lighst"),
+	ECVF_RenderThreadSafe | ECVF_Scalability);
+
 /**
  * The scene info for a directional light.
  */
@@ -344,7 +358,7 @@ public:
 			{
 				FarDistance = GetDistanceFieldShadowDistance();
 			}
-			FarDistance = FMath::Max(FarDistance, FarShadowDistance);
+			FarDistance = FMath::Max(FarDistance, CVarFarShadowDistanceOverride.GetValueOnAnyThread() > 0.0f ? CVarFarShadowDistanceOverride.GetValueOnAnyThread() : FarShadowDistance);
 		}
 	    
 		// The far distance for the dynamic to static fade is the range of the directional light.
@@ -624,7 +638,7 @@ private:
 				// the far cascades start at the after the near cascades
 				uint32 ClampedFarShadowCascadeCount = FMath::Min((uint32)CVarMaxNumFarShadowCascades.GetValueOnAnyThread(), FarShadowCascadeCount);
 
-				return CascadeDistanceWithoutFar + ComputeAccumulatedScale(EffectiveCascadeDistributionExponent, SplitIndex - NumNearCascades, ClampedFarShadowCascadeCount) * (FarShadowDistance - CascadeDistanceWithoutFar);
+				return CascadeDistanceWithoutFar + ComputeAccumulatedScale(EffectiveCascadeDistributionExponent, SplitIndex - NumNearCascades, ClampedFarShadowCascadeCount) * ((CVarFarShadowDistanceOverride.GetValueOnAnyThread() > 0.0f ? CVarFarShadowDistanceOverride.GetValueOnAnyThread() : FarShadowDistance) - CascadeDistanceWithoutFar);
 			}
 		}
 		else
diff --git a/Engine/Source/Runtime/Engine/Public/SceneView.h b/Engine/Source/Runtime/Engine/Public/SceneView.h
--- a/Engine/Source/Runtime/Engine/Public/SceneView.h
+++ b/Engine/Source/Runtime/Engine/Public/SceneView.h
@@ -722,6 +722,7 @@ enum ETranslucencyVolumeCascade
 	VIEW_UNIFORM_BUFFER_MEMBER(FMatrix, WorldToVirtualTexture) \
 	VIEW_UNIFORM_BUFFER_MEMBER(FVector4, VirtualTextureParams) \
 	VIEW_UNIFORM_BUFFER_MEMBER_ARRAY(FVector4, XRPassthroughCameraUVs, [2]) \
+	VIEW_UNIFORM_BUFFER_MEMBER(int32, FarShadowStaticMeshLODBias) \
 
 #define VIEW_UNIFORM_BUFFER_MEMBER(type, identifier) \
 	SHADER_PARAMETER(type, identifier)
@@ -816,10 +817,11 @@ END_GLOBAL_SHADER_PARAMETER_STRUCT()
 
 namespace EDrawDynamicFlags
 {
-	enum Type
+	enum Type : int32
 	{
 		None = 0,
-		ForceLowestLOD = 0x1
+		ForceLowestLOD = 0x1,
+		FarShadowCascade = 0x2,
 	};
 }
 
diff --git a/Engine/Source/Runtime/Landscape/Classes/LandscapeMaterialInstanceConstant.h b/Engine/Source/Runtime/Landscape/Classes/LandscapeMaterialInstanceConstant.h
--- a/Engine/Source/Runtime/Landscape/Classes/LandscapeMaterialInstanceConstant.h
+++ b/Engine/Source/Runtime/Landscape/Classes/LandscapeMaterialInstanceConstant.h
@@ -24,6 +24,9 @@ class ULandscapeMaterialInstanceConstant : public UMaterialInstanceConstant
 	UPROPERTY()
 	uint32 bEditorToolUsage:1;
 
+	UPROPERTY()
+	uint32 bUseDiscreteLOD : 1;
+
 	virtual FMaterialResource* AllocatePermutationResource() override;
 	virtual bool HasOverridenBaseProperties() const override;
 };
diff --git a/Engine/Source/Runtime/Landscape/Classes/LandscapeProxy.h b/Engine/Source/Runtime/Landscape/Classes/LandscapeProxy.h
--- a/Engine/Source/Runtime/Landscape/Classes/LandscapeProxy.h
+++ b/Engine/Source/Runtime/Landscape/Classes/LandscapeProxy.h
@@ -435,6 +435,10 @@ public:
 	UPROPERTY(EditAnywhere, Category = LOD, meta=(ClampMin = "0.01", ClampMax = "1.0", UIMin = "0.01", UIMax = "1.0", DisplayName= "SubSection Min Component ScreenSize"))
 	float ComponentScreenSizeToUseSubSections;
 
+	/** Whether to use discrete LOD rather than continuous LOD. Continuous LOD smooths LOD transition, at the cost of increased RenderThread time. */
+	UPROPERTY(EditAnywhere, Category = "LOD Distribution")
+	bool bUseDiscreteLOD;
+
 	/** This is the starting screen size used to calculate the distribution, by default it's 1, but you can increase the value if you want less LOD0 component, and you use very large landscape component. */
 	UPROPERTY(EditAnywhere, Category = "LOD Distribution", meta = (DisplayName = "LOD 0 Screen Size", ClampMin = "1.0", ClampMax = "10.0", UIMin = "1.0", UIMax = "10.0"))
 	float LOD0ScreenSize;
diff --git a/Engine/Source/Runtime/Landscape/Private/Landscape.cpp b/Engine/Source/Runtime/Landscape/Private/Landscape.cpp
--- a/Engine/Source/Runtime/Landscape/Private/Landscape.cpp
+++ b/Engine/Source/Runtime/Landscape/Private/Landscape.cpp
@@ -2386,6 +2386,7 @@ void ALandscapeProxy::GetSharedProperties(ALandscapeProxy* Landscape)
 		ComponentScreenSizeToUseSubSections = Landscape->ComponentScreenSizeToUseSubSections;
 		UseTessellationComponentScreenSizeFalloff = Landscape->UseTessellationComponentScreenSizeFalloff;
 		TessellationComponentScreenSizeFalloff = Landscape->TessellationComponentScreenSizeFalloff;
+		bUseDiscreteLOD = Landscape->bUseDiscreteLOD;
 		LODDistributionSetting = Landscape->LODDistributionSetting;
 		LOD0DistributionSetting = Landscape->LOD0DistributionSetting;
 		LOD0ScreenSize = Landscape->LOD0ScreenSize;
diff --git a/Engine/Source/Runtime/Landscape/Private/LandscapeEdit.cpp b/Engine/Source/Runtime/Landscape/Private/LandscapeEdit.cpp
--- a/Engine/Source/Runtime/Landscape/Private/LandscapeEdit.cpp
+++ b/Engine/Source/Runtime/Landscape/Private/LandscapeEdit.cpp
@@ -275,6 +275,7 @@ UMaterialInstanceConstant* ULandscapeComponent::GetCombinationMaterial(FMaterial
 
 			ULandscapeMaterialInstanceConstant* LandscapeCombinationMaterialInstance = NewObject<ULandscapeMaterialInstanceConstant>(GetOuter());
 			LandscapeCombinationMaterialInstance->bMobile = bMobile;
+			LandscapeCombinationMaterialInstance->bUseDiscreteLOD = GetLandscapeProxy()->bUseDiscreteLOD;
 			CombinationMaterialInstance = LandscapeCombinationMaterialInstance;
 			UE_LOG(LogLandscape, Log, TEXT("Looking for key %s, making new combination %s"), *LayerKey, *CombinationMaterialInstance->GetName());
 			Proxy->MaterialInstanceConstantMap.Add(*LayerKey, CombinationMaterialInstance);
@@ -363,7 +364,8 @@ void ULandscapeComponent::UpdateMaterialInstances_Internal(FMaterialUpdateContex
 		if (CombinationMaterialInstance != nullptr)
 		{
 			// Create the instance for this component, that will use the layer combination instance.
-			UMaterialInstanceConstant* MaterialInstance = NewObject<ULandscapeMaterialInstanceConstant>(GetOuter());
+			ULandscapeMaterialInstanceConstant* MaterialInstance = NewObject<ULandscapeMaterialInstanceConstant>(GetOuter());
+			MaterialInstance->bUseDiscreteLOD = GetLandscapeProxy()->bUseDiscreteLOD;
 			MaterialInstances[MaterialIndex] = MaterialInstance;
 
 			// Material Instances don't support Undo/Redo (the shader map goes out of sync and crashes happen)
@@ -397,6 +399,9 @@ void ULandscapeComponent::UpdateMaterialInstances_Internal(FMaterialUpdateContex
 			{
 				MaterialInstance->SetTextureParameterValueEditorOnly(FName(TEXT("Heightmap")), BaseHeightmap);
 			}
+
+			MaterialInstance->bUseDiscreteLOD = GetLandscapeProxy()->bUseDiscreteLOD;
+
 			MaterialInstance->PostEditChange();
 
 			// Setup material instance with disabled tessellation
@@ -410,6 +415,7 @@ void ULandscapeComponent::UpdateMaterialInstances_Internal(FMaterialUpdateContex
 				TessellationMaterialInstance->SetParentEditorOnly(MaterialInstance);
 				Context.AddMaterialInstance(TessellationMaterialInstance); // must be done after SetParent
 				TessellationMaterialInstance->bDisableTessellation = true;
+				TessellationMaterialInstance->bUseDiscreteLOD = GetLandscapeProxy()->bUseDiscreteLOD;
 				TessellationMaterialInstance->PostEditChange();
 			}
 		}
@@ -4200,7 +4206,7 @@ void ALandscapeStreamingProxy::PostEditChangeProperty(FPropertyChangedEvent& Pro
 			LandscapeActor = nullptr;
 		}
 	}
-	else if (PropertyName == FName(TEXT("LandscapeMaterial")) || PropertyName == FName(TEXT("LandscapeHoleMaterial")) || PropertyName == FName(TEXT("LandscapeMaterialsOverride")))
+	else if (PropertyName == FName(TEXT("LandscapeMaterial")) || PropertyName == FName(TEXT("LandscapeHoleMaterial")) || PropertyName == FName(TEXT("LandscapeMaterialsOverride")) || PropertyName == FName(TEXT("bUseDiscreteLOD")))
 	{
 		bool RecreateMaterialInstances = true;
 
@@ -4326,6 +4332,30 @@ void ALandscape::PostEditChangeProperty(FPropertyChangedEvent& PropertyChangedEv
 			MaterialInstanceConstantMap.Empty();
 		}
 	}
+	else if (PropertyName == GET_MEMBER_NAME_CHECKED(ALandscapeProxy, bUseDiscreteLOD))
+	{
+		bPropagateToProxies = true;
+
+		if (Info != nullptr)
+		{
+			FMaterialUpdateContext MaterialUpdateContext;
+			Info->UpdateLayerInfoMap(/*this*/);
+
+			ChangedMaterial = true;
+
+			// Clear the parents out of combination material instances
+			for (const auto& MICPair : MaterialInstanceConstantMap)
+			{
+				UMaterialInstanceConstant* MaterialInstance = MICPair.Value;
+				MaterialInstance->BasePropertyOverrides.bOverride_BlendMode = false;
+				MaterialInstance->SetParentEditorOnly(nullptr);
+				MaterialUpdateContext.AddMaterialInstance(MaterialInstance);
+			}
+
+			// Remove our references to any material instances
+			MaterialInstanceConstantMap.Empty();
+		}
+	}
 	else if (PropertyName == FName(TEXT("RelativeScale3D")) ||
 		PropertyName == FName(TEXT("RelativeLocation")) ||
 		PropertyName == FName(TEXT("RelativeRotation")))
@@ -5762,8 +5792,8 @@ void ULandscapeComponent::GeneratePlatformPixelData()
 			MobileCombinationMaterialInstances[MaterialIndex] = GetCombinationMaterial(nullptr, MobileWeightmapLayerAllocations, MaterialLOD, true);
 			check(MobileCombinationMaterialInstances[MaterialIndex] != nullptr);
 
-			UMaterialInstanceConstant* NewMobileMaterialInstance = NewObject<ULandscapeMaterialInstanceConstant>(this);
-
+			ULandscapeMaterialInstanceConstant* NewMobileMaterialInstance = NewObject<ULandscapeMaterialInstanceConstant>(this);
+			NewMobileMaterialInstance->bUseDiscreteLOD = GetLandscapeProxy()->bUseDiscreteLOD;
 			NewMobileMaterialInstance->SetParentEditorOnly(MobileCombinationMaterialInstances[MaterialIndex]);
 
 			// Set the layer mask
diff --git a/Engine/Source/Runtime/Landscape/Private/LandscapeEditLayers.cpp b/Engine/Source/Runtime/Landscape/Private/LandscapeEditLayers.cpp
--- a/Engine/Source/Runtime/Landscape/Private/LandscapeEditLayers.cpp
+++ b/Engine/Source/Runtime/Landscape/Private/LandscapeEditLayers.cpp
@@ -4184,7 +4184,9 @@ void ALandscape::UpdateLayersMaterialInstances(const TArray<ULandscapeComponent*
 				if (NeedToCreateMIC)
 				{
 					// Create the instance for this component, that will use the layer combination instance.
-					MaterialInstance = NewObject<ULandscapeMaterialInstanceConstant>(this);
+					ULandscapeMaterialInstanceConstant* LandscapeMIC = NewObject<ULandscapeMaterialInstanceConstant>(this);
+					LandscapeMIC->bUseDiscreteLOD = this->bUseDiscreteLOD;
+					MaterialInstance = LandscapeMIC;
 					Component->MaterialInstances[MaterialIndex] = MaterialInstance;
 				}
 
diff --git a/Engine/Source/Runtime/Landscape/Private/LandscapeGrass.cpp b/Engine/Source/Runtime/Landscape/Private/LandscapeGrass.cpp
--- a/Engine/Source/Runtime/Landscape/Private/LandscapeGrass.cpp
+++ b/Engine/Source/Runtime/Landscape/Private/LandscapeGrass.cpp
@@ -206,7 +206,11 @@ static bool ShouldCacheLandscapeGrassShaders(EShaderPlatform Platform, const FMa
 	// We only need grass weight shaders for Landscape vertex factories on desktop platforms
 	return (Material->IsUsedWithLandscape() || Material->IsSpecialEngineMaterial()) &&
 		IsFeatureLevelSupported(Platform, ERHIFeatureLevel::SM5) &&
-		((VertexFactoryType == FindVertexFactoryType(FName(TEXT("FLandscapeVertexFactory"), FNAME_Find))) || (VertexFactoryType == FindVertexFactoryType(FName(TEXT("FLandscapeXYOffsetVertexFactory"), FNAME_Find))))
+		(
+			(VertexFactoryType == FindVertexFactoryType(FName(TEXT("FLandscapeVertexFactory"), FNAME_Find))) || 
+			(VertexFactoryType == FindVertexFactoryType(FName(TEXT("FLandscapeFixedGridVertexFactory"), FNAME_Find))) || 
+			(VertexFactoryType == FindVertexFactoryType(FName(TEXT("FLandscapeXYOffsetVertexFactory"), FNAME_Find)))
+		)
 		&& !IsConsolePlatform(Platform);
 }
 
diff --git a/Engine/Source/Runtime/Landscape/Private/LandscapeRender.cpp b/Engine/Source/Runtime/Landscape/Private/LandscapeRender.cpp
--- a/Engine/Source/Runtime/Landscape/Private/LandscapeRender.cpp
+++ b/Engine/Source/Runtime/Landscape/Private/LandscapeRender.cpp
@@ -566,13 +566,22 @@ FLandscapeDebugOptions GLandscapeDebugOptions;
 #if WITH_EDITOR
 LANDSCAPE_API bool GLandscapeEditModeActive = false;
 LANDSCAPE_API int32 GLandscapeEditRenderMode = ELandscapeEditRenderMode::None;
+
 UMaterialInterface* GLayerDebugColorMaterial = nullptr;
 UMaterialInterface* GSelectionColorMaterial = nullptr;
 UMaterialInterface* GSelectionRegionMaterial = nullptr;
 UMaterialInterface* GMaskRegionMaterial = nullptr;
 UMaterialInterface* GColorMaskRegionMaterial = nullptr;
-UTexture2D* GLandscapeBlackTexture = nullptr;
 UMaterialInterface* GLandscapeLayerUsageMaterial = nullptr;
+
+UMaterialInterface* GLayerDebugColorMaterial_DiscreteLOD = nullptr;
+UMaterialInterface* GSelectionColorMaterial_DiscreteLOD = nullptr;
+UMaterialInterface* GSelectionRegionMaterial_DiscreteLOD = nullptr;
+UMaterialInterface* GMaskRegionMaterial_DiscreteLOD = nullptr;
+UMaterialInterface* GColorMaskRegionMaterial_DiscreteLOD = nullptr;
+UMaterialInterface* GLandscapeLayerUsageMaterial_DiscreteLOD = nullptr;
+
+UTexture2D* GLandscapeBlackTexture = nullptr;
 #endif
 
 void ULandscapeComponent::GetUsedMaterials(TArray<UMaterialInterface*>& OutMaterials, bool bGetDebugMaterials) const
@@ -617,12 +626,24 @@ void ULandscapeComponent::GetUsedMaterials(TArray<UMaterialInterface*>& OutMater
 #if WITH_EDITOR
 	//if (bGetDebugMaterials) // TODO: This should be tested and enabled
 	{
-		OutMaterials.Add(GLayerDebugColorMaterial);
-		OutMaterials.Add(GSelectionColorMaterial);
-		OutMaterials.Add(GSelectionRegionMaterial);
-		OutMaterials.Add(GMaskRegionMaterial);
-		OutMaterials.Add(GColorMaskRegionMaterial);
-		OutMaterials.Add(GLandscapeLayerUsageMaterial);
+		if (!GetLandscapeProxy()->bUseDiscreteLOD)
+		{
+			OutMaterials.Add(GLayerDebugColorMaterial);
+			OutMaterials.Add(GSelectionColorMaterial);
+			OutMaterials.Add(GSelectionRegionMaterial);
+			OutMaterials.Add(GMaskRegionMaterial);
+			OutMaterials.Add(GColorMaskRegionMaterial);
+			OutMaterials.Add(GLandscapeLayerUsageMaterial);
+		}
+		else
+		{
+			OutMaterials.Add(GLayerDebugColorMaterial_DiscreteLOD);
+			OutMaterials.Add(GSelectionColorMaterial_DiscreteLOD);
+			OutMaterials.Add(GSelectionRegionMaterial_DiscreteLOD);
+			OutMaterials.Add(GMaskRegionMaterial_DiscreteLOD);
+			OutMaterials.Add(GColorMaskRegionMaterial_DiscreteLOD);
+			OutMaterials.Add(GLandscapeLayerUsageMaterial_DiscreteLOD);
+		}
 	}
 #endif
 }
@@ -639,6 +660,7 @@ FLandscapeComponentSceneProxy::FLandscapeComponentSceneProxy(ULandscapeComponent
 	: FPrimitiveSceneProxy(InComponent, NAME_LandscapeResourceNameForDebugging)
 	, FLandscapeNeighborInfo(InComponent->GetWorld(), InComponent->GetLandscapeProxy()->GetLandscapeGuid(), InComponent->GetSectionBase() / InComponent->ComponentSizeQuads, InComponent->GetHeightmap(), InComponent->ForcedLOD, InComponent->LODBias)
 	, MaxLOD(FMath::CeilLogTwo(InComponent->SubsectionSizeQuads + 1) - 1)
+	, bUseDiscreteLOD(InComponent->GetLandscapeProxy()->bUseDiscreteLOD)
 	, UseTessellationComponentScreenSizeFalloff(InComponent->GetLandscapeProxy()->UseTessellationComponentScreenSizeFalloff)
 	, bRequiresAdjacencyInformation(false)
 	, NumWeightmapLayerAllocations(InComponent->GetWeightmapLayerAllocations().Num())
@@ -668,6 +690,7 @@ FLandscapeComponentSceneProxy::FLandscapeComponentSceneProxy(ULandscapeComponent
 	, SharedBuffersKey(0)
 	, SharedBuffers(nullptr)
 	, VertexFactory(nullptr)
+	, FixedGridVertexFactory(nullptr)
 	, ComponentLightInfo(nullptr)
 #if WITH_EDITORONLY_DATA
 	, EditToolRenderData(InComponent->EditToolRenderData)
@@ -940,7 +963,7 @@ void FLandscapeComponentSceneProxy::CreateRenderThreadResources()
 			SharedBuffers->VertexFactory = LandscapeXYOffsetVertexFactory;
 		}
 
-		if (UseVirtualTexturing(FeatureLevel))
+		if (UseVirtualTexturing(FeatureLevel) || bUseDiscreteLOD)
 		{
 			//todo[vt]: We will need a version of this to support XYOffsetmapTexture
 			FLandscapeFixedGridVertexFactory* LandscapeVertexFactory = new FLandscapeFixedGridVertexFactory(FeatureLevel);
@@ -999,7 +1022,7 @@ void FLandscapeComponentSceneProxy::CreateRenderThreadResources()
 		
 		// Grass is being generated using LOD0 material only
 		FMaterialRenderProxy* RenderProxy = AvailableMaterials[LODIndexToMaterialIndex[0]]->GetRenderProxy();
-		GrassMeshBatch.VertexFactory = VertexFactory;
+		GrassMeshBatch.VertexFactory = bUseDiscreteLOD ? FixedGridVertexFactory : VertexFactory;
 		GrassMeshBatch.MaterialRenderProxy = RenderProxy;
 		GrassMeshBatch.LCI = nullptr;
 		GrassMeshBatch.ReverseCulling = false;
@@ -1127,7 +1150,7 @@ FPrimitiveViewRelevance FLandscapeComponentSceneProxy::GetViewRelevance(const FS
 	Result.bRenderCustomDepth = ShouldRenderCustomDepth();
 	Result.bUsesLightingChannels = GetLightingChannelMask() != GetDefaultLightingChannelMask();
 	Result.bTranslucentSelfShadow = bCastVolumetricTranslucentShadow;
-	Result.bUseCustomViewData = true;
+	Result.bUseCustomViewData = !bUseDiscreteLOD;
 
 	auto FeatureLevel = View->GetFeatureLevel();
 
@@ -1167,30 +1190,30 @@ FPrimitiveViewRelevance FLandscapeComponentSceneProxy::GetViewRelevance(const FS
 			if (EditToolRenderData.SelectedType)
 			{
 				if ((GLandscapeEditRenderMode & ELandscapeEditRenderMode::SelectRegion) && (EditToolRenderData.SelectedType & FLandscapeEditToolRenderData::ST_REGION)
-					&& !(GLandscapeEditRenderMode & ELandscapeEditRenderMode::Mask) && GSelectionRegionMaterial)
+					&& !(GLandscapeEditRenderMode & ELandscapeEditRenderMode::Mask) && (bUseDiscreteLOD ? GSelectionRegionMaterial_DiscreteLOD : GSelectionRegionMaterial))
 				{
 					Result.bDynamicRelevance = true;
-					ToolRelevance |= GSelectionRegionMaterial->GetRelevance_Concurrent(FeatureLevel);
+					ToolRelevance |= (bUseDiscreteLOD ? GSelectionRegionMaterial_DiscreteLOD : GSelectionRegionMaterial)->GetRelevance_Concurrent(FeatureLevel);
 				}
-				if ((GLandscapeEditRenderMode & ELandscapeEditRenderMode::SelectComponent) && (EditToolRenderData.SelectedType & FLandscapeEditToolRenderData::ST_COMPONENT) && GSelectionColorMaterial)
+				if ((GLandscapeEditRenderMode & ELandscapeEditRenderMode::SelectComponent) && (EditToolRenderData.SelectedType & FLandscapeEditToolRenderData::ST_COMPONENT) && (bUseDiscreteLOD ? GSelectionColorMaterial_DiscreteLOD : GSelectionColorMaterial))
 				{
 					Result.bDynamicRelevance = true;
-					ToolRelevance |= GSelectionColorMaterial->GetRelevance_Concurrent(FeatureLevel);
+					ToolRelevance |= (bUseDiscreteLOD ? GSelectionColorMaterial_DiscreteLOD : GSelectionColorMaterial)->GetRelevance_Concurrent(FeatureLevel);
 				}
 			}
 
 			// Mask
-			if ((GLandscapeEditRenderMode & ELandscapeEditRenderMode::Mask) && GMaskRegionMaterial != nullptr &&
+			if ((GLandscapeEditRenderMode & ELandscapeEditRenderMode::Mask) && (bUseDiscreteLOD ? GMaskRegionMaterial_DiscreteLOD : GMaskRegionMaterial) != nullptr &&
 				(((EditToolRenderData.SelectedType & FLandscapeEditToolRenderData::ST_REGION)) || (!(GLandscapeEditRenderMode & ELandscapeEditRenderMode::InvertedMask))))
 			{
 				Result.bDynamicRelevance = true;
-				ToolRelevance |= GMaskRegionMaterial->GetRelevance_Concurrent(FeatureLevel);
+				ToolRelevance |= (bUseDiscreteLOD ? GMaskRegionMaterial_DiscreteLOD : GMaskRegionMaterial)->GetRelevance_Concurrent(FeatureLevel);
 			}
 
 			if (GLandscapeViewMode == ELandscapeViewMode::LayerContribution)
 			{
 				Result.bDynamicRelevance = true;
-				ToolRelevance |= GColorMaskRegionMaterial->GetRelevance_Concurrent(FeatureLevel);
+				ToolRelevance |= (bUseDiscreteLOD ? GColorMaskRegionMaterial_DiscreteLOD : GColorMaskRegionMaterial)->GetRelevance_Concurrent(FeatureLevel);
 			}
 
 			ToolRelevance.SetPrimitiveViewRelevance(Result);
@@ -1199,7 +1222,8 @@ FPrimitiveViewRelevance FLandscapeComponentSceneProxy::GetViewRelevance(const FS
 
 	// Various visualizations need to render using dynamic relevance
 	if ((View->Family->EngineShowFlags.Bounds && IsSelected()) ||
-		GLandscapeDebugOptions.bShowPatches)
+		GLandscapeDebugOptions.bShowPatches
+	)
 	{
 		Result.bDynamicRelevance = true;
 	}
@@ -1219,6 +1243,7 @@ FPrimitiveViewRelevance FLandscapeComponentSceneProxy::GetViewRelevance(const FS
 #if WITH_EDITOR
 		(IsSelected() && !GLandscapeEditModeActive) ||
 		(GLandscapeViewMode != ELandscapeViewMode::Normal) ||
+		(View->Family->LandscapeLODOverride >= 0) ||
 #else
 		IsSelected() ||
 #endif
@@ -1341,6 +1366,7 @@ void FLandscapeComponentSceneProxy::OnTransformChanged()
 
 	// Set FLandscapeUniformVSParameters for this subsection
 	FLandscapeUniformShaderParameters LandscapeParams;
+	LandscapeParams.LastLOD = LastLOD;
 	LandscapeParams.HeightmapUVScaleBias = HeightmapScaleBias;
 	LandscapeParams.WeightmapUVScaleBias = WeightmapScaleBias;
 	LandscapeParams.LocalToWorldNoScaling = LocalToWorldNoScaling;
@@ -1436,14 +1462,18 @@ float FLandscapeComponentSceneProxy::GetComponentScreenSize(const FSceneView* Vi
 	return SquaredScreenRadius * 2.0f;
 }
 
-void FLandscapeComponentSceneProxy::BuildDynamicMeshElement(const FViewCustomDataLOD* InPrimitiveCustomData, bool InToolMesh, bool InHasTessellation, bool InDisableTessellation, FMeshBatch& OutMeshBatch, TArray<FLandscapeBatchElementParams, SceneRenderingAllocator>& OutStaticBatchParamArray) const
+extern ENGINE_API int32 GFarShadowStaticMeshLODBias;
+
+void FLandscapeComponentSceneProxy::BuildDynamicMeshElement(const FSceneView* View, const FViewCustomDataLOD* InPrimitiveCustomData, bool InToolMesh, bool InHasTessellation, bool InDisableTessellation, FMeshBatch& OutMeshBatch, TArray<FLandscapeBatchElementParams, SceneRenderingAllocator>& OutStaticBatchParamArray) const
 {
 	if (AvailableMaterials.Num() == 0 || InPrimitiveCustomData == nullptr || InPrimitiveCustomData->SubSections.Num() == 0)
 	{
 		return;
 	}
 
-	const int8 CurrentLODIndex = InPrimitiveCustomData->SubSections[0].BatchElementCurrentLOD;
+	const int8 CurrentLODIndex = (View->DrawDynamicFlags & EDrawDynamicFlags::FarShadowCascade) ?
+		(FMath::Min(InPrimitiveCustomData->SubSections[0].BatchElementCurrentLOD + GFarShadowStaticMeshLODBias, LastLOD)) : 
+		InPrimitiveCustomData->SubSections[0].BatchElementCurrentLOD;
 	int8 MaterialIndex = LODIndexToMaterialIndex.IsValidIndex(CurrentLODIndex) ? LODIndexToMaterialIndex[CurrentLODIndex] : INDEX_NONE;
 	UMaterialInterface* SelectedMaterial = MaterialIndex != INDEX_NONE ? AvailableMaterials[MaterialIndex] : nullptr;
 
@@ -1469,7 +1499,7 @@ void FLandscapeComponentSceneProxy::BuildDynamicMeshElement(const FViewCustomDat
 		check(SharedBuffers->AdjacencyIndexBuffers);
 	}
 
-	OutMeshBatch.VertexFactory = VertexFactory;
+	OutMeshBatch.VertexFactory = bUseDiscreteLOD ? FixedGridVertexFactory : VertexFactory;
 	OutMeshBatch.MaterialRenderProxy = SelectedMaterial->GetRenderProxy();
 	OutMeshBatch.LCI = ComponentLightInfo.Get();
 	OutMeshBatch.ReverseCulling = IsLocalToWorldDeterminantNegative();
@@ -1488,7 +1518,9 @@ void FLandscapeComponentSceneProxy::BuildDynamicMeshElement(const FViewCustomDat
 			for (int32 SubX = 0; SubX < NumSubsections; SubX++)
 			{
 				const int8 SubSectionIdx = SubX + SubY * NumSubsections;
-				const int8 CurrentLOD = InPrimitiveCustomData->SubSections[SubSectionIdx].BatchElementCurrentLOD;
+				const int8 CurrentLOD = (View->DrawDynamicFlags & EDrawDynamicFlags::FarShadowCascade) ?
+					(FMath::Min(InPrimitiveCustomData->SubSections[SubSectionIdx].BatchElementCurrentLOD + GFarShadowStaticMeshLODBias, LastLOD)) :
+					InPrimitiveCustomData->SubSections[SubSectionIdx].BatchElementCurrentLOD;
 
 				FMeshBatchElement BatchElement;
 				FLandscapeBatchElementParams& BatchElementParams = OutStaticBatchParamArray[SubSectionIdx];
@@ -1585,7 +1617,7 @@ bool FLandscapeComponentSceneProxy::GetMeshElement(bool UseSeperateBatchForShado
 		check(SharedBuffers->AdjacencyIndexBuffers);
 	}
 
-	OutMeshBatch.VertexFactory = VertexFactory;
+	OutMeshBatch.VertexFactory = bUseDiscreteLOD ? FixedGridVertexFactory : VertexFactory;
 	OutMeshBatch.MaterialRenderProxy = InMaterialInterface->GetRenderProxy();
 
 	OutMeshBatch.LCI = ComponentLightInfo.Get();
@@ -1794,68 +1826,71 @@ void FLandscapeComponentSceneProxy::DrawStaticElements(FStaticPrimitiveDrawInter
 	StaticBatchParamArray.Empty(TotalBatchCount);
 	PDI->ReserveMemoryForMeshes(TotalBatchCount);
 
-	MaterialIndexToStaticMeshBatchLOD.SetNumUninitialized(MaterialCount);
-
-	int32 CurrentLODIndex = 0;
-	for (int32 i = 0; i < MaterialCount; ++i)
+	if (!bUseDiscreteLOD)
 	{
-		if (AvailableMaterials[i] == nullptr)
+		MaterialIndexToStaticMeshBatchLOD.SetNumUninitialized(MaterialCount);
+
+		int32 CurrentLODIndex = 0;
+		for (int32 i = 0; i < MaterialCount; ++i)
 		{
-			continue;
-		}
+			if (AvailableMaterials[i] == nullptr)
+			{
+				continue;
+			}
 
-		UMaterialInstance* MaterialInstance = Cast<UMaterialInstance>(AvailableMaterials[i]);
+			UMaterialInstance* MaterialInstance = Cast<UMaterialInstance>(AvailableMaterials[i]);
 
 		bool HasTessellationEnabled = (FeatureLevel >= ERHIFeatureLevel::SM5) ? MaterialInstance != nullptr && MaterialInstance->GetMaterial()->D3D11TessellationMode != EMaterialTessellationMode::MTM_NoTessellation && MaterialIndexToDisabledTessellationMaterial[i] != INDEX_NONE : false;
 
-		// Only add normal batch index (for each material)
-		MaterialIndexToStaticMeshBatchLOD[i] = CurrentLODIndex;
+			// Only add normal batch index (for each material)
+			MaterialIndexToStaticMeshBatchLOD[i] = CurrentLODIndex;
 
-		// Default Batch, tessellated if enabled
-		FMeshBatch NormalMeshBatch;
+			// Default Batch, tessellated if enabled
+			FMeshBatch NormalMeshBatch;
 
-		if (GetMeshElement(false, false, HasTessellationEnabled, CurrentLODIndex++, AvailableMaterials[i], NormalMeshBatch, StaticBatchParamArray))
-		{
-			PDI->DrawMesh(NormalMeshBatch, FLT_MAX);
-		}
+			if (GetMeshElement(false, false, HasTessellationEnabled, CurrentLODIndex++, AvailableMaterials[i], NormalMeshBatch, StaticBatchParamArray))
+			{
+				PDI->DrawMesh(NormalMeshBatch, FLT_MAX);
+			}
 
-		if (HasTessellationEnabled)
-		{
-			UMaterialInstance* NonTessellatedLandscapeMI = Cast<UMaterialInstance>(AvailableMaterials[MaterialIndexToDisabledTessellationMaterial[i]]);
+			if (HasTessellationEnabled)
+			{
+				UMaterialInstance* NonTessellatedLandscapeMI = Cast<UMaterialInstance>(AvailableMaterials[MaterialIndexToDisabledTessellationMaterial[i]]);
 
-			// Make sure that the Material instance we are going to use has the tessellation disabled
-			UMaterialInstanceDynamic* NonTessellatedLandscapeMID = Cast<UMaterialInstanceDynamic>(NonTessellatedLandscapeMI);
-			ULandscapeMaterialInstanceConstant* NonTessellatedLandscapeMIC = Cast<ULandscapeMaterialInstanceConstant>(NonTessellatedLandscapeMI);
+				// Make sure that the Material instance we are going to use has the tessellation disabled
+				UMaterialInstanceDynamic* NonTessellatedLandscapeMID = Cast<UMaterialInstanceDynamic>(NonTessellatedLandscapeMI);
+				ULandscapeMaterialInstanceConstant* NonTessellatedLandscapeMIC = Cast<ULandscapeMaterialInstanceConstant>(NonTessellatedLandscapeMI);
 
-			if (NonTessellatedLandscapeMID != nullptr)
-			{
-				NonTessellatedLandscapeMIC = Cast<ULandscapeMaterialInstanceConstant>(NonTessellatedLandscapeMID->Parent);
-			}
+				if (NonTessellatedLandscapeMID != nullptr)
+				{
+					NonTessellatedLandscapeMIC = Cast<ULandscapeMaterialInstanceConstant>(NonTessellatedLandscapeMID->Parent);
+				}
 
-			check(NonTessellatedLandscapeMIC != nullptr && NonTessellatedLandscapeMIC->bDisableTessellation);
+				check(NonTessellatedLandscapeMIC != nullptr && NonTessellatedLandscapeMIC->bDisableTessellation);
 
-			// Default Batch, no Tessellation enabled
-			FMeshBatch NonTesselatedNormalMeshBatch;
+				// Default Batch, no Tessellation enabled
+				FMeshBatch NonTesselatedNormalMeshBatch;
 
-			if (GetMeshElement(true, false, false, CurrentLODIndex++, NonTessellatedLandscapeMI, NonTesselatedNormalMeshBatch, StaticBatchParamArray))
-			{
-				PDI->DrawMesh(NonTesselatedNormalMeshBatch, FLT_MAX);
-			}
+				if (GetMeshElement(true, false, false, CurrentLODIndex++, NonTessellatedLandscapeMI, NonTesselatedNormalMeshBatch, StaticBatchParamArray))
+				{
+					PDI->DrawMesh(NonTesselatedNormalMeshBatch, FLT_MAX);
+				}
 
-			// Shadow Batch, tessellated if enabled
-			FMeshBatch ShadowMeshBatch;
+				// Shadow Batch, tessellated if enabled
+				FMeshBatch ShadowMeshBatch;
 
-			if (GetMeshElement(true, true, HasTessellationEnabled, CurrentLODIndex++, MaterialInstance, ShadowMeshBatch, StaticBatchParamArray))
-			{
-				PDI->DrawMesh(ShadowMeshBatch, FLT_MAX);
-			}
+				if (GetMeshElement(true, true, HasTessellationEnabled, CurrentLODIndex++, MaterialInstance, ShadowMeshBatch, StaticBatchParamArray))
+				{
+					PDI->DrawMesh(ShadowMeshBatch, FLT_MAX);
+				}
 
-			// Shadow Batch, no tessellation enabled
-			FMeshBatch NonTesselatedShadowMeshBatch;
+				// Shadow Batch, no tessellation enabled
+				FMeshBatch NonTesselatedShadowMeshBatch;
 
-			if (GetMeshElement(true, true, false, CurrentLODIndex++, NonTessellatedLandscapeMI, NonTesselatedShadowMeshBatch, StaticBatchParamArray))
-			{
-				PDI->DrawMesh(NonTesselatedShadowMeshBatch, FLT_MAX);
+				if (GetMeshElement(true, true, false, CurrentLODIndex++, NonTessellatedLandscapeMI, NonTesselatedShadowMeshBatch, StaticBatchParamArray))
+				{
+					PDI->DrawMesh(NonTesselatedShadowMeshBatch, FLT_MAX);
+				}
 			}
 		}
 	}
@@ -1875,6 +1910,113 @@ void FLandscapeComponentSceneProxy::DrawStaticElements(FStaticPrimitiveDrawInter
 		}
  	}
 
+	if (bUseDiscreteLOD)
+	{
+		for (int32 LODIndex = FirstLOD; LODIndex <= LastLOD; LODIndex++)
+		{
+			UMaterialInstance* MaterialInstance = nullptr;
+			UMaterialInstance* NonTessellatedLandscapeMI = nullptr;
+			bool HasTessellationEnabled = false;
+
+			{
+				int8 MaterialIndex = LODIndexToMaterialIndex[LODIndex];
+
+				MaterialInstance = Cast<UMaterialInstance>(AvailableMaterials[MaterialIndex]);
+
+				HasTessellationEnabled = (FeatureLevel >= ERHIFeatureLevel::SM5) ?
+					(MaterialInstance != nullptr && MaterialInstance->GetMaterial()->D3D11TessellationMode != EMaterialTessellationMode::MTM_NoTessellation && MaterialIndexToDisabledTessellationMaterial[MaterialIndex] != INDEX_NONE) :
+					false;
+
+				if (HasTessellationEnabled)
+				{
+					NonTessellatedLandscapeMI = Cast<UMaterialInstance>(AvailableMaterials[MaterialIndexToDisabledTessellationMaterial[MaterialIndex]]);
+
+					// Make sure that the Material instance we are going to use has the tessellation disabled
+					UMaterialInstanceDynamic* NonTessellatedLandscapeMID = Cast<UMaterialInstanceDynamic>(NonTessellatedLandscapeMI);
+					ULandscapeMaterialInstanceConstant* NonTessellatedLandscapeMIC = Cast<ULandscapeMaterialInstanceConstant>(NonTessellatedLandscapeMI);
+
+					if (NonTessellatedLandscapeMID != nullptr)
+					{
+						NonTessellatedLandscapeMIC = Cast<ULandscapeMaterialInstanceConstant>(NonTessellatedLandscapeMID->Parent);
+					}
+
+					check(NonTessellatedLandscapeMIC != nullptr && NonTessellatedLandscapeMIC->bDisableTessellation);
+				}
+			}
+
+			if (MaterialInstance == nullptr)
+			{
+				continue;
+			}
+
+			if (HasTessellationEnabled && LODScreenRatioSquared[LODIndex] < TessellationComponentSquaredScreenSize)
+			{
+				// Disable tessellation by replacing MaterialInstance with the non tessellated version
+				MaterialInstance = NonTessellatedLandscapeMI;
+			}
+
+			// Could be different from bRequiresAdjacencyInformation during shader compilation
+			bool bCurrentRequiresAdjacencyInformation = MaterialRenderingRequiresAdjacencyInformation_RenderingThread(MaterialInstance, VertexFactory->GetType(), GetScene().GetFeatureLevel());
+
+			if (bCurrentRequiresAdjacencyInformation)
+			{
+				check(SharedBuffers->AdjacencyIndexBuffers);
+			}
+
+			FMeshBatch BaseMeshBatch;
+
+			BaseMeshBatch.VertexFactory = FixedGridVertexFactory;
+			BaseMeshBatch.MaterialRenderProxy = MaterialInstance->GetRenderProxy();
+
+			BaseMeshBatch.LCI = ComponentLightInfo.Get();
+			BaseMeshBatch.ReverseCulling = IsLocalToWorldDeterminantNegative();
+			BaseMeshBatch.CastShadow = true;
+			BaseMeshBatch.bUseForDepthPass = true;
+			BaseMeshBatch.bUseAsOccluder = ShouldUseAsOccluder() && GetScene().GetShadingPath() == EShadingPath::Deferred && !IsMovable();
+			BaseMeshBatch.bUseForMaterial = true;
+			BaseMeshBatch.Type = bCurrentRequiresAdjacencyInformation ? PT_12_ControlPointPatchList : PT_TriangleList;
+			BaseMeshBatch.DepthPriorityGroup = SDPG_World;
+			BaseMeshBatch.LODIndex = LODIndex;
+			BaseMeshBatch.bRequiresPerElementVisibility = false;
+			BaseMeshBatch.bDitheredLODTransition = false;
+
+			if (BaseMeshBatch.CastShadow)
+			{
+				float QuadWorldSize = (ComponentMaxExtend * 2.0f) / (SubsectionSizeQuads * NumSubsections);
+				FVector QuadSize = FVector(QuadWorldSize, QuadWorldSize, bCurrentRequiresAdjacencyInformation ? BaseMeshBatch.MaterialRenderProxy->GetMaterial(GetScene().GetFeatureLevel())->GetMaxDisplacement() : 0.0f);
+
+				BaseMeshBatch.TessellationDisablingShadowMapMeshSize = QuadSize.Size();
+			}
+
+			int32 LodSubsectionSizeVerts = SubsectionSizeVerts >> LODIndex;
+			uint32 NumPrimitivesPerSection = FMath::Square((LodSubsectionSizeVerts - 1)) * 2;
+
+			// Combined batch element
+			FLandscapeBatchElementParams* BatchElementParams = new(StaticBatchParamArray) FLandscapeBatchElementParams;
+			BatchElementParams->LandscapeUniformShaderParametersResource = &LandscapeUniformShaderParameters;
+			BatchElementParams->LocalToWorldNoScalingPtr = &LocalToWorldNoScaling;
+			BatchElementParams->SceneProxy = this;
+			BatchElementParams->SubX = -1;
+			BatchElementParams->SubY = -1;
+			BatchElementParams->CurrentLOD = LODIndex;
+
+			FMeshBatchElement BatchElement;
+
+			BatchElement.UserData = BatchElementParams;
+			BatchElement.PrimitiveUniformBuffer = GetUniformBuffer();
+			BatchElement.IndexBuffer = bCurrentRequiresAdjacencyInformation ? SharedBuffers->AdjacencyIndexBuffers->IndexBuffers[LODIndex] : SharedBuffers->IndexBuffers[LODIndex];
+			BatchElement.NumPrimitives = FMath::Square((LodSubsectionSizeVerts - 1)) * FMath::Square(NumSubsections) * 2;
+			BatchElement.FirstIndex = 0;
+			BatchElement.MinVertexIndex = SharedBuffers->IndexRanges[LODIndex].MinIndexFull;
+			BatchElement.MaxVertexIndex = SharedBuffers->IndexRanges[LODIndex].MaxIndexFull;
+
+			FMeshBatch MeshBatch = BaseMeshBatch;
+			MeshBatch.Elements[0] = BatchElement;
+
+			PDI->DrawMesh(MeshBatch, LODIndex == FirstLOD ? FLT_MAX : (FMath::Sqrt(LODScreenRatioSquared[LODIndex]) * 2.0f));
+		}
+	}
+
 	check(StaticBatchParamArray.Num() <= TotalBatchCount);
 }
 
@@ -2479,12 +2621,12 @@ void FLandscapeComponentSceneProxy::PostInitViewCustomData(const FSceneView& InV
 
 bool FLandscapeComponentSceneProxy::IsUsingCustomLODRules() const
 {
-	return true;
+	return !bUseDiscreteLOD;
 }
 
 bool FLandscapeComponentSceneProxy::IsUsingCustomWholeSceneShadowLODRules() const
 {
-	return true;
+	return !bUseDiscreteLOD;
 }
 
 bool FLandscapeComponentSceneProxy::CanUseMeshBatchForShadowCascade(int8 InLODIndex, float InShadowMapTextureResolution, float InShadowMapCascadeSize) const
@@ -2736,7 +2878,12 @@ void FLandscapeComponentSceneProxy::GetDynamicMeshElements(const TArray<const FS
 
 			if (PrimitiveCustomData == nullptr)
 			{
-				continue;
+				float MeshScreenSizeSquared = ComputeBoundsScreenRadiusSquared(GetBounds().Origin, GetBounds().SphereRadius, *View);
+				int32 ForcedLODLevel = (View->Family->EngineShowFlags.LOD) ? GetCVarForceLOD() : 0;
+
+				FLODMask LODToRender;
+				LODToRender.SetLOD(GetLODFromScreenSize(MeshScreenSizeSquared, View->LODDistanceFactor));
+				PrimitiveCustomData = (const FViewCustomDataLOD*)(const_cast<FLandscapeComponentSceneProxy*>(this)->InitViewCustomData(*View, View->LODDistanceFactor, FMemStack::Get(), false, false, &LODToRender, MeshScreenSizeSquared));
 			}
 
 			ParameterArray.ElementParams.AddDefaulted(NumSubsections * NumSubsections);
@@ -2782,13 +2929,13 @@ void FLandscapeComponentSceneProxy::GetDynamicMeshElements(const TArray<const FS
 			}
 
 			FMeshBatch& Mesh = Collector.AllocateMesh();
-			BuildDynamicMeshElement(PrimitiveCustomData, false, HasTessellationEnabled, DisableTessellation, Mesh, ParameterArray.ElementParams);
+			BuildDynamicMeshElement(View, PrimitiveCustomData, false, HasTessellationEnabled, DisableTessellation, Mesh, ParameterArray.ElementParams);
 
 #if WITH_EDITOR
 			FMeshBatch& MeshTools = Collector.AllocateMesh();
 
 			// No Tessellation on tool material
-			BuildDynamicMeshElement(PrimitiveCustomData, true, HasTessellationEnabled, true, MeshTools, ParameterArray.ElementParams);
+			BuildDynamicMeshElement(View, PrimitiveCustomData, true, HasTessellationEnabled, true, MeshTools, ParameterArray.ElementParams);
 #endif
 
 			// Render the landscape component
@@ -2797,9 +2944,9 @@ void FLandscapeComponentSceneProxy::GetDynamicMeshElements(const TArray<const FS
 			{
 			case ELandscapeViewMode::DebugLayer:
 			{
-				if (GLayerDebugColorMaterial)
+				if ((bUseDiscreteLOD ? GLayerDebugColorMaterial_DiscreteLOD : GLayerDebugColorMaterial))
 				{
-					auto DebugColorMaterialInstance = new FLandscapeDebugMaterialRenderProxy(GLayerDebugColorMaterial->GetRenderProxy(),
+					auto DebugColorMaterialInstance = new FLandscapeDebugMaterialRenderProxy((bUseDiscreteLOD ? GLayerDebugColorMaterial_DiscreteLOD : GLayerDebugColorMaterial)->GetRenderProxy(),
 						(EditToolRenderData.DebugChannelR >= 0 ? WeightmapTextures[EditToolRenderData.DebugChannelR / 4] : nullptr),
 						(EditToolRenderData.DebugChannelG >= 0 ? WeightmapTextures[EditToolRenderData.DebugChannelG / 4] : nullptr),
 						(EditToolRenderData.DebugChannelB >= 0 ? WeightmapTextures[EditToolRenderData.DebugChannelB / 4] : nullptr),
@@ -2844,10 +2991,10 @@ void FLandscapeComponentSceneProxy::GetDynamicMeshElements(const TArray<const FS
 
 			case ELandscapeViewMode::LayerUsage:
 			{
-				if (GLandscapeLayerUsageMaterial)
+				if ((bUseDiscreteLOD ? GLandscapeLayerUsageMaterial_DiscreteLOD : GLandscapeLayerUsageMaterial))
 				{
 					float Rotation = ((SectionBase.X / ComponentSizeQuads) ^ (SectionBase.Y / ComponentSizeQuads)) & 1 ? 0 : 2.f * PI;
-					auto LayerUsageMaterialInstance = new FLandscapeLayerUsageRenderProxy(GLandscapeLayerUsageMaterial->GetRenderProxy(), ComponentSizeVerts, LayerColors, Rotation);
+					auto LayerUsageMaterialInstance = new FLandscapeLayerUsageRenderProxy((bUseDiscreteLOD ? GLandscapeLayerUsageMaterial_DiscreteLOD : GLandscapeLayerUsageMaterial)->GetRenderProxy(), ComponentSizeVerts, LayerColors, Rotation);
 					MeshTools.MaterialRenderProxy = LayerUsageMaterialInstance;
 					Collector.RegisterOneFrameMaterialProxy(LayerUsageMaterialInstance);
 					MeshTools.bCanApplyViewModeOverrides = true;
@@ -2924,7 +3071,7 @@ void FLandscapeComponentSceneProxy::GetDynamicMeshElements(const TArray<const FS
 
 				FMeshBatch& MaskMesh = Collector.AllocateMesh();
 				MaskMesh = MeshTools;
-				auto ColorMaskMaterialInstance = new FLandscapeMaskMaterialRenderProxy(GColorMaskRegionMaterial->GetRenderProxy(), EditToolRenderData.LayerContributionTexture ? EditToolRenderData.LayerContributionTexture : GLandscapeBlackTexture, true);
+				auto ColorMaskMaterialInstance = new FLandscapeMaskMaterialRenderProxy((bUseDiscreteLOD ? GColorMaskRegionMaterial_DiscreteLOD : GColorMaskRegionMaterial)->GetRenderProxy(), EditToolRenderData.LayerContributionTexture ? EditToolRenderData.LayerContributionTexture : GLandscapeBlackTexture, true);
 				MaskMesh.MaterialRenderProxy = ColorMaskMaterialInstance;
 				Collector.RegisterOneFrameMaterialProxy(ColorMaskMaterialInstance);
 				Collector.AddMesh(ViewIndex, MaskMesh);
@@ -2969,6 +3116,7 @@ void FLandscapeComponentSceneProxy::GetDynamicMeshElements(const TArray<const FS
 						bIsWireframe ||
 #if WITH_EDITOR
 						(IsSelected() && !GLandscapeEditModeActive) ||
+						ViewFamily.LandscapeLODOverride >= 0 ||
 #else
 						IsSelected() ||
 #endif
@@ -3000,7 +3148,7 @@ void FLandscapeComponentSceneProxy::GetDynamicMeshElements(const TArray<const FS
 					{
 						FMeshBatch& SelectMesh = Collector.AllocateMesh();
 						SelectMesh = MeshTools;
-						auto SelectMaterialInstance = new FLandscapeSelectMaterialRenderProxy(GSelectionRegionMaterial->GetRenderProxy(), EditToolRenderData.DataTexture ? EditToolRenderData.DataTexture : GLandscapeBlackTexture);
+						auto SelectMaterialInstance = new FLandscapeSelectMaterialRenderProxy((bUseDiscreteLOD ? GSelectionRegionMaterial_DiscreteLOD : GSelectionRegionMaterial)->GetRenderProxy(), EditToolRenderData.DataTexture ? EditToolRenderData.DataTexture : GLandscapeBlackTexture);
 						SelectMesh.MaterialRenderProxy = SelectMaterialInstance;
 						Collector.RegisterOneFrameMaterialProxy(SelectMaterialInstance);
 						Collector.AddMesh(ViewIndex, SelectMesh);
@@ -3013,7 +3161,7 @@ void FLandscapeComponentSceneProxy::GetDynamicMeshElements(const TArray<const FS
 					{
 						FMeshBatch& SelectMesh = Collector.AllocateMesh();
 						SelectMesh = MeshTools;
-						SelectMesh.MaterialRenderProxy = GSelectionColorMaterial->GetRenderProxy();
+						SelectMesh.MaterialRenderProxy = (bUseDiscreteLOD ? GSelectionColorMaterial_DiscreteLOD : GSelectionColorMaterial)->GetRenderProxy();
 						Collector.AddMesh(ViewIndex, SelectMesh);
 						NumPasses++;
 						NumTriangles += SelectMesh.GetNumPrimitives();
@@ -3028,7 +3176,7 @@ void FLandscapeComponentSceneProxy::GetDynamicMeshElements(const TArray<const FS
 					{
 						FMeshBatch& MaskMesh = Collector.AllocateMesh();
 						MaskMesh = MeshTools;
-						auto MaskMaterialInstance = new FLandscapeMaskMaterialRenderProxy(GMaskRegionMaterial->GetRenderProxy(), EditToolRenderData.DataTexture ? EditToolRenderData.DataTexture : GLandscapeBlackTexture, !!(GLandscapeEditRenderMode & ELandscapeEditRenderMode::InvertedMask));
+						auto MaskMaterialInstance = new FLandscapeMaskMaterialRenderProxy((bUseDiscreteLOD ? GMaskRegionMaterial_DiscreteLOD : GMaskRegionMaterial)->GetRenderProxy(), EditToolRenderData.DataTexture ? EditToolRenderData.DataTexture : GLandscapeBlackTexture, !!(GLandscapeEditRenderMode & ELandscapeEditRenderMode::InvertedMask));
 						MaskMesh.MaterialRenderProxy = MaskMaterialInstance;
 						Collector.RegisterOneFrameMaterialProxy(MaskMaterialInstance);
 						Collector.AddMesh(ViewIndex, MaskMesh);
@@ -3040,7 +3188,7 @@ void FLandscapeComponentSceneProxy::GetDynamicMeshElements(const TArray<const FS
 					{
 						FMeshBatch& MaskMesh = Collector.AllocateMesh();
 						MaskMesh = MeshTools;
-						auto MaskMaterialInstance = new FLandscapeMaskMaterialRenderProxy(GMaskRegionMaterial->GetRenderProxy(), GLandscapeBlackTexture, false);
+						auto MaskMaterialInstance = new FLandscapeMaskMaterialRenderProxy((bUseDiscreteLOD ? GMaskRegionMaterial_DiscreteLOD : GMaskRegionMaterial)->GetRenderProxy(), GLandscapeBlackTexture, false);
 						MaskMesh.MaterialRenderProxy = MaskMaterialInstance;
 						Collector.RegisterOneFrameMaterialProxy(MaskMaterialInstance);
 						Collector.AddMesh(ViewIndex, MaskMesh);
@@ -3658,6 +3806,7 @@ FLandscapeSharedBuffers::FLandscapeSharedBuffers(const int32 InSharedBuffersKey,
 	, SubsectionSizeVerts(InSubsectionSizeQuads + 1)
 	, NumSubsections(InNumSubsections)
 	, VertexFactory(nullptr)
+	, FixedGridVertexFactory(nullptr)
 	, VertexBuffer(nullptr)
 	, AdjacencyIndexBuffers(nullptr)
 	, bUse32BitIndices(false)
@@ -3743,6 +3892,11 @@ FLandscapeSharedBuffers::~FLandscapeSharedBuffers()
 	delete AdjacencyIndexBuffers;
 	delete VertexFactory;
 
+	if (FixedGridVertexFactory != nullptr)
+	{
+		delete FixedGridVertexFactory;
+	}
+
 	if (OccluderIndicesSP.IsValid())
 	{
 		DEC_DWORD_STAT_BY(STAT_LandscapeOccluderMem, OccluderIndicesSP->GetAllocatedSize());
@@ -3909,108 +4063,126 @@ public:
 
 		int32 SubSectionIndex = BatchElementParams->SubX + BatchElementParams->SubY * SceneProxy->NumSubsections;
 
-		// If we have no custom data for this primitive we will compute of the fly the proper values, this will happen if the shader is not used for normal landscape rendering(i.e grass rendering)
-		FLandscapeComponentSceneProxy::FViewCustomDataLOD* LODData = (FLandscapeComponentSceneProxy::FViewCustomDataLOD*)InView->GetCustomData(SceneProxy->GetPrimitiveSceneInfo()->GetIndex());
+		if (SceneProxy->bUseDiscreteLOD)
+		{
+			// In fact when discrete LOD is used, we're using FLandscapeFixedGridVertexFactory and not going through here
+			// However in certain cases (like in editor landscape tools) this will still be used
+			// Set parameters accordingly to get the same behavior as discrete LOD
+			FVector4 LodTessellationParams(ForceInitToZero);
+			FVector4 LodBias(ForceInitToZero);
+			FVector4 ShaderCurrentLOD(BatchElementParams->CurrentLOD, BatchElementParams->CurrentLOD, BatchElementParams->CurrentLOD, BatchElementParams->CurrentLOD);
+			FVector4 ShaderCurrentNeighborLOD[FLandscapeComponentSceneProxy::NEIGHBOR_COUNT] = { ShaderCurrentLOD, ShaderCurrentLOD, ShaderCurrentLOD, ShaderCurrentLOD };
 
-		if (LODData != nullptr)
+			ShaderBindings.Add(LodTessellationParameter, LodTessellationParams);
+			ShaderBindings.Add(LodBiasParameter, LodBias);
+			ShaderBindings.Add(SectionLodsParameter, ShaderCurrentLOD);
+			ShaderBindings.Add(NeighborSectionLodParameter, ShaderCurrentNeighborLOD);
+		}
+		else
 		{
-			SceneProxy->PostInitViewCustomData(*InView, LODData);
+			// If we have no custom data for this primitive we will compute of the fly the proper values, this will happen if the shader is not used for normal landscape rendering(i.e grass rendering)
+			FLandscapeComponentSceneProxy::FViewCustomDataLOD* LODData = (FLandscapeComponentSceneProxy::FViewCustomDataLOD*)InView->GetCustomData(SceneProxy->GetPrimitiveSceneInfo()->GetIndex());
 
-			if (LodTessellationParameter.IsBound())
+			if (LODData != nullptr)
 			{
-				ShaderBindings.Add(LodTessellationParameter, LODData->LodTessellationParams);
-			}
-				
-			if (LodBiasParameter.IsBound())
-			{
-				check(LODData->LodBias == SceneProxy->GetShaderLODBias());
-				ShaderBindings.Add(LodBiasParameter, LODData->LodBias);
-			}
+				SceneProxy->PostInitViewCustomData(*InView, LODData);
 
-			if (SectionLodsParameter.IsBound())
-			{
-				if (LODData->UseCombinedMeshBatch)
+				if (LodTessellationParameter.IsBound())
 				{
-					ShaderBindings.Add(SectionLodsParameter, LODData->ShaderCurrentLOD);
+					ShaderBindings.Add(LodTessellationParameter, LODData->LodTessellationParams);
 				}
-				else // in non combined, only the one representing us as we'll be called 4 times (once per sub section)
-				{
-					check(SubSectionIndex >= 0);
-					FVector4 ShaderCurrentLOD(ForceInitToZero);
-					ShaderCurrentLOD.Component(SubSectionIndex) = LODData->ShaderCurrentLOD.Component(SubSectionIndex);
 
-					ShaderBindings.Add(SectionLodsParameter, ShaderCurrentLOD);
+				if (LodBiasParameter.IsBound())
+				{
+					check(LODData->LodBias == SceneProxy->GetShaderLODBias());
+					ShaderBindings.Add(LodBiasParameter, LODData->LodBias);
 				}
-			}
-
-			if (NeighborSectionLodParameter.IsBound())
-			{
-				FVector4 ShaderCurrentNeighborLOD[FLandscapeComponentSceneProxy::NEIGHBOR_COUNT] = { FVector4(ForceInitToZero), FVector4(ForceInitToZero), FVector4(ForceInitToZero), FVector4(ForceInitToZero) };
-
-				if (LODData->UseCombinedMeshBatch)
-				{					
-					int32 SubSectionCount = SceneProxy->NumSubsections == 1 ? 1 : FLandscapeComponentSceneProxy::MAX_SUBSECTION_COUNT;
 
-					for (int32 NeighborSubSectionIndex = 0; NeighborSubSectionIndex < SubSectionCount; ++NeighborSubSectionIndex)
+				if (SectionLodsParameter.IsBound())
+				{
+					if (LODData->UseCombinedMeshBatch)
 					{
-						ShaderCurrentNeighborLOD[NeighborSubSectionIndex] = LODData->SubSections[NeighborSubSectionIndex].ShaderCurrentNeighborLOD;
-						check(ShaderCurrentNeighborLOD[NeighborSubSectionIndex].X != -1.0f); // they should all match so only check the 1st one for simplicity
+						ShaderBindings.Add(SectionLodsParameter, LODData->ShaderCurrentLOD);
 					}
+					else // in non combined, only the one representing us as we'll be called 4 times (once per sub section)
+					{
+						check(SubSectionIndex >= 0);
+						FVector4 ShaderCurrentLOD(ForceInitToZero);
+						ShaderCurrentLOD.Component(SubSectionIndex) = LODData->ShaderCurrentLOD.Component(SubSectionIndex);
 
-					ShaderBindings.Add(NeighborSectionLodParameter, ShaderCurrentNeighborLOD);
+						ShaderBindings.Add(SectionLodsParameter, ShaderCurrentLOD);
+					}
 				}
-				else // in non combined, only the one representing us as we'll be called 4 times (once per sub section)
+
+				if (NeighborSectionLodParameter.IsBound())
 				{
-					check(SubSectionIndex >= 0);
-					ShaderCurrentNeighborLOD[SubSectionIndex] = LODData->SubSections[SubSectionIndex].ShaderCurrentNeighborLOD;
-					check(ShaderCurrentNeighborLOD[SubSectionIndex].X != -1.0f); // they should all match so only check the 1st one for simplicity
+					FVector4 ShaderCurrentNeighborLOD[FLandscapeComponentSceneProxy::NEIGHBOR_COUNT] = { FVector4(ForceInitToZero), FVector4(ForceInitToZero), FVector4(ForceInitToZero), FVector4(ForceInitToZero) };
+
+					if (LODData->UseCombinedMeshBatch)
+					{
+						int32 SubSectionCount = SceneProxy->NumSubsections == 1 ? 1 : FLandscapeComponentSceneProxy::MAX_SUBSECTION_COUNT;
 
-					ShaderBindings.Add(NeighborSectionLodParameter, ShaderCurrentNeighborLOD);
+						for (int32 NeighborSubSectionIndex = 0; NeighborSubSectionIndex < SubSectionCount; ++NeighborSubSectionIndex)
+						{
+							ShaderCurrentNeighborLOD[NeighborSubSectionIndex] = LODData->SubSections[NeighborSubSectionIndex].ShaderCurrentNeighborLOD;
+							check(ShaderCurrentNeighborLOD[NeighborSubSectionIndex].X != -1.0f); // they should all match so only check the 1st one for simplicity
+						}
+
+						ShaderBindings.Add(NeighborSectionLodParameter, ShaderCurrentNeighborLOD);
+					}
+					else // in non combined, only the one representing us as we'll be called 4 times (once per sub section)
+					{
+						check(SubSectionIndex >= 0);
+						ShaderCurrentNeighborLOD[SubSectionIndex] = LODData->SubSections[SubSectionIndex].ShaderCurrentNeighborLOD;
+						check(ShaderCurrentNeighborLOD[SubSectionIndex].X != -1.0f); // they should all match so only check the 1st one for simplicity
+
+						ShaderBindings.Add(NeighborSectionLodParameter, ShaderCurrentNeighborLOD);
+					}
 				}
-			}			
-		}
-		else
-		{
+			}
+			else
+			{
 			float ComponentScreenSize = SceneProxy->GetComponentScreenSize(InView, SceneProxy->GetBounds().Origin, SceneProxy->ComponentMaxExtend, SceneProxy->GetBounds().SphereRadius);
-			
-			FLandscapeComponentSceneProxy::FViewCustomDataLOD CurrentLODData;
-			SceneProxy->CalculateBatchElementLOD(*InView, ComponentScreenSize, InView->LODDistanceFactor, CurrentLODData, true);
-			check(CurrentLODData.UseCombinedMeshBatch);
 
-			if (LodBiasParameter.IsBound())
-			{
-				ShaderBindings.Add(LodBiasParameter, SceneProxy->GetShaderLODBias());
-			}
+				FLandscapeComponentSceneProxy::FViewCustomDataLOD CurrentLODData;
+				SceneProxy->CalculateBatchElementLOD(*InView, ComponentScreenSize, InView->LODDistanceFactor, CurrentLODData, true);
+				check(CurrentLODData.UseCombinedMeshBatch);
 
-			if (LodTessellationParameter.IsBound())
-			{
-				FVector4 LodTessellationParams(ForceInitToZero);
-				SceneProxy->ComputeTessellationFalloffShaderValues(CurrentLODData, InView->ViewMatrices.GetProjectionMatrix(), LodTessellationParams.X, LodTessellationParams.Y);
+				if (LodBiasParameter.IsBound())
+				{
+					ShaderBindings.Add(LodBiasParameter, SceneProxy->GetShaderLODBias());
+				}
 
-				ShaderBindings.Add(LodTessellationParameter, LodTessellationParams);
-			}
+				if (LodTessellationParameter.IsBound())
+				{
+					FVector4 LodTessellationParams(ForceInitToZero);
+					SceneProxy->ComputeTessellationFalloffShaderValues(CurrentLODData, InView->ViewMatrices.GetProjectionMatrix(), LodTessellationParams.X, LodTessellationParams.Y);
 
-			if (SectionLodsParameter.IsBound())
-			{
-				ShaderBindings.Add(SectionLodsParameter, CurrentLODData.ShaderCurrentLOD);
-			}
+					ShaderBindings.Add(LodTessellationParameter, LodTessellationParams);
+				}
 
-			if (NeighborSectionLodParameter.IsBound())
-			{
-				FVector4 CurrentNeighborLOD[4] = { FVector4(ForceInitToZero), FVector4(ForceInitToZero), FVector4(ForceInitToZero), FVector4(ForceInitToZero) };
+				if (SectionLodsParameter.IsBound())
+				{
+					ShaderBindings.Add(SectionLodsParameter, CurrentLODData.ShaderCurrentLOD);
+				}
 
-				for (int32 SubY = 0; SubY < SceneProxy->NumSubsections; SubY++)
+				if (NeighborSectionLodParameter.IsBound())
 				{
-					for (int32 SubX = 0; SubX < SceneProxy->NumSubsections; SubX++)
+					FVector4 CurrentNeighborLOD[4] = { FVector4(ForceInitToZero), FVector4(ForceInitToZero), FVector4(ForceInitToZero), FVector4(ForceInitToZero) };
+
+					for (int32 SubY = 0; SubY < SceneProxy->NumSubsections; SubY++)
 					{
-						int32 NeighborSubSectionIndex = SubX + SubY * SceneProxy->NumSubsections;
-						SceneProxy->GetShaderCurrentNeighborLOD(*InView, CurrentLODData.SubSections[NeighborSubSectionIndex].fBatchElementCurrentLOD, SceneProxy->NumSubsections > 1 ? SubX : INDEX_NONE, SceneProxy->NumSubsections > 1 ? SubY : INDEX_NONE, NeighborSubSectionIndex, CurrentNeighborLOD[NeighborSubSectionIndex]);
-						check(CurrentNeighborLOD[NeighborSubSectionIndex].X != -1.0f); // they should all match so only check the 1st one for simplicity
+						for (int32 SubX = 0; SubX < SceneProxy->NumSubsections; SubX++)
+						{
+							int32 NeighborSubSectionIndex = SubX + SubY * SceneProxy->NumSubsections;
+							SceneProxy->GetShaderCurrentNeighborLOD(*InView, CurrentLODData.SubSections[NeighborSubSectionIndex].fBatchElementCurrentLOD, SceneProxy->NumSubsections > 1 ? SubX : INDEX_NONE, SceneProxy->NumSubsections > 1 ? SubY : INDEX_NONE, NeighborSubSectionIndex, CurrentNeighborLOD[NeighborSubSectionIndex]);
+							check(CurrentNeighborLOD[NeighborSubSectionIndex].X != -1.0f); // they should all match so only check the 1st one for simplicity
+						}
 					}
-				}
 
-				ShaderBindings.Add(NeighborSectionLodParameter, CurrentNeighborLOD);
-			}				
+					ShaderBindings.Add(NeighborSectionLodParameter, CurrentNeighborLOD);
+				}
+			}
 		}
 	}
 
@@ -4056,15 +4228,16 @@ public:
 
 		ShaderBindings.Add(Shader->GetUniformBufferParameter<FLandscapeUniformShaderParameters>(), *BatchElementParams->LandscapeUniformShaderParametersResource);
 
-		if (LodValuesParameter.IsBound())
-		{
-			ShaderBindings.Add(LodValuesParameter, SceneProxy->GetShaderLODValues(BatchElementParams->CurrentLOD));
-		}
+		FVector4 LodTessellationParams(ForceInitToZero);
+		FVector4 LodBias(ForceInitToZero);
+		FVector4 ShaderCurrentLOD(BatchElementParams->CurrentLOD, BatchElementParams->CurrentLOD, BatchElementParams->CurrentLOD, BatchElementParams->CurrentLOD);
+		FVector4 ShaderCurrentNeighborLOD[FLandscapeComponentSceneProxy::NEIGHBOR_COUNT] = { ShaderCurrentLOD, ShaderCurrentLOD, ShaderCurrentLOD, ShaderCurrentLOD };
 
-		if (LodBiasParameter.IsBound())
-		{
-			ShaderBindings.Add(LodBiasParameter, FVector4(ForceInitToZero));
-		}
+		ShaderBindings.Add(LodValuesParameter, SceneProxy->GetShaderLODValues(BatchElementParams->CurrentLOD));
+		ShaderBindings.Add(LodTessellationParameter, LodTessellationParams);
+		ShaderBindings.Add(LodBiasParameter, LodBias);
+		ShaderBindings.Add(SectionLodsParameter, ShaderCurrentLOD);
+		ShaderBindings.Add(NeighborSectionLodParameter, ShaderCurrentNeighborLOD);
 	}
 };
 
@@ -4139,7 +4312,7 @@ void FLandscapeVertexFactory::ModifyCompilationEnvironment(const FVertexFactoryT
 }
 
 
-IMPLEMENT_VERTEX_FACTORY_TYPE(FLandscapeVertexFactory, "/Engine/Private/LandscapeVertexFactory.ush", true, true, true, false, false);
+IMPLEMENT_VERTEX_FACTORY_TYPE_EX(FLandscapeVertexFactory, "/Engine/Private/LandscapeVertexFactory.ush", true, true, true, false, false, true, false);
 
 /**
 * Copy the data from another vertex factory
@@ -4195,7 +4368,7 @@ FVertexFactoryShaderParameters* FLandscapeFixedGridVertexFactory::ConstructShade
 	}
 }
 
-IMPLEMENT_VERTEX_FACTORY_TYPE(FLandscapeFixedGridVertexFactory, "/Engine/Private/LandscapeVertexFactory.ush", true, true, true, false, false);
+IMPLEMENT_VERTEX_FACTORY_TYPE_EX(FLandscapeFixedGridVertexFactory, "/Engine/Private/LandscapeVertexFactory.ush", true, true, true, false, false, true, false);
 
 
 /** ULandscapeMaterialInstanceConstant */
@@ -4211,6 +4384,7 @@ class FLandscapeMaterialResource : public FMaterialResource
 	const bool bDisableTessellation;
 	const bool bMobile;
 	const bool bEditorToolUsage;
+	const bool bUseDiscreteLOD;
 
 public:
 	FLandscapeMaterialResource(ULandscapeMaterialInstanceConstant* Parent)
@@ -4218,6 +4392,7 @@ public:
 		, bDisableTessellation(Parent->bDisableTessellation)
 		, bMobile(Parent->bMobile)
 		, bEditorToolUsage(Parent->bEditorToolUsage)
+		, bUseDiscreteLOD(Parent->bUseDiscreteLOD)
 	{
 	}
 
@@ -4317,20 +4492,35 @@ public:
 				// Todo: only compile LandscapeXYOffsetVertexFactory if we are using it
 				bool bIsRuntimeVirtualTextureShaderType = Algo::Find(GetRuntimeVirtualTextureShaderTypes(), ShaderType->GetFName()) != nullptr;
 
-				static const FName LandscapeVertexFactory = FName(TEXT("FLandscapeVertexFactory"));
-				static const FName LandscapeXYOffsetVertexFactory = FName(TEXT("FLandscapeXYOffsetVertexFactory"));
-				static const FName LandscapeVertexFactoryMobile = FName(TEXT("FLandscapeVertexFactoryMobile"));
-				if (VertexFactoryType->GetFName() == LandscapeVertexFactory ||
-					VertexFactoryType->GetFName() == LandscapeXYOffsetVertexFactory ||
-					VertexFactoryType->GetFName() == LandscapeVertexFactoryMobile)
+				if (bIsRuntimeVirtualTextureShaderType)
 				{
-					return !bIsRuntimeVirtualTextureShaderType && FMaterialResource::ShouldCache(Platform, ShaderType, VertexFactoryType);
+					static const FName LandscapeFixedGridVertexFactory = FName(TEXT("FLandscapeFixedGridVertexFactory"));
+					if (VertexFactoryType->GetFName() == LandscapeFixedGridVertexFactory)
+					{
+						return FMaterialResource::ShouldCache(Platform, ShaderType, VertexFactoryType);
+					}
 				}
-
-				static const FName LandscapeFixedGridVertexFactory = FName(TEXT("FLandscapeFixedGridVertexFactory"));
-				if (VertexFactoryType->GetFName() == LandscapeFixedGridVertexFactory)
+				else
 				{
-					return bIsRuntimeVirtualTextureShaderType && FMaterialResource::ShouldCache(Platform, ShaderType, VertexFactoryType);
+					static const FName LandscapeFixedGridVertexFactory = FName(TEXT("FLandscapeFixedGridVertexFactory"));
+					if (VertexFactoryType->GetFName() == LandscapeFixedGridVertexFactory)
+					{
+						return bUseDiscreteLOD && FMaterialResource::ShouldCache(Platform, ShaderType, VertexFactoryType);
+					}
+
+					static const FName LandscapeFixedGridVertexFactoryMobile = FName(TEXT("FLandscapeFixedGridVertexFactoryMobile"));
+					if (VertexFactoryType->GetFName() == LandscapeFixedGridVertexFactoryMobile)
+					{
+						return bUseDiscreteLOD && FMaterialResource::ShouldCache(Platform, ShaderType, VertexFactoryType);
+					}
+
+					static const FName LandscapeVertexFactory = FName(TEXT("FLandscapeVertexFactory"));
+					static const FName LandscapeXYOffsetVertexFactory = FName(TEXT("FLandscapeXYOffsetVertexFactory"));
+					if (VertexFactoryType->GetFName() == LandscapeVertexFactory ||
+						VertexFactoryType->GetFName() == LandscapeXYOffsetVertexFactory)
+					{
+						return !bUseDiscreteLOD && FMaterialResource::ShouldCache(Platform, ShaderType, VertexFactoryType);
+					}
 				}
 			}
 		}
diff --git a/Engine/Source/Runtime/Landscape/Private/LandscapeRenderMobile.cpp b/Engine/Source/Runtime/Landscape/Private/LandscapeRenderMobile.cpp
--- a/Engine/Source/Runtime/Landscape/Private/LandscapeRenderMobile.cpp
+++ b/Engine/Source/Runtime/Landscape/Private/LandscapeRenderMobile.cpp
@@ -34,6 +34,27 @@ void FLandscapeVertexFactoryMobile::InitRHI()
 	InitDeclaration(Elements);
 }
 
+void FLandscapeFixedGridVertexFactoryMobile::InitRHI()
+{
+	// list of declaration items
+	FVertexDeclarationElementList Elements;
+
+	// position decls
+	Elements.Add(AccessStreamComponent(MobileData.PositionComponent, 0));
+
+	if (MobileData.LODHeightsComponent.Num())
+	{
+		const int32 BaseAttribute = 1;
+		for (int32 Index = 0; Index < MobileData.LODHeightsComponent.Num(); Index++)
+		{
+			Elements.Add(AccessStreamComponent(MobileData.LODHeightsComponent[Index], BaseAttribute + Index));
+		}
+	}
+
+	// create the actual device decls
+	InitDeclaration(Elements);
+}
+
 /** Shader parameters for use with FLandscapeVertexFactory */
 class FLandscapeVertexFactoryMobileVertexShaderParameters : public FVertexFactoryShaderParameters
 {
@@ -172,6 +193,108 @@ protected:
 	TShaderUniformBufferParameter<FLandscapeUniformShaderParameters> LandscapeShaderParameters;
 };
 
+
+/** Shader parameters for use with FLandscapeVertexFactory */
+class FLandscapeFixedGridVertexFactoryMobileVertexShaderParameters : public FVertexFactoryShaderParameters
+{
+public:
+	/**
+	* Bind shader constants by name
+	* @param	ParameterMap - mapping of named shader constants to indices
+	*/
+	virtual void Bind(const FShaderParameterMap& ParameterMap) override
+	{
+		LodValuesParameter.Bind(ParameterMap, TEXT("LodValues"));
+		LodTessellationParameter.Bind(ParameterMap, TEXT("LodTessellationParams"));
+		NeighborSectionLodParameter.Bind(ParameterMap, TEXT("NeighborSectionLod"));
+		LodBiasParameter.Bind(ParameterMap, TEXT("LodBias"));
+		SectionLodsParameter.Bind(ParameterMap, TEXT("SectionLods"));
+	}
+
+	/**
+	* Serialize shader params to an archive
+	* @param	Ar - archive to serialize to
+	*/
+	virtual void Serialize(FArchive& Ar) override
+	{
+		Ar << LodValuesParameter;
+		Ar << LodTessellationParameter;
+		Ar << NeighborSectionLodParameter;
+		Ar << LodBiasParameter;
+		Ar << SectionLodsParameter;
+	}
+
+	virtual void GetElementShaderBindings(
+		const class FSceneInterface* Scene,
+		const FSceneView* InView,
+		const class FMeshMaterialShader* Shader,
+		const EVertexInputStreamType InputStreamType,
+		ERHIFeatureLevel::Type FeatureLevel,
+		const FVertexFactory* VertexFactory,
+		const FMeshBatchElement& BatchElement,
+		class FMeshDrawSingleShaderBindings& ShaderBindings,
+		FVertexInputStreamArray& VertexStreams
+	) const override
+	{
+		SCOPE_CYCLE_COUNTER(STAT_LandscapeVFDrawTimeVS);
+
+		const FLandscapeBatchElementParams* BatchElementParams = (const FLandscapeBatchElementParams*)BatchElement.UserData;
+		check(BatchElementParams);
+
+		const FLandscapeComponentSceneProxyMobile* SceneProxy = (const FLandscapeComponentSceneProxyMobile*)BatchElementParams->SceneProxy;
+		ShaderBindings.Add(Shader->GetUniformBufferParameter<FLandscapeUniformShaderParameters>(), *BatchElementParams->LandscapeUniformShaderParametersResource);
+
+		if (LodValuesParameter.IsBound())
+		{
+			FVector4 LodValues(
+				0.0f, // this is the mesh's LOD, ES2 always uses the LOD0 mesh
+				0.0f, // unused
+				(float)SceneProxy->SubsectionSizeQuads,
+				1.f / (float)SceneProxy->SubsectionSizeQuads);
+
+			ShaderBindings.Add(LodValuesParameter, LodValues);
+		}
+
+		if (LodBiasParameter.IsBound())
+		{
+			FVector CameraLocalPos3D = SceneProxy->WorldToLocal.TransformPosition(InView->ViewMatrices.GetViewOrigin());
+
+			FVector4 LodBias(
+				0.0f, // unused
+				0.0f, // unused
+				CameraLocalPos3D.X + SceneProxy->SectionBase.X,
+				CameraLocalPos3D.Y + SceneProxy->SectionBase.Y
+			);
+			ShaderBindings.Add(LodBiasParameter, LodBias);
+		}
+
+		if (LodTessellationParameter.IsBound())
+		{
+			FVector4 LodTessellationParams(ForceInitToZero);
+			ShaderBindings.Add(LodTessellationParameter, LodTessellationParams);
+		}
+
+		if (SectionLodsParameter.IsBound())
+		{
+			FVector4 ShaderCurrentLOD(ForceInitToZero);
+			ShaderBindings.Add(SectionLodsParameter, ShaderCurrentLOD);
+		}
+
+		if (NeighborSectionLodParameter.IsBound())
+		{
+			FVector4 ShaderCurrentNeighborLOD[FLandscapeComponentSceneProxy::NEIGHBOR_COUNT] = { FVector4(ForceInitToZero), FVector4(ForceInitToZero), FVector4(ForceInitToZero), FVector4(ForceInitToZero) };
+			ShaderBindings.Add(NeighborSectionLodParameter, ShaderCurrentNeighborLOD);
+		}
+	}
+protected:
+	FShaderParameter LodValuesParameter;
+	FShaderParameter LodTessellationParameter;
+	FShaderParameter NeighborSectionLodParameter;
+	FShaderParameter LodBiasParameter;
+	FShaderParameter SectionLodsParameter;
+	TShaderUniformBufferParameter<FLandscapeUniformShaderParameters> LandscapeShaderParameters;
+};
+
 /** Shader parameters for use with FLandscapeVertexFactory */
 class FLandscapeVertexFactoryMobilePixelShaderParameters : public FLandscapeVertexFactoryPixelShaderParameters
 {
@@ -243,7 +366,21 @@ FVertexFactoryShaderParameters* FLandscapeVertexFactoryMobile::ConstructShaderPa
 	}
 }
 
+FVertexFactoryShaderParameters* FLandscapeFixedGridVertexFactoryMobile::ConstructShaderParameters(EShaderFrequency ShaderFrequency)
+{
+	switch (ShaderFrequency)
+	{
+	case SF_Vertex:
+		return new FLandscapeFixedGridVertexFactoryMobileVertexShaderParameters();
+	case SF_Pixel:
+		return new FLandscapeVertexFactoryMobilePixelShaderParameters();
+	default:
+		return NULL;
+	}
+}
+
 IMPLEMENT_VERTEX_FACTORY_TYPE(FLandscapeVertexFactoryMobile, "/Engine/Private/LandscapeVertexFactory.ush", true, true, true, false, false);
+IMPLEMENT_VERTEX_FACTORY_TYPE_EX(FLandscapeFixedGridVertexFactoryMobile, "/Engine/Private/LandscapeVertexFactory.ush", true, true, true, false, false, true, false);
 
 /**
 * Initialize the RHI for this rendering resource
@@ -369,6 +506,12 @@ FLandscapeComponentSceneProxyMobile::~FLandscapeComponentSceneProxyMobile()
 		delete VertexFactory;
 		VertexFactory = NULL;
 	}
+
+	if (FixedGridVertexFactory)
+	{
+		delete FixedGridVertexFactory;
+		FixedGridVertexFactory = NULL;
+	}
 }
 
 SIZE_T FLandscapeComponentSceneProxyMobile::GetTypeHash() const
@@ -406,16 +549,31 @@ void FLandscapeComponentSceneProxyMobile::CreateRenderThreadResources()
 	check(MobileRenderData->VertexBuffer);
 	MobileRenderData->VertexBuffer->InitResource();
 
-	FLandscapeVertexFactoryMobile* LandscapeVertexFactory = new FLandscapeVertexFactoryMobile(FeatureLevel);
-	LandscapeVertexFactory->MobileData.PositionComponent = FVertexStreamComponent(MobileRenderData->VertexBuffer, STRUCT_OFFSET(FLandscapeMobileVertex,Position), sizeof(FLandscapeMobileVertex), VET_UByte4N);
-	for( uint32 Index = 0; Index < LANDSCAPE_MAX_ES_LOD_COMP; ++Index )
 	{
-		LandscapeVertexFactory->MobileData.LODHeightsComponent.Add
-			(FVertexStreamComponent(MobileRenderData->VertexBuffer, STRUCT_OFFSET(FLandscapeMobileVertex,LODHeights) + sizeof(uint8) * 4 * Index, sizeof(FLandscapeMobileVertex), VET_UByte4N));
+		FLandscapeVertexFactoryMobile* LandscapeVertexFactory = new FLandscapeVertexFactoryMobile(FeatureLevel);
+		LandscapeVertexFactory->MobileData.PositionComponent = FVertexStreamComponent(MobileRenderData->VertexBuffer, STRUCT_OFFSET(FLandscapeMobileVertex, Position), sizeof(FLandscapeMobileVertex), VET_UByte4N);
+		for (uint32 Index = 0; Index < LANDSCAPE_MAX_ES_LOD_COMP; ++Index)
+		{
+			LandscapeVertexFactory->MobileData.LODHeightsComponent.Add
+			(FVertexStreamComponent(MobileRenderData->VertexBuffer, STRUCT_OFFSET(FLandscapeMobileVertex, LODHeights) + sizeof(uint8) * 4 * Index, sizeof(FLandscapeMobileVertex), VET_UByte4N));
+		}
+
+		LandscapeVertexFactory->InitResource();
+		VertexFactory = LandscapeVertexFactory;
 	}
 
-	LandscapeVertexFactory->InitResource();
-	VertexFactory = LandscapeVertexFactory;
+	{
+		FLandscapeFixedGridVertexFactoryMobile* LandscapeVertexFactory = new FLandscapeFixedGridVertexFactoryMobile(FeatureLevel);
+		LandscapeVertexFactory->MobileData.PositionComponent = FVertexStreamComponent(MobileRenderData->VertexBuffer, STRUCT_OFFSET(FLandscapeMobileVertex, Position), sizeof(FLandscapeMobileVertex), VET_UByte4N);
+		for (uint32 Index = 0; Index < LANDSCAPE_MAX_ES_LOD_COMP; ++Index)
+		{
+			LandscapeVertexFactory->MobileData.LODHeightsComponent.Add
+			(FVertexStreamComponent(MobileRenderData->VertexBuffer, STRUCT_OFFSET(FLandscapeMobileVertex, LODHeights) + sizeof(uint8) * 4 * Index, sizeof(FLandscapeMobileVertex), VET_UByte4N));
+		}
+
+		LandscapeVertexFactory->InitResource();
+		FixedGridVertexFactory = LandscapeVertexFactory;
+	}
 
 	// Assign LandscapeUniformShaderParameters
 	LandscapeUniformShaderParameters.InitResource();
diff --git a/Engine/Source/Runtime/Landscape/Public/LandscapeRender.h b/Engine/Source/Runtime/Landscape/Public/LandscapeRender.h
--- a/Engine/Source/Runtime/Landscape/Public/LandscapeRender.h
+++ b/Engine/Source/Runtime/Landscape/Public/LandscapeRender.h
@@ -75,19 +75,28 @@ namespace ELandscapeEditRenderMode
 
 LANDSCAPE_API extern bool GLandscapeEditModeActive;
 LANDSCAPE_API extern int32 GLandscapeEditRenderMode;
+
 LANDSCAPE_API extern UMaterialInterface* GLayerDebugColorMaterial;
 LANDSCAPE_API extern UMaterialInterface* GSelectionColorMaterial;
 LANDSCAPE_API extern UMaterialInterface* GSelectionRegionMaterial;
 LANDSCAPE_API extern UMaterialInterface* GMaskRegionMaterial;
 LANDSCAPE_API extern UMaterialInterface* GColorMaskRegionMaterial;
-LANDSCAPE_API extern UTexture2D* GLandscapeBlackTexture;
 LANDSCAPE_API extern UMaterialInterface* GLandscapeLayerUsageMaterial;
+
+LANDSCAPE_API extern UMaterialInterface* GLayerDebugColorMaterial_DiscreteLOD;
+LANDSCAPE_API extern UMaterialInterface* GSelectionColorMaterial_DiscreteLOD;
+LANDSCAPE_API extern UMaterialInterface* GSelectionRegionMaterial_DiscreteLOD;
+LANDSCAPE_API extern UMaterialInterface* GMaskRegionMaterial_DiscreteLOD;
+LANDSCAPE_API extern UMaterialInterface* GColorMaskRegionMaterial_DiscreteLOD;
+LANDSCAPE_API extern UMaterialInterface* GLandscapeLayerUsageMaterial_DiscreteLOD;
+LANDSCAPE_API extern UTexture2D* GLandscapeBlackTexture;
 #endif
 
 
 /** The uniform shader parameters for a landscape draw call. */
 BEGIN_GLOBAL_SHADER_PARAMETER_STRUCT(FLandscapeUniformShaderParameters, LANDSCAPE_API)
 /** vertex shader parameters */
+SHADER_PARAMETER(int32, LastLOD)
 SHADER_PARAMETER(FVector4, HeightmapUVScaleBias)
 SHADER_PARAMETER(FVector4, WeightmapUVScaleBias)
 SHADER_PARAMETER(FVector4, LandscapeLightmapScaleBias)
@@ -573,6 +582,7 @@ public:
 
 protected:
 	int8						MaxLOD;		// Maximum LOD level, user override possible
+	bool						bUseDiscreteLOD : 1;
 	bool						UseTessellationComponentScreenSizeFalloff:1;	// Tell if we should apply a Tessellation falloff
 	bool						bRequiresAdjacencyInformation:1;
 	int8						NumWeightmapLayerAllocations;
@@ -718,7 +728,7 @@ protected:
 
 	bool GetMeshElement(bool UseSeperateBatchForShadow, bool ShadowOnly, bool HasTessellation, int8 InLODIndex, UMaterialInterface* InMaterialInterface, FMeshBatch& OutMeshBatch, TArray<FLandscapeBatchElementParams>& OutStaticBatchParamArray) const;
 	bool GetMeshElementForVirtualTexture(int32 InLodIndex, ERuntimeVirtualTextureMaterialType MaterialType, UMaterialInterface* InMaterialInterface, FMeshBatch& OutMeshBatch, TArray<FLandscapeBatchElementParams>& OutStaticBatchParamArray) const;
-	void BuildDynamicMeshElement(const FViewCustomDataLOD* InPrimitiveCustomData, bool InToolMesh, bool InHasTessellation, bool InDisableTessellation, FMeshBatch& OutMeshBatch, TArray<FLandscapeBatchElementParams, SceneRenderingAllocator>& OutStaticBatchParamArray) const;
+	void BuildDynamicMeshElement(const FSceneView* View, const FViewCustomDataLOD* InPrimitiveCustomData, bool InToolMesh, bool InHasTessellation, bool InDisableTessellation, FMeshBatch& OutMeshBatch, TArray<FLandscapeBatchElementParams, SceneRenderingAllocator>& OutStaticBatchParamArray) const;
 
 	float GetComponentScreenSize(const class FSceneView* View, const FVector& Origin,  float MaxExtend, float ElementRadius) const;
 
@@ -753,6 +763,7 @@ public:
 	friend class FLandscapeVertexFactoryMobileVertexShaderParameters;
 	friend class FLandscapeVertexFactoryMobilePixelShaderParameters;
 	friend class FLandscapeFixedGridVertexFactoryVertexShaderParameters;
+	friend class FLandscapeFixedGridVertexFactoryMobileVertexShaderParameters;
 
 	// FLandscapeComponentSceneProxy interface.
 	uint64 GetStaticBatchElementVisibility(const FSceneView& InView, const FMeshBatch* InBatch, const void* InViewCustomData) const;
diff --git a/Engine/Source/Runtime/Landscape/Public/LandscapeRenderMobile.h b/Engine/Source/Runtime/Landscape/Public/LandscapeRenderMobile.h
--- a/Engine/Source/Runtime/Landscape/Public/LandscapeRenderMobile.h
+++ b/Engine/Source/Runtime/Landscape/Public/LandscapeRenderMobile.h
@@ -83,6 +83,68 @@ private:
 	friend class FLandscapeComponentSceneProxyMobile;
 };
 
+/** vertex factory for VTF-heightmap terrain  */
+class FLandscapeFixedGridVertexFactoryMobile : public FLandscapeVertexFactory
+{
+	DECLARE_VERTEX_FACTORY_TYPE(FLandscapeFixedGridVertexFactoryMobile);
+
+	typedef FLandscapeVertexFactory Super;
+public:
+
+	struct FDataType : FLandscapeVertexFactory::FDataType
+	{
+		/** stream which has heights of each LOD levels */
+		TArray<FVertexStreamComponent, TFixedAllocator<LANDSCAPE_MAX_ES_LOD_COMP> > LODHeightsComponent;
+	};
+
+	FLandscapeFixedGridVertexFactoryMobile(ERHIFeatureLevel::Type InFeatureLevel)
+		: FLandscapeVertexFactory(InFeatureLevel)
+	{
+	}
+
+	virtual ~FLandscapeFixedGridVertexFactoryMobile()
+	{
+		ReleaseResource();
+	}
+
+	static FVertexFactoryShaderParameters* ConstructShaderParameters(EShaderFrequency ShaderFrequency);
+
+	/**
+	* Should we cache the material's shadertype on this platform with this vertex factory?
+	*/
+	static bool ShouldCompilePermutation(EShaderPlatform Platform, const class FMaterial* Material, const class FShaderType* ShaderType)
+	{
+		auto FeatureLevel = GetMaxSupportedFeatureLevel(Platform);
+		return (FeatureLevel == ERHIFeatureLevel::ES2 || FeatureLevel == ERHIFeatureLevel::ES3_1) &&
+			(Material->IsUsedWithLandscape() || Material->IsSpecialEngineMaterial());
+	}
+
+	static void ModifyCompilationEnvironment(const FVertexFactoryType* Type, EShaderPlatform Platform, const FMaterial* Material, FShaderCompilerEnvironment& OutEnvironment)
+	{
+		FVertexFactory::ModifyCompilationEnvironment(Type, Platform, Material, OutEnvironment);
+		OutEnvironment.SetDefine(TEXT("NUM_VF_PACKED_INTERPOLANTS"), TEXT("1"));
+		OutEnvironment.SetDefine(TEXT("FIXED_GRID"), TEXT("1"));
+	}
+
+	// FRenderResource interface.
+	virtual void InitRHI() override;
+
+	/**
+	 * An implementation of the interface used by TSynchronizedResource to update the resource with new data from the game thread.
+	 */
+	void SetData(const FDataType& InData)
+	{
+		MobileData = InData;
+		UpdateRHI();
+	}
+
+private:
+	/** stream component data bound to this vertex factory */
+	FDataType MobileData;
+
+	friend class FLandscapeComponentSceneProxyMobile;
+};
+
 //
 // FLandscapeVertexBuffer
 //
diff --git a/Engine/Source/Runtime/Renderer/Private/SceneRendering.cpp b/Engine/Source/Runtime/Renderer/Private/SceneRendering.cpp
--- a/Engine/Source/Runtime/Renderer/Private/SceneRendering.cpp
+++ b/Engine/Source/Runtime/Renderer/Private/SceneRendering.cpp
@@ -1612,6 +1612,16 @@ void FViewInfo::SetupUniformBufferParameters(
 		}
 	}
 
+	if (DrawDynamicFlags & EDrawDynamicFlags::FarShadowCascade)
+	{
+		extern ENGINE_API int32 GFarShadowStaticMeshLODBias;
+		ViewUniformShaderParameters.FarShadowStaticMeshLODBias = GFarShadowStaticMeshLODBias;
+	}
+	else
+	{
+		ViewUniformShaderParameters.FarShadowStaticMeshLODBias = 0;
+	}
+
 	ViewUniformShaderParameters.PreIntegratedBRDF = GEngine->PreIntegratedSkinBRDFTexture->Resource->TextureRHI;
 
 	if (UseGPUScene(GMaxRHIShaderPlatform, RHIFeatureLevel))
diff --git a/Engine/Source/Runtime/Renderer/Private/ShadowDepthRendering.cpp b/Engine/Source/Runtime/Renderer/Private/ShadowDepthRendering.cpp
--- a/Engine/Source/Runtime/Renderer/Private/ShadowDepthRendering.cpp
+++ b/Engine/Source/Runtime/Renderer/Private/ShadowDepthRendering.cpp
@@ -1191,6 +1191,11 @@ void FProjectedShadowInfo::ModifyViewForShadow(FRHICommandList& RHICmdList, FVie
 
 	FoundView->ViewMatrices.HackRemoveTemporalAAProjectionJitter();
 
+	if (CascadeSettings.bFarShadowCascade)
+	{
+		(int32&)FoundView->DrawDynamicFlags |= (int32)EDrawDynamicFlags::FarShadowCascade;
+	}
+
 	// Don't do material texture mip biasing in shadow maps.
 	FoundView->MaterialTextureMipBias = 0;
 
@@ -1223,7 +1228,7 @@ void FProjectedShadowInfo::ModifyViewForShadow(FRHICommandList& RHICmdList, FVie
 
 	if (bPreShadow && GPreshadowsForceLowestLOD)
 	{
-		FoundView->DrawDynamicFlags = EDrawDynamicFlags::ForceLowestLOD;
+		(int32&)FoundView->DrawDynamicFlags |= EDrawDynamicFlags::ForceLowestLOD;
 	}
 }
 
diff --git a/Engine/Source/Runtime/Renderer/Private/ShadowSetup.cpp b/Engine/Source/Runtime/Renderer/Private/ShadowSetup.cpp
--- a/Engine/Source/Runtime/Renderer/Private/ShadowSetup.cpp
+++ b/Engine/Source/Runtime/Renderer/Private/ShadowSetup.cpp
@@ -908,6 +908,21 @@ bool FProjectedShadowInfo::ShouldDrawStaticMeshes(FViewInfo& InCurrentView, bool
 			}
 		}
 
+		if (CascadeSettings.bFarShadowCascade)
+		{
+			extern ENGINE_API int32 GFarShadowStaticMeshLODBias;
+			int8 LODToRenderScan = ShadowLODToRender.DitheredLODIndices[0] + GFarShadowStaticMeshLODBias;
+
+			for (int32 Index = InPrimitiveSceneInfo->StaticMeshRelevances.Num() - 1; Index >= 0; Index--)
+			{
+				if (LODToRenderScan == InPrimitiveSceneInfo->StaticMeshRelevances[Index].LODIndex)
+				{
+					ShadowLODToRender.SetLOD(LODToRenderScan);
+					break;
+				}
+			}
+		}
+
 		if (WholeSceneDirectionalShadow)
 		{
 			if (InPrimitiveSceneInfo->bIsUsingCustomWholeSceneShadowLODRules)
@@ -1351,6 +1366,11 @@ void FProjectedShadowInfo::GatherDynamicMeshElements(FSceneRenderer& Renderer, F
 			ShadowDepthView->DrawDynamicFlags = EDrawDynamicFlags::ForceLowestLOD;
 		}
 
+		if (CascadeSettings.bFarShadowCascade)
+		{
+			(int32&)ShadowDepthView->DrawDynamicFlags |= (int32)EDrawDynamicFlags::FarShadowCascade;
+		}
+
 		if (IsWholeSceneDirectionalShadow())
 		{
 			ShadowDepthView->SetPreShadowTranslation(FVector(0, 0, 0));
