Fix editor crash when trying to change lod settings on the parent landscape actor

diff --git a/Engine/Source/Runtime/Landscape/Private/Landscape.cpp b/Engine/Source/Runtime/Landscape/Private/Landscape.cpp
--- a/Engine/Source/Runtime/Landscape/Private/Landscape.cpp
+++ b/Engine/Source/Runtime/Landscape/Private/Landscape.cpp
@@ -2386,7 +2386,6 @@ void ALandscapeProxy::GetSharedProperties(ALandscapeProxy* Landscape)
 		ComponentScreenSizeToUseSubSections = Landscape->ComponentScreenSizeToUseSubSections;
 		UseTessellationComponentScreenSizeFalloff = Landscape->UseTessellationComponentScreenSizeFalloff;
 		TessellationComponentScreenSizeFalloff = Landscape->TessellationComponentScreenSizeFalloff;
-		bUseDiscreteLOD = Landscape->bUseDiscreteLOD;
 		LODDistributionSetting = Landscape->LODDistributionSetting;
 		LOD0DistributionSetting = Landscape->LOD0DistributionSetting;
 		LOD0ScreenSize = Landscape->LOD0ScreenSize;
