Fix landscape grass rendering

diff --git a/Engine/Source/Runtime/Landscape/Private/LandscapeRender.cpp b/Engine/Source/Runtime/Landscape/Private/LandscapeRender.cpp
--- a/Engine/Source/Runtime/Landscape/Private/LandscapeRender.cpp
+++ b/Engine/Source/Runtime/Landscape/Private/LandscapeRender.cpp
@@ -4228,12 +4228,20 @@ public:
 
 		ShaderBindings.Add(Shader->GetUniformBufferParameter<FLandscapeUniformShaderParameters>(), *BatchElementParams->LandscapeUniformShaderParametersResource);
 
+		int32 CurrentLOD = BatchElementParams->CurrentLOD;
+#if WITH_EDITOR
+		if (InView && InView->Family->LandscapeLODOverride >= 0)
+		{
+			CurrentLOD = InView->Family->LandscapeLODOverride;
+		}
+#endif
+
 		FVector4 LodTessellationParams(ForceInitToZero);
 		FVector4 LodBias(ForceInitToZero);
-		FVector4 ShaderCurrentLOD(BatchElementParams->CurrentLOD, BatchElementParams->CurrentLOD, BatchElementParams->CurrentLOD, BatchElementParams->CurrentLOD);
+		FVector4 ShaderCurrentLOD(CurrentLOD, CurrentLOD, CurrentLOD, CurrentLOD);
 		FVector4 ShaderCurrentNeighborLOD[FLandscapeComponentSceneProxy::NEIGHBOR_COUNT] = { ShaderCurrentLOD, ShaderCurrentLOD, ShaderCurrentLOD, ShaderCurrentLOD };
 
-		ShaderBindings.Add(LodValuesParameter, SceneProxy->GetShaderLODValues(BatchElementParams->CurrentLOD));
+		ShaderBindings.Add(LodValuesParameter, SceneProxy->GetShaderLODValues(CurrentLOD));
 		ShaderBindings.Add(LodTessellationParameter, LodTessellationParams);
 		ShaderBindings.Add(LodBiasParameter, LodBias);
 		ShaderBindings.Add(SectionLodsParameter, ShaderCurrentLOD);
