Implement landscape discrete support for ray tracing

diff --git a/Engine/Source/Runtime/Landscape/Private/LandscapeRender.cpp b/Engine/Source/Runtime/Landscape/Private/LandscapeRender.cpp
--- a/Engine/Source/Runtime/Landscape/Private/LandscapeRender.cpp
+++ b/Engine/Source/Runtime/Landscape/Private/LandscapeRender.cpp
@@ -3250,12 +3250,27 @@ void FLandscapeComponentSceneProxy::GetDynamicRayTracingInstances(FRayTracingMat
 
 	FMemStackBase& PrimitiveCustomDataMemStack = FMemStack::Get();
 
-	float MeshScreenSizeSquared = 0;
-	int32 ForcedLODLevel = (Context.ReferenceView->Family->EngineShowFlags.LOD) ? GetCVarForceLOD() : 0;
+	FViewCustomDataLOD* InPrimitiveCustomData;
 
-	FLODMask LODToRender = GetCustomLOD(*Context.ReferenceView, Context.ReferenceView->LODDistanceFactor, ForcedLODLevel, MeshScreenSizeSquared);
-	FViewCustomDataLOD* InPrimitiveCustomData = (FViewCustomDataLOD*)InitViewCustomData(*Context.ReferenceView, Context.ReferenceView->LODDistanceFactor, PrimitiveCustomDataMemStack, false, false, &LODToRender, MeshScreenSizeSquared);
-	PostInitViewCustomData(*Context.ReferenceView, InPrimitiveCustomData);
+	if (!bUseDiscreteLOD)
+	{
+		float MeshScreenSizeSquared = 0;
+		int32 ForcedLODLevel = (Context.ReferenceView->Family->EngineShowFlags.LOD) ? GetCVarForceLOD() : 0;
+
+		FLODMask LODToRender = GetCustomLOD(*Context.ReferenceView, Context.ReferenceView->LODDistanceFactor, ForcedLODLevel, MeshScreenSizeSquared);
+		InPrimitiveCustomData = (FViewCustomDataLOD*)InitViewCustomData(*Context.ReferenceView, Context.ReferenceView->LODDistanceFactor, PrimitiveCustomDataMemStack, false, false, &LODToRender, MeshScreenSizeSquared);
+		PostInitViewCustomData(*Context.ReferenceView, InPrimitiveCustomData);
+	}
+	else
+	{
+		float MeshScreenSizeSquared = ComputeBoundsScreenRadiusSquared(GetBounds().Origin, GetBounds().SphereRadius, *Context.ReferenceView);
+		int32 ForcedLODLevel = (Context.ReferenceView->Family->EngineShowFlags.LOD) ? GetCVarForceLOD() : 0;
+
+		FLODMask LODToRender;
+		LODToRender.SetLOD(GetLODFromScreenSize(MeshScreenSizeSquared, Context.ReferenceView->LODDistanceFactor));
+		InPrimitiveCustomData = (FViewCustomDataLOD*)InitViewCustomData(*Context.ReferenceView, Context.ReferenceView->LODDistanceFactor, PrimitiveCustomDataMemStack, false, false, &LODToRender, MeshScreenSizeSquared);
+		PostInitViewCustomData(*Context.ReferenceView, InPrimitiveCustomData);
+	}
 
 	FLandscapeElementParamArray& ParameterArray = Context.RayTracingMeshResourceCollector.AllocateOneFrameResource<FLandscapeElementParamArray>();
 	ParameterArray.ElementParams.AddDefaulted(NumSubsections * NumSubsections);
@@ -3276,7 +3291,7 @@ void FLandscapeComponentSceneProxy::GetDynamicRayTracingInstances(FRayTracingMat
 	}
 
 	FMeshBatch BaseMeshBatch;
-	BaseMeshBatch.VertexFactory = VertexFactory;
+	BaseMeshBatch.VertexFactory = bUseDiscreteLOD ? FixedGridVertexFactory : VertexFactory;
 	BaseMeshBatch.MaterialRenderProxy = SelectedMaterial->GetRenderProxy();
 	BaseMeshBatch.LCI = ComponentLightInfo.Get();
 	BaseMeshBatch.CastShadow = true;
@@ -4228,6 +4243,13 @@ public:
 
 		ShaderBindings.Add(Shader->GetUniformBufferParameter<FLandscapeUniformShaderParameters>(), *BatchElementParams->LandscapeUniformShaderParametersResource);
 
+#if RHI_RAYTRACING
+		if (IsRayTracingEnabled())
+		{
+			ShaderBindings.Add(Shader->GetUniformBufferParameter<FLandscapeVertexFactoryMVFParameters>(), BatchElementParams->LandscapeVertexFactoryMVFUniformBuffer);
+		}
+#endif
+
 		int32 CurrentLOD = BatchElementParams->CurrentLOD;
 #if WITH_EDITOR
 		if (InView && InView->Family->LandscapeLODOverride >= 0)
@@ -4366,6 +4388,10 @@ FVertexFactoryShaderParameters* FLandscapeFixedGridVertexFactory::ConstructShade
 	switch (ShaderFrequency)
 	{
 	case SF_Vertex:
+#if RHI_RAYTRACING
+	case SF_Compute:
+	case SF_RayHitGroup:
+#endif
 		return new FLandscapeFixedGridVertexFactoryVertexShaderParameters();
 		break;
 	case SF_Pixel:
diff --git a/Engine/Source/Runtime/Renderer/Private/PrimitiveSceneInfo.cpp b/Engine/Source/Runtime/Renderer/Private/PrimitiveSceneInfo.cpp
--- a/Engine/Source/Runtime/Renderer/Private/PrimitiveSceneInfo.cpp
+++ b/Engine/Source/Runtime/Renderer/Private/PrimitiveSceneInfo.cpp
@@ -216,7 +216,7 @@ void FPrimitiveSceneInfo::CacheMeshDrawCommands(FRHICommandListImmediate& RHICmd
 	}
 
 #if RHI_RAYTRACING
-	if (IsRayTracingEnabled())
+	if (IsRayTracingEnabled() && Proxy->IsRayTracingStaticRelevant())
 	{
 		int MaxLOD = -1;
 
@@ -319,7 +319,7 @@ void FPrimitiveSceneInfo::CacheMeshDrawCommands(FRHICommandListImmediate& RHICmd
 				}
 
 			#if RHI_RAYTRACING
-				if (IsRayTracingEnabled())
+				if (IsRayTracingEnabled() && Proxy->IsRayTracingStaticRelevant())
 				{
 					FCachedRayTracingMeshCommandContext CommandContext(Scene->CachedRayTracingMeshCommands);
 					FRayTracingMeshProcessor RayTracingMeshProcessor(&CommandContext, Scene, nullptr);
diff --git a/Engine/Source/Runtime/Renderer/Private/RayTracing/RayTracingDynamicGeometry.cpp b/Engine/Source/Runtime/Renderer/Private/RayTracing/RayTracingDynamicGeometry.cpp
--- a/Engine/Source/Runtime/Renderer/Private/RayTracing/RayTracingDynamicGeometry.cpp
+++ b/Engine/Source/Runtime/Renderer/Private/RayTracing/RayTracingDynamicGeometry.cpp
@@ -13,6 +13,7 @@ static bool IsSupportedDynamicVertexFactoryType(const FVertexFactoryType* Vertex
 	return VertexFactoryType == FindVertexFactoryType(FName(TEXT("FNiagaraSpriteVertexFactory"), FNAME_Find))
 		|| VertexFactoryType == FindVertexFactoryType(FName(TEXT("FNiagaraRibbonVertexFactory"), FNAME_Find))
 		|| VertexFactoryType == FindVertexFactoryType(FName(TEXT("FLandscapeVertexFactory"), FNAME_Find))
+		|| VertexFactoryType == FindVertexFactoryType(FName(TEXT("FLandscapeFixedGridVertexFactory"), FNAME_Find))
 		|| VertexFactoryType == FindVertexFactoryType(FName(TEXT("FLandscapeXYOffsetVertexFactory"), FNAME_Find))
 		|| VertexFactoryType == FindVertexFactoryType(FName(TEXT("FGPUSkinPassthroughVertexFactory"), FNAME_Find));
 }
diff --git a/Engine/Source/Runtime/Renderer/Private/RayTracing/RayTracingMaterialHitShaders.cpp b/Engine/Source/Runtime/Renderer/Private/RayTracing/RayTracingMaterialHitShaders.cpp
--- a/Engine/Source/Runtime/Renderer/Private/RayTracing/RayTracingMaterialHitShaders.cpp
+++ b/Engine/Source/Runtime/Renderer/Private/RayTracing/RayTracingMaterialHitShaders.cpp
@@ -34,6 +34,7 @@ static bool IsSupportedVertexFactoryType(const FVertexFactoryType* VertexFactory
 	static FName NiagaraSpriteVfFname = FName(TEXT("FNiagaraSpriteVertexFactory"), FNAME_Find);
 	static FName GeometryCacheVfFname = FName(TEXT("FGeometryCacheVertexVertexFactory"), FNAME_Find);
 	static FName LandscapeVfFname = FName(TEXT("FLandscapeVertexFactory"), FNAME_Find);
+	static FName LandscapeFixedGridVfFname = FName(TEXT("FLandscapeFixedGridVertexFactory"), FNAME_Find);
 	static FName LandscapeXYOffsetVfFname = FName(TEXT("FLandscapeXYOffsetVertexFactory"), FNAME_Find);
 
 	return VertexFactoryType == FindVertexFactoryType(LocalVfFname)
@@ -42,6 +43,7 @@ static bool IsSupportedVertexFactoryType(const FVertexFactoryType* VertexFactory
 		|| VertexFactoryType == FindVertexFactoryType(NiagaraSpriteVfFname)
 		|| VertexFactoryType == FindVertexFactoryType(GeometryCacheVfFname)
 		|| VertexFactoryType == FindVertexFactoryType(LandscapeVfFname)
+		|| VertexFactoryType == FindVertexFactoryType(LandscapeFixedGridVfFname)
 		|| VertexFactoryType == FindVertexFactoryType(LandscapeXYOffsetVfFname);
 }
 
