Landscape Splines: Fix Bad version check that was resetting Layer Falloff

diff --git a/Engine/Source/Runtime/Landscape/Private/LandscapeSplines.cpp b/Engine/Source/Runtime/Landscape/Private/LandscapeSplines.cpp
--- a/Engine/Source/Runtime/Landscape/Private/LandscapeSplines.cpp
+++ b/Engine/Source/Runtime/Landscape/Private/LandscapeSplines.cpp
@@ -1506,7 +1506,7 @@ void ULandscapeSplineControlPoint::Serialize(FArchive& Ar)
 		}
 	}
 
-	if (Ar.IsLoading() && Ar.CustomVer(FLandscapeCustomVersion::GUID) << FLandscapeCustomVersion::AddSplineLayerWidth)
+	if (Ar.IsLoading() && Ar.CustomVer(FLandscapeCustomVersion::GUID) < FLandscapeCustomVersion::AddSplineLayerWidth)
 	{
 		for (FLandscapeSplineInterpPoint& Point : Points)
 		{
@@ -2299,7 +2299,7 @@ void ULandscapeSplineSegment::Serialize(FArchive& Ar)
 		}
 	}
 
-	if (Ar.IsLoading() && Ar.CustomVer(FLandscapeCustomVersion::GUID) << FLandscapeCustomVersion::AddSplineLayerWidth)
+	if (Ar.IsLoading() && Ar.CustomVer(FLandscapeCustomVersion::GUID) < FLandscapeCustomVersion::AddSplineLayerWidth)
 	{
 		for (FLandscapeSplineInterpPoint& Point : Points)
 		{
