Revert landscape discrete LOD changes to resolve multiple breakages

diff --git a/Engine/Shaders/Private/LandscapeVertexFactory.ush b/Engine/Shaders/Private/LandscapeVertexFactory.ush
--- a/Engine/Shaders/Private/LandscapeVertexFactory.ush
+++ b/Engine/Shaders/Private/LandscapeVertexFactory.ush
@@ -514,27 +514,17 @@ void GetShadowMapCoordinate(FVertexFactoryInterpolantsVSToPS Interpolants, out f
 }
 #endif
 
-float4 GetSectionLods()
-{
-	return min(SectionLods + View.FarShadowStaticMeshLODBias.xxxx, LandscapeParameters.LastLOD.xxxx);
-}
-
-float4 GetNeighborSectionLod(int NeighborIndex)
-{
-	return min(NeighborSectionLod[NeighborIndex] + View.FarShadowStaticMeshLODBias.xxxx, LandscapeParameters.LastLOD.xxxx);
-}
-
 float CalcLOD_ES2(float2 xy, float2 Subsection)
 {
 	// LOD value based on Barycentric coordinate
 	float4 L0 = float4(xy.y, xy.x, (1 - xy.x), (1 - xy.y)) * 2;
 
 	float4 LODCalculated4 = lerp(
-								(lerp(L0 * GetSectionLods() + (float4(1, 1, 1, 1) - L0) * GetNeighborSectionLod(0),
-									  L0 * GetSectionLods() + (float4(1, 1, 1, 1) - L0) * GetNeighborSectionLod(1),
+								(lerp(L0 * SectionLods.x + (float4(1, 1, 1, 1) - L0) * NeighborSectionLod[0],
+									  L0 * SectionLods.y + (float4(1, 1, 1, 1) - L0) * NeighborSectionLod[1],
 									  Subsection.x)),
-								(lerp(L0 * GetSectionLods() + (float4(1, 1, 1, 1) - L0) * GetNeighborSectionLod(2),
-									  L0 * GetSectionLods() + (float4(1, 1, 1, 1) - L0) * GetNeighborSectionLod(3),
+								(lerp(L0 * SectionLods.z + (float4(1, 1, 1, 1) - L0) * NeighborSectionLod[2],
+									  L0 * SectionLods.w + (float4(1, 1, 1, 1) - L0) * NeighborSectionLod[3],
 									  Subsection.x)),
 								Subsection.y);
 
@@ -565,26 +555,26 @@ float CalcLOD(float2 xy, float2 Subsection)
 	{
 		if (Subsection.x > 0.5f)
 		{
-			float SectionLod = GetSectionLods().w;
-			LODCalculated4 = L0 * SectionLod + (float4(1, 1, 1, 1) - L0) * GetNeighborSectionLod(3);
+			float SectionLod = SectionLods.w;
+			LODCalculated4 = L0 * SectionLod + (float4(1, 1, 1, 1) - L0) * NeighborSectionLod[3];
 		}
 		else
 		{
-			float SectionLod = GetSectionLods().z;
-			LODCalculated4 = L0 * SectionLod + (float4(1, 1, 1, 1) - L0) * GetNeighborSectionLod(2);
+			float SectionLod = SectionLods.z;
+			LODCalculated4 = L0 * SectionLod + (float4(1, 1, 1, 1) - L0) * NeighborSectionLod[2];
 		}
 	}
 	else
 	{
 		if (Subsection.x > 0.5f)
 		{
-			float SectionLod = GetSectionLods().y;
-			LODCalculated4 = L0 * SectionLod + (float4(1, 1, 1, 1) - L0) * GetNeighborSectionLod(1);
+			float SectionLod = SectionLods.y;
+			LODCalculated4 = L0 * SectionLod + (float4(1, 1, 1, 1) - L0) * NeighborSectionLod[1];
 		}
 		else
 		{
-			float SectionLod = GetSectionLods().x;
-			LODCalculated4 = L0 * SectionLod + (float4(1, 1, 1, 1) - L0) * GetNeighborSectionLod(0);
+			float SectionLod = SectionLods.x;
+			LODCalculated4 = L0 * SectionLod + (float4(1, 1, 1, 1) - L0) * NeighborSectionLod[0];
 		}
 	}
 
@@ -686,19 +676,11 @@ FVertexFactoryIntermediates GetVertexFactoryIntermediates(FVertexFactoryInput In
 	// Get the height and normal XY for current and next LOD out of the textures
 #if FEATURE_LEVEL >= FEATURE_LEVEL_SM4
 	float2 SampleCoords = InputPositionLODAdjusted * LandscapeParameters.HeightmapUVScaleBias.xy + LandscapeParameters.HeightmapUVScaleBias.zw + 0.5*LandscapeParameters.HeightmapUVScaleBias.xy + Intermediates.InputPosition.zw * LandscapeParameters.SubsectionOffsetParams.xy;
-#if FIXED_GRID
-	float4 SampleValue = Texture2DSampleLevel(LandscapeParameters.HeightmapTexture, LandscapeParameters.HeightmapTextureSampler, SampleCoords, 0);
-#else
 	float4 SampleValue = Texture2DSampleLevel(LandscapeParameters.HeightmapTexture, LandscapeParameters.HeightmapTextureSampler, SampleCoords, LodValue-LodBias.z);
-#endif
 	float Height = DecodePackedHeight(SampleValue.xy);
 
 	float2 SampleCoordsNextLOD = InputPositionNextLOD * LandscapeParameters.HeightmapUVScaleBias.xy + LandscapeParameters.HeightmapUVScaleBias.zw + 0.5*LandscapeParameters.HeightmapUVScaleBias.xy + Intermediates.InputPosition.zw * LandscapeParameters.SubsectionOffsetParams.xy;
-#if FIXED_GRID
-	float4 SampleValueNextLOD = Texture2DSampleLevel(LandscapeParameters.HeightmapTexture, LandscapeParameters.HeightmapTextureSampler, SampleCoordsNextLOD, 1);
-#else
 	float4 SampleValueNextLOD = Texture2DSampleLevel(LandscapeParameters.HeightmapTexture, LandscapeParameters.HeightmapTextureSampler, SampleCoordsNextLOD, LodValue+1-LodBias.z );
-#endif
 	float HeightNextLOD = DecodePackedHeight(SampleValueNextLOD.xy);
 #endif
 
@@ -740,11 +722,7 @@ float2 LandscapeVertexFactorySampleHeight(FVertexFactoryInput Input, FVertexFact
 	*/
 
 	float2 xy = Intermediates.InputPosition.xy * LodValues.w;
-#if FIXED_GRID
-	float LODCalculated = floor(LodValues.x);
-#else
 	float LODCalculated = CalcLOD(xy, Intermediates.InputPosition.zw);
-#endif
 	float LodValue = floor(LODCalculated);
 	float CoordTranslate_x = (LandscapeParameters.SubsectionSizeVertsLayerUVPan.x * pow(2, -LodValue) - 1) * LandscapeParameters.SubsectionSizeVertsLayerUVPan.y;
 
@@ -776,11 +754,7 @@ float4 VertexFactoryGetRasterizedWorldPosition(FVertexFactoryInput Input, FVerte
 {
 #if RENDERING_LANDSCAPE_GRASS
 	float2 xy = Intermediates.InputPosition.xy * LodValues.w;
-#if FIXED_GRID
-	float LODCalculated = floor(LodValues.x);
-#else
 	float LODCalculated = CalcLOD(xy, Intermediates.InputPosition.zw);
-#endif
 	float LodValue = floor(LODCalculated);
 	float CoordTranslate_x = (LandscapeParameters.SubsectionSizeVertsLayerUVPan.x * pow(2, -LodValue) - 1) * LandscapeParameters.SubsectionSizeVertsLayerUVPan.y;
 
@@ -885,12 +859,7 @@ FVertexFactoryInterpolantsVSToDS VertexFactoryGetInterpolantsVSToDS(FVertexFacto
 	Interpolants.TangentToWorld2 = float4(TangentToWorld[2], 1);
 
 	float2 xy = Intermediates.InputPosition.xy * LodValues.w;
-#if FIXED_GRID
-	float LODCalculated = floor(LodValues.x);
-#else
-	float LODCalculated = CalcLOD(xy, Intermediates.InputPosition.zw);
-#endif
-	Interpolants.CalculatedLOD = LODCalculated;
+	Interpolants.CalculatedLOD = CalcLOD(xy, Intermediates.InputPosition.zw);
 #endif
 	
 	return Interpolants;
diff --git a/Engine/Source/Editor/LandscapeEditor/Private/LandscapeEdMode.cpp b/Engine/Source/Editor/LandscapeEditor/Private/LandscapeEdMode.cpp
--- a/Engine/Source/Editor/LandscapeEditor/Private/LandscapeEdMode.cpp
+++ b/Engine/Source/Editor/LandscapeEditor/Private/LandscapeEdMode.cpp
@@ -274,11 +274,10 @@ void FLandscapeTool::SetEditRenderType()
 
 namespace LandscapeTool
 {
-	UMaterialInstance* CreateMaterialInstance(UMaterialInterface* BaseMaterial, bool bUseDiscreteLOD)
+	UMaterialInstance* CreateMaterialInstance(UMaterialInterface* BaseMaterial)
 	{
 		ULandscapeMaterialInstanceConstant* MaterialInstance = NewObject<ULandscapeMaterialInstanceConstant>(GetTransientPackage());
 		MaterialInstance->bEditorToolUsage = true;
-		MaterialInstance->bUseDiscreteLOD = bUseDiscreteLOD;
 		MaterialInstance->SetParentEditorOnly(BaseMaterial);
 		MaterialInstance->PostEditChange();
 		return MaterialInstance;
@@ -313,19 +312,13 @@ FEdModeLandscape::FEdModeLandscape()
 	, bNeedsUpdateShownLayerList(false)
 	, bUpdatingLandscapeInfo(false)
 {
-	GLayerDebugColorMaterial     = LandscapeTool::CreateMaterialInstance(LoadObject<UMaterial>(nullptr, TEXT("/Engine/EditorLandscapeResources/LayerVisMaterial.LayerVisMaterial")),                                                       false);
-	GSelectionColorMaterial      = LandscapeTool::CreateMaterialInstance(LoadObject<UMaterialInstanceConstant>(nullptr, TEXT("/Engine/EditorLandscapeResources/SelectBrushMaterial_Selected.SelectBrushMaterial_Selected")),               false);
-	GSelectionRegionMaterial     = LandscapeTool::CreateMaterialInstance(LoadObject<UMaterialInstanceConstant>(nullptr, TEXT("/Engine/EditorLandscapeResources/SelectBrushMaterial_SelectedRegion.SelectBrushMaterial_SelectedRegion")),   false);
-	GMaskRegionMaterial          = LandscapeTool::CreateMaterialInstance(LoadObject<UMaterialInstanceConstant>(nullptr, TEXT("/Engine/EditorLandscapeResources/MaskBrushMaterial_MaskedRegion.MaskBrushMaterial_MaskedRegion")),           false);
-	GColorMaskRegionMaterial     = LandscapeTool::CreateMaterialInstance(LoadObject<UMaterialInstanceConstant>(nullptr, TEXT("/Engine/EditorLandscapeResources/ColorMaskBrushMaterial_MaskedRegion.ColorMaskBrushMaterial_MaskedRegion")), false);
-	GLandscapeLayerUsageMaterial = LandscapeTool::CreateMaterialInstance(LoadObject<UMaterial>(nullptr, TEXT("/Engine/EditorLandscapeResources/LandscapeLayerUsageMaterial.LandscapeLayerUsageMaterial")),                                 false);
-	GLayerDebugColorMaterial_DiscreteLOD     = LandscapeTool::CreateMaterialInstance(LoadObject<UMaterial>(nullptr, TEXT("/Engine/EditorLandscapeResources/LayerVisMaterial.LayerVisMaterial")),                                                       true);
-	GSelectionColorMaterial_DiscreteLOD      = LandscapeTool::CreateMaterialInstance(LoadObject<UMaterialInstanceConstant>(nullptr, TEXT("/Engine/EditorLandscapeResources/SelectBrushMaterial_Selected.SelectBrushMaterial_Selected")),               true);
-	GSelectionRegionMaterial_DiscreteLOD     = LandscapeTool::CreateMaterialInstance(LoadObject<UMaterialInstanceConstant>(nullptr, TEXT("/Engine/EditorLandscapeResources/SelectBrushMaterial_SelectedRegion.SelectBrushMaterial_SelectedRegion")),   true);
-	GMaskRegionMaterial_DiscreteLOD          = LandscapeTool::CreateMaterialInstance(LoadObject<UMaterialInstanceConstant>(nullptr, TEXT("/Engine/EditorLandscapeResources/MaskBrushMaterial_MaskedRegion.MaskBrushMaterial_MaskedRegion")),           true);
-	GColorMaskRegionMaterial_DiscreteLOD     = LandscapeTool::CreateMaterialInstance(LoadObject<UMaterialInstanceConstant>(nullptr, TEXT("/Engine/EditorLandscapeResources/ColorMaskBrushMaterial_MaskedRegion.ColorMaskBrushMaterial_MaskedRegion")), true);
-	GLandscapeLayerUsageMaterial_DiscreteLOD = LandscapeTool::CreateMaterialInstance(LoadObject<UMaterial>(nullptr, TEXT("/Engine/EditorLandscapeResources/LandscapeLayerUsageMaterial.LandscapeLayerUsageMaterial")),                                 true);
+	GLayerDebugColorMaterial = LandscapeTool::CreateMaterialInstance(LoadObject<UMaterial>(nullptr, TEXT("/Engine/EditorLandscapeResources/LayerVisMaterial.LayerVisMaterial")));
+	GSelectionColorMaterial  = LandscapeTool::CreateMaterialInstance(LoadObject<UMaterialInstanceConstant>(nullptr, TEXT("/Engine/EditorLandscapeResources/SelectBrushMaterial_Selected.SelectBrushMaterial_Selected")));
+	GSelectionRegionMaterial = LandscapeTool::CreateMaterialInstance(LoadObject<UMaterialInstanceConstant>(nullptr, TEXT("/Engine/EditorLandscapeResources/SelectBrushMaterial_SelectedRegion.SelectBrushMaterial_SelectedRegion")));
+	GMaskRegionMaterial      = LandscapeTool::CreateMaterialInstance(LoadObject<UMaterialInstanceConstant>(nullptr, TEXT("/Engine/EditorLandscapeResources/MaskBrushMaterial_MaskedRegion.MaskBrushMaterial_MaskedRegion")));
+	GColorMaskRegionMaterial = LandscapeTool::CreateMaterialInstance(LoadObject<UMaterialInstanceConstant>(nullptr, TEXT("/Engine/EditorLandscapeResources/ColorMaskBrushMaterial_MaskedRegion.ColorMaskBrushMaterial_MaskedRegion")));
 	GLandscapeBlackTexture   = LoadObject<UTexture2D>(nullptr, TEXT("/Engine/EngineResources/Black.Black"));
+	GLandscapeLayerUsageMaterial = LandscapeTool::CreateMaterialInstance(LoadObject<UMaterial>(nullptr, TEXT("/Engine/EditorLandscapeResources/LandscapeLayerUsageMaterial.LandscapeLayerUsageMaterial")));
 	
 	// Initialize modes
 	UpdateToolModes();
@@ -397,13 +390,6 @@ FEdModeLandscape::~FEdModeLandscape()
 	GLandscapeBlackTexture = NULL;
 	GLandscapeLayerUsageMaterial = NULL;
 
-	GLayerDebugColorMaterial_DiscreteLOD = NULL;
-	GSelectionColorMaterial_DiscreteLOD = NULL;
-	GSelectionRegionMaterial_DiscreteLOD = NULL;
-	GMaskRegionMaterial_DiscreteLOD = NULL;
-	GColorMaskRegionMaterial_DiscreteLOD = NULL;
-	GLandscapeLayerUsageMaterial_DiscreteLOD = NULL;
-
 	InteractorPainting = nullptr;
 }
 
@@ -420,16 +406,8 @@ void FEdModeLandscape::AddReferencedObjects(FReferenceCollector& Collector)
 	Collector.AddReferencedObject(GSelectionRegionMaterial);
 	Collector.AddReferencedObject(GMaskRegionMaterial);
 	Collector.AddReferencedObject(GColorMaskRegionMaterial);
-	Collector.AddReferencedObject(GLandscapeLayerUsageMaterial);
-
-	Collector.AddReferencedObject(GLayerDebugColorMaterial_DiscreteLOD);
-	Collector.AddReferencedObject(GSelectionColorMaterial_DiscreteLOD);
-	Collector.AddReferencedObject(GSelectionRegionMaterial_DiscreteLOD);
-	Collector.AddReferencedObject(GMaskRegionMaterial_DiscreteLOD);
-	Collector.AddReferencedObject(GColorMaskRegionMaterial_DiscreteLOD);
-	Collector.AddReferencedObject(GLandscapeLayerUsageMaterial_DiscreteLOD);
-
 	Collector.AddReferencedObject(GLandscapeBlackTexture);
+	Collector.AddReferencedObject(GLandscapeLayerUsageMaterial);
 }
 
 void FEdModeLandscape::UpdateToolModes()
diff --git a/Engine/Source/Editor/LandscapeEditor/Private/LandscapeEdModeBrushes.cpp b/Engine/Source/Editor/LandscapeEditor/Private/LandscapeEdModeBrushes.cpp
--- a/Engine/Source/Editor/LandscapeEditor/Private/LandscapeEdModeBrushes.cpp
+++ b/Engine/Source/Editor/LandscapeEditor/Private/LandscapeEdModeBrushes.cpp
@@ -12,7 +12,7 @@
 #include "LandscapeGizmoActor.h"
 #include "LandscapeEdMode.h"
 #include "LandscapeEditorObject.h"
-#include "LandscapeMaterialInstanceConstant.h"
+
 #include "LandscapeRender.h"
 
 #include "LevelUtils.h"
@@ -52,7 +52,6 @@ class FLandscapeBrushCircle : public FLandscapeBrush
 protected:
 	FVector2D LastMousePosition;
 	UMaterialInterface* BrushMaterial;
-	UMaterialInterface* BrushMaterial_DiscreteLOD;
 	TMap<ULandscapeComponent*, UMaterialInstanceDynamic*> BrushMaterialInstanceMap;
 
 	virtual float CalculateFalloff(float Distance, float Radius, float Falloff) = 0;
@@ -60,8 +59,7 @@ protected:
 	/** Protected so that only subclasses can create instances of this class. */
 	FLandscapeBrushCircle(FEdModeLandscape* InEdMode, UMaterialInterface* InBrushMaterial)
 		: LastMousePosition(0, 0)
-		, BrushMaterial(LandscapeTool::CreateMaterialInstance(InBrushMaterial, false))
-		, BrushMaterial_DiscreteLOD(LandscapeTool::CreateMaterialInstance(InBrushMaterial, true))
+		, BrushMaterial(LandscapeTool::CreateMaterialInstance(InBrushMaterial))
 		, EdMode(InEdMode)
 	{
 	}
@@ -73,7 +71,6 @@ public:
 	virtual void AddReferencedObjects(FReferenceCollector& Collector) override
 	{
 		Collector.AddReferencedObject(BrushMaterial);
-		Collector.AddReferencedObject(BrushMaterial_DiscreteLOD);
 
 		// Allow any currently unused material instances to be GC'd
 		BrushMaterialFreeInstances.Empty();
@@ -163,26 +160,14 @@ public:
 		for (ULandscapeComponent* AddedComponent : AddedComponents)
 		{
 			UMaterialInstanceDynamic* BrushMaterialInstance = nullptr;
-
 			if (BrushMaterialFreeInstances.Num() > 0)
 			{
-				for (auto Iter = BrushMaterialFreeInstances.CreateIterator(); Iter; Iter++)
-				{
-					ULandscapeMaterialInstanceConstant* LandscapeMIC = Cast<ULandscapeMaterialInstanceConstant>((*Iter)->Parent);
-					if (LandscapeMIC->bUseDiscreteLOD == AddedComponent->GetLandscapeProxy()->bUseDiscreteLOD)
-					{
-						BrushMaterialInstance = *Iter;
-						Iter.RemoveCurrent();
-						break;
-					}
-				}
+				BrushMaterialInstance = BrushMaterialFreeInstances.Pop();
 			}
-			
-			if (BrushMaterialInstance == nullptr)
+			else
 			{
-				BrushMaterialInstance = UMaterialInstanceDynamic::Create(AddedComponent->GetLandscapeProxy()->bUseDiscreteLOD ? BrushMaterial_DiscreteLOD : BrushMaterial, nullptr);
+				BrushMaterialInstance = UMaterialInstanceDynamic::Create(BrushMaterial, nullptr);
 			}
-
 			BrushMaterialInstanceMap.Add(AddedComponent, BrushMaterialInstance);
 			AddedComponent->EditToolRenderData.ToolMaterial = BrushMaterialInstance;
 			AddedComponent->UpdateEditToolRenderData();
@@ -373,18 +358,15 @@ class FLandscapeBrushComponent : public FLandscapeBrush
 protected:
 	FVector2D LastMousePosition;
 	UMaterialInterface* BrushMaterial;
-	UMaterialInterface* BrushMaterial_DiscreteLOD;
 public:
 	FEdModeLandscape* EdMode;
 
 	FLandscapeBrushComponent(FEdModeLandscape* InEdMode)
 		: BrushMaterial(nullptr)
-		, BrushMaterial_DiscreteLOD(nullptr)
 		, EdMode(InEdMode)
 	{
 		UMaterial* BaseBrushMaterial = LoadObject<UMaterial>(nullptr, TEXT("/Engine/EditorLandscapeResources/SelectBrushMaterial.SelectBrushMaterial"));
-		BrushMaterial = LandscapeTool::CreateMaterialInstance(BaseBrushMaterial, false);
-		BrushMaterial_DiscreteLOD = LandscapeTool::CreateMaterialInstance(BaseBrushMaterial, true);
+		BrushMaterial = LandscapeTool::CreateMaterialInstance(BaseBrushMaterial);
 	}
 
 	// FGCObject interface
@@ -392,7 +374,6 @@ public:
 	{
 		Collector.AddReferencedObjects(BrushMaterialComponents);
 		Collector.AddReferencedObject(BrushMaterial);
-		Collector.AddReferencedObject(BrushMaterial_DiscreteLOD);
 	}
 
 	virtual ELandscapeBrushType GetBrushType() override { return ELandscapeBrushType::Component; }
@@ -459,7 +440,7 @@ public:
 				// Set brush material for components in new region
 				for (ULandscapeComponent* NewComponent : NewComponents)
 				{
-					NewComponent->EditToolRenderData.ToolMaterial = NewComponent->GetLandscapeProxy()->bUseDiscreteLOD ? BrushMaterial_DiscreteLOD : BrushMaterial;
+					NewComponent->EditToolRenderData.ToolMaterial = BrushMaterial;
 					NewComponent->UpdateEditToolRenderData();
 				}
 			}
@@ -573,18 +554,15 @@ class FLandscapeBrushGizmo : public FLandscapeBrush
 
 protected:
 	UMaterialInstanceDynamic* BrushMaterial;
-	UMaterialInstanceDynamic* BrushMaterial_DiscreteLOD;
 public:
 	FEdModeLandscape* EdMode;
 
 	FLandscapeBrushGizmo(FEdModeLandscape* InEdMode)
 		: BrushMaterial(nullptr)
-		, BrushMaterial_DiscreteLOD(nullptr)
 		, EdMode(InEdMode)
 	{
 		UMaterialInterface* GizmoMaterial = LoadObject<UMaterialInstanceConstant>(nullptr, TEXT("/Engine/EditorLandscapeResources/MaskBrushMaterial_Gizmo.MaskBrushMaterial_Gizmo"));
-		BrushMaterial = UMaterialInstanceDynamic::Create(LandscapeTool::CreateMaterialInstance(GizmoMaterial, false), nullptr);
-		BrushMaterial_DiscreteLOD = UMaterialInstanceDynamic::Create(LandscapeTool::CreateMaterialInstance(GizmoMaterial, true), nullptr);
+		BrushMaterial = UMaterialInstanceDynamic::Create(LandscapeTool::CreateMaterialInstance(GizmoMaterial), nullptr);
 	}
 
 	// FGCObject interface
@@ -592,7 +570,6 @@ public:
 	{
 		Collector.AddReferencedObjects(BrushMaterialComponents);
 		Collector.AddReferencedObject(BrushMaterial);
-		Collector.AddReferencedObject(BrushMaterial_DiscreteLOD);
 	}
 
 	virtual ELandscapeBrushType GetBrushType() override { return ELandscapeBrushType::Gizmo; }
@@ -662,19 +639,16 @@ public:
 						Gizmo->TextureScale.Y
 						);
 					BrushMaterial->SetVectorParameterValue(FName(TEXT("AlphaScaleBias")), AlphaScaleBias);
-					BrushMaterial_DiscreteLOD->SetVectorParameterValue(FName(TEXT("AlphaScaleBias")), AlphaScaleBias);
 
 					float Angle = (-EdMode->CurrentGizmoActor->GetActorRotation().Euler().Z) * PI / 180.0f;
 					FLinearColor LandscapeLocation(EdMode->CurrentGizmoActor->GetActorLocation().X, EdMode->CurrentGizmoActor->GetActorLocation().Y, EdMode->CurrentGizmoActor->GetActorLocation().Z, Angle);
 					BrushMaterial->SetVectorParameterValue(FName(TEXT("LandscapeLocation")), LandscapeLocation);
-					BrushMaterial_DiscreteLOD->SetVectorParameterValue(FName(TEXT("LandscapeLocation")), LandscapeLocation);
 					BrushMaterial->SetTextureParameterValue(FName(TEXT("AlphaTexture")), DataTexture);
-					BrushMaterial_DiscreteLOD->SetTextureParameterValue(FName(TEXT("AlphaTexture")), DataTexture);
 
 					// Set brush material for components in new region
 					for (ULandscapeComponent* NewComponent : NewComponents)
 					{
-						NewComponent->EditToolRenderData.GizmoMaterial = ((Gizmo->DataType != LGT_None) && (GLandscapeEditRenderMode & ELandscapeEditRenderMode::Gizmo) ? (NewComponent->GetLandscapeProxy()->bUseDiscreteLOD ? BrushMaterial_DiscreteLOD : BrushMaterial) : nullptr);
+						NewComponent->EditToolRenderData.GizmoMaterial = ((Gizmo->DataType != LGT_None) && (GLandscapeEditRenderMode & ELandscapeEditRenderMode::Gizmo) ? BrushMaterial : nullptr);
 						NewComponent->UpdateEditToolRenderData();
 					}
 
diff --git a/Engine/Source/Editor/LandscapeEditor/Public/LandscapeToolInterface.h b/Engine/Source/Editor/LandscapeEditor/Public/LandscapeToolInterface.h
--- a/Engine/Source/Editor/LandscapeEditor/Public/LandscapeToolInterface.h
+++ b/Engine/Source/Editor/LandscapeEditor/Public/LandscapeToolInterface.h
@@ -257,5 +257,5 @@ public:
 
 namespace LandscapeTool
 {
-	UMaterialInstance* CreateMaterialInstance(UMaterialInterface* BaseMaterial, bool bUseDiscreteLOD);
+	UMaterialInstance* CreateMaterialInstance(UMaterialInterface* BaseMaterial);
 }
diff --git a/Engine/Source/Runtime/Landscape/Classes/LandscapeMaterialInstanceConstant.h b/Engine/Source/Runtime/Landscape/Classes/LandscapeMaterialInstanceConstant.h
--- a/Engine/Source/Runtime/Landscape/Classes/LandscapeMaterialInstanceConstant.h
+++ b/Engine/Source/Runtime/Landscape/Classes/LandscapeMaterialInstanceConstant.h
@@ -24,9 +24,6 @@ class ULandscapeMaterialInstanceConstant : public UMaterialInstanceConstant
 	UPROPERTY()
 	uint32 bEditorToolUsage:1;
 
-	UPROPERTY()
-	uint32 bUseDiscreteLOD : 1;
-
 	virtual FMaterialResource* AllocatePermutationResource() override;
 	virtual bool HasOverridenBaseProperties() const override;
 };
diff --git a/Engine/Source/Runtime/Landscape/Classes/LandscapeProxy.h b/Engine/Source/Runtime/Landscape/Classes/LandscapeProxy.h
--- a/Engine/Source/Runtime/Landscape/Classes/LandscapeProxy.h
+++ b/Engine/Source/Runtime/Landscape/Classes/LandscapeProxy.h
@@ -435,10 +435,6 @@ public:
 	UPROPERTY(EditAnywhere, Category = LOD, meta=(ClampMin = "0.01", ClampMax = "1.0", UIMin = "0.01", UIMax = "1.0", DisplayName= "SubSection Min Component ScreenSize"))
 	float ComponentScreenSizeToUseSubSections;
 
-	/** Whether to use discrete LOD rather than continuous LOD. Continuous LOD smooths LOD transition, at the cost of increased RenderThread time. */
-	UPROPERTY(EditAnywhere, Category = "LOD Distribution")
-	bool bUseDiscreteLOD;
-
 	/** This is the starting screen size used to calculate the distribution, by default it's 1, but you can increase the value if you want less LOD0 component, and you use very large landscape component. */
 	UPROPERTY(EditAnywhere, Category = "LOD Distribution", meta = (DisplayName = "LOD 0 Screen Size", ClampMin = "1.0", ClampMax = "10.0", UIMin = "1.0", UIMax = "10.0"))
 	float LOD0ScreenSize;
diff --git a/Engine/Source/Runtime/Landscape/Private/LandscapeEdit.cpp b/Engine/Source/Runtime/Landscape/Private/LandscapeEdit.cpp
--- a/Engine/Source/Runtime/Landscape/Private/LandscapeEdit.cpp
+++ b/Engine/Source/Runtime/Landscape/Private/LandscapeEdit.cpp
@@ -275,7 +275,6 @@ UMaterialInstanceConstant* ULandscapeComponent::GetCombinationMaterial(FMaterial
 
 			ULandscapeMaterialInstanceConstant* LandscapeCombinationMaterialInstance = NewObject<ULandscapeMaterialInstanceConstant>(GetOuter());
 			LandscapeCombinationMaterialInstance->bMobile = bMobile;
-			LandscapeCombinationMaterialInstance->bUseDiscreteLOD = GetLandscapeProxy()->bUseDiscreteLOD;
 			CombinationMaterialInstance = LandscapeCombinationMaterialInstance;
 			UE_LOG(LogLandscape, Log, TEXT("Looking for key %s, making new combination %s"), *LayerKey, *CombinationMaterialInstance->GetName());
 			Proxy->MaterialInstanceConstantMap.Add(*LayerKey, CombinationMaterialInstance);
@@ -364,8 +363,7 @@ void ULandscapeComponent::UpdateMaterialInstances_Internal(FMaterialUpdateContex
 		if (CombinationMaterialInstance != nullptr)
 		{
 			// Create the instance for this component, that will use the layer combination instance.
-			ULandscapeMaterialInstanceConstant* MaterialInstance = NewObject<ULandscapeMaterialInstanceConstant>(GetOuter());
-			MaterialInstance->bUseDiscreteLOD = GetLandscapeProxy()->bUseDiscreteLOD;
+			UMaterialInstanceConstant* MaterialInstance = NewObject<ULandscapeMaterialInstanceConstant>(GetOuter());
 			MaterialInstances[MaterialIndex] = MaterialInstance;
 
 			// Material Instances don't support Undo/Redo (the shader map goes out of sync and crashes happen)
@@ -399,9 +397,6 @@ void ULandscapeComponent::UpdateMaterialInstances_Internal(FMaterialUpdateContex
 			{
 				MaterialInstance->SetTextureParameterValueEditorOnly(FName(TEXT("Heightmap")), BaseHeightmap);
 			}
-
-			MaterialInstance->bUseDiscreteLOD = GetLandscapeProxy()->bUseDiscreteLOD;
-
 			MaterialInstance->PostEditChange();
 
 			// Setup material instance with disabled tessellation
@@ -415,7 +410,6 @@ void ULandscapeComponent::UpdateMaterialInstances_Internal(FMaterialUpdateContex
 				TessellationMaterialInstance->SetParentEditorOnly(MaterialInstance);
 				Context.AddMaterialInstance(TessellationMaterialInstance); // must be done after SetParent
 				TessellationMaterialInstance->bDisableTessellation = true;
-				TessellationMaterialInstance->bUseDiscreteLOD = GetLandscapeProxy()->bUseDiscreteLOD;
 				TessellationMaterialInstance->PostEditChange();
 			}
 		}
@@ -4206,7 +4200,7 @@ void ALandscapeStreamingProxy::PostEditChangeProperty(FPropertyChangedEvent& Pro
 			LandscapeActor = nullptr;
 		}
 	}
-	else if (PropertyName == FName(TEXT("LandscapeMaterial")) || PropertyName == FName(TEXT("LandscapeHoleMaterial")) || PropertyName == FName(TEXT("LandscapeMaterialsOverride")) || PropertyName == FName(TEXT("bUseDiscreteLOD")))
+	else if (PropertyName == FName(TEXT("LandscapeMaterial")) || PropertyName == FName(TEXT("LandscapeHoleMaterial")) || PropertyName == FName(TEXT("LandscapeMaterialsOverride")))
 	{
 		bool RecreateMaterialInstances = true;
 
@@ -4332,30 +4326,6 @@ void ALandscape::PostEditChangeProperty(FPropertyChangedEvent& PropertyChangedEv
 			MaterialInstanceConstantMap.Empty();
 		}
 	}
-	else if (PropertyName == GET_MEMBER_NAME_CHECKED(ALandscapeProxy, bUseDiscreteLOD))
-	{
-		bPropagateToProxies = true;
-
-		if (Info != nullptr)
-		{
-			FMaterialUpdateContext MaterialUpdateContext;
-			Info->UpdateLayerInfoMap(/*this*/);
-
-			ChangedMaterial = true;
-
-			// Clear the parents out of combination material instances
-			for (const auto& MICPair : MaterialInstanceConstantMap)
-			{
-				UMaterialInstanceConstant* MaterialInstance = MICPair.Value;
-				MaterialInstance->BasePropertyOverrides.bOverride_BlendMode = false;
-				MaterialInstance->SetParentEditorOnly(nullptr);
-				MaterialUpdateContext.AddMaterialInstance(MaterialInstance);
-			}
-
-			// Remove our references to any material instances
-			MaterialInstanceConstantMap.Empty();
-		}
-	}
 	else if (PropertyName == FName(TEXT("RelativeScale3D")) ||
 		PropertyName == FName(TEXT("RelativeLocation")) ||
 		PropertyName == FName(TEXT("RelativeRotation")))
@@ -5792,8 +5762,8 @@ void ULandscapeComponent::GeneratePlatformPixelData()
 			MobileCombinationMaterialInstances[MaterialIndex] = GetCombinationMaterial(nullptr, MobileWeightmapLayerAllocations, MaterialLOD, true);
 			check(MobileCombinationMaterialInstances[MaterialIndex] != nullptr);
 
-			ULandscapeMaterialInstanceConstant* NewMobileMaterialInstance = NewObject<ULandscapeMaterialInstanceConstant>(this);
-			NewMobileMaterialInstance->bUseDiscreteLOD = GetLandscapeProxy()->bUseDiscreteLOD;
+			UMaterialInstanceConstant* NewMobileMaterialInstance = NewObject<ULandscapeMaterialInstanceConstant>(this);
+
 			NewMobileMaterialInstance->SetParentEditorOnly(MobileCombinationMaterialInstances[MaterialIndex]);
 
 			// Set the layer mask
diff --git a/Engine/Source/Runtime/Landscape/Private/LandscapeEditLayers.cpp b/Engine/Source/Runtime/Landscape/Private/LandscapeEditLayers.cpp
--- a/Engine/Source/Runtime/Landscape/Private/LandscapeEditLayers.cpp
+++ b/Engine/Source/Runtime/Landscape/Private/LandscapeEditLayers.cpp
@@ -4184,9 +4184,7 @@ void ALandscape::UpdateLayersMaterialInstances(const TArray<ULandscapeComponent*
 				if (NeedToCreateMIC)
 				{
 					// Create the instance for this component, that will use the layer combination instance.
-					ULandscapeMaterialInstanceConstant* LandscapeMIC = NewObject<ULandscapeMaterialInstanceConstant>(this);
-					LandscapeMIC->bUseDiscreteLOD = this->bUseDiscreteLOD;
-					MaterialInstance = LandscapeMIC;
+					MaterialInstance = NewObject<ULandscapeMaterialInstanceConstant>(this);
 					Component->MaterialInstances[MaterialIndex] = MaterialInstance;
 				}
 
diff --git a/Engine/Source/Runtime/Landscape/Private/LandscapeGrass.cpp b/Engine/Source/Runtime/Landscape/Private/LandscapeGrass.cpp
--- a/Engine/Source/Runtime/Landscape/Private/LandscapeGrass.cpp
+++ b/Engine/Source/Runtime/Landscape/Private/LandscapeGrass.cpp
@@ -206,11 +206,7 @@ static bool ShouldCacheLandscapeGrassShaders(EShaderPlatform Platform, const FMa
 	// We only need grass weight shaders for Landscape vertex factories on desktop platforms
 	return (Material->IsUsedWithLandscape() || Material->IsSpecialEngineMaterial()) &&
 		IsFeatureLevelSupported(Platform, ERHIFeatureLevel::SM5) &&
-		(
-			(VertexFactoryType == FindVertexFactoryType(FName(TEXT("FLandscapeVertexFactory"), FNAME_Find))) || 
-			(VertexFactoryType == FindVertexFactoryType(FName(TEXT("FLandscapeFixedGridVertexFactory"), FNAME_Find))) || 
-			(VertexFactoryType == FindVertexFactoryType(FName(TEXT("FLandscapeXYOffsetVertexFactory"), FNAME_Find)))
-		)
+		((VertexFactoryType == FindVertexFactoryType(FName(TEXT("FLandscapeVertexFactory"), FNAME_Find))) || (VertexFactoryType == FindVertexFactoryType(FName(TEXT("FLandscapeXYOffsetVertexFactory"), FNAME_Find))))
 		&& !IsConsolePlatform(Platform);
 }
 
diff --git a/Engine/Source/Runtime/Landscape/Private/LandscapeRender.cpp b/Engine/Source/Runtime/Landscape/Private/LandscapeRender.cpp
--- a/Engine/Source/Runtime/Landscape/Private/LandscapeRender.cpp
+++ b/Engine/Source/Runtime/Landscape/Private/LandscapeRender.cpp
@@ -566,22 +566,13 @@ FLandscapeDebugOptions GLandscapeDebugOptions;
 #if WITH_EDITOR
 LANDSCAPE_API bool GLandscapeEditModeActive = false;
 LANDSCAPE_API int32 GLandscapeEditRenderMode = ELandscapeEditRenderMode::None;
-
 UMaterialInterface* GLayerDebugColorMaterial = nullptr;
 UMaterialInterface* GSelectionColorMaterial = nullptr;
 UMaterialInterface* GSelectionRegionMaterial = nullptr;
 UMaterialInterface* GMaskRegionMaterial = nullptr;
 UMaterialInterface* GColorMaskRegionMaterial = nullptr;
-UMaterialInterface* GLandscapeLayerUsageMaterial = nullptr;
-
-UMaterialInterface* GLayerDebugColorMaterial_DiscreteLOD = nullptr;
-UMaterialInterface* GSelectionColorMaterial_DiscreteLOD = nullptr;
-UMaterialInterface* GSelectionRegionMaterial_DiscreteLOD = nullptr;
-UMaterialInterface* GMaskRegionMaterial_DiscreteLOD = nullptr;
-UMaterialInterface* GColorMaskRegionMaterial_DiscreteLOD = nullptr;
-UMaterialInterface* GLandscapeLayerUsageMaterial_DiscreteLOD = nullptr;
-
 UTexture2D* GLandscapeBlackTexture = nullptr;
+UMaterialInterface* GLandscapeLayerUsageMaterial = nullptr;
 #endif
 
 void ULandscapeComponent::GetUsedMaterials(TArray<UMaterialInterface*>& OutMaterials, bool bGetDebugMaterials) const
@@ -626,24 +617,12 @@ void ULandscapeComponent::GetUsedMaterials(TArray<UMaterialInterface*>& OutMater
 #if WITH_EDITOR
 	//if (bGetDebugMaterials) // TODO: This should be tested and enabled
 	{
-		if (!GetLandscapeProxy()->bUseDiscreteLOD)
-		{
-			OutMaterials.Add(GLayerDebugColorMaterial);
-			OutMaterials.Add(GSelectionColorMaterial);
-			OutMaterials.Add(GSelectionRegionMaterial);
-			OutMaterials.Add(GMaskRegionMaterial);
-			OutMaterials.Add(GColorMaskRegionMaterial);
-			OutMaterials.Add(GLandscapeLayerUsageMaterial);
-		}
-		else
-		{
-			OutMaterials.Add(GLayerDebugColorMaterial_DiscreteLOD);
-			OutMaterials.Add(GSelectionColorMaterial_DiscreteLOD);
-			OutMaterials.Add(GSelectionRegionMaterial_DiscreteLOD);
-			OutMaterials.Add(GMaskRegionMaterial_DiscreteLOD);
-			OutMaterials.Add(GColorMaskRegionMaterial_DiscreteLOD);
-			OutMaterials.Add(GLandscapeLayerUsageMaterial_DiscreteLOD);
-		}
+		OutMaterials.Add(GLayerDebugColorMaterial);
+		OutMaterials.Add(GSelectionColorMaterial);
+		OutMaterials.Add(GSelectionRegionMaterial);
+		OutMaterials.Add(GMaskRegionMaterial);
+		OutMaterials.Add(GColorMaskRegionMaterial);
+		OutMaterials.Add(GLandscapeLayerUsageMaterial);
 	}
 #endif
 }
@@ -660,7 +639,6 @@ FLandscapeComponentSceneProxy::FLandscapeComponentSceneProxy(ULandscapeComponent
 	: FPrimitiveSceneProxy(InComponent, NAME_LandscapeResourceNameForDebugging)
 	, FLandscapeNeighborInfo(InComponent->GetWorld(), InComponent->GetLandscapeProxy()->GetLandscapeGuid(), InComponent->GetSectionBase() / InComponent->ComponentSizeQuads, InComponent->GetHeightmap(), InComponent->ForcedLOD, InComponent->LODBias)
 	, MaxLOD(FMath::CeilLogTwo(InComponent->SubsectionSizeQuads + 1) - 1)
-	, bUseDiscreteLOD(InComponent->GetLandscapeProxy()->bUseDiscreteLOD)
 	, UseTessellationComponentScreenSizeFalloff(InComponent->GetLandscapeProxy()->UseTessellationComponentScreenSizeFalloff)
 	, bRequiresAdjacencyInformation(false)
 	, NumWeightmapLayerAllocations(InComponent->GetWeightmapLayerAllocations().Num())
@@ -690,7 +668,6 @@ FLandscapeComponentSceneProxy::FLandscapeComponentSceneProxy(ULandscapeComponent
 	, SharedBuffersKey(0)
 	, SharedBuffers(nullptr)
 	, VertexFactory(nullptr)
-	, FixedGridVertexFactory(nullptr)
 	, ComponentLightInfo(nullptr)
 #if WITH_EDITORONLY_DATA
 	, EditToolRenderData(InComponent->EditToolRenderData)
@@ -963,7 +940,7 @@ void FLandscapeComponentSceneProxy::CreateRenderThreadResources()
 			SharedBuffers->VertexFactory = LandscapeXYOffsetVertexFactory;
 		}
 
-		if (UseVirtualTexturing(FeatureLevel) || bUseDiscreteLOD)
+		if (UseVirtualTexturing(FeatureLevel))
 		{
 			//todo[vt]: We will need a version of this to support XYOffsetmapTexture
 			FLandscapeFixedGridVertexFactory* LandscapeVertexFactory = new FLandscapeFixedGridVertexFactory(FeatureLevel);
@@ -1022,7 +999,7 @@ void FLandscapeComponentSceneProxy::CreateRenderThreadResources()
 		
 		// Grass is being generated using LOD0 material only
 		FMaterialRenderProxy* RenderProxy = AvailableMaterials[LODIndexToMaterialIndex[0]]->GetRenderProxy();
-		GrassMeshBatch.VertexFactory = bUseDiscreteLOD ? FixedGridVertexFactory : VertexFactory;
+		GrassMeshBatch.VertexFactory = VertexFactory;
 		GrassMeshBatch.MaterialRenderProxy = RenderProxy;
 		GrassMeshBatch.LCI = nullptr;
 		GrassMeshBatch.ReverseCulling = false;
@@ -1150,7 +1127,7 @@ FPrimitiveViewRelevance FLandscapeComponentSceneProxy::GetViewRelevance(const FS
 	Result.bRenderCustomDepth = ShouldRenderCustomDepth();
 	Result.bUsesLightingChannels = GetLightingChannelMask() != GetDefaultLightingChannelMask();
 	Result.bTranslucentSelfShadow = bCastVolumetricTranslucentShadow;
-	Result.bUseCustomViewData = !bUseDiscreteLOD;
+	Result.bUseCustomViewData = true;
 
 	auto FeatureLevel = View->GetFeatureLevel();
 
@@ -1190,30 +1167,30 @@ FPrimitiveViewRelevance FLandscapeComponentSceneProxy::GetViewRelevance(const FS
 			if (EditToolRenderData.SelectedType)
 			{
 				if ((GLandscapeEditRenderMode & ELandscapeEditRenderMode::SelectRegion) && (EditToolRenderData.SelectedType & FLandscapeEditToolRenderData::ST_REGION)
-					&& !(GLandscapeEditRenderMode & ELandscapeEditRenderMode::Mask) && (bUseDiscreteLOD ? GSelectionRegionMaterial_DiscreteLOD : GSelectionRegionMaterial))
+					&& !(GLandscapeEditRenderMode & ELandscapeEditRenderMode::Mask) && GSelectionRegionMaterial)
 				{
 					Result.bDynamicRelevance = true;
-					ToolRelevance |= (bUseDiscreteLOD ? GSelectionRegionMaterial_DiscreteLOD : GSelectionRegionMaterial)->GetRelevance_Concurrent(FeatureLevel);
+					ToolRelevance |= GSelectionRegionMaterial->GetRelevance_Concurrent(FeatureLevel);
 				}
-				if ((GLandscapeEditRenderMode & ELandscapeEditRenderMode::SelectComponent) && (EditToolRenderData.SelectedType & FLandscapeEditToolRenderData::ST_COMPONENT) && (bUseDiscreteLOD ? GSelectionColorMaterial_DiscreteLOD : GSelectionColorMaterial))
+				if ((GLandscapeEditRenderMode & ELandscapeEditRenderMode::SelectComponent) && (EditToolRenderData.SelectedType & FLandscapeEditToolRenderData::ST_COMPONENT) && GSelectionColorMaterial)
 				{
 					Result.bDynamicRelevance = true;
-					ToolRelevance |= (bUseDiscreteLOD ? GSelectionColorMaterial_DiscreteLOD : GSelectionColorMaterial)->GetRelevance_Concurrent(FeatureLevel);
+					ToolRelevance |= GSelectionColorMaterial->GetRelevance_Concurrent(FeatureLevel);
 				}
 			}
 
 			// Mask
-			if ((GLandscapeEditRenderMode & ELandscapeEditRenderMode::Mask) && (bUseDiscreteLOD ? GMaskRegionMaterial_DiscreteLOD : GMaskRegionMaterial) != nullptr &&
+			if ((GLandscapeEditRenderMode & ELandscapeEditRenderMode::Mask) && GMaskRegionMaterial != nullptr &&
 				(((EditToolRenderData.SelectedType & FLandscapeEditToolRenderData::ST_REGION)) || (!(GLandscapeEditRenderMode & ELandscapeEditRenderMode::InvertedMask))))
 			{
 				Result.bDynamicRelevance = true;
-				ToolRelevance |= (bUseDiscreteLOD ? GMaskRegionMaterial_DiscreteLOD : GMaskRegionMaterial)->GetRelevance_Concurrent(FeatureLevel);
+				ToolRelevance |= GMaskRegionMaterial->GetRelevance_Concurrent(FeatureLevel);
 			}
 
 			if (GLandscapeViewMode == ELandscapeViewMode::LayerContribution)
 			{
 				Result.bDynamicRelevance = true;
-				ToolRelevance |= (bUseDiscreteLOD ? GColorMaskRegionMaterial_DiscreteLOD : GColorMaskRegionMaterial)->GetRelevance_Concurrent(FeatureLevel);
+				ToolRelevance |= GColorMaskRegionMaterial->GetRelevance_Concurrent(FeatureLevel);
 			}
 
 			ToolRelevance.SetPrimitiveViewRelevance(Result);
@@ -1222,8 +1199,7 @@ FPrimitiveViewRelevance FLandscapeComponentSceneProxy::GetViewRelevance(const FS
 
 	// Various visualizations need to render using dynamic relevance
 	if ((View->Family->EngineShowFlags.Bounds && IsSelected()) ||
-		GLandscapeDebugOptions.bShowPatches
-	)
+		GLandscapeDebugOptions.bShowPatches)
 	{
 		Result.bDynamicRelevance = true;
 	}
@@ -1243,7 +1219,6 @@ FPrimitiveViewRelevance FLandscapeComponentSceneProxy::GetViewRelevance(const FS
 #if WITH_EDITOR
 		(IsSelected() && !GLandscapeEditModeActive) ||
 		(GLandscapeViewMode != ELandscapeViewMode::Normal) ||
-		(View->Family->LandscapeLODOverride >= 0) ||
 #else
 		IsSelected() ||
 #endif
@@ -1366,7 +1341,6 @@ void FLandscapeComponentSceneProxy::OnTransformChanged()
 
 	// Set FLandscapeUniformVSParameters for this subsection
 	FLandscapeUniformShaderParameters LandscapeParams;
-	LandscapeParams.LastLOD = LastLOD;
 	LandscapeParams.HeightmapUVScaleBias = HeightmapScaleBias;
 	LandscapeParams.WeightmapUVScaleBias = WeightmapScaleBias;
 	LandscapeParams.LocalToWorldNoScaling = LocalToWorldNoScaling;
@@ -1462,18 +1436,14 @@ float FLandscapeComponentSceneProxy::GetComponentScreenSize(const FSceneView* Vi
 	return SquaredScreenRadius * 2.0f;
 }
 
-extern ENGINE_API int32 GFarShadowStaticMeshLODBias;
-
-void FLandscapeComponentSceneProxy::BuildDynamicMeshElement(const FSceneView* View, const FViewCustomDataLOD* InPrimitiveCustomData, bool InToolMesh, bool InHasTessellation, bool InDisableTessellation, FMeshBatch& OutMeshBatch, TArray<FLandscapeBatchElementParams, SceneRenderingAllocator>& OutStaticBatchParamArray) const
+void FLandscapeComponentSceneProxy::BuildDynamicMeshElement(const FViewCustomDataLOD* InPrimitiveCustomData, bool InToolMesh, bool InHasTessellation, bool InDisableTessellation, FMeshBatch& OutMeshBatch, TArray<FLandscapeBatchElementParams, SceneRenderingAllocator>& OutStaticBatchParamArray) const
 {
 	if (AvailableMaterials.Num() == 0 || InPrimitiveCustomData == nullptr || InPrimitiveCustomData->SubSections.Num() == 0)
 	{
 		return;
 	}
 
-	const int8 CurrentLODIndex = (View->DrawDynamicFlags & EDrawDynamicFlags::FarShadowCascade) ?
-		(FMath::Min(InPrimitiveCustomData->SubSections[0].BatchElementCurrentLOD + GFarShadowStaticMeshLODBias, LastLOD)) : 
-		InPrimitiveCustomData->SubSections[0].BatchElementCurrentLOD;
+	const int8 CurrentLODIndex = InPrimitiveCustomData->SubSections[0].BatchElementCurrentLOD;
 	int8 MaterialIndex = LODIndexToMaterialIndex.IsValidIndex(CurrentLODIndex) ? LODIndexToMaterialIndex[CurrentLODIndex] : INDEX_NONE;
 	UMaterialInterface* SelectedMaterial = MaterialIndex != INDEX_NONE ? AvailableMaterials[MaterialIndex] : nullptr;
 
@@ -1499,7 +1469,7 @@ void FLandscapeComponentSceneProxy::BuildDynamicMeshElement(const FSceneView* Vi
 		check(SharedBuffers->AdjacencyIndexBuffers);
 	}
 
-	OutMeshBatch.VertexFactory = bUseDiscreteLOD ? FixedGridVertexFactory : VertexFactory;
+	OutMeshBatch.VertexFactory = VertexFactory;
 	OutMeshBatch.MaterialRenderProxy = SelectedMaterial->GetRenderProxy();
 	OutMeshBatch.LCI = ComponentLightInfo.Get();
 	OutMeshBatch.ReverseCulling = IsLocalToWorldDeterminantNegative();
@@ -1518,9 +1488,7 @@ void FLandscapeComponentSceneProxy::BuildDynamicMeshElement(const FSceneView* Vi
 			for (int32 SubX = 0; SubX < NumSubsections; SubX++)
 			{
 				const int8 SubSectionIdx = SubX + SubY * NumSubsections;
-				const int8 CurrentLOD = (View->DrawDynamicFlags & EDrawDynamicFlags::FarShadowCascade) ?
-					(FMath::Min(InPrimitiveCustomData->SubSections[SubSectionIdx].BatchElementCurrentLOD + GFarShadowStaticMeshLODBias, LastLOD)) :
-					InPrimitiveCustomData->SubSections[SubSectionIdx].BatchElementCurrentLOD;
+				const int8 CurrentLOD = InPrimitiveCustomData->SubSections[SubSectionIdx].BatchElementCurrentLOD;
 
 				FMeshBatchElement BatchElement;
 				FLandscapeBatchElementParams& BatchElementParams = OutStaticBatchParamArray[SubSectionIdx];
@@ -1617,7 +1585,7 @@ bool FLandscapeComponentSceneProxy::GetMeshElement(bool UseSeperateBatchForShado
 		check(SharedBuffers->AdjacencyIndexBuffers);
 	}
 
-	OutMeshBatch.VertexFactory = bUseDiscreteLOD ? FixedGridVertexFactory : VertexFactory;
+	OutMeshBatch.VertexFactory = VertexFactory;
 	OutMeshBatch.MaterialRenderProxy = InMaterialInterface->GetRenderProxy();
 
 	OutMeshBatch.LCI = ComponentLightInfo.Get();
@@ -1826,71 +1794,68 @@ void FLandscapeComponentSceneProxy::DrawStaticElements(FStaticPrimitiveDrawInter
 	StaticBatchParamArray.Empty(TotalBatchCount);
 	PDI->ReserveMemoryForMeshes(TotalBatchCount);
 
-	if (!bUseDiscreteLOD)
-	{
-		MaterialIndexToStaticMeshBatchLOD.SetNumUninitialized(MaterialCount);
+	MaterialIndexToStaticMeshBatchLOD.SetNumUninitialized(MaterialCount);
 
-		int32 CurrentLODIndex = 0;
-		for (int32 i = 0; i < MaterialCount; ++i)
+	int32 CurrentLODIndex = 0;
+	for (int32 i = 0; i < MaterialCount; ++i)
+	{
+		if (AvailableMaterials[i] == nullptr)
 		{
-			if (AvailableMaterials[i] == nullptr)
-			{
-				continue;
-			}
+			continue;
+		}
 
-			UMaterialInstance* MaterialInstance = Cast<UMaterialInstance>(AvailableMaterials[i]);
+		UMaterialInstance* MaterialInstance = Cast<UMaterialInstance>(AvailableMaterials[i]);
 
 		bool HasTessellationEnabled = (FeatureLevel >= ERHIFeatureLevel::SM5) ? MaterialInstance != nullptr && MaterialInstance->GetMaterial()->D3D11TessellationMode != EMaterialTessellationMode::MTM_NoTessellation && MaterialIndexToDisabledTessellationMaterial[i] != INDEX_NONE : false;
 
-			// Only add normal batch index (for each material)
-			MaterialIndexToStaticMeshBatchLOD[i] = CurrentLODIndex;
+		// Only add normal batch index (for each material)
+		MaterialIndexToStaticMeshBatchLOD[i] = CurrentLODIndex;
 
-			// Default Batch, tessellated if enabled
-			FMeshBatch NormalMeshBatch;
+		// Default Batch, tessellated if enabled
+		FMeshBatch NormalMeshBatch;
 
-			if (GetMeshElement(false, false, HasTessellationEnabled, CurrentLODIndex++, AvailableMaterials[i], NormalMeshBatch, StaticBatchParamArray))
-			{
-				PDI->DrawMesh(NormalMeshBatch, FLT_MAX);
-			}
+		if (GetMeshElement(false, false, HasTessellationEnabled, CurrentLODIndex++, AvailableMaterials[i], NormalMeshBatch, StaticBatchParamArray))
+		{
+			PDI->DrawMesh(NormalMeshBatch, FLT_MAX);
+		}
 
-			if (HasTessellationEnabled)
-			{
-				UMaterialInstance* NonTessellatedLandscapeMI = Cast<UMaterialInstance>(AvailableMaterials[MaterialIndexToDisabledTessellationMaterial[i]]);
+		if (HasTessellationEnabled)
+		{
+			UMaterialInstance* NonTessellatedLandscapeMI = Cast<UMaterialInstance>(AvailableMaterials[MaterialIndexToDisabledTessellationMaterial[i]]);
 
-				// Make sure that the Material instance we are going to use has the tessellation disabled
-				UMaterialInstanceDynamic* NonTessellatedLandscapeMID = Cast<UMaterialInstanceDynamic>(NonTessellatedLandscapeMI);
-				ULandscapeMaterialInstanceConstant* NonTessellatedLandscapeMIC = Cast<ULandscapeMaterialInstanceConstant>(NonTessellatedLandscapeMI);
+			// Make sure that the Material instance we are going to use has the tessellation disabled
+			UMaterialInstanceDynamic* NonTessellatedLandscapeMID = Cast<UMaterialInstanceDynamic>(NonTessellatedLandscapeMI);
+			ULandscapeMaterialInstanceConstant* NonTessellatedLandscapeMIC = Cast<ULandscapeMaterialInstanceConstant>(NonTessellatedLandscapeMI);
 
-				if (NonTessellatedLandscapeMID != nullptr)
-				{
-					NonTessellatedLandscapeMIC = Cast<ULandscapeMaterialInstanceConstant>(NonTessellatedLandscapeMID->Parent);
-				}
+			if (NonTessellatedLandscapeMID != nullptr)
+			{
+				NonTessellatedLandscapeMIC = Cast<ULandscapeMaterialInstanceConstant>(NonTessellatedLandscapeMID->Parent);
+			}
 
-				check(NonTessellatedLandscapeMIC != nullptr && NonTessellatedLandscapeMIC->bDisableTessellation);
+			check(NonTessellatedLandscapeMIC != nullptr && NonTessellatedLandscapeMIC->bDisableTessellation);
 
-				// Default Batch, no Tessellation enabled
-				FMeshBatch NonTesselatedNormalMeshBatch;
+			// Default Batch, no Tessellation enabled
+			FMeshBatch NonTesselatedNormalMeshBatch;
 
-				if (GetMeshElement(true, false, false, CurrentLODIndex++, NonTessellatedLandscapeMI, NonTesselatedNormalMeshBatch, StaticBatchParamArray))
-				{
-					PDI->DrawMesh(NonTesselatedNormalMeshBatch, FLT_MAX);
-				}
+			if (GetMeshElement(true, false, false, CurrentLODIndex++, NonTessellatedLandscapeMI, NonTesselatedNormalMeshBatch, StaticBatchParamArray))
+			{
+				PDI->DrawMesh(NonTesselatedNormalMeshBatch, FLT_MAX);
+			}
 
-				// Shadow Batch, tessellated if enabled
-				FMeshBatch ShadowMeshBatch;
+			// Shadow Batch, tessellated if enabled
+			FMeshBatch ShadowMeshBatch;
 
-				if (GetMeshElement(true, true, HasTessellationEnabled, CurrentLODIndex++, MaterialInstance, ShadowMeshBatch, StaticBatchParamArray))
-				{
-					PDI->DrawMesh(ShadowMeshBatch, FLT_MAX);
-				}
+			if (GetMeshElement(true, true, HasTessellationEnabled, CurrentLODIndex++, MaterialInstance, ShadowMeshBatch, StaticBatchParamArray))
+			{
+				PDI->DrawMesh(ShadowMeshBatch, FLT_MAX);
+			}
 
-				// Shadow Batch, no tessellation enabled
-				FMeshBatch NonTesselatedShadowMeshBatch;
+			// Shadow Batch, no tessellation enabled
+			FMeshBatch NonTesselatedShadowMeshBatch;
 
-				if (GetMeshElement(true, true, false, CurrentLODIndex++, NonTessellatedLandscapeMI, NonTesselatedShadowMeshBatch, StaticBatchParamArray))
-				{
-					PDI->DrawMesh(NonTesselatedShadowMeshBatch, FLT_MAX);
-				}
+			if (GetMeshElement(true, true, false, CurrentLODIndex++, NonTessellatedLandscapeMI, NonTesselatedShadowMeshBatch, StaticBatchParamArray))
+			{
+				PDI->DrawMesh(NonTesselatedShadowMeshBatch, FLT_MAX);
 			}
 		}
 	}
@@ -1910,113 +1875,6 @@ void FLandscapeComponentSceneProxy::DrawStaticElements(FStaticPrimitiveDrawInter
 		}
  	}
 
-	if (bUseDiscreteLOD)
-	{
-		for (int32 LODIndex = FirstLOD; LODIndex <= LastLOD; LODIndex++)
-		{
-			UMaterialInstance* MaterialInstance = nullptr;
-			UMaterialInstance* NonTessellatedLandscapeMI = nullptr;
-			bool HasTessellationEnabled = false;
-
-			{
-				int8 MaterialIndex = LODIndexToMaterialIndex[LODIndex];
-
-				MaterialInstance = Cast<UMaterialInstance>(AvailableMaterials[MaterialIndex]);
-
-				HasTessellationEnabled = (FeatureLevel >= ERHIFeatureLevel::SM5) ?
-					(MaterialInstance != nullptr && MaterialInstance->GetMaterial()->D3D11TessellationMode != EMaterialTessellationMode::MTM_NoTessellation && MaterialIndexToDisabledTessellationMaterial[MaterialIndex] != INDEX_NONE) :
-					false;
-
-				if (HasTessellationEnabled)
-				{
-					NonTessellatedLandscapeMI = Cast<UMaterialInstance>(AvailableMaterials[MaterialIndexToDisabledTessellationMaterial[MaterialIndex]]);
-
-					// Make sure that the Material instance we are going to use has the tessellation disabled
-					UMaterialInstanceDynamic* NonTessellatedLandscapeMID = Cast<UMaterialInstanceDynamic>(NonTessellatedLandscapeMI);
-					ULandscapeMaterialInstanceConstant* NonTessellatedLandscapeMIC = Cast<ULandscapeMaterialInstanceConstant>(NonTessellatedLandscapeMI);
-
-					if (NonTessellatedLandscapeMID != nullptr)
-					{
-						NonTessellatedLandscapeMIC = Cast<ULandscapeMaterialInstanceConstant>(NonTessellatedLandscapeMID->Parent);
-					}
-
-					check(NonTessellatedLandscapeMIC != nullptr && NonTessellatedLandscapeMIC->bDisableTessellation);
-				}
-			}
-
-			if (MaterialInstance == nullptr)
-			{
-				continue;
-			}
-
-			if (HasTessellationEnabled && LODScreenRatioSquared[LODIndex] < TessellationComponentSquaredScreenSize)
-			{
-				// Disable tessellation by replacing MaterialInstance with the non tessellated version
-				MaterialInstance = NonTessellatedLandscapeMI;
-			}
-
-			// Could be different from bRequiresAdjacencyInformation during shader compilation
-			bool bCurrentRequiresAdjacencyInformation = MaterialRenderingRequiresAdjacencyInformation_RenderingThread(MaterialInstance, VertexFactory->GetType(), GetScene().GetFeatureLevel());
-
-			if (bCurrentRequiresAdjacencyInformation)
-			{
-				check(SharedBuffers->AdjacencyIndexBuffers);
-			}
-
-			FMeshBatch BaseMeshBatch;
-
-			BaseMeshBatch.VertexFactory = FixedGridVertexFactory;
-			BaseMeshBatch.MaterialRenderProxy = MaterialInstance->GetRenderProxy();
-
-			BaseMeshBatch.LCI = ComponentLightInfo.Get();
-			BaseMeshBatch.ReverseCulling = IsLocalToWorldDeterminantNegative();
-			BaseMeshBatch.CastShadow = true;
-			BaseMeshBatch.bUseForDepthPass = true;
-			BaseMeshBatch.bUseAsOccluder = ShouldUseAsOccluder() && GetScene().GetShadingPath() == EShadingPath::Deferred && !IsMovable();
-			BaseMeshBatch.bUseForMaterial = true;
-			BaseMeshBatch.Type = bCurrentRequiresAdjacencyInformation ? PT_12_ControlPointPatchList : PT_TriangleList;
-			BaseMeshBatch.DepthPriorityGroup = SDPG_World;
-			BaseMeshBatch.LODIndex = LODIndex;
-			BaseMeshBatch.bRequiresPerElementVisibility = false;
-			BaseMeshBatch.bDitheredLODTransition = false;
-
-			if (BaseMeshBatch.CastShadow)
-			{
-				float QuadWorldSize = (ComponentMaxExtend * 2.0f) / (SubsectionSizeQuads * NumSubsections);
-				FVector QuadSize = FVector(QuadWorldSize, QuadWorldSize, bCurrentRequiresAdjacencyInformation ? BaseMeshBatch.MaterialRenderProxy->GetMaterial(GetScene().GetFeatureLevel())->GetMaxDisplacement() : 0.0f);
-
-				BaseMeshBatch.TessellationDisablingShadowMapMeshSize = QuadSize.Size();
-			}
-
-			int32 LodSubsectionSizeVerts = SubsectionSizeVerts >> LODIndex;
-			uint32 NumPrimitivesPerSection = FMath::Square((LodSubsectionSizeVerts - 1)) * 2;
-
-			// Combined batch element
-			FLandscapeBatchElementParams* BatchElementParams = new(StaticBatchParamArray) FLandscapeBatchElementParams;
-			BatchElementParams->LandscapeUniformShaderParametersResource = &LandscapeUniformShaderParameters;
-			BatchElementParams->LocalToWorldNoScalingPtr = &LocalToWorldNoScaling;
-			BatchElementParams->SceneProxy = this;
-			BatchElementParams->SubX = -1;
-			BatchElementParams->SubY = -1;
-			BatchElementParams->CurrentLOD = LODIndex;
-
-			FMeshBatchElement BatchElement;
-
-			BatchElement.UserData = BatchElementParams;
-			BatchElement.PrimitiveUniformBuffer = GetUniformBuffer();
-			BatchElement.IndexBuffer = bCurrentRequiresAdjacencyInformation ? SharedBuffers->AdjacencyIndexBuffers->IndexBuffers[LODIndex] : SharedBuffers->IndexBuffers[LODIndex];
-			BatchElement.NumPrimitives = FMath::Square((LodSubsectionSizeVerts - 1)) * FMath::Square(NumSubsections) * 2;
-			BatchElement.FirstIndex = 0;
-			BatchElement.MinVertexIndex = SharedBuffers->IndexRanges[LODIndex].MinIndexFull;
-			BatchElement.MaxVertexIndex = SharedBuffers->IndexRanges[LODIndex].MaxIndexFull;
-
-			FMeshBatch MeshBatch = BaseMeshBatch;
-			MeshBatch.Elements[0] = BatchElement;
-
-			PDI->DrawMesh(MeshBatch, LODIndex == FirstLOD ? FLT_MAX : (FMath::Sqrt(LODScreenRatioSquared[LODIndex]) * 2.0f));
-		}
-	}
-
 	check(StaticBatchParamArray.Num() <= TotalBatchCount);
 }
 
@@ -2621,12 +2479,12 @@ void FLandscapeComponentSceneProxy::PostInitViewCustomData(const FSceneView& InV
 
 bool FLandscapeComponentSceneProxy::IsUsingCustomLODRules() const
 {
-	return !bUseDiscreteLOD;
+	return true;
 }
 
 bool FLandscapeComponentSceneProxy::IsUsingCustomWholeSceneShadowLODRules() const
 {
-	return !bUseDiscreteLOD;
+	return true;
 }
 
 bool FLandscapeComponentSceneProxy::CanUseMeshBatchForShadowCascade(int8 InLODIndex, float InShadowMapTextureResolution, float InShadowMapCascadeSize) const
@@ -2878,12 +2736,7 @@ void FLandscapeComponentSceneProxy::GetDynamicMeshElements(const TArray<const FS
 
 			if (PrimitiveCustomData == nullptr)
 			{
-				float MeshScreenSizeSquared = ComputeBoundsScreenRadiusSquared(GetBounds().Origin, GetBounds().SphereRadius, *View);
-				int32 ForcedLODLevel = (View->Family->EngineShowFlags.LOD) ? GetCVarForceLOD() : 0;
-
-				FLODMask LODToRender;
-				LODToRender.SetLOD(GetLODFromScreenSize(MeshScreenSizeSquared, View->LODDistanceFactor));
-				PrimitiveCustomData = (const FViewCustomDataLOD*)(const_cast<FLandscapeComponentSceneProxy*>(this)->InitViewCustomData(*View, View->LODDistanceFactor, FMemStack::Get(), false, false, &LODToRender, MeshScreenSizeSquared));
+				continue;
 			}
 
 			ParameterArray.ElementParams.AddDefaulted(NumSubsections * NumSubsections);
@@ -2929,13 +2782,13 @@ void FLandscapeComponentSceneProxy::GetDynamicMeshElements(const TArray<const FS
 			}
 
 			FMeshBatch& Mesh = Collector.AllocateMesh();
-			BuildDynamicMeshElement(View, PrimitiveCustomData, false, HasTessellationEnabled, DisableTessellation, Mesh, ParameterArray.ElementParams);
+			BuildDynamicMeshElement(PrimitiveCustomData, false, HasTessellationEnabled, DisableTessellation, Mesh, ParameterArray.ElementParams);
 
 #if WITH_EDITOR
 			FMeshBatch& MeshTools = Collector.AllocateMesh();
 
 			// No Tessellation on tool material
-			BuildDynamicMeshElement(View, PrimitiveCustomData, true, HasTessellationEnabled, true, MeshTools, ParameterArray.ElementParams);
+			BuildDynamicMeshElement(PrimitiveCustomData, true, HasTessellationEnabled, true, MeshTools, ParameterArray.ElementParams);
 #endif
 
 			// Render the landscape component
@@ -2944,9 +2797,9 @@ void FLandscapeComponentSceneProxy::GetDynamicMeshElements(const TArray<const FS
 			{
 			case ELandscapeViewMode::DebugLayer:
 			{
-				if ((bUseDiscreteLOD ? GLayerDebugColorMaterial_DiscreteLOD : GLayerDebugColorMaterial))
+				if (GLayerDebugColorMaterial)
 				{
-					auto DebugColorMaterialInstance = new FLandscapeDebugMaterialRenderProxy((bUseDiscreteLOD ? GLayerDebugColorMaterial_DiscreteLOD : GLayerDebugColorMaterial)->GetRenderProxy(),
+					auto DebugColorMaterialInstance = new FLandscapeDebugMaterialRenderProxy(GLayerDebugColorMaterial->GetRenderProxy(),
 						(EditToolRenderData.DebugChannelR >= 0 ? WeightmapTextures[EditToolRenderData.DebugChannelR / 4] : nullptr),
 						(EditToolRenderData.DebugChannelG >= 0 ? WeightmapTextures[EditToolRenderData.DebugChannelG / 4] : nullptr),
 						(EditToolRenderData.DebugChannelB >= 0 ? WeightmapTextures[EditToolRenderData.DebugChannelB / 4] : nullptr),
@@ -2991,10 +2844,10 @@ void FLandscapeComponentSceneProxy::GetDynamicMeshElements(const TArray<const FS
 
 			case ELandscapeViewMode::LayerUsage:
 			{
-				if ((bUseDiscreteLOD ? GLandscapeLayerUsageMaterial_DiscreteLOD : GLandscapeLayerUsageMaterial))
+				if (GLandscapeLayerUsageMaterial)
 				{
 					float Rotation = ((SectionBase.X / ComponentSizeQuads) ^ (SectionBase.Y / ComponentSizeQuads)) & 1 ? 0 : 2.f * PI;
-					auto LayerUsageMaterialInstance = new FLandscapeLayerUsageRenderProxy((bUseDiscreteLOD ? GLandscapeLayerUsageMaterial_DiscreteLOD : GLandscapeLayerUsageMaterial)->GetRenderProxy(), ComponentSizeVerts, LayerColors, Rotation);
+					auto LayerUsageMaterialInstance = new FLandscapeLayerUsageRenderProxy(GLandscapeLayerUsageMaterial->GetRenderProxy(), ComponentSizeVerts, LayerColors, Rotation);
 					MeshTools.MaterialRenderProxy = LayerUsageMaterialInstance;
 					Collector.RegisterOneFrameMaterialProxy(LayerUsageMaterialInstance);
 					MeshTools.bCanApplyViewModeOverrides = true;
@@ -3071,7 +2924,7 @@ void FLandscapeComponentSceneProxy::GetDynamicMeshElements(const TArray<const FS
 
 				FMeshBatch& MaskMesh = Collector.AllocateMesh();
 				MaskMesh = MeshTools;
-				auto ColorMaskMaterialInstance = new FLandscapeMaskMaterialRenderProxy((bUseDiscreteLOD ? GColorMaskRegionMaterial_DiscreteLOD : GColorMaskRegionMaterial)->GetRenderProxy(), EditToolRenderData.LayerContributionTexture ? EditToolRenderData.LayerContributionTexture : GLandscapeBlackTexture, true);
+				auto ColorMaskMaterialInstance = new FLandscapeMaskMaterialRenderProxy(GColorMaskRegionMaterial->GetRenderProxy(), EditToolRenderData.LayerContributionTexture ? EditToolRenderData.LayerContributionTexture : GLandscapeBlackTexture, true);
 				MaskMesh.MaterialRenderProxy = ColorMaskMaterialInstance;
 				Collector.RegisterOneFrameMaterialProxy(ColorMaskMaterialInstance);
 				Collector.AddMesh(ViewIndex, MaskMesh);
@@ -3116,7 +2969,6 @@ void FLandscapeComponentSceneProxy::GetDynamicMeshElements(const TArray<const FS
 						bIsWireframe ||
 #if WITH_EDITOR
 						(IsSelected() && !GLandscapeEditModeActive) ||
-						ViewFamily.LandscapeLODOverride >= 0 ||
 #else
 						IsSelected() ||
 #endif
@@ -3148,7 +3000,7 @@ void FLandscapeComponentSceneProxy::GetDynamicMeshElements(const TArray<const FS
 					{
 						FMeshBatch& SelectMesh = Collector.AllocateMesh();
 						SelectMesh = MeshTools;
-						auto SelectMaterialInstance = new FLandscapeSelectMaterialRenderProxy((bUseDiscreteLOD ? GSelectionRegionMaterial_DiscreteLOD : GSelectionRegionMaterial)->GetRenderProxy(), EditToolRenderData.DataTexture ? EditToolRenderData.DataTexture : GLandscapeBlackTexture);
+						auto SelectMaterialInstance = new FLandscapeSelectMaterialRenderProxy(GSelectionRegionMaterial->GetRenderProxy(), EditToolRenderData.DataTexture ? EditToolRenderData.DataTexture : GLandscapeBlackTexture);
 						SelectMesh.MaterialRenderProxy = SelectMaterialInstance;
 						Collector.RegisterOneFrameMaterialProxy(SelectMaterialInstance);
 						Collector.AddMesh(ViewIndex, SelectMesh);
@@ -3161,7 +3013,7 @@ void FLandscapeComponentSceneProxy::GetDynamicMeshElements(const TArray<const FS
 					{
 						FMeshBatch& SelectMesh = Collector.AllocateMesh();
 						SelectMesh = MeshTools;
-						SelectMesh.MaterialRenderProxy = (bUseDiscreteLOD ? GSelectionColorMaterial_DiscreteLOD : GSelectionColorMaterial)->GetRenderProxy();
+						SelectMesh.MaterialRenderProxy = GSelectionColorMaterial->GetRenderProxy();
 						Collector.AddMesh(ViewIndex, SelectMesh);
 						NumPasses++;
 						NumTriangles += SelectMesh.GetNumPrimitives();
@@ -3176,7 +3028,7 @@ void FLandscapeComponentSceneProxy::GetDynamicMeshElements(const TArray<const FS
 					{
 						FMeshBatch& MaskMesh = Collector.AllocateMesh();
 						MaskMesh = MeshTools;
-						auto MaskMaterialInstance = new FLandscapeMaskMaterialRenderProxy((bUseDiscreteLOD ? GMaskRegionMaterial_DiscreteLOD : GMaskRegionMaterial)->GetRenderProxy(), EditToolRenderData.DataTexture ? EditToolRenderData.DataTexture : GLandscapeBlackTexture, !!(GLandscapeEditRenderMode & ELandscapeEditRenderMode::InvertedMask));
+						auto MaskMaterialInstance = new FLandscapeMaskMaterialRenderProxy(GMaskRegionMaterial->GetRenderProxy(), EditToolRenderData.DataTexture ? EditToolRenderData.DataTexture : GLandscapeBlackTexture, !!(GLandscapeEditRenderMode & ELandscapeEditRenderMode::InvertedMask));
 						MaskMesh.MaterialRenderProxy = MaskMaterialInstance;
 						Collector.RegisterOneFrameMaterialProxy(MaskMaterialInstance);
 						Collector.AddMesh(ViewIndex, MaskMesh);
@@ -3188,7 +3040,7 @@ void FLandscapeComponentSceneProxy::GetDynamicMeshElements(const TArray<const FS
 					{
 						FMeshBatch& MaskMesh = Collector.AllocateMesh();
 						MaskMesh = MeshTools;
-						auto MaskMaterialInstance = new FLandscapeMaskMaterialRenderProxy((bUseDiscreteLOD ? GMaskRegionMaterial_DiscreteLOD : GMaskRegionMaterial)->GetRenderProxy(), GLandscapeBlackTexture, false);
+						auto MaskMaterialInstance = new FLandscapeMaskMaterialRenderProxy(GMaskRegionMaterial->GetRenderProxy(), GLandscapeBlackTexture, false);
 						MaskMesh.MaterialRenderProxy = MaskMaterialInstance;
 						Collector.RegisterOneFrameMaterialProxy(MaskMaterialInstance);
 						Collector.AddMesh(ViewIndex, MaskMesh);
@@ -3821,7 +3673,6 @@ FLandscapeSharedBuffers::FLandscapeSharedBuffers(const int32 InSharedBuffersKey,
 	, SubsectionSizeVerts(InSubsectionSizeQuads + 1)
 	, NumSubsections(InNumSubsections)
 	, VertexFactory(nullptr)
-	, FixedGridVertexFactory(nullptr)
 	, VertexBuffer(nullptr)
 	, AdjacencyIndexBuffers(nullptr)
 	, bUse32BitIndices(false)
@@ -3907,11 +3758,6 @@ FLandscapeSharedBuffers::~FLandscapeSharedBuffers()
 	delete AdjacencyIndexBuffers;
 	delete VertexFactory;
 
-	if (FixedGridVertexFactory != nullptr)
-	{
-		delete FixedGridVertexFactory;
-	}
-
 	if (OccluderIndicesSP.IsValid())
 	{
 		DEC_DWORD_STAT_BY(STAT_LandscapeOccluderMem, OccluderIndicesSP->GetAllocatedSize());
@@ -4078,126 +3924,108 @@ public:
 
 		int32 SubSectionIndex = BatchElementParams->SubX + BatchElementParams->SubY * SceneProxy->NumSubsections;
 
-		if (SceneProxy->bUseDiscreteLOD)
-		{
-			// In fact when discrete LOD is used, we're using FLandscapeFixedGridVertexFactory and not going through here
-			// However in certain cases (like in editor landscape tools) this will still be used
-			// Set parameters accordingly to get the same behavior as discrete LOD
-			FVector4 LodTessellationParams(ForceInitToZero);
-			FVector4 LodBias(ForceInitToZero);
-			FVector4 ShaderCurrentLOD(BatchElementParams->CurrentLOD, BatchElementParams->CurrentLOD, BatchElementParams->CurrentLOD, BatchElementParams->CurrentLOD);
-			FVector4 ShaderCurrentNeighborLOD[FLandscapeComponentSceneProxy::NEIGHBOR_COUNT] = { ShaderCurrentLOD, ShaderCurrentLOD, ShaderCurrentLOD, ShaderCurrentLOD };
+		// If we have no custom data for this primitive we will compute of the fly the proper values, this will happen if the shader is not used for normal landscape rendering(i.e grass rendering)
+		FLandscapeComponentSceneProxy::FViewCustomDataLOD* LODData = (FLandscapeComponentSceneProxy::FViewCustomDataLOD*)InView->GetCustomData(SceneProxy->GetPrimitiveSceneInfo()->GetIndex());
 
-			ShaderBindings.Add(LodTessellationParameter, LodTessellationParams);
-			ShaderBindings.Add(LodBiasParameter, LodBias);
-			ShaderBindings.Add(SectionLodsParameter, ShaderCurrentLOD);
-			ShaderBindings.Add(NeighborSectionLodParameter, ShaderCurrentNeighborLOD);
-		}
-		else
+		if (LODData != nullptr)
 		{
-			// If we have no custom data for this primitive we will compute of the fly the proper values, this will happen if the shader is not used for normal landscape rendering(i.e grass rendering)
-			FLandscapeComponentSceneProxy::FViewCustomDataLOD* LODData = (FLandscapeComponentSceneProxy::FViewCustomDataLOD*)InView->GetCustomData(SceneProxy->GetPrimitiveSceneInfo()->GetIndex());
+			SceneProxy->PostInitViewCustomData(*InView, LODData);
 
-			if (LODData != nullptr)
+			if (LodTessellationParameter.IsBound())
 			{
-				SceneProxy->PostInitViewCustomData(*InView, LODData);
+				ShaderBindings.Add(LodTessellationParameter, LODData->LodTessellationParams);
+			}
+				
+			if (LodBiasParameter.IsBound())
+			{
+				check(LODData->LodBias == SceneProxy->GetShaderLODBias());
+				ShaderBindings.Add(LodBiasParameter, LODData->LodBias);
+			}
 
-				if (LodTessellationParameter.IsBound())
+			if (SectionLodsParameter.IsBound())
+			{
+				if (LODData->UseCombinedMeshBatch)
 				{
-					ShaderBindings.Add(LodTessellationParameter, LODData->LodTessellationParams);
+					ShaderBindings.Add(SectionLodsParameter, LODData->ShaderCurrentLOD);
 				}
-
-				if (LodBiasParameter.IsBound())
+				else // in non combined, only the one representing us as we'll be called 4 times (once per sub section)
 				{
-					check(LODData->LodBias == SceneProxy->GetShaderLODBias());
-					ShaderBindings.Add(LodBiasParameter, LODData->LodBias);
+					check(SubSectionIndex >= 0);
+					FVector4 ShaderCurrentLOD(ForceInitToZero);
+					ShaderCurrentLOD.Component(SubSectionIndex) = LODData->ShaderCurrentLOD.Component(SubSectionIndex);
+
+					ShaderBindings.Add(SectionLodsParameter, ShaderCurrentLOD);
 				}
+			}
 
-				if (SectionLodsParameter.IsBound())
-				{
-					if (LODData->UseCombinedMeshBatch)
+			if (NeighborSectionLodParameter.IsBound())
+			{
+				FVector4 ShaderCurrentNeighborLOD[FLandscapeComponentSceneProxy::NEIGHBOR_COUNT] = { FVector4(ForceInitToZero), FVector4(ForceInitToZero), FVector4(ForceInitToZero), FVector4(ForceInitToZero) };
+
+				if (LODData->UseCombinedMeshBatch)
+				{					
+					int32 SubSectionCount = SceneProxy->NumSubsections == 1 ? 1 : FLandscapeComponentSceneProxy::MAX_SUBSECTION_COUNT;
+
+					for (int32 NeighborSubSectionIndex = 0; NeighborSubSectionIndex < SubSectionCount; ++NeighborSubSectionIndex)
 					{
-						ShaderBindings.Add(SectionLodsParameter, LODData->ShaderCurrentLOD);
+						ShaderCurrentNeighborLOD[NeighborSubSectionIndex] = LODData->SubSections[NeighborSubSectionIndex].ShaderCurrentNeighborLOD;
+						check(ShaderCurrentNeighborLOD[NeighborSubSectionIndex].X != -1.0f); // they should all match so only check the 1st one for simplicity
 					}
-					else // in non combined, only the one representing us as we'll be called 4 times (once per sub section)
-					{
-						check(SubSectionIndex >= 0);
-						FVector4 ShaderCurrentLOD(ForceInitToZero);
-						ShaderCurrentLOD.Component(SubSectionIndex) = LODData->ShaderCurrentLOD.Component(SubSectionIndex);
 
-						ShaderBindings.Add(SectionLodsParameter, ShaderCurrentLOD);
-					}
+					ShaderBindings.Add(NeighborSectionLodParameter, ShaderCurrentNeighborLOD);
 				}
-
-				if (NeighborSectionLodParameter.IsBound())
+				else // in non combined, only the one representing us as we'll be called 4 times (once per sub section)
 				{
-					FVector4 ShaderCurrentNeighborLOD[FLandscapeComponentSceneProxy::NEIGHBOR_COUNT] = { FVector4(ForceInitToZero), FVector4(ForceInitToZero), FVector4(ForceInitToZero), FVector4(ForceInitToZero) };
-
-					if (LODData->UseCombinedMeshBatch)
-					{
-						int32 SubSectionCount = SceneProxy->NumSubsections == 1 ? 1 : FLandscapeComponentSceneProxy::MAX_SUBSECTION_COUNT;
-
-						for (int32 NeighborSubSectionIndex = 0; NeighborSubSectionIndex < SubSectionCount; ++NeighborSubSectionIndex)
-						{
-							ShaderCurrentNeighborLOD[NeighborSubSectionIndex] = LODData->SubSections[NeighborSubSectionIndex].ShaderCurrentNeighborLOD;
-							check(ShaderCurrentNeighborLOD[NeighborSubSectionIndex].X != -1.0f); // they should all match so only check the 1st one for simplicity
-						}
+					check(SubSectionIndex >= 0);
+					ShaderCurrentNeighborLOD[SubSectionIndex] = LODData->SubSections[SubSectionIndex].ShaderCurrentNeighborLOD;
+					check(ShaderCurrentNeighborLOD[SubSectionIndex].X != -1.0f); // they should all match so only check the 1st one for simplicity
 
-						ShaderBindings.Add(NeighborSectionLodParameter, ShaderCurrentNeighborLOD);
-					}
-					else // in non combined, only the one representing us as we'll be called 4 times (once per sub section)
-					{
-						check(SubSectionIndex >= 0);
-						ShaderCurrentNeighborLOD[SubSectionIndex] = LODData->SubSections[SubSectionIndex].ShaderCurrentNeighborLOD;
-						check(ShaderCurrentNeighborLOD[SubSectionIndex].X != -1.0f); // they should all match so only check the 1st one for simplicity
-
-						ShaderBindings.Add(NeighborSectionLodParameter, ShaderCurrentNeighborLOD);
-					}
+					ShaderBindings.Add(NeighborSectionLodParameter, ShaderCurrentNeighborLOD);
 				}
-			}
-			else
-			{
+			}			
+		}
+		else
+		{
 			float ComponentScreenSize = SceneProxy->GetComponentScreenSize(InView, SceneProxy->GetBounds().Origin, SceneProxy->ComponentMaxExtend, SceneProxy->GetBounds().SphereRadius);
+			
+			FLandscapeComponentSceneProxy::FViewCustomDataLOD CurrentLODData;
+			SceneProxy->CalculateBatchElementLOD(*InView, ComponentScreenSize, InView->LODDistanceFactor, CurrentLODData, true);
+			check(CurrentLODData.UseCombinedMeshBatch);
 
-				FLandscapeComponentSceneProxy::FViewCustomDataLOD CurrentLODData;
-				SceneProxy->CalculateBatchElementLOD(*InView, ComponentScreenSize, InView->LODDistanceFactor, CurrentLODData, true);
-				check(CurrentLODData.UseCombinedMeshBatch);
+			if (LodBiasParameter.IsBound())
+			{
+				ShaderBindings.Add(LodBiasParameter, SceneProxy->GetShaderLODBias());
+			}
 
-				if (LodBiasParameter.IsBound())
-				{
-					ShaderBindings.Add(LodBiasParameter, SceneProxy->GetShaderLODBias());
-				}
+			if (LodTessellationParameter.IsBound())
+			{
+				FVector4 LodTessellationParams(ForceInitToZero);
+				SceneProxy->ComputeTessellationFalloffShaderValues(CurrentLODData, InView->ViewMatrices.GetProjectionMatrix(), LodTessellationParams.X, LodTessellationParams.Y);
 
-				if (LodTessellationParameter.IsBound())
-				{
-					FVector4 LodTessellationParams(ForceInitToZero);
-					SceneProxy->ComputeTessellationFalloffShaderValues(CurrentLODData, InView->ViewMatrices.GetProjectionMatrix(), LodTessellationParams.X, LodTessellationParams.Y);
+				ShaderBindings.Add(LodTessellationParameter, LodTessellationParams);
+			}
 
-					ShaderBindings.Add(LodTessellationParameter, LodTessellationParams);
-				}
+			if (SectionLodsParameter.IsBound())
+			{
+				ShaderBindings.Add(SectionLodsParameter, CurrentLODData.ShaderCurrentLOD);
+			}
 
-				if (SectionLodsParameter.IsBound())
-				{
-					ShaderBindings.Add(SectionLodsParameter, CurrentLODData.ShaderCurrentLOD);
-				}
+			if (NeighborSectionLodParameter.IsBound())
+			{
+				FVector4 CurrentNeighborLOD[4] = { FVector4(ForceInitToZero), FVector4(ForceInitToZero), FVector4(ForceInitToZero), FVector4(ForceInitToZero) };
 
-				if (NeighborSectionLodParameter.IsBound())
+				for (int32 SubY = 0; SubY < SceneProxy->NumSubsections; SubY++)
 				{
-					FVector4 CurrentNeighborLOD[4] = { FVector4(ForceInitToZero), FVector4(ForceInitToZero), FVector4(ForceInitToZero), FVector4(ForceInitToZero) };
-
-					for (int32 SubY = 0; SubY < SceneProxy->NumSubsections; SubY++)
+					for (int32 SubX = 0; SubX < SceneProxy->NumSubsections; SubX++)
 					{
-						for (int32 SubX = 0; SubX < SceneProxy->NumSubsections; SubX++)
-						{
-							int32 NeighborSubSectionIndex = SubX + SubY * SceneProxy->NumSubsections;
-							SceneProxy->GetShaderCurrentNeighborLOD(*InView, CurrentLODData.SubSections[NeighborSubSectionIndex].fBatchElementCurrentLOD, SceneProxy->NumSubsections > 1 ? SubX : INDEX_NONE, SceneProxy->NumSubsections > 1 ? SubY : INDEX_NONE, NeighborSubSectionIndex, CurrentNeighborLOD[NeighborSubSectionIndex]);
-							check(CurrentNeighborLOD[NeighborSubSectionIndex].X != -1.0f); // they should all match so only check the 1st one for simplicity
-						}
+						int32 NeighborSubSectionIndex = SubX + SubY * SceneProxy->NumSubsections;
+						SceneProxy->GetShaderCurrentNeighborLOD(*InView, CurrentLODData.SubSections[NeighborSubSectionIndex].fBatchElementCurrentLOD, SceneProxy->NumSubsections > 1 ? SubX : INDEX_NONE, SceneProxy->NumSubsections > 1 ? SubY : INDEX_NONE, NeighborSubSectionIndex, CurrentNeighborLOD[NeighborSubSectionIndex]);
+						check(CurrentNeighborLOD[NeighborSubSectionIndex].X != -1.0f); // they should all match so only check the 1st one for simplicity
 					}
-
-					ShaderBindings.Add(NeighborSectionLodParameter, CurrentNeighborLOD);
 				}
-			}
+
+				ShaderBindings.Add(NeighborSectionLodParameter, CurrentNeighborLOD);
+			}				
 		}
 	}
 
@@ -4243,31 +4071,15 @@ public:
 
 		ShaderBindings.Add(Shader->GetUniformBufferParameter<FLandscapeUniformShaderParameters>(), *BatchElementParams->LandscapeUniformShaderParametersResource);
 
-#if RHI_RAYTRACING
-		if (IsRayTracingEnabled())
+		if (LodValuesParameter.IsBound())
 		{
-			ShaderBindings.Add(Shader->GetUniformBufferParameter<FLandscapeVertexFactoryMVFParameters>(), BatchElementParams->LandscapeVertexFactoryMVFUniformBuffer);
+			ShaderBindings.Add(LodValuesParameter, SceneProxy->GetShaderLODValues(BatchElementParams->CurrentLOD));
 		}
-#endif
 
-		int32 CurrentLOD = BatchElementParams->CurrentLOD;
-#if WITH_EDITOR
-		if (InView && InView->Family->LandscapeLODOverride >= 0)
+		if (LodBiasParameter.IsBound())
 		{
-			CurrentLOD = InView->Family->LandscapeLODOverride;
+			ShaderBindings.Add(LodBiasParameter, FVector4(ForceInitToZero));
 		}
-#endif
-
-		FVector4 LodTessellationParams(ForceInitToZero);
-		FVector4 LodBias(ForceInitToZero);
-		FVector4 ShaderCurrentLOD(CurrentLOD, CurrentLOD, CurrentLOD, CurrentLOD);
-		FVector4 ShaderCurrentNeighborLOD[FLandscapeComponentSceneProxy::NEIGHBOR_COUNT] = { ShaderCurrentLOD, ShaderCurrentLOD, ShaderCurrentLOD, ShaderCurrentLOD };
-
-		ShaderBindings.Add(LodValuesParameter, SceneProxy->GetShaderLODValues(CurrentLOD));
-		ShaderBindings.Add(LodTessellationParameter, LodTessellationParams);
-		ShaderBindings.Add(LodBiasParameter, LodBias);
-		ShaderBindings.Add(SectionLodsParameter, ShaderCurrentLOD);
-		ShaderBindings.Add(NeighborSectionLodParameter, ShaderCurrentNeighborLOD);
 	}
 };
 
@@ -4342,7 +4154,7 @@ void FLandscapeVertexFactory::ModifyCompilationEnvironment(const FVertexFactoryT
 }
 
 
-IMPLEMENT_VERTEX_FACTORY_TYPE_EX(FLandscapeVertexFactory, "/Engine/Private/LandscapeVertexFactory.ush", true, true, true, false, false, true, false);
+IMPLEMENT_VERTEX_FACTORY_TYPE(FLandscapeVertexFactory, "/Engine/Private/LandscapeVertexFactory.ush", true, true, true, false, false);
 
 /**
 * Copy the data from another vertex factory
@@ -4402,7 +4214,7 @@ FVertexFactoryShaderParameters* FLandscapeFixedGridVertexFactory::ConstructShade
 	}
 }
 
-IMPLEMENT_VERTEX_FACTORY_TYPE_EX(FLandscapeFixedGridVertexFactory, "/Engine/Private/LandscapeVertexFactory.ush", true, true, true, false, false, true, false);
+IMPLEMENT_VERTEX_FACTORY_TYPE(FLandscapeFixedGridVertexFactory, "/Engine/Private/LandscapeVertexFactory.ush", true, true, true, false, false);
 
 
 /** ULandscapeMaterialInstanceConstant */
@@ -4418,7 +4230,6 @@ class FLandscapeMaterialResource : public FMaterialResource
 	const bool bDisableTessellation;
 	const bool bMobile;
 	const bool bEditorToolUsage;
-	const bool bUseDiscreteLOD;
 
 public:
 	FLandscapeMaterialResource(ULandscapeMaterialInstanceConstant* Parent)
@@ -4426,7 +4237,6 @@ public:
 		, bDisableTessellation(Parent->bDisableTessellation)
 		, bMobile(Parent->bMobile)
 		, bEditorToolUsage(Parent->bEditorToolUsage)
-		, bUseDiscreteLOD(Parent->bUseDiscreteLOD)
 	{
 	}
 
@@ -4526,35 +4336,20 @@ public:
 				// Todo: only compile LandscapeXYOffsetVertexFactory if we are using it
 				bool bIsRuntimeVirtualTextureShaderType = Algo::Find(GetRuntimeVirtualTextureShaderTypes(), ShaderType->GetFName()) != nullptr;
 
-				if (bIsRuntimeVirtualTextureShaderType)
+				static const FName LandscapeVertexFactory = FName(TEXT("FLandscapeVertexFactory"));
+				static const FName LandscapeXYOffsetVertexFactory = FName(TEXT("FLandscapeXYOffsetVertexFactory"));
+				static const FName LandscapeVertexFactoryMobile = FName(TEXT("FLandscapeVertexFactoryMobile"));
+				if (VertexFactoryType->GetFName() == LandscapeVertexFactory ||
+					VertexFactoryType->GetFName() == LandscapeXYOffsetVertexFactory ||
+					VertexFactoryType->GetFName() == LandscapeVertexFactoryMobile)
 				{
-					static const FName LandscapeFixedGridVertexFactory = FName(TEXT("FLandscapeFixedGridVertexFactory"));
-					if (VertexFactoryType->GetFName() == LandscapeFixedGridVertexFactory)
-					{
-						return FMaterialResource::ShouldCache(Platform, ShaderType, VertexFactoryType);
-					}
+					return !bIsRuntimeVirtualTextureShaderType && FMaterialResource::ShouldCache(Platform, ShaderType, VertexFactoryType);
 				}
-				else
-				{
-					static const FName LandscapeFixedGridVertexFactory = FName(TEXT("FLandscapeFixedGridVertexFactory"));
-					if (VertexFactoryType->GetFName() == LandscapeFixedGridVertexFactory)
-					{
-						return bUseDiscreteLOD && FMaterialResource::ShouldCache(Platform, ShaderType, VertexFactoryType);
-					}
-
-					static const FName LandscapeFixedGridVertexFactoryMobile = FName(TEXT("FLandscapeFixedGridVertexFactoryMobile"));
-					if (VertexFactoryType->GetFName() == LandscapeFixedGridVertexFactoryMobile)
-					{
-						return bUseDiscreteLOD && FMaterialResource::ShouldCache(Platform, ShaderType, VertexFactoryType);
-					}
 
-					static const FName LandscapeVertexFactory = FName(TEXT("FLandscapeVertexFactory"));
-					static const FName LandscapeXYOffsetVertexFactory = FName(TEXT("FLandscapeXYOffsetVertexFactory"));
-					if (VertexFactoryType->GetFName() == LandscapeVertexFactory ||
-						VertexFactoryType->GetFName() == LandscapeXYOffsetVertexFactory)
-					{
-						return !bUseDiscreteLOD && FMaterialResource::ShouldCache(Platform, ShaderType, VertexFactoryType);
-					}
+				static const FName LandscapeFixedGridVertexFactory = FName(TEXT("FLandscapeFixedGridVertexFactory"));
+				if (VertexFactoryType->GetFName() == LandscapeFixedGridVertexFactory)
+				{
+					return bIsRuntimeVirtualTextureShaderType && FMaterialResource::ShouldCache(Platform, ShaderType, VertexFactoryType);
 				}
 			}
 		}
diff --git a/Engine/Source/Runtime/Landscape/Private/LandscapeRenderMobile.cpp b/Engine/Source/Runtime/Landscape/Private/LandscapeRenderMobile.cpp
--- a/Engine/Source/Runtime/Landscape/Private/LandscapeRenderMobile.cpp
+++ b/Engine/Source/Runtime/Landscape/Private/LandscapeRenderMobile.cpp
@@ -34,27 +34,6 @@ void FLandscapeVertexFactoryMobile::InitRHI()
 	InitDeclaration(Elements);
 }
 
-void FLandscapeFixedGridVertexFactoryMobile::InitRHI()
-{
-	// list of declaration items
-	FVertexDeclarationElementList Elements;
-
-	// position decls
-	Elements.Add(AccessStreamComponent(MobileData.PositionComponent, 0));
-
-	if (MobileData.LODHeightsComponent.Num())
-	{
-		const int32 BaseAttribute = 1;
-		for (int32 Index = 0; Index < MobileData.LODHeightsComponent.Num(); Index++)
-		{
-			Elements.Add(AccessStreamComponent(MobileData.LODHeightsComponent[Index], BaseAttribute + Index));
-		}
-	}
-
-	// create the actual device decls
-	InitDeclaration(Elements);
-}
-
 /** Shader parameters for use with FLandscapeVertexFactory */
 class FLandscapeVertexFactoryMobileVertexShaderParameters : public FVertexFactoryShaderParameters
 {
@@ -193,108 +172,6 @@ protected:
 	TShaderUniformBufferParameter<FLandscapeUniformShaderParameters> LandscapeShaderParameters;
 };
 
-
-/** Shader parameters for use with FLandscapeVertexFactory */
-class FLandscapeFixedGridVertexFactoryMobileVertexShaderParameters : public FVertexFactoryShaderParameters
-{
-public:
-	/**
-	* Bind shader constants by name
-	* @param	ParameterMap - mapping of named shader constants to indices
-	*/
-	virtual void Bind(const FShaderParameterMap& ParameterMap) override
-	{
-		LodValuesParameter.Bind(ParameterMap, TEXT("LodValues"));
-		LodTessellationParameter.Bind(ParameterMap, TEXT("LodTessellationParams"));
-		NeighborSectionLodParameter.Bind(ParameterMap, TEXT("NeighborSectionLod"));
-		LodBiasParameter.Bind(ParameterMap, TEXT("LodBias"));
-		SectionLodsParameter.Bind(ParameterMap, TEXT("SectionLods"));
-	}
-
-	/**
-	* Serialize shader params to an archive
-	* @param	Ar - archive to serialize to
-	*/
-	virtual void Serialize(FArchive& Ar) override
-	{
-		Ar << LodValuesParameter;
-		Ar << LodTessellationParameter;
-		Ar << NeighborSectionLodParameter;
-		Ar << LodBiasParameter;
-		Ar << SectionLodsParameter;
-	}
-
-	virtual void GetElementShaderBindings(
-		const class FSceneInterface* Scene,
-		const FSceneView* InView,
-		const class FMeshMaterialShader* Shader,
-		const EVertexInputStreamType InputStreamType,
-		ERHIFeatureLevel::Type FeatureLevel,
-		const FVertexFactory* VertexFactory,
-		const FMeshBatchElement& BatchElement,
-		class FMeshDrawSingleShaderBindings& ShaderBindings,
-		FVertexInputStreamArray& VertexStreams
-	) const override
-	{
-		SCOPE_CYCLE_COUNTER(STAT_LandscapeVFDrawTimeVS);
-
-		const FLandscapeBatchElementParams* BatchElementParams = (const FLandscapeBatchElementParams*)BatchElement.UserData;
-		check(BatchElementParams);
-
-		const FLandscapeComponentSceneProxyMobile* SceneProxy = (const FLandscapeComponentSceneProxyMobile*)BatchElementParams->SceneProxy;
-		ShaderBindings.Add(Shader->GetUniformBufferParameter<FLandscapeUniformShaderParameters>(), *BatchElementParams->LandscapeUniformShaderParametersResource);
-
-		if (LodValuesParameter.IsBound())
-		{
-			FVector4 LodValues(
-				0.0f, // this is the mesh's LOD, ES2 always uses the LOD0 mesh
-				0.0f, // unused
-				(float)SceneProxy->SubsectionSizeQuads,
-				1.f / (float)SceneProxy->SubsectionSizeQuads);
-
-			ShaderBindings.Add(LodValuesParameter, LodValues);
-		}
-
-		if (LodBiasParameter.IsBound())
-		{
-			FVector CameraLocalPos3D = SceneProxy->WorldToLocal.TransformPosition(InView->ViewMatrices.GetViewOrigin());
-
-			FVector4 LodBias(
-				0.0f, // unused
-				0.0f, // unused
-				CameraLocalPos3D.X + SceneProxy->SectionBase.X,
-				CameraLocalPos3D.Y + SceneProxy->SectionBase.Y
-			);
-			ShaderBindings.Add(LodBiasParameter, LodBias);
-		}
-
-		if (LodTessellationParameter.IsBound())
-		{
-			FVector4 LodTessellationParams(ForceInitToZero);
-			ShaderBindings.Add(LodTessellationParameter, LodTessellationParams);
-		}
-
-		if (SectionLodsParameter.IsBound())
-		{
-			FVector4 ShaderCurrentLOD(ForceInitToZero);
-			ShaderBindings.Add(SectionLodsParameter, ShaderCurrentLOD);
-		}
-
-		if (NeighborSectionLodParameter.IsBound())
-		{
-			FVector4 ShaderCurrentNeighborLOD[FLandscapeComponentSceneProxy::NEIGHBOR_COUNT] = { FVector4(ForceInitToZero), FVector4(ForceInitToZero), FVector4(ForceInitToZero), FVector4(ForceInitToZero) };
-			ShaderBindings.Add(NeighborSectionLodParameter, ShaderCurrentNeighborLOD);
-		}
-	}
-protected:
-	FShaderParameter LodValuesParameter;
-	FShaderParameter LodTessellationParameter;
-	FShaderParameter NeighborSectionLodParameter;
-	FShaderParameter LodBiasParameter;
-	FShaderParameter SectionLodsParameter;
-	TShaderUniformBufferParameter<FLandscapeUniformShaderParameters> LandscapeShaderParameters;
-};
-
 /** Shader parameters for use with FLandscapeVertexFactory */
 class FLandscapeVertexFactoryMobilePixelShaderParameters : public FLandscapeVertexFactoryPixelShaderParameters
 {
@@ -366,21 +243,7 @@ FVertexFactoryShaderParameters* FLandscapeVertexFactoryMobile::ConstructShaderPa
 	}
 }
 
-FVertexFactoryShaderParameters* FLandscapeFixedGridVertexFactoryMobile::ConstructShaderParameters(EShaderFrequency ShaderFrequency)
-{
-	switch (ShaderFrequency)
-	{
-	case SF_Vertex:
-		return new FLandscapeFixedGridVertexFactoryMobileVertexShaderParameters();
-	case SF_Pixel:
-		return new FLandscapeVertexFactoryMobilePixelShaderParameters();
-	default:
-		return NULL;
-	}
-}
-
 IMPLEMENT_VERTEX_FACTORY_TYPE(FLandscapeVertexFactoryMobile, "/Engine/Private/LandscapeVertexFactory.ush", true, true, true, false, false);
-IMPLEMENT_VERTEX_FACTORY_TYPE_EX(FLandscapeFixedGridVertexFactoryMobile, "/Engine/Private/LandscapeVertexFactory.ush", true, true, true, false, false, true, false);
 
 /**
 * Initialize the RHI for this rendering resource
@@ -506,12 +369,6 @@ FLandscapeComponentSceneProxyMobile::~FLandscapeComponentSceneProxyMobile()
 		delete VertexFactory;
 		VertexFactory = NULL;
 	}
-
-	if (FixedGridVertexFactory)
-	{
-		delete FixedGridVertexFactory;
-		FixedGridVertexFactory = NULL;
-	}
 }
 
 SIZE_T FLandscapeComponentSceneProxyMobile::GetTypeHash() const
@@ -549,31 +406,16 @@ void FLandscapeComponentSceneProxyMobile::CreateRenderThreadResources()
 	check(MobileRenderData->VertexBuffer);
 	MobileRenderData->VertexBuffer->InitResource();
 
+	FLandscapeVertexFactoryMobile* LandscapeVertexFactory = new FLandscapeVertexFactoryMobile(FeatureLevel);
+	LandscapeVertexFactory->MobileData.PositionComponent = FVertexStreamComponent(MobileRenderData->VertexBuffer, STRUCT_OFFSET(FLandscapeMobileVertex,Position), sizeof(FLandscapeMobileVertex), VET_UByte4N);
+	for( uint32 Index = 0; Index < LANDSCAPE_MAX_ES_LOD_COMP; ++Index )
 	{
-		FLandscapeVertexFactoryMobile* LandscapeVertexFactory = new FLandscapeVertexFactoryMobile(FeatureLevel);
-		LandscapeVertexFactory->MobileData.PositionComponent = FVertexStreamComponent(MobileRenderData->VertexBuffer, STRUCT_OFFSET(FLandscapeMobileVertex, Position), sizeof(FLandscapeMobileVertex), VET_UByte4N);
-		for (uint32 Index = 0; Index < LANDSCAPE_MAX_ES_LOD_COMP; ++Index)
-		{
-			LandscapeVertexFactory->MobileData.LODHeightsComponent.Add
-			(FVertexStreamComponent(MobileRenderData->VertexBuffer, STRUCT_OFFSET(FLandscapeMobileVertex, LODHeights) + sizeof(uint8) * 4 * Index, sizeof(FLandscapeMobileVertex), VET_UByte4N));
-		}
-
-		LandscapeVertexFactory->InitResource();
-		VertexFactory = LandscapeVertexFactory;
+		LandscapeVertexFactory->MobileData.LODHeightsComponent.Add
+			(FVertexStreamComponent(MobileRenderData->VertexBuffer, STRUCT_OFFSET(FLandscapeMobileVertex,LODHeights) + sizeof(uint8) * 4 * Index, sizeof(FLandscapeMobileVertex), VET_UByte4N));
 	}
 
-	{
-		FLandscapeFixedGridVertexFactoryMobile* LandscapeVertexFactory = new FLandscapeFixedGridVertexFactoryMobile(FeatureLevel);
-		LandscapeVertexFactory->MobileData.PositionComponent = FVertexStreamComponent(MobileRenderData->VertexBuffer, STRUCT_OFFSET(FLandscapeMobileVertex, Position), sizeof(FLandscapeMobileVertex), VET_UByte4N);
-		for (uint32 Index = 0; Index < LANDSCAPE_MAX_ES_LOD_COMP; ++Index)
-		{
-			LandscapeVertexFactory->MobileData.LODHeightsComponent.Add
-			(FVertexStreamComponent(MobileRenderData->VertexBuffer, STRUCT_OFFSET(FLandscapeMobileVertex, LODHeights) + sizeof(uint8) * 4 * Index, sizeof(FLandscapeMobileVertex), VET_UByte4N));
-		}
-
-		LandscapeVertexFactory->InitResource();
-		FixedGridVertexFactory = LandscapeVertexFactory;
-	}
+	LandscapeVertexFactory->InitResource();
+	VertexFactory = LandscapeVertexFactory;
 
 	// Assign LandscapeUniformShaderParameters
 	LandscapeUniformShaderParameters.InitResource();
diff --git a/Engine/Source/Runtime/Landscape/Public/LandscapeRender.h b/Engine/Source/Runtime/Landscape/Public/LandscapeRender.h
--- a/Engine/Source/Runtime/Landscape/Public/LandscapeRender.h
+++ b/Engine/Source/Runtime/Landscape/Public/LandscapeRender.h
@@ -75,28 +75,19 @@ namespace ELandscapeEditRenderMode
 
 LANDSCAPE_API extern bool GLandscapeEditModeActive;
 LANDSCAPE_API extern int32 GLandscapeEditRenderMode;
-
 LANDSCAPE_API extern UMaterialInterface* GLayerDebugColorMaterial;
 LANDSCAPE_API extern UMaterialInterface* GSelectionColorMaterial;
 LANDSCAPE_API extern UMaterialInterface* GSelectionRegionMaterial;
 LANDSCAPE_API extern UMaterialInterface* GMaskRegionMaterial;
 LANDSCAPE_API extern UMaterialInterface* GColorMaskRegionMaterial;
-LANDSCAPE_API extern UMaterialInterface* GLandscapeLayerUsageMaterial;
-
-LANDSCAPE_API extern UMaterialInterface* GLayerDebugColorMaterial_DiscreteLOD;
-LANDSCAPE_API extern UMaterialInterface* GSelectionColorMaterial_DiscreteLOD;
-LANDSCAPE_API extern UMaterialInterface* GSelectionRegionMaterial_DiscreteLOD;
-LANDSCAPE_API extern UMaterialInterface* GMaskRegionMaterial_DiscreteLOD;
-LANDSCAPE_API extern UMaterialInterface* GColorMaskRegionMaterial_DiscreteLOD;
-LANDSCAPE_API extern UMaterialInterface* GLandscapeLayerUsageMaterial_DiscreteLOD;
 LANDSCAPE_API extern UTexture2D* GLandscapeBlackTexture;
+LANDSCAPE_API extern UMaterialInterface* GLandscapeLayerUsageMaterial;
 #endif
 
 
 /** The uniform shader parameters for a landscape draw call. */
 BEGIN_GLOBAL_SHADER_PARAMETER_STRUCT(FLandscapeUniformShaderParameters, LANDSCAPE_API)
 /** vertex shader parameters */
-SHADER_PARAMETER(int32, LastLOD)
 SHADER_PARAMETER(FVector4, HeightmapUVScaleBias)
 SHADER_PARAMETER(FVector4, WeightmapUVScaleBias)
 SHADER_PARAMETER(FVector4, LandscapeLightmapScaleBias)
@@ -582,7 +573,6 @@ public:
 
 protected:
 	int8						MaxLOD;		// Maximum LOD level, user override possible
-	bool						bUseDiscreteLOD : 1;
 	bool						UseTessellationComponentScreenSizeFalloff:1;	// Tell if we should apply a Tessellation falloff
 	bool						bRequiresAdjacencyInformation:1;
 	int8						NumWeightmapLayerAllocations;
@@ -728,7 +718,7 @@ protected:
 
 	bool GetMeshElement(bool UseSeperateBatchForShadow, bool ShadowOnly, bool HasTessellation, int8 InLODIndex, UMaterialInterface* InMaterialInterface, FMeshBatch& OutMeshBatch, TArray<FLandscapeBatchElementParams>& OutStaticBatchParamArray) const;
 	bool GetMeshElementForVirtualTexture(int32 InLodIndex, ERuntimeVirtualTextureMaterialType MaterialType, UMaterialInterface* InMaterialInterface, FMeshBatch& OutMeshBatch, TArray<FLandscapeBatchElementParams>& OutStaticBatchParamArray) const;
-	void BuildDynamicMeshElement(const FSceneView* View, const FViewCustomDataLOD* InPrimitiveCustomData, bool InToolMesh, bool InHasTessellation, bool InDisableTessellation, FMeshBatch& OutMeshBatch, TArray<FLandscapeBatchElementParams, SceneRenderingAllocator>& OutStaticBatchParamArray) const;
+	void BuildDynamicMeshElement(const FViewCustomDataLOD* InPrimitiveCustomData, bool InToolMesh, bool InHasTessellation, bool InDisableTessellation, FMeshBatch& OutMeshBatch, TArray<FLandscapeBatchElementParams, SceneRenderingAllocator>& OutStaticBatchParamArray) const;
 
 	float GetComponentScreenSize(const class FSceneView* View, const FVector& Origin,  float MaxExtend, float ElementRadius) const;
 
@@ -763,7 +753,6 @@ public:
 	friend class FLandscapeVertexFactoryMobileVertexShaderParameters;
 	friend class FLandscapeVertexFactoryMobilePixelShaderParameters;
 	friend class FLandscapeFixedGridVertexFactoryVertexShaderParameters;
-	friend class FLandscapeFixedGridVertexFactoryMobileVertexShaderParameters;
 
 	// FLandscapeComponentSceneProxy interface.
 	uint64 GetStaticBatchElementVisibility(const FSceneView& InView, const FMeshBatch* InBatch, const void* InViewCustomData) const;
diff --git a/Engine/Source/Runtime/Landscape/Public/LandscapeRenderMobile.h b/Engine/Source/Runtime/Landscape/Public/LandscapeRenderMobile.h
--- a/Engine/Source/Runtime/Landscape/Public/LandscapeRenderMobile.h
+++ b/Engine/Source/Runtime/Landscape/Public/LandscapeRenderMobile.h
@@ -83,68 +83,6 @@ private:
 	friend class FLandscapeComponentSceneProxyMobile;
 };
 
-/** vertex factory for VTF-heightmap terrain  */
-class FLandscapeFixedGridVertexFactoryMobile : public FLandscapeVertexFactory
-{
-	DECLARE_VERTEX_FACTORY_TYPE(FLandscapeFixedGridVertexFactoryMobile);
-
-	typedef FLandscapeVertexFactory Super;
-public:
-
-	struct FDataType : FLandscapeVertexFactory::FDataType
-	{
-		/** stream which has heights of each LOD levels */
-		TArray<FVertexStreamComponent, TFixedAllocator<LANDSCAPE_MAX_ES_LOD_COMP> > LODHeightsComponent;
-	};
-
-	FLandscapeFixedGridVertexFactoryMobile(ERHIFeatureLevel::Type InFeatureLevel)
-		: FLandscapeVertexFactory(InFeatureLevel)
-	{
-	}
-
-	virtual ~FLandscapeFixedGridVertexFactoryMobile()
-	{
-		ReleaseResource();
-	}
-
-	static FVertexFactoryShaderParameters* ConstructShaderParameters(EShaderFrequency ShaderFrequency);
-
-	/**
-	* Should we cache the material's shadertype on this platform with this vertex factory?
-	*/
-	static bool ShouldCompilePermutation(EShaderPlatform Platform, const class FMaterial* Material, const class FShaderType* ShaderType)
-	{
-		auto FeatureLevel = GetMaxSupportedFeatureLevel(Platform);
-		return (FeatureLevel == ERHIFeatureLevel::ES2 || FeatureLevel == ERHIFeatureLevel::ES3_1) &&
-			(Material->IsUsedWithLandscape() || Material->IsSpecialEngineMaterial());
-	}
-
-	static void ModifyCompilationEnvironment(const FVertexFactoryType* Type, EShaderPlatform Platform, const FMaterial* Material, FShaderCompilerEnvironment& OutEnvironment)
-	{
-		FVertexFactory::ModifyCompilationEnvironment(Type, Platform, Material, OutEnvironment);
-		OutEnvironment.SetDefine(TEXT("NUM_VF_PACKED_INTERPOLANTS"), TEXT("1"));
-		OutEnvironment.SetDefine(TEXT("FIXED_GRID"), TEXT("1"));
-	}
-
-	// FRenderResource interface.
-	virtual void InitRHI() override;
-
-	/**
-	 * An implementation of the interface used by TSynchronizedResource to update the resource with new data from the game thread.
-	 */
-	void SetData(const FDataType& InData)
-	{
-		MobileData = InData;
-		UpdateRHI();
-	}
-
-private:
-	/** stream component data bound to this vertex factory */
-	FDataType MobileData;
-
-	friend class FLandscapeComponentSceneProxyMobile;
-};
-
 //
 // FLandscapeVertexBuffer
 //
