Landscape Splines: Add Temp Console Command to Fix Spline Layer Width Data (cmd: Landscape.FixSplines)

diff --git a/Engine/Source/Runtime/Landscape/Private/LandscapeSplines.cpp b/Engine/Source/Runtime/Landscape/Private/LandscapeSplines.cpp
--- a/Engine/Source/Runtime/Landscape/Private/LandscapeSplines.cpp
+++ b/Engine/Source/Runtime/Landscape/Private/LandscapeSplines.cpp
@@ -39,6 +39,7 @@
 #include "Misc/UObjectToken.h"
 #include "Landscape.h"
 #include "LandscapeLayerInfoObject.h"
+#include "LandscapeInfoMap.h"
 #endif
 
 IMPLEMENT_HIT_PROXY(HLandscapeSplineProxy, HHitProxy);
@@ -60,6 +61,50 @@ static FAutoConsoleVariableRef CVarSplinesAlwaysUseBlockAll(
 
 /** Represents a ULandscapeSplinesComponent to the scene manager. */
 #if WITH_EDITOR
+struct FLandscapeFixSplines
+{
+	FLandscapeFixSplines()
+		: FixSplinesConsoleCommand(
+			TEXT("Landscape.FixSplines"),
+			TEXT("One off fix for bad layer width"),
+			FConsoleCommandDelegate::CreateRaw(this, &FLandscapeFixSplines::FixSplines))
+	{
+	}
+
+	FAutoConsoleCommand FixSplinesConsoleCommand;
+
+	void FixSplines()
+	{
+		if (!GWorld || GWorld->IsGameWorld())
+		{
+			return;
+		}
+
+		auto& LandscapeInfoMap = ULandscapeInfoMap::GetLandscapeInfoMap(GWorld);
+		for (TPair<FGuid, ULandscapeInfo*>& Pair : LandscapeInfoMap.Map)
+		{
+			if (Pair.Value)
+			{
+				Pair.Value->ForAllLandscapeProxies([](ALandscapeProxy* Proxy)
+				{
+					if (Proxy && Proxy->SplineComponent)
+					{
+						Proxy->SplineComponent->RebuildAllSplines();
+					}
+				});
+
+				if (Pair.Value->LandscapeActor && Pair.Value->LandscapeActor->HasLayersContent())
+				{
+					Pair.Value->LandscapeActor->RequestSplineLayerUpdate();
+				}
+			}
+		}
+
+	}
+};
+
+FLandscapeFixSplines GLandscapeFixSplines;
+
 class FLandscapeSplinesSceneProxy final : public FPrimitiveSceneProxy
 {
 private:
