Fix Windows ray tracing compilation failure

diff --git a/Engine/Source/Runtime/Landscape/Private/LandscapeRender.cpp b/Engine/Source/Runtime/Landscape/Private/LandscapeRender.cpp
--- a/Engine/Source/Runtime/Landscape/Private/LandscapeRender.cpp
+++ b/Engine/Source/Runtime/Landscape/Private/LandscapeRender.cpp
@@ -3104,26 +3104,13 @@ void FLandscapeComponentSceneProxy::GetDynamicRayTracingInstances(FRayTracingMat
 
 	FViewCustomDataLOD* InPrimitiveCustomData;
 
-	if (!bUseDiscreteLOD)
-	{
-		float MeshScreenSizeSquared = 0;
-		int32 ForcedLODLevel = (Context.ReferenceView->Family->EngineShowFlags.LOD) ? GetCVarForceLOD() : 0;
-
-		FLODMask LODToRender = GetCustomLOD(*Context.ReferenceView, Context.ReferenceView->LODDistanceFactor, ForcedLODLevel, MeshScreenSizeSquared);
-		InPrimitiveCustomData = (FViewCustomDataLOD*)InitViewCustomData(*Context.ReferenceView, Context.ReferenceView->LODDistanceFactor, PrimitiveCustomDataMemStack, false, false, &LODToRender, MeshScreenSizeSquared);
-		PostInitViewCustomData(*Context.ReferenceView, InPrimitiveCustomData);
-	}
-	else
-	{
-		float MeshScreenSizeSquared = ComputeBoundsScreenRadiusSquared(GetBounds().Origin, GetBounds().SphereRadius, *Context.ReferenceView);
-		int32 ForcedLODLevel = (Context.ReferenceView->Family->EngineShowFlags.LOD) ? GetCVarForceLOD() : 0;
-
-		FLODMask LODToRender;
-		LODToRender.SetLOD(GetLODFromScreenSize(MeshScreenSizeSquared, Context.ReferenceView->LODDistanceFactor));
-		InPrimitiveCustomData = (FViewCustomDataLOD*)InitViewCustomData(*Context.ReferenceView, Context.ReferenceView->LODDistanceFactor, PrimitiveCustomDataMemStack, false, false, &LODToRender, MeshScreenSizeSquared);
-		PostInitViewCustomData(*Context.ReferenceView, InPrimitiveCustomData);
-	}
+	float MeshScreenSizeSquared = 0;
+	int32 ForcedLODLevel = (Context.ReferenceView->Family->EngineShowFlags.LOD) ? GetCVarForceLOD() : 0;
 
+	FLODMask LODToRender = GetCustomLOD(*Context.ReferenceView, Context.ReferenceView->LODDistanceFactor, ForcedLODLevel, MeshScreenSizeSquared);
+	InPrimitiveCustomData = (FViewCustomDataLOD*)InitViewCustomData(*Context.ReferenceView, Context.ReferenceView->LODDistanceFactor, PrimitiveCustomDataMemStack, false, false, &LODToRender, MeshScreenSizeSquared);
+	PostInitViewCustomData(*Context.ReferenceView, InPrimitiveCustomData);
+	
 	FLandscapeElementParamArray& ParameterArray = Context.RayTracingMeshResourceCollector.AllocateOneFrameResource<FLandscapeElementParamArray>();
 	ParameterArray.ElementParams.AddDefaulted(NumSubsections * NumSubsections);
 
@@ -3143,7 +3130,7 @@ void FLandscapeComponentSceneProxy::GetDynamicRayTracingInstances(FRayTracingMat
 	}
 
 	FMeshBatch BaseMeshBatch;
-	BaseMeshBatch.VertexFactory = bUseDiscreteLOD ? FixedGridVertexFactory : VertexFactory;
+	BaseMeshBatch.VertexFactory = VertexFactory;
 	BaseMeshBatch.MaterialRenderProxy = SelectedMaterial->GetRenderProxy();
 	BaseMeshBatch.LCI = ComponentLightInfo.Get();
 	BaseMeshBatch.CastShadow = true;
