Landscape Layers: Debugging Features
- CVar: Landscape.OutputDiffBitmap (=1 Heightmap, =2 Weightmap, =3 Both)
- CVar: Landscape.ShowDirty new view mode that shows dirtied data since last clear
- Console Command: Landscape.ClearDirty

#rnx

diff --git a/Engine/Source/Editor/LandscapeEditor/Private/LandscapeEdMode.cpp b/Engine/Source/Editor/LandscapeEditor/Private/LandscapeEdMode.cpp
--- a/Engine/Source/Editor/LandscapeEditor/Private/LandscapeEdMode.cpp
+++ b/Engine/Source/Editor/LandscapeEditor/Private/LandscapeEdMode.cpp
@@ -317,6 +317,7 @@ FEdModeLandscape::FEdModeLandscape()
 	GSelectionRegionMaterial = LandscapeTool::CreateMaterialInstance(LoadObject<UMaterialInstanceConstant>(nullptr, TEXT("/Engine/EditorLandscapeResources/SelectBrushMaterial_SelectedRegion.SelectBrushMaterial_SelectedRegion")));
 	GMaskRegionMaterial      = LandscapeTool::CreateMaterialInstance(LoadObject<UMaterialInstanceConstant>(nullptr, TEXT("/Engine/EditorLandscapeResources/MaskBrushMaterial_MaskedRegion.MaskBrushMaterial_MaskedRegion")));
 	GColorMaskRegionMaterial = LandscapeTool::CreateMaterialInstance(LoadObject<UMaterialInstanceConstant>(nullptr, TEXT("/Engine/EditorLandscapeResources/ColorMaskBrushMaterial_MaskedRegion.ColorMaskBrushMaterial_MaskedRegion")));
+	GLandscapeDirtyMaterial		 = LandscapeTool::CreateMaterialInstance(LoadObject<UMaterial>(nullptr, TEXT("/Engine/EditorLandscapeResources/LandscapeDirtyMaterial.LandscapeDirtyMaterial")));
 	GLandscapeBlackTexture   = LoadObject<UTexture2D>(nullptr, TEXT("/Engine/EngineResources/Black.Black"));
 	GLandscapeLayerUsageMaterial = LandscapeTool::CreateMaterialInstance(LoadObject<UMaterial>(nullptr, TEXT("/Engine/EditorLandscapeResources/LandscapeLayerUsageMaterial.LandscapeLayerUsageMaterial")));
 	
@@ -389,6 +390,7 @@ FEdModeLandscape::~FEdModeLandscape()
 	GColorMaskRegionMaterial = NULL;
 	GLandscapeBlackTexture = NULL;
 	GLandscapeLayerUsageMaterial = NULL;
+	GLandscapeDirtyMaterial = NULL;
 
 	InteractorPainting = nullptr;
 }
@@ -408,6 +410,7 @@ void FEdModeLandscape::AddReferencedObjects(FReferenceCollector& Collector)
 	Collector.AddReferencedObject(GColorMaskRegionMaterial);
 	Collector.AddReferencedObject(GLandscapeBlackTexture);
 	Collector.AddReferencedObject(GLandscapeLayerUsageMaterial);
+	Collector.AddReferencedObject(GLandscapeDirtyMaterial);
 }
 
 void FEdModeLandscape::UpdateToolModes()
diff --git a/Engine/Source/Runtime/Landscape/Classes/Landscape.h b/Engine/Source/Runtime/Landscape/Classes/Landscape.h
--- a/Engine/Source/Runtime/Landscape/Classes/Landscape.h
+++ b/Engine/Source/Runtime/Landscape/Classes/Landscape.h
@@ -329,6 +329,7 @@ public:
 	LANDSCAPE_API void OnPreSave();
 
 	void ReleaseLayersRenderingResource();
+	void ClearDirtyData(ULandscapeComponent* InLandscapeComponent);
 	
 	LANDSCAPE_API void ToggleCanHaveLayersContent();
 	LANDSCAPE_API void ForceUpdateLayersContent(bool bIntermediateRender = false);
@@ -346,8 +347,8 @@ private:
 	bool UpdateCollisionAndClients(const TArray<ULandscapeComponent*>& InLandscapeComponents, const int32 InContentUpdateModes);
 	void ResolveLayersHeightmapTexture(const TArray<ULandscapeComponent*>& InLandscapeComponents);
 	void ResolveLayersWeightmapTexture(const TArray<ULandscapeComponent*>& InLandscapeComponents);
-	bool ResolveLayersTexture(class FLandscapeLayersTexture2DCPUReadBackResource* InCPUReadBackTexture, UTexture2D* InOutputTexture);
-
+	bool ResolveLayersTexture(class FLandscapeLayersTexture2DCPUReadBackResource* InCPUReadBackTexture, UTexture2D* InOutputTexture, bool bHeightmap, ULandscapeComponent* InComponent);
+		
 	bool PrepareLayersBrushTextureResources(bool bInWaitForStreaming, bool bHeightmap) const;
 	bool PrepareLayersHeightmapTextureResources(bool bInWaitForStreaming) const;
 	bool PrepareLayersWeightmapTextureResources(bool bInWaitForStreaming) const;
@@ -403,6 +404,8 @@ private:
 	void PrintLayersDebugTextureResource(const FString& InContext, FTextureResource* InTextureResource, uint8 InMipRender = 0, bool InOutputHeight = true, bool InOutputNormals = false) const;
 	void PrintLayersDebugHeightData(const FString& InContext, const TArray<FColor>& InHeightmapData, const FIntPoint& InDataSize, uint8 InMipRender, bool InOutputNormals = false) const;
 	void PrintLayersDebugWeightData(const FString& InContext, const TArray<FColor>& InWeightmapData, const FIntPoint& InDataSize, uint8 InMipRender) const;
+
+	void UpdateDirtyData(ULandscapeComponent* InLandscapeComponent, const FColor* InOldData, const FColor* InNewData, int32 InSize);
 #endif
 
 public:
diff --git a/Engine/Source/Runtime/Landscape/Classes/LandscapeComponent.h b/Engine/Source/Runtime/Landscape/Classes/LandscapeComponent.h
--- a/Engine/Source/Runtime/Landscape/Classes/LandscapeComponent.h
+++ b/Engine/Source/Runtime/Landscape/Classes/LandscapeComponent.h
@@ -62,7 +62,8 @@ public:
 		DebugChannelG(INDEX_NONE),
 		DebugChannelB(INDEX_NONE),
 		DataTexture(NULL),
-		LayerContributionTexture(NULL)
+		LayerContributionTexture(NULL),
+		DirtyTexture(NULL)
 	{}
 
 	// Material used to render the tool.
@@ -92,6 +93,9 @@ public:
 	UPROPERTY(NonTransactional)
 	UTexture2D* LayerContributionTexture; // Data texture used to represent layer contribution
 
+	UPROPERTY(NonTransactional)
+	UTexture2D* DirtyTexture; // Data texture used to represent layer blend dirtied area
+
 #if WITH_EDITOR
 	void UpdateDebugColorMaterial(const ULandscapeComponent* const Component);
 	void UpdateSelectionMaterial(int32 InSelectedType, const ULandscapeComponent* const Component);
diff --git a/Engine/Source/Runtime/Landscape/Classes/LandscapeInfo.h b/Engine/Source/Runtime/Landscape/Classes/LandscapeInfo.h
--- a/Engine/Source/Runtime/Landscape/Classes/LandscapeInfo.h
+++ b/Engine/Source/Runtime/Landscape/Classes/LandscapeInfo.h
@@ -238,6 +238,9 @@ public:
 
 	/** Will tell if the landscape actor can have some content related to the layer system */
 	LANDSCAPE_API bool CanHaveLayersContent() const;
+
+	/** Will clear all component dirty data */
+	LANDSCAPE_API void ClearDirtyData();
 #endif
 	/** Associates passed actor with this info object
  *  @param	Proxy		Landscape actor to register
diff --git a/Engine/Source/Runtime/Landscape/Private/LandscapeEdit.cpp b/Engine/Source/Runtime/Landscape/Private/LandscapeEdit.cpp
--- a/Engine/Source/Runtime/Landscape/Private/LandscapeEdit.cpp
+++ b/Engine/Source/Runtime/Landscape/Private/LandscapeEdit.cpp
@@ -3939,6 +3939,17 @@ bool ULandscapeInfo::CanHaveLayersContent() const
 	return false;
 }
 
+void ULandscapeInfo::ClearDirtyData()
+{
+	if (ALandscape* Landscape = LandscapeActor.Get())
+	{
+		ForAllLandscapeComponents([=](ULandscapeComponent* InLandscapeComponent)
+		{
+			Landscape->ClearDirtyData(InLandscapeComponent);
+		});
+	}
+}
+
 namespace
 {
 	inline float AdjustStaticLightingResolution(float StaticLightingResolution, int32 NumSubsections, int32 SubsectionSizeQuads, int32 ComponentSizeQuads)
diff --git a/Engine/Source/Runtime/Landscape/Private/LandscapeEditInterface.cpp b/Engine/Source/Runtime/Landscape/Private/LandscapeEditInterface.cpp
--- a/Engine/Source/Runtime/Landscape/Private/LandscapeEditInterface.cpp
+++ b/Engine/Source/Runtime/Landscape/Private/LandscapeEditInterface.cpp
@@ -4560,6 +4560,44 @@ void FLandscapeEditDataInterface::SetLayerContributionData(int32 X1, int32 Y1, i
 	SetEditToolTextureData(X1, Y1, X2, Y2, Data, Stride, ReturnComponentTexture);
 }
 
+template<typename TStoreData>
+void FLandscapeEditDataInterface::GetDirtyDataTempl(const int32 X1, const int32 Y1, const int32 X2, const int32 Y2, TStoreData& StoreData)
+{
+	auto ReturnComponentTexture = [&](const ULandscapeComponent* Component) -> UTexture2D*
+	{
+		check(Component);
+		return Component->EditToolRenderData.DirtyTexture;
+	};
+	GetEditToolTextureData(X1, Y1, X2, Y2, StoreData, ReturnComponentTexture);
+}
+
+void FLandscapeEditDataInterface::GetDirtyData(const int32 X1, const int32 Y1, const int32 X2, const int32 Y2, TMap<FIntPoint, uint8>& SparseData)
+{
+	TSparseStoreData<uint8> SparseStoreData(SparseData);
+	GetDirtyDataTempl(X1, Y1, X2, Y2, SparseStoreData);
+}
+
+void FLandscapeEditDataInterface::GetDirtyData(const int32 X1, const int32 Y1, const int32 X2, const int32 Y2, uint8* Data, int32 Stride)
+{
+	if (Stride == 0)
+	{
+		Stride = (1 + X2 - X1);
+	}
+	TArrayStoreData<uint8> ArrayStoreData(X1, Y1, Data, Stride);
+	GetDirtyDataTempl(X1, Y1, X2, Y2, ArrayStoreData);
+}
+
+
+void FLandscapeEditDataInterface::SetDirtyData(int32 X1, int32 Y1, int32 X2, int32 Y2, const uint8* Data, int32 Stride)
+{
+	auto ReturnComponentTexture = [](ULandscapeComponent* Component) -> UTexture2D*&
+	{
+		check(Component);
+		return Component->EditToolRenderData.DirtyTexture;
+	};
+	SetEditToolTextureData(X1, Y1, X2, Y2, Data, Stride, ReturnComponentTexture);
+}
+
 template<typename TStoreData>
 void FLandscapeEditDataInterface::GetEditToolTextureData(const int32 X1, const int32 Y1, const int32 X2, const int32 Y2, TStoreData& StoreData, TFunctionRef<UTexture2D*(ULandscapeComponent*)> GetComponentTexture)
 {
diff --git a/Engine/Source/Runtime/Landscape/Private/LandscapeEditLayers.cpp b/Engine/Source/Runtime/Landscape/Private/LandscapeEditLayers.cpp
--- a/Engine/Source/Runtime/Landscape/Private/LandscapeEditLayers.cpp
+++ b/Engine/Source/Runtime/Landscape/Private/LandscapeEditLayers.cpp
@@ -36,6 +36,7 @@ LandscapeEditLayers.cpp: Landscape editing layers mode
 #include "GameFramework/WorldSettings.h"
 #include "UObject/UObjectThreadContext.h"
 #include "LandscapeSplinesComponent.h"
+#include "Misc/FileHelper.h"
 #endif
 
 #define LOCTEXT_NAMESPACE "Landscape"
@@ -78,6 +79,51 @@ static TAutoConsoleVariable<int32> CVarLandscapeLayerBrushOptim(
 	TEXT("landscape.BrushOptim"),
 	0,
 	TEXT("This will enable landscape layers optim."));
+
+static TAutoConsoleVariable<int32> CVarLandscapeOutputDiffBitmap(
+	TEXT("landscape.OutputDiffBitmap"),
+	0,
+	TEXT("This will save images for readback textures that have changed in the last layer blend phase. (= 1 Heightmap Diff, = 2 Weightmap Diff, = 3 All Diffs"));
+
+TAutoConsoleVariable<int32> CVarLandscapeShowDirty(
+	TEXT("landscape.ShowDirty"),
+	0,
+	TEXT("This will highlight the data that as changed during the layer blend phase."));
+
+struct FLandscapeDirty
+{
+	FLandscapeDirty()
+		: ClearDiffConsoleCommand(
+			TEXT("Landscape.ClearDirty"),
+			TEXT("Clears all Landscape Dirty Debug Data"),
+			FConsoleCommandDelegate::CreateRaw(this, &FLandscapeDirty::ClearDirty))
+	{
+	}
+
+private:
+	FAutoConsoleCommand ClearDiffConsoleCommand;
+
+	void ClearDirty()
+	{
+		if(!GWorld || GWorld->IsGameWorld())
+		{ 
+			return;
+		}
+
+		auto& LandscapeInfoMap = ULandscapeInfoMap::GetLandscapeInfoMap(GWorld);
+		for (TPair<FGuid, ULandscapeInfo*>& Pair : LandscapeInfoMap.Map)
+		{
+			if (Pair.Value)
+			{
+				Pair.Value->ClearDirtyData();
+			}
+		}
+
+		UE_LOG(LogLandscape, Display, TEXT("Landscape.Dirty: Cleared"));
+	}
+};
+
+FLandscapeDirty GLandscapeDebugDirty;
 #endif
 
 DECLARE_GPU_STAT_NAMED(LandscapeLayersRender, TEXT("Landscape Layer System Render"));
@@ -3067,7 +3113,7 @@ void ALandscape::ResolveLayersHeightmapTexture(const TArray<ULandscapeComponent*
 				continue;
 			}
 						
-			bChanged = ResolveLayersTexture(*CPUReadback, ComponentHeightmap);
+			bChanged = ResolveLayersTexture(*CPUReadback, ComponentHeightmap, true, Component);
 			ProcessedTexture.Add(ComponentHeightmap, bChanged);
 			bValue = &bChanged;
 		}
@@ -3083,7 +3129,58 @@ void ALandscape::ResolveLayersHeightmapTexture(const TArray<ULandscapeComponent*
 
 }
 
-bool ALandscape::ResolveLayersTexture(FLandscapeLayersTexture2DCPUReadBackResource* InCPUReadBackTexture, UTexture2D* InOutputTexture)
+void ALandscape::ClearDirtyData(ULandscapeComponent* InLandscapeComponent)
+{
+	if (InLandscapeComponent->EditToolRenderData.DirtyTexture == nullptr)
+	{
+		return;
+	}
+
+	UpdateDirtyData(InLandscapeComponent, nullptr, nullptr, 0);
+}
+
+void ALandscape::UpdateDirtyData(ULandscapeComponent* InLandscapeComponent, const FColor* InOldData, const FColor* InNewData, int32 InSize)
+{
+	if (!CVarLandscapeShowDirty.GetValueOnAnyThread())
+	{
+		return;
+	}
+
+	FLandscapeEditDataInterface LandscapeEdit(GetLandscapeInfo());
+	const int32 X1 = InLandscapeComponent->GetSectionBase().X;
+	const int32 X2 = X1 + ComponentSizeQuads;
+	const int32 Y1 = InLandscapeComponent->GetSectionBase().Y;
+	const int32 Y2 = Y1 + ComponentSizeQuads;
+	const int32 ComponentWidth = (SubsectionSizeQuads + 1)*NumSubsections;
+	const int32 DirtyDataSize = ComponentWidth * ComponentWidth;
+	uint8* DirtyData = new uint8[DirtyDataSize];
+	
+	if (InOldData && InNewData)
+	{
+		check(InSize == DirtyDataSize);
+		LandscapeEdit.GetDirtyData(X1, Y1, X2, Y2, DirtyData, 0);
+
+		for (int32 i = 0; i < DirtyDataSize; ++i)
+		{
+			if (*InOldData != *InNewData)
+			{
+				DirtyData[i] = 255;
+			}
+			InOldData++;
+			InNewData++;
+		}
+	}
+	else
+	{
+		FMemory::Memzero(DirtyData, DirtyDataSize);
+	}
+
+	LandscapeEdit.SetDirtyData(X1, Y1, X2, Y2, DirtyData, 0);
+
+	delete[] DirtyData;
+}
+
+bool ALandscape::ResolveLayersTexture(FLandscapeLayersTexture2DCPUReadBackResource* InCPUReadBackTexture, UTexture2D* InOutputTexture, bool bHeightmap, ULandscapeComponent* InComponent)
 {
 	TRACE_CPUPROFILER_EVENT_SCOPE("ResolveLayersTexture");
 	SCOPE_CYCLE_COUNTER(STAT_LandscapeLayersResolveTexture);
@@ -3118,6 +3215,8 @@ bool ALandscape::ResolveLayersTexture(FLandscapeLayersTexture2DCPUReadBackResour
 	FlushRenderingCommands();
 
 	bool bChanged = false;
+	const bool HeightmapDiff = (CVarLandscapeOutputDiffBitmap.GetValueOnAnyThread() & 1) != 0;
+	const bool WeightmapDiff = (CVarLandscapeOutputDiffBitmap.GetValueOnAnyThread() & 2) != 0;
     const bool bUpdateHash = !bIntermediateRender;
 	for (int8 MipIndex = 0; MipIndex < OutMipsData.Num(); ++MipIndex)
 	{
@@ -3129,7 +3228,21 @@ bool ALandscape::ResolveLayersTexture(FLandscapeLayersTexture2DCPUReadBackResour
 				bChanged = InCPUReadBackTexture->UpdateHashFromTextureSource((uint8*)OutMipsData[MipIndex].GetData());
 			}
 
+			if (MipIndex == 0 && bChanged)
+			{
+				UpdateDirtyData(InComponent, (FColor*)TextureData, OutMipsData[MipIndex].GetData(), InOutputTexture->Source.GetSizeX()*InOutputTexture->Source.GetSizeY());
+
+				if ((HeightmapDiff && bHeightmap) || (WeightmapDiff && !bHeightmap))
+				{
+					FString LevelName = FPackageName::GetShortName(InComponent->GetOutermost());
+					FString FilePattern = FString::Format(TEXT("LandscapeLayers/{0}-{1}-{2}"), { LevelName, InComponent->GetName(), bHeightmap ? TEXT("HM") : TEXT("WM") });
+					FFileHelper::CreateBitmap(*(FilePattern), InOutputTexture->Source.GetSizeX(), InOutputTexture->Source.GetSizeY(), (FColor*)TextureData, nullptr, &IFileManager::Get(), nullptr, true);
+					FFileHelper::CreateBitmap(*(FilePattern), InOutputTexture->Source.GetSizeX(), InOutputTexture->Source.GetSizeY(), OutMipsData[MipIndex].GetData(), nullptr, &IFileManager::Get(), nullptr, true);
+				}
+			}
+
 			FMemory::Memcpy(TextureData, OutMipsData[MipIndex].GetData(), OutMipsData[MipIndex].Num() * sizeof(FColor));
+
 			InOutputTexture->Source.UnlockMip(MipIndex);
 		}
 	}	
@@ -4305,7 +4418,7 @@ void ALandscape::ResolveLayersWeightmapTexture(const TArray<ULandscapeComponent*
 					continue;
 				}
 								
-				bool bWeightmapChanged = ResolveLayersTexture(*CPUReadback, ComponentLayerWeightmap);
+				bool bWeightmapChanged = ResolveLayersTexture(*CPUReadback, ComponentLayerWeightmap, false, Component);
 				ProcessedWeightmaps.Add(ComponentLayerWeightmap, bWeightmapChanged);
 				bChanged |= bWeightmapChanged;
 			}
diff --git a/Engine/Source/Runtime/Landscape/Private/LandscapeRender.cpp b/Engine/Source/Runtime/Landscape/Private/LandscapeRender.cpp
--- a/Engine/Source/Runtime/Landscape/Private/LandscapeRender.cpp
+++ b/Engine/Source/Runtime/Landscape/Private/LandscapeRender.cpp
@@ -81,6 +81,10 @@ static FAutoConsoleVariableRef CVarAllowLandscapeShadows(
 	TEXT("Allow Landscape Shadows")
 );
 
+#if WITH_EDITOR
+extern TAutoConsoleVariable<int32> CVarLandscapeShowDirty;
+#endif
+
 #if !UE_BUILD_SHIPPING
 static void OnLODDistributionScaleChanged(IConsoleVariable* CVar)
 {
@@ -573,6 +577,7 @@ UMaterialInterface* GMaskRegionMaterial = nullptr;
 UMaterialInterface* GColorMaskRegionMaterial = nullptr;
 UTexture2D* GLandscapeBlackTexture = nullptr;
 UMaterialInterface* GLandscapeLayerUsageMaterial = nullptr;
+UMaterialInterface* GLandscapeDirtyMaterial = nullptr;
 #endif
 
 void ULandscapeComponent::GetUsedMaterials(TArray<UMaterialInterface*>& OutMaterials, bool bGetDebugMaterials) const
@@ -623,6 +628,7 @@ void ULandscapeComponent::GetUsedMaterials(TArray<UMaterialInterface*>& OutMater
 		OutMaterials.Add(GMaskRegionMaterial);
 		OutMaterials.Add(GColorMaskRegionMaterial);
 		OutMaterials.Add(GLandscapeLayerUsageMaterial);
+		OutMaterials.Add(GLandscapeDirtyMaterial);
 	}
 #endif
 }
@@ -1193,6 +1199,12 @@ FPrimitiveViewRelevance FLandscapeComponentSceneProxy::GetViewRelevance(const FS
 				ToolRelevance |= GColorMaskRegionMaterial->GetRelevance_Concurrent(FeatureLevel);
 			}
 
+			if (CVarLandscapeShowDirty.GetValueOnRenderThread())
+			{
+				Result.bDynamicRelevance = true;
+				ToolRelevance |= GLandscapeDirtyMaterial->GetRelevance_Concurrent(FeatureLevel);
+			}
+
 			ToolRelevance.SetPrimitiveViewRelevance(Result);
 		}
 	}
@@ -1219,6 +1231,7 @@ FPrimitiveViewRelevance FLandscapeComponentSceneProxy::GetViewRelevance(const FS
 #if WITH_EDITOR
 		(IsSelected() && !GLandscapeEditModeActive) ||
 		(GLandscapeViewMode != ELandscapeViewMode::Normal) ||
+		CVarLandscapeShowDirty.GetValueOnRenderThread() ||
 #else
 		IsSelected() ||
 #endif
@@ -2961,6 +2974,25 @@ void FLandscapeComponentSceneProxy::GetDynamicMeshElements(const TArray<const FS
 						NumDrawCalls += Mesh.Elements.Num();
 					}
 				}
+				else if (CVarLandscapeShowDirty.GetValueOnRenderThread())
+				{
+					Mesh.bCanApplyViewModeOverrides = false;
+					Collector.AddMesh(ViewIndex, Mesh);
+					NumPasses++;
+					NumTriangles += Mesh.GetNumPrimitives();
+					NumDrawCalls += Mesh.Elements.Num();
+
+					FMeshBatch& MaskMesh = Collector.AllocateMesh();
+					MaskMesh = MeshTools;
+										
+					auto DirtyMaterialInstance = new FLandscapeMaskMaterialRenderProxy(GLandscapeDirtyMaterial->GetRenderProxy(), EditToolRenderData.DirtyTexture ? EditToolRenderData.DirtyTexture : GLandscapeBlackTexture, true);
+					MaskMesh.MaterialRenderProxy = DirtyMaterialInstance;
+					Collector.RegisterOneFrameMaterialProxy(DirtyMaterialInstance);
+					Collector.AddMesh(ViewIndex, MaskMesh);
+					NumPasses++;
+					NumTriangles += MaskMesh.GetNumPrimitives();
+					NumDrawCalls += MaskMesh.Elements.Num();
+				}
 				else
 #endif
 					// Regular Landscape rendering. Only use the dynamic path if we're rendering a rich view or we've disabled the static path for debugging.
diff --git a/Engine/Source/Runtime/Landscape/Public/LandscapeEdit.h b/Engine/Source/Runtime/Landscape/Public/LandscapeEdit.h
--- a/Engine/Source/Runtime/Landscape/Public/LandscapeEdit.h
+++ b/Engine/Source/Runtime/Landscape/Public/LandscapeEdit.h
@@ -228,6 +228,12 @@ struct LANDSCAPE_API FLandscapeEditDataInterface : public FLandscapeTextureDataI
 	void GetLayerContributionData(const int32 X1, const int32 Y1, const int32 X2, const int32 Y2, TMap<FIntPoint, uint8>& SparseData);
 	void SetLayerContributionData(const int32 X1, const int32 Y1, const int32 X2, const int32 Y2, const uint8* Data, int32 Stride);
 
+	template<typename TStoreData>
+	void GetDirtyDataTempl(const int32 X1, const int32 Y1, const int32 X2, const int32 Y2, TStoreData& StoreData);
+	void GetDirtyData(const int32 X1, const int32 Y1, const int32 X2, const int32 Y2, uint8* Data, int32 Stride);
+	void GetDirtyData(const int32 X1, const int32 Y1, const int32 X2, const int32 Y2, TMap<FIntPoint, uint8>& SparseData);
+	void SetDirtyData(const int32 X1, const int32 Y1, const int32 X2, const int32 Y2, const uint8* Data, int32 Stride);
+
 	//
 	// XYOffsetmap access
 	//
diff --git a/Engine/Source/Runtime/Landscape/Public/LandscapeRender.h b/Engine/Source/Runtime/Landscape/Public/LandscapeRender.h
--- a/Engine/Source/Runtime/Landscape/Public/LandscapeRender.h
+++ b/Engine/Source/Runtime/Landscape/Public/LandscapeRender.h
@@ -82,6 +82,7 @@ LANDSCAPE_API extern UMaterialInterface* GMaskRegionMaterial;
 LANDSCAPE_API extern UMaterialInterface* GColorMaskRegionMaterial;
 LANDSCAPE_API extern UTexture2D* GLandscapeBlackTexture;
 LANDSCAPE_API extern UMaterialInterface* GLandscapeLayerUsageMaterial;
+LANDSCAPE_API extern UMaterialInterface* GLandscapeDirtyMaterial;
 #endif
 
 
