Landscape Layers: Clear Dirty Data on Save

#rnx

diff --git a/Engine/Source/Runtime/Landscape/Private/Landscape.cpp b/Engine/Source/Runtime/Landscape/Private/Landscape.cpp
--- a/Engine/Source/Runtime/Landscape/Private/Landscape.cpp
+++ b/Engine/Source/Runtime/Landscape/Private/Landscape.cpp
@@ -1936,6 +1936,14 @@ void ALandscapeProxy::PreSave(const class ITargetPlatform* TargetPlatform)
 	{
 		bHasLandscapeGrass = LandscapeComponents.ContainsByPredicate([](ULandscapeComponent* Component) { return Component->MaterialHasGrass(); });
 	}
+
+	if (ALandscape* Landscape = GetLandscapeActor())
+	{
+		for (ULandscapeComponent* LandscapeComponent : LandscapeComponents)
+		{
+			Landscape->ClearDirtyData(LandscapeComponent);
+		}
+	}
 #endif
 }
 
