Buildfix non editor

#rnx

diff --git a/Engine/Source/Runtime/Landscape/Private/LandscapeRender.cpp b/Engine/Source/Runtime/Landscape/Private/LandscapeRender.cpp
--- a/Engine/Source/Runtime/Landscape/Private/LandscapeRender.cpp
+++ b/Engine/Source/Runtime/Landscape/Private/LandscapeRender.cpp
@@ -2974,6 +2974,7 @@ void FLandscapeComponentSceneProxy::GetDynamicMeshElements(const TArray<const FS
 						NumDrawCalls += Mesh.Elements.Num();
 					}
 				}
+#if WITH_EDITOR
 				else if (CVarLandscapeShowDirty.GetValueOnRenderThread())
 				{
 					Mesh.bCanApplyViewModeOverrides = false;
@@ -2993,6 +2994,7 @@ void FLandscapeComponentSceneProxy::GetDynamicMeshElements(const TArray<const FS
 					NumTriangles += MaskMesh.GetNumPrimitives();
 					NumDrawCalls += MaskMesh.Elements.Num();
 				}
+#endif
 				else
 #endif
 					// Regular Landscape rendering. Only use the dynamic path if we're rendering a rich view or we've disabled the static path for debugging.
