Added a AControlPointMeshActor

diff --git a/Engine/Source/Runtime/Landscape/Classes/ControlPointMeshActor.h b/Engine/Source/Runtime/Landscape/Classes/ControlPointMeshActor.h
new file mode 100644
--- /dev/null
+++ b/Engine/Source/Runtime/Landscape/Classes/ControlPointMeshActor.h
@@ -0,0 +1,47 @@
+// Copyright 1998-2019 Epic Games, Inc. All Rights Reserved.
+
+#pragma once
+
+#include "CoreMinimal.h"
+#include "UObject/ObjectMacros.h"
+#include "GameFramework/Actor.h"
+#include "ControlPointMeshActor.generated.h"
+
+/**
+ * ControlPointMeshActor is an actor with a ControlPointMeshComponent.
+ *
+ * @see UControlPointMeshComponent
+ */
+UCLASS(NotPlaceable, hideCategories = (Input), showCategories = ("Input|MouseInput", "Input|TouchInput"), ConversionRoot, ComponentWrapperClass, meta = (ChildCanTick))
+class LANDSCAPE_API AControlPointMeshActor : public AActor
+{
+	GENERATED_UCLASS_BODY()
+
+private:
+	UPROPERTY(Category = ControlPointMeshActor, VisibleAnywhere, BlueprintReadOnly, meta = (ExposeFunctionCategories = "Mesh,Rendering,Physics,Components|StaticMesh,Components|ControlPointMesh", AllowPrivateAccess = "true"))
+	class UControlPointMeshComponent* ControlPointMeshComponent;
+
+public:
+
+	/** Function to change mobility type */
+	void SetMobility(EComponentMobility::Type InMobility);
+
+#if WITH_EDITOR
+	//~ Begin AActor Interface
+	virtual void CheckForErrors() override;
+	virtual bool GetReferencedContentObjects(TArray<UObject*>& Objects) const override;
+	//~ End AActor Interface
+#endif // WITH_EDITOR	
+
+protected:
+	//~ Begin UObject Interface.
+	virtual FString GetDetailedInfoInternal() const override;
+	//~ End UObject Interface.
+
+public:
+	/** Returns ControlPointMeshComponent subobject **/
+	class UControlPointMeshComponent* GetControlPointMeshComponent() const;
+};
+
+
+
diff --git a/Engine/Source/Runtime/Landscape/Private/ControlPointMeshActor.cpp b/Engine/Source/Runtime/Landscape/Private/ControlPointMeshActor.cpp
new file mode 100644
--- /dev/null
+++ b/Engine/Source/Runtime/Landscape/Private/ControlPointMeshActor.cpp
@@ -0,0 +1,78 @@
+// Copyright 1998-2019 Epic Games, Inc. All Rights Reserved.
+
+/*=============================================================================
+	ControlPointMeshActor.cpp: ControlPoint mesh actor class implementation.
+=============================================================================*/
+
+#include "ControlPointMeshActor.h"
+#include "ControlPointMeshComponent.h"
+#include "Engine/CollisionProfile.h"
+#include "Logging/TokenizedMessage.h"
+#include "Logging/MessageLog.h"
+#include "Misc/UObjectToken.h"
+#include "Misc/MapErrors.h"
+#include "Engine/StaticMesh.h"
+
+#define LOCTEXT_NAMESPACE "ControlPointMeshActor"
+
+AControlPointMeshActor::AControlPointMeshActor(const FObjectInitializer& ObjectInitializer)
+	: Super(ObjectInitializer)
+{
+	bCanBeDamaged = false;
+
+	ControlPointMeshComponent = ObjectInitializer.CreateDefaultSubobject<UControlPointMeshComponent>(this, TEXT("ControlPointMeshComponent0"));
+	ControlPointMeshComponent->SetCollisionProfileName(UCollisionProfile::BlockAll_ProfileName);
+	ControlPointMeshComponent->Mobility = EComponentMobility::Static;
+	ControlPointMeshComponent->SetGenerateOverlapEvents(false);
+
+	RootComponent = ControlPointMeshComponent;
+}
+
+
+FString AControlPointMeshActor::GetDetailedInfoInternal() const
+{
+	check(ControlPointMeshComponent != nullptr);
+	return ControlPointMeshComponent->GetDetailedInfoInternal();
+}
+
+void AControlPointMeshActor::SetMobility(EComponentMobility::Type InMobility)
+{
+	check(ControlPointMeshComponent != nullptr);
+	ControlPointMeshComponent->SetMobility(InMobility);
+}
+
+#if WITH_EDITOR
+
+bool AControlPointMeshActor::GetReferencedContentObjects(TArray<UObject*>& Objects) const
+{
+	Super::GetReferencedContentObjects(Objects);
+
+	check(ControlPointMeshComponent != nullptr);
+	if (ControlPointMeshComponent->GetStaticMesh() != nullptr)
+	{
+		Objects.Add(ControlPointMeshComponent->GetStaticMesh());
+	}
+	return true;
+}
+
+void AControlPointMeshActor::CheckForErrors()
+{
+	Super::CheckForErrors();
+
+	FMessageLog MapCheck("MapCheck");
+
+	if (ControlPointMeshComponent->GetStaticMesh() == nullptr)
+	{
+		MapCheck.Warning()
+			->AddToken(FUObjectToken::Create(this))
+			->AddToken(FTextToken::Create(LOCTEXT("MapCheck_Message_ControlPointMeshNull", "ControlPoint mesh actor has NULL StaticMesh property")))
+			->AddToken(FMapErrorToken::Create(FMapErrors::StaticMeshNull));
+	}
+}
+
+#endif // WITH_EDITOR
+
+/** Returns ControlPointMeshComponent subobject **/
+UControlPointMeshComponent* AControlPointMeshActor::GetControlPointMeshComponent() const { return ControlPointMeshComponent; }
+
+#undef LOCTEXT_NAMESPACE
