Landscape Layers: Fix crash when enabling Landscape.ShowDirty outside of LandscapeEdMode

#rnx

diff --git a/Engine/Source/Runtime/Landscape/Private/LandscapeRender.cpp b/Engine/Source/Runtime/Landscape/Private/LandscapeRender.cpp
--- a/Engine/Source/Runtime/Landscape/Private/LandscapeRender.cpp
+++ b/Engine/Source/Runtime/Landscape/Private/LandscapeRender.cpp
@@ -1199,7 +1199,7 @@ FPrimitiveViewRelevance FLandscapeComponentSceneProxy::GetViewRelevance(const FS
 				ToolRelevance |= GColorMaskRegionMaterial->GetRelevance_Concurrent(FeatureLevel);
 			}
 
-			if (CVarLandscapeShowDirty.GetValueOnRenderThread())
+			if (CVarLandscapeShowDirty.GetValueOnRenderThread() && GLandscapeDirtyMaterial)
 			{
 				Result.bDynamicRelevance = true;
 				ToolRelevance |= GLandscapeDirtyMaterial->GetRelevance_Concurrent(FeatureLevel);
@@ -1231,7 +1231,7 @@ FPrimitiveViewRelevance FLandscapeComponentSceneProxy::GetViewRelevance(const FS
 #if WITH_EDITOR
 		(IsSelected() && !GLandscapeEditModeActive) ||
 		(GLandscapeViewMode != ELandscapeViewMode::Normal) ||
-		CVarLandscapeShowDirty.GetValueOnRenderThread() ||
+		(CVarLandscapeShowDirty.GetValueOnRenderThread() && GLandscapeDirtyMaterial) ||
 #else
 		IsSelected() ||
 #endif
@@ -2975,7 +2975,7 @@ void FLandscapeComponentSceneProxy::GetDynamicMeshElements(const TArray<const FS
 					}
 				}
 #if WITH_EDITOR
-				else if (CVarLandscapeShowDirty.GetValueOnRenderThread())
+				else if (CVarLandscapeShowDirty.GetValueOnRenderThread() && GLandscapeDirtyMaterial)
 				{
 					Mesh.bCanApplyViewModeOverrides = false;
 					Collector.AddMesh(ViewIndex, Mesh);
