Landscape Layers: Prevent Layer Blending / Editing when some proxies are Hidden (Level Hidden)

diff --git a/Engine/Source/Editor/LandscapeEditor/Private/LandscapeEdMode.cpp b/Engine/Source/Editor/LandscapeEditor/Private/LandscapeEdMode.cpp
--- a/Engine/Source/Editor/LandscapeEditor/Private/LandscapeEdMode.cpp
+++ b/Engine/Source/Editor/LandscapeEditor/Private/LandscapeEdMode.cpp
@@ -2636,10 +2636,11 @@ bool FEdModeLandscape::CanEditCurrentTarget(FText* Reason) const
 		return false;
 	}
 
+	ALandscapeProxy* Proxy = CurrentToolTarget.LandscapeInfo->GetLandscapeProxy();
+
 	// Landscape Layer Editing not available without a loaded Landscape Actor
 	if (GetLandscape() == nullptr)
 	{
-		ALandscapeProxy* Proxy = CurrentToolTarget.LandscapeInfo->GetLandscapeProxy();
 		if (!Proxy)
 		{
             LocalReason = NSLOCTEXT("UnrealEd", "LandscapeNotFound", "No Landscape found.");
@@ -2651,6 +2652,14 @@ bool FEdModeLandscape::CanEditCurrentTarget(FText* Reason) const
 			LocalReason = NSLOCTEXT("UnrealEd", "LandscapeActorNotLoaded", "Landscape actor is not loaded. It is needed to do layer editing.");
 			return false;
 		}
+
+	}
+	
+	// To edit a Landscape with layers all components need to be registered
+	if(Proxy && Proxy->HasLayersContent() && !CurrentToolTarget.LandscapeInfo->AreAllComponentsRegistered())
+	{
+		LocalReason = NSLOCTEXT("UnrealEd", "LandscapeMissingComponents", "Landscape has some unregistered components (hidden levels).");
+		return false;
 	}
 
 	return true;
diff --git a/Engine/Source/Runtime/Landscape/Classes/LandscapeInfo.h b/Engine/Source/Runtime/Landscape/Classes/LandscapeInfo.h
--- a/Engine/Source/Runtime/Landscape/Classes/LandscapeInfo.h
+++ b/Engine/Source/Runtime/Landscape/Classes/LandscapeInfo.h
@@ -145,6 +145,7 @@ public:
 #if WITH_EDITOR
 	// @todo document 
 	// all below.
+	LANDSCAPE_API bool AreAllComponentsRegistered() const;
 	LANDSCAPE_API void GetComponentsInRegion(int32 X1, int32 Y1, int32 X2, int32 Y2, TSet<ULandscapeComponent*>& OutComponents, bool bOverlap = true) const;
 	LANDSCAPE_API bool GetLandscapeExtent(int32& MinX, int32& MinY, int32& MaxX, int32& MaxY) const;
 	LANDSCAPE_API void ForAllLandscapeComponents(TFunctionRef<void(ULandscapeComponent*)> Fn) const;
diff --git a/Engine/Source/Runtime/Landscape/Classes/LandscapeProxy.h b/Engine/Source/Runtime/Landscape/Classes/LandscapeProxy.h
--- a/Engine/Source/Runtime/Landscape/Classes/LandscapeProxy.h
+++ b/Engine/Source/Runtime/Landscape/Classes/LandscapeProxy.h
@@ -1058,4 +1058,13 @@ protected:
 private:
 	/** Returns Grass Update interval */
 	int32 GetGrassUpdateInterval() const;
+
+#if WITH_EDITORONLY_DATA
+public:
+	static const TArray<ALandscapeProxy*>& GetLandscapeProxies() { return LandscapeProxies; }
+
+private:
+	/** Maintain list of Proxies for faster iteration */
+	static TArray<ALandscapeProxy*> LandscapeProxies;
+#endif
 };
diff --git a/Engine/Source/Runtime/Landscape/Private/Landscape.cpp b/Engine/Source/Runtime/Landscape/Private/Landscape.cpp
--- a/Engine/Source/Runtime/Landscape/Private/Landscape.cpp
+++ b/Engine/Source/Runtime/Landscape/Private/Landscape.cpp
@@ -992,6 +992,10 @@ void ULandscapeComponent::PostLoad()
 
 #endif // WITH_EDITOR
 
+#if WITH_EDITORONLY_DATA
+TArray<ALandscapeProxy*> ALandscapeProxy::LandscapeProxies;
+#endif
+
 ALandscapeProxy::ALandscapeProxy(const FObjectInitializer& ObjectInitializer)
 	: Super(ObjectInitializer)
 #if WITH_EDITORONLY_DATA
@@ -1080,6 +1084,10 @@ ALandscapeProxy::ALandscapeProxy(const FObjectInitializer& ObjectInitializer)
 
 	static uint32 FrameOffsetForTickIntervalInc = 0;
 	FrameOffsetForTickInterval = FrameOffsetForTickIntervalInc++;
+
+#if WITH_EDITORONLY_DATA
+	LandscapeProxies.Add(this);
+#endif
 }
 
 #if WITH_EDITORONLY_DATA
@@ -3303,6 +3311,10 @@ ALandscapeProxy::~ALandscapeProxy()
 	TotalTexturesToStreamForVisibleGrassMapRender -= NumTexturesToStreamForVisibleGrassMapRender;
 	NumTexturesToStreamForVisibleGrassMapRender = 0;
 #endif
+
+#if WITH_EDITORONLY_DATA
+	LandscapeProxies.Remove(this);
+#endif
 }
 
 //
diff --git a/Engine/Source/Runtime/Landscape/Private/LandscapeEdit.cpp b/Engine/Source/Runtime/Landscape/Private/LandscapeEdit.cpp
--- a/Engine/Source/Runtime/Landscape/Private/LandscapeEdit.cpp
+++ b/Engine/Source/Runtime/Landscape/Private/LandscapeEdit.cpp
@@ -2049,6 +2049,35 @@ void ULandscapeComponent::GetComponentExtent(int32& MinX, int32& MinY, int32& Ma
 //
 // ALandscape
 //
+bool ULandscapeInfo::AreAllComponentsRegistered() const
+{
+	const TArray<ALandscapeProxy*>& LandscapeProxies = ALandscapeProxy::GetLandscapeProxies();
+	for(ALandscapeProxy* LandscapeProxy : LandscapeProxies)
+	{
+		if (LandscapeProxy->IsPendingKill())
+		{
+			continue;
+		}
+
+		if (LandscapeProxy->GetLandscapeGuid() == LandscapeGuid)
+		{
+			if (LandscapeProxy->SplineComponent && !LandscapeProxy->SplineComponent->IsRegistered())
+			{
+				return false;
+			}
+
+			for (ULandscapeComponent* LandscapeComponent : LandscapeProxy->LandscapeComponents)
+			{
+				if (LandscapeComponent && !LandscapeComponent->IsRegistered())
+				{
+					return false;
+				}
+			}
+		}
+	}
+
+	return true;
+}
 
 #define MAX_LANDSCAPE_SUBSECTIONS 2
 
@@ -3281,10 +3310,7 @@ LANDSCAPE_API void ULandscapeInfo::ForAllLandscapeComponents(TFunctionRef<void(U
 	{
 		for (ULandscapeComponent* Component : Proxy->LandscapeComponents)
 		{
-			if (Component->IsRegistered())
-			{
-				Fn(Component);
-			}
+			Fn(Component);
 		}
 	});
 }
diff --git a/Engine/Source/Runtime/Landscape/Private/LandscapeEditLayers.cpp b/Engine/Source/Runtime/Landscape/Private/LandscapeEditLayers.cpp
--- a/Engine/Source/Runtime/Landscape/Private/LandscapeEditLayers.cpp
+++ b/Engine/Source/Runtime/Landscape/Private/LandscapeEditLayers.cpp
@@ -4764,7 +4764,8 @@ void ALandscape::UpdateLayersContent(bool bInWaitForStreaming, bool bInSkipMonit
 		return;
 	}
 
-	if (GetLandscapeInfo() == nullptr || !CanHaveLayersContent())
+	ULandscapeInfo* LandscapeInfo = GetLandscapeInfo();
+	if (LandscapeInfo == nullptr || !CanHaveLayersContent() || !LandscapeInfo->AreAllComponentsRegistered())
 	{
 		return;
 	}
@@ -5840,6 +5841,12 @@ void ALandscape::UpdateLandscapeSplines(const FGuid& InTargetLayer, bool bInUpda
 		TSet<ULandscapeComponent*>* ModifiedComponent = nullptr;
 		if (LandscapeSplinesTargetLayerGuid.IsValid())
 		{
+			// Check that we can modify data
+			if (!LandscapeInfo->AreAllComponentsRegistered())
+			{
+				return;
+			}
+						
 			TMap<ULandscapeComponent*, uint32> PreviousHashes;
 			{
 				FLandscapeEditDataInterface LandscapeEdit(LandscapeInfo);
