Implement DestroyRenderThreadResources in FLandscapeMeshProxySceneProxy/FLandscapeMeshProxySceneProxy. Fix Logspam.

#rnx

diff --git a/Engine/Source/Runtime/Landscape/Private/LandscapeRender.cpp b/Engine/Source/Runtime/Landscape/Private/LandscapeRender.cpp
--- a/Engine/Source/Runtime/Landscape/Private/LandscapeRender.cpp
+++ b/Engine/Source/Runtime/Landscape/Private/LandscapeRender.cpp
@@ -1077,6 +1077,12 @@ void FLandscapeComponentSceneProxy::CreateRenderThreadResources()
 #endif
 }
 
+void FLandscapeComponentSceneProxy::DestroyRenderThreadResources()
+{
+	FPrimitiveSceneProxy::DestroyRenderThreadResources();
+	UnregisterNeighbors();
+}
+
 void FLandscapeComponentSceneProxy::OnLevelAddedToWorld()
 {
 	RegisterNeighbors();
@@ -1084,8 +1090,6 @@ void FLandscapeComponentSceneProxy::OnLevelAddedToWorld()
 
 FLandscapeComponentSceneProxy::~FLandscapeComponentSceneProxy()
 {
-	UnregisterNeighbors();
-
 	// Free the subsection uniform buffer
 	LandscapeUniformShaderParameters.ReleaseResource();
 
@@ -5157,8 +5161,10 @@ void FLandscapeMeshProxySceneProxy::OnLevelAddedToWorld()
 	}
 }
 
-FLandscapeMeshProxySceneProxy::~FLandscapeMeshProxySceneProxy()
+void FLandscapeMeshProxySceneProxy::DestroyRenderThreadResources()
 {
+	FStaticMeshSceneProxy::DestroyRenderThreadResources();
+
 	for (FLandscapeNeighborInfo& Info : ProxyNeighborInfos)
 	{
 		Info.UnregisterNeighbors();
diff --git a/Engine/Source/Runtime/Landscape/Public/LandscapeRender.h b/Engine/Source/Runtime/Landscape/Public/LandscapeRender.h
--- a/Engine/Source/Runtime/Landscape/Public/LandscapeRender.h
+++ b/Engine/Source/Runtime/Landscape/Public/LandscapeRender.h
@@ -463,8 +463,8 @@ public:
 	SIZE_T GetTypeHash() const override;
 
 	FLandscapeMeshProxySceneProxy(UStaticMeshComponent* InComponent, const FGuid& InGuid, const TArray<FIntPoint>& InProxyComponentBases, int8 InProxyLOD);
-	virtual ~FLandscapeMeshProxySceneProxy();
 	virtual void CreateRenderThreadResources() override;
+	virtual void DestroyRenderThreadResources() override;
 	virtual void OnLevelAddedToWorld() override;
 };
 
@@ -738,6 +738,7 @@ public:
 	virtual void GetLightRelevance(const FLightSceneProxy* LightSceneProxy, bool& bDynamic, bool& bRelevant, bool& bLightMapped, bool& bShadowMapped) const override;
 	virtual void OnTransformChanged() override;
 	virtual void CreateRenderThreadResources() override;
+	virtual void DestroyRenderThreadResources() override;
 	virtual void OnLevelAddedToWorld() override;
 	virtual void* InitViewCustomData(const FSceneView& InView, float InViewLODScale, FMemStackBase& InCustomDataMemStack, bool InIsStaticRelevant, bool InIsShadowOnly, const FLODMask* InVisiblePrimitiveLODMask = nullptr, float InMeshScreenSizeSquared = -1.0f) override;
 	void PostInitViewCustomData(const FSceneView& InView, void* InViewCustomData) const;
