Support negative LodBias for runtime virtual texture primitives
To be used by primitives that render at a visibly low Lod in the runtime virtual texture

diff --git a/Engine/Source/Runtime/Engine/Classes/Components/PrimitiveComponent.h b/Engine/Source/Runtime/Engine/Classes/Components/PrimitiveComponent.h
--- a/Engine/Source/Runtime/Engine/Classes/Components/PrimitiveComponent.h
+++ b/Engine/Source/Runtime/Engine/Classes/Components/PrimitiveComponent.h
@@ -573,7 +573,7 @@ public:
 	TArray<URuntimeVirtualTexture*> RuntimeVirtualTextures;
 
 	/** Bias to the LOD selected for rendering to runtime virtual textures. */
-	UPROPERTY(EditAnywhere, AdvancedDisplay, Category = VirtualTexture, meta = (DisplayName = "Virtual Texture LOD Bias", UIMin = "0", UIMax = "7"))
+	UPROPERTY(EditAnywhere, AdvancedDisplay, Category = VirtualTexture, meta = (DisplayName = "Virtual Texture LOD Bias", UIMin = "-7", UIMax = "8"))
 	int8 VirtualTextureLodBias = 0;
 
 	/**
@@ -581,7 +581,7 @@ public:
 	 * Larger values reduce the effective draw distance in the runtime virtual texture.
 	 * This culling method doesn't take into account primitive size or virtual texture size.
 	 */
-	UPROPERTY(EditAnywhere, AdvancedDisplay, Category = VirtualTexture, meta = (DisplayName = "Virtual Texture Skip Mips", UIMin = "0", UIMax = "15"))
+	UPROPERTY(EditAnywhere, AdvancedDisplay, Category = VirtualTexture, meta = (DisplayName = "Virtual Texture Skip Mips", UIMin = "0", UIMax = "7"))
 	int8 VirtualTextureCullMips = 0;
 
 	/**
diff --git a/Engine/Source/Runtime/Foliage/Public/FoliageType.h b/Engine/Source/Runtime/Foliage/Public/FoliageType.h
--- a/Engine/Source/Runtime/Foliage/Public/FoliageType.h
+++ b/Engine/Source/Runtime/Foliage/Public/FoliageType.h
@@ -537,7 +537,7 @@ public:
 	 * Larger values reduce the effective draw distance in the runtime virtual texture.
 	 * This culling method doesn't take into account primitive size or virtual texture size.
 	 */
-	UPROPERTY(EditAnywhere, AdvancedDisplay, BlueprintReadOnly, Category = VirtualTexture, meta = (DisplayName = "Virtual Texture Skip Mips", UIMin = "0", UIMax = "15"))
+	UPROPERTY(EditAnywhere, AdvancedDisplay, BlueprintReadOnly, Category = VirtualTexture, meta = (DisplayName = "Virtual Texture Skip Mips", UIMin = "0", UIMax = "7"))
 	int32 VirtualTextureCullMips = 0;
 
 	/** Render to the main pass based on the virtual texture settings. */
diff --git a/Engine/Source/Runtime/Landscape/Classes/LandscapeSplineControlPoint.h b/Engine/Source/Runtime/Landscape/Classes/LandscapeSplineControlPoint.h
--- a/Engine/Source/Runtime/Landscape/Classes/LandscapeSplineControlPoint.h
+++ b/Engine/Source/Runtime/Landscape/Classes/LandscapeSplineControlPoint.h
@@ -162,7 +162,7 @@ class ULandscapeSplineControlPoint : public UObject
 	TArray<URuntimeVirtualTexture*> RuntimeVirtualTextures;
 
 	/** Lod bias for rendering to runtime virtual texture. */
-	UPROPERTY(EditAnywhere, AdvancedDisplay, BlueprintReadOnly, Category = VirtualTexture, meta = (DisplayName = "Virtual Texture LOD Bias", UIMin = "0", UIMax = "7"))
+	UPROPERTY(EditAnywhere, AdvancedDisplay, BlueprintReadOnly, Category = VirtualTexture, meta = (DisplayName = "Virtual Texture LOD Bias", UIMin = "-7", UIMax = "8"))
 	int32 VirtualTextureLodBias = 0;
 
 	/**
@@ -170,7 +170,7 @@ class ULandscapeSplineControlPoint : public UObject
 	 * Larger values reduce the effective draw distance in the runtime virtual texture.
 	 * This culling method doesn't take into account primitive size or virtual texture size.
 	 */
-	UPROPERTY(EditAnywhere, AdvancedDisplay, BlueprintReadOnly, Category = VirtualTexture, meta = (DisplayName = "Virtual Texture Skip Mips", UIMin = "0", UIMax = "15"))
+	UPROPERTY(EditAnywhere, AdvancedDisplay, BlueprintReadOnly, Category = VirtualTexture, meta = (DisplayName = "Virtual Texture Skip Mips", UIMin = "0", UIMax = "7"))
 	int32 VirtualTextureCullMips = 0;
 
 	/** Desired cull distance in the main pass if we are rendering to both the virtual texture AND the main pass. A value of 0 has no effect. */
diff --git a/Engine/Source/Runtime/Landscape/Classes/LandscapeSplineSegment.h b/Engine/Source/Runtime/Landscape/Classes/LandscapeSplineSegment.h
--- a/Engine/Source/Runtime/Landscape/Classes/LandscapeSplineSegment.h
+++ b/Engine/Source/Runtime/Landscape/Classes/LandscapeSplineSegment.h
@@ -264,7 +264,7 @@ class ULandscapeSplineSegment : public UObject
 	TArray<URuntimeVirtualTexture*> RuntimeVirtualTextures;
 
 	/** Lod bias for rendering to runtime virtual texture. */
-	UPROPERTY(EditAnywhere, AdvancedDisplay, BlueprintReadOnly, Category = VirtualTexture, meta = (DisplayName = "Virtual Texture LOD Bias", UIMin = "0", UIMax = "7"))
+	UPROPERTY(EditAnywhere, AdvancedDisplay, BlueprintReadOnly, Category = VirtualTexture, meta = (DisplayName = "Virtual Texture LOD Bias", UIMin = "-7", UIMax = "8"))
 	int32 VirtualTextureLodBias = 0;
 
 	/**
@@ -272,7 +272,7 @@ class ULandscapeSplineSegment : public UObject
 	 * Larger values reduce the effective draw distance in the runtime virtual texture.
 	 * This culling method doesn't take into account primitive size or virtual texture size.
 	 */
-	UPROPERTY(EditAnywhere, AdvancedDisplay, BlueprintReadOnly, Category = VirtualTexture, meta = (DisplayName = "Virtual Texture Skip Mips", UIMin = "0", UIMax = "15"))
+	UPROPERTY(EditAnywhere, AdvancedDisplay, BlueprintReadOnly, Category = VirtualTexture, meta = (DisplayName = "Virtual Texture Skip Mips", UIMin = "0", UIMax = "7"))
 	int32 VirtualTextureCullMips = 0;
 
 	/** Desired cull distance in the main pass if we are rendering to both the virtual texture AND the main pass. A value of 0 has no effect. */
diff --git a/Engine/Source/Runtime/Renderer/Private/PrimitiveSceneInfo.cpp b/Engine/Source/Runtime/Renderer/Private/PrimitiveSceneInfo.cpp
--- a/Engine/Source/Runtime/Renderer/Private/PrimitiveSceneInfo.cpp
+++ b/Engine/Source/Runtime/Renderer/Private/PrimitiveSceneInfo.cpp
@@ -651,11 +651,11 @@ void FPrimitiveSceneInfo::AddToScene(FRHICommandListImmediate& RHICmdList, bool
 	{
 		int8 MinLod, MaxLod;
 		GetRuntimeVirtualTextureLODRange(StaticMeshRelevances, MinLod, MaxLod);
-
+		
 		FPrimitiveVirtualTextureLodInfo& LodInfo = Scene->PrimitiveVirtualTextureLod[PackedIndex];
-		LodInfo.MinLod = MinLod;
-		LodInfo.MaxLod = MaxLod;
-		LodInfo.LodBias = Proxy->GetVirtualTextureLodBias();
+		LodInfo.MinLod = FMath::Clamp((int32)MinLod, 0, 15);
+		LodInfo.MaxLod = FMath::Clamp((int32)MaxLod, 0, 15);
+		LodInfo.LodBias = FMath::Clamp(Proxy->GetVirtualTextureLodBias() + FPrimitiveVirtualTextureLodInfo::LodBiasOffset, 0, 15);
 		LodInfo.CullMethod = Proxy->GetVirtualTextureMinCoverage() == 0 ? 0 : 1;
 		LodInfo.CullValue = LodInfo.CullMethod == 0 ? Proxy->GetVirtualTextureCullMips() : Proxy->GetVirtualTextureMinCoverage();
 	}
diff --git a/Engine/Source/Runtime/Renderer/Private/VT/RuntimeVirtualTextureRender.cpp b/Engine/Source/Runtime/Renderer/Private/VT/RuntimeVirtualTextureRender.cpp
--- a/Engine/Source/Runtime/Renderer/Private/VT/RuntimeVirtualTextureRender.cpp
+++ b/Engine/Source/Runtime/Renderer/Private/VT/RuntimeVirtualTextureRender.cpp
@@ -358,7 +358,8 @@ namespace RuntimeVirtualTexture
 			const float AreaRatio = 2.f * SphereBounds.W * RcpWorldSize;
 			const int32 CurFirstLODIdx = PrimitiveSceneInfo->Proxy->GetCurrentFirstLODIdx_RenderThread();
 			const int32 MinLODIdx = FMath::Max((int32)LodInfo.MinLod, CurFirstLODIdx);
-			const int32 LodIndex = FMath::Clamp<int32>((int32)LodInfo.LodBias - FMath::FloorToInt(FMath::Log2(AreaRatio)), MinLODIdx, LodInfo.MaxLod);
+			const int32 LodBias = (int32)LodInfo.LodBias - FPrimitiveVirtualTextureLodInfo::LodBiasOffset;
+			const int32 LodIndex = FMath::Clamp<int32>(LodBias - FMath::FloorToInt(FMath::Log2(AreaRatio)), MinLODIdx, LodInfo.MaxLod);
 
 			// Process meshes
 			for (int32 MeshIndex = 0; MeshIndex < PrimitiveSceneInfo->StaticMeshes.Num(); ++MeshIndex)
diff --git a/Engine/Source/Runtime/Renderer/Public/PrimitiveSceneInfo.h b/Engine/Source/Runtime/Renderer/Public/PrimitiveSceneInfo.h
--- a/Engine/Source/Runtime/Renderer/Public/PrimitiveSceneInfo.h
+++ b/Engine/Source/Runtime/Renderer/Public/PrimitiveSceneInfo.h
@@ -186,12 +186,15 @@ struct FPrimitiveVirtualTextureFlags
 /** Lod data used for runtime virtual texture page rendering. Packed to reduce memory overhead since one of these is allocated per primitive. */
 struct FPrimitiveVirtualTextureLodInfo
 {
+	/** LodBias is in range [-7,8] so is stored with this offset. */
+	enum { LodBiasOffset = 7 };
+
 	/** Minimum Lod for primitive in the runtime virtual texture. */
 	uint16 MinLod : 4;
 	/** Maximum Lod for primitive in the runtime virtual texture. */
 	uint16 MaxLod : 4;
 	/** Bias to use for Lod calculation in the runtime virtual texture. */
-	uint16 LodBias : 3;
+	uint16 LodBias : 4;
 	/** 
 	 * Culling method used to remove the primitive from low mips of the runtime virtual texture.
 	 * 0: CullValue is the number of low mips for which we cull the primitive from the runtime virtual texture.
@@ -199,7 +202,7 @@ struct FPrimitiveVirtualTextureLodInfo
 	 */
 	uint16 CullMethod : 1;
 	/** Value used according to the CullMethod. */
-	uint16 CullValue : 4;
+	uint16 CullValue : 3;
 };
 
 /** The type of the octree used by FScene to find primitives. */
