Disable material updates for generated mobile weight textures during cooking.

Calling PostEditChange() on the generated UTextures causes the cooking process to spend a lot of time in DoesMaterialUseTexture(). We already know that these textures aren't used on anything yet, because we've just created them and we'll be assigning them to a material later in the function.

diff --git a/Engine/Source/Runtime/Landscape/Private/LandscapeEdit.cpp b/Engine/Source/Runtime/Landscape/Private/LandscapeEdit.cpp
--- a/Engine/Source/Runtime/Landscape/Private/LandscapeEdit.cpp
+++ b/Engine/Source/Runtime/Landscape/Private/LandscapeEdit.cpp
@@ -5717,10 +5717,12 @@ void ULandscapeComponent::GeneratePlatformPixelData()
 		}
 	}
 
+	GDisableAutomaticTextureMaterialUpdateDependencies = true;
 	for (int TextureIdx = 0; TextureIdx < MobileWeightmapTextures.Num(); TextureIdx++)
 	{
 		MobileWeightmapTextures[TextureIdx]->PostEditChange();
 	}
+	GDisableAutomaticTextureMaterialUpdateDependencies = false;
 
 	FLinearColor Masks[4];
 	Masks[0] = FLinearColor(1, 0, 0, 0);
