Landscape: Avoid ReregisterAllComponent call on every LandscapeEdMode tick if Target Landscape hasn't changed

#rnx

diff --git a/Engine/Source/Runtime/Landscape/Private/LandscapeGizmoActor.cpp b/Engine/Source/Runtime/Landscape/Private/LandscapeGizmoActor.cpp
--- a/Engine/Source/Runtime/Landscape/Private/LandscapeGizmoActor.cpp
+++ b/Engine/Source/Runtime/Landscape/Private/LandscapeGizmoActor.cpp
@@ -666,20 +666,23 @@ void ALandscapeGizmoActiveActor::SetTargetLandscape(ULandscapeInfo* LandscapeInf
 		TargetLandscapeInfo = LandscapeInfo;
 	}
 
-	// if there's no copied data, try to move somewhere useful
-	if (TargetLandscapeInfo && TargetLandscapeInfo != PrevInfo && DataType == LGT_None)
+	if (TargetLandscapeInfo != PrevInfo)
 	{
-		MarginZ = TargetLandscapeInfo->DrawScale.Z * 3;
-		Width = Height = TargetLandscapeInfo->DrawScale.X * (TargetLandscapeInfo->ComponentSizeQuads+1);
-
-		float NewLengthZ;
-		FVector NewLocation = TargetLandscapeInfo->GetLandscapeCenterPos(NewLengthZ);
-		SetLength(NewLengthZ);
-		SetActorLocation( NewLocation, false );
-		SetActorRotation(FRotator::ZeroRotator);
-	}
+		// if there's no copied data, try to move somewhere useful
+		if (TargetLandscapeInfo && DataType == LGT_None)
+		{
+			MarginZ = TargetLandscapeInfo->DrawScale.Z * 3;
+			Width = Height = TargetLandscapeInfo->DrawScale.X * (TargetLandscapeInfo->ComponentSizeQuads + 1);
 
-	ReregisterAllComponents();
+			float NewLengthZ;
+			FVector NewLocation = TargetLandscapeInfo->GetLandscapeCenterPos(NewLengthZ);
+			SetLength(NewLengthZ);
+			SetActorLocation(NewLocation, false);
+			SetActorRotation(FRotator::ZeroRotator);
+		}
+
+		ReregisterAllComponents();
+	}
 }
 
 void ALandscapeGizmoActiveActor::ClearGizmoData()
