Landscape Layers: Fix bad cached layer data when adding new components (causing failed checks and not filling layer properly)

#rnx

diff --git a/Engine/Source/Runtime/Landscape/Private/Landscape.cpp b/Engine/Source/Runtime/Landscape/Private/Landscape.cpp
--- a/Engine/Source/Runtime/Landscape/Private/Landscape.cpp
+++ b/Engine/Source/Runtime/Landscape/Private/Landscape.cpp
@@ -1627,6 +1627,8 @@ void ULandscapeComponent::AddLayerData(const FGuid& InLayerGuid, const FLandscap
 	check(!LandscapeEditingLayer.IsValid());
 	FLandscapeLayerComponentData& Data = LayersData.FindOrAdd(InLayerGuid);
 	Data = InData;
+	CachedEditingLayer.Invalidate();
+	CachedEditingLayerData = nullptr;
 }
 
 void ULandscapeComponent::AddDefaultLayerData(const FGuid& InLayerGuid, const TArray<ULandscapeComponent*>& InComponentsUsingHeightmap, TMap<UTexture2D*, UTexture2D*>& InOutCreatedHeightmapTextures)
