Landscape Layers: Remove some FlushRenderingCommands by moving TextureRHI access in actual RHICommand

#rnx

diff --git a/Engine/Source/Runtime/Landscape/Private/LandscapeEditLayers.cpp b/Engine/Source/Runtime/Landscape/Private/LandscapeEditLayers.cpp
--- a/Engine/Source/Runtime/Landscape/Private/LandscapeEditLayers.cpp
+++ b/Engine/Source/Runtime/Landscape/Private/LandscapeEditLayers.cpp
@@ -1403,9 +1403,6 @@ void ALandscape::CreateLayersRenderingResource()
 		}
 	});
 
-	// Flush because TickLayers can access CPU Readbacks in the same frame (see ResolveLayersTexture)
-	FlushRenderingCommands();
-
 	const FIntPoint ComponentCounts = ComputeComponentCounts();
 
 	ALandscape* Landscape = GetLandscapeActor();
@@ -2691,9 +2688,9 @@ void ALandscape::PrintLayersDebugTextureResource(const FString& InContext, FText
 	Flags.SetMip(InMipRender);
 
 	ENQUEUE_RENDER_COMMAND(FLandscapeLayersDebugReadSurfaceCommand)(
-		[SourceTextureRHI = InTextureResource->TextureRHI, SampleRect = SampleRect, OutData = &OutputTexels, ReadFlags = Flags](FRHICommandListImmediate& RHICmdList) mutable
+		[InTextureResource, SampleRect = SampleRect, OutData = &OutputTexels, ReadFlags = Flags](FRHICommandListImmediate& RHICmdList) mutable
 	{
-		RHICmdList.ReadSurfaceData(SourceTextureRHI, SampleRect, *OutData, ReadFlags);
+		RHICmdList.ReadSurfaceData(InTextureResource->TextureRHI, SampleRect, *OutData, ReadFlags);
 	});
 
 	FlushRenderingCommands();
@@ -3186,11 +3183,12 @@ bool ALandscape::ResolveLayersTexture(FLandscapeLayersTexture2DCPUReadBackResour
 	SCOPE_CYCLE_COUNTER(STAT_LandscapeLayersResolveTexture);
 
 	TArray<TArray<FColor>> OutMipsData;
-	OutMipsData.AddDefaulted(InCPUReadBackTexture->TextureRHI->GetNumMips());
-
+	
 	ENQUEUE_RENDER_COMMAND(FLandscapeLayersReadSurfaceCommand)(
 		[InCPUReadBackTexture, &OutMipsData](FRHICommandListImmediate& RHICmdList) mutable
 	{
+		OutMipsData.AddDefaulted(InCPUReadBackTexture->TextureRHI->GetNumMips());
+
 		int32 MipSizeU = InCPUReadBackTexture->GetSizeX();
 		int32 MipSizeV = InCPUReadBackTexture->GetSizeY();
 		int32 MipIndex = 0;
@@ -3996,7 +3994,6 @@ int32 ALandscape::RegenerateLayersWeightmaps(const TArray<ULandscapeComponent*>&
 		if (bHasWeightmapData)
 		{
 			// Will generate CPU read back resource, if required
-			bool bHasPendingInitResource = false;
 			for (ULandscapeComponent* Component : InLandscapeComponentsToRender)
 			{
 				const TArray<UTexture2D*>& ComponentWeightmapTextures = Component->GetWeightmapTextures();
@@ -4014,17 +4011,10 @@ int32 ALandscape::RegenerateLayersWeightmaps(const TArray<ULandscapeComponent*>&
 						WeightmapTexture->Source.UnlockMip(0);
 						BeginInitResource(NewWeightmapCPUReadBack);
 						Proxy->WeightmapsCPUReadBack.Add(WeightmapTexture, NewWeightmapCPUReadBack);
-						bHasPendingInitResource = true;
 					}
 				}
 			}
-
-			if (bHasPendingInitResource)
-			{
-				// Flush because TickLayers can access CPU Readbacks in the same frame (see ResolveLayersTexture)
-				FlushRenderingCommands();
-			}
-
+			
 			int8 CurrentWeightmapToProcessIndex = 0;
 			bool HasFoundWeightmapToProcess = true; // try processing at least once
 
