Fast Landscape Continuous LOD

diff --git a/Engine/Shaders/Private/LandscapeVertexFactory.ush b/Engine/Shaders/Private/LandscapeVertexFactory.ush
--- a/Engine/Shaders/Private/LandscapeVertexFactory.ush
+++ b/Engine/Shaders/Private/LandscapeVertexFactory.ush
@@ -32,19 +32,67 @@ float4 LodBias;
 */
 float4 LodValues;
 
-#if USING_TESSELLATION
-/* Used to perform adaptative falloff using C - K * d
-   x = C
-   y = K
-   z = unused
-   w = unused
-*/
-float4 LodTessellationParams;
-#endif
+int ForcedLod;
 
 /* Fractional Lod values for each section */
 float4 SectionLods;
 
+int GetComponentLinearIndex()
+{
+	return (LandscapeParameters.ComponentBaseY - LandscapeContinuousLODParameters.Min.y) * LandscapeContinuousLODParameters.Size.x + (LandscapeParameters.ComponentBaseX - LandscapeContinuousLODParameters.Min.x);
+}
+
+float4 GetSectionLods()
+{
+	if (ForcedLod >= 0)
+	{
+		return LodValues.x;
+	}
+	
+	return min(LandscapeContinuousLODParameters.SectionLOD[GetComponentLinearIndex()] + View.FarShadowStaticMeshLODBias, LandscapeParameters.LastLOD);
+}
+
+float4 GetNeighborSectionLod()
+{
+	if (ForcedLod >= 0)
+	{
+		return LodValues.xxxx;
+	}
+
+	//       -Y       
+	//    - - 0 - -   
+	//    |       |   
+	// -X 1   P   2 +X
+	//    |       |   
+	//    - - 3 - -   
+	//       +Y
+
+	int2 NeighborOffset[4] = { {0, -1}, {-1, 0}, {1, 0}, {0, 1} };
+	float NeighborLod[4];
+	for (int i = 0; i < 4; i++)
+	{
+		int2 NeighborPos = NeighborOffset[i] + int2(LandscapeParameters.ComponentBaseX, LandscapeParameters.ComponentBaseY);
+		NeighborPos = max(LandscapeContinuousLODParameters.Min                                                     , NeighborPos);
+		NeighborPos = min(LandscapeContinuousLODParameters.Min + LandscapeContinuousLODParameters.Size - int2(1, 1), NeighborPos);
+		NeighborLod[i] = max(LandscapeContinuousLODParameters.SectionLOD[(NeighborPos.y - LandscapeContinuousLODParameters.Min.y) * LandscapeContinuousLODParameters.Size.x + (NeighborPos.x - LandscapeContinuousLODParameters.Min.x)], GetSectionLods());
+		NeighborLod[i] = min(NeighborLod[i] + View.FarShadowStaticMeshLODBias, LandscapeParameters.LastLOD);
+	}
+	return float4(NeighborLod[0], NeighborLod[1], NeighborLod[2], NeighborLod[3]);
+}
+
+float4 GetLodBias()
+{
+#if FIXED_GRID || FEATURE_LEVEL < FEATURE_LEVEL_SM4
+	return LodBias;
+#else
+	return float4(
+		0,
+		0,
+		LandscapeContinuousLODParameters.SectionLODBias[GetComponentLinearIndex()],
+		0); // TODO: handle XY offset texture mipmap bias
+#endif
+}
+
 /* Array index is section
    x = neighbor section's LOD in position 0               0     -> X
    y = neighbor section's LOD in position 1        1   Current  2
@@ -236,8 +284,8 @@ FLandscapeTexCoords GetLandscapeTexCoords(FVertexFactoryInput Input, FVertexFact
 
 	Result.LayerTexCoord.xy = Intermediates.LocalPosition.xy + LandscapeParameters.SubsectionSizeVertsLayerUVPan.zw + Intermediates.InputPosition.zw * LandscapeParameters.SubsectionOffsetParams.ww;
 #if USING_TESSELLATION
-	Result.LayerTexCoord.z = LodTessellationParams.x;
-	Result.LayerTexCoord.w = LodTessellationParams.y;
+	Result.LayerTexCoord.z = LandscapeContinuousLODParameters.SectionTessellationFalloffC[GetComponentLinearIndex()];
+	Result.LayerTexCoord.w = LandscapeContinuousLODParameters.SectionTessellationFalloffK[GetComponentLinearIndex()];
 #endif
 
 	Result.WeightMapTexCoord	= Intermediates.LocalPosition.xy * LandscapeParameters.WeightmapUVScaleBias.xy + LandscapeParameters.WeightmapUVScaleBias.zw + Intermediates.InputPosition.zw * LandscapeParameters.SubsectionOffsetParams.zz;
@@ -514,19 +562,12 @@ void GetShadowMapCoordinate(FVertexFactoryInterpolantsVSToPS Interpolants, out f
 }
 #endif
 
-float CalcLOD_ES2(float2 xy, float2 Subsection)
+float CalcLOD(float2 xy, float2 Subsection)
 {
 	// LOD value based on Barycentric coordinate
 	float4 L0 = float4(xy.y, xy.x, (1 - xy.x), (1 - xy.y)) * 2;
 
-	float4 LODCalculated4 = lerp(
-								(lerp(L0 * SectionLods.x + (float4(1, 1, 1, 1) - L0) * NeighborSectionLod[0],
-									  L0 * SectionLods.y + (float4(1, 1, 1, 1) - L0) * NeighborSectionLod[1],
-									  Subsection.x)),
-								(lerp(L0 * SectionLods.z + (float4(1, 1, 1, 1) - L0) * NeighborSectionLod[2],
-									  L0 * SectionLods.w + (float4(1, 1, 1, 1) - L0) * NeighborSectionLod[3],
-									  Subsection.x)),
-								Subsection.y);
+	float4 LODCalculated4 = L0 * GetSectionLods() + (float4(1, 1, 1, 1) - L0) * GetNeighborSectionLod();
 
 	// Find which quadrant of the subsection we're in - top, left, right or bottom.
 	// 0 --------- 1 X  1 ----- 2  1          1 
@@ -541,77 +582,6 @@ float CalcLOD_ES2(float2 xy, float2 Subsection)
 	return LODCalculated;
 }
 
-float CalcLOD(float2 xy, float2 Subsection)
-{
-#if ES2_PROFILE
-	// less branchy version for ES2 devices
-	return CalcLOD_ES2(xy, Subsection); 
-#else
-	// LOD value based on Barycentric coordinate
-	float4 L0 = float4(xy.y, xy.x, (1 - xy.x), (1 - xy.y)) * 2;
-
-	float4 LODCalculated4;
-	if (Subsection.y > 0.5f)
-	{
-		if (Subsection.x > 0.5f)
-		{
-			float SectionLod = SectionLods.w;
-			LODCalculated4 = L0 * SectionLod + (float4(1, 1, 1, 1) - L0) * NeighborSectionLod[3];
-		}
-		else
-		{
-			float SectionLod = SectionLods.z;
-			LODCalculated4 = L0 * SectionLod + (float4(1, 1, 1, 1) - L0) * NeighborSectionLod[2];
-		}
-	}
-	else
-	{
-		if (Subsection.x > 0.5f)
-		{
-			float SectionLod = SectionLods.y;
-			LODCalculated4 = L0 * SectionLod + (float4(1, 1, 1, 1) - L0) * NeighborSectionLod[1];
-		}
-		else
-		{
-			float SectionLod = SectionLods.x;
-			LODCalculated4 = L0 * SectionLod + (float4(1, 1, 1, 1) - L0) * NeighborSectionLod[0];
-		}
-	}
-
-	// Find which quadrant of the subsection we're in - top, left, right or bottom.
-	// 0 --------- 1 X  1 ----- 2  1          1 
-	//   | \ 0 / |        \   /    | \      / | 
-	//   | 1 + 2 |          0      |  0    0  |      0       
-	//   | / 3 \ |                 | /      \ |    /   \     
-	// 2 --------- 3               2          2  1 ----- 2   
-	float LODCalculated;
-	if ((xy.x + xy.y) > 1)
-	{
-		if (xy.x < xy.y)
-		{
-			LODCalculated = LODCalculated4.w;
-		}
-		else
-		{
-			LODCalculated = LODCalculated4.z;
-		}
-	}
-	else
-	{
-		if (xy.x < xy.y)
-		{
-			LODCalculated = LODCalculated4.y;
-		}
-		else
-		{
-			LODCalculated = LODCalculated4.x;
-		}
-	}
-
-	return LODCalculated;
-#endif
-}
-
 FVertexFactoryIntermediates GetVertexFactoryIntermediates(FVertexFactoryInput Input)
 {
 	FVertexFactoryIntermediates Intermediates;
@@ -676,20 +646,20 @@ FVertexFactoryIntermediates GetVertexFactoryIntermediates(FVertexFactoryInput In
 	// Get the height and normal XY for current and next LOD out of the textures
 #if FEATURE_LEVEL >= FEATURE_LEVEL_SM4
 	float2 SampleCoords = InputPositionLODAdjusted * LandscapeParameters.HeightmapUVScaleBias.xy + LandscapeParameters.HeightmapUVScaleBias.zw + 0.5*LandscapeParameters.HeightmapUVScaleBias.xy + Intermediates.InputPosition.zw * LandscapeParameters.SubsectionOffsetParams.xy;
-	float4 SampleValue = Texture2DSampleLevel(LandscapeParameters.HeightmapTexture, LandscapeParameters.HeightmapTextureSampler, SampleCoords, LodValue-LodBias.z);
+	float4 SampleValue = Texture2DSampleLevel(LandscapeParameters.HeightmapTexture, LandscapeParameters.HeightmapTextureSampler, SampleCoords, LodValue-GetLodBias().z);
 	float Height = DecodePackedHeight(SampleValue.xy);
 
 	float2 SampleCoordsNextLOD = InputPositionNextLOD * LandscapeParameters.HeightmapUVScaleBias.xy + LandscapeParameters.HeightmapUVScaleBias.zw + 0.5*LandscapeParameters.HeightmapUVScaleBias.xy + Intermediates.InputPosition.zw * LandscapeParameters.SubsectionOffsetParams.xy;
-	float4 SampleValueNextLOD = Texture2DSampleLevel(LandscapeParameters.HeightmapTexture, LandscapeParameters.HeightmapTextureSampler, SampleCoordsNextLOD, LodValue+1-LodBias.z );
+	float4 SampleValueNextLOD = Texture2DSampleLevel(LandscapeParameters.HeightmapTexture, LandscapeParameters.HeightmapTextureSampler, SampleCoordsNextLOD, LodValue+1-GetLodBias().z);
 	float HeightNextLOD = DecodePackedHeight(SampleValueNextLOD.xy);
 #endif
 
 
 #if LANDSCAPE_XYOFFSET // FEATURE_LEVEL >= FEATURE_LEVEL_SM4 only
 	float2 SampleCoords2 = float2(InputPositionLODAdjusted * LandscapeParameters.WeightmapUVScaleBias.xy + LandscapeParameters.WeightmapUVScaleBias.zw + Intermediates.InputPosition.zw * LandscapeParameters.SubsectionOffsetParams.zz);
-	float4 OffsetValue = Texture2DSampleLevel( LandscapeParameters.XYOffsetmapTexture, LandscapeParameters.XYOffsetmapTextureSampler, SampleCoords2, LodValue-LodBias.w );
+	float4 OffsetValue = Texture2DSampleLevel( LandscapeParameters.XYOffsetmapTexture, LandscapeParameters.XYOffsetmapTextureSampler, SampleCoords2, LodValue-GetLodBias().w );
 	float2 SampleCoordsNextLOD2 = float2(InputPositionNextLOD * LandscapeParameters.WeightmapUVScaleBias.xy + LandscapeParameters.WeightmapUVScaleBias.zw + Intermediates.InputPosition.zw * LandscapeParameters.SubsectionOffsetParams.zz);
-	float4 OffsetValueNextLOD = Texture2DSampleLevel( LandscapeParameters.XYOffsetmapTexture, LandscapeParameters.XYOffsetmapTextureSampler, SampleCoordsNextLOD2, LodValue+1-LodBias.w );
+	float4 OffsetValueNextLOD = Texture2DSampleLevel( LandscapeParameters.XYOffsetmapTexture, LandscapeParameters.XYOffsetmapTextureSampler, SampleCoordsNextLOD2, LodValue+1-GetLodBias().w );
 	float2 XYOffset = float2(((OffsetValue.r * 255.0 * 256.0 + OffsetValue.g * 255.0) - 32768.0) * XYOFFSET_SCALE, ((OffsetValue.b * 255.0 * 256.0 + OffsetValue.a * 255.0) - 32768.0) * XYOFFSET_SCALE );
 	float2 XYOffsetNextLOD = float2(((OffsetValueNextLOD.r * 255.0 * 256.0 + OffsetValueNextLOD.g * 255.0) - 32768.0) * XYOFFSET_SCALE, ((OffsetValueNextLOD.b * 255.0 * 256.0 + OffsetValueNextLOD.a * 255.0) - 32768.0) * XYOFFSET_SCALE );
 
@@ -731,7 +701,7 @@ float2 LandscapeVertexFactorySampleHeight(FVertexFactoryInput Input, FVertexFact
 
 	float2 InputPositionLODAdjusted = ActualLODCoordsInt / CoordTranslate_x;
 	float2 SampleCoords = InputPositionLODAdjusted * LandscapeParameters.HeightmapUVScaleBias.xy + LandscapeParameters.HeightmapUVScaleBias.zw + 0.5*LandscapeParameters.HeightmapUVScaleBias.xy + Intermediates.InputPosition.zw * LandscapeParameters.SubsectionOffsetParams.xy;
-	float4 SampleValue = Texture2DSampleLevel(LandscapeParameters.HeightmapTexture, LandscapeParameters.HeightmapTextureSampler, SampleCoords, LodValue-LodBias.z);
+	float4 SampleValue = Texture2DSampleLevel(LandscapeParameters.HeightmapTexture, LandscapeParameters.HeightmapTextureSampler, SampleCoords, LodValue-GetLodBias().z);
 	return SampleValue.rg;
 }
 #endif
diff --git a/Engine/Source/Runtime/Engine/Public/MeshBatch.h b/Engine/Source/Runtime/Engine/Public/MeshBatch.h
--- a/Engine/Source/Runtime/Engine/Public/MeshBatch.h
+++ b/Engine/Source/Runtime/Engine/Public/MeshBatch.h
@@ -245,9 +245,6 @@ struct FMeshBatch
 	/** The current hit proxy ID being rendered. */
 	FHitProxyId BatchHitProxyId;
 
-	/** This is the threshold that will be used to know if we should use this mesh batch or use one with no tessellation enabled */
-	float TessellationDisablingShadowMapMeshSize;
-
 	FORCEINLINE bool IsTranslucent(ERHIFeatureLevel::Type InFeatureLevel) const
 	{
 		// Note: blend mode does not depend on the feature level we are actually rendering in.
@@ -341,7 +338,6 @@ struct FMeshBatch
 	,	LCI(NULL)
 	,	VertexFactory(NULL)
 	,	MaterialRenderProxy(NULL)
-	,	TessellationDisablingShadowMapMeshSize(0.0f)
 	{
 		// By default always add the first element.
 		new(Elements) FMeshBatchElement;
diff --git a/Engine/Source/Runtime/Landscape/Classes/LandscapeProxy.h b/Engine/Source/Runtime/Landscape/Classes/LandscapeProxy.h
--- a/Engine/Source/Runtime/Landscape/Classes/LandscapeProxy.h
+++ b/Engine/Source/Runtime/Landscape/Classes/LandscapeProxy.h
@@ -436,7 +436,7 @@ public:
 	float ComponentScreenSizeToUseSubSections;
 
 	/** This is the starting screen size used to calculate the distribution, by default it's 1, but you can increase the value if you want less LOD0 component, and you use very large landscape component. */
-	UPROPERTY(EditAnywhere, Category = "LOD Distribution", meta = (DisplayName = "LOD 0 Screen Size", ClampMin = "1.0", ClampMax = "10.0", UIMin = "1.0", UIMax = "10.0"))
+	UPROPERTY(EditAnywhere, Category = "LOD Distribution", meta = (DisplayName = "LOD 0 Screen Size", ClampMin = "0.1", ClampMax = "10.0", UIMin = "1.0", UIMax = "10.0"))
 	float LOD0ScreenSize;
 
 	/** The distribution setting used to change the LOD 0 generation, 1.75 is the normal distribution, numbers influence directly the LOD0 proportion on screen. */
diff --git a/Engine/Source/Runtime/Landscape/Private/LandscapeEdit.cpp b/Engine/Source/Runtime/Landscape/Private/LandscapeEdit.cpp
--- a/Engine/Source/Runtime/Landscape/Private/LandscapeEdit.cpp
+++ b/Engine/Source/Runtime/Landscape/Private/LandscapeEdit.cpp
@@ -4409,7 +4409,7 @@ void ALandscape::PostEditChangeProperty(FPropertyChangedEvent& PropertyChangedEv
 	}
 	else if (PropertyName == FName(TEXT("LOD0ScreenSize")))
 	{
-		LOD0ScreenSize = FMath::Clamp<float>(LOD0ScreenSize, 1.0f, 10.0f);
+		LOD0ScreenSize = FMath::Clamp<float>(LOD0ScreenSize, 0.1f, 10.0f);
 		bPropagateToProxies = true;
 	}
 	else if (PropertyName == FName(TEXT("CollisionMipLevel")))
diff --git a/Engine/Source/Runtime/Landscape/Private/LandscapeRender.cpp b/Engine/Source/Runtime/Landscape/Private/LandscapeRender.cpp
--- a/Engine/Source/Runtime/Landscape/Private/LandscapeRender.cpp
+++ b/Engine/Source/Runtime/Landscape/Private/LandscapeRender.cpp
@@ -641,6 +641,319 @@ TMap<FLandscapeNeighborInfo::FLandscapeKey, TMap<FIntPoint, const FLandscapeNeig
 
 const static FName NAME_LandscapeResourceNameForDebugging(TEXT("Landscape"));
 
+IMPLEMENT_GLOBAL_SHADER_PARAMETER_STRUCT(FLandscapeSectionLODUniformParameters, "LandscapeContinuousLODParameters");
+
+TGlobalResource<FNullLandscapeRenderSystemResources> GNullLandscapeRenderSystemResources;
+TMap<FLandscapeNeighborInfo::FLandscapeKey, FLandscapeRenderSystem*> LandscapeRenderSystems;
+
+void FLandscapeRenderSystem::ResizeAndMoveTo(FIntPoint NewMin, FIntPoint NewSize)
+{
+	SectionLODBuffer.SafeRelease();
+	SectionLODBiasBuffer.SafeRelease();
+	SectionTessellationFalloffCBuffer.SafeRelease();
+	SectionTessellationFalloffKBuffer.SafeRelease();
+
+	TResourceArray<float> NewSectionLODValues;
+	TResourceArray<float> NewSectionLODBiases;
+	TResourceArray<float> NewSectionTessellationFalloffC;
+	TResourceArray<float> NewSectionTessellationFalloffK;
+	TArray<LODSettingsComponent> NewSectionLODSettings;
+	TArray<FVector4> NewSectionOriginAndRadius;
+	TArray<FLandscapeComponentSceneProxy*> NewSceneProxies;
+
+	NewSectionLODValues.AddZeroed(NewSize.X * NewSize.Y);
+	NewSectionLODBiases.AddZeroed(NewSize.X * NewSize.Y);
+	NewSectionTessellationFalloffC.AddZeroed(NewSize.X * NewSize.Y);
+	NewSectionTessellationFalloffK.AddZeroed(NewSize.X * NewSize.Y);
+	NewSectionLODSettings.AddZeroed(NewSize.X * NewSize.Y);
+	NewSectionOriginAndRadius.AddZeroed(NewSize.X * NewSize.Y);
+	NewSceneProxies.AddZeroed(NewSize.X * NewSize.Y);
+
+	for (int32 Y = 0; Y < Size.Y; Y++)
+	{
+		for (int32 X = 0; X < Size.X; X++)
+		{
+			int32 LinearIndex = Y * Size.X + X;
+			int32 NewLinearIndex = (Y + (Min.Y - NewMin.Y)) * NewSize.X + (X + (Min.X - NewMin.X));
+			NewSectionLODValues[NewLinearIndex] = SectionLODValues[LinearIndex];
+			NewSectionLODBiases[NewLinearIndex] = SectionLODBiases[LinearIndex];
+			if (TessellationFalloffSettings.UseTessellationComponentScreenSizeFalloff)
+			{
+				NewSectionTessellationFalloffC[NewLinearIndex] = SectionTessellationFalloffC[LinearIndex];
+				NewSectionTessellationFalloffK[NewLinearIndex] = SectionTessellationFalloffK[LinearIndex];
+			}
+			NewSectionLODSettings[NewLinearIndex] = SectionLODSettings[LinearIndex];
+			NewSectionOriginAndRadius[NewLinearIndex] = SectionOriginAndRadius[LinearIndex];
+			NewSceneProxies[NewLinearIndex] = SceneProxies[LinearIndex];
+		}
+	}
+
+	Min = NewMin;
+	Size = NewSize;
+	SectionLODValues = NewSectionLODValues;
+	SectionLODBiases = NewSectionLODBiases;
+	SectionTessellationFalloffC = NewSectionTessellationFalloffC;
+	SectionTessellationFalloffK = NewSectionTessellationFalloffK;
+	SectionLODSettings = NewSectionLODSettings;
+	SectionOriginAndRadius = NewSectionOriginAndRadius;
+	SceneProxies = NewSceneProxies;
+
+	if (!TessellationFalloffSettings.UseTessellationComponentScreenSizeFalloff)
+	{
+		for (float& Value : SectionTessellationFalloffC)
+		{
+			Value = 1.0f;
+		}
+
+		for (float& Value : SectionTessellationFalloffK)
+		{
+			Value = 0.0f;
+		}
+	}
+
+	SectionLODValues.SetAllowCPUAccess(true);
+	SectionLODBiases.SetAllowCPUAccess(true);
+	SectionTessellationFalloffC.SetAllowCPUAccess(true);
+	SectionTessellationFalloffK.SetAllowCPUAccess(true);
+}
+
+void FLandscapeRenderSystem::ComputeSectionPerViewParameters(const FSceneView* View)
+{
+	if (!CachedSectionLODValues.Contains(View))
+	{
+		TRACE_CPUPROFILER_EVENT_SCOPE(FLandscapeRenderSystem::ComputeSectionPerViewParameters());
+
+		CachedSectionLODValues.Add(View, TResourceArray<float>{});
+
+		if (TessellationFalloffSettings.UseTessellationComponentScreenSizeFalloff)
+		{
+			CachedSectionTessellationFalloffC.Add(View, TResourceArray<float>{});
+			CachedSectionTessellationFalloffK.Add(View, TResourceArray<float>{});
+		}
+
+		int32 ForcedLODLevel = View->Family->EngineShowFlags.LOD ? GetCVarForceLOD() : 0;
+
+		for (int32 EntityIndex = 0; EntityIndex < SectionLODSettings.Num(); EntityIndex++)
+		{
+			float MeshScreenSizeSquared = ComputeBoundsScreenRadiusSquared(FVector(SectionOriginAndRadius[EntityIndex]), SectionOriginAndRadius[EntityIndex].W, *View);
+
+			float FractionalLOD;
+			GetLODFromScreenSize(SectionLODSettings[EntityIndex], MeshScreenSizeSquared, View->LODDistanceFactor, FractionalLOD);
+
+			SectionLODValues[EntityIndex] = FractionalLOD;
+
+			if (TessellationFalloffSettings.UseTessellationComponentScreenSizeFalloff)
+			{
+				float MaxTesselationDistance = ComputeBoundsDrawDistance(FMath::Sqrt(TessellationFalloffSettings.TessellationComponentSquaredScreenSize), SectionOriginAndRadius[EntityIndex].W / 2.0f, View->ViewMatrices.GetProjectionMatrix());
+				float FallOffStartingDistance = FMath::Min(
+					ComputeBoundsDrawDistance(FMath::Sqrt(FMath::Min(
+						FMath::Square(TessellationFalloffSettings.TessellationComponentScreenSizeFalloff), 
+						TessellationFalloffSettings.TessellationComponentSquaredScreenSize)), SectionOriginAndRadius[EntityIndex].W / 2.0f, View->ViewMatrices.GetProjectionMatrix()) - MaxTesselationDistance, MaxTesselationDistance);
+
+				// Calculate the falloff using a = C - K * d by sending C & K into the shader
+				SectionTessellationFalloffC[EntityIndex] = MaxTesselationDistance / (MaxTesselationDistance - FallOffStartingDistance);
+				SectionTessellationFalloffK[EntityIndex] = -(1 / (-MaxTesselationDistance + FallOffStartingDistance));
+			}
+		}
+
+		CachedSectionLODValues[View] = SectionLODValues;
+
+		if (TessellationFalloffSettings.UseTessellationComponentScreenSizeFalloff)
+		{
+			CachedSectionTessellationFalloffC[View] = SectionTessellationFalloffC;
+			CachedSectionTessellationFalloffK[View] = SectionTessellationFalloffK;
+		}
+	}
+	else
+	{
+		SectionLODValues = CachedSectionLODValues[View];
+
+		if (TessellationFalloffSettings.UseTessellationComponentScreenSizeFalloff)
+		{
+			SectionTessellationFalloffC = CachedSectionTessellationFalloffC[View];
+			SectionTessellationFalloffK = CachedSectionTessellationFalloffK[View];
+		}
+	}
+}
+
+void FLandscapeRenderSystem::FetchHeightmapLODBiases()
+{
+	TRACE_CPUPROFILER_EVENT_SCOPE(FLandscapeRenderSystem::FetchHeightmapLODBiases());
+
+	// TODO: this function generates A LOT OF cache misses - it should be much better if we have an event of FTexture2DResource::UpdateTexture
+	for (int32 EntityIndex = 0; EntityIndex < SceneProxies.Num(); EntityIndex++)
+	{
+		FLandscapeComponentSceneProxy* SceneProxy = SceneProxies[EntityIndex];
+		if (SceneProxy)
+		{
+			if (SceneProxy->HeightmapTexture)
+			{
+				float SectionLODBias = ((FTexture2DResource*)SceneProxy->HeightmapTexture->Resource)->GetCurrentFirstMip();
+				SectionLODBiases[EntityIndex] = SectionLODBias;
+
+				// TODO: support mipmap LOD bias of XY offset map
+				//XYOffsetmapTexture ? ((FTexture2DResource*)XYOffsetmapTexture->Resource)->GetCurrentFirstMip() : 0.0f);
+			}
+		}
+	}
+}
+
+void FLandscapeRenderSystem::RecreateBuffers(const FSceneView* InView /* = nullptr */)
+{
+	if (InView == nullptr || CachedView != InView)
+	{
+		TRACE_CPUPROFILER_EVENT_SCOPE(FLandscapeRenderSystem::RecreateBuffers());
+
+		{
+			if (!SectionLODBuffer.IsValid())
+			{
+				FRHIResourceCreateInfo CreateInfo(&SectionLODValues);
+				SectionLODBuffer = RHICreateVertexBuffer(SectionLODValues.GetResourceDataSize(), BUF_ShaderResource | BUF_Volatile, CreateInfo);
+				SectionLODSRV = RHICreateShaderResourceView(SectionLODBuffer, sizeof(float), PF_R32_FLOAT);
+			}
+			else
+			{
+				float* Data = (float*)RHILockVertexBuffer(SectionLODBuffer, 0, SectionLODValues.GetResourceDataSize(), RLM_WriteOnly);
+				FMemory::Memcpy(Data, SectionLODValues.GetData(), SectionLODValues.GetResourceDataSize());
+				RHIUnlockVertexBuffer(SectionLODBuffer);
+			}
+		}
+
+		{
+			if (!SectionLODBiasBuffer.IsValid())
+			{
+				FRHIResourceCreateInfo CreateInfo(&SectionLODBiases);
+				SectionLODBiasBuffer = RHICreateVertexBuffer(SectionLODBiases.GetResourceDataSize(), BUF_ShaderResource | BUF_Volatile, CreateInfo);
+				SectionLODBiasSRV = RHICreateShaderResourceView(SectionLODBiasBuffer, sizeof(float), PF_R32_FLOAT);
+			}
+			else
+			{
+				float* Data = (float*)RHILockVertexBuffer(SectionLODBiasBuffer, 0, SectionLODBiases.GetResourceDataSize(), RLM_WriteOnly);
+				FMemory::Memcpy(Data, SectionLODBiases.GetData(), SectionLODBiases.GetResourceDataSize());
+				RHIUnlockVertexBuffer(SectionLODBiasBuffer);
+			}
+		}
+
+		{
+			if (!SectionTessellationFalloffCBuffer.IsValid())
+			{
+				FRHIResourceCreateInfo CreateInfo(&SectionTessellationFalloffC);
+				SectionTessellationFalloffCBuffer = RHICreateVertexBuffer(SectionTessellationFalloffC.GetResourceDataSize(), BUF_ShaderResource | BUF_Volatile, CreateInfo);
+				SectionTessellationFalloffCSRV = RHICreateShaderResourceView(SectionTessellationFalloffCBuffer, sizeof(float), PF_R32_FLOAT);
+			}
+			else
+			{
+				// If we use tessellation falloff, update the buffer, otherwise use the one already filled with default parameters
+				if (TessellationFalloffSettings.UseTessellationComponentScreenSizeFalloff)
+				{
+					float* Data = (float*)RHILockVertexBuffer(SectionTessellationFalloffCBuffer, 0, SectionTessellationFalloffC.GetResourceDataSize(), RLM_WriteOnly);
+					FMemory::Memcpy(Data, SectionTessellationFalloffC.GetData(), SectionTessellationFalloffC.GetResourceDataSize());
+					RHIUnlockVertexBuffer(SectionTessellationFalloffCBuffer);
+				}
+			}
+		}
+
+		{
+			if (!SectionTessellationFalloffKBuffer.IsValid())
+			{
+				FRHIResourceCreateInfo CreateInfo(&SectionTessellationFalloffK);
+				SectionTessellationFalloffKBuffer = RHICreateVertexBuffer(SectionTessellationFalloffK.GetResourceDataSize(), BUF_ShaderResource | BUF_Volatile, CreateInfo);
+				SectionTessellationFalloffKSRV = RHICreateShaderResourceView(SectionTessellationFalloffKBuffer, sizeof(float), PF_R32_FLOAT);
+			}
+			else
+			{
+				// If we use tessellation falloff, update the buffer, otherwise use the one already filled with default parameters
+				if (TessellationFalloffSettings.UseTessellationComponentScreenSizeFalloff)
+				{
+					float* Data = (float*)RHILockVertexBuffer(SectionTessellationFalloffKBuffer, 0, SectionTessellationFalloffK.GetResourceDataSize(), RLM_WriteOnly);
+					FMemory::Memcpy(Data, SectionTessellationFalloffK.GetData(), SectionTessellationFalloffK.GetResourceDataSize());
+					RHIUnlockVertexBuffer(SectionTessellationFalloffKBuffer);
+				}
+			}
+		}
+
+		FLandscapeSectionLODUniformParameters Parameters;
+		Parameters.Min = Min;
+		Parameters.Size = Size;
+		Parameters.SectionLOD = SectionLODSRV;
+		Parameters.SectionLODBias = SectionLODBiasSRV;
+		Parameters.SectionTessellationFalloffC = SectionTessellationFalloffCSRV;
+		Parameters.SectionTessellationFalloffK = SectionTessellationFalloffKSRV;
+
+		if (UniformBuffer.IsValid())
+		{
+			UniformBuffer.UpdateUniformBufferImmediate(Parameters);
+		}
+		else
+		{
+			UniformBuffer = TUniformBufferRef<FLandscapeSectionLODUniformParameters>::CreateUniformBufferImmediate(Parameters, UniformBuffer_SingleFrame);
+		}
+
+		CachedView = InView;
+	}
+}
+
+void FLandscapeRenderSystem::BeginFrame()
+{
+	CachedView = nullptr;
+
+	CachedSectionLODValues.Empty();
+
+	if (TessellationFalloffSettings.UseTessellationComponentScreenSizeFalloff)
+	{
+		CachedSectionTessellationFalloffC.Empty();
+		CachedSectionTessellationFalloffK.Empty();
+	}
+
+	const bool bExecuteInParallel = FApp::ShouldUseThreadingForPerformance()
+		&& GRenderingThread; // Rendering thread is required to safely use rendering resources in parallel.
+
+	if (bExecuteInParallel)
+	{
+		FetchHeightmapLODBiasesEventRef = TGraphTask<FGetSectionLODBiasesTask>::CreateTask(
+			nullptr, ENamedThreads::GetRenderThread()).ConstructAndDispatchWhenReady(*this);
+	}
+	else
+	{
+		FGetSectionLODBiasesTask Task(*this);
+		Task.AnyThreadTask();
+	}
+}
+
+class FLandscapePersistentViewUniformBufferExtension : public IPersistentViewUniformBufferExtension
+{
+public:
+	virtual void BeginFrame() override
+	{
+		for (auto& Pair : LandscapeRenderSystems)
+		{
+			FLandscapeRenderSystem& RenderSystem = *Pair.Value;
+
+			RenderSystem.BeginFrame();
+		}
+	}
+
+	virtual void PrepareView(const FSceneView* View) override
+	{
+		for (auto& Pair : LandscapeRenderSystems)
+		{
+			FLandscapeRenderSystem& RenderSystem = *Pair.Value;
+
+			RenderSystem.PrepareView(View);
+		}
+	}
+
+	virtual void BeginRenderView(const FSceneView* View) override
+	{
+		for (auto& Pair : LandscapeRenderSystems)
+		{
+			FLandscapeRenderSystem& RenderSystem = *Pair.Value;
+
+			RenderSystem.BeginRenderView(View);
+		}
+	}
+} LandscapePersistentViewUniformBufferExtension;
+
 FLandscapeComponentSceneProxy::FLandscapeComponentSceneProxy(ULandscapeComponent* InComponent)
 	: FPrimitiveSceneProxy(InComponent, NAME_LandscapeResourceNameForDebugging)
 	, FLandscapeNeighborInfo(InComponent->GetWorld(), InComponent->GetLandscapeProxy()->GetLandscapeGuid(), InComponent->GetSectionBase() / InComponent->ComponentSizeQuads, InComponent->GetHeightmap(), InComponent->ForcedLOD, InComponent->LODBias)
@@ -748,8 +1061,11 @@ FLandscapeComponentSceneProxy::FLandscapeComponentSceneProxy(ULandscapeComponent
 
 	// LOD 0 handling
 	LODScreenRatioSquared[0] = FMath::Square(CurrentScreenSizeRatio);
+	LODSettings.LOD0ScreenSizeSquared = FMath::Square(CurrentScreenSizeRatio);
 	CurrentScreenSizeRatio /= ScreenSizeRatioDivider;
+	LODSettings.LOD1ScreenSizeSquared = FMath::Square(CurrentScreenSizeRatio);
 	ScreenSizeRatioDivider = FMath::Max(InComponent->GetLandscapeProxy()->LODDistributionSetting * GLandscapeLODDistributionScale, 1.01f);
+	LODSettings.LODOnePlusDistributionScalarSquared = FMath::Square(ScreenSizeRatioDivider);
 
 	// Other LODs
 	for (int32 LODIndex = 1; LODIndex <= MaxLOD; ++LODIndex) // This should ALWAYS be calculated from the component size, not user MaxLOD override
@@ -766,6 +1082,9 @@ FLandscapeComponentSceneProxy::FLandscapeComponentSceneProxy(ULandscapeComponent
 	FirstLOD = 0;
 	LastLOD = MaxLOD;	// we always need to go to MaxLOD regardless of LODBias as we could need the lowest LODs due to streaming.
 
+	LODSettings.LastLODIndex = LastLOD;
+	LODSettings.LastLODScreenSizeSquared = LODScreenRatioSquared[LastLOD];
+
 	// Make sure out LastLOD is > of MinStreamedLOD otherwise we would not be using the right LOD->MIP, the only drawback is a possible minor memory usage for overallocating static mesh element batch
 	const int32 MinStreamedLOD = (HeightmapTexture != nullptr && HeightmapTexture->Resource != nullptr) ? FMath::Min<int32>(((FTexture2DResource*)HeightmapTexture->Resource)->GetCurrentFirstMip(), FMath::CeilLogTwo(SubsectionSizeVerts) - 1) : 0;
 	LastLOD = FMath::Max(MinStreamedLOD, LastLOD);
@@ -917,6 +1236,8 @@ void FLandscapeComponentSceneProxy::CreateRenderThreadResources()
 
 	if (IsComponentLevelVisible())
 	{
+		OriginAndSphereRadius = FVector4(GetBounds().Origin, GetBounds().SphereRadius);
+		LandscapeSceneProxy = this;
 		RegisterNeighbors();
 	}
 
@@ -1022,6 +1343,7 @@ void FLandscapeComponentSceneProxy::CreateRenderThreadResources()
 		BatchElementParams->SubX = -1;
 		BatchElementParams->SubY = -1;
 		BatchElementParams->CurrentLOD = 0;
+		BatchElementParams->ForcedLOD = 0;
 		GrassBatchElement->UserData = BatchElementParams;
 		GrassBatchElement->PrimitiveUniformBuffer = GetUniformBuffer();
 		GrassBatchElement->IndexBuffer = SharedBuffers->GrassIndexBuffer;
@@ -1038,7 +1360,6 @@ void FLandscapeComponentSceneProxy::CreateRenderThreadResources()
 			*CollisionBatchElement = *GrassBatchElement;
 			FLandscapeBatchElementParams* CollisionBatchElementParams = &GrassBatchParams[Mip];
 			*CollisionBatchElementParams = *BatchElementParams;
-			CollisionBatchElementParams->CurrentLOD = Mip;
 			CollisionBatchElement->UserData = CollisionBatchElementParams;
 			CollisionBatchElement->NumPrimitives = FMath::Square(NumSubsections) * FMath::Square(MipSubsectionSizeVerts);
 			CollisionBatchElement->FirstIndex = SharedBuffers->GrassIndexMipOffsets[Mip];
@@ -1085,6 +1406,8 @@ void FLandscapeComponentSceneProxy::DestroyRenderThreadResources()
 
 void FLandscapeComponentSceneProxy::OnLevelAddedToWorld()
 {
+	OriginAndSphereRadius = FVector4(GetBounds().Origin, GetBounds().SphereRadius);
+	LandscapeSceneProxy = this;
 	RegisterNeighbors();
 }
 
@@ -1137,7 +1460,6 @@ FPrimitiveViewRelevance FLandscapeComponentSceneProxy::GetViewRelevance(const FS
 	Result.bRenderCustomDepth = ShouldRenderCustomDepth();
 	Result.bUsesLightingChannels = GetLightingChannelMask() != GetDefaultLightingChannelMask();
 	Result.bTranslucentSelfShadow = bCastVolumetricTranslucentShadow;
-	Result.bUseCustomViewData = true;
 
 	auto FeatureLevel = View->GetFeatureLevel();
 
@@ -1236,6 +1558,7 @@ FPrimitiveViewRelevance FLandscapeComponentSceneProxy::GetViewRelevance(const FS
 		(IsSelected() && !GLandscapeEditModeActive) ||
 		(GLandscapeViewMode != ELandscapeViewMode::Normal) ||
 		(CVarLandscapeShowDirty.GetValueOnRenderThread() && GLandscapeDirtyMaterial) ||
+		(View->Family->LandscapeLODOverride >= 0) ||
 #else
 		IsSelected() ||
 #endif
@@ -1358,6 +1681,10 @@ void FLandscapeComponentSceneProxy::OnTransformChanged()
 
 	// Set FLandscapeUniformVSParameters for this subsection
 	FLandscapeUniformShaderParameters LandscapeParams;
+	LandscapeParams.ComponentBaseX = ComponentBase.X;
+	LandscapeParams.ComponentBaseY = ComponentBase.Y;
+	LandscapeParams.SubsectionSizeVerts = SubsectionSizeVerts;
+	LandscapeParams.LastLOD = LastLOD;
 	LandscapeParams.HeightmapUVScaleBias = HeightmapScaleBias;
 	LandscapeParams.WeightmapUVScaleBias = WeightmapScaleBias;
 	LandscapeParams.LocalToWorldNoScaling = LocalToWorldNoScaling;
@@ -1420,6 +1747,14 @@ void FLandscapeComponentSceneProxy::OnTransformChanged()
 	}
 
 	LandscapeUniformShaderParameters.SetContents(LandscapeParams);
+
+	if (bRegistered)
+	{
+		OriginAndSphereRadius = FVector4(GetBounds().Origin, GetBounds().SphereRadius);
+
+		FLandscapeRenderSystem& RenderSystem = *LandscapeRenderSystems.FindChecked(LandscapeKey);
+		RenderSystem.SetSectionOriginAndRadius(ComponentBase, OriginAndSphereRadius);
+	}
 }
 
 float FLandscapeComponentSceneProxy::GetComponentScreenSize(const FSceneView* View, const FVector& Origin, float MaxExtend, float ElementRadius) const
@@ -1614,17 +1949,8 @@ bool FLandscapeComponentSceneProxy::GetMeshElement(bool UseSeperateBatchForShado
 	OutMeshBatch.Type = bCurrentRequiresAdjacencyInformation ? PT_12_ControlPointPatchList : PT_TriangleList;
 	OutMeshBatch.DepthPriorityGroup = SDPG_World;
 	OutMeshBatch.LODIndex = InLODIndex;
-	OutMeshBatch.bRequiresPerElementVisibility = true;
 	OutMeshBatch.bDitheredLODTransition = false;
 
-	if (OutMeshBatch.CastShadow)
-	{
-		float QuadWorldSize = (ComponentMaxExtend * 2.0f) / (SubsectionSizeQuads * NumSubsections);
-		FVector QuadSize = FVector(QuadWorldSize, QuadWorldSize, HasTessellation ? OutMeshBatch.MaterialRenderProxy->GetMaterial(GetScene().GetFeatureLevel())->GetMaxDisplacement() : 0.0f);
-
-		OutMeshBatch.TessellationDisablingShadowMapMeshSize = QuadSize.Size();
-	}
-
 	int32 BatchElementSize = NumSubsections == 1 ? 1 : MAX_SUBSECTION_COUNT + 1;
 	OutMeshBatch.Elements.Empty(FMath::Max(LastLOD - FirstLOD, 1) * BatchElementSize);
 
@@ -1765,118 +2091,110 @@ void FLandscapeComponentSceneProxy::ApplyWorldOffset(FVector InOffset)
 	}
 }
 
-void FLandscapeComponentSceneProxy::DrawStaticElements(FStaticPrimitiveDrawInterface* PDI)
+template<class ArrayType>
+bool FLandscapeComponentSceneProxy::GetStaticMeshElement(int32 LODIndex, bool bForToolMesh, bool bForcedLOD, FMeshBatch& MeshBatch, ArrayType& OutStaticBatchParamArray) const
 {
-	if (AvailableMaterials.Num() == 0)
-	{
-		return;
-	}
-
-	const int32 NumBatchesPerLOD = NumSubsections == 1 ? 1 : MAX_SUBSECTION_COUNT + 1;
-	const int32 TessellatedRequiredBatchCount = 4;
-	const int32 NoTessellatedRequiredBatchCount = 1;
-	int32 MaterialCount = AvailableMaterials.Num();
-	int32 BatchCount = 0;
-	const auto FeatureLevel = GetScene().GetFeatureLevel();
+	UMaterialInterface* MaterialInterface = nullptr;
 
-	for (int32 i = 0; i < AvailableMaterials.Num(); ++i)
 	{
-		ULandscapeMaterialInstanceConstant* LandscapeMIC = Cast<ULandscapeMaterialInstanceConstant>(AvailableMaterials[i]);
+		int32 MaterialIndex = LODIndexToMaterialIndex[LODIndex];
 
-		if (LandscapeMIC == nullptr)
-		{
-			UMaterialInstanceDynamic* LandscapeMID = Cast<UMaterialInstanceDynamic>(AvailableMaterials[i]);
+		// Defaults to the material interface w/ potential tessellation
+		MaterialInterface = AvailableMaterials[MaterialIndex];
 
-			if (LandscapeMID != nullptr)
-			{
-				LandscapeMIC = Cast<ULandscapeMaterialInstanceConstant>(LandscapeMID->Parent);
-			}
-		}
-
-		bool HasTessellationEnabled = MaterialHasTessellationEnabled[i];
-		check(HasTessellationEnabled && LandscapeMIC != nullptr || !HasTessellationEnabled);
-
-		BatchCount += HasTessellationEnabled && (LandscapeMIC != nullptr && !LandscapeMIC->bDisableTessellation) ? TessellatedRequiredBatchCount : NoTessellatedRequiredBatchCount;
-
-		// Remove material with disabled tessellation as they were generated
-		if (HasTessellationEnabled && (LandscapeMIC != nullptr && LandscapeMIC->bDisableTessellation))
+		if (!MaterialInterface)
 		{
-			--MaterialCount;
+			return false;
 		}
-	}		
-
-	int32 TotalBatchCount = (1 + LastLOD - FirstLOD) * NumBatchesPerLOD * BatchCount;
-	TotalBatchCount += (1 + LastVirtualTextureLOD - FirstVirtualTextureLOD) * RuntimeVirtualTextureMaterialTypes.Num();
 
-	StaticBatchParamArray.Empty(TotalBatchCount);
-	PDI->ReserveMemoryForMeshes(TotalBatchCount);
+		UMaterialInstance* MaterialInstance = Cast<UMaterialInstance>(AvailableMaterials[MaterialIndex]);
+		bool HasTessellationEnabled = (GetScene().GetFeatureLevel() >= ERHIFeatureLevel::SM5) ? MaterialInstance != nullptr && MaterialInstance->GetMaterial()->D3D11TessellationMode != EMaterialTessellationMode::MTM_NoTessellation && MaterialIndexToDisabledTessellationMaterial[MaterialIndex] != INDEX_NONE : false;
 
-	MaterialIndexToStaticMeshBatchLOD.SetNumUninitialized(MaterialCount);
-
-	int32 CurrentLODIndex = 0;
-	for (int32 i = 0; i < MaterialCount; ++i)
-	{
-		if (AvailableMaterials[i] == nullptr)
+		if (HasTessellationEnabled)
 		{
-			continue;
-		}
-
-		UMaterialInstance* MaterialInstance = Cast<UMaterialInstance>(AvailableMaterials[i]);
+			{
+				// Sanity check non-tessellated materials
+				UMaterialInstance* NonTessellatedLandscapeMI = Cast<UMaterialInstance>(AvailableMaterials[MaterialIndexToDisabledTessellationMaterial[MaterialIndex]]);
 
-		bool HasTessellationEnabled = (FeatureLevel >= ERHIFeatureLevel::SM5) ? MaterialInstance != nullptr && MaterialInstance->GetMaterial()->D3D11TessellationMode != EMaterialTessellationMode::MTM_NoTessellation && MaterialIndexToDisabledTessellationMaterial[i] != INDEX_NONE : false;
+				// Make sure that the Material instance we are going to use has the tessellation disabled
+				UMaterialInstanceDynamic* NonTessellatedLandscapeMID = Cast<UMaterialInstanceDynamic>(NonTessellatedLandscapeMI);
+				ULandscapeMaterialInstanceConstant* NonTessellatedLandscapeMIC = Cast<ULandscapeMaterialInstanceConstant>(NonTessellatedLandscapeMI);
 
-		// Only add normal batch index (for each material)
-		MaterialIndexToStaticMeshBatchLOD[i] = CurrentLODIndex;
+				if (NonTessellatedLandscapeMID != nullptr)
+				{
+					NonTessellatedLandscapeMIC = Cast<ULandscapeMaterialInstanceConstant>(NonTessellatedLandscapeMID->Parent);
+				}
 
-		// Default Batch, tessellated if enabled
-		FMeshBatch NormalMeshBatch;
+				check(NonTessellatedLandscapeMIC != nullptr && NonTessellatedLandscapeMIC->bDisableTessellation);
+			}
 
-		if (GetMeshElement(false, false, HasTessellationEnabled, CurrentLODIndex++, AvailableMaterials[i], NormalMeshBatch, StaticBatchParamArray))
-		{
-			PDI->DrawMesh(NormalMeshBatch, FLT_MAX);
+			float TessellationLODScreenSizeThreshold = LODIndex == 0 ? FLT_MAX : LODScreenRatioSquared[LODIndex];
+			if (TessellationLODScreenSizeThreshold < TessellationComponentSquaredScreenSize || bForToolMesh)
+			{
+				// Selectively disable tessellation
+				MaterialInterface = AvailableMaterials[MaterialIndexToDisabledTessellationMaterial[MaterialIndex]];
+			}
 		}
+	}
 
-		if (HasTessellationEnabled)
-		{
-			UMaterialInstance* NonTessellatedLandscapeMI = Cast<UMaterialInstance>(AvailableMaterials[MaterialIndexToDisabledTessellationMaterial[i]]);
-
-			// Make sure that the Material instance we are going to use has the tessellation disabled
-			UMaterialInstanceDynamic* NonTessellatedLandscapeMID = Cast<UMaterialInstanceDynamic>(NonTessellatedLandscapeMI);
-			ULandscapeMaterialInstanceConstant* NonTessellatedLandscapeMIC = Cast<ULandscapeMaterialInstanceConstant>(NonTessellatedLandscapeMI);
+	// Based on the final material we selected, detect if it has tessellation
+	// Could be different from bRequiresAdjacencyInformation during shader compilation
+	bool bCurrentRequiresAdjacencyInformation = MaterialRenderingRequiresAdjacencyInformation_RenderingThread(MaterialInterface, VertexFactory->GetType(), GetScene().GetFeatureLevel());
 
-			if (NonTessellatedLandscapeMID != nullptr)
-			{
-				NonTessellatedLandscapeMIC = Cast<ULandscapeMaterialInstanceConstant>(NonTessellatedLandscapeMID->Parent);
-			}
+	check(!bCurrentRequiresAdjacencyInformation || (bCurrentRequiresAdjacencyInformation && SharedBuffers->AdjacencyIndexBuffers));
 
-			check(NonTessellatedLandscapeMIC != nullptr && NonTessellatedLandscapeMIC->bDisableTessellation);
+	{
+		MeshBatch.VertexFactory = VertexFactory;
+		MeshBatch.MaterialRenderProxy = MaterialInterface->GetRenderProxy();
 
-			// Default Batch, no Tessellation enabled
-			FMeshBatch NonTesselatedNormalMeshBatch;
+		MeshBatch.LCI = ComponentLightInfo.Get();
+		MeshBatch.ReverseCulling = IsLocalToWorldDeterminantNegative();
+		MeshBatch.CastShadow = bForToolMesh ? false : true;
+		MeshBatch.bUseForDepthPass = true;
+		MeshBatch.bUseAsOccluder = ShouldUseAsOccluder() && GetScene().GetShadingPath() == EShadingPath::Deferred && !IsMovable();
+		MeshBatch.bUseForMaterial = true;
+		MeshBatch.Type = bCurrentRequiresAdjacencyInformation ? PT_12_ControlPointPatchList : PT_TriangleList;
+		MeshBatch.DepthPriorityGroup = SDPG_World;
+		MeshBatch.LODIndex = LODIndex;
+		MeshBatch.bDitheredLODTransition = false;
 
-			if (GetMeshElement(true, false, false, CurrentLODIndex++, NonTessellatedLandscapeMI, NonTesselatedNormalMeshBatch, StaticBatchParamArray))
-			{
-				PDI->DrawMesh(NonTesselatedNormalMeshBatch, FLT_MAX);
-			}
+		// Combined batch element
+		FMeshBatchElement& BatchElement = MeshBatch.Elements[0];
 
-			// Shadow Batch, tessellated if enabled
-			FMeshBatch ShadowMeshBatch;
+		FLandscapeBatchElementParams* BatchElementParams = new(OutStaticBatchParamArray) FLandscapeBatchElementParams;
+		BatchElementParams->LandscapeUniformShaderParametersResource = &LandscapeUniformShaderParameters;
+		BatchElementParams->LocalToWorldNoScalingPtr = &LocalToWorldNoScaling;
+		BatchElementParams->SceneProxy = this;
+		BatchElementParams->SubX = -1;
+		BatchElementParams->SubY = -1;
+		BatchElementParams->CurrentLOD = LODIndex;
+		BatchElementParams->ForcedLOD = bForcedLOD ? LODIndex : -1;
 
-			if (GetMeshElement(true, true, HasTessellationEnabled, CurrentLODIndex++, MaterialInstance, ShadowMeshBatch, StaticBatchParamArray))
-			{
-				PDI->DrawMesh(ShadowMeshBatch, FLT_MAX);
-			}
+		BatchElement.UserData = BatchElementParams;
+		BatchElement.PrimitiveUniformBuffer = GetUniformBuffer();
+		BatchElement.IndexBuffer = bCurrentRequiresAdjacencyInformation ? SharedBuffers->AdjacencyIndexBuffers->IndexBuffers[LODIndex] : SharedBuffers->IndexBuffers[LODIndex];
+		BatchElement.NumPrimitives = FMath::Square((SubsectionSizeVerts >> LODIndex) - 1) * FMath::Square(NumSubsections) * 2;
+		BatchElement.FirstIndex = 0;
+		BatchElement.MinVertexIndex = SharedBuffers->IndexRanges[LODIndex].MinIndexFull;
+		BatchElement.MaxVertexIndex = SharedBuffers->IndexRanges[LODIndex].MaxIndexFull;
+	}
 
-			// Shadow Batch, no tessellation enabled
-			FMeshBatch NonTesselatedShadowMeshBatch;
+	return true;
+}
 
-			if (GetMeshElement(true, true, false, CurrentLODIndex++, NonTessellatedLandscapeMI, NonTesselatedShadowMeshBatch, StaticBatchParamArray))
-			{
-				PDI->DrawMesh(NonTesselatedShadowMeshBatch, FLT_MAX);
-			}
-		}
+void FLandscapeComponentSceneProxy::DrawStaticElements(FStaticPrimitiveDrawInterface* PDI)
+{
+	if (AvailableMaterials.Num() == 0)
+	{
+		return;
 	}
 
+	int32 TotalBatchCount = 1 + LastLOD - FirstLOD;
+	TotalBatchCount += (1 + LastVirtualTextureLOD - FirstVirtualTextureLOD) * RuntimeVirtualTextureMaterialTypes.Num();
+
+	StaticBatchParamArray.Empty(TotalBatchCount);
+	PDI->ReserveMemoryForMeshes(TotalBatchCount);
+
 	// Add fixed grid mesh batches for runtime virtual texture usage
  	for (ERuntimeVirtualTextureMaterialType MaterialType : RuntimeVirtualTextureMaterialTypes)
  	{
@@ -1892,6 +2210,16 @@ void FLandscapeComponentSceneProxy::DrawStaticElements(FStaticPrimitiveDrawInter
 		}
  	}
 
+	for (int32 LODIndex = FirstLOD; LODIndex <= LastLOD; LODIndex++)
+	{
+		FMeshBatch MeshBatch;
+
+		if (GetStaticMeshElement(LODIndex, false, false, MeshBatch, StaticBatchParamArray))
+		{
+			PDI->DrawMesh(MeshBatch, LODIndex == FirstLOD ? FLT_MAX : (FMath::Sqrt(LODScreenRatioSquared[LODIndex]) * 2.0f));
+		}
+	}
+
 	check(StaticBatchParamArray.Num() <= TotalBatchCount);
 }
 
@@ -2149,6 +2477,11 @@ float FLandscapeComponentSceneProxy::ComputeBatchElementCurrentLOD(int32 InSelec
 
 int8 FLandscapeComponentSceneProxy::GetLODFromScreenSize(float InScreenSizeSquared, float InViewLODScale) const
 {
+	float FractionalLOD;
+
+	return FLandscapeRenderSystem::GetLODFromScreenSize(LODSettings, InScreenSizeSquared, InViewLODScale, FractionalLOD);
+
+#if 0
 	int32 LODScreenRatioSquaredCount = LODScreenRatioSquared.Num();
 	float ScreenSizeSquared = InScreenSizeSquared / InViewLODScale;
 
@@ -2179,6 +2512,7 @@ int8 FLandscapeComponentSceneProxy::GetLODFromScreenSize(float InScreenSizeSquar
 	}
 
 	return INDEX_NONE;
+#endif
 }
 
 void* FLandscapeComponentSceneProxy::InitViewCustomData(const FSceneView& InView, float InViewLODScale, FMemStackBase& InCustomDataMemStack, bool InIsStaticRelevant, bool InIsShadowOnly, const FLODMask* InVisiblePrimitiveLODMask, float InMeshScreenSizeSquared)
@@ -2494,18 +2828,9 @@ void FLandscapeComponentSceneProxy::PostInitViewCustomData(const FSceneView& InV
 #endif
 }
 
-bool FLandscapeComponentSceneProxy::IsUsingCustomLODRules() const
-{
-	return true;
-}
-
-bool FLandscapeComponentSceneProxy::IsUsingCustomWholeSceneShadowLODRules() const
-{
-	return true;
-}
-
 bool FLandscapeComponentSceneProxy::CanUseMeshBatchForShadowCascade(int8 InLODIndex, float InShadowMapTextureResolution, float InShadowMapCascadeSize) const
 {
+#if 0
 	const FStaticMeshBatch* MeshBatch = nullptr;
 	const TArray<FStaticMeshBatch>& PrimitiveStaticMeshes = GetPrimitiveSceneInfo()->StaticMeshes;
 
@@ -2527,6 +2852,9 @@ bool FLandscapeComponentSceneProxy::CanUseMeshBatchForShadowCascade(int8 InLODIn
 
 	float WorldUnitsForOneTexel = InShadowMapCascadeSize / InShadowMapTextureResolution; // We assume Shadow Map texture to be squared
 	return MeshBatch->TessellationDisablingShadowMapMeshSize >= WorldUnitsForOneTexel * (GShadowMapWorldUnitsToTexelFactor != -1.0f ? GShadowMapWorldUnitsToTexelFactor : 1.0f);
+#endif
+
+	return true;
 }
 
 FLODMask FLandscapeComponentSceneProxy::GetCustomLOD(const FSceneView& InView, float InViewLODScale, int32 InForcedLODLevel, float& OutScreenSizeSquared) const
@@ -2749,63 +3077,24 @@ void FLandscapeComponentSceneProxy::GetDynamicMeshElements(const TArray<const FS
 			const FSceneView* View = Views[ViewIndex];
 			FLandscapeElementParamArray& ParameterArray = Collector.AllocateOneFrameResource<FLandscapeElementParamArray>();
 
-			const FViewCustomDataLOD* PrimitiveCustomData = (const FViewCustomDataLOD*)View->GetCustomData(GetPrimitiveSceneInfo()->GetIndex());
-
-			if (PrimitiveCustomData == nullptr)
-			{
-				continue;
-			}
-
-			ParameterArray.ElementParams.AddDefaulted(NumSubsections * NumSubsections);
-
-			const auto FeatureLevel = View->GetFeatureLevel();
-			bool HasTessellationEnabled = false;
-
-			if (FeatureLevel >= ERHIFeatureLevel::SM5)
-			{
-				const int8 CurrentLODIndex = PrimitiveCustomData->SubSections[0].BatchElementCurrentLOD;
-				ULandscapeMaterialInstanceConstant* LandscapeMIC = Cast<ULandscapeMaterialInstanceConstant>(AvailableMaterials[LODIndexToMaterialIndex[CurrentLODIndex]]);
-
-				if (LandscapeMIC == nullptr)
-				{
-					UMaterialInstanceDynamic* LandscapeMID = Cast<UMaterialInstanceDynamic>(AvailableMaterials[LODIndexToMaterialIndex[CurrentLODIndex]]);
-
-					if (LandscapeMID != nullptr)
-					{
-						LandscapeMIC = Cast<ULandscapeMaterialInstanceConstant>(LandscapeMID->Parent);
-					}
-				}
-
-				HasTessellationEnabled = MaterialHasTessellationEnabled[LODIndexToMaterialIndex[CurrentLODIndex]] && (LandscapeMIC != nullptr && !LandscapeMIC->bDisableTessellation) && AvailableMaterials.Num() > 1;
-			}
-
-			bool DisableTessellation = false;
-
-			if (HasTessellationEnabled)
-			{
-				check(AvailableMaterials.Num() > 1);
-				float ScreenSizeSquared = GetComponentScreenSize(View, GetBounds().Origin, ComponentMaxExtend, GetBounds().SphereRadius);
-				DisableTessellation = (ScreenSizeSquared < TessellationComponentSquaredScreenSize * View->LODDistanceFactor) ? true : false;
-
-				if (!DisableTessellation)
-				{
-					INC_DWORD_STAT(STAT_LandscapeTessellatedComponents);
-				}
-			}
+			float MeshScreenSizeSquared = ComputeBoundsScreenRadiusSquared(GetBounds().Origin, GetBounds().SphereRadius, *View);
+			int32 ForcedLODLevel = (View->Family->EngineShowFlags.LOD) ? GetCVarForceLOD() : -1;
+#if WITH_EDITOR
+			ForcedLODLevel = View->Family->LandscapeLODOverride >= 0 ? View->Family->LandscapeLODOverride : ForcedLODLevel;
+#endif
+			int32 LODToRender = ForcedLODLevel >= 0 ? ForcedLODLevel : GetLODFromScreenSize(MeshScreenSizeSquared, View->LODDistanceFactor);
 
-			if (!PrimitiveCustomData->UseCombinedMeshBatch)
-			{
-				INC_DWORD_STAT(STAT_LandscapeComponentUsingSubSectionDrawCalls);
-			}
+			ParameterArray.ElementParams.AddDefaulted(1);
 
 			FMeshBatch& Mesh = Collector.AllocateMesh();
-			BuildDynamicMeshElement(PrimitiveCustomData, false, HasTessellationEnabled, DisableTessellation, Mesh, ParameterArray.ElementParams);
+
+			GetStaticMeshElement(LODToRender, false, ForcedLODLevel >= 0, Mesh, ParameterArray.ElementParams);
 
 #if WITH_EDITOR
 			FMeshBatch& MeshTools = Collector.AllocateMesh();
 
 			// No Tessellation on tool material
-			BuildDynamicMeshElement(PrimitiveCustomData, true, HasTessellationEnabled, true, MeshTools, ParameterArray.ElementParams);
+			GetStaticMeshElement(LODToRender, true, ForcedLODLevel >= 0, MeshTools, ParameterArray.ElementParams);
 #endif
 
 			// Render the landscape component
@@ -2891,7 +3180,7 @@ void FLandscapeComponentSceneProxy::GetDynamicMeshElements(const TArray<const FS
 					LODMesh.Elements.Add(TemplateMesh.Elements[i]);
 					int32 CurrentLOD = ((FLandscapeBatchElementParams*)TemplateMesh.Elements[i].UserData)->CurrentLOD;
 					LODMesh.VisualizeLODIndex = CurrentLOD;
-					FLinearColor Color = GetColorForLod(CurrentLOD, ForcedLOD, PrimitiveCustomData != nullptr ? PrimitiveCustomData->UseCombinedMeshBatch : true);
+					FLinearColor Color = GetColorForLod(CurrentLOD, ForcedLOD, true);
 					FMaterialRenderProxy* LODMaterialProxy = (FMaterialRenderProxy*)new FColoredMaterialRenderProxy(GEngine->LevelColorationUnlitMaterial->GetRenderProxy(), Color);
 					Collector.RegisterOneFrameMaterialProxy(LODMaterialProxy);
 					LODMesh.MaterialRenderProxy = LODMaterialProxy;
@@ -3007,6 +3296,7 @@ void FLandscapeComponentSceneProxy::GetDynamicMeshElements(const TArray<const FS
 						bIsWireframe ||
 #if WITH_EDITOR
 						(IsSelected() && !GLandscapeEditModeActive) ||
+						ViewFamily.LandscapeLODOverride >= 0 ||
 #else
 						IsSelected() ||
 #endif
@@ -3137,17 +3427,10 @@ void FLandscapeComponentSceneProxy::GetDynamicRayTracingInstances(FRayTracingMat
 	{
 		return;
 	}
-
-	FMemStackBase& PrimitiveCustomDataMemStack = FMemStack::Get();
-
-	FViewCustomDataLOD* InPrimitiveCustomData;
-
-	float MeshScreenSizeSquared = 0;
+	float MeshScreenSizeSquared = ComputeBoundsScreenRadiusSquared(GetBounds().Origin, GetBounds().SphereRadius, *Context.ReferenceView);
 	int32 ForcedLODLevel = (Context.ReferenceView->Family->EngineShowFlags.LOD) ? GetCVarForceLOD() : 0;
 
-	FLODMask LODToRender = GetCustomLOD(*Context.ReferenceView, Context.ReferenceView->LODDistanceFactor, ForcedLODLevel, MeshScreenSizeSquared);
-	InPrimitiveCustomData = (FViewCustomDataLOD*)InitViewCustomData(*Context.ReferenceView, Context.ReferenceView->LODDistanceFactor, PrimitiveCustomDataMemStack, false, false, &LODToRender, MeshScreenSizeSquared);
-	PostInitViewCustomData(*Context.ReferenceView, InPrimitiveCustomData);
+	int32 LODToRender = ForcedLODLevel >= 0 ? ForcedLODLevel : GetLODFromScreenSize(MeshScreenSizeSquared, Context.ReferenceView->LODDistanceFactor);
 	
 	FLandscapeElementParamArray& ParameterArray = Context.RayTracingMeshResourceCollector.AllocateOneFrameResource<FLandscapeElementParamArray>();
 	ParameterArray.ElementParams.AddDefaulted(NumSubsections * NumSubsections);
@@ -3157,7 +3440,7 @@ void FLandscapeComponentSceneProxy::GetDynamicRayTracingInstances(FRayTracingMat
 		return;
 	}
 
-	const int8 CurrentLODIndex = InPrimitiveCustomData->SubSections[0].BatchElementCurrentLOD;
+	const int8 CurrentLODIndex = LODToRender;
 	int8 MaterialIndex = LODIndexToMaterialIndex.IsValidIndex(CurrentLODIndex) ? LODIndexToMaterialIndex[CurrentLODIndex] : INDEX_NONE;
 	UMaterialInterface* SelectedMaterial = MaterialIndex != INDEX_NONE ? AvailableMaterials[MaterialIndex] : nullptr;
 
@@ -3178,12 +3461,14 @@ void FLandscapeComponentSceneProxy::GetDynamicRayTracingInstances(FRayTracingMat
 
 	BaseMeshBatch.Elements.Empty();
 
+	FLandscapeRenderSystem& RenderSystem = *LandscapeRenderSystems.FindChecked(LandscapeKey);
+
 	for (int32 SubY = 0; SubY < NumSubsections; SubY++)
 	{
 		for (int32 SubX = 0; SubX < NumSubsections; SubX++)
 		{
 			const int8 SubSectionIdx = SubX + SubY * NumSubsections;
-			const int8 CurrentLOD = InPrimitiveCustomData->SubSections[SubSectionIdx].BatchElementCurrentLOD;
+			const int8 CurrentLOD = LODToRender;
 
 			FMeshBatch MeshBatch = BaseMeshBatch;
 
@@ -3196,6 +3481,7 @@ void FLandscapeComponentSceneProxy::GetDynamicRayTracingInstances(FRayTracingMat
 			BatchElementParams.SubX = SubX;
 			BatchElementParams.SubY = SubY;
 			BatchElementParams.CurrentLOD = CurrentLOD;
+			BatchElementParams.ForcedLOD = ForcedLODLevel;
 			BatchElement.UserData = &BatchElementParams;
 			BatchElement.PrimitiveUniformBuffer = GetUniformBuffer();
 
@@ -3233,32 +3519,16 @@ void FLandscapeComponentSceneProxy::GetDynamicRayTracingInstances(FRayTracingMat
 					SectionRayTracingStates[SubSectionIdx].CurrentLOD = CurrentLOD;
 					SectionRayTracingStates[SubSectionIdx].RayTracingDynamicVertexBuffer.Release();
 				}
-				if (!SectionRayTracingStates[SubSectionIdx].LODBias.Equals(GetShaderLODBias()))
+				if (SectionRayTracingStates[SubSectionIdx].HeightmapLODBias != RenderSystem.GetSectionLODBias(ComponentBase))
 				{
 					bNeedsRayTracingGeometryUpdate = true;
-					SectionRayTracingStates[SubSectionIdx].LODBias = GetShaderLODBias();
+					SectionRayTracingStates[SubSectionIdx].HeightmapLODBias = RenderSystem.GetSectionLODBias(ComponentBase);
 				}
 
-				float ComponentScreenSize = GetComponentScreenSize(Context.ReferenceView, GetBounds().Origin, ComponentMaxExtend, GetBounds().SphereRadius);
-
-				FViewCustomDataLOD CurrentLODData;
-				CalculateBatchElementLOD(*Context.ReferenceView, ComponentScreenSize, Context.ReferenceView->LODDistanceFactor, CurrentLODData, true);
-
-				if (SectionRayTracingStates[SubSectionIdx].SectionLOD.Equals(CurrentLODData.ShaderCurrentLOD))
+				if (SectionRayTracingStates[SubSectionIdx].FractionalLOD != RenderSystem.GetSectionLODValue(ComponentBase))
 				{
 					bNeedsRayTracingGeometryUpdate = true;
-					SectionRayTracingStates[SubSectionIdx].SectionLOD = CurrentLODData.ShaderCurrentLOD;
-				}
-
-				{
-					FVector4 CurrentNeighborLOD;
-					GetShaderCurrentNeighborLOD(*Context.ReferenceView, CurrentLODData.SubSections[SubSectionIdx].fBatchElementCurrentLOD, NumSubsections > 1 ? SubX : INDEX_NONE, NumSubsections > 1 ? SubY : INDEX_NONE, SubSectionIdx, CurrentNeighborLOD);
-
-					if (!SectionRayTracingStates[SubSectionIdx].CurrentNeighborLOD.Equals(CurrentNeighborLOD))
-					{
-						bNeedsRayTracingGeometryUpdate = true;
-						SectionRayTracingStates[SubSectionIdx].CurrentNeighborLOD = CurrentNeighborLOD;
-					}
+					SectionRayTracingStates[SubSectionIdx].FractionalLOD = RenderSystem.GetSectionLODValue(ComponentBase);
 				}
 			}
 
@@ -3896,6 +4166,7 @@ public:
 	{
 		LodBiasParameter.Bind(ParameterMap, TEXT("LodBias"));
 		LodValuesParameter.Bind(ParameterMap, TEXT("LodValues"));
+		ForcedLodParameter.Bind(ParameterMap, TEXT("ForcedLod"));
 		LodTessellationParameter.Bind(ParameterMap, TEXT("LodTessellationParams"));
 		SectionLodsParameter.Bind(ParameterMap, TEXT("SectionLods"));
 		NeighborSectionLodParameter.Bind(ParameterMap, TEXT("NeighborSectionLod"));
@@ -3908,6 +4179,7 @@ public:
 	virtual void Serialize(FArchive& Ar) override
 	{
 		Ar << LodValuesParameter;
+		Ar << ForcedLodParameter;
 		Ar << LodTessellationParameter;
 		Ar << NeighborSectionLodParameter;
 		Ar << LodBiasParameter;
@@ -3935,6 +4207,15 @@ public:
 
 		ShaderBindings.Add(Shader->GetUniformBufferParameter<FLandscapeUniformShaderParameters>(), *BatchElementParams->LandscapeUniformShaderParametersResource);
 
+		if (SceneProxy->bRegistered)
+		{
+			ShaderBindings.Add(Shader->GetUniformBufferParameter<FLandscapeSectionLODUniformParameters>(), LandscapeRenderSystems.FindChecked(SceneProxy->LandscapeKey)->UniformBuffer);
+		}
+		else
+		{
+			ShaderBindings.Add(Shader->GetUniformBufferParameter<FLandscapeSectionLODUniformParameters>(), GNullLandscapeRenderSystemResources.UniformBuffer);
+		}
+
 #if RHI_RAYTRACING
 		if (IsRayTracingEnabled())
 		{
@@ -3947,8 +4228,12 @@ public:
 			ShaderBindings.Add(LodValuesParameter, SceneProxy->GetShaderLODValues(BatchElementParams->CurrentLOD));
 		}
 
-		int32 SubSectionIndex = BatchElementParams->SubX + BatchElementParams->SubY * SceneProxy->NumSubsections;
+		if (ForcedLodParameter.IsBound())
+		{
+			ShaderBindings.Add(ForcedLodParameter, BatchElementParams->ForcedLOD);
+		}
 
+		int32 SubSectionIndex = BatchElementParams->SubX + BatchElementParams->SubY * SceneProxy->NumSubsections;
 		// If we have no custom data for this primitive we will compute of the fly the proper values, this will happen if the shader is not used for normal landscape rendering(i.e grass rendering)
 		FLandscapeComponentSceneProxy::FViewCustomDataLOD* LODData = (FLandscapeComponentSceneProxy::FViewCustomDataLOD*)InView->GetCustomData(SceneProxy->GetPrimitiveSceneInfo()->GetIndex());
 
@@ -4062,6 +4347,7 @@ public:
 protected:
 	FShaderParameter LodTessellationParameter;
 	FShaderParameter LodValuesParameter;
+	FShaderParameter ForcedLodParameter;
 	FShaderParameter NeighborSectionLodParameter;
 	FShaderParameter LodBiasParameter;
 	FShaderParameter SectionLodsParameter;
@@ -4105,6 +4391,11 @@ public:
 		{
 			ShaderBindings.Add(LodBiasParameter, FVector4(ForceInitToZero));
 		}
+
+		if (ForcedLodParameter.IsBound())
+		{
+			ShaderBindings.Add(ForcedLodParameter, BatchElementParams->ForcedLOD);
+		}
 	}
 };
 
@@ -4179,7 +4470,7 @@ void FLandscapeVertexFactory::ModifyCompilationEnvironment(const FVertexFactoryT
 }
 
 
-IMPLEMENT_VERTEX_FACTORY_TYPE(FLandscapeVertexFactory, "/Engine/Private/LandscapeVertexFactory.ush", true, true, true, false, false);
+IMPLEMENT_VERTEX_FACTORY_TYPE_EX(FLandscapeVertexFactory, "/Engine/Private/LandscapeVertexFactory.ush", true, true, true, false, false, true, false);
 
 /**
 * Copy the data from another vertex factory
@@ -5031,8 +5322,16 @@ void FLandscapeComponentSceneProxy::GetLCIs(FLCIArray& LCIs)
 //
 void FLandscapeNeighborInfo::RegisterNeighbors()
 {
+	check(IsInRenderingThread());
 	if (!bRegistered)
 	{
+		if (!SharedSceneProxyMap.Find(LandscapeKey))
+		{
+			LandscapeRenderSystems.Add(LandscapeKey, new FLandscapeRenderSystem {});
+
+			GetRendererModule().RegisterPersistentViewUniformBufferExtension(&LandscapePersistentViewUniformBufferExtension);
+		}
+
 		// Register ourselves in the map.
 		TMap<FIntPoint, const FLandscapeNeighborInfo*>& SceneProxyMap = SharedSceneProxyMap.FindOrAdd(LandscapeKey);
 
@@ -5065,6 +5364,45 @@ void FLandscapeNeighborInfo::RegisterNeighbors()
 			{
 				Neighbors[3]->Neighbors[0] = this;
 			}
+
+			{
+				FLandscapeRenderSystem& RenderSystem = *LandscapeRenderSystems.FindChecked(LandscapeKey);
+
+				if (RenderSystem.NumRegisteredEntities > 0)
+				{
+					FIntPoint OriginalMin = RenderSystem.Min;
+					FIntPoint OriginalMax = RenderSystem.Min + RenderSystem.Size - FIntPoint(1, 1);
+					FIntPoint NewMin(FMath::Min(RenderSystem.Min.X, ComponentBase.X), FMath::Min(RenderSystem.Min.Y, ComponentBase.Y));
+					FIntPoint NewMax(FMath::Max(OriginalMax.X, ComponentBase.X), FMath::Max(OriginalMax.Y, ComponentBase.Y));
+
+					FIntPoint SizeRequired = (NewMax - NewMin) + FIntPoint(1, 1);
+
+					if (NewMin != RenderSystem.Min || RenderSystem.Size != SizeRequired)
+					{
+						RenderSystem.ResizeAndMoveTo(NewMin, SizeRequired);
+						RenderSystem.RecreateBuffers();
+					}
+
+					// Validate system-wide global parameters
+					check(RenderSystem.TessellationFalloffSettings.UseTessellationComponentScreenSizeFalloff == LandscapeSceneProxy->UseTessellationComponentScreenSizeFalloff);
+					check(RenderSystem.TessellationFalloffSettings.TessellationComponentSquaredScreenSize == LandscapeSceneProxy->TessellationComponentSquaredScreenSize);
+					check(RenderSystem.TessellationFalloffSettings.TessellationComponentScreenSizeFalloff == LandscapeSceneProxy->TessellationComponentScreenSizeFalloff);
+				}
+				else
+				{
+					RenderSystem.TessellationFalloffSettings.UseTessellationComponentScreenSizeFalloff = LandscapeSceneProxy->UseTessellationComponentScreenSizeFalloff;
+					RenderSystem.TessellationFalloffSettings.TessellationComponentSquaredScreenSize = LandscapeSceneProxy->TessellationComponentSquaredScreenSize;
+					RenderSystem.TessellationFalloffSettings.TessellationComponentScreenSizeFalloff = LandscapeSceneProxy->TessellationComponentScreenSizeFalloff;
+
+					RenderSystem.ResizeAndMoveTo(ComponentBase, FIntPoint(1, 1));
+					RenderSystem.RecreateBuffers();
+				}
+
+				RenderSystem.NumRegisteredEntities++;
+				RenderSystem.SetSectionLODSettings(ComponentBase, LODSettings);
+				RenderSystem.SetSectionOriginAndRadius(ComponentBase, OriginAndSphereRadius);
+				RenderSystem.SetSceneProxy(ComponentBase, LandscapeSceneProxy);
+			}
 		}
 		else
 		{
@@ -5075,6 +5413,8 @@ void FLandscapeNeighborInfo::RegisterNeighbors()
 
 void FLandscapeNeighborInfo::UnregisterNeighbors()
 {
+	check(IsInRenderingThread());
+	
 	if (bRegistered)
 	{
 		// Remove ourselves from the map
@@ -5090,9 +5430,18 @@ void FLandscapeNeighborInfo::UnregisterNeighbors()
 			{
 				// remove the entire LandscapeKey entry as this is the last scene proxy
 				SharedSceneProxyMap.Remove(LandscapeKey);
+
+				FLandscapeRenderSystem* RenderSystem = LandscapeRenderSystems.FindChecked(LandscapeKey);
+				delete RenderSystem;
+				LandscapeRenderSystems.Remove(LandscapeKey);
 			}
 			else
 			{
+				FLandscapeRenderSystem* RenderSystem = LandscapeRenderSystems.FindChecked(LandscapeKey);
+				RenderSystem->SetSceneProxy(ComponentBase, nullptr);
+				RenderSystem->SetSectionOriginAndRadius(ComponentBase, FVector4(ForceInitToZero));
+				RenderSystem->NumRegisteredEntities--;
+
 				// remove reference to us from our neighbors
 				if (Neighbors[0])
 				{
diff --git a/Engine/Source/Runtime/Landscape/Private/LandscapeRenderMobile.cpp b/Engine/Source/Runtime/Landscape/Private/LandscapeRenderMobile.cpp
--- a/Engine/Source/Runtime/Landscape/Private/LandscapeRenderMobile.cpp
+++ b/Engine/Source/Runtime/Landscape/Private/LandscapeRenderMobile.cpp
@@ -45,6 +45,7 @@ public:
 	virtual void Bind(const FShaderParameterMap& ParameterMap) override
 	{
 		LodValuesParameter.Bind(ParameterMap,TEXT("LodValues"));
+		ForcedLodParameter.Bind(ParameterMap,TEXT("ForcedLod"));
 		LodTessellationParameter.Bind(ParameterMap, TEXT("LodTessellationParams"));
 		NeighborSectionLodParameter.Bind(ParameterMap,TEXT("NeighborSectionLod"));
 		LodBiasParameter.Bind(ParameterMap,TEXT("LodBias"));
@@ -58,6 +59,7 @@ public:
 	virtual void Serialize(FArchive& Ar) override
 	{
 		Ar << LodValuesParameter;
+		Ar << ForcedLodParameter;
 		Ar << LodTessellationParameter;
 		Ar << NeighborSectionLodParameter;
 		Ar << LodBiasParameter;
@@ -108,63 +110,23 @@ public:
 			ShaderBindings.Add(LodBiasParameter, LodBias);
 		}
 
-		FLandscapeComponentSceneProxy::FViewCustomDataLOD* LODData = (FLandscapeComponentSceneProxy::FViewCustomDataLOD*)InView->GetCustomData(SceneProxy->GetPrimitiveSceneInfo()->GetIndex());
-		int32 SubSectionIndex = BatchElementParams->SubX + BatchElementParams->SubY * SceneProxy->NumSubsections;
-
-		if (LODData != nullptr)
+		if (SceneProxy->bRegistered)
 		{
-			SceneProxy->PostInitViewCustomData(*InView, LODData);
-
-			if (LodTessellationParameter.IsBound())
-			{
-				ShaderBindings.Add(LodTessellationParameter, LODData->LodTessellationParams);
-			}
-
-			if (SectionLodsParameter.IsBound())
-			{
-				if (LODData->UseCombinedMeshBatch)
-				{
-					ShaderBindings.Add(SectionLodsParameter, LODData->ShaderCurrentLOD);
-				}
-				else // in non combined, only the one representing us as we'll be called 4 times (once per sub section)
-				{
-					check(SubSectionIndex >= 0);
-					FVector4 ShaderCurrentLOD(ForceInitToZero);
-					ShaderCurrentLOD.Component(SubSectionIndex) = LODData->ShaderCurrentLOD.Component(SubSectionIndex);
-
-					ShaderBindings.Add(SectionLodsParameter, ShaderCurrentLOD);
-				}
-			}
-
-			if (NeighborSectionLodParameter.IsBound())
-			{
-				FVector4 ShaderCurrentNeighborLOD[FLandscapeComponentSceneProxy::NEIGHBOR_COUNT] = { FVector4(ForceInitToZero), FVector4(ForceInitToZero), FVector4(ForceInitToZero), FVector4(ForceInitToZero) };
-
-				if (LODData->UseCombinedMeshBatch)
-				{
-					int32 SubSectionCount = SceneProxy->NumSubsections == 1 ? 1 : FLandscapeComponentSceneProxy::MAX_SUBSECTION_COUNT;
+			ShaderBindings.Add(Shader->GetUniformBufferParameter<FLandscapeSectionLODUniformParameters>(), LandscapeRenderSystems.FindChecked(SceneProxy->LandscapeKey)->UniformBuffer);
+		}
+		else
+		{
+			ShaderBindings.Add(Shader->GetUniformBufferParameter<FLandscapeSectionLODUniformParameters>(), GNullLandscapeRenderSystemResources.UniformBuffer);
+		}
 
-					for (int32 NeighborSubSectionIndex = 0; NeighborSubSectionIndex < SubSectionCount; ++NeighborSubSectionIndex)
-					{
-						ShaderCurrentNeighborLOD[NeighborSubSectionIndex] = LODData->SubSections[NeighborSubSectionIndex].ShaderCurrentNeighborLOD;
-						check(ShaderCurrentNeighborLOD[NeighborSubSectionIndex].X != -1.0f); // they should all match so only check the 1st one for simplicity
-					}
-
-					ShaderBindings.Add(NeighborSectionLodParameter, ShaderCurrentNeighborLOD);
-				}
-				else // in non combined, only the one representing us as we'll be called 4 times (once per sub section)
-				{
-					check(SubSectionIndex >= 0);
-					ShaderCurrentNeighborLOD[SubSectionIndex] = LODData->SubSections[SubSectionIndex].ShaderCurrentNeighborLOD;
-					check(ShaderCurrentNeighborLOD[SubSectionIndex].X != -1.0f); // they should all match so only check the 1st one for simplicity
-
-					ShaderBindings.Add(NeighborSectionLodParameter, ShaderCurrentNeighborLOD);
-				}
-			}
+		if (ForcedLodParameter.IsBound())
+		{
+			ShaderBindings.Add(ForcedLodParameter, BatchElementParams->ForcedLOD);
 		}
 	}
 protected:
 	FShaderParameter LodValuesParameter;
+	FShaderParameter ForcedLodParameter;
 	FShaderParameter LodTessellationParameter;
 	FShaderParameter NeighborSectionLodParameter;
 	FShaderParameter LodBiasParameter;
@@ -383,6 +345,8 @@ void FLandscapeComponentSceneProxyMobile::CreateRenderThreadResources()
 
 	if (IsComponentLevelVisible())
 	{
+		OriginAndSphereRadius = FVector4(GetBounds().Origin, GetBounds().SphereRadius);
+		LandscapeSceneProxy = this;
 		RegisterNeighbors();
 	}
 	
diff --git a/Engine/Source/Runtime/Landscape/Private/LandscapeVersion.h b/Engine/Source/Runtime/Landscape/Private/LandscapeVersion.h
--- a/Engine/Source/Runtime/Landscape/Private/LandscapeVersion.h
+++ b/Engine/Source/Runtime/Landscape/Private/LandscapeVersion.h
@@ -32,6 +32,8 @@ namespace FLandscapeCustomVersion
 		AddSplineLayerFalloff,
 		// Spline Layer Width
 		AddSplineLayerWidth,
+		// New LOD distribution and tessellation calculations are introduced, needs parameter conversion
+		NewLandscapeContinuousLOD,
 		// -----<new versions can be added above this line>-------------------------------------------------
 		VersionPlusOne,
 		LatestVersion = VersionPlusOne - 1
diff --git a/Engine/Source/Runtime/Landscape/Public/LandscapeRender.h b/Engine/Source/Runtime/Landscape/Public/LandscapeRender.h
--- a/Engine/Source/Runtime/Landscape/Public/LandscapeRender.h
+++ b/Engine/Source/Runtime/Landscape/Public/LandscapeRender.h
@@ -89,6 +89,10 @@ LANDSCAPE_API extern UMaterialInterface* GLandscapeDirtyMaterial;
 /** The uniform shader parameters for a landscape draw call. */
 BEGIN_GLOBAL_SHADER_PARAMETER_STRUCT(FLandscapeUniformShaderParameters, LANDSCAPE_API)
 /** vertex shader parameters */
+SHADER_PARAMETER(int32, ComponentBaseX)
+SHADER_PARAMETER(int32, ComponentBaseY)
+SHADER_PARAMETER(int32, SubsectionSizeVerts)
+SHADER_PARAMETER(int32, LastLOD)
 SHADER_PARAMETER(FVector4, HeightmapUVScaleBias)
 SHADER_PARAMETER(FVector4, WeightmapUVScaleBias)
 SHADER_PARAMETER(FVector4, LandscapeLightmapScaleBias)
@@ -111,6 +115,15 @@ END_GLOBAL_SHADER_PARAMETER_STRUCT()
 
 typedef TUniformBufferRef<FLandscapeVertexFactoryMVFParameters> FLandscapeVertexFactoryMVFUniformBufferRef;
 
+BEGIN_GLOBAL_SHADER_PARAMETER_STRUCT(FLandscapeSectionLODUniformParameters, )
+	SHADER_PARAMETER(FIntPoint, Min)
+	SHADER_PARAMETER(FIntPoint, Size)
+	SHADER_PARAMETER_SRV(Buffer<float>, SectionLOD)
+	SHADER_PARAMETER_SRV(Buffer<float>, SectionLODBias)
+	SHADER_PARAMETER_SRV(Buffer<float>, SectionTessellationFalloffC)
+	SHADER_PARAMETER_SRV(Buffer<float>, SectionTessellationFalloffK)
+END_GLOBAL_SHADER_PARAMETER_STRUCT()
+
 /* Data needed for the landscape vertex factory to set the render state for an individual batch element */
 struct FLandscapeBatchElementParams
 {
@@ -125,6 +138,7 @@ struct FLandscapeBatchElementParams
 	int32 SubX;
 	int32 SubY;
 	int32 CurrentLOD;
+	int32 ForcedLOD = -1;
 };
 
 class FLandscapeElementParamArray : public FOneFrameResource
@@ -366,9 +380,19 @@ public:
 //
 // FLandscapeNeighborInfo
 //
+
+struct LODSettingsComponent
+{
+	float LOD0ScreenSizeSquared;
+	float LOD1ScreenSizeSquared;
+	float LODOnePlusDistributionScalarSquared;
+	float LastLODScreenSizeSquared;
+	int8 LastLODIndex;
+};
+
 class FLandscapeNeighborInfo
 {
-protected:
+public:
 	static const int8 NEIGHBOR_COUNT = 4;
 
 	// Key to uniquely identify the landscape to find the correct render proxy map
@@ -403,6 +427,10 @@ protected:
 		return nullptr;
 	}
 
+	UTexture2D*				HeightmapTexture; // PC : Heightmap, Mobile : Weightmap
+
+protected:
+
 	virtual const ULandscapeComponent* GetLandscapeComponent() const { return nullptr; }
 
 	// Map of currently registered landscape proxies, used to register with our neighbors
@@ -417,23 +445,27 @@ protected:
 
 	
 	// Data we need to be able to access about our neighbor
-	UTexture2D*				HeightmapTexture; // PC : Heightmap, Mobile : Weightmap
 	int8					ForcedLOD;
 	int8					LODBias;
 	bool					bRegistered;
 	int32					PrimitiveCustomDataIndex;
 
+	LODSettingsComponent LODSettings;
+	FVector4 OriginAndSphereRadius;
+	FLandscapeComponentSceneProxy* LandscapeSceneProxy;
+
 	friend class FLandscapeComponentSceneProxy;
 
 public:
 	FLandscapeNeighborInfo(const UWorld* InWorld, const FGuid& InGuid, const FIntPoint& InComponentBase, UTexture2D* InHeightmapTexture, int8 InForcedLOD, int8 InLODBias)
-	: LandscapeKey(InWorld, InGuid)
+	: HeightmapTexture(InHeightmapTexture)
+	, LandscapeKey(InWorld, InGuid)
 	, ComponentBase(InComponentBase)
-	, HeightmapTexture(InHeightmapTexture)
 	, ForcedLOD(InForcedLOD)
 	, LODBias(InLODBias)
 	, bRegistered(false)
 	, PrimitiveCustomDataIndex(INDEX_NONE)
+	, LandscapeSceneProxy(nullptr)
 	{
 		//       -Y       
 		//    - - 0 - -   
@@ -453,6 +485,268 @@ public:
 	void UnregisterNeighbors();
 };
 
+
+class FNullLandscapeRenderSystemResources : public FRenderResource
+{
+public:
+
+	FVertexBufferRHIRef SectionLODBuffer;
+	FShaderResourceViewRHIRef SectionLODSRV;
+	TUniformBufferRef<FLandscapeSectionLODUniformParameters> UniformBuffer;
+
+	virtual void InitRHI() override
+	{
+		TResourceArray<float> ResourceBuffer;
+		ResourceBuffer.Add(0.0f);
+		FRHIResourceCreateInfo CreateInfo(&ResourceBuffer);
+		SectionLODBuffer = RHICreateVertexBuffer(ResourceBuffer.GetResourceDataSize(), BUF_ShaderResource, CreateInfo);
+		SectionLODSRV = RHICreateShaderResourceView(SectionLODBuffer, sizeof(float), PF_R32_FLOAT);
+
+		FLandscapeSectionLODUniformParameters Parameters;
+		Parameters.Size = FIntPoint(1, 1);
+		Parameters.SectionLOD = SectionLODSRV;
+		Parameters.SectionLODBias = SectionLODSRV;
+		Parameters.SectionTessellationFalloffC = SectionLODSRV;
+		Parameters.SectionTessellationFalloffK = SectionLODSRV;
+		UniformBuffer = TUniformBufferRef<FLandscapeSectionLODUniformParameters>::CreateUniformBufferImmediate(Parameters, UniformBuffer_MultiFrame);
+	}
+
+	virtual void ReleaseRHI() override
+	{
+		SectionLODBuffer.SafeRelease();
+		SectionLODSRV.SafeRelease();
+		UniformBuffer.SafeRelease();
+	}
+};
+
+extern TGlobalResource<FNullLandscapeRenderSystemResources> GNullLandscapeRenderSystemResources;
+
+struct FLandscapeRenderSystem
+{
+	static int8 GetLODFromScreenSize(LODSettingsComponent LODSettings, float InScreenSizeSquared, float InViewLODScale, float& OutFractionalLOD)
+	{
+		float ScreenSizeSquared = InScreenSizeSquared / InViewLODScale;
+
+		if (ScreenSizeSquared <= LODSettings.LastLODScreenSizeSquared)
+		{
+			OutFractionalLOD = LODSettings.LastLODIndex;
+			return LODSettings.LastLODIndex;
+		}
+		else if (ScreenSizeSquared > LODSettings.LOD1ScreenSizeSquared)
+		{
+			OutFractionalLOD = (LODSettings.LOD0ScreenSizeSquared - FMath::Min(ScreenSizeSquared, LODSettings.LOD0ScreenSizeSquared)) / (LODSettings.LOD0ScreenSizeSquared - LODSettings.LOD1ScreenSizeSquared);
+			return 0;
+		}
+		else
+		{
+			// No longer linear fraction, but worth the cache misses
+			OutFractionalLOD = 1 + FMath::LogX(LODSettings.LODOnePlusDistributionScalarSquared, LODSettings.LOD1ScreenSizeSquared / InScreenSizeSquared);
+			return (int8)OutFractionalLOD;
+		}
+
+		OutFractionalLOD = -1.0f;
+		return INDEX_NONE;
+	}
+
+	int32 NumRegisteredEntities;
+
+	FIntPoint Min;
+	FIntPoint Size;
+
+	struct SystemTessellationFalloffSettings // Global settings on the render system, not as a component of an entity
+	{
+		bool UseTessellationComponentScreenSizeFalloff;
+		float TessellationComponentSquaredScreenSize;
+		float TessellationComponentScreenSizeFalloff;
+	} TessellationFalloffSettings;
+
+	TArray<LODSettingsComponent> SectionLODSettings;
+	TResourceArray<float> SectionLODValues;
+	TResourceArray<float> SectionLODBiases;
+	TResourceArray<float> SectionTessellationFalloffC;
+	TResourceArray<float> SectionTessellationFalloffK;
+	TArray<FVector4> SectionOriginAndRadius;
+	TArray<FLandscapeComponentSceneProxy*> SceneProxies;
+
+	FVertexBufferRHIRef SectionLODBuffer;
+	FShaderResourceViewRHIRef SectionLODSRV;
+	FVertexBufferRHIRef SectionLODBiasBuffer;
+	FShaderResourceViewRHIRef SectionLODBiasSRV;
+	FVertexBufferRHIRef SectionTessellationFalloffCBuffer;
+	FShaderResourceViewRHIRef SectionTessellationFalloffCSRV;
+	FVertexBufferRHIRef SectionTessellationFalloffKBuffer;
+	FShaderResourceViewRHIRef SectionTessellationFalloffKSRV;
+
+	TUniformBufferRef<FLandscapeSectionLODUniformParameters> UniformBuffer;
+
+	TMap<const FSceneView*, TResourceArray<float>> CachedSectionLODValues;
+	TMap<const FSceneView*, TResourceArray<float>> CachedSectionTessellationFalloffC;
+	TMap<const FSceneView*, TResourceArray<float>> CachedSectionTessellationFalloffK;
+	const FSceneView* CachedView;
+
+	TMap<const FSceneView*, FGraphEventRef> TaskEventRef;
+	FGraphEventRef FetchHeightmapLODBiasesEventRef;
+
+	struct FComputeSectionPerViewParametersTask
+	{
+		FLandscapeRenderSystem& RenderSystem;
+		const FSceneView* View;
+
+		FComputeSectionPerViewParametersTask(FLandscapeRenderSystem& InRenderSystem, const FSceneView* InView)
+			: RenderSystem(InRenderSystem)
+			, View(InView)
+		{
+		}
+
+		FORCEINLINE TStatId GetStatId() const
+		{
+			RETURN_QUICK_DECLARE_CYCLE_STAT(FComputeSectionPerViewParametersTask, STATGROUP_TaskGraphTasks);
+		}
+
+		ENamedThreads::Type GetDesiredThread()
+		{
+			return ENamedThreads::AnyNormalThreadNormalTask;
+		}
+
+		static ESubsequentsMode::Type GetSubsequentsMode()
+		{
+			return ESubsequentsMode::TrackSubsequents;
+		}
+
+		void AnyThreadTask()
+		{
+			RenderSystem.ComputeSectionPerViewParameters(View);
+		}
+
+		void DoTask(ENamedThreads::Type CurrentThread, const FGraphEventRef& MyCompletionGraphEvent)
+		{
+			AnyThreadTask();
+		}
+	};
+
+	struct FGetSectionLODBiasesTask
+	{
+		FLandscapeRenderSystem& RenderSystem;
+
+		FGetSectionLODBiasesTask(FLandscapeRenderSystem& InRenderSystem)
+			: RenderSystem(InRenderSystem)
+		{
+		}
+
+		FORCEINLINE TStatId GetStatId() const
+		{
+			RETURN_QUICK_DECLARE_CYCLE_STAT(FGetSectionLODBiasesTask, STATGROUP_TaskGraphTasks);
+		}
+
+		ENamedThreads::Type GetDesiredThread()
+		{
+			return ENamedThreads::AnyNormalThreadNormalTask;
+		}
+
+		static ESubsequentsMode::Type GetSubsequentsMode()
+		{
+			return ESubsequentsMode::TrackSubsequents;
+		}
+
+		void AnyThreadTask()
+		{
+			RenderSystem.FetchHeightmapLODBiases();
+		}
+
+		void DoTask(ENamedThreads::Type CurrentThread, const FGraphEventRef& MyCompletionGraphEvent)
+		{
+			AnyThreadTask();
+		}
+	};
+
+	FLandscapeRenderSystem()
+		: NumRegisteredEntities(0)
+		, Min(MAX_int32, MAX_int32)
+		, Size(EForceInit::ForceInitToZero)
+		, CachedView(nullptr)
+	{
+		SectionLODValues.SetAllowCPUAccess(true);
+		SectionLODBiases.SetAllowCPUAccess(true);
+		SectionTessellationFalloffC.SetAllowCPUAccess(true);
+		SectionTessellationFalloffK.SetAllowCPUAccess(true);
+	}
+
+	int32 GetComponentLinearIndex(FIntPoint ComponentBase)
+	{
+		return (ComponentBase.Y - Min.Y) * Size.X + ComponentBase.X - Min.X;
+	}
+	void ResizeAndMoveTo(FIntPoint NewMin, FIntPoint NewSize);
+
+	void SetSectionLODSettings(FIntPoint ComponentBase, LODSettingsComponent LODSettings)
+	{
+		SectionLODSettings[GetComponentLinearIndex(ComponentBase)] = LODSettings;
+	}
+
+	void SetSectionOriginAndRadius(FIntPoint ComponentBase, FVector4 OriginAndRadius)
+	{
+		SectionOriginAndRadius[GetComponentLinearIndex(ComponentBase)] = OriginAndRadius;
+	}
+
+	void SetSceneProxy(FIntPoint ComponentBase, FLandscapeComponentSceneProxy* SceneProxy)
+	{
+		SceneProxies[GetComponentLinearIndex(ComponentBase)] = SceneProxy;
+	}
+
+	float GetSectionLODValue(FIntPoint ComponentBase)
+	{
+		return SectionLODValues[GetComponentLinearIndex(ComponentBase)];
+	}
+
+	float GetSectionLODBias(FIntPoint ComponentBase)
+	{
+		return SectionLODBiases[GetComponentLinearIndex(ComponentBase)];
+	}
+
+	void ComputeSectionPerViewParameters(const FSceneView* View);
+
+	void PrepareView(const FSceneView* View)
+	{
+		const bool bExecuteInParallel = FApp::ShouldUseThreadingForPerformance()
+			&& GRenderingThread; // Rendering thread is required to safely use rendering resources in parallel.
+
+		if (bExecuteInParallel)
+		{
+			TaskEventRef.Add(View, TGraphTask<FComputeSectionPerViewParametersTask>::CreateTask(
+				nullptr, ENamedThreads::GetRenderThread()).ConstructAndDispatchWhenReady(*this, View));
+		}
+	}
+
+	void BeginRenderView(const FSceneView* View)
+	{
+		if (FetchHeightmapLODBiasesEventRef.IsValid())
+		{
+			FTaskGraphInterface::Get().WaitUntilTaskCompletes(FetchHeightmapLODBiasesEventRef, ENamedThreads::GetRenderThread());
+			FetchHeightmapLODBiasesEventRef.SafeRelease();
+		}
+
+		if (TaskEventRef.Contains(View))
+		{
+			FTaskGraphInterface::Get().WaitUntilTaskCompletes(TaskEventRef[View], ENamedThreads::GetRenderThread());
+			TaskEventRef.Remove(View);
+		}
+		else
+		{
+			FComputeSectionPerViewParametersTask Task(*this, View);
+			Task.AnyThreadTask();
+		}
+
+		RecreateBuffers(View);
+	}
+
+	void BeginFrame();
+
+	void FetchHeightmapLODBiases();
+
+	void RecreateBuffers(const FSceneView* InView = nullptr);
+};
+
+extern TMap<FLandscapeNeighborInfo::FLandscapeKey, FLandscapeRenderSystem*> LandscapeRenderSystems;
+
 //
 // FLandscapeMeshProxySceneProxy
 //
@@ -552,9 +846,8 @@ public:
 	struct FLandscapeSectionRayTracingState
 	{
 		int8 CurrentLOD;
-		FVector4 LODBias;
-		FVector4 SectionLOD;
-		FVector4 CurrentNeighborLOD;
+		float FractionalLOD;
+		float HeightmapLODBias;
 		uint32 ReferencedTextureRHIHash;
 
 		FRayTracingGeometry Geometry;
@@ -563,15 +856,16 @@ public:
 
 		FLandscapeSectionRayTracingState() 
 			: CurrentLOD(-1)
-			, LODBias(-1000.0f, -1000.0f, -1000.0f, -1000.0f)
-			, SectionLOD(-1000.0f, -1000.0f, -1000.0f, -1000.0f)
-			, CurrentNeighborLOD(-1000.0f, -1000.0f, -1000.0f, -1000.0f)
+			, FractionalLOD(-1000.0f)
+			, HeightmapLODBias(-1000.0f)
 			, ReferencedTextureRHIHash(0) {}
 	};
 
 	TStaticArray<FLandscapeSectionRayTracingState, MAX_SUBSECTION_COUNT> SectionRayTracingStates;
 #endif
 
+	friend FLandscapeNeighborInfo;
+
 protected:
 	int8						MaxLOD;		// Maximum LOD level, user override possible
 	bool						UseTessellationComponentScreenSizeFalloff:1;	// Tell if we should apply a Tessellation falloff
@@ -719,6 +1013,7 @@ protected:
 
 	bool GetMeshElement(bool UseSeperateBatchForShadow, bool ShadowOnly, bool HasTessellation, int8 InLODIndex, UMaterialInterface* InMaterialInterface, FMeshBatch& OutMeshBatch, TArray<FLandscapeBatchElementParams>& OutStaticBatchParamArray) const;
 	bool GetMeshElementForVirtualTexture(int32 InLodIndex, ERuntimeVirtualTextureMaterialType MaterialType, UMaterialInterface* InMaterialInterface, FMeshBatch& OutMeshBatch, TArray<FLandscapeBatchElementParams>& OutStaticBatchParamArray) const;
+	template<class ArrayType> bool GetStaticMeshElement(int32 LODIndex, bool bForToolMesh, bool bForcedLOD, FMeshBatch& MeshBatch, ArrayType& OutStaticBatchParamArray) const;
 	void BuildDynamicMeshElement(const FViewCustomDataLOD* InPrimitiveCustomData, bool InToolMesh, bool InHasTessellation, bool InDisableTessellation, FMeshBatch& OutMeshBatch, TArray<FLandscapeBatchElementParams, SceneRenderingAllocator>& OutStaticBatchParamArray) const;
 
 	float GetComponentScreenSize(const class FSceneView* View, const FVector& Origin,  float MaxExtend, float ElementRadius) const;
@@ -742,9 +1037,7 @@ public:
 	virtual void OnLevelAddedToWorld() override;
 	virtual void* InitViewCustomData(const FSceneView& InView, float InViewLODScale, FMemStackBase& InCustomDataMemStack, bool InIsStaticRelevant, bool InIsShadowOnly, const FLODMask* InVisiblePrimitiveLODMask = nullptr, float InMeshScreenSizeSquared = -1.0f) override;
 	void PostInitViewCustomData(const FSceneView& InView, void* InViewCustomData) const;
-	virtual bool IsUsingCustomLODRules() const override;
 	virtual FLODMask GetCustomLOD(const FSceneView& InView, float InViewLODScale, int32 InForcedLODLevel, float& OutScreenSizeSquared) const override;
-	virtual bool IsUsingCustomWholeSceneShadowLODRules() const override;
 	virtual FLODMask GetCustomWholeSceneShadowLOD(const FSceneView& InView, float InViewLODScale, int32 InForcedLODLevel, const struct FLODMask& InVisibilePrimitiveLODMask, float InShadowMapTextureResolution, float InShadowMapCascadeSize, int8 InShadowCascadeId, bool InHasSelfShadow) const override;
 	
 	friend class ULandscapeComponent;
diff --git a/Engine/Source/Runtime/RenderCore/Public/RendererInterface.h b/Engine/Source/Runtime/RenderCore/Public/RendererInterface.h
--- a/Engine/Source/Runtime/RenderCore/Public/RendererInterface.h
+++ b/Engine/Source/Runtime/RenderCore/Public/RendererInterface.h
@@ -571,6 +571,14 @@ public:
 	float PreExposure;
 };
 
+class IPersistentViewUniformBufferExtension
+{
+public:
+	virtual void BeginFrame() {}
+	virtual void PrepareView(const FSceneView* View) {}
+	virtual void BeginRenderView(const FSceneView* View) {}
+};
+
 /**
  * The public interface of the renderer module.
  */
@@ -703,5 +711,7 @@ public:
 
 	/** Evict all data from virtual texture caches*/
 	virtual void FlushVirtualTextureCache() = 0;
+
+	virtual void RegisterPersistentViewUniformBufferExtension(IPersistentViewUniformBufferExtension* Extension) = 0;
 };
 
diff --git a/Engine/Source/Runtime/Renderer/Private/DeferredShadingRenderer.cpp b/Engine/Source/Runtime/Renderer/Private/DeferredShadingRenderer.cpp
--- a/Engine/Source/Runtime/Renderer/Private/DeferredShadingRenderer.cpp
+++ b/Engine/Source/Runtime/Renderer/Private/DeferredShadingRenderer.cpp
@@ -596,6 +596,13 @@ bool FDeferredShadingSceneRenderer::GatherRayTracingWorldInstances(FRHICommandLi
 
 		FViewInfo& ReferenceView = Views[0];
 
+		extern TSet<IPersistentViewUniformBufferExtension*> PersistentViewUniformBufferExtensions;
+
+		for (IPersistentViewUniformBufferExtension* Extension : PersistentViewUniformBufferExtensions)
+		{
+			Extension->BeginRenderView(&ReferenceView);
+		}
+
 		ReferenceView.RayTracingMeshResourceCollector = MakeUnique<FRayTracingMeshResourceCollector>(
 			Scene->GetFeatureLevel(),
 			&DynamicIndexBufferForInitViews,
@@ -994,6 +1001,16 @@ void FDeferredShadingSceneRenderer::Render(FRHICommandListImmediate& RHICmdList)
 		}	
 	}
 
+	for (int32 ViewIndex = 0; ViewIndex < Views.Num(); ViewIndex++)
+	{
+		extern TSet<IPersistentViewUniformBufferExtension*> PersistentViewUniformBufferExtensions;
+
+		for (IPersistentViewUniformBufferExtension* Extension : PersistentViewUniformBufferExtensions)
+		{
+			Extension->PrepareView(&Views[ViewIndex]);
+		}
+	}
+
 	if (GDoPrepareDistanceFieldSceneAfterRHIFlush && (GRHINeedsExtraDeletionLatency || !GRHICommandList.Bypass()))
 	{
 		// we will probably stall on occlusion queries, so might as well have the RHI thread and GPU work while we wait.
diff --git a/Engine/Source/Runtime/Renderer/Private/MobileShadingRenderer.cpp b/Engine/Source/Runtime/Renderer/Private/MobileShadingRenderer.cpp
--- a/Engine/Source/Runtime/Renderer/Private/MobileShadingRenderer.cpp
+++ b/Engine/Source/Runtime/Renderer/Private/MobileShadingRenderer.cpp
@@ -300,6 +300,16 @@ void FMobileSceneRenderer::InitViews(FRHICommandListImmediate& RHICmdList)
 	UpdatePrimitiveIndirectLightingCacheBuffers();
 	
 	OnStartRender(RHICmdList);
+
+	for (int32 ViewIndex = 0; ViewIndex < Views.Num(); ViewIndex++)
+	{
+		extern TSet<IPersistentViewUniformBufferExtension*> PersistentViewUniformBufferExtensions;
+
+		for (IPersistentViewUniformBufferExtension* Extension : PersistentViewUniformBufferExtensions)
+		{
+			Extension->PrepareView(&Views[ViewIndex]);
+		}
+	}
 }
 
 /** 
diff --git a/Engine/Source/Runtime/Renderer/Private/RendererModule.h b/Engine/Source/Runtime/Renderer/Private/RendererModule.h
--- a/Engine/Source/Runtime/Renderer/Private/RendererModule.h
+++ b/Engine/Source/Runtime/Renderer/Private/RendererModule.h
@@ -106,6 +106,7 @@ public:
 	virtual void RequestVirtualTextureTilesForRegion(IAllocatedVirtualTexture* AllocatedVT, const FVector2D& InScreenSpaceSize, const FIntRect& InTextureRegion, int32 InMipLevel) override;
 	virtual void LoadPendingVirtualTextureTiles(FRHICommandListImmediate& RHICmdList, ERHIFeatureLevel::Type FeatureLevel) override;
 	virtual void FlushVirtualTextureCache() override;
+	virtual void RegisterPersistentViewUniformBufferExtension(IPersistentViewUniformBufferExtension* Extension) override;
 private:
 	TSet<FSceneInterface*> AllocatedScenes;
 	FPostOpaqueRenderDelegate PostOpaqueRenderDelegate;
diff --git a/Engine/Source/Runtime/Renderer/Private/RendererScene.cpp b/Engine/Source/Runtime/Renderer/Private/RendererScene.cpp
--- a/Engine/Source/Runtime/Renderer/Private/RendererScene.cpp
+++ b/Engine/Source/Runtime/Renderer/Private/RendererScene.cpp
@@ -920,6 +920,13 @@ void FPersistentUniformBuffers::Initialize()
 #endif
 }
 
+TSet<IPersistentViewUniformBufferExtension*> PersistentViewUniformBufferExtensions;
+
+void FRendererModule::RegisterPersistentViewUniformBufferExtension(IPersistentViewUniformBufferExtension* Extension)
+{
+	PersistentViewUniformBufferExtensions.Add(Extension);
+}
+
 bool FPersistentUniformBuffers::UpdateViewUniformBuffer(const FViewInfo& View)
 {
 	// ViewUniformBuffer can be cached by mesh commands, so we need to update it every time we change current view.
@@ -941,6 +948,12 @@ bool FPersistentUniformBuffers::UpdateViewUniformBuffer(const FViewInfo& View)
 		}
 
 		CachedView = &View;
+
+		for (IPersistentViewUniformBufferExtension* Extension : PersistentViewUniformBufferExtensions)
+		{
+			Extension->BeginRenderView(&View);
+		}
+
 		return true;
 	}
 	return false;
@@ -3341,7 +3354,10 @@ void FScene::OnLevelAddedToWorld_RenderThread(FName InLevelName)
 			Proxy->bIsComponentLevelVisible = true;
 			if (Proxy->NeedsLevelAddedToWorldNotification())
 			{
+				// The only type of SceneProxy using this is landscape
+				(*It)->RemoveStaticMeshes();
 				Proxy->OnLevelAddedToWorld();
+				(*It)->AddStaticMeshes(FRHICommandListExecutor::GetImmediateCommandList());
 			}
 		}
 	}
diff --git a/Engine/Source/Runtime/Renderer/Private/SceneHitProxyRendering.cpp b/Engine/Source/Runtime/Renderer/Private/SceneHitProxyRendering.cpp
--- a/Engine/Source/Runtime/Renderer/Private/SceneHitProxyRendering.cpp
+++ b/Engine/Source/Runtime/Renderer/Private/SceneHitProxyRendering.cpp
@@ -580,6 +580,13 @@ void FDeferredShadingSceneRenderer::RenderHitProxies(FRHICommandListImmediate& R
 		for (int32 ViewIndex = 0; ViewIndex < Views.Num(); ViewIndex++)
 		{
 			UploadDynamicPrimitiveShaderDataForView(RHICmdList, *Scene, Views[ViewIndex]);
+
+			extern TSet<IPersistentViewUniformBufferExtension*> PersistentViewUniformBufferExtensions;
+
+			for (IPersistentViewUniformBufferExtension* Extension : PersistentViewUniformBufferExtensions)
+			{
+				Extension->BeginRenderView(&Views[ViewIndex]);
+			}
 		}	
 
 		if (UpdateViewCustomDataEvents.Num())
diff --git a/Engine/Source/Runtime/Renderer/Private/SceneRendering.cpp b/Engine/Source/Runtime/Renderer/Private/SceneRendering.cpp
--- a/Engine/Source/Runtime/Renderer/Private/SceneRendering.cpp
+++ b/Engine/Source/Runtime/Renderer/Private/SceneRendering.cpp
@@ -3118,6 +3118,13 @@ void FSceneRenderer::RenderCustomDepthPass(FRHICommandListImmediate& RHICmdList)
 						Scene->UniformBuffers.InstancedCustomDepthViewUniformBuffer.UpdateUniformBufferImmediate(reinterpret_cast<FInstancedViewUniformShaderParameters&>(*View.CachedViewUniformShaderParameters));
 					}
 				}
+
+				extern TSet<IPersistentViewUniformBufferExtension*> PersistentViewUniformBufferExtensions;
+				
+				for (IPersistentViewUniformBufferExtension* Extension : PersistentViewUniformBufferExtensions)
+				{
+					Extension->BeginRenderView(&View);
+				}
 	
 				View.ParallelMeshDrawCommandPasses[EMeshPass::CustomDepth].DispatchDraw(nullptr, RHICmdList);
 			}
@@ -3144,6 +3151,13 @@ void FSceneRenderer::OnStartRender(FRHICommandListImmediate& RHICmdList)
 			View.ViewState->OnStartRender(View, ViewFamily);
 		}
 	}
+
+	extern TSet<IPersistentViewUniformBufferExtension*> PersistentViewUniformBufferExtensions;
+
+	for (IPersistentViewUniformBufferExtension* Extension : PersistentViewUniformBufferExtensions)
+	{
+		Extension->BeginFrame();
+	}
 }
 
 bool FSceneRenderer::ShouldCompositeEditorPrimitives(const FViewInfo& View)
diff --git a/Engine/Source/Runtime/Renderer/Private/ShadowDepthRendering.cpp b/Engine/Source/Runtime/Renderer/Private/ShadowDepthRendering.cpp
--- a/Engine/Source/Runtime/Renderer/Private/ShadowDepthRendering.cpp
+++ b/Engine/Source/Runtime/Renderer/Private/ShadowDepthRendering.cpp
@@ -1089,6 +1089,16 @@ void FProjectedShadowInfo::SetupShadowUniformBuffers(FRHICommandListImmediate& R
 		ShadowDepthPassUniformBuffer.UpdateUniformBufferImmediate(ShadowDepthPassParameters);
 
 		UploadDynamicPrimitiveShaderDataForView(RHICmdList, *Scene, *ShadowDepthView);
+
+		if (DependentView)
+		{
+			extern TSet<IPersistentViewUniformBufferExtension*> PersistentViewUniformBufferExtensions;
+
+			for (IPersistentViewUniformBufferExtension* Extension : PersistentViewUniformBufferExtensions)
+			{
+				Extension->BeginRenderView(DependentView);
+			}
+		}
 	}
 }
 
@@ -1116,6 +1126,16 @@ void FProjectedShadowInfo::RenderDepthInner(FRHICommandListImmediate& RHICmdList
 	{
 		// CSM shadow depth cached mesh draw commands are all referencing the same view uniform buffer.  We need to update it before rendering each cascade.
 		ShadowDepthView->ViewUniformBuffer.UpdateUniformBufferImmediate(*ShadowDepthView->CachedViewUniformShaderParameters);
+		
+		if (DependentView)
+		{
+			extern TSet<IPersistentViewUniformBufferExtension*> PersistentViewUniformBufferExtensions;
+
+			for (IPersistentViewUniformBufferExtension* Extension : PersistentViewUniformBufferExtensions)
+			{
+				Extension->BeginRenderView(DependentView);
+			}
+		}
 	}
 
 	if (FSceneInterface::GetShadingPath(FeatureLevel) == EShadingPath::Mobile)
