Fixed landscape GetHeightAtLocation wasn't calculated properly for negative locations. (From PatrickE)

diff --git a/Engine/Source/Runtime/Landscape/Private/LandscapeCollision.cpp b/Engine/Source/Runtime/Landscape/Private/LandscapeCollision.cpp
--- a/Engine/Source/Runtime/Landscape/Private/LandscapeCollision.cpp
+++ b/Engine/Source/Runtime/Landscape/Private/LandscapeCollision.cpp
@@ -2209,8 +2209,8 @@ LANDSCAPE_API TOptional<float> ALandscapeProxy::GetHeightAtLocation(FVector Loca
 {
 	TOptional<float> Height;
 	ULandscapeInfo* Info = GetLandscapeInfo();
-	const FVector ActorSpaceLocation = GetActorTransform().InverseTransformPosition(Location);
-	const FIntPoint Key = FIntPoint(ActorSpaceLocation.X, ActorSpaceLocation.Y) / ComponentSizeQuads;
+	const FVector ActorSpaceLocation = LandscapeActorToWorld().InverseTransformPosition(Location);
+	const FIntPoint Key = FIntPoint(FMath::FloorToInt(ActorSpaceLocation.X / ComponentSizeQuads), FMath::FloorToInt(ActorSpaceLocation.Y / ComponentSizeQuads));
 	ULandscapeComponent* Component = Info->XYtoComponentMap.FindRef(Key);
 	if (Component)
 	{
@@ -2225,4 +2225,4 @@ LANDSCAPE_API TOptional<float> ALandscapeProxy::GetHeightAtLocation(FVector Loca
 		}
 	}
 	return Height;
-}
\ No newline at end of file
+}
