#UE4 Fix deadlock in landscape rendering

diff --git a/Engine/Source/Runtime/Landscape/Public/LandscapeRender.h b/Engine/Source/Runtime/Landscape/Public/LandscapeRender.h
--- a/Engine/Source/Runtime/Landscape/Public/LandscapeRender.h
+++ b/Engine/Source/Runtime/Landscape/Public/LandscapeRender.h
@@ -720,13 +720,13 @@ struct FLandscapeRenderSystem
 	{
 		if (FetchHeightmapLODBiasesEventRef.IsValid())
 		{
-			FTaskGraphInterface::Get().WaitUntilTaskCompletes(FetchHeightmapLODBiasesEventRef, ENamedThreads::GetRenderThread());
+			FTaskGraphInterface::Get().WaitUntilTaskCompletes(FetchHeightmapLODBiasesEventRef, ENamedThreads::GetRenderThread_Local());
 			FetchHeightmapLODBiasesEventRef.SafeRelease();
 		}
 
 		if (TaskEventRef.Contains(View))
 		{
-			FTaskGraphInterface::Get().WaitUntilTaskCompletes(TaskEventRef[View], ENamedThreads::GetRenderThread());
+			FTaskGraphInterface::Get().WaitUntilTaskCompletes(TaskEventRef[View], ENamedThreads::GetRenderThread_Local());
 			TaskEventRef.Remove(View);
 		}
 		else
