[NO NON-PII COMMIT MESSAGE]

diff --git a/Engine/Source/Runtime/Landscape/Private/LandscapeEdit.cpp b/Engine/Source/Runtime/Landscape/Private/LandscapeEdit.cpp
--- a/Engine/Source/Runtime/Landscape/Private/LandscapeEdit.cpp
+++ b/Engine/Source/Runtime/Landscape/Private/LandscapeEdit.cpp
@@ -5997,7 +5997,7 @@ void ULandscapeComponent::GeneratePlatformVertexData(const ITargetPlatform* Targ
 	
 	// Generate occlusion mesh
 	TArray<FVector> OccluderVertices;
-	const int32 OcclusionMeshMip = FMath::Clamp<int32>(GetLandscapeProxy()->OccluderGeometryLOD, -1, MaxLOD);
+	const int32 OcclusionMeshMip = FMath::Clamp<int32>(GetLandscapeProxy()->OccluderGeometryLOD, -1, HeightmapMipData.Num() - 1);
 
 	if (OcclusionMeshMip >= 0 && (!TargetPlatform || TargetPlatform->SupportsFeature(ETargetPlatformFeatures::SoftwareOcclusion)))
 	{
