Fix landscape startup crash from a bad merge

diff --git a/Engine/Source/Runtime/Landscape/Private/LandscapeRender.cpp b/Engine/Source/Runtime/Landscape/Private/LandscapeRender.cpp
--- a/Engine/Source/Runtime/Landscape/Private/LandscapeRender.cpp
+++ b/Engine/Source/Runtime/Landscape/Private/LandscapeRender.cpp
@@ -4470,7 +4470,7 @@ void FLandscapeVertexFactory::ModifyCompilationEnvironment(const FVertexFactoryT
 }
 
 
-IMPLEMENT_VERTEX_FACTORY_TYPE_EX(FLandscapeVertexFactory, "/Engine/Private/LandscapeVertexFactory.ush", true, true, true, false, false, true, false);
+IMPLEMENT_VERTEX_FACTORY_TYPE(FLandscapeVertexFactory, "/Engine/Private/LandscapeVertexFactory.ush", true, true, true, false, false);
 
 /**
 * Copy the data from another vertex factory
