Fix landscape volatile buffer usage

diff --git a/Engine/Source/Runtime/Landscape/Private/LandscapeRender.cpp b/Engine/Source/Runtime/Landscape/Private/LandscapeRender.cpp
--- a/Engine/Source/Runtime/Landscape/Private/LandscapeRender.cpp
+++ b/Engine/Source/Runtime/Landscape/Private/LandscapeRender.cpp
@@ -808,7 +808,7 @@ void FLandscapeRenderSystem::RecreateBuffers(const FSceneView* InView /* = nullp
 			if (!SectionLODBuffer.IsValid())
 			{
 				FRHIResourceCreateInfo CreateInfo(&SectionLODValues);
-				SectionLODBuffer = RHICreateVertexBuffer(SectionLODValues.GetResourceDataSize(), BUF_ShaderResource | BUF_Volatile, CreateInfo);
+				SectionLODBuffer = RHICreateVertexBuffer(SectionLODValues.GetResourceDataSize(), BUF_ShaderResource | BUF_Dynamic, CreateInfo);
 				SectionLODSRV = RHICreateShaderResourceView(SectionLODBuffer, sizeof(float), PF_R32_FLOAT);
 			}
 			else
@@ -823,7 +823,7 @@ void FLandscapeRenderSystem::RecreateBuffers(const FSceneView* InView /* = nullp
 			if (!SectionLODBiasBuffer.IsValid())
 			{
 				FRHIResourceCreateInfo CreateInfo(&SectionLODBiases);
-				SectionLODBiasBuffer = RHICreateVertexBuffer(SectionLODBiases.GetResourceDataSize(), BUF_ShaderResource | BUF_Volatile, CreateInfo);
+				SectionLODBiasBuffer = RHICreateVertexBuffer(SectionLODBiases.GetResourceDataSize(), BUF_ShaderResource | BUF_Dynamic, CreateInfo);
 				SectionLODBiasSRV = RHICreateShaderResourceView(SectionLODBiasBuffer, sizeof(float), PF_R32_FLOAT);
 			}
 			else
@@ -838,7 +838,7 @@ void FLandscapeRenderSystem::RecreateBuffers(const FSceneView* InView /* = nullp
 			if (!SectionTessellationFalloffCBuffer.IsValid())
 			{
 				FRHIResourceCreateInfo CreateInfo(&SectionTessellationFalloffC);
-				SectionTessellationFalloffCBuffer = RHICreateVertexBuffer(SectionTessellationFalloffC.GetResourceDataSize(), BUF_ShaderResource | BUF_Volatile, CreateInfo);
+				SectionTessellationFalloffCBuffer = RHICreateVertexBuffer(SectionTessellationFalloffC.GetResourceDataSize(), BUF_ShaderResource | BUF_Dynamic, CreateInfo);
 				SectionTessellationFalloffCSRV = RHICreateShaderResourceView(SectionTessellationFalloffCBuffer, sizeof(float), PF_R32_FLOAT);
 			}
 			else
@@ -857,7 +857,7 @@ void FLandscapeRenderSystem::RecreateBuffers(const FSceneView* InView /* = nullp
 			if (!SectionTessellationFalloffKBuffer.IsValid())
 			{
 				FRHIResourceCreateInfo CreateInfo(&SectionTessellationFalloffK);
-				SectionTessellationFalloffKBuffer = RHICreateVertexBuffer(SectionTessellationFalloffK.GetResourceDataSize(), BUF_ShaderResource | BUF_Volatile, CreateInfo);
+				SectionTessellationFalloffKBuffer = RHICreateVertexBuffer(SectionTessellationFalloffK.GetResourceDataSize(), BUF_ShaderResource | BUF_Dynamic, CreateInfo);
 				SectionTessellationFalloffKSRV = RHICreateShaderResourceView(SectionTessellationFalloffKBuffer, sizeof(float), PF_R32_FLOAT);
 			}
 			else
diff --git a/Engine/Source/Runtime/Landscape/Public/LandscapeRender.h b/Engine/Source/Runtime/Landscape/Public/LandscapeRender.h
--- a/Engine/Source/Runtime/Landscape/Public/LandscapeRender.h
+++ b/Engine/Source/Runtime/Landscape/Public/LandscapeRender.h
@@ -499,7 +499,7 @@ public:
 		TResourceArray<float> ResourceBuffer;
 		ResourceBuffer.Add(0.0f);
 		FRHIResourceCreateInfo CreateInfo(&ResourceBuffer);
-		SectionLODBuffer = RHICreateVertexBuffer(ResourceBuffer.GetResourceDataSize(), BUF_ShaderResource, CreateInfo);
+		SectionLODBuffer = RHICreateVertexBuffer(ResourceBuffer.GetResourceDataSize(), BUF_ShaderResource | BUF_Static, CreateInfo);
 		SectionLODSRV = RHICreateShaderResourceView(SectionLODBuffer, sizeof(float), PF_R32_FLOAT);
 
 		FLandscapeSectionLODUniformParameters Parameters;
