Fix CIS

diff --git a/Engine/Source/Runtime/Landscape/Public/LandscapeRender.h b/Engine/Source/Runtime/Landscape/Public/LandscapeRender.h
--- a/Engine/Source/Runtime/Landscape/Public/LandscapeRender.h
+++ b/Engine/Source/Runtime/Landscape/Public/LandscapeRender.h
@@ -543,9 +543,6 @@ struct FLandscapeRenderSystem
 			OutFractionalLOD = 1 + FMath::LogX(LODSettings.LODOnePlusDistributionScalarSquared, LODSettings.LOD1ScreenSizeSquared / InScreenSizeSquared);
 			return (int8)OutFractionalLOD;
 		}
-
-		OutFractionalLOD = -1.0f;
-		return INDEX_NONE;
 	}
 
 	int32 NumRegisteredEntities;
