Do not Monitor Shaders while out of LandscapeMode

#rnx

diff --git a/Engine/Source/Runtime/Landscape/Private/LandscapeEditLayers.cpp b/Engine/Source/Runtime/Landscape/Private/LandscapeEditLayers.cpp
--- a/Engine/Source/Runtime/Landscape/Private/LandscapeEditLayers.cpp
+++ b/Engine/Source/Runtime/Landscape/Private/LandscapeEditLayers.cpp
@@ -4686,6 +4686,12 @@ void ALandscape::MonitorLandscapeEdModeChanges()
 
 void ALandscape::MonitorShaderCompilation()
 {
+	// Do not monitor changes when not editing Landscape
+	if (!LandscapeEdMode)
+	{
+		return;
+	}
+
 	// If doing editing while shader are compiling or at load of a map, it's possible we will need another update pass after shader are completed to see the correct result
 	const int32 RemainingShadersThisFrame = GShaderCompilingManager->GetNumRemainingJobs();
 	if (!WasCompilingShaders && RemainingShadersThisFrame > 0)
