Fix landscape being affected by r.StaticMeshLODDistanceScale which caused seams between landscape tiles

diff --git a/Engine/Source/Runtime/Landscape/Private/Landscape.cpp b/Engine/Source/Runtime/Landscape/Private/Landscape.cpp
--- a/Engine/Source/Runtime/Landscape/Private/Landscape.cpp
+++ b/Engine/Source/Runtime/Landscape/Private/Landscape.cpp
@@ -1426,6 +1426,25 @@ FBoxSphereBounds ULandscapeComponent::CalcBounds(const FTransform& LocalToWorld)
 	return FBoxSphereBounds(MyBounds);
 }
 
+static void OnStaticMeshLODDistanceScaleChanged()
+{
+	extern RENDERER_API TAutoConsoleVariable<float> CVarStaticMeshLODDistanceScale;
+
+	static float LastValue = 1.0f;
+
+	if (LastValue != CVarStaticMeshLODDistanceScale.GetValueOnAnyThread())
+	{
+		LastValue = CVarStaticMeshLODDistanceScale.GetValueOnAnyThread();
+
+		for (auto* LandscapeComponent : TObjectRange<ULandscapeComponent>(RF_ClassDefaultObject | RF_ArchetypeObject, true, EInternalObjectFlags::PendingKill))
+		{
+			LandscapeComponent->MarkRenderStateDirty();
+		}
+	}
+}
+
+FAutoConsoleVariableSink OnStaticMeshLODDistanceScaleChangedSink(FConsoleCommandDelegate::CreateStatic(&OnStaticMeshLODDistanceScaleChanged));
+
 void ULandscapeComponent::OnRegister()
 {
 	Super::OnRegister();
diff --git a/Engine/Source/Runtime/Landscape/Private/LandscapeRender.cpp b/Engine/Source/Runtime/Landscape/Private/LandscapeRender.cpp
--- a/Engine/Source/Runtime/Landscape/Private/LandscapeRender.cpp
+++ b/Engine/Source/Runtime/Landscape/Private/LandscapeRender.cpp
@@ -731,16 +731,17 @@ void FLandscapeRenderSystem::ComputeSectionPerViewParameters(const FSceneView* V
 			CachedSectionTessellationFalloffK.Add(View, TResourceArray<float>{});
 		}
 
-		int32 ForcedLODLevel = View->Family->EngineShowFlags.LOD ? GetCVarForceLOD() : 0;
+		int32 ForcedLODLevel = View->Family->EngineShowFlags.LOD ? GetCVarForceLOD() : -1;
+		float LODScale = View->LODDistanceFactor * CVarStaticMeshLODDistanceScale.GetValueOnRenderThread();
 
 		for (int32 EntityIndex = 0; EntityIndex < SectionLODSettings.Num(); EntityIndex++)
 		{
 			float MeshScreenSizeSquared = ComputeBoundsScreenRadiusSquared(FVector(SectionOriginAndRadius[EntityIndex]), SectionOriginAndRadius[EntityIndex].W, *View);
 
 			float FractionalLOD;
-			GetLODFromScreenSize(SectionLODSettings[EntityIndex], MeshScreenSizeSquared, View->LODDistanceFactor, FractionalLOD);
+			GetLODFromScreenSize(SectionLODSettings[EntityIndex], MeshScreenSizeSquared, LODScale * LODScale, FractionalLOD);
 
-			SectionLODValues[EntityIndex] = FractionalLOD;
+			SectionLODValues[EntityIndex] = ForcedLODLevel >= 0 ? ForcedLODLevel : FractionalLOD;
 
 			if (TessellationFalloffSettings.UseTessellationComponentScreenSizeFalloff)
 			{
@@ -1055,7 +1056,8 @@ FLandscapeComponentSceneProxy::FLandscapeComponentSceneProxy(ULandscapeComponent
 	}
 
 	float ScreenSizeRatioDivider = FMath::Max(InComponent->GetLandscapeProxy()->LOD0DistributionSetting * GLandscapeLOD0DistributionScale, 1.01f);
-	float CurrentScreenSizeRatio = InComponent->GetLandscapeProxy()->LOD0ScreenSize;
+	// Cancel out so that landscape is not affected by r.StaticMeshLODDistanceScale
+	float CurrentScreenSizeRatio = InComponent->GetLandscapeProxy()->LOD0ScreenSize / CVarStaticMeshLODDistanceScale.GetValueOnAnyThread();
 
 	LODScreenRatioSquared.AddUninitialized(MaxLOD + 1);
 
@@ -3079,10 +3081,12 @@ void FLandscapeComponentSceneProxy::GetDynamicMeshElements(const TArray<const FS
 
 			float MeshScreenSizeSquared = ComputeBoundsScreenRadiusSquared(GetBounds().Origin, GetBounds().SphereRadius, *View);
 			int32 ForcedLODLevel = (View->Family->EngineShowFlags.LOD) ? GetCVarForceLOD() : -1;
+
+			float LODScale = View->LODDistanceFactor * CVarStaticMeshLODDistanceScale.GetValueOnRenderThread();
 #if WITH_EDITOR
 			ForcedLODLevel = View->Family->LandscapeLODOverride >= 0 ? View->Family->LandscapeLODOverride : ForcedLODLevel;
 #endif
-			int32 LODToRender = ForcedLODLevel >= 0 ? ForcedLODLevel : GetLODFromScreenSize(MeshScreenSizeSquared, View->LODDistanceFactor);
+			int32 LODToRender = ForcedLODLevel >= 0 ? ForcedLODLevel : GetLODFromScreenSize(MeshScreenSizeSquared, LODScale * LODScale);
 
 			ParameterArray.ElementParams.AddDefaulted(1);
 
diff --git a/Engine/Source/Runtime/Landscape/Public/LandscapeRender.h b/Engine/Source/Runtime/Landscape/Public/LandscapeRender.h
--- a/Engine/Source/Runtime/Landscape/Public/LandscapeRender.h
+++ b/Engine/Source/Runtime/Landscape/Public/LandscapeRender.h
@@ -521,6 +521,8 @@ public:
 
 extern TGlobalResource<FNullLandscapeRenderSystemResources> GNullLandscapeRenderSystemResources;
 
+extern RENDERER_API TAutoConsoleVariable<float> CVarStaticMeshLODDistanceScale;
+
 struct FLandscapeRenderSystem
 {
 	static int8 GetLODFromScreenSize(LODSettingsComponent LODSettings, float InScreenSizeSquared, float InViewLODScale, float& OutFractionalLOD)
@@ -540,7 +542,7 @@ struct FLandscapeRenderSystem
 		else
 		{
 			// No longer linear fraction, but worth the cache misses
-			OutFractionalLOD = 1 + FMath::LogX(LODSettings.LODOnePlusDistributionScalarSquared, LODSettings.LOD1ScreenSizeSquared / InScreenSizeSquared);
+			OutFractionalLOD = 1 + FMath::LogX(LODSettings.LODOnePlusDistributionScalarSquared, LODSettings.LOD1ScreenSizeSquared / ScreenSizeSquared);
 			return (int8)OutFractionalLOD;
 		}
 	}
diff --git a/Engine/Source/Runtime/Renderer/Private/SceneVisibility.cpp b/Engine/Source/Runtime/Renderer/Private/SceneVisibility.cpp
--- a/Engine/Source/Runtime/Renderer/Private/SceneVisibility.cpp
+++ b/Engine/Source/Runtime/Renderer/Private/SceneVisibility.cpp
@@ -110,7 +110,7 @@ static FAutoConsoleVariableRef CVarAllowSubPrimitiveQueries(
 	ECVF_RenderThreadSafe
 	);
 
-static TAutoConsoleVariable<float> CVarStaticMeshLODDistanceScale(
+RENDERER_API TAutoConsoleVariable<float> CVarStaticMeshLODDistanceScale(
 	TEXT("r.StaticMeshLODDistanceScale"),
 	1.0f,
 	TEXT("Scale factor for the distance used in computing discrete LOD for static meshes. (defaults to 1)\n")
