Heightfield cooked data support

diff --git a/Engine/Source/Runtime/Experimental/Chaos/Private/Chaos/HeightField.cpp b/Engine/Source/Runtime/Experimental/Chaos/Private/Chaos/HeightField.cpp
--- a/Engine/Source/Runtime/Experimental/Chaos/Private/Chaos/HeightField.cpp
+++ b/Engine/Source/Runtime/Experimental/Chaos/Private/Chaos/HeightField.cpp
@@ -154,7 +154,7 @@ namespace Chaos
 			};
 
 			TVector<T, 3> Points[4];
-			GeomData->GetPoints(FullIndex, Points);
+			GeomData->GetPointsScaled(FullIndex, Points);
 
 			// Test both triangles that are in this cell, as we could hit both in any order
 			TestTriangle(FullIndex * 2, Points[0], Points[1], Points[2]);
@@ -250,7 +250,7 @@ namespace Chaos
 			};
 
 			TVector<T, 3> Points[4];
-			HfData->GetPoints(FullIndex, Points);
+			HfData->GetPointsScaled(FullIndex, Points);
 
 			TestTriangle(FullIndex * 2, Points[0], Points[1], Points[2]);
 			TestTriangle(FullIndex * 2 + 1, Points[2], Points[1], Points[3]);
@@ -273,80 +273,82 @@ namespace Chaos
 
 	};
 
-	template <typename T>
-	THeightField<T>::THeightField(TArray<T>&& Height, int32 NumRows, int32 NumCols, const TVector<T, 3>& InScale)
-		: TImplicitObject<T, 3>(EImplicitObject::HasBoundingBox, ImplicitObjectType::HeightField)
+	template<typename HeightfieldT, typename BufferType>
+	void BuildGeomData(TArrayView<BufferType> BufferView, int32 NumRows, int32 NumCols, const TVector<HeightfieldT, 3>& InScale, TUniqueFunction<HeightfieldT(const BufferType)> ToRealFunc, typename THeightField<HeightfieldT>::FDataType& OutData, TBox<HeightfieldT, 3>& OutBounds)
 	{
-		//For now just generate a triangle mesh
+		using FDataType = THeightField<HeightfieldT>::FDataType;
+		using RealType = FDataType::RealType;
+
 		const int32 NumCells = NumRows * NumCols;
-		ensure(Height.Num() == NumCells);
+		ensure(BufferView.Num() == NumCells);
 		ensure(NumRows > 1);
 		ensure(NumCols > 1);
 
 		// Populate data.
-		const int32 NumHeights = Height.Num();
-		GeomData.Heights.SetNum(NumHeights);
+		const int32 NumHeights = BufferView.Num();
+		OutData.Heights.SetNum(NumHeights);
 
-		GeomData.NumRows = NumRows;
-		GeomData.NumCols = NumCols;
-		GeomData.MinValue = Height[0];
-		GeomData.MaxValue = Height[0];
-		GeomData.Scale = InScale;
+		OutData.NumRows = NumRows;
+		OutData.NumCols = NumCols;
+		OutData.MinValue = ToRealFunc(BufferView[0]);
+		OutData.MaxValue = ToRealFunc(BufferView[0]);
+		OutData.Scale = InScale;
 
 		for(int32 HeightIndex = 1; HeightIndex < NumHeights; ++HeightIndex)
 		{
-			const typename FDataType::RealType CurrHeight = (typename FDataType::RealType)Height[HeightIndex];
+			const typename RealType CurrHeight = ToRealFunc(BufferView[HeightIndex]);
 
-			if(CurrHeight > GeomData.MaxValue)
+			if(CurrHeight > OutData.MaxValue)
 			{
-				GeomData.MaxValue = CurrHeight;
+				OutData.MaxValue = CurrHeight;
 			}
-			else if(CurrHeight < GeomData.MinValue)
+			else if(CurrHeight < OutData.MinValue)
 			{
-				GeomData.MinValue = CurrHeight;
+				OutData.MinValue = CurrHeight;
 			}
 		}
 
-		GeomData.Range = GeomData.MaxValue - GeomData.MinValue;
-		GeomData.HeightPerUnit = GeomData.Range / FDataType::StorageRange;
+		OutData.Range = OutData.MaxValue - OutData.MinValue;
+		OutData.HeightPerUnit = OutData.Range / FDataType::StorageRange;
 
 		for(int32 HeightIndex = 0; HeightIndex < NumHeights; ++HeightIndex)
 		{
-			GeomData.Heights[HeightIndex] = static_cast<typename FDataType::StorageType>((Height[HeightIndex] - GeomData.MinValue) / GeomData.HeightPerUnit);
+			OutData.Heights[HeightIndex] = static_cast<typename FDataType::StorageType>((ToRealFunc(BufferView[HeightIndex]) - OutData.MinValue) / OutData.HeightPerUnit);
 
 			int32 X = HeightIndex % (NumCols);
 			int32 Y = HeightIndex / (NumCols);
-			TVector<T, 3> Position(typename FDataType::RealType(X), typename FDataType::RealType(Y), GeomData.MinValue + GeomData.Heights[HeightIndex] * GeomData.HeightPerUnit);
+			TVector<HeightfieldT, 3> Position(typename RealType(X), typename RealType(Y), OutData.MinValue + OutData.Heights[HeightIndex] * OutData.HeightPerUnit);
 			if(HeightIndex == 0)
 			{
-				LocalBounds = TBox<T, 3>(Position * InScale, Position * InScale);
+				OutBounds = TBox<HeightfieldT, 3>(Position * InScale, Position * InScale);
 			}
 			else
 			{
-				LocalBounds.GrowToInclude(Position * InScale);
+				OutBounds.GrowToInclude(Position * InScale);
 			}
 		}
+		OutBounds.Thicken(KINDA_SMALL_NUMBER);
+	}
 
-		// Flatten out the Z axis
-		FlattenedBounds.Min = TVector<T, 2>(LocalBounds.Min()[0], LocalBounds.Min()[1]);
-		FlattenedBounds.Max = TVector<T, 2>(LocalBounds.Max()[0], LocalBounds.Max()[1]);
-
-		BuildQueryData();
+	template <typename T>
+	THeightField<T>::THeightField(TArray<T>&& Height, int32 NumRows, int32 NumCols, const TVector<T, 3>& InScale)
+		: TImplicitObject<T, 3>(EImplicitObject::HasBoundingBox, ImplicitObjectType::HeightField)
+	{
+		BuildGeomData<T, T>(MakeArrayView(Height), NumRows, NumCols, InScale, [](const T InVal) {return InVal; }, GeomData, LocalBounds);
+		CalcBounds();
+	}
 
-		// Cache per-cell bounds
-		const int32 NumX = GeomData.NumCols - 1;
-		const int32 NumY = GeomData.NumRows - 1;
-		GeomData.CellBounds.SetNum(NumX * NumY);
-		for(int32 XIndex = 0; XIndex < NumX; ++XIndex)
+	template<typename T>
+	Chaos::THeightField<T>::THeightField(TArrayView<const uint16> InHeights, int32 InNumRows, int32 InNumCols, const TVector<T, 3>& InScale)
+		: TImplicitObject<T, 3>(EImplicitObject::HasBoundingBox, ImplicitObjectType::HeightField)
+	{
+		TUniqueFunction<T(const uint16)> ConversionFunc = [](const T InVal) -> float
 		{
-			for(int32 YIndex = 0; YIndex < NumY; ++YIndex)
-			{
-				const TVector<int32, 2> Cell(XIndex, YIndex);
-				TVector<T, 3> Min, Max;
-				CalcCellBounds3D(Cell, Min, Max);
-				GeomData.CellBounds[YIndex * NumX + XIndex] = TBox<T, 3>(Min, Max);
-			}
-		}
+			return (float)((int32)InVal - 32768);
+		};
+
+		BuildGeomData<T, const uint16>(InHeights, InNumRows, InNumCols, InScale, MoveTemp(ConversionFunc), GeomData, LocalBounds);
+		CalcBounds();
 	}
 
 	template<typename T>
@@ -381,6 +383,40 @@ namespace Chaos
 		return false;
 	}
 
+	template<typename T>
+	bool Chaos::THeightField<T>::GetCellBounds2DScaled(const TVector<int32, 2> InCoord, FBounds2D& OutBounds, const TVector<T, 2>& InInflate /*= {0}*/) const
+	{
+		if(FlatGrid.IsValid(InCoord))
+		{
+			const TBox<T, 3>& Bound = GeomData.CellBounds[InCoord[1] * (GeomData.NumCols - 1) + InCoord[0]];
+			OutBounds.Min = {Bound.Min()[0], Bound.Min()[1]};
+			OutBounds.Max = {Bound.Max()[0], Bound.Max()[1]};
+			OutBounds.Min -= InInflate;
+			OutBounds.Max += InInflate;
+			const TVector<float, 2> Scale2D = TVector<T, 2>(GeomData.Scale[0], GeomData.Scale[1]);
+			OutBounds.Min *= Scale2D;
+			OutBounds.Max *= Scale2D;
+			return true;
+		}
+
+		return false;
+	}
+
+	template<typename T>
+	bool Chaos::THeightField<T>::GetCellBounds3DScaled(const TVector<int32, 2> InCoord, TVector<T, 3>& OutMin, TVector<T, 3>& OutMax, const TVector<T, 3>& InInflate /*= {0}*/) const
+	{
+		if(FlatGrid.IsValid(InCoord))
+		{
+			const TBox<T, 3>& Bound = GeomData.CellBounds[InCoord[1] * (GeomData.NumCols - 1) + InCoord[0]];
+			OutMin = Bound.Min() - InInflate;
+			OutMax = Bound.Max() + InInflate;
+			OutMin *= GeomData.Scale;
+			OutMax *= GeomData.Scale;
+			return true;
+		}
+
+		return false;
+	}
 
 	template<typename T>
 	bool Chaos::THeightField<T>::CalcCellBounds3D(const TVector<int32, 2> InCoord, TVector<T, 3>& OutMin, TVector<T, 3>& OutMax, const TVector<T, 3>& InInflate /*= {0}*/) const
@@ -417,7 +453,7 @@ namespace Chaos
 		TVector<T, 2> ClippedFlatRayStart;
 		TVector<T, 2> ClippedFlatRayEnd;
 
-		if(FlattenedBounds.ClipLine(StartPoint, StartPoint + Dir * Length, ClippedFlatRayStart, ClippedFlatRayEnd))
+		if(GetFlatBounds().ClipLine(StartPoint, StartPoint + Dir * Length, ClippedFlatRayStart, ClippedFlatRayEnd))
 		{
 			// The line is now valid and is entirely enclosed by the bounds (and thus, the grid)
 			if(FMath::Abs((ClippedFlatRayEnd - ClippedFlatRayStart).SizeSquared()) < SMALL_NUMBER)
@@ -430,8 +466,9 @@ namespace Chaos
 			else
 			{
 				// Rasterize the line over the grid
-				TVector<int32, 2> StartCell = FlatGrid.Cell(ClippedFlatRayStart);
-				TVector<int32, 2> EndCell = FlatGrid.Cell(ClippedFlatRayEnd);
+				const TVector<T, 2> Scale2D(GeomData.Scale[0], GeomData.Scale[1]);
+				TVector<int32, 2> StartCell = FlatGrid.Cell(ClippedFlatRayStart/Scale2D);
+				TVector<int32, 2> EndCell = FlatGrid.Cell(ClippedFlatRayEnd/Scale2D);
 
 				// Boundaries might push us one cell over
 				StartCell = FlatGrid.ClampIndex(StartCell);
@@ -525,7 +562,7 @@ namespace Chaos
 	{
 		T CurrentLength = Length;
 
-		FBounds2D InflatedBounds = FlattenedBounds;
+		FBounds2D InflatedBounds = GetFlatBounds();
 		InflatedBounds.Min -= InHalfExtents;
 		InflatedBounds.Max += InHalfExtents;
 
@@ -534,6 +571,7 @@ namespace Chaos
 		const TVector<T, 3> EndPoint = StartPoint + Dir * Length;
 		const TVector<T, 2> Start2D(StartPoint[0], StartPoint[1]);
 		const TVector<T, 2> End2D(EndPoint[0], EndPoint[1]);
+		const TVector<T, 2> Scale2D(GeomData.Scale[0], GeomData.Scale[1]);
 
 		TVector<T, 2> ClippedStart;
 		TVector<T, 2> ClippedEnd;
@@ -541,8 +579,8 @@ namespace Chaos
 		if(InflatedBounds.ClipLine(StartPoint, StartPoint + Dir * Length, ClippedStart, ClippedEnd))
 		{
 			// Rasterize the line over the grid
-			TVector<int32, 2> StartCell = FlatGrid.Cell(ClippedStart);
-			TVector<int32, 2> EndCell = FlatGrid.Cell(ClippedEnd);
+			TVector<int32, 2> StartCell = FlatGrid.Cell(ClippedStart / Scale2D);
+			TVector<int32, 2> EndCell = FlatGrid.Cell(ClippedEnd / Scale2D);
 
 			// Boundaries might push us one cell over
 			StartCell = FlatGrid.ClampIndex(StartCell);
@@ -617,7 +655,7 @@ namespace Chaos
 						TVector<int32, 2> NeighCoord = CellCoord.Index + Neighbor;
 
 						FBounds2D CellBounds;
-						if(!Seen.Contains(NeighCoord) && GetCellBounds3D(NeighCoord, Min, Max, HalfExtents3D))
+						if(!Seen.Contains(NeighCoord) && GetCellBounds3DScaled(NeighCoord, Min, Max, HalfExtents3D))
 						{
 							if(TBox<T, 3>::RaycastFast(Min, Max, StartPoint, Dir, InvDir, bParallel, CurrentLength, InvCurrentLength, ToI, HitPoint))
 							{
@@ -660,7 +698,7 @@ namespace Chaos
 
 					// Check the current cell, if we hit its 3D bound we can move on to narrow phase
 					const TVector<int32, 2> Coord = CellCoord.Index;
-					if(GetCellBounds3D(Coord, Min, Max, HalfExtents3D) &&
+					if(GetCellBounds3DScaled(Coord, Min, Max, HalfExtents3D) &&
 						TBox<T, 3>::RaycastFast(Min, Max, StartPoint, Dir, InvDir, bParallel, CurrentLength, InvCurrentLength, ToI, HitPoint))
 					{
 						bool bContinue = Visitor.VisitSweep(CellCoord.Index[1] * (GeomData.NumCols - 1) + CellCoord.Index[0], CurrentLength);
@@ -754,10 +792,15 @@ namespace Chaos
 	{
 		OutInterssctions.Reset();
 
-		InFlatBounds.Min = FlattenedBounds.Clamp(InFlatBounds.Min);
-		InFlatBounds.Max = FlattenedBounds.Clamp(InFlatBounds.Max);
-		TVector<int32, 2> MinCell = FlatGrid.Cell(InFlatBounds.Min);
-		TVector<int32, 2> MaxCell = FlatGrid.Cell(InFlatBounds.Max);
+		const FBounds2D FlatBounds = GetFlatBounds();
+		const TVector<T, 2> Scale2D(GeomData.Scale[0], GeomData.Scale[1]);
+
+		InFlatBounds.Min = FlatBounds.Clamp(InFlatBounds.Min);
+		InFlatBounds.Max = FlatBounds.Clamp(InFlatBounds.Max);
+		TVector<int32, 2> MinCell = FlatGrid.Cell(InFlatBounds.Min / Scale2D);
+		TVector<int32, 2> MaxCell = FlatGrid.Cell(InFlatBounds.Max / Scale2D);
+		MinCell = FlatGrid.ClampIndex(MinCell);
+		MaxCell = FlatGrid.ClampIndex(MaxCell);
 
 		const int32 DeltaX = MaxCell[0] - MinCell[0];
 		const int32 DeltaY = MaxCell[1] - MinCell[1];
@@ -773,6 +816,15 @@ namespace Chaos
 		return OutInterssctions.Num() > 0;
 	}
 
+	template<typename T>
+	typename Chaos::THeightField<T>::FBounds2D Chaos::THeightField<T>::GetFlatBounds() const
+	{
+		FBounds2D Result;
+		Result.Min = TVector<T, 2>(CachedBounds.Min()[0], CachedBounds.Min()[1]);
+		Result.Max = TVector<T, 2>(CachedBounds.Max()[0], CachedBounds.Max()[1]);
+		return Result;
+	}
+
 	template <typename T>
 	bool THeightField<T>::Overlap(const TVector<T, 3>& Point, const T Thickness) const
 	{
@@ -986,7 +1038,7 @@ namespace Chaos
 
 			TVector<T, 3> Points[4];
 
-			GeomData.GetPoints(FaceIndex, Points);
+			GeomData.GetPointsScaled(FaceIndex, Points);
 
 			TVector<T, 3> A;
 			TVector<T, 3> B;
@@ -1016,6 +1068,30 @@ namespace Chaos
 		return TVector<T, 3>(0, 0, 1);
 	}
 
+	template<typename T>
+	void THeightField<T>::CalcBounds()
+	{
+		// Flatten out the Z axis
+		FlattenedBounds = GetFlatBounds();
+
+		BuildQueryData();
+
+		// Cache per-cell bounds
+		const int32 NumX = GeomData.NumCols - 1;
+		const int32 NumY = GeomData.NumRows - 1;
+		GeomData.CellBounds.SetNum(NumX * NumY);
+		for(int32 XIndex = 0; XIndex < NumX; ++XIndex)
+		{
+			for(int32 YIndex = 0; YIndex < NumY; ++YIndex)
+			{
+				const TVector<int32, 2> Cell(XIndex, YIndex);
+				TVector<T, 3> Min, Max;
+				CalcCellBounds3D(Cell, Min, Max);
+				GeomData.CellBounds[YIndex * NumX + XIndex] = TBox<T, 3>(Min, Max);
+			}
+		}
+	}
+
 	template<typename T>
 	void THeightField<T>::BuildQueryData()
 	{
@@ -1032,3 +1108,4 @@ namespace Chaos
 }
 
 template class Chaos::THeightField<float>;
+
diff --git a/Engine/Source/Runtime/Experimental/Chaos/Public/Chaos/HeightField.h b/Engine/Source/Runtime/Experimental/Chaos/Public/Chaos/HeightField.h
--- a/Engine/Source/Runtime/Experimental/Chaos/Public/Chaos/HeightField.h
+++ b/Engine/Source/Runtime/Experimental/Chaos/Public/Chaos/HeightField.h
@@ -11,7 +11,6 @@
 #include "Templates/UnrealTypeTraits.h"
 #include "UniformGrid.h"
 #include "Utilities.h"
-PRAGMA_DISABLE_OPTIMIZATION
 
 namespace Chaos
 {
@@ -28,8 +27,10 @@ namespace Chaos
 		using TImplicitObject<T, 3>::GetTypeName;
 
 		THeightField(TArray<T>&& Height, int32 InNumRows, int32 InNumCols, const TVector<T,3>& InScale);
+		THeightField(TArrayView<const uint16> InHeights, int32 InNumRows, int32 InNumCols, const TVector<T, 3>& InScale);
 		THeightField(const THeightField& Other) = delete;
 		THeightField(THeightField&& Other) = default;
+
 		virtual ~THeightField() {}
 
 		virtual T PhiWithNormal(const TVector<T, 3>& x, TVector<T, 3>& Normal) const
@@ -47,7 +48,8 @@ namespace Chaos
 
 		virtual const TBox<T, 3>& BoundingBox() const
 		{
-			return LocalBounds;
+			CachedBounds = TBox<T, 3>(LocalBounds.Min() * GeomData.Scale, LocalBounds.Max() * GeomData.Scale);
+			return CachedBounds;
 		}
 
 		virtual uint32 GetTypeHash() const override
@@ -75,12 +77,22 @@ namespace Chaos
 			
 			GeomData.Serialize(Ar);
 
+			Ar << FlatGrid;
+			Ar << FlattenedBounds.Min;
+			Ar << FlattenedBounds.Max;
+			Ar << LocalBounds;
+
 			if(Ar.IsLoading())
 			{
 				BuildQueryData();
 			}
 		}
 
+		void SetScale(const TVector<T, 3>& InScale)
+		{
+			GeomData.Scale = InScale;
+		}
+
 		template<typename InStorageType, typename InRealType>
 		struct FData
 		{
@@ -137,15 +149,20 @@ namespace Chaos
 				const int32 X = Index % (NumCols);
 				const int32 Y = Index / (NumCols);
 
-				TVector<typename FDataType::RealType, 3> P0 = {(typename FDataType::RealType)X, (typename FDataType::RealType)Y, H0};
-				TVector<typename FDataType::RealType, 3> P1 = {(typename FDataType::RealType)X + 1, (typename FDataType::RealType)Y, H1};
-				TVector<typename FDataType::RealType, 3> P2 = {(typename FDataType::RealType)X, (typename FDataType::RealType)Y + 1, H2};
-				TVector<typename FDataType::RealType, 3> P3 = {(typename FDataType::RealType)X + 1, (typename FDataType::RealType)Y + 1, H3};
-				
-				OutPts[0] = (P0 * Scale);
-				OutPts[1] = (P1 * Scale);
-				OutPts[2] = (P2 * Scale);
-				OutPts[3] = (P3 * Scale);
+				OutPts[0] = {(typename FDataType::RealType)X, (typename FDataType::RealType)Y, H0};
+				OutPts[1] = {(typename FDataType::RealType)X + 1, (typename FDataType::RealType)Y, H1};
+				OutPts[2] = {(typename FDataType::RealType)X, (typename FDataType::RealType)Y + 1, H2};
+				OutPts[3] = {(typename FDataType::RealType)X + 1, (typename FDataType::RealType)Y + 1, H3};
+			}
+
+			FORCEINLINE void GetPointsScaled(int32 Index, TVector<T, 3> OutPts[4]) const
+			{
+				GetPoints(Index, OutPts);
+
+				OutPts[0] *= Scale;
+				OutPts[1] *= Scale;
+				OutPts[2] *= Scale;
+				OutPts[3] *= Scale;
 			}
 
 			void Serialize(FChaosArchive& Ar)
@@ -168,6 +185,9 @@ namespace Chaos
 				Ar << MaxValue;
 				Ar << NumRows;
 				Ar << NumCols;
+				Ar << Range;
+				Ar << HeightPerUnit;
+				Ar << CellBounds;
 			}
 		};
 
@@ -329,6 +349,8 @@ namespace Chaos
 		// Helpers for accessing bounds
 		bool GetCellBounds2D(const TVector<int32, 2> InCoord, FBounds2D& OutBounds, const TVector<T, 2>& InInflate = {0}) const;
 		bool GetCellBounds3D(const TVector<int32, 2> InCoord, TVector<T, 3>& OutMin, TVector<T, 3>& OutMax, const TVector<T, 3>& InInflate = TVector<T, 3>(0)) const;
+		bool GetCellBounds2DScaled(const TVector<int32, 2> InCoord, FBounds2D& OutBounds, const TVector<T, 2>& InInflate = {0}) const;
+		bool GetCellBounds3DScaled(const TVector<int32, 2> InCoord, TVector<T, 3>& OutMin, TVector<T, 3>& OutMax, const TVector<T, 3>& InInflate = TVector<T, 3>(0)) const;
 		bool CalcCellBounds3D(const TVector<int32, 2> InCoord, TVector<T, 3>& OutMin, TVector<T, 3>& OutMax, const TVector<T, 3>& InInflate = TVector<T, 3>(0)) const;
 
 		// Query functions - sweep, ray, overlap
@@ -337,13 +359,18 @@ namespace Chaos
 		bool GridCast(const TVector<T, 3>& StartPoint, const TVector<T, 3>& Dir, const T Length, THeightfieldRaycastVisitor<T>& Visitor) const;
 		bool GetGridIntersections(FBounds2D InFlatBounds, TArray<TVector<int32, 2>>& OutInterssctions) const;
 		
+		FBounds2D GetFlatBounds() const;
+
 		// Grid for queries, faster than bounding volumes for heightfields
 		TUniformGrid<T, 2> FlatGrid;
 		// Bounds in 2D of the whole heightfield, to clip queries against
 		FBounds2D FlattenedBounds;
 		// 3D bounds for the heightfield, for insertion to the scene structure
 		TBox<T, 3> LocalBounds;
+		// Cached when bounds are requested. Mutable to allow GetBounds to be logical const
+		mutable TBox<T, 3> CachedBounds;
 
+		void CalcBounds();
 		void BuildQueryData();
 		
 		// Needed for serialization
@@ -351,4 +378,3 @@ namespace Chaos
 		friend TImplicitObject<T, 3>;
 	};
 }
-PRAGMA_ENABLE_OPTIMIZATION
\ No newline at end of file
diff --git a/Engine/Source/Runtime/Landscape/Classes/LandscapeHeightfieldCollisionComponent.h b/Engine/Source/Runtime/Landscape/Classes/LandscapeHeightfieldCollisionComponent.h
--- a/Engine/Source/Runtime/Landscape/Classes/LandscapeHeightfieldCollisionComponent.h
+++ b/Engine/Source/Runtime/Landscape/Classes/LandscapeHeightfieldCollisionComponent.h
@@ -23,6 +23,14 @@ struct FConvexVolume;
 struct FEngineShowFlags;
 struct FNavigableGeometryExport;
 
+#if WITH_CHAOS
+namespace Chaos
+{
+	template<typename T>
+	class THeightField;
+}
+#endif
+
 #if WITH_PHYSX
 namespace physx
 {
@@ -79,45 +87,31 @@ class ULandscapeHeightfieldCollisionComponent : public UPrimitiveComponent
 	UFUNCTION(BlueprintCallable, Category = "Landscape")
 	ULandscapeComponent* GetRenderComponent() const;
 
-	struct FPhysXHeightfieldRef : public FRefCountedObject
+	struct FHeightfieldGeometryRef : public FRefCountedObject
 	{
 		FGuid Guid;
 
 #if WITH_PHYSX
 		/** List of PxMaterials used on this landscape */
 		TArray<physx::PxMaterial*> UsedPhysicalMaterialArray;
-		physx::PxHeightField* RBHeightfield;
-		physx::PxHeightField* RBHeightfieldSimple;
+		physx::PxHeightField* RBHeightfield = nullptr;
+		physx::PxHeightField* RBHeightfieldSimple = nullptr;
 #if WITH_EDITOR
-		physx::PxHeightField* RBHeightfieldEd; // Used only by landscape editor, does not have holes in it
+		physx::PxHeightField* RBHeightfieldEd = nullptr; // Used only by landscape editor, does not have holes in it
 #endif	//WITH_EDITOR
 #endif	//WITH_PHYSX
 
-		/** tors **/
-		FPhysXHeightfieldRef()
-#if WITH_PHYSX
-			: RBHeightfield(nullptr)
-			, RBHeightfieldSimple(nullptr)
+#if WITH_CHAOS
+		TUniquePtr<Chaos::THeightField<float>> Heightfield = nullptr;
+	    TUniquePtr<Chaos::THeightField<float>> HeightfieldSimple = nullptr;
 #if WITH_EDITOR
-			, RBHeightfieldEd(nullptr)
-#endif	//WITH_EDITOR
-#endif	//WITH_PHYSX
-		{
-		}
+		TUniquePtr<Chaos::THeightField<float>> EditorHeightfield = nullptr;
+#endif
+#endif
 
-		FPhysXHeightfieldRef(FGuid& InGuid)
-			: Guid(InGuid)
-#if WITH_PHYSX
-			, RBHeightfield(nullptr)
-			, RBHeightfieldSimple(nullptr)
-#if WITH_EDITOR
-			, RBHeightfieldEd(nullptr)
-#endif	//WITH_EDITOR
-#endif	//WITH_PHYSX
-		{
-		}
+		FHeightfieldGeometryRef(FGuid& InGuid);
 
-		virtual ~FPhysXHeightfieldRef();
+		virtual ~FHeightfieldGeometryRef();
 	};
 	
 #if WITH_EDITORONLY_DATA
@@ -160,7 +154,7 @@ class ULandscapeHeightfieldCollisionComponent : public UPrimitiveComponent
 	TArray<UPhysicalMaterial*>					CookedPhysicalMaterials;
 	
 	/** Physics engine version of heightfield data. */
-	TRefCountPtr<FPhysXHeightfieldRef>	HeightfieldRef;
+	TRefCountPtr<FHeightfieldGeometryRef>	HeightfieldRef;
 	
 	/** Cached PxHeightFieldSamples values for navmesh generation. Note that it's being used only if navigation octree is set up for lazy geometry exporting */
 	int32 HeightfieldRowsCount;
diff --git a/Engine/Source/Runtime/Landscape/Private/LandscapeCollision.cpp b/Engine/Source/Runtime/Landscape/Private/LandscapeCollision.cpp
--- a/Engine/Source/Runtime/Landscape/Private/LandscapeCollision.cpp
+++ b/Engine/Source/Runtime/Landscape/Private/LandscapeCollision.cpp
@@ -52,10 +52,15 @@
 #if WITH_CHAOS
 #include "PhysXToChaosUtil.h"
 #include "Chaos/ParticleHandle.h"
+#include "Chaos/HeightField.h"
+#include "Chaos/Vector.h"
 #endif
 
 using namespace PhysicsInterfaceTypes;
 
+// Global switch for whether to read/write to DDC for landscape cooked data
+bool GLandscapeCollisionSkipDDC = false;
+
 #if ENABLE_COOK_STATS
 namespace LandscapeCollisionCookStats
 {
@@ -69,9 +74,14 @@ namespace LandscapeCollisionCookStats
 }
 #endif
 
-TMap<FGuid, ULandscapeHeightfieldCollisionComponent::FPhysXHeightfieldRef* > GSharedHeightfieldRefs;
+TMap<FGuid, ULandscapeHeightfieldCollisionComponent::FHeightfieldGeometryRef* > GSharedHeightfieldRefs;
 
-ULandscapeHeightfieldCollisionComponent::FPhysXHeightfieldRef::~FPhysXHeightfieldRef()
+ULandscapeHeightfieldCollisionComponent::FHeightfieldGeometryRef::FHeightfieldGeometryRef(FGuid& InGuid)
+	: Guid(InGuid)
+{
+}
+
+ULandscapeHeightfieldCollisionComponent::FHeightfieldGeometryRef::~FHeightfieldGeometryRef()
 {
 #if WITH_PHYSX
 	// Free the existing heightfield data.
@@ -119,7 +129,7 @@ ULandscapeMeshCollisionComponent::FPhysXMeshRef::~FPhysXMeshRef()
 }
 
 // Generate a new guid to force a recache of landscape collison derived data
-#define LANDSCAPE_COLLISION_DERIVEDDATA_VER	TEXT("84A5A09B87CA4ED3B9B301DECE89D011")
+#define LANDSCAPE_COLLISION_DERIVEDDATA_VER	TEXT("582D8D523EAF429B92760385A69C8D63")
 
 static FString GetHFDDCKeyString(const FName& Format, bool bDefMaterial, const FGuid& StateId, const TArray<UPhysicalMaterial*>& PhysicalMaterials)
 {
@@ -152,7 +162,15 @@ static FString GetHFDDCKeyString(const FName& Format, bool bDefMaterial, const F
 		CombinedStateId = FGuid(Hash[0] ^ Hash[4], Hash[1], Hash[2], Hash[3]);
 	}
 
-	const FString KeyPrefix = FString::Printf(TEXT("%s_%s"), *Format.ToString(), (bDefMaterial ? TEXT("VIS") : TEXT("FULL")));
+#if PHYSICS_INTERFACE_PHYSX
+	const FString InterfacePrefix = TEXT("PHYSX");
+#elif WITH_CHAOS
+	const FString InterfacePrefix = TEXT("CHAOS");
+#else
+	const FString InterfacePrefix = TEXT("UNDEFINED");
+#endif
+
+	const FString KeyPrefix = FString::Printf(TEXT("%s_%s_%s"), *InterfacePrefix, *Format.ToString(), (bDefMaterial ? TEXT("VIS") : TEXT("FULL")));
 	return FDerivedDataCacheInterface::BuildCacheKey(*KeyPrefix, LANDSCAPE_COLLISION_DERIVEDDATA_VER, *CombinedStateId.ToString());
 }
 
@@ -229,8 +247,10 @@ void ULandscapeHeightfieldCollisionComponent::OnCreatePhysicsState()
 			const float SimpleCollisionScale = bCreateSimpleCollision ? CollisionScale * CollisionSizeQuads / SimpleCollisionSizeQuads : 0;
 
 			// Create the geometry
+			FVector FinalScale(LandscapeScale.X * CollisionScale, LandscapeScale.Y * CollisionScale, LandscapeScale.Z * LANDSCAPE_ZSCALE);
 			PxHeightFieldGeometry LandscapeComponentGeom(HeightfieldRef->RBHeightfield, PxMeshGeometryFlag::eDOUBLE_SIDED, LandscapeScale.Z * LANDSCAPE_ZSCALE, LandscapeScale.Y * CollisionScale, LandscapeScale.X * CollisionScale);
 
+#if PHYSICS_INTERFACE_PHYSX
 			if (LandscapeComponentGeom.isValid())
 			{
 				// Creating both a sync and async actor, since this object is static
@@ -256,31 +276,6 @@ void ULandscapeHeightfieldCollisionComponent::OnCreatePhysicsState()
 
 				HeightFieldActorSync->attachShape(*HeightFieldShapeSync);
 
-#if WITH_CHAOS
-				FActorCreationParams Params;
-				Params.InitialTM = LandscapeShapeTM;
-				Params.bQueryOnly = true;
-				Params.Scene = GetWorld()->GetPhysicsScene();
-				FPhysicsActorHandle PhysHandle;
-				FPhysicsInterface::CreateActor(Params, PhysHandle);
-
-				TUniquePtr<Chaos::TImplicitObjectTransformed<float, 3>> ChaosHeightField = PxShapeToChaosGeom(HeightFieldShapeSync);
-				TUniquePtr<Chaos::TPerShapeData<float, 3>> NewShape = Chaos::TPerShapeData<float, 3>::CreatePerShapeData();
-
-				NewShape->Geometry = MakeSerializable(ChaosHeightField);
-				NewShape->QueryData = QueryFilterData;
-				NewShape->SimData = SimFilterData;
-
-				Chaos::TShapesArray<float, 3> ShapeArray;
-				ShapeArray.Emplace(MoveTemp(NewShape));
-				PhysHandle->SetGeometry(MoveTemp(ChaosHeightField));
-				PhysHandle->SetShapesArray(MoveTemp(ShapeArray));
-
-				FTransform ActorTM(LandscapeComponentMatrix);
-				PhysHandle->SetX(ActorTM.GetTranslation());
-				PhysHandle->SetR(ActorTM.GetRotation());
-#endif
-
 				// attachShape holds its own ref(), so release this here.
 				HeightFieldShapeSync->release();
 
@@ -350,15 +345,6 @@ void ULandscapeHeightfieldCollisionComponent::OnCreatePhysicsState()
 				BodyInstance.PhysicsUserData = FPhysicsUserData(&BodyInstance);
 				BodyInstance.OwnerComponent = this;
 
-#if WITH_CHAOS
-				PhysHandle->SetUserData(&BodyInstance.PhysicsUserData);
-#endif
-
-#if WITH_CHAOS || WITH_IMMEDIATE_PHYSX
-				TArray<FPhysicsActorHandle> Actors;
-				Actors.Add(PhysHandle);
-				PhysScene->AddActorsToScene_AssumesLocked(Actors);
-#else
 				BodyInstance.ActorHandle.SyncActor = HeightFieldActorSync;
 				HeightFieldActorSync->userData = &BodyInstance.PhysicsUserData;
 
@@ -366,6 +352,85 @@ void ULandscapeHeightfieldCollisionComponent::OnCreatePhysicsState()
 				PxScene* SyncScene = PhysScene->GetPxScene();
 				SCOPED_SCENE_WRITE_LOCK(SyncScene);
 				SyncScene->addActor(*HeightFieldActorSync);
+			}
+			else
+#elif WITH_CHAOS
+			{
+				FActorCreationParams Params;
+				Params.InitialTM = LandscapeComponentTransform;
+				Params.InitialTM.SetScale3D(FVector(0));
+				Params.bQueryOnly = true;
+				Params.Scene = GetWorld()->GetPhysicsScene();
+				FPhysicsActorHandle PhysHandle;
+				FPhysicsInterface::CreateActor(Params, PhysHandle);
+
+				Chaos::TShapesArray<float, 3> ShapeArray;
+				TArray<TUniquePtr<Chaos::TImplicitObject<float, 3>>> Geoms;
+
+				// First add complex geometry
+				TUniquePtr<Chaos::TPerShapeData<float, 3>> NewShape = Chaos::TPerShapeData<float, 3>::CreatePerShapeData();
+
+				HeightfieldRef->Heightfield->SetScale(FinalScale);
+				TUniquePtr<Chaos::TImplicitObjectTransformed<float, 3>> ChaosHeightFieldFromCooked = MakeUnique<Chaos::TImplicitObjectTransformed<float, 3>>(MakeSerializable(HeightfieldRef->Heightfield), Chaos::TRigidTransform<float, 3>(FTransform::Identity));
+
+				// Setup filtering
+				FCollisionFilterData QueryFilterData, SimFilterData;
+				CreateShapeFilterData(GetCollisionObjectType(), FMaskFilter(0), GetOwner()->GetUniqueID(), GetCollisionResponseToChannels(), GetUniqueID(), 0, QueryFilterData, SimFilterData, true, false, true);
+
+				// Heightfield is used for simple and complex collision
+				QueryFilterData.Word3 |= bCreateSimpleCollision ? EPDF_ComplexCollision : (EPDF_SimpleCollision | EPDF_ComplexCollision);
+				SimFilterData.Word3 |= bCreateSimpleCollision ? EPDF_ComplexCollision : (EPDF_SimpleCollision | EPDF_ComplexCollision);
+
+				NewShape->Geometry = MakeSerializable(ChaosHeightFieldFromCooked);
+				NewShape->QueryData = QueryFilterData;
+				NewShape->SimData = SimFilterData;
+
+				Geoms.Emplace(MoveTemp(ChaosHeightFieldFromCooked));
+				ShapeArray.Emplace(MoveTemp(NewShape));
+
+				// Add simple geometry if necessary
+				if(bCreateSimpleCollision)
+				{
+					TUniquePtr<Chaos::TPerShapeData<float, 3>> NewSimpleShape = Chaos::TPerShapeData<float, 3>::CreatePerShapeData();
+
+					HeightfieldRef->HeightfieldSimple->SetScale(FinalScale);
+					TUniquePtr<Chaos::TImplicitObjectTransformed<float, 3>> ChaosSimpleHeightFieldFromCooked = MakeUnique<Chaos::TImplicitObjectTransformed<float, 3>>(MakeSerializable(HeightfieldRef->HeightfieldSimple), Chaos::TRigidTransform<float, 3>(FTransform::Identity));
+
+					FCollisionFilterData QueryFilterDataSimple = QueryFilterData;
+					FCollisionFilterData SimFilterDataSimple = SimFilterData;
+					QueryFilterDataSimple.Word3 = (QueryFilterDataSimple.Word3 & ~EPDF_ComplexCollision) | EPDF_SimpleCollision;
+					SimFilterDataSimple.Word3 = (SimFilterDataSimple.Word3 & ~EPDF_ComplexCollision) | EPDF_SimpleCollision;
+
+					NewShape->Geometry = MakeSerializable(ChaosSimpleHeightFieldFromCooked);
+					NewShape->QueryData = QueryFilterDataSimple;
+					NewShape->SimData = SimFilterDataSimple;
+
+					Geoms.Emplace(MoveTemp(ChaosHeightFieldFromCooked));
+					ShapeArray.Emplace(MoveTemp(NewSimpleShape));
+				}
+
+				if(Geoms.Num() == 1)
+				{
+					PhysHandle->SetGeometry(MoveTemp(Geoms[0]));
+				}
+				else
+				{
+					PhysHandle->SetGeometry(MakeUnique<Chaos::TImplicitObjectUnion<float, 3>>(MoveTemp(Geoms)));
+				}
+				
+				PhysHandle->SetShapesArray(MoveTemp(ShapeArray));
+
+				FPhysScene* PhysScene = GetWorld()->GetPhysicsScene();
+
+				// Set body instance data
+				BodyInstance.PhysicsUserData = FPhysicsUserData(&BodyInstance);
+				BodyInstance.OwnerComponent = this;
+
+				PhysHandle->SetUserData(&BodyInstance.PhysicsUserData);
+
+				TArray<FPhysicsActorHandle> Actors;
+				Actors.Add(PhysHandle);
+				PhysScene->AddActorsToScene_AssumesLocked(Actors);
 #endif
 			}
 		}
@@ -385,13 +450,12 @@ void ULandscapeHeightfieldCollisionComponent::ApplyWorldOffset(const FVector& In
 
 void ULandscapeHeightfieldCollisionComponent::CreateCollisionObject()
 {
-#if WITH_PHYSX
 	// If we have not created a heightfield yet - do it now.
 	if (!IsValidRef(HeightfieldRef))
 	{
 		UWorld* World = GetWorld();
 
-		FPhysXHeightfieldRef* ExistingHeightfieldRef = nullptr;
+		FHeightfieldGeometryRef* ExistingHeightfieldRef = nullptr;
 		bool bCheckDDC = true;
 
 		if (!HeightfieldGuid.IsValid())
@@ -429,8 +493,9 @@ void ULandscapeHeightfieldCollisionComponent::CreateCollisionObject()
 
 			if (CookedCollisionData.Num())
 			{
-				HeightfieldRef = GSharedHeightfieldRefs.Add(HeightfieldGuid, new FPhysXHeightfieldRef(HeightfieldGuid));
+				HeightfieldRef = GSharedHeightfieldRefs.Add(HeightfieldGuid, new FHeightfieldGeometryRef(HeightfieldGuid));
 
+#if WITH_PHYSX && PHYSICS_INTERFACE_PHYSX
 				// Create heightfield shape
 				{
 					FPhysXInputStream HeightFieldStream(CookedCollisionData.GetData(), CookedCollisionData.Num());
@@ -443,15 +508,7 @@ void ULandscapeHeightfieldCollisionComponent::CreateCollisionObject()
 
 				for (UPhysicalMaterial* PhysicalMaterial : CookedPhysicalMaterials)
 				{
-#if WITH_CHAOS || WITH_IMMEDIATE_PHYSX
-#if WITH_PHYSX
-					//todo: total hack until we get landscape fully converted to chaos
-					HeightfieldRef->UsedPhysicalMaterialArray.Add(GPhysXSDK->createMaterial(1,1,1));
-#endif
-#else
-					const FPhysicsMaterialHandle_PhysX& MaterialHandle = PhysicalMaterial->GetPhysicsMaterial();
-					HeightfieldRef->UsedPhysicalMaterialArray.Add(MaterialHandle.Material);
-#endif
+					HeightfieldRef->UsedPhysicalMaterialArray.Add(GPhysXSDK->createMaterial(1, 1, 1));
 				}
 
 				// Release cooked collison data
@@ -473,17 +530,58 @@ void ULandscapeHeightfieldCollisionComponent::CreateCollisionObject()
 					}
 				}
 #endif //WITH_EDITOR
+#elif WITH_CHAOS
+				// Create heightfields
+				{
+					FMemoryReader Reader(CookedCollisionData);
+					Chaos::FChaosArchive Ar(Reader);
+					bool bContainsSimple = false;
+					Ar << bContainsSimple;
+					Ar << HeightfieldRef->Heightfield;
+
+					if(bContainsSimple)
+					{
+						Ar << HeightfieldRef->HeightfieldSimple;
+					}
+				}
+
+				// Register materials
+				for(UPhysicalMaterial* PhysicalMaterial : CookedPhysicalMaterials)
+				{
+					//todo: total hack until we get landscape fully converted to chaos
+					HeightfieldRef->UsedPhysicalMaterialArray.Add(GPhysXSDK->createMaterial(1, 1, 1));
+				}
+
+#if WITH_EDITOR
+				// Create heightfield for the landscape editor (no holes in it)
+				if(!World->IsGameWorld())
+				{
+					TArray<UPhysicalMaterial*> CookedMaterialsEd;
+					if(CookCollisionData(PhysicsFormatName, true, bCheckDDC, CookedCollisionDataEd, CookedMaterialsEd))
+					{
+						FMemoryReader Reader(CookedCollisionDataEd);
+						Chaos::FChaosArchive Ar(Reader);
+
+						// Don't actually care about this but need to strip it out of the data
+						bool bContainsSimple = false;
+						Ar << bContainsSimple;
+						Ar << HeightfieldRef->EditorHeightfield;
+
+						CookedCollisionDataEd.Empty();
+					}
+				}
+#endif //WITH_EDITOR
+#endif
 			}
 		}
 	}
-#endif //WITH_PHYSX
 }
 
 #if WITH_EDITOR
 void ULandscapeHeightfieldCollisionComponent::SpeculativelyLoadAsyncDDCCollsionData()
 {
 #if WITH_PHYSX
-	if (GetLinkerUE4Version() >= VER_UE4_LANDSCAPE_SERIALIZE_PHYSICS_MATERIALS)
+	if (GetLinkerUE4Version() >= VER_UE4_LANDSCAPE_SERIALIZE_PHYSICS_MATERIALS && !GLandscapeCollisionSkipDDC)
 	{
 		UWorld* World = GetWorld();
 		if (World && HeightfieldGuid.IsValid() && CookedPhysicalMaterials.Num() > 0 && GSharedHeightfieldRefs.FindRef(HeightfieldGuid) == nullptr)
@@ -557,12 +655,15 @@ TArray<PxHeightFieldSample> ConvertHeightfieldDataForPhysx(const ULandscapeHeigh
 
 bool ULandscapeHeightfieldCollisionComponent::CookCollisionData(const FName& Format, bool bUseDefMaterial, bool bCheckDDC, TArray<uint8>& OutCookedData, TArray<UPhysicalMaterial*>& InOutMaterials) const
 {
-#if WITH_PHYSX
 	COOK_STAT(auto Timer = LandscapeCollisionCookStats::HeightfieldUsageStats.TimeSyncWork());
+	 
+	bool Succeeded = false;
+	TArray<uint8> OutData;
+
 	// we have 2 versions of collision objects
 	const int32 CookedDataIndex = bUseDefMaterial ? 0 : 1;
 
-	if (bCheckDDC && HeightfieldGuid.IsValid())
+	if (!GLandscapeCollisionSkipDDC && bCheckDDC && HeightfieldGuid.IsValid())
 	{
 		// Ensure that content was saved with physical materials before using DDC data
 		if (GetLinkerUE4Version() >= VER_UE4_LANDSCAPE_SERIALIZE_PHYSICS_MATERIALS)
@@ -633,6 +734,7 @@ bool ULandscapeHeightfieldCollisionComponent::CookCollisionData(const FName& For
 		check(DominantLayerData.GetElementCount() == NumSamples + NumSimpleSamples);
 	}
 
+#if WITH_PHYSX && PHYSICS_INTERFACE_PHYSX
 	// List of materials which is actually used by heightfield
 	InOutMaterials.Empty();
 
@@ -659,25 +761,53 @@ bool ULandscapeHeightfieldCollisionComponent::CookCollisionData(const FName& For
 
 	//
 	FIntPoint HFSize = FIntPoint(CollisionSizeVerts, CollisionSizeVerts);
-	TArray<uint8> OutData;
 
 	ITargetPlatformManagerModule* TPM = GetTargetPlatformManager();
 	const IPhysXCooking* Cooker = TPM->FindPhysXCooking(Format);
-	bool Result = Cooker->CookHeightField(Format, HFSize, Samples.GetData(), Samples.GetTypeSize(), OutData);
+	Succeeded = Cooker->CookHeightField(Format, HFSize, Samples.GetData(), Samples.GetTypeSize(), OutData);
 
-	if (Result && bGenerateSimpleCollision)
+	if (Succeeded && bGenerateSimpleCollision)
 	{
 		FIntPoint HFSizeSimple = FIntPoint(SimpleCollisionSizeVerts, SimpleCollisionSizeVerts);
-		Result = Cooker->CookHeightField(Format, HFSizeSimple, SimpleSamples.GetData(), SimpleSamples.GetTypeSize(), OutData);
+		Succeeded = Cooker->CookHeightField(Format, HFSizeSimple, SimpleSamples.GetData(), SimpleSamples.GetTypeSize(), OutData);
 	}
+#elif WITH_CHAOS
+	TUniquePtr<Chaos::THeightField<float>> Heightfield = nullptr;
+	TUniquePtr<Chaos::THeightField<float>> HeightfieldSimple = nullptr;
 
-	if (Result)
+	FMemoryWriter Writer(OutData);
+	Chaos::FChaosArchive Ar(Writer);
+
+	bool bSerializeGenerateSimpleCollision = bGenerateSimpleCollision;
+	Ar << bSerializeGenerateSimpleCollision;
+
+	TArrayView<const uint16> ComplexHeightView(Heights, NumSamples);
+	Heightfield = MakeUnique<Chaos::THeightField<float>>(ComplexHeightView, CollisionSizeVerts, CollisionSizeVerts, Chaos::TVector<float, 3>(1));
+	Ar << Heightfield;
+	if(bGenerateSimpleCollision)
+	{
+		TArrayView<const uint16> SimpleHeightView(Heights + NumSamples, NumSimpleSamples);
+		HeightfieldSimple = MakeUnique<Chaos::THeightField<float>>(SimpleHeightView, CollisionSizeVerts, CollisionSizeVerts, Chaos::TVector<float, 3>(1));
+		Ar << HeightfieldSimple;
+	}
+
+	CollisionHeightData.Unlock();
+	if(DominantLayers)
+	{
+		DominantLayerData.Unlock();
+	}
+
+	Succeeded = true;
+
+#endif
+
+	if (Succeeded)
 	{
 		COOK_STAT(Timer.AddMiss(OutData.Num()));
 		OutCookedData.SetNumUninitialized(OutData.Num());
 		FMemory::Memcpy(OutCookedData.GetData(), OutData.GetData(), OutData.Num());
 
-		if (bShouldSaveCookedDataToDDC[CookedDataIndex] && HeightfieldGuid.IsValid())
+		if (!GLandscapeCollisionSkipDDC && bShouldSaveCookedDataToDDC[CookedDataIndex] && HeightfieldGuid.IsValid())
 		{
 			GetDerivedDataCacheRef().Put(*GetHFDDCKeyString(Format, bUseDefMaterial, HeightfieldGuid, InOutMaterials), OutCookedData);
 			bShouldSaveCookedDataToDDC[CookedDataIndex] = false;
@@ -691,10 +821,7 @@ bool ULandscapeHeightfieldCollisionComponent::CookCollisionData(const FName& For
 		InOutMaterials.Empty();
 	}
 
-	return Result;
-#endif	// WITH_PHYSX
-
-	return false;
+	return Succeeded;
 }
 
 bool ULandscapeMeshCollisionComponent::CookCollisionData(const FName& Format, bool bUseDefMaterial, bool bCheckDDC, TArray<uint8>& OutCookedData, TArray<UPhysicalMaterial*>& InOutMaterials) const
@@ -704,7 +831,7 @@ bool ULandscapeMeshCollisionComponent::CookCollisionData(const FName& Format, bo
 	// we have 2 versions of collision objects
 	const int32 CookedDataIndex = bUseDefMaterial ? 0 : 1;
 
-	if (bCheckDDC)
+	if (!GLandscapeCollisionSkipDDC && bCheckDDC)
 	{
 		// Ensure that content was saved with physical materials before using DDC data
 		if (GetLinkerUE4Version() >= VER_UE4_LANDSCAPE_SERIALIZE_PHYSICS_MATERIALS && MeshGuid.IsValid())
@@ -880,7 +1007,7 @@ bool ULandscapeMeshCollisionComponent::CookCollisionData(const FName& Format, bo
 		OutCookedData.SetNumUninitialized(OutData.Num());
 		FMemory::Memcpy(OutCookedData.GetData(), OutData.GetData(), OutData.Num());
 
-		if (bShouldSaveCookedDataToDDC[CookedDataIndex] && MeshGuid.IsValid())
+		if (!GLandscapeCollisionSkipDDC && bShouldSaveCookedDataToDDC[CookedDataIndex] && MeshGuid.IsValid())
 		{
 			GetDerivedDataCacheRef().Put(*GetHFDDCKeyString(Format, bUseDefMaterial, MeshGuid, InOutMaterials), OutCookedData);
 			bShouldSaveCookedDataToDDC[CookedDataIndex] = false;
@@ -1442,7 +1569,7 @@ void ULandscapeHeightfieldCollisionComponent::Serialize(FArchive& Ar)
 		{
 			FName Format = Ar.CookingTarget()->GetPhysicsFormat(nullptr);
 			CookCollisionData(Format, false, true, CookedCollisionData, CookedPhysicalMaterials);
-			if (HeightfieldGuid.IsValid())
+			if (!GLandscapeCollisionSkipDDC && HeightfieldGuid.IsValid())
 			{
 				GetDerivedDataCacheRef().Put(*GetHFDDCKeyString(Format, false, HeightfieldGuid, CookedPhysicalMaterials), CookedCollisionData);
 			}
@@ -1735,15 +1862,18 @@ void ULandscapeHeightfieldCollisionComponent::PreSave(const class ITargetPlatfor
 			}
 		}
 
-		static const FName PhysicsFormatName(FPlatformProperties::GetPhysicsFormat());
-		if (CookedCollisionData.Num() && HeightfieldGuid.IsValid())
+		if(!GLandscapeCollisionSkipDDC)
 		{
-			GetDerivedDataCacheRef().Put(*GetHFDDCKeyString(PhysicsFormatName, false, HeightfieldGuid, CookedPhysicalMaterials), CookedCollisionData);
-		}
+			static const FName PhysicsFormatName(FPlatformProperties::GetPhysicsFormat());
+			if(CookedCollisionData.Num() && HeightfieldGuid.IsValid())
+			{
+				GetDerivedDataCacheRef().Put(*GetHFDDCKeyString(PhysicsFormatName, false, HeightfieldGuid, CookedPhysicalMaterials), CookedCollisionData);
+			}
 
-		if (CookedCollisionDataEd.Num() && HeightfieldGuid.IsValid())
-		{
-			GetDerivedDataCacheRef().Put(*GetHFDDCKeyString(PhysicsFormatName, true, HeightfieldGuid, TArray<UPhysicalMaterial*>()), CookedCollisionDataEd);
+			if(CookedCollisionDataEd.Num() && HeightfieldGuid.IsValid())
+			{
+				GetDerivedDataCacheRef().Put(*GetHFDDCKeyString(PhysicsFormatName, true, HeightfieldGuid, TArray<UPhysicalMaterial*>()), CookedCollisionDataEd);
+			}
 		}
 #endif// WITH_EDITOR
 	}
