Use landscape collision component map for to retrieve landscape height (checking in for Patrick Enfedaque).

diff --git a/Engine/Source/Runtime/Landscape/Classes/LandscapeHeightfieldCollisionComponent.h b/Engine/Source/Runtime/Landscape/Classes/LandscapeHeightfieldCollisionComponent.h
--- a/Engine/Source/Runtime/Landscape/Classes/LandscapeHeightfieldCollisionComponent.h
+++ b/Engine/Source/Runtime/Landscape/Classes/LandscapeHeightfieldCollisionComponent.h
@@ -189,6 +189,8 @@ public:
 	virtual ECollisionResponse GetCollisionResponseToChannel(ECollisionChannel Channel) const override;
 	virtual ECollisionChannel GetCollisionObjectType() const override;
 	virtual const FCollisionResponseContainer& GetCollisionResponseToChannels() const override;
+	virtual void OnRegister() override;
+	virtual void OnUnregister() override;
 	//~ End USceneComponent Interface.
 
 	//~ Begin UPrimitiveComponent Interface
@@ -217,10 +219,7 @@ public:
 	virtual void PostEditImport() override;
 	virtual void PostEditUndo() override;
 	//~ End UObject Interface.
-
-	/** Gets the landscape info object for this landscape */
-	ULandscapeInfo* GetLandscapeInfo() const;
-
+		
 	/**  We speculatively async load collision object from DDC to remove hitch when streaming */
 	void SpeculativelyLoadAsyncDDCCollsionData();
 
@@ -235,6 +234,8 @@ public:
 	/** Computes a hash of all the data that will impact final collision */
 	virtual uint32 ComputeCollisionHash() const;
 #endif
+	/** Gets the landscape info object for this landscape */
+	ULandscapeInfo* GetLandscapeInfo() const;
 
 	/** Creates collision object from a cooked collision data */
 	virtual void CreateCollisionObject();
diff --git a/Engine/Source/Runtime/Landscape/Classes/LandscapeInfo.h b/Engine/Source/Runtime/Landscape/Classes/LandscapeInfo.h
--- a/Engine/Source/Runtime/Landscape/Classes/LandscapeInfo.h
+++ b/Engine/Source/Runtime/Landscape/Classes/LandscapeInfo.h
@@ -17,6 +17,7 @@ class ULandscapeLayerInfoObject;
 class ULevel;
 class UMaterialInstanceConstant;
 struct FLandscapeEditorLayerSettings;
+class ULandscapeHeightfieldCollisionComponent;
 
 /** Structure storing Collision for LandscapeComponent Add */
 #if WITH_EDITORONLY_DATA
@@ -120,6 +121,8 @@ class ULandscapeInfo : public UObject
 public:
 	/** Map of the offsets (in component space) to the component. Valid in editor only. */
 	TMap<FIntPoint, ULandscapeComponent*> XYtoComponentMap;
+    /** Map of the offsets (in component space) to the collision components. Should always be valid. */
+	TMap<FIntPoint, ULandscapeHeightfieldCollisionComponent*> XYtoCollisionComponentMap;
 
 #if WITH_EDITORONLY_DATA
 	/** Lookup map used by the "add component" tool. Only available near valid LandscapeComponents.
@@ -260,4 +263,8 @@ public:
 
 	/** Deassociates passed landscape component with this info object*/
 	LANDSCAPE_API void UnregisterActorComponent(ULandscapeComponent* Component);
+
+	/** Server doesn't have ULandscapeComponent use CollisionComponents instead to get height on landscape */
+	LANDSCAPE_API void RegisterCollisionComponent(ULandscapeHeightfieldCollisionComponent* Component);
+	LANDSCAPE_API void UnregisterCollisionComponent(ULandscapeHeightfieldCollisionComponent* Component);
 };
diff --git a/Engine/Source/Runtime/Landscape/Private/Landscape.cpp b/Engine/Source/Runtime/Landscape/Private/Landscape.cpp
--- a/Engine/Source/Runtime/Landscape/Private/Landscape.cpp
+++ b/Engine/Source/Runtime/Landscape/Private/Landscape.cpp
@@ -3029,6 +3029,11 @@ void ULandscapeInfo::RegisterActor(ALandscapeProxy* Proxy, bool bMapCheck)
 	{
 		RegisterActorComponent(Proxy->LandscapeComponents[CompIdx], bMapCheck);
 	}
+
+	for (ULandscapeHeightfieldCollisionComponent* CollComp: Proxy->CollisionComponents)
+	{
+		RegisterCollisionComponent(CollComp);
+	}
 }
 
 void ULandscapeInfo::UnregisterActor(ALandscapeProxy* Proxy)
@@ -3070,12 +3075,54 @@ void ULandscapeInfo::UnregisterActor(ALandscapeProxy* Proxy)
 	}
 	XYtoComponentMap.Compact();
 
+	for (ULandscapeHeightfieldCollisionComponent* CollComp : Proxy->CollisionComponents)
+	{
+		if (CollComp)
+		{
+			UnregisterCollisionComponent(CollComp);
+		}
+	}
+	XYtoCollisionComponentMap.Compact();
+
 #if WITH_EDITOR
 	UpdateLayerInfoMap();
 	UpdateAllAddCollisions();
 #endif
 }
 
+void ULandscapeInfo::RegisterCollisionComponent(ULandscapeHeightfieldCollisionComponent* Component)
+{
+	if (Component == nullptr || !Component->IsRegistered())
+	{
+		return;
+	}
+
+	FIntPoint ComponentKey = Component->GetSectionBase() / Component->CollisionSizeQuads;
+	auto RegisteredComponent = XYtoCollisionComponentMap.FindRef(ComponentKey);
+
+	if (RegisteredComponent != Component)
+	{
+		if (RegisteredComponent == nullptr)
+		{
+			XYtoCollisionComponentMap.Add(ComponentKey, Component);
+		}
+	}
+}
+
+void ULandscapeInfo::UnregisterCollisionComponent(ULandscapeHeightfieldCollisionComponent* Component)
+{
+	if (ensure(Component))
+	{
+		FIntPoint ComponentKey = Component->GetSectionBase() / Component->CollisionSizeQuads;
+		auto RegisteredComponent = XYtoCollisionComponentMap.FindRef(ComponentKey);
+
+		if (RegisteredComponent == Component)
+		{
+			XYtoCollisionComponentMap.Remove(ComponentKey);
+		}
+	}
+}
+
 void ULandscapeInfo::RegisterActorComponent(ULandscapeComponent* Component, bool bMapCheck)
 {
 	// Do not register components which are not part of the world
diff --git a/Engine/Source/Runtime/Landscape/Private/LandscapeCollision.cpp b/Engine/Source/Runtime/Landscape/Private/LandscapeCollision.cpp
--- a/Engine/Source/Runtime/Landscape/Private/LandscapeCollision.cpp
+++ b/Engine/Source/Runtime/Landscape/Private/LandscapeCollision.cpp
@@ -174,6 +174,46 @@ static FString GetHFDDCKeyString(const FName& Format, bool bDefMaterial, const F
 	return FDerivedDataCacheInterface::BuildCacheKey(*KeyPrefix, LANDSCAPE_COLLISION_DERIVEDDATA_VER, *CombinedStateId.ToString());
 }
 
+void ULandscapeHeightfieldCollisionComponent::OnRegister()
+{
+	Super::OnRegister();
+
+	if (GetLandscapeProxy())
+	{
+		// AActor::GetWorld checks for Unreachable and BeginDestroyed
+		UWorld* World = GetLandscapeProxy()->GetWorld();
+		if (World)
+		{
+			ULandscapeInfo* Info = GetLandscapeInfo();
+			if (Info)
+			{
+				Info->RegisterCollisionComponent(this);
+			}
+		}
+	}
+}
+
+void ULandscapeHeightfieldCollisionComponent::OnUnregister()
+{
+	Super::OnUnregister();
+
+	if (GetLandscapeProxy())
+	{
+		// AActor::GetWorld checks for Unreachable and BeginDestroyed
+		UWorld* World = GetLandscapeProxy()->GetWorld();
+
+		// Game worlds don't have landscape infos
+		if (World)
+		{
+			ULandscapeInfo* Info = GetLandscapeInfo();
+			if (Info)
+			{
+				Info->UnregisterCollisionComponent(this);
+			}
+		}
+	}
+}
+
 ECollisionEnabled::Type ULandscapeHeightfieldCollisionComponent::GetCollisionEnabled() const
 {
 	if (!HasAnyFlags(RF_ClassDefaultObject))
@@ -2275,13 +2315,13 @@ void ULandscapeMeshCollisionComponent::ImportCustomProperties(const TCHAR* Sourc
 	}
 }
 
+#endif // WITH_EDITOR
+
 ULandscapeInfo* ULandscapeHeightfieldCollisionComponent::GetLandscapeInfo() const
 {
 	return GetLandscapeProxy()->GetLandscapeInfo();
 }
 
-#endif // WITH_EDITOR
-
 ALandscapeProxy* ULandscapeHeightfieldCollisionComponent::GetLandscapeProxy() const
 {
 	return CastChecked<ALandscapeProxy>(GetOuter());
@@ -2341,17 +2381,14 @@ LANDSCAPE_API TOptional<float> ALandscapeProxy::GetHeightAtLocation(FVector Loca
 	ULandscapeInfo* Info = GetLandscapeInfo();
 	const FVector ActorSpaceLocation = LandscapeActorToWorld().InverseTransformPosition(Location);
 	const FIntPoint Key = FIntPoint(FMath::FloorToInt(ActorSpaceLocation.X / ComponentSizeQuads), FMath::FloorToInt(ActorSpaceLocation.Y / ComponentSizeQuads));
-	ULandscapeComponent* Component = Info->XYtoComponentMap.FindRef(Key);
+	ULandscapeHeightfieldCollisionComponent* Component = Info->XYtoCollisionComponentMap.FindRef(Key);
 	if (Component)
 	{
-		if (ULandscapeHeightfieldCollisionComponent* CollisionComp = Component->CollisionComponent.Get())
+		const FVector ComponentSpaceLocation = Component->GetComponentToWorld().InverseTransformPosition(Location);
+		const TOptional<float> LocalHeight = Component->GetHeight(ComponentSpaceLocation.X, ComponentSpaceLocation.Y);
+		if (LocalHeight.IsSet())
 		{
-			const FVector ComponentSpaceLocation = CollisionComp->GetComponentToWorld().InverseTransformPosition(Location);
-			const TOptional<float> LocalHeight = CollisionComp->GetHeight(ComponentSpaceLocation.X, ComponentSpaceLocation.Y);
-			if (LocalHeight.IsSet())
-			{
-				Height = CollisionComp->GetComponentToWorld().TransformPositionNoScale(FVector(0, 0, LocalHeight.GetValue())).Z;
-			}
+			Height = Component->GetComponentToWorld().TransformPositionNoScale(FVector(0, 0, LocalHeight.GetValue())).Z;
 		}
 	}
 	return Height;
