Landscape Grass: Render missing grass maps on PreSave so that maps are in consistent state and don't get dirtied on load.

#rnx

diff --git a/Engine/Source/Runtime/Landscape/Classes/LandscapeProxy.h b/Engine/Source/Runtime/Landscape/Classes/LandscapeProxy.h
--- a/Engine/Source/Runtime/Landscape/Classes/LandscapeProxy.h
+++ b/Engine/Source/Runtime/Landscape/Classes/LandscapeProxy.h
@@ -841,6 +841,9 @@ public:
 	LANDSCAPE_API void InvalidateGeneratedComponentData();
 
 #if WITH_EDITOR
+	/** Update Grass maps */
+	void UpdateGrassData();
+
 	/** Render grass maps for the specified components */
 	void RenderGrassMaps(const TArray<ULandscapeComponent*>& LandscapeComponents, const TArray<ULandscapeGrassType*>& GrassTypes);
 
@@ -1059,6 +1062,10 @@ private:
 	/** Returns Grass Update interval */
 	int32 GetGrassUpdateInterval() const;
 
+#if WITH_EDITOR
+	void UpdateGrassDataStatus(TSet<UTexture2D*>& OutCurrentForcedStreamedTextures, TSet<UTexture2D*>* OutDesiredForcedStreamedTextures, TSet<ULandscapeComponent*>& OutComponentsNeedingGrassMapRender, TSet<ULandscapeComponent*>* OutOutdatedComponents, bool bInEnableForceResidentFlag);
+#endif
+
 #if WITH_EDITORONLY_DATA
 public:
 	static const TArray<ALandscapeProxy*>& GetLandscapeProxies() { return LandscapeProxies; }
diff --git a/Engine/Source/Runtime/Landscape/Private/Landscape.cpp b/Engine/Source/Runtime/Landscape/Private/Landscape.cpp
--- a/Engine/Source/Runtime/Landscape/Private/Landscape.cpp
+++ b/Engine/Source/Runtime/Landscape/Private/Landscape.cpp
@@ -1978,6 +1978,8 @@ void ALandscapeProxy::PreSave(const class ITargetPlatform* TargetPlatform)
 	if (!HasAnyFlags(RF_ClassDefaultObject))
 	{
 		bHasLandscapeGrass = LandscapeComponents.ContainsByPredicate([](ULandscapeComponent* Component) { return Component->MaterialHasGrass(); });
+
+		UpdateGrassData();
 	}
 
 	if (ALandscape* Landscape = GetLandscapeActor())
diff --git a/Engine/Source/Runtime/Landscape/Private/LandscapeGrass.cpp b/Engine/Source/Runtime/Landscape/Private/LandscapeGrass.cpp
--- a/Engine/Source/Runtime/Landscape/Private/LandscapeGrass.cpp
+++ b/Engine/Source/Runtime/Landscape/Private/LandscapeGrass.cpp
@@ -2137,6 +2137,153 @@ void ALandscapeProxy::RemoveAllExclusionBoxes()
 int32 ALandscapeProxy::TotalComponentsNeedingGrassMapRender = 0;
 int32 ALandscapeProxy::TotalTexturesToStreamForVisibleGrassMapRender = 0;
 int32 ALandscapeProxy::TotalComponentsNeedingTextureBaking = 0;
+
+void ALandscapeProxy::UpdateGrassDataStatus(TSet<UTexture2D*>& OutCurrentForcedStreamedTextures, TSet<UTexture2D*>* OutDesiredForcedStreamedTextures, TSet<ULandscapeComponent*>& OutComponentsNeedingGrassMapRender, TSet<ULandscapeComponent*>* OutOutdatedComponents, bool bInEnableForceResidentFlag)
+{
+	OutCurrentForcedStreamedTextures.Empty();
+	OutComponentsNeedingGrassMapRender.Empty();
+	
+	if (OutOutdatedComponents)
+	{
+		OutOutdatedComponents->Empty();
+	}
+
+	if (OutDesiredForcedStreamedTextures)
+	{
+		OutDesiredForcedStreamedTextures->Empty();
+	}
+		
+	const UWorld* World = GetWorld();
+	if (!World || World->IsGameWorld())
+	{
+		return;
+	}
+
+	// In either case we want to check the Grass textures stream state
+	const bool bCheckStreamingState = OutDesiredForcedStreamedTextures || bInEnableForceResidentFlag;
+
+	for (auto Component : LandscapeComponents)
+	{
+		if (Component != nullptr)
+		{
+			UTexture2D* Heightmap = Component->GetHeightmap();
+			// check textures currently needing force streaming
+			if (Heightmap->bForceMiplevelsToBeResident)
+			{
+				OutCurrentForcedStreamedTextures.Add(Heightmap);
+			}
+			TArray<UTexture2D*>& ComponentWeightmapTextures = Component->GetWeightmapTextures();
+
+			for (auto WeightmapTexture : ComponentWeightmapTextures)
+			{
+				if (WeightmapTexture->bForceMiplevelsToBeResident)
+				{
+					OutCurrentForcedStreamedTextures.Add(WeightmapTexture);
+				}
+			}
+
+			if (Component->IsGrassMapOutdated() && OutOutdatedComponents)
+			{
+				OutOutdatedComponents->Add(Component);
+			}
+
+			if (Component->IsGrassMapOutdated() || !Component->GrassData->HasData())
+			{
+				OutComponentsNeedingGrassMapRender.Add(Component);
+
+				if (bCheckStreamingState && !Component->AreTexturesStreamedForGrassMapRender())
+				{
+					if (OutDesiredForcedStreamedTextures)
+					{
+						OutDesiredForcedStreamedTextures->Add(Heightmap);
+					}
+
+					if (bInEnableForceResidentFlag)
+					{
+						Heightmap->bForceMiplevelsToBeResident = true;
+					}
+				
+					for (auto WeightmapTexture : ComponentWeightmapTextures)
+					{
+						if (OutDesiredForcedStreamedTextures)
+						{
+							OutDesiredForcedStreamedTextures->Add(WeightmapTexture);
+						}
+						
+						if (bInEnableForceResidentFlag)
+						{
+							WeightmapTexture->bForceMiplevelsToBeResident = true;
+						}
+					}
+				}
+			}
+		}
+	}
+}
+
+void ALandscapeProxy::UpdateGrassData()
+{
+	UWorld* World = GetWorld();
+	if (!World || World->IsGameWorld())
+	{
+		return;
+	}
+
+	if (!CVarGrassEnable.GetValueOnAnyThread())
+	{
+		return;
+	}
+
+	TArray<ULandscapeGrassType*> GrassTypes = GetGrassTypes();
+	if (GrassTypes.Num() == 0 && !bBakeMaterialPositionOffsetIntoCollision)
+	{
+		return;
+	}
+
+	// see if we need to flush grass for any components
+	TSet<UTexture2D*> DesiredForcedStreamedTextures;
+	TSet<UTexture2D*> CurrentForcedStreamedTextures;
+	TSet<ULandscapeComponent*> ComponentsNeedingGrassMapRender;
+	TSet<ULandscapeComponent*> OutdatedComponents;
+	const bool bEnableForceResidentFlag = true;
+	UpdateGrassDataStatus(CurrentForcedStreamedTextures, &DesiredForcedStreamedTextures, ComponentsNeedingGrassMapRender, &OutdatedComponents, bEnableForceResidentFlag);
+	if (OutdatedComponents.Num())
+	{
+		FlushGrassComponents(&OutdatedComponents);
+	}
+
+	// Remove local count from global count
+	TotalTexturesToStreamForVisibleGrassMapRender -= NumTexturesToStreamForVisibleGrassMapRender;
+	NumTexturesToStreamForVisibleGrassMapRender = 0;
+
+	// Remove local count from global count
+	TotalComponentsNeedingGrassMapRender -= NumComponentsNeedingGrassMapRender;
+	NumComponentsNeedingGrassMapRender = 0;
+
+	// Wait for Texture Streaming
+	for (UTexture2D* TextureToStream : DesiredForcedStreamedTextures)
+	{
+		TextureToStream->WaitForStreaming();
+		CurrentForcedStreamedTextures.Add(TextureToStream);
+	}
+		
+	// render grass data if we don't have any
+	for(ULandscapeComponent* Component : ComponentsNeedingGrassMapRender)
+	{
+		if (!Component->CanRenderGrassMap())
+		{
+			// we can't currently render grassmaps (eg shaders not compiled)
+			continue;
+		}
+		
+		Component->RenderGrassMap();
+	}
+		
+	for (UTexture2D* TextureToStream : CurrentForcedStreamedTextures)
+	{
+		TextureToStream->bForceMiplevelsToBeResident = false;
+	}
+}
 #endif
 
 void ALandscapeProxy::UpdateGrass(const TArray<FVector>& Cameras, bool bForceSync)
@@ -2174,52 +2321,17 @@ void ALandscapeProxy::UpdateGrass(const TArray<FVector>& Cameras, bool bForceSyn
 #if WITH_EDITOR
 			int32 RequiredTexturesNotStreamedIn = 0;
 			TSet<ULandscapeComponent*> ComponentsNeedingGrassMapRender;
+			TSet<ULandscapeComponent*> OutdatedComponents;
 			TSet<UTexture2D*> CurrentForcedStreamedTextures;
 			TSet<UTexture2D*> DesiredForceStreamedTextures;
-
-			if (!World->IsGameWorld())
+			const bool bEnableForceResidentFlag = false;
+			
+			// Do not pass in DesiredForceStreamedTextures because we build our own list
+			UpdateGrassDataStatus(CurrentForcedStreamedTextures, nullptr, ComponentsNeedingGrassMapRender, &OutdatedComponents, bEnableForceResidentFlag);
+			
+			if(OutdatedComponents.Num())
 			{
-				// see if we need to flush grass for any components
-				TSet<ULandscapeComponent*> FlushComponents;
-				for (auto Component : LandscapeComponents)
-				{
-					if (Component != nullptr)
-					{
-						UTexture2D* Heightmap = Component->GetHeightmap();
-						// check textures currently needing force streaming
-						if (Heightmap->bForceMiplevelsToBeResident)
-						{
-							CurrentForcedStreamedTextures.Add(Heightmap);
-						}
-						TArray<UTexture2D*>& ComponentWeightmapTextures = Component->GetWeightmapTextures();
-
-						for (auto WeightmapTexture : ComponentWeightmapTextures)
-						{
-							if (WeightmapTexture->bForceMiplevelsToBeResident)
-							{
-								CurrentForcedStreamedTextures.Add(WeightmapTexture);
-							}
-						}
-
-						if (Component->IsGrassMapOutdated())
-						{
-							FlushComponents.Add(Component);
-						}
-
-						if (GrassTypes.Num() > 0 || bBakeMaterialPositionOffsetIntoCollision)
-						{
-							if (Component->IsGrassMapOutdated() ||
-								!Component->GrassData->HasData())
-							{
-								ComponentsNeedingGrassMapRender.Add(Component);
-							}
-						}
-					}
-				}
-				if (FlushComponents.Num())
-				{
-					FlushGrassComponents(&FlushComponents);
-				}
+				FlushGrassComponents(&OutdatedComponents);
 			}
 #endif
 			ERHIFeatureLevel::Type FeatureLevel = World->Scene->GetFeatureLevel();
