Add speculative sanity checks to unblock QA

diff --git a/Engine/Source/Runtime/Landscape/Private/LandscapeRender.cpp b/Engine/Source/Runtime/Landscape/Private/LandscapeRender.cpp
--- a/Engine/Source/Runtime/Landscape/Private/LandscapeRender.cpp
+++ b/Engine/Source/Runtime/Landscape/Private/LandscapeRender.cpp
@@ -719,6 +719,10 @@ void FLandscapeRenderSystem::ResizeAndMoveTo(FIntPoint NewMin, FIntPoint NewSize
 
 void FLandscapeRenderSystem::ComputeSectionPerViewParameters(const FSceneView* View)
 {
+	ensureMsgf(View != nullptr, TEXT("View should not be nullptr when in ComputeSectionPerViewParameters()"));
+
+	if (View == nullptr) return;
+
 	if (!CachedSectionLODValues.Contains(View))
 	{
 		TRACE_CPUPROFILER_EVENT_SCOPE(FLandscapeRenderSystem::ComputeSectionPerViewParameters());
@@ -731,7 +735,16 @@ void FLandscapeRenderSystem::ComputeSectionPerViewParameters(const FSceneView* V
 			CachedSectionTessellationFalloffK.Add(View, TResourceArray<float>{});
 		}
 
+		ensureMsgf(View->Family != nullptr, TEXT("View->Family should not be nullptr when in ComputeSectionPerViewParameters()"));
+
+		if (View->Family == nullptr)
+			return;
+
+#if (!(UE_BUILD_SHIPPING || UE_BUILD_TEST))
 		int32 ForcedLODLevel = View->Family->EngineShowFlags.LOD ? GetCVarForceLOD() : -1;
+#else
+		int32 ForcedLODLevel = -1;
+#endif
 		float LODScale = View->LODDistanceFactor * CVarStaticMeshLODDistanceScale.GetValueOnRenderThread();
 
 		for (int32 EntityIndex = 0; EntityIndex < SectionLODSettings.Num(); EntityIndex++)
@@ -936,6 +949,10 @@ public:
 
 	virtual void PrepareView(const FSceneView* View) override
 	{
+		ensureMsgf(View != nullptr, TEXT("View should not be nullptr when in PrepareView()"));
+
+		if (View == nullptr) return;
+
 		for (auto& Pair : LandscapeRenderSystems)
 		{
 			FLandscapeRenderSystem& RenderSystem = *Pair.Value;
@@ -946,6 +963,10 @@ public:
 
 	virtual void BeginRenderView(const FSceneView* View) override
 	{
+		ensureMsgf(View != nullptr, TEXT("View should not be nullptr when in BeginRenderView()"));
+
+		if (View == nullptr) return;
+
 		for (auto& Pair : LandscapeRenderSystems)
 		{
 			FLandscapeRenderSystem& RenderSystem = *Pair.Value;
