Fix FLandscapeRenderSystem tasks not being properly waited in FMobileSceneRenderer
Fix race condition when accessing the View pointer and cached view values from parallel task threads

diff --git a/Engine/Source/Runtime/Engine/Public/SceneManagement.h b/Engine/Source/Runtime/Engine/Public/SceneManagement.h
--- a/Engine/Source/Runtime/Engine/Public/SceneManagement.h
+++ b/Engine/Source/Runtime/Engine/Public/SceneManagement.h
@@ -2835,6 +2835,7 @@ float ENGINE_API ComputeBoundsScreenSize(const FVector4& BoundsOrigin, const flo
 /**
  * Computes the screen radius squared of a given sphere bounds in the given view. This is used at
  * runtime instead of ComputeBoundsScreenSize to avoid a square root.
+ * It is a wrapper for the version below that does not take a FSceneView reference but parameters directly
  * @param Origin - Origin of the bounds in world space
  * @param SphereRadius - Radius of the sphere to use to calculate screen coverage
  * @param View - The view to calculate the display factor for
@@ -2842,6 +2843,17 @@ float ENGINE_API ComputeBoundsScreenSize(const FVector4& BoundsOrigin, const flo
  */
 float ENGINE_API ComputeBoundsScreenRadiusSquared(const FVector4& Origin, const float SphereRadius, const FSceneView& View);
 
+/**
+ * Computes the screen radius squared of a given sphere bounds in the given view. This is used at
+ * runtime instead of ComputeBoundsScreenSize to avoid a square root.
+ * @param Origin - Origin of the bounds in world space
+ * @param SphereRadius - Radius of the sphere to use to calculate screen coverage
+ * @param ViewOrigin - The view origin involved in the calculation
+ * @param ProjMatrix - The projection matrix of the view involved in the calculation
+ * @return float - The screen size calculated
+ */
+float ENGINE_API ComputeBoundsScreenRadiusSquared(const FVector4& BoundsOrigin, const float SphereRadius, const FVector4& ViewOrigin, const FMatrix& ProjMatrix);
+
 /**
  * Computes the draw distance of a given sphere bounds in the given view with the specified screen size.
  * @param ScreenSize - The screen size (as computed by ComputeBoundsScreenSize)
diff --git a/Engine/Source/Runtime/Landscape/Private/LandscapeRender.cpp b/Engine/Source/Runtime/Landscape/Private/LandscapeRender.cpp
--- a/Engine/Source/Runtime/Landscape/Private/LandscapeRender.cpp
+++ b/Engine/Source/Runtime/Landscape/Private/LandscapeRender.cpp
@@ -717,39 +717,66 @@ void FLandscapeRenderSystem::ResizeAndMoveTo(FIntPoint NewMin, FIntPoint NewSize
 	SectionTessellationFalloffK.SetAllowCPUAccess(true);
 }
 
-void FLandscapeRenderSystem::ComputeSectionPerViewParameters(const FSceneView* View)
+void FLandscapeRenderSystem::BeginRenderView(const FSceneView* View)
 {
-	ensureMsgf(View != nullptr, TEXT("View should not be nullptr when in ComputeSectionPerViewParameters()"));
+	TRACE_CPUPROFILER_EVENT_SCOPE(FLandscapeRenderSystem::BeginRenderView());
 
-	if (View == nullptr) return;
+	if (FetchHeightmapLODBiasesEventRef.IsValid())
+	{
+		FTaskGraphInterface::Get().WaitUntilTaskCompletes(FetchHeightmapLODBiasesEventRef, ENamedThreads::GetRenderThread_Local());
+		FetchHeightmapLODBiasesEventRef.SafeRelease();
+	}
 
-	if (!CachedSectionLODValues.Contains(View))
+	if (TaskEventRef.Contains(View))
 	{
-		TRACE_CPUPROFILER_EVENT_SCOPE(FLandscapeRenderSystem::ComputeSectionPerViewParameters());
+		FTaskGraphInterface::Get().WaitUntilTaskCompletes(TaskEventRef[View], ENamedThreads::GetRenderThread_Local());
+		TaskEventRef.Remove(View);
+	}
+	else
+	{
+		FComputeSectionPerViewParametersTask Task(*this, View);
+		Task.AnyThreadTask();
+	}
 
-		CachedSectionLODValues.Add(View, TResourceArray<float>{});
+	RecreateBuffers(View);
+}
+
+void FLandscapeRenderSystem::ComputeSectionPerViewParameters(
+	const FSceneView* ViewPtrAsIdentifier,
+	bool ViewEngineShowFlagLOD,
+	float ViewLODDistanceFactor,
+	FVector ViewOrigin,
+	FMatrix ViewProjectionMarix
+)
+{
+	bool bValuesCached = false;
 
-		if (TessellationFalloffSettings.UseTessellationComponentScreenSizeFalloff)
+	{
+		FScopeLock Lock(&CachedValuesCS);
+		bValuesCached = CachedSectionLODValues.Contains(ViewPtrAsIdentifier);
+
+		if (!bValuesCached)
 		{
-			CachedSectionTessellationFalloffC.Add(View, TResourceArray<float>{});
-			CachedSectionTessellationFalloffK.Add(View, TResourceArray<float>{});
-		}
+			CachedSectionLODValues.Add(ViewPtrAsIdentifier, TResourceArray<float>{});
 
-		ensureMsgf(View->Family != nullptr, TEXT("View->Family should not be nullptr when in ComputeSectionPerViewParameters()"));
+			if (TessellationFalloffSettings.UseTessellationComponentScreenSizeFalloff)
+			{
+				CachedSectionTessellationFalloffC.Add(ViewPtrAsIdentifier, TResourceArray<float>{});
+				CachedSectionTessellationFalloffK.Add(ViewPtrAsIdentifier, TResourceArray<float>{});
+			}
+		}
+	}
 
-		if (View->Family == nullptr)
-			return;
+	if (!bValuesCached)
+	{
+		TRACE_CPUPROFILER_EVENT_SCOPE(FLandscapeRenderSystem::ComputeSectionPerViewParameters());
 
-#if (!(UE_BUILD_SHIPPING || UE_BUILD_TEST))
-		int32 ForcedLODLevel = View->Family->EngineShowFlags.LOD ? GetCVarForceLOD() : -1;
-#else
-		int32 ForcedLODLevel = -1;
-#endif
-		float LODScale = View->LODDistanceFactor * CVarStaticMeshLODDistanceScale.GetValueOnRenderThread();
+		int32 ForcedLODLevel = ViewEngineShowFlagLOD ? GetCVarForceLOD() : -1;
+		float LODScale = ViewLODDistanceFactor * CVarStaticMeshLODDistanceScale.GetValueOnRenderThread();
 
 		for (int32 EntityIndex = 0; EntityIndex < SectionLODSettings.Num(); EntityIndex++)
 		{
-			float MeshScreenSizeSquared = ComputeBoundsScreenRadiusSquared(FVector(SectionOriginAndRadius[EntityIndex]), SectionOriginAndRadius[EntityIndex].W, *View);
+			float MeshScreenSizeSquared = ComputeBoundsScreenRadiusSquared(FVector(SectionOriginAndRadius[EntityIndex]), SectionOriginAndRadius[EntityIndex].W, ViewOrigin, ViewProjectionMarix);
 
 			float FractionalLOD;
 			GetLODFromScreenSize(SectionLODSettings[EntityIndex], MeshScreenSizeSquared, LODScale * LODScale, FractionalLOD);
@@ -758,11 +785,11 @@ void FLandscapeRenderSystem::ComputeSectionPerViewParameters(const FSceneView* V
 
 			if (TessellationFalloffSettings.UseTessellationComponentScreenSizeFalloff)
 			{
-				float MaxTesselationDistance = ComputeBoundsDrawDistance(FMath::Sqrt(TessellationFalloffSettings.TessellationComponentSquaredScreenSize), SectionOriginAndRadius[EntityIndex].W / 2.0f, View->ViewMatrices.GetProjectionMatrix());
+				float MaxTesselationDistance = ComputeBoundsDrawDistance(FMath::Sqrt(TessellationFalloffSettings.TessellationComponentSquaredScreenSize), SectionOriginAndRadius[EntityIndex].W / 2.0f, ViewProjectionMarix);
 				float FallOffStartingDistance = FMath::Min(
 					ComputeBoundsDrawDistance(FMath::Sqrt(FMath::Min(
 						FMath::Square(TessellationFalloffSettings.TessellationComponentScreenSizeFalloff), 
-						TessellationFalloffSettings.TessellationComponentSquaredScreenSize)), SectionOriginAndRadius[EntityIndex].W / 2.0f, View->ViewMatrices.GetProjectionMatrix()) - MaxTesselationDistance, MaxTesselationDistance);
+						TessellationFalloffSettings.TessellationComponentSquaredScreenSize)), SectionOriginAndRadius[EntityIndex].W / 2.0f, ViewProjectionMarix) - MaxTesselationDistance, MaxTesselationDistance);
 
 				// Calculate the falloff using a = C - K * d by sending C & K into the shader
 				SectionTessellationFalloffC[EntityIndex] = MaxTesselationDistance / (MaxTesselationDistance - FallOffStartingDistance);
@@ -770,22 +797,30 @@ void FLandscapeRenderSystem::ComputeSectionPerViewParameters(const FSceneView* V
 			}
 		}
 
-		CachedSectionLODValues[View] = SectionLODValues;
-
-		if (TessellationFalloffSettings.UseTessellationComponentScreenSizeFalloff)
 		{
-			CachedSectionTessellationFalloffC[View] = SectionTessellationFalloffC;
-			CachedSectionTessellationFalloffK[View] = SectionTessellationFalloffK;
+			FScopeLock Lock(&CachedValuesCS);
+
+			CachedSectionLODValues[ViewPtrAsIdentifier] = SectionLODValues;
+
+			if (TessellationFalloffSettings.UseTessellationComponentScreenSizeFalloff)
+			{
+				CachedSectionTessellationFalloffC[ViewPtrAsIdentifier] = SectionTessellationFalloffC;
+				CachedSectionTessellationFalloffK[ViewPtrAsIdentifier] = SectionTessellationFalloffK;
+			}
 		}
 	}
 	else
 	{
-		SectionLODValues = CachedSectionLODValues[View];
-
-		if (TessellationFalloffSettings.UseTessellationComponentScreenSizeFalloff)
 		{
-			SectionTessellationFalloffC = CachedSectionTessellationFalloffC[View];
-			SectionTessellationFalloffK = CachedSectionTessellationFalloffK[View];
+			FScopeLock Lock(&CachedValuesCS);
+
+			SectionLODValues = CachedSectionLODValues[ViewPtrAsIdentifier];
+
+			if (TessellationFalloffSettings.UseTessellationComponentScreenSizeFalloff)
+			{
+				SectionTessellationFalloffC = CachedSectionTessellationFalloffC[ViewPtrAsIdentifier];
+				SectionTessellationFalloffK = CachedSectionTessellationFalloffK[ViewPtrAsIdentifier];
+			}
 		}
 	}
 }
@@ -934,6 +969,24 @@ void FLandscapeRenderSystem::BeginFrame()
 	}
 }
 
+void FLandscapeRenderSystem::EndFrame()
+{
+	// Finalize any outstanding jobs before ~FSceneRenderer() so we don't have corrupted accesses
+	if (FetchHeightmapLODBiasesEventRef.IsValid())
+	{
+		FTaskGraphInterface::Get().WaitUntilTaskCompletes(FetchHeightmapLODBiasesEventRef, ENamedThreads::GetRenderThread_Local());
+		FetchHeightmapLODBiasesEventRef.SafeRelease();
+	}
+
+	for (auto& Pair : TaskEventRef)
+	{
+		const FSceneView* View = Pair.Key;
+		FTaskGraphInterface::Get().WaitUntilTaskCompletes(TaskEventRef[View], ENamedThreads::GetRenderThread_Local());
+	}
+
+	TaskEventRef.Empty();
+}
+
 class FLandscapePersistentViewUniformBufferExtension : public IPersistentViewUniformBufferExtension
 {
 public:
@@ -949,10 +1002,6 @@ public:
 
 	virtual void PrepareView(const FSceneView* View) override
 	{
-		ensureMsgf(View != nullptr, TEXT("View should not be nullptr when in PrepareView()"));
-
-		if (View == nullptr) return;
-
 		for (auto& Pair : LandscapeRenderSystems)
 		{
 			FLandscapeRenderSystem& RenderSystem = *Pair.Value;
@@ -961,11 +1010,12 @@ public:
 		}
 	}
 
-	virtual void BeginRenderView(const FSceneView* View) override
+	virtual void BeginRenderView(const FSceneView* View, bool bShouldWaitForJobs = true) override
 	{
-		ensureMsgf(View != nullptr, TEXT("View should not be nullptr when in BeginRenderView()"));
-
-		if (View == nullptr) return;
+		if (!bShouldWaitForJobs)
+		{
+			return;
+		}
 
 		for (auto& Pair : LandscapeRenderSystems)
 		{
@@ -974,6 +1024,17 @@ public:
 			RenderSystem.BeginRenderView(View);
 		}
 	}
+
+	virtual void EndFrame() override
+	{
+		for (auto& Pair : LandscapeRenderSystems)
+		{
+			FLandscapeRenderSystem& RenderSystem = *Pair.Value;
+
+			RenderSystem.EndFrame();
+		}
+	}
+
 } LandscapePersistentViewUniformBufferExtension;
 
 FLandscapeComponentSceneProxy::FLandscapeComponentSceneProxy(ULandscapeComponent* InComponent)
diff --git a/Engine/Source/Runtime/Landscape/Public/LandscapeRender.h b/Engine/Source/Runtime/Landscape/Public/LandscapeRender.h
--- a/Engine/Source/Runtime/Landscape/Public/LandscapeRender.h
+++ b/Engine/Source/Runtime/Landscape/Public/LandscapeRender.h
@@ -578,6 +578,7 @@ struct FLandscapeRenderSystem
 
 	TUniformBufferRef<FLandscapeSectionLODUniformParameters> UniformBuffer;
 
+	FCriticalSection CachedValuesCS;
 	TMap<const FSceneView*, TResourceArray<float>> CachedSectionLODValues;
 	TMap<const FSceneView*, TResourceArray<float>> CachedSectionTessellationFalloffC;
 	TMap<const FSceneView*, TResourceArray<float>> CachedSectionTessellationFalloffK;
@@ -589,11 +590,19 @@ struct FLandscapeRenderSystem
 	struct FComputeSectionPerViewParametersTask
 	{
 		FLandscapeRenderSystem& RenderSystem;
-		const FSceneView* View;
+		const FSceneView* ViewPtrAsIdentifier;
+		bool ViewEngineShowFlagLOD;
+		float ViewLODDistanceFactor;
+		FVector ViewOrigin;
+		FMatrix ViewProjectionMatrix;
 
 		FComputeSectionPerViewParametersTask(FLandscapeRenderSystem& InRenderSystem, const FSceneView* InView)
 			: RenderSystem(InRenderSystem)
-			, View(InView)
+			, ViewPtrAsIdentifier(InView)
+			, ViewEngineShowFlagLOD(InView->Family->EngineShowFlags.LOD)
+			, ViewLODDistanceFactor(InView->LODDistanceFactor)
+			, ViewOrigin(InView->ViewMatrices.GetViewOrigin())
+			, ViewProjectionMatrix(InView->ViewMatrices.GetProjectionMatrix())
 		{
 		}
 
@@ -614,7 +623,7 @@ struct FLandscapeRenderSystem
 
 		void AnyThreadTask()
 		{
-			RenderSystem.ComputeSectionPerViewParameters(View);
+			RenderSystem.ComputeSectionPerViewParameters(ViewPtrAsIdentifier, ViewEngineShowFlagLOD, ViewLODDistanceFactor, ViewOrigin, ViewProjectionMatrix);
 		}
 
 		void DoTask(ENamedThreads::Type CurrentThread, const FGraphEventRef& MyCompletionGraphEvent)
@@ -701,7 +710,12 @@ struct FLandscapeRenderSystem
 		return SectionLODBiases[GetComponentLinearIndex(ComponentBase)];
 	}
 
-	void ComputeSectionPerViewParameters(const FSceneView* View);
+	void ComputeSectionPerViewParameters(
+		const FSceneView* ViewPtrAsIdentifier,
+		bool ViewEngineShowFlagLOD,
+		float ViewLODDistanceFactor,
+		FVector ViewOrigin,
+		FMatrix ViewProjectionMarix);
 
 	void PrepareView(const FSceneView* View)
 	{
@@ -715,33 +729,15 @@ struct FLandscapeRenderSystem
 		}
 	}
 
-	void BeginRenderView(const FSceneView* View)
-	{
-		if (FetchHeightmapLODBiasesEventRef.IsValid())
-		{
-			FTaskGraphInterface::Get().WaitUntilTaskCompletes(FetchHeightmapLODBiasesEventRef, ENamedThreads::GetRenderThread_Local());
-			FetchHeightmapLODBiasesEventRef.SafeRelease();
-		}
-
-		if (TaskEventRef.Contains(View))
-		{
-			FTaskGraphInterface::Get().WaitUntilTaskCompletes(TaskEventRef[View], ENamedThreads::GetRenderThread_Local());
-			TaskEventRef.Remove(View);
-		}
-		else
-		{
-			FComputeSectionPerViewParametersTask Task(*this, View);
-			Task.AnyThreadTask();
-		}
-
-		RecreateBuffers(View);
-	}
+	void BeginRenderView(const FSceneView* View);
 
 	void BeginFrame();
 
 	void FetchHeightmapLODBiases();
 
 	void RecreateBuffers(const FSceneView* InView = nullptr);
+
+	void EndFrame();
 };
 
 extern TMap<FLandscapeNeighborInfo::FLandscapeKey, FLandscapeRenderSystem*> LandscapeRenderSystems;
diff --git a/Engine/Source/Runtime/RenderCore/Public/RendererInterface.h b/Engine/Source/Runtime/RenderCore/Public/RendererInterface.h
--- a/Engine/Source/Runtime/RenderCore/Public/RendererInterface.h
+++ b/Engine/Source/Runtime/RenderCore/Public/RendererInterface.h
@@ -576,7 +576,8 @@ class IPersistentViewUniformBufferExtension
 public:
 	virtual void BeginFrame() {}
 	virtual void PrepareView(const FSceneView* View) {}
-	virtual void BeginRenderView(const FSceneView* View) {}
+	virtual void BeginRenderView(const FSceneView* View, bool bShouldWaitForJobs = true) {}
+	virtual void EndFrame() {}
 };
 
 /**
diff --git a/Engine/Source/Runtime/Renderer/Private/DeferredShadingRenderer.cpp b/Engine/Source/Runtime/Renderer/Private/DeferredShadingRenderer.cpp
--- a/Engine/Source/Runtime/Renderer/Private/DeferredShadingRenderer.cpp
+++ b/Engine/Source/Runtime/Renderer/Private/DeferredShadingRenderer.cpp
@@ -1010,12 +1010,15 @@ void FDeferredShadingSceneRenderer::Render(FRHICommandListImmediate& RHICmdList)
 		}	
 	}
 
-	for (int32 ViewIndex = 0; ViewIndex < Views.Num(); ViewIndex++)
+	extern TSet<IPersistentViewUniformBufferExtension*> PersistentViewUniformBufferExtensions;
+
+	for (IPersistentViewUniformBufferExtension* Extension : PersistentViewUniformBufferExtensions)
 	{
-		extern TSet<IPersistentViewUniformBufferExtension*> PersistentViewUniformBufferExtensions;
+		Extension->BeginFrame();
 
-		for (IPersistentViewUniformBufferExtension* Extension : PersistentViewUniformBufferExtensions)
+		for (int32 ViewIndex = 0; ViewIndex < Views.Num(); ViewIndex++)
 		{
+			// Must happen before RHI thread flush so any tasks we dispatch here can land in the idle gap during the flush
 			Extension->PrepareView(&Views[ViewIndex]);
 		}
 	}
diff --git a/Engine/Source/Runtime/Renderer/Private/MobileShadingRenderer.cpp b/Engine/Source/Runtime/Renderer/Private/MobileShadingRenderer.cpp
--- a/Engine/Source/Runtime/Renderer/Private/MobileShadingRenderer.cpp
+++ b/Engine/Source/Runtime/Renderer/Private/MobileShadingRenderer.cpp
@@ -278,14 +278,27 @@ void FMobileSceneRenderer::InitViews(FRHICommandListImmediate& RHICmdList)
 	{
 		UploadDynamicPrimitiveShaderDataForView(RHICmdList, *Scene, Views[ViewIndex]);
 	}
-	
+
+	extern TSet<IPersistentViewUniformBufferExtension*> PersistentViewUniformBufferExtensions;
+
+	for (IPersistentViewUniformBufferExtension* Extension : PersistentViewUniformBufferExtensions)
+	{
+		Extension->BeginFrame();
+
+		for (int32 ViewIndex = 0; ViewIndex < Views.Num(); ViewIndex++)
+		{
+			// Must happen before RHI thread flush so any tasks we dispatch here can land in the idle gap during the flush
+			Extension->PrepareView(&Views[ViewIndex]);
+		}
+	}
 
 	// update buffers used in cached mesh path
 	// in case there are multiple views, these buffers will be updated before rendering each view
 	if (Views.Num() > 0)
 	{
 		const FViewInfo& View = Views[0];
-		Scene->UniformBuffers.UpdateViewUniformBuffer(View);
+		// We want to wait for the extension jobs only when the view is being actually rendered for the first time
+		Scene->UniformBuffers.UpdateViewUniformBuffer(View, false);
 		UpdateOpaqueBasePassUniformBuffer(RHICmdList, View);
 		UpdateTranslucentBasePassUniformBuffer(RHICmdList, View);
 		UpdateDirectionalLightUniformBuffers(RHICmdList, View);
@@ -300,16 +313,6 @@ void FMobileSceneRenderer::InitViews(FRHICommandListImmediate& RHICmdList)
 	UpdatePrimitiveIndirectLightingCacheBuffers();
 	
 	OnStartRender(RHICmdList);
-
-	for (int32 ViewIndex = 0; ViewIndex < Views.Num(); ViewIndex++)
-	{
-		extern TSet<IPersistentViewUniformBufferExtension*> PersistentViewUniformBufferExtensions;
-
-		for (IPersistentViewUniformBufferExtension* Extension : PersistentViewUniformBufferExtensions)
-		{
-			Extension->PrepareView(&Views[ViewIndex]);
-		}
-	}
 }
 
 /** 
diff --git a/Engine/Source/Runtime/Renderer/Private/RendererScene.cpp b/Engine/Source/Runtime/Renderer/Private/RendererScene.cpp
--- a/Engine/Source/Runtime/Renderer/Private/RendererScene.cpp
+++ b/Engine/Source/Runtime/Renderer/Private/RendererScene.cpp
@@ -927,8 +927,14 @@ void FRendererModule::RegisterPersistentViewUniformBufferExtension(IPersistentVi
 	PersistentViewUniformBufferExtensions.Add(Extension);
 }
 
-bool FPersistentUniformBuffers::UpdateViewUniformBuffer(const FViewInfo& View)
+bool FPersistentUniformBuffers::UpdateViewUniformBuffer(const FViewInfo& View, bool bShouldWaitForPersistentViewUniformBufferExtensionsJobs)
 {
+	// Let the implementation of each extension decide whether it can cache the result for CachedView
+	for (IPersistentViewUniformBufferExtension* Extension : PersistentViewUniformBufferExtensions)
+	{
+		Extension->BeginRenderView(&View, bShouldWaitForPersistentViewUniformBufferExtensionsJobs);
+	}
+
 	// ViewUniformBuffer can be cached by mesh commands, so we need to update it every time we change current view.
 	if (CachedView != &View)
 	{
@@ -949,11 +955,6 @@ bool FPersistentUniformBuffers::UpdateViewUniformBuffer(const FViewInfo& View)
 
 		CachedView = &View;
 
-		for (IPersistentViewUniformBufferExtension* Extension : PersistentViewUniformBufferExtensions)
-		{
-			Extension->BeginRenderView(&View);
-		}
-
 		return true;
 	}
 	return false;
diff --git a/Engine/Source/Runtime/Renderer/Private/ScenePrivate.h b/Engine/Source/Runtime/Renderer/Private/ScenePrivate.h
--- a/Engine/Source/Runtime/Renderer/Private/ScenePrivate.h
+++ b/Engine/Source/Runtime/Renderer/Private/ScenePrivate.h
@@ -2317,8 +2317,10 @@ public:
 
 	/** Compares the provided view against the cached view and updates the view uniform buffer
 	 *  if the views differ. Returns whether uniform buffer was updated.
+	 *  If bShouldWaitForPersistentViewUniformBufferExtensionsJobs == true, it calls Extension->BeginRenderView() which
+	 *  waits on the potential jobs dispatched in Extension->PrepareView(). Currently it is false only in FMobileSceneRenderer::InitViews()
 	 */
-	bool UpdateViewUniformBuffer(const FViewInfo& View);
+	bool UpdateViewUniformBuffer(const FViewInfo& View, bool bShouldWaitForPersistentViewUniformBufferExtensionsJobs = true);
 
 	/** Updates view uniform buffer and invalidates the internally cached view instance. */
 	void UpdateViewUniformBufferImmediate(const FViewUniformShaderParameters& Parameters);
diff --git a/Engine/Source/Runtime/Renderer/Private/SceneRendering.cpp b/Engine/Source/Runtime/Renderer/Private/SceneRendering.cpp
--- a/Engine/Source/Runtime/Renderer/Private/SceneRendering.cpp
+++ b/Engine/Source/Runtime/Renderer/Private/SceneRendering.cpp
@@ -2570,6 +2570,13 @@ FSceneRenderer::~FSceneRenderer()
 	// Manually release references to TRefCountPtrs that are allocated on the mem stack, which doesn't call dtors
 	SortedShadowsForShadowDepthPass.Release();
 
+	extern TSet<IPersistentViewUniformBufferExtension*> PersistentViewUniformBufferExtensions;
+
+	for (IPersistentViewUniformBufferExtension* Extension : PersistentViewUniformBufferExtensions)
+	{
+		Extension->EndFrame();
+	}
+
 	Views.Empty();
 }
 
@@ -3151,13 +3158,6 @@ void FSceneRenderer::OnStartRender(FRHICommandListImmediate& RHICmdList)
 			View.ViewState->OnStartRender(View, ViewFamily);
 		}
 	}
-
-	extern TSet<IPersistentViewUniformBufferExtension*> PersistentViewUniformBufferExtensions;
-
-	for (IPersistentViewUniformBufferExtension* Extension : PersistentViewUniformBufferExtensions)
-	{
-		Extension->BeginFrame();
-	}
 }
 
 bool FSceneRenderer::ShouldCompositeEditorPrimitives(const FViewInfo& View)
