Change HeightField Fwd Declare to an include to resolve a unique_ptr instantiation error on linux.

diff --git a/Engine/Source/Runtime/Landscape/Classes/LandscapeHeightfieldCollisionComponent.h b/Engine/Source/Runtime/Landscape/Classes/LandscapeHeightfieldCollisionComponent.h
--- a/Engine/Source/Runtime/Landscape/Classes/LandscapeHeightfieldCollisionComponent.h
+++ b/Engine/Source/Runtime/Landscape/Classes/LandscapeHeightfieldCollisionComponent.h
@@ -11,6 +11,7 @@
 #include "AI/Navigation/NavigationTypes.h"
 #include "Components/PrimitiveComponent.h"
 #include "Serialization/BulkData.h"
+#include "Chaos/HeightField.h"
 #include "LandscapeHeightfieldCollisionComponent.generated.h"
 
 class ALandscapeProxy;
@@ -23,14 +24,6 @@ struct FConvexVolume;
 struct FEngineShowFlags;
 struct FNavigableGeometryExport;
 
-#if WITH_CHAOS
-namespace Chaos
-{
-	template<typename T>
-	class THeightField;
-}
-#endif
-
 #if WITH_PHYSX
 namespace physx
 {
diff --git a/Engine/Source/Runtime/Landscape/Private/LandscapeCollision.cpp b/Engine/Source/Runtime/Landscape/Private/LandscapeCollision.cpp
--- a/Engine/Source/Runtime/Landscape/Private/LandscapeCollision.cpp
+++ b/Engine/Source/Runtime/Landscape/Private/LandscapeCollision.cpp
@@ -52,7 +52,6 @@
 #if WITH_CHAOS
 #include "PhysXToChaosUtil.h"
 #include "Chaos/ParticleHandle.h"
-#include "Chaos/HeightField.h"
 #include "Chaos/Vector.h"
 #endif
 
