Landscape tools working with Chaos objects for painting and sculpting, however materials and holes still not functional

diff --git a/Engine/Source/Runtime/Engine/Private/PhysicsEngine/Experimental/PhysScene_Chaos.cpp b/Engine/Source/Runtime/Engine/Private/PhysicsEngine/Experimental/PhysScene_Chaos.cpp
--- a/Engine/Source/Runtime/Engine/Private/PhysicsEngine/Experimental/PhysScene_Chaos.cpp
+++ b/Engine/Source/Runtime/Engine/Private/PhysicsEngine/Experimental/PhysScene_Chaos.cpp
@@ -982,6 +982,41 @@ const Chaos::FPhysicsSolver* FPhysScene_ChaosInterface::GetSolver() const
 	return Scene.GetSolver();
 }
 
+void FPhysScene_ChaosInterface::Flush_AssumesLocked()
+{
+	check(IsInGameThread());
+
+	// Flush all of our pending commands
+	Chaos::IDispatcher* Dispatcher = FChaosSolversModule::GetModule()->GetDispatcher();
+
+	if(Dispatcher->GetMode() != Chaos::EThreadingMode::SingleThread)
+	{
+		Dispatcher->Execute();
+	}
+
+	Chaos::FPBDRigidsSolver* Solver = Scene.GetSolver();
+
+	if(Solver)
+	{
+		TQueue<TFunction<void(Chaos::FPhysicsSolver*)>, EQueueMode::Mpsc>& Queue = Solver->GetCommandQueue();
+		TFunction<void(Chaos::FPhysicsSolver*)> Command;
+		while(Queue.Dequeue(Command))
+		{
+			Command(Solver);
+		}
+	}
+
+	// Populate the spacial acceleration
+	Chaos::FPBDRigidsSolver::FPBDRigidsEvolution* Evolution = GetSolver()->GetEvolution();
+
+	if(Evolution)
+	{
+		Evolution->FlushSpatialAcceleration();
+	}
+
+	Scene.CopySolverAccelerationStructure();
+}
+
 FPhysicsReplication* FPhysScene_ChaosInterface::GetPhysicsReplication()
 {
 	return nullptr;
diff --git a/Engine/Source/Runtime/Engine/Public/Physics/Experimental/PhysScene_Chaos.h b/Engine/Source/Runtime/Engine/Public/Physics/Experimental/PhysScene_Chaos.h
--- a/Engine/Source/Runtime/Engine/Public/Physics/Experimental/PhysScene_Chaos.h
+++ b/Engine/Source/Runtime/Engine/Public/Physics/Experimental/PhysScene_Chaos.h
@@ -290,6 +290,13 @@ public:
 	FPhysScene_Chaos& GetScene() { return Scene; }
 	const FPhysScene_Chaos& GetScene() const { return Scene; }
 
+	/**
+	 * Flushes all pending global, task and solver command queues and refreshes the spatial acceleration
+	 * for the scene. Required when querying against a currently non-running scene to ensure the scene
+	 * is correctly represented
+	 */
+	ENGINE_API void Flush_AssumesLocked();
+
 	FPhysicsReplication* GetPhysicsReplication();
 	void RemoveBodyInstanceFromPendingLists_AssumesLocked(FBodyInstance* BodyInstance, int32 SceneType);
 	void AddCustomPhysics_AssumesLocked(FBodyInstance* BodyInstance, FCalculateCustomPhysics& CalculateCustomPhysics);
diff --git a/Engine/Source/Runtime/Experimental/Chaos/Private/Chaos/HeightField.cpp b/Engine/Source/Runtime/Experimental/Chaos/Private/Chaos/HeightField.cpp
--- a/Engine/Source/Runtime/Experimental/Chaos/Private/Chaos/HeightField.cpp
+++ b/Engine/Source/Runtime/Experimental/Chaos/Private/Chaos/HeightField.cpp
@@ -277,6 +277,7 @@ namespace Chaos
 	void BuildGeomData(TArrayView<BufferType> BufferView, int32 NumRows, int32 NumCols, const TVector<HeightfieldT, 3>& InScale, TUniqueFunction<HeightfieldT(const BufferType)> ToRealFunc, typename THeightField<HeightfieldT>::FDataType& OutData, TBox<HeightfieldT, 3>& OutBounds)
 	{
 		using FDataType = typename THeightField<HeightfieldT>::FDataType;
+
 		using RealType = typename FDataType::RealType;
 
 		const int32 NumCells = NumRows * NumCols;
@@ -330,6 +331,90 @@ namespace Chaos
 		OutBounds.Thicken(KINDA_SMALL_NUMBER);
 	}
 
+	template<typename HeightfieldT, typename BufferType>
+	void EditGeomData(TArrayView<BufferType> BufferView, int32 InBeginRow, int32 InBeginCol, int32 NumRows, int32 NumCols, TUniqueFunction<HeightfieldT(const BufferType)> ToRealFunc, typename THeightField<HeightfieldT>::FDataType& OutData, TBox<HeightfieldT, 3>& OutBounds)
+	{
+		using FDataType = typename THeightField<HeightfieldT>::FDataType;
+		using RealType = typename FDataType::RealType;
+
+		RealType MinValue = TNumericLimits<HeightfieldT>::Max();
+		RealType MaxValue = TNumericLimits<HeightfieldT>::Min();
+
+		for(BufferType& Value : BufferView)
+		{
+			MinValue = FMath::Min(MinValue, ToRealFunc(Value));
+			MaxValue = FMath::Max(MaxValue, ToRealFunc(Value));
+		}
+
+		const int32 EndRow = InBeginRow + NumRows;
+		const int32 EndCol = InBeginCol + NumCols;
+
+		// If our range now falls outside of the original ranges we need to resample the whole heightfield to perform the edit.
+		// Here we resample everything outside of the edit and update our ranges
+		const bool bNeedsResample = MinValue < OutData.MinValue || MaxValue > OutData.MaxValue;
+		if(bNeedsResample)
+		{
+			const RealType NewMin = FMath::Min(MinValue, OutData.MinValue);
+			const RealType NewMax = FMath::Max(MaxValue, OutData.MaxValue);
+			const RealType NewRange = NewMax - NewMin;
+			const RealType NewHeightPerUnit = NewRange / FDataType::StorageRange;
+
+			for(int32 RowIdx = 0; RowIdx < OutData.NumRows; ++RowIdx)
+			{
+				for(int32 ColIdx = 0; ColIdx < OutData.NumCols; ++ColIdx)
+				{
+					const int32 HeightIndex = RowIdx * OutData.NumCols + ColIdx;
+
+					if(RowIdx >= InBeginRow && RowIdx < EndRow &&
+						ColIdx >= InBeginCol && ColIdx < EndCol)
+					{
+						// From the new set
+						const int32 NewSetIndex = (RowIdx - InBeginRow) * NumCols + (ColIdx - InBeginCol);
+						OutData.Heights[HeightIndex] = static_cast<typename FDataType::StorageType>((ToRealFunc(BufferView[NewSetIndex]) - NewMin) / NewHeightPerUnit);
+					}
+					else
+					{
+						// Resample existing
+						const RealType ExpandedHeight = OutData.MinValue + OutData.Heights[HeightIndex] * OutData.HeightPerUnit;
+						OutData.Heights[HeightIndex] = static_cast<typename FDataType::StorageType>((ExpandedHeight - NewMin) / NewHeightPerUnit);
+					}
+
+					int32 X = HeightIndex % (OutData.NumCols);
+					int32 Y = HeightIndex / (OutData.NumCols);
+					TVector<HeightfieldT, 3> Position(RealType(X), RealType(Y), NewMin + OutData.Heights[HeightIndex] * NewHeightPerUnit);
+					if(HeightIndex == 0)
+					{
+						OutBounds = TBox<HeightfieldT, 3>(Position, Position);
+					}
+					else
+					{
+						OutBounds.GrowToInclude(Position);
+					}
+				}
+			}
+
+			OutBounds.Thicken(KINDA_SMALL_NUMBER);
+
+			OutData.MinValue = NewMin;
+			OutData.MaxValue = NewMax;
+			OutData.HeightPerUnit = NewHeightPerUnit;
+			OutData.Range = NewRange;
+		}
+		else
+		{
+			// No resample, just push new heights into the data
+			for(int32 RowIdx = InBeginRow; RowIdx < EndRow; ++RowIdx)
+			{
+				for(int32 ColIdx = InBeginCol; ColIdx < EndCol; ++ColIdx)
+				{
+					const int32 HeightIndex = RowIdx * NumCols + ColIdx;
+					const int32 NewSetIndex = (RowIdx - InBeginRow) * NumCols + (ColIdx - InBeginCol);
+					OutData.Heights[HeightIndex] = static_cast<typename FDataType::StorageType>((ToRealFunc(BufferView[NewSetIndex]) - OutData.MinValue) / OutData.HeightPerUnit);
+				}
+			}
+		}
+	}
+
 	template <typename T>
 	THeightField<T>::THeightField(TArray<T>&& Height, int32 NumRows, int32 NumCols, const TVector<T, 3>& InScale)
 		: TImplicitObject<T, 3>(EImplicitObject::HasBoundingBox, ImplicitObjectType::HeightField)
@@ -353,6 +438,42 @@ namespace Chaos
 		SetScale(InScale);
 	}
 
+	template<typename T>
+	void Chaos::THeightField<T>::EditHeights(TArrayView<const uint16> InHeights, int32 InBeginRow, int32 InBeginCol, int32 InNumRows, int32 InNumCols)
+	{
+		const int32 NumExpectedValues = InNumRows * InNumCols;
+		const int32 EndRow = InBeginRow + InNumRows - 1;
+		const int32 EndCol = InBeginCol + InNumCols - 1;
+
+		if(ensure(InHeights.Num() == NumExpectedValues && InBeginRow >= 0 && InBeginCol >= 0 && EndRow < GeomData.NumRows && EndCol < GeomData.NumCols))
+		{
+			TUniqueFunction<T(const uint16)> ConversionFunc = [](const T InVal) -> float
+			{
+				return (float)((int32)InVal - 32768);
+			};
+
+			EditGeomData<T, const uint16>(InHeights, InBeginRow, InBeginCol, InNumRows, InNumCols, MoveTemp(ConversionFunc), GeomData, LocalBounds);
+		}
+	}
+
+	template<typename T>
+	void Chaos::THeightField<T>::EditHeights(TArrayView<T> InHeights, int32 InBeginRow, int32 InBeginCol, int32 InNumRows, int32 InNumCols)
+	{
+		const int32 NumExpectedValues = InNumRows * InNumCols;
+		const int32 EndRow = InBeginRow + InNumRows - 1;
+		const int32 EndCol = InBeginCol + InNumCols - 1;
+
+		if(ensure(InHeights.Num() == NumExpectedValues && InBeginRow >= 0 && InBeginCol >= 0 && EndRow < GeomData.NumRows && EndCol < GeomData.NumCols))
+		{
+			TUniqueFunction<T(T)> ConversionFunc = [](const T InVal) -> float
+			{
+				return InVal;
+			};
+
+			EditGeomData<T, T>(InHeights, InBeginRow, InBeginCol, InNumRows, InNumCols, MoveTemp(ConversionFunc), GeomData, LocalBounds);
+		}
+	}
+
 	template<typename T>
 	bool Chaos::THeightField<T>::GetCellBounds2D(const TVector<int32, 2> InCoord, FBounds2D& OutBounds, const TVector<T, 2>& InInflate /*= {0}*/) const
 	{
diff --git a/Engine/Source/Runtime/Experimental/Chaos/Private/Chaos/PBDRigidsEvolution.cpp b/Engine/Source/Runtime/Experimental/Chaos/Private/Chaos/PBDRigidsEvolution.cpp
--- a/Engine/Source/Runtime/Experimental/Chaos/Private/Chaos/PBDRigidsEvolution.cpp
+++ b/Engine/Source/Runtime/Experimental/Chaos/Private/Chaos/PBDRigidsEvolution.cpp
@@ -392,6 +392,15 @@ namespace Chaos
 		}
 	}
 
+	template<class FPBDRigidsEvolution, class FPBDCollisionConstraint, class T, int d>
+	void TPBDRigidsEvolutionBase<FPBDRigidsEvolution, FPBDCollisionConstraint, T, d>::FlushSpatialAcceleration()
+	{
+		//force build acceleration structure with latest data
+		ComputeIntermediateSpatialAcceleration(true);
+		ComputeIntermediateSpatialAcceleration(true);	//having to do it multiple times because of the various caching involved over multiple frames.
+		ComputeIntermediateSpatialAcceleration(true);
+	}
+
 	template<class FPBDRigidsEvolution, class FPBDCollisionConstraint, class T, int d>
 	void TPBDRigidsEvolutionBase<FPBDRigidsEvolution, FPBDCollisionConstraint, T, d>::Serialize(FChaosArchive& Ar)
 	{
@@ -430,10 +439,7 @@ namespace Chaos
 				DirtyParticle(Particle);
 			}
 
-			//force build acceleration structure with latest data
-			ComputeIntermediateSpatialAcceleration(true);
-			ComputeIntermediateSpatialAcceleration(true);	//having to do it multiple times because of the various caching involved over multiple frames.
-			ComputeIntermediateSpatialAcceleration(true);
+			FlushSpatialAcceleration();
 		}
 	}
 }
diff --git a/Engine/Source/Runtime/Experimental/Chaos/Public/Chaos/HeightField.h b/Engine/Source/Runtime/Experimental/Chaos/Public/Chaos/HeightField.h
--- a/Engine/Source/Runtime/Experimental/Chaos/Public/Chaos/HeightField.h
+++ b/Engine/Source/Runtime/Experimental/Chaos/Public/Chaos/HeightField.h
@@ -34,6 +34,10 @@ namespace Chaos
 
 		virtual ~THeightField() {}
 
+		/** Support for editing a subsection of the heightfield */
+		void EditHeights(TArrayView<T> InHeights, int32 InBeginRow, int32 InBeginCol, int32 InNumRows, int32 InNumCols);
+		void EditHeights(TArrayView<const uint16> InHeights, int32 InBeginRow, int32 InBeginCol, int32 InNumRows, int32 InNumCols);
+
 		virtual T PhiWithNormal(const TVector<T, 3>& x, TVector<T, 3>& Normal) const
 		{
 			check(false);	//not supported yet - might support it in the future or we may change the interface
diff --git a/Engine/Source/Runtime/Experimental/Chaos/Public/Chaos/PBDRigidsEvolution.h b/Engine/Source/Runtime/Experimental/Chaos/Public/Chaos/PBDRigidsEvolution.h
--- a/Engine/Source/Runtime/Experimental/Chaos/Public/Chaos/PBDRigidsEvolution.h
+++ b/Engine/Source/Runtime/Experimental/Chaos/Public/Chaos/PBDRigidsEvolution.h
@@ -262,6 +262,9 @@ class TPBDRigidsEvolutionBase
 	CHAOS_API void UpdateExternalAccelerationStructure(TUniquePtr<ISpatialAccelerationCollection<TAccelerationStructureHandle<T, d>, T, d>>& ExternalStructure);
 	ISpatialAccelerationCollection<TAccelerationStructureHandle<T, d>, T, d>* GetSpatialAcceleration() { return InternalAcceleration.Get(); }
 
+	/** Perform a blocking flush of the spatial acceleration structure for situations where we aren't simulating but must have an up to date structure */
+	CHAOS_API void FlushSpatialAcceleration();
+
 	const auto& GetRigidClustering() const { return Clustering; }
 	auto& GetRigidClustering() { return Clustering; }
 
diff --git a/Engine/Source/Runtime/Landscape/Classes/LandscapeHeightfieldCollisionComponent.h b/Engine/Source/Runtime/Landscape/Classes/LandscapeHeightfieldCollisionComponent.h
--- a/Engine/Source/Runtime/Landscape/Classes/LandscapeHeightfieldCollisionComponent.h
+++ b/Engine/Source/Runtime/Landscape/Classes/LandscapeHeightfieldCollisionComponent.h
@@ -165,10 +165,11 @@ class ULandscapeHeightfieldCollisionComponent : public UPrimitiveComponent
 private:
 	bool bEnableCollisionHashOptim = false;
 #endif //WITH_EDITORONLY_DATA
-
+	
 	//~ Begin UActorComponent Interface.
 protected:
 	virtual void OnCreatePhysicsState() override;
+	virtual void OnDestroyPhysicsState() override;
 public:
 
 	virtual void ApplyWorldOffset(const FVector& InOffset, bool bWorldShift) override;
@@ -251,6 +252,7 @@ public:
 
 	LANDSCAPE_API void SnapFoliageInstances();
 #endif
+
 	public:
 	TOptional<float> GetHeight(float X, float Y);
 };
diff --git a/Engine/Source/Runtime/Landscape/Private/LandscapeCollision.cpp b/Engine/Source/Runtime/Landscape/Private/LandscapeCollision.cpp
--- a/Engine/Source/Runtime/Landscape/Private/LandscapeCollision.cpp
+++ b/Engine/Source/Runtime/Landscape/Private/LandscapeCollision.cpp
@@ -448,6 +448,32 @@ void ULandscapeHeightfieldCollisionComponent::OnCreatePhysicsState()
 					ShapeArray.Emplace(MoveTemp(NewSimpleShape));
 				}
 
+#if WITH_EDITOR
+				// Create a shape for a heightfield which is used only by the landscape editor
+				if(!GetWorld()->IsGameWorld())
+				{
+					TUniquePtr<Chaos::TPerShapeData<float, 3>> NewEditorShape = Chaos::TPerShapeData<float, 3>::CreatePerShapeData();
+
+					HeightfieldRef->EditorHeightfield->SetScale(FinalScale);
+					TUniquePtr<Chaos::TImplicitObjectTransformed<float, 3>> ChaosEditorHeightFieldFromCooked = MakeUnique<Chaos::TImplicitObjectTransformed<float, 3>>(MakeSerializable(HeightfieldRef->EditorHeightfield), Chaos::TRigidTransform<float, 3>(FTransform::Identity));
+
+					FCollisionResponseContainer CollisionResponse;
+					CollisionResponse.SetAllChannels(ECollisionResponse::ECR_Ignore);
+					CollisionResponse.SetResponse(ECollisionChannel::ECC_Visibility, ECR_Block);
+					FCollisionFilterData QueryFilterDataEd, SimFilterDataEd;
+					CreateShapeFilterData(ECollisionChannel::ECC_Visibility, FMaskFilter(0), GetOwner()->GetUniqueID(), CollisionResponse, GetUniqueID(), 0, QueryFilterDataEd, SimFilterDataEd, true, false, true);
+
+					QueryFilterDataEd.Word3 |= (EPDF_SimpleCollision | EPDF_ComplexCollision);
+
+					NewEditorShape->Geometry = MakeSerializable(ChaosEditorHeightFieldFromCooked);
+					NewEditorShape->QueryData = QueryFilterDataEd;
+					NewEditorShape->SimData = SimFilterDataEd;
+
+					Geoms.Emplace(MoveTemp(ChaosEditorHeightFieldFromCooked));
+					ShapeArray.Emplace(MoveTemp(NewEditorShape));
+				}
+#endif
+				// Push the shapes to the actor
 				if(Geoms.Num() == 1)
 				{
 					PhysHandle->SetGeometry(MoveTemp(Geoms[0]));
@@ -456,27 +482,54 @@ void ULandscapeHeightfieldCollisionComponent::OnCreatePhysicsState()
 				{
 					PhysHandle->SetGeometry(MakeUnique<Chaos::TImplicitObjectUnion<float, 3>>(MoveTemp(Geoms)));
 				}
-				
+
 				PhysHandle->SetShapesArray(MoveTemp(ShapeArray));
 
+				// Push the actor to the scene
 				FPhysScene* PhysScene = GetWorld()->GetPhysicsScene();
 
 				// Set body instance data
 				BodyInstance.PhysicsUserData = FPhysicsUserData(&BodyInstance);
 				BodyInstance.OwnerComponent = this;
+				BodyInstance.ActorHandle = PhysHandle;
 
 				PhysHandle->SetUserData(&BodyInstance.PhysicsUserData);
 
 				TArray<FPhysicsActorHandle> Actors;
 				Actors.Add(PhysHandle);
 				PhysScene->AddActorsToScene_AssumesLocked(Actors);
-			}
+
+#if WITH_EDITOR
+				// If we're in an editor world we will never simulate but we require up to date SQ for the tools
+				// #BGTODO consider a more automatic way to have scenes update when they aren't simulating to avoid this explicit flush
+				if(!GetWorld()->IsGameWorld())
+				{
+					PhysScene->Flush_AssumesLocked();
+				}
 #endif
+			}
+#endif // WITH_CHAOS
 		}
 #endif// WITH_PHYSX
 	}
 }
 
+void ULandscapeHeightfieldCollisionComponent::OnDestroyPhysicsState()
+{
+	Super::OnDestroyPhysicsState();
+	
+#if WITH_EDITOR && WITH_CHAOS
+	// If we're in an editor world we will never simulate but we require up to date SQ for the tools
+	// #BGTODO consider a more automatic way to have scenes update when they aren't simulating to avoid this explicit flush
+	if(!GetWorld()->IsGameWorld())
+	{
+		FPhysScene* PhysScene = GetWorld()->GetPhysicsScene();
+		check(PhysScene);
+		PhysScene->Flush_AssumesLocked();
+	}
+#endif
+}
+
 void ULandscapeHeightfieldCollisionComponent::ApplyWorldOffset(const FVector& InOffset, bool bWorldShift)
 {
 	Super::ApplyWorldOffset(InOffset, bWorldShift);
@@ -1325,7 +1378,10 @@ void ULandscapeHeightfieldCollisionComponent::UpdateHeightfieldRegion(int32 Comp
 		}
 
 #if WITH_CHAOS || WITH_IMMEDIATE_PHYSX
-        ensure(false);
+		if(!BodyInstance.ActorHandle)
+		{
+			return;
+		}
 #else
 		if (BodyInstance.ActorHandle.SyncActor == NULL)
 		{
@@ -1347,6 +1403,7 @@ void ULandscapeHeightfieldCollisionComponent::UpdateHeightfieldRegion(int32 Comp
 			uint16* Heights = (uint16*)CollisionHeightData.Lock(LOCK_READ_ONLY);
 			check(CollisionHeightData.GetElementCount() == (FMath::Square(CollisionSizeVerts) + FMath::Square(SimpleCollisionSizeVerts)));
 	
+#if PHYSICS_INTERFACE_PHYSX
 			// PhysX heightfield has the X and Y axis swapped, and the X component is also inverted
 			int32 HeightfieldX1 = ComponentY1;
 			int32 HeightfieldY1 = (bIsMirrored ? ComponentX1 : (CollisionSizeVerts - ComponentX2 - 1));
@@ -1393,20 +1450,43 @@ void ULandscapeHeightfieldCollisionComponent::UpdateHeightfieldRegion(int32 Comp
 			FVector LandscapeScale = GetComponentToWorld().GetScale3D().GetAbs();
 			// Create the geometry
 			PxHeightFieldGeometry LandscapeComponentGeom(HeightfieldRef->RBHeightfieldEd, PxMeshGeometryFlags(), LandscapeScale.Z * LANDSCAPE_ZSCALE, LandscapeScale.Y * CollisionScale, LandscapeScale.X * CollisionScale);
-	
+
 			{
 				FInlineShapeArray PShapes;
-#if WITH_CHAOS
-				ensure(false);
-				const int32 NumShapes = 0;
-#else
+
 				const int32 NumShapes = FillInlineShapeArray_AssumesLocked(PShapes, Actor);
-#endif
-				if (NumShapes > 1)
+				if(NumShapes > 1)
 				{
 					FPhysicsInterface::SetGeometry(PShapes[1], LandscapeComponentGeom);
 				}
 			}
+
+#elif WITH_CHAOS
+			int32 HeightfieldY1 = ComponentY1;
+			int32 HeightfieldX1 = (bIsMirrored ? ComponentX1 : (CollisionSizeVerts - ComponentX2 - 1));
+			int32 DstVertsX = ComponentX2 - ComponentX1 + 1;
+			int32 DstVertsY = ComponentY2 - ComponentY1 + 1;
+			TArray<uint16> Samples;
+			Samples.AddZeroed(DstVertsX*DstVertsY);
+
+			for(int32 RowIndex = 0; RowIndex < DstVertsY; RowIndex++)
+			{
+				for(int32 ColIndex = 0; ColIndex < DstVertsX; ColIndex++)
+				{
+					int32 SrcX = bIsMirrored ? (ColIndex + ComponentX1) : (ComponentX2 - ColIndex);
+					int32 SrcY = RowIndex + ComponentY1;
+					int32 SrcSampleIndex = (SrcY * CollisionSizeVerts) + SrcX;
+					check(SrcSampleIndex < FMath::Square(CollisionSizeVerts));
+					int32 DstSampleIndex = (RowIndex * DstVertsX) + ColIndex;
+
+					Samples[DstSampleIndex] = Heights[SrcSampleIndex];
+				}
+			}
+
+			CollisionHeightData.Unlock();
+
+			HeightfieldRef->EditorHeightfield->EditHeights(Samples, HeightfieldY1, HeightfieldX1, DstVertsY, DstVertsX);
+#endif
 		});
 	}
 
