Fix build game client build error introduced by 5235541

diff --git a/Engine/Source/Runtime/Landscape/Private/LandscapeRender.cpp b/Engine/Source/Runtime/Landscape/Private/LandscapeRender.cpp
--- a/Engine/Source/Runtime/Landscape/Private/LandscapeRender.cpp
+++ b/Engine/Source/Runtime/Landscape/Private/LandscapeRender.cpp
@@ -4415,6 +4415,8 @@ void FLandscapeComponentSceneProxy::GetHeightfieldRepresentation(UTexture2D*& Ou
 
 	OutVisibilityTexture = nullptr;
 	OutDescription.VisibilityChannel = -1;
+
+#if WITH_EDITOR
 	for (auto& Allocation : LandscapeComponent->WeightmapLayerAllocations)
 	{
 		if (Allocation.LayerInfo == ALandscapeProxy::VisibilityLayer)
@@ -4424,6 +4426,7 @@ void FLandscapeComponentSceneProxy::GetHeightfieldRepresentation(UTexture2D*& Ou
 			break;
 		}
 	}
+#endif
 }
 
 void FLandscapeComponentSceneProxy::GetLCIs(FLCIArray& LCIs)
