Fix some reading of ULandscapeComponent from the render thread

diff --git a/Engine/Source/Runtime/Landscape/Private/LandscapeRender.cpp b/Engine/Source/Runtime/Landscape/Private/LandscapeRender.cpp
--- a/Engine/Source/Runtime/Landscape/Private/LandscapeRender.cpp
+++ b/Engine/Source/Runtime/Landscape/Private/LandscapeRender.cpp
@@ -631,6 +631,8 @@ FLandscapeComponentSceneProxy::FLandscapeComponentSceneProxy(ULandscapeComponent
 	, LandscapeComponent(InComponent)
 	, WeightmapScaleBias(InComponent->WeightmapScaleBias)
 	, WeightmapTextures(InComponent->WeightmapTextures)
+	, VisibilityWeightmapTexture(nullptr)
+	, VisibilityWeightmapChannel(-1)
 	, NormalmapTexture(InComponent->GetHeightmap())
 	, BaseColorForGITexture(InComponent->GIBakedBaseColorTexture)
 	, HeightmapScaleBias(InComponent->HeightmapScaleBias)
@@ -845,6 +847,17 @@ FLandscapeComponentSceneProxy::FLandscapeComponentSceneProxy(ULandscapeComponent
 			LayerColors.Add(Allocation.LayerInfo->LayerUsageDebugColor);
 		}
 	}
+
+	for (int32 Idx = 0; Idx < InComponent->WeightmapLayerAllocations.Num(); Idx++)
+	{
+		FWeightmapLayerAllocationInfo& Allocation = InComponent->WeightmapLayerAllocations[Idx];
+		if (Allocation.LayerInfo == ALandscapeProxy::VisibilityLayer)
+		{
+			VisibilityWeightmapTexture = WeightmapTextures[Idx];
+			VisibilityWeightmapChannel = Allocation.WeightmapTextureChannel;
+			break;
+		}
+	}
 #endif
 }
 
@@ -4399,6 +4412,8 @@ void FLandscapeComponentSceneProxy::GetHeightfieldRepresentation(UTexture2D*& Ou
 {
 	OutHeightmapTexture = HeightmapTexture;
 	OutDiffuseColorTexture = BaseColorForGITexture;
+	OutVisibilityTexture = VisibilityWeightmapTexture;
+	
 	OutDescription.HeightfieldScaleBias = HeightmapScaleBias;
 
 	OutDescription.MinMaxUV = FVector4(
@@ -4413,20 +4428,7 @@ void FLandscapeComponentSceneProxy::GetHeightfieldRepresentation(UTexture2D*& Ou
 
 	OutDescription.SubsectionScaleAndBias = FVector4(SubsectionSizeQuads, SubsectionSizeQuads, HeightmapSubsectionOffsetU, HeightmapSubsectionOffsetV);
 
-	OutVisibilityTexture = nullptr;
-	OutDescription.VisibilityChannel = -1;
-
-#if WITH_EDITOR
-	for (auto& Allocation : LandscapeComponent->WeightmapLayerAllocations)
-	{
-		if (Allocation.LayerInfo == ALandscapeProxy::VisibilityLayer)
-		{
-			OutVisibilityTexture = WeightmapTextures[Allocation.WeightmapTextureIndex];
-			OutDescription.VisibilityChannel = Allocation.WeightmapTextureChannel;
-			break;
-		}
-	}
-#endif
+	OutDescription.VisibilityChannel = VisibilityWeightmapChannel;
 }
 
 void FLandscapeComponentSceneProxy::GetLCIs(FLCIArray& LCIs)
diff --git a/Engine/Source/Runtime/Landscape/Public/LandscapeRender.h b/Engine/Source/Runtime/Landscape/Public/LandscapeRender.h
--- a/Engine/Source/Runtime/Landscape/Public/LandscapeRender.h
+++ b/Engine/Source/Runtime/Landscape/Public/LandscapeRender.h
@@ -574,6 +574,10 @@ protected:
 
 	FVector4 WeightmapScaleBias;
 	TArray<UTexture2D*> WeightmapTextures;
+
+	UTexture2D* VisibilityWeightmapTexture;
+	int32 VisibilityWeightmapChannel;
+
 #if WITH_EDITOR
 	TArray<FLinearColor> LayerColors;
 #endif
