Break ARuntimeVirtualTexturePlane up to use a component
Fixes a couple of update issues and opens up use of scene proxy

diff --git a/Engine/Source/Editor/VirtualTexturingEditor/Classes/RuntimeVirtualTextureDetailsCustomization.h b/Engine/Source/Editor/VirtualTexturingEditor/Classes/RuntimeVirtualTextureDetailsCustomization.h
--- a/Engine/Source/Editor/VirtualTexturingEditor/Classes/RuntimeVirtualTextureDetailsCustomization.h
+++ b/Engine/Source/Editor/VirtualTexturingEditor/Classes/RuntimeVirtualTextureDetailsCustomization.h
@@ -38,13 +38,13 @@ private:
 
 
 /** UI customization for ARuntimeVirtualTexturePlane */
-class FRuntimeVirtualTextureActorDetailsCustomization : public IDetailCustomization
+class FRuntimeVirtualTextureComponentDetailsCustomization : public IDetailCustomization
 {
 public:
 	static TSharedRef<IDetailCustomization> MakeInstance();
 
 protected:
-	FRuntimeVirtualTextureActorDetailsCustomization();
+	FRuntimeVirtualTextureComponentDetailsCustomization();
 
 	/** Callback for Copy Rotation button */
 	FReply SetRotation();
@@ -56,5 +56,5 @@ protected:
 	//~ End IDetailCustomization Interface.
 
 private:
-	class ARuntimeVirtualTexturePlane* VirtualTexturePlane;
+	class URuntimeVirtualTextureComponent* RuntimeVirtualTextureComponent;
 };
diff --git a/Engine/Source/Editor/VirtualTexturingEditor/Private/RuntimeVirtualTextureDetailsCustomization.cpp b/Engine/Source/Editor/VirtualTexturingEditor/Private/RuntimeVirtualTextureDetailsCustomization.cpp
--- a/Engine/Source/Editor/VirtualTexturingEditor/Private/RuntimeVirtualTextureDetailsCustomization.cpp
+++ b/Engine/Source/Editor/VirtualTexturingEditor/Private/RuntimeVirtualTextureDetailsCustomization.cpp
@@ -153,32 +153,32 @@ void FRuntimeVirtualTextureDetailsCustomization::RefreshDetails()
 }
 
 
-FRuntimeVirtualTextureActorDetailsCustomization::FRuntimeVirtualTextureActorDetailsCustomization()
+FRuntimeVirtualTextureComponentDetailsCustomization::FRuntimeVirtualTextureComponentDetailsCustomization()
 {
 }
 
-TSharedRef<IDetailCustomization> FRuntimeVirtualTextureActorDetailsCustomization::MakeInstance()
+TSharedRef<IDetailCustomization> FRuntimeVirtualTextureComponentDetailsCustomization::MakeInstance()
 {
-	return MakeShareable(new FRuntimeVirtualTextureActorDetailsCustomization);
+	return MakeShareable(new FRuntimeVirtualTextureComponentDetailsCustomization);
 }
 
-void FRuntimeVirtualTextureActorDetailsCustomization::CustomizeDetails(IDetailLayoutBuilder& DetailBuilder)
+void FRuntimeVirtualTextureComponentDetailsCustomization::CustomizeDetails(IDetailLayoutBuilder& DetailBuilder)
 {
-	// Get and store the linked ARuntimeVirtualTexturePlane
+	// Get and store the linked ARuntimeRuntimeVirtualTextureComponent
 	TArray<TWeakObjectPtr<UObject>> ObjectsBeingCustomized;
 	DetailBuilder.GetObjectsBeingCustomized(ObjectsBeingCustomized);
 	if (ObjectsBeingCustomized.Num() > 1)
 	{
 		return;
 	}
-	VirtualTexturePlane = Cast<ARuntimeVirtualTexturePlane>(ObjectsBeingCustomized[0].Get());
-	if (VirtualTexturePlane == nullptr)
+	RuntimeVirtualTextureComponent = Cast<URuntimeVirtualTextureComponent>(ObjectsBeingCustomized[0].Get());
+	if (RuntimeVirtualTextureComponent == nullptr)
 	{
 		return;
 	}
 
 	// Use SourceActor property to add buttons
-	TSharedPtr<IPropertyHandle> SourceActorValue = DetailBuilder.GetProperty("SourceActor");
+	TSharedPtr<IPropertyHandle> SourceActorValue = DetailBuilder.GetProperty("BoundsSourceActor");
 	DetailBuilder.HideProperty(SourceActorValue);
 
 	IDetailCategoryBuilder& BoundsCategory = DetailBuilder.EditCategory("TransformFromBounds", FText::GetEmpty(), ECategoryPriority::Important);
@@ -218,7 +218,7 @@ void FRuntimeVirtualTextureActorDetailsCustomization::CustomizeDetails(IDetailLa
 						.HAlign(HAlign_Center)
 						.Text(LOCTEXT("Button_CopyRotation", "Copy Rotation"))
 						.ToolTipText(LOCTEXT("Button_CopyRotation_Tooltip", "Set the virtual texture rotation to match the source actor"))
-						.OnClicked(this, &FRuntimeVirtualTextureActorDetailsCustomization::SetRotation)
+						.OnClicked(this, &FRuntimeVirtualTextureComponentDetailsCustomization::SetRotation)
 					]
 
 					+ SVerticalBox::Slot()
@@ -228,7 +228,7 @@ void FRuntimeVirtualTextureActorDetailsCustomization::CustomizeDetails(IDetailLa
 						.HAlign(HAlign_Center)
 						.Text(LOCTEXT("Button_CopyBounds", "Copy Bounds"))
 						.ToolTipText(LOCTEXT("Button_CopyBounds_Tooltip", "Set the virtual texture transform so that it includes the full bounds of the source actor"))
-						.OnClicked(this, &FRuntimeVirtualTextureActorDetailsCustomization::SetTransformToBounds)
+						.OnClicked(this, &FRuntimeVirtualTextureComponentDetailsCustomization::SetTransformToBounds)
 					]
 				]
 			]
@@ -236,15 +236,15 @@ void FRuntimeVirtualTextureActorDetailsCustomization::CustomizeDetails(IDetailLa
 	];
 }
 
-FReply FRuntimeVirtualTextureActorDetailsCustomization::SetRotation()
+FReply FRuntimeVirtualTextureComponentDetailsCustomization::SetRotation()
 {
-	VirtualTexturePlane->SetRotation();
+	RuntimeVirtualTextureComponent->SetRotation();
 	return FReply::Handled();
 }
 
-FReply FRuntimeVirtualTextureActorDetailsCustomization::SetTransformToBounds()
+FReply FRuntimeVirtualTextureComponentDetailsCustomization::SetTransformToBounds()
 {
-	VirtualTexturePlane->SetTransformToBounds();
+	RuntimeVirtualTextureComponent->SetTransformToBounds();
 	return FReply::Handled();
 }
 
diff --git a/Engine/Source/Editor/VirtualTexturingEditor/Private/VirtualTexturingEditorModule.cpp b/Engine/Source/Editor/VirtualTexturingEditor/Private/VirtualTexturingEditorModule.cpp
--- a/Engine/Source/Editor/VirtualTexturingEditor/Private/VirtualTexturingEditorModule.cpp
+++ b/Engine/Source/Editor/VirtualTexturingEditor/Private/VirtualTexturingEditorModule.cpp
@@ -19,7 +19,7 @@ public:
 	
 		FPropertyEditorModule& PropertyModule = FModuleManager::GetModuleChecked<FPropertyEditorModule>("PropertyEditor");
 		PropertyModule.RegisterCustomClassLayout("RuntimeVirtualTexture", FOnGetDetailCustomizationInstance::CreateStatic(&FRuntimeVirtualTextureDetailsCustomization::MakeInstance));
-		PropertyModule.RegisterCustomClassLayout("RuntimeVirtualTexturePlane", FOnGetDetailCustomizationInstance::CreateStatic(&FRuntimeVirtualTextureActorDetailsCustomization::MakeInstance));
+		PropertyModule.RegisterCustomClassLayout("RuntimeVirtualTextureComponent", FOnGetDetailCustomizationInstance::CreateStatic(&FRuntimeVirtualTextureComponentDetailsCustomization::MakeInstance));
 	}
 
 	virtual void ShutdownModule() override
diff --git a/Engine/Source/Runtime/Landscape/Classes/RuntimeVirtualTexturePlane.h b/Engine/Source/Runtime/Landscape/Classes/RuntimeVirtualTexturePlane.h
--- a/Engine/Source/Runtime/Landscape/Classes/RuntimeVirtualTexturePlane.h
+++ b/Engine/Source/Runtime/Landscape/Classes/RuntimeVirtualTexturePlane.h
@@ -8,59 +8,72 @@
 
 class UMaterialInterface;
 class URuntimeVirtualTexture;
+class URuntimeVirtualTextureComponent;
 
-/** Actor used to place a URuntimeVirtualTexture in the world.
- *  todo[vt]: I think that this belongs in the Engine lib (not Landscape) but that requires removing the dependency on Renderer for direct creation of a FRuntimeVirtualTextureProducer
- *  todo[vt]: If we fix that and move to Engine then remove build dependency on Landscape in VirtualTexturingEditor.Build.cs
- */
+/** Actor used to place a URuntimeVirtualTexture in the world. */
 UCLASS(hidecategories=(Actor, Collision, Cooking, Input, LOD, Replication), MinimalAPI)
 class ARuntimeVirtualTexturePlane : public AActor
 {
 	GENERATED_UCLASS_BODY()
 
-public:
-	/** Actor to copy the bounds from to set up the transform. */
-	UPROPERTY(EditAnywhere, DuplicateTransient, Category = TransformFromBounds)
-	AActor* SourceActor = nullptr;
-
-	/** The virtual texture object to use. */
-	UPROPERTY(EditAnywhere, DuplicateTransient, Category = VirtualTexture)
-	URuntimeVirtualTexture* VirtualTexture = nullptr;
-
 private:
+	/** Component that owns the runtime virtual texture. */
+	UPROPERTY(VisibleAnywhere, BlueprintReadOnly, Category = VirtualTexture, meta = (AllowPrivateAccess = "true"))
+	class URuntimeVirtualTextureComponent* VirtualTextureComponent;
+
 #if WITH_EDITORONLY_DATA
 	/** Box for visualizing virtual texture extents. */
 	UPROPERTY(Transient)
 	class UBoxComponent* Box = nullptr;
 #endif // WITH_EDITORONLY_DATA
 
+protected:
+	//~ Begin AActor Interface.
+	virtual void PostEditMove(bool bFinished) override;
+	virtual bool IsLevelBoundsRelevant() const override { return false; }
+	//~ End AActor Interface
+};
+
+/** Component used to place a URuntimeVirtualTexture in the world. */
+UCLASS(ClassGroup = Rendering, collapsecategories, hidecategories = (Activation, Collision, Cooking, Mobility, LOD, Object, Physics, Rendering), editinlinenew)
+class LANDSCAPE_API URuntimeVirtualTextureComponent : public USceneComponent
+{
+	GENERATED_UCLASS_BODY()
+
+private:
+	/** The virtual texture object to use. */
+	UPROPERTY(EditAnywhere, DuplicateTransient, Category = VirtualTexture)
+	URuntimeVirtualTexture* VirtualTexture = nullptr;
+
+	/** Actor to copy the bounds from to set up the transform. */
+	UPROPERTY(EditAnywhere, DuplicateTransient, Category = TransformFromBounds, meta = (DisplayName = "Source Actor"))
+	AActor* BoundsSourceActor = nullptr;
+
+public:
+	/** Call whenever we need to update the underlying URuntimeVirtualTexture. */
+	void UpdateVirtualTexture();
+
+	/** Call when we need to disconnect from the underlying URuntimeVirtualTexture. */
+	void ReleaseVirtualTexture();
+
 #if WITH_EDITOR
 	UFUNCTION(CallInEditor)
 	void OnVirtualTextureEditProperty(URuntimeVirtualTexture const* InVirtualTexture);
-#endif
 
-public:
-#if WITH_EDITOR
 	/** Copy the rotation from SourceActor to this component. Called by our UI details customization. */
-	LANDSCAPE_API void SetRotation();
+	void SetRotation();
 
 	/** Set this component transform to include the SourceActor bounds. Called by our UI details customization. */
-	LANDSCAPE_API void SetTransformToBounds();
+	void SetTransformToBounds();
 #endif
 
 protected:
-	/** Call whenever we need to update the underlying URuntimeVirtualTexture. */
-	void UpdateVirtualTexture();
-	/** Call when we need to disconnect from the underlying URuntimeVirtualTexture. */
-	void ReleaseVirtualTexture();
-
-	//~ Begin AActor Interface.
-	virtual void PostRegisterAllComponents() override;
+	//~ Begin UObject Interface
 	virtual void PostLoad() override;
 	virtual void BeginDestroy() override;
-	virtual bool IsLevelBoundsRelevant() const override { return false; }
-#if WITH_EDITOR
-	virtual void PostEditMove(bool bFinished);
-#endif
-	//~ End AActor Interface.
+	//~ End UObject Interface
+
+	//~ Begin UActorComponent Interface
+	virtual void OnRegister() override;
+	//~ End UActorComponent Interface
 };
diff --git a/Engine/Source/Runtime/Landscape/Private/RuntimeVirtualTexturePlane.cpp b/Engine/Source/Runtime/Landscape/Private/RuntimeVirtualTexturePlane.cpp
--- a/Engine/Source/Runtime/Landscape/Private/RuntimeVirtualTexturePlane.cpp
+++ b/Engine/Source/Runtime/Landscape/Private/RuntimeVirtualTexturePlane.cpp
@@ -10,9 +10,10 @@
 ARuntimeVirtualTexturePlane::ARuntimeVirtualTexturePlane(const FObjectInitializer& ObjectInitializer)
 	: Super(ObjectInitializer)
 {
-	RootComponent = CreateDefaultSubobject<USceneComponent>(TEXT("RootComponent"));
+	RootComponent = VirtualTextureComponent = CreateDefaultSubobject<URuntimeVirtualTextureComponent>(TEXT("VirtualTextureComponent"));
 
 #if WITH_EDITORONLY_DATA
+	// Add box for visualization of bounds
 	Box = CreateDefaultSubobject<UBoxComponent>(TEXT("Box"));
 	Box->SetBoxExtent(FVector(0.5f, 0.5f, 1.f), false);
 	Box->SetIsVisualizationComponent(true);
@@ -20,11 +21,51 @@ ARuntimeVirtualTexturePlane::ARuntimeVirtualTexturePlane(const FObjectInitialize
 	Box->SetCanEverAffectNavigation(false);
 	Box->CanCharacterStepUpOn = ECanBeCharacterBase::ECB_No;
 	Box->SetGenerateOverlapEvents(false);
-	Box->SetupAttachment(RootComponent);
+	Box->SetupAttachment(VirtualTextureComponent);
 #endif
 }
 
-void ARuntimeVirtualTexturePlane::UpdateVirtualTexture()
+#if WITH_EDITOR
+
+void ARuntimeVirtualTexturePlane::PostEditMove(bool bFinished)
+{
+	if (bFinished)
+	{
+		if (VirtualTextureComponent != nullptr)
+		{
+			VirtualTextureComponent->UpdateVirtualTexture();
+		}
+	}
+	Super::PostEditMove(bFinished);
+}
+
+#endif
+
+
+URuntimeVirtualTextureComponent::URuntimeVirtualTextureComponent(const FObjectInitializer& ObjectInitializer)
+	: Super(ObjectInitializer)
+{
+}
+
+void URuntimeVirtualTextureComponent::OnRegister()
+{
+	Super::OnRegister();
+	UpdateVirtualTexture();
+}
+
+void URuntimeVirtualTextureComponent::PostLoad()
+{
+	Super::PostLoad();
+	UpdateVirtualTexture();
+}
+
+void URuntimeVirtualTextureComponent::BeginDestroy()
+{
+	ReleaseVirtualTexture();
+	Super::BeginDestroy();
+}
+
+void URuntimeVirtualTextureComponent::UpdateVirtualTexture()
 {
 	if (VirtualTexture != nullptr)
 	{
@@ -35,11 +76,11 @@ void ARuntimeVirtualTexturePlane::UpdateVirtualTexture()
 		const ERuntimeVirtualTextureMaterialType MaterialType = VirtualTexture->GetMaterialType();
 
 		// Transform is based on bottom left of the box
-		FTransform Transform = FTransform(FVector(-0.5f, -0.5f, 0.f)) * GetTransform();
-		
-		FRuntimeVirtualTextureProducer* Producer = new FRuntimeVirtualTextureProducer(Desc, MaterialType, RootComponent->GetScene(), Transform);
+		FTransform Transform = FTransform(FVector(-0.5f, -0.5f, 0.f)) * GetComponentToWorld();
+
+		FRuntimeVirtualTextureProducer* Producer = new FRuntimeVirtualTextureProducer(Desc, MaterialType, GetScene(), Transform);
 		VirtualTexture->Initialize(Producer, Transform);
-		
+
 #if WITH_EDITOR
 		// Bind function to ensure we call ReInit again if the virtual texture properties are modified
 		static const FName BinderFunction(TEXT("OnVirtualTextureEditProperty"));
@@ -48,7 +89,7 @@ void ARuntimeVirtualTexturePlane::UpdateVirtualTexture()
 	}
 }
 
-void ARuntimeVirtualTexturePlane::ReleaseVirtualTexture()
+void URuntimeVirtualTextureComponent::ReleaseVirtualTexture()
 {
 	if (VirtualTexture != nullptr)
 	{
@@ -60,27 +101,9 @@ void ARuntimeVirtualTexturePlane::ReleaseVirtualTexture()
 	}
 }
 
-void ARuntimeVirtualTexturePlane::PostRegisterAllComponents()
-{
-	Super::PostRegisterAllComponents();
-	UpdateVirtualTexture();
-}
-
-void ARuntimeVirtualTexturePlane::PostLoad()
-{
-	Super::PostLoad();
-	UpdateVirtualTexture();
-}
-
-void ARuntimeVirtualTexturePlane::BeginDestroy()
-{
-	ReleaseVirtualTexture();
-	Super::BeginDestroy();
-}
-
 #if WITH_EDITOR
 
-void ARuntimeVirtualTexturePlane::OnVirtualTextureEditProperty(URuntimeVirtualTexture const* InVirtualTexture)
+void URuntimeVirtualTextureComponent::OnVirtualTextureEditProperty(URuntimeVirtualTexture const* InVirtualTexture)
 {
 	if (InVirtualTexture == VirtualTexture)
 	{
@@ -88,41 +111,31 @@ void ARuntimeVirtualTexturePlane::OnVirtualTextureEditProperty(URuntimeVirtualTe
 	}
 }
 
-void ARuntimeVirtualTexturePlane::PostEditMove(bool bFinished)
+void URuntimeVirtualTextureComponent::SetRotation()
 {
-	if (bFinished)
+	if (BoundsSourceActor != nullptr)
 	{
-		UpdateVirtualTexture();
+		// Copy the source actor rotation and notify the parent actor
+		SetWorldRotation(BoundsSourceActor->GetTransform().GetRotation());
+		GetOwner()->PostEditMove(true);
 	}
-	Super::PostEditMove(bFinished);
 }
 
-void ARuntimeVirtualTexturePlane::SetRotation()
+void URuntimeVirtualTextureComponent::SetTransformToBounds()
 {
-	if (SourceActor != nullptr)
+	if (BoundsSourceActor != nullptr)
 	{
-		RootComponent->SetWorldRotation(SourceActor->GetTransform().GetRotation());
-
-		// Update the virtual texture to match the new transform
-		UpdateVirtualTexture();
-	}
-}
-
-void ARuntimeVirtualTexturePlane::SetTransformToBounds()
-{
-	if (SourceActor != nullptr)
-	{
-		// Calculate the bounds in our local rotation space translated to the SourceActor center
-		const FQuat TargetRotation = GetTransform().GetRotation();
-		const FVector InitialPosition = SourceActor->GetComponentsBoundingBox().GetCenter();
+		// Calculate the bounds in our local rotation space translated to the BoundsSourceActor center
+		const FQuat TargetRotation = GetComponentToWorld().GetRotation();
+		const FVector InitialPosition = BoundsSourceActor->GetComponentsBoundingBox().GetCenter();
 		const FVector InitialScale = FVector(0.5f, 0.5, 1.f);
 
 		FTransform LocalTransform;
 		LocalTransform.SetComponents(TargetRotation, InitialPosition, InitialScale);
 		FTransform WorldToLocal = LocalTransform.Inverse();
 
-		FBox Bounds(ForceInit);
-		for (const UActorComponent* Component : SourceActor->GetComponents())
+		FBox BoundBox(ForceInit);
+		for (const UActorComponent* Component : BoundsSourceActor->GetComponents())
 		{
 			// Only gather visual components in the bounds calculation
 			const UPrimitiveComponent* PrimitiveComponent = Cast<const UPrimitiveComponent>(Component);
@@ -130,24 +143,23 @@ void ARuntimeVirtualTexturePlane::SetTransformToBounds()
 			{
 				const FTransform ComponentToActor = PrimitiveComponent->GetComponentTransform() * WorldToLocal;
 				FBoxSphereBounds LocalSpaceComponentBounds = PrimitiveComponent->CalcBounds(ComponentToActor);
-				Bounds += LocalSpaceComponentBounds.GetBox();
+				BoundBox += LocalSpaceComponentBounds.GetBox();
 			}
 		}
 
 		// Create transform from bounds
 		FVector Origin;
 		FVector Extent;
-		Bounds.GetCenterAndExtents(Origin, Extent);
+		BoundBox.GetCenterAndExtents(Origin, Extent);
 
 		Origin = LocalTransform.TransformPosition(Origin);
 
 		FTransform Transform;
 		Transform.SetComponents(TargetRotation, Origin, Extent);
 
-		RootComponent->SetWorldTransform(Transform);
-
-		// Update the virtual texture to match the new transform
-		UpdateVirtualTexture();
+		// Apply final result and notify the parent actor
+		SetWorldTransform(Transform);
+		GetOwner()->PostEditMove(true);
 	}
 }
 
