Fixes non editor builds

diff --git a/Engine/Source/Runtime/Landscape/Classes/RuntimeVirtualTexturePlane.h b/Engine/Source/Runtime/Landscape/Classes/RuntimeVirtualTexturePlane.h
--- a/Engine/Source/Runtime/Landscape/Classes/RuntimeVirtualTexturePlane.h
+++ b/Engine/Source/Runtime/Landscape/Classes/RuntimeVirtualTexturePlane.h
@@ -29,7 +29,9 @@ private:
 
 protected:
 	//~ Begin AActor Interface.
+#if WITH_EDITOR
 	virtual void PostEditMove(bool bFinished) override;
+#endif
 	virtual bool IsLevelBoundsRelevant() const override { return false; }
 	//~ End AActor Interface
 };
