Add ray tracing support for landscapes

* Landscapes are triangulated into meshes every frame. This makes it possible to match the raster path perfectly to avoid shadow artifacts. World position offset is also supported.
* Measured performance in KiteDemo: ~6ms geometry update time (1.5ms VSinCS dispatch, the rest is BLAS refit time) and ~500MB video memory

diff --git a/Engine/Plugins/FX/Niagara/Source/Niagara/Private/NiagaraRendererSprites.cpp b/Engine/Plugins/FX/Niagara/Source/Niagara/Private/NiagaraRendererSprites.cpp
--- a/Engine/Plugins/FX/Niagara/Source/Niagara/Private/NiagaraRendererSprites.cpp
+++ b/Engine/Plugins/FX/Niagara/Source/Niagara/Private/NiagaraRendererSprites.cpp
@@ -372,6 +372,7 @@ void NiagaraRendererSprites::CreateMeshBatchForView(
 
 	MeshBatch.VertexFactory = &CollectorResources.VertexFactory;
 	MeshBatch.CastShadow = SceneProxy->CastsDynamicShadow();
+	MeshBatch.CastRayTracedShadow = SceneProxy->CastsDynamicShadow();
 	MeshBatch.bUseAsOccluder = false;
 	MeshBatch.ReverseCulling = SceneProxy->IsLocalToWorldDeterminantNegative();
 	MeshBatch.Type = PT_TriangleList;
@@ -493,14 +494,18 @@ void NiagaraRendererSprites::GetDynamicRayTracingInstances(FRayTracingMaterialGa
 		RayTracingInstance.Materials.Add(MeshBatch);
 
 		// Update dynamic ray tracing geometry
-		Context.DynamicRayTracingGeometriesToUpdate.AddDynamicMeshBatchForGeometryUpdate(
-			Context.Scene,
-			Context.ReferenceView,
-			SceneProxy,
-			MeshBatch,
-			RayTracingGeometry,
-			6 * DynamicDataSprites->RTParticleData.GetNumInstances(),
-			RayTracingDynamicVertexBuffer);
+		Context.DynamicRayTracingGeometriesToUpdate.Add(
+			FRayTracingDynamicGeometryUpdateParams
+			{
+				MeshBatch,
+				MeshBatch.Elements[0].NumPrimitives == 0,
+				6 * DynamicDataSprites->RTParticleData.GetNumInstances(),
+				6 * DynamicDataSprites->RTParticleData.GetNumInstances() * sizeof(FVector),
+				2 * DynamicDataSprites->RTParticleData.GetNumInstances(),
+				&RayTracingGeometry,
+				&RayTracingDynamicVertexBuffer
+			}
+		);
 	}
 
 	RayTracingInstance.BuildInstanceMaskAndFlags();
diff --git a/Engine/Shaders/Private/LandscapeVertexFactory.ush b/Engine/Shaders/Private/LandscapeVertexFactory.ush
--- a/Engine/Shaders/Private/LandscapeVertexFactory.ush
+++ b/Engine/Shaders/Private/LandscapeVertexFactory.ush
@@ -53,21 +53,8 @@ float4 SectionLods;
 */
 float4 NeighborSectionLod[4];
 
-float4x4 LocalToWorldNoScaling;
-
-#if FEATURE_LEVEL >= FEATURE_LEVEL_SM4
-Texture2D HeightmapTexture;
-SamplerState HeightmapTextureSampler;
-#endif
-#if PIXELSHADER
-Texture2D NormalmapTexture;
-SamplerState NormalmapTextureSampler;
-#endif
-
 #if LANDSCAPE_XYOFFSET
-#define XYOFFSET_SCALE			(1.0f/256.f)
-Texture2D XYOffsetmapTexture;
-SamplerState XYOffsetmapTextureSampler;
+	#define XYOFFSET_SCALE			(1.0f/256.f)
 #endif
 
 #if ES2_PROFILE || ES3_1_PROFILE
@@ -84,6 +71,48 @@ struct FVertexFactoryInput
 #endif
 };
 
+// RHI_RAYTRACING
+#if COMPUTESHADER || RAYHITGROUPSHADER
+FVertexFactoryInput LoadVertexFactoryInput(uint TriangleIndex, int VertexIndex)
+{
+	FVertexFactoryInput Input;
+
+	uint SubsectionSizeQuadsCurrentLOD = LodValues.z;
+	uint SubsectionSizeVertsCurrentLOD = LodValues.z + 1;
+
+#if RAYHITGROUPSHADER
+	#if 1
+		FTriangleBaseAttributes Triangle = LoadTriangleBaseAttributes(TriangleIndex);
+		uint VertexId = Triangle.Indices[VertexIndex];
+	#else
+		uint QuadId = TriangleIndex / 2;
+		uint QuadX = QuadId % SubsectionSizeQuadsCurrentLOD;
+		uint QuadY = QuadId / SubsectionSizeQuadsCurrentLOD;
+		uint i00 = (QuadX + 0) + (QuadY + 0) * SubsectionSizeVertsCurrentLOD;
+		uint i10 = (QuadX + 1) + (QuadY + 0) * SubsectionSizeVertsCurrentLOD;
+		uint i11 = (QuadX + 1) + (QuadY + 1) * SubsectionSizeVertsCurrentLOD;
+		uint i01 = (QuadX + 0) + (QuadY + 1) * SubsectionSizeVertsCurrentLOD;
+		uint Indices[6] = {i00, i11, i10, i00, i01, i10};
+		uint VertexId = Indices[TriangleIndex % 2 * 3 + VertexIndex];
+	#endif
+#elif COMPUTESHADER
+	uint VertexId = TriangleIndex * 3 + VertexIndex;
+#endif
+
+	Input.Position.x = VertexId % SubsectionSizeVertsCurrentLOD;
+	Input.Position.y = VertexId / SubsectionSizeVertsCurrentLOD;
+	Input.Position.z = LandscapeMVF.SubXY.x;
+	Input.Position.w = LandscapeMVF.SubXY.y;
+
+	return Input;
+}
+
+uint GetNumRayTracingDynamicMeshVerticesIndirect()
+{
+	return 0;
+}
+#endif
+
 struct FVertexFactoryInterpolantsVSToPS
 {
 #if USING_TESSELLATION
@@ -237,11 +266,11 @@ float3x3 CalcTangentBasisFromWorldNormal(float3 Normal)
 float3x3 VertexFactoryGetPerPixelTangentBasis(FVertexFactoryInterpolantsVSToPS Interpolants)
 {
 	float3x3 Result;
-#if PIXELSHADER
+#if PIXELSHADER || RAYHITGROUPSHADER || COMPUTESHADER
 #if FEATURE_LEVEL >= FEATURE_LEVEL_SM4
-	float4 SampleValue = Texture2DSample(NormalmapTexture, NormalmapTextureSampler, Interpolants.WeightHeightMapTexCoord.zw);
+	float4 SampleValue = Texture2DSample(LandscapeParameters.NormalmapTexture, LandscapeParameters.NormalmapTextureSampler, Interpolants.WeightHeightMapTexCoord.zw);
 #else
-	float4 SampleValue = Texture2DSample(NormalmapTexture, NormalmapTextureSampler, Interpolants.WeightMapTexCoord);
+	float4 SampleValue = Texture2DSample(LandscapeParameters.NormalmapTexture, LandscapeParameters.NormalmapTextureSampler, Interpolants.WeightMapTexCoord);
 #endif
 #if WEBGL
 	float2 SampleNormal = float2(SampleValue.b, SampleValue.a) * float2(2.0,2.0) - float2(1.0,1.0);
@@ -386,7 +415,7 @@ FMaterialPixelParameters GetMaterialPixelParameters(FVertexFactoryInterpolantsVS
 
 	// Calculate LocalToTangent directly from normal map texture.
 	float3x3 TangentToLocal = VertexFactoryGetPerPixelTangentBasis(Interpolants);
-	Result.TangentToWorld = mul(TangentToLocal, (float3x3)LocalToWorldNoScaling);
+	Result.TangentToWorld = mul(TangentToLocal, (float3x3)LandscapeParameters.LocalToWorldNoScaling);
 	Result.UnMirrored = 1;
 
 	Result.VertexColor = 1;
@@ -397,7 +426,7 @@ FMaterialPixelParameters GetMaterialPixelParameters(FVertexFactoryInterpolantsVS
 
 #if (ES2_PROFILE || ES3_1_PROFILE) && PIXELSHADER
 
-	Result.LayerWeights.xy = Texture2DSample(NormalmapTexture, NormalmapTextureSampler, Interpolants.WeightMapTexCoord).rg;
+	Result.LayerWeights.xy = Texture2DSample(LandscapeParameters.NormalmapTexture, LandscapeParameters.NormalmapTextureSampler, Interpolants.WeightMapTexCoord).rg;
 	Result.LayerWeights.z = saturate(BlendableLayerMask.z - dot(Result.LayerWeights.xy, BlendableLayerMask.xy));
 	Result.LayerWeights.w = 0; // Total weight need to be 1
 #endif
@@ -647,20 +676,20 @@ FVertexFactoryIntermediates GetVertexFactoryIntermediates(FVertexFactoryInput In
 	// Get the height and normal XY for current and next LOD out of the textures
 #if FEATURE_LEVEL >= FEATURE_LEVEL_SM4
 	float2 SampleCoords = InputPositionLODAdjusted * LandscapeParameters.HeightmapUVScaleBias.xy + LandscapeParameters.HeightmapUVScaleBias.zw + 0.5*LandscapeParameters.HeightmapUVScaleBias.xy + Intermediates.InputPosition.zw * LandscapeParameters.SubsectionOffsetParams.xy;
-	float4 SampleValue = Texture2DSampleLevel(HeightmapTexture, HeightmapTextureSampler, SampleCoords, LodValue-LodBias.z);
+	float4 SampleValue = Texture2DSampleLevel(LandscapeParameters.HeightmapTexture, LandscapeParameters.HeightmapTextureSampler, SampleCoords, LodValue-LodBias.z);
 	float Height = DecodePackedHeight(SampleValue.xy);
 
 	float2 SampleCoordsNextLOD = InputPositionNextLOD * LandscapeParameters.HeightmapUVScaleBias.xy + LandscapeParameters.HeightmapUVScaleBias.zw + 0.5*LandscapeParameters.HeightmapUVScaleBias.xy + Intermediates.InputPosition.zw * LandscapeParameters.SubsectionOffsetParams.xy;
-	float4 SampleValueNextLOD = Texture2DSampleLevel(HeightmapTexture, HeightmapTextureSampler, SampleCoordsNextLOD, LodValue+1-LodBias.z );
+	float4 SampleValueNextLOD = Texture2DSampleLevel(LandscapeParameters.HeightmapTexture, LandscapeParameters.HeightmapTextureSampler, SampleCoordsNextLOD, LodValue+1-LodBias.z );
 	float HeightNextLOD = DecodePackedHeight(SampleValueNextLOD.xy);
 #endif
 
 
 #if LANDSCAPE_XYOFFSET // FEATURE_LEVEL >= FEATURE_LEVEL_SM4 only
 	float2 SampleCoords2 = float2(InputPositionLODAdjusted * LandscapeParameters.WeightmapUVScaleBias.xy + LandscapeParameters.WeightmapUVScaleBias.zw + Intermediates.InputPosition.zw * LandscapeParameters.SubsectionOffsetParams.zz);
-	float4 OffsetValue = Texture2DSampleLevel( XYOffsetmapTexture, XYOffsetmapTextureSampler, SampleCoords2, LodValue-LodBias.w );
+	float4 OffsetValue = Texture2DSampleLevel( LandscapeParameters.XYOffsetmapTexture, LandscapeParameters.XYOffsetmapTextureSampler, SampleCoords2, LodValue-LodBias.w );
 	float2 SampleCoordsNextLOD2 = float2(InputPositionNextLOD * LandscapeParameters.WeightmapUVScaleBias.xy + LandscapeParameters.WeightmapUVScaleBias.zw + Intermediates.InputPosition.zw * LandscapeParameters.SubsectionOffsetParams.zz);
-	float4 OffsetValueNextLOD = Texture2DSampleLevel( XYOffsetmapTexture, XYOffsetmapTextureSampler, SampleCoordsNextLOD2, LodValue+1-LodBias.w );
+	float4 OffsetValueNextLOD = Texture2DSampleLevel( LandscapeParameters.XYOffsetmapTexture, LandscapeParameters.XYOffsetmapTextureSampler, SampleCoordsNextLOD2, LodValue+1-LodBias.w );
 	float2 XYOffset = float2(((OffsetValue.r * 255.0 * 256.0 + OffsetValue.g * 255.0) - 32768.0) * XYOFFSET_SCALE, ((OffsetValue.b * 255.0 * 256.0 + OffsetValue.a * 255.0) - 32768.0) * XYOFFSET_SCALE );
 	float2 XYOffsetNextLOD = float2(((OffsetValueNextLOD.r * 255.0 * 256.0 + OffsetValueNextLOD.g * 255.0) - 32768.0) * XYOFFSET_SCALE, ((OffsetValueNextLOD.b * 255.0 * 256.0 + OffsetValueNextLOD.a * 255.0) - 32768.0) * XYOFFSET_SCALE );
 
@@ -702,7 +731,7 @@ float2 LandscapeVertexFactorySampleHeight(FVertexFactoryInput Input, FVertexFact
 
 	float2 InputPositionLODAdjusted = ActualLODCoordsInt / CoordTranslate_x;
 	float2 SampleCoords = InputPositionLODAdjusted * LandscapeParameters.HeightmapUVScaleBias.xy + LandscapeParameters.HeightmapUVScaleBias.zw + 0.5*LandscapeParameters.HeightmapUVScaleBias.xy + Intermediates.InputPosition.zw * LandscapeParameters.SubsectionOffsetParams.xy;
-	float4 SampleValue = Texture2DSampleLevel(HeightmapTexture, HeightmapTextureSampler, SampleCoords, LodValue-LodBias.z);
+	float4 SampleValue = Texture2DSampleLevel(LandscapeParameters.HeightmapTexture, LandscapeParameters.HeightmapTextureSampler, SampleCoords, LodValue-LodBias.z);
 	return SampleValue.rg;
 }
 #endif
diff --git a/Engine/Shaders/Private/RayTracing/RayTracingDynamicMesh.usf b/Engine/Shaders/Private/RayTracing/RayTracingDynamicMesh.usf
--- a/Engine/Shaders/Private/RayTracing/RayTracingDynamicMesh.usf
+++ b/Engine/Shaders/Private/RayTracing/RayTracingDynamicMesh.usf
@@ -63,7 +63,13 @@ void RayTracingDynamicGeometryConverterCS(uint3 DispatchThreadId : SV_DispatchTh
 	{
 		FVertexFactoryInput Input = LoadVertexFactoryInput(VertexIndex / 3, VertexIndex % 3);
 		FVertexFactoryIntermediates VFIntermediates = GetVertexFactoryIntermediates(Input);
-		float3 WorldPosition = (VertexFactoryGetWorldPosition(Input, VFIntermediates) - float4(ResolvedView.PreViewTranslation, 0.0f)).xyz;
+		float4 WorldPosition = VertexFactoryGetWorldPosition(Input, VFIntermediates);
+		float3x3 TangentToLocal = VertexFactoryGetTangentToLocal(Input, VFIntermediates);
+		FMaterialVertexParameters VertexParameters = GetMaterialVertexParameters(Input, VFIntermediates, WorldPosition.xyz, TangentToLocal);
+
+		WorldPosition.xyz -= ResolvedView.PreViewTranslation;
+		WorldPosition.xyz += GetMaterialWorldPositionOffset(VertexParameters);
+
 		RWVertexPositions[VertexIndex * 3 + 0] = WorldPosition.x;
 		RWVertexPositions[VertexIndex * 3 + 1] = WorldPosition.y;
 		RWVertexPositions[VertexIndex * 3 + 2] = WorldPosition.z;
diff --git a/Engine/Source/Runtime/Engine/Public/SceneManagement.h b/Engine/Source/Runtime/Engine/Public/SceneManagement.h
--- a/Engine/Source/Runtime/Engine/Public/SceneManagement.h
+++ b/Engine/Source/Runtime/Engine/Public/SceneManagement.h
@@ -2248,6 +2248,21 @@ public:
 	}
 };
 
+struct FRayTracingDynamicGeometryUpdateParams
+{
+	FMeshBatch MeshBatch;
+
+	bool bUsingIndirectDraw;
+	// When bUsingIndirectDraw == false, NumVertices == the actual number of vertices to process
+	// When bUsingIndirectDraw == true, it is the maximum possible vertices that GPU can emit
+	uint32 NumVertices;
+	uint32 VertexBufferSize;
+	uint32 NumTriangles;
+
+	FRayTracingGeometry* Geometry;
+	FRWBuffer* Buffer;
+};
+
 struct FRayTracingMaterialGatheringContext
 {
 	const class FScene* Scene;
@@ -2255,7 +2270,7 @@ struct FRayTracingMaterialGatheringContext
 	const FSceneViewFamily& ReferenceViewFamily;
 
 	FRayTracingMeshResourceCollector& RayTracingMeshResourceCollector;
-	class FRayTracingDynamicGeometryCollection& DynamicRayTracingGeometriesToUpdate;
+	TArray<FRayTracingDynamicGeometryUpdateParams> DynamicRayTracingGeometriesToUpdate;
 };
 #endif
 
diff --git a/Engine/Source/Runtime/Landscape/Landscape.Build.cs b/Engine/Source/Runtime/Landscape/Landscape.Build.cs
--- a/Engine/Source/Runtime/Landscape/Landscape.Build.cs
+++ b/Engine/Source/Runtime/Landscape/Landscape.Build.cs
@@ -10,8 +10,9 @@ public class Landscape : ModuleRules
 		PrivateIncludePaths.AddRange(
 			new string[] {
 				"Runtime/Engine/Private", // for Engine/Private/Collision/PhysXCollision.h
-				"Runtime/Landscape/Private"
-			}
+				"Runtime/Landscape/Private",
+                "../Shaders/Shared"
+            }
 		);
 
 		PrivateIncludePathModuleNames.AddRange(
diff --git a/Engine/Source/Runtime/Landscape/Private/LandscapeRender.cpp b/Engine/Source/Runtime/Landscape/Private/LandscapeRender.cpp
--- a/Engine/Source/Runtime/Landscape/Private/LandscapeRender.cpp
+++ b/Engine/Source/Runtime/Landscape/Private/LandscapeRender.cpp
@@ -38,6 +38,7 @@ LandscapeRender.cpp: New terrain rendering
 #include "HAL/LowLevelMemTracker.h"
 #include "MeshMaterialShader.h"
 #include "VT/RuntimeVirtualTexture.h"
+#include "RayTracingInstance.h"
 
 IMPLEMENT_GLOBAL_SHADER_PARAMETER_STRUCT(FLandscapeUniformShaderParameters, "LandscapeParameters");
 
@@ -1023,6 +1024,34 @@ void FLandscapeComponentSceneProxy::CreateRenderThreadResources()
 		}
 	}
 #endif
+
+#if RHI_RAYTRACING
+	for (int32 SubY = 0; SubY < NumSubsections; SubY++)
+	{
+		for (int32 SubX = 0; SubX < NumSubsections; SubX++)
+		{
+			const int8 SubSectionIdx = SubX + SubY * NumSubsections;
+
+			int32 LodSubsectionSizeVerts = (SubsectionSizeVerts >> 0);
+			uint32 NumPrimitives = FMath::Square((LodSubsectionSizeVerts - 1)) * 2;
+
+			FRayTracingGeometryInitializer Initializer;
+			FRHIResourceCreateInfo CreateInfo;
+			Initializer.PositionVertexBuffer = RHICreateVertexBuffer(sizeof(FVector4) * NumPrimitives * 3, BUF_UnorderedAccess | BUF_ShaderResource, CreateInfo);
+			Initializer.IndexBuffer = nullptr;
+			Initializer.BaseVertexIndex = 0;
+			Initializer.VertexBufferStride = sizeof(FVector);
+			Initializer.TotalPrimitiveCount = NumPrimitives;
+			Initializer.VertexBufferElementType = VET_Float3;
+			Initializer.GeometryType = RTGT_Triangles;
+			Initializer.bFastBuild = true;
+			Initializer.bAllowUpdate = true;
+
+			SectionRayTracingStates[SubSectionIdx].Geometry.SetInitializer(Initializer);
+			SectionRayTracingStates[SubSectionIdx].Geometry.InitResource();
+		}
+	}
+#endif
 }
 
 void FLandscapeComponentSceneProxy::OnLevelAddedToWorld()
@@ -1046,6 +1075,18 @@ FLandscapeComponentSceneProxy::~FLandscapeComponentSceneProxy()
 		}
 		SharedBuffers = nullptr;
 	}
+
+#if RHI_RAYTRACING
+	for (int32 SubY = 0; SubY < NumSubsections; SubY++)
+	{
+		for (int32 SubX = 0; SubX < NumSubsections; SubX++)
+		{
+			const int8 SubSectionIdx = SubX + SubY * NumSubsections;
+			SectionRayTracingStates[SubSectionIdx].Geometry.ReleaseResource();
+			SectionRayTracingStates[SubSectionIdx].RayTracingDynamicVertexBuffer.Release();
+		}
+	}
+#endif
 }
 
 bool FLandscapeComponentSceneProxy::CanBeOccluded() const
@@ -1311,6 +1352,39 @@ void FLandscapeComponentSceneProxy::OnTransformChanged()
 		0
 	);
 
+	if (HeightmapTexture)
+	{
+		LandscapeParams.HeightmapTexture = HeightmapTexture->TextureReference.TextureReferenceRHI;
+		LandscapeParams.HeightmapTextureSampler = TStaticSamplerState<SF_Point>::GetRHI();
+	}
+	else
+	{
+		LandscapeParams.HeightmapTexture = GBlackTexture->TextureRHI;
+		LandscapeParams.HeightmapTextureSampler = GBlackTexture->SamplerStateRHI;
+	}
+
+	if (XYOffsetmapTexture)
+	{
+		LandscapeParams.XYOffsetmapTexture = XYOffsetmapTexture->TextureReference.TextureReferenceRHI;
+		LandscapeParams.XYOffsetmapTextureSampler = TStaticSamplerState<SF_Point>::GetRHI();
+	}
+	else
+	{
+		LandscapeParams.XYOffsetmapTexture = GBlackTexture->TextureRHI;
+		LandscapeParams.XYOffsetmapTextureSampler = GBlackTexture->SamplerStateRHI;
+	}
+
+	if (NormalmapTexture)
+	{
+		LandscapeParams.NormalmapTexture = NormalmapTexture->TextureReference.TextureReferenceRHI;
+		LandscapeParams.NormalmapTextureSampler = NormalmapTexture->Resource->SamplerStateRHI;
+	}
+	else
+	{
+		LandscapeParams.NormalmapTexture = GBlackTexture->TextureRHI;
+		LandscapeParams.NormalmapTextureSampler = GBlackTexture->SamplerStateRHI;
+	}
+
 	LandscapeUniformShaderParameters.SetContents(LandscapeParams);
 }
 
@@ -2218,7 +2292,6 @@ float FLandscapeComponentSceneProxy::GetNeighborLOD(const FSceneView& InView, fl
 	// Handle subsection
 	if (InSubSectionX != INDEX_NONE && InSubSectionY != INDEX_NONE)
 	{
-
 		const SubSectionData& Data = SubSectionValues[CurrentSubSectionIndex][InNeighborIndex];
 		DesiredSubSectionX = InSubSectionX + Data.SubSectionOffsetX;
 		DesiredSubSectionY = InSubSectionY + Data.SubSectionOffsetY;
@@ -2238,6 +2311,8 @@ float FLandscapeComponentSceneProxy::GetNeighborLOD(const FSceneView& InView, fl
 		}
 		else
 		{
+			DesiredSubSectionX = InSubSectionX;
+			DesiredSubSectionY = InSubSectionY;
 			DesiredSubSectionIndex = CurrentSubSectionIndex;
 		}
 	}
@@ -2260,44 +2335,54 @@ float FLandscapeComponentSceneProxy::GetNeighborLOD(const FSceneView& InView, fl
 
 	if (ComputeNeighborCustomDataLOD)
 	{
-		FVector LandscapeComponentOrigin = LandscapeComponent->Bounds.Origin;
-		float LandscapeComponentMaxExtends = ComponentMaxExtend;
+		FBoxSphereBounds NeighborBounds = LandscapeComponent->Bounds;
+		float NeighborMaxExtends = ComponentMaxExtend;
+		FLandscapeComponentSceneProxy* NeighborSceneProxy = nullptr;
 
 		if (Neighbor != nullptr)
 		{
-			const ULandscapeComponent* NeighborComponent = Neighbor->GetLandscapeComponent();
-
-			if (Neighbor->GetLandscapeComponent() != nullptr)
-			{
-				LandscapeComponentOrigin = NeighborComponent->Bounds.Origin;
-				LandscapeComponentMaxExtends = NeighborComponent->SubsectionSizeQuads * FMath::Max(NeighborComponent->GetComponentTransform().GetScale3D().X, NeighborComponent->GetComponentTransform().GetScale3D().Y);
-			}
+			NeighborSceneProxy = (FLandscapeComponentSceneProxy*)Neighbor;
+			NeighborBounds = NeighborSceneProxy->GetBounds();
+			NeighborMaxExtends = NeighborSceneProxy->SubsectionSizeQuads * FMath::Max(NeighborSceneProxy->GetLocalToWorld().GetScaleVector().X, NeighborSceneProxy->GetLocalToWorld().GetScaleVector().Y);
 		}
 
 		if (NumSubsections > 1)
 		{
-			float SubSectionMaxExtend = LandscapeComponentMaxExtends / 2.0f;
-			FVector ComponentTopLeftCorner = LandscapeComponentOrigin - FVector(SubSectionMaxExtend, SubSectionMaxExtend, 0.0f);
+			float NeighborSubsectionMaxExtends = NeighborMaxExtends / 2.0f;
+			FVector ComponentTopLeftCorner = NeighborBounds.Origin - FVector(NeighborSubsectionMaxExtends, NeighborSubsectionMaxExtends, 0.0f);
 
-			FVector SubSectionOrigin = ComponentTopLeftCorner + FVector(LandscapeComponentMaxExtends * DesiredSubSectionX, LandscapeComponentMaxExtends * DesiredSubSectionY, 0.0f);
-			float MeshBatchScreenSizeSquared = GetComponentScreenSize(&InView, SubSectionOrigin, SubSectionMaxExtend, LandscapeComponent->Bounds.SphereRadius / 2.0f);
+			FVector SubSectionOrigin = ComponentTopLeftCorner + FVector(NeighborMaxExtends * DesiredSubSectionX, NeighborMaxExtends * DesiredSubSectionY, 0.0f);
+			float MeshBatchScreenSizeSquared = GetComponentScreenSize(&InView, SubSectionOrigin, NeighborSubsectionMaxExtends, NeighborBounds.SphereRadius / 2.0f);
 
 			FViewCustomDataLOD NeighborLODData;
-			CalculateLODFromScreenSize(InView, MeshBatchScreenSizeSquared, InView.LODDistanceFactor, DesiredSubSectionIndex, NeighborLODData);
+			if (NeighborSceneProxy != nullptr)
+			{
+				// Needs to pull some data like per-component LOD bias from neighbor scene proxy
+				NeighborSceneProxy->CalculateLODFromScreenSize(InView, MeshBatchScreenSizeSquared, InView.LODDistanceFactor, DesiredSubSectionIndex, NeighborLODData);
+			}
+			else
+			{
+				CalculateLODFromScreenSize(InView, MeshBatchScreenSizeSquared, InView.LODDistanceFactor, DesiredSubSectionIndex, NeighborLODData);
+			}
+
 			const FViewCustomDataSubSectionLOD& SubSectionData = NeighborLODData.SubSections[DesiredSubSectionIndex];
 			check(SubSectionData.fBatchElementCurrentLOD != -1.0f);
 
-			if (SubSectionData.fBatchElementCurrentLOD > InBatchElementCurrentLOD)
-			{
-				NeighborLOD = SubSectionData.fBatchElementCurrentLOD;
-			}
+			NeighborLOD = FMath::Max(SubSectionData.fBatchElementCurrentLOD, InBatchElementCurrentLOD);
 		}
 		else
 		{
-			float MeshBatchScreenSizeSquared = GetComponentScreenSize(&InView, LandscapeComponentOrigin, LandscapeComponentMaxExtends, LandscapeComponent->Bounds.SphereRadius);
+			float MeshBatchScreenSizeSquared = GetComponentScreenSize(&InView, NeighborBounds.Origin, NeighborMaxExtends, NeighborBounds.SphereRadius);
 
 			FViewCustomDataLOD NeighborLODData;
-			CalculateLODFromScreenSize(InView, MeshBatchScreenSizeSquared, InView.LODDistanceFactor, 0, NeighborLODData);
+			if (NeighborSceneProxy != nullptr)
+			{
+				NeighborSceneProxy->CalculateLODFromScreenSize(InView, MeshBatchScreenSizeSquared, InView.LODDistanceFactor, DesiredSubSectionIndex, NeighborLODData);
+			}
+			else
+			{
+				CalculateLODFromScreenSize(InView, MeshBatchScreenSizeSquared, InView.LODDistanceFactor, DesiredSubSectionIndex, NeighborLODData);
+			}
 
 			FViewCustomDataSubSectionLOD& SubSectionLODData = NeighborLODData.SubSections[0];
 			check(SubSectionLODData.fBatchElementCurrentLOD != -1.0f);
@@ -2973,6 +3058,111 @@ void FLandscapeComponentSceneProxy::GetDynamicMeshElements(const TArray<const FS
 	INC_DWORD_STAT_BY(STAT_LandscapeTriangles, NumTriangles * NumPasses);
 }
 
+#if RHI_RAYTRACING
+void FLandscapeComponentSceneProxy::GetDynamicRayTracingInstances(FRayTracingMaterialGatheringContext& Context, TArray<FRayTracingInstance>& OutRayTracingInstances)
+{
+	FMemStackBase& PrimitiveCustomDataMemStack = FMemStack::Get();
+
+	float MeshScreenSizeSquared = 0;
+	int32 ForcedLODLevel = (Context.ReferenceView->Family->EngineShowFlags.LOD) ? GetCVarForceLOD() : 0;
+
+	FLODMask LODToRender = GetCustomLOD(*Context.ReferenceView, Context.ReferenceView->LODDistanceFactor, ForcedLODLevel, MeshScreenSizeSquared);
+	FViewCustomDataLOD* InPrimitiveCustomData = (FViewCustomDataLOD*)InitViewCustomData(*Context.ReferenceView, Context.ReferenceView->LODDistanceFactor, PrimitiveCustomDataMemStack, false, false, &LODToRender, MeshScreenSizeSquared);
+	PostInitViewCustomData(*Context.ReferenceView, InPrimitiveCustomData);
+
+	FLandscapeElementParamArray& ParameterArray = Context.RayTracingMeshResourceCollector.AllocateOneFrameResource<FLandscapeElementParamArray>();
+	ParameterArray.ElementParams.AddDefaulted(NumSubsections * NumSubsections);
+
+	if (AvailableMaterials.Num() == 0)
+	{
+		return;
+	}
+
+	const int8 CurrentLODIndex = InPrimitiveCustomData->SubSections[0].BatchElementCurrentLOD;
+	int8 MaterialIndex = LODIndexToMaterialIndex.IsValidIndex(CurrentLODIndex) ? LODIndexToMaterialIndex[CurrentLODIndex] : INDEX_NONE;
+	UMaterialInterface* SelectedMaterial = MaterialIndex != INDEX_NONE ? AvailableMaterials[MaterialIndex] : nullptr;
+
+	// this is really not normal that we have no material at this point, so do not continue
+	if (SelectedMaterial == nullptr)
+	{
+		return;
+	}
+
+	FMeshBatch BaseMeshBatch;
+	BaseMeshBatch.VertexFactory = VertexFactory;
+	BaseMeshBatch.MaterialRenderProxy = SelectedMaterial->GetRenderProxy();
+	BaseMeshBatch.LCI = ComponentLightInfo.Get();
+	BaseMeshBatch.CastShadow = true;
+	BaseMeshBatch.CastRayTracedShadow = true;
+	BaseMeshBatch.bUseForMaterial = true;
+	BaseMeshBatch.SegmentIndex = 0;
+
+	BaseMeshBatch.Elements.Empty();
+
+	for (int32 SubY = 0; SubY < NumSubsections; SubY++)
+	{
+		for (int32 SubX = 0; SubX < NumSubsections; SubX++)
+		{
+			const int8 SubSectionIdx = SubX + SubY * NumSubsections;
+			const int8 CurrentLOD = InPrimitiveCustomData->SubSections[SubSectionIdx].BatchElementCurrentLOD;
+
+			FMeshBatch MeshBatch = BaseMeshBatch;
+
+			FMeshBatchElement BatchElement;
+			FLandscapeBatchElementParams& BatchElementParams = ParameterArray.ElementParams[SubSectionIdx];
+
+			BatchElementParams.LocalToWorldNoScalingPtr = &LocalToWorldNoScaling;
+			BatchElementParams.LandscapeUniformShaderParametersResource = &LandscapeUniformShaderParameters;
+			BatchElementParams.SceneProxy = this;
+			BatchElementParams.SubX = SubX;
+			BatchElementParams.SubY = SubY;
+			BatchElementParams.CurrentLOD = CurrentLOD;
+			BatchElement.UserData = &BatchElementParams;
+			BatchElement.PrimitiveUniformBuffer = GetUniformBuffer();
+
+			int32 LodSubsectionSizeVerts = SubsectionSizeVerts >> CurrentLOD;
+			uint32 NumPrimitives = FMath::Square(LodSubsectionSizeVerts - 1) * 2;
+
+			BatchElement.IndexBuffer = SharedBuffers->ZeroOffsetIndexBuffers[CurrentLOD];
+			BatchElement.FirstIndex = 0;
+			BatchElement.NumPrimitives = NumPrimitives;
+
+			MeshBatch.Elements.Add(BatchElement);
+
+			SectionRayTracingStates[SubSectionIdx].Geometry.Initializer.IndexBuffer = BatchElement.IndexBuffer->IndexBufferRHI;
+
+			{
+				FLandscapeVertexFactoryMVFParameters UniformBufferParams;
+
+				UniformBufferParams.SubXY = FIntPoint(SubX, SubY);
+
+				BatchElementParams.LandscapeVertexFactoryMVFUniformBuffer = FLandscapeVertexFactoryMVFUniformBufferRef::CreateUniformBufferImmediate(UniformBufferParams, UniformBuffer_SingleFrame);
+			}
+
+			Context.DynamicRayTracingGeometriesToUpdate.Add(
+				FRayTracingDynamicGeometryUpdateParams
+				{
+					MeshBatch,
+					false,
+					(uint32)FMath::Square(LodSubsectionSizeVerts),
+					(uint32)FMath::Square(LodSubsectionSizeVerts) * sizeof(FVector),
+					(uint32)FMath::Square(LodSubsectionSizeVerts - 1) * 2,
+					&SectionRayTracingStates[SubSectionIdx].Geometry,
+					&SectionRayTracingStates[SubSectionIdx].RayTracingDynamicVertexBuffer
+				}
+			);
+
+			FRayTracingInstance RayTracingInstance;
+			RayTracingInstance.Geometry = &SectionRayTracingStates[SubSectionIdx].Geometry;
+			RayTracingInstance.InstanceTransforms.Add(FMatrix::Identity);
+			RayTracingInstance.Materials.Add(MeshBatch);
+			RayTracingInstance.BuildInstanceMaskAndFlags();
+			OutRayTracingInstances.Add(RayTracingInstance);
+		}
+	}
+}
+#endif
+
 bool FLandscapeComponentSceneProxy::CollectOccluderElements(FOccluderElementsCollector& Collector) const
 {
 	// TODO: implement
@@ -3234,6 +3424,37 @@ void FLandscapeSharedBuffers::CreateIndexBuffers(ERHIFeatureLevel::Type InFeatur
 		}
 
 		IndexBuffers[Mip] = IndexBuffer;
+
+#if RHI_RAYTRACING
+		if (IsRayTracingEnabled())
+		{
+			TArray<INDEX_TYPE> ZeroOffsetIndices;
+
+			for (int32 y = 0; y < LodSubsectionSizeQuads; y++)
+			{
+				for (int32 x = 0; x < LodSubsectionSizeQuads; x++)
+				{
+					INDEX_TYPE i00 = (x + 0) + (y + 0) * (SubsectionSizeVerts >> Mip);
+					INDEX_TYPE i10 = (x + 1) + (y + 0) * (SubsectionSizeVerts >> Mip);
+					INDEX_TYPE i11 = (x + 1) + (y + 1) * (SubsectionSizeVerts >> Mip);
+					INDEX_TYPE i01 = (x + 0) + (y + 1) * (SubsectionSizeVerts >> Mip);
+
+					ZeroOffsetIndices.Add(i00);
+					ZeroOffsetIndices.Add(i11);
+					ZeroOffsetIndices.Add(i10);
+
+					ZeroOffsetIndices.Add(i00);
+					ZeroOffsetIndices.Add(i01);
+					ZeroOffsetIndices.Add(i11);
+				}
+			}
+
+			FRawStaticIndexBuffer16or32<INDEX_TYPE>* ZeroOffsetIndexBuffer = new FRawStaticIndexBuffer16or32<INDEX_TYPE>(false);
+			ZeroOffsetIndexBuffer->AssignNewBuffer(ZeroOffsetIndices);
+			ZeroOffsetIndexBuffer->InitResource();
+			ZeroOffsetIndexBuffers[Mip] = ZeroOffsetIndexBuffer;
+		}
+#endif
 	}
 }
 
@@ -3346,6 +3567,13 @@ FLandscapeSharedBuffers::FLandscapeSharedBuffers(const int32 InSharedBuffersKey,
 	FMemory::Memzero(IndexBuffers, sizeof(FIndexBuffer*)* NumIndexBuffers);
 	IndexRanges = new FLandscapeIndexRanges[NumIndexBuffers]();
 
+#if RHI_RAYTRACING
+	if (IsRayTracingEnabled())
+	{
+		ZeroOffsetIndexBuffers.AddZeroed(NumIndexBuffers);
+	}
+#endif
+
 	// See if we need to use 16 or 32-bit index buffers
 	if (NumVertices > 65535)
 	{
@@ -3384,6 +3612,18 @@ FLandscapeSharedBuffers::~FLandscapeSharedBuffers()
 	delete[] IndexBuffers;
 	delete[] IndexRanges;
 
+#if RHI_RAYTRACING
+	if (IsRayTracingEnabled())
+	{
+		while (ZeroOffsetIndexBuffers.Num() > 0)
+		{
+			FIndexBuffer* Buffer = ZeroOffsetIndexBuffers.Pop();
+			Buffer->ReleaseResource();
+			delete Buffer;
+		}
+	}
+#endif
+
 #if WITH_EDITOR
 	if (GrassIndexBuffer)
 	{
@@ -3506,15 +3746,11 @@ public:
 	*/
 	virtual void Bind(const FShaderParameterMap& ParameterMap) override
 	{
-		HeightmapTextureParameter.Bind(ParameterMap, TEXT("HeightmapTexture"));
-		HeightmapTextureParameterSampler.Bind(ParameterMap, TEXT("HeightmapTextureSampler"));
+		LodBiasParameter.Bind(ParameterMap, TEXT("LodBias"));
 		LodValuesParameter.Bind(ParameterMap, TEXT("LodValues"));
 		LodTessellationParameter.Bind(ParameterMap, TEXT("LodTessellationParams"));
-		NeighborSectionLodParameter.Bind(ParameterMap, TEXT("NeighborSectionLod"));
-		LodBiasParameter.Bind(ParameterMap, TEXT("LodBias"));
 		SectionLodsParameter.Bind(ParameterMap, TEXT("SectionLods"));
-		XYOffsetTextureParameter.Bind(ParameterMap, TEXT("XYOffsetmapTexture"));
-		XYOffsetTextureParameterSampler.Bind(ParameterMap, TEXT("XYOffsetmapTextureSampler"));
+		NeighborSectionLodParameter.Bind(ParameterMap, TEXT("NeighborSectionLod"));
 	}
 
 	/**
@@ -3523,15 +3759,11 @@ public:
 	*/
 	virtual void Serialize(FArchive& Ar) override
 	{
-		Ar << HeightmapTextureParameter;
-		Ar << HeightmapTextureParameterSampler;
 		Ar << LodValuesParameter;
 		Ar << LodTessellationParameter;
 		Ar << NeighborSectionLodParameter;
 		Ar << LodBiasParameter;
 		Ar << SectionLodsParameter;
-		Ar << XYOffsetTextureParameter;
-		Ar << XYOffsetTextureParameterSampler;
 	}
 
 	virtual void GetElementShaderBindings(
@@ -3555,10 +3787,12 @@ public:
 
 		ShaderBindings.Add(Shader->GetUniformBufferParameter<FLandscapeUniformShaderParameters>(), *BatchElementParams->LandscapeUniformShaderParametersResource);
 
-		if (HeightmapTextureParameter.IsBound())
+#if RHI_RAYTRACING
+		if (IsRayTracingEnabled())
 		{
-			ShaderBindings.AddTexture(HeightmapTextureParameter, HeightmapTextureParameterSampler, TStaticSamplerState<SF_Point>::GetRHI(), SceneProxy->HeightmapTexture->Resource->TextureRHI);
+			ShaderBindings.Add(Shader->GetUniformBufferParameter<FLandscapeVertexFactoryMVFParameters>(), BatchElementParams->LandscapeVertexFactoryMVFUniformBuffer);
 		}
+#endif
 
 		if (LodValuesParameter.IsBound())
 		{
@@ -3670,11 +3904,6 @@ public:
 				ShaderBindings.Add(NeighborSectionLodParameter, CurrentNeighborLOD);
 			}				
 		}
-
-		if (XYOffsetTextureParameter.IsBound() && SceneProxy->XYOffsetmapTexture)
-		{
-			ShaderBindings.AddTexture(XYOffsetTextureParameter, XYOffsetTextureParameterSampler, TStaticSamplerState<SF_Point>::GetRHI(), SceneProxy->XYOffsetmapTexture->Resource->TextureRHI);
-		}
 	}
 
 	virtual uint32 GetSize() const override
@@ -3688,10 +3917,6 @@ protected:
 	FShaderParameter NeighborSectionLodParameter;
 	FShaderParameter LodBiasParameter;
 	FShaderParameter SectionLodsParameter;
-	FShaderResourceParameter HeightmapTextureParameter;
-	FShaderResourceParameter HeightmapTextureParameterSampler;
-	FShaderResourceParameter XYOffsetTextureParameter;
-	FShaderResourceParameter XYOffsetTextureParameterSampler;
 	TShaderUniformBufferParameter<FLandscapeUniformShaderParameters> LandscapeShaderParameters;
 };
 
@@ -3723,11 +3948,6 @@ public:
 
 		ShaderBindings.Add(Shader->GetUniformBufferParameter<FLandscapeUniformShaderParameters>(), *BatchElementParams->LandscapeUniformShaderParametersResource);
 
-		if (HeightmapTextureParameter.IsBound())
-		{
-			ShaderBindings.AddTexture(HeightmapTextureParameter, HeightmapTextureParameterSampler, TStaticSamplerState<SF_Point>::GetRHI(), SceneProxy->HeightmapTexture->Resource->TextureRHI);
-		}
-
 		if (LodValuesParameter.IsBound())
 		{
 			ShaderBindings.Add(LodValuesParameter, SceneProxy->GetShaderLODValues(BatchElementParams->CurrentLOD));
@@ -3743,27 +3963,6 @@ public:
 //
 // FLandscapeVertexFactoryPixelShaderParameters
 //
-/**
-* Bind shader constants by name
-* @param	ParameterMap - mapping of named shader constants to indices
-*/
-void FLandscapeVertexFactoryPixelShaderParameters::Bind(const FShaderParameterMap& ParameterMap)
-{
-	NormalmapTextureParameter.Bind(ParameterMap, TEXT("NormalmapTexture"));
-	NormalmapTextureParameterSampler.Bind(ParameterMap, TEXT("NormalmapTextureSampler"));
-	LocalToWorldNoScalingParameter.Bind(ParameterMap, TEXT("LocalToWorldNoScaling"));
-}
-
-/**
-* Serialize shader params to an archive
-* @param	Ar - archive to serialize to
-*/
-void FLandscapeVertexFactoryPixelShaderParameters::Serialize(FArchive& Ar)
-{
-	Ar << NormalmapTextureParameter
-		<< NormalmapTextureParameterSampler
-		<< LocalToWorldNoScalingParameter;
-}
 
 void FLandscapeVertexFactoryPixelShaderParameters::GetElementShaderBindings(
 	const class FSceneInterface* Scene,
@@ -3781,26 +3980,15 @@ void FLandscapeVertexFactoryPixelShaderParameters::GetElementShaderBindings(
 
 	const FLandscapeBatchElementParams* BatchElementParams = (const FLandscapeBatchElementParams*)BatchElement.UserData;
 
-	if (LocalToWorldNoScalingParameter.IsBound())
-	{
-		ShaderBindings.Add(LocalToWorldNoScalingParameter, *BatchElementParams->LocalToWorldNoScalingPtr);
-	}
-
-	if (NormalmapTextureParameter.IsBound())
-	{
-		const FTexture* NormalmapTexture = BatchElementParams->SceneProxy->NormalmapTexture->Resource;
-		ShaderBindings.AddTexture(
-			NormalmapTextureParameter,
-			NormalmapTextureParameterSampler,
-			NormalmapTexture->SamplerStateRHI,
-			NormalmapTexture->TextureRHI);
-	}
+	ShaderBindings.Add(Shader->GetUniformBufferParameter<FLandscapeUniformShaderParameters>(), *BatchElementParams->LandscapeUniformShaderParametersResource);
 }
 
 //
 // FLandscapeVertexFactory
 //
 
+IMPLEMENT_GLOBAL_SHADER_PARAMETER_STRUCT(FLandscapeVertexFactoryMVFParameters, "LandscapeMVF");
+
 void FLandscapeVertexFactory::InitRHI()
 {
 	// list of declaration items
@@ -3823,6 +4011,10 @@ FVertexFactoryShaderParameters* FLandscapeVertexFactory::ConstructShaderParamete
 	switch (ShaderFrequency)
 	{
 	case SF_Vertex:
+#if RHI_RAYTRACING
+	case SF_Compute:
+	case SF_RayHitGroup:
+#endif
 		return new FLandscapeVertexFactoryVertexShaderParameters();
 		break;
 	case SF_Pixel:
diff --git a/Engine/Source/Runtime/Landscape/Public/LandscapeRender.h b/Engine/Source/Runtime/Landscape/Public/LandscapeRender.h
--- a/Engine/Source/Runtime/Landscape/Public/LandscapeRender.h
+++ b/Engine/Source/Runtime/Landscape/Public/LandscapeRender.h
@@ -95,11 +95,27 @@ SHADER_PARAMETER(FVector4, SubsectionSizeVertsLayerUVPan)
 SHADER_PARAMETER(FVector4, SubsectionOffsetParams)
 SHADER_PARAMETER(FVector4, LightmapSubsectionOffsetParams)
 SHADER_PARAMETER(FMatrix, LocalToWorldNoScaling)
+SHADER_PARAMETER_TEXTURE(Texture2D, HeightmapTexture)
+SHADER_PARAMETER_SAMPLER(SamplerState, HeightmapTextureSampler)
+SHADER_PARAMETER_TEXTURE(Texture2D, NormalmapTexture)
+SHADER_PARAMETER_SAMPLER(SamplerState, NormalmapTextureSampler)
+SHADER_PARAMETER_TEXTURE(Texture2D, XYOffsetmapTexture)
+SHADER_PARAMETER_SAMPLER(SamplerState, XYOffsetmapTextureSampler)
 END_GLOBAL_SHADER_PARAMETER_STRUCT()
 
+
+BEGIN_GLOBAL_SHADER_PARAMETER_STRUCT(FLandscapeVertexFactoryMVFParameters, LANDSCAPE_API)
+	SHADER_PARAMETER(FIntPoint, SubXY)
+END_GLOBAL_SHADER_PARAMETER_STRUCT()
+
+typedef TUniformBufferRef<FLandscapeVertexFactoryMVFParameters> FLandscapeVertexFactoryMVFUniformBufferRef;
+
 /* Data needed for the landscape vertex factory to set the render state for an individual batch element */
 struct FLandscapeBatchElementParams
 {
+#if RHI_RAYTRACING
+	FLandscapeVertexFactoryMVFUniformBufferRef LandscapeVertexFactoryMVFUniformBuffer;
+#endif
 	const TUniformBuffer<FLandscapeUniformShaderParameters>* LandscapeUniformShaderParametersResource;
 	const FMatrix* LocalToWorldNoScalingPtr;
 
@@ -120,17 +136,8 @@ public:
 class FLandscapeVertexFactoryPixelShaderParameters : public FVertexFactoryShaderParameters
 {
 public:
-	/**
-	* Bind shader constants by name
-	* @param	ParameterMap - mapping of named shader constants to indices
-	*/
-	virtual void Bind(const FShaderParameterMap& ParameterMap) override;
-
-	/**
-	* Serialize shader params to an archive
-	* @param	Ar - archive to serialize to
-	*/
-	virtual void Serialize(FArchive& Ar) override;
+	virtual void Bind(const class FShaderParameterMap& ParameterMap) {}
+	virtual void Serialize(FArchive& Ar) {}
 
 	virtual void GetElementShaderBindings(
 		const class FSceneInterface* Scene,
@@ -148,11 +155,6 @@ public:
 	{
 		return sizeof(*this);
 	}
-
-private:
-	FShaderResourceParameter NormalmapTextureParameter;
-	FShaderResourceParameter NormalmapTextureParameterSampler;
-	FShaderParameter LocalToWorldNoScalingParameter;
 };
 
 /** vertex factory for VTF-heightmap terrain  */
@@ -341,6 +343,10 @@ public:
 	TArray<int32, TInlineAllocator<8>> GrassIndexMipOffsets;
 #endif
 
+#if RHI_RAYTRACING
+	TArray<FIndexBuffer*> ZeroOffsetIndexBuffers;
+#endif
+
 	FLandscapeSharedBuffers(int32 SharedBuffersKey, int32 SubsectionSizeQuads, int32 NumSubsections, ERHIFeatureLevel::Type FeatureLevel, bool bRequiresAdjacencyInformation, int32 NumOcclusionVertices);
 
 	template <typename INDEX_TYPE>
@@ -541,6 +547,20 @@ public:
 		FVector4 LodTessellationParams;
 	};
 
+#if RHI_RAYTRACING
+	struct FLandscapeSectionRayTracingState
+	{
+		int8 CurrentLOD;
+		FRayTracingGeometry Geometry;
+		FRWBuffer RayTracingDynamicVertexBuffer;
+		FLandscapeVertexFactoryMVFUniformBufferRef UniformBuffer;
+
+		FLandscapeSectionRayTracingState() : CurrentLOD(-1) {}
+	};
+
+	TStaticArray<FLandscapeSectionRayTracingState, MAX_SUBSECTION_COUNT> SectionRayTracingStates;
+#endif
+
 protected:
 	int8						MaxLOD;		// Maximum LOD level, user override possible
 	bool						UseTessellationComponentScreenSizeFalloff:1;	// Tell if we should apply a Tessellation falloff
@@ -743,6 +763,11 @@ public:
 #if !(UE_BUILD_SHIPPING || UE_BUILD_TEST)
 	virtual int32 GetLightMapResolution() const override { return LightMapResolution; }
 #endif
+
+#if RHI_RAYTRACING
+	virtual void GetDynamicRayTracingInstances(FRayTracingMaterialGatheringContext& Context, TArray<FRayTracingInstance>& OutRayTracingInstances) override final;
+	virtual bool IsRayTracingRelevant() const override { return true; }
+#endif
 };
 
 class FLandscapeDebugMaterialRenderProxy : public FMaterialRenderProxy
diff --git a/Engine/Source/Runtime/Renderer/Private/DeferredShadingRenderer.cpp b/Engine/Source/Runtime/Renderer/Private/DeferredShadingRenderer.cpp
--- a/Engine/Source/Runtime/Renderer/Private/DeferredShadingRenderer.cpp
+++ b/Engine/Source/Runtime/Renderer/Private/DeferredShadingRenderer.cpp
@@ -572,8 +572,7 @@ bool FDeferredShadingSceneRenderer::GatherRayTracingWorldInstances(FRHICommandLi
 			Scene,
 			&ReferenceView,
 			ViewFamily,
-			*ReferenceView.RayTracingMeshResourceCollector,
-			*Scene->RayTracingDynamicGeometryCollection
+			*ReferenceView.RayTracingMeshResourceCollector
 		};
 
 		int32 BroadIndex = 0;
@@ -706,6 +705,20 @@ bool FDeferredShadingSceneRenderer::GatherRayTracingWorldInstances(FRHICommandLi
 				TArray<FRayTracingInstance> RayTracingInstances;
 				SceneProxy->GetDynamicRayTracingInstances(MaterialGatheringContext, RayTracingInstances);
 
+				{
+					for (auto DynamicRayTracingGeometryUpdate : MaterialGatheringContext.DynamicRayTracingGeometriesToUpdate)
+					{
+						Scene->GetRayTracingDynamicGeometryCollection()->AddDynamicMeshBatchForGeometryUpdate(
+							Scene,
+							&ReferenceView,
+							SceneProxy,
+							DynamicRayTracingGeometryUpdate
+						);
+					}
+
+					MaterialGatheringContext.DynamicRayTracingGeometriesToUpdate.Reset();
+				}
+
 				if (RayTracingInstances.Num() > 0)
 				{
 					for (FRayTracingInstance& Instance : RayTracingInstances)
diff --git a/Engine/Source/Runtime/Renderer/Private/MeshPassProcessor.cpp b/Engine/Source/Runtime/Renderer/Private/MeshPassProcessor.cpp
--- a/Engine/Source/Runtime/Renderer/Private/MeshPassProcessor.cpp
+++ b/Engine/Source/Runtime/Renderer/Private/MeshPassProcessor.cpp
@@ -304,13 +304,35 @@ void FMeshDrawShaderBindings::SetRayTracingShaderBindingsForHitGroup(
 
 	const TArray<FShaderLooseParameterBufferInfo>& LooseParameterBuffers = SingleShaderBindings.ParameterMapInfo.LooseParameterBuffers;	
 
-	const void* LooseParameterData = SingleShaderBindings.GetLooseDataStart();
-	uint32 LooseParameterDataSize = 0;
-	check(LooseParameterBuffers.Num() <= 1);
+	// This array must be alive until RHICmdList.SetRayTracingHitGroup() finishes
+	TArray<uint8> ReorderedLooseParameters;
+
 	if (LooseParameterBuffers.Num())
 	{
-		LooseParameterDataSize = LooseParameterBuffers[0].BufferSize;
-		check(LooseParameterBuffers[0].BufferIndex == 0);
+		check(LooseParameterBuffers.Num() <= 1);
+
+		const FShaderLooseParameterBufferInfo& LooseParameterBuffer = SingleShaderBindings.ParameterMapInfo.LooseParameterBuffers[0];
+		check(LooseParameterBuffer.BufferIndex == 0);
+
+		uint32 LooseParameterDataSize = 0;
+
+		for (int32 LooseParameterIndex = 0; LooseParameterIndex < LooseParameterBuffer.Parameters.Num(); LooseParameterIndex++)
+		{
+			FShaderParameterInfo LooseParameter = LooseParameterBuffer.Parameters[LooseParameterIndex];
+			LooseParameterDataSize = FMath::Max<uint32>(LooseParameterDataSize, LooseParameter.BaseIndex + LooseParameter.Size);
+		}
+
+		ReorderedLooseParameters.AddZeroed(LooseParameterDataSize);
+
+		const uint8* LooseDataOffset = SingleShaderBindings.GetLooseDataStart();
+
+		for (int32 LooseParameterIndex = 0; LooseParameterIndex < LooseParameterBuffer.Parameters.Num(); LooseParameterIndex++)
+		{
+			FShaderParameterInfo LooseParameter = LooseParameterBuffer.Parameters[LooseParameterIndex];
+			FMemory::Memcpy(ReorderedLooseParameters.GetData() + LooseParameter.BaseIndex, LooseDataOffset, LooseParameter.Size);
+
+			LooseDataOffset += LooseParameter.Size;
+		}
 	}
 
 	check(SegmentIndex < 0xFF);
@@ -318,7 +340,7 @@ void FMeshDrawShaderBindings::SetRayTracingShaderBindingsForHitGroup(
 	const uint32 UserData = 0; // UserData could be used to store material ID or any other kind of per-material constant. This can be retrieved in hit shaders via GetHitGroupUserData().
 	RHICmdList.SetRayTracingHitGroup(Scene, InstanceIndex, SegmentIndex, ShaderSlot, PipelineState, HitGroupIndex, 
 		NumUniformBuffersToSet, LocalUniformBuffers,
-		LooseParameterDataSize, LooseParameterData,
+		ReorderedLooseParameters.Num(), ReorderedLooseParameters.GetData(),
 		UserData);
 }
 #endif // RHI_RAYTRACING
diff --git a/Engine/Source/Runtime/Renderer/Private/RayTracing/RayTracingDynamicGeometry.cpp b/Engine/Source/Runtime/Renderer/Private/RayTracing/RayTracingDynamicGeometry.cpp
--- a/Engine/Source/Runtime/Renderer/Private/RayTracing/RayTracingDynamicGeometry.cpp
+++ b/Engine/Source/Runtime/Renderer/Private/RayTracing/RayTracingDynamicGeometry.cpp
@@ -2,14 +2,13 @@
 
 #include "MeshMaterialShader.h"
 #include "ScenePrivate.h"
-#include "MeshPassProcessor.inl"
 #include "RayTracingDynamicGeometryCollection.h"
 
 #if RHI_RAYTRACING
 
 static bool IsSupportedDynamicVertexFactoryType(const FVertexFactoryType* VertexFactoryType)
 {
-	return VertexFactoryType == FindVertexFactoryType(FName(TEXT("FNiagaraSpriteVertexFactory"), FNAME_Find));
+	return VertexFactoryType == FindVertexFactoryType(FName(TEXT("FNiagaraSpriteVertexFactory"), FNAME_Find)) || VertexFactoryType == FindVertexFactoryType(FName(TEXT("FLandscapeVertexFactory"), FNAME_Find));
 }
 
 class FRayTracingDynamicGeometryConverterCS : public FMeshMaterialShader
@@ -87,11 +86,15 @@ void FRayTracingDynamicGeometryCollection::AddDynamicMeshBatchForGeometryUpdate(
 	const FScene* Scene, 
 	const FSceneView* View, 
 	const FPrimitiveSceneProxy* PrimitiveSceneProxy, 
-	const FMeshBatch& MeshBatch, 
-	FRayTracingGeometry& Geometry,
-	uint32 NumMaxVertices,
-	FRWBuffer& Buffer)
+	FRayTracingDynamicGeometryUpdateParams UpdateParams
+)
 {
+	const FMeshBatch& MeshBatch = UpdateParams.MeshBatch;
+	FRayTracingGeometry& Geometry = *UpdateParams.Geometry;
+	bool bUsingIndirectDraw = UpdateParams.bUsingIndirectDraw;
+	uint32 NumMaxVertices = UpdateParams.NumVertices;
+	FRWBuffer& Buffer = *UpdateParams.Buffer;
+
 	const FMaterialRenderProxy* FallbackMaterialRenderProxyPtr = nullptr;
 	const FMaterial& Material = MeshBatch.MaterialRenderProxy->GetMaterialWithFallback(Scene->GetFeatureLevel(), FallbackMaterialRenderProxyPtr);
 
@@ -128,15 +131,30 @@ void FRayTracingDynamicGeometryCollection::AddDynamicMeshBatchForGeometryUpdate(
 	DispatchCmd.TargetBuffer = &Buffer;
 	DispatchCmd.TargetGeometry = &Geometry;
 	DispatchCmd.NumMaxVertices = NumMaxVertices;
-	DispatchCmd.NumCPUVertices = MeshBatch.Elements[0].NumPrimitives * 2 * MeshBatch.Elements[0].NumInstances;
+	DispatchCmd.NumCPUVertices = !bUsingIndirectDraw ? UpdateParams.NumVertices : 0;
+
+	bool bRefit = true;
 
-	uint32 DesiredVertexBufferSize = FMath::DivideAndRoundUp((uint32)(NumMaxVertices * sizeof(FVector)), 4096u) * 4096u;
+	uint32 DesiredVertexBufferSize = UpdateParams.VertexBufferSize;
 	if (Buffer.NumBytes != DesiredVertexBufferSize)
 	{
 		int32 OriginalSize = Buffer.NumBytes;
-		Buffer.Initialize(4, DesiredVertexBufferSize / 4, PF_R32_FLOAT, BUF_UnorderedAccess | BUF_ShaderResource, TEXT("RayTracingDynamicVertexBuffer"));
+		Buffer.Initialize(sizeof(float), DesiredVertexBufferSize / sizeof(float), PF_R32_FLOAT, BUF_UnorderedAccess | BUF_ShaderResource, TEXT("RayTracingDynamicVertexBuffer"));
+		bRefit = false;
+	}
+
+	if (!Geometry.RayTracingGeometryRHI.IsValid())
+	{
+		bRefit = false;
 	}
 
+	if (!Geometry.Initializer.bAllowUpdate)
+	{
+		bRefit = false;
+	}
+
+	DispatchCmd.bRefit = bRefit;
+
 	check(DispatchCmd.TargetBuffer->NumBytes >= NumMaxVertices * sizeof(FVector));
 
 #if MESH_DRAW_COMMAND_DEBUG_DATA
@@ -148,8 +166,12 @@ void FRayTracingDynamicGeometryCollection::AddDynamicMeshBatchForGeometryUpdate(
 
 	check(Geometry.IsInitialized());
 	Geometry.Initializer.PositionVertexBuffer = Buffer.Buffer;
-	Geometry.Initializer.TotalPrimitiveCount = NumMaxVertices / 3;
-	Geometry.RayTracingGeometryRHI = RHICreateRayTracingGeometry(Geometry.Initializer);
+	Geometry.Initializer.TotalPrimitiveCount = UpdateParams.NumTriangles;
+
+	if (!bRefit)
+	{
+		Geometry.RayTracingGeometryRHI = RHICreateRayTracingGeometry(Geometry.Initializer);
+	}
 }
 
 void FRayTracingDynamicGeometryCollection::DispatchUpdates(FRHICommandListImmediate& RHICmdList)
@@ -158,29 +180,47 @@ void FRayTracingDynamicGeometryCollection::DispatchUpdates(FRHICommandListImmedi
 	{
 		SCOPED_DRAW_EVENT(RHICmdList, RayTracingDynamicGeometryUpdate);
 
-		for (auto& Cmd : *DispatchCommands)
 		{
-			FRayTracingDynamicGeometryConverterCS* Shader = Cmd.MaterialShader;
+			SCOPED_DRAW_EVENT(RHICmdList, VSinCSComputeDispatch);
+			for (auto& Cmd : *DispatchCommands)
+			{
+				{
+					FRayTracingDynamicGeometryConverterCS* Shader = Cmd.MaterialShader;
 
-			RHICmdList.SetComputeShader(Shader->GetComputeShader());
+					RHICmdList.SetComputeShader(Shader->GetComputeShader());
 
-			Cmd.ShaderBindings.SetOnCommandListForCompute(RHICmdList, Shader->GetComputeShader());
-			Shader->RWVertexPositions.SetBuffer(RHICmdList, Shader->GetComputeShader(), *Cmd.TargetBuffer);
-			SetShaderValue(RHICmdList, Shader->GetComputeShader(), Shader->NumMaxVertices, Cmd.NumMaxVertices);
-			SetShaderValue(RHICmdList, Shader->GetComputeShader(), Shader->NumCPUVertices, Cmd.NumCPUVertices);
+					Cmd.ShaderBindings.SetOnCommandListForCompute(RHICmdList, Shader->GetComputeShader());
+					Shader->RWVertexPositions.SetBuffer(RHICmdList, Shader->GetComputeShader(), *Cmd.TargetBuffer);
+					SetShaderValue(RHICmdList, Shader->GetComputeShader(), Shader->NumMaxVertices, Cmd.TargetBuffer->NumBytes / sizeof(FVector));
+					SetShaderValue(RHICmdList, Shader->GetComputeShader(), Shader->NumCPUVertices, Cmd.NumCPUVertices);
 
-			RHICmdList.DispatchComputeShader(FMath::DivideAndRoundUp<uint32>(Cmd.NumMaxVertices, 256), 1, 1);
-			Shader->RWVertexPositions.UnsetUAV(RHICmdList, Shader->GetComputeShader());
+					RHICmdList.DispatchComputeShader(FMath::DivideAndRoundUp<uint32>(Cmd.NumMaxVertices, 256), 1, 1);
+
+					Shader->RWVertexPositions.UnsetUAV(RHICmdList, Shader->GetComputeShader());
+				}
+			}
 		}
 
 		TArray<FAccelerationStructureUpdateParams> BuildParams;
+		TArray<FAccelerationStructureUpdateParams> RefitParams;
 
 		for (auto& Cmd : *DispatchCommands)
 		{
-			BuildParams.Add(FAccelerationStructureUpdateParams { Cmd.TargetGeometry->RayTracingGeometryRHI, Cmd.TargetBuffer->Buffer });
+			if (Cmd.bRefit)
+				RefitParams.Add(FAccelerationStructureUpdateParams{ Cmd.TargetGeometry->RayTracingGeometryRHI, Cmd.TargetBuffer->Buffer });
+			else
+				BuildParams.Add(FAccelerationStructureUpdateParams { Cmd.TargetGeometry->RayTracingGeometryRHI, Cmd.TargetBuffer->Buffer });
 		}
 
-		RHICmdList.BuildAccelerationStructures(BuildParams);
+		{
+			SCOPED_DRAW_EVENT(RHICmdList, Build);
+			RHICmdList.BuildAccelerationStructures(BuildParams);
+		}
+
+		{
+			SCOPED_DRAW_EVENT(RHICmdList, Refit);
+			RHICmdList.UpdateAccelerationStructures(RefitParams);
+		}
 
 		Clear();
 	}
diff --git a/Engine/Source/Runtime/Renderer/Private/RayTracing/RayTracingMaterialHitShaders.cpp b/Engine/Source/Runtime/Renderer/Private/RayTracing/RayTracingMaterialHitShaders.cpp
--- a/Engine/Source/Runtime/Renderer/Private/RayTracing/RayTracingMaterialHitShaders.cpp
+++ b/Engine/Source/Runtime/Renderer/Private/RayTracing/RayTracingMaterialHitShaders.cpp
@@ -32,11 +32,13 @@ static bool IsSupportedVertexFactoryType(const FVertexFactoryType* VertexFactory
 	static FName InstancedVfFname = FName(TEXT("FInstancedStaticMeshVertexFactory"), FNAME_Find);
 	static FName NiagaraSpriteVfFname = FName(TEXT("FNiagaraSpriteVertexFactory"), FNAME_Find);
 	static FName GeometryCacheVfFname = FName(TEXT("FGeometryCacheVertexVertexFactory"), FNAME_Find);
+	static FName LandscapeVfFname = FName(TEXT("FLandscapeVertexFactory"), FNAME_Find);
 
-	return (VertexFactoryType == FindVertexFactoryType(LocalVfFname)
+	return VertexFactoryType == FindVertexFactoryType(LocalVfFname)
 		|| VertexFactoryType == FindVertexFactoryType(LSkinnedVfFname)
 		|| VertexFactoryType == FindVertexFactoryType(NiagaraSpriteVfFname)
-		|| VertexFactoryType == FindVertexFactoryType(GeometryCacheVfFname));
+		|| VertexFactoryType == FindVertexFactoryType(GeometryCacheVfFname)
+		|| VertexFactoryType == FindVertexFactoryType(LandscapeVfFname);
 }
 
 class FMaterialCHS : public FMeshMaterialShader, public FUniformLightMapPolicyShaderParametersType
diff --git a/Engine/Source/Runtime/Renderer/Private/ScenePrivate.h b/Engine/Source/Runtime/Renderer/Private/ScenePrivate.h
--- a/Engine/Source/Runtime/Renderer/Private/ScenePrivate.h
+++ b/Engine/Source/Runtime/Renderer/Private/ScenePrivate.h
@@ -2351,8 +2351,10 @@ struct FMeshComputeDispatchCommand
 
 	uint32 NumMaxVertices;
 	uint32 NumCPUVertices;
+	uint32 BaseVertexIndex;
 	FRWBuffer* TargetBuffer;
 	FRayTracingGeometry* TargetGeometry;
+	bool bRefit;
 };
 #endif
 
diff --git a/Engine/Source/Runtime/Renderer/Public/RayTracingDynamicGeometryCollection.h b/Engine/Source/Runtime/Renderer/Public/RayTracingDynamicGeometryCollection.h
--- a/Engine/Source/Runtime/Renderer/Public/RayTracingDynamicGeometryCollection.h
+++ b/Engine/Source/Runtime/Renderer/Public/RayTracingDynamicGeometryCollection.h
@@ -9,7 +9,12 @@ class RENDERER_API FRayTracingDynamicGeometryCollection
 public:
 	FRayTracingDynamicGeometryCollection();
 
-	void AddDynamicMeshBatchForGeometryUpdate(const FScene* Scene, const FSceneView* View, const FPrimitiveSceneProxy* PrimitiveSceneProxy, const FMeshBatch& MeshBatch, FRayTracingGeometry& Geometry, uint32 NumMaxVertices, FRWBuffer& Buffer);
+	void AddDynamicMeshBatchForGeometryUpdate(
+		const FScene* Scene, 
+		const FSceneView* View, 
+		const FPrimitiveSceneProxy* PrimitiveSceneProxy, 
+		FRayTracingDynamicGeometryUpdateParams Params);
+
 	void DispatchUpdates(FRHICommandListImmediate& RHICmdList);
 	void Clear();
 
