Fix landscape clang CIS

diff --git a/Engine/Source/Runtime/Landscape/Private/LandscapeRender.cpp b/Engine/Source/Runtime/Landscape/Private/LandscapeRender.cpp
--- a/Engine/Source/Runtime/Landscape/Private/LandscapeRender.cpp
+++ b/Engine/Source/Runtime/Landscape/Private/LandscapeRender.cpp
@@ -3145,7 +3145,7 @@ void FLandscapeComponentSceneProxy::GetDynamicRayTracingInstances(FRayTracingMat
 					MeshBatch,
 					false,
 					(uint32)FMath::Square(LodSubsectionSizeVerts),
-					(uint32)FMath::Square(LodSubsectionSizeVerts) * sizeof(FVector),
+					FMath::Square(LodSubsectionSizeVerts) * (uint32)sizeof(FVector),
 					(uint32)FMath::Square(LodSubsectionSizeVerts - 1) * 2,
 					&SectionRayTracingStates[SubSectionIdx].Geometry,
 					&SectionRayTracingStates[SubSectionIdx].RayTracingDynamicVertexBuffer
