Fix GPU crash when trying to mistakenly refit landscape ray tracing geometry

diff --git a/Engine/Source/Runtime/Landscape/Private/LandscapeRender.cpp b/Engine/Source/Runtime/Landscape/Private/LandscapeRender.cpp
--- a/Engine/Source/Runtime/Landscape/Private/LandscapeRender.cpp
+++ b/Engine/Source/Runtime/Landscape/Private/LandscapeRender.cpp
@@ -3129,6 +3129,11 @@ void FLandscapeComponentSceneProxy::GetDynamicRayTracingInstances(FRayTracingMat
 
 			MeshBatch.Elements.Add(BatchElement);
 
+			if (SectionRayTracingStates[SubSectionIdx].CurrentLOD != CurrentLOD)
+			{
+				SectionRayTracingStates[SubSectionIdx].CurrentLOD = CurrentLOD;
+				SectionRayTracingStates[SubSectionIdx].RayTracingDynamicVertexBuffer.Release();
+			}
 			SectionRayTracingStates[SubSectionIdx].Geometry.Initializer.IndexBuffer = BatchElement.IndexBuffer->IndexBufferRHI;
 
 			{
