Optimization for landscape ray tracing BLAS by only updating when LOD state changes. Not compatible with world position offset as texture streaming may make the cached BLAS mismatch the raster one

diff --git a/Engine/Source/Runtime/Landscape/Private/LandscapeRender.cpp b/Engine/Source/Runtime/Landscape/Private/LandscapeRender.cpp
--- a/Engine/Source/Runtime/Landscape/Private/LandscapeRender.cpp
+++ b/Engine/Source/Runtime/Landscape/Private/LandscapeRender.cpp
@@ -120,6 +120,15 @@ FAutoConsoleVariableRef CVarLandscapeDebugViewMode(
 );
 #endif
 
+#if RHI_RAYTRACING
+int32 GLandscapeRayTracingGeometryLODsThatUpdateEveryFrame = 3;
+static FAutoConsoleVariableRef CVarLandscapeRayTracingGeometryLODsThatUpdateEveryFrame(
+	TEXT("r.RayTracing.Landscape.LODsUpdateEveryFrame"),
+	GLandscapeRayTracingGeometryLODsThatUpdateEveryFrame,
+	TEXT("If on, LODs that are lower than the specified level will be updated every frame, which can be used to workaround some artifacts caused by texture streaming if you're using WPO on the landscape")
+);
+#endif
+
 /*------------------------------------------------------------------------------
 Forsyth algorithm for cache optimizing index buffers.
 ------------------------------------------------------------------------------*/
@@ -3129,11 +3138,6 @@ void FLandscapeComponentSceneProxy::GetDynamicRayTracingInstances(FRayTracingMat
 
 			MeshBatch.Elements.Add(BatchElement);
 
-			if (SectionRayTracingStates[SubSectionIdx].CurrentLOD != CurrentLOD)
-			{
-				SectionRayTracingStates[SubSectionIdx].CurrentLOD = CurrentLOD;
-				SectionRayTracingStates[SubSectionIdx].RayTracingDynamicVertexBuffer.Release();
-			}
 			SectionRayTracingStates[SubSectionIdx].Geometry.Initializer.IndexBuffer = BatchElement.IndexBuffer->IndexBufferRHI;
 
 			{
@@ -3144,18 +3148,63 @@ void FLandscapeComponentSceneProxy::GetDynamicRayTracingInstances(FRayTracingMat
 				BatchElementParams.LandscapeVertexFactoryMVFUniformBuffer = FLandscapeVertexFactoryMVFUniformBufferRef::CreateUniformBufferImmediate(UniformBufferParams, UniformBuffer_SingleFrame);
 			}
 
-			Context.DynamicRayTracingGeometriesToUpdate.Add(
-				FRayTracingDynamicGeometryUpdateParams
+			bool bNeedsRayTracingGeometryUpdate = false;
+
+			// Detect force update CVar
+			bNeedsRayTracingGeometryUpdate |= (CurrentLOD <= GLandscapeRayTracingGeometryLODsThatUpdateEveryFrame) ? true : false;
+
+			// Detect continuous LOD parameter changes. This is for far-away high LODs - they change rarely yet the BLAS refit time is not ideal, even if they contains tiny amount of triangles
+			{
+				if (SectionRayTracingStates[SubSectionIdx].CurrentLOD != CurrentLOD)
+				{
+					bNeedsRayTracingGeometryUpdate = true;
+					SectionRayTracingStates[SubSectionIdx].CurrentLOD = CurrentLOD;
+					SectionRayTracingStates[SubSectionIdx].RayTracingDynamicVertexBuffer.Release();
+				}
+				if (!SectionRayTracingStates[SubSectionIdx].LODBias.Equals(GetShaderLODBias()))
+				{
+					bNeedsRayTracingGeometryUpdate = true;
+					SectionRayTracingStates[SubSectionIdx].LODBias = GetShaderLODBias();
+				}
+
+				float ComponentScreenSize = GetComponentScreenSize(Context.ReferenceView, LandscapeComponent->Bounds.Origin, ComponentMaxExtend, LandscapeComponent->Bounds.SphereRadius);
+
+				FViewCustomDataLOD CurrentLODData;
+				CalculateBatchElementLOD(*Context.ReferenceView, ComponentScreenSize, Context.ReferenceView->LODDistanceFactor, CurrentLODData, true);
+
+				if (SectionRayTracingStates[SubSectionIdx].SectionLOD.Equals(CurrentLODData.ShaderCurrentLOD))
 				{
-					MeshBatch,
-					false,
-					(uint32)FMath::Square(LodSubsectionSizeVerts),
-					FMath::Square(LodSubsectionSizeVerts) * (uint32)sizeof(FVector),
-					(uint32)FMath::Square(LodSubsectionSizeVerts - 1) * 2,
-					&SectionRayTracingStates[SubSectionIdx].Geometry,
-					&SectionRayTracingStates[SubSectionIdx].RayTracingDynamicVertexBuffer
+					bNeedsRayTracingGeometryUpdate = true;
+					SectionRayTracingStates[SubSectionIdx].SectionLOD = CurrentLODData.ShaderCurrentLOD;
+				}
+
+				{
+					FVector4 CurrentNeighborLOD;
+					GetShaderCurrentNeighborLOD(*Context.ReferenceView, CurrentLODData.SubSections[SubSectionIdx].fBatchElementCurrentLOD, NumSubsections > 1 ? SubX : INDEX_NONE, NumSubsections > 1 ? SubY : INDEX_NONE, SubSectionIdx, CurrentNeighborLOD);
+
+					if (!SectionRayTracingStates[SubSectionIdx].CurrentNeighborLOD.Equals(CurrentNeighborLOD))
+					{
+						bNeedsRayTracingGeometryUpdate = true;
+						SectionRayTracingStates[SubSectionIdx].CurrentNeighborLOD = CurrentNeighborLOD;
+					}
 				}
-			);
+			}
+
+			if (bNeedsRayTracingGeometryUpdate)
+			{
+				Context.DynamicRayTracingGeometriesToUpdate.Add(
+					FRayTracingDynamicGeometryUpdateParams
+					{
+						MeshBatch,
+						false,
+						(uint32)FMath::Square(LodSubsectionSizeVerts),
+						FMath::Square(LodSubsectionSizeVerts) * (uint32)sizeof(FVector),
+						(uint32)FMath::Square(LodSubsectionSizeVerts - 1) * 2,
+						&SectionRayTracingStates[SubSectionIdx].Geometry,
+						&SectionRayTracingStates[SubSectionIdx].RayTracingDynamicVertexBuffer
+					}
+				);
+			}
 
 			FRayTracingInstance RayTracingInstance;
 			RayTracingInstance.Geometry = &SectionRayTracingStates[SubSectionIdx].Geometry;
diff --git a/Engine/Source/Runtime/Landscape/Public/LandscapeRender.h b/Engine/Source/Runtime/Landscape/Public/LandscapeRender.h
--- a/Engine/Source/Runtime/Landscape/Public/LandscapeRender.h
+++ b/Engine/Source/Runtime/Landscape/Public/LandscapeRender.h
@@ -551,11 +551,19 @@ public:
 	struct FLandscapeSectionRayTracingState
 	{
 		int8 CurrentLOD;
+		FVector4 LODBias;
+		FVector4 SectionLOD;
+		FVector4 CurrentNeighborLOD;
+
 		FRayTracingGeometry Geometry;
 		FRWBuffer RayTracingDynamicVertexBuffer;
 		FLandscapeVertexFactoryMVFUniformBufferRef UniformBuffer;
 
-		FLandscapeSectionRayTracingState() : CurrentLOD(-1) {}
+		FLandscapeSectionRayTracingState() 
+			: CurrentLOD(-1)
+			, LODBias(-1000.0f, -1000.0f, -1000.0f, -1000.0f)
+			, SectionLOD(-1000.0f, -1000.0f, -1000.0f, -1000.0f)
+			, CurrentNeighborLOD(-1000.0f, -1000.0f, -1000.0f, -1000.0f) {}
 	};
 
 	TStaticArray<FLandscapeSectionRayTracingState, MAX_SUBSECTION_COUNT> SectionRayTracingStates;
diff --git a/Engine/Source/Runtime/Renderer/Private/RayTracing/RayTracingDynamicGeometry.cpp b/Engine/Source/Runtime/Renderer/Private/RayTracing/RayTracingDynamicGeometry.cpp
--- a/Engine/Source/Runtime/Renderer/Private/RayTracing/RayTracingDynamicGeometry.cpp
+++ b/Engine/Source/Runtime/Renderer/Private/RayTracing/RayTracingDynamicGeometry.cpp
@@ -6,6 +6,8 @@
 
 #if RHI_RAYTRACING
 
+DECLARE_GPU_STAT_NAMED(RayTracingDynamicGeom, TEXT("Ray Tracing Dynamic Geometry Update"));
+
 static bool IsSupportedDynamicVertexFactoryType(const FVertexFactoryType* VertexFactoryType)
 {
 	return VertexFactoryType == FindVertexFactoryType(FName(TEXT("FNiagaraSpriteVertexFactory"), FNAME_Find)) || VertexFactoryType == FindVertexFactoryType(FName(TEXT("FLandscapeVertexFactory"), FNAME_Find));
@@ -179,6 +181,7 @@ void FRayTracingDynamicGeometryCollection::DispatchUpdates(FRHICommandListImmedi
 	if (DispatchCommands->Num() > 0)
 	{
 		SCOPED_DRAW_EVENT(RHICmdList, RayTracingDynamicGeometryUpdate);
+		SCOPED_GPU_STAT(RHICmdList, RayTracingDynamicGeom);
 
 		{
 			SCOPED_DRAW_EVENT(RHICmdList, VSinCSComputeDispatch);
