Detect texture streaming changes in landscape materials to overcome artifacts caused by using textures to calculate WorldPositionOffset

diff --git a/Engine/Source/Runtime/Engine/Private/Materials/MaterialUniformExpressions.cpp b/Engine/Source/Runtime/Engine/Private/Materials/MaterialUniformExpressions.cpp
--- a/Engine/Source/Runtime/Engine/Private/Materials/MaterialUniformExpressions.cpp
+++ b/Engine/Source/Runtime/Engine/Private/Materials/MaterialUniformExpressions.cpp
@@ -855,6 +855,28 @@ void FUniformExpressionSet::FillUniformBuffer(const FMaterialRenderContext& Mate
 	}
 }
 
+uint32 FUniformExpressionSet::GetReferencedTexture2DRHIHash(const FMaterialRenderContext& MaterialRenderContext) const
+{
+	uint32 BaseHash = 0;
+
+	for (int32 ExpressionIndex = 0; ExpressionIndex < Uniform2DTextureExpressions.Num(); ExpressionIndex++)
+	{
+		const UTexture* Value;
+		Uniform2DTextureExpressions[ExpressionIndex]->GetTextureValue(MaterialRenderContext, MaterialRenderContext.Material, Value);
+
+		const uint32 ValidTextureTypes = MCT_Texture2D | MCT_TextureVirtual | MCT_TextureExternal;
+
+		FRHITexture* TexturePtr = nullptr;
+		if (Value && Value->Resource && Value->TextureReference.TextureReferenceRHI && (Value->GetMaterialType() & ValidTextureTypes) != 0u)
+		{
+			TexturePtr = Value->TextureReference.TextureReferenceRHI->GetReferencedTexture();
+		}
+		BaseHash = PointerHash(TexturePtr, BaseHash);
+	}
+
+	return BaseHash;
+}
+
 FMaterialUniformExpressionTexture::FMaterialUniformExpressionTexture() :
 	TextureIndex(INDEX_NONE),
 	LayerIndex(INDEX_NONE),
diff --git a/Engine/Source/Runtime/Engine/Public/MaterialShared.h b/Engine/Source/Runtime/Engine/Public/MaterialShared.h
--- a/Engine/Source/Runtime/Engine/Public/MaterialShared.h
+++ b/Engine/Source/Runtime/Engine/Public/MaterialShared.h
@@ -466,6 +466,9 @@ public:
 
 	ENGINE_API void FillUniformBuffer(const FMaterialRenderContext& MaterialRenderContext, const FUniformExpressionCache& UniformExpressionCache, uint8* TempBuffer, int TempBufferSize) const;
 
+	// Get a combined hash of all referenced Texture2D's underlying RHI textures, going through TextureReferences. Can be used to tell if any texture has gone through texture streaming mip changes recently.
+	ENGINE_API uint32 GetReferencedTexture2DRHIHash(const FMaterialRenderContext& MaterialRenderContext) const;
+
 	inline bool HasExternalTextureExpressions() const
 	{
 		return UniformExternalTextureExpressions.Num() > 0;
diff --git a/Engine/Source/Runtime/Landscape/Private/LandscapeRender.cpp b/Engine/Source/Runtime/Landscape/Private/LandscapeRender.cpp
--- a/Engine/Source/Runtime/Landscape/Private/LandscapeRender.cpp
+++ b/Engine/Source/Runtime/Landscape/Private/LandscapeRender.cpp
@@ -122,11 +122,18 @@ FAutoConsoleVariableRef CVarLandscapeDebugViewMode(
 #endif
 
 #if RHI_RAYTRACING
-int32 GLandscapeRayTracingGeometryLODsThatUpdateEveryFrame = 3;
+int32 GLandscapeRayTracingGeometryLODsThatUpdateEveryFrame = 0;
 static FAutoConsoleVariableRef CVarLandscapeRayTracingGeometryLODsThatUpdateEveryFrame(
 	TEXT("r.RayTracing.Landscape.LODsUpdateEveryFrame"),
 	GLandscapeRayTracingGeometryLODsThatUpdateEveryFrame,
-	TEXT("If on, LODs that are lower than the specified level will be updated every frame, which can be used to workaround some artifacts caused by texture streaming if you're using WPO on the landscape")
+	TEXT("If on, LODs that are lower than the specified level will be updated every frame, which can be used to workaround some artifacts caused by texture streaming if you're using WorldPositionOffset on the landscape")
+);
+
+int32 GLandscapeRayTracingGeometryDetectTextureStreaming = 1;
+static FAutoConsoleVariableRef CVarLandscapeRayTracingGeometryDetectTextureStreaming(
+	TEXT("r.RayTracing.Landscape.DetectTextureStreaming"),
+	GLandscapeRayTracingGeometryDetectTextureStreaming,
+	TEXT("If on, update ray tracing geometry when texture streaming state changes. Useful when WorldPositionOffset is used in the landscape material")
 );
 #endif
 
@@ -3179,6 +3186,28 @@ void FLandscapeComponentSceneProxy::GetDynamicRayTracingInstances(FRayTracingMat
 				}
 			}
 
+			if (GLandscapeRayTracingGeometryDetectTextureStreaming > 0)
+			{
+				const FMaterialRenderProxy* FallbackMaterialRenderProxyPtr = nullptr;
+				const FMaterial& Material = MeshBatch.MaterialRenderProxy->GetMaterialWithFallback(((FSceneInterface*)Context.Scene)->GetFeatureLevel(), FallbackMaterialRenderProxyPtr);
+
+				if (Material.HasVertexPositionOffsetConnected())
+				{
+					const FMaterialRenderProxy* MaterialRenderProxy = FallbackMaterialRenderProxyPtr ? FallbackMaterialRenderProxyPtr : MeshBatch.MaterialRenderProxy;
+
+					FMaterialRenderContext MaterialRenderContext(MaterialRenderProxy, Material, Context.ReferenceView);
+
+					const FUniformExpressionSet& UniformExpressionSet = Material.GetRenderingThreadShaderMap()->GetUniformExpressionSet();
+					const uint32 Hash = UniformExpressionSet.GetReferencedTexture2DRHIHash(MaterialRenderContext);
+
+					if (SectionRayTracingStates[SubSectionIdx].ReferencedTextureRHIHash != Hash)
+					{
+						bNeedsRayTracingGeometryUpdate = true;
+						SectionRayTracingStates[SubSectionIdx].ReferencedTextureRHIHash = Hash;
+					}
+				}
+			}
+
 			if (bNeedsRayTracingGeometryUpdate)
 			{
 				Context.DynamicRayTracingGeometriesToUpdate.Add(
diff --git a/Engine/Source/Runtime/Landscape/Public/LandscapeRender.h b/Engine/Source/Runtime/Landscape/Public/LandscapeRender.h
--- a/Engine/Source/Runtime/Landscape/Public/LandscapeRender.h
+++ b/Engine/Source/Runtime/Landscape/Public/LandscapeRender.h
@@ -554,6 +554,7 @@ public:
 		FVector4 LODBias;
 		FVector4 SectionLOD;
 		FVector4 CurrentNeighborLOD;
+		uint32 ReferencedTextureRHIHash;
 
 		FRayTracingGeometry Geometry;
 		FRWBuffer RayTracingDynamicVertexBuffer;
@@ -563,7 +564,8 @@ public:
 			: CurrentLOD(-1)
 			, LODBias(-1000.0f, -1000.0f, -1000.0f, -1000.0f)
 			, SectionLOD(-1000.0f, -1000.0f, -1000.0f, -1000.0f)
-			, CurrentNeighborLOD(-1000.0f, -1000.0f, -1000.0f, -1000.0f) {}
+			, CurrentNeighborLOD(-1000.0f, -1000.0f, -1000.0f, -1000.0f)
+			, ReferencedTextureRHIHash(0) {}
 	};
 
 	TStaticArray<FLandscapeSectionRayTracingState, MAX_SUBSECTION_COUNT> SectionRayTracingStates;
