Remove some unnecessary landscape shader variations added by runtime virtual texture

diff --git a/Engine/Source/Runtime/Landscape/Private/LandscapeRender.cpp b/Engine/Source/Runtime/Landscape/Private/LandscapeRender.cpp
--- a/Engine/Source/Runtime/Landscape/Private/LandscapeRender.cpp
+++ b/Engine/Source/Runtime/Landscape/Private/LandscapeRender.cpp
@@ -4276,17 +4276,25 @@ public:
 			else
 			{
 				// Landscape MICs are only for use with the Landscape vertex factories
+
+				// For now only compile FLandscapeFixedGridVertexFactory for runtime virtual texture page rendering (can change if we need for other cases)
 				// Todo: only compile LandscapeXYOffsetVertexFactory if we are using it
+				bool bIsRuntimeVirtualTextureShaderType = Algo::Find(GetRuntimeVirtualTextureShaderTypes(), ShaderType->GetFName());
+
 				static const FName LandscapeVertexFactory = FName(TEXT("FLandscapeVertexFactory"));
 				static const FName LandscapeXYOffsetVertexFactory = FName(TEXT("FLandscapeXYOffsetVertexFactory"));
-				static const FName LandscapeFixedGridVertexFactory = FName(TEXT("FLandscapeFixedGridVertexFactory"));
 				static const FName LandscapeVertexFactoryMobile = FName(TEXT("FLandscapeVertexFactoryMobile"));
 				if (VertexFactoryType->GetFName() == LandscapeVertexFactory ||
 					VertexFactoryType->GetFName() == LandscapeXYOffsetVertexFactory ||
-					VertexFactoryType->GetFName() == LandscapeFixedGridVertexFactory ||
 					VertexFactoryType->GetFName() == LandscapeVertexFactoryMobile)
 				{
-					return FMaterialResource::ShouldCache(Platform, ShaderType, VertexFactoryType);
+					return !bIsRuntimeVirtualTextureShaderType && FMaterialResource::ShouldCache(Platform, ShaderType, VertexFactoryType);
+				}
+
+				static const FName LandscapeFixedGridVertexFactory = FName(TEXT("FLandscapeFixedGridVertexFactory"));
+				if (VertexFactoryType->GetFName() == LandscapeFixedGridVertexFactory)
+				{
+					return bIsRuntimeVirtualTextureShaderType && FMaterialResource::ShouldCache(Platform, ShaderType, VertexFactoryType);
 				}
 			}
 		}
@@ -4539,6 +4547,20 @@ public:
 		};
 		return ExcludedShaderTypes;
 	}
+
+	static const TArray<FName>& GetRuntimeVirtualTextureShaderTypes()
+	{
+		static const TArray<FName> ShaderTypes =
+		{
+			FName(TEXT("TVirtualTextureVSBaseColor")),
+			FName(TEXT("TVirtualTextureVSBaseColorNormal")),
+			FName(TEXT("TVirtualTextureVSBaseColorNormalSpecular")),
+			FName(TEXT("TVirtualTexturePSBaseColor")),
+			FName(TEXT("TVirtualTexturePSBaseColorNormal")),
+			FName(TEXT("TVirtualTexturePSBaseColorNormalSpecular")),
+		};
+		return ShaderTypes;
+	}
 };
 
 FMaterialResource* ULandscapeMaterialInstanceConstant::AllocatePermutationResource()
diff --git a/Engine/Source/Runtime/Renderer/Private/VT/RuntimeVirtualTextureRender.cpp b/Engine/Source/Runtime/Renderer/Private/VT/RuntimeVirtualTextureRender.cpp
--- a/Engine/Source/Runtime/Renderer/Private/VT/RuntimeVirtualTextureRender.cpp
+++ b/Engine/Source/Runtime/Renderer/Private/VT/RuntimeVirtualTextureRender.cpp
@@ -156,6 +156,8 @@ namespace RuntimeVirtualTexture
 		}
 	};
 
+	// If we change this macro or add additional policy types then we need to update GetRuntimeVirtualTextureShaderTypes() in LandscapeRender.cpp
+	// That code is used to filter out unnecessary shader variations
 #define IMPLEMENT_VIRTUALTEXTURE_SHADER_TYPE(PolicyType, PolicyName) \
 	typedef FShader_VirtualTextureMaterialDraw_VS<PolicyType> TVirtualTextureVS##PolicyName; \
 	IMPLEMENT_MATERIAL_SHADER_TYPE(template<>,TVirtualTextureVS##PolicyName, TEXT("/Engine/Private/VirtualTextureMaterial.usf"), TEXT("MainVS"), SF_Vertex); \
