Refactors the 4 different implementation of Halton() as a single one.

diff --git a/Engine/Source/Runtime/Core/Public/Math/Halton.h b/Engine/Source/Runtime/Core/Public/Math/Halton.h
new file mode 100644
--- /dev/null
+++ b/Engine/Source/Runtime/Core/Public/Math/Halton.h
@@ -0,0 +1,21 @@
+// Copyright 1998-2019 Epic Games, Inc. All Rights Reserved.
+
+#pragma once
+
+#include "CoreTypes.h"
+
+
+/** [ Halton 1964, "Radical-inverse quasi-random point sequence" ] */
+inline float Halton(int32 Index, int32 Base)
+{
+	float Result = 0.0f;
+	float InvBase = 1.0f / Base;
+	float Fraction = InvBase;
+	while (Index > 0)
+	{
+		Result += (Index % Base) * Fraction;
+		Index /= Base;
+		Fraction *= InvBase;
+	}
+	return Result;
+}
diff --git a/Engine/Source/Runtime/Landscape/Private/LandscapeGrass.cpp b/Engine/Source/Runtime/Landscape/Private/LandscapeGrass.cpp
--- a/Engine/Source/Runtime/Landscape/Private/LandscapeGrass.cpp
+++ b/Engine/Source/Runtime/Landscape/Private/LandscapeGrass.cpp
@@ -55,6 +55,7 @@
 #include "InstancedStaticMesh.h"
 #include "MeshPassProcessor.h"
 #include "MeshPassProcessor.inl"
+#include "Math/Halton.h"
 
 #define LOCTEXT_NAMESPACE "Landscape"
 
@@ -1545,21 +1546,6 @@ private:
 	int32 Stride;
 };
 
-template<uint32 Base>
-static FORCEINLINE float Halton(uint32 Index)
-{
-	float Result = 0.0f;
-	float InvBase = 1.0f / Base;
-	float Fraction = InvBase;
-	while( Index > 0 )
-	{
-		Result += ( Index % Base ) * Fraction;
-		Index /= Base;
-		Fraction *= InvBase;
-	}
-	return Result;
-}
-
 struct FAsyncGrassBuilder : public FGrassBuilderBase
 {
 	FLandscapeComponentGrassAccess GrassData;
@@ -1783,8 +1769,8 @@ struct FAsyncGrassBuilder : public FGrassBuilderBase
 			FVector DivExtent(Extent * Div);
 			for (int32 InstanceIndex = 0; InstanceIndex < MaxNum; InstanceIndex++)
 			{
-				float HaltonX = Halton<2>(InstanceIndex + HaltonBaseIndex);
-				float HaltonY = Halton<3>(InstanceIndex + HaltonBaseIndex);
+				float HaltonX = Halton(InstanceIndex + HaltonBaseIndex, 2);
+				float HaltonY = Halton(InstanceIndex + HaltonBaseIndex, 3);
 				FVector Location(Origin.X + HaltonX * Extent.X, Origin.Y + HaltonY * Extent.Y, 0.0f);
 				FVector LocationWithHeight;
 				float Weight = GetLayerWeightAtLocationLocal(Location, &LocationWithHeight);
diff --git a/Engine/Source/Runtime/Renderer/Private/PostProcess/PostProcessTonemap.h b/Engine/Source/Runtime/Renderer/Private/PostProcess/PostProcessTonemap.h
--- a/Engine/Source/Runtime/Renderer/Private/PostProcess/PostProcessTonemap.h
+++ b/Engine/Source/Runtime/Renderer/Private/PostProcess/PostProcessTonemap.h
@@ -16,25 +16,13 @@
 #include "PostProcessParameters.h"
 #include "PostProcess/RenderingCompositionGraph.h"
 #include "PostProcess/PostProcessEyeAdaptation.h"
+#include "Math/Halton.h"
 
-static float GrainHalton( int32 Index, int32 Base )
-{
-	float Result = 0.0f;
-	float InvBase = 1.0f / Base;
-	float Fraction = InvBase;
-	while( Index > 0 )
-	{
-		Result += ( Index % Base ) * Fraction;
-		Index /= Base;
-		Fraction *= InvBase;
-	}
-	return Result;
-}
 
 static void GrainRandomFromFrame(FVector* RESTRICT const Constant, uint32 FrameNumber)
 {
-	Constant->X = GrainHalton(FrameNumber & 1023, 2);
-	Constant->Y = GrainHalton(FrameNumber & 1023, 3);
+	Constant->X = Halton(FrameNumber & 1023, 2);
+	Constant->Y = Halton(FrameNumber & 1023, 3);
 }
 
 
diff --git a/Engine/Source/Runtime/Renderer/Private/SceneVisibility.cpp b/Engine/Source/Runtime/Renderer/Private/SceneVisibility.cpp
--- a/Engine/Source/Runtime/Renderer/Private/SceneVisibility.cpp
+++ b/Engine/Source/Runtime/Renderer/Private/SceneVisibility.cpp
@@ -37,6 +37,7 @@
 #include "TranslucentRendering.h"
 #include "Async/ParallelFor.h"
 #include "RectLightSceneProxy.h"
+#include "Math/Halton.h"
 
 /*------------------------------------------------------------------------------
 	Globals
@@ -2887,20 +2888,6 @@ static bool IsLargeCameraMovement(FSceneView& View, const FMatrix& PrevViewMatri
 		Distance.SizeSquared() > CameraTranslationThreshold * CameraTranslationThreshold;
 }
 
-float Halton( int32 Index, int32 Base )
-{
-	float Result = 0.0f;
-	float InvBase = 1.0f / Base;
-	float Fraction = InvBase;
-	while( Index > 0 )
-	{
-		Result += ( Index % Base ) * Fraction;
-		Index /= Base;
-		Fraction *= InvBase;
-	}
-	return Result;
-}
-
 void FSceneRenderer::PreVisibilityFrameSetup(FRHICommandListImmediate& RHICmdList)
 {
 	// Notify the RHI we are beginning to render a scene.
diff --git a/Engine/Source/Runtime/Renderer/Private/VolumetricFog.cpp b/Engine/Source/Runtime/Renderer/Private/VolumetricFog.cpp
--- a/Engine/Source/Runtime/Renderer/Private/VolumetricFog.cpp
+++ b/Engine/Source/Runtime/Renderer/Private/VolumetricFog.cpp
@@ -18,6 +18,7 @@ VolumetricFog.cpp
 #include "VolumeLighting.h"
 #include "PipelineStateCache.h"
 #include "ShaderParameterStruct.h"
+#include "Math/Halton.h"
 
 int32 GVolumetricFog = 1;
 FAutoConsoleVariableRef CVarVolumetricFog(
@@ -107,20 +108,6 @@ DECLARE_GPU_STAT(VolumetricFog);
 FVolumetricFogGlobalData::FVolumetricFogGlobalData()
 {}
 
-float TemporalHalton(int32 Index, int32 Base)
-{
-	float Result = 0.0f;
-	float InvBase = 1.0f / Base;
-	float Fraction = InvBase;
-	while (Index > 0)
-	{
-		Result += (Index % Base) * Fraction;
-		Index /= Base;
-		Fraction *= InvBase;
-	}
-	return Result;
-}
-
 FVector VolumetricFogTemporalRandom(uint32 FrameNumber)
 {
 	// Center of the voxel
@@ -128,7 +115,7 @@ FVector VolumetricFogTemporalRandom(uint32 FrameNumber)
 
 	if (GVolumetricFogJitter && GVolumetricFogTemporalReprojection)
 	{
-		RandomOffsetValue = FVector(TemporalHalton(FrameNumber & 1023, 2), TemporalHalton(FrameNumber & 1023, 3), TemporalHalton(FrameNumber & 1023, 5));
+		RandomOffsetValue = FVector(Halton(FrameNumber & 1023, 2), Halton(FrameNumber & 1023, 3), Halton(FrameNumber & 1023, 5));
 	}
 
 	return RandomOffsetValue;
