Landscape Layers: Fix Crash when Clearing Weightmaps

#rnx

diff --git a/Engine/Source/Runtime/Landscape/Private/LandscapeEditInterface.cpp b/Engine/Source/Runtime/Landscape/Private/LandscapeEditInterface.cpp
--- a/Engine/Source/Runtime/Landscape/Private/LandscapeEditInterface.cpp
+++ b/Engine/Source/Runtime/Landscape/Private/LandscapeEditInterface.cpp
@@ -1602,6 +1602,7 @@ void ULandscapeComponent::DeleteLayer(ULandscapeLayerInfoObject* LayerInfo, FLan
 		if (LayerIdx != DeleteLayerIdx && Allocation.WeightmapTextureIndex == DeleteLayerWeightmapTextureIndex)
 		{
 			bCanRemoveLayerTexture = false;
+			break;
 		}
 	}
 
@@ -1738,6 +1739,7 @@ void ULandscapeComponent::DeleteLayer(ULandscapeLayerInfoObject* LayerInfo, FLan
 		ComponentWeightmapTextures[DeleteLayerWeightmapTextureIndex]->Modify(LandscapeEdit.GetShouldDirtyPackage());
 		ComponentWeightmapTextures[DeleteLayerWeightmapTextureIndex]->ClearFlags(RF_Standalone);
 
+		ComponentWeightmapTexturesUsage.RemoveAt(DeleteLayerWeightmapTextureIndex);
 		ComponentWeightmapTextures.RemoveAt(DeleteLayerWeightmapTextureIndex);
 
 		// Adjust WeightmapTextureIndex index for other layers
@@ -5693,4 +5695,4 @@ FLandscapeTextureDataInfo::~FLandscapeTextureDataInfo()
 	Texture->ClearFlags(RF_Transactional);
 }
 
-#endif // WITH_EDITOR
\ No newline at end of file
+#endif // WITH_EDITOR
