Landscape post merge fixes

* Enable cached mesh draw commands that got accidentally disabled by merge
* Fix seams in multi-subsection landscapes - the subsections are treated as one big subsection now
* Small refactoring to the FLandscapeRenderSystem code

(ushell-p4-cherrypick of 9575894 by [REDACTED_PERSON])

diff --git a/Engine/Source/Runtime/Landscape/Private/LandscapeRender.cpp b/Engine/Source/Runtime/Landscape/Private/LandscapeRender.cpp
--- a/Engine/Source/Runtime/Landscape/Private/LandscapeRender.cpp
+++ b/Engine/Source/Runtime/Landscape/Private/LandscapeRender.cpp
@@ -646,6 +646,69 @@ IMPLEMENT_GLOBAL_SHADER_PARAMETER_STRUCT(FLandscapeSectionLODUniformParameters,
 TGlobalResource<FNullLandscapeRenderSystemResources> GNullLandscapeRenderSystemResources;
 TMap<FLandscapeNeighborInfo::FLandscapeKey, FLandscapeRenderSystem*> LandscapeRenderSystems;
 
+void FLandscapeRenderSystem::RegisterEntity(FLandscapeComponentSceneProxy* SceneProxy)
+{
+	check(IsInRenderingThread());
+	check(SceneProxy != nullptr);
+
+	if (NumRegisteredEntities > 0)
+	{
+		// Calculate new bounding rect of landscape components
+		FIntPoint OriginalMin = Min;
+		FIntPoint OriginalMax = Min + Size - FIntPoint(1, 1);
+		FIntPoint NewMin(FMath::Min(Min.X, SceneProxy->ComponentBase.X), FMath::Min(Min.Y, SceneProxy->ComponentBase.Y));
+		FIntPoint NewMax(FMath::Max(OriginalMax.X, SceneProxy->ComponentBase.X), FMath::Max(OriginalMax.Y, SceneProxy->ComponentBase.Y));
+
+		FIntPoint SizeRequired = (NewMax - NewMin) + FIntPoint(1, 1);
+
+		if (NewMin != Min || Size != SizeRequired)
+		{
+			ResizeAndMoveTo(NewMin, SizeRequired);
+			RecreateBuffers();
+		}
+
+		// Validate system-wide global parameters
+		check(TessellationFalloffSettings.UseTessellationComponentScreenSizeFalloff == SceneProxy->UseTessellationComponentScreenSizeFalloff);
+		check(TessellationFalloffSettings.TessellationComponentSquaredScreenSize == SceneProxy->TessellationComponentSquaredScreenSize);
+		check(TessellationFalloffSettings.TessellationComponentScreenSizeFalloff == SceneProxy->TessellationComponentScreenSizeFalloff);
+
+		if (SceneProxy->MaterialHasTessellationEnabled.Find(true) != INDEX_NONE)
+		{
+			NumEntitiesWithTessellation++;
+		}
+	}
+	else
+	{
+		TessellationFalloffSettings.UseTessellationComponentScreenSizeFalloff = SceneProxy->UseTessellationComponentScreenSizeFalloff;
+		TessellationFalloffSettings.TessellationComponentSquaredScreenSize = SceneProxy->TessellationComponentSquaredScreenSize;
+		TessellationFalloffSettings.TessellationComponentScreenSizeFalloff = SceneProxy->TessellationComponentScreenSizeFalloff;
+
+		ResizeAndMoveTo(SceneProxy->ComponentBase, FIntPoint(1, 1));
+		RecreateBuffers();
+	}
+
+	NumRegisteredEntities++;
+	SetSectionLODSettings(SceneProxy->ComponentBase, SceneProxy->LODSettings);
+	SetSectionOriginAndRadius(SceneProxy->ComponentBase, FVector4(SceneProxy->GetBounds().Origin, SceneProxy->GetBounds().SphereRadius));
+	SetSceneProxy(SceneProxy->ComponentBase, SceneProxy);
+}
+
+void FLandscapeRenderSystem::UnregisterEntity(FLandscapeComponentSceneProxy* SceneProxy)
+{
+	check(IsInRenderingThread());
+	check(SceneProxy != nullptr);
+
+	SetSceneProxy(SceneProxy->ComponentBase, nullptr);
+	SetSectionOriginAndRadius(SceneProxy->ComponentBase, FVector4(ForceInitToZero));
+
+	if (SceneProxy->MaterialHasTessellationEnabled.Find(true) != INDEX_NONE)
+	{
+		NumEntitiesWithTessellation--;
+	}
+
+	NumRegisteredEntities--;
+}
+
 void FLandscapeRenderSystem::ResizeAndMoveTo(FIntPoint NewMin, FIntPoint NewSize)
 {
 	SectionLODBuffer.SafeRelease();
@@ -677,7 +740,7 @@ void FLandscapeRenderSystem::ResizeAndMoveTo(FIntPoint NewMin, FIntPoint NewSize
 			int32 NewLinearIndex = (Y + (Min.Y - NewMin.Y)) * NewSize.X + (X + (Min.X - NewMin.X));
 			NewSectionLODValues[NewLinearIndex] = SectionLODValues[LinearIndex];
 			NewSectionLODBiases[NewLinearIndex] = SectionLODBiases[LinearIndex];
-			if (TessellationFalloffSettings.UseTessellationComponentScreenSizeFalloff)
+			if (TessellationFalloffSettings.UseTessellationComponentScreenSizeFalloff && NumEntitiesWithTessellation > 0)
 			{
 				NewSectionTessellationFalloffC[NewLinearIndex] = SectionTessellationFalloffC[LinearIndex];
 				NewSectionTessellationFalloffK[NewLinearIndex] = SectionTessellationFalloffK[LinearIndex];
@@ -698,7 +761,7 @@ void FLandscapeRenderSystem::ResizeAndMoveTo(FIntPoint NewMin, FIntPoint NewSize
 	SectionOriginAndRadius = NewSectionOriginAndRadius;
 	SceneProxies = NewSceneProxies;
 
-	if (!TessellationFalloffSettings.UseTessellationComponentScreenSizeFalloff)
+	if (!(TessellationFalloffSettings.UseTessellationComponentScreenSizeFalloff && NumEntitiesWithTessellation > 0))
 	{
 		for (float& Value : SectionTessellationFalloffC)
 		{
@@ -717,6 +780,18 @@ void FLandscapeRenderSystem::ResizeAndMoveTo(FIntPoint NewMin, FIntPoint NewSize
 	SectionTessellationFalloffK.SetAllowCPUAccess(true);
 }
 
+void FLandscapeRenderSystem::PrepareView(const FSceneView* View)
+{
+	const bool bExecuteInParallel = FApp::ShouldUseThreadingForPerformance()
+		&& GRenderingThread; // Rendering thread is required to safely use rendering resources in parallel.
+
+	if (bExecuteInParallel)
+	{
+		PerViewParametersTasks.Add(View, TGraphTask<FComputeSectionPerViewParametersTask>::CreateTask(
+			nullptr, ENamedThreads::GetRenderThread()).ConstructAndDispatchWhenReady(*this, View));
+	}
+}
+
 void FLandscapeRenderSystem::BeginRenderView(const FSceneView* View)
 {
 	TRACE_CPUPROFILER_EVENT_SCOPE(FLandscapeRenderSystem::BeginRenderView());
@@ -727,10 +802,10 @@ void FLandscapeRenderSystem::BeginRenderView(const FSceneView* View)
 		FetchHeightmapLODBiasesEventRef.SafeRelease();
 	}
 
-	if (TaskEventRef.Contains(View))
+	if (PerViewParametersTasks.Contains(View))
 	{
-		FTaskGraphInterface::Get().WaitUntilTaskCompletes(TaskEventRef[View], ENamedThreads::GetRenderThread_Local());
-		TaskEventRef.Remove(View);
+		FTaskGraphInterface::Get().WaitUntilTaskCompletes(PerViewParametersTasks[View], ENamedThreads::GetRenderThread_Local());
+		PerViewParametersTasks.Remove(View);
 	}
 	else
 	{
@@ -759,7 +834,7 @@ void FLandscapeRenderSystem::ComputeSectionPerViewParameters(
 		{
 			CachedSectionLODValues.Add(ViewPtrAsIdentifier, TResourceArray<float>{});
 
-			if (TessellationFalloffSettings.UseTessellationComponentScreenSizeFalloff)
+			if (TessellationFalloffSettings.UseTessellationComponentScreenSizeFalloff && NumEntitiesWithTessellation > 0)
 			{
 				CachedSectionTessellationFalloffC.Add(ViewPtrAsIdentifier, TResourceArray<float>{});
 				CachedSectionTessellationFalloffK.Add(ViewPtrAsIdentifier, TResourceArray<float>{});
@@ -783,7 +858,7 @@ void FLandscapeRenderSystem::ComputeSectionPerViewParameters(
 
 			SectionLODValues[EntityIndex] = ForcedLODLevel >= 0 ? ForcedLODLevel : FractionalLOD;
 
-			if (TessellationFalloffSettings.UseTessellationComponentScreenSizeFalloff)
+			if (TessellationFalloffSettings.UseTessellationComponentScreenSizeFalloff && NumEntitiesWithTessellation > 0)
 			{
 				float MaxTesselationDistance = ComputeBoundsDrawDistance(FMath::Sqrt(TessellationFalloffSettings.TessellationComponentSquaredScreenSize), SectionOriginAndRadius[EntityIndex].W / 2.0f, ViewProjectionMarix);
 				float FallOffStartingDistance = FMath::Min(
@@ -802,7 +877,7 @@ void FLandscapeRenderSystem::ComputeSectionPerViewParameters(
 
 			CachedSectionLODValues[ViewPtrAsIdentifier] = SectionLODValues;
 
-			if (TessellationFalloffSettings.UseTessellationComponentScreenSizeFalloff)
+			if (TessellationFalloffSettings.UseTessellationComponentScreenSizeFalloff && NumEntitiesWithTessellation > 0)
 			{
 				CachedSectionTessellationFalloffC[ViewPtrAsIdentifier] = SectionTessellationFalloffC;
 				CachedSectionTessellationFalloffK[ViewPtrAsIdentifier] = SectionTessellationFalloffK;
@@ -816,7 +891,7 @@ void FLandscapeRenderSystem::ComputeSectionPerViewParameters(
 
 			SectionLODValues = CachedSectionLODValues[ViewPtrAsIdentifier];
 
-			if (TessellationFalloffSettings.UseTessellationComponentScreenSizeFalloff)
+			if (TessellationFalloffSettings.UseTessellationComponentScreenSizeFalloff && NumEntitiesWithTessellation > 0)
 			{
 				SectionTessellationFalloffC = CachedSectionTessellationFalloffC[ViewPtrAsIdentifier];
 				SectionTessellationFalloffK = CachedSectionTessellationFalloffK[ViewPtrAsIdentifier];
@@ -893,7 +968,7 @@ void FLandscapeRenderSystem::RecreateBuffers(const FSceneView* InView /* = nullp
 			else
 			{
 				// If we use tessellation falloff, update the buffer, otherwise use the one already filled with default parameters
-				if (TessellationFalloffSettings.UseTessellationComponentScreenSizeFalloff)
+				if (TessellationFalloffSettings.UseTessellationComponentScreenSizeFalloff && NumEntitiesWithTessellation > 0)
 				{
 					float* Data = (float*)RHILockVertexBuffer(SectionTessellationFalloffCBuffer, 0, SectionTessellationFalloffC.GetResourceDataSize(), RLM_WriteOnly);
 					FMemory::Memcpy(Data, SectionTessellationFalloffC.GetData(), SectionTessellationFalloffC.GetResourceDataSize());
@@ -912,7 +987,7 @@ void FLandscapeRenderSystem::RecreateBuffers(const FSceneView* InView /* = nullp
 			else
 			{
 				// If we use tessellation falloff, update the buffer, otherwise use the one already filled with default parameters
-				if (TessellationFalloffSettings.UseTessellationComponentScreenSizeFalloff)
+				if (TessellationFalloffSettings.UseTessellationComponentScreenSizeFalloff && NumEntitiesWithTessellation > 0)
 				{
 					float* Data = (float*)RHILockVertexBuffer(SectionTessellationFalloffKBuffer, 0, SectionTessellationFalloffK.GetResourceDataSize(), RLM_WriteOnly);
 					FMemory::Memcpy(Data, SectionTessellationFalloffK.GetData(), SectionTessellationFalloffK.GetResourceDataSize());
@@ -948,7 +1023,7 @@ void FLandscapeRenderSystem::BeginFrame()
 
 	CachedSectionLODValues.Empty();
 
-	if (TessellationFalloffSettings.UseTessellationComponentScreenSizeFalloff)
+	if (TessellationFalloffSettings.UseTessellationComponentScreenSizeFalloff && NumEntitiesWithTessellation > 0)
 	{
 		CachedSectionTessellationFalloffC.Empty();
 		CachedSectionTessellationFalloffK.Empty();
@@ -978,13 +1053,13 @@ void FLandscapeRenderSystem::EndFrame()
 		FetchHeightmapLODBiasesEventRef.SafeRelease();
 	}
 
-	for (auto& Pair : TaskEventRef)
+	for (auto& Pair : PerViewParametersTasks)
 	{
 		const FSceneView* View = Pair.Key;
-		FTaskGraphInterface::Get().WaitUntilTaskCompletes(TaskEventRef[View], ENamedThreads::GetRenderThread_Local());
+		FTaskGraphInterface::Get().WaitUntilTaskCompletes(PerViewParametersTasks[View], ENamedThreads::GetRenderThread_Local());
 	}
 
-	TaskEventRef.Empty();
+	PerViewParametersTasks.Empty();
 }
 
 class FLandscapePersistentViewUniformBufferExtension : public IPersistentViewUniformBufferExtension
@@ -1320,9 +1395,7 @@ void FLandscapeComponentSceneProxy::CreateRenderThreadResources()
 
 	if (IsComponentLevelVisible())
 	{
-		OriginAndSphereRadius = FVector4(GetBounds().Origin, GetBounds().SphereRadius);
-		LandscapeSceneProxy = this;
-		RegisterNeighbors();
+		RegisterNeighbors(this);
 	}
 
 	auto FeatureLevel = GetScene().GetFeatureLevel();
@@ -1487,14 +1560,12 @@ void FLandscapeComponentSceneProxy::CreateRenderThreadResources()
 void FLandscapeComponentSceneProxy::DestroyRenderThreadResources()
 {
 	FPrimitiveSceneProxy::DestroyRenderThreadResources();
-	UnregisterNeighbors();
+	UnregisterNeighbors(this);
 }
 
 void FLandscapeComponentSceneProxy::OnLevelAddedToWorld()
 {
-	OriginAndSphereRadius = FVector4(GetBounds().Origin, GetBounds().SphereRadius);
-	LandscapeSceneProxy = this;
-	RegisterNeighbors();
+	RegisterNeighbors(this);
 }
 
 FLandscapeComponentSceneProxy::~FLandscapeComponentSceneProxy()
@@ -1770,6 +1841,7 @@ void FLandscapeComponentSceneProxy::OnTransformChanged()
 	LandscapeParams.ComponentBaseX = ComponentBase.X;
 	LandscapeParams.ComponentBaseY = ComponentBase.Y;
 	LandscapeParams.SubsectionSizeVerts = SubsectionSizeVerts;
+	LandscapeParams.NumSubsections = NumSubsections;
 	LandscapeParams.LastLOD = LastLOD;
 	LandscapeParams.HeightmapUVScaleBias = HeightmapScaleBias;
 	LandscapeParams.WeightmapUVScaleBias = WeightmapScaleBias;
@@ -1836,7 +1908,7 @@ void FLandscapeComponentSceneProxy::OnTransformChanged()
 
 	if (bRegistered)
 	{
-		OriginAndSphereRadius = FVector4(GetBounds().Origin, GetBounds().SphereRadius);
+		FVector4 OriginAndSphereRadius(GetBounds().Origin, GetBounds().SphereRadius);
 
 		FLandscapeRenderSystem& RenderSystem = *LandscapeRenderSystems.FindChecked(LandscapeKey);
 		RenderSystem.SetSectionOriginAndRadius(ComponentBase, OriginAndSphereRadius);
@@ -4320,111 +4392,6 @@ public:
 		{
 			ShaderBindings.Add(ForcedLodParameter, BatchElementParams->ForcedLOD);
 		}
-
-		int32 SubSectionIndex = BatchElementParams->SubX + BatchElementParams->SubY * SceneProxy->NumSubsections;
-		// If we have no custom data for this primitive we will compute of the fly the proper values, this will happen if the shader is not used for normal landscape rendering(i.e grass rendering)
-		FLandscapeComponentSceneProxy::FViewCustomDataLOD* LODData = (FLandscapeComponentSceneProxy::FViewCustomDataLOD*)InView->GetCustomData(SceneProxy->GetPrimitiveSceneInfo()->GetIndex());
-
-		if (LODData != nullptr)
-		{
-			SceneProxy->PostInitViewCustomData(*InView, LODData);
-
-			if (LodTessellationParameter.IsBound())
-			{
-				ShaderBindings.Add(LodTessellationParameter, LODData->LodTessellationParams);
-			}
-				
-			if (LodBiasParameter.IsBound())
-			{
-				check(LODData->LodBias == SceneProxy->GetShaderLODBias());
-				ShaderBindings.Add(LodBiasParameter, LODData->LodBias);
-			}
-
-			if (SectionLodsParameter.IsBound())
-			{
-				if (LODData->UseCombinedMeshBatch)
-				{
-					ShaderBindings.Add(SectionLodsParameter, LODData->ShaderCurrentLOD);
-				}
-				else // in non combined, only the one representing us as we'll be called 4 times (once per sub section)
-				{
-					check(SubSectionIndex >= 0);
-					FVector4 ShaderCurrentLOD(ForceInitToZero);
-					ShaderCurrentLOD.Component(SubSectionIndex) = LODData->ShaderCurrentLOD.Component(SubSectionIndex);
-
-					ShaderBindings.Add(SectionLodsParameter, ShaderCurrentLOD);
-				}
-			}
-
-			if (NeighborSectionLodParameter.IsBound())
-			{
-				FVector4 ShaderCurrentNeighborLOD[FLandscapeComponentSceneProxy::NEIGHBOR_COUNT] = { FVector4(ForceInitToZero), FVector4(ForceInitToZero), FVector4(ForceInitToZero), FVector4(ForceInitToZero) };
-
-				if (LODData->UseCombinedMeshBatch)
-				{					
-					int32 SubSectionCount = SceneProxy->NumSubsections == 1 ? 1 : FLandscapeComponentSceneProxy::MAX_SUBSECTION_COUNT;
-
-					for (int32 NeighborSubSectionIndex = 0; NeighborSubSectionIndex < SubSectionCount; ++NeighborSubSectionIndex)
-					{
-						ShaderCurrentNeighborLOD[NeighborSubSectionIndex] = LODData->SubSections[NeighborSubSectionIndex].ShaderCurrentNeighborLOD;
-						check(ShaderCurrentNeighborLOD[NeighborSubSectionIndex].X != -1.0f); // they should all match so only check the 1st one for simplicity
-					}
-
-					ShaderBindings.Add(NeighborSectionLodParameter, ShaderCurrentNeighborLOD);
-				}
-				else // in non combined, only the one representing us as we'll be called 4 times (once per sub section)
-				{
-					check(SubSectionIndex >= 0);
-					ShaderCurrentNeighborLOD[SubSectionIndex] = LODData->SubSections[SubSectionIndex].ShaderCurrentNeighborLOD;
-					check(ShaderCurrentNeighborLOD[SubSectionIndex].X != -1.0f); // they should all match so only check the 1st one for simplicity
-
-					ShaderBindings.Add(NeighborSectionLodParameter, ShaderCurrentNeighborLOD);
-				}
-			}			
-		}
-		else
-		{
-			float ComponentScreenSize = SceneProxy->GetComponentScreenSize(InView, SceneProxy->GetBounds().Origin, SceneProxy->ComponentMaxExtend, SceneProxy->GetBounds().SphereRadius);
-			
-			FLandscapeComponentSceneProxy::FViewCustomDataLOD CurrentLODData;
-			SceneProxy->CalculateBatchElementLOD(*InView, ComponentScreenSize, InView->LODDistanceFactor, CurrentLODData, true);
-			check(CurrentLODData.UseCombinedMeshBatch);
-
-			if (LodBiasParameter.IsBound())
-			{
-				ShaderBindings.Add(LodBiasParameter, SceneProxy->GetShaderLODBias());
-			}
-
-			if (LodTessellationParameter.IsBound())
-			{
-				FVector4 LodTessellationParams(ForceInitToZero);
-				SceneProxy->ComputeTessellationFalloffShaderValues(CurrentLODData, InView->ViewMatrices.GetProjectionMatrix(), LodTessellationParams.X, LodTessellationParams.Y);
-
-				ShaderBindings.Add(LodTessellationParameter, LodTessellationParams);
-			}
-
-			if (SectionLodsParameter.IsBound())
-			{
-				ShaderBindings.Add(SectionLodsParameter, CurrentLODData.ShaderCurrentLOD);
-			}
-
-			if (NeighborSectionLodParameter.IsBound())
-			{
-				FVector4 CurrentNeighborLOD[4] = { FVector4(ForceInitToZero), FVector4(ForceInitToZero), FVector4(ForceInitToZero), FVector4(ForceInitToZero) };
-
-				for (int32 SubY = 0; SubY < SceneProxy->NumSubsections; SubY++)
-				{
-					for (int32 SubX = 0; SubX < SceneProxy->NumSubsections; SubX++)
-					{
-						int32 NeighborSubSectionIndex = SubX + SubY * SceneProxy->NumSubsections;
-						SceneProxy->GetShaderCurrentNeighborLOD(*InView, CurrentLODData.SubSections[NeighborSubSectionIndex].fBatchElementCurrentLOD, SceneProxy->NumSubsections > 1 ? SubX : INDEX_NONE, SceneProxy->NumSubsections > 1 ? SubY : INDEX_NONE, NeighborSubSectionIndex, CurrentNeighborLOD[NeighborSubSectionIndex]);
-						check(CurrentNeighborLOD[NeighborSubSectionIndex].X != -1.0f); // they should all match so only check the 1st one for simplicity
-					}
-				}
-
-				ShaderBindings.Add(NeighborSectionLodParameter, CurrentNeighborLOD);
-			}				
-		}
 	}
 
 	virtual uint32 GetSize() const override
@@ -4558,7 +4525,7 @@ void FLandscapeVertexFactory::ModifyCompilationEnvironment(const FVertexFactoryT
 }
 
 
-IMPLEMENT_VERTEX_FACTORY_TYPE(FLandscapeVertexFactory, "/Engine/Private/LandscapeVertexFactory.ush", true, true, true, false, false);
+IMPLEMENT_VERTEX_FACTORY_TYPE_EX(FLandscapeVertexFactory, "/Engine/Private/LandscapeVertexFactory.ush", true, true, true, false, false, true, false);
 
 /**
 * Copy the data from another vertex factory
@@ -4587,7 +4554,7 @@ void FLandscapeXYOffsetVertexFactory::ModifyCompilationEnvironment(const FVertex
 	OutEnvironment.SetDefine(TEXT("LANDSCAPE_XYOFFSET"), TEXT("1"));
 }
 
-IMPLEMENT_VERTEX_FACTORY_TYPE(FLandscapeXYOffsetVertexFactory, "/Engine/Private/LandscapeVertexFactory.ush", true, true, true, false, false);
+IMPLEMENT_VERTEX_FACTORY_TYPE_EX(FLandscapeXYOffsetVertexFactory, "/Engine/Private/LandscapeVertexFactory.ush", true, true, true, false, false, true, false);
 
 //
 // FLandscapeFixedGridVertexFactory
@@ -5408,7 +5375,7 @@ void FLandscapeComponentSceneProxy::GetLCIs(FLCIArray& LCIs)
 //
 // FLandscapeNeighborInfo
 //
-void FLandscapeNeighborInfo::RegisterNeighbors()
+void FLandscapeNeighborInfo::RegisterNeighbors(FLandscapeComponentSceneProxy* SceneProxy /* = nullptr */)
 {
 	check(IsInRenderingThread());
 	if (!bRegistered)
@@ -5453,43 +5420,10 @@ void FLandscapeNeighborInfo::RegisterNeighbors()
 				Neighbors[3]->Neighbors[0] = this;
 			}
 
+			if (SceneProxy != nullptr)
 			{
 				FLandscapeRenderSystem& RenderSystem = *LandscapeRenderSystems.FindChecked(LandscapeKey);
-
-				if (RenderSystem.NumRegisteredEntities > 0)
-				{
-					FIntPoint OriginalMin = RenderSystem.Min;
-					FIntPoint OriginalMax = RenderSystem.Min + RenderSystem.Size - FIntPoint(1, 1);
-					FIntPoint NewMin(FMath::Min(RenderSystem.Min.X, ComponentBase.X), FMath::Min(RenderSystem.Min.Y, ComponentBase.Y));
-					FIntPoint NewMax(FMath::Max(OriginalMax.X, ComponentBase.X), FMath::Max(OriginalMax.Y, ComponentBase.Y));
-
-					FIntPoint SizeRequired = (NewMax - NewMin) + FIntPoint(1, 1);
-
-					if (NewMin != RenderSystem.Min || RenderSystem.Size != SizeRequired)
-					{
-						RenderSystem.ResizeAndMoveTo(NewMin, SizeRequired);
-						RenderSystem.RecreateBuffers();
-					}
-
-					// Validate system-wide global parameters
-					check(RenderSystem.TessellationFalloffSettings.UseTessellationComponentScreenSizeFalloff == LandscapeSceneProxy->UseTessellationComponentScreenSizeFalloff);
-					check(RenderSystem.TessellationFalloffSettings.TessellationComponentSquaredScreenSize == LandscapeSceneProxy->TessellationComponentSquaredScreenSize);
-					check(RenderSystem.TessellationFalloffSettings.TessellationComponentScreenSizeFalloff == LandscapeSceneProxy->TessellationComponentScreenSizeFalloff);
-				}
-				else
-				{
-					RenderSystem.TessellationFalloffSettings.UseTessellationComponentScreenSizeFalloff = LandscapeSceneProxy->UseTessellationComponentScreenSizeFalloff;
-					RenderSystem.TessellationFalloffSettings.TessellationComponentSquaredScreenSize = LandscapeSceneProxy->TessellationComponentSquaredScreenSize;
-					RenderSystem.TessellationFalloffSettings.TessellationComponentScreenSizeFalloff = LandscapeSceneProxy->TessellationComponentScreenSizeFalloff;
-
-					RenderSystem.ResizeAndMoveTo(ComponentBase, FIntPoint(1, 1));
-					RenderSystem.RecreateBuffers();
-				}
-
-				RenderSystem.NumRegisteredEntities++;
-				RenderSystem.SetSectionLODSettings(ComponentBase, LODSettings);
-				RenderSystem.SetSectionOriginAndRadius(ComponentBase, OriginAndSphereRadius);
-				RenderSystem.SetSceneProxy(ComponentBase, LandscapeSceneProxy);
+				RenderSystem.RegisterEntity(SceneProxy);
 			}
 		}
 		else
@@ -5499,7 +5433,7 @@ void FLandscapeNeighborInfo::RegisterNeighbors()
 	}
 }
 
-void FLandscapeNeighborInfo::UnregisterNeighbors()
+void FLandscapeNeighborInfo::UnregisterNeighbors(FLandscapeComponentSceneProxy* SceneProxy /* = nullptr */)
 {
 	check(IsInRenderingThread());
 	
@@ -5514,6 +5448,12 @@ void FLandscapeNeighborInfo::UnregisterNeighbors()
 		{
 			SceneProxyMap->Remove(ComponentBase);
 
+			if (SceneProxy != nullptr)
+			{
+				FLandscapeRenderSystem& RenderSystem = *LandscapeRenderSystems.FindChecked(LandscapeKey);
+				RenderSystem.UnregisterEntity(SceneProxy);
+			}
+
 			if (SceneProxyMap->Num() == 0)
 			{
 				// remove the entire LandscapeKey entry as this is the last scene proxy
@@ -5525,11 +5465,6 @@ void FLandscapeNeighborInfo::UnregisterNeighbors()
 			}
 			else
 			{
-				FLandscapeRenderSystem* RenderSystem = LandscapeRenderSystems.FindChecked(LandscapeKey);
-				RenderSystem->SetSceneProxy(ComponentBase, nullptr);
-				RenderSystem->SetSectionOriginAndRadius(ComponentBase, FVector4(ForceInitToZero));
-				RenderSystem->NumRegisteredEntities--;
-
 				// remove reference to us from our neighbors
 				if (Neighbors[0])
 				{
diff --git a/Engine/Source/Runtime/Landscape/Private/LandscapeRenderMobile.cpp b/Engine/Source/Runtime/Landscape/Private/LandscapeRenderMobile.cpp
--- a/Engine/Source/Runtime/Landscape/Private/LandscapeRenderMobile.cpp
+++ b/Engine/Source/Runtime/Landscape/Private/LandscapeRenderMobile.cpp
@@ -345,9 +345,7 @@ void FLandscapeComponentSceneProxyMobile::CreateRenderThreadResources()
 
 	if (IsComponentLevelVisible())
 	{
-		OriginAndSphereRadius = FVector4(GetBounds().Origin, GetBounds().SphereRadius);
-		LandscapeSceneProxy = this;
-		RegisterNeighbors();
+		RegisterNeighbors(this);
 	}
 	
 	auto FeatureLevel = GetScene().GetFeatureLevel();
diff --git a/Engine/Source/Runtime/Landscape/Public/LandscapeRender.h b/Engine/Source/Runtime/Landscape/Public/LandscapeRender.h
--- a/Engine/Source/Runtime/Landscape/Public/LandscapeRender.h
+++ b/Engine/Source/Runtime/Landscape/Public/LandscapeRender.h
@@ -92,6 +92,7 @@ BEGIN_GLOBAL_SHADER_PARAMETER_STRUCT(FLandscapeUniformShaderParameters, LANDSCAP
 SHADER_PARAMETER(int32, ComponentBaseX)
 SHADER_PARAMETER(int32, ComponentBaseY)
 SHADER_PARAMETER(int32, SubsectionSizeVerts)
+SHADER_PARAMETER(int32, NumSubsections)
 SHADER_PARAMETER(int32, LastLOD)
 SHADER_PARAMETER(FVector4, HeightmapUVScaleBias)
 SHADER_PARAMETER(FVector4, WeightmapUVScaleBias)
@@ -381,15 +382,6 @@ public:
 // FLandscapeNeighborInfo
 //
 
-struct LODSettingsComponent
-{
-	float LOD0ScreenSizeSquared;
-	float LOD1ScreenSizeSquared;
-	float LODOnePlusDistributionScalarSquared;
-	float LastLODScreenSizeSquared;
-	int8 LastLODIndex;
-};
-
 class FLandscapeNeighborInfo
 {
 public:
@@ -450,10 +442,6 @@ protected:
 	bool					bRegistered;
 	int32					PrimitiveCustomDataIndex;
 
-	LODSettingsComponent LODSettings;
-	FVector4 OriginAndSphereRadius;
-	FLandscapeComponentSceneProxy* LandscapeSceneProxy;
-
 	friend class FLandscapeComponentSceneProxy;
 
 public:
@@ -465,7 +453,6 @@ public:
 	, LODBias(InLODBias)
 	, bRegistered(false)
 	, PrimitiveCustomDataIndex(INDEX_NONE)
-	, LandscapeSceneProxy(nullptr)
 	{
 		//       -Y       
 		//    - - 0 - -   
@@ -481,8 +468,8 @@ public:
 		Neighbors[3] = nullptr;
 	}
 
-	void RegisterNeighbors();
-	void UnregisterNeighbors();
+	void RegisterNeighbors(FLandscapeComponentSceneProxy* SceneProxy = nullptr);
+	void UnregisterNeighbors(FLandscapeComponentSceneProxy* SceneProxy = nullptr);
 };
 
 
@@ -525,6 +512,15 @@ extern RENDERER_API TAutoConsoleVariable<float> CVarStaticMeshLODDistanceScale;
 
 struct FLandscapeRenderSystem
 {
+	struct LODSettingsComponent
+	{
+		float LOD0ScreenSizeSquared;
+		float LOD1ScreenSizeSquared;
+		float LODOnePlusDistributionScalarSquared;
+		float LastLODScreenSizeSquared;
+		int8 LastLODIndex;
+	};
+
 	static int8 GetLODFromScreenSize(LODSettingsComponent LODSettings, float InScreenSizeSquared, float InViewLODScale, float& OutFractionalLOD)
 	{
 		float ScreenSizeSquared = InScreenSizeSquared / InViewLODScale;
@@ -548,6 +544,7 @@ struct FLandscapeRenderSystem
 	}
 
 	int32 NumRegisteredEntities;
+	int32 NumEntitiesWithTessellation;
 
 	FIntPoint Min;
 	FIntPoint Size;
@@ -584,7 +581,7 @@ struct FLandscapeRenderSystem
 	TMap<const FSceneView*, TResourceArray<float>> CachedSectionTessellationFalloffK;
 	const FSceneView* CachedView;
 
-	TMap<const FSceneView*, FGraphEventRef> TaskEventRef;
+	TMap<const FSceneView*, FGraphEventRef> PerViewParametersTasks;
 	FGraphEventRef FetchHeightmapLODBiasesEventRef;
 
 	struct FComputeSectionPerViewParametersTask
@@ -669,6 +666,7 @@ struct FLandscapeRenderSystem
 
 	FLandscapeRenderSystem()
 		: NumRegisteredEntities(0)
+		, NumEntitiesWithTessellation(0)
 		, Min(MAX_int32, MAX_int32)
 		, Size(EForceInit::ForceInitToZero)
 		, CachedView(nullptr)
@@ -679,6 +677,10 @@ struct FLandscapeRenderSystem
 		SectionTessellationFalloffK.SetAllowCPUAccess(true);
 	}
 
+	void RegisterEntity(FLandscapeComponentSceneProxy* SceneProxy);
+
+	void UnregisterEntity(FLandscapeComponentSceneProxy* SceneProxy);
+
 	int32 GetComponentLinearIndex(FIntPoint ComponentBase)
 	{
 		return (ComponentBase.Y - Min.Y) * Size.X + ComponentBase.X - Min.X;
@@ -717,17 +719,7 @@ struct FLandscapeRenderSystem
 		FVector ViewOrigin,
 		FMatrix ViewProjectionMarix);
 
-	void PrepareView(const FSceneView* View)
-	{
-		const bool bExecuteInParallel = FApp::ShouldUseThreadingForPerformance()
-			&& GRenderingThread; // Rendering thread is required to safely use rendering resources in parallel.
-
-		if (bExecuteInParallel)
-		{
-			TaskEventRef.Add(View, TGraphTask<FComputeSectionPerViewParametersTask>::CreateTask(
-				nullptr, ENamedThreads::GetRenderThread()).ConstructAndDispatchWhenReady(*this, View));
-		}
-	}
+	void PrepareView(const FSceneView* View);
 
 	void BeginRenderView(const FSceneView* View);
 
@@ -859,7 +851,7 @@ public:
 	TStaticArray<FLandscapeSectionRayTracingState, MAX_SUBSECTION_COUNT> SectionRayTracingStates;
 #endif
 
-	friend FLandscapeNeighborInfo;
+	friend FLandscapeRenderSystem;
 
 protected:
 	int8						MaxLOD;		// Maximum LOD level, user override possible
@@ -880,6 +872,8 @@ protected:
 	float						TessellationComponentSquaredScreenSize;	// Screen size of the component at which we start to apply tessellation
 	float						TessellationComponentScreenSizeFalloff;	// Min Component screen size before we start applying the tessellation falloff
 
+	FLandscapeRenderSystem::LODSettingsComponent LODSettings;
+
 	/** 
 	 * Number of subsections within the component in each dimension, this can be 1 or 2.
 	 * Subsections exist to improve the speed at which LOD transitions can take place over distance.
diff --git a/Engine/Source/Runtime/Renderer/Private/RendererScene.cpp b/Engine/Source/Runtime/Renderer/Private/RendererScene.cpp
--- a/Engine/Source/Runtime/Renderer/Private/RendererScene.cpp
+++ b/Engine/Source/Runtime/Renderer/Private/RendererScene.cpp
@@ -3341,6 +3341,7 @@ void FScene::OnLevelAddedToWorld(FName LevelAddedName, UWorld* InWorld, bool bIs
 	ENQUEUE_RENDER_COMMAND(FLevelAddedToWorld)(
 		[Scene, LevelAddedName](FRHICommandListImmediate& RHICmdList)
 		{
+			Scene->UpdateAllPrimitiveSceneInfos(RHICmdList);
 			Scene->OnLevelAddedToWorld_RenderThread(LevelAddedName);
 		});
 }
