Fix UpdateGrass when GrassTypes is empty

#rnx

diff --git a/Engine/Source/Runtime/Landscape/Private/Landscape.cpp b/Engine/Source/Runtime/Landscape/Private/Landscape.cpp
--- a/Engine/Source/Runtime/Landscape/Private/Landscape.cpp
+++ b/Engine/Source/Runtime/Landscape/Private/Landscape.cpp
@@ -1056,6 +1056,10 @@ ALandscapeProxy::ALandscapeProxy(const FObjectInitializer& ObjectInitializer)
 #endif
 
 #if WITH_EDITOR
+	NumComponentsNeedingGrassMapRender = 0;
+	NumTexturesToStreamForVisibleGrassMapRender = 0;
+	NumComponentsNeedingTextureBaking = 0;
+
 	if (VisibilityLayer == nullptr)
 	{
 		// Structure to hold one-time initialization
diff --git a/Engine/Source/Runtime/Landscape/Private/LandscapeGrass.cpp b/Engine/Source/Runtime/Landscape/Private/LandscapeGrass.cpp
--- a/Engine/Source/Runtime/Landscape/Private/LandscapeGrass.cpp
+++ b/Engine/Source/Runtime/Landscape/Private/LandscapeGrass.cpp
@@ -2161,6 +2161,7 @@ void ALandscapeProxy::UpdateGrassDataStatus(TSet<UTexture2D*>& OutCurrentForcedS
 
 	// In either case we want to check the Grass textures stream state
 	const bool bCheckStreamingState = OutDesiredForcedStreamedTextures || bInEnableForceResidentFlag;
+	const bool bHasGrassTypes = GetGrassTypes().Num();
 
 	for (auto Component : LandscapeComponents)
 	{
@@ -2187,32 +2188,35 @@ void ALandscapeProxy::UpdateGrassDataStatus(TSet<UTexture2D*>& OutCurrentForcedS
 				OutOutdatedComponents->Add(Component);
 			}
 
-			if (Component->IsGrassMapOutdated() || !Component->GrassData->HasData())
+			if (bHasGrassTypes || bBakeMaterialPositionOffsetIntoCollision)
 			{
-				OutComponentsNeedingGrassMapRender.Add(Component);
-
-				if (bCheckStreamingState && !Component->AreTexturesStreamedForGrassMapRender())
+				if (Component->IsGrassMapOutdated() || !Component->GrassData->HasData())
 				{
-					if (OutDesiredForcedStreamedTextures)
-					{
-						OutDesiredForcedStreamedTextures->Add(Heightmap);
-					}
+					OutComponentsNeedingGrassMapRender.Add(Component);
 
-					if (bInEnableForceResidentFlag)
-					{
-						Heightmap->bForceMiplevelsToBeResident = true;
-					}
-				
-					for (auto WeightmapTexture : ComponentWeightmapTextures)
+					if (bCheckStreamingState && !Component->AreTexturesStreamedForGrassMapRender())
 					{
 						if (OutDesiredForcedStreamedTextures)
 						{
-							OutDesiredForcedStreamedTextures->Add(WeightmapTexture);
+							OutDesiredForcedStreamedTextures->Add(Heightmap);
 						}
-						
+
 						if (bInEnableForceResidentFlag)
 						{
-							WeightmapTexture->bForceMiplevelsToBeResident = true;
+							Heightmap->bForceMiplevelsToBeResident = true;
+						}
+
+						for (auto WeightmapTexture : ComponentWeightmapTextures)
+						{
+							if (OutDesiredForcedStreamedTextures)
+							{
+								OutDesiredForcedStreamedTextures->Add(WeightmapTexture);
+							}
+
+							if (bInEnableForceResidentFlag)
+							{
+								WeightmapTexture->bForceMiplevelsToBeResident = true;
+							}
 						}
 					}
 				}
@@ -2234,12 +2238,6 @@ void ALandscapeProxy::UpdateGrassData()
 		return;
 	}
 
-	TArray<ULandscapeGrassType*> GrassTypes = GetGrassTypes();
-	if (GrassTypes.Num() == 0 && !bBakeMaterialPositionOffsetIntoCollision)
-	{
-		return;
-	}
-
 	// see if we need to flush grass for any components
 	TSet<UTexture2D*> DesiredForcedStreamedTextures;
 	TSet<UTexture2D*> CurrentForcedStreamedTextures;
