Fix VS2019 Strict

#rnx

diff --git a/Engine/Source/Runtime/Landscape/Private/LandscapeGrass.cpp b/Engine/Source/Runtime/Landscape/Private/LandscapeGrass.cpp
--- a/Engine/Source/Runtime/Landscape/Private/LandscapeGrass.cpp
+++ b/Engine/Source/Runtime/Landscape/Private/LandscapeGrass.cpp
@@ -2161,7 +2161,7 @@ void ALandscapeProxy::UpdateGrassDataStatus(TSet<UTexture2D*>& OutCurrentForcedS
 
 	// In either case we want to check the Grass textures stream state
 	const bool bCheckStreamingState = OutDesiredForcedStreamedTextures || bInEnableForceResidentFlag;
-	const bool bHasGrassTypes = GetGrassTypes().Num();
+	const bool bHasGrassTypes = GetGrassTypes().Num() > 0;
 
 	for (auto Component : LandscapeComponents)
 	{
