Fix landscape simple collision returning incorrect physical material.

diff --git a/Engine/Source/Runtime/Landscape/Private/LandscapeCollision.cpp b/Engine/Source/Runtime/Landscape/Private/LandscapeCollision.cpp
--- a/Engine/Source/Runtime/Landscape/Private/LandscapeCollision.cpp
+++ b/Engine/Source/Runtime/Landscape/Private/LandscapeCollision.cpp
@@ -836,7 +836,7 @@ bool ULandscapeHeightfieldCollisionComponent::CookCollisionData(const FName& For
 
 	if (bGenerateSimpleCollision)
 	{
-		SimpleSamples = ConvertHeightfieldDataForPhysx(this, SimpleCollisionSizeVerts, bIsMirrored, Heights + NumSamples, bUseDefMaterial, DominantLayers, DefMaterial, InOutMaterials);
+		SimpleSamples = ConvertHeightfieldDataForPhysx(this, SimpleCollisionSizeVerts, bIsMirrored, Heights + NumSamples, bUseDefMaterial, DominantLayers + NumSamples, DefMaterial, InOutMaterials);
 	}
 
 	CollisionHeightData.Unlock();
