Fix race condition of landscape multi view LOD calculation

(ushell-p4-cherrypick of 9668425 by [REDACTED_PERSON])

diff --git a/Engine/Source/Runtime/Landscape/Private/LandscapeRender.cpp b/Engine/Source/Runtime/Landscape/Private/LandscapeRender.cpp
--- a/Engine/Source/Runtime/Landscape/Private/LandscapeRender.cpp
+++ b/Engine/Source/Runtime/Landscape/Private/LandscapeRender.cpp
@@ -790,6 +790,11 @@ void FLandscapeRenderSystem::PrepareView(const FSceneView* View)
 		PerViewParametersTasks.Add(View, TGraphTask<FComputeSectionPerViewParametersTask>::CreateTask(
 			nullptr, ENamedThreads::GetRenderThread()).ConstructAndDispatchWhenReady(*this, View));
 	}
+	else
+	{
+		FComputeSectionPerViewParametersTask Task(*this, View);
+		Task.AnyThreadTask();
+	}
 }
 
 void FLandscapeRenderSystem::BeginRenderView(const FSceneView* View)
@@ -807,10 +812,17 @@ void FLandscapeRenderSystem::BeginRenderView(const FSceneView* View)
 		FTaskGraphInterface::Get().WaitUntilTaskCompletes(PerViewParametersTasks[View], ENamedThreads::GetRenderThread_Local());
 		PerViewParametersTasks.Remove(View);
 	}
-	else
+
 	{
-		FComputeSectionPerViewParametersTask Task(*this, View);
-		Task.AnyThreadTask();
+		FScopeLock Lock(&CachedValuesCS);
+
+		SectionLODValues = CachedSectionLODValues[View];
+
+		if (TessellationFalloffSettings.UseTessellationComponentScreenSizeFalloff && NumEntitiesWithTessellation > 0)
+		{
+			SectionTessellationFalloffC = CachedSectionTessellationFalloffC[View];
+			SectionTessellationFalloffC = CachedSectionTessellationFalloffK[View];
+		}
 	}
 
 	RecreateBuffers(View);
@@ -824,78 +836,55 @@ void FLandscapeRenderSystem::ComputeSectionPerViewParameters(
 	FMatrix ViewProjectionMarix
 )
 {
-	bool bValuesCached = false;
-
-	{
-		FScopeLock Lock(&CachedValuesCS);
-		bValuesCached = CachedSectionLODValues.Contains(ViewPtrAsIdentifier);
+	TRACE_CPUPROFILER_EVENT_SCOPE(FLandscapeRenderSystem::ComputeSectionPerViewParameters());
 
-		if (!bValuesCached)
-		{
-			CachedSectionLODValues.Add(ViewPtrAsIdentifier, TResourceArray<float>{});
+	TResourceArray<float> NewSectionLODValues;
+	TResourceArray<float> NewSectionTessellationFalloffC;
+	TResourceArray<float> NewSectionTessellationFalloffK;
 
-			if (TessellationFalloffSettings.UseTessellationComponentScreenSizeFalloff && NumEntitiesWithTessellation > 0)
-			{
-				CachedSectionTessellationFalloffC.Add(ViewPtrAsIdentifier, TResourceArray<float>{});
-				CachedSectionTessellationFalloffK.Add(ViewPtrAsIdentifier, TResourceArray<float>{});
-			}
-		}
-	}
+	NewSectionLODValues.AddZeroed(SectionLODSettings.Num());
 
-	if (!bValuesCached)
+	if (TessellationFalloffSettings.UseTessellationComponentScreenSizeFalloff && NumEntitiesWithTessellation > 0)
 	{
-		TRACE_CPUPROFILER_EVENT_SCOPE(FLandscapeRenderSystem::ComputeSectionPerViewParameters());
-
-		int32 ForcedLODLevel = ViewEngineShowFlagLOD ? GetCVarForceLOD() : -1;
-		float LODScale = ViewLODDistanceFactor * CVarStaticMeshLODDistanceScale.GetValueOnRenderThread();
+		NewSectionTessellationFalloffC.AddZeroed(SectionLODSettings.Num());
+		NewSectionTessellationFalloffK.AddZeroed(SectionLODSettings.Num());
+	}
 
-		for (int32 EntityIndex = 0; EntityIndex < SectionLODSettings.Num(); EntityIndex++)
-		{
-			float MeshScreenSizeSquared = ComputeBoundsScreenRadiusSquared(FVector(SectionOriginAndRadius[EntityIndex]), SectionOriginAndRadius[EntityIndex].W, ViewOrigin, ViewProjectionMarix);
+	int32 ForcedLODLevel = ViewEngineShowFlagLOD ? GetCVarForceLOD() : -1;
+	float LODScale = ViewLODDistanceFactor * CVarStaticMeshLODDistanceScale.GetValueOnRenderThread();
 
-			float FractionalLOD;
-			GetLODFromScreenSize(SectionLODSettings[EntityIndex], MeshScreenSizeSquared, LODScale * LODScale, FractionalLOD);
+	for (int32 EntityIndex = 0; EntityIndex < SectionLODSettings.Num(); EntityIndex++)
+	{
+		float MeshScreenSizeSquared = ComputeBoundsScreenRadiusSquared(FVector(SectionOriginAndRadius[EntityIndex]), SectionOriginAndRadius[EntityIndex].W, ViewOrigin, ViewProjectionMarix);
 
-			SectionLODValues[EntityIndex] = ForcedLODLevel >= 0 ? ForcedLODLevel : FractionalLOD;
+		float FractionalLOD;
+		GetLODFromScreenSize(SectionLODSettings[EntityIndex], MeshScreenSizeSquared, LODScale * LODScale, FractionalLOD);
 
-			if (TessellationFalloffSettings.UseTessellationComponentScreenSizeFalloff && NumEntitiesWithTessellation > 0)
-			{
-				float MaxTesselationDistance = ComputeBoundsDrawDistance(FMath::Sqrt(TessellationFalloffSettings.TessellationComponentSquaredScreenSize), SectionOriginAndRadius[EntityIndex].W / 2.0f, ViewProjectionMarix);
-				float FallOffStartingDistance = FMath::Min(
-					ComputeBoundsDrawDistance(FMath::Sqrt(FMath::Min(
-						FMath::Square(TessellationFalloffSettings.TessellationComponentScreenSizeFalloff), 
-						TessellationFalloffSettings.TessellationComponentSquaredScreenSize)), SectionOriginAndRadius[EntityIndex].W / 2.0f, ViewProjectionMarix) - MaxTesselationDistance, MaxTesselationDistance);
-
-				// Calculate the falloff using a = C - K * d by sending C & K into the shader
-				SectionTessellationFalloffC[EntityIndex] = MaxTesselationDistance / (MaxTesselationDistance - FallOffStartingDistance);
-				SectionTessellationFalloffK[EntityIndex] = -(1 / (-MaxTesselationDistance + FallOffStartingDistance));
-			}
-		}
+		NewSectionLODValues[EntityIndex] = ForcedLODLevel >= 0 ? ForcedLODLevel : FractionalLOD;
 
+		if (TessellationFalloffSettings.UseTessellationComponentScreenSizeFalloff && NumEntitiesWithTessellation > 0)
 		{
-			FScopeLock Lock(&CachedValuesCS);
-
-			CachedSectionLODValues[ViewPtrAsIdentifier] = SectionLODValues;
+			float MaxTesselationDistance = ComputeBoundsDrawDistance(FMath::Sqrt(TessellationFalloffSettings.TessellationComponentSquaredScreenSize), SectionOriginAndRadius[EntityIndex].W / 2.0f, ViewProjectionMarix);
+			float FallOffStartingDistance = FMath::Min(
+				ComputeBoundsDrawDistance(FMath::Sqrt(FMath::Min(
+					FMath::Square(TessellationFalloffSettings.TessellationComponentScreenSizeFalloff),
+					TessellationFalloffSettings.TessellationComponentSquaredScreenSize)), SectionOriginAndRadius[EntityIndex].W / 2.0f, ViewProjectionMarix) - MaxTesselationDistance, MaxTesselationDistance);
 
-			if (TessellationFalloffSettings.UseTessellationComponentScreenSizeFalloff && NumEntitiesWithTessellation > 0)
-			{
-				CachedSectionTessellationFalloffC[ViewPtrAsIdentifier] = SectionTessellationFalloffC;
-				CachedSectionTessellationFalloffK[ViewPtrAsIdentifier] = SectionTessellationFalloffK;
-			}
+			// Calculate the falloff using a = C - K * d by sending C & K into the shader
+			NewSectionTessellationFalloffC[EntityIndex] = MaxTesselationDistance / (MaxTesselationDistance - FallOffStartingDistance);
+			NewSectionTessellationFalloffC[EntityIndex] = -(1 / (-MaxTesselationDistance + FallOffStartingDistance));
 		}
 	}
-	else
+
 	{
-		{
-			FScopeLock Lock(&CachedValuesCS);
+		FScopeLock Lock(&CachedValuesCS);
 
-			SectionLODValues = CachedSectionLODValues[ViewPtrAsIdentifier];
+		CachedSectionLODValues.Add(ViewPtrAsIdentifier, NewSectionLODValues);
 
-			if (TessellationFalloffSettings.UseTessellationComponentScreenSizeFalloff && NumEntitiesWithTessellation > 0)
-			{
-				SectionTessellationFalloffC = CachedSectionTessellationFalloffC[ViewPtrAsIdentifier];
-				SectionTessellationFalloffK = CachedSectionTessellationFalloffK[ViewPtrAsIdentifier];
-			}
+		if (TessellationFalloffSettings.UseTessellationComponentScreenSizeFalloff && NumEntitiesWithTessellation > 0)
+		{
+			CachedSectionTessellationFalloffC.Add(ViewPtrAsIdentifier, NewSectionTessellationFalloffC);
+			CachedSectionTessellationFalloffK.Add(ViewPtrAsIdentifier, NewSectionTessellationFalloffC);
 		}
 	}
 }
diff --git a/Engine/Source/Runtime/Renderer/Private/SceneHitProxyRendering.cpp b/Engine/Source/Runtime/Renderer/Private/SceneHitProxyRendering.cpp
--- a/Engine/Source/Runtime/Renderer/Private/SceneHitProxyRendering.cpp
+++ b/Engine/Source/Runtime/Renderer/Private/SceneHitProxyRendering.cpp
@@ -576,18 +576,24 @@ void FDeferredShadingSceneRenderer::RenderHitProxies(FRHICommandListImmediate& R
 			InitViewsPossiblyAfterPrepass(RHICmdList, ILCTaskData, UpdateViewCustomDataEvents);
 		}
 
-		UpdateGPUScene(RHICmdList, *Scene);
+		extern TSet<IPersistentViewUniformBufferExtension*> PersistentViewUniformBufferExtensions;
 
-		for (int32 ViewIndex = 0; ViewIndex < Views.Num(); ViewIndex++)
+		for (IPersistentViewUniformBufferExtension* Extension : PersistentViewUniformBufferExtensions)
 		{
-			UploadDynamicPrimitiveShaderDataForView(RHICmdList, *Scene, Views[ViewIndex]);
+			Extension->BeginFrame();
 
-			extern TSet<IPersistentViewUniformBufferExtension*> PersistentViewUniformBufferExtensions;
-
-			for (IPersistentViewUniformBufferExtension* Extension : PersistentViewUniformBufferExtensions)
+			for (int32 ViewIndex = 0; ViewIndex < Views.Num(); ViewIndex++)
 			{
-				Extension->BeginRenderView(&Views[ViewIndex]);
+				// Must happen before RHI thread flush so any tasks we dispatch here can land in the idle gap during the flush
+				Extension->PrepareView(&Views[ViewIndex]);
 			}
+		}
+
+		UpdateGPUScene(RHICmdList, *Scene);
+
+		for (int32 ViewIndex = 0; ViewIndex < Views.Num(); ViewIndex++)
+		{
+			UploadDynamicPrimitiveShaderDataForView(RHICmdList, *Scene, Views[ViewIndex]);
 		}	
 
 		if (UpdateViewCustomDataEvents.Num())
