Fixed issue in landscape render

diff --git a/Engine/Source/Runtime/Landscape/Private/LandscapeRender.cpp b/Engine/Source/Runtime/Landscape/Private/LandscapeRender.cpp
--- a/Engine/Source/Runtime/Landscape/Private/LandscapeRender.cpp
+++ b/Engine/Source/Runtime/Landscape/Private/LandscapeRender.cpp
@@ -872,7 +872,7 @@ void FLandscapeRenderSystem::ComputeSectionPerViewParameters(
 
 			// Calculate the falloff using a = C - K * d by sending C & K into the shader
 			NewSectionTessellationFalloffC[EntityIndex] = MaxTesselationDistance / (MaxTesselationDistance - FallOffStartingDistance);
-			NewSectionTessellationFalloffC[EntityIndex] = -(1 / (-MaxTesselationDistance + FallOffStartingDistance));
+			NewSectionTessellationFalloffK[EntityIndex] = -(1 / (-MaxTesselationDistance + FallOffStartingDistance));
 		}
 	}
 
@@ -884,7 +884,7 @@ void FLandscapeRenderSystem::ComputeSectionPerViewParameters(
 		if (TessellationFalloffSettings.UseTessellationComponentScreenSizeFalloff && NumEntitiesWithTessellation > 0)
 		{
 			CachedSectionTessellationFalloffC.Add(ViewPtrAsIdentifier, NewSectionTessellationFalloffC);
-			CachedSectionTessellationFalloffK.Add(ViewPtrAsIdentifier, NewSectionTessellationFalloffC);
+			CachedSectionTessellationFalloffK.Add(ViewPtrAsIdentifier, NewSectionTessellationFalloffK);
 		}
 	}
 }
