Enable mesh draw commands for rendering landscape to runtime virtual texture

diff --git a/Engine/Source/Runtime/Landscape/Private/LandscapeRender.cpp b/Engine/Source/Runtime/Landscape/Private/LandscapeRender.cpp
--- a/Engine/Source/Runtime/Landscape/Private/LandscapeRender.cpp
+++ b/Engine/Source/Runtime/Landscape/Private/LandscapeRender.cpp
@@ -4574,7 +4574,7 @@ FVertexFactoryShaderParameters* FLandscapeFixedGridVertexFactory::ConstructShade
 	}
 }
 
-IMPLEMENT_VERTEX_FACTORY_TYPE(FLandscapeFixedGridVertexFactory, "/Engine/Private/LandscapeVertexFactory.ush", true, true, true, false, false);
+IMPLEMENT_VERTEX_FACTORY_TYPE_EX(FLandscapeFixedGridVertexFactory, "/Engine/Private/LandscapeVertexFactory.ush", true, true, true, false, false, true, false);
 
 
 /** ULandscapeMaterialInstanceConstant */
