Cherrypick CL 9839712

Fix HasLayersContent returning false and causing dirtyness when applying splines

#rnx

diff --git a/Engine/Source/Runtime/Landscape/Private/Landscape.cpp b/Engine/Source/Runtime/Landscape/Private/Landscape.cpp
--- a/Engine/Source/Runtime/Landscape/Private/Landscape.cpp
+++ b/Engine/Source/Runtime/Landscape/Private/Landscape.cpp
@@ -2366,7 +2366,8 @@ void ALandscapeProxy::PostLoad()
 		ULandscapeInfo* LandscapeInfo = CreateLandscapeInfo();
 
 		// Cache the value at this point as RegisterActor might create/destroy layers content if there was a mismatch between landscape & proxy
-		const bool bNeedOldDataMigration = !HasLayersContent() && CanHaveLayersContent();
+		// Check the actual flag here not HasLayersContent() which could return true if the LandscapeActor is valid.
+		const bool bNeedOldDataMigration = !bHasLayersContent && CanHaveLayersContent();
 				
 		LandscapeInfo->RegisterActor(this, true);
 
diff --git a/Engine/Source/Runtime/Landscape/Private/LandscapeEditLayers.cpp b/Engine/Source/Runtime/Landscape/Private/LandscapeEditLayers.cpp
--- a/Engine/Source/Runtime/Landscape/Private/LandscapeEditLayers.cpp
+++ b/Engine/Source/Runtime/Landscape/Private/LandscapeEditLayers.cpp
@@ -5172,7 +5172,7 @@ bool ALandscapeProxy::CanHaveLayersContent() const
 
 bool ALandscapeProxy::HasLayersContent() const
 {
-	return bHasLayersContent;
+	return bHasLayersContent || (GetLandscapeActor() != nullptr && GetLandscapeActor()->HasLayersContent());
 }
 
 void ALandscapeProxy::UpdateCachedHasLayersContent(bool InCheckComponentDataIntegrity)
