Changed log warning to mapcheck warning

diff --git a/Engine/Source/Runtime/Landscape/Private/LandscapeEditLayers.cpp b/Engine/Source/Runtime/Landscape/Private/LandscapeEditLayers.cpp
--- a/Engine/Source/Runtime/Landscape/Private/LandscapeEditLayers.cpp
+++ b/Engine/Source/Runtime/Landscape/Private/LandscapeEditLayers.cpp
@@ -37,6 +37,8 @@ LandscapeEditLayers.cpp: Landscape editing layers mode
 #include "UObject/UObjectThreadContext.h"
 #include "LandscapeSplinesComponent.h"
 #include "Misc/FileHelper.h"
+#include "Misc/MapErrors.h"
+#include "Misc/UObjectToken.h"
 #endif
 
 #define LOCTEXT_NAMESPACE "Landscape"
@@ -5203,8 +5205,15 @@ bool ALandscapeProxy::RemoveObsoleteLayers(const TSet<FGuid>& InExistingLayers)
 	{
 		if (!InExistingLayers.Contains(LayerGuid))
 		{
-			UE_LOG(LogLandscape, Warning, TEXT("Layer '%s' was removed from LandscapeProxy '%s' because it doesn't match any of the LandscapeActor Layers. Possible loss of data."), 
-				*LayerGuid.ToString(EGuidFormats::HexValuesInBraces), *GetPathName());
+			FFormatNamedArguments Arguments;
+			Arguments.Add(TEXT("LayerGuid"), FText::FromString(LayerGuid.ToString(EGuidFormats::HexValuesInBraces)));
+			Arguments.Add(TEXT("ProxyPackage"), FText::FromString(GetOutermost()->GetName()));
+
+			FMessageLog("MapCheck").Warning()
+				->AddToken(FUObjectToken::Create(this))
+				->AddToken(FTextToken::Create(FText::Format(LOCTEXT("MapCheck_Message_LandscapeProxyObsoleteLayer","Layer '{LayerGuid}' was removed from LandscapeProxy because it doesn't match any of the LandscapeActor Layers. Please resave '{ProxyPackage}'."), Arguments)))
+				->AddToken(FMapErrorToken::Create(FMapErrors::LandscapeComponentPostLoad_Warning));
+
 			DeleteLayer(LayerGuid);
 			bModified = true;
 		}
