Update default landscape LOD distribution parameters in favor of the new system

diff --git a/Engine/Source/Runtime/Landscape/Private/Landscape.cpp b/Engine/Source/Runtime/Landscape/Private/Landscape.cpp
--- a/Engine/Source/Runtime/Landscape/Private/Landscape.cpp
+++ b/Engine/Source/Runtime/Landscape/Private/Landscape.cpp
@@ -1041,9 +1041,9 @@ ALandscapeProxy::ALandscapeProxy(const FObjectInitializer& ObjectInitializer)
 	ComponentScreenSizeToUseSubSections = 0.65f;
 	UseTessellationComponentScreenSizeFalloff = true;
 	TessellationComponentScreenSizeFalloff = 0.75f;
-	LOD0ScreenSize = 1.0f;
-	LOD0DistributionSetting = 1.75f;
-	LODDistributionSetting = 2.0f;
+	LOD0ScreenSize = 0.5f;
+	LOD0DistributionSetting = 1.25f;
+	LODDistributionSetting = 3.0f;
 	bCastStaticShadow = true;
 	bCastShadowAsTwoSided = false;
 	bUsedForNavigation = true;
