Landscape Layers: Wait for all Textures before doing any render (This completes the random broken Landscape with Ryan's BP fix)

#rnx

diff --git a/Engine/Source/Runtime/Landscape/Classes/Landscape.h b/Engine/Source/Runtime/Landscape/Classes/Landscape.h
--- a/Engine/Source/Runtime/Landscape/Classes/Landscape.h
+++ b/Engine/Source/Runtime/Landscape/Classes/Landscape.h
@@ -350,6 +350,7 @@ private:
 	void ResolveLayersWeightmapTexture(const TArray<ULandscapeComponent*>& InLandscapeComponents);
 	bool ResolveLayersTexture(class FLandscapeLayersTexture2DCPUReadBackResource* InCPUReadBackTexture, UTexture2D* InOutputTexture, bool bHeightmap, ULandscapeComponent* InComponent);
 		
+	bool AreLayersTextureResourcesReady(bool bInWaitForStreaming) const;
 	bool PrepareLayersBrushTextureResources(bool bInWaitForStreaming, bool bHeightmap) const;
 	bool PrepareLayersHeightmapTextureResources(bool bInWaitForStreaming) const;
 	bool PrepareLayersWeightmapTextureResources(bool bInWaitForStreaming) const;
diff --git a/Engine/Source/Runtime/Landscape/Private/LandscapeEditLayers.cpp b/Engine/Source/Runtime/Landscape/Private/LandscapeEditLayers.cpp
--- a/Engine/Source/Runtime/Landscape/Private/LandscapeEditLayers.cpp
+++ b/Engine/Source/Runtime/Landscape/Private/LandscapeEditLayers.cpp
@@ -2804,16 +2804,7 @@ int32 ALandscape::RegenerateLayersHeightmaps(const TArray<ULandscapeComponent*>&
 	{
 		return 0;
 	}
-
-	const bool bHeightmapTexturesReady = PrepareLayersHeightmapTextureResources(bInWaitForStreaming);
-	const bool bHeightmap = true;
-	const bool bBrushTexturesReady = PrepareLayersBrushTextureResources(bInWaitForStreaming, bHeightmap);
-
-	if (!(bHeightmapTexturesReady && bBrushTexturesReady))
-	{
-		return 0;
-	}
-
+		
 	// Nothing to do (return that we did the processing)
 	if (InLandscapeComponentsToResolve.Num() == 0)
 	{
@@ -3740,16 +3731,7 @@ int32 ALandscape::RegenerateLayersWeightmaps(const TArray<ULandscapeComponent*>&
 	{
 		return 0;
 	}
-
-	const bool bWeightmapTexturesReady = PrepareLayersWeightmapTextureResources(bInWaitForStreaming);
-	const bool bHeightmap = false;
-	const bool bBrushTexturesReady = PrepareLayersBrushTextureResources(bInWaitForStreaming, bHeightmap);
-
-	if (!(bWeightmapTexturesReady && bBrushTexturesReady))
-	{
-		return 0;
-	}
-
+		
 	if (InLandscapeComponentsToResolve.Num() == 0)
 	{
 		return WeightmapUpdateModes;
@@ -4753,6 +4735,15 @@ void ALandscape::GetLandscapeComponentWeightmapsToRender(ULandscapeComponent* La
 	}
 }
 
+bool ALandscape::AreLayersTextureResourcesReady(bool bInWaitForStreaming) const
+{
+	const bool bHeightmapReady = PrepareLayersHeightmapTextureResources(bInWaitForStreaming);
+	const bool bWeightmapReady = PrepareLayersWeightmapTextureResources(bInWaitForStreaming);
+	const bool bBrushHeightmapbReady = PrepareLayersBrushTextureResources(bInWaitForStreaming, true);
+	const bool bBrushWeightmapReady = PrepareLayersBrushTextureResources(bInWaitForStreaming, false);
+	return bHeightmapReady && bWeightmapReady && bBrushHeightmapbReady && bBrushWeightmapReady;
+}
+
 void ALandscape::UpdateLayersContent(bool bInWaitForStreaming, bool bInSkipMonitorLandscapeEdModeChanges)
 {
 	TRACE_CPUPROFILER_EVENT_SCOPE("UpdateLayersContent");
@@ -4785,6 +4776,11 @@ void ALandscape::UpdateLayersContent(bool bInWaitForStreaming, bool bInSkipMonit
 		bSplineLayerUpdateRequested = false;
 	}
 
+	if (!AreLayersTextureResourcesReady(bInWaitForStreaming))
+	{
+		return;
+	}
+	
 	const bool bForceRender = CVarOutputLayersDebugDrawCallName.GetValueOnAnyThread() == 1;
 
 	if (LayerContentUpdateModes == 0 && !bForceRender)
