Fixed crash when creating a layer info on a list of layers with none assigned yet

#rnx

diff --git a/Engine/Source/Runtime/Landscape/Private/LandscapeEditInterface.cpp b/Engine/Source/Runtime/Landscape/Private/LandscapeEditInterface.cpp
--- a/Engine/Source/Runtime/Landscape/Private/LandscapeEditInterface.cpp
+++ b/Engine/Source/Runtime/Landscape/Private/LandscapeEditInterface.cpp
@@ -1898,7 +1898,7 @@ void ULandscapeComponent::FillLayer(ULandscapeLayerInfoObject* LayerInfo, FLands
 				continue;
 			}
 			FWeightmapLayerAllocationInfo& Allocation = ComponentWeightmapLayerAllocations[LayerIdx];
-			if (Allocation.LayerInfo->bNoWeightBlend)
+			if ((Allocation.LayerInfo == nullptr) || Allocation.LayerInfo->bNoWeightBlend)
 			{
 				continue;
 			}
