LevelGrid: Editgrate CL 9688728 + 9743673 from //UE4/Dev-Editor-OpenWorld

MoveComponentsToLevel now part of ULandscapeInfo

#rnx

diff --git a/Engine/Source/Editor/LandscapeEditor/Private/LandscapeEdMode.cpp b/Engine/Source/Editor/LandscapeEditor/Private/LandscapeEdMode.cpp
--- a/Engine/Source/Editor/LandscapeEditor/Private/LandscapeEdMode.cpp
+++ b/Engine/Source/Editor/LandscapeEditor/Private/LandscapeEdMode.cpp
@@ -97,176 +97,6 @@ struct HNewLandscapeGrabHandleProxy : public HHitProxy
 
 IMPLEMENT_HIT_PROXY(HNewLandscapeGrabHandleProxy, HHitProxy)
 
-ENGINE_API extern bool GDisableAutomaticTextureMaterialUpdateDependencies;
-
-void ALandscape::SplitHeightmap(ULandscapeComponent* Comp, ALandscapeProxy* TargetProxy,FMaterialUpdateContext* InOutUpdateContext, TArray<FComponentRecreateRenderStateContext>* InOutRecreateRenderStateContext, bool InReregisterComponent)
-{
-	ULandscapeInfo* Info = Comp->GetLandscapeInfo();
-	
-	// Make sure the heightmap UVs are powers of two.
-	int32 ComponentSizeVerts = Comp->NumSubsections * (Comp->SubsectionSizeQuads + 1);
-	int32 HeightmapSizeU = (1 << FMath::CeilLogTwo(ComponentSizeVerts));
-	int32 HeightmapSizeV = (1 << FMath::CeilLogTwo(ComponentSizeVerts));
-
-	ALandscapeProxy* SrcProxy = Comp->GetLandscapeProxy();
-	ALandscapeProxy* DstProxy = TargetProxy ? TargetProxy : SrcProxy;
-	SrcProxy->Modify();
-	DstProxy->Modify();
-	
-	UTexture2D* OldHeightmapTexture = Comp->GetHeightmap(false);
-	UTexture2D* NewHeightmapTexture = NULL;
-	FVector4 OldHeightmapScaleBias = Comp->HeightmapScaleBias;
-	FVector4 NewHeightmapScaleBias = FVector4(1.0f / (float)HeightmapSizeU, 1.0f / (float)HeightmapSizeV, 0.0f, 0.0f);
-
-	{
-		// Read old data and split
-		FLandscapeEditDataInterface LandscapeEdit(Info);
-		TArray<uint8> HeightData;
-		HeightData.AddZeroed((1 + Comp->ComponentSizeQuads)*(1 + Comp->ComponentSizeQuads)*sizeof(uint16));
-		// Because of edge problem, normal would be just copy from old component data
-		TArray<uint8> NormalData;
-		NormalData.AddZeroed((1 + Comp->ComponentSizeQuads)*(1 + Comp->ComponentSizeQuads)*sizeof(uint16));
-		LandscapeEdit.GetHeightDataFast(Comp->GetSectionBase().X, Comp->GetSectionBase().Y, Comp->GetSectionBase().X + Comp->ComponentSizeQuads, Comp->GetSectionBase().Y + Comp->ComponentSizeQuads, (uint16*)HeightData.GetData(), 0, (uint16*)NormalData.GetData());
-
-		// Create the new heightmap texture
-		NewHeightmapTexture = DstProxy->CreateLandscapeTexture(HeightmapSizeU, HeightmapSizeV, TEXTUREGROUP_Terrain_Heightmap, TSF_BGRA8);
-		ULandscapeComponent::CreateEmptyTextureMips(NewHeightmapTexture, true);
-		Comp->HeightmapScaleBias = NewHeightmapScaleBias;
-		Comp->SetHeightmap(NewHeightmapTexture);
-
-		check(Comp->GetHeightmap(false) == Comp->GetHeightmap(true));
-		LandscapeEdit.SetHeightData(Comp->GetSectionBase().X, Comp->GetSectionBase().Y, Comp->GetSectionBase().X + Comp->ComponentSizeQuads, Comp->GetSectionBase().Y + Comp->ComponentSizeQuads, (uint16*)HeightData.GetData(), 0, false, (uint16*)NormalData.GetData());
-	}
-
-	// End material update
-	if (InOutUpdateContext != nullptr && InOutRecreateRenderStateContext != nullptr)
-	{
-		Comp->UpdateMaterialInstances(*InOutUpdateContext, *InOutRecreateRenderStateContext);
-	}
-	else
-	{
-		Comp->UpdateMaterialInstances();
-	}
-
-	// We disable automatic material update context, to manage it manually if we have a custom update context specified
-	GDisableAutomaticTextureMaterialUpdateDependencies = (InOutUpdateContext != nullptr);
-
-	NewHeightmapTexture->PostEditChange();
-
-	if (InOutUpdateContext != nullptr)
-	{
-		// Build a list of all unique materials the landscape uses
-		TArray<UMaterialInterface*> LandscapeMaterials;
-
-		int8 MaxLOD = FMath::CeilLogTwo(Comp->SubsectionSizeQuads + 1) - 1;
-
-		for (int8 LODIndex = 0; LODIndex < MaxLOD; ++LODIndex)
-		{
-			UMaterialInterface* Material = Comp->GetLandscapeMaterial(LODIndex);
-			LandscapeMaterials.AddUnique(Material);
-		}
-
-		TSet<UMaterial*> BaseMaterialsThatUseThisTexture;
-
-		for (UMaterialInterface* MaterialInterface : LandscapeMaterials)
-		{
-			if (DoesMaterialUseTexture(MaterialInterface, NewHeightmapTexture))
-			{
-				UMaterial* Material = MaterialInterface->GetMaterial();
-				bool MaterialAlreadyCompute = false;
-				BaseMaterialsThatUseThisTexture.Add(Material, &MaterialAlreadyCompute);
-
-				if (!MaterialAlreadyCompute)
-				{
-					if (Material->IsTextureForceRecompileCacheRessource(NewHeightmapTexture))
-					{
-						InOutUpdateContext->AddMaterial(Material);
-						Material->UpdateMaterialShaderCacheAndTextureReferences();
-					}
-				}
-			}
-		}
-	}
-	
-	GDisableAutomaticTextureMaterialUpdateDependencies = false;
-
-#if WITH_EDITORONLY_DATA
-	check(Comp->GetLandscapeProxy()->HasLayersContent() == DstProxy->CanHaveLayersContent());
-	if (Comp->GetLandscapeProxy()->HasLayersContent() && DstProxy->CanHaveLayersContent())
-	{
-		FLandscapeLayersTexture2DCPUReadBackResource* NewCPUReadBackResource = new FLandscapeLayersTexture2DCPUReadBackResource(NewHeightmapTexture->Source.GetSizeX(), NewHeightmapTexture->Source.GetSizeY(), NewHeightmapTexture->GetPixelFormat(), NewHeightmapTexture->Source.GetNumMips());
-		BeginInitResource(NewCPUReadBackResource);
-		DstProxy->HeightmapsCPUReadBack.Add(NewHeightmapTexture, NewCPUReadBackResource);
-
-		// Free OldHeightmapTexture's CPUReadBackResource if not used by any component
-		bool FreeCPUReadBack = true;
-		for (ULandscapeComponent* Component : SrcProxy->LandscapeComponents)
-		{
-			if (Component != Comp && Component->GetHeightmap(false) == OldHeightmapTexture)
-			{
-				FreeCPUReadBack = false;
-				break;
-			}
-		}
-		if (FreeCPUReadBack)
-		{
-			FLandscapeLayersTexture2DCPUReadBackResource** OldCPUReadBackResource = SrcProxy->HeightmapsCPUReadBack.Find(OldHeightmapTexture);
-			if (OldCPUReadBackResource)
-			{
-				if (FLandscapeLayersTexture2DCPUReadBackResource* ResourceToDelete = *OldCPUReadBackResource)
-				{
-					ReleaseResourceAndFlush(ResourceToDelete);
-					delete ResourceToDelete;
-					SrcProxy->HeightmapsCPUReadBack.Remove(OldHeightmapTexture);
-				}
-			}
-		}
-
-		// Move layer content to new layer heightmap
-		FLandscapeEditDataInterface LandscapeEdit(Info);
-		ALandscape* Landscape = Info->LandscapeActor.Get();
-		Comp->ForEachLayer([&](const FGuid& LayerGuid, FLandscapeLayerComponentData& LayerData)
-		{
-			UTexture2D* OldLayerHeightmap = LayerData.HeightmapData.Texture;
-			if (OldLayerHeightmap != nullptr)
-			{
-				FScopedSetLandscapeEditingLayer Scope(Landscape, LayerGuid);
-				// Read old data and split
-				TArray<uint8> LayerHeightData;
-				LayerHeightData.AddZeroed((1 + Comp->ComponentSizeQuads)*(1 + Comp->ComponentSizeQuads) * sizeof(uint16));
-				// Because of edge problem, normal would be just copy from old component data
-				TArray<uint8> LayerNormalData;
-				LayerNormalData.AddZeroed((1 + Comp->ComponentSizeQuads)*(1 + Comp->ComponentSizeQuads) * sizeof(uint16));
-
-				// Read using old heightmap scale/bias
-				Comp->HeightmapScaleBias = OldHeightmapScaleBias;
-				LandscapeEdit.GetHeightDataFast(Comp->GetSectionBase().X, Comp->GetSectionBase().Y, Comp->GetSectionBase().X + Comp->ComponentSizeQuads, Comp->GetSectionBase().Y + Comp->ComponentSizeQuads, (uint16*)LayerHeightData.GetData(), 0, (uint16*)LayerNormalData.GetData());
-				// Restore new heightmap scale/bias
-				Comp->HeightmapScaleBias = NewHeightmapScaleBias;
-				{
-					UTexture2D* LayerHeightmapTexture = DstProxy->CreateLandscapeTexture(HeightmapSizeU, HeightmapSizeV, TEXTUREGROUP_Terrain_Heightmap, TSF_BGRA8);
-					ULandscapeComponent::CreateEmptyTextureMips(LayerHeightmapTexture, true);
-					LayerHeightmapTexture->PostEditChange();
-					// Set Layer heightmap texture
-					LayerData.HeightmapData.Texture = LayerHeightmapTexture;
-					LandscapeEdit.SetHeightData(Comp->GetSectionBase().X, Comp->GetSectionBase().Y, Comp->GetSectionBase().X + Comp->ComponentSizeQuads, Comp->GetSectionBase().Y + Comp->ComponentSizeQuads, (uint16*)LayerHeightData.GetData(), 0, false, (uint16*)LayerNormalData.GetData());
-				}
-			}
-		});
-
-		Landscape->RequestLayersContentUpdateForceAll();
-	}
-#endif
-
-	// Reregister
-	if (InReregisterComponent)
-	{
-		FComponentReregisterContext ReregisterContext(Comp);
-	}
-}
-
-
-
 void FLandscapeTool::SetEditRenderType()
 {
 	GLandscapeEditRenderMode = ELandscapeEditRenderMode::SelectRegion | (GLandscapeEditRenderMode & ELandscapeEditRenderMode::BitMaskForMask);
diff --git a/Engine/Source/Editor/LandscapeEditor/Private/LandscapeEdModeComponentTools.cpp b/Engine/Source/Editor/LandscapeEditor/Private/LandscapeEdModeComponentTools.cpp
--- a/Engine/Source/Editor/LandscapeEditor/Private/LandscapeEdModeComponentTools.cpp
+++ b/Engine/Source/Editor/LandscapeEditor/Private/LandscapeEdModeComponentTools.cpp
@@ -480,205 +480,16 @@ public:
 				FScopedSlowTask SlowTask(0, LOCTEXT("BeginMovingLandscapeComponentsToCurrentLevelTask", "Moving Landscape components to current level"));
 				SlowTask.MakeDialogDelayed(10); // show slow task dialog after 10 seconds
 
-				LandscapeInfo->SortSelectedComponents();
-				const int32 ComponentSizeVerts = Landscape->NumSubsections * (Landscape->SubsectionSizeQuads + 1);
-				const int32 NeedHeightmapSize = 1 << FMath::CeilLogTwo(ComponentSizeVerts);
-
-				TSet<ALandscapeProxy*> SelectProxies;
-				TSet<ULandscapeComponent*> TargetSelectedComponents;
-				TArray<ULandscapeHeightfieldCollisionComponent*> TargetSelectedCollisionComponents;
-				for (ULandscapeComponent* Component : SelectedComponents)
-				{
-					SelectProxies.Add(Component->GetLandscapeProxy());
-					if (Component->GetLandscapeProxy()->GetOuter() != World->GetCurrentLevel())
-					{
-						TargetSelectedComponents.Add(Component);
-					}
-
-					ULandscapeHeightfieldCollisionComponent* CollisionComp = Component->CollisionComponent.Get();
-					SelectProxies.Add(CollisionComp->GetLandscapeProxy());
-					if (CollisionComp->GetLandscapeProxy()->GetOuter() != World->GetCurrentLevel())
-					{
-						TargetSelectedCollisionComponents.Add(CollisionComp);
-					}
-				}
-
-				// Check which ones are need for height map change
-				TSet<UTexture2D*> OldHeightmapTextures;
-				for (ULandscapeComponent* Component : TargetSelectedComponents)
-				{
-					Component->Modify();
-					OldHeightmapTextures.Add(Component->GetHeightmap());
-				}
-
-				// Need to split all the component which share Heightmap with selected components
-				TMap<ULandscapeComponent*, bool> HeightmapUpdateComponents;
-				HeightmapUpdateComponents.Reserve(TargetSelectedComponents.Num() * 4); // worst case
-				for (ULandscapeComponent* Component : TargetSelectedComponents)
-				{
-					// Search neighbor only
-					const int32 SearchX = Component->GetHeightmap()->Source.GetSizeX() / NeedHeightmapSize - 1;
-					const int32 SearchY = Component->GetHeightmap()->Source.GetSizeY() / NeedHeightmapSize - 1;
-					const FIntPoint ComponentBase = Component->GetSectionBase() / Component->ComponentSizeQuads;
-
-					for (int32 Y = -SearchY; Y <= SearchY; ++Y)
-					{
-						for (int32 X = -SearchX; X <= SearchX; ++X)
-						{
-							ULandscapeComponent* const Neighbor = LandscapeInfo->XYtoComponentMap.FindRef(ComponentBase + FIntPoint(X, Y));
-							if (Neighbor && Neighbor->GetHeightmap() == Component->GetHeightmap() && !HeightmapUpdateComponents.Contains(Neighbor))
-							{
-								Neighbor->Modify();
-								bool bNeedsMoveToCurrentLevel = TargetSelectedComponents.Contains(Neighbor);
-								HeightmapUpdateComponents.Add(Neighbor, bNeedsMoveToCurrentLevel);
-							}
-						}
-					}
-				}
-
-				ALandscapeProxy* LandscapeProxy = LandscapeInfo->GetCurrentLevelLandscapeProxy(false);
-				if (!LandscapeProxy)
-				{
-					LandscapeProxy = World->SpawnActor<ALandscapeStreamingProxy>();
-					// copy shared properties to this new proxy
-					LandscapeProxy->GetSharedProperties(Landscape);
-					LandscapeProxy->CreateLandscapeInfo();
-					
-					// set proxy location
-					// by default first component location
-					ULandscapeComponent* FirstComponent = *TargetSelectedComponents.CreateConstIterator();
-					LandscapeProxy->GetRootComponent()->SetWorldLocationAndRotation(FirstComponent->GetComponentLocation(), FirstComponent->GetComponentRotation());
-					LandscapeProxy->LandscapeSectionOffset = FirstComponent->GetSectionBase();
-
-					// Hide(unregister) the new landscape if owning level currently in hidden state
-					if (LandscapeProxy->GetLevel()->bIsVisible == false)
-					{
-						LandscapeProxy->UnregisterAllComponents();
-					}
-				}
-
-				// Changing Heightmap format for selected components
-				for (const auto& HeightmapUpdateComponentPair : HeightmapUpdateComponents)
-				{
-					ALandscape::SplitHeightmap(HeightmapUpdateComponentPair.Key, HeightmapUpdateComponentPair.Value ? LandscapeProxy : nullptr);
-				}
-								
-				// Delete if it is no referenced textures...
-				for (UTexture2D* Texture : OldHeightmapTextures)
-				{
-					Texture->SetFlags(RF_Transactional);
-					Texture->Modify();
-					Texture->MarkPackageDirty();
-					Texture->ClearFlags(RF_Standalone);
-				}
-
-				for (ALandscapeProxy* Proxy : SelectProxies)
-				{
-					Proxy->Modify();
-				}
-
-				LandscapeProxy->Modify();
-				LandscapeProxy->MarkPackageDirty();
-
-				// Handle XY-offset textures (these don't need splitting, as they aren't currently shared between components like heightmaps/weightmaps can be)
-				for (ULandscapeComponent* Component : TargetSelectedComponents)
-				{
-					if (Component->XYOffsetmapTexture)
-					{
-						Component->XYOffsetmapTexture->Modify();
-						Component->XYOffsetmapTexture->Rename(nullptr, LandscapeProxy);
-					}
-				}
-
-				// Change Weight maps...
-				{
-					FLandscapeEditDataInterface LandscapeEdit(LandscapeInfo);
-					for (ULandscapeComponent* Component : TargetSelectedComponents)
-					{
-						Component->ReallocateWeightmaps(&LandscapeEdit, false, true, false, true, LandscapeProxy);
-						Component->ForEachLayer([&](const FGuid& LayerGuid, FLandscapeLayerComponentData& LayerData)
-						{
-							FScopedSetLandscapeEditingLayer Scope(Landscape, LayerGuid);
-							Component->ReallocateWeightmaps(&LandscapeEdit, true, true, false, true, LandscapeProxy);
-						});
-						Landscape->RequestLayersContentUpdateForceAll();
-					}
-
-					// Need to Repacking all the Weight map (to make it packed well...)
-					for (ALandscapeProxy* Proxy : SelectProxies)
-					{
-						Proxy->RemoveInvalidWeightmaps();
-					}
-				}
-
-				// Move the components to the Proxy actor
-				// This does not use the MoveSelectedActorsToCurrentLevel path as there is no support to only move certain components.
-				for (ULandscapeComponent* Component :TargetSelectedComponents)
+				if (ALandscapeProxy* LandscapeProxy = LandscapeInfo->MoveComponentsToLevel(SelectedComponents.Array(), World->GetCurrentLevel()))
 				{
-					// Need to move or recreate all related data (Height map, Weight map, maybe collision components, allocation info)
-					Component->GetLandscapeProxy()->LandscapeComponents.Remove(Component);
-					Component->UnregisterComponent();
-					Component->DetachFromComponent(FDetachmentTransformRules::KeepWorldTransform);
-					Component->InvalidateLightingCache();
-					Component->Rename(nullptr, LandscapeProxy);
-					LandscapeProxy->LandscapeComponents.Add(Component);
-					Component->AttachToComponent(LandscapeProxy->GetRootComponent(), FAttachmentTransformRules::KeepWorldTransform);
-					
-					// clear transient mobile data
-					Component->MobileDataSourceHash.Invalidate();
-					Component->MobileMaterialInterfaces.Reset();
-					Component->MobileWeightmapTextures.Reset();
-					
-					Component->UpdateMaterialInstances();
-										
-					FFormatNamedArguments Args;
-					Args.Add(TEXT("ComponentName"), FText::FromString(Component->GetName()));
-				}
-				LandscapeProxy->UpdateCachedHasLayersContent();
-
-				for (ULandscapeHeightfieldCollisionComponent* Component : TargetSelectedCollisionComponents)
-				{
-					// Need to move or recreate all related data (Height map, Weight map, maybe collision components, allocation info)
+					GEditor->SelectNone(false, true);
+					GEditor->SelectActor(LandscapeProxy, true, false, true);
 
-					Component->GetLandscapeProxy()->CollisionComponents.Remove(Component);
-					Component->UnregisterComponent();
-					Component->DetachFromComponent(FDetachmentTransformRules::KeepWorldTransform);
-					Component->Rename(nullptr, LandscapeProxy);
-					LandscapeProxy->CollisionComponents.Add(Component);
-					Component->AttachToComponent(LandscapeProxy->GetRootComponent(), FAttachmentTransformRules::KeepWorldTransform);
+					GEditor->SelectNone(false, true);
 
-					// Move any foliage associated
-					AInstancedFoliageActor::MoveInstancesForComponentToCurrentLevel(Component);
-
-					FFormatNamedArguments Args;
-					Args.Add(TEXT("ComponentName"), FText::FromString(Component->GetName()));
-				}
-
-				GEditor->SelectNone(false, true);
-				GEditor->SelectActor(LandscapeProxy, true, false, true);
-
-				GEditor->SelectNone(false, true);
-
-				// Register our new components if destination landscape is registered in scene 
-				if (LandscapeProxy->GetRootComponent()->IsRegistered())
-				{
-					LandscapeProxy->RegisterAllComponents();
+					// Remove Selection
+					LandscapeInfo->ClearSelectedRegion(true);
 				}
-
-				for (ALandscapeProxy* Proxy : SelectProxies)
-				{
-					if (Proxy->GetRootComponent()->IsRegistered())
-					{
-						Proxy->RegisterAllComponents();
-					}
-				}
-
-				//Landscape->bLockLocation = (LandscapeInfo->XYtoComponentMap.Num() != Landscape->LandscapeComponents.Num());
-
-				// Remove Selection
-				LandscapeInfo->ClearSelectedRegion(true);
-
-				//EdMode->SetMaskEnable(Landscape->SelectedRegion.Num());
 			}
 		}
 	}
diff --git a/Engine/Source/Runtime/Foliage/Private/InstancedFoliage.cpp b/Engine/Source/Runtime/Foliage/Private/InstancedFoliage.cpp
--- a/Engine/Source/Runtime/Foliage/Private/InstancedFoliage.cpp
+++ b/Engine/Source/Runtime/Foliage/Private/InstancedFoliage.cpp
@@ -2607,6 +2607,12 @@ bool AInstancedFoliageActor::ContainsInstancesFromProceduralFoliageComponent(con
 }
 
 void AInstancedFoliageActor::MoveInstancesForComponentToCurrentLevel(UActorComponent* InComponent)
+{
+	UWorld* InWorld = InComponent->GetWorld();
+	MoveInstancesForComponentToLevel(InComponent, InWorld->GetCurrentLevel());
+}
+
+void AInstancedFoliageActor::MoveInstancesForComponentToLevel(UActorComponent* InComponent, ULevel* TargetLevel)
 {
 	if (!HasFoliageAttached(InComponent))
 	{
@@ -2614,11 +2620,11 @@ void AInstancedFoliageActor::MoveInstancesForComponentToCurrentLevel(UActorCompo
 		return;
 	}
 
-	UWorld* InWorld = InComponent->GetWorld();
-	AInstancedFoliageActor* NewIFA = AInstancedFoliageActor::GetInstancedFoliageActorForCurrentLevel(InWorld, true);
+	
+	AInstancedFoliageActor* NewIFA = AInstancedFoliageActor::GetInstancedFoliageActorForLevel(TargetLevel, true);
 	NewIFA->Modify();
 
-	for (TActorIterator<AInstancedFoliageActor> It(InWorld); It; ++It)
+	for (TActorIterator<AInstancedFoliageActor> It(InComponent->GetWorld()); It; ++It)
 	{
 		AInstancedFoliageActor* IFA = (*It);
 
diff --git a/Engine/Source/Runtime/Foliage/Public/InstancedFoliageActor.h b/Engine/Source/Runtime/Foliage/Public/InstancedFoliageActor.h
--- a/Engine/Source/Runtime/Foliage/Public/InstancedFoliageActor.h
+++ b/Engine/Source/Runtime/Foliage/Public/InstancedFoliageActor.h
@@ -166,6 +166,7 @@ public:
 
 	// Moves instances based on the specified component to the current streaming level
 	static FOLIAGE_API void MoveInstancesForComponentToCurrentLevel(UActorComponent* InComponent);
+	static FOLIAGE_API void MoveInstancesForComponentToLevel(UActorComponent* InComponent, ULevel* TargetLevel);
 
 	// Change all instances based on one component to a new component (possible in another level).
 	// The instances keep the same world locations
diff --git a/Engine/Source/Runtime/Landscape/Classes/Landscape.h b/Engine/Source/Runtime/Landscape/Classes/Landscape.h
--- a/Engine/Source/Runtime/Landscape/Classes/Landscape.h
+++ b/Engine/Source/Runtime/Landscape/Classes/Landscape.h
@@ -243,7 +243,7 @@ public:
 	LANDSCAPE_API static void CalcComponentIndicesNoOverlap(const int32 X1, const int32 Y1, const int32 X2, const int32 Y2, const int32 ComponentSizeQuads,
 		int32& ComponentIndexX1, int32& ComponentIndexY1, int32& ComponentIndexX2, int32& ComponentIndexY2);
 
-	static void SplitHeightmap(ULandscapeComponent* Comp, ALandscapeProxy* TargetProxy = nullptr, class FMaterialUpdateContext* InOutUpdateContext = nullptr, TArray<class FComponentRecreateRenderStateContext>* InOutRecreateRenderStateContext = nullptr, bool InReregisterComponent = true);
+	LANDSCAPE_API static void SplitHeightmap(ULandscapeComponent* Comp, ALandscapeProxy* TargetProxy = nullptr, class FMaterialUpdateContext* InOutUpdateContext = nullptr, TArray<class FComponentRecreateRenderStateContext>* InOutRecreateRenderStateContext = nullptr, bool InReregisterComponent = true);
 	
 	//~ Begin UObject Interface.
 	virtual void PreSave(const class ITargetPlatform* TargetPlatform) override;
diff --git a/Engine/Source/Runtime/Landscape/Classes/LandscapeInfo.h b/Engine/Source/Runtime/Landscape/Classes/LandscapeInfo.h
--- a/Engine/Source/Runtime/Landscape/Classes/LandscapeInfo.h
+++ b/Engine/Source/Runtime/Landscape/Classes/LandscapeInfo.h
@@ -169,7 +169,6 @@ public:
 	LANDSCAPE_API TSet<ULandscapeComponent*> GetSelectedComponents() const;
 	LANDSCAPE_API TSet<ULandscapeComponent*> GetSelectedRegionComponents() const;
 	LANDSCAPE_API void UpdateSelectedComponents(TSet<ULandscapeComponent*>& NewComponents, bool bIsComponentwise = true);
-	LANDSCAPE_API void SortSelectedComponents();
 	LANDSCAPE_API void ClearSelectedRegion(bool bIsComponentwise = true);
 
 	// only for use by the "add component" tool. Todo - move into the tool?
@@ -245,6 +244,10 @@ public:
 
 	/** Will clear all component dirty data */
 	LANDSCAPE_API void ClearDirtyData();
+
+	/** Moves Components to target level. Creates ALandscapeProxy if needed. */
+	LANDSCAPE_API ALandscapeProxy* MoveComponentsToLevel(const TArray<ULandscapeComponent*>& InComponents, ULevel* TargetLevel);
+
 #endif
 	/** Associates passed actor with this info object
  *  @param	Proxy		Landscape actor to register
diff --git a/Engine/Source/Runtime/Landscape/Private/LandscapeEdit.cpp b/Engine/Source/Runtime/Landscape/Private/LandscapeEdit.cpp
--- a/Engine/Source/Runtime/Landscape/Private/LandscapeEdit.cpp
+++ b/Engine/Source/Runtime/Landscape/Private/LandscapeEdit.cpp
@@ -50,10 +50,12 @@ LandscapeEdit.cpp: Landscape editing
 #include "LandscapeEditorModule.h"
 #include "LandscapeFileFormatInterface.h"
 #include "ComponentRecreateRenderStateContext.h"
+#include "ComponentReregisterContext.h"
 #include "Interfaces/ITargetPlatform.h"
 #include "Engine/TextureRenderTarget2D.h"
 #include "ScopedTransaction.h"
 #include "Editor.h"
+#include "PhysicalMaterials/PhysicalMaterial.h"
 #endif
 #include "Algo/Count.h"
 #include "Serialization/MemoryWriter.h"
@@ -3974,6 +3976,400 @@ void ULandscapeInfo::ClearDirtyData()
 	}
 }
 
+ALandscapeProxy* ULandscapeInfo::MoveComponentsToLevel(const TArray<ULandscapeComponent*>& InComponents, ULevel* TargetLevel)
+{
+	ALandscape* Landscape = LandscapeActor.Get();
+	check(Landscape != nullptr);
+
+	// Make sure references are in a different package (should be fixed up before calling this method)
+	// Check the Physical Material is same package with Landscape
+	if(Landscape->DefaultPhysMaterial && Landscape->DefaultPhysMaterial->GetOutermost() == Landscape->GetOutermost())
+	{
+		return nullptr;
+	}
+
+	// Check the LayerInfoObjects are not in same package as Landscape
+	for (int32 i = 0; i < Layers.Num(); ++i)
+	{
+		ULandscapeLayerInfoObject* LayerInfo = Layers[i].LayerInfoObj;
+		if (LayerInfo && LayerInfo->GetOutermost() == Landscape->GetOutermost())
+		{
+			return nullptr;
+		}
+	}
+
+	// Check the Landscape Materials are not in same package as moved components
+	for (ULandscapeComponent* Component : InComponents)
+	{
+		UMaterialInterface* LandscapeMaterial = Component->GetLandscapeMaterial();
+		if (LandscapeMaterial && LandscapeMaterial->GetOutermost() == Component->GetOutermost())
+		{
+			return nullptr;
+		}
+	}
+
+	struct FCompareULandscapeComponentBySectionBase
+	{
+		FORCEINLINE bool operator()(const ULandscapeComponent& A, const ULandscapeComponent& B) const
+		{
+			return (A.GetSectionBase().X == B.GetSectionBase().X) ? (A.GetSectionBase().Y < B.GetSectionBase().Y) : (A.GetSectionBase().X < B.GetSectionBase().X);
+		}
+	};
+	TArray<ULandscapeComponent*> ComponentsToMove(InComponents);
+	ComponentsToMove.Sort(FCompareULandscapeComponentBySectionBase());
+		
+	const int32 ComponentSizeVerts = Landscape->NumSubsections * (Landscape->SubsectionSizeQuads + 1);
+	const int32 NeedHeightmapSize = 1 << FMath::CeilLogTwo(ComponentSizeVerts);
+
+	TSet<ALandscapeProxy*> SelectProxies;
+	TSet<ULandscapeComponent*> TargetSelectedComponents;
+	TArray<ULandscapeHeightfieldCollisionComponent*> TargetSelectedCollisionComponents;
+	for (ULandscapeComponent* Component : ComponentsToMove)
+	{
+		SelectProxies.Add(Component->GetLandscapeProxy());
+		if (Component->GetLandscapeProxy()->GetOuter() != TargetLevel)
+		{
+			TargetSelectedComponents.Add(Component);
+		}
+
+		ULandscapeHeightfieldCollisionComponent* CollisionComp = Component->CollisionComponent.Get();
+		SelectProxies.Add(CollisionComp->GetLandscapeProxy());
+		if (CollisionComp->GetLandscapeProxy()->GetOuter() != TargetLevel)
+		{
+			TargetSelectedCollisionComponents.Add(CollisionComp);
+		}
+	}
+
+	// Check which ones are need for height map change
+	TSet<UTexture2D*> OldHeightmapTextures;
+	for (ULandscapeComponent* Component : TargetSelectedComponents)
+	{
+		Component->Modify();
+		OldHeightmapTextures.Add(Component->GetHeightmap());
+	}
+
+	// Need to split all the component which share Heightmap with selected components
+	TMap<ULandscapeComponent*, bool> HeightmapUpdateComponents;
+	HeightmapUpdateComponents.Reserve(TargetSelectedComponents.Num() * 4); // worst case
+	for (ULandscapeComponent* Component : TargetSelectedComponents)
+	{
+		// Search neighbor only
+		const int32 SearchX = Component->GetHeightmap()->Source.GetSizeX() / NeedHeightmapSize - 1;
+		const int32 SearchY = Component->GetHeightmap()->Source.GetSizeY() / NeedHeightmapSize - 1;
+		const FIntPoint ComponentBase = Component->GetSectionBase() / Component->ComponentSizeQuads;
+
+		for (int32 Y = -SearchY; Y <= SearchY; ++Y)
+		{
+			for (int32 X = -SearchX; X <= SearchX; ++X)
+			{
+				ULandscapeComponent* const Neighbor = XYtoComponentMap.FindRef(ComponentBase + FIntPoint(X, Y));
+				if (Neighbor && Neighbor->GetHeightmap() == Component->GetHeightmap() && !HeightmapUpdateComponents.Contains(Neighbor))
+				{
+					Neighbor->Modify();
+					bool bNeedsMoveToCurrentLevel = TargetSelectedComponents.Contains(Neighbor);
+					HeightmapUpdateComponents.Add(Neighbor, bNeedsMoveToCurrentLevel);
+				}
+			}
+		}
+	}
+
+	ALandscapeProxy* LandscapeProxy = GetLandscapeProxyForLevel(TargetLevel);
+	if (!LandscapeProxy)
+	{
+		FActorSpawnParameters SpawnParams;
+		SpawnParams.OverrideLevel = TargetLevel;
+		LandscapeProxy = TargetLevel->GetWorld()->SpawnActor<ALandscapeStreamingProxy>(SpawnParams);
+		// copy shared properties to this new proxy
+		LandscapeProxy->GetSharedProperties(Landscape);
+		LandscapeProxy->CreateLandscapeInfo();
+
+		// set proxy location
+		// by default first component location
+		ULandscapeComponent* FirstComponent = *TargetSelectedComponents.CreateConstIterator();
+		LandscapeProxy->GetRootComponent()->SetWorldLocationAndRotation(FirstComponent->GetComponentLocation(), FirstComponent->GetComponentRotation());
+		LandscapeProxy->LandscapeSectionOffset = FirstComponent->GetSectionBase();
+
+		// Hide(unregister) the new landscape if owning level currently in hidden state
+		if (LandscapeProxy->GetLevel()->bIsVisible == false)
+		{
+			LandscapeProxy->UnregisterAllComponents();
+		}
+	}
+
+	// Changing Heightmap format for selected components
+	for (const auto& HeightmapUpdateComponentPair : HeightmapUpdateComponents)
+	{
+		ALandscape::SplitHeightmap(HeightmapUpdateComponentPair.Key, HeightmapUpdateComponentPair.Value ? LandscapeProxy : nullptr);
+	}
+
+	// Delete if it is no referenced textures...
+	for (UTexture2D* Texture : OldHeightmapTextures)
+	{
+		Texture->SetFlags(RF_Transactional);
+		Texture->Modify();
+		Texture->MarkPackageDirty();
+		Texture->ClearFlags(RF_Standalone);
+	}
+
+	for (ALandscapeProxy* Proxy : SelectProxies)
+	{
+		Proxy->Modify();
+	}
+
+	LandscapeProxy->Modify();
+	LandscapeProxy->MarkPackageDirty();
+
+	// Handle XY-offset textures (these don't need splitting, as they aren't currently shared between components like heightmaps/weightmaps can be)
+	for (ULandscapeComponent* Component : TargetSelectedComponents)
+	{
+		if (Component->XYOffsetmapTexture)
+		{
+			Component->XYOffsetmapTexture->Modify();
+			Component->XYOffsetmapTexture->Rename(nullptr, LandscapeProxy);
+		}
+	}
+
+	// Change Weight maps...
+	{
+		FLandscapeEditDataInterface LandscapeEdit(this);
+		for (ULandscapeComponent* Component : TargetSelectedComponents)
+		{
+			Component->ReallocateWeightmaps(&LandscapeEdit, false, true, false, true, LandscapeProxy);
+			Component->ForEachLayer([&](const FGuid& LayerGuid, FLandscapeLayerComponentData& LayerData)
+			{
+				FScopedSetLandscapeEditingLayer Scope(Landscape, LayerGuid);
+				Component->ReallocateWeightmaps(&LandscapeEdit, true, true, false, true, LandscapeProxy);
+			});
+			Landscape->RequestLayersContentUpdateForceAll();
+		}
+
+		// Need to Repacking all the Weight map (to make it packed well...)
+		for (ALandscapeProxy* Proxy : SelectProxies)
+		{
+			Proxy->RemoveInvalidWeightmaps();
+		}
+	}
+
+	// Move the components to the Proxy actor
+	// This does not use the MoveSelectedActorsToCurrentLevel path as there is no support to only move certain components.
+	for (ULandscapeComponent* Component : TargetSelectedComponents)
+	{
+		// Need to move or recreate all related data (Height map, Weight map, maybe collision components, allocation info)
+		Component->GetLandscapeProxy()->LandscapeComponents.Remove(Component);
+		Component->UnregisterComponent();
+		Component->DetachFromComponent(FDetachmentTransformRules::KeepWorldTransform);
+		Component->InvalidateLightingCache();
+		Component->Rename(nullptr, LandscapeProxy);
+		LandscapeProxy->LandscapeComponents.Add(Component);
+		Component->AttachToComponent(LandscapeProxy->GetRootComponent(), FAttachmentTransformRules::KeepWorldTransform);
+
+		// clear transient mobile data
+		Component->MobileDataSourceHash.Invalidate();
+		Component->MobileMaterialInterfaces.Reset();
+		Component->MobileWeightmapTextures.Reset();
+
+		Component->UpdateMaterialInstances();
+	}
+	LandscapeProxy->UpdateCachedHasLayersContent();
+
+	for (ULandscapeHeightfieldCollisionComponent* Component : TargetSelectedCollisionComponents)
+	{
+		// Need to move or recreate all related data (Height map, Weight map, maybe collision components, allocation info)
+
+		Component->GetLandscapeProxy()->CollisionComponents.Remove(Component);
+		Component->UnregisterComponent();
+		Component->DetachFromComponent(FDetachmentTransformRules::KeepWorldTransform);
+		Component->Rename(nullptr, LandscapeProxy);
+		LandscapeProxy->CollisionComponents.Add(Component);
+		Component->AttachToComponent(LandscapeProxy->GetRootComponent(), FAttachmentTransformRules::KeepWorldTransform);
+
+		// Move any foliage associated
+		AInstancedFoliageActor::MoveInstancesForComponentToLevel(Component, TargetLevel);
+	}
+		
+	// Register our new components if destination landscape is registered in scene 
+	if (LandscapeProxy->GetRootComponent()->IsRegistered())
+	{
+		LandscapeProxy->RegisterAllComponents();
+	}
+
+	for (ALandscapeProxy* Proxy : SelectProxies)
+	{
+		if (Proxy->GetRootComponent()->IsRegistered())
+		{
+			Proxy->RegisterAllComponents();
+		}
+	}
+
+	return LandscapeProxy;
+}
+
+void ALandscape::SplitHeightmap(ULandscapeComponent* Comp, ALandscapeProxy* TargetProxy, FMaterialUpdateContext* InOutUpdateContext, TArray<FComponentRecreateRenderStateContext>* InOutRecreateRenderStateContext, bool InReregisterComponent)
+{
+	ULandscapeInfo* Info = Comp->GetLandscapeInfo();
+
+	// Make sure the heightmap UVs are powers of two.
+	int32 ComponentSizeVerts = Comp->NumSubsections * (Comp->SubsectionSizeQuads + 1);
+	int32 HeightmapSizeU = (1 << FMath::CeilLogTwo(ComponentSizeVerts));
+	int32 HeightmapSizeV = (1 << FMath::CeilLogTwo(ComponentSizeVerts));
+
+	ALandscapeProxy* SrcProxy = Comp->GetLandscapeProxy();
+	ALandscapeProxy* DstProxy = TargetProxy ? TargetProxy : SrcProxy;
+	SrcProxy->Modify();
+	DstProxy->Modify();
+
+	UTexture2D* OldHeightmapTexture = Comp->GetHeightmap(false);
+	UTexture2D* NewHeightmapTexture = NULL;
+	FVector4 OldHeightmapScaleBias = Comp->HeightmapScaleBias;
+	FVector4 NewHeightmapScaleBias = FVector4(1.0f / (float)HeightmapSizeU, 1.0f / (float)HeightmapSizeV, 0.0f, 0.0f);
+
+	{
+		// Read old data and split
+		FLandscapeEditDataInterface LandscapeEdit(Info);
+		TArray<uint8> HeightData;
+		HeightData.AddZeroed((1 + Comp->ComponentSizeQuads) * (1 + Comp->ComponentSizeQuads) * sizeof(uint16));
+		// Because of edge problem, normal would be just copy from old component data
+		TArray<uint8> NormalData;
+		NormalData.AddZeroed((1 + Comp->ComponentSizeQuads) * (1 + Comp->ComponentSizeQuads) * sizeof(uint16));
+		LandscapeEdit.GetHeightDataFast(Comp->GetSectionBase().X, Comp->GetSectionBase().Y, Comp->GetSectionBase().X + Comp->ComponentSizeQuads, Comp->GetSectionBase().Y + Comp->ComponentSizeQuads, (uint16*)HeightData.GetData(), 0, (uint16*)NormalData.GetData());
+
+		// Create the new heightmap texture
+		NewHeightmapTexture = DstProxy->CreateLandscapeTexture(HeightmapSizeU, HeightmapSizeV, TEXTUREGROUP_Terrain_Heightmap, TSF_BGRA8);
+		ULandscapeComponent::CreateEmptyTextureMips(NewHeightmapTexture, true);
+		Comp->HeightmapScaleBias = NewHeightmapScaleBias;
+		Comp->SetHeightmap(NewHeightmapTexture);
+
+		check(Comp->GetHeightmap(false) == Comp->GetHeightmap(true));
+		LandscapeEdit.SetHeightData(Comp->GetSectionBase().X, Comp->GetSectionBase().Y, Comp->GetSectionBase().X + Comp->ComponentSizeQuads, Comp->GetSectionBase().Y + Comp->ComponentSizeQuads, (uint16*)HeightData.GetData(), 0, false, (uint16*)NormalData.GetData());
+	}
+
+	// End material update
+	if (InOutUpdateContext != nullptr && InOutRecreateRenderStateContext != nullptr)
+	{
+		Comp->UpdateMaterialInstances(*InOutUpdateContext, *InOutRecreateRenderStateContext);
+	}
+	else
+	{
+		Comp->UpdateMaterialInstances();
+	}
+
+	// We disable automatic material update context, to manage it manually if we have a custom update context specified
+	GDisableAutomaticTextureMaterialUpdateDependencies = (InOutUpdateContext != nullptr);
+
+	NewHeightmapTexture->PostEditChange();
+
+	if (InOutUpdateContext != nullptr)
+	{
+		// Build a list of all unique materials the landscape uses
+		TArray<UMaterialInterface*> LandscapeMaterials;
+
+		int8 MaxLOD = FMath::CeilLogTwo(Comp->SubsectionSizeQuads + 1) - 1;
+
+		for (int8 LODIndex = 0; LODIndex < MaxLOD; ++LODIndex)
+		{
+			UMaterialInterface* Material = Comp->GetLandscapeMaterial(LODIndex);
+			LandscapeMaterials.AddUnique(Material);
+		}
+
+		TSet<UMaterial*> BaseMaterialsThatUseThisTexture;
+
+		for (UMaterialInterface* MaterialInterface : LandscapeMaterials)
+		{
+			if (DoesMaterialUseTexture(MaterialInterface, NewHeightmapTexture))
+			{
+				UMaterial* Material = MaterialInterface->GetMaterial();
+				bool MaterialAlreadyCompute = false;
+				BaseMaterialsThatUseThisTexture.Add(Material, &MaterialAlreadyCompute);
+
+				if (!MaterialAlreadyCompute)
+				{
+					if (Material->IsTextureForceRecompileCacheRessource(NewHeightmapTexture))
+					{
+						InOutUpdateContext->AddMaterial(Material);
+						Material->UpdateMaterialShaderCacheAndTextureReferences();
+					}
+				}
+			}
+		}
+	}
+
+	GDisableAutomaticTextureMaterialUpdateDependencies = false;
+
+#if WITH_EDITORONLY_DATA
+	check(Comp->GetLandscapeProxy()->HasLayersContent() == DstProxy->CanHaveLayersContent());
+	if (Comp->GetLandscapeProxy()->HasLayersContent() && DstProxy->CanHaveLayersContent())
+	{
+		FLandscapeLayersTexture2DCPUReadBackResource* NewCPUReadBackResource = new FLandscapeLayersTexture2DCPUReadBackResource(NewHeightmapTexture->Source.GetSizeX(), NewHeightmapTexture->Source.GetSizeY(), NewHeightmapTexture->GetPixelFormat(), NewHeightmapTexture->Source.GetNumMips());
+		BeginInitResource(NewCPUReadBackResource);
+		DstProxy->HeightmapsCPUReadBack.Add(NewHeightmapTexture, NewCPUReadBackResource);
+
+		// Free OldHeightmapTexture's CPUReadBackResource if not used by any component
+		bool FreeCPUReadBack = true;
+		for (ULandscapeComponent* Component : SrcProxy->LandscapeComponents)
+		{
+			if (Component != Comp && Component->GetHeightmap(false) == OldHeightmapTexture)
+			{
+				FreeCPUReadBack = false;
+				break;
+			}
+		}
+		if (FreeCPUReadBack)
+		{
+			FLandscapeLayersTexture2DCPUReadBackResource** OldCPUReadBackResource = SrcProxy->HeightmapsCPUReadBack.Find(OldHeightmapTexture);
+			if (OldCPUReadBackResource)
+			{
+				if (FLandscapeLayersTexture2DCPUReadBackResource* ResourceToDelete = *OldCPUReadBackResource)
+				{
+					ReleaseResourceAndFlush(ResourceToDelete);
+					delete ResourceToDelete;
+					SrcProxy->HeightmapsCPUReadBack.Remove(OldHeightmapTexture);
+				}
+			}
+		}
+
+		// Move layer content to new layer heightmap
+		FLandscapeEditDataInterface LandscapeEdit(Info);
+		ALandscape* Landscape = Info->LandscapeActor.Get();
+		Comp->ForEachLayer([&](const FGuid& LayerGuid, FLandscapeLayerComponentData& LayerData)
+		{
+			UTexture2D* OldLayerHeightmap = LayerData.HeightmapData.Texture;
+			if (OldLayerHeightmap != nullptr)
+			{
+				FScopedSetLandscapeEditingLayer Scope(Landscape, LayerGuid);
+				// Read old data and split
+				TArray<uint8> LayerHeightData;
+				LayerHeightData.AddZeroed((1 + Comp->ComponentSizeQuads) * (1 + Comp->ComponentSizeQuads) * sizeof(uint16));
+				// Because of edge problem, normal would be just copy from old component data
+				TArray<uint8> LayerNormalData;
+				LayerNormalData.AddZeroed((1 + Comp->ComponentSizeQuads) * (1 + Comp->ComponentSizeQuads) * sizeof(uint16));
+
+				// Read using old heightmap scale/bias
+				Comp->HeightmapScaleBias = OldHeightmapScaleBias;
+				LandscapeEdit.GetHeightDataFast(Comp->GetSectionBase().X, Comp->GetSectionBase().Y, Comp->GetSectionBase().X + Comp->ComponentSizeQuads, Comp->GetSectionBase().Y + Comp->ComponentSizeQuads, (uint16*)LayerHeightData.GetData(), 0, (uint16*)LayerNormalData.GetData());
+				// Restore new heightmap scale/bias
+				Comp->HeightmapScaleBias = NewHeightmapScaleBias;
+				{
+					UTexture2D* LayerHeightmapTexture = DstProxy->CreateLandscapeTexture(HeightmapSizeU, HeightmapSizeV, TEXTUREGROUP_Terrain_Heightmap, TSF_BGRA8);
+					ULandscapeComponent::CreateEmptyTextureMips(LayerHeightmapTexture, true);
+					LayerHeightmapTexture->PostEditChange();
+					// Set Layer heightmap texture
+					LayerData.HeightmapData.Texture = LayerHeightmapTexture;
+					LandscapeEdit.SetHeightData(Comp->GetSectionBase().X, Comp->GetSectionBase().Y, Comp->GetSectionBase().X + Comp->ComponentSizeQuads, Comp->GetSectionBase().Y + Comp->ComponentSizeQuads, (uint16*)LayerHeightData.GetData(), 0, false, (uint16*)LayerNormalData.GetData());
+				}
+			}
+		});
+
+		Landscape->RequestLayersContentUpdateForceAll();
+	}
+#endif
+
+	// Reregister
+	if (InReregisterComponent)
+	{
+		FComponentReregisterContext ReregisterContext(Comp);
+	}
+}
+
 namespace
 {
 	inline float AdjustStaticLightingResolution(float StaticLightingResolution, int32 NumSubsections, int32 SubsectionSizeQuads, int32 ComponentSizeQuads)
@@ -4780,18 +5176,6 @@ void ULandscapeInfo::UpdateSelectedComponents(TSet<ULandscapeComponent*>& NewCom
 	}
 }
 
-void ULandscapeInfo::SortSelectedComponents()
-{
-	struct FCompareULandscapeComponentBySectionBase
-	{
-		FORCEINLINE bool operator()(const ULandscapeComponent& A, const ULandscapeComponent& B) const
-		{
-			return (A.GetSectionBase().X == B.GetSectionBase().X) ? (A.GetSectionBase().Y < B.GetSectionBase().Y) : (A.GetSectionBase().X < B.GetSectionBase().X);
-		}
-	};
-	SelectedComponents.Sort(FCompareULandscapeComponentBySectionBase());
-}
-
 void ULandscapeInfo::ClearSelectedRegion(bool bIsComponentwise /*= true*/)
 {
 	TSet<ULandscapeComponent*> NewComponents;
