LevelGrid: Partial Editgrate of CL 9692186 from //dev-editor-openworld

Fix Rebuild Material Instances in LandscapeEdMode wasn't called on all Proxies

#rnx

diff --git a/Engine/Source/Editor/LandscapeEditor/Private/LandscapeEditorDetailCustomization_TargetLayers.cpp b/Engine/Source/Editor/LandscapeEditor/Private/LandscapeEditorDetailCustomization_TargetLayers.cpp
--- a/Engine/Source/Editor/LandscapeEditor/Private/LandscapeEditorDetailCustomization_TargetLayers.cpp
+++ b/Engine/Source/Editor/LandscapeEditor/Private/LandscapeEditorDetailCustomization_TargetLayers.cpp
@@ -1141,7 +1141,7 @@ void FLandscapeEditorCustomNodeBuilder_TargetLayers::OnRebuildMICs(const TShared
 {
 	if (Target->LandscapeInfo.IsValid())
 	{
-		Target->LandscapeInfo.Get()->GetLandscapeProxy()->UpdateAllComponentMaterialInstances();
+		Target->LandscapeInfo.Get()->UpdateAllComponentMaterialInstances();
 	}
 }
 
diff --git a/Engine/Source/Runtime/Landscape/Classes/LandscapeInfo.h b/Engine/Source/Runtime/Landscape/Classes/LandscapeInfo.h
--- a/Engine/Source/Runtime/Landscape/Classes/LandscapeInfo.h
+++ b/Engine/Source/Runtime/Landscape/Classes/LandscapeInfo.h
@@ -245,6 +245,9 @@ public:
 	/** Will clear all component dirty data */
 	LANDSCAPE_API void ClearDirtyData();
 
+	/** Will call UpdateAllComponentMaterialInstances on all LandscapeProxies */
+	LANDSCAPE_API void UpdateAllComponentMaterialInstances();
+
 	/** Moves Components to target level. Creates ALandscapeProxy if needed. */
 	LANDSCAPE_API ALandscapeProxy* MoveComponentsToLevel(const TArray<ULandscapeComponent*>& InComponents, ULevel* TargetLevel);
 
diff --git a/Engine/Source/Runtime/Landscape/Private/LandscapeEdit.cpp b/Engine/Source/Runtime/Landscape/Private/LandscapeEdit.cpp
--- a/Engine/Source/Runtime/Landscape/Private/LandscapeEdit.cpp
+++ b/Engine/Source/Runtime/Landscape/Private/LandscapeEdit.cpp
@@ -3680,15 +3680,7 @@ void ALandscapeProxy::PostEditImport()
 		CreateLandscapeInfo();
 	}
 
-	for (int32 ComponentIndex = 0; ComponentIndex < LandscapeComponents.Num(); ++ComponentIndex)
-	{
-		ULandscapeComponent* Comp = LandscapeComponents[ComponentIndex];
-		if (Comp)
-		{
-			// Update the MIC
-			Comp->UpdateMaterialInstances();
-		}
-	}
+	UpdateAllComponentMaterialInstances();
 }
 
 void ALandscape::PostEditMove(bool bFinished)
@@ -3976,6 +3968,14 @@ void ULandscapeInfo::ClearDirtyData()
 	}
 }
 
+void ULandscapeInfo::UpdateAllComponentMaterialInstances()
+{
+	ForAllLandscapeProxies([](ALandscapeProxy* Proxy)
+	{
+		Proxy->UpdateAllComponentMaterialInstances();
+	});
+}
+
 ALandscapeProxy* ULandscapeInfo::MoveComponentsToLevel(const TArray<ULandscapeComponent*>& InComponents, ULevel* TargetLevel)
 {
 	ALandscape* Landscape = LandscapeActor.Get();
