Level Grid: EditGrate CL 974968

#rnx

diff --git a/Engine/Source/Runtime/Landscape/Classes/Landscape.h b/Engine/Source/Runtime/Landscape/Classes/Landscape.h
--- a/Engine/Source/Runtime/Landscape/Classes/Landscape.h
+++ b/Engine/Source/Runtime/Landscape/Classes/Landscape.h
@@ -333,6 +333,7 @@ public:
 	
 	LANDSCAPE_API void ToggleCanHaveLayersContent();
 	LANDSCAPE_API void ForceUpdateLayersContent(bool bIntermediateRender = false);
+	LANDSCAPE_API void InitializeLandscapeLayersWeightmapUsage();
 
 private:
 	void TickLayers(float DeltaTime, ELevelTick TickType, FActorTickFunction& ThisTickFunction);
@@ -399,8 +400,7 @@ private:
 	void CommitDeferredCopyLayersTexture();
 
 	void InitializeLayers();
-	void InitializeLandscapeLayersWeightmapUsage();
-
+	
 	void PrintLayersDebugRT(const FString& InContext, UTextureRenderTarget2D* InDebugRT, uint8 InMipRender = 0, bool InOutputHeight = true, bool InOutputNormals = false) const;
 	void PrintLayersDebugTextureResource(const FString& InContext, FTextureResource* InTextureResource, uint8 InMipRender = 0, bool InOutputHeight = true, bool InOutputNormals = false) const;
 	void PrintLayersDebugHeightData(const FString& InContext, const TArray<FColor>& InHeightmapData, const FIntPoint& InDataSize, uint8 InMipRender, bool InOutputNormals = false) const;
diff --git a/Engine/Source/Runtime/Landscape/Classes/LandscapeInfo.h b/Engine/Source/Runtime/Landscape/Classes/LandscapeInfo.h
--- a/Engine/Source/Runtime/Landscape/Classes/LandscapeInfo.h
+++ b/Engine/Source/Runtime/Landscape/Classes/LandscapeInfo.h
@@ -245,12 +245,12 @@ public:
 	/** Will clear all component dirty data */
 	LANDSCAPE_API void ClearDirtyData();
 
+	/** Moves Components to target level. Creates ALandscapeProxy if needed. */
+	LANDSCAPE_API ALandscapeProxy* MoveComponentsToLevel(const TArray<ULandscapeComponent*>& InComponents, ULevel* TargetLevel, FName NewProxyName = NAME_None);
+
 	/** Will call UpdateAllComponentMaterialInstances on all LandscapeProxies */
 	LANDSCAPE_API void UpdateAllComponentMaterialInstances();
 
-	/** Moves Components to target level. Creates ALandscapeProxy if needed. */
-	LANDSCAPE_API ALandscapeProxy* MoveComponentsToLevel(const TArray<ULandscapeComponent*>& InComponents, ULevel* TargetLevel);
-
 #endif
 	/** Associates passed actor with this info object
  *  @param	Proxy		Landscape actor to register
diff --git a/Engine/Source/Runtime/Landscape/Private/LandscapeEdit.cpp b/Engine/Source/Runtime/Landscape/Private/LandscapeEdit.cpp
--- a/Engine/Source/Runtime/Landscape/Private/LandscapeEdit.cpp
+++ b/Engine/Source/Runtime/Landscape/Private/LandscapeEdit.cpp
@@ -3976,7 +3976,7 @@ void ULandscapeInfo::UpdateAllComponentMaterialInstances()
 	});
 }
 
-ALandscapeProxy* ULandscapeInfo::MoveComponentsToLevel(const TArray<ULandscapeComponent*>& InComponents, ULevel* TargetLevel)
+ALandscapeProxy* ULandscapeInfo::MoveComponentsToLevel(const TArray<ULandscapeComponent*>& InComponents, ULevel* TargetLevel, FName NewProxyName)
 {
 	ALandscape* Landscape = LandscapeActor.Get();
 	check(Landscape != nullptr);
@@ -4077,11 +4077,14 @@ ALandscapeProxy* ULandscapeInfo::MoveComponentsToLevel(const TArray<ULandscapeCo
 	if (!LandscapeProxy)
 	{
 		FActorSpawnParameters SpawnParams;
+		SpawnParams.Name = NewProxyName;
 		SpawnParams.OverrideLevel = TargetLevel;
 		LandscapeProxy = TargetLevel->GetWorld()->SpawnActor<ALandscapeStreamingProxy>(SpawnParams);
+		
 		// copy shared properties to this new proxy
 		LandscapeProxy->GetSharedProperties(Landscape);
 		LandscapeProxy->CreateLandscapeInfo();
+		LandscapeProxy->SetActorLabel(LandscapeProxy->GetName());
 
 		// set proxy location
 		// by default first component location
