LevelGrid: EdiGrate CL 9839759

MoveSplinesToLevel now part of LandscapeInfo

#rnx

diff --git a/Engine/Source/Editor/LandscapeEditor/Private/LandscapeEdModeSplineTools.cpp b/Engine/Source/Editor/LandscapeEditor/Private/LandscapeEdModeSplineTools.cpp
--- a/Engine/Source/Editor/LandscapeEditor/Private/LandscapeEdModeSplineTools.cpp
+++ b/Engine/Source/Editor/LandscapeEditor/Private/LandscapeEdModeSplineTools.cpp
@@ -96,13 +96,9 @@ public:
 	virtual void SetEditRenderType() override { GLandscapeEditRenderMode = ELandscapeEditRenderMode::None | (GLandscapeEditRenderMode & ELandscapeEditRenderMode::BitMaskForMask); }
 	virtual bool SupportsMask() override { return false; }
 
-	void CreateSplineComponent(ALandscapeProxy* Landscape, FVector Scale3D)
+	void CreateSplineComponent(ALandscapeProxy* Proxy, FVector Scale3D)
 	{
-		Landscape->Modify();
-		Landscape->SplineComponent = NewObject<ULandscapeSplinesComponent>(Landscape, NAME_None, RF_Transactional);
-		Landscape->SplineComponent->SetRelativeScale3D_Direct(Scale3D);
-		Landscape->SplineComponent->AttachToComponent(Landscape->GetRootComponent(), FAttachmentTransformRules::KeepRelativeTransform);
-		Landscape->SplineComponent->ShowSplineEditorMesh(true);
+		Proxy->CreateSplineComponent(Scale3D);
 	}
 
 	void UpdatePropertiesWindows()
@@ -866,10 +862,10 @@ public:
 			ALandscapeProxy* FromProxy = LandscapeSplinesComp ? Cast<ALandscapeProxy>(LandscapeSplinesComp->GetOuter()) : nullptr;
 			if (FromProxy)
 			{
+				ULandscapeInfo* ProxyLandscapeInfo = FromProxy->GetLandscapeInfo();
+				check(ProxyLandscapeInfo);
 				if (!ToLandscape)
 				{
-					ULandscapeInfo* ProxyLandscapeInfo = FromProxy->GetLandscapeInfo();
-					check(ProxyLandscapeInfo);
 					ToLandscape = ProxyLandscapeInfo->GetCurrentLevelLandscapeProxy(true);
 					if (!ToLandscape)
 					{
@@ -878,133 +874,10 @@ public:
 					}
 				}
 
-				if (ToLandscape != FromProxy)
-				{
-					ToLandscape->Modify();
-					if (ToLandscape->SplineComponent == nullptr)
-					{
-						CreateSplineComponent(ToLandscape, FromProxy->SplineComponent->GetRelativeScale3D());
-						check(ToLandscape->SplineComponent);
-					}
-					ToLandscape->SplineComponent->Modify();
-
-					const FTransform OldToNewTransform =
-						FromProxy->SplineComponent->GetComponentTransform().GetRelativeTransform(ToLandscape->SplineComponent->GetComponentTransform());
-
-					if (FromProxies.Find(FromProxy) == nullptr)
-					{
-						FromProxies.Add(FromProxy);
-						FromProxy->Modify();
-						FromProxy->SplineComponent->Modify();
-						FromProxy->SplineComponent->MarkRenderStateDirty();
-					}
-
-					// Delete all Mesh Components associated with the ControlPoint. (Will get recreated in UpdateSplinePoints)
-					if (ControlPoint->LocalMeshComponent)
-					{
-						ControlPoint->LocalMeshComponent->Modify();
-						ControlPoint->LocalMeshComponent->UnregisterComponent();
-						ControlPoint->LocalMeshComponent->DestroyComponent();
-						FromProxy->SplineComponent->Modify();
-						FromProxy->SplineComponent->MeshComponentLocalOwnersMap.Remove(ControlPoint->LocalMeshComponent);
-						ControlPoint->LocalMeshComponent = nullptr;
-					}
-
-					TMap<ULandscapeSplinesComponent*, UControlPointMeshComponent*> ForeignMeshComponents = ControlPoint->GetForeignMeshComponents();
-					for (auto Pair : ForeignMeshComponents)
-					{
-						Pair.Key->RemoveForeignMeshComponent(ControlPoint, Pair.Value);
-						Pair.Value->Modify();
-						Pair.Value->UnregisterComponent();
-						Pair.Value->DestroyComponent();
-					}
-					
-					// Move control point to new level
-					FromProxy->SplineComponent->ControlPoints.Remove(ControlPoint);
-					ControlPoint->Rename(nullptr, ToLandscape->SplineComponent);
-					ToLandscape->SplineComponent->ControlPoints.Add(ControlPoint);
-
-					ControlPoint->Location = OldToNewTransform.TransformPosition(ControlPoint->Location);
-
-					const bool bUpdateCollision = true; // default value
-					const bool bUpdateSegments = false; // done in next loop
-					const bool bUpdateMeshLevel = false; // no need because mesh have been deleted
-					ControlPoint->UpdateSplinePoints(bUpdateCollision, bUpdateSegments, bUpdateMeshLevel);
-				}
-			}
-		}
-
-		for (ULandscapeSplineSegment* Segment : SelectedSplineSegments)
-		{
-			ULandscapeSplinesComponent* LandscapeSplinesComp = Segment->GetOuterULandscapeSplinesComponent();
-			ALandscapeProxy* FromProxy = LandscapeSplinesComp ? Cast<ALandscapeProxy>(LandscapeSplinesComp->GetOuter()) : nullptr;
-			if (FromProxy)
-			{
-				if (!ToLandscape)
-				{
-					ULandscapeInfo* ProxyLandscapeInfo = FromProxy->GetLandscapeInfo();
-					check(ProxyLandscapeInfo);
-					ToLandscape = ProxyLandscapeInfo->GetCurrentLevelLandscapeProxy(true);
-					if (!ToLandscape)
-					{
-						// No Landscape Proxy, don't support for creating only for Spline now
-						return;
-					}
-				}
-
-				if (ToLandscape != FromProxy)
-				{
-					ToLandscape->Modify();
-					if (ToLandscape->SplineComponent == nullptr)
-					{
-						CreateSplineComponent(ToLandscape, FromProxy->SplineComponent->GetRelativeScale3D());
-						check(ToLandscape->SplineComponent);
-					}
-					ToLandscape->SplineComponent->Modify();
-
-					if (FromProxies.Find(FromProxy) == nullptr)
-					{
-						FromProxies.Add(FromProxy);
-						FromProxy->Modify();
-						FromProxy->SplineComponent->Modify();
-						FromProxy->SplineComponent->MarkRenderStateDirty();
-					}
-
-					// Delete all Mesh Components associated with the Segment. (Will get recreated in UpdateSplinePoints)
-					for (auto* MeshComponent : Segment->LocalMeshComponents)
-					{
-						MeshComponent->Modify();
-						MeshComponent->UnregisterComponent();
-						MeshComponent->DestroyComponent();
-						FromProxy->Modify();
-						FromProxy->SplineComponent->MeshComponentLocalOwnersMap.Remove(MeshComponent);
-					}
-					Segment->LocalMeshComponents.Empty();
-
-					TMap<ULandscapeSplinesComponent*, TArray<USplineMeshComponent*>> ForeignMeshComponents = Segment->GetForeignMeshComponents();
-					for (auto Pair : ForeignMeshComponents)
-					{
-						Pair.Key->RemoveAllForeignMeshComponents(Segment);
-						for (auto MeshComponent : Pair.Value)
-						{
-							MeshComponent->Modify();
-							MeshComponent->UnregisterComponent();
-							MeshComponent->DestroyComponent();
-						}
-					}
-										
-					// Move segment to new level
-					FromProxy->SplineComponent->Segments.Remove(Segment);
-					Segment->Rename(nullptr, ToLandscape->SplineComponent);
-					ToLandscape->SplineComponent->Segments.Add(Segment);
-
-					const bool bUpdateCollision = true; // default value
-					const bool bUpdateMeshLevel = false; // no need because mesh have been deleted 
-					Segment->UpdateSplinePoints(bUpdateCollision, bUpdateMeshLevel);
-				}
+				ProxyLandscapeInfo->MoveSplineToLevel(ControlPoint, ToLandscape->GetLevel());
 			}
 		}
-
+				
 		if (ToLandscape && ToLandscape->SplineComponent)
 		{
 			if (!ToLandscape->SplineComponent->IsRegistered())
diff --git a/Engine/Source/Runtime/Landscape/Classes/LandscapeInfo.h b/Engine/Source/Runtime/Landscape/Classes/LandscapeInfo.h
--- a/Engine/Source/Runtime/Landscape/Classes/LandscapeInfo.h
+++ b/Engine/Source/Runtime/Landscape/Classes/LandscapeInfo.h
@@ -17,6 +17,9 @@ class ULandscapeLayerInfoObject;
 class ULevel;
 class UMaterialInstanceConstant;
 struct FLandscapeEditorLayerSettings;
+class ULandscapeSplinesComponent;
+class ULandscapeSplineControlPoint;
+class ULandscapeSplineSegment;
 class ULandscapeHeightfieldCollisionComponent;
 
 /** Structure storing Collision for LandscapeComponent Add */
@@ -248,9 +251,14 @@ public:
 	/** Moves Components to target level. Creates ALandscapeProxy if needed. */
 	LANDSCAPE_API ALandscapeProxy* MoveComponentsToLevel(const TArray<ULandscapeComponent*>& InComponents, ULevel* TargetLevel, FName NewProxyName = NAME_None);
 
+	/** Moves Splines connected to this control point to target level. Creates ULandscapeSplineComponent if needed. */
+	LANDSCAPE_API void MoveSplineToLevel(ULandscapeSplineControlPoint* InControlPoint, ULevel* TargetLevel);
+
+	/** Moves all Splines to target level. Creates ULandscapeSplineComponent if needed. */
+	LANDSCAPE_API void MoveSplinesToLevel(ULandscapeSplinesComponent* InSplineComponent, ULevel* TargetLevel);
+
 	/** Will call UpdateAllComponentMaterialInstances on all LandscapeProxies */
 	LANDSCAPE_API void UpdateAllComponentMaterialInstances();
-
 #endif
 	/** Associates passed actor with this info object
  *  @param	Proxy		Landscape actor to register
@@ -273,4 +281,10 @@ public:
 	/** Server doesn't have ULandscapeComponent use CollisionComponents instead to get height on landscape */
 	LANDSCAPE_API void RegisterCollisionComponent(ULandscapeHeightfieldCollisionComponent* Component);
 	LANDSCAPE_API void UnregisterCollisionComponent(ULandscapeHeightfieldCollisionComponent* Component);
+
+#if WITH_EDITOR
+private:
+	void MoveSegmentToLandscape(ULandscapeSplineSegment* InSegment, ALandscapeProxy* FromProxy, ALandscapeProxy* ToProxy);
+	void MoveControlPointToLandscape(ULandscapeSplineControlPoint* InControlPoint, ALandscapeProxy* FromProxy, ALandscapeProxy* ToProxy);
+#endif
 };
diff --git a/Engine/Source/Runtime/Landscape/Classes/LandscapeProxy.h b/Engine/Source/Runtime/Landscape/Classes/LandscapeProxy.h
--- a/Engine/Source/Runtime/Landscape/Classes/LandscapeProxy.h
+++ b/Engine/Source/Runtime/Landscape/Classes/LandscapeProxy.h
@@ -878,6 +878,8 @@ public:
 	LANDSCAPE_API FTransform LandscapeActorToWorld() const;
 
 #if WITH_EDITOR
+	LANDSCAPE_API void CreateSplineComponent(const FVector& Scale3D);
+
 	virtual void PostEditChangeProperty(FPropertyChangedEvent& PropertyChangedEvent) override;
 	virtual void PostEditImport() override;
 	//~ End UObject Interface
diff --git a/Engine/Source/Runtime/Landscape/Classes/LandscapeSplineControlPoint.h b/Engine/Source/Runtime/Landscape/Classes/LandscapeSplineControlPoint.h
--- a/Engine/Source/Runtime/Landscape/Classes/LandscapeSplineControlPoint.h
+++ b/Engine/Source/Runtime/Landscape/Classes/LandscapeSplineControlPoint.h
@@ -270,4 +270,5 @@ public:
 #endif // WITH_EDITOR
 
 	friend class FLandscapeToolSplines;
+	friend class ULandscapeInfo;
 };
diff --git a/Engine/Source/Runtime/Landscape/Classes/LandscapeSplineSegment.h b/Engine/Source/Runtime/Landscape/Classes/LandscapeSplineSegment.h
--- a/Engine/Source/Runtime/Landscape/Classes/LandscapeSplineSegment.h
+++ b/Engine/Source/Runtime/Landscape/Classes/LandscapeSplineSegment.h
@@ -364,4 +364,5 @@ public:
 	//~ End UObject Interface
 
 	friend class FLandscapeToolSplines;
+	friend class ULandscapeInfo;
 };
diff --git a/Engine/Source/Runtime/Landscape/Classes/LandscapeSplinesComponent.h b/Engine/Source/Runtime/Landscape/Classes/LandscapeSplinesComponent.h
--- a/Engine/Source/Runtime/Landscape/Classes/LandscapeSplinesComponent.h
+++ b/Engine/Source/Runtime/Landscape/Classes/LandscapeSplinesComponent.h
@@ -154,6 +154,8 @@ public:
 	bool ModifySplines(bool bAlwaysMarkDirty = true);
 
 #if WITH_EDITOR
+	bool HasAnyControlPointsOrSegments() const { return ControlPoints.Num() > 0 || Segments.Num() > 0; }
+
 	virtual void ShowSplineEditorMesh(bool bShow);
 
 	// Rebuilds all spline points and meshes for all spline control points and segments in this splines component
@@ -226,6 +228,7 @@ public:
 	friend class FLandscapeSplinesSceneProxy;
 	friend class ULandscapeSplineControlPoint;
 	friend class ULandscapeSplineSegment;
+	friend class ULandscapeInfo;
 #if WITH_EDITOR
 	// TODO - move this out of ULandscapeInfo
 	friend bool ULandscapeInfo::ApplySplinesInternal(bool bOnlySelected, ALandscapeProxy* Landscape, TSet<ULandscapeComponent*>* OutModifiedComponents, bool bMarkPackageDirty);
diff --git a/Engine/Source/Runtime/Landscape/Private/LandscapeSplines.cpp b/Engine/Source/Runtime/Landscape/Private/LandscapeSplines.cpp
--- a/Engine/Source/Runtime/Landscape/Private/LandscapeSplines.cpp
+++ b/Engine/Source/Runtime/Landscape/Private/LandscapeSplines.cpp
@@ -3224,6 +3224,188 @@ void ULandscapeSplineSegment::PostEditChangeProperty(FPropertyChangedEvent& Prop
 		}
 	}
 }
+
+void ALandscapeProxy::CreateSplineComponent(const FVector& Scale3D)
+{
+	check(SplineComponent == nullptr);
+	Modify();
+	SplineComponent = NewObject<ULandscapeSplinesComponent>(this, NAME_None, RF_Transactional);
+	SplineComponent->SetRelativeScale3D_Direct(Scale3D);
+	SplineComponent->AttachToComponent(GetRootComponent(), FAttachmentTransformRules::KeepRelativeTransform);
+	SplineComponent->ShowSplineEditorMesh(true);
+}
+
+void FindStartingControlPoints(ULandscapeSplineControlPoint* ControlPoint, TSet<ULandscapeSplineControlPoint*>& StartingControlPoints)
+{
+	bool bIsStartingControlPoint = true;
+	for (const FLandscapeSplineConnection& Connection : ControlPoint->ConnectedSegments)
+	{
+		if (Connection.End)
+		{
+			ULandscapeSplineControlPoint* Start = Connection.Segment->Connections[0].ControlPoint;
+			check(Start != ControlPoint);
+			FindStartingControlPoints(Start, StartingControlPoints);
+			bIsStartingControlPoint = false;
+		}
+	}
+
+	if (bIsStartingControlPoint)
+	{
+		StartingControlPoints.Add(ControlPoint);
+	}
+}
+
+void ULandscapeInfo::MoveControlPointToLandscape(ULandscapeSplineControlPoint* InControlPoint, ALandscapeProxy* FromProxy, ALandscapeProxy* ToProxy)
+{
+	if (ToProxy->SplineComponent == nullptr)
+	{
+		ToProxy->CreateSplineComponent(FromProxy->SplineComponent->GetRelativeScale3D());
+		check(ToProxy->SplineComponent);
+	}
+
+	if (InControlPoint->GetOuterULandscapeSplinesComponent() == ToProxy->SplineComponent)
+	{
+		return;
+	}
+	check(InControlPoint->GetOuterULandscapeSplinesComponent() == FromProxy->SplineComponent);
+
+	ToProxy->SplineComponent->Modify();
+
+	const FTransform OldToNewTransform =
+		FromProxy->SplineComponent->GetComponentTransform().GetRelativeTransform(ToProxy->SplineComponent->GetComponentTransform());
+		
+	// Delete all Mesh Components associated with the ControlPoint. (Will get recreated in UpdateSplinePoints)
+	if (InControlPoint->LocalMeshComponent)
+	{
+		InControlPoint->LocalMeshComponent->Modify();
+		InControlPoint->LocalMeshComponent->UnregisterComponent();
+		InControlPoint->LocalMeshComponent->DestroyComponent();
+		FromProxy->SplineComponent->Modify();
+		FromProxy->SplineComponent->MeshComponentLocalOwnersMap.Remove(InControlPoint->LocalMeshComponent);
+		InControlPoint->LocalMeshComponent = nullptr;
+	}
+
+	TMap<ULandscapeSplinesComponent*, UControlPointMeshComponent*> ForeignMeshComponents = InControlPoint->GetForeignMeshComponents();
+	for (auto Pair : ForeignMeshComponents)
+	{
+		Pair.Key->RemoveForeignMeshComponent(InControlPoint, Pair.Value);
+		Pair.Value->Modify();
+		Pair.Value->UnregisterComponent();
+		Pair.Value->DestroyComponent();
+	}
+
+	// Move control point to new level
+	FromProxy->SplineComponent->ControlPoints.Remove(InControlPoint);
+	InControlPoint->Rename(nullptr, ToProxy->SplineComponent);
+	ToProxy->SplineComponent->ControlPoints.Add(InControlPoint);
+
+	InControlPoint->Location = OldToNewTransform.TransformPosition(InControlPoint->Location);
+
+	const bool bUpdateCollision = true; // default value
+	const bool bUpdateSegments = false; // done in next loop
+	const bool bUpdateMeshLevel = false; // no need because mesh have been deleted
+	InControlPoint->UpdateSplinePoints(bUpdateCollision, bUpdateSegments, bUpdateMeshLevel);
+
+	// Continue
+	for (const FLandscapeSplineConnection& Connection : InControlPoint->ConnectedSegments)
+	{
+		if (!Connection.End)
+		{
+			MoveSegmentToLandscape(Connection.Segment, FromProxy, ToProxy);
+		}
+	}
+}
+
+void ULandscapeInfo::MoveSegmentToLandscape(ULandscapeSplineSegment* InSegment, ALandscapeProxy* FromProxy, ALandscapeProxy* ToProxy)
+{
+	ToProxy->Modify();
+	if (ToProxy->SplineComponent == nullptr)
+	{
+		ToProxy->CreateSplineComponent(FromProxy->SplineComponent->GetRelativeScale3D());
+		check(ToProxy->SplineComponent);
+	}
+		
+	if (InSegment->GetOuterULandscapeSplinesComponent() == ToProxy->SplineComponent)
+	{
+		return;
+	}
+	check(InSegment->GetOuterULandscapeSplinesComponent() == FromProxy->SplineComponent);
+
+	ToProxy->SplineComponent->Modify();
+		
+	// Delete all Mesh Components associated with the Segment. (Will get recreated in UpdateSplinePoints)
+	for (auto* MeshComponent : InSegment->LocalMeshComponents)
+	{
+		MeshComponent->Modify();
+		MeshComponent->UnregisterComponent();
+		MeshComponent->DestroyComponent();
+		FromProxy->Modify();
+		FromProxy->SplineComponent->MeshComponentLocalOwnersMap.Remove(MeshComponent);
+	}
+	InSegment->LocalMeshComponents.Empty();
+
+	TMap<ULandscapeSplinesComponent*, TArray<USplineMeshComponent*>> ForeignMeshComponents = InSegment->GetForeignMeshComponents();
+	for (auto Pair : ForeignMeshComponents)
+	{
+		Pair.Key->RemoveAllForeignMeshComponents(InSegment);
+		for (auto MeshComponent : Pair.Value)
+		{
+			MeshComponent->Modify();
+			MeshComponent->UnregisterComponent();
+			MeshComponent->DestroyComponent();
+		}
+	}
+
+	// Move segment to new level
+	FromProxy->SplineComponent->Segments.Remove(InSegment);
+	InSegment->Rename(nullptr, ToProxy->SplineComponent);
+	ToProxy->SplineComponent->Segments.Add(InSegment);
+		
+	// Continue ...
+	MoveControlPointToLandscape(InSegment->Connections[1].ControlPoint, FromProxy, ToProxy);
+
+	const bool bUpdateCollision = true; // default value
+	const bool bUpdateMeshLevel = false; // no need because mesh have been deleted 
+	InSegment->UpdateSplinePoints(bUpdateCollision, bUpdateMeshLevel);
+}
+
+void ULandscapeInfo::MoveSplineToLevel(ULandscapeSplineControlPoint* InControlPoint, ULevel* TargetLevel)
+{
+	ALandscapeProxy* FromProxy = InControlPoint->GetTypedOuter<ALandscapeProxy>();
+	if (FromProxy->GetLevel() == TargetLevel)
+	{
+		return;
+	}
+
+	ALandscapeProxy* ToProxy = GetLandscapeProxyForLevel(TargetLevel);
+	if (!ToProxy)
+	{
+		return;
+	}
+	
+	TSet<ULandscapeSplineControlPoint*> ControlPoints;
+	
+	FindStartingControlPoints(InControlPoint, ControlPoints);
+	
+	FromProxy->Modify();
+	FromProxy->SplineComponent->Modify();
+	FromProxy->SplineComponent->MarkRenderStateDirty();
+
+	for (ULandscapeSplineControlPoint* ControlPoint : ControlPoints)
+	{
+		MoveControlPointToLandscape(ControlPoint, FromProxy, ToProxy);
+	}
+}
+
+void ULandscapeInfo::MoveSplinesToLevel(ULandscapeSplinesComponent * InSplineComponent, ULevel * TargetLevel)
+{
+	TArray<ULandscapeSplineControlPoint*> CopyControlPoints(InSplineComponent->ControlPoints);
+	for(ULandscapeSplineControlPoint* ControlPoint : CopyControlPoints)
+	{
+		// Even if ControlPoints are part of same connected spline they will be skipped if already moved
+		MoveSplineToLevel(ControlPoint, TargetLevel);
+	}
+}
 #endif // WITH_EDITOR
 
 
