Change the cached World ref in ULandscapeInfoMap to a TWeakObjectPtr to safeguard against order or destruction issues on World unload.

diff --git a/Engine/Source/Runtime/Landscape/Classes/LandscapeInfoMap.h b/Engine/Source/Runtime/Landscape/Classes/LandscapeInfoMap.h
--- a/Engine/Source/Runtime/Landscape/Classes/LandscapeInfoMap.h
+++ b/Engine/Source/Runtime/Landscape/Classes/LandscapeInfoMap.h
@@ -5,6 +5,7 @@
 #include "CoreMinimal.h"
 #include "UObject/ObjectMacros.h"
 #include "UObject/Object.h"
+#include "UObject/WeakObjectPtr.h"
 #include "Misc/Guid.h"
 
 #include "LandscapeInfoMap.generated.h"
@@ -24,8 +25,7 @@ public:
 	static void AddReferencedObjects(UObject* InThis, FReferenceCollector& Collector);
 
 	TMap<FGuid, ULandscapeInfo*> Map;
-	UWorld* World;
-
+	TWeakObjectPtr<UWorld> World;
 
 	/**
 	* Gets landscape-specific data for given world.
diff --git a/Engine/Source/Runtime/Landscape/Private/LandscapeInfoMap.cpp b/Engine/Source/Runtime/Landscape/Private/LandscapeInfoMap.cpp
--- a/Engine/Source/Runtime/Landscape/Private/LandscapeInfoMap.cpp
+++ b/Engine/Source/Runtime/Landscape/Private/LandscapeInfoMap.cpp
@@ -29,7 +29,7 @@ void ULandscapeInfoMap::Serialize(FArchive& Ar)
 
 void ULandscapeInfoMap::BeginDestroy()
 {
-	if (World != nullptr)
+	if (World.IsValid())
 	{
 		World->PerModuleDataObjects.Remove(this);
 	}
