Landscape: Fix VT Parameters not being copied between LandscapeProxies

#rnx

diff --git a/Engine/Source/Runtime/Landscape/Private/Landscape.cpp b/Engine/Source/Runtime/Landscape/Private/Landscape.cpp
--- a/Engine/Source/Runtime/Landscape/Private/Landscape.cpp
+++ b/Engine/Source/Runtime/Landscape/Private/Landscape.cpp
@@ -2506,6 +2506,11 @@ void ALandscapeProxy::GetSharedProperties(ALandscapeProxy* Landscape)
 		PositiveZBoundsExtension = Landscape->PositiveZBoundsExtension;
 		CollisionMipLevel = Landscape->CollisionMipLevel;
 		bBakeMaterialPositionOffsetIntoCollision = Landscape->bBakeMaterialPositionOffsetIntoCollision;
+		RuntimeVirtualTextures = Landscape->RuntimeVirtualTextures;
+		VirtualTextureLodBias = Landscape->VirtualTextureLodBias;
+		VirtualTextureNumLods = Landscape->VirtualTextureNumLods;
+		VirtualTextureRenderPassType = Landscape->VirtualTextureRenderPassType;
+
 		if (!LandscapeMaterial)
 		{
 			LandscapeMaterial = Landscape->LandscapeMaterial;
