Chaos physical material support

diff --git a/Engine/Source/Runtime/Core/Public/UObject/ExternalPhysicsCustomObjectVersion.h b/Engine/Source/Runtime/Core/Public/UObject/ExternalPhysicsCustomObjectVersion.h
--- a/Engine/Source/Runtime/Core/Public/UObject/ExternalPhysicsCustomObjectVersion.h
+++ b/Engine/Source/Runtime/Core/Public/UObject/ExternalPhysicsCustomObjectVersion.h
@@ -63,6 +63,9 @@ struct CORE_API FExternalPhysicsCustomObjectVersion
 		// Add an SOA which contains particles with full dynamic data but which are in the kinematic object state.
 		AddDynamicKinematicSOA,
 
+		// Added material manager to Chaos
+		AddedMaterialManager,
+		
 		// -----<new versions can be added above this line>-------------------------------------------------
 		VersionPlusOne,
 		LatestVersion = VersionPlusOne - 1
diff --git a/Engine/Source/Runtime/Engine/Private/Collision/CollisionConversions.cpp b/Engine/Source/Runtime/Engine/Private/Collision/CollisionConversions.cpp
--- a/Engine/Source/Runtime/Engine/Private/Collision/CollisionConversions.cpp
+++ b/Engine/Source/Runtime/Engine/Private/Collision/CollisionConversions.cpp
@@ -184,7 +184,7 @@ static void SetHitResultFromShapeAndFaceIndex(const FPhysicsShape& Shape,  const
 		if (bReturnPhysMat)
 		{
 			// This function returns the single material in all cases other than trimesh or heightfield
-			if(const FPhysicsMaterial* PhysicsMaterial = GetMaterialFromInternalFaceIndex(Shape, FaceIndex))
+			if(const FPhysicsMaterial* PhysicsMaterial = GetMaterialFromInternalFaceIndex(Shape, Actor, FaceIndex))
 			{
 				OutResult.PhysMaterial = GetUserData(*PhysicsMaterial);
 			}
@@ -320,7 +320,7 @@ EConvertQueryResult ConvertQueryImpactHit(const UWorld* World, const FHitLocatio
 		// Lookup physical material for heightfields
 		if (bReturnPhysMat && InternalFaceIndex != GetInvalidPhysicsFaceIndex())
 		{
-			if (const FPhysicsMaterial* Material = GetMaterialFromInternalFaceIndex(HitShape, InternalFaceIndex))
+			if (const FPhysicsMaterial* Material = GetMaterialFromInternalFaceIndex(HitShape, HitActor, InternalFaceIndex))
 			{
 				OutResult.PhysMaterial = GetUserData(*Material);
 			}
diff --git a/Engine/Source/Runtime/Engine/Private/PhysicsEngine/Experimental/ChaosDerivedData.cpp b/Engine/Source/Runtime/Engine/Private/PhysicsEngine/Experimental/ChaosDerivedData.cpp
--- a/Engine/Source/Runtime/Engine/Private/PhysicsEngine/Experimental/ChaosDerivedData.cpp
+++ b/Engine/Source/Runtime/Engine/Private/PhysicsEngine/Experimental/ChaosDerivedData.cpp
@@ -21,7 +21,7 @@ const TCHAR* FChaosDerivedDataCooker::GetPluginName() const
 
 const TCHAR* FChaosDerivedDataCooker::GetVersionString() const
 {
-	return TEXT("40C102F05B01482DB74F98FE0F6C4166");
+	return TEXT("64EABBC88E0D47A98B01CD8BA0042E13");
 }
 
 FString FChaosDerivedDataCooker::GetPluginSpecificCacheKeySuffix() const
diff --git a/Engine/Source/Runtime/Engine/Private/PhysicsEngine/Experimental/ChaosInterfaceWrapper.cpp b/Engine/Source/Runtime/Engine/Private/PhysicsEngine/Experimental/ChaosInterfaceWrapper.cpp
--- a/Engine/Source/Runtime/Engine/Private/PhysicsEngine/Experimental/ChaosInterfaceWrapper.cpp
+++ b/Engine/Source/Runtime/Engine/Private/PhysicsEngine/Experimental/ChaosInterfaceWrapper.cpp
@@ -14,9 +14,10 @@ namespace ChaosInterface
 		return UserData ? FPhysxUserData::Get<FBodyInstance>(Actor.UserData()) : nullptr;
 	}
 
-	UPhysicalMaterial* GetUserData(const FPhysTypeDummy& Material)
+	UPhysicalMaterial* GetUserData(const Chaos::FChaosPhysicsMaterial& Material)
 	{
-		return nullptr;
+		void* UserData = Material.UserData;
+		return UserData ? FPhysxUserData::Get<UPhysicalMaterial>(UserData) : nullptr;
 	}
 
 #if WITH_CHAOS
diff --git a/Engine/Source/Runtime/Engine/Private/PhysicsEngine/Experimental/PhysInterface_Chaos.cpp b/Engine/Source/Runtime/Engine/Private/PhysicsEngine/Experimental/PhysInterface_Chaos.cpp
--- a/Engine/Source/Runtime/Engine/Private/PhysicsEngine/Experimental/PhysInterface_Chaos.cpp
+++ b/Engine/Source/Runtime/Engine/Private/PhysicsEngine/Experimental/PhysInterface_Chaos.cpp
@@ -36,12 +36,15 @@
 #include "Collision/CollisionConversions.h"
 #include "PhysicsInterfaceUtilsCore.h"
 #include "Components/SkeletalMeshComponent.h"
+#include "PBDRigidsSolver.h"
 
 #if WITH_PHYSX
 #include "geometry/PxConvexMesh.h"
 #include "geometry/PxTriangleMesh.h"
 #include "foundation/PxVec3.h"
 #include "extensions/PxMassProperties.h"
+#include "PhysicalMaterials/PhysicalMaterial.h"
+#include "Containers/ArrayView.h"
 #endif
 
 DEFINE_STAT(STAT_TotalPhysicsTime);
@@ -83,6 +86,34 @@ ECollisionShapeType GetGeometryType(const Chaos::TPerShapeData<float, 3>& Shape)
 	return GetType(*Shape.Geometry);
 }
 
+Chaos::FChaosPhysicsMaterial* GetMaterialFromInternalFaceIndex(const FPhysicsShape& Shape, const FPhysicsActor& Actor, uint32 InternalFaceIndex)
+{
+	if(Shape.Materials.Num() > 0 && Actor.Proxy)
+	{
+		Chaos::FPBDRigidsSolver* Solver = Actor.Proxy->GetSolver();
+
+		if(ensure(Solver))
+		{
+			if(Shape.Materials.Num() == 1)
+			{
+				Chaos::TSolverQueryMaterialScope<Chaos::ELockType::Read> Scope(Solver);
+				return Solver->GetQueryMaterials().Get(Shape.Materials[0].InnerHandle);
+			}
+
+			uint8 Index = Shape.Geometry->GetMaterialIndex(InternalFaceIndex);
+
+			if(Shape.Materials.IsValidIndex(Index))
+			{
+				Chaos::TSolverQueryMaterialScope<Chaos::ELockType::Read> Scope(Solver);
+				return Solver->GetQueryMaterials().Get(Shape.Materials[Index].InnerHandle);
+			}
+		}
+	}
+
+	return nullptr;
+}
+
+
 const Chaos::FImplicitObject& FPhysicsShapeReference_Chaos::GetGeometry() const
 {
 	check(IsValid()); return *Shape->Geometry;
@@ -180,6 +211,39 @@ int32 FPhysInterface_Chaos::GetNumActorsInAggregate(const FPhysicsAggregateRefer
 void FPhysInterface_Chaos::AddActorToAggregate_AssumesLocked(const FPhysicsAggregateReference_Chaos& InAggregate, const FPhysicsActorHandle& InActor) {}
 
 
+FPhysicsMaterialHandle FPhysInterface_Chaos::CreateMaterial(const UPhysicalMaterial* InMaterial)
+{
+	Chaos::FMaterialHandle NewHandle = Chaos::FPhysicalMaterialManager::Get().Create();
+
+	return NewHandle;
+}
+
+void FPhysInterface_Chaos::ReleaseMaterial(FPhysicsMaterialHandle& InHandle)
+{
+	Chaos::FPhysicalMaterialManager::Get().Destroy(InHandle);
+}
+
+void FPhysInterface_Chaos::UpdateMaterial(FPhysicsMaterialHandle& InHandle, UPhysicalMaterial* InMaterial)
+{
+	if(Chaos::FChaosPhysicsMaterial* Material = InHandle.Get())
+	{
+		Material->Friction = InMaterial->Friction;
+		Material->Restitution = InMaterial->Restitution;
+	}
+
+	Chaos::FPhysicalMaterialManager::Get().UpdateMaterial(InHandle);
+}
+
+void FPhysInterface_Chaos::SetUserData(FPhysicsMaterialHandle& InHandle, void* InUserData)
+{
+	if(Chaos::FChaosPhysicsMaterial* Material = InHandle.Get())
+	{
+		Material->UserData = InUserData;
+	}
+
+	Chaos::FPhysicalMaterialManager::Get().UpdateMaterial(InHandle);
+}
+
 int32 FPhysInterface_Chaos::GetNumShapes(const FPhysicsActorHandle& InHandle)
 {
 	// #todo : Implement
@@ -1103,12 +1167,18 @@ void FPhysInterface_Chaos::AddGeometry(FPhysicsActorHandle& InActor, const FGeom
 #if WITH_CHAOS
 	if (InActor && Geoms.Num())
 	{
-		if (OutOptShapes)
+		for (TUniquePtr<Chaos::TPerShapeData<float, 3>>& Shape : Shapes)
 		{
-			for (auto& Shape : Shapes)
+			FPhysicsShapeHandle NewHandle(Shape.Get(), true, true, InActor);
+			if (OutOptShapes)
 			{
-				OutOptShapes->Add({ Shape.Get(), true, true, InActor });
+				OutOptShapes->Add(NewHandle);
 			}
+
+			FBodyInstance::ApplyMaterialToShape_AssumesLocked(NewHandle, InParams.SimpleMaterial, InParams.ComplexMaterials);
+
+			//TArrayView<UPhysicalMaterial*> SimpleView = MakeArrayView(&(const_cast<UPhysicalMaterial*>(InParams.SimpleMaterial)), 1);
+			//FPhysInterface_Chaos::SetMaterials(NewHandle, InParams.ComplexMaterials.Num() > 0 ? InParams.ComplexMaterials : SimpleView);
 		}
 
 		//todo: we should not be creating unique geometry per actor
@@ -1254,6 +1324,20 @@ void FPhysInterface_Chaos::SetLocalTransform(const FPhysicsShapeHandle& InShape,
 #endif
 }
 
+void FPhysInterface_Chaos::SetMaterials(const FPhysicsShapeHandle& InShape, const TArrayView<UPhysicalMaterial*> InMaterials)
+{
+	// Build a list of handles to store on the shape
+	TArray<Chaos::FMaterialHandle> NewMaterialHandles;
+	NewMaterialHandles.Reserve(InMaterials.Num());
+
+	for(UPhysicalMaterial* UnrealMaterial : InMaterials)
+	{
+		NewMaterialHandles.Add(UnrealMaterial->GetPhysicsMaterial());
+	}
+
+	InShape.Shape->Materials = NewMaterialHandles;
+}
+
 void FinishSceneStat()
 {
 }
diff --git a/Engine/Source/Runtime/Engine/Private/PhysicsEngine/Experimental/PhysScene_Chaos.cpp b/Engine/Source/Runtime/Engine/Private/PhysicsEngine/Experimental/PhysScene_Chaos.cpp
--- a/Engine/Source/Runtime/Engine/Private/PhysicsEngine/Experimental/PhysScene_Chaos.cpp
+++ b/Engine/Source/Runtime/Engine/Private/PhysicsEngine/Experimental/PhysScene_Chaos.cpp
@@ -1369,6 +1369,11 @@ void FPhysScene_ChaosInterface::EndFrame(ULineBatchComponent* InLineBatcher)
 			{
 				SyncBodies(Solver);
 				Solver->SyncEvents_GameThread();
+
+				{
+					SCOPE_CYCLE_COUNTER(STAT_SqUpdateMaterials);
+					Solver->SyncQueryMaterials();
+				}
 			}
 		}
 
diff --git a/Engine/Source/Runtime/Engine/Private/PhysicsEngine/PhysicalMaterial.cpp b/Engine/Source/Runtime/Engine/Private/PhysicsEngine/PhysicalMaterial.cpp
--- a/Engine/Source/Runtime/Engine/Private/PhysicsEngine/PhysicalMaterial.cpp
+++ b/Engine/Source/Runtime/Engine/Private/PhysicsEngine/PhysicalMaterial.cpp
@@ -81,24 +81,16 @@ void UPhysicalMaterial::FinishDestroy()
 
 FPhysicsMaterialHandle& UPhysicalMaterial::GetPhysicsMaterial()
 {
-#if WITH_CHAOS
-    check(false);
-    MaterialHandle = nullptr;
-    return MaterialHandle;
-#else
 	if(!MaterialHandle.IsValid())
 	{
 		MaterialHandle = FPhysicsInterface::CreateMaterial(this);
 		check(MaterialHandle.IsValid());
 
-#if WITH_PHYSX
 		FPhysicsInterface::SetUserData(MaterialHandle, &PhysxUserData);
-#endif
 		FPhysicsInterface::UpdateMaterial(MaterialHandle, this);
 	}
 
 	return MaterialHandle;
-#endif
 }
 
 EPhysicalSurface UPhysicalMaterial::DetermineSurfaceType(UPhysicalMaterial const* PhysicalMaterial)
@@ -110,4 +102,3 @@ EPhysicalSurface UPhysicalMaterial::DetermineSurfaceType(UPhysicalMaterial const
 	
 	return PhysicalMaterial->SurfaceType;
 }
-
diff --git a/Engine/Source/Runtime/Engine/Public/Physics/Experimental/ChaosInterfaceWrapper.h b/Engine/Source/Runtime/Engine/Public/Physics/Experimental/ChaosInterfaceWrapper.h
--- a/Engine/Source/Runtime/Engine/Public/Physics/Experimental/ChaosInterfaceWrapper.h
+++ b/Engine/Source/Runtime/Engine/Public/Physics/Experimental/ChaosInterfaceWrapper.h
@@ -44,7 +44,9 @@ inline FQueryFilterData MakeQueryFilterData(const FCollisionFilterData& FilterDa
 
 FBodyInstance* GetUserData(const Chaos::TGeometryParticle<float, 3>& Actor);
 
-UPhysicalMaterial* GetUserData(const FPhysTypeDummy& Material);
+#if WITH_CHAOS
+UPhysicalMaterial* GetUserData(const Chaos::FChaosPhysicsMaterial& Material);
+#endif
 
 }
 
diff --git a/Engine/Source/Runtime/Engine/Public/Physics/Experimental/PhysInterface_Chaos.h b/Engine/Source/Runtime/Engine/Public/Physics/Experimental/PhysInterface_Chaos.h
--- a/Engine/Source/Runtime/Engine/Public/Physics/Experimental/PhysInterface_Chaos.h
+++ b/Engine/Source/Runtime/Engine/Public/Physics/Experimental/PhysInterface_Chaos.h
@@ -13,6 +13,7 @@
 #include "Physics/PhysicsInterfaceTypes.h"
 #include "Physics/GenericPhysicsInterface.h"
 #include "Physics/Experimental/PhysicsUserData_Chaos.h"
+#include "Chaos/PhysicalMaterials.h"
 
 //NOTE: Do not include Chaos headers directly as it means recompiling all of engine. This should be reworked to avoid allocations
 
@@ -194,11 +195,10 @@ public:
 	static void AddActorToAggregate_AssumesLocked(const FPhysicsAggregateReference_Chaos& InAggregate, const FPhysicsActorHandle& InActor);
 
 	// Material interface functions
-    // @todo(mlentine): How do we set material on the solver?
-	static FPhysicsMaterialHandle CreateMaterial(const UPhysicalMaterial* InMaterial) { return nullptr; }
-    static void ReleaseMaterial(FPhysicsMaterialHandle& InHandle) {}
-    static void UpdateMaterial(const FPhysicsMaterialHandle& InHandle, UPhysicalMaterial* InMaterial) {}
-    static void SetUserData(const FPhysicsMaterialHandle& InHandle, void* InUserData) {}
+	static FPhysicsMaterialHandle CreateMaterial(const UPhysicalMaterial* InMaterial);
+    static void ReleaseMaterial(FPhysicsMaterialHandle& InHandle);
+    static void UpdateMaterial(FPhysicsMaterialHandle& InHandle, UPhysicalMaterial* InMaterial);
+    static void SetUserData(FPhysicsMaterialHandle& InHandle, void* InUserData);
 
 	// Actor interface functions
 	template<typename AllocatorType>
@@ -397,7 +397,7 @@ public:
     static void SetUserData(const FPhysicsShapeHandle& InShape, void* InUserData) {}
     static void SetGeometry(const FPhysicsShapeHandle& InShape, physx::PxGeometry& InGeom) {}
 	static void SetLocalTransform(const FPhysicsShapeHandle& InShape, const FTransform& NewLocalTransform);
-    static void SetMaterials(const FPhysicsShapeHandle& InShape, const TArrayView<UPhysicalMaterial*>InMaterials) {}
+    static void SetMaterials(const FPhysicsShapeHandle& InShape, const TArrayView<UPhysicalMaterial*>InMaterials);
 };
 
 /*
@@ -442,10 +442,7 @@ FORCEINLINE void ComputeZeroDistanceImpactNormalAndPenetration(const UWorld* Wor
 	//TODO_SQ_IMPLEMENTATION
 }
 
-inline FPhysTypeDummy* GetMaterialFromInternalFaceIndex(const FPhysicsShape& Shape, uint32 InternalFaceIndex)
-{
-	return nullptr;
-}
+Chaos::FChaosPhysicsMaterial* GetMaterialFromInternalFaceIndex(const FPhysicsShape& Shape, const FPhysicsActor& Actor, uint32 InternalFaceIndex);
 
 inline uint32 GetTriangleMeshExternalFaceIndex(const FPhysicsShape& Shape, uint32 InternalFaceIndex)
 {
diff --git a/Engine/Source/Runtime/Experimental/Chaos/Private/Chaos/GeometryParticles.cpp b/Engine/Source/Runtime/Experimental/Chaos/Private/Chaos/GeometryParticles.cpp
--- a/Engine/Source/Runtime/Experimental/Chaos/Private/Chaos/GeometryParticles.cpp
+++ b/Engine/Source/Runtime/Experimental/Chaos/Private/Chaos/GeometryParticles.cpp
@@ -85,6 +85,11 @@ namespace Chaos
 			// This should be set by particle serializing this TPerShapeData.
 			WorldSpaceInflatedShapeBounds = TAABB<FReal, 3>(FVec3(0.0f, 0.0f, 0.0f), FVec3(0.0f, 0.0f, 0.0f));
 		}
+
+		if(Ar.CustomVer(FExternalPhysicsCustomObjectVersion::GUID) >= FExternalPhysicsCustomObjectVersion::AddedMaterialManager)
+		{
+			Ar << Materials;
+		}
 	}
 
 
diff --git a/Engine/Source/Runtime/Experimental/Chaos/Private/Chaos/HeightField.cpp b/Engine/Source/Runtime/Experimental/Chaos/Private/Chaos/HeightField.cpp
--- a/Engine/Source/Runtime/Experimental/Chaos/Private/Chaos/HeightField.cpp
+++ b/Engine/Source/Runtime/Experimental/Chaos/Private/Chaos/HeightField.cpp
@@ -158,8 +158,8 @@ namespace Chaos
 			GeomData->GetPointsScaled(FullIndex, Points);
 
 			// Test both triangles that are in this cell, as we could hit both in any order
-			TestTriangle(FullIndex * 2, Points[0], Points[1], Points[2]);
-			TestTriangle(FullIndex * 2 + 1, Points[2], Points[1], Points[3]);
+			TestTriangle(Payload * 2, Points[0], Points[1], Points[2]);
+			TestTriangle(Payload * 2 + 1, Points[2], Points[1], Points[3]);
 
 			return OutTime > 0;
 		}
@@ -244,8 +244,8 @@ namespace Chaos
 			TVector<T, 3> Points[4];
 			HfData->GetPointsScaled(FullIndex, Points);
 
-			TestTriangle(FullIndex * 2, Points[0], Points[1], Points[2]);
-			TestTriangle(FullIndex * 2 + 1, Points[2], Points[1], Points[3]);
+			TestTriangle(Payload * 2, Points[0], Points[1], Points[2]);
+			TestTriangle(Payload * 2 + 1, Points[2], Points[1], Points[3]);
 
 			return OutTime != 0;
 		}
@@ -266,14 +266,15 @@ namespace Chaos
 	};
 
 	template<typename HeightfieldT, typename BufferType>
-	void BuildGeomData(TArrayView<BufferType> BufferView, int32 NumRows, int32 NumCols, const TVector<HeightfieldT, 3>& InScale, TUniqueFunction<HeightfieldT(const BufferType)> ToRealFunc, typename THeightField<HeightfieldT>::FDataType& OutData, TBox<HeightfieldT, 3>& OutBounds)
+	void BuildGeomData(TArrayView<BufferType> BufferView, TArrayView<uint8> MaterialIndexView, int32 NumRows, int32 NumCols, const TVector<HeightfieldT, 3>& InScale, TUniqueFunction<HeightfieldT(const BufferType)> ToRealFunc, typename THeightField<HeightfieldT>::FDataType& OutData, TBox<HeightfieldT, 3>& OutBounds)
 	{
 		using FDataType = typename THeightField<HeightfieldT>::FDataType;
 
 		using RealType = typename FDataType::RealType;
 
-		const int32 NumCells = NumRows * NumCols;
-		ensure(BufferView.Num() == NumCells);
+		const bool bHaveMaterials = MaterialIndexView.Num() > 0;
+		const bool bOnlyDefaultMaterial = MaterialIndexView.Num() == 1;
+		ensure(BufferView.Num() == NumRows * NumCols);
 		ensure(NumRows > 1);
 		ensure(NumCols > 1);
 
@@ -321,6 +322,21 @@ namespace Chaos
 			}
 		}
 		OutBounds.Thicken(KINDA_SMALL_NUMBER);
+
+		if(bHaveMaterials)
+		{
+			if(bOnlyDefaultMaterial)
+			{
+				OutData.MaterialIndices.Add(0);
+			}
+			else
+			{
+				const int32 NumCells = NumHeights - NumRows - NumCols + 1;
+				ensure(MaterialIndexView.Num() == NumCells);
+				OutData.MaterialIndices.Empty();
+				OutData.MaterialIndices.Append(MaterialIndexView.GetData(), MaterialIndexView.Num());
+			}
+		}
 	}
 
 	template<typename HeightfieldT, typename BufferType>
@@ -408,16 +424,16 @@ namespace Chaos
 	}
 
 	template <typename T>
-	THeightField<T>::THeightField(TArray<T>&& Height, int32 NumRows, int32 NumCols, const TVector<T, 3>& InScale)
+	THeightField<T>::THeightField(TArray<T>&& Height, TArray<uint8>&& InMaterialIndices, int32 NumRows, int32 NumCols, const TVector<T, 3>& InScale)
 		: FImplicitObject(EImplicitObject::HasBoundingBox, ImplicitObjectType::HeightField)
 	{
-		BuildGeomData<T, T>(MakeArrayView(Height), NumRows, NumCols, TVector<T, 3>(1), [](const T InVal) {return InVal; }, GeomData, LocalBounds);
+		BuildGeomData<T, T>(MakeArrayView(Height), MakeArrayView(InMaterialIndices), NumRows, NumCols, TVector<T, 3>(1), [](const T InVal) {return InVal; }, GeomData, LocalBounds);
 		CalcBounds();
 		SetScale(InScale);
 	}
 
 	template<typename T>
-	Chaos::THeightField<T>::THeightField(TArrayView<const uint16> InHeights, int32 InNumRows, int32 InNumCols, const TVector<T, 3>& InScale)
+	Chaos::THeightField<T>::THeightField(TArrayView<const uint16> InHeights, TArrayView<uint8> InMaterialIndices, int32 InNumRows, int32 InNumCols, const TVector<T, 3>& InScale)
 		: FImplicitObject(EImplicitObject::HasBoundingBox, ImplicitObjectType::HeightField)
 	{
 		TUniqueFunction<T(const uint16)> ConversionFunc = [](const T InVal) -> float
@@ -425,7 +441,7 @@ namespace Chaos
 			return (float)((int32)InVal - 32768);
 		};
 
-		BuildGeomData<T, const uint16>(InHeights, InNumRows, InNumCols, TVector<T, 3>(1), MoveTemp(ConversionFunc), GeomData, LocalBounds);
+		BuildGeomData<T, const uint16>(InHeights, InMaterialIndices, InNumRows, InNumCols, TVector<T, 3>(1), MoveTemp(ConversionFunc), GeomData, LocalBounds);
 		CalcBounds();
 		SetScale(InScale);
 	}
diff --git a/Engine/Source/Runtime/Experimental/Chaos/Private/Chaos/PhysicalMaterials.cpp b/Engine/Source/Runtime/Experimental/Chaos/Private/Chaos/PhysicalMaterials.cpp
new file mode 100644
--- /dev/null
+++ b/Engine/Source/Runtime/Experimental/Chaos/Private/Chaos/PhysicalMaterials.cpp
@@ -0,0 +1,82 @@
+// Copyright 1998-2019 Epic Games, Inc. All Rights Reserved.
+
+#include "Chaos/PhysicalMaterials.h"
+
+namespace Chaos
+{
+	FChaosPhysicsMaterial* FMaterialHandle::Get() const
+	{
+		if(InnerHandle.IsValid())
+		{
+			return FPhysicalMaterialManager::Get().Resolve(InnerHandle);
+		}
+		return nullptr;
+	}
+
+	const FChaosPhysicsMaterial* FConstMaterialHandle::Get() const
+	{
+		if(InnerHandle.IsValid())
+		{
+			return FPhysicalMaterialManager::Get().Resolve(InnerHandle);
+		}
+		return nullptr;
+	}
+
+	FPhysicalMaterialManager::FPhysicalMaterialManager()
+		: Materials(InitialCapacity)
+	{
+
+	}
+
+	FPhysicalMaterialManager& FPhysicalMaterialManager::Get()
+	{
+		static FPhysicalMaterialManager Instance;
+		return Instance;
+	}
+
+	FChaosPhysicsMaterial* FPhysicalMaterialManager::Resolve(FChaosMaterialHandle InHandle) const
+	{
+		return Materials.Get(InHandle);
+	}
+
+	const FChaosPhysicsMaterial* FPhysicalMaterialManager::Resolve(FChaosConstMaterialHandle InHandle) const
+	{
+		return Materials.Get(InHandle);
+	}
+
+	void FPhysicalMaterialManager::UpdateMaterial(FMaterialHandle InHandle)
+	{
+		check(IsInGameThread());
+
+		OnMaterialUpdated.Broadcast(InHandle);
+	}
+
+	const Chaos::THandleArray<FChaosPhysicsMaterial>& FPhysicalMaterialManager::GetMasterMaterials() const
+	{
+		return Materials;
+	}
+
+	FMaterialHandle FPhysicalMaterialManager::Create()
+	{
+		check(IsInGameThread());
+		FMaterialHandle OutHandle;
+		OutHandle.InnerHandle = Materials.Create();
+
+		OnMaterialCreated.Broadcast(OutHandle);
+
+		return OutHandle;
+	}
+
+
+	void FPhysicalMaterialManager::Destroy(FMaterialHandle InHandle)
+	{
+		check(IsInGameThread());
+		if(InHandle.InnerHandle.IsValid())
+		{
+			Materials.Destroy(InHandle.InnerHandle);
+
+			OnMaterialDestroyed.Broadcast(InHandle);
+		}
+	}
+
+}
\ No newline at end of file
diff --git a/Engine/Source/Runtime/Experimental/Chaos/Private/ChaosStats.cpp b/Engine/Source/Runtime/Experimental/Chaos/Private/ChaosStats.cpp
--- a/Engine/Source/Runtime/Experimental/Chaos/Private/ChaosStats.cpp
+++ b/Engine/Source/Runtime/Experimental/Chaos/Private/ChaosStats.cpp
@@ -30,6 +30,7 @@ DEFINE_STAT(STAT_DispatchCollisionEvents);
 DEFINE_STAT(STAT_DispatchBreakEvents);
 DEFINE_STAT(STAT_BufferPhysicsResults);
 DEFINE_STAT(STAT_FlipResults);
+DEFINE_STAT(STAT_SqUpdateMaterials);
 DEFINE_STAT(STAT_CacheResultGeomCollection);
 DEFINE_STAT(STAT_CacheResultStaticMesh);
 DEFINE_STAT(STAT_CaptureDisabledState);
diff --git a/Engine/Source/Runtime/Experimental/Chaos/Public/Chaos/Defines.h b/Engine/Source/Runtime/Experimental/Chaos/Public/Chaos/Defines.h
--- a/Engine/Source/Runtime/Experimental/Chaos/Public/Chaos/Defines.h
+++ b/Engine/Source/Runtime/Experimental/Chaos/Public/Chaos/Defines.h
@@ -30,6 +30,7 @@ namespace Chaos
 		Chaos::FReal SleepingAngularThreshold;
 		Chaos::FReal DisabledLinearThreshold;
 		Chaos::FReal DisabledAngularThreshold;
+		void* UserData;
 
 		FChaosPhysicsMaterial()
 			: Friction(0.5)
@@ -38,6 +39,7 @@ namespace Chaos
 			, SleepingAngularThreshold(1)
 			, DisabledLinearThreshold(0)
 			, DisabledAngularThreshold(0)
+			, UserData(nullptr)
 		{
 		}
 
diff --git a/Engine/Source/Runtime/Experimental/Chaos/Public/Chaos/Framework/Handles.h b/Engine/Source/Runtime/Experimental/Chaos/Public/Chaos/Framework/Handles.h
--- a/Engine/Source/Runtime/Experimental/Chaos/Public/Chaos/Framework/Handles.h
+++ b/Engine/Source/Runtime/Experimental/Chaos/Public/Chaos/Framework/Handles.h
@@ -65,11 +65,6 @@ namespace Chaos
 			Generation = InUint & ((1 << IndexWidth) - 1);
 		}
 
-		bool operator==(const THandleData<IndexWidth, GenerationWidth>& InOther) const
-		{
-			return Generation == InOther.Generation && Index == InOther.Index;
-		}
-
 		friend FArchive& operator <<(FArchive& Ar, THandleData<IndexWidth, GenerationWidth>& InHandle)
 		{
 			uint32 Packed = InHandle.AsUint();
@@ -102,6 +97,11 @@ namespace Chaos
 
 		}
 
+		friend bool operator ==(const THandle<ElementType, IndexWidth, GenerationWidth>& A, const THandle<ElementType, IndexWidth, GenerationWidth>& B)
+		{
+			return A.AsUint() == B.AsUint();
+		}
+
 		friend FArchive& operator <<(FArchive& Ar, THandle<ElementType, IndexWidth, GenerationWidth>& InHandle)
 		{
 			Ar << *((THandleData<IndexWidth, GenerationWidth>*)&InHandle);
@@ -129,6 +129,11 @@ namespace Chaos
 
 		}
 
+		friend bool operator ==(const TConstHandle<ElementType, IndexWidth, GenerationWidth>& A, const TConstHandle<ElementType, IndexWidth, GenerationWidth>& B)
+		{
+			return A.AsUint() == B.AsUint();
+		}
+
 		friend FArchive& operator <<(FArchive& Ar, TConstHandle<ElementType, IndexWidth, GenerationWidth>& InHandle)
 		{
 			Ar << *((THandleData<IndexWidth, GenerationWidth>*)&InHandle);
@@ -166,11 +171,13 @@ namespace Chaos
 		THandleArray& operator =(const THandleArray<ElementType, IndexWidth, GenerationWidth>& Other)
 		{
 			CopyFrom(Other);
+			return *this;
 		}
 
 		THandleArray& operator =(THandleArray<ElementType, IndexWidth, GenerationWidth>&& Other)
 		{
 			MoveFrom(Other);
+			return *this;
 		}
 
 		THandleArray()
@@ -217,11 +224,11 @@ namespace Chaos
 			}
 		}
 
-		ElementType* Get(FHandle InHandle)
+		ElementType* Get(FHandle InHandle) const
 		{
 			if(InHandle.IsValid() && Validity[InHandle.Index])
 			{
-				FHandleEntry& Entry = HandleEntries[InHandle.Index];
+				const FHandleEntry& Entry = HandleEntries[InHandle.Index];
 				return Entry.Generation == InHandle.Generation ? &(Data[InHandle.Index]) : nullptr;
 			}
 			return nullptr;
diff --git a/Engine/Source/Runtime/Experimental/Chaos/Public/Chaos/GeometryParticles.h b/Engine/Source/Runtime/Experimental/Chaos/Public/Chaos/GeometryParticles.h
--- a/Engine/Source/Runtime/Experimental/Chaos/Public/Chaos/GeometryParticles.h
+++ b/Engine/Source/Runtime/Experimental/Chaos/Public/Chaos/GeometryParticles.h
@@ -9,6 +9,7 @@
 #include "Chaos/GeometryParticlesfwd.h"
 #include "Chaos/CollisionFilterData.h"
 #include "Chaos/Box.h"
+#include "Chaos/PhysicalMaterials.h"
 #include "UObject/PhysicsObjectVersion.h"
 #include "UObject/ExternalPhysicsCustomObjectVersion.h"
 
@@ -33,7 +34,7 @@ namespace Chaos
 		void* UserData;
 		TSerializablePtr<FImplicitObject> Geometry;
 		TAABB<FReal, 3> WorldSpaceInflatedShapeBounds;
-		
+		TArray<FMaterialHandle> Materials;
 
 		static TPerShapeData<T, d>* SerializationFactory(FChaosArchive& Ar, TPerShapeData<T, d>*);
 		void Serialize(FChaosArchive& Ar);
diff --git a/Engine/Source/Runtime/Experimental/Chaos/Public/Chaos/HeightField.h b/Engine/Source/Runtime/Experimental/Chaos/Public/Chaos/HeightField.h
--- a/Engine/Source/Runtime/Experimental/Chaos/Public/Chaos/HeightField.h
+++ b/Engine/Source/Runtime/Experimental/Chaos/Public/Chaos/HeightField.h
@@ -27,8 +27,8 @@ namespace Chaos
 	public:
 		using FImplicitObject::GetTypeName;
 
-		THeightField(TArray<T>&& Height, int32 InNumRows, int32 InNumCols, const TVector<T,3>& InScale);
-		THeightField(TArrayView<const uint16> InHeights, int32 InNumRows, int32 InNumCols, const TVector<T, 3>& InScale);
+		THeightField(TArray<T>&& Height, TArray<uint8>&& InMaterialIndices, int32 InNumRows, int32 InNumCols, const TVector<T,3>& InScale);
+		THeightField(TArrayView<const uint16> InHeights, TArrayView<uint8> InMaterialIndices, int32 InNumRows, int32 InNumCols, const TVector<T, 3>& InScale);
 		THeightField(const THeightField& Other) = delete;
 		THeightField(THeightField&& Other) = default;
 
@@ -68,6 +68,28 @@ namespace Chaos
 		virtual int32 FindMostOpposingFace(const TVector<T, 3>& Position, const TVector<T, 3>& UnitDir, int32 HintFaceIndex, T SearchDist) const override;
 		virtual TVector<T, 3> FindGeometryOpposingNormal(const TVector<T, 3>& DenormDir, int32 FaceIndex, const TVector<T, 3>& OriginalNormal) const override;
 
+		virtual uint8 GetMaterialIndex(uint32 HintIndex) const override
+		{
+			ensure(GeomData.MaterialIndices.Num() > 0);
+
+			// If we've only got a default
+			if(GeomData.MaterialIndices.Num() == 1)
+			{
+				return GeomData.MaterialIndices[0];
+			}
+			else
+			{
+				// We store per cell for materials, so change to cell index
+				int32 CellIndex = HintIndex / 2;
+				if(ensure(GeomData.MaterialIndices.IsValidIndex(CellIndex)))
+				{
+					return GeomData.MaterialIndices[CellIndex];
+				}
+			}
+
+			return 0;
+		}
+
 		virtual const TBox<T, 3>& BoundingBox() const
 		{
 			CachedBounds = TBox<T, 3>(LocalBounds.Min() * GeomData.Scale, LocalBounds.Max() * GeomData.Scale);
@@ -152,6 +174,7 @@ namespace Chaos
 			// With HeightPerUnit being the range of the min/max realtype values of
 			// the heightfield divided by the range of StorageType
 			TArray<StorageType> Heights;
+			TArray<uint8> MaterialIndices;
 			TVector<RealType, 3> Scale;
 			RealType MinValue;
 			RealType MaxValue;
@@ -225,6 +248,11 @@ namespace Chaos
 					Ar << HeightPerUnit;
 					Ar << CellBounds;
 				}
+
+				if(Ar.CustomVer(FExternalPhysicsCustomObjectVersion::GUID) >= FExternalPhysicsCustomObjectVersion::AddedMaterialManager)
+				{
+					Ar << MaterialIndices;
+				}
 			}
 		};
 
@@ -419,5 +447,6 @@ namespace Chaos
 
 		template <typename QueryGeomType>
 		bool SweepGeomImp(const QueryGeomType& QueryGeom, const TRigidTransform<T, 3>& StartTM, const TVector<T, 3>& Dir, const T Length, T& OutTime, TVector<T, 3>& OutPosition, TVector<T, 3>& OutNormal, int32& OutFaceIndex, const T Thickness) const;
+
 	};
 }
diff --git a/Engine/Source/Runtime/Experimental/Chaos/Public/Chaos/ImplicitObject.h b/Engine/Source/Runtime/Experimental/Chaos/Public/Chaos/ImplicitObject.h
--- a/Engine/Source/Runtime/Experimental/Chaos/Public/Chaos/ImplicitObject.h
+++ b/Engine/Source/Runtime/Experimental/Chaos/Public/Chaos/ImplicitObject.h
@@ -272,6 +272,8 @@ public:
 
 	static const FName GetTypeName(const EImplicitObjectType InType);
 
+	virtual uint8 GetMaterialIndex(uint32 HintIndex) const { return 0; }
+
 protected:
 	EImplicitObjectType Type;
 	bool bIsConvex;
diff --git a/Engine/Source/Runtime/Experimental/Chaos/Public/Chaos/ImplicitObjectTransformed.h b/Engine/Source/Runtime/Experimental/Chaos/Public/Chaos/ImplicitObjectTransformed.h
--- a/Engine/Source/Runtime/Experimental/Chaos/Public/Chaos/ImplicitObjectTransformed.h
+++ b/Engine/Source/Runtime/Experimental/Chaos/Public/Chaos/ImplicitObjectTransformed.h
@@ -204,6 +204,11 @@ public:
 		return HashCombine(MObject->GetTypeHash(), ::GetTypeHash(MTransform));
 	}
 
+	virtual uint8 GetMaterialIndex(uint32 HintIndex) const override
+	{
+		return MObject->GetMaterialIndex(HintIndex);
+	}
+
 private:
 	ObjectType MObject;
 	TUniquePtr<Chaos::FImplicitObject> MObjectOwner;
@@ -212,6 +217,7 @@ private:
 
 	//needed for serialization
 	TImplicitObjectTransformed() : FImplicitObject(EImplicitObject::HasBoundingBox, ImplicitObjectType::Transformed) {}
+
 	friend FImplicitObject;	//needed for serialization
 };
 
diff --git a/Engine/Source/Runtime/Experimental/Chaos/Public/Chaos/PhysicalMaterials.h b/Engine/Source/Runtime/Experimental/Chaos/Public/Chaos/PhysicalMaterials.h
new file mode 100644
--- /dev/null
+++ b/Engine/Source/Runtime/Experimental/Chaos/Public/Chaos/PhysicalMaterials.h
@@ -0,0 +1,105 @@
+// Copyright 1998-2019 Epic Games, Inc. All Rights Reserved.
+
+#pragma once
+
+#include "Chaos/Core.h"
+#include "Chaos/Defines.h"
+#include "Chaos/Framework/Handles.h"
+
+namespace Chaos
+{
+	using FMaterialArray = THandleArray<FChaosPhysicsMaterial>;
+	using FChaosMaterialHandle = FMaterialArray::FHandle;
+	using FChaosConstMaterialHandle = FMaterialArray::FConstHandle;
+
+	/** 
+	 * Helper wrapper to encapsulate the access through the material manager
+	 * Handles returned from the material manager are only designed to be used (resolved) on the game thread
+	 * by the physics interface.
+	 */
+	struct CHAOS_API FMaterialHandle
+	{
+		FChaosPhysicsMaterial* Get() const;
+		FChaosMaterialHandle InnerHandle;
+
+		bool IsValid() const { return Get() != nullptr; }
+
+		friend bool operator ==(const FMaterialHandle& A, const FMaterialHandle& B) { return A.InnerHandle == B.InnerHandle; }
+
+		friend FArchive& operator <<(FArchive& Ar, FMaterialHandle& Value) { Ar << Value.InnerHandle; return Ar; }
+	};
+
+	/**
+	 * Helper wrapper to encapsulate the access through the material manager
+	 * Handles returned from the material manager are only designed to be used (resolved) on the game thread
+	 * by the physics interface.
+	 */
+	struct CHAOS_API FConstMaterialHandle
+	{
+		const FChaosPhysicsMaterial* Get() const;
+		FChaosConstMaterialHandle InnerHandle;
+
+		bool IsValid() const { return Get() != nullptr; }
+
+		friend bool operator ==(const FConstMaterialHandle& A, const FConstMaterialHandle& B) { return A.InnerHandle == B.InnerHandle; }
+
+		friend FArchive& operator <<(FArchive& Ar, FConstMaterialHandle& Value) { Ar << Value.InnerHandle; return Ar; }
+	};
+
+	DECLARE_MULTICAST_DELEGATE_OneParam(FOnMaterialCreated, FMaterialHandle);
+	DECLARE_MULTICAST_DELEGATE_OneParam(FOnMaterialDestroyed, FMaterialHandle);
+	DECLARE_MULTICAST_DELEGATE_OneParam(FOnMaterialUpdated, FMaterialHandle);
+	using FMaterialCreatedDelegate = FOnMaterialCreated::FDelegate;
+	using FMaterialDestroyedDelegate = FOnMaterialDestroyed::FDelegate;
+	using FMaterialUpdatedDelegate = FOnMaterialUpdated::FDelegate;
+
+	/** 
+	 * Global manager for physical materials.
+	 * Materials are created, updated and destroyed only on the game thread and an immutable
+	 * copy of the materials are stored on each solver. The solvers module binds to the updated event
+	 * in this manager and enqueues updates to all active solvers when a material is updated.
+	 *
+	 * The material manager provides handles for the objects which should be stored instead of the
+	 * material pointer. When accessing the internal material always use the handle rather than
+	 * storing the result of Get()
+	 */
+	class CHAOS_API FPhysicalMaterialManager
+	{
+	public:
+
+		static FPhysicalMaterialManager& Get();
+
+		/** Create a new material, returning a stable handle to it - this should be stored and not the actual material pointer */
+		FMaterialHandle Create();
+
+		/** Destroy the material referenced by the provided handle */
+		void Destroy(FMaterialHandle InHandle);
+
+		/** Get the actual material from a handle */
+		FChaosPhysicsMaterial* Resolve(FChaosMaterialHandle InHandle) const;
+		const FChaosPhysicsMaterial* Resolve(FChaosConstMaterialHandle InHandle) const;
+
+		/** Signals stakeholders that the stored material for the provided handle has changed */
+		void UpdateMaterial(FMaterialHandle InHandle);
+
+		/** Gets the internal list of master materials representing the current user state of the material data */
+		const THandleArray<FChaosPhysicsMaterial>& GetMasterMaterials() const;
+
+		/** Events */
+		FOnMaterialUpdated OnMaterialUpdated;
+		FOnMaterialCreated OnMaterialCreated;
+		FOnMaterialDestroyed OnMaterialDestroyed;
+
+	private:
+
+		/** Initial size for the handle-managed array */
+		static constexpr int32 InitialCapacity = 4;
+
+		/** Material manager is a singleton - access with Get() */
+		FPhysicalMaterialManager();
+
+		/** Handle-managed array of global materials. This is pushed to all solvers who all maintain a copy */
+		THandleArray<FChaosPhysicsMaterial> Materials;
+	};
+}
+
diff --git a/Engine/Source/Runtime/Experimental/Chaos/Public/ChaosStats.h b/Engine/Source/Runtime/Experimental/Chaos/Public/ChaosStats.h
--- a/Engine/Source/Runtime/Experimental/Chaos/Public/ChaosStats.h
+++ b/Engine/Source/Runtime/Experimental/Chaos/Public/ChaosStats.h
@@ -39,6 +39,7 @@ DECLARE_CYCLE_STAT_EXTERN(TEXT("Dispatch Collision Events"), STAT_DispatchCollis
 DECLARE_CYCLE_STAT_EXTERN(TEXT("Dispatch Break Events"), STAT_DispatchBreakEvents, STATGROUP_Chaos, CHAOS_API);
 DECLARE_CYCLE_STAT_EXTERN(TEXT("BufferPhysicsResults"), STAT_BufferPhysicsResults, STATGROUP_Chaos, CHAOS_API);
 DECLARE_CYCLE_STAT_EXTERN(TEXT("Flip Results"), STAT_FlipResults, STATGROUP_Chaos, CHAOS_API);
+DECLARE_CYCLE_STAT_EXTERN(TEXT("SQ - Update Materials"), STAT_SqUpdateMaterials, STATGROUP_Chaos, CHAOS_API);
 DECLARE_CYCLE_STAT_EXTERN(TEXT("[BufferPhysicsResults] - Geometry Collection"), STAT_CacheResultGeomCollection, STATGROUP_ChaosWide, CHAOS_API);
 DECLARE_CYCLE_STAT_EXTERN(TEXT("[BufferPhysicsResults] - StaticMesh"), STAT_CacheResultStaticMesh, STATGROUP_Chaos, CHAOS_API);
 DECLARE_CYCLE_STAT_EXTERN(TEXT("Capture Disabled State"), STAT_CaptureDisabledState, STATGROUP_Chaos, CHAOS_API);
diff --git a/Engine/Source/Runtime/Experimental/ChaosSolvers/Private/ChaosSolversModule.cpp b/Engine/Source/Runtime/Experimental/ChaosSolvers/Private/ChaosSolversModule.cpp
--- a/Engine/Source/Runtime/Experimental/ChaosSolvers/Private/ChaosSolversModule.cpp
+++ b/Engine/Source/Runtime/Experimental/ChaosSolvers/Private/ChaosSolversModule.cpp
@@ -19,6 +19,7 @@
 #include "Framework/Dispatcher.h"
 #include "Framework/DispatcherImpl.h"
 #include "Misc/App.h"
+#include "Chaos/PhysicalMaterials.h"
 
 TAutoConsoleVariable<int32> CVarChaosThreadEnabled(
 	TEXT("p.Chaos.DedicatedThreadEnabled"),
@@ -222,6 +223,12 @@ void FChaosSolversModule::Initialize()
 			break;
 		}
 
+		// Bind to the material manager
+		Chaos::FPhysicalMaterialManager& MaterialManager = Chaos::FPhysicalMaterialManager::Get();
+		OnCreateMaterialHandle = MaterialManager.OnMaterialCreated.Add(Chaos::FMaterialCreatedDelegate::CreateRaw(this, &FChaosSolversModule::OnCreateMaterial));
+		OnDestroyMaterialHandle = MaterialManager.OnMaterialDestroyed.Add(Chaos::FMaterialDestroyedDelegate::CreateRaw(this, &FChaosSolversModule::OnDestroyMaterial));
+		OnUpdateMaterialHandle = MaterialManager.OnMaterialUpdated.Add(Chaos::FMaterialUpdatedDelegate::CreateRaw(this, &FChaosSolversModule::OnUpdateMaterial));
+
 		bModuleInitialized = true;
 	}
 }
@@ -232,6 +239,12 @@ void FChaosSolversModule::Shutdown()
 	{
 		EndPhysicsTask();
 
+		// Unbind material events
+		Chaos::FPhysicalMaterialManager& MaterialManager = Chaos::FPhysicalMaterialManager::Get();
+		MaterialManager.OnMaterialCreated.Remove(OnCreateMaterialHandle);
+		MaterialManager.OnMaterialDestroyed.Remove(OnDestroyMaterialHandle);
+		MaterialManager.OnMaterialUpdated.Remove(OnUpdateMaterialHandle);
+
 		bModuleInitialized = false;
 	}
 }
@@ -464,6 +477,15 @@ Chaos::FPhysicsSolver* FChaosSolversModule::CreateSolver(bool bStandalone /*= fa
 	NewSolver->SetDebugName(NewDebugName);
 #endif
 
+	// Set up the material lists on the new solver, copying from the current master list
+	{
+		Chaos::FPhysicalMaterialManager& Manager =	Chaos::FPhysicalMaterialManager::Get();
+		NewSolver->QueryMaterialLock.WriteLock();
+		NewSolver->QueryMaterials = Manager.GetMasterMaterials();
+		NewSolver->SimMaterials = Manager.GetMasterMaterials();
+		NewSolver->QueryMaterialLock.WriteUnlock();
+	}
+
 	if(!bStandalone && IsPersistentTaskRunning() && Dispatcher)
 	{
 		// Need to let the thread know there's a new solver to care about
@@ -869,6 +891,66 @@ Chaos::EMultiBufferMode FChaosSolversModule::GetDesiredBufferingMode() const
 	return GetBufferModeFromThreadingModel(GetDesiredThreadingMode());
 }
 
+void FChaosSolversModule::OnUpdateMaterial(Chaos::FMaterialHandle InHandle)
+{
+	check(Dispatcher);
+
+	// Grab the material
+	Chaos::FChaosPhysicsMaterial* Material = InHandle.Get();
+
+	if(ensure(Material))
+	{
+		for(Chaos::FPhysicsSolver* Solver : Solvers)
+		{
+			// Send a copy of the material to each solver
+			Dispatcher->EnqueueCommandImmediate(Solver, [InHandle, MaterialCopy = *Material](Chaos::FPhysicsSolver* InSolver)
+			{
+				InSolver->UpdateMaterial(InHandle, MaterialCopy);
+			});
+		}
+	}
+}
+
+void FChaosSolversModule::OnCreateMaterial(Chaos::FMaterialHandle InHandle)
+{
+	check(Dispatcher);
+
+	// Grab the material
+	Chaos::FChaosPhysicsMaterial* Material = InHandle.Get();
+
+	if(ensure(Material))
+	{
+		for(Chaos::FPhysicsSolver* Solver : Solvers)
+		{
+			// Send a copy of the material to each solver
+			Dispatcher->EnqueueCommandImmediate(Solver, [InHandle, MaterialCopy = *Material](Chaos::FPhysicsSolver* InSolver)
+			{
+				InSolver->CreateMaterial(InHandle, MaterialCopy);
+			});
+		}
+	}
+}
+
+void FChaosSolversModule::OnDestroyMaterial(Chaos::FMaterialHandle InHandle)
+{
+	check(Dispatcher);
+
+	// Grab the material
+	Chaos::FChaosPhysicsMaterial* Material = InHandle.Get();
+
+	if(ensure(Material))
+	{
+		for(Chaos::FPhysicsSolver* Solver : Solvers)
+		{
+			// Notify each solver
+			Dispatcher->EnqueueCommandImmediate(Solver, [InHandle](Chaos::FPhysicsSolver* InSolver)
+			{
+				InSolver->DestroyMaterial(InHandle);
+			});
+		}
+	}
+}
+
 const IChaosSettingsProvider& FChaosSolversModule::GetSettingsProvider() const
 {
 	return SettingsProvider ? *SettingsProvider : Chaos::GDefaultChaosSettings;
diff --git a/Engine/Source/Runtime/Experimental/ChaosSolvers/Private/PBDRigidsSolver.cpp b/Engine/Source/Runtime/Experimental/ChaosSolvers/Private/PBDRigidsSolver.cpp
--- a/Engine/Source/Runtime/Experimental/ChaosSolvers/Private/PBDRigidsSolver.cpp
+++ b/Engine/Source/Runtime/Experimental/ChaosSolvers/Private/PBDRigidsSolver.cpp
@@ -520,4 +520,25 @@ namespace Chaos
 	}
 
 
+	void FPBDRigidsSolver::UpdateMaterial(Chaos::FMaterialHandle InHandle, const Chaos::FChaosPhysicsMaterial& InNewData)
+	{
+		*SimMaterials.Get(InHandle.InnerHandle) = InNewData;
+	}
+
+	void FPBDRigidsSolver::CreateMaterial(Chaos::FMaterialHandle InHandle, const Chaos::FChaosPhysicsMaterial& InNewData)
+	{
+		ensure(SimMaterials.Create(InNewData) == InHandle.InnerHandle);
+	}
+
+	void FPBDRigidsSolver::DestroyMaterial(Chaos::FMaterialHandle InHandle)
+	{
+		SimMaterials.Destroy(InHandle.InnerHandle);
+	}
+
+	void FPBDRigidsSolver::SyncQueryMaterials()
+	{
+		TSolverQueryMaterialScope<ELockType::Write> Scope(this);
+		QueryMaterials = SimMaterials;
+	}
+
 }; // namespace Chaos
diff --git a/Engine/Source/Runtime/Experimental/ChaosSolvers/Public/ChaosSolversModule.h b/Engine/Source/Runtime/Experimental/ChaosSolvers/Public/ChaosSolversModule.h
--- a/Engine/Source/Runtime/Experimental/ChaosSolvers/Public/ChaosSolversModule.h
+++ b/Engine/Source/Runtime/Experimental/ChaosSolvers/Public/ChaosSolversModule.h
@@ -15,6 +15,7 @@
 #include "PhysicsCoreTypes.h"
 #include "Chaos/Framework/MultiBufferResource.h"
 #include "Chaos/Declares.h"
+#include "Chaos/PhysicalMaterials.h"
 
 /** Classes that want to set the solver actor class can implement this. */
 class IChaosSolverActorClassProvider
@@ -288,6 +289,11 @@ public:
 	Chaos::EThreadingMode GetDesiredThreadingMode() const;
 	Chaos::EMultiBufferMode GetDesiredBufferingMode() const;
 
+	/** Events that the material manager will emit for us to update our threaded data */
+	void OnUpdateMaterial(Chaos::FMaterialHandle InHandle);
+	void OnCreateMaterial(Chaos::FMaterialHandle InHandle);
+	void OnDestroyMaterial(Chaos::FMaterialHandle InHandle);
+
 private:
 
 	/** Safe method for always getting a settings provider (from the external caller or an internal default) */
@@ -352,6 +358,11 @@ private:
 
 	// Whether we're initialized, gates work in Initialize() and Shutdown()
 	bool bModuleInitialized;
+
+	/** Event binding handles */
+	FDelegateHandle OnCreateMaterialHandle;
+	FDelegateHandle OnDestroyMaterialHandle;
+	FDelegateHandle OnUpdateMaterialHandle;
 };
 
 /**
diff --git a/Engine/Source/Runtime/Experimental/ChaosSolvers/Public/PBDRigidsSolver.h b/Engine/Source/Runtime/Experimental/ChaosSolvers/Public/PBDRigidsSolver.h
--- a/Engine/Source/Runtime/Experimental/ChaosSolvers/Public/PBDRigidsSolver.h
+++ b/Engine/Source/Runtime/Experimental/ChaosSolvers/Public/PBDRigidsSolver.h
@@ -41,6 +41,20 @@ namespace Chaos
 	class FPersistentPhysicsTask;
 	class FPhysicsCommand;
 	class FChaosArchive;
+	class FPBDRigidsSolver;
+
+	enum class ELockType : uint8
+	{
+		Read,
+		Write
+	};
+
+	template<ELockType LockType>
+	struct TSolverQueryMaterialScope
+	{
+		TSolverQueryMaterialScope() = delete;
+	};
+
 	/**
 	*
 	*/
@@ -262,6 +276,18 @@ namespace Chaos
 		/**/
 		void PostTickDebugDraw() const;
 
+		/** Events hooked up to the Chaos material manager */
+		void UpdateMaterial(Chaos::FMaterialHandle InHandle, const Chaos::FChaosPhysicsMaterial& InNewData);
+		void CreateMaterial(Chaos::FMaterialHandle InHandle, const Chaos::FChaosPhysicsMaterial& InNewData);
+		void DestroyMaterial(Chaos::FMaterialHandle InHandle);
+
+		/** Access to the intenal material mirrors */
+		const THandleArray<FChaosPhysicsMaterial>& GetQueryMaterials() const { return QueryMaterials; }
+		const THandleArray<FChaosPhysicsMaterial>& GetSimMaterials() const { return SimMaterials; }
+
+		/** Copy the simulation material list to the query material list, to be done when the SQ commits an update */
+		void SyncQueryMaterials();
+
 	private:
 
 		template<typename ParticleType>
@@ -319,7 +345,62 @@ namespace Chaos
 		TArray< FStaticMeshPhysicsProxy* > StaticMeshPhysicsProxies; // dep
 		TArray< FGeometryCollectionPhysicsProxy* > GeometryCollectionPhysicsProxies;
 		TArray< FFieldSystemPhysicsProxy* > FieldSystemPhysicsProxies;
+
+		// Physics material mirrors for the solver. These should generally stay in sync with the global material list from
+		// the game thread. This data is read only in the solver as we should never need to update it here. External threads can
+		// Enqueue commands to change parameters.
+		//
+		// There are two copies here to enable SQ to lock only the solvers that it needs to handle the material access during a query
+		// instead of having to lock the entire physics state of the runtime.
+		FRWLock QueryMaterialLock;
+		THandleArray<FChaosPhysicsMaterial> QueryMaterials;
+		THandleArray<FChaosPhysicsMaterial> SimMaterials;
 				
+		template<ELockType>
+		friend struct TSolverQueryMaterialScope;
+	};
+
+	template<>
+	struct TSolverQueryMaterialScope<ELockType::Read>
+	{
+		TSolverQueryMaterialScope() = delete;
+
+
+		explicit TSolverQueryMaterialScope(FPBDRigidsSolver* InSolver)
+			: Solver(InSolver)
+		{
+			check(Solver);
+			Solver->QueryMaterialLock.ReadLock();
+		}
+
+		~TSolverQueryMaterialScope()
+		{
+			Solver->QueryMaterialLock.ReadUnlock();
+		}
+
+	private:
+		FPBDRigidsSolver* Solver;
+	};
+
+	template<>
+	struct TSolverQueryMaterialScope<ELockType::Write>
+	{
+		TSolverQueryMaterialScope() = delete;
+
+		explicit TSolverQueryMaterialScope(FPBDRigidsSolver* InSolver)
+			: Solver(InSolver)
+		{
+			check(Solver);
+			Solver->QueryMaterialLock.WriteLock();
+		}
+
+		~TSolverQueryMaterialScope()
+		{
+			Solver->QueryMaterialLock.WriteUnlock();
+		}
+
+	private:
+		FPBDRigidsSolver* Solver;
 	};
 
 }; // namespace Chaos
diff --git a/Engine/Source/Runtime/Landscape/Classes/LandscapeHeightfieldCollisionComponent.h b/Engine/Source/Runtime/Landscape/Classes/LandscapeHeightfieldCollisionComponent.h
--- a/Engine/Source/Runtime/Landscape/Classes/LandscapeHeightfieldCollisionComponent.h
+++ b/Engine/Source/Runtime/Landscape/Classes/LandscapeHeightfieldCollisionComponent.h
@@ -95,6 +95,7 @@ class ULandscapeHeightfieldCollisionComponent : public UPrimitiveComponent
 #endif	//WITH_PHYSX
 
 #if WITH_CHAOS
+		TArray<Chaos::FMaterialHandle> UsedChaosMaterials;
 		TUniquePtr<Chaos::THeightField<float>> Heightfield = nullptr;
 	    TUniquePtr<Chaos::THeightField<float>> HeightfieldSimple = nullptr;
 #if WITH_EDITOR
diff --git a/Engine/Source/Runtime/Landscape/Private/LandscapeCollision.cpp b/Engine/Source/Runtime/Landscape/Private/LandscapeCollision.cpp
--- a/Engine/Source/Runtime/Landscape/Private/LandscapeCollision.cpp
+++ b/Engine/Source/Runtime/Landscape/Private/LandscapeCollision.cpp
@@ -128,7 +128,7 @@ ULandscapeMeshCollisionComponent::FPhysXMeshRef::~FPhysXMeshRef()
 }
 
 // Generate a new guid to force a recache of landscape collison derived data
-#define LANDSCAPE_COLLISION_DERIVEDDATA_VER	TEXT("9CAF1A48A5224D4BB307CEBE2776DE27")
+#define LANDSCAPE_COLLISION_DERIVEDDATA_VER	TEXT("1ECE4BBBBC7146DA8C6ADAD75BA52C77")
 
 static FString GetHFDDCKeyString(const FName& Format, bool bDefMaterial, const FGuid& StateId, const TArray<UPhysicalMaterial*>& PhysicalMaterials)
 {
@@ -423,6 +423,7 @@ void ULandscapeHeightfieldCollisionComponent::OnCreatePhysicsState()
 				NewShape->Geometry = MakeSerializable(ChaosHeightFieldFromCooked);
 				NewShape->QueryData = QueryFilterData;
 				NewShape->SimData = SimFilterData;
+				NewShape->Materials = HeightfieldRef->UsedChaosMaterials;
 
 				Geoms.Emplace(MoveTemp(ChaosHeightFieldFromCooked));
 				ShapeArray.Emplace(MoveTemp(NewShape));
@@ -600,7 +601,8 @@ void ULandscapeHeightfieldCollisionComponent::CreateCollisionObject()
 
 				for (UPhysicalMaterial* PhysicalMaterial : CookedPhysicalMaterials)
 				{
-					HeightfieldRef->UsedPhysicalMaterialArray.Add(GPhysXSDK->createMaterial(1, 1, 1));
+					const FPhysicsMaterialHandle_PhysX& MaterialHandle = PhysicalMaterial->GetPhysicsMaterial();
+					HeightfieldRef->UsedPhysicalMaterialArray.Add(MaterialHandle.Material);
 				}
 
 				// Release cooked collison data
@@ -642,6 +644,8 @@ void ULandscapeHeightfieldCollisionComponent::CreateCollisionObject()
 				{
 					//todo: total hack until we get landscape fully converted to chaos
 					HeightfieldRef->UsedPhysicalMaterialArray.Add(GPhysXSDK->createMaterial(1, 1, 1));
+
+					HeightfieldRef->UsedChaosMaterials.Add(PhysicalMaterial->GetPhysicsMaterial());
 				}
 
 #if WITH_EDITOR
@@ -864,6 +868,44 @@ bool ULandscapeHeightfieldCollisionComponent::CookCollisionData(const FName& For
 		Succeeded = Cooker->CookHeightField(Format, HFSizeSimple, SimpleSamples.GetData(), SimpleSamples.GetTypeSize(), OutData);
 	}
 #elif WITH_CHAOS
+
+	// Generate material indices
+	TArray<uint8> MaterialIndices;
+	MaterialIndices.Reserve(NumSamples + NumSimpleSamples);
+	for(int32 RowIndex = 0; RowIndex < CollisionSizeVerts; RowIndex++)
+	{
+		for(int32 ColIndex = 0; ColIndex < CollisionSizeVerts; ColIndex++)
+		{
+			const int32 SrcSampleIndex = (ColIndex * CollisionSizeVerts) + (bIsMirrored ? RowIndex : (CollisionSizeVerts - RowIndex - 1));
+
+			// Materials are not relevant on the last row/column because they are per-triangle and the last row/column don't own any
+			if(RowIndex < CollisionSizeVerts - 1 &&
+				ColIndex < CollisionSizeVerts - 1)
+			{
+				int32 MaterialIndex = 0; // Default physical material.
+				if(!bUseDefMaterial && DominantLayers)
+				{
+					uint8 DominantLayerIdx = DominantLayers[SrcSampleIndex];
+					if(ComponentLayerInfos.IsValidIndex(DominantLayerIdx))
+					{
+						ULandscapeLayerInfoObject* Layer = ComponentLayerInfos[DominantLayerIdx];
+						if(Layer == ALandscapeProxy::VisibilityLayer)
+						{
+							// If it's a hole, use the final index
+							MaterialIndex = TNumericLimits<uint8>::Max();
+						}
+						else
+						{
+							UPhysicalMaterial* DominantMaterial = Layer && Layer->PhysMaterial ? Layer->PhysMaterial : DefMaterial;
+							MaterialIndex = InOutMaterials.AddUnique(DominantMaterial);
+						}
+					}
+				}
+				MaterialIndices.Add(MaterialIndex);
+			}
+		}
+	}
+
 	TUniquePtr<Chaos::THeightField<float>> Heightfield = nullptr;
 	TUniquePtr<Chaos::THeightField<float>> HeightfieldSimple = nullptr;
 
@@ -874,12 +916,13 @@ bool ULandscapeHeightfieldCollisionComponent::CookCollisionData(const FName& For
 	Ar << bSerializeGenerateSimpleCollision;
 
 	TArrayView<const uint16> ComplexHeightView(Heights, NumSamples);
-	Heightfield = MakeUnique<Chaos::THeightField<float>>(ComplexHeightView, CollisionSizeVerts, CollisionSizeVerts, Chaos::TVector<float, 3>(1));
+	Heightfield = MakeUnique<Chaos::THeightField<float>>(ComplexHeightView, MakeArrayView(MaterialIndices), CollisionSizeVerts, CollisionSizeVerts, Chaos::TVector<float, 3>(1));
 	Ar << Heightfield;
 	if(bGenerateSimpleCollision)
 	{
+		// #BGTODO Materials for simple geometry, currently just passing in the default
 		TArrayView<const uint16> SimpleHeightView(Heights + NumSamples, NumSimpleSamples);
-		HeightfieldSimple = MakeUnique<Chaos::THeightField<float>>(SimpleHeightView, CollisionSizeVerts, CollisionSizeVerts, Chaos::TVector<float, 3>(1));
+		HeightfieldSimple = MakeUnique<Chaos::THeightField<float>>(SimpleHeightView, MakeArrayView(MaterialIndices.GetData(), 1), CollisionSizeVerts, CollisionSizeVerts, Chaos::TVector<float, 3>(1));
 		Ar << HeightfieldSimple;
 	}
 
diff --git a/Engine/Source/Runtime/PhysicsCore/Public/ChaosInterfaceWrapperCore.h b/Engine/Source/Runtime/PhysicsCore/Public/ChaosInterfaceWrapperCore.h
--- a/Engine/Source/Runtime/PhysicsCore/Public/ChaosInterfaceWrapperCore.h
+++ b/Engine/Source/Runtime/PhysicsCore/Public/ChaosInterfaceWrapperCore.h
@@ -138,11 +138,6 @@ inline FVector GetNormal(const FLocationHit& Hit)
 	return Hit.WorldNormal;
 }
 
-inline FDummyPhysType* GetMaterialFromInternalFaceIndex(const FDummyPhysType& Shape, uint32 InternalFaceIndex)
-{
-	return nullptr;
-}
-
 inline FHitFlags GetFlags(const FLocationHit& Hit)
 {
 	return Hit.Flags;
diff --git a/Engine/Source/Runtime/PhysicsCore/Public/PhysXInterfaceWrapperCore.h b/Engine/Source/Runtime/PhysicsCore/Public/PhysXInterfaceWrapperCore.h
--- a/Engine/Source/Runtime/PhysicsCore/Public/PhysXInterfaceWrapperCore.h
+++ b/Engine/Source/Runtime/PhysicsCore/Public/PhysXInterfaceWrapperCore.h
@@ -216,7 +216,7 @@ FORCEINLINE ECollisionShapeType GetGeometryType(const PxShape& Shape)
 	return P2UGeometryType(Shape.getGeometryType());
 }
 
-FORCEINLINE PxMaterial* GetMaterialFromInternalFaceIndex(const PxShape& Shape, uint32 InternalFaceIndex)
+FORCEINLINE PxMaterial* GetMaterialFromInternalFaceIndex(const PxShape& Shape, const FPhysicsActor& Actor, uint32 InternalFaceIndex)
 {
 	return Shape.getMaterialFromInternalFaceIndex(InternalFaceIndex);
 }
diff --git a/Engine/Source/Runtime/PhysicsCore/Public/PhysXToChaosUtil.h b/Engine/Source/Runtime/PhysicsCore/Public/PhysXToChaosUtil.h
--- a/Engine/Source/Runtime/PhysicsCore/Public/PhysXToChaosUtil.h
+++ b/Engine/Source/Runtime/PhysicsCore/Public/PhysXToChaosUtil.h
@@ -88,6 +88,15 @@ inline TUniquePtr<Chaos::TImplicitObjectTransformed<float, 3>> PxShapeToChaosGeo
 
 		HeightField->saveCells(CellBuffer.GetData(), sizeof(CellBuffer[0]) * CellBuffer.Num());
 
+		const int32 NumTris = (NumCells - NumRows - NumCols + 1) * 2;
+		TArray<uint8> MaterialIndices;
+		MaterialIndices.Reserve(NumTris / 2);
+		for(int32 TriangleIndex = 0; TriangleIndex < NumTris; TriangleIndex += 2)
+		{
+			// We're only grabbing every other material because UE4 will only use one material per cell
+			MaterialIndices.Add((uint8)(HeightField->getTriangleMaterialIndex(TriangleIndex)));
+		}
+		
 		TArray<float> Height;
 		Height.AddUninitialized(NumRows * NumCols);
 		//PhysX and unreal have opposite handedness so we flip the data (see LandscapeComponent)
@@ -116,7 +125,7 @@ inline TUniquePtr<Chaos::TImplicitObjectTransformed<float, 3>> PxShapeToChaosGeo
 
 		ShapeTM = U2PTransform(FTransform(LandscapeComponentMatrix));
 
-		InnerObj = MakeUnique<THeightField<float>>(MoveTemp(Height), NumRows, NumCols, TVector<float, 3>(HeightFieldGeom.columnScale, HeightFieldGeom.rowScale, HeightFieldGeom.heightScale));
+		InnerObj = MakeUnique<THeightField<float>>(MoveTemp(Height), MoveTemp(MaterialIndices), NumRows, NumCols, TVector<float, 3>(HeightFieldGeom.columnScale, HeightFieldGeom.rowScale, HeightFieldGeom.heightScale));
 		break;
 	}
 	case PxGeometryType::eTRIANGLEMESH:
diff --git a/Engine/Source/Runtime/PhysicsCore/Public/PhysicsInterfaceDeclaresCore.h b/Engine/Source/Runtime/PhysicsCore/Public/PhysicsInterfaceDeclaresCore.h
--- a/Engine/Source/Runtime/PhysicsCore/Public/PhysicsInterfaceDeclaresCore.h
+++ b/Engine/Source/Runtime/PhysicsCore/Public/PhysicsInterfaceDeclaresCore.h
@@ -79,6 +79,15 @@ namespace Chaos
 
 	template <typename T, int d>
 	class TPerShapeData;
+
+	class FPhysicalMaterial;
+
+	template<typename, uint32, uint32>
+	class THandle;
+
+	struct FMaterialHandle;
+
+	class FChaosPhysicsMaterial;
 }
 
 // Temporary dummy types until SQ implemented
@@ -107,7 +116,7 @@ using FPhysicsTransform = FTransform;
 using FPhysicsShape = Chaos::TPerShapeData<float, 3>;
 using FPhysicsGeometry = Chaos::FImplicitObject;
 using FPhysicsCapsuleGeometry = Chaos::TCapsule<float>;
-using FPhysicsMaterial = FPhysTypeDummy;
+using FPhysicsMaterial = Chaos::FChaosPhysicsMaterial;
 using FPhysicsActor = Chaos::TGeometryParticle<float,3>;
 
 template <typename T>
@@ -206,7 +215,7 @@ typedef FPhysicsAggregateReference_Chaos	FPhysicsAggregateHandle;
 typedef FPhysInterface_Chaos				FPhysicsCommand;
 typedef FPhysicsShapeReference_Chaos		FPhysicsShapeHandle;
 typedef FPhysicsGeometryCollection_Chaos	FPhysicsGeometryCollection;
-typedef void*								FPhysicsMaterialHandle;
+typedef Chaos::FMaterialHandle				FPhysicsMaterialHandle;
 typedef FPhysicsShapeAdapter_Chaos			FPhysicsShapeAdapter;
 typedef FPhysicsUserData_Chaos				FPhysicsUserData;
 
