Fix landscape flickering issue on iOS
Invalidated mobile landscape DDC key

diff --git a/Engine/Source/Runtime/Landscape/Private/Landscape.cpp b/Engine/Source/Runtime/Landscape/Private/Landscape.cpp
--- a/Engine/Source/Runtime/Landscape/Private/Landscape.cpp
+++ b/Engine/Source/Runtime/Landscape/Private/Landscape.cpp
@@ -116,7 +116,7 @@ namespace LandscapeCookStats
 // differences, etc.) replace the version GUID below with a new one.
 // In case of merge conflicts with DDC versions, you MUST generate a new GUID
 // and set this new GUID as the version.                                       
-#define LANDSCAPE_MOBILE_COOK_VERSION TEXT("A048A0D4A24644BA9948FB08068AE8D7")
+#define LANDSCAPE_MOBILE_COOK_VERSION TEXT("6862AA3DD9FB4368B9DDAF9EB7E9901F")
 
 #define LOCTEXT_NAMESPACE "Landscape"
 
