Landscape: Fix DefaultMaterial Collision when there is no painted data

#rnx

diff --git a/Engine/Source/Runtime/Landscape/Private/LandscapeCollision.cpp b/Engine/Source/Runtime/Landscape/Private/LandscapeCollision.cpp
--- a/Engine/Source/Runtime/Landscape/Private/LandscapeCollision.cpp
+++ b/Engine/Source/Runtime/Landscape/Private/LandscapeCollision.cpp
@@ -708,7 +708,10 @@ void ULandscapeHeightfieldCollisionComponent::SpeculativelyLoadAsyncDDCCollsionD
 TArray<PxHeightFieldSample> ConvertHeightfieldDataForPhysx(const ULandscapeHeightfieldCollisionComponent* const Component, const int32 CollisionSizeVerts, const bool bIsMirrored, const uint16* Heights, const bool bUseDefMaterial, const uint8* DominantLayers, UPhysicalMaterial* const DefMaterial, TArray<UPhysicalMaterial*> &InOutMaterials)
 {
 	const int32 NumSamples = FMath::Square(CollisionSizeVerts);
-
+	check(DefMaterial);
+	// Might return INDEX_NONE if DefMaterial wasn't added yet
+	int32 DefaultMaterialIndex = InOutMaterials.IndexOfByKey(DefMaterial);
+	
 	TArray<PxHeightFieldSample> Samples;
 	Samples.Reserve(NumSamples);
 	Samples.AddZeroed(NumSamples);
@@ -727,7 +730,7 @@ TArray<PxHeightFieldSample> ConvertHeightfieldDataForPhysx(const ULandscapeHeigh
 			if (RowIndex < CollisionSizeVerts - 1 &&
 				ColIndex < CollisionSizeVerts - 1)
 			{
-				int32 MaterialIndex = 0; // Default physical material.
+				int32 MaterialIndex = DefaultMaterialIndex; // Default physical material.
 				if (!bUseDefMaterial && DominantLayers)
 				{
 					uint8 DominantLayerIdx = DominantLayers[SrcSampleIndex];
@@ -739,14 +742,23 @@ TArray<PxHeightFieldSample> ConvertHeightfieldDataForPhysx(const ULandscapeHeigh
 							// If it's a hole, override with the hole flag.
 							MaterialIndex = PxHeightFieldMaterial::eHOLE;
 						}
+						else if (Layer && Layer->PhysMaterial)
+						{
+							MaterialIndex = InOutMaterials.AddUnique(Layer->PhysMaterial);
+						}
 						else
 						{
-							UPhysicalMaterial* DominantMaterial = Layer && Layer->PhysMaterial ? Layer->PhysMaterial : DefMaterial;
-							MaterialIndex = InOutMaterials.AddUnique(DominantMaterial);
+							MaterialIndex = DefaultMaterialIndex;
 						}
 					}
 				}
 
+				// Default Material but Def Material wasn't added yet...
+				if (MaterialIndex == INDEX_NONE)
+				{
+					MaterialIndex = DefaultMaterialIndex = InOutMaterials.Add(DefMaterial);
+				}
+
 				Sample.materialIndex0 = MaterialIndex;
 				Sample.materialIndex1 = MaterialIndex;
 			}
@@ -843,7 +855,7 @@ bool ULandscapeHeightfieldCollisionComponent::CookCollisionData(const FName& For
 #if WITH_PHYSX && PHYSICS_INTERFACE_PHYSX
 	// List of materials which is actually used by heightfield
 	InOutMaterials.Empty();
-
+		
 	TArray<PxHeightFieldSample> Samples;
 	TArray<PxHeightFieldSample> SimpleSamples;
 	Samples = ConvertHeightfieldDataForPhysx(this, CollisionSizeVerts, bIsMirrored, Heights, bUseDefMaterial, DominantLayers, DefMaterial, InOutMaterials);
@@ -858,13 +870,7 @@ bool ULandscapeHeightfieldCollisionComponent::CookCollisionData(const FName& For
 	{
 		DominantLayerData.Unlock();
 	}
-
-	// Add the default physical material to be used used when we have no dominant data.
-	if (InOutMaterials.Num() == 0)
-	{
-		InOutMaterials.Add(DefMaterial);
-	}
-
+	
 	//
 	FIntPoint HFSize = FIntPoint(CollisionSizeVerts, CollisionSizeVerts);
 
