Landscape Collision Material: Add Minimum Collision Relevance Value

#rnx

diff --git a/Engine/Source/Runtime/Landscape/Classes/LandscapeComponent.h b/Engine/Source/Runtime/Landscape/Classes/LandscapeComponent.h
--- a/Engine/Source/Runtime/Landscape/Classes/LandscapeComponent.h
+++ b/Engine/Source/Runtime/Landscape/Classes/LandscapeComponent.h
@@ -963,7 +963,7 @@ protected:
 	void UpdateCollisionHeightBuffer(int32 InComponentX1, int32 InComponentY1, int32 InComponentX2, int32 InComponentY2, int32 InCollisionMipLevel, int32 InHeightmapSizeU, int32 InHeightmapSizeV,
 		const FColor* const InHeightmapTextureMipData, uint16* CollisionHeightData, uint16* GrassHeightData,
 		const FColor* const InXYOffsetTextureMipData, uint16* CollisionXYOffsetData);
-	void UpdateDominantLayerBuffer(int32 InComponentX1, int32 InComponentY1, int32 InComponentX2, int32 InComponentY2, int32 InCollisionMipLevel, int32 InWeightmapSizeU, int32 InDataLayerIdx, const TArray<uint8*>& InCollisionDataPtrs, uint8* DominantLayerData);
+	void UpdateDominantLayerBuffer(int32 InComponentX1, int32 InComponentY1, int32 InComponentX2, int32 InComponentY2, int32 InCollisionMipLevel, int32 InWeightmapSizeU, int32 InDataLayerIdx, const TArray<uint8*>& InCollisionDataPtrs, const TArray<ULandscapeLayerInfoObject*>& InLayerInfos, uint8* DominantLayerData);
 #endif
 
 	/** Whether the component type supports static lighting. */
diff --git a/Engine/Source/Runtime/Landscape/Classes/LandscapeLayerInfoObject.h b/Engine/Source/Runtime/Landscape/Classes/LandscapeLayerInfoObject.h
--- a/Engine/Source/Runtime/Landscape/Classes/LandscapeLayerInfoObject.h
+++ b/Engine/Source/Runtime/Landscape/Classes/LandscapeLayerInfoObject.h
@@ -35,6 +35,10 @@ class ULandscapeLayerInfoObject : public UObject
 	float Hardness;
 
 #if WITH_EDITORONLY_DATA
+	/* The minimum weight that needs to be painted for that layer to be picked up as the dominant physical layer */
+	UPROPERTY(EditAnywhere, Category=LandscapeLayerInfoObject, Meta = (ClampMin="0", ClampMax="1"))
+	float MinimumCollisionRelevanceWeight;
+
 	UPROPERTY(EditAnywhere, Category=LandscapeLayerInfoObject)
 	uint32 bNoWeightBlend:1;
 
diff --git a/Engine/Source/Runtime/Landscape/Private/LandscapeEdit.cpp b/Engine/Source/Runtime/Landscape/Private/LandscapeEdit.cpp
--- a/Engine/Source/Runtime/Landscape/Private/LandscapeEdit.cpp
+++ b/Engine/Source/Runtime/Landscape/Private/LandscapeEdit.cpp
@@ -1218,7 +1218,7 @@ void ULandscapeComponent::UpdateCollisionHeightBuffer(	int32 InComponentX1, int3
 	}
 }
 
-void ULandscapeComponent::UpdateDominantLayerBuffer(int32 InComponentX1, int32 InComponentY1, int32 InComponentX2, int32 InComponentY2, int32 InCollisionMipLevel, int32 InWeightmapSizeU, int32 InDataLayerIdx, const TArray<uint8*>& InCollisionDataPtrs, uint8* OutDominantLayerData)
+void ULandscapeComponent::UpdateDominantLayerBuffer(int32 InComponentX1, int32 InComponentY1, int32 InComponentX2, int32 InComponentY2, int32 InCollisionMipLevel, int32 InWeightmapSizeU, int32 InDataLayerIdx, const TArray<uint8*>& InCollisionDataPtrs, const TArray<ULandscapeLayerInfoObject*>& InLayerInfos, uint8* OutDominantLayerData)
 {
 	const int32 MipSizeU = InWeightmapSizeU >> InCollisionMipLevel;
 
@@ -1267,6 +1267,7 @@ void ULandscapeComponent::UpdateDominantLayerBuffer(int32 InComponentX1, int32 I
 					for (int32 LayerIdx = 0; LayerIdx < InCollisionDataPtrs.Num(); LayerIdx++)
 					{
 						const uint8 LayerWeight = InCollisionDataPtrs[LayerIdx][DataOffset];
+						const uint8 LayerMinimumWeight = InLayerInfos[LayerIdx] ? (uint8)(InLayerInfos[LayerIdx]->MinimumCollisionRelevanceWeight * 255) :  0;
 
 						if (LayerIdx == InDataLayerIdx) // Override value for hole
 						{
@@ -1276,7 +1277,7 @@ void ULandscapeComponent::UpdateDominantLayerBuffer(int32 InComponentX1, int32 I
 								DominantWeight = INT_MAX;
 							}
 						}
-						else if (LayerWeight > DominantWeight)
+						else if (LayerWeight > DominantWeight && LayerWeight >= LayerMinimumWeight)
 						{
 							DominantLayer = LayerIdx;
 							DominantWeight = LayerWeight;
@@ -1391,12 +1392,12 @@ void ULandscapeComponent::UpdateCollisionLayerData(const FColor* const* const We
 			const int32 WeightmapSizeU = ComponentWeightmapsTexture[0]->Source.GetSizeX();
 						
 			// gmartin: WeightmapScaleBias not handled?			
-			UpdateDominantLayerBuffer(ComponentX1, ComponentY1, ComponentX2, ComponentY2, CollisionMipLevel, WeightmapSizeU, DataLayerIdx, CandidateDataPtrs, DominantLayerData);
+			UpdateDominantLayerBuffer(ComponentX1, ComponentY1, ComponentX2, ComponentY2, CollisionMipLevel, WeightmapSizeU, DataLayerIdx, CandidateDataPtrs, CollisionComp->ComponentLayerInfos, DominantLayerData);
 
 			if (bUsingSimpleCollision)
 			{
 				uint8* const SimpleCollisionHeightData = DominantLayerData + CollisionSize.SizeVertsSquare;
-				UpdateDominantLayerBuffer(ComponentX1, ComponentY1, ComponentX2, ComponentY2, SimpleCollisionMipLevel, WeightmapSizeU, DataLayerIdx, SimpleCollisionDataPtrs, SimpleCollisionHeightData);
+				UpdateDominantLayerBuffer(ComponentX1, ComponentY1, ComponentX2, ComponentY2, SimpleCollisionMipLevel, WeightmapSizeU, DataLayerIdx, SimpleCollisionDataPtrs, CollisionComp->ComponentLayerInfos, SimpleCollisionHeightData);
 			}
 
 			CollisionComp->DominantLayerData.Unlock();
@@ -3747,6 +3748,7 @@ ULandscapeLayerInfoObject::ULandscapeLayerInfoObject(const FObjectInitializer& O
 {
 	Hardness = 0.5f;
 #if WITH_EDITORONLY_DATA
+	MinimumCollisionRelevanceWeight = 0.0f;
 	bNoWeightBlend = false;
 	SplineFalloffModulationTexture = nullptr;
 	SplineFalloffModulationColorMask = ESplineModulationColorMask::Red;
@@ -3768,9 +3770,10 @@ ULandscapeLayerInfoObject::ULandscapeLayerInfoObject(const FObjectInitializer& O
 #if WITH_EDITOR
 void ULandscapeLayerInfoObject::PostEditChangeProperty(FPropertyChangedEvent& PropertyChangedEvent)
 {
-	static const FName NAME_Hardness = FName(TEXT("Hardness"));
-	static const FName NAME_PhysMaterial = FName(TEXT("PhysMaterial"));
-	static const FName NAME_LayerUsageDebugColor = FName(TEXT("LayerUsageDebugColor"));
+	static const FName NAME_Hardness = GET_MEMBER_NAME_CHECKED(ULandscapeLayerInfoObject, Hardness);
+	static const FName NAME_PhysMaterial = GET_MEMBER_NAME_CHECKED(ULandscapeLayerInfoObject, PhysMaterial);
+	static const FName NAME_LayerUsageDebugColor = GET_MEMBER_NAME_CHECKED(ULandscapeLayerInfoObject, LayerUsageDebugColor);
+	static const FName NAME_MinimumCollisionRelevanceWeight = GET_MEMBER_NAME_CHECKED(ULandscapeLayerInfoObject, MinimumCollisionRelevanceWeight);
 	static const FName NAME_R = FName(TEXT("R"));
 	static const FName NAME_G = FName(TEXT("G"));
 	static const FName NAME_B = FName(TEXT("B"));
@@ -3786,7 +3789,7 @@ void ULandscapeLayerInfoObject::PostEditChangeProperty(FPropertyChangedEvent& Pr
 		{
 			Hardness = FMath::Clamp<float>(Hardness, 0.0f, 1.0f);
 		}
-		else if (PropertyName == NAME_PhysMaterial)
+		else if (PropertyName == NAME_PhysMaterial || PropertyName == NAME_MinimumCollisionRelevanceWeight)
 		{
 			for (TObjectIterator<ALandscapeProxy> It; It; ++It)
 			{
