Crash Assertion Failed in Landscape Mountain Project After Setting Fixed LOD to 7
#rnx

diff --git a/Engine/Source/Runtime/Landscape/Private/LandscapeRender.cpp b/Engine/Source/Runtime/Landscape/Private/LandscapeRender.cpp
--- a/Engine/Source/Runtime/Landscape/Private/LandscapeRender.cpp
+++ b/Engine/Source/Runtime/Landscape/Private/LandscapeRender.cpp
@@ -3221,27 +3221,26 @@ void FLandscapeComponentSceneProxy::GetDynamicMeshElements(const TArray<const FS
 	{
 		if (VisibilityMap & (1 << ViewIndex))
 		{
-			const FSceneView* View = Views[ViewIndex];
 			FLandscapeElementParamArray& ParameterArray = Collector.AllocateOneFrameResource<FLandscapeElementParamArray>();
+			ParameterArray.ElementParams.AddDefaulted(1);
 
-			float MeshScreenSizeSquared = ComputeBoundsScreenRadiusSquared(GetBounds().Origin, GetBounds().SphereRadius, *View);
-			int32 ForcedLODLevel = (View->Family->EngineShowFlags.LOD) ? GetCVarForceLOD() : -1;
+			const FSceneView* View = Views[ViewIndex];
 
-			float LODScale = View->LODDistanceFactor * CVarStaticMeshLODDistanceScale.GetValueOnRenderThread();
+			int32 ForcedLODLevel = (View->Family->EngineShowFlags.LOD) ? GetCVarForceLOD() : -1;
 #if WITH_EDITOR
 			ForcedLODLevel = View->Family->LandscapeLODOverride >= 0 ? View->Family->LandscapeLODOverride : ForcedLODLevel;
 #endif
-			int32 LODToRender = ForcedLODLevel >= 0 ? ForcedLODLevel : GetLODFromScreenSize(MeshScreenSizeSquared, LODScale * LODScale);
+			ForcedLODLevel = FMath::Min(ForcedLODLevel, (int32)LODSettings.LastLODIndex);
 
-			ParameterArray.ElementParams.AddDefaulted(1);
+			const float LODScale = View->LODDistanceFactor * CVarStaticMeshLODDistanceScale.GetValueOnRenderThread();
+			const float MeshScreenSizeSquared = ComputeBoundsScreenRadiusSquared(GetBounds().Origin, GetBounds().SphereRadius, *View);
+			const int32 LODToRender = ForcedLODLevel >= 0 ? ForcedLODLevel : GetLODFromScreenSize(MeshScreenSizeSquared, LODScale * LODScale);
 
 			FMeshBatch& Mesh = Collector.AllocateMesh();
-
 			GetStaticMeshElement(LODToRender, false, ForcedLODLevel >= 0, Mesh, ParameterArray.ElementParams);
 
 #if WITH_EDITOR
 			FMeshBatch& MeshTools = Collector.AllocateMesh();
-
 			// No Tessellation on tool material
 			GetStaticMeshElement(LODToRender, true, ForcedLODLevel >= 0, MeshTools, ParameterArray.ElementParams);
 #endif
