Fix PlayerCollision LandscapeRender

#rnx

diff --git a/Engine/Source/Runtime/Landscape/Private/LandscapeRender.cpp b/Engine/Source/Runtime/Landscape/Private/LandscapeRender.cpp
--- a/Engine/Source/Runtime/Landscape/Private/LandscapeRender.cpp
+++ b/Engine/Source/Runtime/Landscape/Private/LandscapeRender.cpp
@@ -3210,6 +3210,7 @@ void FLandscapeComponentSceneProxy::GetDynamicMeshElements(const TArray<const FS
 	const bool bInCollisionView = ViewFamily.EngineShowFlags.CollisionVisibility || ViewFamily.EngineShowFlags.CollisionPawn;
 	const bool bDrawSimpleCollision = ViewFamily.EngineShowFlags.CollisionPawn       && CollisionResponse.GetResponse(ECC_Pawn) != ECR_Ignore;
 	const bool bDrawComplexCollision = ViewFamily.EngineShowFlags.CollisionVisibility && CollisionResponse.GetResponse(ECC_Visibility) != ECR_Ignore;
+	const int32 CollisionLODLevel = bDrawSimpleCollision ? FMath::Max(CollisionMipLevel, SimpleCollisionMipLevel) : bDrawComplexCollision ? CollisionMipLevel : -1;
 #endif
 
 	int32 NumPasses = 0;
@@ -3230,6 +3231,11 @@ void FLandscapeComponentSceneProxy::GetDynamicMeshElements(const TArray<const FS
 #if WITH_EDITOR
 			ForcedLODLevel = View->Family->LandscapeLODOverride >= 0 ? View->Family->LandscapeLODOverride : ForcedLODLevel;
 #endif
+
+#if WITH_EDITOR || !(UE_BUILD_SHIPPING || UE_BUILD_TEST)
+			ForcedLODLevel = CollisionLODLevel >= 0 ? CollisionLODLevel : ForcedLODLevel;
+#endif
+
 			ForcedLODLevel = FMath::Min(ForcedLODLevel, (int32)LODSettings.LastLODIndex);
 
 			const float LODScale = View->LODDistanceFactor * CVarStaticMeshLODDistanceScale.GetValueOnRenderThread();
