Fixed crash when displaying a landscape with no proxy in it, which can happen during the LOD baking process of the levels

diff --git a/Engine/Source/Runtime/Landscape/Private/LandscapeRender.cpp b/Engine/Source/Runtime/Landscape/Private/LandscapeRender.cpp
--- a/Engine/Source/Runtime/Landscape/Private/LandscapeRender.cpp
+++ b/Engine/Source/Runtime/Landscape/Private/LandscapeRender.cpp
@@ -917,6 +917,7 @@ void FLandscapeRenderSystem::RecreateBuffers(const FSceneView* InView /* = nullp
 	{
 		TRACE_CPUPROFILER_EVENT_SCOPE(FLandscapeRenderSystem::RecreateBuffers());
 
+		if (Size != FIntPoint::ZeroValue)
 		{
 			if (!SectionLODBuffer.IsValid())
 			{
@@ -930,9 +931,7 @@ void FLandscapeRenderSystem::RecreateBuffers(const FSceneView* InView /* = nullp
 				FMemory::Memcpy(Data, SectionLODValues.GetData(), SectionLODValues.GetResourceDataSize());
 				RHIUnlockVertexBuffer(SectionLODBuffer);
 			}
-		}
 
-		{
 			if (!SectionLODBiasBuffer.IsValid())
 			{
 				FRHIResourceCreateInfo CreateInfo(&SectionLODBiases);
@@ -945,9 +944,7 @@ void FLandscapeRenderSystem::RecreateBuffers(const FSceneView* InView /* = nullp
 				FMemory::Memcpy(Data, SectionLODBiases.GetData(), SectionLODBiases.GetResourceDataSize());
 				RHIUnlockVertexBuffer(SectionLODBiasBuffer);
 			}
-		}
 
-		{
 			if (!SectionTessellationFalloffCBuffer.IsValid())
 			{
 				FRHIResourceCreateInfo CreateInfo(&SectionTessellationFalloffC);
@@ -964,9 +961,7 @@ void FLandscapeRenderSystem::RecreateBuffers(const FSceneView* InView /* = nullp
 					RHIUnlockVertexBuffer(SectionTessellationFalloffCBuffer);
 				}
 			}
-		}
 
-		{
 			if (!SectionTessellationFalloffKBuffer.IsValid())
 			{
 				FRHIResourceCreateInfo CreateInfo(&SectionTessellationFalloffK);
@@ -983,23 +978,23 @@ void FLandscapeRenderSystem::RecreateBuffers(const FSceneView* InView /* = nullp
 					RHIUnlockVertexBuffer(SectionTessellationFalloffKBuffer);
 				}
 			}
-		}
 
-		FLandscapeSectionLODUniformParameters Parameters;
-		Parameters.Min = Min;
-		Parameters.Size = Size;
-		Parameters.SectionLOD = SectionLODSRV;
-		Parameters.SectionLODBias = SectionLODBiasSRV;
-		Parameters.SectionTessellationFalloffC = SectionTessellationFalloffCSRV;
-		Parameters.SectionTessellationFalloffK = SectionTessellationFalloffKSRV;
+			FLandscapeSectionLODUniformParameters Parameters;
+			Parameters.Min = Min;
+			Parameters.Size = Size;
+			Parameters.SectionLOD = SectionLODSRV;
+			Parameters.SectionLODBias = SectionLODBiasSRV;
+			Parameters.SectionTessellationFalloffC = SectionTessellationFalloffCSRV;
+			Parameters.SectionTessellationFalloffK = SectionTessellationFalloffKSRV;
 
-		if (UniformBuffer.IsValid())
-		{
-			UniformBuffer.UpdateUniformBufferImmediate(Parameters);
-		}
-		else
-		{
-			UniformBuffer = TUniformBufferRef<FLandscapeSectionLODUniformParameters>::CreateUniformBufferImmediate(Parameters, UniformBuffer_SingleFrame);
+			if (UniformBuffer.IsValid())
+			{
+				UniformBuffer.UpdateUniformBufferImmediate(Parameters);
+			}
+			else
+			{
+				UniformBuffer = TUniformBufferRef<FLandscapeSectionLODUniformParameters>::CreateUniformBufferImmediate(Parameters, UniformBuffer_SingleFrame);
+			}
 		}
 
 		CachedView = InView;
