Addition of an option to prevent navmesh generation under the terrain

diff --git a/Engine/Source/Runtime/Landscape/Classes/LandscapeComponent.h b/Engine/Source/Runtime/Landscape/Classes/LandscapeComponent.h
--- a/Engine/Source/Runtime/Landscape/Classes/LandscapeComponent.h
+++ b/Engine/Source/Runtime/Landscape/Classes/LandscapeComponent.h
@@ -937,6 +937,9 @@ public:
 
 	/** Updates navigation properties to match landscape's master switch */
 	void UpdateNavigationRelevance();
+
+	/** Updates the reject navmesh underneath flag in the collision component */
+	void UpdateRejectNavmeshUnderneath();
 	
 	/** Updates the values of component-level properties exposed by the Landscape Actor */
 	LANDSCAPE_API void UpdatedSharedPropertiesFromActor();
diff --git a/Engine/Source/Runtime/Landscape/Classes/LandscapeProxy.h b/Engine/Source/Runtime/Landscape/Classes/LandscapeProxy.h
--- a/Engine/Source/Runtime/Landscape/Classes/LandscapeProxy.h
+++ b/Engine/Source/Runtime/Landscape/Classes/LandscapeProxy.h
@@ -695,6 +695,10 @@ public:
 	UPROPERTY(EditAnywhere, Category=Landscape)
 	uint32 bUsedForNavigation:1;
 
+	/** Set to true to prevent navmesh generation under the terrain geometry */
+	UPROPERTY(EditAnywhere, Category = Landscape)
+	uint32 bRejectNavmeshUnderLandscapeGeometry:1;
+
 	/** When set to true it will generate MaterialInstanceDynamic for each components, so material can be changed at runtime */
 	UPROPERTY(EditAnywhere, Category = Landscape)
 	bool bUseDynamicMaterialInstance;
diff --git a/Engine/Source/Runtime/Landscape/Private/Landscape.cpp b/Engine/Source/Runtime/Landscape/Private/Landscape.cpp
--- a/Engine/Source/Runtime/Landscape/Private/Landscape.cpp
+++ b/Engine/Source/Runtime/Landscape/Private/Landscape.cpp
@@ -1048,6 +1048,7 @@ ALandscapeProxy::ALandscapeProxy(const FObjectInitializer& ObjectInitializer)
 	bCastStaticShadow = true;
 	bCastShadowAsTwoSided = false;
 	bUsedForNavigation = true;
+	bRejectNavmeshUnderLandscapeGeometry = false;
 	CollisionThickness = 16;
 	BodyInstance.SetCollisionProfileName(UCollisionProfile::BlockAll_ProfileName);
 	bGenerateOverlapEvents = false;
@@ -2002,6 +2003,14 @@ void ALandscape::PostLoad()
 			Brush.SetOwner(this);
 		}
 	}
+
+	for (ULandscapeComponent* Comp : LandscapeComponents)
+	{
+		if (Comp)
+		{
+			Comp->UpdateRejectNavmeshUnderneath();
+		}
+	}
 #endif
 
 	Super::PostLoad();
diff --git a/Engine/Source/Runtime/Landscape/Private/LandscapeEdit.cpp b/Engine/Source/Runtime/Landscape/Private/LandscapeEdit.cpp
--- a/Engine/Source/Runtime/Landscape/Private/LandscapeEdit.cpp
+++ b/Engine/Source/Runtime/Landscape/Private/LandscapeEdit.cpp
@@ -149,6 +149,15 @@ void ULandscapeComponent::UpdateNavigationRelevance()
 	}
 }
 
+void ULandscapeComponent::UpdateRejectNavmeshUnderneath()
+{
+	ALandscapeProxy* Proxy = GetLandscapeProxy();
+	if (CollisionComponent && Proxy)
+	{
+		CollisionComponent->bRejectNavmeshUnderneath = Proxy->bRejectNavmeshUnderLandscapeGeometry;
+	}
+}
+
 ULandscapeMaterialInstanceConstant* ALandscapeProxy::GetLayerThumbnailMIC(UMaterialInterface* LandscapeMaterial, FName LayerName, UTexture2D* ThumbnailWeightmap, UTexture2D* ThumbnailHeightmap, ALandscapeProxy* Proxy)
 {
 	if (!LandscapeMaterial)
@@ -4705,6 +4714,7 @@ void ALandscape::PostEditChangeProperty(FPropertyChangedEvent& PropertyChangedEv
 	bool bNeedsRecalcBoundingBox = false;
 	bool bChangedLighting = false;
 	bool bChangedNavRelevance = false;
+	bool bChangeRejectNavmeshUnder = false;
 	bool bPropagateToProxies = false;
 
 	ULandscapeInfo* Info = GetLandscapeInfo();
@@ -4850,6 +4860,10 @@ void ALandscape::PostEditChangeProperty(FPropertyChangedEvent& PropertyChangedEv
 	{
 		bChangedNavRelevance = true;
 	}
+	else if (GIsEditor && PropertyName == GET_MEMBER_NAME_CHECKED(ALandscapeProxy, bRejectNavmeshUnderLandscapeGeometry))
+	{
+		bChangeRejectNavmeshUnder = true;
+	}
 
 	// Must do this *after* clamping values
 	Super::PostEditChangeProperty(PropertyChangedEvent);
@@ -4875,7 +4889,7 @@ void ALandscape::PostEditChangeProperty(FPropertyChangedEvent& PropertyChangedEv
 			LandscapeEdit.RecalculateNormals();
 		}
 
-		if (bNeedsRecalcBoundingBox || ChangedMaterial || bChangedLighting || bChangedNavRelevance)
+		if (bNeedsRecalcBoundingBox || ChangedMaterial || bChangedLighting || bChangedNavRelevance || bChangeRejectNavmeshUnder)
 		{
 			// We cannot iterate the XYtoComponentMap directly because reregistering components modifies the array.
 			TArray<ULandscapeComponent*> AllComponents;
@@ -4901,6 +4915,11 @@ void ALandscape::PostEditChangeProperty(FPropertyChangedEvent& PropertyChangedEv
 					{
 						Comp->UpdateNavigationRelevance();
 					}
+
+					if (bChangeRejectNavmeshUnder)
+					{
+						Comp->UpdateRejectNavmeshUnderneath();
+					}
 				}
 			}
 
