Fixed issues with mobile weightmap textures causing non-deterministic builds:

* the weightmaps were not zero-initialized, so unused channels contained random data
* generated textures got a random lighting GUID at each build. Changed it so that the GUID is actually the MD5 of the texture name, so it's stable between builds.

#rnx

diff --git a/Engine/Source/Runtime/Engine/Classes/Engine/Texture.h b/Engine/Source/Runtime/Engine/Classes/Engine/Texture.h
--- a/Engine/Source/Runtime/Engine/Classes/Engine/Texture.h
+++ b/Engine/Source/Runtime/Engine/Classes/Engine/Texture.h
@@ -1108,6 +1108,11 @@ public:
 #endif // WITH_EDITORONLY_DATA
 	}
 
+	void SetLightingGuid(const FGuid& Guid)
+	{
+		LightingGuid = Guid;
+	}
+
 	/**
 	 * Retrieves the pixel format enum for enum <-> string conversions.
 	 */
diff --git a/Engine/Source/Runtime/Landscape/Private/LandscapeEdit.cpp b/Engine/Source/Runtime/Landscape/Private/LandscapeEdit.cpp
--- a/Engine/Source/Runtime/Landscape/Private/LandscapeEdit.cpp
+++ b/Engine/Source/Runtime/Landscape/Private/LandscapeEdit.cpp
@@ -6067,7 +6067,7 @@ void ULandscapeComponent::GeneratePlatformPixelData()
 	MobileWeightmapTextures.Empty();
 
     UTexture2D* MobileWeightNormalmapTexture = GetLandscapeProxy()->CreateLandscapeTexture(WeightmapSize, WeightmapSize, TEXTUREGROUP_Terrain_Weightmap, TSF_BGRA8, nullptr, GMobileCompressLandscapeWeightMaps ? true : false );
-	CreateEmptyTextureMips(MobileWeightNormalmapTexture);
+	CreateEmptyTextureMips(MobileWeightNormalmapTexture, true);
 
 	{
 		FLandscapeTextureDataInterface LandscapeData;
@@ -6110,7 +6110,7 @@ void ULandscapeComponent::GeneratePlatformPixelData()
 					CurrentChannel = 0;
 					RemainingChannels = 4;
                     CurrentWeightmapTexture = GetLandscapeProxy()->CreateLandscapeTexture(WeightmapSize, WeightmapSize, TEXTUREGROUP_Terrain_Weightmap, TSF_BGRA8, nullptr, GMobileCompressLandscapeWeightMaps ? true : false);
-					CreateEmptyTextureMips(CurrentWeightmapTexture);
+					CreateEmptyTextureMips(CurrentWeightmapTexture, true);
 					MobileWeightmapTextures.Add(CurrentWeightmapTexture);
 				}
 
@@ -6127,7 +6127,23 @@ void ULandscapeComponent::GeneratePlatformPixelData()
 	GDisableAutomaticTextureMaterialUpdateDependencies = true;
 	for (int TextureIdx = 0; TextureIdx < MobileWeightmapTextures.Num(); TextureIdx++)
 	{
-		MobileWeightmapTextures[TextureIdx]->PostEditChange();
+		UTexture* Texture = MobileWeightmapTextures[TextureIdx];
+		Texture->PostEditChange();
+
+		// PostEditChange() will assign a random GUID to the texture, which leads to non-deterministic builds.
+		// Compute a 128-bit hash based on the texture name and use that as a GUID to fix this issue.
+		FTCHARToUTF8 Converted(*Texture->GetFullName());
+		FMD5 MD5Gen;
+		MD5Gen.Update((const uint8*)Converted.Get(), Converted.Length());
+		uint32 Digest[4];
+		MD5Gen.Final((uint8*)Digest);
+
+		// FGuid::NewGuid() creates a version 4 UUID (at least on Windows), which will have the top 4 bits of the
+		// second field set to 0100. We'll set the top bit to 1 in the GUID we create, to ensure that we can never
+		// have a collision with textures which use implicitly generated GUIDs.
+		Digest[1] |= 0x80000000;
+		FGuid TextureGUID(Digest[0], Digest[1], Digest[2], Digest[3]);
+		Texture->SetLightingGuid(TextureGUID);
 	}
 	GDisableAutomaticTextureMaterialUpdateDependencies = false;
 
