Fixed LOD level computation for landscape in ES2 profile

#rnx

diff --git a/Engine/Source/Runtime/Landscape/Private/LandscapeRender.cpp b/Engine/Source/Runtime/Landscape/Private/LandscapeRender.cpp
--- a/Engine/Source/Runtime/Landscape/Private/LandscapeRender.cpp
+++ b/Engine/Source/Runtime/Landscape/Private/LandscapeRender.cpp
@@ -1601,6 +1601,7 @@ FPrimitiveViewRelevance FLandscapeComponentSceneProxy::GetViewRelevance(const FS
 	Result.bRenderCustomDepth = ShouldRenderCustomDepth();
 	Result.bUsesLightingChannels = GetLightingChannelMask() != GetDefaultLightingChannelMask();
 	Result.bTranslucentSelfShadow = bCastVolumetricTranslucentShadow;
+	Result.bUseCustomViewData = true;
 
 	auto FeatureLevel = View->GetFeatureLevel();
 
diff --git a/Engine/Source/Runtime/Landscape/Private/LandscapeRenderMobile.cpp b/Engine/Source/Runtime/Landscape/Private/LandscapeRenderMobile.cpp
--- a/Engine/Source/Runtime/Landscape/Private/LandscapeRenderMobile.cpp
+++ b/Engine/Source/Runtime/Landscape/Private/LandscapeRenderMobile.cpp
@@ -123,6 +123,61 @@ public:
 		{
 			ShaderBindings.Add(ForcedLodParameter, BatchElementParams->ForcedLOD);
 		}
+
+		FLandscapeComponentSceneProxy::FViewCustomDataLOD* LODData = (FLandscapeComponentSceneProxy::FViewCustomDataLOD*)InView->GetCustomData(SceneProxy->GetPrimitiveSceneInfo()->GetIndex());
+		int32 SubSectionIndex = BatchElementParams->SubX + BatchElementParams->SubY * SceneProxy->NumSubsections;
+
+		if (LODData != nullptr)
+		{
+			SceneProxy->PostInitViewCustomData(*InView, LODData);
+
+			if (LodTessellationParameter.IsBound())
+			{
+				ShaderBindings.Add(LodTessellationParameter, LODData->LodTessellationParams);
+			}
+
+			if (SectionLodsParameter.IsBound())
+			{
+				if (LODData->UseCombinedMeshBatch)
+				{
+					ShaderBindings.Add(SectionLodsParameter, LODData->ShaderCurrentLOD);
+				}
+				else // in non combined, only the one representing us as we'll be called 4 times (once per sub section)
+				{
+					check(SubSectionIndex >= 0);
+					FVector4 ShaderCurrentLOD(ForceInitToZero);
+					ShaderCurrentLOD.Component(SubSectionIndex) = LODData->ShaderCurrentLOD.Component(SubSectionIndex);
+
+					ShaderBindings.Add(SectionLodsParameter, ShaderCurrentLOD);
+				}
+			}
+
+			if (NeighborSectionLodParameter.IsBound())
+			{
+				FVector4 ShaderCurrentNeighborLOD[FLandscapeComponentSceneProxy::NEIGHBOR_COUNT] = { FVector4(ForceInitToZero), FVector4(ForceInitToZero), FVector4(ForceInitToZero), FVector4(ForceInitToZero) };
+
+				if (LODData->UseCombinedMeshBatch)
+				{
+					int32 SubSectionCount = SceneProxy->NumSubsections == 1 ? 1 : FLandscapeComponentSceneProxy::MAX_SUBSECTION_COUNT;
+
+					for (int32 NeighborSubSectionIndex = 0; NeighborSubSectionIndex < SubSectionCount; ++NeighborSubSectionIndex)
+					{
+						ShaderCurrentNeighborLOD[NeighborSubSectionIndex] = LODData->SubSections[NeighborSubSectionIndex].ShaderCurrentNeighborLOD;
+						check(ShaderCurrentNeighborLOD[NeighborSubSectionIndex].X != -1.0f); // they should all match so only check the 1st one for simplicity
+					}
+
+					ShaderBindings.Add(NeighborSectionLodParameter, ShaderCurrentNeighborLOD);
+				}
+				else // in non combined, only the one representing us as we'll be called 4 times (once per sub section)
+				{
+					check(SubSectionIndex >= 0);
+					ShaderCurrentNeighborLOD[SubSectionIndex] = LODData->SubSections[SubSectionIndex].ShaderCurrentNeighborLOD;
+					check(ShaderCurrentNeighborLOD[SubSectionIndex].X != -1.0f); // they should all match so only check the 1st one for simplicity
+
+					ShaderBindings.Add(NeighborSectionLodParameter, ShaderCurrentNeighborLOD);
+				}
+			}
+		}
 	}
 protected:
 	FShaderParameter LodValuesParameter;
