Forbid ray tracing shaders in the landscape thumbnail render

diff --git a/Engine/Source/Runtime/Landscape/Private/LandscapeRender.cpp b/Engine/Source/Runtime/Landscape/Private/LandscapeRender.cpp
--- a/Engine/Source/Runtime/Landscape/Private/LandscapeRender.cpp
+++ b/Engine/Source/Runtime/Landscape/Private/LandscapeRender.cpp
@@ -4667,13 +4667,13 @@ public:
 				static const FName LocalVertexFactory = FName(TEXT("FLocalVertexFactory"));
 				if (!IsMobilePlatform(Platform) && VertexFactoryType->GetFName() == LocalVertexFactory)
 				{
-					if (Algo::Find(GetAllowedShaderTypes(), ShaderType->GetFName()))
+					if (Algo::Find(GetAllowedShaderTypesInThumbnailRender(), ShaderType->GetFName()))
 					{
 						return FMaterialResource::ShouldCache(Platform, ShaderType, VertexFactoryType);
 					}
 					else
 					{
-						if (Algo::Find(GetExcludedShaderTypes(), ShaderType->GetFName()))
+						if (Algo::Find(GetExcludedShaderTypesInThumbnailRender(), ShaderType->GetFName()))
 						{
 							UE_LOG(LogLandscape, VeryVerbose, TEXT("Excluding shader %s from landscape thumbnail material"), ShaderType->GetName());
 							return false;
@@ -4718,7 +4718,7 @@ public:
 		return false;
 	}
 
-	static const TArray<FName>& GetAllowedShaderTypes()
+	static const TArray<FName>& GetAllowedShaderTypesInThumbnailRender()
 	{
 		// reduce the number of shaders compiled for the thumbnail materials by only compiling with shader types known to be used by the preview scene
 		static const TArray<FName> AllowedShaderTypes =
@@ -4815,7 +4815,7 @@ public:
 		return AllowedShaderTypes;
 	}
 
-	static const TArray<FName>& GetExcludedShaderTypes()
+	static const TArray<FName>& GetExcludedShaderTypesInThumbnailRender()
 	{
 		// shader types known *not* to be used by the preview scene
 		static const TArray<FName> ExcludedShaderTypes =
@@ -4970,6 +4970,13 @@ public:
 
 			FName(TEXT("TBasePassDSFPrecomputedVolumetricLightmapLightingPolicy")),
 			FName(TEXT("TBasePassHSFPrecomputedVolumetricLightmapLightingPolicy")),
+
+#if RHI_RAYTRACING
+				// No ray tracing on thumbnails
+				FName(TEXT("TMaterialCHSFPrecomputedVolumetricLightmapLightingPolicy")),
+				FName(TEXT("TMaterialCHSFNoLightMapPolicy")),
+				FName(TEXT("FRayTracingDynamicGeometryConverterCS")),
+#endif // RHI_RAYTRACING
 		};
 		return ExcludedShaderTypes;
 	}
