Cleanup chaos cooked data after it's used in landscape initialization

diff --git a/Engine/Source/Runtime/Landscape/Private/LandscapeCollision.cpp b/Engine/Source/Runtime/Landscape/Private/LandscapeCollision.cpp
--- a/Engine/Source/Runtime/Landscape/Private/LandscapeCollision.cpp
+++ b/Engine/Source/Runtime/Landscape/Private/LandscapeCollision.cpp
@@ -661,6 +661,13 @@ void ULandscapeHeightfieldCollisionComponent::CreateCollisionObject()
 					HeightfieldRef->UsedChaosMaterials.Add(PhysicalMaterial->GetPhysicsMaterial());
 				}
 
+				// Release cooked collison data
+				// In cooked builds created collision object will never be deleted while component is alive, so we don't need this data anymore
+				if(FPlatformProperties::RequiresCookedData() || World->IsGameWorld())
+				{
+					CookedCollisionData.Empty();
+				}
+
 #if WITH_EDITOR
 				// Create heightfield for the landscape editor (no holes in it)
 				if(!World->IsGameWorld())
