Added support for Procedural scatter tool to read Paint Layers of any Landscape Edit Layers.

diff --git a/Engine/Source/Runtime/Landscape/Classes/LandscapeComponent.h b/Engine/Source/Runtime/Landscape/Classes/LandscapeComponent.h
--- a/Engine/Source/Runtime/Landscape/Classes/LandscapeComponent.h
+++ b/Engine/Source/Runtime/Landscape/Classes/LandscapeComponent.h
@@ -743,7 +743,7 @@ public:
 	LANDSCAPE_API float EditorGetPaintLayerWeightByNameAtLocation(const FVector& InLocation, const FName InPaintLayerName);
 		
 	/** Get the landscape actor associated with this component. */
-	ALandscape* GetLandscapeActor() const;
+	LANDSCAPE_API ALandscape* GetLandscapeActor() const;
 
 	/** Get the level in which the owning actor resides */
 	ULevel* GetLevel() const;
@@ -930,7 +930,7 @@ public:
 	LANDSCAPE_API void InitWeightmapData(TArray<ULandscapeLayerInfoObject*>& LayerInfos, TArray<TArray<uint8> >& Weights);
 
 	/** @todo document */
-	LANDSCAPE_API float GetLayerWeightAtLocation( const FVector& InLocation, ULandscapeLayerInfoObject* LayerInfo, TArray<uint8>* LayerCache=NULL );
+	LANDSCAPE_API float GetLayerWeightAtLocation( const FVector& InLocation, ULandscapeLayerInfoObject* LayerInfo, TArray<uint8>* LayerCache = NULL, bool bUseEditingWeightmap = false);
 
 	/** Extends passed region with this component section size */
 	LANDSCAPE_API void GetComponentExtent(int32& MinX, int32& MinY, int32& MaxX, int32& MaxY) const;
diff --git a/Engine/Source/Runtime/Landscape/Private/LandscapeEdit.cpp b/Engine/Source/Runtime/Landscape/Private/LandscapeEdit.cpp
--- a/Engine/Source/Runtime/Landscape/Private/LandscapeEdit.cpp
+++ b/Engine/Source/Runtime/Landscape/Private/LandscapeEdit.cpp
@@ -1981,7 +1981,7 @@ void ULandscapeComponent::UpdateDataMips(int32 InNumSubsections, int32 InSubsect
 	UpdateMipsTempl<uint8>(InNumSubsections, InSubsectionSizeQuads, Texture, TextureMipData, ComponentX1, ComponentY1, ComponentX2, ComponentY2, TextureDataInfo);
 }
 
-float ULandscapeComponent::GetLayerWeightAtLocation(const FVector& InLocation, ULandscapeLayerInfoObject* LayerInfo, TArray<uint8>* LayerCache)
+float ULandscapeComponent::GetLayerWeightAtLocation(const FVector& InLocation, ULandscapeLayerInfoObject* LayerInfo, TArray<uint8>* LayerCache, bool bUseEditingWeightmap)
 {
 	// Allocate and discard locally if no external cache is passed in.
 	TArray<uint8> LocalCache;
@@ -1994,7 +1994,7 @@ float ULandscapeComponent::GetLayerWeightAtLocation(const FVector& InLocation, U
 	if (LayerCache->Num() == 0)
 	{
 		FLandscapeComponentDataInterface CDI(this);
-		if (!CDI.GetWeightmapTextureData(LayerInfo, *LayerCache))
+		if (!CDI.GetWeightmapTextureData(LayerInfo, *LayerCache, bUseEditingWeightmap))
 		{
 			// no data for this layer for this component.
 			return 0.0f;
