Landscape Splines: Optimize Raster by moving cache up one level (was about to get even slower with forthcoming LS map split)

#rnx

diff --git a/Engine/Source/Runtime/Landscape/Classes/LandscapeInfo.h b/Engine/Source/Runtime/Landscape/Classes/LandscapeInfo.h
--- a/Engine/Source/Runtime/Landscape/Classes/LandscapeInfo.h
+++ b/Engine/Source/Runtime/Landscape/Classes/LandscapeInfo.h
@@ -21,6 +21,7 @@ class ULandscapeSplinesComponent;
 class ULandscapeSplineControlPoint;
 class ULandscapeSplineSegment;
 class ULandscapeHeightfieldCollisionComponent;
+class FModulateAlpha;
 
 /** Structure storing Collision for LandscapeComponent Add */
 #if WITH_EDITORONLY_DATA
@@ -158,7 +159,6 @@ public:
 	LANDSCAPE_API void ExportHeightmap(const FString& Filename);
 	LANDSCAPE_API void ExportLayer(ULandscapeLayerInfoObject* LayerInfo, const FString& Filename);
 	LANDSCAPE_API bool ApplySplines(bool bOnlySelected, TSet<ULandscapeComponent*>* OutModifiedComponents = nullptr, bool bMarkPackageDirty = true);
-	bool ApplySplinesInternal(bool bOnlySelected, ALandscapeProxy* Proxy, TSet<ULandscapeComponent*>* OutModifiedComponents, bool bMarkPackageDirty);
 
 	LANDSCAPE_API bool GetSelectedExtent(int32& MinX, int32& MinY, int32& MaxX, int32& MaxY) const;
 	FVector GetLandscapeCenterPos(float& LengthZ, int32 MinX = MAX_int32, int32 MinY = MAX_int32, int32 MaxX = MIN_int32, int32 MaxY = MIN_int32);
@@ -284,6 +284,7 @@ public:
 
 #if WITH_EDITOR
 private:
+	bool ApplySplinesInternal(bool bOnlySelected, ALandscapeProxy* Proxy, TSet<ULandscapeComponent*>* OutModifiedComponents, bool bMarkPackageDirty, int32 LandscapeMinX, int32 LandscapeMinY, int32 LandscapeMaxX, int32 LandscapeMaxY, TFunctionRef<TSharedPtr<FModulateAlpha>(ULandscapeLayerInfoObject*)> GetOrCreateModulate);
 	void MoveSegmentToLandscape(ULandscapeSplineSegment* InSegment, ALandscapeProxy* FromProxy, ALandscapeProxy* ToProxy);
 	void MoveControlPointToLandscape(ULandscapeSplineControlPoint* InControlPoint, ALandscapeProxy* FromProxy, ALandscapeProxy* ToProxy);
 #endif
diff --git a/Engine/Source/Runtime/Landscape/Classes/LandscapeSplinesComponent.h b/Engine/Source/Runtime/Landscape/Classes/LandscapeSplinesComponent.h
--- a/Engine/Source/Runtime/Landscape/Classes/LandscapeSplinesComponent.h
+++ b/Engine/Source/Runtime/Landscape/Classes/LandscapeSplinesComponent.h
@@ -229,8 +229,4 @@ public:
 	friend class ULandscapeSplineControlPoint;
 	friend class ULandscapeSplineSegment;
 	friend class ULandscapeInfo;
-#if WITH_EDITOR
-	// TODO - move this out of ULandscapeInfo
-	friend bool ULandscapeInfo::ApplySplinesInternal(bool bOnlySelected, ALandscapeProxy* Landscape, TSet<ULandscapeComponent*>* OutModifiedComponents, bool bMarkPackageDirty);
-#endif
 };
diff --git a/Engine/Source/Runtime/Landscape/Private/LandscapeSplineRaster.cpp b/Engine/Source/Runtime/Landscape/Private/LandscapeSplineRaster.cpp
--- a/Engine/Source/Runtime/Landscape/Private/LandscapeSplineRaster.cpp
+++ b/Engine/Source/Runtime/Landscape/Private/LandscapeSplineRaster.cpp
@@ -557,15 +557,35 @@ bool ULandscapeInfo::ApplySplines(bool bOnlySelected, TSet<ULandscapeComponent*>
 	FGuid SplinesTargetLayerGuid = Layer ? Layer->Guid : Landscape ? Landscape->GetEditingLayer() : FGuid();
 	FScopedSetLandscapeEditingLayer Scope(Landscape, SplinesTargetLayerGuid, [=] { Landscape->RequestLayersContentUpdate(ELandscapeLayerUpdateMode::Update_All); });
 
+	TMap<ULandscapeLayerInfoObject*, TSharedPtr<FModulateAlpha>> ModulatePerLayerInfo;
+	int32 LandscapeMinX, LandscapeMinY, LandscapeMaxX, LandscapeMaxY;
+	if (!GetLandscapeExtent(LandscapeMinX, LandscapeMinY, LandscapeMaxX, LandscapeMaxY))
+	{
+		return false;
+	}
+
+	auto GetOrCreateModulate = [&](ULandscapeLayerInfoObject* LayerInfo) -> TSharedPtr<FModulateAlpha>
+	{
+		if (const TSharedPtr<FModulateAlpha>* SharedPtr = ModulatePerLayerInfo.Find(LayerInfo))
+		{
+			return *SharedPtr;
+		}
+
+		TSharedPtr<FModulateAlpha> SharedPtr = FModulateAlpha::CreateFromLayerInfo(LayerInfo, LandscapeMinX, LandscapeMinY);
+		ModulatePerLayerInfo.Add(LayerInfo, SharedPtr);
+
+		return SharedPtr;
+	};
+
 	ForAllLandscapeProxies([&](ALandscapeProxy* Proxy)
 	{
-		bResult |= ApplySplinesInternal(bOnlySelected, Proxy, OutModifiedComponents, bMarkPackageDirty);
+		bResult |= ApplySplinesInternal(bOnlySelected, Proxy, OutModifiedComponents, bMarkPackageDirty, LandscapeMinX, LandscapeMinY, LandscapeMaxX, LandscapeMaxY, GetOrCreateModulate);
 	});
 
 	return bResult;
 }
 
-bool ULandscapeInfo::ApplySplinesInternal(bool bOnlySelected, ALandscapeProxy* Proxy, TSet<ULandscapeComponent*>* OutModifiedComponents, bool bMarkPackageDirty)
+bool ULandscapeInfo::ApplySplinesInternal(bool bOnlySelected, ALandscapeProxy* Proxy, TSet<ULandscapeComponent*>* OutModifiedComponents, bool bMarkPackageDirty, int32 LandscapeMinX, int32 LandscapeMinY, int32 LandscapeMaxX, int32 LandscapeMaxY, TFunctionRef<TSharedPtr<FModulateAlpha>(ULandscapeLayerInfoObject*)> GetOrCreateModulate)
 {
 	if (!Proxy || !Proxy->SplineComponent || !Proxy->SplineComponent->IsRegistered() || Proxy->SplineComponent->ControlPoints.Num() == 0 || Proxy->SplineComponent->Segments.Num() == 0)
 	{
@@ -577,28 +597,7 @@ bool ULandscapeInfo::ApplySplinesInternal(bool bOnlySelected, ALandscapeProxy* P
 	FLandscapeEditDataInterface LandscapeEdit(this);
 	FLandscapeDoNotDirtyScope DoNotDirtyScope(LandscapeEdit, !bMarkPackageDirty);
 	TSet<ULandscapeComponent*> ModifiedComponents;
-
-	// I'd dearly love to use FIntRect in this code, but Landscape works with "Inclusive Max" and FIntRect is "Exclusive Max"
-	int32 LandscapeMinX, LandscapeMinY, LandscapeMaxX, LandscapeMaxY;
-	if (!GetLandscapeExtent(LandscapeMinX, LandscapeMinY, LandscapeMaxX, LandscapeMaxY))
-	{
-		return false;
-	}
-
-	TMap<ULandscapeLayerInfoObject*, TSharedPtr<FModulateAlpha>> ModulatePerLayerInfo;
 	
-	auto GetOrCreateModulate = [&](ULandscapeLayerInfoObject* LayerInfo) -> TSharedPtr<FModulateAlpha>
-	{
-		if (const TSharedPtr<FModulateAlpha>* SharedPtr = ModulatePerLayerInfo.Find(LayerInfo))
-		{
-			return *SharedPtr;
-		}
-		
-		TSharedPtr<FModulateAlpha> SharedPtr = FModulateAlpha::CreateFromLayerInfo(LayerInfo, LandscapeMinX, LandscapeMinY);
-		ModulatePerLayerInfo.Add(LayerInfo, SharedPtr);
-
-		return SharedPtr;
-	};
 
 	for (const ULandscapeSplineControlPoint* ControlPoint : Proxy->SplineComponent->ControlPoints)
 	{
