Landscape Layers: Rework Dirty Debug (tool to investigate map dirtyness on load)

#rnx

diff --git a/Engine/Source/Runtime/Core/Private/Misc/FileHelper.cpp b/Engine/Source/Runtime/Core/Private/Misc/FileHelper.cpp
--- a/Engine/Source/Runtime/Core/Private/Misc/FileHelper.cpp
+++ b/Engine/Source/Runtime/Core/Private/Misc/FileHelper.cpp
@@ -430,10 +430,11 @@ void FFileHelper::GenerateDateTimeBasedBitmapFilename(const FString& Pattern, co
  * @param SubRectangle optional, specifies a sub-rectangle of the source image to save out. If NULL, the whole bitmap is saved
  * @param FileManager must not be 0
  * @param OutFilename optional, if specified filename will be output
+ * @param ChannelMask optional, specifies a specific channel to write out (will be written out to all channels gray scale).
  *
  * @return true if success
  */
-bool FFileHelper::CreateBitmap( const TCHAR* Pattern, int32 SourceWidth, int32 SourceHeight, const FColor* Data, struct FIntRect* SubRectangle, IFileManager* FileManager /*= &IFileManager::Get()*/, FString* OutFilename /*= NULL*/, bool bInWriteAlpha /*= false*/ )
+bool FFileHelper::CreateBitmap( const TCHAR* Pattern, int32 SourceWidth, int32 SourceHeight, const FColor* Data, struct FIntRect* SubRectangle, IFileManager* FileManager /*= &IFileManager::Get()*/, FString* OutFilename /*= NULL*/, bool bInWriteAlpha /*= false*/, EChannelMask ChannelMask /*= All */ )
 {
 #if ALLOW_DEBUG_FILES
 	FIntRect Src(0, 0, SourceWidth, SourceHeight);
@@ -565,14 +566,54 @@ bool FFileHelper::CreateBitmap( const TCHAR* Pattern, int32 SourceWidth, int32 S
 		{
 			for( int32 j = SubRectangle->Min.X; j < SubRectangle->Max.X; j++ )
 			{
-				Ar->Serialize( (void *)&Data[i*SourceWidth+j].B, 1 );
-				Ar->Serialize( (void *)&Data[i*SourceWidth+j].G, 1 );
-				Ar->Serialize( (void *)&Data[i*SourceWidth+j].R, 1 );
+				if (ChannelMask == EChannelMask::All)
+				{
+					Ar->Serialize((void*)&Data[i * SourceWidth + j].B, 1);
+					Ar->Serialize((void*)&Data[i * SourceWidth + j].G, 1);
+					Ar->Serialize((void*)&Data[i * SourceWidth + j].R, 1);
 
-				if (bInWriteAlpha)
+					if (bInWriteAlpha)
+					{
+						Ar->Serialize((void*)&Data[i * SourceWidth + j].A, 1);
+					}
+				}
+				else
 				{
-					Ar->Serialize( (void *)&Data[i * SourceWidth + j].A, 1 );
+					const uint8 Max = 255;
+					uint8 ChannelValue = 0;
+					// When using Channel mask write the masked channel to blue channel.
+					switch (ChannelMask)
+					{
+					case EChannelMask::B:
+						ChannelValue = Data[i * SourceWidth + j].B;
+						break;
+					case EChannelMask::G:
+						ChannelValue = Data[i * SourceWidth + j].G;
+						break;
+					case EChannelMask::R:
+						ChannelValue = Data[i * SourceWidth + j].R;
+						break;
+					case EChannelMask::A:
+						ChannelValue = Data[i * SourceWidth + j].A;
+						break;
+					default:
+						check(0);
+						break;
+					}
+										
+					// replicate Channel in B, G, R
+					Ar->Serialize((void*)&ChannelValue, 1);
+					Ar->Serialize((void*)&ChannelValue, 1);
+					Ar->Serialize((void*)&ChannelValue, 1);
+
+					// if write alpha write max value in there (we don't want transparency)
+					if (bInWriteAlpha)
+					{
+						Ar->Serialize((void*)&Max, 1);
+					}
 				}
+
+				
 			}
 
 			// Pad each row's length to be a multiple of 4 bytes.
diff --git a/Engine/Source/Runtime/Core/Public/Misc/FileHelper.h b/Engine/Source/Runtime/Core/Public/Misc/FileHelper.h
--- a/Engine/Source/Runtime/Core/Public/Misc/FileHelper.h
+++ b/Engine/Source/Runtime/Core/Public/Misc/FileHelper.h
@@ -31,6 +31,15 @@ struct CORE_API FFileHelper
 		ForceUTF8WithoutBOM
 	};
 
+	enum class EChannelMask
+	{
+		R = STRUCT_OFFSET(FColor, R),
+		G = STRUCT_OFFSET(FColor, G),
+		B = STRUCT_OFFSET(FColor, B),
+		A = STRUCT_OFFSET(FColor, A),
+		All = R|G|B|A
+	};
+
 	/**
 	 * Load a text file to an FString.
 	 * Supports all combination of ANSI/Unicode files and platforms.
@@ -121,10 +130,11 @@ struct CORE_API FFileHelper
 	 * @param FileManager must not be 0
 	 * @param OutFilename optional, if specified filename will be output
 	 * @param bInWriteAlpha optional, specifies whether to write out the alpha channel. Will force BMP V4 format.
+	 * @param ChannelMask optional, specifies a specific channel to write out (will be written out to all channels gray scale).
 	 *
 	 * @return true if success
 	 */
-	static bool CreateBitmap( const TCHAR* Pattern, int32 DataWidth, int32 DataHeight, const struct FColor* Data, struct FIntRect* SubRectangle = NULL, IFileManager* FileManager = &IFileManager::Get(), FString* OutFilename = NULL, bool bInWriteAlpha = false );
+	static bool CreateBitmap( const TCHAR* Pattern, int32 DataWidth, int32 DataHeight, const struct FColor* Data, struct FIntRect* SubRectangle = NULL, IFileManager* FileManager = &IFileManager::Get(), FString* OutFilename = NULL, bool bInWriteAlpha = false, EChannelMask ChannelMask = EChannelMask::All );
 
 	/**
 	 * Generates the next unique bitmap filename with a specified extension
diff --git a/Engine/Source/Runtime/Landscape/Classes/Landscape.h b/Engine/Source/Runtime/Landscape/Classes/Landscape.h
--- a/Engine/Source/Runtime/Landscape/Classes/Landscape.h
+++ b/Engine/Source/Runtime/Landscape/Classes/Landscape.h
@@ -348,7 +348,9 @@ private:
 	bool UpdateCollisionAndClients(const TArray<ULandscapeComponent*>& InLandscapeComponents, const int32 InContentUpdateModes);
 	void ResolveLayersHeightmapTexture(const TArray<ULandscapeComponent*>& InLandscapeComponents);
 	void ResolveLayersWeightmapTexture(const TArray<ULandscapeComponent*>& InLandscapeComponents);
-	bool ResolveLayersTexture(class FLandscapeLayersTexture2DCPUReadBackResource* InCPUReadBackTexture, UTexture2D* InOutputTexture, bool bHeightmap, ULandscapeComponent* InComponent);
+
+	using FDirtyDelegate = TFunctionRef<void(UTexture2D*, FColor*, FColor*)>;
+	bool ResolveLayersTexture(class FLandscapeLayersTexture2DCPUReadBackResource* InCPUReadBackTexture, UTexture2D* InOutputTexture, FDirtyDelegate DirtyDelegate);
 		
 	bool AreLayersTextureResourcesReady(bool bInWaitForStreaming) const;
 	bool PrepareLayersBrushTextureResources(bool bInWaitForStreaming, bool bHeightmap) const;
@@ -406,7 +408,8 @@ private:
 	void PrintLayersDebugHeightData(const FString& InContext, const TArray<FColor>& InHeightmapData, const FIntPoint& InDataSize, uint8 InMipRender, bool InOutputNormals = false) const;
 	void PrintLayersDebugWeightData(const FString& InContext, const TArray<FColor>& InWeightmapData, const FIntPoint& InDataSize, uint8 InMipRender) const;
 
-	void UpdateDirtyData(ULandscapeComponent* InLandscapeComponent, const FColor* InOldData, const FColor* InNewData, int32 InSize);
+	void UpdateWeightDirtyData(ULandscapeComponent* InLandscapeComponent, UTexture2D* Heightmap, FColor* InOldData, const FColor* InNewData, uint8 Channel);
+	void UpdateHeightDirtyData(ULandscapeComponent* InLandscapeComponent, UTexture2D* Heightmap, FColor* InOldData, const FColor* InNewData);
 #endif
 
 public:
diff --git a/Engine/Source/Runtime/Landscape/Classes/LandscapeProxy.h b/Engine/Source/Runtime/Landscape/Classes/LandscapeProxy.h
--- a/Engine/Source/Runtime/Landscape/Classes/LandscapeProxy.h
+++ b/Engine/Source/Runtime/Landscape/Classes/LandscapeProxy.h
@@ -971,6 +971,10 @@ public:
 	/** Creates a Texture2D for use by this landscape proxy or one of it's components. If OptionalOverrideOuter is not specified, the proxy is used. */
 	LANDSCAPE_API UTexture2D* CreateLandscapeTexture(int32 InSizeX, int32 InSizeY, TextureGroup InLODGroup, ETextureSourceFormat InFormat, UObject* OptionalOverrideOuter = nullptr, bool bCompress = false) const;
 
+	/** Creates a Texture2D for use by this landscape proxy or one of it's components for tools .*/ 
+	LANDSCAPE_API UTexture2D* CreateLandscapeToolTexture(int32 InSizeX, int32 InSizeY, TextureGroup InLODGroup, ETextureSourceFormat InFormat) const;
+
+
 	/** Creates a LandscapeWeightMapUsage object outered to this proxy. */
 	LANDSCAPE_API ULandscapeWeightmapUsage* CreateWeightmapUsage();
 
diff --git a/Engine/Source/Runtime/Landscape/Private/LandscapeEdit.cpp b/Engine/Source/Runtime/Landscape/Private/LandscapeEdit.cpp
--- a/Engine/Source/Runtime/Landscape/Private/LandscapeEdit.cpp
+++ b/Engine/Source/Runtime/Landscape/Private/LandscapeEdit.cpp
@@ -6481,6 +6481,21 @@ UTexture2D* ALandscapeProxy::CreateLandscapeTexture(int32 InSizeX, int32 InSizeY
 	return NewTexture;
 }
 
+UTexture2D* ALandscapeProxy::CreateLandscapeToolTexture(int32 InSizeX, int32 InSizeY, TextureGroup InLODGroup, ETextureSourceFormat InFormat) const
+{
+	UObject* TexOuter = const_cast<ALandscapeProxy*>(this);
+	UTexture2D* NewTexture = NewObject<UTexture2D>(TexOuter);
+	NewTexture->Source.Init(InSizeX, InSizeY, 1, 1, InFormat);
+	NewTexture->SRGB = false;
+	NewTexture->CompressionNone = true;
+	NewTexture->MipGenSettings = TMGS_NoMipmaps;
+	NewTexture->AddressX = TA_Clamp;
+	NewTexture->AddressY = TA_Clamp;
+	NewTexture->LODGroup = InLODGroup;
+
+	return NewTexture;
+}
+
 ULandscapeWeightmapUsage* ALandscapeProxy::CreateWeightmapUsage()
 {
 	return NewObject<ULandscapeWeightmapUsage>(this, ULandscapeWeightmapUsage::StaticClass(), NAME_None, RF_Transactional);
diff --git a/Engine/Source/Runtime/Landscape/Private/LandscapeEditInterface.cpp b/Engine/Source/Runtime/Landscape/Private/LandscapeEditInterface.cpp
--- a/Engine/Source/Runtime/Landscape/Private/LandscapeEditInterface.cpp
+++ b/Engine/Source/Runtime/Landscape/Private/LandscapeEditInterface.cpp
@@ -4628,7 +4628,7 @@ void FLandscapeEditDataInterface::SetDirtyData(int32 X1, int32 Y1, int32 X2, int
 		check(Component);
 		return Component->EditToolRenderData.DirtyTexture;
 	};
-	SetEditToolTextureData(X1, Y1, X2, Y2, Data, Stride, ReturnComponentTexture);
+	SetEditToolTextureData(X1, Y1, X2, Y2, Data, Stride, ReturnComponentTexture, TEXTUREGROUP_8BitData);
 }
 
 template<typename TStoreData>
@@ -4714,7 +4714,7 @@ void FLandscapeEditDataInterface::GetEditToolTextureData(const int32 X1, const i
 	}
 }
 
-void FLandscapeEditDataInterface::SetEditToolTextureData(int32 X1, int32 Y1, int32 X2, int32 Y2, const uint8* Data, int32 Stride, TFunctionRef<UTexture2D*&(ULandscapeComponent*)> GetComponentTexture)
+void FLandscapeEditDataInterface::SetEditToolTextureData(int32 X1, int32 Y1, int32 X2, int32 Y2, const uint8* Data, int32 Stride, TFunctionRef<UTexture2D*&(ULandscapeComponent*)> GetComponentTexture, TextureGroup InTextureGroup)
 {
 	if( Stride==0 )
 	{
@@ -4741,26 +4741,12 @@ void FLandscapeEditDataInterface::SetEditToolTextureData(int32 X1, int32 Y1, int
 			// if NULL, it was painted away
 			if (EditToolTexture == NULL)
 			{
-				//FlushRenderingCommands();
 				// Construct Texture...
 				int32 WeightmapSize = (Component->SubsectionSizeQuads+1) * Component->NumSubsections;
-				EditToolTexture = Component->GetLandscapeProxy()->CreateLandscapeTexture(WeightmapSize, WeightmapSize, TEXTUREGROUP_Terrain_Weightmap, TSF_G8);
-				// Alloc dummy mips
-				ULandscapeComponent::CreateEmptyTextureMips(EditToolTexture, true);
+				EditToolTexture = Component->GetLandscapeProxy()->CreateLandscapeToolTexture(WeightmapSize, WeightmapSize, InTextureGroup, TSF_G8);
 				EditToolTexture->PostEditChange();
 
-				//FlushRenderingCommands();
 				ZeroTexture(EditToolTexture);
-				FLandscapeTextureDataInfo* TexDataInfo = GetTextureDataInfo(EditToolTexture);
-				int32 NumMips = EditToolTexture->Source.GetNumMips();
-				TArray<uint8*> TextureMipData;
-				TextureMipData.AddUninitialized(NumMips);
-				for( int32 MipIdx=0;MipIdx<NumMips;MipIdx++ )
-				{
-					TextureMipData[MipIdx] = (uint8*)TexDataInfo->GetMipData(MipIdx);
-				}
-				ULandscapeComponent::UpdateDataMips(ComponentNumSubsections, SubsectionSizeQuads, EditToolTexture, TextureMipData, 0, 0, MAX_int32, MAX_int32, TexDataInfo);
-
 				Component->UpdateEditToolRenderData();
 			}
 
@@ -4825,15 +4811,6 @@ void FLandscapeEditDataInterface::SetEditToolTextureData(int32 X1, int32 Y1, int
 					TexDataInfo->AddMipUpdateRegion(0,TexX1,TexY1,TexX2,TexY2);
 				}
 			}
-			// Update mipmaps
-			int32 NumMips = EditToolTexture->Source.GetNumMips();
-			TArray<uint8*> TextureMipData;
-			TextureMipData.AddUninitialized(NumMips);
-			for( int32 MipIdx=0;MipIdx<NumMips;MipIdx++ )
-			{
-				TextureMipData[MipIdx] = (uint8*)TexDataInfo->GetMipData(MipIdx);
-			}
-			ULandscapeComponent::UpdateDataMips(ComponentNumSubsections, SubsectionSizeQuads, EditToolTexture, TextureMipData, ComponentX1, ComponentY1, ComponentX2, ComponentY2, TexDataInfo);
 		}
 	}
 }
diff --git a/Engine/Source/Runtime/Landscape/Private/LandscapeEditLayers.cpp b/Engine/Source/Runtime/Landscape/Private/LandscapeEditLayers.cpp
--- a/Engine/Source/Runtime/Landscape/Private/LandscapeEditLayers.cpp
+++ b/Engine/Source/Runtime/Landscape/Private/LandscapeEditLayers.cpp
@@ -93,6 +93,11 @@ TAutoConsoleVariable<int32> CVarLandscapeShowDirty(
 	0,
 	TEXT("This will highlight the data that as changed during the layer blend phase."));
 
+TAutoConsoleVariable<int32> CVarLandscapeTrackDirty(
+	TEXT("landscape.TrackDirty"),
+	0,
+	TEXT("This will track the accumulation of data changes during the layer blend phase."));
+
 struct FLandscapeDirty
 {
 	FLandscapeDirty()
@@ -3087,37 +3092,66 @@ void ALandscape::ResolveLayersHeightmapTexture(const TArray<ULandscapeComponent*
 		return;
 	}
 
-	TMap<UTexture*, bool> ProcessedTexture;
 	TArray<ULandscapeComponent*> ChangedComponents;
-
+	TMap<UTexture2D*, TArray<ULandscapeComponent*>> HeightmapsToResolve;
 	for (ULandscapeComponent* Component : InLandscapeComponents)
 	{
-		UTexture2D* ComponentHeightmap = Component->GetHeightmap(false);
-		bool* bValue = ProcessedTexture.Find(ComponentHeightmap);
-		bool bChanged = false;
-		if (bValue == nullptr)
+		HeightmapsToResolve.FindOrAdd(Component->GetHeightmap(false)).Add(Component);
+	}
+	
+	// Dirty Delegate
+	const bool HeightmapDiff = (CVarLandscapeOutputDiffBitmap.GetValueOnAnyThread() & 1) != 0;
+	auto DirtyDelegate = [&](UTexture2D* Heightmap, FColor* OldData, FColor* NewData)
+	{
+		if (!HeightmapDiff && !CVarLandscapeTrackDirty.GetValueOnAnyThread())
 		{
-			FLandscapeLayersTexture2DCPUReadBackResource** CPUReadback = Component->GetLandscapeProxy()->HeightmapsCPUReadBack.Find(ComponentHeightmap);
+			return;
+		}
 
-			if (CPUReadback == nullptr)
+		if (TArray<ULandscapeComponent*>* Components = HeightmapsToResolve.Find(Heightmap))
+		{
+			for (ULandscapeComponent* Component : *Components)
 			{
-				continue;
+				if (CVarLandscapeTrackDirty.GetValueOnAnyThread())
+				{
+					UpdateHeightDirtyData(Component, Heightmap, OldData, NewData);
+				}
+
+				if (HeightmapDiff)
+				{
+					FString LevelName = FPackageName::GetShortName(Component->GetOutermost());
+					FString FilePattern = FString::Format(TEXT("LandscapeLayers/{0}-{1}-HM"), { LevelName, Component->GetName() });
+
+					const int32 SizeU = Heightmap->Source.GetSizeX();
+					const int32 SizeV = Heightmap->Source.GetSizeY();
+					const int32 HeightmapOffsetX = Component->HeightmapScaleBias.Z * (float)SizeU;
+					const int32 HeightmapOffsetY = Component->HeightmapScaleBias.W * (float)SizeV;
+					const int32 ComponentWidth = (SubsectionSizeQuads + 1) * NumSubsections;
+					FIntRect SubRegion(HeightmapOffsetX, HeightmapOffsetY, HeightmapOffsetX + ComponentWidth, HeightmapOffsetY + ComponentWidth);
+					
+					FFileHelper::CreateBitmap(*(FilePattern + "-Pre.bmp"), Heightmap->Source.GetSizeX(), Heightmap->Source.GetSizeY(), OldData, &SubRegion, &IFileManager::Get(), nullptr, true);
+					FFileHelper::CreateBitmap(*(FilePattern + "-Post.bmp"), Heightmap->Source.GetSizeX(), Heightmap->Source.GetSizeY(), NewData, &SubRegion, &IFileManager::Get(), nullptr, true);
+				}
 			}
-						
-			bChanged = ResolveLayersTexture(*CPUReadback, ComponentHeightmap, true, Component);
-			ProcessedTexture.Add(ComponentHeightmap, bChanged);
-			bValue = &bChanged;
 		}
-		
-		if (*bValue)
+	};
+
+	for (const auto& HeightmapToResolve : HeightmapsToResolve)
+	{
+		UTexture2D* Heightmap = HeightmapToResolve.Key;
+		ALandscapeProxy* LandscapeProxy = Heightmap->GetTypedOuter<ALandscapeProxy>();
+		check(LandscapeProxy);
+		if (FLandscapeLayersTexture2DCPUReadBackResource** CPUReadback = LandscapeProxy->HeightmapsCPUReadBack.Find(Heightmap))
 		{
-			Component->MarkPackageDirty();
-			ChangedComponents.Add(Component);
+			if (ResolveLayersTexture(*CPUReadback, Heightmap, DirtyDelegate))
+			{
+				ChangedComponents.Append(HeightmapToResolve.Value);
+				Heightmap->MarkPackageDirty();
+			}
 		}
 	}
-
+		
 	InvalidateGeneratedComponentData(ChangedComponents);
-
 }
 
 void ALandscape::ClearDirtyData(ULandscapeComponent* InLandscapeComponent)
@@ -3127,16 +3161,59 @@ void ALandscape::ClearDirtyData(ULandscapeComponent* InLandscapeComponent)
 		return;
 	}
 
-	UpdateDirtyData(InLandscapeComponent, nullptr, nullptr, 0);
+	if (!CVarLandscapeTrackDirty.GetValueOnAnyThread())
+	{
+		return;
+	}
+
+	FLandscapeEditDataInterface LandscapeEdit(GetLandscapeInfo());
+	const int32 X1 = InLandscapeComponent->GetSectionBase().X;
+	const int32 X2 = X1 + ComponentSizeQuads;
+	const int32 Y1 = InLandscapeComponent->GetSectionBase().Y;
+	const int32 Y2 = Y1 + ComponentSizeQuads;
+	const int32 ComponentWidth = (SubsectionSizeQuads + 1) * NumSubsections;
+	const int32 DirtyDataSize = ComponentWidth * ComponentWidth;
+	TUniquePtr<uint8[]> DirtyData = MakeUnique<uint8[]>(DirtyDataSize);
+	FMemory::Memzero(DirtyData.Get(), DirtyDataSize);
+	LandscapeEdit.SetDirtyData(X1, Y1, X2, Y2, DirtyData.Get(), 0);
 }
 
-void ALandscape::UpdateDirtyData(ULandscapeComponent* InLandscapeComponent, const FColor* InOldData, const FColor* InNewData, int32 InSize)
+void ALandscape::UpdateWeightDirtyData(ULandscapeComponent* InLandscapeComponent, UTexture2D* Heightmap, FColor* InOldData, const FColor* InNewData, uint8 Channel)
 {
-	if (!CVarLandscapeShowDirty.GetValueOnAnyThread())
+	check(InOldData && InNewData);
+
+	FLandscapeEditDataInterface LandscapeEdit(GetLandscapeInfo());
+	const int32 X1 = InLandscapeComponent->GetSectionBase().X;
+	const int32 X2 = X1 + ComponentSizeQuads;
+	const int32 Y1 = InLandscapeComponent->GetSectionBase().Y;
+	const int32 Y2 = Y1 + ComponentSizeQuads;
+	const int32 ComponentWidth = (SubsectionSizeQuads + 1) * NumSubsections;
+	const int32 DirtyDataSize = ComponentWidth * ComponentWidth;
+	TUniquePtr<uint8[]> DirtyData = MakeUnique<uint8[]>(DirtyDataSize);
+	const int32 SizeU = Heightmap->Source.GetSizeX();
+	const int32 SizeV = Heightmap->Source.GetSizeY();
+	check(DirtyDataSize == SizeU * SizeV);
+
+	const uint8 DirtyWeight = 1<<1;
+	LandscapeEdit.GetDirtyData(X1, Y1, X2, Y2, DirtyData.Get(), 0);
+
+	for (int32 Index = 0; Index < DirtyDataSize; ++Index)
 	{
-		return;
+		uint8* OldChannelValue = (uint8*)&InOldData[Index] + ChannelOffsets[Channel];
+		uint8* NewChannelValue = (uint8*)&InNewData[Index] + ChannelOffsets[Channel];
+		if (*OldChannelValue != *NewChannelValue)
+		{
+			DirtyData[Index] |= DirtyWeight;
+		}
 	}
 
+	LandscapeEdit.SetDirtyData(X1, Y1, X2, Y2, DirtyData.Get(), 0);
+}
+
+void ALandscape::UpdateHeightDirtyData(ULandscapeComponent* InLandscapeComponent, UTexture2D* Heightmap, FColor* InOldData, const FColor* InNewData)
+{
+	check(InOldData && InNewData);
+
 	FLandscapeEditDataInterface LandscapeEdit(GetLandscapeInfo());
 	const int32 X1 = InLandscapeComponent->GetSectionBase().X;
 	const int32 X2 = X1 + ComponentSizeQuads;
@@ -3144,34 +3221,33 @@ void ALandscape::UpdateDirtyData(ULandscapeComponent* InLandscapeComponent, cons
 	const int32 Y2 = Y1 + ComponentSizeQuads;
 	const int32 ComponentWidth = (SubsectionSizeQuads + 1)*NumSubsections;
 	const int32 DirtyDataSize = ComponentWidth * ComponentWidth;
-	uint8* DirtyData = new uint8[DirtyDataSize];
+	TUniquePtr<uint8[]> DirtyData = MakeUnique<uint8[]>(DirtyDataSize);
+	const int32 SizeU = Heightmap->Source.GetSizeX();
+	const int32 SizeV = Heightmap->Source.GetSizeY();
+	const int32 HeightmapOffsetX = InLandscapeComponent->HeightmapScaleBias.Z * (float)SizeU;
+	const int32 HeightmapOffsetY = InLandscapeComponent->HeightmapScaleBias.W * (float)SizeV;
+	const uint8 DirtyHeight = 1<<0;
+	LandscapeEdit.GetDirtyData(X1, Y1, X2, Y2, DirtyData.Get(), 0);
 	
-	if (InOldData && InNewData)
+	for (int32 X = 0; X < ComponentWidth; ++X)
 	{
-		check(InSize == DirtyDataSize);
-		LandscapeEdit.GetDirtyData(X1, Y1, X2, Y2, DirtyData, 0);
-
-		for (int32 i = 0; i < DirtyDataSize; ++i)
+		for (int32 Y = 0; Y < ComponentWidth; ++Y)
 		{
-			if (*InOldData != *InNewData)
+			int32 TexX = HeightmapOffsetX + X;
+			int32 TexY = HeightmapOffsetY + Y;
+			int32 TexIndex = TexX + TexY * SizeU;
+			check(TexIndex < SizeU * SizeV);
+			if (InOldData[TexIndex] != InNewData[TexIndex])
 			{
-				DirtyData[i] = 255;
+				DirtyData[X + Y * ComponentWidth] |= DirtyHeight;
 			}
-			InOldData++;
-			InNewData++;
 		}
 	}
-	else
-	{
-		FMemory::Memzero(DirtyData, DirtyDataSize);
-	}
 
-	LandscapeEdit.SetDirtyData(X1, Y1, X2, Y2, DirtyData, 0);
-
-	delete[] DirtyData;
+	LandscapeEdit.SetDirtyData(X1, Y1, X2, Y2, DirtyData.Get(), 0);
 }
 
-bool ALandscape::ResolveLayersTexture(FLandscapeLayersTexture2DCPUReadBackResource* InCPUReadBackTexture, UTexture2D* InOutputTexture, bool bHeightmap, ULandscapeComponent* InComponent)
+bool ALandscape::ResolveLayersTexture(FLandscapeLayersTexture2DCPUReadBackResource* InCPUReadBackTexture, UTexture2D* InOutputTexture, FDirtyDelegate DirtyDelegate)
 {
 	TRACE_CPUPROFILER_EVENT_SCOPE("ResolveLayersTexture");
 	SCOPE_CYCLE_COUNTER(STAT_LandscapeLayersResolveTexture);
@@ -3207,9 +3283,7 @@ bool ALandscape::ResolveLayersTexture(FLandscapeLayersTexture2DCPUReadBackResour
 	FlushRenderingCommands();
 
 	bool bChanged = false;
-	const bool HeightmapDiff = (CVarLandscapeOutputDiffBitmap.GetValueOnAnyThread() & 1) != 0;
-	const bool WeightmapDiff = (CVarLandscapeOutputDiffBitmap.GetValueOnAnyThread() & 2) != 0;
-    const bool bUpdateHash = !bIntermediateRender;
+	const bool bUpdateHash = !bIntermediateRender;
 	for (int8 MipIndex = 0; MipIndex < OutMipsData.Num(); ++MipIndex)
 	{
 		if (OutMipsData[MipIndex].Num() > 0)
@@ -3222,15 +3296,7 @@ bool ALandscape::ResolveLayersTexture(FLandscapeLayersTexture2DCPUReadBackResour
 
 			if (MipIndex == 0 && bChanged)
 			{
-				UpdateDirtyData(InComponent, (FColor*)TextureData, OutMipsData[MipIndex].GetData(), InOutputTexture->Source.GetSizeX()*InOutputTexture->Source.GetSizeY());
-
-				if ((HeightmapDiff && bHeightmap) || (WeightmapDiff && !bHeightmap))
-				{
-					FString LevelName = FPackageName::GetShortName(InComponent->GetOutermost());
-					FString FilePattern = FString::Format(TEXT("LandscapeLayers/{0}-{1}-{2}"), { LevelName, InComponent->GetName(), bHeightmap ? TEXT("HM") : TEXT("WM") });
-					FFileHelper::CreateBitmap(*(FilePattern), InOutputTexture->Source.GetSizeX(), InOutputTexture->Source.GetSizeY(), (FColor*)TextureData, nullptr, &IFileManager::Get(), nullptr, true);
-					FFileHelper::CreateBitmap(*(FilePattern), InOutputTexture->Source.GetSizeX(), InOutputTexture->Source.GetSizeY(), OutMipsData[MipIndex].GetData(), nullptr, &IFileManager::Get(), nullptr, true);
-				}
+				DirtyDelegate(InOutputTexture, (FColor*)TextureData, OutMipsData[MipIndex].GetData());
 			}
 
 			FMemory::Memcpy(TextureData, OutMipsData[MipIndex].GetData(), OutMipsData[MipIndex].Num() * sizeof(FColor));
@@ -4373,44 +4439,70 @@ void ALandscape::ResolveLayersWeightmapTexture(const TArray<ULandscapeComponent*
 		return;
 	}
 
-	TMap<UTexture2D*,bool> ProcessedWeightmaps;
-	TArray<ULandscapeComponent*> ChangedComponents;
+	TMap<UTexture2D*, TArray<TPair<ULandscapeComponent*,const FWeightmapLayerAllocationInfo*>>> WeightmapsToResolve;
+	TSet<ULandscapeComponent*> ChangedComponents;
 
 	for (ULandscapeComponent* Component : InLandscapeComponents)
 	{
 		TArray<UTexture2D*>& ComponentWeightmaps = Component->GetWeightmapTextures(false);
-		bool bChanged = false;
+		TArray<FWeightmapLayerAllocationInfo>& AllocInfos = Component->GetWeightmapLayerAllocations(false);
 
-		for (UTexture2D* ComponentLayerWeightmap : ComponentWeightmaps)
+		for (const FWeightmapLayerAllocationInfo& AllocInfo : AllocInfos)
 		{
-			bool* bValue = ProcessedWeightmaps.Find(ComponentLayerWeightmap);
-			if (bValue == nullptr)
-			{
-				FLandscapeLayersTexture2DCPUReadBackResource** CPUReadback = Component->GetLandscapeProxy()->WeightmapsCPUReadBack.Find(ComponentLayerWeightmap);
+			check(AllocInfo.WeightmapTextureIndex != 255 && AllocInfo.WeightmapTextureIndex < ComponentWeightmaps.Num());
+			UTexture2D* Weightmap = ComponentWeightmaps[AllocInfo.WeightmapTextureIndex];
+			WeightmapsToResolve.FindOrAdd(Weightmap).Add(TPair<ULandscapeComponent*,const FWeightmapLayerAllocationInfo*>(Component, &AllocInfo));
+		}
+	}
+
+	const bool WeightmapDiff = (CVarLandscapeOutputDiffBitmap.GetValueOnAnyThread() & 2) != 0;
+	auto DirtyDelegate = [&](UTexture2D* Weightmap, FColor* OldData, FColor* NewData)
+	{
+		if (!WeightmapDiff && !CVarLandscapeTrackDirty.GetValueOnAnyThread())
+		{
+			return;
+		}
 
-				if (CPUReadback == nullptr)
+		if (TArray<TPair<ULandscapeComponent*,const FWeightmapLayerAllocationInfo*>>* Components = WeightmapsToResolve.Find(Weightmap))
+		{
+			for (const auto& Pair : *Components)
+			{
+				if (CVarLandscapeTrackDirty.GetValueOnAnyThread())
 				{
-					continue;
+					UpdateWeightDirtyData(Pair.Key, Weightmap, OldData, NewData, Pair.Value->WeightmapTextureChannel);
+				}
+				
+				if (WeightmapDiff)
+				{
+					size_t ChannelOffset = ChannelOffsets[Pair.Value->WeightmapTextureChannel];
+					FString LevelName = FPackageName::GetShortName(Pair.Key->GetOutermost());
+					FString FilePattern = FString::Format(TEXT("LandscapeLayers/{0}-{1}-{2}-WM"), { LevelName, Pair.Key->GetName(), Pair.Value->GetLayerName().ToString() });
+					FFileHelper::CreateBitmap(*(FilePattern + "-Pre.bmp"), Weightmap->Source.GetSizeX(), Weightmap->Source.GetSizeY(), OldData, nullptr, &IFileManager::Get(), nullptr, true, (FFileHelper::EChannelMask)ChannelOffset);
+					FFileHelper::CreateBitmap(*(FilePattern + "-Post.bmp"), Weightmap->Source.GetSizeX(), Weightmap->Source.GetSizeY(), NewData, nullptr, &IFileManager::Get(), nullptr, true, (FFileHelper::EChannelMask)ChannelOffset);
 				}
-								
-				bool bWeightmapChanged = ResolveLayersTexture(*CPUReadback, ComponentLayerWeightmap, false, Component);
-				ProcessedWeightmaps.Add(ComponentLayerWeightmap, bWeightmapChanged);
-				bChanged |= bWeightmapChanged;
-			}
-			else
-			{
-				bChanged |= *bValue;
 			}
 		}
+	};
 
-		if (bChanged)
+	for (const auto& WeightmapToResolve : WeightmapsToResolve)
+	{
+		UTexture2D* Weightmap = WeightmapToResolve.Key;
+		ALandscapeProxy* LandscapeProxy = Weightmap->GetTypedOuter<ALandscapeProxy>();
+		check(LandscapeProxy);
+		if (FLandscapeLayersTexture2DCPUReadBackResource** CPUReadback = LandscapeProxy->WeightmapsCPUReadBack.Find(Weightmap))
 		{
-			Component->MarkPackageDirty();
-			ChangedComponents.Add(Component);
+			if (ResolveLayersTexture(*CPUReadback, Weightmap, DirtyDelegate))
+			{
+				for (const auto& Pair : WeightmapToResolve.Value)
+				{
+					ChangedComponents.Add(Pair.Key);
+				}
+				Weightmap->MarkPackageDirty();
+			}
 		}
 	}
-
-	InvalidateGeneratedComponentData(ChangedComponents);
+		
+	InvalidateGeneratedComponentData(ChangedComponents.Array());	
 }
 
 bool ALandscape::HasLayersContent() const
diff --git a/Engine/Source/Runtime/Landscape/Private/LandscapeRender.cpp b/Engine/Source/Runtime/Landscape/Private/LandscapeRender.cpp
--- a/Engine/Source/Runtime/Landscape/Private/LandscapeRender.cpp
+++ b/Engine/Source/Runtime/Landscape/Private/LandscapeRender.cpp
@@ -5174,12 +5174,31 @@ void ULandscapeComponent::GetStreamingRenderAssetInfo(FStreamingTextureLevelCont
 	}
 
 #if WITH_EDITOR
-	if (GIsEditor && EditToolRenderData.DataTexture)
+	if (GIsEditor)
 	{
-		FStreamingRenderAssetPrimitiveInfo& StreamingDatamap = *new(OutStreamingRenderAssets)FStreamingRenderAssetPrimitiveInfo;
-		StreamingDatamap.Bounds = BoundingSphere;
-		StreamingDatamap.TexelFactor = TexelFactor;
-		StreamingDatamap.RenderAsset = EditToolRenderData.DataTexture;
+		if (EditToolRenderData.DataTexture)
+		{
+			FStreamingRenderAssetPrimitiveInfo& StreamingDatamap = *new(OutStreamingRenderAssets)FStreamingRenderAssetPrimitiveInfo;
+			StreamingDatamap.Bounds = BoundingSphere;
+			StreamingDatamap.TexelFactor = TexelFactor;
+			StreamingDatamap.RenderAsset = EditToolRenderData.DataTexture;
+		}
+
+		if (EditToolRenderData.LayerContributionTexture)
+		{
+			FStreamingRenderAssetPrimitiveInfo& StreamingDatamap = *new(OutStreamingRenderAssets)FStreamingRenderAssetPrimitiveInfo;
+			StreamingDatamap.Bounds = BoundingSphere;
+			StreamingDatamap.TexelFactor = TexelFactor;
+			StreamingDatamap.RenderAsset = EditToolRenderData.LayerContributionTexture;
+		}
+
+		if (EditToolRenderData.DirtyTexture)
+		{
+			FStreamingRenderAssetPrimitiveInfo& StreamingDatamap = *new(OutStreamingRenderAssets)FStreamingRenderAssetPrimitiveInfo;
+			StreamingDatamap.Bounds = BoundingSphere;
+			StreamingDatamap.TexelFactor = TexelFactor;
+			StreamingDatamap.RenderAsset = EditToolRenderData.DirtyTexture;
+		}
 	}
 #endif
 }
diff --git a/Engine/Source/Runtime/Landscape/Public/LandscapeEdit.h b/Engine/Source/Runtime/Landscape/Public/LandscapeEdit.h
--- a/Engine/Source/Runtime/Landscape/Public/LandscapeEdit.h
+++ b/Engine/Source/Runtime/Landscape/Public/LandscapeEdit.h
@@ -217,7 +217,7 @@ struct LANDSCAPE_API FLandscapeEditDataInterface : public FLandscapeTextureDataI
 
 	template<typename TStoreData>
 	void GetEditToolTextureData(const int32 X1, const int32 Y1, const int32 X2, const int32 Y2, TStoreData& StoreData, TFunctionRef<UTexture2D*(ULandscapeComponent*)> GetComponentTexture);
-	void SetEditToolTextureData(int32 X1, int32 Y1, int32 X2, int32 Y2, const uint8* Data, int32 Stride, TFunctionRef<UTexture2D*&(ULandscapeComponent*)> GetComponentTexture);
+	void SetEditToolTextureData(int32 X1, int32 Y1, int32 X2, int32 Y2, const uint8* Data, int32 Stride, TFunctionRef<UTexture2D*&(ULandscapeComponent*)> GetComponentTexture, TextureGroup InTextureGroup = TEXTUREGROUP_Terrain_Weightmap);
 
 	// Without data interpolation, Select Data 
 	template<typename TStoreData>
