Fix Creation of Unregistered SplineComponents (causing LandscapeEdMode to suspend edition)

#rnx

diff --git a/Engine/Source/Runtime/Landscape/Private/LandscapeSplines.cpp b/Engine/Source/Runtime/Landscape/Private/LandscapeSplines.cpp
--- a/Engine/Source/Runtime/Landscape/Private/LandscapeSplines.cpp
+++ b/Engine/Source/Runtime/Landscape/Private/LandscapeSplines.cpp
@@ -988,11 +988,9 @@ ULandscapeSplinesComponent* ULandscapeSplinesComponent::GetStreamingSplinesCompo
 			ALandscapeProxy* ComponentLandscapeProxy = LandscapeComponent->GetLandscapeProxy();
 			if (!ComponentLandscapeProxy->SplineComponent && bCreate)
 			{
-				ComponentLandscapeProxy->Modify();
-				ComponentLandscapeProxy->SplineComponent = NewObject<ULandscapeSplinesComponent>(ComponentLandscapeProxy, NAME_None, RF_Transactional);
-				ComponentLandscapeProxy->SplineComponent->SetRelativeScale3D_Direct(GetRelativeScale3D());
-				ComponentLandscapeProxy->SplineComponent->AttachToComponent(ComponentLandscapeProxy->GetRootComponent(), FAttachmentTransformRules::KeepRelativeTransform);
+				ComponentLandscapeProxy->CreateSplineComponent(GetRelativeScale3D());
 			}
+
 			if (ComponentLandscapeProxy->SplineComponent)
 			{
 				return ComponentLandscapeProxy->SplineComponent;
@@ -1016,10 +1014,7 @@ ULandscapeSplinesComponent* ULandscapeSplinesComponent::GetStreamingSplinesCompo
 		{
 			if (!Proxy->SplineComponent && bCreate)
 			{
-				Proxy->Modify();
-				Proxy->SplineComponent = NewObject<ULandscapeSplinesComponent>(Proxy, NAME_None, RF_Transactional);
-				Proxy->SplineComponent->SetRelativeScale3D_Direct(GetRelativeScale3D());
-				Proxy->SplineComponent->AttachToComponent(Proxy->GetRootComponent(), FAttachmentTransformRules::KeepRelativeTransform);
+				Proxy->CreateSplineComponent(GetRelativeScale3D());
 			}
 			return Proxy->SplineComponent;
 		}
