Hide Chaos Warnings behing bChaosChecked to make cook clean.

diff --git a/Engine/Source/Runtime/Engine/Private/PhysicsEngine/BodySetup.cpp b/Engine/Source/Runtime/Engine/Private/PhysicsEngine/BodySetup.cpp
--- a/Engine/Source/Runtime/Engine/Private/PhysicsEngine/BodySetup.cpp
+++ b/Engine/Source/Runtime/Engine/Private/PhysicsEngine/BodySetup.cpp
@@ -726,7 +726,7 @@ void UBodySetup::FinishCreatingPhysicsMeshes_Chaos(FChaosDerivedDataReader<float
 					UE_LOG(LogPhysics, Warning, TEXT("InReader.ConvexImplicitObjects.Num() [%d], AggGeom.ConvexElems.Num() [%d]"),
 						InReader.ConvexImplicitObjects.Num(), AggGeom.ConvexElems.Num());
 				}
-				UE_LOG(LogPhysics, Warning, TEXT("TConvex Name:%s, Element [%d] has no Geometry"), FullName.GetCharArray().GetData(), ElementIndex);
+				CHAOS_LOG(LogPhysics, Warning, TEXT("TConvex Name:%s, Element [%d] has no Geometry"), FullName.GetCharArray().GetData(), ElementIndex);
 			}
 		}
 		InReader.ConvexImplicitObjects.Reset();
diff --git a/Engine/Source/Runtime/Engine/Private/PhysicsEngine/Experimental/ChaosDerivedDataUtil.cpp b/Engine/Source/Runtime/Engine/Private/PhysicsEngine/Experimental/ChaosDerivedDataUtil.cpp
--- a/Engine/Source/Runtime/Engine/Private/PhysicsEngine/Experimental/ChaosDerivedDataUtil.cpp
+++ b/Engine/Source/Runtime/Engine/Private/PhysicsEngine/Experimental/ChaosDerivedDataUtil.cpp
@@ -151,8 +151,8 @@ namespace Chaos
 			}
 		}
 
-		UE_CLOG(NumDiscardedTriangles_Welded > 0, LogChaos, Warning, TEXT("Discarded %d welded triangles when cooking trimesh."), NumDiscardedTriangles_Welded);
-		UE_CLOG(NumDiscardedTriangles_Area > 0, LogChaos, Warning, TEXT("Discarded %d small triangles when cooking trimesh."), NumDiscardedTriangles_Area);
+		CHAOS_CLOG(NumDiscardedTriangles_Welded > 0, LogChaos, Warning, TEXT("Discarded %d welded triangles when cooking trimesh."), NumDiscardedTriangles_Welded);
+		CHAOS_CLOG(NumDiscardedTriangles_Area > 0, LogChaos, Warning, TEXT("Discarded %d small triangles when cooking trimesh."), NumDiscardedTriangles_Area);
 
 		InOutVertices = LocalUniqueVerts;
 		InOutIndices = LocalUniqueIndices;
diff --git a/Engine/Source/Runtime/Experimental/Chaos/Public/ChaosCheck.h b/Engine/Source/Runtime/Experimental/Chaos/Public/ChaosCheck.h
--- a/Engine/Source/Runtime/Experimental/Chaos/Public/ChaosCheck.h
+++ b/Engine/Source/Runtime/Experimental/Chaos/Public/ChaosCheck.h
@@ -8,8 +8,12 @@
 #define CHAOS_CHECK(Condition) check(Condition)
 #define CHAOS_ENSURE(Condition) ensure(Condition)
 #define CHAOS_ENSURE_MSG(InExpression, InFormat, ... ) ensureMsgf(InExpression, InFormat, ##__VA_ARGS__)
+#define CHAOS_LOG(InLog, InSeverity, InFormat, ...) UE_LOG(InLog, InSeverity, InFormat, ##__VA_ARGS__)
+#define CHAOS_CLOG(InExpression, InLog, InSeverity, InFormat, ...) UE_CLOG(InExpression, InLog, InSeverity, InFormat, ##__VA_ARGS__)
 #else
 #define CHAOS_CHECK(Condition) (!!(Condition))
 #define CHAOS_ENSURE(Condition) (!!(Condition))
 #define CHAOS_ENSURE_MSG(InExpression, InFormat, ... ) (!!(InExpression))
+#define CHAOS_LOG(InLog, InSeverity, InFormat, ...)
+#define CHAOS_CLOG(InExpression, InLog, InSeverity, InFormat, ...) (!!(InExpression))
 #endif
diff --git a/Engine/Source/Runtime/Landscape/Private/LandscapeCollision.cpp b/Engine/Source/Runtime/Landscape/Private/LandscapeCollision.cpp
--- a/Engine/Source/Runtime/Landscape/Private/LandscapeCollision.cpp
+++ b/Engine/Source/Runtime/Landscape/Private/LandscapeCollision.cpp
@@ -1239,7 +1239,7 @@ void ULandscapeMeshCollisionComponent::CreateCollisionObject()
 				for (UPhysicalMaterial* PhysicalMaterial : CookedPhysicalMaterials)
 				{
 #if WITH_CHAOS || WITH_IMMEDIATE_PHYSX
-					UE_LOG(LogLandscape, Warning, TEXT("CHAOS - Landscape material setup not implemented"));
+					CHAOS_LOG(LogLandscape, Warning, TEXT("CHAOS - Landscape material setup not implemented"));
 #else
 					MeshRef->UsedPhysicalMaterialArray.Add(PhysicalMaterial->GetPhysicsMaterial().Material);
 #endif
@@ -1342,7 +1342,7 @@ void ULandscapeMeshCollisionComponent::OnCreatePhysicsState()
 					if (PTriMeshGeomEd.isValid())
 					{
 #if WITH_CHAOS || WITH_IMMEDIATE_PHYSX
-						UE_LOG(LogLandscape, Warning, TEXT("CHAOS - Landscape shape creation not implemented"));
+						CHAOS_LOG(LogLandscape, Warning, TEXT("CHAOS - Landscape shape creation not implemented"));
 #else
 						PxMaterial* PDefaultMat = GEngine->DefaultPhysMaterial->GetPhysicsMaterial().Material;
 						PxShape* MeshShapeEdSync = GPhysXSDK->createShape(PTriMeshGeomEd, &PDefaultMat, 1, true);
@@ -1370,7 +1370,7 @@ void ULandscapeMeshCollisionComponent::OnCreatePhysicsState()
 				BodyInstance.OwnerComponent = this;
 
 #if WITH_CHAOS || WITH_IMMEDIATE_PHYSX
-				UE_LOG(LogLandscape, Warning, TEXT("CHAOS - Landscape sim scene addition not implemented"));
+				CHAOS_LOG(LogLandscape, Warning, TEXT("CHAOS - Landscape sim scene addition not implemented"));
 #else
 				BodyInstance.ActorHandle.SyncActor = MeshActorSync;
 				MeshActorSync->userData = &BodyInstance.PhysicsUserData;
