Implemented navigation data export for Chaos heightfields

diff --git a/Engine/Source/Runtime/Engine/Public/AI/NavigationSystemHelpers.h b/Engine/Source/Runtime/Engine/Public/AI/NavigationSystemHelpers.h
--- a/Engine/Source/Runtime/Engine/Public/AI/NavigationSystemHelpers.h
+++ b/Engine/Source/Runtime/Engine/Public/AI/NavigationSystemHelpers.h
@@ -23,6 +23,14 @@ namespace physx
 }
 #endif // WITH_PHYSX
 
+#if WITH_CHAOS
+namespace Chaos
+{
+	template<typename T>
+	class THeightField;
+}
+#endif
+
 struct FNavigableGeometryExport
 {
 	virtual ~FNavigableGeometryExport() {}
@@ -32,6 +40,9 @@ struct FNavigableGeometryExport
 	virtual void ExportPxConvexMesh(physx::PxConvexMesh const * const ConvexMesh, const FTransform& LocalToWorld) = 0;
 	virtual void ExportPxHeightField(physx::PxHeightField const * const HeightField, const FTransform& LocalToWorld) = 0;
 #endif // WITH_PHYSX
+#if WITH_CHAOS
+	virtual void ExportChaosHeightField(const Chaos::THeightField<float>* const Heightfield, const FTransform& LocalToWorld) = 0;
+#endif
 	virtual void ExportHeightFieldSlice(const FNavHeightfieldSamples& PrefetchedHeightfieldSamples, const int32 NumRows, const int32 NumCols, const FTransform& LocalToWorld, const FBox& SliceBox) = 0;
 	virtual void ExportRigidBodySetup(UBodySetup& BodySetup, const FTransform& LocalToWorld) = 0;
 	virtual void ExportCustomMesh(const FVector* VertexBuffer, int32 NumVerts, const int32* IndexBuffer, int32 NumIndices, const FTransform& LocalToWorld) = 0;
diff --git a/Engine/Source/Runtime/Experimental/Chaos/Private/Chaos/HeightField.cpp b/Engine/Source/Runtime/Experimental/Chaos/Private/Chaos/HeightField.cpp
--- a/Engine/Source/Runtime/Experimental/Chaos/Private/Chaos/HeightField.cpp
+++ b/Engine/Source/Runtime/Experimental/Chaos/Private/Chaos/HeightField.cpp
@@ -505,6 +505,85 @@ namespace Chaos
 		return false;
 	}
 
+	template<typename T>
+	T Chaos::THeightField<T>::GetHeight(int32 InIndex) const
+	{
+		if (CHAOS_ENSURE(InIndex >= 0 && InIndex < GeomData.Heights.Num()))
+		{
+			return GeomData.GetPoint(InIndex).Z;
+		}
+
+		return TNumericLimits<T>::Max();
+	}
+
+	template<typename T>
+	T Chaos::THeightField<T>::GetHeight(int32 InX, int32 InY) const
+	{
+		const int32 Index = InY * GeomData.NumCols + InX;
+		return GetHeight(Index);
+	}
+
+	template<typename T>
+	T Chaos::THeightField<T>::GetHeightAt(const TVector<T, 2>& InGridLocation) const
+	{
+		const TVector<T, 2> Scale2D(GeomData.Scale[0], GeomData.Scale[1]);
+		const TVector<T, 2> UnscaledLocation = InGridLocation / Scale2D;
+
+		if(UnscaledLocation == FlatGrid.Clamp(UnscaledLocation))
+		{
+			TVector<int32, 2> CellCoord = FlatGrid.Cell(UnscaledLocation);
+
+			const int32 SingleIndex = CellCoord[1] * (GeomData.NumCols - 1) + CellCoord[0];
+			TVector<T, 3> Pts[4];
+			GeomData.GetPoints(SingleIndex, Pts);
+
+			float FractionX = FMath::Frac(UnscaledLocation[0]);
+			float FractionY = FMath::Frac(UnscaledLocation[1]);
+
+			if(FractionX > FractionY)
+			{
+				// In the second triangle (0,3,2)
+#if 0
+				// Reduced form of Bary calculation - TODO confirm working
+				float U = -1.0f * (FractionY - 1);
+				float V = FractionX;
+				float W = 1.0f - U - V;
+#endif
+
+				FVector Tri[3];
+				Tri[0] = FVector(0.0f, 0.0f, 0.0f);
+				Tri[1] = FVector(1.0f, 1.0f, 0.0f);
+				Tri[2] = FVector(0.0f, 1.0f, 0.0f);
+				
+				FVector Bary = FMath::GetBaryCentric2D({ FractionX, FractionY, 0.0f }, Tri[0], Tri[1], Tri[2]);
+
+				return Pts[0].Z * Bary[0] + Pts[1].Z * Bary[1] * Pts[3].Z * Bary[2];
+			}
+			else
+			{
+				// In the first triangle (0,1,3)
+
+#if 0
+				// Reduced form of Bary calculation - TODO confirm working
+				float U = ((-1) * (FractionX - 1));
+				float V = ((FractionX - 1) - (FractionY - 1));
+				float W = 1.0f - U - V;
+#endif
+
+				FVector Tri[3];
+				Tri[0] = FVector(0.0f, 0.0f, 0.0f);
+				Tri[1] = FVector(1.0f, 0.0f, 0.0f);
+				Tri[2] = FVector(1.0f, 1.0f, 0.0f);
+
+				FVector Bary = FMath::GetBaryCentric2D({FractionX, FractionY, 0.0f}, Tri[0], Tri[1], Tri[2]);
+
+				return Pts[0].Z * Bary[0] + Pts[1].Z * Bary[1] * Pts[3].Z * Bary[2];
+			}
+		}
+
+		return 0.0f;
+	}
+
 	template<typename T>
 	bool Chaos::THeightField<T>::GetCellBounds3D(const TVector<int32, 2> InCoord, TVector<T, 3>& OutMin, TVector<T, 3>& OutMax, const TVector<T, 3>& InInflate /*= {0}*/) const
 	{
diff --git a/Engine/Source/Runtime/Experimental/Chaos/Private/Chaos/UniformGrid.cpp b/Engine/Source/Runtime/Experimental/Chaos/Private/Chaos/UniformGrid.cpp
--- a/Engine/Source/Runtime/Experimental/Chaos/Private/Chaos/UniformGrid.cpp
+++ b/Engine/Source/Runtime/Experimental/Chaos/Private/Chaos/UniformGrid.cpp
@@ -259,5 +259,6 @@ bool Chaos::TUniformGrid<T, 3>::IsValid(const TVector<int32, 3>& X) const
 
 template class Chaos::TUniformGridBase<float, 3>;
 template class Chaos::TUniformGrid<float, 3>;
+template class Chaos::TUniformGrid<float, 2>;
 template TVector<float, 3> Chaos::TUniformGridBase<float, 3>::LinearlyInterpolate<TVector<float, 3>>(const TArrayND<TVector<float, 3>, 3>&, const TVector<float, 3>&) const;
 template CHAOS_API float Chaos::TUniformGridBase<float, 3>::LinearlyInterpolate<float>(const TArrayND<float, 3>&, const TVector<float, 3>&) const;
diff --git a/Engine/Source/Runtime/Experimental/Chaos/Public/Chaos/HeightField.h b/Engine/Source/Runtime/Experimental/Chaos/Public/Chaos/HeightField.h
--- a/Engine/Source/Runtime/Experimental/Chaos/Public/Chaos/HeightField.h
+++ b/Engine/Source/Runtime/Experimental/Chaos/Public/Chaos/HeightField.h
@@ -38,6 +38,11 @@ namespace Chaos
 		/** Support for editing a subsection of the heightfield */
 		void EditHeights(TArrayView<T> InHeights, int32 InBeginRow, int32 InBeginCol, int32 InNumRows, int32 InNumCols);
 		void EditHeights(TArrayView<const uint16> InHeights, int32 InBeginRow, int32 InBeginCol, int32 InNumRows, int32 InNumCols);
+		T GetHeight(int32 InIndex) const;
+		T GetHeight(int32 InX, int32 InY) const;
+		T GetHeightAt(const TVector<T, 2>& InGridLocation) const;
+		int32 GetNumRows() const { return GeomData.NumRows; }
+		int32 GetNumCols() const { return GeomData.NumCols; }
 
 		virtual T PhiWithNormal(const TVector<T, 3>& x, TVector<T, 3>& Normal) const
 		{
@@ -195,6 +200,21 @@ namespace Chaos
 				return Scale[1];
 			}
 
+			FORCEINLINE TVector<T, 3> GetPoint(int32 Index) const
+			{
+				const typename FDataType::RealType Height = MinValue + Heights[Index] * HeightPerUnit;
+
+				const int32 X = Index % (NumCols);
+				const int32 Y = Index / (NumCols);
+
+				return {(typename FDataType::RealType)X, (typename FDataType::RealType)Y, Height};
+			}
+
+			FORCEINLINE TVector<T, 3> GetPointScaled(int32 Index) const
+			{
+				return GetPoint(Index) * Scale;
+			}
+
 			FORCEINLINE void GetPoints(int32 Index, TVector<T, 3> OutPts[4]) const
 			{
 				const typename FDataType::RealType H0 = MinValue + Heights[Index] * HeightPerUnit;
diff --git a/Engine/Source/Runtime/Experimental/ChaosCore/Public/Chaos/Vector.h b/Engine/Source/Runtime/Experimental/ChaosCore/Public/Chaos/Vector.h
--- a/Engine/Source/Runtime/Experimental/ChaosCore/Public/Chaos/Vector.h
+++ b/Engine/Source/Runtime/Experimental/ChaosCore/Public/Chaos/Vector.h
@@ -622,6 +622,14 @@ namespace Chaos
 		{
 			return TVector<float, 2>(X / Other, Y / Other);
 		}
+		TVector<float, 2> operator*(const float Other) const
+		{
+			return TVector<float, 2>(X * Other, Y * Other);
+		}
+		TVector<float, 2> operator*(const TVector<float, 2>& Other) const
+		{
+			return TVector<float, 2>(X * Other[0], Y * Other[1]);
+		}
 	};
 #endif // !COMPILE_WITHOUT_UNREAL_SUPPORT
 
@@ -804,6 +812,11 @@ namespace Chaos
 			: X((int32)InVector.X)
 			, Y((int32)InVector.Y)
 		{}
+		FORCEINLINE TVector(std::istream& Stream)
+		{
+			Stream.read(reinterpret_cast<char*>(&X), sizeof(int32));
+			Stream.read(reinterpret_cast<char*>(&Y), sizeof(int32));
+		}
 		FORCEINLINE ~TVector()
 		{}
 
diff --git a/Engine/Source/Runtime/Landscape/Private/LandscapeCollision.cpp b/Engine/Source/Runtime/Landscape/Private/LandscapeCollision.cpp
--- a/Engine/Source/Runtime/Landscape/Private/LandscapeCollision.cpp
+++ b/Engine/Source/Runtime/Landscape/Private/LandscapeCollision.cpp
@@ -83,7 +83,7 @@ ULandscapeHeightfieldCollisionComponent::FHeightfieldGeometryRef::FHeightfieldGe
 
 ULandscapeHeightfieldCollisionComponent::FHeightfieldGeometryRef::~FHeightfieldGeometryRef()
 {
-#if WITH_PHYSX
+#if WITH_PHYSX && PHYSICS_INTERFACE_PHYSX
 	// Free the existing heightfield data.
 	if (RBHeightfield)
 	{
@@ -107,7 +107,7 @@ TMap<FGuid, ULandscapeMeshCollisionComponent::FPhysXMeshRef* > GSharedMeshRefs;
 
 ULandscapeMeshCollisionComponent::FPhysXMeshRef::~FPhysXMeshRef()
 {
-#if WITH_PHYSX
+#if WITH_PHYSX && PHYSICS_INTERFACE_PHYSX
 	// Free the existing heightfield data.
 	if (RBTriangleMesh)
 	{
@@ -288,9 +288,10 @@ void ULandscapeHeightfieldCollisionComponent::OnCreatePhysicsState()
 
 			// Create the geometry
 			FVector FinalScale(LandscapeScale.X * CollisionScale, LandscapeScale.Y * CollisionScale, LandscapeScale.Z * LANDSCAPE_ZSCALE);
-			PxHeightFieldGeometry LandscapeComponentGeom(HeightfieldRef->RBHeightfield, PxMeshGeometryFlag::eDOUBLE_SIDED, LandscapeScale.Z * LANDSCAPE_ZSCALE, LandscapeScale.Y * CollisionScale, LandscapeScale.X * CollisionScale);
 
 #if PHYSICS_INTERFACE_PHYSX
+			PxHeightFieldGeometry LandscapeComponentGeom(HeightfieldRef->RBHeightfield, PxMeshGeometryFlag::eDOUBLE_SIDED, LandscapeScale.Z * LANDSCAPE_ZSCALE, LandscapeScale.Y * CollisionScale, LandscapeScale.X * CollisionScale);
+
 			if (LandscapeComponentGeom.isValid())
 			{
 				// Creating both a sync and async actor, since this object is static
@@ -656,8 +657,6 @@ void ULandscapeHeightfieldCollisionComponent::CreateCollisionObject()
 				for(UPhysicalMaterial* PhysicalMaterial : CookedPhysicalMaterials)
 				{
 					//todo: total hack until we get landscape fully converted to chaos
-					HeightfieldRef->UsedPhysicalMaterialArray.Add(GPhysXSDK->createMaterial(1, 1, 1));
-
 					HeightfieldRef->UsedChaosMaterials.Add(PhysicalMaterial->GetPhysicsMaterial());
 				}
 
@@ -1897,7 +1896,7 @@ bool ULandscapeHeightfieldCollisionComponent::ComponentIsTouchingSelectionFrustu
 bool ULandscapeHeightfieldCollisionComponent::DoCustomNavigableGeometryExport(FNavigableGeometryExport& GeomExport) const
 {
 	check(IsInGameThread());
-#if WITH_PHYSX
+#if WITH_PHYSX && PHYSICS_INTERFACE_PHYSX
 	if (IsValidRef(HeightfieldRef) && HeightfieldRef->RBHeightfield)
 	{
 		FTransform HFToW = GetComponentTransform();
@@ -1913,7 +1912,24 @@ bool ULandscapeHeightfieldCollisionComponent::DoCustomNavigableGeometryExport(FN
 			GeomExport.ExportPxHeightField(HeightfieldRef->RBHeightfield, HFToW);
 		}
 	}
-#endif// WITH_PHYSX
+#elif WITH_CHAOS
+	if(IsValidRef(HeightfieldRef) && HeightfieldRef->Heightfield)
+	{
+		FTransform HFToW = GetComponentTransform();
+		if(HeightfieldRef->HeightfieldSimple)
+		{
+			const float SimpleCollisionScale = CollisionScale * CollisionSizeQuads / SimpleCollisionSizeQuads;
+			HFToW.MultiplyScale3D(FVector(SimpleCollisionScale, SimpleCollisionScale, LANDSCAPE_ZSCALE));
+			GeomExport.ExportChaosHeightField(HeightfieldRef->HeightfieldSimple.Get(), HFToW);
+		}
+		else
+		{
+			HFToW.MultiplyScale3D(FVector(CollisionScale, CollisionScale, LANDSCAPE_ZSCALE));
+			GeomExport.ExportChaosHeightField(HeightfieldRef->Heightfield.Get(), HFToW);
+		}
+	}
+#endif
+
 	return false;
 }
 
@@ -1938,7 +1954,7 @@ ENavDataGatheringMode ULandscapeHeightfieldCollisionComponent::GetGeometryGather
 void ULandscapeHeightfieldCollisionComponent::PrepareGeometryExportSync()
 {
 	//check(IsInGameThread());
-#if WITH_PHYSX
+#if WITH_PHYSX && PHYSICS_INTERFACE_PHYSX
 	if (IsValidRef(HeightfieldRef) && HeightfieldRef->RBHeightfield != nullptr && CachedHeightFieldSamples.IsEmpty())
 	{
 		const UWorld* World = GetWorld();
@@ -1967,6 +1983,32 @@ void ULandscapeHeightfieldCollisionComponent::PrepareGeometryExportSync()
 			}
 		}
 	}
+#elif WITH_CHAOS
+	if(IsValidRef(HeightfieldRef) && HeightfieldRef->Heightfield.Get() && CachedHeightFieldSamples.IsEmpty())
+	{
+		const UWorld* World = GetWorld();
+
+		if(World != nullptr)
+		{
+			HeightfieldRowsCount = HeightfieldRef->Heightfield->GetNumRows();
+			HeightfieldColumnsCount = HeightfieldRef->Heightfield->GetNumCols();
+
+			if(CachedHeightFieldSamples.Heights.Num() != HeightfieldRowsCount * HeightfieldRowsCount)
+			{
+				QUICK_SCOPE_CYCLE_COUNTER(STAT_NavMesh_ExportPxHeightField_saveCells);
+
+				CachedHeightFieldSamples.Heights.SetNumUninitialized(HeightfieldRowsCount * HeightfieldRowsCount);
+				
+				for(int32 SampleIndex = 0; SampleIndex < CachedHeightFieldSamples.Heights.Num(); ++SampleIndex)
+				{
+					CachedHeightFieldSamples.Heights[SampleIndex] = HeightfieldRef->Heightfield->GetHeight(SampleIndex);
+					// #PHYSTODO Hole support
+					CachedHeightFieldSamples.Holes.Add(false);
+				}
+			}
+		}
+	}
+
 #endif// WITH_PHYSX
 }
 
@@ -2513,12 +2555,17 @@ ULandscapeComponent* ULandscapeHeightfieldCollisionComponent::GetRenderComponent
 TOptional<float> ULandscapeHeightfieldCollisionComponent::GetHeight(float X, float Y)
 {
 	TOptional<float> Height;
-#if WITH_PHYSX
 	const float ZScale = GetComponentTransform().GetScale3D().Z * LANDSCAPE_ZSCALE;
+#if WITH_PHYSX && PHYSICS_INTERFACE_PHYSX
 	if (IsValidRef(HeightfieldRef) && HeightfieldRef->RBHeightfield != nullptr)
 	{
 		Height = HeightfieldRef->RBHeightfield->getHeight(HeightfieldRef->RBHeightfield->getNbRows() - 1 - X, Y) * ZScale;
 	}
+#elif WITH_CHAOS
+	if(IsValidRef(HeightfieldRef) && HeightfieldRef->Heightfield.Get())
+	{
+		Height = HeightfieldRef->Heightfield->GetHeightAt({ X, Y }) * ZScale;
+	}
 #endif
 	return Height;
 }
diff --git a/Engine/Source/Runtime/NavigationSystem/Private/NavMesh/RecastNavMeshGenerator.cpp b/Engine/Source/Runtime/NavigationSystem/Private/NavMesh/RecastNavMeshGenerator.cpp
--- a/Engine/Source/Runtime/NavigationSystem/Private/NavMesh/RecastNavMeshGenerator.cpp
+++ b/Engine/Source/Runtime/NavigationSystem/Private/NavMesh/RecastNavMeshGenerator.cpp
@@ -16,6 +16,9 @@
 #if WITH_PHYSX
 	#include "PhysXPublic.h"
 #endif
+#if WITH_CHAOS
+	#include "Chaos/HeightField.h"
+#endif
 #include "NavMesh/PImplRecastNavMesh.h"
 
 // recast includes
@@ -205,6 +208,9 @@ struct FRecastGeometryExport : public FNavigableGeometryExport
 	virtual void ExportPxConvexMesh(physx::PxConvexMesh const * const ConvexMesh, const FTransform& LocalToWorld) override;
 	virtual void ExportPxHeightField(physx::PxHeightField const * const HeightField, const FTransform& LocalToWorld) override;
 #endif // WITH_PHYSX
+#if WITH_CHAOS
+	virtual void ExportChaosHeightField(const Chaos::THeightField<float>* const Heightfield, const FTransform& LocalToWorld) override;
+#endif
 	virtual void ExportHeightFieldSlice(const FNavHeightfieldSamples& PrefetchedHeightfieldSamples, const int32 NumRows, const int32 NumCols, const FTransform& LocalToWorld, const FBox& SliceBox) override;
 	virtual void ExportCustomMesh(const FVector* InVertices, int32 NumVerts, const int32* InIndices, int32 NumIndices, const FTransform& LocalToWorld) override;
 	virtual void ExportRigidBodySetup(UBodySetup& BodySetup, const FTransform& LocalToWorld) override;
@@ -528,6 +534,81 @@ void ExportPxHeightField(PxHeightField const * const HeightField, const FTransfo
 }
 #endif // WITH_PHYSX
 
+#if WITH_CHAOS
+void ExportChaosHeightField(const Chaos::THeightField<float>* const HeightField, const FTransform& LocalToWorld
+	, TNavStatArray<float>& VertexBuffer, TNavStatArray<int32>& IndexBuffer
+	, FBox& UnrealBounds)
+{
+	using namespace Chaos;
+
+	if(HeightField == nullptr)
+	{
+		return;
+	}
+
+	QUICK_SCOPE_CYCLE_COUNTER(STAT_NavMesh_ExportPxHeightField);
+
+	const int32 NumRows = HeightField->GetNumRows();
+	const int32 NumCols = HeightField->GetNumCols();
+	const int32 VertexCount = NumRows * NumCols;
+
+	const int32 VertOffset = VertexBuffer.Num() / 3;
+	const int32 NumQuads = (NumRows - 1) * (NumCols - 1);
+
+	VertexBuffer.Reserve(VertexBuffer.Num() + VertexCount * 3);
+	IndexBuffer.Reserve(IndexBuffer.Num() + NumQuads * 6);
+
+	const bool bMirrored = (LocalToWorld.GetDeterminant() < 0.f);
+
+	for(int32 Y = 0; Y < NumRows; Y++)
+	{
+		for(int32 X = 0; X < NumCols; X++)
+		{
+			const int32 SampleIdx = Y * NumCols + X;
+
+			const FVector UnrealCoords = LocalToWorld.TransformPosition(FVector(X, Y, HeightField->GetHeight(SampleIdx)));
+			UnrealBounds += UnrealCoords;
+
+			VertexBuffer.Add(UnrealCoords.X);
+			VertexBuffer.Add(UnrealCoords.Y);
+			VertexBuffer.Add(UnrealCoords.Z);
+		}
+	}
+
+	for(int32 Y = 0; Y < NumRows - 1; Y++)
+	{
+		for(int32 X = 0; X < NumCols - 1; X++)
+		{
+			// #PHYSTODO Hole support for chaos heightfields
+			const bool bIsHole = false;
+			if(bIsHole)
+			{
+				continue;
+			}
+
+			const int32 I0 = Y * NumCols + X;
+			int32 I1 = I0 + 1;
+			int32 I2 = I0 + NumCols;
+			const int32 I3 = I2 + 1;
+
+			if(bMirrored)
+			{
+				// Flip the winding so the triangles face the right way after scaling
+				Swap(I1, I2);
+			}
+
+			IndexBuffer.Add(VertOffset + I0);
+			IndexBuffer.Add(VertOffset + I3);
+			IndexBuffer.Add(VertOffset + I1);
+
+			IndexBuffer.Add(VertOffset + I0);
+			IndexBuffer.Add(VertOffset + I2);
+			IndexBuffer.Add(VertOffset + I3);
+		}
+	}
+}
+#endif
+
 void ExportHeightFieldSlice(const FNavHeightfieldSamples& PrefetchedHeightfieldSamples, const int32 NumRows, const int32 NumCols, const FTransform& LocalToWorld
 	, TNavStatArray<float>& VertexBuffer, TNavStatArray<int32>& IndexBuffer, const FBox& SliceBox
 	, FBox& UnrealBounds)
@@ -1089,6 +1170,13 @@ void FRecastGeometryExport::ExportPxHeightField(physx::PxHeightField const * con
 }
 #endif // WITH_PHYSX
 
+#if WITH_CHAOS
+void FRecastGeometryExport::ExportChaosHeightField(const Chaos::THeightField<float>* const Heightfield, const FTransform& LocalToWorld)
+{
+	RecastGeometryExport::ExportChaosHeightField(Heightfield, LocalToWorld, VertexBuffer, IndexBuffer, Data->Bounds);
+}
+#endif
+
 void FRecastGeometryExport::ExportHeightFieldSlice(const FNavHeightfieldSamples& PrefetchedHeightfieldSamples, const int32 NumRows, const int32 NumCols, const FTransform& LocalToWorld, const FBox& SliceBox)
 {
 	RecastGeometryExport::ExportHeightFieldSlice(PrefetchedHeightfieldSamples, NumRows, NumCols, LocalToWorld, VertexBuffer, IndexBuffer, SliceBox, Data->Bounds);
