Fix landscape LOD calculation not taking LOD scale into account
https://udn.unrealengine.com/questions/488728/landscape-lod-popping-with-non-90-fov.html
#rn minor Landscape

diff --git a/Engine/Source/Runtime/Landscape/Private/LandscapeRender.cpp b/Engine/Source/Runtime/Landscape/Private/LandscapeRender.cpp
--- a/Engine/Source/Runtime/Landscape/Private/LandscapeRender.cpp
+++ b/Engine/Source/Runtime/Landscape/Private/LandscapeRender.cpp
@@ -1695,7 +1695,7 @@ void FLandscapeComponentSceneProxy::CalculateLODFromScreenSize(const FSceneView&
 	}
 	else
 	{
-		PreferedLOD = FMath::Clamp<float>(ComputeBatchElementCurrentLOD(GetLODFromScreenSize(InMeshScreenSizeSquared, InViewLODScale), InMeshScreenSizeSquared) + LocalLODBias, FMath::Max((float)MinStreamedLOD, MinValidLOD), FMath::Min((float)LastLOD, MaxValidLOD));
+		PreferedLOD = FMath::Clamp<float>(ComputeBatchElementCurrentLOD(GetLODFromScreenSize(InMeshScreenSizeSquared, InViewLODScale), InMeshScreenSizeSquared, InViewLODScale) + LocalLODBias, FMath::Max((float)MinStreamedLOD, MinValidLOD), FMath::Min((float)LastLOD, MaxValidLOD));
 	}
 
 	check(PreferedLOD != -1.0f && PreferedLOD <= MaxLOD);
@@ -1876,7 +1876,7 @@ int32 FLandscapeComponentSceneProxy::ConvertBatchElementLODToBatchElementIndex(i
 	return BatchElementIndex;
 }
 
-float FLandscapeComponentSceneProxy::ComputeBatchElementCurrentLOD(int32 InSelectedLODIndex, float InComponentScreenSize) const
+float FLandscapeComponentSceneProxy::ComputeBatchElementCurrentLOD(int32 InSelectedLODIndex, float InComponentScreenSize, float InViewLODScale) const
 {
 	check(LODScreenRatioSquared.IsValidIndex(InSelectedLODIndex));
 
@@ -1885,17 +1885,18 @@ float FLandscapeComponentSceneProxy::ComputeBatchElementCurrentLOD(int32 InSelec
 	float NextLODScreenRatio = LastElement ? 0 : LODScreenRatioSquared[InSelectedLODIndex + 1];
 
 	float LODScreenRatioRange = CurrentLODScreenRatio - NextLODScreenRatio;
+	float ScreenSizeWithLODScale = InComponentScreenSize / InViewLODScale;
 
-	if (InComponentScreenSize > CurrentLODScreenRatio || InComponentScreenSize < NextLODScreenRatio)
+	if (ScreenSizeWithLODScale > CurrentLODScreenRatio || ScreenSizeWithLODScale < NextLODScreenRatio)
 	{
 		// Find corresponding LODIndex to appropriately calculate Ratio and apply it to new LODIndex
-		int32 LODFromScreenSize = GetLODFromScreenSize(InComponentScreenSize, 1.0f); // for 4.19 only
+		int32 LODFromScreenSize = GetLODFromScreenSize(InComponentScreenSize, InViewLODScale);
 		CurrentLODScreenRatio = LODScreenRatioSquared[LODFromScreenSize];
 		NextLODScreenRatio = LODFromScreenSize == LODScreenRatioSquared.Num() - 1 ? 0 : LODScreenRatioSquared[LODFromScreenSize + 1];
 		LODScreenRatioRange = CurrentLODScreenRatio - NextLODScreenRatio;
 	}
 
-	float CurrentLODRangeRatio = (InComponentScreenSize - NextLODScreenRatio) / LODScreenRatioRange;
+	float CurrentLODRangeRatio = (ScreenSizeWithLODScale - NextLODScreenRatio) / LODScreenRatioRange;
 	float fLOD = (float)InSelectedLODIndex + (1.0f - CurrentLODRangeRatio);
 
 	return fLOD;
diff --git a/Engine/Source/Runtime/Landscape/Public/LandscapeRender.h b/Engine/Source/Runtime/Landscape/Public/LandscapeRender.h
--- a/Engine/Source/Runtime/Landscape/Public/LandscapeRender.h
+++ b/Engine/Source/Runtime/Landscape/Public/LandscapeRender.h
@@ -652,7 +652,7 @@ protected:
 	void CalculateLODFromScreenSize(const FSceneView& InView, float InMeshScreenSizeSquared, float InViewLODScale, int32 InSubSectionIndex, FViewCustomDataLOD& InOutLODData) const;
 	FORCEINLINE void ComputeStaticBatchIndexToRender(FViewCustomDataLOD& OutLODData, int32 InSubSectionIndex);
 	int8 GetLODFromScreenSize(float InScreenSizeSquared, float InViewLODScale) const;
-	FORCEINLINE float ComputeBatchElementCurrentLOD(int32 InSelectedLODIndex, float InComponentScreenSize) const;
+	FORCEINLINE float ComputeBatchElementCurrentLOD(int32 InSelectedLODIndex, float InComponentScreenSize, float InViewLODScale) const;
 	
 	FORCEINLINE void GetShaderCurrentNeighborLOD(const FSceneView& InView, float InBatchElementCurrentLOD, int8 InSubSectionX, int8 InSubSectionY, int8 InCurrentSubSectionIndex, FVector4& OutShaderCurrentNeighborLOD) const;
 	FORCEINLINE FVector4 GetShaderLODBias() const;
