Reparent landscape components owned textures and materials to the landscape actor instead of the level package.
This makes more sense and will be required if we save actors outside levels in the future.
Tested by [REDACTED_PERSON].

#rnx

diff --git a/Engine/Source/Editor/LandscapeEditor/Private/LandscapeEdMode.cpp b/Engine/Source/Editor/LandscapeEditor/Private/LandscapeEdMode.cpp
--- a/Engine/Source/Editor/LandscapeEditor/Private/LandscapeEdMode.cpp
+++ b/Engine/Source/Editor/LandscapeEditor/Private/LandscapeEdMode.cpp
@@ -99,7 +99,7 @@ IMPLEMENT_HIT_PROXY(HNewLandscapeGrabHandleProxy, HHitProxy)
 
 ENGINE_API extern bool GDisableAutomaticTextureMaterialUpdateDependencies;
 
-void ALandscape::SplitHeightmap(ULandscapeComponent* Comp, bool bMoveToCurrentLevel, FMaterialUpdateContext* InOutUpdateContext, TArray<FComponentRecreateRenderStateContext>* InOutRecreateRenderStateContext, bool InReregisterComponent)
+void ALandscape::SplitHeightmap(ULandscapeComponent* Comp,  UObject* TextureOuter, FMaterialUpdateContext* InOutUpdateContext, TArray<FComponentRecreateRenderStateContext>* InOutRecreateRenderStateContext, bool InReregisterComponent)
 {
 	ULandscapeInfo* Info = Comp->GetLandscapeInfo();
 	ALandscape* Landscape = Info->LandscapeActor.Get();
@@ -122,7 +122,6 @@ void ALandscape::SplitHeightmap(ULandscapeComponent* Comp, bool bMoveToCurrentLe
 		LandscapeEdit.GetHeightDataFast(Comp->GetSectionBase().X, Comp->GetSectionBase().Y, Comp->GetSectionBase().X + Comp->ComponentSizeQuads, Comp->GetSectionBase().Y + Comp->ComponentSizeQuads, (uint16*)HeightData.GetData(), 0, (uint16*)NormalData.GetData());
 
 		// Construct the heightmap textures
-		UObject* TextureOuter = bMoveToCurrentLevel ? Comp->GetWorld()->GetCurrentLevel()->GetOutermost() : nullptr;
 		HeightmapTexture = Comp->GetLandscapeProxy()->CreateLandscapeTexture(HeightmapSizeU, HeightmapSizeV, TEXTUREGROUP_Terrain_Heightmap, TSF_BGRA8, TextureOuter);
 
 		int32 MipSubsectionSizeQuads = Comp->SubsectionSizeQuads;
@@ -3661,7 +3660,7 @@ void FEdModeLandscape::DeleteLandscapeComponents(ULandscapeInfo* LandscapeInfo,
 		// Changing Heightmap format for selected components
 		for (ULandscapeComponent* Component : HeightmapUpdateComponents)
 		{
-			ALandscape::SplitHeightmap(Component, false, &MaterialUpdateContext, &RecreateRenderStateContexts, false);
+			ALandscape::SplitHeightmap(Component, nullptr, &MaterialUpdateContext, &RecreateRenderStateContexts, false);
 		}
 
 		FMultiComponentReregisterContext RegisterContext(ComponentsToReregister);
diff --git a/Engine/Source/Editor/LandscapeEditor/Private/LandscapeEdModeComponentTools.cpp b/Engine/Source/Editor/LandscapeEditor/Private/LandscapeEdModeComponentTools.cpp
--- a/Engine/Source/Editor/LandscapeEditor/Private/LandscapeEdModeComponentTools.cpp
+++ b/Engine/Source/Editor/LandscapeEditor/Private/LandscapeEdModeComponentTools.cpp
@@ -577,21 +577,6 @@ public:
 					}
 				}
 
-				// Changing Heightmap format for selected components
-				for (const auto& HeightmapUpdateComponentPair : HeightmapUpdateComponents)
-				{
-					ALandscape::SplitHeightmap(HeightmapUpdateComponentPair.Key, HeightmapUpdateComponentPair.Value);
-				}
-
-				// Delete if it is no referenced textures...
-				for (UTexture2D* Texture : OldHeightmapTextures)
-				{
-					Texture->SetFlags(RF_Transactional);
-					Texture->Modify();
-					Texture->MarkPackageDirty();
-					Texture->ClearFlags(RF_Standalone);
-				}
-
 				ALandscapeProxy* LandscapeProxy = LandscapeInfo->GetCurrentLevelLandscapeProxy(false);
 				if (!LandscapeProxy)
 				{
@@ -612,6 +597,21 @@ public:
 					}
 				}
 
+				// Changing Heightmap format for selected components
+				for (const auto& HeightmapUpdateComponentPair : HeightmapUpdateComponents)
+				{
+					ALandscape::SplitHeightmap(HeightmapUpdateComponentPair.Key, HeightmapUpdateComponentPair.Value ? LandscapeProxy : nullptr);
+				}
+
+				// Delete if it is no referenced textures...
+				for (UTexture2D* Texture : OldHeightmapTextures)
+				{
+					Texture->SetFlags(RF_Transactional);
+					Texture->Modify();
+					Texture->MarkPackageDirty();
+					Texture->ClearFlags(RF_Standalone);
+				}
+
 				for (ALandscapeProxy* Proxy : SelectProxies)
 				{
 					Proxy->Modify();
@@ -626,7 +626,7 @@ public:
 					if (Component->XYOffsetmapTexture)
 					{
 						Component->XYOffsetmapTexture->Modify();
-						Component->XYOffsetmapTexture->Rename(nullptr, LandscapeProxy->GetOutermost());
+						Component->XYOffsetmapTexture->Rename(nullptr, LandscapeProxy->GetOuter());
 					}
 				}
 
diff --git a/Engine/Source/Runtime/Landscape/Classes/Landscape.h b/Engine/Source/Runtime/Landscape/Classes/Landscape.h
--- a/Engine/Source/Runtime/Landscape/Classes/Landscape.h
+++ b/Engine/Source/Runtime/Landscape/Classes/Landscape.h
@@ -211,7 +211,7 @@ public:
 	LANDSCAPE_API static void CalcComponentIndicesNoOverlap(const int32 X1, const int32 Y1, const int32 X2, const int32 Y2, const int32 ComponentSizeQuads,
 		int32& ComponentIndexX1, int32& ComponentIndexY1, int32& ComponentIndexX2, int32& ComponentIndexY2);
 
-	static void SplitHeightmap(ULandscapeComponent* Comp, bool bMoveToCurrentLevel = false, class FMaterialUpdateContext* InOutUpdateContext = nullptr, TArray<class FComponentRecreateRenderStateContext>* InOutRecreateRenderStateContext = nullptr, bool InReregisterComponent = true);
+	static void SplitHeightmap(ULandscapeComponent* Comp,  UObject* TextureOuter, class FMaterialUpdateContext* InOutUpdateContext = nullptr, TArray<class FComponentRecreateRenderStateContext>* InOutRecreateRenderStateContext = nullptr, bool InReregisterComponent = true);
 	
 	//~ Begin UObject Interface.
 	virtual void PreSave(const class ITargetPlatform* TargetPlatform) override;
diff --git a/Engine/Source/Runtime/Landscape/Classes/LandscapeProxy.h b/Engine/Source/Runtime/Landscape/Classes/LandscapeProxy.h
--- a/Engine/Source/Runtime/Landscape/Classes/LandscapeProxy.h
+++ b/Engine/Source/Runtime/Landscape/Classes/LandscapeProxy.h
@@ -922,8 +922,8 @@ public:
 	/** @return Current size of bounding rectangle in quads space */
 	LANDSCAPE_API FIntRect GetBoundingRect() const;
 
-	/** Creates a Texture2D for use by this landscape proxy or one of it's components. If OptionalOverrideOuter is not specified, the level is used. */
-	LANDSCAPE_API UTexture2D* CreateLandscapeTexture(int32 InSizeX, int32 InSizeY, TextureGroup InLODGroup, ETextureSourceFormat InFormat, UObject* OptionalOverrideOuter = nullptr, bool bCompress = false) const;
+	/** Creates a Texture2D for use by this landscape proxy or one of it's components. If OptionalOverrideOuter is not specified, the proxy is used. */
+	LANDSCAPE_API UTexture2D* CreateLandscapeTexture(int32 InSizeX, int32 InSizeY, TextureGroup InLODGroup, ETextureSourceFormat InFormat, UObject* OptionalOverrideOuter = nullptr, bool bCompress = false);
 
 	/* For the grassmap rendering notification */
 	int32 NumComponentsNeedingGrassMapRender;
diff --git a/Engine/Source/Runtime/Landscape/Private/Landscape.cpp b/Engine/Source/Runtime/Landscape/Private/Landscape.cpp
--- a/Engine/Source/Runtime/Landscape/Private/Landscape.cpp
+++ b/Engine/Source/Runtime/Landscape/Private/Landscape.cpp
@@ -745,6 +745,55 @@ void ULandscapeComponent::PostLoad()
 	}
 #endif
 
+#if WITH_EDITOR
+	auto ReparentObject = [this](UObject* Object)
+	{
+		if(Object && !Object->HasAllFlags(RF_Public | RF_Standalone) && (Object->GetOuter() != GetOuter()))
+		{
+			Object->Rename(nullptr, GetOuter(), REN_ForceNoResetLoaders);
+			return true;
+		}
+		return false;
+	};
+
+	for(UMaterialInstance* MaterialInstance: MaterialInstances)
+	{
+		while(ReparentObject(MaterialInstance))
+		{
+			MaterialInstance = Cast<UMaterialInstance>(MaterialInstance->Parent);
+		}
+	}
+
+	for(UMaterialInterface* MobileMaterialInterface: MobileMaterialInterfaces)
+	{
+		while(ReparentObject(MobileMaterialInterface))
+		{
+			MobileMaterialInterface = Cast<UMaterialInstance>(MobileMaterialInterface) ? Cast<UMaterialInstance>(((UMaterialInstance*)MobileMaterialInterface)->Parent) : nullptr;
+		}
+	}
+
+	for(UMaterialInstance* MobileCombinationMaterialInstance: MobileCombinationMaterialInstances)
+	{
+		while(ReparentObject(MobileCombinationMaterialInstance))
+		{
+			MobileCombinationMaterialInstance = Cast<UMaterialInstance>(MobileCombinationMaterialInstance->Parent);
+		}
+	}
+
+	ReparentObject(HeightmapTexture);
+	ReparentObject(XYOffsetmapTexture);
+
+	for(UTexture2D* WeightmapTexture: WeightmapTextures)
+	{
+		ReparentObject(WeightmapTexture);
+	}
+
+	for(UTexture2D* MobileWeightmapTexture: MobileWeightmapTextures)
+	{
+		ReparentObject(MobileWeightmapTexture);
+	}
+#endif
+
 #if !UE_BUILD_SHIPPING
 	// This will fix the data in case there is mismatch between save of asset/maps
 	int8 MaxLOD = FMath::CeilLogTwo(SubsectionSizeQuads + 1) - 1;
diff --git a/Engine/Source/Runtime/Landscape/Private/LandscapeEdit.cpp b/Engine/Source/Runtime/Landscape/Private/LandscapeEdit.cpp
--- a/Engine/Source/Runtime/Landscape/Private/LandscapeEdit.cpp
+++ b/Engine/Source/Runtime/Landscape/Private/LandscapeEdit.cpp
@@ -251,11 +251,11 @@ UMaterialInstanceConstant* ULandscapeComponent::GetCombinationMaterial(FMaterial
 
 		// Find or set a matching MIC in the Landscape's map.
 		UMaterialInstanceConstant* CombinationMaterialInstance = Proxy->MaterialInstanceConstantMap.FindRef(*LayerKey);
-		if (CombinationMaterialInstance == nullptr || CombinationMaterialInstance->Parent != MaterialToUse || GetOutermost() != CombinationMaterialInstance->GetOutermost())
+		if (CombinationMaterialInstance == nullptr || CombinationMaterialInstance->Parent != MaterialToUse || GetOuter() != CombinationMaterialInstance->GetOuter())
 		{
 			FlushRenderingCommands();
 
-			ULandscapeMaterialInstanceConstant* LandscapeCombinationMaterialInstance = NewObject<ULandscapeMaterialInstanceConstant>(GetOutermost());
+			ULandscapeMaterialInstanceConstant* LandscapeCombinationMaterialInstance = NewObject<ULandscapeMaterialInstanceConstant>(GetOuter());
 			LandscapeCombinationMaterialInstance->bMobile = bMobile;
 			CombinationMaterialInstance = LandscapeCombinationMaterialInstance;
 			UE_LOG(LogLandscape, Log, TEXT("Looking for key %s, making new combination %s"), *LayerKey, *CombinationMaterialInstance->GetName());
@@ -345,7 +345,7 @@ void ULandscapeComponent::UpdateMaterialInstances_Internal(FMaterialUpdateContex
 		if (CombinationMaterialInstance != nullptr)
 		{
 			// Create the instance for this component, that will use the layer combination instance.
-			UMaterialInstanceConstant* MaterialInstance = NewObject<ULandscapeMaterialInstanceConstant>(GetOutermost());
+			UMaterialInstanceConstant* MaterialInstance = NewObject<ULandscapeMaterialInstanceConstant>(GetOuter());
 			MaterialInstances[MaterialIndex] = MaterialInstance;
 
 			// Material Instances don't support Undo/Redo (the shader map goes out of sync and crashes happen)
@@ -384,7 +384,7 @@ void ULandscapeComponent::UpdateMaterialInstances_Internal(FMaterialUpdateContex
 			// Setup material instance with disabled tessellation
 			if (CombinationMaterialInstance->GetMaterial()->D3D11TessellationMode != EMaterialTessellationMode::MTM_NoTessellation)
 			{
-				ULandscapeMaterialInstanceConstant* TessellationMaterialInstance = NewObject<ULandscapeMaterialInstanceConstant>(GetOutermost());
+				ULandscapeMaterialInstanceConstant* TessellationMaterialInstance = NewObject<ULandscapeMaterialInstanceConstant>(GetOuter());
 				int32 TessellatedMaterialIndex = MaterialPerLOD.Num() + TessellatedMaterialCount++;
 				MaterialInstances[TessellatedMaterialIndex] = TessellationMaterialInstance;
 				MaterialIndexToDisabledTessellationMaterial[MaterialIndex] = TessellatedMaterialIndex;
@@ -5499,7 +5499,7 @@ void ULandscapeComponent::GeneratePlatformPixelData()
 			MobileCombinationMaterialInstances[MaterialIndex] = GetCombinationMaterial(nullptr, MobileWeightmapLayerAllocations, MaterialLOD, true);
 			check(MobileCombinationMaterialInstances[MaterialIndex] != nullptr);
 
-			UMaterialInstanceConstant* NewMobileMaterialInstance = NewObject<ULandscapeMaterialInstanceConstant>(GetOutermost());
+			UMaterialInstanceConstant* NewMobileMaterialInstance = NewObject<ULandscapeMaterialInstanceConstant>(GetOuter());
 
 			NewMobileMaterialInstance->SetParentEditorOnly(MobileCombinationMaterialInstances[MaterialIndex]);
 
@@ -5740,9 +5740,9 @@ void ULandscapeComponent::GeneratePlatformVertexData(const ITargetPlatform* Targ
 	PlatformData.InitializeFromUncompressedData(NewPlatformData);
 }
 
-UTexture2D* ALandscapeProxy::CreateLandscapeTexture(int32 InSizeX, int32 InSizeY, TextureGroup InLODGroup, ETextureSourceFormat InFormat, UObject* OptionalOverrideOuter, bool bCompress) const
+UTexture2D* ALandscapeProxy::CreateLandscapeTexture(int32 InSizeX, int32 InSizeY, TextureGroup InLODGroup, ETextureSourceFormat InFormat, UObject* OptionalOverrideOuter, bool bCompress)
 {
-	UObject* TexOuter = OptionalOverrideOuter ? OptionalOverrideOuter : GetOutermost();
+	UObject* TexOuter = OptionalOverrideOuter ? OptionalOverrideOuter : this;
 	UTexture2D* NewTexture = NewObject<UTexture2D>(TexOuter);
 	NewTexture->Source.Init2DWithMipChain(InSizeX, InSizeY, InFormat);
 	NewTexture->SRGB = false;
