Don't give warnings for duplicated landscape actors guid when loading them, as it can happen when streaming the same landscape several times.

#rnx

diff --git a/Engine/Source/Runtime/Landscape/Private/Landscape.cpp b/Engine/Source/Runtime/Landscape/Private/Landscape.cpp
--- a/Engine/Source/Runtime/Landscape/Private/Landscape.cpp
+++ b/Engine/Source/Runtime/Landscape/Private/Landscape.cpp
@@ -1688,24 +1688,9 @@ void ALandscape::PostLoad()
 		{
 			if (Landscape && Landscape != this && Landscape->LandscapeGuid == LandscapeGuid && Landscape->GetWorld() == CurrentWorld)
 			{
-				// Duplicated landscape level, need to generate new GUID
+				// Duplicated landscape level, need to generate new GUID. This can happen during PIE or gameplay when streaming the same landscape actor.
 				Modify();
 				LandscapeGuid = FGuid::NewGuid();
-
-
-				// Show MapCheck window
-
-				FFormatNamedArguments Arguments;
-				Arguments.Add(TEXT("ProxyName1"), FText::FromString(Landscape->GetName()));
-				Arguments.Add(TEXT("LevelName1"), FText::FromString(Landscape->GetLevel()->GetOutermost()->GetName()));
-				Arguments.Add(TEXT("ProxyName2"), FText::FromString(this->GetName()));
-				Arguments.Add(TEXT("LevelName2"), FText::FromString(this->GetLevel()->GetOutermost()->GetName()));
-				FMessageLog("LoadErrors").Warning()
-					->AddToken(FUObjectToken::Create(this))
-					->AddToken(FTextToken::Create(FText::Format(LOCTEXT("LoadError_DuplicateLandscapeGuid", "Landscape {ProxyName1} of {LevelName1} has the same guid as {ProxyName2} of {LevelName2}. {LevelName2}.{ProxyName2} has had its guid automatically changed, please save {LevelName2}!"), Arguments)));
-
-				// Show MapCheck window
-				FMessageLog("LoadErrors").Open();
 				break;
 			}
 		}
