Make sure we don't try to register an actor pending destruction

diff --git a/Engine/Source/Runtime/Landscape/Private/Landscape.cpp b/Engine/Source/Runtime/Landscape/Private/Landscape.cpp
--- a/Engine/Source/Runtime/Landscape/Private/Landscape.cpp
+++ b/Engine/Source/Runtime/Landscape/Private/Landscape.cpp
@@ -1742,7 +1742,7 @@ void ALandscapeProxy::PostRegisterAllComponents()
 
 #if WITH_EDITOR
 	// Game worlds don't have landscape infos
-	if (!GetWorld()->IsGameWorld())
+	if (!GetWorld()->IsGameWorld() && !IsPendingKillPending())
 	{
 		// Duplicated Landscapes don't have a valid guid until PostEditImport is called, we'll register then
 		if (LandscapeGuid.IsValid())
