Detect when changes required to rebuild the material occur instead of always rebuilding them. Which mean that when using the slider, it will only update the materials when the slider stopped being interactive.

diff --git a/Engine/Source/Runtime/Landscape/Classes/Landscape.h b/Engine/Source/Runtime/Landscape/Classes/Landscape.h
--- a/Engine/Source/Runtime/Landscape/Classes/Landscape.h
+++ b/Engine/Source/Runtime/Landscape/Classes/Landscape.h
@@ -218,6 +218,7 @@ public:
 	
 	//~ Begin UObject Interface.
 	virtual void PreSave(const class ITargetPlatform* TargetPlatform) override;
+	virtual void PreEditChange(UProperty* PropertyThatWillChange) override;
 	virtual void PostEditChangeProperty(FPropertyChangedEvent& PropertyChangedEvent) override;
 	virtual void PostEditMove(bool bFinished) override;
 	virtual void PostEditUndo() override;
diff --git a/Engine/Source/Runtime/Landscape/Classes/LandscapeProxy.h b/Engine/Source/Runtime/Landscape/Classes/LandscapeProxy.h
--- a/Engine/Source/Runtime/Landscape/Classes/LandscapeProxy.h
+++ b/Engine/Source/Runtime/Landscape/Classes/LandscapeProxy.h
@@ -321,6 +321,29 @@ struct FLandscapeProxyMaterialOverride
 
 	UPROPERTY(EditAnywhere, Category = Landscape)
 	UMaterialInterface* Material;
+
+	bool operator==(const FLandscapeProxyMaterialOverride& InOther) const
+	{
+		if (Material != InOther.Material)
+		{
+			return false;
+		}
+
+		if (LODIndex.Default != InOther.LODIndex.Default || LODIndex.PerPlatform.Num() != InOther.LODIndex.PerPlatform.Num())
+		{
+			return false;
+		}
+
+		for (auto& ItPair : LODIndex.PerPlatform)
+		{
+			if (!InOther.LODIndex.PerPlatform.Contains(ItPair.Key))
+			{
+				return false;
+			}
+		}
+
+		return true;
+	}
 };
 
 class FLandscapeLayersTexture2DCPUReadBackResource : public FTextureResource
@@ -461,6 +484,20 @@ public:
 	UPROPERTY(EditAnywhere, Category = Landscape)
 	TArray<FLandscapeProxyMaterialOverride> LandscapeMaterialsOverride;
 
+#if WITH_EDITORONLY_DATA
+	UPROPERTY(Transient)
+	UMaterialInterface* PreEditLandscapeMaterial;
+
+	UPROPERTY(Transient)
+	UMaterialInterface* PreEditLandscapeHoleMaterial;
+
+	UPROPERTY(Transient)
+	TArray<FLandscapeProxyMaterialOverride> PreEditLandscapeMaterialsOverride;
+
+	UPROPERTY(Transient)
+	bool bIsPerformingInteractiveActionOnLandscapeMaterialOverride;
+#endif
+
 	/** Allows overriding the landscape bounds. This is useful if you distort the landscape with world-position-offset, for example
 	 *  Extension value in the negative Z axis, positive value increases bound size
 	 *  Note that this can also be overridden per-component when the component is selected with the component select tool */
diff --git a/Engine/Source/Runtime/Landscape/Private/LandscapeEdit.cpp b/Engine/Source/Runtime/Landscape/Private/LandscapeEdit.cpp
--- a/Engine/Source/Runtime/Landscape/Private/LandscapeEdit.cpp
+++ b/Engine/Source/Runtime/Landscape/Private/LandscapeEdit.cpp
@@ -4138,6 +4138,15 @@ void ALandscapeStreamingProxy::PostEditChangeProperty(FPropertyChangedEvent& Pro
 	Super::PostEditChangeProperty(PropertyChangedEvent);
 }
 
+void ALandscape::PreEditChange(UProperty* PropertyThatWillChange)
+{
+	PreEditLandscapeMaterial = LandscapeMaterial;
+	PreEditLandscapeHoleMaterial = LandscapeHoleMaterial;
+	PreEditLandscapeMaterialsOverride = LandscapeMaterialsOverride;
+
+	Super::PreEditChange(PropertyThatWillChange);
+}
+
 void ALandscape::PostEditChangeProperty(FPropertyChangedEvent& PropertyChangedEvent)
 {
 	const FName PropertyName = PropertyChangedEvent.Property ? PropertyChangedEvent.Property->GetFName() : NAME_None;
@@ -4151,16 +4160,42 @@ void ALandscape::PostEditChangeProperty(FPropertyChangedEvent& PropertyChangedEv
 
 	ULandscapeInfo* Info = GetLandscapeInfo();
 
-	if (PropertyName == FName(TEXT("LandscapeMaterial")) || PropertyName == FName(TEXT("LandscapeHoleMaterial")) || MemberPropertyName == FName(TEXT("LandscapeMaterialsOverride")))
+	if ((PropertyName == GET_MEMBER_NAME_CHECKED(ALandscapeProxy, LandscapeMaterial) || PropertyName == GET_MEMBER_NAME_CHECKED(ALandscapeProxy, LandscapeHoleMaterial) || MemberPropertyName == GET_MEMBER_NAME_CHECKED(ALandscapeProxy, LandscapeMaterialsOverride))
+		&& PropertyChangedEvent.ChangeType != EPropertyChangeType::ArrayAdd)
 	{
-		bool RecreateMaterialInstances = true;
+		bool HasMaterialChanged = false;
 
-		if (PropertyName == FName(TEXT("LandscapeMaterialsOverride")) && PropertyChangedEvent.ChangeType == EPropertyChangeType::ArrayAdd)
+		if (PropertyChangedEvent.ChangeType != EPropertyChangeType::Interactive)
 		{
-			RecreateMaterialInstances = false;
-		}		
+			if (PreEditLandscapeMaterial != LandscapeMaterial || PreEditLandscapeHoleMaterial != LandscapeHoleMaterial || PreEditLandscapeMaterialsOverride.Num() != LandscapeMaterialsOverride.Num() || bIsPerformingInteractiveActionOnLandscapeMaterialOverride)
+			{
+				HasMaterialChanged = true;
+			}
+
+			if (!HasMaterialChanged)
+			{
+				for (int32 i = 0; i < LandscapeMaterialsOverride.Num(); ++i)
+				{
+					const FLandscapeProxyMaterialOverride& NewMaterialOverride = LandscapeMaterialsOverride[i];
+					const FLandscapeProxyMaterialOverride& PreEditMaterialOverride = PreEditLandscapeMaterialsOverride[i];
+
+					if (!(PreEditMaterialOverride == NewMaterialOverride))
+					{
+						HasMaterialChanged = true;
+						break;
+					}
+				}
+			}
 
-		if (Info != nullptr && RecreateMaterialInstances)
+			bIsPerformingInteractiveActionOnLandscapeMaterialOverride = false;
+		}
+		else
+		{
+			// We are probably using a slider or something similar in LandscapeMaterialsOverride
+			bIsPerformingInteractiveActionOnLandscapeMaterialOverride = MemberPropertyName == GET_MEMBER_NAME_CHECKED(ALandscapeProxy, LandscapeMaterialsOverride);
+		}
+
+		if (Info != nullptr && HasMaterialChanged)
 		{
 			FMaterialUpdateContext MaterialUpdateContext;
 			Info->UpdateLayerInfoMap(/*this*/);
@@ -4349,6 +4384,10 @@ void ALandscape::PostEditChangeProperty(FPropertyChangedEvent& PropertyChangedEv
 			LandscapeMaterialChangedDelegate.Broadcast();
 		}
 	}
+
+	PreEditLandscapeMaterial = nullptr;
+	PreEditLandscapeHoleMaterial = nullptr;
+	PreEditLandscapeMaterialsOverride.Empty();
 }
 
 void ALandscapeProxy::ChangedPhysMaterial()
