Update the mobile landscape content on load if we preview in ES3.1 but the landscape is in a streaming level that is loaded at runtime.

diff --git a/Engine/Source/Runtime/Landscape/Private/Landscape.cpp b/Engine/Source/Runtime/Landscape/Private/Landscape.cpp
--- a/Engine/Source/Runtime/Landscape/Private/Landscape.cpp
+++ b/Engine/Source/Runtime/Landscape/Private/Landscape.cpp
@@ -970,8 +970,10 @@ void ULandscapeComponent::PostLoad()
 
 	if (!HasAnyFlags(RF_ClassDefaultObject))
 	{
+		UWorld* World = GetWorld();
+
 		// If we're loading on a platform that doesn't require cooked data, but *only* supports OpenGL ES, generate or preload data from the DDC
-		if (!FPlatformProperties::RequiresCookedData() && GMaxRHIFeatureLevel <= ERHIFeatureLevel::ES3_1)
+		if (!FPlatformProperties::RequiresCookedData() && (GMaxRHIFeatureLevel <= ERHIFeatureLevel::ES3_1 || World && World->FeatureLevel <= ERHIFeatureLevel::ES3_1))
 		{
 			CheckGenerateLandscapePlatformData(false, nullptr);
 		}
