Fixed non editor build error

diff --git a/Engine/Source/Runtime/Landscape/Classes/LandscapeProxy.h b/Engine/Source/Runtime/Landscape/Classes/LandscapeProxy.h
--- a/Engine/Source/Runtime/Landscape/Classes/LandscapeProxy.h
+++ b/Engine/Source/Runtime/Landscape/Classes/LandscapeProxy.h
@@ -322,6 +322,7 @@ struct FLandscapeProxyMaterialOverride
 	UPROPERTY(EditAnywhere, Category = Landscape)
 	UMaterialInterface* Material;
 
+#if WITH_EDITORONLY_DATA
 	bool operator==(const FLandscapeProxyMaterialOverride& InOther) const
 	{
 		if (Material != InOther.Material)
@@ -344,6 +345,7 @@ struct FLandscapeProxyMaterialOverride
 
 		return true;
 	}
+#endif
 };
 
 class FLandscapeLayersTexture2DCPUReadBackResource : public FTextureResource
