Prevent possible crash if the LODIndex do not exist in the primitives static mesh list

diff --git a/Engine/Source/Runtime/Landscape/Private/LandscapeRender.cpp b/Engine/Source/Runtime/Landscape/Private/LandscapeRender.cpp
--- a/Engine/Source/Runtime/Landscape/Private/LandscapeRender.cpp
+++ b/Engine/Source/Runtime/Landscape/Private/LandscapeRender.cpp
@@ -2270,7 +2270,11 @@ bool FLandscapeComponentSceneProxy::CanUseMeshBatchForShadowCascade(int8 InLODIn
 	const FStaticMeshBatch* MeshBatch = nullptr;
 	const TArray<FStaticMeshBatch>& PrimitiveStaticMeshes = GetPrimitiveSceneInfo()->StaticMeshes;
 
-	check(PrimitiveStaticMeshes.IsValidIndex(InLODIndex));
+	if (!PrimitiveStaticMeshes.IsValidIndex(InLODIndex))
+	{
+		return false;
+	}
+
 	check(PrimitiveStaticMeshes[InLODIndex].CastShadow);
 	check(InLODIndex == PrimitiveStaticMeshes[InLODIndex].LODIndex);
 	MeshBatch = &PrimitiveStaticMeshes[InLODIndex];
