Prevent possible crash if GetDynamicMeshElement was called but InitViewCustomData wasn't called yet.

diff --git a/Engine/Source/Runtime/Landscape/Private/LandscapeRender.cpp b/Engine/Source/Runtime/Landscape/Private/LandscapeRender.cpp
--- a/Engine/Source/Runtime/Landscape/Private/LandscapeRender.cpp
+++ b/Engine/Source/Runtime/Landscape/Private/LandscapeRender.cpp
@@ -1298,7 +1298,7 @@ float FLandscapeComponentSceneProxy::GetComponentScreenSize(const FSceneView* Vi
 
 void FLandscapeComponentSceneProxy::BuildDynamicMeshElement(const FViewCustomDataLOD* InPrimitiveCustomData, bool InToolMesh, bool InHasTessellation, bool InDisableTessellation, FMeshBatch& OutMeshBatch, TArray<FLandscapeBatchElementParams, SceneRenderingAllocator>& OutStaticBatchParamArray) const
 {
-	if (AvailableMaterials.Num() == 0)
+	if (AvailableMaterials.Num() == 0 || InPrimitiveCustomData == nullptr || InPrimitiveCustomData->SubSections.Num() == 0)
 	{
 		return;
 	}
