Removed HoleMaterial, as it's no longer required, the same behavior can be achieved using a normal material. (Was discussed with Jack P at the time)

diff --git a/Engine/Source/Editor/LandscapeEditor/Private/LandscapeEdMode.cpp b/Engine/Source/Editor/LandscapeEditor/Private/LandscapeEdMode.cpp
--- a/Engine/Source/Editor/LandscapeEditor/Private/LandscapeEdMode.cpp
+++ b/Engine/Source/Editor/LandscapeEditor/Private/LandscapeEdMode.cpp
@@ -4061,7 +4061,6 @@ ALandscape* FEdModeLandscape::ChangeComponentSetting(int32 NumComponentsX, int32
 			NewLandscape->PositiveZBoundsExtension = OldLandscape->PositiveZBoundsExtension;
 			NewLandscape->DefaultPhysMaterial = OldLandscape->DefaultPhysMaterial;
 			NewLandscape->StreamingDistanceMultiplier = OldLandscape->StreamingDistanceMultiplier;
-			NewLandscape->LandscapeHoleMaterial = OldLandscape->LandscapeHoleMaterial;
 			NewLandscape->StaticLightingResolution = OldLandscape->StaticLightingResolution;
 			NewLandscape->bCastStaticShadow = OldLandscape->bCastStaticShadow;
 			NewLandscape->bCastShadowAsTwoSided = OldLandscape->bCastShadowAsTwoSided;
diff --git a/Engine/Source/Editor/LandscapeEditor/Private/LandscapeEdModeComponentTools.cpp b/Engine/Source/Editor/LandscapeEditor/Private/LandscapeEdModeComponentTools.cpp
--- a/Engine/Source/Editor/LandscapeEditor/Private/LandscapeEdModeComponentTools.cpp
+++ b/Engine/Source/Editor/LandscapeEditor/Private/LandscapeEdModeComponentTools.cpp
@@ -330,11 +330,8 @@ public:
 	virtual bool BeginTool(FEditorViewportClient* ViewportClient, const FLandscapeToolTarget& InTarget, const FVector& InHitLocation) override
 	{
 		ALandscapeProxy* Proxy = InTarget.LandscapeInfo->GetLandscapeProxy();
-		UMaterialInterface* HoleMaterial = Proxy->GetLandscapeHoleMaterial();
-		if (!HoleMaterial)
-		{
-			HoleMaterial = Proxy->GetLandscapeMaterial();
-		}
+		UMaterialInterface* HoleMaterial = Proxy->GetLandscapeMaterial();
+
 		if (!HoleMaterial->GetMaterial()->HasAnyExpressionsInMaterialAndFunctionsOfType<UMaterialExpressionLandscapeVisibilityMask>())
 		{
 			FMessageDialog::Open(EAppMsgType::Ok,
diff --git a/Engine/Source/Editor/LandscapeEditor/Private/LandscapeEditorDetailCustomization_TargetLayers.cpp b/Engine/Source/Editor/LandscapeEditor/Private/LandscapeEditorDetailCustomization_TargetLayers.cpp
--- a/Engine/Source/Editor/LandscapeEditor/Private/LandscapeEditorDetailCustomization_TargetLayers.cpp
+++ b/Engine/Source/Editor/LandscapeEditor/Private/LandscapeEditorDetailCustomization_TargetLayers.cpp
@@ -139,12 +139,9 @@ bool FLandscapeEditorDetailCustomization_TargetLayers::ShouldShowVisibilityTip()
 		if (LandscapeEdMode->CurrentToolTarget.TargetType == ELandscapeToolTargetType::Visibility)
 		{
 			ALandscapeProxy* Proxy = LandscapeEdMode->CurrentToolTarget.LandscapeInfo->GetLandscapeProxy();
-			UMaterialInterface* HoleMaterial = Proxy->GetLandscapeHoleMaterial();
-			if (!HoleMaterial)
-			{
-				HoleMaterial = Proxy->GetLandscapeMaterial();
-			}
-			if (!HoleMaterial->GetMaterial()->HasAnyExpressionsInMaterialAndFunctionsOfType<UMaterialExpressionLandscapeVisibilityMask>())
+			UMaterialInterface* HoleMaterial = Proxy->GetLandscapeMaterial();
+
+			if (HoleMaterial && !HoleMaterial->GetMaterial()->HasAnyExpressionsInMaterialAndFunctionsOfType<UMaterialExpressionLandscapeVisibilityMask>())
 			{
 				return true;
 			}
diff --git a/Engine/Source/Runtime/Core/Public/UObject/EditorObjectVersion.h b/Engine/Source/Runtime/Core/Public/UObject/EditorObjectVersion.h
--- a/Engine/Source/Runtime/Core/Public/UObject/EditorObjectVersion.h
+++ b/Engine/Source/Runtime/Core/Public/UObject/EditorObjectVersion.h
@@ -77,6 +77,8 @@ struct CORE_API FEditorObjectVersion
 		CultureInvariantTextSerializationKeyStability,
 		// Change to UScrollBar and UScrollBox thickness property (removed implicit padding of 2, so thickness value must be incremented by 4).
 		ScrollBarThicknessChange,
+		// Deprecated LandscapeHoleMaterial
+		RemoveLandscapeHoleMaterial,
 		// -----<new versions can be added above this line>-------------------------------------------------
 		VersionPlusOne,
 		LatestVersion = VersionPlusOne - 1
diff --git a/Engine/Source/Runtime/Landscape/Classes/LandscapeComponent.h b/Engine/Source/Runtime/Landscape/Classes/LandscapeComponent.h
--- a/Engine/Source/Runtime/Landscape/Classes/LandscapeComponent.h
+++ b/Engine/Source/Runtime/Landscape/Classes/LandscapeComponent.h
@@ -355,8 +355,8 @@ class ULandscapeComponent : public UPrimitiveComponent
 	UPROPERTY(EditAnywhere, BlueprintReadWrite, Category=LandscapeComponent)
 	UMaterialInterface* OverrideMaterial;
 
-	UPROPERTY(EditAnywhere, BlueprintReadWrite, Category=LandscapeComponent, AdvancedDisplay)
-	UMaterialInterface* OverrideHoleMaterial;
+	UPROPERTY()
+	UMaterialInterface* OverrideHoleMaterial_DEPRECATED;
 
 	UPROPERTY(EditAnywhere, Category = LandscapeComponent)
 	TArray<FLandscapeComponentMaterialOverride> OverrideMaterials;
@@ -841,9 +841,6 @@ public:
 	/** Returns the actor's LandscapeMaterial, or the Component's OverrideLandscapeMaterial if set */
 	LANDSCAPE_API UMaterialInterface* GetLandscapeMaterial(int8 InLODIndex = INDEX_NONE) const;
 
-	/** Returns the actor's LandscapeHoleMaterial, or the Component's OverrideLandscapeHoleMaterial if set */
-	LANDSCAPE_API UMaterialInterface* GetLandscapeHoleMaterial() const;
-
 	/** Returns true if this component has visibility painted */
 	LANDSCAPE_API bool ComponentHasVisibilityPainted() const;
 
diff --git a/Engine/Source/Runtime/Landscape/Classes/LandscapeProxy.h b/Engine/Source/Runtime/Landscape/Classes/LandscapeProxy.h
--- a/Engine/Source/Runtime/Landscape/Classes/LandscapeProxy.h
+++ b/Engine/Source/Runtime/Landscape/Classes/LandscapeProxy.h
@@ -479,9 +479,8 @@ public:
 	UPROPERTY(EditAnywhere, BlueprintSetter=EditorSetLandscapeMaterial, Category=Landscape)
 	UMaterialInterface* LandscapeMaterial;
 
-	/** Material used to render landscape components with holes. If not set, LandscapeMaterial will be used (blend mode will be overridden to Masked if it is set to Opaque) */
-	UPROPERTY(EditAnywhere, Category=Landscape, AdvancedDisplay)
-	UMaterialInterface* LandscapeHoleMaterial;
+	UPROPERTY()
+	UMaterialInterface* LandscapeHoleMaterial_DEPRECATED;
 
 	UPROPERTY(EditAnywhere, Category = Landscape)
 	TArray<FLandscapeProxyMaterialOverride> LandscapeMaterialsOverride;
@@ -490,9 +489,6 @@ public:
 	UPROPERTY(Transient)
 	UMaterialInterface* PreEditLandscapeMaterial;
 
-	UPROPERTY(Transient)
-	UMaterialInterface* PreEditLandscapeHoleMaterial;
-
 	UPROPERTY(Transient)
 	TArray<FLandscapeProxyMaterialOverride> PreEditLandscapeMaterialsOverride;
 
@@ -844,9 +840,6 @@ public:
 	// Get Landscape Material assigned to this Landscape
 	virtual UMaterialInterface* GetLandscapeMaterial(int8 InLODIndex = INDEX_NONE) const;
 
-	// Get Hole Landscape Material assigned to this Landscape
-	virtual UMaterialInterface* GetLandscapeHoleMaterial() const;
-
 	// 
 	void FixupWeightmaps();
 
diff --git a/Engine/Source/Runtime/Landscape/Classes/LandscapeStreamingProxy.h b/Engine/Source/Runtime/Landscape/Classes/LandscapeStreamingProxy.h
--- a/Engine/Source/Runtime/Landscape/Classes/LandscapeStreamingProxy.h
+++ b/Engine/Source/Runtime/Landscape/Classes/LandscapeStreamingProxy.h
@@ -32,7 +32,6 @@ public:
 	virtual ALandscape* GetLandscapeActor() override;
 #if WITH_EDITOR
 	virtual UMaterialInterface* GetLandscapeMaterial(int8 InLODIndex = INDEX_NONE) const override;
-	virtual UMaterialInterface* GetLandscapeHoleMaterial() const override;
 #endif
 	//~ End ALandscapeBase Interface
 
diff --git a/Engine/Source/Runtime/Landscape/Private/Landscape.cpp b/Engine/Source/Runtime/Landscape/Private/Landscape.cpp
--- a/Engine/Source/Runtime/Landscape/Private/Landscape.cpp
+++ b/Engine/Source/Runtime/Landscape/Private/Landscape.cpp
@@ -61,6 +61,7 @@ Landscape.cpp: Terrain rendering
 #include "LandscapeVersion.h"
 #include "UObject/FortniteMainBranchObjectVersion.h"
 #include "LandscapeDataAccess.h"
+#include "UObject/EditorObjectVersion.h"
 
 /** Landscape stats */
 
@@ -313,6 +314,7 @@ void ULandscapeComponent::Serialize(FArchive& Ar)
 	LLM_SCOPE(ELLMTag::Landscape);
 	Ar.UsingCustomVersion(FRenderingObjectVersion::GUID);
 	Ar.UsingCustomVersion(FFortniteMainBranchObjectVersion::GUID);
+	Ar.UsingCustomVersion(FEditorObjectVersion::GUID);
 
 #if WITH_EDITOR
 	if (Ar.IsCooking() && !HasAnyFlags(RF_ClassDefaultObject) && Ar.CookingTarget()->SupportsFeature(ETargetPlatformFeatures::MobileRendering))
@@ -438,7 +440,24 @@ void ULandscapeComponent::Serialize(FArchive& Ar)
 		LegacyComponentData.Data.Emplace(MapBuildDataId, LegacyMapBuildData);
 		GComponentsWithLegacyLightmaps.AddAnnotation(this, LegacyComponentData);
 	}
-
+	
+	if (Ar.IsLoading() && Ar.CustomVer(FEditorObjectVersion::GUID) < FEditorObjectVersion::RemoveLandscapeHoleMaterial)
+	{
+		if (OverrideHoleMaterial_DEPRECATED != nullptr)
+		{
+			if (OverrideMaterial == nullptr)
+			{
+				OverrideMaterial = OverrideHoleMaterial_DEPRECATED;
+			}
+			else
+			{
+				if (OverrideMaterial->GetBlendMode() != EBlendMode::BLEND_Masked)
+				{
+					UE_LOG(LogLandscape, Warning, TEXT("OverrideHoleMaterial was deprecated, your OverrideMaterial is not currently setup correctly to support visibility painting, please correct your material, otherwise your landscape component \"%s\" won't render as before."), *GetFullName());
+				}
+			}
+		}
+	}
 	if (Ar.IsLoading() && Ar.CustomVer(FFortniteMainBranchObjectVersion::GUID) < FFortniteMainBranchObjectVersion::NewLandscapeMaterialPerLOD)
 	{
 		if (MobileMaterialInterface_DEPRECATED != nullptr)
@@ -570,20 +589,6 @@ UMaterialInterface* ULandscapeComponent::GetLandscapeMaterial(int8 InLODIndex) c
 	return UMaterial::GetDefaultMaterial(MD_Surface);
 }
 
-UMaterialInterface* ULandscapeComponent::GetLandscapeHoleMaterial() const
-{
-	if (OverrideHoleMaterial)
-	{
-		return OverrideHoleMaterial;
-	}
-	ALandscapeProxy* Proxy = GetLandscapeProxy();
-	if (Proxy)
-	{
-		return Proxy->GetLandscapeHoleMaterial();
-	}
-	return nullptr;
-}
-
 bool ULandscapeComponent::ComponentHasVisibilityPainted() const
 {
 	for (const FWeightmapLayerAllocationInfo& Allocation : WeightmapLayerAllocations)
@@ -1899,6 +1904,7 @@ void ALandscapeProxy::Serialize(FArchive& Ar)
 	Super::Serialize(Ar);
 
 	Ar.UsingCustomVersion(FLandscapeCustomVersion::GUID);
+	Ar.UsingCustomVersion(FEditorObjectVersion::GUID);
 
 	if (Ar.IsLoading() && Ar.CustomVer(FLandscapeCustomVersion::GUID) < FLandscapeCustomVersion::MigrateOldPropertiesToNewRenderingProperties)
 	{
@@ -1921,6 +1927,24 @@ void ALandscapeProxy::Serialize(FArchive& Ar)
 			}
 		}
 	}
+	
+	if (Ar.IsLoading() && Ar.CustomVer(FEditorObjectVersion::GUID) < FEditorObjectVersion::RemoveLandscapeHoleMaterial)
+	{
+		if (LandscapeHoleMaterial_DEPRECATED != nullptr)
+		{
+			if (LandscapeMaterial == nullptr)
+			{
+				LandscapeMaterial = LandscapeHoleMaterial_DEPRECATED;
+			}
+			else
+			{
+				if (LandscapeMaterial->GetBlendMode() != EBlendMode::BLEND_Masked)
+				{
+					UE_LOG(LogLandscape, Warning, TEXT("LandscapeHoleMaterial was deprecated, your landscape material is not currently setup correctly to support visibility painting, please correct your material, otherwise your landscape \"%s\" won't render as before."), *GetFullName());
+				}
+			}
+		}
+	}
 }
 
 void ALandscapeProxy::AddReferencedObjects(UObject* InThis, FReferenceCollector& Collector)
@@ -2338,10 +2362,6 @@ void ALandscapeProxy::GetSharedProperties(ALandscapeProxy* Landscape)
 			LandscapeMaterial = Landscape->LandscapeMaterial;
 			LandscapeMaterialsOverride = Landscape->LandscapeMaterialsOverride;
 		}
-		if (!LandscapeHoleMaterial)
-		{
-			LandscapeHoleMaterial = Landscape->LandscapeHoleMaterial;
-		}
 		if (LandscapeMaterial == Landscape->LandscapeMaterial)
 		{
 			EditorLayerSettings = Landscape->EditorLayerSettings;
@@ -2531,15 +2551,6 @@ UMaterialInterface* ALandscapeProxy::GetLandscapeMaterial(int8 InLODIndex) const
 	return LandscapeMaterial != nullptr ? LandscapeMaterial : UMaterial::GetDefaultMaterial(MD_Surface);
 }
 
-UMaterialInterface* ALandscapeProxy::GetLandscapeHoleMaterial() const
-{
-	if (LandscapeHoleMaterial)
-	{
-		return LandscapeHoleMaterial;
-	}
-	return nullptr;
-}
-
 UMaterialInterface* ALandscapeStreamingProxy::GetLandscapeMaterial(int8 InLODIndex) const
 {
 	if (InLODIndex != INDEX_NONE)
@@ -2576,19 +2587,6 @@ UMaterialInterface* ALandscapeStreamingProxy::GetLandscapeMaterial(int8 InLODInd
 	return UMaterial::GetDefaultMaterial(MD_Surface);
 }
 
-UMaterialInterface* ALandscapeStreamingProxy::GetLandscapeHoleMaterial() const
-{
-	if (LandscapeHoleMaterial)
-	{
-		return LandscapeHoleMaterial;
-	}
-	else if (ALandscape* Landscape = LandscapeActor.Get())
-	{
-		return Landscape->GetLandscapeHoleMaterial();
-	}
-	return nullptr;
-}
-
 void ALandscape::PreSave(const class ITargetPlatform* TargetPlatform)
 {
 	Super::PreSave(TargetPlatform);
diff --git a/Engine/Source/Runtime/Landscape/Private/LandscapeEdit.cpp b/Engine/Source/Runtime/Landscape/Private/LandscapeEdit.cpp
--- a/Engine/Source/Runtime/Landscape/Private/LandscapeEdit.cpp
+++ b/Engine/Source/Runtime/Landscape/Private/LandscapeEdit.cpp
@@ -229,9 +229,8 @@ UMaterialInstanceConstant* ULandscapeComponent::GetCombinationMaterial(FMaterial
 
 	const bool bComponentHasHoles = ComponentHasVisibilityPainted();
 	UMaterialInterface* const LandscapeMaterial = GetLandscapeMaterial(InLODIndex);
-	UMaterialInterface* const HoleMaterial = bComponentHasHoles ? GetLandscapeHoleMaterial() : nullptr;
-	UMaterialInterface* const MaterialToUse = bComponentHasHoles && HoleMaterial ? HoleMaterial : LandscapeMaterial;
-	bool bOverrideBlendMode = bComponentHasHoles && !HoleMaterial && LandscapeMaterial->GetBlendMode() == BLEND_Opaque;
+	UMaterialInterface* const MaterialToUse = LandscapeMaterial;
+	bool bOverrideBlendMode = bComponentHasHoles && LandscapeMaterial->GetBlendMode() == BLEND_Opaque;
 
 	if (bOverrideBlendMode)
 	{
@@ -4089,7 +4088,7 @@ void ALandscapeStreamingProxy::PostEditChangeProperty(FPropertyChangedEvent& Pro
 			LandscapeActor = nullptr;
 		}
 	}
-	else if (PropertyName == FName(TEXT("LandscapeMaterial")) || PropertyName == FName(TEXT("LandscapeHoleMaterial")) || PropertyName == FName(TEXT("LandscapeMaterialsOverride")))
+	else if (PropertyName == FName(TEXT("LandscapeMaterial")) || PropertyName == FName(TEXT("LandscapeMaterialsOverride")))
 	{
 		bool RecreateMaterialInstances = true;
 
@@ -4141,7 +4140,6 @@ void ALandscapeStreamingProxy::PostEditChangeProperty(FPropertyChangedEvent& Pro
 void ALandscape::PreEditChange(UProperty* PropertyThatWillChange)
 {
 	PreEditLandscapeMaterial = LandscapeMaterial;
-	PreEditLandscapeHoleMaterial = LandscapeHoleMaterial;
 	PreEditLandscapeMaterialsOverride = LandscapeMaterialsOverride;
 
 	Super::PreEditChange(PropertyThatWillChange);
@@ -4160,14 +4158,14 @@ void ALandscape::PostEditChangeProperty(FPropertyChangedEvent& PropertyChangedEv
 
 	ULandscapeInfo* Info = GetLandscapeInfo();
 
-	if ((PropertyName == GET_MEMBER_NAME_CHECKED(ALandscapeProxy, LandscapeMaterial) || PropertyName == GET_MEMBER_NAME_CHECKED(ALandscapeProxy, LandscapeHoleMaterial) || MemberPropertyName == GET_MEMBER_NAME_CHECKED(ALandscapeProxy, LandscapeMaterialsOverride))
+	if ((PropertyName == GET_MEMBER_NAME_CHECKED(ALandscapeProxy, LandscapeMaterial) || MemberPropertyName == GET_MEMBER_NAME_CHECKED(ALandscapeProxy, LandscapeMaterialsOverride))
 		&& PropertyChangedEvent.ChangeType != EPropertyChangeType::ArrayAdd)
 	{
 		bool HasMaterialChanged = false;
 
 		if (PropertyChangedEvent.ChangeType != EPropertyChangeType::Interactive)
 		{
-			if (PreEditLandscapeMaterial != LandscapeMaterial || PreEditLandscapeHoleMaterial != LandscapeHoleMaterial || PreEditLandscapeMaterialsOverride.Num() != LandscapeMaterialsOverride.Num() || bIsPerformingInteractiveActionOnLandscapeMaterialOverride)
+			if (PreEditLandscapeMaterial != LandscapeMaterial || PreEditLandscapeMaterialsOverride.Num() != LandscapeMaterialsOverride.Num() || bIsPerformingInteractiveActionOnLandscapeMaterialOverride)
 			{
 				HasMaterialChanged = true;
 			}
@@ -4386,7 +4384,6 @@ void ALandscape::PostEditChangeProperty(FPropertyChangedEvent& PropertyChangedEv
 	}
 
 	PreEditLandscapeMaterial = nullptr;
-	PreEditLandscapeHoleMaterial = nullptr;
 	PreEditLandscapeMaterialsOverride.Empty();
 }
 
diff --git a/Engine/Source/Runtime/Landscape/Private/LandscapeRender.cpp b/Engine/Source/Runtime/Landscape/Private/LandscapeRender.cpp
--- a/Engine/Source/Runtime/Landscape/Private/LandscapeRender.cpp
+++ b/Engine/Source/Runtime/Landscape/Private/LandscapeRender.cpp
@@ -571,11 +571,6 @@ void ULandscapeComponent::GetUsedMaterials(TArray<UMaterialInterface*>& OutMater
 		OutMaterials.Add(OverrideMaterial);
 	}
 
-	if (OverrideHoleMaterial)
-	{
-		OutMaterials.Add(OverrideHoleMaterial);
-	}
-
 	OutMaterials.Append(MobileMaterialInterfaces);
 
 #if WITH_EDITORONLY_DATA
