When we have no pin name specified, init to empty so all the condition for PinFriendlyName.IsNone() will return true, which mean that when you have a pin name of None, it will get init automatically to Input/Output and we will also display this name in the material editor instead of nothing as the PinFriendlyName was initialized to " " before.

Prevent having an input pin named None for node LandscapeGrassOutput

diff --git a/Engine/Source/Editor/UnrealEd/Private/MaterialGraphNode.cpp b/Engine/Source/Editor/UnrealEd/Private/MaterialGraphNode.cpp
--- a/Engine/Source/Editor/UnrealEd/Private/MaterialGraphNode.cpp
+++ b/Engine/Source/Editor/UnrealEd/Private/MaterialGraphNode.cpp
@@ -553,7 +553,7 @@ void UMaterialGraphNode::CreateInputPins()
 		{
 			// Makes sure pin has a name for lookup purposes but user will never see it
 			NewPin->PinName = CreateUniquePinName(TEXT("Input"));
-			NewPin->PinFriendlyName = FText::FromString(TEXT(" "));
+			NewPin->PinFriendlyName = FText::GetEmpty();
 		}
 	}
 }
@@ -607,7 +607,7 @@ void UMaterialGraphNode::CreateOutputPins()
 		{
 			// Makes sure pin has a name for lookup purposes but user will never see it
 			NewPin->PinName = CreateUniquePinName(TEXT("Output"));
-			NewPin->PinFriendlyName = FText::FromString(TEXT(" "));
+			NewPin->PinFriendlyName = FText::GetEmpty();
 		}
 	}
 }
diff --git a/Engine/Source/Runtime/Landscape/Classes/Materials/MaterialExpressionLandscapeGrassOutput.h b/Engine/Source/Runtime/Landscape/Classes/Materials/MaterialExpressionLandscapeGrassOutput.h
--- a/Engine/Source/Runtime/Landscape/Classes/Materials/MaterialExpressionLandscapeGrassOutput.h
+++ b/Engine/Source/Runtime/Landscape/Classes/Materials/MaterialExpressionLandscapeGrassOutput.h
@@ -48,6 +48,7 @@ class LANDSCAPE_API UMaterialExpressionLandscapeGrassOutput : public UMaterialEx
 	virtual const TArray<FExpressionInput*> GetInputs() override;
 	virtual FExpressionInput* GetInput(int32 InputIndex) override;
 	virtual FName GetInputName(int32 InputIndex) const override;
+	void ValidateInputName(FGrassInput& Input) const;
 #endif
 
 	//~ Begin UObject Interface
@@ -64,6 +65,9 @@ class LANDSCAPE_API UMaterialExpressionLandscapeGrassOutput : public UMaterialEx
 
 	UPROPERTY(EditAnywhere, Category = UMaterialExpressionLandscapeGrassOutput)
 	TArray<FGrassInput> GrassTypes;
+
+private:
+	static FName PinDefaultName;
 };
 
 
diff --git a/Engine/Source/Runtime/Landscape/Private/LandscapeGrass.cpp b/Engine/Source/Runtime/Landscape/Private/LandscapeGrass.cpp
--- a/Engine/Source/Runtime/Landscape/Private/LandscapeGrass.cpp
+++ b/Engine/Source/Runtime/Landscape/Private/LandscapeGrass.cpp
@@ -1042,6 +1042,8 @@ public:
 //
 // UMaterialExpressionLandscapeGrassOutput
 //
+FName UMaterialExpressionLandscapeGrassOutput::PinDefaultName = TEXT("Pin");
+
 UMaterialExpressionLandscapeGrassOutput::UMaterialExpressionLandscapeGrassOutput(const FObjectInitializer& ObjectInitializer)
 : Super(ObjectInitializer)
 {
@@ -1128,6 +1130,11 @@ void UMaterialExpressionLandscapeGrassOutput::PostEditChangeProperty(FPropertyCh
 		const FName PropertyName = PropertyChangedEvent.MemberProperty->GetFName();
 		if (PropertyName == GET_MEMBER_NAME_CHECKED(UMaterialExpressionLandscapeGrassOutput, GrassTypes))
 		{
+			for (FGrassInput& Input : GrassTypes)
+			{
+				ValidateInputName(Input);
+			}
+
 			if (GraphNode)
 			{
 				GraphNode->ReconstructNode();
@@ -1135,6 +1142,46 @@ void UMaterialExpressionLandscapeGrassOutput::PostEditChangeProperty(FPropertyCh
 		}
 	}
 }
+
+void UMaterialExpressionLandscapeGrassOutput::ValidateInputName(FGrassInput& InInput) const
+{
+	if (Material != nullptr)
+	{
+		int32 NameIndex = 1;
+		bool bFoundValidName = false;
+
+		// Parameters cannot be named Name_None, use the default name instead
+		FName PotentialName = InInput.Name == NAME_None ? UMaterialExpressionLandscapeGrassOutput::PinDefaultName : InInput.Name;
+
+		// Find an available unique name
+		while (!bFoundValidName)
+		{
+			if (NameIndex != 1)
+			{
+				PotentialName.SetNumber(NameIndex);
+			}
+
+			bFoundValidName = true;
+
+			// Make sure the name is unique among others pins of this node
+			for (const FGrassInput& OtherInput : GrassTypes)
+			{
+				if (&OtherInput != &InInput)
+				{
+					if (OtherInput.Name == PotentialName)
+					{
+						bFoundValidName = false;
+						break;
+					}
+				}
+			}
+
+			++NameIndex;
+		}
+
+		InInput.Name = PotentialName;
+	}
+}
 #endif
 
 //
