PR #5537: Update LandscapeProxy C++ function access (Contributed by [REDACTED_PERSON])

diff --git a/Engine/Source/Runtime/Landscape/Classes/LandscapeProxy.h b/Engine/Source/Runtime/Landscape/Classes/LandscapeProxy.h
--- a/Engine/Source/Runtime/Landscape/Classes/LandscapeProxy.h
+++ b/Engine/Source/Runtime/Landscape/Classes/LandscapeProxy.h
@@ -729,15 +729,15 @@ public:
 
 	/** Set an MID texture parameter value for all landscape components. */
 	UFUNCTION(BlueprintCallable, Category = "Landscape|Runtime|Material")
-	void SetLandscapeMaterialTextureParameterValue(FName ParameterName, class UTexture* Value);
+	LANDSCAPE_API void SetLandscapeMaterialTextureParameterValue(FName ParameterName, class UTexture* Value);
 
 	/** Set an MID vector parameter value for all landscape components. */
 	UFUNCTION(BlueprintCallable, meta = (Keywords = "SetColorParameterValue"), Category = "Landscape|Runtime|Material")
-	void SetLandscapeMaterialVectorParameterValue(FName ParameterName, FLinearColor Value);
+	LANDSCAPE_API void SetLandscapeMaterialVectorParameterValue(FName ParameterName, FLinearColor Value);
 
 	/** Set a MID scalar (float) parameter value for all landscape components. */
 	UFUNCTION(BlueprintCallable, meta = (Keywords = "SetFloatParameterValue"), Category = "Landscape|Runtime|Material")
-	void SetLandscapeMaterialScalarParameterValue(FName ParameterName, float Value);
+	LANDSCAPE_API void SetLandscapeMaterialScalarParameterValue(FName ParameterName, float Value);
 
 	// End blueprint functions
 
