Fix linux compile warning in Landscape.cpp
#rnx

diff --git a/Engine/Source/Runtime/Landscape/Private/Landscape.cpp b/Engine/Source/Runtime/Landscape/Private/Landscape.cpp
--- a/Engine/Source/Runtime/Landscape/Private/Landscape.cpp
+++ b/Engine/Source/Runtime/Landscape/Private/Landscape.cpp
@@ -978,7 +978,7 @@ void ULandscapeComponent::PostLoad()
 		UWorld* World = GetWorld();
 
 		// If we're loading on a platform that doesn't require cooked data, but *only* supports OpenGL ES, generate or preload data from the DDC
-		if (!FPlatformProperties::RequiresCookedData() && (GMaxRHIFeatureLevel <= ERHIFeatureLevel::ES3_1 || World && World->FeatureLevel <= ERHIFeatureLevel::ES3_1))
+		if (!FPlatformProperties::RequiresCookedData() && (GMaxRHIFeatureLevel <= ERHIFeatureLevel::ES3_1 || World && (World->FeatureLevel <= ERHIFeatureLevel::ES3_1)))
 		{
 			CheckGenerateLandscapePlatformData(false, nullptr);
 		}
