Allow user specifying at what Screen Size should LOD0 be, by default it's 1.0, and also removed the check that ensure the screen size was between 0-1

diff --git a/Engine/Source/Runtime/Landscape/Classes/LandscapeProxy.h b/Engine/Source/Runtime/Landscape/Classes/LandscapeProxy.h
--- a/Engine/Source/Runtime/Landscape/Classes/LandscapeProxy.h
+++ b/Engine/Source/Runtime/Landscape/Classes/LandscapeProxy.h
@@ -421,6 +421,10 @@ public:
 	UPROPERTY(EditAnywhere, Category = LOD, meta=(ClampMin = "0.01", ClampMax = "1.0", UIMin = "0.01", UIMax = "1.0", DisplayName= "SubSection Min Component ScreenSize"))
 	float ComponentScreenSizeToUseSubSections;
 
+	/** This is the starting screen size used to calculate the distribution, by default it's 1, but you can increase the value if you want less LOD0 component, and you use very large landscape component. */
+	UPROPERTY(EditAnywhere, Category = "LOD Distribution", meta = (DisplayName = "LOD 0 Screen Size", ClampMin = "1.0", ClampMax = "10.0", UIMin = "1.0", UIMax = "10.0"))
+	float LOD0ScreenSize;
+
 	/** The distribution setting used to change the LOD 0 generation, 1.75 is the normal distribution, numbers influence directly the LOD0 proportion on screen. */
 	UPROPERTY(EditAnywhere, Category = "LOD Distribution", meta = (DisplayName = "LOD 0", ClampMin = "1.0", ClampMax = "10.0", UIMin = "1.0", UIMax = "10.0"))
 	float LOD0DistributionSetting;
diff --git a/Engine/Source/Runtime/Landscape/Private/Landscape.cpp b/Engine/Source/Runtime/Landscape/Private/Landscape.cpp
--- a/Engine/Source/Runtime/Landscape/Private/Landscape.cpp
+++ b/Engine/Source/Runtime/Landscape/Private/Landscape.cpp
@@ -1030,6 +1030,7 @@ ALandscapeProxy::ALandscapeProxy(const FObjectInitializer& ObjectInitializer)
 	ComponentScreenSizeToUseSubSections = 0.65f;
 	UseTessellationComponentScreenSizeFalloff = true;
 	TessellationComponentScreenSizeFalloff = 0.75f;
+	LOD0ScreenSize = 1.0f;
 	LOD0DistributionSetting = 1.75f;
 	LODDistributionSetting = 2.0f;
 	bCastStaticShadow = true;
@@ -2386,6 +2387,7 @@ void ALandscapeProxy::GetSharedProperties(ALandscapeProxy* Landscape)
 		TessellationComponentScreenSizeFalloff = Landscape->TessellationComponentScreenSizeFalloff;
 		LODDistributionSetting = Landscape->LODDistributionSetting;
 		LOD0DistributionSetting = Landscape->LOD0DistributionSetting;
+		LOD0ScreenSize = Landscape->LOD0ScreenSize;
 		OccluderGeometryLOD = Landscape->OccluderGeometryLOD;
 		NegativeZBoundsExtension = Landscape->NegativeZBoundsExtension;
 		PositiveZBoundsExtension = Landscape->PositiveZBoundsExtension;
@@ -2459,6 +2461,12 @@ void ALandscapeProxy::FixupSharedData(ALandscape* Landscape)
 		bUpdated = true;
 	}
 
+	if (LOD0ScreenSize != Landscape->LOD0ScreenSize)
+	{
+		LOD0ScreenSize = Landscape->LOD0ScreenSize;
+		bUpdated = true;
+	}
+
 	if (OccluderGeometryLOD != Landscape->OccluderGeometryLOD)
 	{
 		OccluderGeometryLOD = Landscape->OccluderGeometryLOD;
diff --git a/Engine/Source/Runtime/Landscape/Private/LandscapeEdit.cpp b/Engine/Source/Runtime/Landscape/Private/LandscapeEdit.cpp
--- a/Engine/Source/Runtime/Landscape/Private/LandscapeEdit.cpp
+++ b/Engine/Source/Runtime/Landscape/Private/LandscapeEdit.cpp
@@ -3995,7 +3995,8 @@ void ALandscapeProxy::PostEditChangeProperty(FPropertyChangedEvent& PropertyChan
 		ChangeTessellationComponentScreenSizeFalloff(TessellationComponentScreenSizeFalloff);
 	}
 	else if (PropertyName == GET_MEMBER_NAME_CHECKED(ALandscapeProxy, LODDistributionSetting)
-		|| PropertyName == GET_MEMBER_NAME_CHECKED(ALandscapeProxy, LOD0DistributionSetting))
+		|| PropertyName == GET_MEMBER_NAME_CHECKED(ALandscapeProxy, LOD0DistributionSetting)
+		|| PropertyName == GET_MEMBER_NAME_CHECKED(ALandscapeProxy, LOD0ScreenSize))
 	{		
 		MarkComponentsRenderStateDirty();
 	}
@@ -4264,6 +4265,11 @@ void ALandscape::PostEditChangeProperty(FPropertyChangedEvent& PropertyChangedEv
 		LOD0DistributionSetting = FMath::Clamp<float>(LOD0DistributionSetting, 1.0f, 10.0f);
 		bPropagateToProxies = true;
 	}
+	else if (PropertyName == FName(TEXT("LOD0ScreenSize")))
+	{
+		LOD0ScreenSize = FMath::Clamp<float>(LOD0ScreenSize, 1.0f, 10.0f);
+		bPropagateToProxies = true;
+	}
 	else if (PropertyName == FName(TEXT("CollisionMipLevel")))
 	{
 		CollisionMipLevel = FMath::Clamp<int32>(CollisionMipLevel, 0, FMath::CeilLogTwo(SubsectionSizeQuads + 1) - 1);
diff --git a/Engine/Source/Runtime/Landscape/Private/LandscapeRender.cpp b/Engine/Source/Runtime/Landscape/Private/LandscapeRender.cpp
--- a/Engine/Source/Runtime/Landscape/Private/LandscapeRender.cpp
+++ b/Engine/Source/Runtime/Landscape/Private/LandscapeRender.cpp
@@ -704,7 +704,7 @@ FLandscapeComponentSceneProxy::FLandscapeComponentSceneProxy(ULandscapeComponent
 	}
 
 	float ScreenSizeRatioDivider = FMath::Max(InComponent->GetLandscapeProxy()->LOD0DistributionSetting * GLandscapeLOD0DistributionScale, 1.01f);
-	float CurrentScreenSizeRatio = 1.0f;
+	float CurrentScreenSizeRatio = InComponent->GetLandscapeProxy()->LOD0ScreenSize;
 
 	LODScreenRatioSquared.AddUninitialized(MaxLOD + 1);
 
@@ -1288,7 +1288,7 @@ float FLandscapeComponentSceneProxy::GetComponentScreenSize(const FSceneView* Vi
 	// Calculate screen-space projected radius
 	float SquaredScreenRadius = FMath::Square(ScreenMultiple * ElementRadius) / FMath::Max(1.0f, DistSquared);
 
-	return FMath::Min(SquaredScreenRadius * 2.0f, 1.0f);
+	return SquaredScreenRadius * 2.0f;
 }
 
 void FLandscapeComponentSceneProxy::BuildDynamicMeshElement(const FViewCustomDataLOD* InPrimitiveCustomData, bool InToolMesh, bool InHasTessellation, bool InDisableTessellation, FMeshBatch& OutMeshBatch, TArray<FLandscapeBatchElementParams, SceneRenderingAllocator>& OutStaticBatchParamArray) const
@@ -1796,7 +1796,7 @@ void FLandscapeComponentSceneProxy::CalculateBatchElementLOD(const FSceneView& I
 {
 	float SquaredViewLODScale = FMath::Square(InViewLODScale);
 
-	check(InMeshScreenSizeSquared >= 0.0f && InMeshScreenSizeSquared <= 1.0f);
+	check(InMeshScreenSizeSquared >= 0.0f);
 	float ComponentScreenSize = InMeshScreenSizeSquared;
 
 	if (NumSubsections > 1)
@@ -1900,7 +1900,7 @@ float FLandscapeComponentSceneProxy::ComputeBatchElementCurrentLOD(int32 InSelec
 	float NextLODScreenRatio = LastElement ? 0 : LODScreenRatioSquared[InSelectedLODIndex + 1];
 
 	float LODScreenRatioRange = CurrentLODScreenRatio - NextLODScreenRatio;
-	float ScreenSizeWithLODScale = InComponentScreenSize / InViewLODScale;
+	float ScreenSizeWithLODScale = FMath::Clamp(InComponentScreenSize / InViewLODScale, 0.0f, LODScreenRatioSquared[0]);
 
 	if (ScreenSizeWithLODScale > CurrentLODScreenRatio || ScreenSizeWithLODScale < NextLODScreenRatio)
 	{
@@ -1961,7 +1961,6 @@ void* FLandscapeComponentSceneProxy::InitViewCustomData(const FSceneView& InView
 
 	FViewCustomDataLOD* LODData = (FViewCustomDataLOD*)new(InCustomDataMemStack) FViewCustomDataLOD();
 
-	check(InMeshScreenSizeSquared <= 1.0f);
 	LODData->ComponentScreenSize = InMeshScreenSizeSquared;
 
 	// If a valid screen size was provided, we use it instead of recomputing it
diff --git a/Engine/Source/Runtime/Landscape/Public/LandscapeRender.h b/Engine/Source/Runtime/Landscape/Public/LandscapeRender.h
--- a/Engine/Source/Runtime/Landscape/Public/LandscapeRender.h
+++ b/Engine/Source/Runtime/Landscape/Public/LandscapeRender.h
@@ -654,7 +654,7 @@ protected:
 	void CalculateLODFromScreenSize(const FSceneView& InView, float InMeshScreenSizeSquared, float InViewLODScale, int32 InSubSectionIndex, FViewCustomDataLOD& InOutLODData) const;
 	FORCEINLINE void ComputeStaticBatchIndexToRender(FViewCustomDataLOD& OutLODData, int32 InSubSectionIndex);
 	int8 GetLODFromScreenSize(float InScreenSizeSquared, float InViewLODScale) const;
-	FORCEINLINE float ComputeBatchElementCurrentLOD(int32 InSelectedLODIndex, float InComponentScreenSize, float InViewLODScale) const;
+	FORCEINLINE_DEBUGGABLE float ComputeBatchElementCurrentLOD(int32 InSelectedLODIndex, float InComponentScreenSize, float InViewLODScale) const;
 	
 	FORCEINLINE void GetShaderCurrentNeighborLOD(const FSceneView& InView, float InBatchElementCurrentLOD, int8 InSubSectionX, int8 InSubSectionY, int8 InCurrentSubSectionIndex, FVector4& OutShaderCurrentNeighborLOD) const;
 	FORCEINLINE FVector4 GetShaderLODBias() const;
