Wrap Channel and Texture index + Prevent crash caused by incoming CL in QA Landscape map when switching Landscape to Edit Layers

#rnx

diff --git a/Engine/Source/Runtime/Landscape/Classes/LandscapeComponent.h b/Engine/Source/Runtime/Landscape/Classes/LandscapeComponent.h
--- a/Engine/Source/Runtime/Landscape/Classes/LandscapeComponent.h
+++ b/Engine/Source/Runtime/Landscape/Classes/LandscapeComponent.h
@@ -205,6 +205,14 @@ struct FWeightmapLayerAllocationInfo
 	FName GetLayerName() const;
 
 	uint32 GetHash() const;
+
+	void Free()
+	{
+		WeightmapTextureChannel = 255;
+		WeightmapTextureIndex = 255;
+	}
+
+	bool IsAllocated() const { return (WeightmapTextureChannel != 255 && WeightmapTextureIndex != 255); }
 };
 
 struct FLandscapeComponentGrassData
diff --git a/Engine/Source/Runtime/Landscape/Private/LandscapeEdit.cpp b/Engine/Source/Runtime/Landscape/Private/LandscapeEdit.cpp
--- a/Engine/Source/Runtime/Landscape/Private/LandscapeEdit.cpp
+++ b/Engine/Source/Runtime/Landscape/Private/LandscapeEdit.cpp
@@ -5229,7 +5229,7 @@ void ULandscapeComponent::ReallocateWeightmaps(FLandscapeEditDataInterface* Data
 	{
 		for (int32 LayerIdx = 0; LayerIdx < ComponentWeightmapLayerAllocations.Num(); LayerIdx++)
 		{
-			if (ComponentWeightmapLayerAllocations[LayerIdx].WeightmapTextureIndex == 255)
+			if (!ComponentWeightmapLayerAllocations[LayerIdx].IsAllocated())
 			{
 				NeededNewChannels++;
 			}
@@ -5288,14 +5288,14 @@ void ULandscapeComponent::ReallocateWeightmaps(FLandscapeEditDataInterface* Data
 						{
 							FWeightmapLayerAllocationInfo& AllocInfo = ComponentWeightmapLayerAllocations[CurrentAlloc];
 
-							if (AllocInfo.WeightmapTextureIndex == 255)
+							if (!AllocInfo.IsAllocated())
 							{
 								break;
 							}
 						}
 
 						FWeightmapLayerAllocationInfo& AllocInfo = ComponentWeightmapLayerAllocations[CurrentAlloc];
-						check(AllocInfo.WeightmapTextureIndex == 255);
+						check(!AllocInfo.IsAllocated());
 
 						// Zero out the data for this texture channel
 						if (DataInterface)
@@ -5433,7 +5433,7 @@ void ULandscapeComponent::ReallocateWeightmaps(FLandscapeEditDataInterface* Data
 				// Use this allocation
 				FWeightmapLayerAllocationInfo& AllocInfo = ComponentWeightmapLayerAllocations[CurrentLayer];
 
-				if (AllocInfo.WeightmapTextureIndex == 255)
+				if (!AllocInfo.IsAllocated())
 				{
 					// New layer - zero out the data for this texture channel
 					if (DataInterface)
diff --git a/Engine/Source/Runtime/Landscape/Private/LandscapeEditInterface.cpp b/Engine/Source/Runtime/Landscape/Private/LandscapeEditInterface.cpp
--- a/Engine/Source/Runtime/Landscape/Private/LandscapeEditInterface.cpp
+++ b/Engine/Source/Runtime/Landscape/Private/LandscapeEditInterface.cpp
@@ -1910,11 +1910,11 @@ void ULandscapeComponent::FillLayer(ULandscapeLayerInfoObject* LayerInfo, FLands
 				Usage->ChannelUsage[Allocation.WeightmapTextureChannel] = nullptr;
 			}
 
-			Allocation.WeightmapTextureIndex = 255;
+			Allocation.Free();
 		}
 
 		ComponentWeightmapLayerAllocations.RemoveAll(
-			[](const FWeightmapLayerAllocationInfo& Allocation) { return Allocation.WeightmapTextureIndex == 255; });
+			[](const FWeightmapLayerAllocationInfo& Allocation) { return !Allocation.IsAllocated(); });
 
 		// remove any textures we're no longer using
 		for (int32 TextureIdx = 0; TextureIdx < ComponentWeightmapTextures.Num(); ++TextureIdx)
diff --git a/Engine/Source/Runtime/Landscape/Private/LandscapeEditLayers.cpp b/Engine/Source/Runtime/Landscape/Private/LandscapeEditLayers.cpp
--- a/Engine/Source/Runtime/Landscape/Private/LandscapeEditLayers.cpp
+++ b/Engine/Source/Runtime/Landscape/Private/LandscapeEditLayers.cpp
@@ -1842,8 +1842,7 @@ void ALandscape::CopyOldDataToDefaultLayer(ALandscapeProxy* InProxy)
 
 		for (FWeightmapLayerAllocationInfo& Allocation : ComponentLayerAllocations)
 		{
-			Allocation.WeightmapTextureChannel = 255;
-			Allocation.WeightmapTextureIndex = 255;
+			Allocation.Free();
 		}
 	}
 }
@@ -3343,7 +3342,7 @@ void ALandscape::PrepareComponentDataToExtractMaterialLayersCS(const TArray<ULan
 
 				for (const FWeightmapLayerAllocationInfo& WeightmapLayerAllocation : ComponentLayerData->WeightmapData.LayerAllocations)
 				{
-					if (WeightmapLayerAllocation.LayerInfo != nullptr && WeightmapLayerAllocation.WeightmapTextureIndex != 255 && ComponentLayerData->WeightmapData.Textures[WeightmapLayerAllocation.WeightmapTextureIndex] == LayerWeightmap)
+					if (WeightmapLayerAllocation.LayerInfo != nullptr && WeightmapLayerAllocation.IsAllocated() && ComponentLayerData->WeightmapData.Textures[WeightmapLayerAllocation.WeightmapTextureIndex] == LayerWeightmap)
 					{
 						FLandscapeLayerWeightmapExtractMaterialLayersComponentData Data;
 
@@ -3518,8 +3517,7 @@ void ALandscape::ReallocateLayersWeightmaps(const TArray<ULandscapeComponent*>&
 
 		for (FWeightmapLayerAllocationInfo& BaseWeightmapAllocation : BaseLayerAllocations)
 		{
-			BaseWeightmapAllocation.WeightmapTextureChannel = 255;
-			BaseWeightmapAllocation.WeightmapTextureIndex = 255;
+			BaseWeightmapAllocation.Free();
 		}
 
 		TArray<ULandscapeWeightmapUsage*>& WeightmapTexturesUsage = Component->GetWeightmapTexturesUsage();
@@ -4449,7 +4447,7 @@ void ALandscape::ResolveLayersWeightmapTexture(const TArray<ULandscapeComponent*
 
 		for (const FWeightmapLayerAllocationInfo& AllocInfo : AllocInfos)
 		{
-			check(AllocInfo.WeightmapTextureIndex != 255 && AllocInfo.WeightmapTextureIndex < ComponentWeightmaps.Num());
+			check(AllocInfo.IsAllocated() && AllocInfo.WeightmapTextureIndex < ComponentWeightmaps.Num());
 			UTexture2D* Weightmap = ComponentWeightmaps[AllocInfo.WeightmapTextureIndex];
 			WeightmapsToResolve.FindOrAdd(Weightmap).Add(TPair<ULandscapeComponent*,const FWeightmapLayerAllocationInfo*>(Component, &AllocInfo));
 		}
@@ -5960,7 +5958,7 @@ uint32 ULandscapeComponent::ComputeLayerHash() const
 
 	for (const FWeightmapLayerAllocationInfo& AllocationInfo : AllocationInfos)
 	{
-		if (AllocationInfo.WeightmapTextureChannel != 255 && AllocationInfo.WeightmapTextureIndex != 255)
+		if (AllocationInfo.IsAllocated())
 		{
             // Compute hash of actual data of the texture that is owned by the component (per Texture Channel)
 			UTexture2D* Weightmap = Weightmaps[AllocationInfo.WeightmapTextureIndex];
diff --git a/Engine/Source/Runtime/Landscape/Private/LandscapeRender.cpp b/Engine/Source/Runtime/Landscape/Private/LandscapeRender.cpp
--- a/Engine/Source/Runtime/Landscape/Private/LandscapeRender.cpp
+++ b/Engine/Source/Runtime/Landscape/Private/LandscapeRender.cpp
@@ -1361,7 +1361,7 @@ FLandscapeComponentSceneProxy::FLandscapeComponentSceneProxy(ULandscapeComponent
 	for (int32 Idx = 0; Idx < InComponent->WeightmapLayerAllocations.Num(); Idx++)
 	{
 		FWeightmapLayerAllocationInfo& Allocation = InComponent->WeightmapLayerAllocations[Idx];
-		if (Allocation.LayerInfo == ALandscapeProxy::VisibilityLayer)
+		if (Allocation.LayerInfo == ALandscapeProxy::VisibilityLayer && Allocation.IsAllocated())
 		{
 			VisibilityWeightmapTexture = WeightmapTextures[Allocation.WeightmapTextureIndex];
 			VisibilityWeightmapChannel = Allocation.WeightmapTextureChannel;
