Remove FlushRenderingCommands
- When uploading TextureEdits to GPU (double buffer)
- When invalidating lighting cache. (dependency has already been removed a while ago with the MapBuildData changes)

Move InvalidateLightingCache calls to the resolve (like all the other client code ex: Coll, Foliage, etc)

#rnx

diff --git a/Engine/Source/Runtime/Landscape/Classes/LandscapeProxy.h b/Engine/Source/Runtime/Landscape/Classes/LandscapeProxy.h
--- a/Engine/Source/Runtime/Landscape/Classes/LandscapeProxy.h
+++ b/Engine/Source/Runtime/Landscape/Classes/LandscapeProxy.h
@@ -838,11 +838,11 @@ public:
 	TArray<ULandscapeGrassType*> GetGrassTypes() const;
 
 	/* Invalidate the precomputed grass and baked texture data for the specified components */
-	LANDSCAPE_API static void InvalidateGeneratedComponentData(const TSet<ULandscapeComponent*>& Components);
-	LANDSCAPE_API static void InvalidateGeneratedComponentData(const TArray<ULandscapeComponent*>& Components);
+	LANDSCAPE_API static void InvalidateGeneratedComponentData(const TSet<ULandscapeComponent*>& Components, bool bInvalidateLightingCache = false);
+	LANDSCAPE_API static void InvalidateGeneratedComponentData(const TArray<ULandscapeComponent*>& Components, bool bInvalidateLightingCache = false);
 
 	/* Invalidate the precomputed grass and baked texture data on all components */
-	LANDSCAPE_API void InvalidateGeneratedComponentData();
+	LANDSCAPE_API void InvalidateGeneratedComponentData(bool bInvalidateLightingCache = false);
 
 #if WITH_EDITOR
 	/** Update Grass maps */
diff --git a/Engine/Source/Runtime/Landscape/Private/Landscape.cpp b/Engine/Source/Runtime/Landscape/Private/Landscape.cpp
--- a/Engine/Source/Runtime/Landscape/Private/Landscape.cpp
+++ b/Engine/Source/Runtime/Landscape/Private/Landscape.cpp
@@ -3757,12 +3757,16 @@ void ALandscapeProxy::UpdateBakedTextures()
 #endif
 
 template<class ContainerType>
-void InvalidateGeneratedComponentDataImpl(const ContainerType& Components)
+void InvalidateGeneratedComponentDataImpl(const ContainerType& Components, bool bInvalidateLightingCache)
 {
 	TMap<ALandscapeProxy*, TSet<ULandscapeComponent*>> ByProxy;
 	for (auto Iter = Components.CreateConstIterator(); Iter; ++Iter)
 	{
 		ULandscapeComponent* Component = *Iter;
+		if (bInvalidateLightingCache)
+		{
+			Component->InvalidateLightingCache();
+		}
 		Component->BakedTextureMaterialGuid.Invalidate();
 		ByProxy.FindOrAdd(Component->GetLandscapeProxy()).Add(Component);
 	}
@@ -3773,19 +3777,19 @@ void InvalidateGeneratedComponentDataImpl(const ContainerType& Components)
 	}
 }
 
-void ALandscapeProxy::InvalidateGeneratedComponentData()
+void ALandscapeProxy::InvalidateGeneratedComponentData(bool bInvalidateLightingCache)
 {
-	InvalidateGeneratedComponentDataImpl(LandscapeComponents);
+	InvalidateGeneratedComponentDataImpl(LandscapeComponents, bInvalidateLightingCache);
 }
 
-void ALandscapeProxy::InvalidateGeneratedComponentData(const TArray<ULandscapeComponent*>& Components)
+void ALandscapeProxy::InvalidateGeneratedComponentData(const TArray<ULandscapeComponent*>& Components, bool bInvalidateLightingCache)
 {
-	InvalidateGeneratedComponentDataImpl(Components);
+	InvalidateGeneratedComponentDataImpl(Components, bInvalidateLightingCache);
 }
 
-void ALandscapeProxy::InvalidateGeneratedComponentData(const TSet<ULandscapeComponent*>& Components)
+void ALandscapeProxy::InvalidateGeneratedComponentData(const TSet<ULandscapeComponent*>& Components, bool bInvalidateLightingCache)
 {
-	InvalidateGeneratedComponentDataImpl(Components);
+	InvalidateGeneratedComponentDataImpl(Components, bInvalidateLightingCache);
 }
 
 #undef LOCTEXT_NAMESPACE
diff --git a/Engine/Source/Runtime/Landscape/Private/LandscapeEditInterface.cpp b/Engine/Source/Runtime/Landscape/Private/LandscapeEditInterface.cpp
--- a/Engine/Source/Runtime/Landscape/Private/LandscapeEditInterface.cpp
+++ b/Engine/Source/Runtime/Landscape/Private/LandscapeEditInterface.cpp
@@ -1964,11 +1964,11 @@ void ULandscapeComponent::FillLayer(ULandscapeLayerInfoObject* LayerInfo, FLands
 	// Update the shaders for this component
 	if (!Component->GetLandscapeProxy()->HasLayersContent())
 	{
+		Component->InvalidateLightingCache();
 		Component->UpdateMaterialInstances();
 		Component->EditToolRenderData.UpdateDebugColorMaterial(Component);
 		Component->UpdateEditToolRenderData();
 	}
-	Component->InvalidateLightingCache();
 	Component->RequestWeightmapUpdate();
 
 	// Update dominant layer info stored in collision component
@@ -5677,16 +5677,32 @@ bool FLandscapeTextureDataInfo::UpdateTextureData()
 	{
 		if (MipInfo[i].MipData && MipInfo[i].MipUpdateRegions.Num() > 0)
 		{
-			bNeedToWaitForUpdate = true;
 			if (bCompressed)
 			{
+				bNeedToWaitForUpdate = true;
 				// Cannot update regions on compressed textures so we will update the whole texture below.
 				break;
 			}
-			Texture->UpdateTextureRegions(i, MipInfo[i].MipUpdateRegions.Num(), &MipInfo[i].MipUpdateRegions[0], ((Texture->Source.GetSizeX()) >> i)*DataSize, DataSize, (uint8*)MipInfo[i].MipData);
+
+			const uint32 SrcSize = (Texture->Source.GetSizeX()) >> i;
+			const uint32 SrcPitch = (SrcSize * DataSize);
+			const uint32 BufferSize = SrcSize * SrcSize * DataSize;
+
+			// Copy Mip update data so we can avoid waiting for Render thread in calling method
+			FMipInfo* CopyMipInfo = new FMipInfo();
+			CopyMipInfo->MipUpdateRegions = MipInfo[i].MipUpdateRegions;
+			CopyMipInfo->MipData = FMemory::Malloc(BufferSize);
+			FMemory::Memcpy(CopyMipInfo->MipData, MipInfo[i].MipData, BufferSize);
+
+			Texture->UpdateTextureRegions(i, CopyMipInfo->MipUpdateRegions.Num(), &CopyMipInfo->MipUpdateRegions[0], SrcPitch, DataSize, (uint8*)CopyMipInfo->MipData, 
+				[CopyMipInfo](uint8* SrcData, const FUpdateTextureRegion2D*) 
+				{ 
+					FMemory::Free(CopyMipInfo->MipData);
+					delete CopyMipInfo; 
+				});
 		}
 	}
-	if (bCompressed && bNeedToWaitForUpdate)
+	if (bNeedToWaitForUpdate)
 	{
 		Texture->UpdateResource();
 	}
diff --git a/Engine/Source/Runtime/Landscape/Private/LandscapeEditLayers.cpp b/Engine/Source/Runtime/Landscape/Private/LandscapeEditLayers.cpp
--- a/Engine/Source/Runtime/Landscape/Private/LandscapeEditLayers.cpp
+++ b/Engine/Source/Runtime/Landscape/Private/LandscapeEditLayers.cpp
@@ -3150,7 +3150,8 @@ void ALandscape::ResolveLayersHeightmapTexture(const TArray<ULandscapeComponent*
 		}
 	}
 		
-	InvalidateGeneratedComponentData(ChangedComponents);
+	const bool bInvalidateLightingCache = true;
+	InvalidateGeneratedComponentData(ChangedComponents, bInvalidateLightingCache);
 }
 
 void ALandscape::ClearDirtyData(ULandscapeComponent* InLandscapeComponent)
@@ -4500,7 +4501,9 @@ void ALandscape::ResolveLayersWeightmapTexture(const TArray<ULandscapeComponent*
 		}
 	}
 		
-	InvalidateGeneratedComponentData(ChangedComponents.Array());	
+	// Weightmaps shouldn't invalidate lighting
+	const bool bInvalidateLightingCache = false;
+	InvalidateGeneratedComponentData(ChangedComponents, bInvalidateLightingCache);	
 }
 
 bool ALandscape::HasLayersContent() const
diff --git a/Engine/Source/Runtime/Landscape/Private/LandscapeLight.cpp b/Engine/Source/Runtime/Landscape/Private/LandscapeLight.cpp
--- a/Engine/Source/Runtime/Landscape/Private/LandscapeLight.cpp
+++ b/Engine/Source/Runtime/Landscape/Private/LandscapeLight.cpp
@@ -753,9 +753,6 @@ void ULandscapeComponent::InvalidateLightingCacheDetailed(bool bInvalidateBuildE
 {
 	Modify();
 
-	// Block until the RT processes the unregister before modifying variables that it may need to access
-	FlushRenderingCommands();
-
 	FComponentReregisterContext ReregisterContext(this);
 
 	Super::InvalidateLightingCacheDetailed(bInvalidateBuildEnqueuedLighting, bTranslationOnly);
diff --git a/Engine/Source/Runtime/Landscape/Public/LandscapeEdit.h b/Engine/Source/Runtime/Landscape/Public/LandscapeEdit.h
--- a/Engine/Source/Runtime/Landscape/Public/LandscapeEdit.h
+++ b/Engine/Source/Runtime/Landscape/Public/LandscapeEdit.h
@@ -60,6 +60,7 @@ struct FLandscapeTextureDataInfo
 	{
 		void* MipData;
 		TArray<FUpdateTextureRegion2D> MipUpdateRegions;
+		bool bFull;
 	};
 
 	FLandscapeTextureDataInfo(UTexture2D* InTexture, bool bShouldDirtyPackage);
@@ -72,10 +73,27 @@ struct FLandscapeTextureDataInfo
 
 	void AddMipUpdateRegion(int32 MipNum, int32 InX1, int32 InY1, int32 InX2, int32 InY2)
 	{
-		check( MipNum < MipInfo.Num() );
-		new(MipInfo[MipNum].MipUpdateRegions) FUpdateTextureRegion2D(InX1, InY1, InX1, InY1, 1+InX2-InX1, 1+InY2-InY1);
-	}
+		if (MipInfo[MipNum].bFull)
+		{
+			return;
+		}
+
+		check(MipNum < MipInfo.Num());
+		uint32 Width = 1 + InX2 - InX1;
+		uint32 Height = 1 + InY2 - InY1;
+		// Catch situation where we are updating the whole texture to avoid adding redundant regions once the whole region as been included.
+		if (Width == GetMipSizeX(MipNum) && Width == GetMipSizeY(MipNum))
+		{
+			MipInfo[MipNum].bFull = true;
+			MipInfo[MipNum].MipUpdateRegions.Reset();
+			// Push a full region for UpdateTextureData() to process later
+			new(MipInfo[MipNum].MipUpdateRegions) FUpdateTextureRegion2D(0, 0, 0, 0, Width, Height);
+			return;
+		}
 
+		new(MipInfo[MipNum].MipUpdateRegions) FUpdateTextureRegion2D(InX1, InY1, InX1, InY1, Width, Height);
+	}
+		
 	void* GetMipData(int32 MipNum)
 	{
 		check( MipNum < MipInfo.Num() );
@@ -397,7 +415,6 @@ struct FHeightmapAccessor
 
 			for (ULandscapeComponent* Component : Components)
 			{
-				Component->InvalidateLightingCache();
 				Component->RequestHeightmapUpdate();
 			}
 						
@@ -408,7 +425,8 @@ struct FHeightmapAccessor
             // Landscape Layers are updates are delayed and done in  ALandscape::TickLayers
 			if (!LandscapeEdit->HasLandscapeLayersContent())
 			{
-				ALandscapeProxy::InvalidateGeneratedComponentData(Components);
+				const bool bInvalidateLightingCache = true;
+				ALandscapeProxy::InvalidateGeneratedComponentData(Components, bInvalidateLightingCache);
 				bUpdateNormals = true;
 				for (ULandscapeComponent* Component : Components)
 				{
