Fix for flickering landscape when using edit layers in D3D12
Needs heavyweight BlockUntilGPUIdle call until RHI behaviour is unified

diff --git a/Engine/Source/Runtime/Landscape/Private/LandscapeEditLayers.cpp b/Engine/Source/Runtime/Landscape/Private/LandscapeEditLayers.cpp
--- a/Engine/Source/Runtime/Landscape/Private/LandscapeEditLayers.cpp
+++ b/Engine/Source/Runtime/Landscape/Private/LandscapeEditLayers.cpp
@@ -3257,6 +3257,10 @@ bool ALandscape::ResolveLayersTexture(FLandscapeLayersTexture2DCPUReadBackResour
 	ENQUEUE_RENDER_COMMAND(FLandscapeLayersReadSurfaceCommand)(
 		[InCPUReadBackTexture, &OutMipsData](FRHICommandListImmediate& RHICmdList) mutable
 	{
+ 		// D3D12RHI requires heavyweight BlockUntilGPUIdle() call to ensure commands have finished.
+ 		// Note that this needs improving but for now is enshrined in the RHIUnitTest::VerifyBufferContents unit test
+		RHICmdList.BlockUntilGPUIdle();
+
 		OutMipsData.AddDefaulted(InCPUReadBackTexture->TextureRHI->GetNumMips());
 
 		int32 MipSizeU = InCPUReadBackTexture->GetSizeX();
