Improve compression quality of landscape weightmaps

diff --git a/Engine/Source/Runtime/Landscape/Private/LandscapeEdit.cpp b/Engine/Source/Runtime/Landscape/Private/LandscapeEdit.cpp
--- a/Engine/Source/Runtime/Landscape/Private/LandscapeEdit.cpp
+++ b/Engine/Source/Runtime/Landscape/Private/LandscapeEdit.cpp
@@ -6478,6 +6478,12 @@ UTexture2D* ALandscapeProxy::CreateLandscapeTexture(int32 InSizeX, int32 InSizeY
 	NewTexture->AddressY = TA_Clamp;
 	NewTexture->LODGroup = InLODGroup;
 
+	if (bCompress && InLODGroup == TEXTUREGROUP_Terrain_Weightmap)
+	{
+		// Compress weightmaps with a best quality to reduce blocky artifacts on a landscape 
+		NewTexture->CompressionQuality = TCQ_Highest;
+	}
+
 	return NewTexture;
 }
 
