Landscape Layers: PostLoad Fixup Weightmap code not supporting Layers.

#rnx

diff --git a/Engine/Source/Runtime/Landscape/Private/LandscapeEdit.cpp b/Engine/Source/Runtime/Landscape/Private/LandscapeEdit.cpp
--- a/Engine/Source/Runtime/Landscape/Private/LandscapeEdit.cpp
+++ b/Engine/Source/Runtime/Landscape/Private/LandscapeEdit.cpp
@@ -670,11 +670,28 @@ void ULandscapeComponent::FixupWeightmaps()
 
 			if (bFixedLayerDeletion)
 			{
-				FLandscapeEditDataInterface LandscapeEdit(Info);
-				for (int32 Idx = 0; Idx < LayersToDelete.Num(); ++Idx)
 				{
-					DeleteLayer(LayersToDelete[Idx], LandscapeEdit);
+					FLandscapeEditDataInterface LandscapeEdit(Info);
+					for (int32 Idx = 0; Idx < LayersToDelete.Num(); ++Idx)
+					{
+						DeleteLayer(LayersToDelete[Idx], LandscapeEdit);
+					}
 				}
+
+				ForEachLayer([&](const FGuid& LayerGuid, FLandscapeLayerComponentData& LayerData)
+				{
+					SetEditingLayer(LayerGuid);
+					FLandscapeEditDataInterface LandscapeEdit(Info);
+					for (int32 Idx = 0; Idx < LayersToDelete.Num(); ++Idx)
+					{
+						DeleteLayer(LayersToDelete[Idx], LandscapeEdit);
+					}
+				});
+								
+				// Make sure to clear editing layer and cache
+				SetEditingLayer(FGuid());
+				CachedEditingLayer.Invalidate();
+				CachedEditingLayerData = nullptr;
 			}
 
 			bool bFixedWeightmapTextureIndex = false;
diff --git a/Engine/Source/Runtime/Landscape/Private/LandscapeEditInterface.cpp b/Engine/Source/Runtime/Landscape/Private/LandscapeEditInterface.cpp
--- a/Engine/Source/Runtime/Landscape/Private/LandscapeEditInterface.cpp
+++ b/Engine/Source/Runtime/Landscape/Private/LandscapeEditInterface.cpp
@@ -1739,7 +1739,12 @@ void ULandscapeComponent::DeleteLayer(ULandscapeLayerInfoObject* LayerInfo, FLan
 		ComponentWeightmapTextures[DeleteLayerWeightmapTextureIndex]->Modify(LandscapeEdit.GetShouldDirtyPackage());
 		ComponentWeightmapTextures[DeleteLayerWeightmapTextureIndex]->ClearFlags(RF_Standalone);
 
-		ComponentWeightmapTexturesUsage.RemoveAt(DeleteLayerWeightmapTextureIndex);
+		// possible that usages have not been built yet
+		if (ComponentWeightmapTexturesUsage.IsValidIndex(DeleteLayerWeightmapTextureIndex))
+		{
+			ComponentWeightmapTexturesUsage.RemoveAt(DeleteLayerWeightmapTextureIndex);
+		}
+
 		ComponentWeightmapTextures.RemoveAt(DeleteLayerWeightmapTextureIndex);
 
 		// Adjust WeightmapTextureIndex index for other layers
