Fix copy paste typo which broke landscape tessellation
Cherry pick integrated from //UE4/Dev-Rendering/...

diff --git a/Engine/Source/Runtime/Landscape/Private/LandscapeRender.cpp b/Engine/Source/Runtime/Landscape/Private/LandscapeRender.cpp
--- a/Engine/Source/Runtime/Landscape/Private/LandscapeRender.cpp
+++ b/Engine/Source/Runtime/Landscape/Private/LandscapeRender.cpp
@@ -821,7 +821,7 @@ void FLandscapeRenderSystem::BeginRenderView(const FSceneView* View)
 		if (TessellationFalloffSettings.UseTessellationComponentScreenSizeFalloff && NumEntitiesWithTessellation > 0)
 		{
 			SectionTessellationFalloffC = CachedSectionTessellationFalloffC[View];
-			SectionTessellationFalloffC = CachedSectionTessellationFalloffK[View];
+			SectionTessellationFalloffK = CachedSectionTessellationFalloffK[View];
 		}
 	}
 
