Landscape Brush: Fix rare crash when Brush was Ticked with GFrameNumber < CVarLandscapeBrushPadding

diff --git a/Engine/Source/Runtime/Landscape/Private/LandscapeBlueprintBrushBase.cpp b/Engine/Source/Runtime/Landscape/Private/LandscapeBlueprintBrushBase.cpp
--- a/Engine/Source/Runtime/Landscape/Private/LandscapeBlueprintBrushBase.cpp
+++ b/Engine/Source/Runtime/Landscape/Private/LandscapeBlueprintBrushBase.cpp
@@ -8,6 +8,8 @@
 #define LOCTEXT_NAMESPACE "Landscape"
 
 #if WITH_EDITOR
+static const uint32 InvalidLastRequestLayersContentUpdateFrameNumber = 0;
+
 static TAutoConsoleVariable<int32> CVarLandscapeBrushPadding(
 	TEXT("landscape.BrushFramePadding"),
 	5,
@@ -20,7 +22,7 @@ ALandscapeBlueprintBrushBase::ALandscapeBlueprintBrushBase(const FObjectInitiali
 	, AffectHeightmap(false)
 	, AffectWeightmap(false)
 	, bIsVisible(true)
-	, LastRequestLayersContentUpdateFrameNumber(0)
+	, LastRequestLayersContentUpdateFrameNumber(InvalidLastRequestLayersContentUpdateFrameNumber)
 #endif
 {
 #if WITH_EDITOR
@@ -52,7 +54,8 @@ void ALandscapeBlueprintBrushBase::RequestLandscapeUpdate()
 		if (ModeMask)
 		{
 			OwningLandscape->RequestLayersContentUpdateForceAll((ELandscapeLayerUpdateMode)ModeMask);
-			LastRequestLayersContentUpdateFrameNumber = GFrameNumber;
+			// Just in case differentiate between 0 (default value and frame number)
+			LastRequestLayersContentUpdateFrameNumber = GFrameNumber == InvalidLastRequestLayersContentUpdateFrameNumber ? GFrameNumber + 1 : GFrameNumber;
 		}
 	}
 #endif
@@ -64,7 +67,9 @@ void ALandscapeBlueprintBrushBase::Tick(float DeltaSeconds)
 #if WITH_EDITORONLY_DATA
 	// Avoid computing collision and client updates every frame
 	// Wait until we didn't trigger any more landscape update requests (padding of a couple of frames)
-	if (LastRequestLayersContentUpdateFrameNumber + CVarLandscapeBrushPadding.GetValueOnAnyThread() == GFrameNumber)
+	if (OwningLandscape != nullptr && 
+		LastRequestLayersContentUpdateFrameNumber != InvalidLastRequestLayersContentUpdateFrameNumber &&
+		LastRequestLayersContentUpdateFrameNumber + CVarLandscapeBrushPadding.GetValueOnAnyThread() == GFrameNumber)
 	{
 		uint32 ModeMask = 0;
 		if (AffectHeightmap)
