Changed display name of Landscape spline control point's "Width" parameter to "Half-Width" so that it reflects what it really controls.

diff --git a/Engine/Source/Runtime/Landscape/Classes/LandscapeSplineControlPoint.h b/Engine/Source/Runtime/Landscape/Classes/LandscapeSplineControlPoint.h
--- a/Engine/Source/Runtime/Landscape/Classes/LandscapeSplineControlPoint.h
+++ b/Engine/Source/Runtime/Landscape/Classes/LandscapeSplineControlPoint.h
@@ -60,8 +60,8 @@ class ULandscapeSplineControlPoint : public UObject
 	UPROPERTY(EditAnywhere, Category=LandscapeSpline)
 	FRotator Rotation;
 
-	/** Width of the spline at this point. */
-	UPROPERTY(EditAnywhere, Category=LandscapeSpline)
+	/** Half-Width of the spline at this point. */
+	UPROPERTY(EditAnywhere, Category=LandscapeSpline, meta = (DisplayName = "Half-Width"))
 	float Width;
 
 	/** Layer Width ratio of the spline at this point. */
