Fix Update of Physical Material not being propagated properly

diff --git a/Engine/Source/Runtime/Landscape/Private/LandscapeEdit.cpp b/Engine/Source/Runtime/Landscape/Private/LandscapeEdit.cpp
--- a/Engine/Source/Runtime/Landscape/Private/LandscapeEdit.cpp
+++ b/Engine/Source/Runtime/Landscape/Private/LandscapeEdit.cpp
@@ -4444,6 +4444,8 @@ void ALandscapeProxy::PostEditChangeProperty(FPropertyChangedEvent& PropertyChan
 	const FName PropertyName = PropertyChangedEvent.MemberProperty ? PropertyChangedEvent.MemberProperty->GetFName() : NAME_None;
 	const FName SubPropertyName = PropertyChangedEvent.Property ? PropertyChangedEvent.Property->GetFName() : NAME_None;
 
+	bool bChangedPhysMaterial = false;
+
 	if (PropertyName == FName(TEXT("RelativeScale3D")))
 	{
 		// RelativeScale3D isn't even a property of ALandscapeProxy, it's a property of the root component
@@ -4517,9 +4519,9 @@ void ALandscapeProxy::PostEditChangeProperty(FPropertyChangedEvent& PropertyChan
 		// Recalculate in a few seconds.
 		GetWorld()->TriggerStreamingDataRebuild();
 	}
-	else if (GIsEditor && PropertyName == FName(TEXT("DefaultPhysMaterial")))
+	else if (GIsEditor && PropertyName == GET_MEMBER_NAME_CHECKED(ALandscapeProxy, DefaultPhysMaterial))
 	{
-		ChangedPhysMaterial();
+		bChangedPhysMaterial = true;
 	}
 	else if (GIsEditor &&
 		(PropertyName == GET_MEMBER_NAME_CHECKED(ALandscapeProxy, CollisionMipLevel) ||
@@ -4640,6 +4642,12 @@ void ALandscapeProxy::PostEditChangeProperty(FPropertyChangedEvent& PropertyChan
 	// Must do this *after* correcting the scale or reattaching the landscape components will crash!
 	// Must do this *after* clamping values / propogating values to components
 	Super::PostEditChangeProperty(PropertyChangedEvent);
+
+	// Call that posteditchange when components are registered
+	if (bChangedPhysMaterial)
+	{
+		ChangedPhysMaterial();
+	}
 }
 
 void ALandscapeStreamingProxy::PostEditChangeProperty(FPropertyChangedEvent& PropertyChangedEvent)
@@ -4846,6 +4854,10 @@ void ALandscape::PostEditChangeProperty(FPropertyChangedEvent& PropertyChangedEv
 		CollisionMipLevel = FMath::Clamp<int32>(CollisionMipLevel, 0, FMath::CeilLogTwo(SubsectionSizeQuads + 1) - 1);
 		bPropagateToProxies = true;
 	}
+	else if (PropertyName == GET_MEMBER_NAME_CHECKED(ALandscapeProxy, DefaultPhysMaterial))
+	{
+		bPropagateToProxies = true;
+	}
 	else if (PropertyName == GET_MEMBER_NAME_CHECKED(ALandscapeProxy, SimpleCollisionMipLevel))
 	{
 		SimpleCollisionMipLevel = FMath::Clamp<int32>(SimpleCollisionMipLevel, 0, FMath::CeilLogTwo(SubsectionSizeQuads + 1) - 1);
@@ -4982,17 +4994,14 @@ void ALandscape::PostEditChangeProperty(FPropertyChangedEvent& PropertyChangedEv
 
 void ALandscapeProxy::ChangedPhysMaterial()
 {
-	ULandscapeInfo* LandscapeInfo = GetLandscapeInfo();
-	if (!LandscapeInfo) return;
-	for (auto It = LandscapeInfo->XYtoComponentMap.CreateIterator(); It; ++It)
+	for (ULandscapeComponent* LandscapeComponent : LandscapeComponents)
 	{
-		ULandscapeComponent* Comp = It.Value();
-		if (Comp)
+		if (LandscapeComponent && LandscapeComponent->IsRegistered())
 		{
-			ULandscapeHeightfieldCollisionComponent* CollisionComponent = Comp->CollisionComponent.Get();
+			ULandscapeHeightfieldCollisionComponent* CollisionComponent = LandscapeComponent->CollisionComponent.Get();
 			if (CollisionComponent)
 			{
-				Comp->UpdateCollisionLayerData();
+				LandscapeComponent->UpdateCollisionLayerData();
 				// Physical materials cooked into collision object, so we need to recreate it
 				CollisionComponent->RecreateCollision();
 			}
