Shaved 100us per frame off of ALandscapeProxy::UpdateGrass on switch.

- Optimized math in FGrassBuilderBase
- Inlined SetSectionBase/GetSectionBase
- Cached results of GetGrassTypes so it's not called every frame
- Converted TAutoConsoleVariable to FAutoConsoleVariableRef to avoid Get overhead

diff --git a/Engine/Source/Runtime/Landscape/Classes/LandscapeComponent.h b/Engine/Source/Runtime/Landscape/Classes/LandscapeComponent.h
--- a/Engine/Source/Runtime/Landscape/Classes/LandscapeComponent.h
+++ b/Engine/Source/Runtime/Landscape/Classes/LandscapeComponent.h
@@ -765,10 +765,17 @@ public:
 	LANDSCAPE_API ALandscapeProxy* GetLandscapeProxy() const;
 
 	/** @return Component section base as FIntPoint */
-	LANDSCAPE_API FIntPoint GetSectionBase() const; 
+	LANDSCAPE_API FIntPoint GetSectionBase() const
+	{
+		return FIntPoint(SectionBaseX, SectionBaseY);
+	}
 
 	/** @param InSectionBase new section base for a component */
-	LANDSCAPE_API void SetSectionBase(FIntPoint InSectionBase);
+	LANDSCAPE_API void SetSectionBase(FIntPoint InSectionBase)
+	{
+		SectionBaseX = InSectionBase.X;
+		SectionBaseY = InSectionBase.Y;
+	}
 
 	/** @todo document */
 	const FGuid& GetLightingGuid() const
diff --git a/Engine/Source/Runtime/Landscape/Classes/LandscapeProxy.h b/Engine/Source/Runtime/Landscape/Classes/LandscapeProxy.h
--- a/Engine/Source/Runtime/Landscape/Classes/LandscapeProxy.h
+++ b/Engine/Source/Runtime/Landscape/Classes/LandscapeProxy.h
@@ -497,6 +497,14 @@ public:
 	UPROPERTY(EditAnywhere, BlueprintSetter=EditorSetLandscapeMaterial, Category=Landscape)
 	UMaterialInterface* LandscapeMaterial;
 
+#if !WITH_EDITORONLY_DATA
+	/** Used to cache grass types from GetGrassTypes */
+	UMaterialInterface* LandscapeMaterialCached;
+
+	/** Cached grass types from GetGrassTypes */
+	TArray<ULandscapeGrassType*> LandscapeGrassTypes;
+#endif
+
 	/** Material used to render landscape components with holes. If not set, LandscapeMaterial will be used (blend mode will be overridden to Masked if it is set to Opaque) */
 	UPROPERTY(EditAnywhere, Category=Landscape, AdvancedDisplay)
 	UMaterialInterface* LandscapeHoleMaterial;
@@ -835,7 +843,7 @@ public:
 
 
 	/* Get the list of grass types on this landscape */
-	TArray<ULandscapeGrassType*> GetGrassTypes() const;
+	static void GetGrassTypes(UMaterialInterface* LandscapeMat, TArray<ULandscapeGrassType*>& GrassTypesOut);
 
 	/* Invalidate the precomputed grass and baked texture data for the specified components */
 	LANDSCAPE_API static void InvalidateGeneratedComponentData(const TSet<ULandscapeComponent*>& Components, bool bInvalidateLightingCache = false);
diff --git a/Engine/Source/Runtime/Landscape/Private/Landscape.cpp b/Engine/Source/Runtime/Landscape/Private/Landscape.cpp
--- a/Engine/Source/Runtime/Landscape/Private/Landscape.cpp
+++ b/Engine/Source/Runtime/Landscape/Private/Landscape.cpp
@@ -1017,6 +1017,10 @@ ALandscapeProxy::ALandscapeProxy(const FObjectInitializer& ObjectInitializer)
 #if WITH_EDITORONLY_DATA
 	, TargetDisplayOrder(ELandscapeLayerDisplayMode::Default)
 #endif // WITH_EDITORONLY_DATA
+#if !WITH_EDITORONLY_DATA
+	, LandscapeMaterialCached(nullptr)
+	, LandscapeGrassTypes()
+#endif
 	, bHasLandscapeGrass(true)
 {
 	PrimaryActorTick.bCanEverTick = true;
@@ -1293,17 +1297,6 @@ ALandscapeProxy* ULandscapeComponent::GetLandscapeProxy() const
 	return CastChecked<ALandscapeProxy>(GetOuter());
 }
 
-FIntPoint ULandscapeComponent::GetSectionBase() const
-{
-	return FIntPoint(SectionBaseX, SectionBaseY);
-}
-
-void ULandscapeComponent::SetSectionBase(FIntPoint InSectionBase)
-{
-	SectionBaseX = InSectionBase.X;
-	SectionBaseY = InSectionBase.Y;
-}
-
 const FMeshMapBuildData* ULandscapeComponent::GetMeshMapBuildData() const
 {
 	AActor* Owner = GetOwner();
diff --git a/Engine/Source/Runtime/Landscape/Private/LandscapeGrass.cpp b/Engine/Source/Runtime/Landscape/Private/LandscapeGrass.cpp
--- a/Engine/Source/Runtime/Landscape/Private/LandscapeGrass.cpp
+++ b/Engine/Source/Runtime/Landscape/Private/LandscapeGrass.cpp
@@ -60,97 +60,115 @@
 
 DEFINE_LOG_CATEGORY_STATIC(LogGrass, Log, All);
 
-static TAutoConsoleVariable<float> CVarGuardBandMultiplier(
+static float GGuardBandMultiplier = 1.3f;
+static FAutoConsoleVariableRef CVarGuardBandMultiplier(
 	TEXT("grass.GuardBandMultiplier"),
-	1.3f,
+	GGuardBandMultiplier,
 	TEXT("Used to control discarding in the grass system. Approximate range, 1-4. Multiplied by the cull distance to control when we add grass components."));
 
-static TAutoConsoleVariable<float> CVarGuardBandDiscardMultiplier(
+static float GGuardBandDiscardMultiplier = 1.4f;
+static FAutoConsoleVariableRef CVarGuardBandDiscardMultiplier(
 	TEXT("grass.GuardBandDiscardMultiplier"),
-	1.4f,
+	GGuardBandDiscardMultiplier,
 	TEXT("Used to control discarding in the grass system. Approximate range, 1-4. Multiplied by the cull distance to control when we discard grass components."));
 
-static TAutoConsoleVariable<int32> CVarMinFramesToKeepGrass(
+static int32 GMinFramesToKeepGrass = 30;
+static FAutoConsoleVariableRef CVarMinFramesToKeepGrass(
 	TEXT("grass.MinFramesToKeepGrass"),
-	30,
+	GMinFramesToKeepGrass,
 	TEXT("Minimum number of frames before cached grass can be discarded; used to prevent thrashing."));
 
-static TAutoConsoleVariable<int32> CVarGrassTickInterval(
+static int32 GGrassTickInterval = 1;
+static FAutoConsoleVariableRef CVarGrassTickInterval(
 	TEXT("grass.TickInterval"),
-	1,
+	GGrassTickInterval,
 	TEXT("Number of frames between grass ticks."));
 
-static TAutoConsoleVariable<float> CVarMinTimeToKeepGrass(
+static float GMinTimeToKeepGrass = 5.0f;
+static FAutoConsoleVariableRef CVarMinTimeToKeepGrass(
 	TEXT("grass.MinTimeToKeepGrass"),
-	5.0f,
+	GMinTimeToKeepGrass,
 	TEXT("Minimum number of seconds before cached grass can be discarded; used to prevent thrashing."));
 
-static TAutoConsoleVariable<int32> CVarMaxInstancesPerComponent(
+static int32 GMaxInstancesPerComponent = 65536;
+static FAutoConsoleVariableRef CVarMaxInstancesPerComponent(
 	TEXT("grass.MaxInstancesPerComponent"),
-	65536,
+	GMaxInstancesPerComponent,
 	TEXT("Used to control the number of hierarchical components created. More can be more efficient, but can be hitchy as new components come into range"));
 
-static TAutoConsoleVariable<int32> CVarMaxAsyncTasks(
+static int32 GMaxAsyncTasks = 4;
+static FAutoConsoleVariableRef CVarMaxAsyncTasks(
 	TEXT("grass.MaxAsyncTasks"),
-	4,
+	GMaxAsyncTasks,
 	TEXT("Used to control the number of hierarchical components created at a time."));
 
-static TAutoConsoleVariable<int32> CVarUseHaltonDistribution(
+static int32 GUseHaltonDistribution = 0;
+static FAutoConsoleVariableRef CVarUseHaltonDistribution(
 	TEXT("grass.UseHaltonDistribution"),
-	0,
+	GUseHaltonDistribution,
 	TEXT("Used to control the distribution of grass instances. If non-zero, use a halton sequence."));
 
-static TAutoConsoleVariable<float> CVarGrassDensityScale(
+static float GGrassDensityScale = 1;
+static FAutoConsoleVariableRef CVarGrassDensityScale(
 	TEXT("grass.densityScale"),
-	1,
+	GGrassDensityScale,
 	TEXT("Multiplier on all grass densities."),
 	ECVF_Scalability);
 
-static TAutoConsoleVariable<float> CVarGrassCullDistanceScale(
+static float GGrassCullDistanceScale = 1;
+static FAutoConsoleVariableRef CVarGrassCullDistanceScale(
 	TEXT("grass.CullDistanceScale"),
-	1,
+	GGrassCullDistanceScale,
 	TEXT("Multiplier on all grass cull distances."),
 	ECVF_Scalability);
 
-static TAutoConsoleVariable<int32> CVarGrassEnable(
+static int32 GGrassEnable = 1;
+static FAutoConsoleVariableRef CVarGrassEnable(
 	TEXT("grass.Enable"),
-	1,
+	GGrassEnable,
 	TEXT("1: Enable Grass; 0: Disable Grass"));
 
-static TAutoConsoleVariable<int32> CVarGrassDiscardDataOnLoad(
+static int32 GGrassDiscardDataOnLoad = 0;
+static FAutoConsoleVariableRef CVarGrassDiscardDataOnLoad(
 	TEXT("grass.DiscardDataOnLoad"),
-	0,
+	GGrassDiscardDataOnLoad,
 	TEXT("1: Discard grass data on load (disables grass); 0: Keep grass data (requires reloading level)"),
 	ECVF_Scalability);
 
-static TAutoConsoleVariable<int32> CVarUseStreamingManagerForCameras(
+static int32 GUseStreamingManagerForCameras = 1;
+static FAutoConsoleVariableRef CVarUseStreamingManagerForCameras(
 	TEXT("grass.UseStreamingManagerForCameras"),
-	1,
+	GUseStreamingManagerForCameras,
 	TEXT("1: Use Streaming Manager; 0: Use ViewLocationsRenderedLastFrame"));
 
-static TAutoConsoleVariable<int32> CVarCullSubsections(
+static int32 GCullSubsections = 1;
+static FAutoConsoleVariableRef CVarCullSubsections(
 	TEXT("grass.CullSubsections"),
-	1,
+	GCullSubsections,
 	TEXT("1: Cull each foliage component; 0: Cull only based on the landscape component."));
 
-static TAutoConsoleVariable<int32> CVarDisableGPUCull(
+static int32 GDisableGPUCull = 0;
+static FAutoConsoleVariableRef CVarDisableGPUCull(
 	TEXT("grass.DisableGPUCull"),
-	0,
+	GDisableGPUCull,
 	TEXT("For debugging. Set this to zero to see where the grass is generated. Useful for tweaking the guard bands."));
 
-static TAutoConsoleVariable<int32> CVarPrerenderGrassmaps(
+static int32 GPrerenderGrassmaps = 1;
+static FAutoConsoleVariableRef CVarPrerenderGrassmaps(
 	TEXT("grass.PrerenderGrassmaps"),
-	1,
+	GPrerenderGrassmaps,
 	TEXT("1: Pre-render grass maps for all components in the editor; 0: Generate grass maps on demand while moving through the editor"));
 
-static TAutoConsoleVariable<int32> CVarDisableDynamicShadows(
+static int32 GDisableDynamicShadows = 0;
+static FAutoConsoleVariableRef CVarDisableDynamicShadows(
 	TEXT("grass.DisableDynamicShadows"),
-	0,
+	GDisableDynamicShadows,
 	TEXT("0: Dynamic shadows from grass follow the grass type bCastDynamicShadow flag; 1: Dynamic shadows are disabled for all grass"));
 
-static TAutoConsoleVariable<int32> CVarIgnoreExcludeBoxes(
+static int32 GIgnoreExcludeBoxes = 0;
+static FAutoConsoleVariableRef CVarIgnoreExcludeBoxes(
 	TEXT("grass.IgnoreExcludeBoxes"),
-	0,
+	GIgnoreExcludeBoxes,
 	TEXT("For debugging. Ignores any exclusion boxes."));
 
 DECLARE_CYCLE_STAT(TEXT("Grass Async Build Time"), STAT_FoliageGrassAsyncBuildTime, STATGROUP_Foliage);
@@ -164,15 +182,15 @@ static int32 GGrassUpdateInterval = 1;
 static void GrassCVarSinkFunction()
 {
 	static float CachedGrassDensityScale = 1.0f;
-	float GrassDensityScale = CVarGrassDensityScale.GetValueOnGameThread();
+	float GrassDensityScale = GGrassDensityScale;
 
 	if (FApp::IsGame())
 	{
-		GGrassUpdateInterval = FMath::Clamp<int32>(CVarGrassTickInterval.GetValueOnGameThread(), 1, 60);
+		GGrassUpdateInterval = FMath::Clamp<int32>(GGrassTickInterval, 1, 60);
 	}
 
 	static float CachedGrassCullDistanceScale = 1.0f;
-	float GrassCullDistanceScale = CVarGrassCullDistanceScale.GetValueOnGameThread();
+	float GrassCullDistanceScale = GGrassCullDistanceScale;
 
 	static const IConsoleVariable* DetailModeCVar = IConsoleManager::Get().FindConsoleVariable(TEXT("r.DetailMode"));
 	static int32 CachedDetailMode = DetailModeCVar ? DetailModeCVar->GetInt() : 0;
@@ -1350,7 +1368,7 @@ FArchive& operator<<(FArchive& Ar, FLandscapeComponentGrassData& Data)
 
 void FLandscapeComponentGrassData::ConditionalDiscardDataOnLoad()
 {
-	if (!GIsEditor && CVarGrassDiscardDataOnLoad.GetValueOnAnyThread())
+	if (!GIsEditor && GGrassDiscardDataOnLoad)
 	{
 		// Remove data for grass types which have scalability enabled
 		for (auto GrassTypeIt = WeightData.CreateIterator(); GrassTypeIt; ++GrassTypeIt)
@@ -1411,7 +1429,7 @@ void ALandscapeProxy::TickGrass()
 
 	// Update foliage
 	static TArray<FVector> OldCameras;
-	if (CVarUseStreamingManagerForCameras.GetValueOnGameThread() == 0)
+	if (GUseStreamingManagerForCameras == 0)
 	{
 		UWorld* World = GetWorld();
 		if (!World)
@@ -1472,9 +1490,7 @@ struct FGrassBuilderBase
 
 	FGrassBuilderBase(ALandscapeProxy* Landscape, ULandscapeComponent* Component, const FGrassVariety& GrassVariety, ERHIFeatureLevel::Type FeatureLevel, int32 SqrtSubsections = 1, int32 SubX = 0, int32 SubY = 0, bool bEnableDensityScaling = true)
 	{
-		bHaveValidData = true;
-
-		const float DensityScale = bEnableDensityScaling ? CVarGrassDensityScale.GetValueOnAnyThread() : 1.0f;
+		const float DensityScale = bEnableDensityScaling ? GGrassDensityScale : 1.0f;
 		GrassDensity = GrassVariety.GrassDensity.GetValueForFeatureLevel(FeatureLevel) * DensityScale;
 
 		DrawScale = Landscape->GetRootComponent()->GetRelativeScale3D();
@@ -1484,18 +1500,22 @@ struct FGrassBuilderBase
 		SectionBase = Component->GetSectionBase();
 		ComponentSizeQuads = Component->ComponentSizeQuads;
 
-		Origin = FVector(DrawScale.X * float(SectionBase.X), DrawScale.Y * float(SectionBase.Y), 0.0f);
-		Extent = FVector(DrawScale.X * float(SectionBase.X + ComponentSizeQuads), DrawScale.Y * float(SectionBase.Y + ComponentSizeQuads), 0.0f) - Origin;
+		Origin.X = DrawScale.X * SectionBase.X;
+		Origin.Y = DrawScale.Y * SectionBase.Y;
+		Origin.Z = 0.0f;
+
+		Extent.X = DrawScale.X * ComponentSizeQuads;
+		Extent.Y = DrawScale.Y * ComponentSizeQuads;
+		Extent.Z = 0.0f;
 
-		ComponentOrigin = Origin - FVector(DrawScale.X * LandscapeSectionOffset.X, DrawScale.Y * LandscapeSectionOffset.Y, 0.0f);
+		ComponentOrigin.X = DrawScale.X * (SectionBase.X - LandscapeSectionOffset.X);
+		ComponentOrigin.Y = DrawScale.Y * (SectionBase.Y - LandscapeSectionOffset.Y);
+		ComponentOrigin.Z = 0.0f;
 
 		SqrtMaxInstances = FMath::CeilToInt(FMath::Sqrt(FMath::Abs(Extent.X * Extent.Y * GrassDensity / 1000.0f / 1000.0f)));
 
-		if (SqrtMaxInstances == 0)
-		{
-			bHaveValidData = false;
-		}
-		const FRotator DrawRot = Landscape->GetActorRotation();
+		bHaveValidData = SqrtMaxInstances != 0;
+
 		LandscapeToWorld = Landscape->GetRootComponent()->GetComponentTransform().ToMatrixNoScale();
 
 		if (bHaveValidData && SqrtSubsections != 1)
@@ -1504,8 +1524,11 @@ struct FGrassBuilderBase
 			SqrtMaxInstances /= SqrtSubsections;
 			check(SqrtMaxInstances > 0);
 
-			Extent /= float(SqrtSubsections);
-			Origin += Extent * FVector(float(SubX), float(SubY), 0.0f);
+			Extent.X /= SqrtSubsections;
+			Extent.Y /= SqrtSubsections;
+
+			Origin.X += Extent.X * SubX;
+			Origin.Y += Extent.Y * SubY;
 		}
 	}
 };
@@ -2091,22 +2114,20 @@ void ALandscapeProxy::FlushGrassComponents(const TSet<ULandscapeComponent*>* Onl
 	}
 }
 
-TArray<ULandscapeGrassType*> ALandscapeProxy::GetGrassTypes() const
+void ALandscapeProxy::GetGrassTypes(UMaterialInterface* LandscapeMat, TArray<ULandscapeGrassType*>& GrassTypesOut)
 {
-	TArray<ULandscapeGrassType*> GrassTypes;
-	if (LandscapeMaterial)
+	if (LandscapeMat)
 	{
 		TArray<const UMaterialExpressionLandscapeGrassOutput*> GrassExpressions;
-		LandscapeMaterial->GetMaterial()->GetAllExpressionsOfType<UMaterialExpressionLandscapeGrassOutput>(GrassExpressions);
+		LandscapeMat->GetMaterial()->GetAllExpressionsOfType<UMaterialExpressionLandscapeGrassOutput>(GrassExpressions);
 		if (GrassExpressions.Num() > 0)
 		{
-			for (auto& Type : GrassExpressions[0]->GrassTypes)
+			for (const FGrassInput& Type : GrassExpressions[0]->GrassTypes)
 			{
-				GrassTypes.Add(Type.GrassType);
+				GrassTypesOut.Add(Type.GrassType);
 			}
 		}
 	}
-	return GrassTypes;
 }
 
 static uint32 GGrassExclusionChangeTag = 1;
@@ -2161,7 +2182,10 @@ void ALandscapeProxy::UpdateGrassDataStatus(TSet<UTexture2D*>& OutCurrentForcedS
 
 	// In either case we want to check the Grass textures stream state
 	const bool bCheckStreamingState = OutDesiredForcedStreamedTextures || bInEnableForceResidentFlag;
-	const bool bHasGrassTypes = GetGrassTypes().Num() > 0;
+
+	TArray<ULandscapeGrassType*> GrassTypes;
+	GetGrassTypes(LandscapeMaterial, GrassTypes);
+	const bool bHasGrassTypes = GrassTypes.Num() > 0;
 
 	for (auto Component : LandscapeComponents)
 	{
@@ -2233,7 +2257,7 @@ void ALandscapeProxy::UpdateGrassData()
 		return;
 	}
 
-	if (!CVarGrassEnable.GetValueOnAnyThread())
+	if (!GGrassEnable)
 	{
 		return;
 	}
@@ -2287,7 +2311,7 @@ void ALandscapeProxy::UpdateGrassData()
 void ALandscapeProxy::UpdateGrass(const TArray<FVector>& Cameras, bool bForceSync)
 {
 	SCOPE_CYCLE_COUNTER(STAT_GrassUpdate);
-	if (GFrameNumberLastStaleCheck != GFrameNumber && CVarIgnoreExcludeBoxes.GetValueOnAnyThread() == 0)
+	if (GFrameNumberLastStaleCheck != GFrameNumber && GIgnoreExcludeBoxes == 0)
 	{
 		GFrameNumberLastStaleCheck = GFrameNumber;
 		for (auto Iter = GGrassExclusionBoxes.CreateIterator(); Iter; ++Iter)
@@ -2300,18 +2324,29 @@ void ALandscapeProxy::UpdateGrass(const TArray<FVector>& Cameras, bool bForceSyn
 		}
 	}
 
-	if (CVarGrassEnable.GetValueOnAnyThread() > 0)
+	if (GGrassEnable > 0)
 	{
-		TArray<ULandscapeGrassType*> GrassTypes = GetGrassTypes();
+#if WITH_EDITOR
+		TArray<ULandscapeGrassType*> LandscapeGrassTypes;
+		GetGrassTypes(LandscapeMaterial, LandscapeGrassTypes);
+#else
+		// In non editor builds, cache grass types for performance.
+		if (LandscapeMaterial != LandscapeMaterialCached)
+		{
+			LandscapeMaterialCached = LandscapeMaterial;
+			LandscapeGrassTypes.Reset();
+			GetGrassTypes(LandscapeMaterial, LandscapeGrassTypes);
+		}
+#endif
 
-		float GuardBand = CVarGuardBandMultiplier.GetValueOnAnyThread();
-		float DiscardGuardBand = CVarGuardBandDiscardMultiplier.GetValueOnAnyThread();
-		bool bCullSubsections = CVarCullSubsections.GetValueOnAnyThread() > 0;
-		bool bDisableGPUCull = CVarDisableGPUCull.GetValueOnAnyThread() > 0;
-		bool bDisableDynamicShadows = CVarDisableDynamicShadows.GetValueOnAnyThread() > 0;
-		int32 MaxInstancesPerComponent = FMath::Max<int32>(1024, CVarMaxInstancesPerComponent.GetValueOnAnyThread());
-		int32 MaxTasks = CVarMaxAsyncTasks.GetValueOnAnyThread();
-		const float CullDistanceScale = CVarGrassCullDistanceScale.GetValueOnAnyThread();
+		float GuardBand = GGuardBandMultiplier;
+		float DiscardGuardBand = GGuardBandDiscardMultiplier;
+		bool bCullSubsections = GCullSubsections > 0;
+		bool bDisableGPUCull = GDisableGPUCull > 0;
+		bool bDisableDynamicShadows = GDisableDynamicShadows > 0;
+		int32 MaxInstancesPerComponent = FMath::Max<int32>(1024, GMaxInstancesPerComponent);
+		int32 MaxTasks = GMaxAsyncTasks;
+		const float CullDistanceScale = GGrassCullDistanceScale;
 
 		UWorld* World = GetWorld();
 		if (World)
@@ -2382,7 +2417,7 @@ void ALandscapeProxy::UpdateGrass(const TArray<FVector>& Cameras, bool bForceSyn
 				if (Component->ChangeTag != GGrassExclusionChangeTag)
 				{
 					Component->ActiveExcludedBoxes.Empty();
-					if (GGrassExclusionBoxes.Num() && CVarIgnoreExcludeBoxes.GetValueOnAnyThread() == 0)
+					if (GGrassExclusionBoxes.Num() && GIgnoreExcludeBoxes == 0)
 					{
 						const FBox& WorldBox = SortedLandscapeComponent.BoundsBox;
 
@@ -2397,7 +2432,7 @@ void ALandscapeProxy::UpdateGrass(const TArray<FVector>& Cameras, bool bForceSyn
 					Component->ChangeTag = GGrassExclusionChangeTag;
 				}
 
-				for (auto GrassType : GrassTypes)
+				for (ULandscapeGrassType* GrassType : LandscapeGrassTypes)
 				{
 					if (GrassType)
 					{
@@ -2711,7 +2746,7 @@ void ALandscapeProxy::UpdateGrass(const TArray<FVector>& Cameras, bool bForceSyn
 			{
 				int32 NumComponentsRendered = 0;
 				int32 NumComponentsUnableToRender = 0;
-				if ((GrassTypes.Num() > 0 && CVarPrerenderGrassmaps.GetValueOnAnyThread() > 0) || bBakeMaterialPositionOffsetIntoCollision)
+				if ((LandscapeGrassTypes.Num() > 0 && GPrerenderGrassmaps > 0) || bBakeMaterialPositionOffsetIntoCollision)
 				{
 					// try to render some grassmaps
 					TArray<ULandscapeComponent*> ComponentsToRender;
@@ -2749,7 +2784,7 @@ void ALandscapeProxy::UpdateGrass(const TArray<FVector>& Cameras, bool bForceSyn
 					}
 					if (ComponentsToRender.Num())
 					{
-						RenderGrassMaps(ComponentsToRender, GrassTypes);
+						RenderGrassMaps(ComponentsToRender, LandscapeGrassTypes);
 						MarkPackageDirty();
 					}
 				}
@@ -2778,8 +2813,8 @@ void ALandscapeProxy::UpdateGrass(const TArray<FVector>& Cameras, bool bForceSyn
 		QUICK_SCOPE_CYCLE_COUNTER(STAT_Grass_StillUsed);
 
 		// trim cached items based on time, pending and emptiness
-		double OldestToKeepTime = FPlatformTime::Seconds() - CVarMinTimeToKeepGrass.GetValueOnGameThread();
-		uint32 OldestToKeepFrame = GFrameNumber - CVarMinFramesToKeepGrass.GetValueOnGameThread() * GetGrassUpdateInterval();
+		double OldestToKeepTime = FPlatformTime::Seconds() - GMinTimeToKeepGrass;
+		uint32 OldestToKeepFrame = GFrameNumber - GMinFramesToKeepGrass * GetGrassUpdateInterval();
 		for (FCachedLandscapeFoliage::TGrassSet::TIterator Iter(FoliageCache.CachedGrassComps); Iter; ++Iter)
 		{
 			const FCachedLandscapeFoliage::FGrassComp& GrassItem = *Iter;
