Fix seams between landscape subsections when setting MaxLODLevel

diff --git a/Engine/Source/Runtime/Landscape/Private/LandscapeRender.cpp b/Engine/Source/Runtime/Landscape/Private/LandscapeRender.cpp
--- a/Engine/Source/Runtime/Landscape/Private/LandscapeRender.cpp
+++ b/Engine/Source/Runtime/Landscape/Private/LandscapeRender.cpp
@@ -1217,21 +1217,24 @@ FLandscapeComponentSceneProxy::FLandscapeComponentSceneProxy(ULandscapeComponent
 		CurrentScreenSizeRatio /= ScreenSizeRatioDivider;
 	}
 
-	if (InComponent->GetLandscapeProxy()->MaxLODLevel >= 0)
-	{
-		MaxLOD = FMath::Min<int8>(MaxLOD, InComponent->GetLandscapeProxy()->MaxLODLevel);
-	}
-
 	FirstLOD = 0;
 	LastLOD = MaxLOD;	// we always need to go to MaxLOD regardless of LODBias as we could need the lowest LODs due to streaming.
 
-	LODSettings.LastLODIndex = LastLOD;
-	LODSettings.LastLODScreenSizeSquared = LODScreenRatioSquared[LastLOD];
-
 	// Make sure out LastLOD is > of MinStreamedLOD otherwise we would not be using the right LOD->MIP, the only drawback is a possible minor memory usage for overallocating static mesh element batch
 	const int32 MinStreamedLOD = (HeightmapTexture != nullptr && HeightmapTexture->Resource != nullptr) ? FMath::Min<int32>(((FTexture2DResource*)HeightmapTexture->Resource)->GetCurrentFirstMip(), FMath::CeilLogTwo(SubsectionSizeVerts) - 1) : 0;
 	LastLOD = FMath::Max(MinStreamedLOD, LastLOD);
 
+	// Clamp to MaxLODLevel
+	const int32 MaxLODLevel = InComponent->GetLandscapeProxy()->MaxLODLevel;
+	if ( MaxLODLevel >= 0)
+	{
+		MaxLOD = FMath::Min<int8>(MaxLODLevel, MaxLOD);
+		LastLOD = FMath::Min<int32>(MaxLODLevel, LastLOD);
+	}
+
+	LODSettings.LastLODIndex = LastLOD;
+	LODSettings.LastLODScreenSizeSquared = LODScreenRatioSquared[LastLOD];
+
 	ForcedLOD = ForcedLOD != INDEX_NONE ? FMath::Clamp<int32>(ForcedLOD, FirstLOD, LastLOD) : ForcedLOD;
 	LODBias = FMath::Clamp<int8>(LODBias, -MaxLOD, MaxLOD);
 
