Landscape Grass Optim:
- Batch Ticks into LandscapeSubsystem (remove overhead of early exit virtual Ticks)
- Gather Cameras only once per frame

diff --git a/Engine/Source/Runtime/Landscape/Classes/Landscape.h b/Engine/Source/Runtime/Landscape/Classes/Landscape.h
--- a/Engine/Source/Runtime/Landscape/Classes/Landscape.h
+++ b/Engine/Source/Runtime/Landscape/Classes/Landscape.h
@@ -223,8 +223,6 @@ class ALandscape : public ALandscapeProxy
 public:
 	ALandscape(const FObjectInitializer& ObjectInitializer);
 
-	virtual void TickActor(float DeltaTime, ELevelTick TickType, FActorTickFunction& ThisTickFunction) override;
-
 	//~ Begin ALandscapeProxy Interface
 	LANDSCAPE_API virtual ALandscape* GetLandscapeActor() override;
 	LANDSCAPE_API virtual const ALandscape* GetLandscapeActor() const override;
@@ -261,7 +259,7 @@ public:
 	//~ End UObject Interface
 
 	LANDSCAPE_API bool IsUpToDate() const;
-
+	LANDSCAPE_API void TickLayers(float DeltaTime);
 	// Layers stuff
 #if WITH_EDITOR
 	LANDSCAPE_API void RegisterLandscapeEdMode(ILandscapeEdModeInterface* InLandscapeEdMode) { LandscapeEdMode = InLandscapeEdMode; }
@@ -336,7 +334,6 @@ public:
 	LANDSCAPE_API void InitializeLandscapeLayersWeightmapUsage();
 
 private:
-	void TickLayers(float DeltaTime, ELevelTick TickType, FActorTickFunction& ThisTickFunction);
 	void CreateLayersRenderingResource();
 	void GetLandscapeComponentNeighborsToRender(ULandscapeComponent* LandscapeComponent, TSet<ULandscapeComponent*>& NeighborComponents) const;
 	void GetLandscapeComponentWeightmapsToRender(ULandscapeComponent* LandscapeComponent, TSet<ULandscapeComponent*>& WeightmapComponents) const;
diff --git a/Engine/Source/Runtime/Landscape/Classes/LandscapeInfo.h b/Engine/Source/Runtime/Landscape/Classes/LandscapeInfo.h
--- a/Engine/Source/Runtime/Landscape/Classes/LandscapeInfo.h
+++ b/Engine/Source/Runtime/Landscape/Classes/LandscapeInfo.h
@@ -206,16 +206,6 @@ public:
 	 */
 	LANDSCAPE_API ALandscapeProxy* GetLandscapeProxy() const;
 
-	/**
-	 * Runs the given function on the root landscape actor and all streaming proxies
-	 * Most easily used with a lambda as follows:
-	 * ForAllLandscapeProxies([](ALandscapeProxy* Proxy)
-	 * {
-	 *     // Code
-	 * });
-	 */
-	LANDSCAPE_API void ForAllLandscapeProxies(TFunctionRef<void(ALandscapeProxy*)> Fn) const;
-
 	/** Resets all actors, proxies, components registrations */
 	LANDSCAPE_API void Reset();
 
@@ -260,6 +250,17 @@ public:
 	/** Will call UpdateAllComponentMaterialInstances on all LandscapeProxies */
 	LANDSCAPE_API void UpdateAllComponentMaterialInstances();
 #endif
+
+	/**
+	 * Runs the given function on the root landscape actor and all streaming proxies
+	 * Most easily used with a lambda as follows:
+	 * ForAllLandscapeProxies([](ALandscapeProxy* Proxy)
+	 * {
+	 *     // Code
+	 * });
+	 */
+	LANDSCAPE_API void ForAllLandscapeProxies(TFunctionRef<void(ALandscapeProxy*)> Fn) const;
+
 	/** Associates passed actor with this info object
  *  @param	Proxy		Landscape actor to register
  *  @param  bMapCheck	Whether to warn about landscape errors
@@ -282,6 +283,8 @@ public:
 	LANDSCAPE_API void RegisterCollisionComponent(ULandscapeHeightfieldCollisionComponent* Component);
 	LANDSCAPE_API void UnregisterCollisionComponent(ULandscapeHeightfieldCollisionComponent* Component);
 
+	LANDSCAPE_API void Tick(UWorld* World, float DeltaTime);
+
 #if WITH_EDITOR
 private:
 	bool ApplySplinesInternal(bool bOnlySelected, ALandscapeProxy* Proxy, TSet<ULandscapeComponent*>* OutModifiedComponents, bool bMarkPackageDirty, int32 LandscapeMinX, int32 LandscapeMinY, int32 LandscapeMaxX, int32 LandscapeMaxY, TFunctionRef<TSharedPtr<FModulateAlpha>(ULandscapeLayerInfoObject*)> GetOrCreateModulate);
diff --git a/Engine/Source/Runtime/Landscape/Classes/LandscapeProxy.h b/Engine/Source/Runtime/Landscape/Classes/LandscapeProxy.h
--- a/Engine/Source/Runtime/Landscape/Classes/LandscapeProxy.h
+++ b/Engine/Source/Runtime/Landscape/Classes/LandscapeProxy.h
@@ -38,6 +38,10 @@ struct FLandscapeInfoLayerSettings;
 struct FMeshDescription;
 enum class ENavDataGatheringMode : uint8;
 
+#if WITH_EDITOR
+LANDSCAPE_API extern bool GLandscapeEditModeActive;
+#endif
+
 USTRUCT()
 struct FLandscapeEditorLayerSettings
 {
@@ -824,8 +828,28 @@ public:
 	virtual ALandscape* GetLandscapeActor() PURE_VIRTUAL(GetLandscapeActor, return nullptr;)
 	virtual const ALandscape* GetLandscapeActor() const PURE_VIRTUAL(GetLandscapeActor, return nullptr;)
 
+	static void SetGrassUpdateInterval(int32 Interval) { GrassUpdateInterval = Interval; }
+
 	/* Per-frame call to update dynamic grass placement and render grassmaps */
-	void TickGrass();
+	FORCEINLINE bool ShouldTickGrass() const
+	{
+		if (!bHasLandscapeGrass)
+		{
+			return false;
+		}
+
+		const int32 UpdateInterval = GetGrassUpdateInterval();
+		if (UpdateInterval > 1)
+		{
+			if ((GFrameNumber + FrameOffsetForTickInterval) % uint32(UpdateInterval))
+			{
+				return false;
+			}
+		}
+
+		return true;
+	}
+	void TickGrass(const TArray<FVector>& Cameras);
 
 	/** Flush the grass cache */
 	LANDSCAPE_API void FlushGrassComponents(const TSet<ULandscapeComponent*>* OnlyForComponents = nullptr, bool bFlushGrassMaps = true);
@@ -872,11 +896,6 @@ public:
 	FDelegateHandle FeatureLevelChangedDelegateHandle;
 #endif
 
-	//~ Begin AActor Interface.
-	virtual void TickActor(float DeltaTime, ELevelTick TickType, FActorTickFunction& ThisTickFunction) override;
-	virtual bool ShouldTickIfViewportsOnly() const override;
-	//~ End AActor Interface
-
 	//~ Begin UObject Interface.
 	virtual void PreSave(const class ITargetPlatform* TargetPlatform) override;
 	virtual void Serialize(FArchive& Ar) override;
@@ -889,6 +908,8 @@ public:
 	/** Get the LandcapeActor-to-world transform with respect to landscape section offset*/
 	LANDSCAPE_API FTransform LandscapeActorToWorld() const;
 
+	/** Get landscape position in section space */
+	LANDSCAPE_API FIntPoint GetSectionBaseOffset() const;
 #if WITH_EDITOR
 	LANDSCAPE_API void CreateSplineComponent(const FVector& Scale3D);
 
@@ -926,10 +947,7 @@ public:
 
 	/** Set landscape absolute location in section space */
 	LANDSCAPE_API void SetAbsoluteSectionBase(FIntPoint SectionOffset);
-	
-	/** Get landscape position in section space */
-	LANDSCAPE_API FIntPoint GetSectionBaseOffset() const;
-	
+
 	/** Recreate all components rendering and collision states */
 	LANDSCAPE_API void RecreateComponentsState();
 
@@ -1078,7 +1096,18 @@ protected:
 #endif
 private:
 	/** Returns Grass Update interval */
-	int32 GetGrassUpdateInterval() const;
+	FORCEINLINE int32 GetGrassUpdateInterval() const 
+	{
+#if WITH_EDITOR
+		// When editing landscape, force update interval to be every frame
+		if (GLandscapeEditModeActive)
+		{
+			return 1;
+		}
+#endif
+		return GrassUpdateInterval;
+	}
+	static int32 GrassUpdateInterval;
 
 #if WITH_EDITOR
 	void UpdateGrassDataStatus(TSet<UTexture2D*>& OutCurrentForcedStreamedTextures, TSet<UTexture2D*>* OutDesiredForcedStreamedTextures, TSet<ULandscapeComponent*>& OutComponentsNeedingGrassMapRender, TSet<ULandscapeComponent*>* OutOutdatedComponents, bool bInEnableForceResidentFlag);
diff --git a/Engine/Source/Runtime/Landscape/Private/Landscape.cpp b/Engine/Source/Runtime/Landscape/Private/Landscape.cpp
--- a/Engine/Source/Runtime/Landscape/Private/Landscape.cpp
+++ b/Engine/Source/Runtime/Landscape/Private/Landscape.cpp
@@ -50,6 +50,7 @@ Landscape.cpp: Terrain rendering
 #include "EngineUtils.h"
 #include "ComponentRecreateRenderStateContext.h"
 #include "LandscapeWeightmapUsage.h"
+#include "ContentStreaming.h"
 
 #if WITH_EDITOR
 #include "LandscapeEdit.h"
@@ -64,6 +65,12 @@ Landscape.cpp: Terrain rendering
 #include "UObject/EditorObjectVersion.h"
 #include "Algo/BinarySearch.h"
 
+static int32 GUseStreamingManagerForCameras = 1;
+static FAutoConsoleVariableRef CVarUseStreamingManagerForCameras(
+	TEXT("grass.UseStreamingManagerForCameras"),
+	GUseStreamingManagerForCameras,
+	TEXT("1: Use Streaming Manager; 0: Use ViewLocationsRenderedLastFrame"));
+
 /** Landscape stats */
 
 DEFINE_STAT(STAT_LandscapeDynamicDrawTime);
@@ -1023,12 +1030,6 @@ ALandscapeProxy::ALandscapeProxy(const FObjectInitializer& ObjectInitializer)
 #endif
 	, bHasLandscapeGrass(true)
 {
-	PrimaryActorTick.bCanEverTick = true;
-	PrimaryActorTick.bTickEvenWhenPaused = true;
-	PrimaryActorTick.bStartWithTickEnabled = true;
-	PrimaryActorTick.TickGroup = TG_DuringPhysics;
-	bAllowTickBeforeBeginPlay = true;
-
 	bReplicates = false;
 	NetUpdateFrequency = 10.0f;
 	SetHidden(false);
@@ -1911,11 +1912,9 @@ void ALandscapeProxy::PostRegisterAllComponents()
 
 void ALandscapeProxy::UnregisterAllComponents(const bool bForReregister)
 {
-#if WITH_EDITOR
 	// Game worlds don't have landscape infos
-	if (GetWorld() && !GetWorld()->IsGameWorld()
-		// On shutdown the world will be unreachable
-		&& !GetWorld()->IsPendingKillOrUnreachable() &&
+	// On shutdown the world will be unreachable
+	if (GetWorld() && !GetWorld()->IsPendingKillOrUnreachable() &&
 		// When redoing the creation of a landscape we may get UnregisterAllComponents called when
 		// we are in a "pre-initialized" state (empty guid, etc)
 		LandscapeGuid.IsValid())
@@ -1926,7 +1925,6 @@ void ALandscapeProxy::UnregisterAllComponents(const bool bForReregister)
 			LandscapeInfo->UnregisterActor(this);
 		}
 	}
-#endif
 
 	Super::UnregisterAllComponents(bForReregister);
 }
@@ -2365,13 +2363,6 @@ void ALandscapeProxy::PostLoad()
 {
 	Super::PostLoad();
 
-	// disable ticking if we have no grass to tick
-	if (!GIsEditor && !bHasLandscapeGrass)
-	{
-		SetActorTickEnabled(false);
-		PrimaryActorTick.bCanEverTick = false;
-	}
-
 	// Temporary
 	if (ComponentSizeQuads == 0 && LandscapeComponents.Num() > 0)
 	{
@@ -2452,6 +2443,11 @@ void ALandscapeProxy::PostLoad()
 #endif
 }
 
+FIntPoint ALandscapeProxy::GetSectionBaseOffset() const
+{
+	return LandscapeSectionOffset;
+}
+
 #if WITH_EDITOR
 void ALandscapeProxy::Destroyed()
 {
@@ -2664,11 +2660,6 @@ void ALandscapeProxy::SetAbsoluteSectionBase(FIntPoint InSectionBase)
 	}
 }
 
-FIntPoint ALandscapeProxy::GetSectionBaseOffset() const
-{
-	return LandscapeSectionOffset;
-}
-
 void ALandscapeProxy::RecreateComponentsState()
 {
 	for (int32 ComponentIndex = 0; ComponentIndex < LandscapeComponents.Num(); ComponentIndex++)
@@ -2865,20 +2856,6 @@ ALandscapeProxy* ULandscapeInfo::GetLandscapeProxy() const
 	return nullptr;
 }
 
-void ULandscapeInfo::ForAllLandscapeProxies(TFunctionRef<void(ALandscapeProxy*)> Fn) const
-{
-	ALandscape* Landscape = LandscapeActor.Get();
-	if (Landscape)
-	{
-		Fn(Landscape);
-	}
-
-	for (ALandscapeProxy* LandscapeProxy : Proxies)
-	{
-		Fn(LandscapeProxy);
-	}
-}
-
 void ULandscapeInfo::Reset()
 {
 	LandscapeActor.Reset();
@@ -3017,90 +2994,169 @@ void ULandscapeInfo::RecreateLandscapeInfo(UWorld* InWorld, bool bMapCheck)
 
 #endif
 
-void ULandscapeInfo::RegisterActor(ALandscapeProxy* Proxy, bool bMapCheck)
+void ULandscapeInfo::Tick(UWorld* World, float DeltaTime)
 {
-#if WITH_EDITOR
-	if (!Proxy->GetWorld()->IsGameWorld())
+	// Update Grass
+	static TArray<FVector> OldCameras;
+	TArray<FVector>* Cameras = nullptr;
+	if (GUseStreamingManagerForCameras == 0)
 	{
-		// do not pass here invalid actors
-		checkSlow(Proxy);
-		check(Proxy->GetLandscapeGuid().IsValid());
-		UWorld* OwningWorld = Proxy->GetWorld();
-
-		// in case this Info object is not initialized yet
-		// initialized it with properties from passed actor
-		if (LandscapeGuid.IsValid() == false ||
-			(GetLandscapeProxy() == nullptr && ensure(LandscapeGuid == Proxy->GetLandscapeGuid())))
+		if (OldCameras.Num() || World->ViewLocationsRenderedLastFrame.Num())
 		{
-			LandscapeGuid = Proxy->GetLandscapeGuid();
-			ComponentSizeQuads = Proxy->ComponentSizeQuads;
-			ComponentNumSubsections = Proxy->NumSubsections;
-			SubsectionSizeQuads = Proxy->SubsectionSizeQuads;
-			DrawScale = Proxy->GetRootComponent() != nullptr ? Proxy->GetRootComponent()->GetRelativeScale3D() : FVector(100.0f);
+			Cameras = &OldCameras;
+			// there is a bug here, which often leaves us with no cameras in the editor
+			if (World->ViewLocationsRenderedLastFrame.Num())
+			{
+				check(IsInGameThread());
+				Cameras = &World->ViewLocationsRenderedLastFrame;
+				OldCameras = *Cameras;
+			}
 		}
-
-		// check that passed actor matches all shared parameters
-		check(LandscapeGuid == Proxy->GetLandscapeGuid());
-		check(ComponentSizeQuads == Proxy->ComponentSizeQuads);
-		check(ComponentNumSubsections == Proxy->NumSubsections);
-		check(SubsectionSizeQuads == Proxy->SubsectionSizeQuads);
-
-		if (Proxy->GetRootComponent() != nullptr && !DrawScale.Equals(Proxy->GetRootComponent()->GetRelativeScale3D()))
+	}
+	else
+	{
+		int32 Num = IStreamingManager::Get().GetNumViews();
+		if (Num)
 		{
-			UE_LOG(LogLandscape, Warning, TEXT("Landscape proxy (%s) scale (%s) does not match to main actor scale (%s)."),
-				*Proxy->GetName(), *Proxy->GetRootComponent()->GetRelativeScale3D().ToCompactString(), *DrawScale.ToCompactString());
+			OldCameras.Reset(Num);
+			for (int32 Index = 0; Index < Num; Index++)
+			{
+				auto& ViewInfo = IStreamingManager::Get().GetViewInformation(Index);
+				OldCameras.Add(ViewInfo.ViewOrigin);
+			}
+			Cameras = &OldCameras;
 		}
+	}
 
-		// register
-		if (ALandscape* Landscape = Cast<ALandscape>(Proxy))
+	ForAllLandscapeProxies([DeltaTime,Cameras,World](ALandscapeProxy* Proxy)
+	{
+#if WITH_EDITOR
+		if (GIsEditor)
 		{
-			checkf(!LandscapeActor || LandscapeActor == Landscape, TEXT("Multiple landscapes with the same GUID detected: %s vs %s"), *LandscapeActor->GetPathName(), *Landscape->GetPathName());
-			LandscapeActor = Landscape;
-			// In world composition user is not allowed to move landscape in editor, only through WorldBrowser 
-			LandscapeActor->bLockLocation = OwningWorld != nullptr ? OwningWorld->WorldComposition != nullptr : false;
+			if (ALandscape* Landscape = Cast<ALandscape>(Proxy))
+			{
+				Landscape->TickLayers(DeltaTime);
+			}
 
-			// update proxies reference actor
-			for (ALandscapeStreamingProxy* StreamingProxy : Proxies)
+			// editor-only
+			if (World && !World->IsPlayInEditor())
 			{
-				StreamingProxy->LandscapeActor = LandscapeActor;
-				StreamingProxy->FixupSharedData(Landscape);
+				Proxy->UpdateBakedTextures();
 			}
 		}
-		else
+#endif
+		if (Cameras && Proxy->ShouldTickGrass())
 		{
-			// Insert Proxies in a sorted fashion for generating determisnitic results in the Layer system
-			ALandscapeStreamingProxy* StreamingProxy = CastChecked<ALandscapeStreamingProxy>(Proxy);
-			if (!Proxies.Contains(Proxy))
-			{
-				uint32 InsertIndex = Algo::LowerBound(Proxies, Proxy, [](ALandscapeProxy* A, ALandscapeProxy* B)
-				{
-					FIntPoint SectionBaseA = A->GetSectionBaseOffset();
-					FIntPoint SectionBaseB = B->GetSectionBaseOffset();
+			Proxy->TickGrass(*Cameras);
+		}
+	});
+}
 
-					if (SectionBaseA.X != SectionBaseB.X)
-					{
-						return SectionBaseA.X < SectionBaseB.X;
-					}
+void ULandscapeInfo::ForAllLandscapeProxies(TFunctionRef<void(ALandscapeProxy*)> Fn) const
+{
+	ALandscape* Landscape = LandscapeActor.Get();
+	if (Landscape)
+	{
+		Fn(Landscape);
+	}
 
-					return SectionBaseA.Y < SectionBaseB.Y;
-				});
+	for (ALandscapeProxy* LandscapeProxy : Proxies)
+	{
+		Fn(LandscapeProxy);
+	}
+}
 
-				Proxies.Insert(StreamingProxy, InsertIndex);
-			}
+void ULandscapeInfo::RegisterActor(ALandscapeProxy* Proxy, bool bMapCheck)
+{
+	// do not pass here invalid actors
+	checkSlow(Proxy);
+	check(Proxy->GetLandscapeGuid().IsValid());
+#if WITH_EDITOR
+	UWorld* OwningWorld = Proxy->GetWorld();
+
+	// in case this Info object is not initialized yet
+	// initialized it with properties from passed actor
+	if (LandscapeGuid.IsValid() == false ||
+		(GetLandscapeProxy() == nullptr && ensure(LandscapeGuid == Proxy->GetLandscapeGuid())))
+	{
+		LandscapeGuid = Proxy->GetLandscapeGuid();
+		ComponentSizeQuads = Proxy->ComponentSizeQuads;
+		ComponentNumSubsections = Proxy->NumSubsections;
+		SubsectionSizeQuads = Proxy->SubsectionSizeQuads;
+		DrawScale = Proxy->GetRootComponent() != nullptr ? Proxy->GetRootComponent()->GetRelativeScale3D() : FVector(100.0f);
+	}
+
+	// check that passed actor matches all shared parameters
+	check(LandscapeGuid == Proxy->GetLandscapeGuid());
+	check(ComponentSizeQuads == Proxy->ComponentSizeQuads);
+	check(ComponentNumSubsections == Proxy->NumSubsections);
+	check(SubsectionSizeQuads == Proxy->SubsectionSizeQuads);
+
+	if (Proxy->GetRootComponent() != nullptr && !DrawScale.Equals(Proxy->GetRootComponent()->GetRelativeScale3D()))
+	{
+		UE_LOG(LogLandscape, Warning, TEXT("Landscape proxy (%s) scale (%s) does not match to main actor scale (%s)."),
+			*Proxy->GetName(), *Proxy->GetRootComponent()->GetRelativeScale3D().ToCompactString(), *DrawScale.ToCompactString());
+	}
+#endif
+
+	// register
+	if (ALandscape* Landscape = Cast<ALandscape>(Proxy))
+	{
+		checkf(!LandscapeActor || LandscapeActor == Landscape, TEXT("Multiple landscapes with the same GUID detected: %s vs %s"), *LandscapeActor->GetPathName(), *Landscape->GetPathName());
+		LandscapeActor = Landscape;
+
+#if WITH_EDITOR
+		// In world composition user is not allowed to move landscape in editor, only through WorldBrowser 
+		LandscapeActor->bLockLocation = OwningWorld != nullptr ? OwningWorld->WorldComposition != nullptr : false;
+#endif
+		// update proxies reference actor
+		for (ALandscapeStreamingProxy* StreamingProxy : Proxies)
+		{
 			StreamingProxy->LandscapeActor = LandscapeActor;
-			StreamingProxy->FixupSharedData(LandscapeActor.Get());
+#if WITH_EDITOR
+			StreamingProxy->FixupSharedData(Landscape);
+#endif
 		}
-
-		if (LandscapeActor && LandscapeActor->HasLayersContent())
+	}
+	else
+	{
+		// Insert Proxies in a sorted fashion for generating determisnitic results in the Layer system
+		ALandscapeStreamingProxy* StreamingProxy = CastChecked<ALandscapeStreamingProxy>(Proxy);
+		if (!Proxies.Contains(Proxy))
 		{
-			const bool bInRequestContentUpdate = false;
-			LandscapeActor->RequestLayersInitialization(bInRequestContentUpdate);
+			uint32 InsertIndex = Algo::LowerBound(Proxies, Proxy, [](ALandscapeProxy* A, ALandscapeProxy* B)
+			{
+				FIntPoint SectionBaseA = A->GetSectionBaseOffset();
+				FIntPoint SectionBaseB = B->GetSectionBaseOffset();
+
+				if (SectionBaseA.X != SectionBaseB.X)
+				{
+					return SectionBaseA.X < SectionBaseB.X;
+				}
+
+				return SectionBaseA.Y < SectionBaseB.Y;
+			});
+
+			Proxies.Insert(StreamingProxy, InsertIndex);
 		}
+		StreamingProxy->LandscapeActor = LandscapeActor;
+#if WITH_EDITOR
+		StreamingProxy->FixupSharedData(LandscapeActor.Get());
+#endif
+	}
 
-		UpdateLayerInfoMap(Proxy);
-		UpdateAllAddCollisions();
+#if WITH_EDITOR
+	if (LandscapeActor && LandscapeActor->HasLayersContent())
+	{
+		const bool bInRequestContentUpdate = false;
+		LandscapeActor->RequestLayersInitialization(bInRequestContentUpdate);
 	}
+
+	UpdateLayerInfoMap(Proxy);
+	UpdateAllAddCollisions();
 #endif
+	
+
 	//
 	// add proxy components to the XY map
 	//
@@ -3117,7 +3173,6 @@ void ULandscapeInfo::RegisterActor(ALandscapeProxy* Proxy, bool bMapCheck)
 
 void ULandscapeInfo::UnregisterActor(ALandscapeProxy* Proxy)
 {
-#if WITH_EDITOR
 	if (ALandscape* Landscape = Cast<ALandscape>(Proxy))
 	{
 		// Note: UnregisterActor sometimes gets triggered twice, e.g. it has been observed to happen during redo
@@ -3141,7 +3196,6 @@ void ULandscapeInfo::UnregisterActor(ALandscapeProxy* Proxy)
 		Proxies.Remove(StreamingProxy);
 		StreamingProxy->LandscapeActor = nullptr;
 	}
-#endif
 
 	// remove proxy components from the XY map
 	for (int32 CompIdx = 0; CompIdx < Proxy->LandscapeComponents.Num(); ++CompIdx)
@@ -3409,45 +3463,6 @@ bool LandscapeMaterialsParameterSetUpdater(FStaticParameterSet& StaticParameterS
 	return UpdateParameterSet<FStaticTerrainLayerWeightParameter, UMaterialExpressionLandscapeLayerWeight>(StaticParameterSet.TerrainLayerWeightParameters, ParentMaterial);
 }
 
-bool ALandscapeProxy::ShouldTickIfViewportsOnly() const
-{
-	return true;
-}
-
-void ALandscape::TickActor(float DeltaTime, ELevelTick TickType, FActorTickFunction& ThisTickFunction)
-{
-	Super::TickActor(DeltaTime, TickType, ThisTickFunction);
-
-#if WITH_EDITOR
-	if (GIsEditor)
-	{
-		TickLayers(DeltaTime, TickType, ThisTickFunction);
-	}
-#endif
-}
-
-void ALandscapeProxy::TickActor(float DeltaTime, ELevelTick TickType, FActorTickFunction& ThisTickFunction)
-{
-	CSV_SCOPED_TIMING_STAT_EXCLUSIVE(Landscape);
-	LLM_SCOPE(ELLMTag::Landscape);
-#if WITH_EDITOR
-	// editor-only
-	UWorld* World = GetWorld();
-	if (GIsEditor && World && !World->IsPlayInEditor())
-	{
-		UpdateBakedTextures();
-	}
-#endif
-
-	// Tick grass even while paused or in the editor
-	if (GIsEditor || bHasLandscapeGrass)
-	{
-		TickGrass();
-	}
-
-	Super::TickActor(DeltaTime, TickType, ThisTickFunction);
-}
-
 ALandscapeProxy::~ALandscapeProxy()
 {
 	for (int32 Index = 0; Index < AsyncFoliageTasks.Num(); Index++)
diff --git a/Engine/Source/Runtime/Landscape/Private/LandscapeEditLayers.cpp b/Engine/Source/Runtime/Landscape/Private/LandscapeEditLayers.cpp
--- a/Engine/Source/Runtime/Landscape/Private/LandscapeEditLayers.cpp
+++ b/Engine/Source/Runtime/Landscape/Private/LandscapeEditLayers.cpp
@@ -5155,7 +5155,7 @@ void ALandscape::ForceUpdateLayersContent(bool bInIntermediateRender)
 	bInIntermediateRender = false;
 }
 
-void ALandscape::TickLayers(float DeltaTime, ELevelTick TickType, FActorTickFunction& ThisTickFunction)
+void ALandscape::TickLayers(float DeltaTime)
 {
 	check(GIsEditor);
 
diff --git a/Engine/Source/Runtime/Landscape/Private/LandscapeGrass.cpp b/Engine/Source/Runtime/Landscape/Private/LandscapeGrass.cpp
--- a/Engine/Source/Runtime/Landscape/Private/LandscapeGrass.cpp
+++ b/Engine/Source/Runtime/Landscape/Private/LandscapeGrass.cpp
@@ -25,6 +25,7 @@
 #include "Shader.h"
 #include "Landscape.h"
 #include "LandscapeProxy.h"
+#include "LandscapeInfo.h"
 #include "LightMap.h"
 #include "Engine/MapBuildDataRegistry.h"
 #include "ShadowMap.h"
@@ -37,7 +38,6 @@
 #include "LandscapeGrassType.h"
 #include "Materials/MaterialExpressionLandscapeGrassOutput.h"
 #include "Engine/TextureRenderTarget2D.h"
-#include "ContentStreaming.h"
 #include "LandscapeDataAccess.h"
 #include "StaticMeshResources.h"
 #include "LandscapeLight.h"
@@ -135,12 +135,6 @@ static FAutoConsoleVariableRef CVarGrassDiscardDataOnLoad(
 	TEXT("1: Discard grass data on load (disables grass); 0: Keep grass data (requires reloading level)"),
 	ECVF_Scalability);
 
-static int32 GUseStreamingManagerForCameras = 1;
-static FAutoConsoleVariableRef CVarUseStreamingManagerForCameras(
-	TEXT("grass.UseStreamingManagerForCameras"),
-	GUseStreamingManagerForCameras,
-	TEXT("1: Use Streaming Manager; 0: Use ViewLocationsRenderedLastFrame"));
-
 static int32 GCullSubsections = 1;
 static FAutoConsoleVariableRef CVarCullSubsections(
 	TEXT("grass.CullSubsections"),
@@ -177,7 +171,7 @@ DECLARE_CYCLE_STAT(TEXT("Grass End Comp"), STAT_FoliageGrassEndComp, STATGROUP_F
 DECLARE_CYCLE_STAT(TEXT("Grass Destroy Comps"), STAT_FoliageGrassDestoryComp, STATGROUP_Foliage);
 DECLARE_CYCLE_STAT(TEXT("Grass Update"), STAT_GrassUpdate, STATGROUP_Foliage);
 
-static int32 GGrassUpdateInterval = 1;
+int32 ALandscapeProxy::GrassUpdateInterval = 1;
 
 static void GrassCVarSinkFunction()
 {
@@ -186,7 +180,7 @@ static void GrassCVarSinkFunction()
 
 	if (FApp::IsGame())
 	{
-		GGrassUpdateInterval = FMath::Clamp<int32>(GGrassTickInterval, 1, 60);
+		ALandscapeProxy::SetGrassUpdateInterval(FMath::Clamp<int32>(GGrassTickInterval, 1, 60));
 	}
 
 	static float CachedGrassCullDistanceScale = 1.0f;
@@ -1392,83 +1386,20 @@ void FLandscapeComponentGrassData::ConditionalDiscardDataOnLoad()
 // ALandscapeProxy grass-related functions
 //
 
-int32 ALandscapeProxy::GetGrassUpdateInterval() const
-{
-#if WITH_EDITOR
-	// When editing landscape, force update interval to be every frame
-	if (GLandscapeEditModeActive)
-	{
-		return 1;
-	}
-#endif
-	return GGrassUpdateInterval;
-}
-
-void ALandscapeProxy::TickGrass()
+void ALandscapeProxy::TickGrass(const TArray<FVector>& Cameras)
 {
-	const int32 UpdateInterval = GetGrassUpdateInterval();
-	if (UpdateInterval > 1)
-	{
-		if ((GFrameNumber + FrameOffsetForTickInterval) % uint32(UpdateInterval))
-		{
-			return;
-		}
-	}
-
+	TRACE_CPUPROFILER_EVENT_SCOPE(ALandscapeProxy::TickGrass);
+#if WITH_EDITORONLY_DATA
 	if (ALandscape* Landscape = GetLandscapeActor())
 	{
-		if (!Landscape->IsUpToDate() 
-#if WITH_EDITORONLY_DATA
-			|| !Landscape->bGrassUpdateEnabled
-#endif
-			)
+		if (!Landscape->IsUpToDate() || !Landscape->bGrassUpdateEnabled)
 		{
 			return;
 		}
 	}
+#endif
 
-	// Update foliage
-	static TArray<FVector> OldCameras;
-	if (GUseStreamingManagerForCameras == 0)
-	{
-		UWorld* World = GetWorld();
-		if (!World)
-		{
-			return;
-		}
-
-		if (!OldCameras.Num() && !World->ViewLocationsRenderedLastFrame.Num())
-		{
-			// no cameras, no grass update
-			return;
-		}
-
-		// there is a bug here, which often leaves us with no cameras in the editor
-		const TArray<FVector>& Cameras = World->ViewLocationsRenderedLastFrame.Num() ? World->ViewLocationsRenderedLastFrame : OldCameras;
-
-		if (&Cameras != &OldCameras)
-		{
-			check(IsInGameThread());
-			OldCameras = Cameras;
-		}
-		UpdateGrass(Cameras);
-	}
-	else
-	{
-		int32 Num = IStreamingManager::Get().GetNumViews();
-		if (!Num)
-		{
-			// no cameras, no grass update
-			return;
-		}
-		OldCameras.Reset(Num);
-		for (int32 Index = 0; Index < Num; Index++)
-		{
-			auto& ViewInfo = IStreamingManager::Get().GetViewInformation(Index);
-			OldCameras.Add(ViewInfo.ViewOrigin);
-		}
-		UpdateGrass(OldCameras);
-	}
+	UpdateGrass(Cameras);
 }
 
 struct FGrassBuilderBase
@@ -2311,6 +2242,7 @@ void ALandscapeProxy::UpdateGrassData()
 void ALandscapeProxy::UpdateGrass(const TArray<FVector>& Cameras, bool bForceSync)
 {
 	SCOPE_CYCLE_COUNTER(STAT_GrassUpdate);
+	TRACE_CPUPROFILER_EVENT_SCOPE(ALandscapeProxy::UpdateGrass);
 	if (GFrameNumberLastStaleCheck != GFrameNumber && GIgnoreExcludeBoxes == 0)
 	{
 		GFrameNumberLastStaleCheck = GFrameNumber;
diff --git a/Engine/Source/Runtime/Landscape/Private/LandscapeSubsystem.cpp b/Engine/Source/Runtime/Landscape/Private/LandscapeSubsystem.cpp
new file mode 100644
--- /dev/null
+++ b/Engine/Source/Runtime/Landscape/Private/LandscapeSubsystem.cpp
@@ -0,0 +1,65 @@
+#include "LandscapeSubsystem.h"
+#include "UObject/UObjectGlobals.h"
+#include "Engine\EngineBaseTypes.h"
+#include "LandscapeInfoMap.h"
+#include "LandscapeInfo.h"
+
+DECLARE_CYCLE_STAT(TEXT("LandscapeSubsystem Tick"), STAT_LandscapeSubsystemTick, STATGROUP_Landscape);
+
+ULandscapeSubsystem::ULandscapeSubsystem()
+	: TickFunction(this), LandscapeInfoMap(nullptr)
+{}
+
+ULandscapeSubsystem::~ULandscapeSubsystem()
+{
+
+}
+
+void ULandscapeSubsystem::Initialize(FSubsystemCollectionBase& Collection)
+{
+	// Register Tick 
+	TickFunction.bCanEverTick = true;
+	TickFunction.bTickEvenWhenPaused = true;
+	TickFunction.bStartWithTickEnabled = true;
+	TickFunction.TickGroup = TG_DuringPhysics;
+	TickFunction.RegisterTickFunction(GetWorld()->PersistentLevel);
+}
+
+void ULandscapeSubsystem::Deinitialize()
+{
+	TickFunction.UnRegisterTickFunction();
+	LandscapeInfoMap = nullptr;
+}
+
+void ULandscapeSubsystem::Tick(float DeltaTime, enum ELevelTick TickType, ENamedThreads::Type CurrentThread, const FGraphEventRef& MyCompletionGraphEvent)
+{
+	SCOPE_CYCLE_COUNTER(STAT_LandscapeSubsystemTick);
+	TRACE_CPUPROFILER_EVENT_SCOPE(ULandscapeSubsystem::Tick);
+	UWorld* World = GetWorld();
+	if (!LandscapeInfoMap)
+	{
+		LandscapeInfoMap = &ULandscapeInfoMap::GetLandscapeInfoMap(World);
+	}
+		
+	for (const auto& Pair : LandscapeInfoMap->Map)
+	{
+		Pair.Value->Tick(World, DeltaTime);
+	}
+}
+
+void FLandscapeSubsystemTickFunction::ExecuteTick(float DeltaTime, enum ELevelTick TickType, ENamedThreads::Type CurrentThread, const FGraphEventRef& MyCompletionGraphEvent)
+{
+	Subsystem->Tick(DeltaTime, TickType, CurrentThread, MyCompletionGraphEvent);
+}
+
+FString FLandscapeSubsystemTickFunction::DiagnosticMessage()
+{
+	static const FString Message(TEXT("FLandscapeSubsystemTickFunction"));
+	return Message;
+}
+
+FName FLandscapeSubsystemTickFunction::DiagnosticContext(bool bDetailed)
+{
+	static const FName Context(TEXT("FLandscapeSubsystemTickFunction"));
+	return Context;
+}
\ No newline at end of file
diff --git a/Engine/Source/Runtime/Landscape/Public/LandscapeSubsystem.h b/Engine/Source/Runtime/Landscape/Public/LandscapeSubsystem.h
new file mode 100644
--- /dev/null
+++ b/Engine/Source/Runtime/Landscape/Public/LandscapeSubsystem.h
@@ -0,0 +1,50 @@
+// Copyright Epic Games, Inc. All Rights Reserved.
+
+#pragma once
+
+#include "CoreMinimal.h"
+#include "UObject/ObjectMacros.h"
+#include "Subsystems/WorldSubsystem.h"
+#include "Engine/EngineBaseTypes.h"
+#include "LandscapeSubsystem.generated.h"
+
+class ULandscapeInfoMap;
+
+struct FLandscapeSubsystemTickFunction : public FTickFunction
+{
+	FLandscapeSubsystemTickFunction() : FLandscapeSubsystemTickFunction(nullptr) {}
+	FLandscapeSubsystemTickFunction(ULandscapeSubsystem* InSubsystem) : Subsystem(InSubsystem) {}
+	virtual ~FLandscapeSubsystemTickFunction() {}
+
+	// Begin FTickFunction overrides
+	virtual void ExecuteTick(float DeltaTime, enum ELevelTick TickType, ENamedThreads::Type CurrentThread, const FGraphEventRef& MyCompletionGraphEvent) override;
+	virtual FString DiagnosticMessage() override;
+	virtual FName DiagnosticContext(bool bDetailed) override;
+	// End FTickFunction overrides
+
+	class ULandscapeSubsystem* Subsystem;
+};
+
+UCLASS()
+class ULandscapeSubsystem : public UWorldSubsystem, public FTickFunction
+{
+	GENERATED_BODY()
+
+	ULandscapeSubsystem();
+	virtual ~ULandscapeSubsystem();
+
+public:
+	// Begin USubsystem
+	virtual void Initialize(FSubsystemCollectionBase& Collection) override;
+	virtual void Deinitialize() override;
+	// End USubsystem
+		
+private:
+	void Tick(float DeltaTime, enum ELevelTick TickType, ENamedThreads::Type CurrentThread, const FGraphEventRef& MyCompletionGraphEvent);
+
+	friend struct FLandscapeSubsystemTickFunction;
+	FLandscapeSubsystemTickFunction TickFunction;
+
+	ULandscapeInfoMap* LandscapeInfoMap;
+};
+
