Back out changelist 11033080 (until crash is resolved)

#rnx

diff --git a/Engine/Source/Runtime/Landscape/Classes/Landscape.h b/Engine/Source/Runtime/Landscape/Classes/Landscape.h
--- a/Engine/Source/Runtime/Landscape/Classes/Landscape.h
+++ b/Engine/Source/Runtime/Landscape/Classes/Landscape.h
@@ -223,6 +223,8 @@ class ALandscape : public ALandscapeProxy
 public:
 	ALandscape(const FObjectInitializer& ObjectInitializer);
 
+	virtual void TickActor(float DeltaTime, ELevelTick TickType, FActorTickFunction& ThisTickFunction) override;
+
 	//~ Begin ALandscapeProxy Interface
 	LANDSCAPE_API virtual ALandscape* GetLandscapeActor() override;
 	LANDSCAPE_API virtual const ALandscape* GetLandscapeActor() const override;
@@ -259,7 +261,7 @@ public:
 	//~ End UObject Interface
 
 	LANDSCAPE_API bool IsUpToDate() const;
-	LANDSCAPE_API void TickLayers(float DeltaTime);
+
 	// Layers stuff
 #if WITH_EDITOR
 	LANDSCAPE_API void RegisterLandscapeEdMode(ILandscapeEdModeInterface* InLandscapeEdMode) { LandscapeEdMode = InLandscapeEdMode; }
@@ -334,6 +336,7 @@ public:
 	LANDSCAPE_API void InitializeLandscapeLayersWeightmapUsage();
 
 private:
+	void TickLayers(float DeltaTime, ELevelTick TickType, FActorTickFunction& ThisTickFunction);
 	void CreateLayersRenderingResource();
 	void GetLandscapeComponentNeighborsToRender(ULandscapeComponent* LandscapeComponent, TSet<ULandscapeComponent*>& NeighborComponents) const;
 	void GetLandscapeComponentWeightmapsToRender(ULandscapeComponent* LandscapeComponent, TSet<ULandscapeComponent*>& WeightmapComponents) const;
diff --git a/Engine/Source/Runtime/Landscape/Classes/LandscapeInfo.h b/Engine/Source/Runtime/Landscape/Classes/LandscapeInfo.h
--- a/Engine/Source/Runtime/Landscape/Classes/LandscapeInfo.h
+++ b/Engine/Source/Runtime/Landscape/Classes/LandscapeInfo.h
@@ -206,6 +206,16 @@ public:
 	 */
 	LANDSCAPE_API ALandscapeProxy* GetLandscapeProxy() const;
 
+	/**
+	 * Runs the given function on the root landscape actor and all streaming proxies
+	 * Most easily used with a lambda as follows:
+	 * ForAllLandscapeProxies([](ALandscapeProxy* Proxy)
+	 * {
+	 *     // Code
+	 * });
+	 */
+	LANDSCAPE_API void ForAllLandscapeProxies(TFunctionRef<void(ALandscapeProxy*)> Fn) const;
+
 	/** Resets all actors, proxies, components registrations */
 	LANDSCAPE_API void Reset();
 
@@ -250,17 +260,6 @@ public:
 	/** Will call UpdateAllComponentMaterialInstances on all LandscapeProxies */
 	LANDSCAPE_API void UpdateAllComponentMaterialInstances();
 #endif
-
-	/**
-	 * Runs the given function on the root landscape actor and all streaming proxies
-	 * Most easily used with a lambda as follows:
-	 * ForAllLandscapeProxies([](ALandscapeProxy* Proxy)
-	 * {
-	 *     // Code
-	 * });
-	 */
-	LANDSCAPE_API void ForAllLandscapeProxies(TFunctionRef<void(ALandscapeProxy*)> Fn) const;
-
 	/** Associates passed actor with this info object
  *  @param	Proxy		Landscape actor to register
  *  @param  bMapCheck	Whether to warn about landscape errors
@@ -283,8 +282,6 @@ public:
 	LANDSCAPE_API void RegisterCollisionComponent(ULandscapeHeightfieldCollisionComponent* Component);
 	LANDSCAPE_API void UnregisterCollisionComponent(ULandscapeHeightfieldCollisionComponent* Component);
 
-	LANDSCAPE_API void Tick(UWorld* World, float DeltaTime);
-
 #if WITH_EDITOR
 private:
 	bool ApplySplinesInternal(bool bOnlySelected, ALandscapeProxy* Proxy, TSet<ULandscapeComponent*>* OutModifiedComponents, bool bMarkPackageDirty, int32 LandscapeMinX, int32 LandscapeMinY, int32 LandscapeMaxX, int32 LandscapeMaxY, TFunctionRef<TSharedPtr<FModulateAlpha>(ULandscapeLayerInfoObject*)> GetOrCreateModulate);
diff --git a/Engine/Source/Runtime/Landscape/Classes/LandscapeProxy.h b/Engine/Source/Runtime/Landscape/Classes/LandscapeProxy.h
--- a/Engine/Source/Runtime/Landscape/Classes/LandscapeProxy.h
+++ b/Engine/Source/Runtime/Landscape/Classes/LandscapeProxy.h
@@ -38,10 +38,6 @@ struct FLandscapeInfoLayerSettings;
 struct FMeshDescription;
 enum class ENavDataGatheringMode : uint8;
 
-#if WITH_EDITOR
-LANDSCAPE_API extern bool GLandscapeEditModeActive;
-#endif
-
 USTRUCT()
 struct FLandscapeEditorLayerSettings
 {
@@ -828,28 +824,8 @@ public:
 	virtual ALandscape* GetLandscapeActor() PURE_VIRTUAL(GetLandscapeActor, return nullptr;)
 	virtual const ALandscape* GetLandscapeActor() const PURE_VIRTUAL(GetLandscapeActor, return nullptr;)
 
-	static void SetGrassUpdateInterval(int32 Interval) { GrassUpdateInterval = Interval; }
-
 	/* Per-frame call to update dynamic grass placement and render grassmaps */
-	FORCEINLINE bool ShouldTickGrass() const
-	{
-		if (!bHasLandscapeGrass)
-		{
-			return false;
-		}
-
-		const int32 UpdateInterval = GetGrassUpdateInterval();
-		if (UpdateInterval > 1)
-		{
-			if ((GFrameNumber + FrameOffsetForTickInterval) % uint32(UpdateInterval))
-			{
-				return false;
-			}
-		}
-
-		return true;
-	}
-	void TickGrass(const TArray<FVector>& Cameras);
+	void TickGrass();
 
 	/** Flush the grass cache */
 	LANDSCAPE_API void FlushGrassComponents(const TSet<ULandscapeComponent*>* OnlyForComponents = nullptr, bool bFlushGrassMaps = true);
@@ -896,6 +872,11 @@ public:
 	FDelegateHandle FeatureLevelChangedDelegateHandle;
 #endif
 
+	//~ Begin AActor Interface.
+	virtual void TickActor(float DeltaTime, ELevelTick TickType, FActorTickFunction& ThisTickFunction) override;
+	virtual bool ShouldTickIfViewportsOnly() const override;
+	//~ End AActor Interface
+
 	//~ Begin UObject Interface.
 	virtual void PreSave(const class ITargetPlatform* TargetPlatform) override;
 	virtual void Serialize(FArchive& Ar) override;
@@ -908,8 +889,6 @@ public:
 	/** Get the LandcapeActor-to-world transform with respect to landscape section offset*/
 	LANDSCAPE_API FTransform LandscapeActorToWorld() const;
 
-	/** Get landscape position in section space */
-	LANDSCAPE_API FIntPoint GetSectionBaseOffset() const;
 #if WITH_EDITOR
 	LANDSCAPE_API void CreateSplineComponent(const FVector& Scale3D);
 
@@ -947,7 +926,10 @@ public:
 
 	/** Set landscape absolute location in section space */
 	LANDSCAPE_API void SetAbsoluteSectionBase(FIntPoint SectionOffset);
-
+	
+	/** Get landscape position in section space */
+	LANDSCAPE_API FIntPoint GetSectionBaseOffset() const;
+	
 	/** Recreate all components rendering and collision states */
 	LANDSCAPE_API void RecreateComponentsState();
 
@@ -1096,18 +1078,7 @@ protected:
 #endif
 private:
 	/** Returns Grass Update interval */
-	FORCEINLINE int32 GetGrassUpdateInterval() const 
-	{
-#if WITH_EDITOR
-		// When editing landscape, force update interval to be every frame
-		if (GLandscapeEditModeActive)
-		{
-			return 1;
-		}
-#endif
-		return GrassUpdateInterval;
-	}
-	static int32 GrassUpdateInterval;
+	int32 GetGrassUpdateInterval() const;
 
 #if WITH_EDITOR
 	void UpdateGrassDataStatus(TSet<UTexture2D*>& OutCurrentForcedStreamedTextures, TSet<UTexture2D*>* OutDesiredForcedStreamedTextures, TSet<ULandscapeComponent*>& OutComponentsNeedingGrassMapRender, TSet<ULandscapeComponent*>* OutOutdatedComponents, bool bInEnableForceResidentFlag);
diff --git a/Engine/Source/Runtime/Landscape/Private/Landscape.cpp b/Engine/Source/Runtime/Landscape/Private/Landscape.cpp
--- a/Engine/Source/Runtime/Landscape/Private/Landscape.cpp
+++ b/Engine/Source/Runtime/Landscape/Private/Landscape.cpp
@@ -50,7 +50,6 @@ Landscape.cpp: Terrain rendering
 #include "EngineUtils.h"
 #include "ComponentRecreateRenderStateContext.h"
 #include "LandscapeWeightmapUsage.h"
-#include "ContentStreaming.h"
 
 #if WITH_EDITOR
 #include "LandscapeEdit.h"
@@ -65,12 +64,6 @@ Landscape.cpp: Terrain rendering
 #include "UObject/EditorObjectVersion.h"
 #include "Algo/BinarySearch.h"
 
-static int32 GUseStreamingManagerForCameras = 1;
-static FAutoConsoleVariableRef CVarUseStreamingManagerForCameras(
-	TEXT("grass.UseStreamingManagerForCameras"),
-	GUseStreamingManagerForCameras,
-	TEXT("1: Use Streaming Manager; 0: Use ViewLocationsRenderedLastFrame"));
-
 /** Landscape stats */
 
 DEFINE_STAT(STAT_LandscapeDynamicDrawTime);
@@ -1030,6 +1023,12 @@ ALandscapeProxy::ALandscapeProxy(const FObjectInitializer& ObjectInitializer)
 #endif
 	, bHasLandscapeGrass(true)
 {
+	PrimaryActorTick.bCanEverTick = true;
+	PrimaryActorTick.bTickEvenWhenPaused = true;
+	PrimaryActorTick.bStartWithTickEnabled = true;
+	PrimaryActorTick.TickGroup = TG_DuringPhysics;
+	bAllowTickBeforeBeginPlay = true;
+
 	bReplicates = false;
 	NetUpdateFrequency = 10.0f;
 	SetHidden(false);
@@ -1912,9 +1911,11 @@ void ALandscapeProxy::PostRegisterAllComponents()
 
 void ALandscapeProxy::UnregisterAllComponents(const bool bForReregister)
 {
+#if WITH_EDITOR
 	// Game worlds don't have landscape infos
-	// On shutdown the world will be unreachable
-	if (GetWorld() && !GetWorld()->IsPendingKillOrUnreachable() &&
+	if (GetWorld() && !GetWorld()->IsGameWorld()
+		// On shutdown the world will be unreachable
+		&& !GetWorld()->IsPendingKillOrUnreachable() &&
 		// When redoing the creation of a landscape we may get UnregisterAllComponents called when
 		// we are in a "pre-initialized" state (empty guid, etc)
 		LandscapeGuid.IsValid())
@@ -1925,6 +1926,7 @@ void ALandscapeProxy::UnregisterAllComponents(const bool bForReregister)
 			LandscapeInfo->UnregisterActor(this);
 		}
 	}
+#endif
 
 	Super::UnregisterAllComponents(bForReregister);
 }
@@ -2363,6 +2365,13 @@ void ALandscapeProxy::PostLoad()
 {
 	Super::PostLoad();
 
+	// disable ticking if we have no grass to tick
+	if (!GIsEditor && !bHasLandscapeGrass)
+	{
+		SetActorTickEnabled(false);
+		PrimaryActorTick.bCanEverTick = false;
+	}
+
 	// Temporary
 	if (ComponentSizeQuads == 0 && LandscapeComponents.Num() > 0)
 	{
@@ -2443,11 +2452,6 @@ void ALandscapeProxy::PostLoad()
 #endif
 }
 
-FIntPoint ALandscapeProxy::GetSectionBaseOffset() const
-{
-	return LandscapeSectionOffset;
-}
-
 #if WITH_EDITOR
 void ALandscapeProxy::Destroyed()
 {
@@ -2660,6 +2664,11 @@ void ALandscapeProxy::SetAbsoluteSectionBase(FIntPoint InSectionBase)
 	}
 }
 
+FIntPoint ALandscapeProxy::GetSectionBaseOffset() const
+{
+	return LandscapeSectionOffset;
+}
+
 void ALandscapeProxy::RecreateComponentsState()
 {
 	for (int32 ComponentIndex = 0; ComponentIndex < LandscapeComponents.Num(); ComponentIndex++)
@@ -2856,6 +2865,20 @@ ALandscapeProxy* ULandscapeInfo::GetLandscapeProxy() const
 	return nullptr;
 }
 
+void ULandscapeInfo::ForAllLandscapeProxies(TFunctionRef<void(ALandscapeProxy*)> Fn) const
+{
+	ALandscape* Landscape = LandscapeActor.Get();
+	if (Landscape)
+	{
+		Fn(Landscape);
+	}
+
+	for (ALandscapeProxy* LandscapeProxy : Proxies)
+	{
+		Fn(LandscapeProxy);
+	}
+}
+
 void ULandscapeInfo::Reset()
 {
 	LandscapeActor.Reset();
@@ -2994,169 +3017,90 @@ void ULandscapeInfo::RecreateLandscapeInfo(UWorld* InWorld, bool bMapCheck)
 
 #endif
 
-void ULandscapeInfo::Tick(UWorld* World, float DeltaTime)
+void ULandscapeInfo::RegisterActor(ALandscapeProxy* Proxy, bool bMapCheck)
 {
-	// Update Grass
-	static TArray<FVector> OldCameras;
-	TArray<FVector>* Cameras = nullptr;
-	if (GUseStreamingManagerForCameras == 0)
+#if WITH_EDITOR
+	if (!Proxy->GetWorld()->IsGameWorld())
 	{
-		if (OldCameras.Num() || World->ViewLocationsRenderedLastFrame.Num())
+		// do not pass here invalid actors
+		checkSlow(Proxy);
+		check(Proxy->GetLandscapeGuid().IsValid());
+		UWorld* OwningWorld = Proxy->GetWorld();
+
+		// in case this Info object is not initialized yet
+		// initialized it with properties from passed actor
+		if (LandscapeGuid.IsValid() == false ||
+			(GetLandscapeProxy() == nullptr && ensure(LandscapeGuid == Proxy->GetLandscapeGuid())))
 		{
-			Cameras = &OldCameras;
-			// there is a bug here, which often leaves us with no cameras in the editor
-			if (World->ViewLocationsRenderedLastFrame.Num())
-			{
-				check(IsInGameThread());
-				Cameras = &World->ViewLocationsRenderedLastFrame;
-				OldCameras = *Cameras;
-			}
+			LandscapeGuid = Proxy->GetLandscapeGuid();
+			ComponentSizeQuads = Proxy->ComponentSizeQuads;
+			ComponentNumSubsections = Proxy->NumSubsections;
+			SubsectionSizeQuads = Proxy->SubsectionSizeQuads;
+			DrawScale = Proxy->GetRootComponent() != nullptr ? Proxy->GetRootComponent()->GetRelativeScale3D() : FVector(100.0f);
 		}
-	}
-	else
-	{
-		int32 Num = IStreamingManager::Get().GetNumViews();
-		if (Num)
+
+		// check that passed actor matches all shared parameters
+		check(LandscapeGuid == Proxy->GetLandscapeGuid());
+		check(ComponentSizeQuads == Proxy->ComponentSizeQuads);
+		check(ComponentNumSubsections == Proxy->NumSubsections);
+		check(SubsectionSizeQuads == Proxy->SubsectionSizeQuads);
+
+		if (Proxy->GetRootComponent() != nullptr && !DrawScale.Equals(Proxy->GetRootComponent()->GetRelativeScale3D()))
 		{
-			OldCameras.Reset(Num);
-			for (int32 Index = 0; Index < Num; Index++)
-			{
-				auto& ViewInfo = IStreamingManager::Get().GetViewInformation(Index);
-				OldCameras.Add(ViewInfo.ViewOrigin);
-			}
-			Cameras = &OldCameras;
+			UE_LOG(LogLandscape, Warning, TEXT("Landscape proxy (%s) scale (%s) does not match to main actor scale (%s)."),
+				*Proxy->GetName(), *Proxy->GetRootComponent()->GetRelativeScale3D().ToCompactString(), *DrawScale.ToCompactString());
 		}
-	}
 
-	ForAllLandscapeProxies([DeltaTime,Cameras,World](ALandscapeProxy* Proxy)
-	{
-#if WITH_EDITOR
-		if (GIsEditor)
+		// register
+		if (ALandscape* Landscape = Cast<ALandscape>(Proxy))
 		{
-			if (ALandscape* Landscape = Cast<ALandscape>(Proxy))
-			{
-				Landscape->TickLayers(DeltaTime);
-			}
+			checkf(!LandscapeActor || LandscapeActor == Landscape, TEXT("Multiple landscapes with the same GUID detected: %s vs %s"), *LandscapeActor->GetPathName(), *Landscape->GetPathName());
+			LandscapeActor = Landscape;
+			// In world composition user is not allowed to move landscape in editor, only through WorldBrowser 
+			LandscapeActor->bLockLocation = OwningWorld != nullptr ? OwningWorld->WorldComposition != nullptr : false;
 
-			// editor-only
-			if (World && !World->IsPlayInEditor())
+			// update proxies reference actor
+			for (ALandscapeStreamingProxy* StreamingProxy : Proxies)
 			{
-				Proxy->UpdateBakedTextures();
+				StreamingProxy->LandscapeActor = LandscapeActor;
+				StreamingProxy->FixupSharedData(Landscape);
 			}
 		}
-#endif
-		if (Cameras && Proxy->ShouldTickGrass())
+		else
 		{
-			Proxy->TickGrass(*Cameras);
-		}
-	});
-}
-
-void ULandscapeInfo::ForAllLandscapeProxies(TFunctionRef<void(ALandscapeProxy*)> Fn) const
-{
-	ALandscape* Landscape = LandscapeActor.Get();
-	if (Landscape)
-	{
-		Fn(Landscape);
-	}
-
-	for (ALandscapeProxy* LandscapeProxy : Proxies)
-	{
-		Fn(LandscapeProxy);
-	}
-}
-
-void ULandscapeInfo::RegisterActor(ALandscapeProxy* Proxy, bool bMapCheck)
-{
-	// do not pass here invalid actors
-	checkSlow(Proxy);
-	check(Proxy->GetLandscapeGuid().IsValid());
-#if WITH_EDITOR
-	UWorld* OwningWorld = Proxy->GetWorld();
-
-	// in case this Info object is not initialized yet
-	// initialized it with properties from passed actor
-	if (LandscapeGuid.IsValid() == false ||
-		(GetLandscapeProxy() == nullptr && ensure(LandscapeGuid == Proxy->GetLandscapeGuid())))
-	{
-		LandscapeGuid = Proxy->GetLandscapeGuid();
-		ComponentSizeQuads = Proxy->ComponentSizeQuads;
-		ComponentNumSubsections = Proxy->NumSubsections;
-		SubsectionSizeQuads = Proxy->SubsectionSizeQuads;
-		DrawScale = Proxy->GetRootComponent() != nullptr ? Proxy->GetRootComponent()->GetRelativeScale3D() : FVector(100.0f);
-	}
-
-	// check that passed actor matches all shared parameters
-	check(LandscapeGuid == Proxy->GetLandscapeGuid());
-	check(ComponentSizeQuads == Proxy->ComponentSizeQuads);
-	check(ComponentNumSubsections == Proxy->NumSubsections);
-	check(SubsectionSizeQuads == Proxy->SubsectionSizeQuads);
+			// Insert Proxies in a sorted fashion for generating determisnitic results in the Layer system
+			ALandscapeStreamingProxy* StreamingProxy = CastChecked<ALandscapeStreamingProxy>(Proxy);
+			if (!Proxies.Contains(Proxy))
+			{
+				uint32 InsertIndex = Algo::LowerBound(Proxies, Proxy, [](ALandscapeProxy* A, ALandscapeProxy* B)
+				{
+					FIntPoint SectionBaseA = A->GetSectionBaseOffset();
+					FIntPoint SectionBaseB = B->GetSectionBaseOffset();
 
-	if (Proxy->GetRootComponent() != nullptr && !DrawScale.Equals(Proxy->GetRootComponent()->GetRelativeScale3D()))
-	{
-		UE_LOG(LogLandscape, Warning, TEXT("Landscape proxy (%s) scale (%s) does not match to main actor scale (%s)."),
-			*Proxy->GetName(), *Proxy->GetRootComponent()->GetRelativeScale3D().ToCompactString(), *DrawScale.ToCompactString());
-	}
-#endif
+					if (SectionBaseA.X != SectionBaseB.X)
+					{
+						return SectionBaseA.X < SectionBaseB.X;
+					}
 
-	// register
-	if (ALandscape* Landscape = Cast<ALandscape>(Proxy))
-	{
-		checkf(!LandscapeActor || LandscapeActor == Landscape, TEXT("Multiple landscapes with the same GUID detected: %s vs %s"), *LandscapeActor->GetPathName(), *Landscape->GetPathName());
-		LandscapeActor = Landscape;
+					return SectionBaseA.Y < SectionBaseB.Y;
+				});
 
-#if WITH_EDITOR
-		// In world composition user is not allowed to move landscape in editor, only through WorldBrowser 
-		LandscapeActor->bLockLocation = OwningWorld != nullptr ? OwningWorld->WorldComposition != nullptr : false;
-#endif
-		// update proxies reference actor
-		for (ALandscapeStreamingProxy* StreamingProxy : Proxies)
-		{
+				Proxies.Insert(StreamingProxy, InsertIndex);
+			}
 			StreamingProxy->LandscapeActor = LandscapeActor;
-#if WITH_EDITOR
-			StreamingProxy->FixupSharedData(Landscape);
-#endif
+			StreamingProxy->FixupSharedData(LandscapeActor.Get());
 		}
-	}
-	else
-	{
-		// Insert Proxies in a sorted fashion for generating determisnitic results in the Layer system
-		ALandscapeStreamingProxy* StreamingProxy = CastChecked<ALandscapeStreamingProxy>(Proxy);
-		if (!Proxies.Contains(Proxy))
-		{
-			uint32 InsertIndex = Algo::LowerBound(Proxies, Proxy, [](ALandscapeProxy* A, ALandscapeProxy* B)
-			{
-				FIntPoint SectionBaseA = A->GetSectionBaseOffset();
-				FIntPoint SectionBaseB = B->GetSectionBaseOffset();
-
-				if (SectionBaseA.X != SectionBaseB.X)
-				{
-					return SectionBaseA.X < SectionBaseB.X;
-				}
 
-				return SectionBaseA.Y < SectionBaseB.Y;
-			});
-
-			Proxies.Insert(StreamingProxy, InsertIndex);
+		if (LandscapeActor && LandscapeActor->HasLayersContent())
+		{
+			const bool bInRequestContentUpdate = false;
+			LandscapeActor->RequestLayersInitialization(bInRequestContentUpdate);
 		}
-		StreamingProxy->LandscapeActor = LandscapeActor;
-#if WITH_EDITOR
-		StreamingProxy->FixupSharedData(LandscapeActor.Get());
-#endif
-	}
 
-#if WITH_EDITOR
-	if (LandscapeActor && LandscapeActor->HasLayersContent())
-	{
-		const bool bInRequestContentUpdate = false;
-		LandscapeActor->RequestLayersInitialization(bInRequestContentUpdate);
+		UpdateLayerInfoMap(Proxy);
+		UpdateAllAddCollisions();
 	}
-
-	UpdateLayerInfoMap(Proxy);
-	UpdateAllAddCollisions();
 #endif
-	
-
 	//
 	// add proxy components to the XY map
 	//
@@ -3173,6 +3117,7 @@ void ULandscapeInfo::RegisterActor(ALandscapeProxy* Proxy, bool bMapCheck)
 
 void ULandscapeInfo::UnregisterActor(ALandscapeProxy* Proxy)
 {
+#if WITH_EDITOR
 	if (ALandscape* Landscape = Cast<ALandscape>(Proxy))
 	{
 		// Note: UnregisterActor sometimes gets triggered twice, e.g. it has been observed to happen during redo
@@ -3196,6 +3141,7 @@ void ULandscapeInfo::UnregisterActor(ALandscapeProxy* Proxy)
 		Proxies.Remove(StreamingProxy);
 		StreamingProxy->LandscapeActor = nullptr;
 	}
+#endif
 
 	// remove proxy components from the XY map
 	for (int32 CompIdx = 0; CompIdx < Proxy->LandscapeComponents.Num(); ++CompIdx)
@@ -3463,6 +3409,45 @@ bool LandscapeMaterialsParameterSetUpdater(FStaticParameterSet& StaticParameterS
 	return UpdateParameterSet<FStaticTerrainLayerWeightParameter, UMaterialExpressionLandscapeLayerWeight>(StaticParameterSet.TerrainLayerWeightParameters, ParentMaterial);
 }
 
+bool ALandscapeProxy::ShouldTickIfViewportsOnly() const
+{
+	return true;
+}
+
+void ALandscape::TickActor(float DeltaTime, ELevelTick TickType, FActorTickFunction& ThisTickFunction)
+{
+	Super::TickActor(DeltaTime, TickType, ThisTickFunction);
+
+#if WITH_EDITOR
+	if (GIsEditor)
+	{
+		TickLayers(DeltaTime, TickType, ThisTickFunction);
+	}
+#endif
+}
+
+void ALandscapeProxy::TickActor(float DeltaTime, ELevelTick TickType, FActorTickFunction& ThisTickFunction)
+{
+	CSV_SCOPED_TIMING_STAT_EXCLUSIVE(Landscape);
+	LLM_SCOPE(ELLMTag::Landscape);
+#if WITH_EDITOR
+	// editor-only
+	UWorld* World = GetWorld();
+	if (GIsEditor && World && !World->IsPlayInEditor())
+	{
+		UpdateBakedTextures();
+	}
+#endif
+
+	// Tick grass even while paused or in the editor
+	if (GIsEditor || bHasLandscapeGrass)
+	{
+		TickGrass();
+	}
+
+	Super::TickActor(DeltaTime, TickType, ThisTickFunction);
+}
+
 ALandscapeProxy::~ALandscapeProxy()
 {
 	for (int32 Index = 0; Index < AsyncFoliageTasks.Num(); Index++)
diff --git a/Engine/Source/Runtime/Landscape/Private/LandscapeEditLayers.cpp b/Engine/Source/Runtime/Landscape/Private/LandscapeEditLayers.cpp
--- a/Engine/Source/Runtime/Landscape/Private/LandscapeEditLayers.cpp
+++ b/Engine/Source/Runtime/Landscape/Private/LandscapeEditLayers.cpp
@@ -5155,7 +5155,7 @@ void ALandscape::ForceUpdateLayersContent(bool bInIntermediateRender)
 	bInIntermediateRender = false;
 }
 
-void ALandscape::TickLayers(float DeltaTime)
+void ALandscape::TickLayers(float DeltaTime, ELevelTick TickType, FActorTickFunction& ThisTickFunction)
 {
 	check(GIsEditor);
 
diff --git a/Engine/Source/Runtime/Landscape/Private/LandscapeGrass.cpp b/Engine/Source/Runtime/Landscape/Private/LandscapeGrass.cpp
--- a/Engine/Source/Runtime/Landscape/Private/LandscapeGrass.cpp
+++ b/Engine/Source/Runtime/Landscape/Private/LandscapeGrass.cpp
@@ -25,7 +25,6 @@
 #include "Shader.h"
 #include "Landscape.h"
 #include "LandscapeProxy.h"
-#include "LandscapeInfo.h"
 #include "LightMap.h"
 #include "Engine/MapBuildDataRegistry.h"
 #include "ShadowMap.h"
@@ -38,6 +37,7 @@
 #include "LandscapeGrassType.h"
 #include "Materials/MaterialExpressionLandscapeGrassOutput.h"
 #include "Engine/TextureRenderTarget2D.h"
+#include "ContentStreaming.h"
 #include "LandscapeDataAccess.h"
 #include "StaticMeshResources.h"
 #include "LandscapeLight.h"
@@ -135,6 +135,12 @@ static FAutoConsoleVariableRef CVarGrassDiscardDataOnLoad(
 	TEXT("1: Discard grass data on load (disables grass); 0: Keep grass data (requires reloading level)"),
 	ECVF_Scalability);
 
+static int32 GUseStreamingManagerForCameras = 1;
+static FAutoConsoleVariableRef CVarUseStreamingManagerForCameras(
+	TEXT("grass.UseStreamingManagerForCameras"),
+	GUseStreamingManagerForCameras,
+	TEXT("1: Use Streaming Manager; 0: Use ViewLocationsRenderedLastFrame"));
+
 static int32 GCullSubsections = 1;
 static FAutoConsoleVariableRef CVarCullSubsections(
 	TEXT("grass.CullSubsections"),
@@ -171,7 +177,7 @@ DECLARE_CYCLE_STAT(TEXT("Grass End Comp"), STAT_FoliageGrassEndComp, STATGROUP_F
 DECLARE_CYCLE_STAT(TEXT("Grass Destroy Comps"), STAT_FoliageGrassDestoryComp, STATGROUP_Foliage);
 DECLARE_CYCLE_STAT(TEXT("Grass Update"), STAT_GrassUpdate, STATGROUP_Foliage);
 
-int32 ALandscapeProxy::GrassUpdateInterval = 1;
+static int32 GGrassUpdateInterval = 1;
 
 static void GrassCVarSinkFunction()
 {
@@ -180,7 +186,7 @@ static void GrassCVarSinkFunction()
 
 	if (FApp::IsGame())
 	{
-		ALandscapeProxy::SetGrassUpdateInterval(FMath::Clamp<int32>(GGrassTickInterval, 1, 60));
+		GGrassUpdateInterval = FMath::Clamp<int32>(GGrassTickInterval, 1, 60);
 	}
 
 	static float CachedGrassCullDistanceScale = 1.0f;
@@ -1386,20 +1392,83 @@ void FLandscapeComponentGrassData::ConditionalDiscardDataOnLoad()
 // ALandscapeProxy grass-related functions
 //
 
-void ALandscapeProxy::TickGrass(const TArray<FVector>& Cameras)
+int32 ALandscapeProxy::GetGrassUpdateInterval() const
 {
-	TRACE_CPUPROFILER_EVENT_SCOPE(ALandscapeProxy::TickGrass);
-#if WITH_EDITORONLY_DATA
-	if (ALandscape* Landscape = GetLandscapeActor())
+#if WITH_EDITOR
+	// When editing landscape, force update interval to be every frame
+	if (GLandscapeEditModeActive)
+	{
+		return 1;
+	}
+#endif
+	return GGrassUpdateInterval;
+}
+
+void ALandscapeProxy::TickGrass()
+{
+	const int32 UpdateInterval = GetGrassUpdateInterval();
+	if (UpdateInterval > 1)
 	{
-		if (!Landscape->IsUpToDate() || !Landscape->bGrassUpdateEnabled)
+		if ((GFrameNumber + FrameOffsetForTickInterval) % uint32(UpdateInterval))
 		{
 			return;
 		}
 	}
+
+	if (ALandscape* Landscape = GetLandscapeActor())
+	{
+		if (!Landscape->IsUpToDate() 
+#if WITH_EDITORONLY_DATA
+			|| !Landscape->bGrassUpdateEnabled
 #endif
+			)
+		{
+			return;
+		}
+	}
+
+	// Update foliage
+	static TArray<FVector> OldCameras;
+	if (GUseStreamingManagerForCameras == 0)
+	{
+		UWorld* World = GetWorld();
+		if (!World)
+		{
+			return;
+		}
+
+		if (!OldCameras.Num() && !World->ViewLocationsRenderedLastFrame.Num())
+		{
+			// no cameras, no grass update
+			return;
+		}
 
-	UpdateGrass(Cameras);
+		// there is a bug here, which often leaves us with no cameras in the editor
+		const TArray<FVector>& Cameras = World->ViewLocationsRenderedLastFrame.Num() ? World->ViewLocationsRenderedLastFrame : OldCameras;
+
+		if (&Cameras != &OldCameras)
+		{
+			check(IsInGameThread());
+			OldCameras = Cameras;
+		}
+		UpdateGrass(Cameras);
+	}
+	else
+	{
+		int32 Num = IStreamingManager::Get().GetNumViews();
+		if (!Num)
+		{
+			// no cameras, no grass update
+			return;
+		}
+		OldCameras.Reset(Num);
+		for (int32 Index = 0; Index < Num; Index++)
+		{
+			auto& ViewInfo = IStreamingManager::Get().GetViewInformation(Index);
+			OldCameras.Add(ViewInfo.ViewOrigin);
+		}
+		UpdateGrass(OldCameras);
+	}
 }
 
 struct FGrassBuilderBase
@@ -2242,7 +2311,6 @@ void ALandscapeProxy::UpdateGrassData()
 void ALandscapeProxy::UpdateGrass(const TArray<FVector>& Cameras, bool bForceSync)
 {
 	SCOPE_CYCLE_COUNTER(STAT_GrassUpdate);
-	TRACE_CPUPROFILER_EVENT_SCOPE(ALandscapeProxy::UpdateGrass);
 	if (GFrameNumberLastStaleCheck != GFrameNumber && GIgnoreExcludeBoxes == 0)
 	{
 		GFrameNumberLastStaleCheck = GFrameNumber;
diff --git a/Engine/Source/Runtime/Landscape/Private/LandscapeSubsystem.cpp b/Engine/Source/Runtime/Landscape/Private/LandscapeSubsystem.cpp
deleted file mode 100644
diff --git a/Engine/Source/Runtime/Landscape/Public/LandscapeSubsystem.h b/Engine/Source/Runtime/Landscape/Public/LandscapeSubsystem.h
deleted file mode 100644
