Grass Optims:
- Precompute Greatest Cull Distance and Early Out (based on a git pull request)
- Avoid list allocation per UpdateGrass by making it static (non-threaded code)
- Limit Component Creation to all Proxies (not per proxy anymore)

diff --git a/Engine/Source/Editor/WorldBrowser/Private/Tiles/WorldTileCollectionModel.cpp b/Engine/Source/Editor/WorldBrowser/Private/Tiles/WorldTileCollectionModel.cpp
--- a/Engine/Source/Editor/WorldBrowser/Private/Tiles/WorldTileCollectionModel.cpp
+++ b/Engine/Source/Editor/WorldBrowser/Private/Tiles/WorldTileCollectionModel.cpp
@@ -2097,7 +2097,8 @@ bool FWorldTileCollectionModel::GenerateLODLevels(FLevelModelList InLevelList, i
 				/* Flush existing grass components, but not grass maps */
 				Landscape->FlushGrassComponents(nullptr, false);
 				TArray<FVector> Cameras;
-				Landscape->UpdateGrass(Cameras, true);
+				int32 NumCompsCreated = 0;
+				Landscape->UpdateGrass(Cameras, NumCompsCreated, true);
 			}
 								
 			// This is texture resolution for a landscape mesh, probably needs to be calculated using landscape size
diff --git a/Engine/Source/Runtime/Landscape/Classes/LandscapeProxy.h b/Engine/Source/Runtime/Landscape/Classes/LandscapeProxy.h
--- a/Engine/Source/Runtime/Landscape/Classes/LandscapeProxy.h
+++ b/Engine/Source/Runtime/Landscape/Classes/LandscapeProxy.h
@@ -507,6 +507,9 @@ public:
 
 	/** Cached grass types from GetGrassTypes */
 	TArray<ULandscapeGrassType*> LandscapeGrassTypes;
+
+	/** Cached grass max square discard distance for all grass in GetGrassTypes */
+	float GrassMaxSquareDiscardDistance;
 #endif
 
 	/** Material used to render landscape components with holes. If not set, LandscapeMaterial will be used (blend mode will be overridden to Masked if it is set to Opaque) */
@@ -849,7 +852,7 @@ public:
 
 		return true;
 	}
-	void TickGrass(const TArray<FVector>& Cameras);
+	void TickGrass(const TArray<FVector>& Cameras, int32& InOutNumCompsCreated);
 
 	/** Flush the grass cache */
 	LANDSCAPE_API void FlushGrassComponents(const TSet<ULandscapeComponent*>* OnlyForComponents = nullptr, bool bFlushGrassMaps = true);
@@ -859,7 +862,7 @@ public:
 		* @param Cameras to use for culling, if empty, then NO culling
 		* @param bForceSync if true, block and finish all work
 	*/
-	LANDSCAPE_API void UpdateGrass(const TArray<FVector>& Cameras, bool bForceSync = false);
+	LANDSCAPE_API void UpdateGrass(const TArray<FVector>& Cameras, int32& InOutNumComponentsCreated, bool bForceSync = false);
 
 	LANDSCAPE_API static void AddExclusionBox(FWeakObjectPtr Owner, const FBox& BoxToRemove);
 	LANDSCAPE_API static void RemoveExclusionBox(FWeakObjectPtr Owner);
@@ -867,7 +870,7 @@ public:
 
 
 	/* Get the list of grass types on this landscape */
-	static void GetGrassTypes(UMaterialInterface* LandscapeMat, TArray<ULandscapeGrassType*>& GrassTypesOut);
+	static void GetGrassTypes(const UWorld* World, UMaterialInterface* LandscapeMat, TArray<ULandscapeGrassType*>& GrassTypesOut, float& OutMaxDiscardDistance);
 
 	/* Invalidate the precomputed grass and baked texture data for the specified components */
 	LANDSCAPE_API static void InvalidateGeneratedComponentData(const TSet<ULandscapeComponent*>& Components, bool bInvalidateLightingCache = false);
diff --git a/Engine/Source/Runtime/Landscape/Private/Landscape.cpp b/Engine/Source/Runtime/Landscape/Private/Landscape.cpp
--- a/Engine/Source/Runtime/Landscape/Private/Landscape.cpp
+++ b/Engine/Source/Runtime/Landscape/Private/Landscape.cpp
@@ -1027,6 +1027,7 @@ ALandscapeProxy::ALandscapeProxy(const FObjectInitializer& ObjectInitializer)
 #if !WITH_EDITORONLY_DATA
 	, LandscapeMaterialCached(nullptr)
 	, LandscapeGrassTypes()
+	, GrassMaxSquareDiscardDistance(0.0f)
 #endif
 	, bHasLandscapeGrass(true)
 {
@@ -3028,7 +3029,8 @@ void ULandscapeInfo::Tick(UWorld* World, float DeltaTime)
 		}
 	}
 
-	ForAllLandscapeProxies([DeltaTime,Cameras,World](ALandscapeProxy* Proxy)
+	int32 NumCompsCreated = 0;
+	ForAllLandscapeProxies([DeltaTime,Cameras,World,&NumCompsCreated](ALandscapeProxy* Proxy)
 	{
 #if WITH_EDITOR
 		if (GIsEditor)
@@ -3047,7 +3049,7 @@ void ULandscapeInfo::Tick(UWorld* World, float DeltaTime)
 #endif
 		if (Cameras && Proxy->ShouldTickGrass())
 		{
-			Proxy->TickGrass(*Cameras);
+			Proxy->TickGrass(*Cameras, NumCompsCreated);
 		}
 	});
 }
diff --git a/Engine/Source/Runtime/Landscape/Private/LandscapeGrass.cpp b/Engine/Source/Runtime/Landscape/Private/LandscapeGrass.cpp
--- a/Engine/Source/Runtime/Landscape/Private/LandscapeGrass.cpp
+++ b/Engine/Source/Runtime/Landscape/Private/LandscapeGrass.cpp
@@ -165,6 +165,12 @@ static FAutoConsoleVariableRef CVarIgnoreExcludeBoxes(
 	GIgnoreExcludeBoxes,
 	TEXT("For debugging. Ignores any exclusion boxes."));
 
+static int32 GGrassMaxCreatePerFrame = 1;
+static FAutoConsoleVariableRef CVarGrassMaxCreatePerFrame(
+	TEXT("grass.MaxCreatePerFrame"),
+	GIgnoreExcludeBoxes,
+	TEXT("Maximum number of Grass components to create per frame"));
+
 DECLARE_CYCLE_STAT(TEXT("Grass Async Build Time"), STAT_FoliageGrassAsyncBuildTime, STATGROUP_Foliage);
 DECLARE_CYCLE_STAT(TEXT("Grass Start Comp"), STAT_FoliageGrassStartComp, STATGROUP_Foliage);
 DECLARE_CYCLE_STAT(TEXT("Grass End Comp"), STAT_FoliageGrassEndComp, STATGROUP_Foliage);
@@ -1386,7 +1392,7 @@ void FLandscapeComponentGrassData::ConditionalDiscardDataOnLoad()
 // ALandscapeProxy grass-related functions
 //
 
-void ALandscapeProxy::TickGrass(const TArray<FVector>& Cameras)
+void ALandscapeProxy::TickGrass(const TArray<FVector>& Cameras, int32& InOutNumCompsCreated)
 {
 	TRACE_CPUPROFILER_EVENT_SCOPE(ALandscapeProxy::TickGrass);
 #if WITH_EDITORONLY_DATA
@@ -1399,7 +1405,7 @@ void ALandscapeProxy::TickGrass(const TArray<FVector>& Cameras)
 	}
 #endif
 
-	UpdateGrass(Cameras);
+	UpdateGrass(Cameras, InOutNumCompsCreated);
 }
 
 struct FGrassBuilderBase
@@ -2045,8 +2051,10 @@ void ALandscapeProxy::FlushGrassComponents(const TSet<ULandscapeComponent*>* Onl
 	}
 }
 
-void ALandscapeProxy::GetGrassTypes(UMaterialInterface* LandscapeMat, TArray<ULandscapeGrassType*>& GrassTypesOut)
+void ALandscapeProxy::GetGrassTypes(const UWorld* World, UMaterialInterface* LandscapeMat, TArray<ULandscapeGrassType*>& GrassTypesOut, float& OutMaxSquareDiscardDistance)
 {
+	float MaxDiscardDistance = 0.0f;
+	ERHIFeatureLevel::Type FeatureLevel = World->Scene->GetFeatureLevel();
 	if (LandscapeMat)
 	{
 		TArray<const UMaterialExpressionLandscapeGrassOutput*> GrassExpressions;
@@ -2055,10 +2063,23 @@ void ALandscapeProxy::GetGrassTypes(UMaterialInterface* LandscapeMat, TArray<ULa
 		{
 			for (const FGrassInput& Type : GrassExpressions[0]->GrassTypes)
 			{
-				GrassTypesOut.Add(Type.GrassType);
+				if (Type.GrassType)
+				{
+					GrassTypesOut.Add(Type.GrassType);
+
+					for (auto& GrassVariety : Type.GrassType->GrassVarieties)
+					{
+						int32 EndCullDistance = GrassVariety.EndCullDistance.GetValueForFeatureLevel(FeatureLevel);
+						if (EndCullDistance > MaxDiscardDistance)
+						{
+							MaxDiscardDistance = EndCullDistance;
+						}
+					}
+				}
 			}
 		}
 	}
+	OutMaxSquareDiscardDistance = MaxDiscardDistance * MaxDiscardDistance;
 }
 
 static uint32 GGrassExclusionChangeTag = 1;
@@ -2115,7 +2136,8 @@ void ALandscapeProxy::UpdateGrassDataStatus(TSet<UTexture2D*>& OutCurrentForcedS
 	const bool bCheckStreamingState = OutDesiredForcedStreamedTextures || bInEnableForceResidentFlag;
 
 	TArray<ULandscapeGrassType*> GrassTypes;
-	GetGrassTypes(LandscapeMaterial, GrassTypes);
+	float OutMaxSquareDiscardDistance = 0.0f;
+	GetGrassTypes(World, LandscapeMaterial, GrassTypes, OutMaxSquareDiscardDistance);
 	const bool bHasGrassTypes = GrassTypes.Num() > 0;
 
 	for (auto Component : LandscapeComponents)
@@ -2239,7 +2261,7 @@ void ALandscapeProxy::UpdateGrassData()
 }
 #endif
 
-void ALandscapeProxy::UpdateGrass(const TArray<FVector>& Cameras, bool bForceSync)
+void ALandscapeProxy::UpdateGrass(const TArray<FVector>& Cameras, int32& InOutNumCompsCreated, bool bForceSync)
 {
 	SCOPE_CYCLE_COUNTER(STAT_GrassUpdate);
 	TRACE_CPUPROFILER_EVENT_SCOPE(ALandscapeProxy::UpdateGrass);
@@ -2258,18 +2280,7 @@ void ALandscapeProxy::UpdateGrass(const TArray<FVector>& Cameras, bool bForceSyn
 
 	if (GGrassEnable > 0)
 	{
-#if WITH_EDITOR
-		TArray<ULandscapeGrassType*> LandscapeGrassTypes;
-		GetGrassTypes(LandscapeMaterial, LandscapeGrassTypes);
-#else
-		// In non editor builds, cache grass types for performance.
-		if (LandscapeMaterial != LandscapeMaterialCached)
-		{
-			LandscapeMaterialCached = LandscapeMaterial;
-			LandscapeGrassTypes.Reset();
-			GetGrassTypes(LandscapeMaterial, LandscapeGrassTypes);
-		}
-#endif
+		UWorld* World = GetWorld();
 
 		float GuardBand = GGuardBandMultiplier;
 		float DiscardGuardBand = GGuardBandDiscardMultiplier;
@@ -2280,9 +2291,23 @@ void ALandscapeProxy::UpdateGrass(const TArray<FVector>& Cameras, bool bForceSyn
 		int32 MaxTasks = GMaxAsyncTasks;
 		const float CullDistanceScale = GGrassCullDistanceScale;
 
-		UWorld* World = GetWorld();
 		if (World)
 		{
+#if WITH_EDITOR
+			TArray<ULandscapeGrassType*> LandscapeGrassTypes;
+			float GrassMaxSquareDiscardDistance = 0.0f;
+			GetGrassTypes(World, LandscapeMaterial, LandscapeGrassTypes, GrassMaxSquareDiscardDistance);
+#else
+			// In non editor builds, cache grass types for performance.
+			if (LandscapeMaterial != LandscapeMaterialCached)
+			{
+				LandscapeMaterialCached = LandscapeMaterial;
+				LandscapeGrassTypes.Reset();
+				GrassMaxSquareDiscardDistance = 0.0f;
+				GetGrassTypes(World, LandscapeMaterial, LandscapeGrassTypes, GrassMaxSquareDiscardDistance);
+			}
+#endif
+			
 #if WITH_EDITOR
 			int32 RequiredTexturesNotStreamedIn = 0;
 			TSet<ULandscapeComponent*> ComponentsNeedingGrassMapRender;
@@ -2314,8 +2339,8 @@ void ALandscapeProxy::UpdateGrass(const TArray<FVector>& Cameras, bool bForceSyn
 				FBox BoundsBox;
 			};
 
-			TArray<SortedLandscapeElement> SortedLandscapeComponents;
-			SortedLandscapeComponents.Reserve(LandscapeComponents.Num());
+			static TArray<SortedLandscapeElement> SortedLandscapeComponents;
+			SortedLandscapeComponents.Reset(LandscapeComponents.Num());
 			for (ULandscapeComponent* Component : LandscapeComponents)
 			{
 				// skip if we have no data and no way to generate it
@@ -2329,6 +2354,12 @@ void ALandscapeProxy::UpdateGrass(const TArray<FVector>& Cameras, bool bForceSyn
 				{
 					MinSqrDistanceToComponent = FMath::Min<float>(MinSqrDistanceToComponent, WorldBounds.ComputeSquaredDistanceFromBoxToPoint(CameraPos));
 				}
+
+				if (MinSqrDistanceToComponent > GrassMaxSquareDiscardDistance)
+				{
+					continue;
+				}
+
 				SortedLandscapeComponents.Emplace(Component, FMath::Sqrt(MinSqrDistanceToComponent), WorldBounds.GetBox());
 			}
 			
@@ -2340,7 +2371,6 @@ void ALandscapeProxy::UpdateGrass(const TArray<FVector>& Cameras, bool bForceSyn
 			}
 #endif
 
-			int32 NumCompsCreated = 0;
 			for (const SortedLandscapeElement& SortedLandscapeComponent : SortedLandscapeComponents)
 			{
 				ULandscapeComponent* Component = SortedLandscapeComponent.Component;
@@ -2496,7 +2526,7 @@ void ALandscapeProxy::UpdateGrass(const TArray<FVector>& Cameras, bool bForceSyn
 											}
 										}
 
-										if (!bRebuildForBoxes && !bForceSync && (NumCompsCreated || AsyncFoliageTasks.Num() >= MaxTasks))
+										if (!bRebuildForBoxes && !bForceSync && (InOutNumCompsCreated >= GGrassMaxCreatePerFrame || AsyncFoliageTasks.Num() >= MaxTasks))
 										{
 											continue; // one per frame, but we still want to touch the existing ones and we must do the rebuilds because we changed the tag
 										}
@@ -2541,7 +2571,7 @@ void ALandscapeProxy::UpdateGrass(const TArray<FVector>& Cameras, bool bForceSyn
 										}
 #endif
 
-										NumCompsCreated++;
+										InOutNumCompsCreated++;
 
 										SCOPE_CYCLE_COUNTER(STAT_FoliageGrassStartComp);
 
