Fix API Backward compatibility: UpdateGrass

#rnx

diff --git a/Engine/Source/Runtime/Landscape/Classes/LandscapeProxy.h b/Engine/Source/Runtime/Landscape/Classes/LandscapeProxy.h
--- a/Engine/Source/Runtime/Landscape/Classes/LandscapeProxy.h
+++ b/Engine/Source/Runtime/Landscape/Classes/LandscapeProxy.h
@@ -860,10 +860,13 @@ public:
 	/**
 		Update Grass 
 		* @param Cameras to use for culling, if empty, then NO culling
+		* @param InOutNumComponentsCreated, value can increase if components were created, it is also used internally to limit the number of creations
 		* @param bForceSync if true, block and finish all work
 	*/
 	LANDSCAPE_API void UpdateGrass(const TArray<FVector>& Cameras, int32& InOutNumComponentsCreated, bool bForceSync = false);
 
+	LANDSCAPE_API void UpdateGrass(const TArray<FVector>& Cameras, bool bForceSync = false);
+
 	LANDSCAPE_API static void AddExclusionBox(FWeakObjectPtr Owner, const FBox& BoxToRemove);
 	LANDSCAPE_API static void RemoveExclusionBox(FWeakObjectPtr Owner);
 	LANDSCAPE_API static void RemoveAllExclusionBoxes();
diff --git a/Engine/Source/Runtime/Landscape/Private/LandscapeGrass.cpp b/Engine/Source/Runtime/Landscape/Private/LandscapeGrass.cpp
--- a/Engine/Source/Runtime/Landscape/Private/LandscapeGrass.cpp
+++ b/Engine/Source/Runtime/Landscape/Private/LandscapeGrass.cpp
@@ -2261,7 +2261,13 @@ void ALandscapeProxy::UpdateGrassData()
 }
 #endif
 
-void ALandscapeProxy::UpdateGrass(const TArray<FVector>& Cameras, int32& InOutNumCompsCreated, bool bForceSync)
+void ALandscapeProxy::UpdateGrass(const TArray<FVector>& Cameras, bool bForceSync /* = false */)
+{
+	int32 InOutNumCompsCreated = 0;
+	UpdateGrass(Cameras, InOutNumCompsCreated, bForceSync);
+}
+
+void ALandscapeProxy::UpdateGrass(const TArray<FVector>& Cameras, int32& InOutNumCompsCreated, bool bForceSync /* = false */)
 {
 	SCOPE_CYCLE_COUNTER(STAT_GrassUpdate);
 	TRACE_CPUPROFILER_EVENT_SCOPE(ALandscapeProxy::UpdateGrass);
