Fix Creative Landscape Guid Clash

diff --git a/Engine/Source/Runtime/Landscape/Classes/LandscapeInfo.h b/Engine/Source/Runtime/Landscape/Classes/LandscapeInfo.h
--- a/Engine/Source/Runtime/Landscape/Classes/LandscapeInfo.h
+++ b/Engine/Source/Runtime/Landscape/Classes/LandscapeInfo.h
@@ -283,8 +283,6 @@ public:
 	LANDSCAPE_API void RegisterCollisionComponent(ULandscapeHeightfieldCollisionComponent* Component);
 	LANDSCAPE_API void UnregisterCollisionComponent(ULandscapeHeightfieldCollisionComponent* Component);
 
-	LANDSCAPE_API void Tick(UWorld* World, float DeltaTime);
-
 #if WITH_EDITOR
 private:
 	bool ApplySplinesInternal(bool bOnlySelected, ALandscapeProxy* Proxy, TSet<ULandscapeComponent*>* OutModifiedComponents, bool bMarkPackageDirty, int32 LandscapeMinX, int32 LandscapeMinY, int32 LandscapeMaxX, int32 LandscapeMaxY, TFunctionRef<TSharedPtr<FModulateAlpha>(ULandscapeLayerInfoObject*)> GetOrCreateModulate);
diff --git a/Engine/Source/Runtime/Landscape/Private/Landscape.cpp b/Engine/Source/Runtime/Landscape/Private/Landscape.cpp
--- a/Engine/Source/Runtime/Landscape/Private/Landscape.cpp
+++ b/Engine/Source/Runtime/Landscape/Private/Landscape.cpp
@@ -50,7 +50,7 @@ Landscape.cpp: Terrain rendering
 #include "EngineUtils.h"
 #include "ComponentRecreateRenderStateContext.h"
 #include "LandscapeWeightmapUsage.h"
-#include "ContentStreaming.h"
+#include "LandscapeSubsystem.h"
 
 #if WITH_EDITOR
 #include "LandscapeEdit.h"
@@ -65,12 +65,6 @@ Landscape.cpp: Terrain rendering
 #include "UObject/EditorObjectVersion.h"
 #include "Algo/BinarySearch.h"
 
-static int32 GUseStreamingManagerForCameras = 1;
-static FAutoConsoleVariableRef CVarUseStreamingManagerForCameras(
-	TEXT("grass.UseStreamingManagerForCameras"),
-	GUseStreamingManagerForCameras,
-	TEXT("1: Use Streaming Manager; 0: Use ViewLocationsRenderedLastFrame"));
-
 /** Landscape stats */
 
 DEFINE_STAT(STAT_LandscapeDynamicDrawTime);
@@ -2995,65 +2989,6 @@ void ULandscapeInfo::RecreateLandscapeInfo(UWorld* InWorld, bool bMapCheck)
 
 #endif
 
-void ULandscapeInfo::Tick(UWorld* World, float DeltaTime)
-{
-	// Update Grass
-	static TArray<FVector> OldCameras;
-	TArray<FVector>* Cameras = nullptr;
-	if (GUseStreamingManagerForCameras == 0)
-	{
-		if (OldCameras.Num() || World->ViewLocationsRenderedLastFrame.Num())
-		{
-			Cameras = &OldCameras;
-			// there is a bug here, which often leaves us with no cameras in the editor
-			if (World->ViewLocationsRenderedLastFrame.Num())
-			{
-				check(IsInGameThread());
-				Cameras = &World->ViewLocationsRenderedLastFrame;
-				OldCameras = *Cameras;
-			}
-		}
-	}
-	else
-	{
-		int32 Num = IStreamingManager::Get().GetNumViews();
-		if (Num)
-		{
-			OldCameras.Reset(Num);
-			for (int32 Index = 0; Index < Num; Index++)
-			{
-				auto& ViewInfo = IStreamingManager::Get().GetViewInformation(Index);
-				OldCameras.Add(ViewInfo.ViewOrigin);
-			}
-			Cameras = &OldCameras;
-		}
-	}
-
-	int32 NumCompsCreated = 0;
-	ForAllLandscapeProxies([DeltaTime,Cameras,World,&NumCompsCreated](ALandscapeProxy* Proxy)
-	{
-#if WITH_EDITOR
-		if (GIsEditor)
-		{
-			if (ALandscape* Landscape = Cast<ALandscape>(Proxy))
-			{
-				Landscape->TickLayers(DeltaTime);
-			}
-
-			// editor-only
-			if (World && !World->IsPlayInEditor())
-			{
-				Proxy->UpdateBakedTextures();
-			}
-		}
-#endif
-		if (Cameras && Proxy->ShouldTickGrass())
-		{
-			Proxy->TickGrass(*Cameras, NumCompsCreated);
-		}
-	});
-}
-
 void ULandscapeInfo::ForAllLandscapeProxies(TFunctionRef<void(ALandscapeProxy*)> Fn) const
 {
 	ALandscape* Landscape = LandscapeActor.Get();
@@ -3074,46 +3009,11 @@ void ULandscapeInfo::RegisterActor(ALandscapeProxy* Proxy, bool bMapCheck)
 	// do not pass here invalid actors
 	checkSlow(Proxy);
 	check(Proxy->GetLandscapeGuid().IsValid());
+	
+	OwningWorld->GetSubsystem<ULandscapeSubsystem>()->RegisterActor(Proxy);
 
-	auto LamdbdaLowerBound = [](ALandscapeProxy* A, ALandscapeProxy* B)
-	{
-		FIntPoint SectionBaseA = A->GetSectionBaseOffset();
-		FIntPoint SectionBaseB = B->GetSectionBaseOffset();
-
-		if (SectionBaseA.X != SectionBaseB.X)
-		{
-			return SectionBaseA.X < SectionBaseB.X;
-		}
-
-		return SectionBaseA.Y < SectionBaseB.Y;
-	};
-
-	if (OwningWorld->IsGameWorld())
-	{
-		// register
-		if (ALandscape* Landscape = Cast<ALandscape>(Proxy))
-		{
-			checkf(!LandscapeActor || LandscapeActor == Landscape, TEXT("Multiple landscapes with the same GUID detected: %s vs %s"), *LandscapeActor->GetPathName(), *Landscape->GetPathName());
-			LandscapeActor = Landscape;
-			// update proxies reference actor
-			for (ALandscapeStreamingProxy* StreamingProxy : Proxies)
-			{
-				StreamingProxy->LandscapeActor = LandscapeActor;
-			}
-		}
-		else
-		{
-			ALandscapeStreamingProxy* StreamingProxy = CastChecked<ALandscapeStreamingProxy>(Proxy);
-			if (!Proxies.Contains(Proxy))
-			{
-				uint32 InsertIndex = Algo::LowerBound(Proxies, Proxy, LamdbdaLowerBound);
-				Proxies.Insert(StreamingProxy, InsertIndex);
-			}
-			StreamingProxy->LandscapeActor = LandscapeActor;
-		}
-	}
 #if WITH_EDITOR
-	else
+	if (!OwningWorld->IsGameWorld())
 	{
 		// in case this Info object is not initialized yet
 		// initialized it with properties from passed actor
@@ -3156,6 +3056,19 @@ void ULandscapeInfo::RegisterActor(ALandscapeProxy* Proxy, bool bMapCheck)
 		}
 		else
 		{
+			auto LamdbdaLowerBound = [](ALandscapeProxy* A, ALandscapeProxy* B)
+			{
+				FIntPoint SectionBaseA = A->GetSectionBaseOffset();
+				FIntPoint SectionBaseB = B->GetSectionBaseOffset();
+
+				if (SectionBaseA.X != SectionBaseB.X)
+				{
+					return SectionBaseA.X < SectionBaseB.X;
+				}
+
+				return SectionBaseA.Y < SectionBaseB.Y;
+			};
+
 			// Insert Proxies in a sorted fashion for generating deterministic results in the Layer system
 			ALandscapeStreamingProxy* StreamingProxy = CastChecked<ALandscapeStreamingProxy>(Proxy);
 			if (!Proxies.Contains(Proxy))
@@ -3194,29 +3107,35 @@ void ULandscapeInfo::RegisterActor(ALandscapeProxy* Proxy, bool bMapCheck)
 
 void ULandscapeInfo::UnregisterActor(ALandscapeProxy* Proxy)
 {
-	if (ALandscape* Landscape = Cast<ALandscape>(Proxy))
+	UWorld* OwningWorld = Proxy->GetWorld();
+#if WITH_EDITOR
+	if (!OwningWorld->IsGameWorld())
 	{
-		// Note: UnregisterActor sometimes gets triggered twice, e.g. it has been observed to happen during redo
-		// Note: In some cases LandscapeActor could be updated to a new landscape actor before the old landscape is unregistered/destroyed
-		// e.g. this has been observed when merging levels in the editor
-
-		if (LandscapeActor.Get() == Landscape)
+		if (ALandscape* Landscape = Cast<ALandscape>(Proxy))
 		{
-			LandscapeActor = nullptr;
-		}
+			// Note: UnregisterActor sometimes gets triggered twice, e.g. it has been observed to happen during redo
+			// Note: In some cases LandscapeActor could be updated to a new landscape actor before the old landscape is unregistered/destroyed
+			// e.g. this has been observed when merging levels in the editor
 
-		// update proxies reference to landscape actor
-		for (ALandscapeStreamingProxy* StreamingProxy : Proxies)
+			if (LandscapeActor.Get() == Landscape)
+			{
+				LandscapeActor = nullptr;
+			}
+
+			// update proxies reference to landscape actor
+			for (ALandscapeStreamingProxy* StreamingProxy : Proxies)
+			{
+				StreamingProxy->LandscapeActor = LandscapeActor;
+			}
+		}
+		else
 		{
-			StreamingProxy->LandscapeActor = LandscapeActor;
+			ALandscapeStreamingProxy* StreamingProxy = CastChecked<ALandscapeStreamingProxy>(Proxy);
+			Proxies.Remove(StreamingProxy);
+			StreamingProxy->LandscapeActor = nullptr;
 		}
 	}
-	else
-	{
-		ALandscapeStreamingProxy* StreamingProxy = CastChecked<ALandscapeStreamingProxy>(Proxy);
-		Proxies.Remove(StreamingProxy);
-		StreamingProxy->LandscapeActor = nullptr;
-	}
+#endif
 
 	// remove proxy components from the XY map
 	for (int32 CompIdx = 0; CompIdx < Proxy->LandscapeComponents.Num(); ++CompIdx)
@@ -3239,12 +3158,14 @@ void ULandscapeInfo::UnregisterActor(ALandscapeProxy* Proxy)
 	XYtoCollisionComponentMap.Compact();
 
 #if WITH_EDITOR
-	if (!Proxy->GetWorld()->IsGameWorld())
+	if (!OwningWorld->IsGameWorld())
 	{
 		UpdateLayerInfoMap();
 		UpdateAllAddCollisions();
 	}
 #endif
+
+	OwningWorld->GetSubsystem<ULandscapeSubsystem>()->UnregisterActor(Proxy);
 }
 
 void ULandscapeInfo::RegisterCollisionComponent(ULandscapeHeightfieldCollisionComponent* Component)
diff --git a/Engine/Source/Runtime/Landscape/Private/LandscapeSubsystem.cpp b/Engine/Source/Runtime/Landscape/Private/LandscapeSubsystem.cpp
--- a/Engine/Source/Runtime/Landscape/Private/LandscapeSubsystem.cpp
+++ b/Engine/Source/Runtime/Landscape/Private/LandscapeSubsystem.cpp
@@ -4,14 +4,20 @@
 #include "UObject/UObjectGlobals.h"
 #include "Engine/EngineBaseTypes.h"
 #include "Engine/World.h"
-#include "LandscapeInfoMap.h"
-#include "LandscapeInfo.h"
+#include "ContentStreaming.h"
+#include "Landscape.h"
 #include "ProfilingDebugging/CsvProfiler.h"
 
+static int32 GUseStreamingManagerForCameras = 1;
+static FAutoConsoleVariableRef CVarUseStreamingManagerForCameras(
+	TEXT("grass.UseStreamingManagerForCameras"),
+	GUseStreamingManagerForCameras,
+	TEXT("1: Use Streaming Manager; 0: Use ViewLocationsRenderedLastFrame"));
+
 DECLARE_CYCLE_STAT(TEXT("LandscapeSubsystem Tick"), STAT_LandscapeSubsystemTick, STATGROUP_Landscape);
 
 ULandscapeSubsystem::ULandscapeSubsystem()
-	: TickFunction(this), LandscapeInfoMap(nullptr)
+	: TickFunction(this)
 {}
 
 ULandscapeSubsystem::~ULandscapeSubsystem()
@@ -33,7 +39,17 @@ void ULandscapeSubsystem::Initialize(FSubsystemCollectionBase& Collection)
 void ULandscapeSubsystem::Deinitialize()
 {
 	TickFunction.UnRegisterTickFunction();
-	LandscapeInfoMap = nullptr;
+	Proxies.Empty();
+}
+
+void ULandscapeSubsystem::RegisterActor(ALandscapeProxy* Proxy)
+{
+	Proxies.AddUnique(Proxy);
+}
+
+void ULandscapeSubsystem::UnregisterActor(ALandscapeProxy* Proxy)
+{
+	Proxies.Remove(Proxy);
 }
 
 void ULandscapeSubsystem::Tick(float DeltaTime, enum ELevelTick TickType, ENamedThreads::Type CurrentThread, const FGraphEventRef& MyCompletionGraphEvent)
@@ -44,14 +60,60 @@ void ULandscapeSubsystem::Tick(float DeltaTime, enum ELevelTick TickType, ENamed
 	LLM_SCOPE(ELLMTag::Landscape);
 
 	UWorld* World = GetWorld();
-	if (!LandscapeInfoMap)
+
+	static TArray<FVector> OldCameras;
+	TArray<FVector>* Cameras = nullptr;
+	if (GUseStreamingManagerForCameras == 0)
+	{
+		if (OldCameras.Num() || World->ViewLocationsRenderedLastFrame.Num())
+		{
+			Cameras = &OldCameras;
+			// there is a bug here, which often leaves us with no cameras in the editor
+			if (World->ViewLocationsRenderedLastFrame.Num())
+			{
+				check(IsInGameThread());
+				Cameras = &World->ViewLocationsRenderedLastFrame;
+				OldCameras = *Cameras;
+			}
+		}
+	}
+	else
 	{
-		LandscapeInfoMap = &ULandscapeInfoMap::GetLandscapeInfoMap(World);
+		int32 Num = IStreamingManager::Get().GetNumViews();
+		if (Num)
+		{
+			OldCameras.Reset(Num);
+			for (int32 Index = 0; Index < Num; Index++)
+			{
+				auto& ViewInfo = IStreamingManager::Get().GetViewInformation(Index);
+				OldCameras.Add(ViewInfo.ViewOrigin);
+			}
+			Cameras = &OldCameras;
+		}
 	}
-		
-	for (const auto& Pair : LandscapeInfoMap->Map)
+
+	int32 InOutNumComponentsCreated = 0;
+	for(ALandscapeProxy* Proxy : Proxies)
 	{
-		Pair.Value->Tick(World, DeltaTime);
+#if WITH_EDITOR
+		if (GIsEditor)
+		{
+			if (ALandscape* Landscape = Cast<ALandscape>(Proxy))
+			{
+				Landscape->TickLayers(DeltaTime);
+			}
+
+			// editor-only
+			if (!World->IsPlayInEditor())
+			{
+				Proxy->UpdateBakedTextures();
+			}
+		}
+#endif
+		if (Cameras && Proxy->ShouldTickGrass())
+		{
+			Proxy->TickGrass(*Cameras, InOutNumComponentsCreated);
+		}
 	}
 }
 
diff --git a/Engine/Source/Runtime/Landscape/Public/LandscapeSubsystem.h b/Engine/Source/Runtime/Landscape/Public/LandscapeSubsystem.h
--- a/Engine/Source/Runtime/Landscape/Public/LandscapeSubsystem.h
+++ b/Engine/Source/Runtime/Landscape/Public/LandscapeSubsystem.h
@@ -8,7 +8,7 @@
 #include "Engine/EngineBaseTypes.h"
 #include "LandscapeSubsystem.generated.h"
 
-class ULandscapeInfoMap;
+class ALandscapeProxy;
 
 struct FLandscapeSubsystemTickFunction : public FTickFunction
 {
@@ -39,12 +39,15 @@ public:
 	virtual void Deinitialize() override;
 	// End USubsystem
 		
+	void RegisterActor(ALandscapeProxy* Proxy);
+	void UnregisterActor(ALandscapeProxy* Proxy);
+
 private:
 	void Tick(float DeltaTime, enum ELevelTick TickType, ENamedThreads::Type CurrentThread, const FGraphEventRef& MyCompletionGraphEvent);
 
 	friend struct FLandscapeSubsystemTickFunction;
 	FLandscapeSubsystemTickFunction TickFunction;
 
-	ULandscapeInfoMap* LandscapeInfoMap;
+	TArray<ALandscapeProxy*> Proxies;
 };
 
