Fix Crash in ALandscapeProxy::PreSave on null Scene.

#rnx

diff --git a/Engine/Source/Runtime/Landscape/Private/LandscapeGrass.cpp b/Engine/Source/Runtime/Landscape/Private/LandscapeGrass.cpp
--- a/Engine/Source/Runtime/Landscape/Private/LandscapeGrass.cpp
+++ b/Engine/Source/Runtime/Landscape/Private/LandscapeGrass.cpp
@@ -2054,7 +2054,7 @@ void ALandscapeProxy::FlushGrassComponents(const TSet<ULandscapeComponent*>* Onl
 void ALandscapeProxy::GetGrassTypes(const UWorld* World, UMaterialInterface* LandscapeMat, TArray<ULandscapeGrassType*>& GrassTypesOut, float& OutMaxSquareDiscardDistance)
 {
 	float MaxDiscardDistance = 0.0f;
-	ERHIFeatureLevel::Type FeatureLevel = World->Scene->GetFeatureLevel();
+	ERHIFeatureLevel::Type FeatureLevel = World->Scene != nullptr? World->Scene->GetFeatureLevel() : GMaxRHIFeatureLevel;
 	if (LandscapeMat)
 	{
 		TArray<const UMaterialExpressionLandscapeGrassOutput*> GrassExpressions;
