Fix Crash: Change registration to LandscapeSubsystem

#rnx

diff --git a/Engine/Source/Runtime/Landscape/Private/Landscape.cpp b/Engine/Source/Runtime/Landscape/Private/Landscape.cpp
--- a/Engine/Source/Runtime/Landscape/Private/Landscape.cpp
+++ b/Engine/Source/Runtime/Landscape/Private/Landscape.cpp
@@ -1892,6 +1892,14 @@ void ALandscapeProxy::PostRegisterAllComponents()
 		{
 			LandscapeInfo = CreateLandscapeInfo();
 		}
+
+		if (UWorld* OwningWorld = GetWorld())
+		{
+			if (ULandscapeSubsystem* LandscapeSubsystem = OwningWorld->GetSubsystem<ULandscapeSubsystem>())
+			{
+				LandscapeSubsystem->RegisterActor(this);
+			}
+		}
 	}
 #if WITH_EDITOR
 	// Game worlds don't have landscape infos
@@ -1919,6 +1927,11 @@ void ALandscapeProxy::UnregisterAllComponents(const bool bForReregister)
 		{
 			LandscapeInfo->UnregisterActor(this);
 		}
+
+		if (ULandscapeSubsystem* LandscapeSubsystem = GetWorld()->GetSubsystem<ULandscapeSubsystem>())
+		{
+			LandscapeSubsystem->UnregisterActor(this);
+		}
 	}
 
 	Super::UnregisterAllComponents(bForReregister);
@@ -3010,8 +3023,6 @@ void ULandscapeInfo::RegisterActor(ALandscapeProxy* Proxy, bool bMapCheck)
 	checkSlow(Proxy);
 	check(Proxy->GetLandscapeGuid().IsValid());
 	
-	OwningWorld->GetSubsystem<ULandscapeSubsystem>()->RegisterActor(Proxy);
-
 #if WITH_EDITOR
 	if (!OwningWorld->IsGameWorld())
 	{
@@ -3164,8 +3175,6 @@ void ULandscapeInfo::UnregisterActor(ALandscapeProxy* Proxy)
 		UpdateAllAddCollisions();
 	}
 #endif
-
-	OwningWorld->GetSubsystem<ULandscapeSubsystem>()->UnregisterActor(Proxy);
 }
 
 void ULandscapeInfo::RegisterCollisionComponent(ULandscapeHeightfieldCollisionComponent* Component)
