Mobile landscape hole support.
On mobile a unique index buffer can be created for each landscape component which removes quads marked as holes by the visibility map.
Mobile vertex format has been changed to add a lod clamp value so that the landscape holes can retain their shape over distance.
To make room for the new vertex clamp value the vertex Min/Max height values have been quantized. This shouldn't impact height resolution in the general case and actually frees up space to improve precision to the low lods.
Holes will come with memory overhead for the unique index buffers and some perf overhead for the increase in rendered vertices (to keep hole shape). But perf cost should still be much less then using alpha test.
Also mobile holes are offet by a half pixel from non-mobile version since the visiblity map is sampled at vertex positions on non-mobile. In practice this isn't a problem since hole edges are often hidden.
Another potential problem is that hole LODding info isn't shared by neighbouring components so cracks may appear at camera distance in some use cases.

diff --git a/Engine/Shaders/Private/LandscapeVertexFactory.ush b/Engine/Shaders/Private/LandscapeVertexFactory.ush
--- a/Engine/Shaders/Private/LandscapeVertexFactory.ush
+++ b/Engine/Shaders/Private/LandscapeVertexFactory.ush
@@ -623,10 +623,12 @@ FVertexFactoryIntermediates GetVertexFactoryIntermediates(FVertexFactoryInput In
 #if FEATURE_LEVEL >= FEATURE_LEVEL_SM4
 	Intermediates.InputPosition = Input.Position;
 #else
-	Intermediates.InputPosition = Input.PackedPosition * 255.f;
-	float2 SubXY = frac((Intermediates.InputPosition.zw) * .5f) * 2.f;
-	float2 LOD45 = (Intermediates.InputPosition.zw - SubXY) / 255.f;
-	Intermediates.InputPosition.zw = SubXY;
+	Intermediates.InputPosition.xy = Input.PackedPosition.xy * 255.f;
+	float SubX = (float)((uint)(Input.PackedPosition.z * 255.f) >> 4);
+	float SubY = (float)((uint)(Input.PackedPosition.z * 255.f) & 0xf);
+	Intermediates.InputPosition.zw = float2(SubX, SubY);
+	float HoleMaxLod = (float)((uint)(Input.PackedPosition.w * 255.f) >> 4);
+	float HoleVertexLod = (float)((uint)(Input.PackedPosition.w * 255.f) & 0xf);
 #endif
 
 	float2 xy = Intermediates.InputPosition.xy * LodValues.w;
@@ -646,12 +648,19 @@ FVertexFactoryIntermediates GetVertexFactoryIntermediates(FVertexFactoryInput In
 	float2 ActualLODCoordsInt = floor(Intermediates.InputPosition.xy * pow(2, -(LodValue - LodValues.x)));
 	float InvLODScaleFactor = pow(2, -LodValue);
 #else
-	float3 HeightData = (LodValue < 1 ? float3(1.0, Input.LODHeights[1][0], Input.LODHeights[1][1]) :
-						(LodValue < 2 ? float3(0.5, Input.LODHeights[1][1], Input.LODHeights[1][2]) :
-						(LodValue < 3 ? float3(0.25, Input.LODHeights[1][2], Input.LODHeights[1][3]) :
-						(LodValue < 4 ? float3(0.125, Input.LODHeights[1][3], LOD45[0]) :
-						(LodValue < 5 ? float3(0.0625, LOD45[0], LOD45[1]) :
-						(float3(0.03125, LOD45[1], LOD45[1])))))));
+
+	// Clamp lod values around holes until we reach HoleMaxLod. Kill morphing when clamped.
+	bool bMorph = LodValue < HoleVertexLod || HoleMaxLod <= LodValue;
+	MorphAlpha = bMorph ? MorphAlpha : 0;
+	bool bClampLodForHole = HoleVertexLod < LodValue && LodValue < HoleMaxLod;
+	LodValue = bClampLodForHole ? HoleVertexLod : LodValue;
+
+	float3 HeightData = (LodValue < 1 ? float3(1.0, Input.LODHeights[0][2], Input.LODHeights[0][3]) :
+						(LodValue < 2 ? float3(0.5, Input.LODHeights[0][3], Input.LODHeights[1][0]) :
+						(LodValue < 3 ? float3(0.25, Input.LODHeights[1][0], Input.LODHeights[1][1]) :
+						(LodValue < 4 ? float3(0.125, Input.LODHeights[1][1], Input.LODHeights[1][2]) :
+						(LodValue < 5 ? float3(0.0625, Input.LODHeights[1][2], Input.LODHeights[1][3]) :
+						(float3(0.03125, Input.LODHeights[1][3], Input.LODHeights[1][3])))))));
 
 	// we can hardcode InvLODScaleFactor
 	float InvLODScaleFactor = HeightData.x;
@@ -659,8 +668,8 @@ FVertexFactoryIntermediates GetVertexFactoryIntermediates(FVertexFactoryInput In
 	float InputHeight = HeightData.y;
 	float InputHeightNextLOD = HeightData.z;
 
-	float MinHeight = DecodePackedHeight(Input.LODHeights[0].xy);
-	float MaxHeight = DecodePackedHeight(Input.LODHeights[0].zw);
+	float MinHeight = DecodePackedHeight(float2(Input.LODHeights[0].x, 0));
+	float MaxHeight = DecodePackedHeight(float2(Input.LODHeights[0].y, 0));
 
 	float Height = lerp(MinHeight, MaxHeight, InputHeight);
 	float HeightNextLOD = lerp(MinHeight, MaxHeight, InputHeightNextLOD);
diff --git a/Engine/Source/Runtime/Landscape/Classes/LandscapeProxy.h b/Engine/Source/Runtime/Landscape/Classes/LandscapeProxy.h
--- a/Engine/Source/Runtime/Landscape/Classes/LandscapeProxy.h
+++ b/Engine/Source/Runtime/Landscape/Classes/LandscapeProxy.h
@@ -533,6 +533,14 @@ public:
 	bool bIsPerformingInteractiveActionOnLandscapeMaterialOverride;
 #endif 
 
+	/** Use unique geometry instead of material alpha tests for holes on mobile platforms. This requires additional memory and will render more vertices at lower LODs.*/
+	UPROPERTY(EditAnywhere, AdvancedDisplay, BlueprintReadOnly, Category = Mobile, meta = (DisplayName = "Unique Hole Meshes"))
+	bool bMeshHoles = false;
+
+	/** Maximum geometry LOD at which to render unique hole meshes.*/
+	UPROPERTY(EditAnywhere, AdvancedDisplay, BlueprintReadOnly, Category = Mobile, meta = (DisplayName = "Unique Hole Meshes Max LOD", UIMin = "1", UIMax = "6"))
+	uint8 MeshHolesMaxLod = 6;
+
 	/**
 	 * Array of runtime virtual textures into which we render this landscape.
 	 * The material also needs to be set up to output to a virtual texture.
diff --git a/Engine/Source/Runtime/Landscape/Private/Landscape.cpp b/Engine/Source/Runtime/Landscape/Private/Landscape.cpp
--- a/Engine/Source/Runtime/Landscape/Private/Landscape.cpp
+++ b/Engine/Source/Runtime/Landscape/Private/Landscape.cpp
@@ -97,6 +97,7 @@ DEFINE_STAT(STAT_LandscapeLayersRegenerateWeightmaps);
 
 DEFINE_STAT(STAT_LandscapeVertexMem);
 DEFINE_STAT(STAT_LandscapeOccluderMem);
+DEFINE_STAT(STAT_LandscapeHoleMem);
 DEFINE_STAT(STAT_LandscapeComponentMem);
 
 #if ENABLE_COOK_STATS
@@ -117,7 +118,7 @@ namespace LandscapeCookStats
 // differences, etc.) replace the version GUID below with a new one.
 // In case of merge conflicts with DDC versions, you MUST generate a new GUID
 // and set this new GUID as the version.                                       
-#define LANDSCAPE_MOBILE_COOK_VERSION TEXT("6862AA3DD9FB4368B9DDAF9EB7E9901F")
+#define LANDSCAPE_MOBILE_COOK_VERSION TEXT("EAA6E15CEDD644308E8B8D5427EC180")
 
 #define LOCTEXT_NAMESPACE "Landscape"
 
@@ -3317,7 +3318,7 @@ ULandscapeWeightmapUsage::ULandscapeWeightmapUsage(const FObjectInitializer& Obj
 }
 
 // Generate a new guid to force a recache of all landscape derived data
-#define LANDSCAPE_FULL_DERIVEDDATA_VER			TEXT("016D326F3A954BBA9CCDFA00CEFA31E9")
+#define LANDSCAPE_FULL_DERIVEDDATA_VER			TEXT("89D752865B0642E89CC8A5A2FED808A2")
 
 FString FLandscapeComponentDerivedData::GetDDCKeyString(const FGuid& StateId)
 {
@@ -3496,6 +3497,10 @@ void ULandscapeComponent::SerializeStateHashes(FArchive& Ar)
 	int32 OccluderGeometryLOD = GetLandscapeProxy()->OccluderGeometryLOD;
 	Ar << OccluderGeometryLOD;
 
+	bool bMeshHoles = GetLandscapeProxy()->bMeshHoles;
+	uint8 MeshHolesMaxLod = GetLandscapeProxy()->MeshHolesMaxLod;
+	Ar << bMeshHoles << MeshHolesMaxLod;
+
 	// Take into account the Heightmap offset per component
 	Ar << HeightmapScaleBias.Z;
 	Ar << HeightmapScaleBias.W;
diff --git a/Engine/Source/Runtime/Landscape/Private/LandscapeEdit.cpp b/Engine/Source/Runtime/Landscape/Private/LandscapeEdit.cpp
--- a/Engine/Source/Runtime/Landscape/Private/LandscapeEdit.cpp
+++ b/Engine/Source/Runtime/Landscape/Private/LandscapeEdit.cpp
@@ -4584,7 +4584,8 @@ void ALandscapeProxy::PostEditChangeProperty(FPropertyChangedEvent& PropertyChan
 			}
 		}
 	}
-	else if (GIsEditor && PropertyName == GET_MEMBER_NAME_CHECKED(ALandscapeProxy, OccluderGeometryLOD))
+	else if (GIsEditor && 
+		(PropertyName == GET_MEMBER_NAME_CHECKED(ALandscapeProxy, OccluderGeometryLOD) || PropertyName == GET_MEMBER_NAME_CHECKED(ALandscapeProxy, bMeshHoles) || PropertyName == GET_MEMBER_NAME_CHECKED(ALandscapeProxy, MeshHolesMaxLod)))
 	{
 		CheckGenerateLandscapePlatformData(false, nullptr);
 		MarkComponentsRenderStateDirty();
@@ -6278,9 +6279,325 @@ void ULandscapeComponent::GeneratePlatformPixelData()
 	}
 }
 
-//
-// Generates vertex buffer data from the component's heightmap texture, for use on platforms without vertex texture fetch
-//
+
+// FBox2D version that uses integers
+struct FIntBox2D
+{
+	FIntBox2D() : Min(INT32_MAX, INT32_MAX), Max(-INT32_MAX, -INT32_MAX) {}
+
+	void Add(FIntPoint const& Pos) 
+	{
+		Min = FIntPoint(FMath::Min(Min.X, Pos.X), FMath::Min(Min.Y, Pos.Y));
+		Max = FIntPoint(FMath::Max(Max.X, Pos.X), FMath::Max(Max.Y, Pos.Y));
+	}
+
+	void Add(FIntBox2D const& Rhs)
+	{
+		Min = FIntPoint(FMath::Min(Min.X, Rhs.Min.X), FMath::Min(Min.Y, Rhs.Min.Y));
+		Max = FIntPoint(FMath::Max(Max.X, Rhs.Max.X), FMath::Max(Max.Y, Rhs.Max.Y));
+	}
+
+	bool Intersects(FIntBox2D const& Rhs)
+	{
+		return !((Rhs.Max.X < Min.X) || (Rhs.Min.X > Max.X) || (Rhs.Max.Y < Min.Y) || (Rhs.Min.Y > Max.Y));
+	}
+
+	FIntPoint Min;
+	FIntPoint Max;
+};
+
+// Segment the hole map and return an array of hole bounding rectangles
+void GetHoleBounds(int32 InSize, TArray<uint8> const& InVisibilityData, TArray<FIntBox2D>& OutHoleBounds)
+{
+	check(InVisibilityData.Num() == InSize * InSize);
+
+	TArray<uint32> HoleSegmentLabels;
+	HoleSegmentLabels.AddZeroed(InSize*InSize);
+
+	TArray<uint32, TInlineAllocator<32>> LabelEquivalenceMap;
+	LabelEquivalenceMap.Add(0);
+	uint32 NextLabel = 1;
+
+	// First pass fills HoleSegmentLabels with labels.
+	for (int32 y = 0; y < InSize; ++y)
+	{
+		for (int32 x = 0; x < InSize; ++x)
+		{
+			const uint8 VisThreshold = 170;
+			const bool bIsHole = InVisibilityData[y * InSize + x] >= VisThreshold;
+			if (bIsHole)
+			{
+				uint8 WestLabel = (x > 0) ? HoleSegmentLabels[y * InSize + x - 1] : 0;
+				uint8 NorthLabel = (y > 0) ? HoleSegmentLabels[(y - 1) * InSize + x] : 0;
+
+				if (WestLabel != 0 && NorthLabel != 0 && WestLabel != NorthLabel)
+				{
+					uint32 MinLabel = FMath::Min(WestLabel, NorthLabel);
+					uint32 MaxLabel = FMath::Max(WestLabel, NorthLabel);
+					LabelEquivalenceMap[MaxLabel] = MinLabel;
+					HoleSegmentLabels[y * InSize + x] = MinLabel;
+				}
+				else if (WestLabel != 0)
+				{
+					HoleSegmentLabels[y * InSize + x] = WestLabel;
+				}
+				else if (NorthLabel != 0)
+				{
+					HoleSegmentLabels[y * InSize + x] = NorthLabel;
+				}
+				else
+				{
+					LabelEquivalenceMap.Add(NextLabel);
+					HoleSegmentLabels[y * InSize + x] = NextLabel++;
+				}
+			}
+		}
+	}
+
+	// Resolve label equivalences.
+	for (int32 Index = 0; Index < LabelEquivalenceMap.Num(); ++ Index)
+	{
+		int32 CommonIndex = Index;
+		while (LabelEquivalenceMap[CommonIndex] != CommonIndex)
+		{
+			CommonIndex = LabelEquivalenceMap[CommonIndex];
+		}
+		LabelEquivalenceMap[Index] = CommonIndex;
+	}
+
+	// Flatten labels to be contiguous.
+	int32 NumLabels = 0;
+	for (int32 Index = 0; Index < LabelEquivalenceMap.Num(); ++Index)
+	{
+		if (LabelEquivalenceMap[Index] == Index)
+		{
+			LabelEquivalenceMap[Index] = NumLabels++;
+		}
+		else
+		{
+			LabelEquivalenceMap[Index] = LabelEquivalenceMap[LabelEquivalenceMap[Index]];
+		}
+	}
+
+	// Second pass finds bounds for each label.
+	// Could also write contiguous labels to HoleSegmentLabels here if we want to keep that info.
+	OutHoleBounds.AddDefaulted(NumLabels);
+	for (int32 y = 0; y < InSize - 1; ++y)
+	{
+		for (int32 x = 0; x < InSize - 1; ++x)
+		{
+			const int32 Index = InSize * y + x;
+			const int32 Label = LabelEquivalenceMap[HoleSegmentLabels[Index]];
+			OutHoleBounds[Label].Add(FIntPoint(x, y));
+		}
+	}
+}
+
+// Move vertex index up to the next location which obeys the condition:
+// PosAt(VertexIndex, LodIndex) > PosAt(VertexIndex - 1, LodIndex + 1)
+// Maths derived from pattern when analyzing a spreadsheet containing a dump of lod vertex positions.
+inline void AlignVertexDown(int32 InLodIndex, int32& InOutVertexIndex)
+{
+	const int32 Offset = InOutVertexIndex & ((2 << InLodIndex) - 1);
+	if (Offset < (1 << InLodIndex))
+	{
+		InOutVertexIndex -= Offset;
+	}
+}
+
+// Move vertex index up to the next location which obeys the condition:
+// PosAt(VertexIndex, LodIndex) < PosAt(VertexIndex + 1, LodIndex + 1)
+// Maths derived from pattern when analyzing a spreadsheet containing a dump of lod vertex positions.
+inline void AlignVertexUp(int32 InLodIndex, int32& InOutVertexIndex)
+{
+	const int32 Offset = (InOutVertexIndex + 1) & ((2 << InLodIndex) - 1);
+	if (Offset > (1 << InLodIndex))
+	{
+		InOutVertexIndex += (1 << InLodIndex) - Offset;
+	}
+}
+
+// Expand bounding rectangles from LodIndex-1 to LodIndex
+void ExpandBoundsForLod(int32 InSize, int32 InLodIndex, TArray<FIntBox2D> const& InHoleBounds, TArray<FIntBox2D>& OutHoleBounds)
+{
+	OutHoleBounds.AddZeroed(InHoleBounds.Num());
+	for (int32 i = 0; i < InHoleBounds.Num(); ++i)
+	{
+		// Expand
+		const int32 ExpandDistance = (2 << InLodIndex) - 1;
+		OutHoleBounds[i].Min.X = InHoleBounds[i].Min.X - ExpandDistance;
+		OutHoleBounds[i].Min.Y = InHoleBounds[i].Min.Y - ExpandDistance;
+		OutHoleBounds[i].Max.X = InHoleBounds[i].Max.X + ExpandDistance;
+		OutHoleBounds[i].Max.Y = InHoleBounds[i].Max.Y + ExpandDistance;
+
+		// Snap to continuous LOD borders so that consecutive vertices with different LODs don't overlap
+		if (InLodIndex > 0)
+		{
+			AlignVertexDown(InLodIndex, OutHoleBounds[i].Min.X);
+			AlignVertexDown(InLodIndex, OutHoleBounds[i].Min.Y);
+			AlignVertexUp(InLodIndex, OutHoleBounds[i].Max.X);
+			AlignVertexUp(InLodIndex, OutHoleBounds[i].Max.Y);
+		}
+
+		// Clamp to edges
+		OutHoleBounds[i].Min.X = FMath::Max(OutHoleBounds[i].Min.X, 0);
+		OutHoleBounds[i].Max.X = FMath::Min(OutHoleBounds[i].Max.X, InSize - 1);
+		OutHoleBounds[i].Min.Y = FMath::Max(OutHoleBounds[i].Min.Y, 0);
+		OutHoleBounds[i].Max.Y = FMath::Min(OutHoleBounds[i].Max.Y, InSize - 1);
+	}
+}
+
+// Combine intersecting bounding rectangles into to form their bounding rectangles.
+void CombineIntersectingBounds(TArray<FIntBox2D>& InOutHoleBounds)
+{
+	int i = 1;
+	while (i < InOutHoleBounds.Num())
+	{
+		int j = i + 1;
+		for (; j < InOutHoleBounds.Num(); ++j)
+		{
+			if (InOutHoleBounds[i].Intersects(InOutHoleBounds[j]))
+			{
+				InOutHoleBounds[i].Add(InOutHoleBounds[j]);
+				InOutHoleBounds.RemoveAtSwap(j);
+				break;
+			}
+		}
+		if (j == InOutHoleBounds.Num())
+		{
+			++i;
+		}
+	}
+}
+
+// Build an array with an entry per vertex which contains the Lod at which that vertex falls inside a hole bounding rectangle. 
+// This is the Lod at which we should clamp the vertex in the vertex shader.
+void BuildHoleVertexLods(int32 InSize, int32 InNumLods, TArray<FIntBox2D> const& InHoleBounds, TArray<uint8>& OutHoleVertexLods)
+{
+	// Generate hole bounds for each Lod level from Lod0 InHoleBounds
+	TArray< TArray<FIntBox2D> > HoleBoundsPerLevel;
+	HoleBoundsPerLevel.AddDefaulted(InNumLods);
+	HoleBoundsPerLevel[0] = InHoleBounds;
+
+	for (int32 LodIndex = 1; LodIndex < InNumLods; ++LodIndex)
+	{
+		ExpandBoundsForLod(InSize, LodIndex, HoleBoundsPerLevel[LodIndex - 1], HoleBoundsPerLevel[LodIndex]);
+	}
+
+	for (int32 LodIndex = 0; LodIndex < InNumLods; ++LodIndex)
+	{
+		CombineIntersectingBounds(HoleBoundsPerLevel[LodIndex]);
+	}
+
+	// Initialize output to the max Lod
+	OutHoleVertexLods.Init(InNumLods, InSize * InSize);
+
+	// Fill by writing each Lod level in turn
+	for (int32 LodIndex = InNumLods - 1; LodIndex >= 0; --LodIndex)
+	{
+		TArray<FIntBox2D> const& HoleBoundsAtLevel = HoleBoundsPerLevel[LodIndex];
+		for (int32 BoxIndex = 1; BoxIndex < HoleBoundsAtLevel.Num(); ++BoxIndex)
+		{
+			const FIntPoint Min = HoleBoundsAtLevel[BoxIndex].Min;
+			const FIntPoint Max = HoleBoundsAtLevel[BoxIndex].Max;
+			
+			for (int32 y = Min.Y; y <= Max.Y; ++y)
+			{
+				for (int32 x = Min.X; x <= Max.X; ++x)
+				{
+					OutHoleVertexLods[y * InSize + x] = LodIndex;
+				}
+			}
+		}
+	}
+}
+
+// Structure containing the hole render data required by the runtime rendering.
+template <typename INDEX_TYPE>
+struct FLandscapeHoleRenderData
+{
+	TArray<INDEX_TYPE> HoleIndices;
+	int32 MinIndex;
+	int32 MaxIndex;
+};
+
+// Serialize the hole render data.
+template <typename INDEX_TYPE>
+void SerializeHoleRenderData(FMemoryArchive& Ar, FLandscapeHoleRenderData<INDEX_TYPE>& InHoleRenderData)
+{
+	bool b16BitIndices = sizeof(INDEX_TYPE) == 2;
+	Ar << b16BitIndices;
+
+	Ar << InHoleRenderData.MinIndex;
+	Ar << InHoleRenderData.MaxIndex;
+
+	int32 HoleIndexCount = InHoleRenderData.HoleIndices.Num();
+	Ar << HoleIndexCount;
+	Ar.Serialize(InHoleRenderData.HoleIndices.GetData(), HoleIndexCount * sizeof(INDEX_TYPE));
+}
+
+// Take the processed hole map and generate the hole render data.
+template <typename INDEX_TYPE>
+void BuildHoleRenderData(int32 InNumSubsections, int32 InSubsectionSizeVerts, TArray<uint8> const& InVisibilityData, TMap<uint64, int32>& InVertexMap, FLandscapeHoleRenderData<INDEX_TYPE>& OutHoleRenderData)
+{
+	const int32 SizeVerts = InNumSubsections * InSubsectionSizeVerts;
+	const int32 SubsectionSizeQuads = InSubsectionSizeVerts - 1;
+	const uint8 VisThreshold = 170;
+
+	INDEX_TYPE MaxIndex = 0;
+	INDEX_TYPE MinIndex = TNumericLimits<INDEX_TYPE>::Max();
+
+	for (int32 SubY = 0; SubY < InNumSubsections; SubY++)
+	{
+		for (int32 SubX = 0; SubX < InNumSubsections; SubX++)
+		{
+			for (int32 y = 0; y < SubsectionSizeQuads; y++)
+			{
+				for (int32 x = 0; x < SubsectionSizeQuads; x++)
+				{
+					const int32 x0 = x;
+					const int32 y0 = y;
+					const int32 x1 = x + 1;
+					const int32 y1 = y + 1;
+
+					const int32 VertexIndex = (SubY * InSubsectionSizeVerts + y0) * SizeVerts + SubX * InSubsectionSizeVerts + x0;
+					const bool bIsHole = InVisibilityData[VertexIndex] < VisThreshold;
+					if (bIsHole)
+					{
+						INDEX_TYPE i00 = *(InVertexMap.Find(FLandscapeVertexRef(x0, y0, SubX, SubY).MakeKey()));
+						INDEX_TYPE i10 = *(InVertexMap.Find(FLandscapeVertexRef(x1, y0, SubX, SubY).MakeKey()));
+						INDEX_TYPE i11 = *(InVertexMap.Find(FLandscapeVertexRef(x1, y1, SubX, SubY).MakeKey()));
+						INDEX_TYPE i01 = *(InVertexMap.Find(FLandscapeVertexRef(x0, y1, SubX, SubY).MakeKey()));
+
+						OutHoleRenderData.HoleIndices.Add(i00);
+						OutHoleRenderData.HoleIndices.Add(i11);
+						OutHoleRenderData.HoleIndices.Add(i10);
+
+						OutHoleRenderData.HoleIndices.Add(i00);
+						OutHoleRenderData.HoleIndices.Add(i01);
+						OutHoleRenderData.HoleIndices.Add(i11);
+
+						// Update the min/max index ranges
+						MaxIndex = FMath::Max<INDEX_TYPE>(MaxIndex, i00);
+						MinIndex = FMath::Min<INDEX_TYPE>(MinIndex, i00);
+						MaxIndex = FMath::Max<INDEX_TYPE>(MaxIndex, i10);
+						MinIndex = FMath::Min<INDEX_TYPE>(MinIndex, i10);
+						MaxIndex = FMath::Max<INDEX_TYPE>(MaxIndex, i11);
+						MinIndex = FMath::Min<INDEX_TYPE>(MinIndex, i11);
+						MaxIndex = FMath::Max<INDEX_TYPE>(MaxIndex, i01);
+						MinIndex = FMath::Min<INDEX_TYPE>(MinIndex, i01);
+					}
+				}
+			}
+		}
+	}
+
+	OutHoleRenderData.MinIndex = MinIndex;
+	OutHoleRenderData.MinIndex = MaxIndex;
+}
+
+// Generates vertex and index buffer data from the component's height map and visibility textures.
+// For use on mobile platforms that don't use vertex texture fetch for height or alpha testing for visibility.
 void ULandscapeComponent::GeneratePlatformVertexData(const ITargetPlatform* TargetPlatform)
 {
 	if (IsTemplate())
@@ -6293,16 +6610,17 @@ void ULandscapeComponent::GeneratePlatformVertexData(const ITargetPlatform* Targ
 	TArray<uint8> NewPlatformData;
 	FMemoryWriter PlatformAr(NewPlatformData);
 
-	int32 SubsectionSizeVerts = SubsectionSizeQuads + 1;
-	int32 MaxLOD = FMath::CeilLogTwo(SubsectionSizeVerts) - 1;
+	const int32 SubsectionSizeVerts = SubsectionSizeQuads + 1;
+	const int32 MaxLOD = FMath::CeilLogTwo(SubsectionSizeVerts) - 1;
+	const int32 NumMips = FMath::Min(LANDSCAPE_MAX_ES_LOD, GetHeightmap()->Source.GetNumMips());
 
-	float HeightmapSubsectionOffsetU = (float)(SubsectionSizeVerts) / (float)GetHeightmap()->Source.GetSizeX();
-	float HeightmapSubsectionOffsetV = (float)(SubsectionSizeVerts) / (float)GetHeightmap()->Source.GetSizeY();
-	
-	// Get the required mip data
+	const float HeightmapSubsectionOffsetU = (float)(SubsectionSizeVerts) / (float)GetHeightmap()->Source.GetSizeX();
+	const float HeightmapSubsectionOffsetV = (float)(SubsectionSizeVerts) / (float)GetHeightmap()->Source.GetSizeY();
+
+	// Get the required height mip data
 	TArray<TArray64<uint8>> HeightmapMipRawData;
 	TArray64<FColor*> HeightmapMipData;
-	for (int32 MipIdx = 0; MipIdx < FMath::Min(LANDSCAPE_MAX_ES_LOD, GetHeightmap()->Source.GetNumMips()); MipIdx++)
+	for (int32 MipIdx = 0; MipIdx < NumMips; MipIdx++)
 	{
 		int32 MipSubsectionSizeVerts = (SubsectionSizeVerts) >> MipIdx;
 		if (MipSubsectionSizeVerts > 1)
@@ -6313,11 +6631,32 @@ void ULandscapeComponent::GeneratePlatformVertexData(const ITargetPlatform* Targ
 		}
 	}
 
+	// Get any hole data
+	int32 NumHoleLods = 0;
+	TArray< uint8 > VisibilityData;
+	if (ComponentHasVisibilityPainted() && GetLandscapeProxy()->bMeshHoles)
+	{
+		TArray<FWeightmapLayerAllocationInfo>& ComponentWeightmapLayerAllocations = GetWeightmapLayerAllocations();
+		for (int32 AllocIdx = 0; AllocIdx < ComponentWeightmapLayerAllocations.Num(); AllocIdx++)
+		{
+			FWeightmapLayerAllocationInfo& AllocInfo = ComponentWeightmapLayerAllocations[AllocIdx];
+			if (AllocInfo.LayerInfo == ALandscapeProxy::VisibilityLayer)
+			{
+				NumHoleLods = FMath::Clamp<int32>(GetLandscapeProxy()->MeshHolesMaxLod, 1, NumMips);
+
+				FLandscapeComponentDataInterface CDI(this, 0);
+				CDI.GetWeightmapTextureData(AllocInfo.LayerInfo, VisibilityData);
+				break;
+			}
+		}
+	}
+
+	// Layout index buffer to determine best vertex order.
+	// This vertex layout code is duplicated in FLandscapeSharedBuffers::CreateIndexBuffers() to create matching index buffers at runtime.
 	TMap<uint64, int32> VertexMap;
 	TArray<FLandscapeVertexRef> VertexOrder;
 	VertexOrder.Empty(FMath::Square(SubsectionSizeVerts * NumSubsections));
 
-	// Layout index buffer to determine best vertex order
 	for (int32 Mip = MaxLOD; Mip >= 0; Mip--)
 	{
 		int32 LodSubsectionSizeQuads = (SubsectionSizeVerts >> Mip) - 1;
@@ -6377,34 +6716,67 @@ void ULandscapeComponent::GeneratePlatformVertexData(const ITargetPlatform* Targ
 			VertexOrder.Num(), FMath::Square(SubsectionSizeVerts) * FMath::Square(NumSubsections));
 	}
 
-	int32 NumMobileVerices = FMath::Square(SubsectionSizeVerts * NumSubsections);
+	// Build and serialize hole render data which includes a unique index buffer with the holes missing.
+	// This fills HoleVertexLods which is required for filling the vertex data.
+	TArray<uint8> HoleVertexLods;
+	PlatformAr << NumHoleLods;
+	if (NumHoleLods > 0)
+	{
+		TArray<FIntBox2D> HoleBounds;
+		GetHoleBounds(SubsectionSizeVerts * NumSubsections, VisibilityData, HoleBounds);
+		BuildHoleVertexLods(SubsectionSizeVerts * NumSubsections, NumHoleLods, HoleBounds, HoleVertexLods);
+
+		if (VertexOrder.Num() <= UINT16_MAX)
+		{
+			FLandscapeHoleRenderData<uint16> HoleRenderData;
+			BuildHoleRenderData(NumSubsections, SubsectionSizeVerts, VisibilityData, VertexMap, HoleRenderData);
+			SerializeHoleRenderData(PlatformAr, HoleRenderData);
+		}
+		else
+		{
+			FLandscapeHoleRenderData<uint32> HoleRenderData;
+			BuildHoleRenderData(NumSubsections, SubsectionSizeVerts, VisibilityData, VertexMap, HoleRenderData);
+			SerializeHoleRenderData(PlatformAr, HoleRenderData);
+		}
+	}
+
+	// Fill in the vertices in the specified order.
+	const int32 SizeVerts = SubsectionSizeVerts * NumSubsections;
+	int32 NumMobileVertices = FMath::Square(SizeVerts);
 	TArray<FLandscapeMobileVertex> MobileVertices;
-	MobileVertices.AddZeroed(NumMobileVerices);
+	MobileVertices.AddZeroed(NumMobileVertices);
 	FLandscapeMobileVertex* DstVert = MobileVertices.GetData();
 
-	// Fill in the vertices in the specified order
 	for (int32 Idx = 0; Idx < VertexOrder.Num(); Idx++)
 	{
-		int32 X = VertexOrder[Idx].X;
-		int32 Y = VertexOrder[Idx].Y;
-		int32 SubX = VertexOrder[Idx].SubX;
-		int32 SubY = VertexOrder[Idx].SubY;
-
-		float HeightmapScaleBiasZ = HeightmapScaleBias.Z + HeightmapSubsectionOffsetU * (float)SubX;
-		float HeightmapScaleBiasW = HeightmapScaleBias.W + HeightmapSubsectionOffsetV * (float)SubY;
-		int32 BaseMipOfsX = FMath::RoundToInt(HeightmapScaleBiasZ * (float)GetHeightmap()->Source.GetSizeX());
-		int32 BaseMipOfsY = FMath::RoundToInt(HeightmapScaleBiasW * (float)GetHeightmap()->Source.GetSizeY());
+		// Store XY position info
+		const int32 X = VertexOrder[Idx].X;
+		const int32 Y = VertexOrder[Idx].Y;
+		const int32 SubX = VertexOrder[Idx].SubX;
+		const int32 SubY = VertexOrder[Idx].SubY;
 
 		DstVert->Position[0] = X;
 		DstVert->Position[1] = Y;
-		DstVert->Position[2] = SubX;
-		DstVert->Position[3] = SubY;
+		DstVert->Position[2] = (SubX << 4) | SubY;
+
+		// Store hole info
+		const int32 VertexIndex = (SubY * SubsectionSizeVerts + Y) * SizeVerts + SubX * SubsectionSizeVerts + X;
+		const int32 HoleVertexLod = (NumHoleLods > 0) ? HoleVertexLods[VertexIndex] : 0;
+		const int32 HoleMaxLod = (NumHoleLods > 0) ? NumHoleLods : 0;
+
+		DstVert->Position[3] = (HoleMaxLod << 4) | HoleVertexLod;
 
+		// Calculate min/max height for packing
 		TArray<int32> MipHeights;
 		MipHeights.AddZeroed(HeightmapMipData.Num());
 		int32 LastIndex = 0;
 		uint16 MaxHeight = 0, MinHeight = 65535;
 
+		float HeightmapScaleBiasZ = HeightmapScaleBias.Z + HeightmapSubsectionOffsetU * (float)SubX;
+		float HeightmapScaleBiasW = HeightmapScaleBias.W + HeightmapSubsectionOffsetV * (float)SubY;
+		int32 BaseMipOfsX = FMath::RoundToInt(HeightmapScaleBiasZ * (float)GetHeightmap()->Source.GetSizeX());
+		int32 BaseMipOfsY = FMath::RoundToInt(HeightmapScaleBiasW * (float)GetHeightmap()->Source.GetSizeY());
+
 		for (int32 Mip = 0; Mip < HeightmapMipData.Num(); ++Mip)
 		{
 			int32 MipSizeX = GetHeightmap()->Source.GetSizeX() >> Mip;
@@ -6423,29 +6795,27 @@ void ULandscapeComponent::GeneratePlatformVertexData(const ITargetPlatform* Targ
 			MinHeight = FMath::Min(MinHeight, Height);
 		}
 
+		// Quantize min/max height so we can store each in 8 bits
+		MaxHeight = (MaxHeight + 255) & (~255);
+		MinHeight = (MinHeight) & (~255);
+
 		DstVert->LODHeights[0] = MinHeight >> 8;
-		DstVert->LODHeights[1] = (MinHeight & 255);
-		DstVert->LODHeights[2] = MaxHeight >> 8;
-		DstVert->LODHeights[3] = (MaxHeight & 255);
+		DstVert->LODHeights[1] = MaxHeight >> 8;
 
+		// Now quantize the mip heights to steps between MinHeight and MaxHeight
 		for (int32 Mip = 0; Mip < HeightmapMipData.Num(); ++Mip)
 		{
-			if (Mip < 4)
-			{
-				DstVert->LODHeights[4 + Mip] = FMath::RoundToInt(float(MipHeights[Mip] - MinHeight) / (MaxHeight - MinHeight) * 255);
-			}
-			else // Mip 4 5 packed into SubX, SubY
-			{
-				DstVert->Position[Mip - 2] += (FMath::RoundToInt(float(MipHeights[Mip] - MinHeight) / (MaxHeight - MinHeight) * 255)) & (0xfffe);
-			}
+			check(Mip < 6);
+			DstVert->LODHeights[2 + Mip] = FMath::RoundToInt(float(MipHeights[Mip] - MinHeight) / (MaxHeight - MinHeight) * 255);
 		}
 
 		DstVert++;
 	}
 
-	PlatformAr << NumMobileVerices;
-	PlatformAr.Serialize(MobileVertices.GetData(), NumMobileVerices*sizeof(FLandscapeMobileVertex));
-	
+	// Serialize vertex buffer
+	PlatformAr << NumMobileVertices;
+	PlatformAr.Serialize(MobileVertices.GetData(), NumMobileVertices*sizeof(FLandscapeMobileVertex));
+
 	// Generate occlusion mesh
 	TArray<FVector> OccluderVertices;
 	const int32 OcclusionMeshMip = FMath::Clamp<int32>(GetLandscapeProxy()->OccluderGeometryLOD, -1, HeightmapMipData.Num() - 1);
diff --git a/Engine/Source/Runtime/Landscape/Private/LandscapePrivate.h b/Engine/Source/Runtime/Landscape/Private/LandscapePrivate.h
--- a/Engine/Source/Runtime/Landscape/Private/LandscapePrivate.h
+++ b/Engine/Source/Runtime/Landscape/Private/LandscapePrivate.h
@@ -43,4 +43,5 @@ DECLARE_DWORD_COUNTER_STAT_EXTERN(TEXT("DrawCalls"), STAT_LandscapeDrawCalls, ST
 
 DECLARE_MEMORY_STAT_EXTERN(TEXT("Vertex Mem"), STAT_LandscapeVertexMem, STATGROUP_Landscape, );
 DECLARE_MEMORY_STAT_EXTERN(TEXT("Occluder Mem"), STAT_LandscapeOccluderMem, STATGROUP_Landscape, );
+DECLARE_MEMORY_STAT_EXTERN(TEXT("Hole Mem"), STAT_LandscapeHoleMem, STATGROUP_Landscape, );
 DECLARE_DWORD_COUNTER_STAT_EXTERN(TEXT("Component Mem"), STAT_LandscapeComponentMem, STATGROUP_Landscape, );
diff --git a/Engine/Source/Runtime/Landscape/Private/LandscapeRender.cpp b/Engine/Source/Runtime/Landscape/Private/LandscapeRender.cpp
--- a/Engine/Source/Runtime/Landscape/Private/LandscapeRender.cpp
+++ b/Engine/Source/Runtime/Landscape/Private/LandscapeRender.cpp
@@ -2323,6 +2323,9 @@ bool FLandscapeComponentSceneProxy::GetStaticMeshElement(int32 LODIndex, bool bF
 		BatchElement.FirstIndex = 0;
 		BatchElement.MinVertexIndex = SharedBuffers->IndexRanges[LODIndex].MinIndexFull;
 		BatchElement.MaxVertexIndex = SharedBuffers->IndexRanges[LODIndex].MaxIndexFull;
+
+		// The default is overridden here only by mobile landscape to punch holes in the geometry
+		ApplyMeshElementModifier(BatchElement, LODIndex);
 	}
 
 	return true;
diff --git a/Engine/Source/Runtime/Landscape/Private/LandscapeRenderMobile.cpp b/Engine/Source/Runtime/Landscape/Private/LandscapeRenderMobile.cpp
--- a/Engine/Source/Runtime/Landscape/Private/LandscapeRenderMobile.cpp
+++ b/Engine/Source/Runtime/Landscape/Private/LandscapeRenderMobile.cpp
@@ -13,6 +13,13 @@ LandscapeRenderMobile.cpp: Landscape Rendering without using vertex texture fetc
 #include "HAL/LowLevelMemTracker.h"
 #include "MeshMaterialShader.h"
 
+// Debug CVar for disabling the loading of landscape hole meshes
+static TAutoConsoleVariable<int32> CVarMobileLandscapeHoleMesh(
+	TEXT("r.Mobile.LandscapeHoleMesh"),
+	1,
+	TEXT("Set to 0 to skip loading of landscape hole meshes on mobile."),
+	ECVF_Default);
+
 void FLandscapeVertexFactoryMobile::InitRHI()
 {
 	// list of declaration items
@@ -349,38 +356,100 @@ void FLandscapeVertexBufferMobile::InitRHI()
 	RHIUnlockVertexBuffer(VertexBufferRHI);
 }
 
+struct FLandscapeMobileHoleData
+{
+	FRawStaticIndexBuffer16or32Interface* IndexBuffer = nullptr;
+	int32 NumHoleLods;
+	int32 MinHoleIndex;
+	int32 MaxHoleIndex;
+
+	~FLandscapeMobileHoleData()
+	{
+		if (IndexBuffer != nullptr)
+		{
+			DEC_DWORD_STAT_BY(STAT_LandscapeHoleMem, IndexBuffer->GetResourceDataSize());
+			IndexBuffer->ReleaseResource();
+			delete IndexBuffer;
+		}
+	}
+};
+
+template <typename INDEX_TYPE>
+void SerializeLandscapeMobileHoleData(FMemoryArchive& Ar, FLandscapeMobileHoleData& HoleData)
+{
+	Ar << HoleData.MinHoleIndex;
+	Ar << HoleData.MaxHoleIndex;
+
+	TArray<INDEX_TYPE> IndexData;
+	int32 IndexCount = 0;
+	Ar << IndexCount;
+	IndexData.SetNumUninitialized(IndexCount);
+	Ar.Serialize(IndexData.GetData(), IndexCount * sizeof(INDEX_TYPE));
+
+	const bool bLoadHoleMeshData = IndexCount > 0 && CVarMobileLandscapeHoleMesh.GetValueOnGameThread();
+	if (bLoadHoleMeshData)
+	{
+		FRawStaticIndexBuffer16or32<INDEX_TYPE>* IndexBuffer = new FRawStaticIndexBuffer16or32<INDEX_TYPE>(false);
+		IndexBuffer->AssignNewBuffer(IndexData);
+		HoleData.IndexBuffer = IndexBuffer;
+		BeginInitResource(HoleData.IndexBuffer);
+		INC_DWORD_STAT_BY(STAT_LandscapeHoleMem, HoleData.IndexBuffer->GetResourceDataSize());
+	}
+}
+
 /**
  * Container for FLandscapeVertexBufferMobile that we can reference from a thread-safe shared pointer
  * while ensuring the vertex buffer is always destroyed on the render thread.
  **/
 struct FLandscapeMobileRenderData
 {
-	FLandscapeVertexBufferMobile* VertexBuffer;
+	FLandscapeVertexBufferMobile* VertexBuffer = nullptr;
+	FLandscapeMobileHoleData* HoleData = nullptr;
 	FOccluderVertexArraySP OccluderVerticesSP;
 
 	FLandscapeMobileRenderData(const TArray<uint8>& InPlatformData)
 	{
 		FMemoryReader MemAr(InPlatformData);
+
+		int32 NumHoleLods;
+		MemAr << NumHoleLods;
+		if (NumHoleLods > 0)
 		{
-			int32 NumMobileVertices;
-			TArray<uint8> MobileVerticesData;
+			HoleData = new FLandscapeMobileHoleData;
+			HoleData->NumHoleLods = NumHoleLods;
 
-			MemAr << NumMobileVertices;
-			MobileVerticesData.SetNumUninitialized(NumMobileVertices*sizeof(FLandscapeMobileVertex));
-			MemAr.Serialize(MobileVerticesData.GetData(), MobileVerticesData.Num());
+			bool b16BitIndices;
+			MemAr << b16BitIndices;
+			if (b16BitIndices)
+			{
+				SerializeLandscapeMobileHoleData<uint16>(MemAr, *HoleData);
+			}
+			else
+			{
+				SerializeLandscapeMobileHoleData<uint32>(MemAr, *HoleData);
+			}
+		}
 
-			VertexBuffer = new FLandscapeVertexBufferMobile(MoveTemp(MobileVerticesData));
+		{
+			int32 VertexCount;
+			MemAr << VertexCount;
+			TArray<uint8> VertexData;
+			VertexData.SetNumUninitialized(VertexCount*sizeof(FLandscapeMobileVertex));
+			MemAr.Serialize(VertexData.GetData(), VertexData.Num());
+			VertexBuffer = new FLandscapeVertexBufferMobile(MoveTemp(VertexData));
 		}
 		
-		int32 NumOccluderVertices;
-		MemAr << NumOccluderVertices;
-		if (NumOccluderVertices > 0)
 		{
-			OccluderVerticesSP = MakeShared<FOccluderVertexArray, ESPMode::ThreadSafe>();
-			OccluderVerticesSP->SetNumUninitialized(NumOccluderVertices);
-			MemAr.Serialize(OccluderVerticesSP->GetData(), NumOccluderVertices*sizeof(FVector));
+			int32 NumOccluderVertices;
+			MemAr << NumOccluderVertices;
+			if (NumOccluderVertices > 0)
+			{
+				OccluderVerticesSP = MakeShared<FOccluderVertexArray, ESPMode::ThreadSafe>();
+				OccluderVerticesSP->SetNumUninitialized(NumOccluderVertices);
+				MemAr.Serialize(OccluderVerticesSP->GetData(), NumOccluderVertices * sizeof(FVector));
 
-			INC_DWORD_STAT_BY(STAT_LandscapeOccluderMem, OccluderVerticesSP->GetAllocatedSize());
+				INC_DWORD_STAT_BY(STAT_LandscapeOccluderMem, OccluderVerticesSP->GetAllocatedSize());
+			}
 		}
 	}
 
@@ -392,15 +461,18 @@ struct FLandscapeMobileRenderData
 			if (IsInRenderingThread())
 			{
 				delete VertexBuffer;
+				delete HoleData;
 			}
 			else
 			{
 				FLandscapeVertexBufferMobile* InVertexBuffer = VertexBuffer;
+				FLandscapeMobileHoleData* InHoleData = HoleData;
 				ENQUEUE_RENDER_COMMAND(InitCommand)(
-					[InVertexBuffer](FRHICommandListImmediate& RHICmdList)
-					{
-						delete InVertexBuffer;
-					});
+					[InVertexBuffer, InHoleData](FRHICommandListImmediate& RHICmdList)
+				{
+					delete InVertexBuffer;
+					delete InHoleData;
+				});
 			}
 		}
 
@@ -571,4 +643,18 @@ TSharedPtr<FLandscapeMobileRenderData, ESPMode::ThreadSafe> FLandscapeComponentD
 
 		return RenderData;
 	}
-}
\ No newline at end of file
+}
+
+void FLandscapeComponentSceneProxyMobile::ApplyMeshElementModifier(FMeshBatchElement& InOutMeshElement, int32 InLodIndex) const
+{
+	const bool bHoleDataExists = MobileRenderData->HoleData != nullptr && MobileRenderData->HoleData->IndexBuffer != nullptr && InLodIndex < MobileRenderData->HoleData->NumHoleLods;
+	if (bHoleDataExists)
+	{
+		FLandscapeMobileHoleData const& HoleData = *MobileRenderData->HoleData;
+		InOutMeshElement.IndexBuffer = HoleData.IndexBuffer;
+		InOutMeshElement.NumPrimitives = HoleData.IndexBuffer->Num() / 3;
+		InOutMeshElement.FirstIndex = 0;
+		InOutMeshElement.MinVertexIndex = HoleData.MinHoleIndex;
+		InOutMeshElement.MaxVertexIndex = HoleData.MaxHoleIndex;
+	}
+}
diff --git a/Engine/Source/Runtime/Landscape/Public/LandscapeRender.h b/Engine/Source/Runtime/Landscape/Public/LandscapeRender.h
--- a/Engine/Source/Runtime/Landscape/Public/LandscapeRender.h
+++ b/Engine/Source/Runtime/Landscape/Public/LandscapeRender.h
@@ -1005,6 +1005,8 @@ protected:
 	template<class ArrayType> bool GetStaticMeshElement(int32 LODIndex, bool bForToolMesh, bool bForcedLOD, FMeshBatch& MeshBatch, ArrayType& OutStaticBatchParamArray) const;
 	void BuildDynamicMeshElement(const FViewCustomDataLOD* InPrimitiveCustomData, bool InToolMesh, bool InHasTessellation, bool InDisableTessellation, FMeshBatch& OutMeshBatch, TArray<FLandscapeBatchElementParams, SceneRenderingAllocator>& OutStaticBatchParamArray) const;
 
+	virtual void ApplyMeshElementModifier(FMeshBatchElement& InOutMeshElement, int32 InLodIndex) const {}
+
 	float GetComponentScreenSize(const class FSceneView* View, const FVector& Origin,  float MaxExtend, float ElementRadius) const;
 
 public:
diff --git a/Engine/Source/Runtime/Landscape/Public/LandscapeRenderMobile.h b/Engine/Source/Runtime/Landscape/Public/LandscapeRenderMobile.h
--- a/Engine/Source/Runtime/Landscape/Public/LandscapeRenderMobile.h
+++ b/Engine/Source/Runtime/Landscape/Public/LandscapeRenderMobile.h
@@ -133,6 +133,8 @@ public:
 	uint8 BlendableLayerMask;
 
 	friend class FLandscapeVertexBufferMobile;
+
+	virtual void ApplyMeshElementModifier(FMeshBatchElement& InOutMeshElement, int32 InLodIndex) const override;
 };
 
 
