PR #6588: Optimized ALandscapeProxy::UpdateGrass (Contributed by [REDACTED_PERSON])

diff --git a/Engine/Source/Runtime/Landscape/Private/LandscapeGrass.cpp b/Engine/Source/Runtime/Landscape/Private/LandscapeGrass.cpp
--- a/Engine/Source/Runtime/Landscape/Private/LandscapeGrass.cpp
+++ b/Engine/Source/Runtime/Landscape/Private/LandscapeGrass.cpp
@@ -2350,7 +2350,7 @@ void ALandscapeProxy::UpdateGrass(const TArray<FVector>& Cameras, int32& InOutNu
 			for (ULandscapeComponent* Component : LandscapeComponents)
 			{
 				// skip if we have no data and no way to generate it
-				if (Component == nullptr || (World->IsGameWorld() && !Component->GrassData->HasData()))
+				if (Component == nullptr || (World->IsGameWorld() && Component->GrassData->WeightData.Num() == 0))
 				{
 					continue;
 				}
@@ -2404,6 +2404,11 @@ void ALandscapeProxy::UpdateGrass(const TArray<FVector>& Cameras, int32& InOutNu
 				{
 					if (GrassType)
 					{
+						if (World->IsGameWorld() && Component->GrassData->WeightData.Find(GrassType) == nullptr)
+						{
+							continue;
+						}
+
 						int32 GrassVarietyIndex = -1;
 						uint32 HaltonBaseIndex = 1;
 						for (auto& GrassVariety : GrassType->GrassVarieties)
