Fixes for mobile landscape holes working on device
#rnx

diff --git a/Engine/Source/Runtime/Landscape/Private/Landscape.cpp b/Engine/Source/Runtime/Landscape/Private/Landscape.cpp
--- a/Engine/Source/Runtime/Landscape/Private/Landscape.cpp
+++ b/Engine/Source/Runtime/Landscape/Private/Landscape.cpp
@@ -118,7 +118,7 @@ namespace LandscapeCookStats
 // differences, etc.) replace the version GUID below with a new one.
 // In case of merge conflicts with DDC versions, you MUST generate a new GUID
 // and set this new GUID as the version.                                       
-#define LANDSCAPE_MOBILE_COOK_VERSION TEXT("EAA6E15CEDD644308E8B8D5427EC180")
+#define LANDSCAPE_MOBILE_COOK_VERSION TEXT("8BCE3D5A734C48D9BB53B4B9B1778D28")
 
 #define LOCTEXT_NAMESPACE "Landscape"
 
diff --git a/Engine/Source/Runtime/Landscape/Private/LandscapeEdit.cpp b/Engine/Source/Runtime/Landscape/Private/LandscapeEdit.cpp
--- a/Engine/Source/Runtime/Landscape/Private/LandscapeEdit.cpp
+++ b/Engine/Source/Runtime/Landscape/Private/LandscapeEdit.cpp
@@ -6593,7 +6593,7 @@ void BuildHoleRenderData(int32 InNumSubsections, int32 InSubsectionSizeVerts, TA
 	}
 
 	OutHoleRenderData.MinIndex = MinIndex;
-	OutHoleRenderData.MinIndex = MaxIndex;
+	OutHoleRenderData.MaxIndex = MaxIndex;
 }
 
 // Generates vertex and index buffer data from the component's height map and visibility textures.
diff --git a/Engine/Source/Runtime/Landscape/Private/LandscapeRenderMobile.cpp b/Engine/Source/Runtime/Landscape/Private/LandscapeRenderMobile.cpp
--- a/Engine/Source/Runtime/Landscape/Private/LandscapeRenderMobile.cpp
+++ b/Engine/Source/Runtime/Landscape/Private/LandscapeRenderMobile.cpp
@@ -362,6 +362,7 @@ struct FLandscapeMobileHoleData
 {
 	FRawStaticIndexBuffer16or32Interface* IndexBuffer = nullptr;
 	int32 NumHoleLods;
+	int32 IndexBufferSize;
 	int32 MinHoleIndex;
 	int32 MaxHoleIndex;
 
@@ -383,12 +384,11 @@ void SerializeLandscapeMobileHoleData(FMemoryArchive& Ar, FLandscapeMobileHoleDa
 	Ar << HoleData.MaxHoleIndex;
 
 	TArray<INDEX_TYPE> IndexData;
-	int32 IndexCount = 0;
-	Ar << IndexCount;
-	IndexData.SetNumUninitialized(IndexCount);
-	Ar.Serialize(IndexData.GetData(), IndexCount * sizeof(INDEX_TYPE));
+	Ar << HoleData.IndexBufferSize;
+	IndexData.SetNumUninitialized(HoleData.IndexBufferSize);
+	Ar.Serialize(IndexData.GetData(), HoleData.IndexBufferSize * sizeof(INDEX_TYPE));
 
-	const bool bLoadHoleMeshData = IndexCount > 0 && CVarMobileLandscapeHoleMesh.GetValueOnGameThread();
+	const bool bLoadHoleMeshData = HoleData.IndexBufferSize > 0 && CVarMobileLandscapeHoleMesh.GetValueOnGameThread();
 	if (bLoadHoleMeshData)
 	{
 		FRawStaticIndexBuffer16or32<INDEX_TYPE>* IndexBuffer = new FRawStaticIndexBuffer16or32<INDEX_TYPE>(false);
@@ -654,7 +654,7 @@ void FLandscapeComponentSceneProxyMobile::ApplyMeshElementModifier(FMeshBatchEle
 	{
 		FLandscapeMobileHoleData const& HoleData = *MobileRenderData->HoleData;
 		InOutMeshElement.IndexBuffer = HoleData.IndexBuffer;
-		InOutMeshElement.NumPrimitives = HoleData.IndexBuffer->Num() / 3;
+		InOutMeshElement.NumPrimitives = HoleData.IndexBufferSize / 3;
 		InOutMeshElement.FirstIndex = 0;
 		InOutMeshElement.MinVertexIndex = HoleData.MinHoleIndex;
 		InOutMeshElement.MaxVertexIndex = HoleData.MaxHoleIndex;
