Fixed lost landscape data when opening/making visible a level with non-edit layers data only while the main landscape level has been converted to edit layers : moved copy from old data to layers in PostRegisterAllComponents, which is executed upon making the level visible.
Also fixed the logic for detecting data conversion is needed and fixed a crash when cleaning up weightmap layer allocations

patrick.enfedaque

diff --git a/Engine/Source/Runtime/Landscape/Classes/LandscapeProxy.h b/Engine/Source/Runtime/Landscape/Classes/LandscapeProxy.h
--- a/Engine/Source/Runtime/Landscape/Classes/LandscapeProxy.h
+++ b/Engine/Source/Runtime/Landscape/Classes/LandscapeProxy.h
@@ -916,7 +916,7 @@ public:
 	static void AddReferencedObjects(UObject* InThis, FReferenceCollector& Collector);
 	virtual void PostLoad() override;
 
-	LANDSCAPE_API ULandscapeInfo* CreateLandscapeInfo();
+	LANDSCAPE_API ULandscapeInfo* CreateLandscapeInfo(bool bMapCheck = false);
 	LANDSCAPE_API ULandscapeInfo* GetLandscapeInfo() const;
 
 	/** Get the LandcapeActor-to-world transform with respect to landscape section offset*/
diff --git a/Engine/Source/Runtime/Landscape/Private/Landscape.cpp b/Engine/Source/Runtime/Landscape/Private/Landscape.cpp
--- a/Engine/Source/Runtime/Landscape/Private/Landscape.cpp
+++ b/Engine/Source/Runtime/Landscape/Private/Landscape.cpp
@@ -1167,7 +1167,7 @@ ALandscape* ALandscapeStreamingProxy::GetLandscapeActor()
 	return LandscapeActor.Get();
 }
 
-ULandscapeInfo* ALandscapeProxy::CreateLandscapeInfo()
+ULandscapeInfo* ALandscapeProxy::CreateLandscapeInfo(bool bMapCheck)
 {
 	ULandscapeInfo* LandscapeInfo = nullptr;
 
@@ -1186,7 +1186,7 @@ ULandscapeInfo* ALandscapeProxy::CreateLandscapeInfo()
 		LandscapeInfoMap.Map.Add(LandscapeGuid, LandscapeInfo);
 	}
 	check(LandscapeInfo);
-	LandscapeInfo->RegisterActor(this);
+	LandscapeInfo->RegisterActor(this, bMapCheck);
 
 	return LandscapeInfo;
 }
@@ -1891,7 +1891,30 @@ void ALandscapeProxy::PostRegisterAllComponents()
 		// Duplicated Landscapes don't have a valid guid until PostEditImport is called, we'll register then
 		if (LandscapeGuid.IsValid())
 		{
-			LandscapeInfo = CreateLandscapeInfo();
+#if WITH_EDITOR
+			if (GIsEditor && !GetWorld()->IsGameWorld())
+			{
+				UpdateCachedHasLayersContent(true);
+
+				// Cache the value at this point as CreateLandscapeInfo (-> RegisterActor) might create/destroy layers content if there was a mismatch between landscape & proxy
+				// Check the actual flag here not HasLayersContent() which could return true if the LandscapeActor is valid.
+				bool bHasLayersContentBefore = bHasLayersContent;
+
+				LandscapeInfo = CreateLandscapeInfo(true);
+
+				FixupWeightmaps();
+
+				const bool bNeedOldDataMigration = !bHasLayersContentBefore && CanHaveLayersContent();
+				if (bNeedOldDataMigration && LandscapeInfo->LandscapeActor.IsValid())
+				{
+					LandscapeInfo->LandscapeActor.Get()->CopyOldDataToDefaultLayer(this);
+				}
+			}
+			else
+#endif
+			{
+				LandscapeInfo = CreateLandscapeInfo(true);
+			}
 		}
 
 		if (UWorld* OwningWorld = GetWorld())
@@ -2418,26 +2441,6 @@ void ALandscapeProxy::PostLoad()
 		EditorCachedLayerInfos_DEPRECATED.Empty();
 	}
 
-	if (GIsEditor && !GetWorld()->IsGameWorld())
-	{
-		UpdateCachedHasLayersContent(true);
-		
-		ULandscapeInfo* LandscapeInfo = CreateLandscapeInfo();
-
-		// Cache the value at this point as RegisterActor might create/destroy layers content if there was a mismatch between landscape & proxy
-		// Check the actual flag here not HasLayersContent() which could return true if the LandscapeActor is valid.
-		const bool bNeedOldDataMigration = !bHasLayersContent && CanHaveLayersContent();
-				
-		LandscapeInfo->RegisterActor(this, true);
-
-		FixupWeightmaps();
-
-		if (bNeedOldDataMigration && LandscapeInfo->LandscapeActor.IsValid())
-		{
-			LandscapeInfo->LandscapeActor.Get()->CopyOldDataToDefaultLayer(this);
-		}
-	}
-
 	for (ULandscapeComponent* Comp : LandscapeComponents)
 	{
 		if (Comp)
diff --git a/Engine/Source/Runtime/Landscape/Private/LandscapeEdit.cpp b/Engine/Source/Runtime/Landscape/Private/LandscapeEdit.cpp
--- a/Engine/Source/Runtime/Landscape/Private/LandscapeEdit.cpp
+++ b/Engine/Source/Runtime/Landscape/Private/LandscapeEdit.cpp
@@ -697,9 +697,14 @@ void ULandscapeComponent::FixupWeightmaps()
 			bool bFixedWeightmapTextureIndex = false;
 
 			// Store the weightmap allocations in WeightmapUsageMap
-			for (int32 LayerIdx = 0; LayerIdx < WeightmapLayerAllocations.Num(); LayerIdx++)
+			for (int32 LayerIdx = 0; LayerIdx < WeightmapLayerAllocations.Num();)
 			{
 				FWeightmapLayerAllocationInfo& Allocation = WeightmapLayerAllocations[LayerIdx];
+				if (!Allocation.IsAllocated())
+				{
+					WeightmapLayerAllocations.RemoveAt(LayerIdx);
+					continue;
+				}
 
 				// Fix up any problems caused by the layer deletion bug.
 				if (Allocation.WeightmapTextureIndex >= WeightmapTextures.Num())
@@ -741,13 +746,13 @@ void ULandscapeComponent::FixupWeightmaps()
 						->AddToken(FTextToken::Create(FText::Format(LOCTEXT("MapCheck_Message_FixedUpSharedLayerWeightmap", "Fixed up shared weightmap texture for layer {LayerName} in component '{LandscapeName}' (shares with '{ChannelName}')"), Arguments)))
 						->AddToken(FMapErrorToken::Create(FMapErrors::FixedUpSharedLayerWeightmap));
 					WeightmapLayerAllocations.RemoveAt(LayerIdx);
-					LayerIdx--;
 					continue;
 				}
 				else
 				{
 					Usage->ChannelUsage[Allocation.WeightmapTextureChannel] = this;
 				}
+				++LayerIdx;
 			}
 
 			RemoveInvalidWeightmaps();
