Merging //UE4/Main @ 9404494 to //UE4/Dev-Core (stefan.boberg_rdu-p4merge-02_6162)

(ushell-p4-mergedown)

diff --git a/Engine/Plugins/Experimental/GeometryCollectionPlugin/Source/GeometryCollectionExampleCore/Private/GeometryCollection/GeometryCollectionExampleSimulationField.cpp b/Engine/Plugins/Experimental/GeometryCollectionPlugin/Source/GeometryCollectionExampleCore/Private/GeometryCollection/GeometryCollectionExampleSimulationField.cpp
--- a/Engine/Plugins/Experimental/GeometryCollectionPlugin/Source/GeometryCollectionExampleCore/Private/GeometryCollection/GeometryCollectionExampleSimulationField.cpp
+++ b/Engine/Plugins/Experimental/GeometryCollectionPlugin/Source/GeometryCollectionExampleCore/Private/GeometryCollection/GeometryCollectionExampleSimulationField.cpp
@@ -646,7 +646,6 @@ namespace GeometryCollectionExample
 
 		delete PhysObject;
 		delete RadialMask;
-#endif
 	}
 	template void RigidBodies_Field_CollisionGroup<float>();
 
@@ -738,7 +737,6 @@ namespace GeometryCollectionExample
 		delete PhysObject;
 		delete FalloffField;
 		delete FieldObject;
-#endif
 	}
 	template void RigidBodies_Field_ClusterBreak_StrainModel_Test1<float>();
 
@@ -844,7 +842,6 @@ namespace GeometryCollectionExample
 		delete PhysObject;
 		delete FalloffField;
 		delete FieldObject;
-#endif
 	}
 	template void RigidBodies_Field_ClusterBreak_StrainModel_Test2<float>();
 
@@ -949,7 +946,6 @@ namespace GeometryCollectionExample
 		delete PhysObject;
 		delete FalloffField;
 		delete FieldObject;
-#endif
 	}
 	template void RigidBodies_Field_ClusterBreak_StrainModel_Test3<float>();
 
@@ -1046,7 +1042,6 @@ namespace GeometryCollectionExample
 		delete PhysObject;
 		delete FalloffField;
 		delete FieldObject;
-#endif
 	}
 	template void RigidBodies_Field_ClusterBreak_StrainModel_Test4<float>();
 
diff --git a/Engine/Plugins/Experimental/GeometryCollectionPlugin/Source/GeometryCollectionExampleCore/Private/GeometryCollection/GeometryCollectionExampleUtility.cpp b/Engine/Plugins/Experimental/GeometryCollectionPlugin/Source/GeometryCollectionExampleCore/Private/GeometryCollection/GeometryCollectionExampleUtility.cpp
--- a/Engine/Plugins/Experimental/GeometryCollectionPlugin/Source/GeometryCollectionExampleCore/Private/GeometryCollection/GeometryCollectionExampleUtility.cpp
+++ b/Engine/Plugins/Experimental/GeometryCollectionPlugin/Source/GeometryCollectionExampleCore/Private/GeometryCollection/GeometryCollectionExampleUtility.cpp
@@ -284,6 +284,7 @@ namespace GeometryCollectionExample {
 		return PhysObject;
 #else
 		return nullptr;
+#endif
 	}
 	template FGeometryCollectionPhysicsProxy* RigidBodySetup(
 		TUniquePtr<Chaos::TChaosPhysicsMaterial<float>> & PhysicalMaterial,
diff --git a/Engine/Source/Runtime/Landscape/Classes/LandscapeProxy.h b/Engine/Source/Runtime/Landscape/Classes/LandscapeProxy.h
--- a/Engine/Source/Runtime/Landscape/Classes/LandscapeProxy.h
+++ b/Engine/Source/Runtime/Landscape/Classes/LandscapeProxy.h
@@ -436,7 +436,7 @@ public:
 	float ComponentScreenSizeToUseSubSections;
 
 	/** This is the starting screen size used to calculate the distribution, by default it's 1, but you can increase the value if you want less LOD0 component, and you use very large landscape component. */
-	UPROPERTY(EditAnywhere, Category = "LOD Distribution", meta = (DisplayName = "LOD 0 Screen Size", ClampMin = "1.0", ClampMax = "10.0", UIMin = "1.0", UIMax = "10.0"))
+	UPROPERTY(EditAnywhere, Category = "LOD Distribution", meta = (DisplayName = "LOD 0 Screen Size", ClampMin = "0.1", ClampMax = "10.0", UIMin = "1.0", UIMax = "10.0"))
 	float LOD0ScreenSize;
 
 	/** The distribution setting used to change the LOD 0 generation, 1.75 is the normal distribution, numbers influence directly the LOD0 proportion on screen. */
@@ -871,6 +871,9 @@ public:
 	LANDSCAPE_API ULandscapeInfo* CreateLandscapeInfo();
 	LANDSCAPE_API ULandscapeInfo* GetLandscapeInfo() const;
 
+	/** Get the LandcapeActor-to-world transform with respect to landscape section offset*/
+	LANDSCAPE_API FTransform LandscapeActorToWorld() const;
+
 #if WITH_EDITOR
 	virtual void PostEditChangeProperty(FPropertyChangedEvent& PropertyChangedEvent) override;
 	virtual void PostEditImport() override;
@@ -903,10 +906,7 @@ public:
 
 	// Assign only mismatching data and mark proxy package dirty
 	LANDSCAPE_API void FixupSharedData(ALandscape* Landscape);
-	
-	/** Get the LandcapeActor-to-world transform with respect to landscape section offset*/
-	LANDSCAPE_API FTransform LandscapeActorToWorld() const;
-	
+
 	/** Set landscape absolute location in section space */
 	LANDSCAPE_API void SetAbsoluteSectionBase(FIntPoint SectionOffset);
 	
