Copying Private-LoadTimes-4.24 stream to Main. Biggest changes are in Materials/Shader memory freezing.

diff --git a/Engine/Source/Runtime/Engine/Public/PerPlatformProperties.h b/Engine/Source/Runtime/Engine/Public/PerPlatformProperties.h
--- a/Engine/Source/Runtime/Engine/Public/PerPlatformProperties.h
+++ b/Engine/Source/Runtime/Engine/Public/PerPlatformProperties.h
@@ -10,6 +10,7 @@ PerPlatformProperties.h: Property types that can be overridden on a per-platform
 #include "RHIDefinitions.h"
 #include "Containers/Map.h"
 #include "Algo/Find.h"
+#include "Serialization/MemoryLayout.h"
 #include "PerPlatformProperties.generated.h"
 
 /** TPerPlatformProperty - template parent class for per-platform properties 
@@ -117,6 +118,8 @@ ENGINE_API FArchive& operator<<(FArchive& Ar, TPerPlatformProperty<_StructType,
 template<typename _StructType, typename _ValueType, EName _BasePropertyName>
 ENGINE_API void operator<<(FStructuredArchive::FSlot Slot, TPerPlatformProperty<_StructType, _ValueType, _BasePropertyName>& P);
 
+struct FFreezablePerPlatformInt;
+
 /** FPerPlatformInt - int32 property with per-platform overrides */
 USTRUCT()
 struct ENGINE_API FPerPlatformInt
@@ -145,7 +148,46 @@ struct ENGINE_API FPerPlatformInt
 	}
 
 	FString ToString() const;
+
+	FPerPlatformInt(const FFreezablePerPlatformInt& Other);
+
+};
+
+USTRUCT()
+struct ENGINE_API FFreezablePerPlatformInt
+#if CPP
+	: public TPerPlatformProperty<FFreezablePerPlatformInt, int32, NAME_IntProperty>
+#endif
+{
+	DECLARE_TYPE_LAYOUT(FFreezablePerPlatformInt, NonVirtual);
+
+	GENERATED_USTRUCT_BODY()
+
+public:
+	using FPerPlatformMap = TMemoryImageMap<FName, int32>;
+
+	LAYOUT_FIELD(int32, Default);
+	LAYOUT_FIELD_EDITORONLY(FPerPlatformMap, PerPlatform);
+
+	FFreezablePerPlatformInt() : Default(0) {}
+	FFreezablePerPlatformInt(int32 InDefaultValue) : Default(InDefaultValue) {}
+	FFreezablePerPlatformInt(const FPerPlatformInt& Other)
+		: Default(Other.Default)
+#if WITH_EDITORONLY_DATA
+		, PerPlatform(Other.PerPlatform)
+#endif
+	{}
+
+	FString ToString() const;
 };
+
+inline FPerPlatformInt::FPerPlatformInt(const FFreezablePerPlatformInt& Other)
+	: Default(Other.Default)
+#if WITH_EDITORONLY_DATA
+	, PerPlatform(Other.PerPlatform)
+#endif
+{}
+
 extern template ENGINE_API FArchive& operator<<(FArchive&, TPerPlatformProperty<FPerPlatformInt, int32, NAME_IntProperty>&);
 extern template ENGINE_API void operator<<(FStructuredArchive::FSlot Slot, TPerPlatformProperty<FPerPlatformInt, int32, NAME_IntProperty>&);
 
@@ -160,8 +202,10 @@ struct TStructOpsTypeTraits<FPerPlatformInt>
 	};
 };
 
+struct FFreezablePerPlatformFloat;
+
 /** FPerPlatformFloat - float property with per-platform overrides */
-USTRUCT()
+USTRUCT(meta = (CanFlattenStruct))
 struct ENGINE_API FPerPlatformFloat
 #if CPP
 :	public TPerPlatformProperty<FPerPlatformFloat, float, NAME_FloatProperty>
@@ -186,9 +230,40 @@ struct ENGINE_API FPerPlatformFloat
 	:	Default(InDefaultValue)
 	{
 	}
+
+	FPerPlatformFloat(const FFreezablePerPlatformFloat& Other);
 };
 extern template ENGINE_API FArchive& operator<<(FArchive&, TPerPlatformProperty<FPerPlatformFloat, float, NAME_FloatProperty>&);
 
+struct ENGINE_API FFreezablePerPlatformFloat
+#if CPP
+	: public TPerPlatformProperty<FFreezablePerPlatformFloat, float, NAME_FloatProperty>
+#endif
+{
+	DECLARE_TYPE_LAYOUT(FFreezablePerPlatformFloat, NonVirtual);
+public:
+	using FPerPlatformMap = TMemoryImageMap<FName, float>;
+	
+	LAYOUT_FIELD(float, Default);
+	LAYOUT_FIELD_EDITORONLY(FPerPlatformMap, PerPlatform);
+
+	FFreezablePerPlatformFloat() : Default(0.0f) {}
+	FFreezablePerPlatformFloat(float InDefaultValue) : Default(InDefaultValue) {}
+	FFreezablePerPlatformFloat(const FPerPlatformFloat& Other)
+		: Default(Other.Default)
+#if WITH_EDITORONLY_DATA
+		, PerPlatform(Other.PerPlatform)
+#endif
+	{}
+};
+
+inline FPerPlatformFloat::FPerPlatformFloat(const FFreezablePerPlatformFloat& Other)
+	: Default(Other.Default)
+#if WITH_EDITORONLY_DATA
+	, PerPlatform(Other.PerPlatform)
+#endif
+{}
+
 template<>
 struct TStructOpsTypeTraits<FPerPlatformFloat>
 :	public TStructOpsTypeTraitsBase2<FPerPlatformFloat>
diff --git a/Engine/Source/Runtime/Engine/Public/PrimitiveSceneProxy.h b/Engine/Source/Runtime/Engine/Public/PrimitiveSceneProxy.h
--- a/Engine/Source/Runtime/Engine/Public/PrimitiveSceneProxy.h
+++ b/Engine/Source/Runtime/Engine/Public/PrimitiveSceneProxy.h
@@ -1115,10 +1115,10 @@ protected:
 /**
  * Returns if specified mesh command can be cached, or needs to be recreated every frame.
  */
-ENGINE_API extern bool SupportsCachingMeshDrawCommands(const FPrimitiveSceneProxy* RESTRICT PrimitiveSceneProxy, const FMeshBatch& MeshBatch);
+ENGINE_API extern bool SupportsCachingMeshDrawCommands(const FMeshBatch& MeshBatch);
 
 /**
  * Returns if specified mesh command can be cached, or needs to be recreated every frame; this is a slightly slower version
  * used for materials with external textures that need invalidating their PSOs.
  */
-ENGINE_API extern bool SupportsCachingMeshDrawCommands(const FPrimitiveSceneProxy* RESTRICT PrimitiveSceneProxy, const FMeshBatch& MeshBatch, ERHIFeatureLevel::Type FeatureLevel);
+ENGINE_API extern bool SupportsCachingMeshDrawCommands(const FMeshBatch& MeshBatch, ERHIFeatureLevel::Type FeatureLevel);
\ No newline at end of file
diff --git a/Engine/Source/Runtime/Engine/Public/PrimitiveUniformShaderParameters.h b/Engine/Source/Runtime/Engine/Public/PrimitiveUniformShaderParameters.h
--- a/Engine/Source/Runtime/Engine/Public/PrimitiveUniformShaderParameters.h
+++ b/Engine/Source/Runtime/Engine/Public/PrimitiveUniformShaderParameters.h
@@ -245,7 +245,7 @@ public:
 
 	ENGINE_API void UploadToGPU();
 
-	EShaderPlatform ShaderPlatform;
+	EShaderPlatform ShaderPlatform=SP_NumPlatforms;
 
 	FPrimitiveSceneShaderData PrimitiveSceneData;
 	FLightmapSceneShaderData LightmapSceneData;
diff --git a/Engine/Source/Runtime/Engine/Public/RawIndexBuffer.h b/Engine/Source/Runtime/Engine/Public/RawIndexBuffer.h
--- a/Engine/Source/Runtime/Engine/Public/RawIndexBuffer.h
+++ b/Engine/Source/Runtime/Engine/Public/RawIndexBuffer.h
@@ -11,6 +11,7 @@
 #include "RenderResource.h"
 #include "Containers/DynamicRHIResourceArray.h"
 
+
 class FRawIndexBuffer : public FIndexBuffer
 {
 public:
@@ -290,7 +291,7 @@ private:
 	FIndexBufferRHIRef CreateRHIBuffer_Internal();
 
 	/** Storage for indices. */
-	TResourceArray<uint8,INDEXBUFFER_ALIGNMENT> IndexStorage;
+	TResourceArray<uint8, INDEXBUFFER_ALIGNMENT> IndexStorage;
 
 	/** If >= 0, represents the number of indices in this IB. Needed in cooked build since InitRHI may discard CPU data */
 	int32 CachedNumIndices;
@@ -446,7 +447,8 @@ public:
 
 	virtual void AssignNewBuffer(const TArray<INDEX_TYPE>& Buffer)
 	{
-		Indices = TArray<INDEX_TYPE,TAlignedHeapAllocator<INDEXBUFFER_ALIGNMENT> >(Buffer);
+		using IndexBufferType = typename TResourceArray<INDEX_TYPE, INDEXBUFFER_ALIGNMENT>::Super;
+		Indices = IndexBufferType(Buffer);
 		CachedNumIndices = Indices.Num();
 	}
 
diff --git a/Engine/Source/Runtime/Engine/Public/Rendering/ColorVertexBuffer.h b/Engine/Source/Runtime/Engine/Public/Rendering/ColorVertexBuffer.h
--- a/Engine/Source/Runtime/Engine/Public/Rendering/ColorVertexBuffer.h
+++ b/Engine/Source/Runtime/Engine/Public/Rendering/ColorVertexBuffer.h
@@ -3,6 +3,7 @@
 #pragma once
 
 #include "RenderResource.h"
+#include "StaticMeshVertexData.h"
 
 struct FStaticMeshBuildVertex;
 
@@ -211,4 +212,4 @@ private:
 
 	/** Purposely hidden */
 	ENGINE_API FColorVertexBuffer(const FColorVertexBuffer &rhs);
-};
\ No newline at end of file
+};
diff --git a/Engine/Source/Runtime/Engine/Public/Rendering/PositionVertexBuffer.h b/Engine/Source/Runtime/Engine/Public/Rendering/PositionVertexBuffer.h
--- a/Engine/Source/Runtime/Engine/Public/Rendering/PositionVertexBuffer.h
+++ b/Engine/Source/Runtime/Engine/Public/Rendering/PositionVertexBuffer.h
@@ -143,7 +143,7 @@ private:
 	FShaderResourceViewRHIRef PositionComponentSRV;
 
 	/** The vertex data storage type */
-	class FPositionVertexData* VertexData;
+	TMemoryImagePtr<class FPositionVertexData> VertexData;
 
 	/** The cached vertex data pointer. */
 	uint8* Data;
@@ -162,3 +162,4 @@ private:
 	template <bool bRenderThread>
 	FVertexBufferRHIRef CreateRHIBuffer_Internal();
 };
+
diff --git a/Engine/Source/Runtime/Engine/Public/Rendering/SkeletalMeshDuplicatedVerticesBuffer.h b/Engine/Source/Runtime/Engine/Public/Rendering/SkeletalMeshDuplicatedVerticesBuffer.h
--- a/Engine/Source/Runtime/Engine/Public/Rendering/SkeletalMeshDuplicatedVerticesBuffer.h
+++ b/Engine/Source/Runtime/Engine/Public/Rendering/SkeletalMeshDuplicatedVerticesBuffer.h
@@ -25,6 +25,7 @@ struct FIndexLengthPair
         return Ar;
     }
 };
+template<> struct TCanBulkSerialize<FIndexLengthPair> { enum { Value = true }; };
 
 class FDuplicatedVerticesBuffer : public FRenderResource
 {
diff --git a/Engine/Source/Runtime/Engine/Public/Rendering/StaticMeshVertexDataInterface.h b/Engine/Source/Runtime/Engine/Public/Rendering/StaticMeshVertexDataInterface.h
--- a/Engine/Source/Runtime/Engine/Public/Rendering/StaticMeshVertexDataInterface.h
+++ b/Engine/Source/Runtime/Engine/Public/Rendering/StaticMeshVertexDataInterface.h
@@ -48,6 +48,8 @@ public:
 	/** Serializer. */
 	virtual void Serialize(FArchive& Ar, bool bForcePerElementSerialization = false) = 0;
 
+	virtual void OverrideFreezeSizeAndAlignment(int64& Size, int32& Alignment) const = 0;
+
 	virtual SIZE_T GetResourceSize() const = 0;
 
 	virtual bool GetAllowCPUAccess() const = 0;
diff --git a/Engine/Source/Runtime/Engine/Public/SceneManagement.h b/Engine/Source/Runtime/Engine/Public/SceneManagement.h
--- a/Engine/Source/Runtime/Engine/Public/SceneManagement.h
+++ b/Engine/Source/Runtime/Engine/Public/SceneManagement.h
@@ -1002,7 +1002,7 @@ public:
 	}
 };
 
-inline bool DoesPlatformSupportDistanceFields(EShaderPlatform Platform)
+inline bool DoesPlatformSupportDistanceFields(const FStaticShaderPlatform Platform)
 {
 	return Platform == SP_PCD3D_SM5
 		|| Platform == SP_PS4
@@ -1011,7 +1011,7 @@ inline bool DoesPlatformSupportDistanceFields(EShaderPlatform Platform)
 		|| IsVulkanSM5Platform(Platform)
 		|| Platform == SP_SWITCH
 		|| Platform == SP_SWITCH_FORWARD
-		|| FDataDrivenShaderPlatformInfo::GetInfo(Platform).bSupportsDistanceFields;
+		|| FDataDrivenShaderPlatformInfo::GetSupportsDistanceFields(Platform);
 }
 
 inline bool DoesPlatformSupportDistanceFieldShadowing(EShaderPlatform Platform)
diff --git a/Engine/Source/Runtime/Engine/Public/ScreenRendering.h b/Engine/Source/Runtime/Engine/Public/ScreenRendering.h
--- a/Engine/Source/Runtime/Engine/Public/ScreenRendering.h
+++ b/Engine/Source/Runtime/Engine/Public/ScreenRendering.h
@@ -77,25 +77,17 @@ public:
 
 	void SetParameters(FRHICommandList& RHICmdList, const FTexture* Texture)
 	{
-		SetTextureParameter(RHICmdList, GetPixelShader(),InTexture,InTextureSampler,Texture);
+		SetTextureParameter(RHICmdList, RHICmdList.GetBoundPixelShader(),InTexture,InTextureSampler,Texture);
 	}
 
 	void SetParameters(FRHICommandList& RHICmdList, FRHISamplerState* SamplerStateRHI, FRHITexture* TextureRHI)
 	{
-		SetTextureParameter(RHICmdList, GetPixelShader(),InTexture,InTextureSampler,SamplerStateRHI,TextureRHI);
-	}
-
-	virtual bool Serialize(FArchive& Ar) override
-	{
-		bool bShaderHasOutdatedParameters = FGlobalShader::Serialize(Ar);
-		Ar << InTexture;
-		Ar << InTextureSampler;
-		return bShaderHasOutdatedParameters;
+		SetTextureParameter(RHICmdList, RHICmdList.GetBoundPixelShader(),InTexture,InTextureSampler,SamplerStateRHI,TextureRHI);
 	}
 
 private:
-	FShaderResourceParameter InTexture;
-	FShaderResourceParameter InTextureSampler;
+	LAYOUT_FIELD(FShaderResourceParameter, InTexture);
+	LAYOUT_FIELD(FShaderResourceParameter, InTextureSampler);
 };
 
 /**
@@ -118,25 +110,17 @@ public:
 
 	void SetParameters(FRHICommandList& RHICmdList, const FTexture* Texture)
 	{
-		SetTextureParameter(RHICmdList, GetPixelShader(), InTexture, InTextureSampler, Texture);
+		SetTextureParameter(RHICmdList, RHICmdList.GetBoundPixelShader(), InTexture, InTextureSampler, Texture);
 	}
 
 	void SetParameters(FRHICommandList& RHICmdList, FRHISamplerState* SamplerStateRHI, FRHITexture* TextureRHI)
 	{
-		SetTextureParameter(RHICmdList, GetPixelShader(), InTexture, InTextureSampler, SamplerStateRHI, TextureRHI);
-	}
-
-	virtual bool Serialize(FArchive& Ar) override
-	{
-		bool bShaderHasOutdatedParameters = FGlobalShader::Serialize(Ar);
-		Ar << InTexture;
-		Ar << InTextureSampler;
-		return bShaderHasOutdatedParameters;
+		SetTextureParameter(RHICmdList, RHICmdList.GetBoundPixelShader(), InTexture, InTextureSampler, SamplerStateRHI, TextureRHI);
 	}
 
 private:
-	FShaderResourceParameter InTexture;
-	FShaderResourceParameter InTextureSampler;
+	LAYOUT_FIELD(FShaderResourceParameter, InTexture);
+	LAYOUT_FIELD(FShaderResourceParameter, InTextureSampler);
 };
 
 class FScreenPS_OSE : public FGlobalShader
@@ -157,25 +141,17 @@ public:
 
     void SetParameters(FRHICommandList& RHICmdList, const FTexture* Texture)
     {
-        SetTextureParameter(RHICmdList, GetPixelShader(),InTexture,InTextureSampler,Texture);
+        SetTextureParameter(RHICmdList, RHICmdList.GetBoundPixelShader(),InTexture,InTextureSampler,Texture);
     }
 
     void SetParameters(FRHICommandList& RHICmdList, FRHISamplerState* SamplerStateRHI, FRHITexture* TextureRHI)
     {
-        SetTextureParameter(RHICmdList, GetPixelShader(),InTexture,InTextureSampler,SamplerStateRHI,TextureRHI);
-    }
-
-    virtual bool Serialize(FArchive& Ar) override
-    {
-        bool bShaderHasOutdatedParameters = FGlobalShader::Serialize(Ar);
-        Ar << InTexture;
-        Ar << InTextureSampler;
-        return bShaderHasOutdatedParameters;
+        SetTextureParameter(RHICmdList, RHICmdList.GetBoundPixelShader(),InTexture,InTextureSampler,SamplerStateRHI,TextureRHI);
     }
 
 private:
-    FShaderResourceParameter InTexture;
-    FShaderResourceParameter InTextureSampler;
+	LAYOUT_FIELD(FShaderResourceParameter, InTexture);
+	LAYOUT_FIELD(FShaderResourceParameter, InTextureSampler);
 };
 
 /**
@@ -194,16 +170,9 @@ public:
 	}
 	FScreenVS() {}
 
-
 	void SetParameters(FRHICommandList& RHICmdList, FRHIUniformBuffer* ViewUniformBuffer)
 	{
-		FGlobalShader::SetParameters<FViewUniformShaderParameters>(RHICmdList, GetVertexShader(), ViewUniformBuffer);
-	}
-
-
-	virtual bool Serialize(FArchive& Ar) override
-	{
-		return FGlobalShader::Serialize(Ar);
+		FGlobalShader::SetParameters<FViewUniformShaderParameters>(RHICmdList, RHICmdList.GetBoundVertexShader(), ViewUniformBuffer);
 	}
 };
 
diff --git a/Engine/Source/Runtime/Engine/Public/ShaderCompiler.h b/Engine/Source/Runtime/Engine/Public/ShaderCompiler.h
--- a/Engine/Source/Runtime/Engine/Public/ShaderCompiler.h
+++ b/Engine/Source/Runtime/Engine/Public/ShaderCompiler.h
@@ -662,6 +662,8 @@ extern ENGINE_API void GlobalBeginCompileShader(
 	const FString& DebugExtension = ""
 	);
 
+extern void GetOutdatedShaderTypes(TArray<const FShaderType*>& OutdatedShaderTypes, TArray<const FShaderPipelineType*>& OutdatedShaderPipelineTypes, TArray<const FVertexFactoryType*>& OutdatedFactoryTypes);
+
 /** Implementation of the 'recompileshaders' console command.  Recompiles shaders at runtime based on various criteria. */
 extern bool RecompileShaders(const TCHAR* Cmd, FOutputDevice& Ar);
 
@@ -693,7 +695,7 @@ extern ENGINE_API bool RecompileChangedShadersForPlatform(const FString& Platfor
 * Begins recompiling the specified global shader types, and flushes their bound shader states.
 * FinishRecompileGlobalShaders must be called after this and before using the global shaders for anything.
 */
-extern ENGINE_API void BeginRecompileGlobalShaders(const TArray<FShaderType*>& OutdatedShaderTypes, const TArray<const FShaderPipelineType*>& OutdatedShaderPipelineTypes, EShaderPlatform ShaderPlatform);
+extern ENGINE_API void BeginRecompileGlobalShaders(const TArray<const FShaderType*>& OutdatedShaderTypes, const TArray<const FShaderPipelineType*>& OutdatedShaderPipelineTypes, EShaderPlatform ShaderPlatform, const ITargetPlatform* TargetPlatform = nullptr);
 
 /** Finishes recompiling global shaders.  Must be called after BeginRecompileGlobalShaders. */
 extern ENGINE_API void FinishRecompileGlobalShaders();
@@ -729,8 +731,9 @@ extern ENGINE_API void RecompileShadersForRemote(
 	bool bCompileChangedShaders = true);
 
 extern ENGINE_API void CompileGlobalShaderMap(bool bRefreshShaderMap=false);
-extern ENGINE_API void CompileGlobalShaderMap(EShaderPlatform Platform, bool bRefreshShaderMap = false);
 extern ENGINE_API void CompileGlobalShaderMap(ERHIFeatureLevel::Type InFeatureLevel, bool bRefreshShaderMap=false);
+extern ENGINE_API void CompileGlobalShaderMap(EShaderPlatform Platform, bool bRefreshShaderMap = false);
+extern ENGINE_API void CompileGlobalShaderMap(EShaderPlatform Platform, const ITargetPlatform* TargetPlatform, bool bRefreshShaderMap);
 
 extern ENGINE_API FString GetGlobalShaderMapDDCKey();
 
diff --git a/Engine/Source/Runtime/Engine/Public/SimpleElementShaders.h b/Engine/Source/Runtime/Engine/Public/SimpleElementShaders.h
--- a/Engine/Source/Runtime/Engine/Public/SimpleElementShaders.h
+++ b/Engine/Source/Runtime/Engine/Public/SimpleElementShaders.h
@@ -30,13 +30,13 @@ public:
 
 	/*ENGINE_API */void SetParameters(FRHICommandList& RHICmdList, const FMatrix& TransformValue, bool bSwitchVerticalAxis = false);
 
-	virtual bool Serialize(FArchive& Ar) override;
+	//virtual bool Serialize(FArchive& Ar) override;
 
 	static void ModifyCompilationEnvironment(const FGlobalShaderPermutationParameters& Parameters, FShaderCompilerEnvironment& OutEnvironment);
 
 private:
-	FShaderParameter Transform;
-	FShaderParameter SwitchVerticalAxis;
+	LAYOUT_FIELD(FShaderParameter, Transform)
+	LAYOUT_FIELD(FShaderParameter, SwitchVerticalAxis)
 };
 
 /**
@@ -63,15 +63,15 @@ public:
 
 	void SetParameters(FRHICommandList& RHICmdList, const FTexture* TextureValue );
 
-	virtual bool Serialize(FArchive& Ar) override;
+	//virtual bool Serialize(FArchive& Ar) override;
 
 private:
-	FShaderResourceParameter InTexture;
-	FShaderResourceParameter InTextureSampler;
-	FShaderParameter TextureComponentReplicate;
-	FShaderParameter TextureComponentReplicateAlpha;
-	FShaderParameter EditorCompositeDepthTestParameter;
-	FShaderParameter ScreenToPixel;
+	LAYOUT_FIELD(FShaderResourceParameter, InTexture)
+	LAYOUT_FIELD(FShaderResourceParameter, InTextureSampler)
+	LAYOUT_FIELD(FShaderParameter, TextureComponentReplicate)
+	LAYOUT_FIELD(FShaderParameter, TextureComponentReplicateAlpha)
+	LAYOUT_FIELD(FShaderParameter, EditorCompositeDepthTestParameter)
+	LAYOUT_FIELD(FShaderParameter, ScreenToPixel)
 };
 
 /**
@@ -91,7 +91,7 @@ public:
  */
 class FSimpleElementGammaBasePS : public FSimpleElementPS
 {
-	DECLARE_SHADER_TYPE(FSimpleElementGammaBasePS, Global);
+	DECLARE_TYPE_LAYOUT(FSimpleElementGammaBasePS, NonVirtual);
 public:
 
 	FSimpleElementGammaBasePS(const ShaderMetaType::CompiledShaderInitializerType& Initializer);
@@ -99,10 +99,12 @@ public:
 
 	void SetParameters(FRHICommandList& RHICmdList, const FTexture* Texture,float GammaValue,ESimpleElementBlendMode BlendMode);
 
-	virtual bool Serialize(FArchive& Ar) override;
+	//virtual bool Serialize(FArchive& Ar) override;
 
 private:
-	FShaderParameter Gamma;
+	
+		LAYOUT_FIELD(FShaderParameter, Gamma)
+	
 };
 
 template <bool bSRGBTexture>
@@ -141,7 +143,7 @@ public:
  */
 class FSimpleElementMaskedGammaBasePS : public FSimpleElementGammaBasePS
 {
-	DECLARE_SHADER_TYPE(FSimpleElementMaskedGammaBasePS, Global);
+	DECLARE_TYPE_LAYOUT(FSimpleElementMaskedGammaBasePS, NonVirtual);
 public:
 
 	FSimpleElementMaskedGammaBasePS(const ShaderMetaType::CompiledShaderInitializerType& Initializer);
@@ -149,10 +151,12 @@ public:
 
 	void SetParameters(FRHICommandList& RHICmdList, const FTexture* Texture,float Gamma,float ClipRefValue,ESimpleElementBlendMode BlendMode);
 
-	virtual bool Serialize(FArchive& Ar) override;
+	//virtual bool Serialize(FArchive& Ar) override;
 
 private:
-	FShaderParameter ClipRef;
+	
+		LAYOUT_FIELD(FShaderParameter, ClipRef)
+	
 };
 
 template <bool bSRGBTexture>
@@ -235,27 +239,27 @@ public:
 	* @param Ar - archive to serialize to
 	* @return true if any of the parameters were outdated
 	*/
-	virtual bool Serialize(FArchive& Ar) override;
+	//virtual bool Serialize(FArchive& Ar) override;
 	
 private:
 	/** The width to smooth the edge the texture */
-	FShaderParameter SmoothWidth;
+	LAYOUT_FIELD(FShaderParameter, SmoothWidth)
 	/** Toggles drop shadow rendering */
-	FShaderParameter EnableShadow;
+	LAYOUT_FIELD(FShaderParameter, EnableShadow)
 	/** 2D vector specifying the direction of shadow */
-	FShaderParameter ShadowDirection;
+	LAYOUT_FIELD(FShaderParameter, ShadowDirection)
 	/** Color of the shadowed pixels */
-	FShaderParameter ShadowColor;	
+	LAYOUT_FIELD(FShaderParameter, ShadowColor)	
 	/** The width to smooth the edge the shadow of the texture */
-	FShaderParameter ShadowSmoothWidth;
+	LAYOUT_FIELD(FShaderParameter, ShadowSmoothWidth)
 	/** whether to turn on the outline glow */
-	FShaderParameter EnableGlow;
+	LAYOUT_FIELD(FShaderParameter, EnableGlow)
 	/** base color to use for the glow */
-	FShaderParameter GlowColor;
+	LAYOUT_FIELD(FShaderParameter, GlowColor)
 	/** outline glow outer radius */
-	FShaderParameter GlowOuterRadius;
+	LAYOUT_FIELD(FShaderParameter, GlowOuterRadius)
 	/** outline glow inner radius */
-	FShaderParameter GlowInnerRadius;
+	LAYOUT_FIELD(FShaderParameter, GlowInnerRadius)
 };
 
 /**
@@ -277,11 +281,11 @@ public:
 
 	void SetParameters(FRHICommandList& RHICmdList, const FTexture* TextureValue);
 
-	virtual bool Serialize(FArchive& Ar) override;
+	//virtual bool Serialize(FArchive& Ar) override;
 
 private:
-	FShaderResourceParameter InTexture;
-	FShaderResourceParameter InTextureSampler;
+	LAYOUT_FIELD(FShaderResourceParameter, InTexture)
+	LAYOUT_FIELD(FShaderResourceParameter, InTextureSampler)
 };
 
 
@@ -313,13 +317,13 @@ public:
 	*/
 	void SetParameters(FRHICommandList& RHICmdList, const FTexture* TextureValue, const FMatrix& ColorWeightsValue, float GammaValue);
 
-	virtual bool Serialize(FArchive& Ar) override;
+	//virtual bool Serialize(FArchive& Ar) override;
 
 private:
-	FShaderResourceParameter InTexture;
-	FShaderResourceParameter InTextureSampler;
-	FShaderParameter ColorWeights; 
-	FShaderParameter Gamma;
+	LAYOUT_FIELD(FShaderResourceParameter, InTexture)
+	LAYOUT_FIELD(FShaderResourceParameter, InTextureSampler)
+	LAYOUT_FIELD(FShaderParameter, ColorWeights) 
+	LAYOUT_FIELD(FShaderParameter, Gamma)
 };
 
 template <class TSimpleElementBase, uint32 EncodedBlendMode>
diff --git a/Engine/Source/Runtime/Engine/Public/StaticMeshResources.h b/Engine/Source/Runtime/Engine/Public/StaticMeshResources.h
--- a/Engine/Source/Runtime/Engine/Public/StaticMeshResources.h
+++ b/Engine/Source/Runtime/Engine/Public/StaticMeshResources.h
@@ -250,7 +250,7 @@ protected:
 	FStaticMeshLODResources* Owner;
 };
 
-typedef TArray<FStaticMeshSectionAreaWeightedTriangleSampler> FStaticMeshSectionAreaWeightedTriangleSamplerArray;
+typedef TArray<FStaticMeshSectionAreaWeightedTriangleSampler, FMemoryImageAllocator> FStaticMeshSectionAreaWeightedTriangleSamplerArray;
 
 /** Represents GPU resource needed for area weighted uniform sampling of a mesh surface. */
 class FStaticMeshSectionAreaWeightedTriangleSamplerBuffer : public FRenderResource
@@ -341,7 +341,8 @@ struct FStaticMeshLODResources
 #endif // RHI_RAYTRACING
 
 	/** Sections for this LOD. */
-	TArray<FStaticMeshSection> Sections;
+	using FStaticMeshSectionArray = TArray<FStaticMeshSection, TInlineAllocator<1>>;
+	FStaticMeshSectionArray Sections;
 
 	/** Distance field data associated with this mesh, null if not present.  */
 	class FDistanceFieldVolumeData* DistanceFieldData; 
@@ -359,7 +360,7 @@ struct FStaticMeshLODResources
 	uint32 bHasReversedIndices : 1;
 
 	/** True if the reversed index buffers contained data at init. Needed as it will not be available to the CPU afterwards. */
-	uint32 bHasReversedDepthOnlyIndices: 1;
+	uint32 bHasReversedDepthOnlyIndices : 1;
 
 	uint32 bHasColorVertexData : 1;
 
@@ -548,6 +549,9 @@ struct ENGINE_API FStaticMeshVertexFactories
 	void ReleaseResources();
 };
 
+using FStaticMeshLODResourcesArray = TArray<FStaticMeshLODResources>;
+using FStaticMeshVertexFactoriesArray = TArray<FStaticMeshVertexFactories>;
+
 /**
  * FStaticMeshRenderData - All data needed to render a static mesh.
  */
@@ -558,8 +562,8 @@ public:
 	ENGINE_API FStaticMeshRenderData();
 
 	/** Per-LOD resources. */
-	TIndirectArray<FStaticMeshLODResources> LODResources;
-	TIndirectArray<FStaticMeshVertexFactories> LODVertexFactories;
+	FStaticMeshLODResourcesArray LODResources;
+	FStaticMeshVertexFactoriesArray LODVertexFactories;
 
 	/** Screen size to switch LODs */
 	FPerPlatformFloat ScreenSize[MAX_STATIC_MESH_LODS];
@@ -583,6 +587,7 @@ public:
 	uint8 CurrentFirstLODIdx;
 
 #if WITH_EDITORONLY_DATA
+
 	/** The derived data key associated with this render data. */
 	FString DerivedDataKey;
 
@@ -592,11 +597,13 @@ public:
 	/** UV data used for streaming accuracy debug view modes. In sync for rendering thread */
 	TArray<FMeshUVChannelInfo> UVChannelDataPerMaterial;
 
-	void SyncUVChannelData(const TArray<FStaticMaterial>& ObjectData);
 
 	/** The next cached derived data in the list. */
 	TUniquePtr<class FStaticMeshRenderData> NextCachedRenderData;
 
+
+	void SyncUVChannelData(const TArray<FStaticMaterial>& ObjectData);
+
 	/**
 	 * Cache derived renderable data for the static mesh with the provided
 	 * level of detail settings.
@@ -748,7 +755,7 @@ public:
 
 				if (Component->IsRegistered() && !Component->bRenderStateCreated)
 				{
-					Component->CreateRenderState_Concurrent();
+					Component->CreateRenderState_Concurrent(nullptr);
 					Scenes.Add(Component->GetScene());
 				}
 			}
@@ -799,7 +806,7 @@ public:
 	virtual int32 CollectOccluderElements(class FOccluderElementsCollector& Collector) const override;
 
 	virtual void CreateRenderThreadResources() override;
-		
+
 	virtual void DestroyRenderThreadResources() override;
 
 	/** Sets up a wireframe FMeshBatch for a specific LOD. */
@@ -931,7 +938,7 @@ protected:
 		};
 
 		/** Per-section information. */
-		TArray<FSectionInfo> Sections;
+		TArray<FSectionInfo, TInlineAllocator<1>> Sections;
 
 		/** Vertex color data for this LOD (or NULL when not overridden), FStaticMeshComponentLODInfo handle the release of the memory */
 		FColorVertexBuffer* OverrideColorVertexBuffer;
@@ -941,7 +948,7 @@ protected:
 		const FRawStaticIndexBuffer* PreCulledIndexBuffer;
 
 		/** Initialization constructor. */
-		FLODInfo(const UStaticMeshComponent* InComponent, const TIndirectArray<FStaticMeshVertexFactories>& InLODVertexFactories, int32 InLODIndex, int32 InClampedMinLOD, bool bLODsShareStaticLighting);
+		FLODInfo(const UStaticMeshComponent* InComponent, const FStaticMeshVertexFactoriesArray& InLODVertexFactories, int32 InLODIndex, int32 InClampedMinLOD, bool bLODsShareStaticLighting);
 
 		bool UsesMeshModifyingMaterials() const { return bUsesMeshModifyingMaterials; }
 
@@ -959,7 +966,7 @@ protected:
 
 	FStaticMeshOccluderData* OccluderData;
 
-	TIndirectArray<FLODInfo> LODs;
+	TArray<FLODInfo> LODs;
 
 	const FDistanceFieldVolumeData* DistanceFieldData;	
 
diff --git a/Engine/Source/Runtime/Engine/Public/StaticMeshVertexData.h b/Engine/Source/Runtime/Engine/Public/StaticMeshVertexData.h
--- a/Engine/Source/Runtime/Engine/Public/StaticMeshVertexData.h
+++ b/Engine/Source/Runtime/Engine/Public/StaticMeshVertexData.h
@@ -11,8 +11,8 @@ template<typename VertexDataType>
 class TStaticMeshVertexData :
 	public FStaticMeshVertexDataInterface
 {
-	TResourceArray<VertexDataType, VERTEXBUFFER_ALIGNMENT> Data;
-
+	using FVertexResourceArray = TResourceArray<VertexDataType, VERTEXBUFFER_ALIGNMENT>;
+	FVertexResourceArray Data;
 public:
 
 	/**
@@ -96,7 +96,7 @@ public:
 	*/
 	void Serialize(FArchive& Ar, bool bForcePerElementSerialization = false) override
 	{
-		Data.TResourceArray<VertexDataType, VERTEXBUFFER_ALIGNMENT>::BulkSerialize(Ar, bForcePerElementSerialization);
+		Data.FVertexResourceArray::BulkSerialize(Ar, bForcePerElementSerialization);
 	}
 	/**
 	* Assignment. This is currently the only method which allows for 
@@ -137,4 +137,9 @@ public:
 	{
 		return Data.GetAllowCPUAccess();
 	}
+
+	void OverrideFreezeSizeAndAlignment(int64& Size, int32& Alignment) const override
+	{
+		Size = sizeof(*this);
+	}
 };
diff --git a/Engine/Source/Runtime/Engine/Public/StaticParameterSet.h b/Engine/Source/Runtime/Engine/Public/StaticParameterSet.h
--- a/Engine/Source/Runtime/Engine/Public/StaticParameterSet.h
+++ b/Engine/Source/Runtime/Engine/Public/StaticParameterSet.h
@@ -330,9 +330,8 @@ struct FStaticMaterialLayersParameter : public FStaticParameterBase
 	
 	const ID GetID() const;
 
-	UMaterialFunctionInterface* GetParameterAssociatedFunction(const FMaterialParameterInfo& InParameterInfo) const;
-	void GetParameterAssociatedFunctions(const FMaterialParameterInfo& InParameterInfo, TArray<UMaterialFunctionInterface*>& AssociatedFunctions) const;
-	
+	UMaterialFunctionInterface* GetParameterAssociatedFunction(const FHashedMaterialParameterInfo& InParameterInfo) const;
+
 	void AppendKeyString(FString& InKeyString) const
 	{
 		InKeyString += ParameterInfo.ToString() + ExpressionGUID.ToString() + Value.GetStaticPermutationString();
diff --git a/Engine/Source/Runtime/Engine/Public/StereoLayerRendering.h b/Engine/Source/Runtime/Engine/Public/StereoLayerRendering.h
--- a/Engine/Source/Runtime/Engine/Public/StereoLayerRendering.h
+++ b/Engine/Source/Runtime/Engine/Public/StereoLayerRendering.h
@@ -34,7 +34,7 @@ public:
 
 	void SetParameters(FRHICommandList& RHICmdList, FVector2D QuadSize, FBox2D UVRect, const FMatrix& ViewProjection, const FMatrix& World)
 	{
-		FRHIVertexShader* VS = GetVertexShader();
+		FRHIVertexShader* VS = RHICmdList.GetBoundVertexShader();
 
 		if (InQuadAdjust.IsBound())
 		{
@@ -62,48 +62,28 @@ public:
 		}
 	}
 
-	virtual bool Serialize(FArchive& Ar) override
-	{
-		bool bShaderHasOutdatedParameters = FGlobalShader::Serialize(Ar);
-		Ar << InQuadAdjust;
-		Ar << InUVAdjust;
-		Ar << InViewProjection;
-		Ar << InWorld;
-		return bShaderHasOutdatedParameters;
-	}
-
 private:
-	FShaderParameter InQuadAdjust;
-	FShaderParameter InUVAdjust;
-	FShaderParameter InViewProjection;
-	FShaderParameter InWorld;
+	LAYOUT_FIELD(FShaderParameter, InQuadAdjust);
+	LAYOUT_FIELD(FShaderParameter, InUVAdjust);
+	LAYOUT_FIELD(FShaderParameter, InViewProjection);
+	LAYOUT_FIELD(FShaderParameter, InWorld);
 };
 
 class FStereoLayerPS_Base : public FGlobalShader
 {
+	DECLARE_TYPE_LAYOUT(FStereoLayerPS_Base, NonVirtual);
 public:
 
 	static bool ShouldCompilePermutation(const FGlobalShaderPermutationParameters& Parameters) { return true; }
 
 	void SetParameters(FRHICommandList& RHICmdList, FRHISamplerState* SamplerStateRHI, FRHITexture* TextureRHI)
 	{
-		FRHIPixelShader* PS = GetPixelShader();
+		FRHIPixelShader* PS = RHICmdList.GetBoundPixelShader();
 
 		SetTextureParameter(RHICmdList, PS, InTexture, InTextureSampler, SamplerStateRHI, TextureRHI);
 	}
 
-	virtual bool Serialize(FArchive& Ar) override
-	{
-		bool bShaderHasOutdatedParameters = FGlobalShader::Serialize(Ar);
-		Ar << InTexture;
-		Ar << InTextureSampler;
-		return bShaderHasOutdatedParameters;
-	}
-
 protected:
-	FShaderResourceParameter InTexture;
-	FShaderResourceParameter InTextureSampler;
-
 	FStereoLayerPS_Base(const ShaderMetaType::CompiledShaderInitializerType& Initializer, const TCHAR* TextureParamName) :
 		FGlobalShader(Initializer) 
 	{
@@ -112,6 +92,8 @@ protected:
 	}
 	FStereoLayerPS_Base() {}
 
+	LAYOUT_FIELD(FShaderResourceParameter, InTexture);
+	LAYOUT_FIELD(FShaderResourceParameter, InTextureSampler);
 };
 
 /**
diff --git a/Engine/Source/Runtime/Engine/Public/TessellationRendering.h b/Engine/Source/Runtime/Engine/Public/TessellationRendering.h
--- a/Engine/Source/Runtime/Engine/Public/TessellationRendering.h
+++ b/Engine/Source/Runtime/Engine/Public/TessellationRendering.h
@@ -14,11 +14,11 @@ class UMaterialInterface;
 
 /** Returns true if the Material and Vertex Factory combination require adjacency information.
   * Game thread version that looks at the material settings. Will not change answer during a shader compile */
-ENGINE_API bool MaterialSettingsRequireAdjacencyInformation_GameThread(UMaterialInterface* Material, const FVertexFactoryType* VertexFactoryType, ERHIFeatureLevel::Type InFeatureLevel);
+ENGINE_API bool MaterialSettingsRequireAdjacencyInformation_GameThread(UMaterialInterface* Material, const FVertexFactoryType* VertexFactoryType, const FStaticFeatureLevel InFeatureLevel);
 
 /** Returns true if the Material and Vertex Factory combination require adjacency information.
   * Rendering thread version that looks at the current shader that will be used. **Will change answer during a shader compile** */
-ENGINE_API bool MaterialRenderingRequiresAdjacencyInformation_RenderingThread(UMaterialInterface* Material, const FVertexFactoryType* VertexFactoryType, ERHIFeatureLevel::Type InFeatureLevel);
+ENGINE_API bool MaterialRenderingRequiresAdjacencyInformation_RenderingThread(UMaterialInterface* Material, const FVertexFactoryType* VertexFactoryType, const FStaticFeatureLevel InFeatureLevel);
 
 /** Returns true if the Material and Vertex Factory combination require adjacency information.
   * Returns different information depending on whether it is called on the rendering thread or game thread -
@@ -27,4 +27,4 @@ ENGINE_API bool MaterialRenderingRequiresAdjacencyInformation_RenderingThread(UM
   *
   * WARNING: In single-threaded mode as the game thread will return the rendering thread information
   * Please use the explicit game/render thread functions above instead */
-ENGINE_API bool RequiresAdjacencyInformation(UMaterialInterface* Material, const FVertexFactoryType* VertexFactoryType, ERHIFeatureLevel::Type InFeatureLevel);
+ENGINE_API bool RequiresAdjacencyInformation(UMaterialInterface* Material, const FVertexFactoryType* VertexFactoryType, const FStaticFeatureLevel InFeatureLevel);
diff --git a/Engine/Source/Runtime/Engine/Public/UnrealEngine.h b/Engine/Source/Runtime/Engine/Public/UnrealEngine.h
--- a/Engine/Source/Runtime/Engine/Public/UnrealEngine.h
+++ b/Engine/Source/Runtime/Engine/Public/UnrealEngine.h
@@ -457,7 +457,7 @@ protected:
 	friend void ScalabilityCVarsSinkCallback();
 };
 
-ENGINE_API bool AllowHighQualityLightmaps(ERHIFeatureLevel::Type FeatureLevel);
+ENGINE_API bool AllowHighQualityLightmaps(const FStaticFeatureLevel FeatureLevel);
 
 ENGINE_API const FCachedSystemScalabilityCVars& GetCachedScalabilityCVars();
 
diff --git a/Engine/Source/Runtime/Engine/Public/UpdateTextureShaders.h b/Engine/Source/Runtime/Engine/Public/UpdateTextureShaders.h
--- a/Engine/Source/Runtime/Engine/Public/UpdateTextureShaders.h
+++ b/Engine/Source/Runtime/Engine/Public/UpdateTextureShaders.h
@@ -22,24 +22,15 @@ public:
 		DestTexture.Bind(Initializer.ParameterMap, TEXT("DestTexture"), SPF_Mandatory);
 	}
 
-	// FShader interface.
-	virtual bool Serialize(FArchive& Ar) override
-	{
-		bool bShaderHasOutdatedParameters = FGlobalShader::Serialize(Ar);
-		Ar << SrcPitchParameter << SrcBuffer << DestPosSizeParameter << DestTexture;
-		return bShaderHasOutdatedParameters;
-	}
-
 	static bool ShouldCompilePermutation(const FGlobalShaderPermutationParameters& Parameters)
 	{
 		return IsFeatureLevelSupported(Parameters.Platform, ERHIFeatureLevel::SM5);
 	}
 
-//protected:
-	FShaderParameter SrcPitchParameter;
-	FShaderResourceParameter SrcBuffer;
-	FShaderParameter DestPosSizeParameter;
-	FShaderResourceParameter DestTexture;
+	LAYOUT_FIELD(FShaderParameter, SrcPitchParameter);
+	LAYOUT_FIELD(FShaderResourceParameter, SrcBuffer);
+	LAYOUT_FIELD(FShaderParameter, DestPosSizeParameter);
+	LAYOUT_FIELD(FShaderResourceParameter, DestTexture);
 };
 
 template<uint32 NumComponents>
@@ -57,14 +48,6 @@ public:
 		DestTexture.Bind(Initializer.ParameterMap, TEXT("TDestTexture"), SPF_Mandatory);
 	}
 
-	// FShader interface.
-	virtual bool Serialize(FArchive& Ar) override
-	{
-		bool bShaderHasOutdatedParameters = FGlobalShader::Serialize(Ar);
-		Ar << SrcPitchParameter << SrcBuffer << DestPosSizeParameter << DestTexture;
-		return bShaderHasOutdatedParameters;
-	}
-
 	static void ModifyCompilationEnvironment(const FGlobalShaderPermutationParameters& Parameters, FShaderCompilerEnvironment& OutEnvironment)
 	{
 		const TCHAR* Type = nullptr;
@@ -85,11 +68,10 @@ public:
 		return IsFeatureLevelSupported(Parameters.Platform, ERHIFeatureLevel::SM5);
 	}
 
-	//protected:
-	FShaderParameter SrcPitchParameter;
-	FShaderResourceParameter SrcBuffer;
-	FShaderParameter DestPosSizeParameter;
-	FShaderResourceParameter DestTexture;
+	LAYOUT_FIELD(FShaderParameter, SrcPitchParameter);
+	LAYOUT_FIELD(FShaderResourceParameter, SrcBuffer);
+	LAYOUT_FIELD(FShaderParameter, DestPosSizeParameter);
+	LAYOUT_FIELD(FShaderResourceParameter, DestTexture);
 };
 
 class FUpdateTexture3DSubresouceCS : public FGlobalShader
@@ -111,28 +93,19 @@ public:
 		DestTexture3D.Bind(Initializer.ParameterMap, TEXT("DestTexture3D"), SPF_Mandatory);
 	}
 
-	// FShader interface.
-	virtual bool Serialize(FArchive& Ar) override
-	{
-		bool bShaderHasOutdatedParameters = FGlobalShader::Serialize(Ar);
-		Ar << SrcPitchParameter << SrcDepthPitchParameter << SrcBuffer << DestPosParameter << DestSizeParameter  << DestTexture3D;
-		return bShaderHasOutdatedParameters;
-	}
-
 	static bool ShouldCompilePermutation(const FGlobalShaderPermutationParameters& Parameters)
 	{
 		return IsFeatureLevelSupported(Parameters.Platform, ERHIFeatureLevel::SM5);
 	}
 
-	//protected:
-	FShaderParameter SrcPitchParameter;
-	FShaderParameter SrcDepthPitchParameter;
-	FShaderResourceParameter SrcBuffer;
+	LAYOUT_FIELD(FShaderParameter, SrcPitchParameter);
+	LAYOUT_FIELD(FShaderParameter, SrcDepthPitchParameter);
+	LAYOUT_FIELD(FShaderResourceParameter, SrcBuffer);
 
-	FShaderParameter DestPosParameter;
-	FShaderParameter DestSizeParameter;
+	LAYOUT_FIELD(FShaderParameter, DestPosParameter);
+	LAYOUT_FIELD(FShaderParameter, DestSizeParameter);
 
-	FShaderResourceParameter DestTexture3D;
+	LAYOUT_FIELD(FShaderResourceParameter, DestTexture3D)
 };
 
 class FCopyTexture2DCS : public FGlobalShader
@@ -148,23 +121,14 @@ public:
 		DestPosSize.Bind(Initializer.ParameterMap, TEXT("DestPosSize"), SPF_Mandatory);
 	}
 
-	// FShader interface.
-	virtual bool Serialize(FArchive& Ar) override
-	{
-		bool bShaderHasOutdatedParameters = FGlobalShader::Serialize(Ar);
-		Ar << SrcTexture << DestTexture << DestPosSize;
-		return bShaderHasOutdatedParameters;
-	}
-
 	static bool ShouldCompilePermutation(const FGlobalShaderPermutationParameters& Parameters)
 	{
 		return IsFeatureLevelSupported(Parameters.Platform, ERHIFeatureLevel::SM5);
 	}
 
-//protected:
-	FShaderResourceParameter SrcTexture;
-	FShaderResourceParameter DestTexture;
-	FShaderParameter DestPosSize;
+	LAYOUT_FIELD(FShaderResourceParameter, SrcTexture);
+	LAYOUT_FIELD(FShaderResourceParameter, DestTexture);
+	LAYOUT_FIELD(FShaderParameter, DestPosSize);
 };
 
 template<uint32 NumComponents, typename ComponentType = uint32>
@@ -182,14 +146,6 @@ public:
 		DestPosSizeParameter.Bind(Initializer.ParameterMap, TEXT("DestPosSize"), SPF_Mandatory);
 	}
 
-	// FShader interface.
-	virtual bool Serialize(FArchive& Ar) override
-	{
-		bool bShaderHasOutdatedParameters = FGlobalShader::Serialize(Ar);
-		Ar << SrcTexture << DestTexture << SrcPosParameter << DestPosSizeParameter;
-		return bShaderHasOutdatedParameters;
-	}
-
 	static const TCHAR* GetTypename(uint32 Dummy)
 	{
 		switch (NumComponents)
@@ -225,10 +181,10 @@ public:
 		return IsFeatureLevelSupported(Parameters.Platform, ERHIFeatureLevel::SM5);
 	}
 
-	FShaderResourceParameter SrcTexture;
-	FShaderResourceParameter DestTexture;
-	FShaderParameter SrcPosParameter;
-	FShaderParameter DestPosSizeParameter;
+	LAYOUT_FIELD(FShaderResourceParameter, SrcTexture);
+	LAYOUT_FIELD(FShaderResourceParameter, DestTexture);
+	LAYOUT_FIELD(FShaderParameter, SrcPosParameter);
+	LAYOUT_FIELD(FShaderParameter, DestPosSizeParameter);
 };
 
 template<uint32 ElementsPerThread>
@@ -244,20 +200,11 @@ public:
 		DestBuffer.Bind(Initializer.ParameterMap, TEXT("DestBuffer"), SPF_Mandatory);		
 	}
 
-	// FShader interface.
-	virtual bool Serialize(FArchive& Ar) override
-	{
-		bool bShaderHasOutdatedParameters = FGlobalShader::Serialize(Ar);
-		Ar << SrcBuffer << DestBuffer;
-		return bShaderHasOutdatedParameters;
-	}
-
 	static bool ShouldCompilePermutation(const FGlobalShaderPermutationParameters& Parameters)
 	{
 		return IsFeatureLevelSupported(Parameters.Platform, ERHIFeatureLevel::SM5);
 	}
 
-	//protected:
-	FShaderResourceParameter SrcBuffer;
-	FShaderResourceParameter DestBuffer;	
+	LAYOUT_FIELD(FShaderResourceParameter, SrcBuffer);
+	LAYOUT_FIELD(FShaderResourceParameter, DestBuffer);
 };
diff --git a/Engine/Source/Runtime/Engine/Public/WeightedRandomSampler.h b/Engine/Source/Runtime/Engine/Public/WeightedRandomSampler.h
--- a/Engine/Source/Runtime/Engine/Public/WeightedRandomSampler.h
+++ b/Engine/Source/Runtime/Engine/Public/WeightedRandomSampler.h
@@ -23,9 +23,11 @@ GetEntryIndex takes two random floating values. This does not generate randoms f
 #pragma once
 
 #include "CoreMinimal.h"
+#include "Serialization/MemoryLayout.h"
 
 struct ENGINE_API FWeightedRandomSampler
 {
+	DECLARE_TYPE_LAYOUT(FWeightedRandomSampler, Abstract);
 public:
 	FWeightedRandomSampler();
 	virtual ~FWeightedRandomSampler() { }
@@ -50,11 +52,11 @@ public:
 
 	int32 GetNumEntries() const { return Prob.Num(); }
 
-	const TArray<float>& GetProb() const { return Prob; }
-	const TArray<int32>& GetAlias() const { return Alias; }
+	const TArray<float, FMemoryImageAllocator>& GetProb() const { return Prob; }
+	const TArray<int32, FMemoryImageAllocator>& GetAlias() const { return Alias; }
 
 protected:
-	TArray<float> Prob;
-	TArray<int32> Alias;
-	float TotalWeight;
+	LAYOUT_FIELD(TMemoryImageArray<float>, Prob);
+	LAYOUT_FIELD(TMemoryImageArray<int32>, Alias);
+	LAYOUT_FIELD(float, TotalWeight);
 };
\ No newline at end of file
diff --git a/Engine/Source/Runtime/Experimental/Chaos/Private/Chaos/CollisionConvexMesh.cpp b/Engine/Source/Runtime/Experimental/Chaos/Private/Chaos/CollisionConvexMesh.cpp
--- a/Engine/Source/Runtime/Experimental/Chaos/Private/Chaos/CollisionConvexMesh.cpp
+++ b/Engine/Source/Runtime/Experimental/Chaos/Private/Chaos/CollisionConvexMesh.cpp
@@ -5,6 +5,9 @@
 
 namespace Chaos
 {
+	constexpr float FConvexBuilder::TriQuadPrismInflation;
+
+
 	// CVars variables for controlling geometry complexity checking and simplification
 	int32 FConvexBuilder::PerformGeometryCheck = 0;
 
diff --git a/Engine/Source/Runtime/Experimental/Chaos/Public/Chaos/ISpatialAcceleration.h b/Engine/Source/Runtime/Experimental/Chaos/Public/Chaos/ISpatialAcceleration.h
--- a/Engine/Source/Runtime/Experimental/Chaos/Public/Chaos/ISpatialAcceleration.h
+++ b/Engine/Source/Runtime/Experimental/Chaos/Public/Chaos/ISpatialAcceleration.h
@@ -28,7 +28,7 @@ struct CHAOS_API FQueryFastData
 
 	const bool bParallel[3];
 
-#if PLATFORM_WINDOWS || PLATFORM_XBOXONE
+#ifdef _MSC_VER
 	#pragma warning (push)
 	#pragma warning(disable:4723)
 #endif
@@ -45,7 +45,7 @@ struct CHAOS_API FQueryFastData
 		}
 	}
 
-#if PLATFORM_WINDOWS || PLATFORM_XBOXONE
+#ifdef _MSC_VER
 	#pragma warning(pop)
 #endif
 
diff --git a/Engine/Source/Runtime/Experimental/GeometryCollectionEngine/Private/GeometryCollection/GeometryCollectionComponent.cpp b/Engine/Source/Runtime/Experimental/GeometryCollectionEngine/Private/GeometryCollection/GeometryCollectionComponent.cpp
--- a/Engine/Source/Runtime/Experimental/GeometryCollectionEngine/Private/GeometryCollection/GeometryCollectionComponent.cpp
+++ b/Engine/Source/Runtime/Experimental/GeometryCollectionEngine/Private/GeometryCollection/GeometryCollectionComponent.cpp
@@ -338,9 +338,9 @@ FBoxSphereBounds UGeometryCollectionComponent::CalcBounds(const FTransform& Loca
 	return FBoxSphereBounds(ForceInitToZero);
 }
 
-void UGeometryCollectionComponent::CreateRenderState_Concurrent()
+void UGeometryCollectionComponent::CreateRenderState_Concurrent(FRegisterComponentContext* Context)
 {
-	Super::CreateRenderState_Concurrent();
+	Super::CreateRenderState_Concurrent(Context);
 
 	if (SceneProxy && RestCollection && RestCollection->HasVisibleGeometry())
 	{
diff --git a/Engine/Source/Runtime/HeadMountedDisplay/Private/DefaultStereoLayers.cpp b/Engine/Source/Runtime/HeadMountedDisplay/Private/DefaultStereoLayers.cpp
--- a/Engine/Source/Runtime/HeadMountedDisplay/Private/DefaultStereoLayers.cpp
+++ b/Engine/Source/Runtime/HeadMountedDisplay/Private/DefaultStereoLayers.cpp
@@ -83,7 +83,7 @@ void FDefaultStereoLayers::StereoLayerRender(FRHICommandListImmediate& RHICmdLis
 	TShaderMapRef<FStereoLayerPS_External> PixelShader_External(ShaderMap);
 
 	GraphicsPSOInit.BoundShaderState.VertexDeclarationRHI = GFilterVertexDeclaration.VertexDeclarationRHI;
-	GraphicsPSOInit.BoundShaderState.VertexShaderRHI = GETSAFERHISHADER_VERTEX(*VertexShader);
+	GraphicsPSOInit.BoundShaderState.VertexShaderRHI = VertexShader.GetVertexShader();
 
 	GraphicsPSOInit.PrimitiveType = PT_TriangleList;
 
@@ -110,7 +110,7 @@ void FDefaultStereoLayers::StereoLayerRender(FRHICommandListImmediate& RHICmdLis
 		if (bIsExternal != bLastWasExternal)
 		{
 			bLastWasExternal = bIsExternal;
-			GraphicsPSOInit.BoundShaderState.PixelShaderRHI = bIsExternal ? GETSAFERHISHADER_PIXEL(*PixelShader_External) : GETSAFERHISHADER_PIXEL(*PixelShader);
+			GraphicsPSOInit.BoundShaderState.PixelShaderRHI = bIsExternal ? PixelShader_External.GetPixelShader() : PixelShader.GetPixelShader();
 			bPipelineStateNeedsUpdate = true;
 		}
 
@@ -161,7 +161,7 @@ void FDefaultStereoLayers::StereoLayerRender(FRHICommandListImmediate& RHICmdLis
 			1.0f, 1.0f,
 			TargetSize,
 			FIntPoint(1, 1),
-			*VertexShader
+			VertexShader
 		);
 	}
 }
diff --git a/Engine/Source/Runtime/HeadMountedDisplay/Private/HeadMountedDisplayBase.cpp b/Engine/Source/Runtime/HeadMountedDisplay/Private/HeadMountedDisplayBase.cpp
--- a/Engine/Source/Runtime/HeadMountedDisplay/Private/HeadMountedDisplayBase.cpp
+++ b/Engine/Source/Runtime/HeadMountedDisplay/Private/HeadMountedDisplayBase.cpp
@@ -21,6 +21,9 @@
 #include "IXRSystemAssets.h"
 
 
+constexpr float FHeadMountedDisplayBase::PixelDensityMin;
+constexpr float FHeadMountedDisplayBase::PixelDensityMax;
+
 FHeadMountedDisplayBase::FHeadMountedDisplayBase(IARSystemSupport* InARImplementation)
 	: FXRTrackingSystemBase(InARImplementation)
 	, bHeadTrackingEnforced(false)
diff --git a/Engine/Source/Runtime/HeadMountedDisplay/Private/MotionControllerComponent.cpp b/Engine/Source/Runtime/HeadMountedDisplay/Private/MotionControllerComponent.cpp
--- a/Engine/Source/Runtime/HeadMountedDisplay/Private/MotionControllerComponent.cpp
+++ b/Engine/Source/Runtime/HeadMountedDisplay/Private/MotionControllerComponent.cpp
@@ -92,9 +92,9 @@ void UMotionControllerComponent::BeginDestroy()
 	}
 }
 
-void UMotionControllerComponent::CreateRenderState_Concurrent()
+void UMotionControllerComponent::CreateRenderState_Concurrent(FRegisterComponentContext* Context)
 {
-	Super::CreateRenderState_Concurrent();
+	Super::CreateRenderState_Concurrent(Context);
 	RenderThreadRelativeTransform = GetRelativeTransform();
 	RenderThreadComponentScale = GetComponentScale();
 }
diff --git a/Engine/Source/Runtime/Landscape/Classes/LandscapeMaterialInstanceConstant.h b/Engine/Source/Runtime/Landscape/Classes/LandscapeMaterialInstanceConstant.h
--- a/Engine/Source/Runtime/Landscape/Classes/LandscapeMaterialInstanceConstant.h
+++ b/Engine/Source/Runtime/Landscape/Classes/LandscapeMaterialInstanceConstant.h
@@ -7,11 +7,26 @@
 #include "Materials/MaterialInstanceConstant.h"
 #include "LandscapeMaterialInstanceConstant.generated.h"
 
+USTRUCT()
+struct FLandscapeMaterialTextureStreamingInfo
+{
+	GENERATED_USTRUCT_BODY()
+
+	UPROPERTY()
+	FName TextureName;
+
+	UPROPERTY()
+	float TexelFactor;
+};
+
 UCLASS(MinimalAPI)
 class ULandscapeMaterialInstanceConstant : public UMaterialInstanceConstant
 {
 	GENERATED_UCLASS_BODY()
 
+	UPROPERTY()
+	TArray<FLandscapeMaterialTextureStreamingInfo> TextureStreamingInfo;
+
 	UPROPERTY()
 	uint32 bIsLayerThumbnail:1;
 
@@ -24,7 +39,19 @@ class ULandscapeMaterialInstanceConstant : public UMaterialInstanceConstant
 	UPROPERTY()
 	uint32 bEditorToolUsage:1;
 
+	virtual void PostLoad() override;
+#if WITH_EDITOR
+	virtual void PostEditChangeProperty(FPropertyChangedEvent& PropertyChangedEvent) override;
+#endif
 	virtual FMaterialResource* AllocatePermutationResource() override;
 	virtual bool HasOverridenBaseProperties() const override;
+
+#if WITH_EDITOR
+	FLandscapeMaterialTextureStreamingInfo& AcquireTextureStreamingInfo(const FName& TextureName);
+	void UpdateCachedTextureStreaming();
+#endif
+
+public:
+	float GetLandscapeTexelFactor(const FName& TextureName) const;
 };
 
diff --git a/Engine/Source/Runtime/Landscape/Private/Landscape.cpp b/Engine/Source/Runtime/Landscape/Private/Landscape.cpp
--- a/Engine/Source/Runtime/Landscape/Private/Landscape.cpp
+++ b/Engine/Source/Runtime/Landscape/Private/Landscape.cpp
@@ -2082,7 +2082,7 @@ void ALandscapeProxy::PreSave(const class ITargetPlatform* TargetPlatform)
 			Landscape->ClearDirtyData(LandscapeComponent);
 		}
 	}
-#endif
+#endif // WITH_EDITOR
 }
 
 void ALandscapeProxy::Serialize(FArchive& Ar)
@@ -3369,6 +3369,7 @@ void FLandscapeComponentDerivedData::SaveToDDC(const FGuid& StateId, UObject* Co
 	GetDerivedDataCacheRef().Put(*GetDDCKeyString(StateId), CompressedLandscapeData, Component->GetPathName());
 }
 
+#if WITH_EDITOR
 void LandscapeMaterialsParameterValuesGetter(FStaticParameterSet& OutStaticParameterSet, UMaterialInstance* Material)
 {
 	if (Material->Parent)
@@ -3420,6 +3421,7 @@ bool LandscapeMaterialsParameterSetUpdater(FStaticParameterSet& StaticParameterS
 {
 	return UpdateParameterSet<FStaticTerrainLayerWeightParameter, UMaterialExpressionLandscapeLayerWeight>(StaticParameterSet.TerrainLayerWeightParameters, ParentMaterial);
 }
+#endif // WITH_EDITOR
 
 ALandscapeProxy::~ALandscapeProxy()
 {
@@ -3510,8 +3512,7 @@ void ULandscapeComponent::SerializeStateHashes(FArchive& Ar)
 
 	if (OverrideMaterial != nullptr)
 	{
-		UMaterialInterface::TMicRecursionGuard RecursionGuard;
-		FGuid LocalStateId = OverrideMaterial->GetMaterial_Concurrent(RecursionGuard)->StateId;
+		FGuid LocalStateId = OverrideMaterial->GetMaterial_Concurrent()->StateId;
 		Ar << LocalStateId;
 	}
 
@@ -3519,8 +3520,7 @@ void ULandscapeComponent::SerializeStateHashes(FArchive& Ar)
 	{
 		if (MaterialOverride.Material != nullptr)
 		{
-			UMaterialInterface::TMicRecursionGuard RecursionGuard;
-			FGuid LocalStateId = MaterialOverride.Material->GetMaterial_Concurrent(RecursionGuard)->StateId;
+			FGuid LocalStateId = MaterialOverride.Material->GetMaterial_Concurrent()->StateId;
 			Ar << LocalStateId;
 			Ar << MaterialOverride.LODIndex;
 		}
@@ -3530,8 +3530,7 @@ void ULandscapeComponent::SerializeStateHashes(FArchive& Ar)
 
 	if (Proxy->LandscapeMaterial != nullptr)
 	{
-		UMaterialInterface::TMicRecursionGuard RecursionGuard;
-		FGuid LocalStateId = Proxy->LandscapeMaterial->GetMaterial_Concurrent(RecursionGuard)->StateId;
+		FGuid LocalStateId = Proxy->LandscapeMaterial->GetMaterial_Concurrent()->StateId;
 		Ar << LocalStateId;
 	}
 
@@ -3539,8 +3538,7 @@ void ULandscapeComponent::SerializeStateHashes(FArchive& Ar)
 	{
 		if (MaterialOverride.Material != nullptr)
 		{
-			UMaterialInterface::TMicRecursionGuard RecursionGuard;
-			FGuid LocalStateId = MaterialOverride.Material->GetMaterial_Concurrent(RecursionGuard)->StateId;
+			FGuid LocalStateId = MaterialOverride.Material->GetMaterial_Concurrent()->StateId;
 			Ar << LocalStateId;
 			Ar << MaterialOverride.LODIndex;
 		}
diff --git a/Engine/Source/Runtime/Landscape/Private/LandscapeEditLayers.cpp b/Engine/Source/Runtime/Landscape/Private/LandscapeEditLayers.cpp
--- a/Engine/Source/Runtime/Landscape/Private/LandscapeEditLayers.cpp
+++ b/Engine/Source/Runtime/Landscape/Private/LandscapeEditLayers.cpp
@@ -363,18 +363,11 @@ public:
 
 	void SetParameters(FRHICommandList& RHICmdList, const FMatrix& InProjectionMatrix)
 	{
-		SetShaderValue(RHICmdList, GetVertexShader(), TransformParam, InProjectionMatrix);
-	}
-
-	virtual bool Serialize(FArchive& Ar) override
-	{
-		bool bShaderHasOutdatedParameters = FShader::Serialize(Ar);
-		Ar << TransformParam;
-		return bShaderHasOutdatedParameters;
+		SetShaderValue(RHICmdList, RHICmdList.GetBoundVertexShader(), TransformParam, InProjectionMatrix);
 	}
 
 private:
-	FShaderParameter TransformParam;
+	LAYOUT_FIELD(FShaderParameter, TransformParam);
 };
 
 IMPLEMENT_GLOBAL_SHADER(FLandscapeLayersVS, "/Engine/Private/LandscapeLayersVS.usf", "VSMain", SF_Vertex);
@@ -444,45 +437,30 @@ public:
 
 	void SetParameters(FRHICommandList& RHICmdList, const FLandscapeLayersHeightmapShaderParameters& InParams)
 	{
-		SetTextureParameter(RHICmdList, GetPixelShader(), ReadTexture1Param, ReadTexture1SamplerParam, TStaticSamplerState<SF_Point, AM_Clamp, AM_Clamp, AM_Clamp>::GetRHI(), InParams.ReadHeightmap1->Resource->TextureRHI);
-		SetTextureParameter(RHICmdList, GetPixelShader(), ReadTexture2Param, ReadTexture2SamplerParam, TStaticSamplerState<SF_Point, AM_Clamp, AM_Clamp, AM_Clamp>::GetRHI(), InParams.ReadHeightmap2 != nullptr ? InParams.ReadHeightmap2->Resource->TextureRHI : GWhiteTexture->TextureRHI);
+		SetTextureParameter(RHICmdList, RHICmdList.GetBoundPixelShader(), ReadTexture1Param, ReadTexture1SamplerParam, TStaticSamplerState<SF_Point, AM_Clamp, AM_Clamp, AM_Clamp>::GetRHI(), InParams.ReadHeightmap1->Resource->TextureRHI);
+		SetTextureParameter(RHICmdList, RHICmdList.GetBoundPixelShader(), ReadTexture2Param, ReadTexture2SamplerParam, TStaticSamplerState<SF_Point, AM_Clamp, AM_Clamp, AM_Clamp>::GetRHI(), InParams.ReadHeightmap2 != nullptr ? InParams.ReadHeightmap2->Resource->TextureRHI : GWhiteTexture->TextureRHI);
 
 		FVector4 LayerInfo(InParams.LayerAlpha, InParams.LayerVisible ? 1.0f : 0.0f, InParams.LayerBlendMode == LSBM_AlphaBlend ? 1.0f : 0.f, 0.f);
 		FVector4 OutputConfig(InParams.ApplyLayerModifiers ? 1.0f : 0.0f, 0.0f /*unused*/, InParams.ReadHeightmap2 ? 1.0f : 0.0f, InParams.GenerateNormals ? 1.0f : 0.0f);
 		FVector2D TextureSize(InParams.HeightmapSize.X, InParams.HeightmapSize.Y);
 
-		SetShaderValue(RHICmdList, GetPixelShader(), LayerInfoParam, LayerInfo);
-		SetShaderValue(RHICmdList, GetPixelShader(), OutputConfigParam, OutputConfig);
-		SetShaderValue(RHICmdList, GetPixelShader(), TextureSizeParam, TextureSize);
-		SetShaderValue(RHICmdList, GetPixelShader(), LandscapeGridScaleParam, InParams.GridSize);
-		SetShaderValue(RHICmdList, GetPixelShader(), ComponentVertexCountParam, (float)InParams.CurrentMipComponentVertexCount);
-	}
-
-	virtual bool Serialize(FArchive& Ar) override
-	{
-		bool bShaderHasOutdatedParameters = FShader::Serialize(Ar);
-		Ar << ReadTexture1Param;
-		Ar << ReadTexture2Param;
-		Ar << ReadTexture1SamplerParam;
-		Ar << ReadTexture2SamplerParam;
-		Ar << LayerInfoParam;
-		Ar << OutputConfigParam;
-		Ar << TextureSizeParam;
-		Ar << LandscapeGridScaleParam;
-		Ar << ComponentVertexCountParam;
-		return bShaderHasOutdatedParameters;
+		SetShaderValue(RHICmdList, RHICmdList.GetBoundPixelShader(), LayerInfoParam, LayerInfo);
+		SetShaderValue(RHICmdList, RHICmdList.GetBoundPixelShader(), OutputConfigParam, OutputConfig);
+		SetShaderValue(RHICmdList, RHICmdList.GetBoundPixelShader(), TextureSizeParam, TextureSize);
+		SetShaderValue(RHICmdList, RHICmdList.GetBoundPixelShader(), LandscapeGridScaleParam, InParams.GridSize);
+		SetShaderValue(RHICmdList, RHICmdList.GetBoundPixelShader(), ComponentVertexCountParam, (float)InParams.CurrentMipComponentVertexCount);
 	}
 
 private:
-	FShaderResourceParameter ReadTexture1Param;
-	FShaderResourceParameter ReadTexture2Param;
-	FShaderResourceParameter ReadTexture1SamplerParam;
-	FShaderResourceParameter ReadTexture2SamplerParam;
-	FShaderParameter LayerInfoParam;
-	FShaderParameter OutputConfigParam;
-	FShaderParameter TextureSizeParam;
-	FShaderParameter LandscapeGridScaleParam;
-	FShaderParameter ComponentVertexCountParam;
+	LAYOUT_FIELD(FShaderResourceParameter, ReadTexture1Param);
+	LAYOUT_FIELD(FShaderResourceParameter, ReadTexture2Param);
+	LAYOUT_FIELD(FShaderResourceParameter, ReadTexture1SamplerParam);
+	LAYOUT_FIELD(FShaderResourceParameter, ReadTexture2SamplerParam);
+	LAYOUT_FIELD(FShaderParameter, LayerInfoParam);
+	LAYOUT_FIELD(FShaderParameter, OutputConfigParam);
+	LAYOUT_FIELD(FShaderParameter, TextureSizeParam);
+	LAYOUT_FIELD(FShaderParameter, LandscapeGridScaleParam);
+	LAYOUT_FIELD(FShaderParameter, ComponentVertexCountParam);
 };
 
 IMPLEMENT_GLOBAL_SHADER(FLandscapeLayersHeightmapPS, "/Engine/Private/LandscapeLayersPS.usf", "PSHeightmapMain", SF_Pixel);
@@ -516,31 +494,19 @@ public:
 
 	void SetParameters(FRHICommandList& RHICmdList, const FLandscapeLayersHeightmapShaderParameters& InParams)
 	{
-		SetTextureParameter(RHICmdList, GetPixelShader(), ReadTexture1Param, ReadTexture1SamplerParam, TStaticSamplerState<SF_Point, AM_Clamp, AM_Clamp, AM_Clamp>::GetRHI(), InParams.ReadHeightmap1->Resource->TextureRHI);
-
-		SetShaderValue(RHICmdList, GetPixelShader(), CurrentMipSizeParam, FVector2D(InParams.CurrentMipSize.X, InParams.CurrentMipSize.Y));
-		SetShaderValue(RHICmdList, GetPixelShader(), ParentMipSizeParam, FVector2D(InParams.ParentMipSize.X, InParams.ParentMipSize.Y));
-		SetShaderValue(RHICmdList, GetPixelShader(), CurrentMipComponentVertexCountParam, (float)InParams.CurrentMipComponentVertexCount);
-	}
-
-	virtual bool Serialize(FArchive& Ar) override
-	{
-		bool bShaderHasOutdatedParameters = FShader::Serialize(Ar);
-		Ar << ReadTexture1Param;
-		Ar << ReadTexture1SamplerParam;
-		Ar << CurrentMipSizeParam;
-		Ar << ParentMipSizeParam;
-		Ar << CurrentMipComponentVertexCountParam;
+		SetTextureParameter(RHICmdList, RHICmdList.GetBoundPixelShader(), ReadTexture1Param, ReadTexture1SamplerParam, TStaticSamplerState<SF_Point, AM_Clamp, AM_Clamp, AM_Clamp>::GetRHI(), InParams.ReadHeightmap1->Resource->TextureRHI);
 
-		return bShaderHasOutdatedParameters;
+		SetShaderValue(RHICmdList, RHICmdList.GetBoundPixelShader(), CurrentMipSizeParam, FVector2D(InParams.CurrentMipSize.X, InParams.CurrentMipSize.Y));
+		SetShaderValue(RHICmdList, RHICmdList.GetBoundPixelShader(), ParentMipSizeParam, FVector2D(InParams.ParentMipSize.X, InParams.ParentMipSize.Y));
+		SetShaderValue(RHICmdList, RHICmdList.GetBoundPixelShader(), CurrentMipComponentVertexCountParam, (float)InParams.CurrentMipComponentVertexCount);
 	}
 
 private:
-	FShaderResourceParameter ReadTexture1Param;
-	FShaderResourceParameter ReadTexture1SamplerParam;
-	FShaderParameter CurrentMipSizeParam;
-	FShaderParameter ParentMipSizeParam;
-	FShaderParameter CurrentMipComponentVertexCountParam;
+	LAYOUT_FIELD(FShaderResourceParameter, ReadTexture1Param);
+	LAYOUT_FIELD(FShaderResourceParameter, ReadTexture1SamplerParam);
+	LAYOUT_FIELD(FShaderParameter, CurrentMipSizeParam);
+	LAYOUT_FIELD(FShaderParameter, ParentMipSizeParam);
+	LAYOUT_FIELD(FShaderParameter, CurrentMipComponentVertexCountParam);
 };
 
 IMPLEMENT_GLOBAL_SHADER(FLandscapeLayersHeightmapMipsPS, "/Engine/Private/LandscapeLayersPS.usf", "PSHeightmapMainMips", SF_Pixel);
@@ -605,38 +571,25 @@ public:
 
 	void SetParameters(FRHICommandList& RHICmdList, const FLandscapeLayersWeightmapShaderParameters& InParams)
 	{
-		SetTextureParameter(RHICmdList, GetPixelShader(), ReadTexture1Param, ReadTexture1SamplerParam, TStaticSamplerState<SF_Point, AM_Clamp, AM_Clamp, AM_Clamp>::GetRHI(), InParams.ReadWeightmap1->Resource->TextureRHI);
-		SetTextureParameter(RHICmdList, GetPixelShader(), ReadTexture2Param, ReadTexture2SamplerParam, TStaticSamplerState<SF_Point, AM_Clamp, AM_Clamp, AM_Clamp>::GetRHI(), InParams.ReadWeightmap2 != nullptr ? InParams.ReadWeightmap2->Resource->TextureRHI : GWhiteTexture->TextureRHI);
+		SetTextureParameter(RHICmdList, RHICmdList.GetBoundPixelShader(), ReadTexture1Param, ReadTexture1SamplerParam, TStaticSamplerState<SF_Point, AM_Clamp, AM_Clamp, AM_Clamp>::GetRHI(), InParams.ReadWeightmap1->Resource->TextureRHI);
+		SetTextureParameter(RHICmdList, RHICmdList.GetBoundPixelShader(), ReadTexture2Param, ReadTexture2SamplerParam, TStaticSamplerState<SF_Point, AM_Clamp, AM_Clamp, AM_Clamp>::GetRHI(), InParams.ReadWeightmap2 != nullptr ? InParams.ReadWeightmap2->Resource->TextureRHI : GWhiteTexture->TextureRHI);
 
 		FVector4 LayerInfo(InParams.LayerAlpha, InParams.LayerVisible ? 1.0f : 0.0f, InParams.LayerBlendMode == LSBM_AlphaBlend ? 1.0f : 0.f, 0.f);
 		FVector4 OutputConfig(InParams.ApplyLayerModifiers ? 1.0f : 0.0f, InParams.OutputAsSubstractive ? 1.0f : 0.0f, InParams.ReadWeightmap2 != nullptr ? 1.0f : 0.0f, InParams.OutputAsNormalized ? 1.0f : 0.0f);
 
-		SetShaderValue(RHICmdList, GetPixelShader(), LayerInfoParam, LayerInfo);
-		SetShaderValue(RHICmdList, GetPixelShader(), OutputConfigParam, OutputConfig);
-		SetShaderValue(RHICmdList, GetPixelShader(), ComponentVertexCountParam, (float)InParams.CurrentMipComponentVertexCount);
-	}
-
-	virtual bool Serialize(FArchive& Ar) override
-	{
-		bool bShaderHasOutdatedParameters = FShader::Serialize(Ar);
-		Ar << ReadTexture1Param;
-		Ar << ReadTexture2Param;
-		Ar << ReadTexture1SamplerParam;
-		Ar << ReadTexture2SamplerParam;
-		Ar << LayerInfoParam;
-		Ar << OutputConfigParam;
-		Ar << ComponentVertexCountParam;
-		return bShaderHasOutdatedParameters;
+		SetShaderValue(RHICmdList, RHICmdList.GetBoundPixelShader(), LayerInfoParam, LayerInfo);
+		SetShaderValue(RHICmdList, RHICmdList.GetBoundPixelShader(), OutputConfigParam, OutputConfig);
+		SetShaderValue(RHICmdList, RHICmdList.GetBoundPixelShader(), ComponentVertexCountParam, (float)InParams.CurrentMipComponentVertexCount);
 	}
 
 private:
-	FShaderResourceParameter ReadTexture1Param;
-	FShaderResourceParameter ReadTexture2Param;
-	FShaderResourceParameter ReadTexture1SamplerParam;
-	FShaderResourceParameter ReadTexture2SamplerParam;
-	FShaderParameter LayerInfoParam;
-	FShaderParameter OutputConfigParam;
-	FShaderParameter ComponentVertexCountParam;
+	LAYOUT_FIELD(FShaderResourceParameter, ReadTexture1Param);
+	LAYOUT_FIELD(FShaderResourceParameter, ReadTexture2Param);
+	LAYOUT_FIELD(FShaderResourceParameter, ReadTexture1SamplerParam);
+	LAYOUT_FIELD(FShaderResourceParameter, ReadTexture2SamplerParam);
+	LAYOUT_FIELD(FShaderParameter, LayerInfoParam);
+	LAYOUT_FIELD(FShaderParameter, OutputConfigParam);
+	LAYOUT_FIELD(FShaderParameter, ComponentVertexCountParam);
 };
 
 IMPLEMENT_GLOBAL_SHADER(FLandscapeLayersWeightmapPS, "/Engine/Private/LandscapeLayersPS.usf", "PSWeightmapMain", SF_Pixel);
@@ -670,31 +623,19 @@ public:
 
 	void SetParameters(FRHICommandList& RHICmdList, const FLandscapeLayersWeightmapShaderParameters& InParams)
 	{
-		SetTextureParameter(RHICmdList, GetPixelShader(), ReadTexture1Param, ReadTexture1SamplerParam, TStaticSamplerState<SF_Point, AM_Clamp, AM_Clamp, AM_Clamp>::GetRHI(), InParams.ReadWeightmap1->Resource->TextureRHI);
+		SetTextureParameter(RHICmdList, RHICmdList.GetBoundPixelShader(), ReadTexture1Param, ReadTexture1SamplerParam, TStaticSamplerState<SF_Point, AM_Clamp, AM_Clamp, AM_Clamp>::GetRHI(), InParams.ReadWeightmap1->Resource->TextureRHI);
 
-		SetShaderValue(RHICmdList, GetPixelShader(), CurrentMipSizeParam, FVector2D(InParams.CurrentMipSize.X, InParams.CurrentMipSize.Y));
-		SetShaderValue(RHICmdList, GetPixelShader(), ParentMipSizeParam, FVector2D(InParams.ParentMipSize.X, InParams.ParentMipSize.Y));
-		SetShaderValue(RHICmdList, GetPixelShader(), CurrentMipComponentVertexCountParam, (float)InParams.CurrentMipComponentVertexCount);
-	}
-
-	virtual bool Serialize(FArchive& Ar) override
-	{
-		bool bShaderHasOutdatedParameters = FShader::Serialize(Ar);
-		Ar << ReadTexture1Param;
-		Ar << ReadTexture1SamplerParam;
-		Ar << CurrentMipSizeParam;
-		Ar << ParentMipSizeParam;
-		Ar << CurrentMipComponentVertexCountParam;
-
-		return bShaderHasOutdatedParameters;
+		SetShaderValue(RHICmdList, RHICmdList.GetBoundPixelShader(), CurrentMipSizeParam, FVector2D(InParams.CurrentMipSize.X, InParams.CurrentMipSize.Y));
+		SetShaderValue(RHICmdList, RHICmdList.GetBoundPixelShader(), ParentMipSizeParam, FVector2D(InParams.ParentMipSize.X, InParams.ParentMipSize.Y));
+		SetShaderValue(RHICmdList, RHICmdList.GetBoundPixelShader(), CurrentMipComponentVertexCountParam, (float)InParams.CurrentMipComponentVertexCount);
 	}
 
 private:
-	FShaderResourceParameter ReadTexture1Param;
-	FShaderResourceParameter ReadTexture1SamplerParam;
-	FShaderParameter CurrentMipSizeParam;
-	FShaderParameter ParentMipSizeParam;
-	FShaderParameter CurrentMipComponentVertexCountParam;
+	LAYOUT_FIELD(FShaderResourceParameter, ReadTexture1Param);
+	LAYOUT_FIELD(FShaderResourceParameter, ReadTexture1SamplerParam);
+	LAYOUT_FIELD(FShaderParameter, CurrentMipSizeParam);
+	LAYOUT_FIELD(FShaderParameter, ParentMipSizeParam);
+	LAYOUT_FIELD(FShaderParameter, CurrentMipComponentVertexCountParam);
 };
 
 IMPLEMENT_GLOBAL_SHADER(FLandscapeLayersWeightmapMipsPS, "/Engine/Private/LandscapeLayersPS.usf", "PSWeightmapMainMips", SF_Pixel);
@@ -743,20 +684,12 @@ public:
 
 	void SetParameters(FRHICommandList& RHICmdList, FRHITexture* TextureRHI)
 	{
-		SetTextureParameter(RHICmdList, GetPixelShader(), ReadTexture1Param, ReadTexture1SamplerParam, TStaticSamplerState<SF_Point, AM_Clamp, AM_Clamp, AM_Clamp>::GetRHI(), TextureRHI);
-	}
-
-	virtual bool Serialize(FArchive& Ar) override
-	{
-		bool bShaderHasOutdatedParameters = FShader::Serialize(Ar);
-		Ar << ReadTexture1Param;
-		Ar << ReadTexture1SamplerParam;
-		return bShaderHasOutdatedParameters;
+		SetTextureParameter(RHICmdList, RHICmdList.GetBoundPixelShader(), ReadTexture1Param, ReadTexture1SamplerParam, TStaticSamplerState<SF_Point, AM_Clamp, AM_Clamp, AM_Clamp>::GetRHI(), TextureRHI);
 	}
 
 private:
-	FShaderResourceParameter ReadTexture1Param;
-	FShaderResourceParameter ReadTexture1SamplerParam;
+	LAYOUT_FIELD(FShaderResourceParameter, ReadTexture1Param);
+	LAYOUT_FIELD(FShaderResourceParameter, ReadTexture1SamplerParam);
 };
 
 IMPLEMENT_GLOBAL_SHADER(FLandscapeCopyTextureVS, "/Engine/Private/LandscapeLayersPS.usf", "CopyTextureVS", SF_Vertex);
@@ -866,33 +799,22 @@ public:
 
 	void SetParameters(FRHICommandList& RHICmdList, const FLandscapeLayerWeightmapExtractMaterialLayersComputeShaderParameters& InParams)
 	{
-		SetTextureParameter(RHICmdList, GetComputeShader(), ComponentWeightmapParam, InParams.ComponentWeightmapResource->TextureRHI);
-		SetUAVParameter(RHICmdList, GetComputeShader(), AtlasPaintListsParam, InParams.AtlasWeightmapsPerLayer->TextureUAV);
-		SetSRVParameter(RHICmdList, GetComputeShader(), ComponentsDataParam, InParams.ComputeShaderResource->ComponentsDataSRV);
-		SetShaderValue(RHICmdList, GetComputeShader(), ComponentSizeParam, InParams.ComponentSize);
+		SetTextureParameter(RHICmdList, RHICmdList.GetBoundComputeShader(), ComponentWeightmapParam, InParams.ComponentWeightmapResource->TextureRHI);
+		SetUAVParameter(RHICmdList, RHICmdList.GetBoundComputeShader(), AtlasPaintListsParam, InParams.AtlasWeightmapsPerLayer->TextureUAV);
+		SetSRVParameter(RHICmdList, RHICmdList.GetBoundComputeShader(), ComponentsDataParam, InParams.ComputeShaderResource->ComponentsDataSRV);
+		SetShaderValue(RHICmdList, RHICmdList.GetBoundComputeShader(), ComponentSizeParam, InParams.ComponentSize);
 	}
 
 	void UnsetParameters(FRHICommandList& RHICmdList)
 	{
-		SetUAVParameter(RHICmdList, GetComputeShader(), AtlasPaintListsParam, nullptr);
-	}
-
-	virtual bool Serialize(FArchive& Ar) override
-	{
-		bool bShaderHasOutdatedParameters = FShader::Serialize(Ar);
-		Ar << ComponentWeightmapParam;
-		Ar << AtlasPaintListsParam;
-		Ar << ComponentsDataParam;
-		Ar << ComponentSizeParam;
-
-		return bShaderHasOutdatedParameters;
+		SetUAVParameter(RHICmdList, RHICmdList.GetBoundComputeShader(), AtlasPaintListsParam, nullptr);
 	}
 
 private:
-	FShaderResourceParameter ComponentWeightmapParam;
-	FShaderResourceParameter AtlasPaintListsParam;
-	FShaderResourceParameter ComponentsDataParam;
-	FShaderParameter ComponentSizeParam;
+	LAYOUT_FIELD(FShaderResourceParameter, ComponentWeightmapParam);
+	LAYOUT_FIELD(FShaderResourceParameter, AtlasPaintListsParam);
+	LAYOUT_FIELD(FShaderResourceParameter, ComponentsDataParam);
+	LAYOUT_FIELD(FShaderParameter, ComponentSizeParam);
 };
 
 IMPLEMENT_GLOBAL_SHADER(FLandscapeLayerWeightmapExtractMaterialLayersCS, "/Engine/Private/LandscapeLayersCS.usf", "ComputeWeightmapPerPaintLayer", SF_Compute);
@@ -911,14 +833,14 @@ public:
 		SCOPED_GPU_STAT(InRHICmdList, LandscapeLayersRender);
 
 		TShaderMapRef<FLandscapeLayerWeightmapExtractMaterialLayersCS> ComputeShader(GetGlobalShaderMap(GMaxRHIFeatureLevel));
-		InRHICmdList.SetComputeShader(ComputeShader->GetComputeShader());
+		InRHICmdList.SetComputeShader(ComputeShader.GetComputeShader());
 		ComputeShader->SetParameters(InRHICmdList, ShaderParams);
 
 		uint32 ThreadGroupCountX = FMath::CeilToInt((float)ShaderParams.ComponentSize / (float)GLandscapeLayerWeightmapThreadGroupSizeX);
 		uint32 ThreadGroupCountY = FMath::CeilToInt((float)ShaderParams.ComponentSize / (float)GLandscapeLayerWeightmapThreadGroupSizeY);
 		check(ThreadGroupCountX > 0 && ThreadGroupCountY > 0);
 
-		DispatchComputeShader(InRHICmdList, *ComputeShader, ThreadGroupCountX, ThreadGroupCountY, ShaderParams.ComputeShaderResource->GetComponentsDataCount());
+		DispatchComputeShader(InRHICmdList, ComputeShader.GetShader(), ThreadGroupCountX, ThreadGroupCountY, ShaderParams.ComputeShaderResource->GetComponentsDataCount());
 		ComputeShader->UnsetParameters(InRHICmdList);
 		ShaderParams.ComputeShaderResource->ReleaseResource();
 		delete ShaderParams.ComputeShaderResource;
@@ -1062,38 +984,26 @@ public:
 
 	void SetParameters(FRHICommandList& RHICmdList, const FLandscapeLayerWeightmapPackMaterialLayersComputeShaderParameters& InParams)
 	{
-		SetUAVParameter(RHICmdList, GetComputeShader(), ComponentWeightmapParam, InParams.ComponentWeightmapResource->TextureUAV);
-		SetTextureParameter(RHICmdList, GetComputeShader(), AtlasPaintListsParam, InParams.AtlasWeightmapsPerLayer->TextureRHI);
-		SetSRVParameter(RHICmdList, GetComputeShader(), ComponentsDataParam, InParams.ComputeShaderResource->ComponentsDataSRV);
-		SetShaderValue(RHICmdList, GetComputeShader(), ComponentSizeParam, InParams.ComponentSize);
-		SetSRVParameter(RHICmdList, GetComputeShader(), WeightmapWeightBlendModeParam, InParams.ComputeShaderResource->WeightmapWeightBlendModeSRV);
-		SetSRVParameter(RHICmdList, GetComputeShader(), WeightmapTextureOutputOffsetParam, InParams.ComputeShaderResource->WeightmapTextureOutputOffsetSRV);
+		SetUAVParameter(RHICmdList, RHICmdList.GetBoundComputeShader(), ComponentWeightmapParam, InParams.ComponentWeightmapResource->TextureUAV);
+		SetTextureParameter(RHICmdList, RHICmdList.GetBoundComputeShader(), AtlasPaintListsParam, InParams.AtlasWeightmapsPerLayer->TextureRHI);
+		SetSRVParameter(RHICmdList, RHICmdList.GetBoundComputeShader(), ComponentsDataParam, InParams.ComputeShaderResource->ComponentsDataSRV);
+		SetShaderValue(RHICmdList, RHICmdList.GetBoundComputeShader(), ComponentSizeParam, InParams.ComponentSize);
+		SetSRVParameter(RHICmdList, RHICmdList.GetBoundComputeShader(), WeightmapWeightBlendModeParam, InParams.ComputeShaderResource->WeightmapWeightBlendModeSRV);
+		SetSRVParameter(RHICmdList, RHICmdList.GetBoundComputeShader(), WeightmapTextureOutputOffsetParam, InParams.ComputeShaderResource->WeightmapTextureOutputOffsetSRV);
 	}
 
 	void UnsetParameters(FRHICommandList& RHICmdList)
 	{
-		SetUAVParameter(RHICmdList, GetComputeShader(), ComponentWeightmapParam, nullptr);
-	}
-
-	virtual bool Serialize(FArchive& Ar) override
-	{
-		bool bShaderHasOutdatedParameters = FShader::Serialize(Ar);
-		Ar << ComponentWeightmapParam;
-		Ar << AtlasPaintListsParam;
-		Ar << ComponentsDataParam;
-		Ar << ComponentSizeParam;
-		Ar << WeightmapWeightBlendModeParam;
-		Ar << WeightmapTextureOutputOffsetParam;
-		return bShaderHasOutdatedParameters;
+		SetUAVParameter(RHICmdList, RHICmdList.GetBoundComputeShader(), ComponentWeightmapParam, nullptr);
 	}
 
 private:
-	FShaderResourceParameter ComponentWeightmapParam;
-	FShaderResourceParameter AtlasPaintListsParam;
-	FShaderResourceParameter ComponentsDataParam;
-	FShaderParameter ComponentSizeParam;
-	FShaderResourceParameter WeightmapWeightBlendModeParam;
-	FShaderResourceParameter WeightmapTextureOutputOffsetParam;
+	LAYOUT_FIELD(FShaderResourceParameter, ComponentWeightmapParam);
+	LAYOUT_FIELD(FShaderResourceParameter, AtlasPaintListsParam);
+	LAYOUT_FIELD(FShaderResourceParameter, ComponentsDataParam);
+	LAYOUT_FIELD(FShaderParameter, ComponentSizeParam);
+	LAYOUT_FIELD(FShaderResourceParameter, WeightmapWeightBlendModeParam);
+	LAYOUT_FIELD(FShaderResourceParameter, WeightmapTextureOutputOffsetParam);
 };
 
 IMPLEMENT_GLOBAL_SHADER(FLandscapeLayerWeightmapPackMaterialLayersCS, "/Engine/Private/LandscapeLayersCS.usf", "PackPaintLayerToWeightmap", SF_Compute);
@@ -1112,14 +1022,14 @@ public:
 		SCOPED_GPU_STAT(InRHICmdList, LandscapeLayersRender);
 
 		TShaderMapRef<FLandscapeLayerWeightmapPackMaterialLayersCS> ComputeShader(GetGlobalShaderMap(GMaxRHIFeatureLevel));
-		InRHICmdList.SetComputeShader(ComputeShader->GetComputeShader());
+		InRHICmdList.SetComputeShader(ComputeShader.GetComputeShader());
 		ComputeShader->SetParameters(InRHICmdList, ShaderParams);
 
 		uint32 ThreadGroupCountX = FMath::CeilToInt((float)ShaderParams.ComponentSize / (float)GLandscapeLayerWeightmapThreadGroupSizeX);
 		uint32 ThreadGroupCountY = FMath::CeilToInt((float)ShaderParams.ComponentSize / (float)GLandscapeLayerWeightmapThreadGroupSizeY);
 		check(ThreadGroupCountX > 0 && ThreadGroupCountY > 0);
 
-		DispatchComputeShader(InRHICmdList, *ComputeShader, ThreadGroupCountX, ThreadGroupCountY, ShaderParams.ComputeShaderResource->GetComponentsDataCount());
+		DispatchComputeShader(InRHICmdList, ComputeShader.GetShader(), ThreadGroupCountX, ThreadGroupCountY, ShaderParams.ComputeShaderResource->GetComponentsDataCount());
 		ComputeShader->UnsetParameters(InRHICmdList);
 		ShaderParams.ComputeShaderResource->ReleaseResource();
 		delete ShaderParams.ComputeShaderResource;
@@ -1301,8 +1211,8 @@ public:
 			TShaderMapRef<FLandscapeLayersVS> VertexShader(GetGlobalShaderMap(View->GetFeatureLevel()));
 			TShaderMapRef<ShaderPixelClass> PixelShader(GetGlobalShaderMap(View->GetFeatureLevel()));
 
-			GraphicsPSOInit.BoundShaderState.VertexShaderRHI = GETSAFERHISHADER_VERTEX(*VertexShader);
-			GraphicsPSOInit.BoundShaderState.PixelShaderRHI = GETSAFERHISHADER_PIXEL(*PixelShader);
+			GraphicsPSOInit.BoundShaderState.VertexShaderRHI = VertexShader.GetVertexShader();
+			GraphicsPSOInit.BoundShaderState.PixelShaderRHI = PixelShader.GetPixelShader();
 
 			InRHICmdList.SetViewport(View->UnscaledViewRect.Min.X, View->UnscaledViewRect.Min.Y, 0.0f, View->UnscaledViewRect.Max.X, View->UnscaledViewRect.Max.Y, 1.0f);
 
@@ -1319,8 +1229,8 @@ public:
 			TShaderMapRef<FLandscapeLayersVS> VertexShader(GetGlobalShaderMap(View->GetFeatureLevel()));
 			TShaderMapRef<ShaderPixelMipsClass> PixelShader(GetGlobalShaderMap(View->GetFeatureLevel()));
 
-			GraphicsPSOInit.BoundShaderState.VertexShaderRHI = GETSAFERHISHADER_VERTEX(*VertexShader);
-			GraphicsPSOInit.BoundShaderState.PixelShaderRHI = GETSAFERHISHADER_PIXEL(*PixelShader);
+			GraphicsPSOInit.BoundShaderState.VertexShaderRHI = VertexShader.GetVertexShader();
+			GraphicsPSOInit.BoundShaderState.PixelShaderRHI = PixelShader.GetPixelShader();
 
 			InRHICmdList.SetViewport(0.0f, 0.0f, 0.0f, WriteRenderTargetSize.X, WriteRenderTargetSize.Y, 1.0f);
 
@@ -1952,7 +1862,7 @@ void ALandscape::CopyTexturePS(const FString& InSourceDebugName, FTextureResourc
 		FRHIRenderPassInfo RPInfo(InDestResource->TextureRHI, ERenderTargetActions::DontLoad_Store);
 		RHICmdList.BeginRenderPass(RPInfo, TEXT("CopyTexture"));
 
-		TShaderMap<FGlobalShaderType>* GlobalShaderMap = GetGlobalShaderMap(GMaxRHIFeatureLevel);
+		FGlobalShaderMap* GlobalShaderMap = GetGlobalShaderMap(GMaxRHIFeatureLevel);
 		TShaderMapRef< FLandscapeCopyTextureVS > VertexShader(GlobalShaderMap);
 		TShaderMapRef< FLandscapeCopyTexturePS > PixelShader(GlobalShaderMap);
 
@@ -1963,8 +1873,8 @@ void ALandscape::CopyTexturePS(const FString& InSourceDebugName, FTextureResourc
 		GraphicsPSOInit.DepthStencilState = TStaticDepthStencilState<false, CF_Always>::GetRHI();
 		GraphicsPSOInit.PrimitiveType = PT_TriangleList;
 		GraphicsPSOInit.BoundShaderState.VertexDeclarationRHI = GetVertexDeclarationFVector4();
-		GraphicsPSOInit.BoundShaderState.VertexShaderRHI = GETSAFERHISHADER_VERTEX(*VertexShader);
-		GraphicsPSOInit.BoundShaderState.PixelShaderRHI = GETSAFERHISHADER_PIXEL(*PixelShader);
+		GraphicsPSOInit.BoundShaderState.VertexShaderRHI = VertexShader.GetVertexShader();
+		GraphicsPSOInit.BoundShaderState.PixelShaderRHI = PixelShader.GetPixelShader();
 		SetGraphicsPipelineState(RHICmdList, GraphicsPSOInit);
 
 		PixelShader->SetParameters(RHICmdList, InSourceResource->TextureRHI);
@@ -2809,7 +2719,7 @@ int32 ALandscape::RegenerateLayersHeightmaps(const TArray<ULandscapeComponent*>&
 	{
 		return 0;
 	}
-		
+
 	// Nothing to do (return that we did the processing)
 	if (InLandscapeComponentsToResolve.Num() == 0)
 	{
@@ -3801,7 +3711,7 @@ int32 ALandscape::RegenerateLayersWeightmaps(const TArray<ULandscapeComponent*>&
 	{
 		return 0;
 	}
-		
+
 	if (InLandscapeComponentsToResolve.Num() == 0)
 	{
 		return WeightmapUpdateModes;
diff --git a/Engine/Source/Runtime/Landscape/Private/LandscapeGizmoActor.cpp b/Engine/Source/Runtime/Landscape/Private/LandscapeGizmoActor.cpp
--- a/Engine/Source/Runtime/Landscape/Private/LandscapeGizmoActor.cpp
+++ b/Engine/Source/Runtime/Landscape/Private/LandscapeGizmoActor.cpp
@@ -52,7 +52,7 @@ public:
 		return Parent->GetMaterialWithFallback(InFeatureLevel, OutFallbackMaterialRenderProxy);
 	}
 
-	virtual bool GetVectorValue(const FMaterialParameterInfo& ParameterInfo, FLinearColor* OutValue, const FMaterialRenderContext& Context) const
+	virtual bool GetVectorValue(const FHashedMaterialParameterInfo& ParameterInfo, FLinearColor* OutValue, const FMaterialRenderContext& Context) const override
 	{
 		if (ParameterInfo.Name == FName(TEXT("AlphaScaleBias")))
 		{
@@ -86,7 +86,7 @@ public:
 
 		return Parent->GetVectorValue(ParameterInfo, OutValue, Context);
 	}
-	virtual bool GetScalarValue(const FMaterialParameterInfo& ParameterInfo, float* OutValue, const FMaterialRenderContext& Context) const
+	virtual bool GetScalarValue(const FHashedMaterialParameterInfo& ParameterInfo, float* OutValue, const FMaterialRenderContext& Context) const override
 	{
 		if (ParameterInfo.Name == FName(TEXT("Top")))
 		{
@@ -100,7 +100,7 @@ public:
 		}
 		return Parent->GetScalarValue(ParameterInfo, OutValue, Context);
 	}
-	virtual bool GetTextureValue(const FMaterialParameterInfo& ParameterInfo,const UTexture** OutValue, const FMaterialRenderContext& Context) const
+	virtual bool GetTextureValue(const FHashedMaterialParameterInfo& ParameterInfo,const UTexture** OutValue, const FMaterialRenderContext& Context) const override
 	{
 		if (ParameterInfo.Name == FName(TEXT("AlphaTexture")))
 		{
@@ -111,7 +111,7 @@ public:
 		}
 		return Parent->GetTextureValue(ParameterInfo, OutValue, Context);
 	}
-	virtual bool GetTextureValue(const FMaterialParameterInfo& ParameterInfo, const URuntimeVirtualTexture** OutValue, const FMaterialRenderContext& Context) const
+	virtual bool GetTextureValue(const FHashedMaterialParameterInfo& ParameterInfo, const URuntimeVirtualTexture** OutValue, const FMaterialRenderContext& Context) const
 	{
 		return Parent->GetTextureValue(ParameterInfo, OutValue, Context);
 	}
diff --git a/Engine/Source/Runtime/Landscape/Private/LandscapeGrass.cpp b/Engine/Source/Runtime/Landscape/Private/LandscapeGrass.cpp
--- a/Engine/Source/Runtime/Landscape/Private/LandscapeGrass.cpp
+++ b/Engine/Source/Runtime/Landscape/Private/LandscapeGrass.cpp
@@ -218,13 +218,13 @@ static FAutoConsoleVariableSink CVarGrassSink(FConsoleCommandDelegate::CreateSta
 //
 
 #if WITH_EDITOR
-static bool ShouldCacheLandscapeGrassShaders(EShaderPlatform Platform, const FMaterial* Material, const FVertexFactoryType* VertexFactoryType)
+static bool ShouldCacheLandscapeGrassShaders(const FMeshMaterialShaderPermutationParameters& Parameters)
 {
 	// We only need grass weight shaders for Landscape vertex factories on desktop platforms
-	return (Material->IsUsedWithLandscape() || Material->IsSpecialEngineMaterial()) &&
-		IsFeatureLevelSupported(Platform, ERHIFeatureLevel::SM5) &&
-		((VertexFactoryType == FindVertexFactoryType(FName(TEXT("FLandscapeVertexFactory"), FNAME_Find))) || (VertexFactoryType == FindVertexFactoryType(FName(TEXT("FLandscapeXYOffsetVertexFactory"), FNAME_Find))))
-		&& !IsConsolePlatform(Platform);
+	return (Parameters.MaterialParameters.bIsUsedWithLandscape || Parameters.MaterialParameters.bIsSpecialEngineMaterial) &&
+		IsFeatureLevelSupported(Parameters.Platform, ERHIFeatureLevel::SM5) &&
+		((Parameters.VertexFactoryType == FindVertexFactoryType(FName(TEXT("FLandscapeVertexFactory"), FNAME_Find))) || (Parameters.VertexFactoryType == FindVertexFactoryType(FName(TEXT("FLandscapeXYOffsetVertexFactory"), FNAME_Find))))
+		&& !IsConsolePlatform(Parameters.Platform);
 }
 
 class FLandscapeGrassWeightShaderElementData : public FMeshMaterialShaderElementData
@@ -239,7 +239,7 @@ class FLandscapeGrassWeightVS : public FMeshMaterialShader
 {
 	DECLARE_SHADER_TYPE(FLandscapeGrassWeightVS, MeshMaterial);
 
-	FShaderParameter RenderOffsetParameter;
+	LAYOUT_FIELD(FShaderParameter, RenderOffsetParameter);
 
 protected:
 
@@ -257,7 +257,7 @@ public:
 
 	static bool ShouldCompilePermutation(const FMeshMaterialShaderPermutationParameters& Parameters)
 	{
-		return ShouldCacheLandscapeGrassShaders(Parameters.Platform, Parameters.Material, Parameters.VertexFactoryType);
+		return ShouldCacheLandscapeGrassShaders(Parameters);
 	}
 
 	void GetShaderBindings(
@@ -274,13 +274,6 @@ public:
 
 		ShaderBindings.Add(RenderOffsetParameter, ShaderElementData.RenderOffset);
 	}
-
-	virtual bool Serialize(FArchive& Ar) override
-	{
-		bool bShaderHasOutdatedParameters = FMeshMaterialShader::Serialize(Ar);
-		Ar << RenderOffsetParameter;
-		return bShaderHasOutdatedParameters;
-	}
 };
 
 IMPLEMENT_MATERIAL_SHADER_TYPE(, FLandscapeGrassWeightVS, TEXT("/Engine/Private/LandscapeGrassWeight.usf"), TEXT("VSMain"), SF_Vertex);
@@ -288,12 +281,12 @@ IMPLEMENT_MATERIAL_SHADER_TYPE(, FLandscapeGrassWeightVS, TEXT("/Engine/Private/
 class FLandscapeGrassWeightPS : public FMeshMaterialShader
 {
 	DECLARE_SHADER_TYPE(FLandscapeGrassWeightPS, MeshMaterial);
-	FShaderParameter OutputPassParameter;
+	LAYOUT_FIELD(FShaderParameter, OutputPassParameter);
 public:
 
 	static bool ShouldCompilePermutation(const FMeshMaterialShaderPermutationParameters& Parameters)
 	{
-		return ShouldCacheLandscapeGrassShaders(Parameters.Platform, Parameters.Material, Parameters.VertexFactoryType);
+		return ShouldCacheLandscapeGrassShaders(Parameters);
 	}
 
 	FLandscapeGrassWeightPS(const ShaderMetaType::CompiledShaderInitializerType& Initializer)
@@ -320,13 +313,6 @@ public:
 
 		ShaderBindings.Add(OutputPassParameter, ShaderElementData.OutputPass);
 	}
-
-	virtual bool Serialize(FArchive& Ar) override
-	{
-		bool bShaderHasOutdatedParameters = FMeshMaterialShader::Serialize(Ar);
-		Ar << OutputPassParameter;
-		return bShaderHasOutdatedParameters;
-	}
 };
 
 IMPLEMENT_MATERIAL_SHADER_TYPE(, FLandscapeGrassWeightPS, TEXT("/Engine/Private/LandscapeGrassWeight.usf"), TEXT("PSMain"), SF_Pixel);
@@ -418,7 +404,8 @@ void FLandscapeGrassWeightMeshProcessor::Process(
 	PassShaders.PixelShader = MaterialResource.GetShader<FLandscapeGrassWeightPS>(VertexFactory->GetType());
 	PassShaders.VertexShader = MaterialResource.GetShader<FLandscapeGrassWeightVS>(VertexFactory->GetType());
 
-	const ERasterizerFillMode MeshFillMode = ComputeMeshFillMode(MeshBatch, MaterialResource);
+	const FMeshDrawingPolicyOverrideSettings OverrideSettings = ComputeMeshOverrideSettings(MeshBatch);
+	const ERasterizerFillMode MeshFillMode = ComputeMeshFillMode(MeshBatch, MaterialResource, OverrideSettings);
 	const ERasterizerCullMode MeshCullMode = CM_None;
 
 	FLandscapeGrassWeightShaderElementData ShaderElementData;
@@ -2058,28 +2045,24 @@ void ALandscapeProxy::GetGrassTypes(const UWorld* World, UMaterialInterface* Lan
 	ERHIFeatureLevel::Type FeatureLevel = World->Scene != nullptr? World->Scene->GetFeatureLevel() : GMaxRHIFeatureLevel;
 	if (LandscapeMat)
 	{
-		TArray<const UMaterialExpressionLandscapeGrassOutput*> GrassExpressions;
-		LandscapeMat->GetMaterial()->GetAllExpressionsOfType<UMaterialExpressionLandscapeGrassOutput>(GrassExpressions);
-		if (GrassExpressions.Num() > 0)
+		GrassTypesOut.Append(LandscapeMat->GetMaterial()->GetCachedExpressionData().GrassTypes);
+
+		for (const ULandscapeGrassType* GrassType : LandscapeMat->GetMaterial()->GetCachedExpressionData().GrassTypes)
 		{
-			for (const FGrassInput& Type : GrassExpressions[0]->GrassTypes)
+			if (GrassType != nullptr)
 			{
-				if (Type.GrassType)
+				for (auto& GrassVariety : GrassType->GrassVarieties)
 				{
-					GrassTypesOut.Add(Type.GrassType);
-
-					for (auto& GrassVariety : Type.GrassType->GrassVarieties)
+					const int32 EndCullDistance = GrassVariety.EndCullDistance.GetValueForFeatureLevel(FeatureLevel);
+					if (EndCullDistance > MaxDiscardDistance)
 					{
-						int32 EndCullDistance = GrassVariety.EndCullDistance.GetValueForFeatureLevel(FeatureLevel);
-						if (EndCullDistance > MaxDiscardDistance)
-						{
-							MaxDiscardDistance = EndCullDistance;
-						}
+						MaxDiscardDistance = EndCullDistance;
 					}
 				}
 			}
 		}
 	}
+
 	OutMaxSquareDiscardDistance = MaxDiscardDistance * MaxDiscardDistance;
 }
 
