Copying Private-LoadTimes-4.24 stream to Main. Biggest changes are in Materials/Shader memory freezing.

diff --git a/Engine/Source/Runtime/Landscape/Private/LandscapeRender.cpp b/Engine/Source/Runtime/Landscape/Private/LandscapeRender.cpp
--- a/Engine/Source/Runtime/Landscape/Private/LandscapeRender.cpp
+++ b/Engine/Source/Runtime/Landscape/Private/LandscapeRender.cpp
@@ -40,8 +40,10 @@ LandscapeRender.cpp: New terrain rendering
 #include "MeshMaterialShader.h"
 #include "VT/RuntimeVirtualTexture.h"
 #include "RayTracingInstance.h"
+#include "ProfilingDebugging/LoadTimeTracker.h"
 
 IMPLEMENT_GLOBAL_SHADER_PARAMETER_STRUCT(FLandscapeUniformShaderParameters, "LandscapeParameters");
+IMPLEMENT_TYPE_LAYOUT(FLandscapeVertexFactoryPixelShaderParameters);
 
 int32 GLandscapeMeshLODBias = 0;
 FAutoConsoleVariableRef CVarLandscapeMeshLODBias(
@@ -1286,7 +1288,7 @@ FLandscapeComponentSceneProxy::FLandscapeComponentSceneProxy(ULandscapeComponent
 		for (UMaterialInterface*& MaterialInterface : AvailableMaterials)
 		{
 			if (MaterialInterface == nullptr ||
-				(bHasStaticLighting && !MaterialInterface->CheckMaterialUsage(MATUSAGE_StaticLighting)))
+				(bHasStaticLighting && !MaterialInterface->CheckMaterialUsage_Concurrent(MATUSAGE_StaticLighting)))
 			{
 				MaterialInterface = UMaterial::GetDefaultMaterial(MD_Surface);
 			}
@@ -1301,7 +1303,7 @@ FLandscapeComponentSceneProxy::FLandscapeComponentSceneProxy(ULandscapeComponent
 
 	for (UMaterialInterface*& MaterialInterface : AvailableMaterials)
 	{
-		UMaterial* LandscapeMaterial = MaterialInterface != nullptr ? MaterialInterface->GetMaterial() : nullptr;
+		const UMaterial* LandscapeMaterial = MaterialInterface != nullptr ? MaterialInterface->GetMaterial_Concurrent() : nullptr;
 
 		if (LandscapeMaterial != nullptr)
 		{
@@ -1311,10 +1313,10 @@ FLandscapeComponentSceneProxy::FLandscapeComponentSceneProxy(ULandscapeComponent
 			// and since the MIC is not really valid, dont generate the relevance.
 			if (MaterialInstance == nullptr || MaterialInstance->IsChildOf(LandscapeMaterial))
 			{
-				MaterialRelevances.Add(MaterialInterface->GetRelevance(FeatureLevel));
+				MaterialRelevances.Add(MaterialInterface->GetRelevance_Concurrent(FeatureLevel));
 			}
 
-			bRequiresAdjacencyInformation |= MaterialSettingsRequireAdjacencyInformation_GameThread(MaterialInterface, XYOffsetmapTexture == nullptr ? &FLandscapeVertexFactory::StaticType : &FLandscapeXYOffsetVertexFactory::StaticType, InComponent->GetWorld()->FeatureLevel);
+			bRequiresAdjacencyInformation |= RequiresAdjacencyInformation(MaterialInterface, XYOffsetmapTexture == nullptr ? &FLandscapeVertexFactory::StaticType : &FLandscapeXYOffsetVertexFactory::StaticType, InComponent->GetWorld()->FeatureLevel.GetValue());
 
 			bool HasTessellationEnabled = false;
 
@@ -2076,7 +2078,7 @@ bool FLandscapeComponentSceneProxy::GetMeshElement(bool UseSeperateBatchForShado
 	}
 
 	// Could be different from bRequiresAdjacencyInformation during shader compilation
-	bool bCurrentRequiresAdjacencyInformation = MaterialRenderingRequiresAdjacencyInformation_RenderingThread(InMaterialInterface, VertexFactory->GetType(), GetScene().GetFeatureLevel());
+	bool bCurrentRequiresAdjacencyInformation = RequiresAdjacencyInformation(InMaterialInterface, VertexFactory->GetType(), GetScene().GetFeatureLevel());
 
 	if (bCurrentRequiresAdjacencyInformation)
 	{
@@ -2254,7 +2256,7 @@ bool FLandscapeComponentSceneProxy::GetStaticMeshElement(int32 LODIndex, bool bF
 		}
 
 		UMaterialInstance* MaterialInstance = Cast<UMaterialInstance>(AvailableMaterials[MaterialIndex]);
-		bool HasTessellationEnabled = (GetScene().GetFeatureLevel() >= ERHIFeatureLevel::SM5) ? MaterialInstance != nullptr && MaterialInstance->GetMaterial()->D3D11TessellationMode != EMaterialTessellationMode::MTM_NoTessellation && MaterialIndexToDisabledTessellationMaterial[MaterialIndex] != INDEX_NONE : false;
+		bool HasTessellationEnabled = (GetScene().GetFeatureLevel() >= ERHIFeatureLevel::SM5) ? MaterialInstance != nullptr && MaterialInstance->GetMaterial_Concurrent()->D3D11TessellationMode != EMaterialTessellationMode::MTM_NoTessellation && MaterialIndexToDisabledTessellationMaterial[MaterialIndex] != INDEX_NONE : false;
 
 		if (HasTessellationEnabled)
 		{
@@ -2285,7 +2287,7 @@ bool FLandscapeComponentSceneProxy::GetStaticMeshElement(int32 LODIndex, bool bF
 
 	// Based on the final material we selected, detect if it has tessellation
 	// Could be different from bRequiresAdjacencyInformation during shader compilation
-	bool bCurrentRequiresAdjacencyInformation = MaterialRenderingRequiresAdjacencyInformation_RenderingThread(MaterialInterface, VertexFactory->GetType(), GetScene().GetFeatureLevel());
+	bool bCurrentRequiresAdjacencyInformation = RequiresAdjacencyInformation(MaterialInterface, VertexFactory->GetType(), GetScene().GetFeatureLevel());
 
 	check(!bCurrentRequiresAdjacencyInformation || (bCurrentRequiresAdjacencyInformation && SharedBuffers->AdjacencyIndexBuffers));
 
@@ -3752,6 +3754,8 @@ int32 FLandscapeComponentSceneProxy::CollectOccluderElements(FOccluderElementsCo
 */
 void FLandscapeVertexBuffer::InitRHI()
 {
+	SCOPED_LOADTIMER(FLandscapeVertexBuffer_InitRHI);
+
 	// create a static vertex buffer
 	FRHIResourceCreateInfo CreateInfo;
 	void* BufferData = nullptr;
@@ -4313,12 +4317,13 @@ FLandscapeSharedAdjacencyIndexBuffer::~FLandscapeSharedAdjacencyIndexBuffer()
 /** Shader parameters for use with FLandscapeVertexFactory */
 class FLandscapeVertexFactoryVertexShaderParameters : public FVertexFactoryShaderParameters
 {
+	DECLARE_INLINE_TYPE_LAYOUT(FLandscapeVertexFactoryVertexShaderParameters, NonVirtual);
 public:
 	/**
 	* Bind shader constants by name
 	* @param	ParameterMap - mapping of named shader constants to indices
 	*/
-	virtual void Bind(const FShaderParameterMap& ParameterMap) override
+	void Bind(const FShaderParameterMap& ParameterMap)
 	{
 		LodBiasParameter.Bind(ParameterMap, TEXT("LodBias"));
 		LodValuesParameter.Bind(ParameterMap, TEXT("LodValues"));
@@ -4328,21 +4333,7 @@ public:
 		NeighborSectionLodParameter.Bind(ParameterMap, TEXT("NeighborSectionLod"));
 	}
 
-	/**
-	* Serialize shader params to an archive
-	* @param	Ar - archive to serialize to
-	*/
-	virtual void Serialize(FArchive& Ar) override
-	{
-		Ar << LodValuesParameter;
-		Ar << ForcedLodParameter;
-		Ar << LodTessellationParameter;
-		Ar << NeighborSectionLodParameter;
-		Ar << LodBiasParameter;
-		Ar << SectionLodsParameter;
-	}
-
-	virtual void GetElementShaderBindings(
+	void GetElementShaderBindings(
 		const class FSceneInterface* Scene,
 		const FSceneView* InView,
 		const class FMeshMaterialShader* Shader,
@@ -4352,7 +4343,7 @@ public:
 		const FMeshBatchElement& BatchElement,
 		class FMeshDrawSingleShaderBindings& ShaderBindings,
 		FVertexInputStreamArray& VertexStreams
-		) const override
+	) const
 	{
 		SCOPE_CYCLE_COUNTER(STAT_LandscapeVFDrawTimeVS);
 
@@ -4390,19 +4381,14 @@ public:
 		}
 	}
 
-	virtual uint32 GetSize() const override
-	{
-		return sizeof(*this);
-	}
-
 protected:
-	FShaderParameter LodTessellationParameter;
-	FShaderParameter LodValuesParameter;
-	FShaderParameter ForcedLodParameter;
-	FShaderParameter NeighborSectionLodParameter;
-	FShaderParameter LodBiasParameter;
-	FShaderParameter SectionLodsParameter;
-	TShaderUniformBufferParameter<FLandscapeUniformShaderParameters> LandscapeShaderParameters;
+	LAYOUT_FIELD(FShaderParameter, LodTessellationParameter);
+	LAYOUT_FIELD(FShaderParameter, LodValuesParameter);
+	LAYOUT_FIELD(FShaderParameter, ForcedLodParameter);
+	LAYOUT_FIELD(FShaderParameter, NeighborSectionLodParameter);
+	LAYOUT_FIELD(FShaderParameter, LodBiasParameter);
+	LAYOUT_FIELD(FShaderParameter, SectionLodsParameter);
+	LAYOUT_FIELD(TShaderUniformBufferParameter<FLandscapeUniformShaderParameters>, LandscapeShaderParameters);
 };
 
 /** 
@@ -4411,8 +4397,9 @@ protected:
   */
 class FLandscapeFixedGridVertexFactoryVertexShaderParameters : public FLandscapeVertexFactoryVertexShaderParameters
 {
+	DECLARE_INLINE_TYPE_LAYOUT(FLandscapeFixedGridVertexFactoryVertexShaderParameters, NonVirtual);
 public:
-	virtual void GetElementShaderBindings(
+	void GetElementShaderBindings(
 		const class FSceneInterface* Scene,
 		const FSceneView* InView,
 		const class FMeshMaterialShader* Shader,
@@ -4422,7 +4409,7 @@ public:
 		const FMeshBatchElement& BatchElement,
 		class FMeshDrawSingleShaderBindings& ShaderBindings,
 		FVertexInputStreamArray& VertexStreams
-	) const override
+	) const
 	{
 		SCOPE_CYCLE_COUNTER(STAT_LandscapeVFDrawTimeVS);
 
@@ -4448,6 +4435,9 @@ public:
 			ShaderBindings.Add(ForcedLodParameter, BatchElementParams->ForcedLOD);
 		}
 	}
+
+	
+	
 };
 
 //
@@ -4496,30 +4486,25 @@ FLandscapeVertexFactory::FLandscapeVertexFactory(ERHIFeatureLevel::Type InFeatur
 {
 }
 
-FVertexFactoryShaderParameters* FLandscapeVertexFactory::ConstructShaderParameters(EShaderFrequency ShaderFrequency)
+bool FLandscapeVertexFactory::ShouldCompilePermutation(const FVertexFactoryShaderPermutationParameters& Parameters)
 {
-	switch (ShaderFrequency)
-	{
-	case SF_Vertex:
-#if RHI_RAYTRACING
-	case SF_Compute:
-	case SF_RayHitGroup:
-#endif
-		return new FLandscapeVertexFactoryVertexShaderParameters();
-		break;
-	case SF_Pixel:
-		return new FLandscapeVertexFactoryPixelShaderParameters();
-		break;
-	default:
-		return nullptr;
-	}
+	// only compile landscape materials for landscape vertex factory
+	// The special engine materials must be compiled for the landscape vertex factory because they are used with it for wireframe, etc.
+	return IsFeatureLevelSupported(Parameters.Platform, ERHIFeatureLevel::SM5) &&
+		(Parameters.MaterialParameters.bIsUsedWithLandscape || Parameters.MaterialParameters.bIsSpecialEngineMaterial);
 }
 
-void FLandscapeVertexFactory::ModifyCompilationEnvironment(const FVertexFactoryType* Type, EShaderPlatform Platform, const FMaterial* Material, FShaderCompilerEnvironment& OutEnvironment)
+void FLandscapeVertexFactory::ModifyCompilationEnvironment(const FVertexFactoryShaderPermutationParameters& Parameters, FShaderCompilerEnvironment& OutEnvironment)
 {
-	FVertexFactory::ModifyCompilationEnvironment(Type, Platform, Material, OutEnvironment);
+	FVertexFactory::ModifyCompilationEnvironment(Parameters, OutEnvironment);
 }
 
+IMPLEMENT_VERTEX_FACTORY_PARAMETER_TYPE(FLandscapeVertexFactory, SF_Vertex, FLandscapeVertexFactoryVertexShaderParameters);
+#if RHI_RAYTRACING
+IMPLEMENT_VERTEX_FACTORY_PARAMETER_TYPE(FLandscapeVertexFactory, SF_Compute, FLandscapeVertexFactoryVertexShaderParameters);
+IMPLEMENT_VERTEX_FACTORY_PARAMETER_TYPE(FLandscapeVertexFactory, SF_RayHitGroup, FLandscapeVertexFactoryVertexShaderParameters);
+#endif // RHI_RAYTRACING
+IMPLEMENT_VERTEX_FACTORY_PARAMETER_TYPE(FLandscapeVertexFactory, SF_Pixel, FLandscapeVertexFactoryPixelShaderParameters);
 
 IMPLEMENT_VERTEX_FACTORY_TYPE_EX(FLandscapeVertexFactory, "/Engine/Private/LandscapeVertexFactory.ush", true, true, true, false, false, true, false);
 
@@ -4544,9 +4529,9 @@ void FLandscapeVertexFactory::Copy(const FLandscapeVertexFactory& Other)
 // FLandscapeXYOffsetVertexFactory
 //
 
-void FLandscapeXYOffsetVertexFactory::ModifyCompilationEnvironment(const FVertexFactoryType* Type, EShaderPlatform Platform, const FMaterial* Material, FShaderCompilerEnvironment& OutEnvironment)
+void FLandscapeXYOffsetVertexFactory::ModifyCompilationEnvironment(const FVertexFactoryShaderPermutationParameters& Parameters, FShaderCompilerEnvironment& OutEnvironment)
 {
-	FLandscapeVertexFactory::ModifyCompilationEnvironment(Type, Platform, Material, OutEnvironment);
+	FLandscapeVertexFactory::ModifyCompilationEnvironment(Parameters, OutEnvironment);
 	OutEnvironment.SetDefine(TEXT("LANDSCAPE_XYOFFSET"), TEXT("1"));
 }
 
@@ -4556,30 +4541,18 @@ IMPLEMENT_VERTEX_FACTORY_TYPE_EX(FLandscapeXYOffsetVertexFactory, "/Engine/Priva
 // FLandscapeFixedGridVertexFactory
 //
 
-void FLandscapeFixedGridVertexFactory::ModifyCompilationEnvironment(const FVertexFactoryType* Type, EShaderPlatform Platform, const FMaterial* Material, FShaderCompilerEnvironment& OutEnvironment)
+void FLandscapeFixedGridVertexFactory::ModifyCompilationEnvironment(const FVertexFactoryShaderPermutationParameters& Parameters, FShaderCompilerEnvironment& OutEnvironment)
 {
-	FLandscapeVertexFactory::ModifyCompilationEnvironment(Type, Platform, Material, OutEnvironment);
+	FLandscapeVertexFactory::ModifyCompilationEnvironment(Parameters, OutEnvironment);
 	OutEnvironment.SetDefine(TEXT("FIXED_GRID"), TEXT("1"));
 }
 
-FVertexFactoryShaderParameters* FLandscapeFixedGridVertexFactory::ConstructShaderParameters(EShaderFrequency ShaderFrequency)
-{
-	switch (ShaderFrequency)
-	{
-	case SF_Vertex:
+IMPLEMENT_VERTEX_FACTORY_PARAMETER_TYPE(FLandscapeFixedGridVertexFactory, SF_Vertex, FLandscapeFixedGridVertexFactoryVertexShaderParameters);
 #if RHI_RAYTRACING
-	case SF_Compute:
-	case SF_RayHitGroup:
+IMPLEMENT_VERTEX_FACTORY_PARAMETER_TYPE(FLandscapeFixedGridVertexFactory, SF_Compute, FLandscapeFixedGridVertexFactoryVertexShaderParameters);
+IMPLEMENT_VERTEX_FACTORY_PARAMETER_TYPE(FLandscapeFixedGridVertexFactory, SF_RayHitGroup, FLandscapeFixedGridVertexFactoryVertexShaderParameters);
 #endif
-		return new FLandscapeFixedGridVertexFactoryVertexShaderParameters();
-		break;
-	case SF_Pixel:
-		return new FLandscapeVertexFactoryPixelShaderParameters();
-		break;
-	default:
-		return nullptr;
-	}
-}
+IMPLEMENT_VERTEX_FACTORY_PARAMETER_TYPE(FLandscapeFixedGridVertexFactory, SF_Pixel, FLandscapeVertexFactoryPixelShaderParameters);
 
 IMPLEMENT_VERTEX_FACTORY_TYPE_EX(FLandscapeFixedGridVertexFactory, "/Engine/Private/LandscapeVertexFactory.ush", true, true, true, false, false, true, false);
 
@@ -4591,6 +4564,80 @@ ULandscapeMaterialInstanceConstant::ULandscapeMaterialInstanceConstant(const FOb
 	bIsLayerThumbnail = false;
 }
 
+void ULandscapeMaterialInstanceConstant::PostLoad()
+{
+	Super::PostLoad();
+#if WITH_EDITOR
+	UpdateCachedTextureStreaming();
+#endif // WITH_EDITOR
+}
+
+float ULandscapeMaterialInstanceConstant::GetLandscapeTexelFactor(const FName& TextureName) const
+{
+	for (const FLandscapeMaterialTextureStreamingInfo& Info : TextureStreamingInfo)
+	{
+		if (Info.TextureName == TextureName)
+		{
+			return Info.TexelFactor;
+		}
+	}
+	return 1.0f;
+}
+
+#if WITH_EDITOR
+
+void ULandscapeMaterialInstanceConstant::PostEditChangeProperty(FPropertyChangedEvent& PropertyChangedEvent)
+{
+	Super::PostEditChangeProperty(PropertyChangedEvent);
+	UpdateCachedTextureStreaming();
+}
+
+FLandscapeMaterialTextureStreamingInfo& ULandscapeMaterialInstanceConstant::AcquireTextureStreamingInfo(const FName& TextureName)
+{
+	for (FLandscapeMaterialTextureStreamingInfo& Info : TextureStreamingInfo)
+	{
+		if (Info.TextureName == TextureName)
+		{
+			return Info;
+		}
+	}
+	FLandscapeMaterialTextureStreamingInfo& Info = TextureStreamingInfo.AddDefaulted_GetRef();
+	Info.TextureName = TextureName;
+	Info.TexelFactor = 1.0f;
+	return Info;
+}
+
+void ULandscapeMaterialInstanceConstant::UpdateCachedTextureStreaming()
+{
+	const UMaterial* Material = GetMaterial();
+	if (Material)
+	{
+		int32 NumExpressions = Material->Expressions.Num();
+		for (int32 ExpressionIndex = 0; ExpressionIndex < NumExpressions; ExpressionIndex++)
+		{
+			UMaterialExpression* Expression = Material->Expressions[ExpressionIndex];
+			UMaterialExpressionTextureSample* TextureSample = Cast<UMaterialExpressionTextureSample>(Expression);
+
+			// TODO: This is only works for direct Coordinate Texture Sample cases
+			if (TextureSample && TextureSample->Coordinates.IsConnected())
+			{
+				if (UMaterialExpressionTextureCoordinate* TextureCoordinate = Cast<UMaterialExpressionTextureCoordinate>(TextureSample->Coordinates.Expression))
+				{
+					FLandscapeMaterialTextureStreamingInfo& Info = AcquireTextureStreamingInfo(TextureSample->Texture->GetFName());
+					Info.TexelFactor *= FPlatformMath::Max(TextureCoordinate->UTiling, TextureCoordinate->VTiling);
+				}
+				else if (UMaterialExpressionLandscapeLayerCoords* TerrainTextureCoordinate = Cast<UMaterialExpressionLandscapeLayerCoords>(TextureSample->Coordinates.Expression))
+				{
+					FLandscapeMaterialTextureStreamingInfo& Info = AcquireTextureStreamingInfo(TextureSample->Texture->GetFName());
+					Info.TexelFactor *= TerrainTextureCoordinate->MappingScale;
+				}
+			}
+		}
+	}
+}
+
+#endif // WITH_EDITOR
+
 class FLandscapeMaterialResource : public FMaterialResource
 {
 	const bool bIsLayerThumbnail;
@@ -5057,6 +5104,9 @@ void ULandscapeComponent::GetStreamingRenderAssetInfo(FStreamingTextureLevelCont
 		{
 			TArray<UTexture*> Textures;
 			MaterialInterface->GetUsedTextures(Textures, EMaterialQualityLevel::Num, false, FeatureLevel, false);
+
+			const ULandscapeMaterialInstanceConstant* LandscapeMaterial = Cast<ULandscapeMaterialInstanceConstant>(MaterialInterface);
+
 			// Add each texture to the output with the appropriate parameters.
 			// TODO: Take into account which UVIndex is being used.
 			for (int32 TextureIndex = 0; TextureIndex < Textures.Num(); TextureIndex++)
@@ -5068,56 +5118,11 @@ void ULandscapeComponent::GetStreamingRenderAssetInfo(FStreamingTextureLevelCont
 				StreamingTexture.Bounds = BoundingSphere;
 				StreamingTexture.TexelFactor = TexelFactor;
 				StreamingTexture.RenderAsset = Texture2D;
-			}
 
-			const UMaterial* Material = MaterialInterface->GetMaterial();
-			if (Material)
-			{
-				int32 NumExpressions = Material->Expressions.Num();
-				for (int32 ExpressionIndex = 0; ExpressionIndex < NumExpressions; ExpressionIndex++)
+				if (LandscapeMaterial)
 				{
-					UMaterialExpression* Expression = Material->Expressions[ExpressionIndex];
-					UMaterialExpressionTextureSample* TextureSample = Cast<UMaterialExpressionTextureSample>(Expression);
-
-					// TODO: This is only works for direct Coordinate Texture Sample cases
-					if (TextureSample && TextureSample->Coordinates.IsConnected())
-					{
-						UMaterialExpressionTextureCoordinate* TextureCoordinate = nullptr;
-						UMaterialExpressionLandscapeLayerCoords* TerrainTextureCoordinate = nullptr;
-
-						for (UMaterialExpression* FindExp : Material->Expressions)
-						{
-							if (FindExp && FindExp->GetFName() == TextureSample->Coordinates.ExpressionName)
-							{
-								TextureCoordinate = Cast<UMaterialExpressionTextureCoordinate>(FindExp);
-								if (!TextureCoordinate)
-								{
-									TerrainTextureCoordinate = Cast<UMaterialExpressionLandscapeLayerCoords>(FindExp);
-								}
-								break;
-							}
-						}
-
-						if (TextureCoordinate || TerrainTextureCoordinate)
-						{
-							for (int32 i = 0; i < OutStreamingRenderAssets.Num(); ++i)
-							{
-								FStreamingRenderAssetPrimitiveInfo& StreamingTexture = OutStreamingRenderAssets[i];
-								if (StreamingTexture.RenderAsset == TextureSample->Texture)
-								{
-									if (TextureCoordinate)
-									{
-										StreamingTexture.TexelFactor = TexelFactor * FPlatformMath::Max(TextureCoordinate->UTiling, TextureCoordinate->VTiling);
-									}
-									else //if ( TerrainTextureCoordinate )
-									{
-										StreamingTexture.TexelFactor = TexelFactor * TerrainTextureCoordinate->MappingScale;
-									}
-									break;
-								}
-							}
-						}
-					}
+					const float MaterialTexelFactor = LandscapeMaterial->GetLandscapeTexelFactor(Texture2D->GetFName());
+					StreamingTexture.TexelFactor *= MaterialTexelFactor;
 				}
 			}
 
diff --git a/Engine/Source/Runtime/Landscape/Private/LandscapeRenderMobile.cpp b/Engine/Source/Runtime/Landscape/Private/LandscapeRenderMobile.cpp
--- a/Engine/Source/Runtime/Landscape/Private/LandscapeRenderMobile.cpp
+++ b/Engine/Source/Runtime/Landscape/Private/LandscapeRenderMobile.cpp
@@ -20,6 +20,13 @@ static TAutoConsoleVariable<int32> CVarMobileLandscapeHoleMesh(
 	TEXT("Set to 0 to skip loading of landscape hole meshes on mobile."),
 	ECVF_Default);
 
+bool FLandscapeVertexFactoryMobile::ShouldCompilePermutation(const FVertexFactoryShaderPermutationParameters& Parameters)
+{
+	auto FeatureLevel = GetMaxSupportedFeatureLevel(Parameters.Platform);
+	return (FeatureLevel == ERHIFeatureLevel::ES2 || FeatureLevel == ERHIFeatureLevel::ES3_1) &&
+		(Parameters.MaterialParameters.bIsUsedWithLandscape || Parameters.MaterialParameters.bIsSpecialEngineMaterial);
+}
+
 void FLandscapeVertexFactoryMobile::InitRHI()
 {
 	// list of declaration items
@@ -44,12 +51,13 @@ void FLandscapeVertexFactoryMobile::InitRHI()
 /** Shader parameters for use with FLandscapeVertexFactory */
 class FLandscapeVertexFactoryMobileVertexShaderParameters : public FVertexFactoryShaderParameters
 {
+	DECLARE_INLINE_TYPE_LAYOUT(FLandscapeVertexFactoryMobileVertexShaderParameters, NonVirtual);
 public:
 	/**
 	* Bind shader constants by name
 	* @param	ParameterMap - mapping of named shader constants to indices
 	*/
-	virtual void Bind(const FShaderParameterMap& ParameterMap) override
+	void Bind(const FShaderParameterMap& ParameterMap)
 	{
 		LodValuesParameter.Bind(ParameterMap,TEXT("LodValues"));
 		ForcedLodParameter.Bind(ParameterMap,TEXT("ForcedLod"));
@@ -59,21 +67,7 @@ public:
 		SectionLodsParameter.Bind(ParameterMap,TEXT("SectionLods"));
 	}
 
-	/**
-	* Serialize shader params to an archive
-	* @param	Ar - archive to serialize to
-	*/
-	virtual void Serialize(FArchive& Ar) override
-	{
-		Ar << LodValuesParameter;
-		Ar << ForcedLodParameter;
-		Ar << LodTessellationParameter;
-		Ar << NeighborSectionLodParameter;
-		Ar << LodBiasParameter;
-		Ar << SectionLodsParameter;
-	}
-	
-	virtual void GetElementShaderBindings(
+	void GetElementShaderBindings(
 		const class FSceneInterface* Scene,
 		const FSceneView* InView,
 		const class FMeshMaterialShader* Shader,
@@ -83,7 +77,7 @@ public:
 		const FMeshBatchElement& BatchElement,
 		class FMeshDrawSingleShaderBindings& ShaderBindings,
 		FVertexInputStreamArray& VertexStreams
-	) const override
+	) const
 	{
 		SCOPE_CYCLE_COUNTER(STAT_LandscapeVFDrawTimeVS);
 
@@ -189,40 +183,31 @@ public:
 #endif
 	}
 protected:
-	FShaderParameter LodValuesParameter;
-	FShaderParameter ForcedLodParameter;
-	FShaderParameter LodTessellationParameter;
-	FShaderParameter NeighborSectionLodParameter;
-	FShaderParameter LodBiasParameter;
-	FShaderParameter SectionLodsParameter;
-	TShaderUniformBufferParameter<FLandscapeUniformShaderParameters> LandscapeShaderParameters;
+	LAYOUT_FIELD(FShaderParameter, LodValuesParameter);
+	LAYOUT_FIELD(FShaderParameter, ForcedLodParameter);
+	LAYOUT_FIELD(FShaderParameter, LodTessellationParameter);
+	LAYOUT_FIELD(FShaderParameter, NeighborSectionLodParameter);
+	LAYOUT_FIELD(FShaderParameter, LodBiasParameter);
+	LAYOUT_FIELD(FShaderParameter, SectionLodsParameter);
+	LAYOUT_FIELD(TShaderUniformBufferParameter<FLandscapeUniformShaderParameters>, LandscapeShaderParameters);
 };
 
 /** Shader parameters for use with FLandscapeVertexFactory */
 class FLandscapeVertexFactoryMobilePixelShaderParameters : public FLandscapeVertexFactoryPixelShaderParameters
 {
+	DECLARE_INLINE_TYPE_LAYOUT(FLandscapeVertexFactoryMobilePixelShaderParameters, NonVirtual);
 public:
 	/**
 	* Bind shader constants by name
 	* @param	ParameterMap - mapping of named shader constants to indices
 	*/
-	virtual void Bind(const FShaderParameterMap& ParameterMap) override
+	void Bind(const FShaderParameterMap& ParameterMap)
 	{
 		FLandscapeVertexFactoryPixelShaderParameters::Bind(ParameterMap);
 		BlendableLayerMaskParameter.Bind(ParameterMap, TEXT("BlendableLayerMask"));
 	}
 
-	/**
-	* Serialize shader params to an archive
-	* @param	Ar - archive to serialize to
-	*/
-	virtual void Serialize(FArchive& Ar) override
-	{
-		FLandscapeVertexFactoryPixelShaderParameters::Serialize(Ar);
-		Ar << BlendableLayerMaskParameter;
-	}
-
-	virtual void GetElementShaderBindings(
+	void GetElementShaderBindings(
 		const class FSceneInterface* Scene,
 		const FSceneView* InView,
 		const class FMeshMaterialShader* Shader,
@@ -232,7 +217,7 @@ public:
 		const FMeshBatchElement& BatchElement,
 		class FMeshDrawSingleShaderBindings& ShaderBindings,
 		FVertexInputStreamArray& VertexStreams
-	) const override final
+	) const
 	{
 		SCOPE_CYCLE_COUNTER(STAT_LandscapeVFDrawTimePS);
 		
@@ -253,34 +238,17 @@ public:
 	}
 
 protected:
-	FShaderParameter BlendableLayerMaskParameter;
+	LAYOUT_FIELD(FShaderParameter, BlendableLayerMaskParameter);
 };
 
-FVertexFactoryShaderParameters* FLandscapeVertexFactoryMobile::ConstructShaderParameters(EShaderFrequency ShaderFrequency)
-{
-	switch( ShaderFrequency )
-	{
-	case SF_Vertex:
-		return new FLandscapeVertexFactoryMobileVertexShaderParameters();
-	case SF_Pixel:
-		return new FLandscapeVertexFactoryMobilePixelShaderParameters();
-	default:
-		return NULL;
-	}
-}
-
-IMPLEMENT_VERTEX_FACTORY_TYPE(FLandscapeVertexFactoryMobile, "/Engine/Private/LandscapeVertexFactory.ush", true, true, true, false, false);
-IMPLEMENT_VERTEX_FACTORY_TYPE_EX(FLandscapeFixedGridVertexFactoryMobile, "/Engine/Private/LandscapeVertexFactory.ush", true, true, true, false, false, true, false);
-
-
-/** 
+/**
   * Shader parameters for use with FLandscapeFixedGridVertexFactory
   * Simple grid rendering (without dynamic lod blend) needs a simpler fixed setup.
   */
 class FLandscapeFixedGridVertexFactoryMobileVertexShaderParameters : public FLandscapeVertexFactoryMobileVertexShaderParameters
 {
 public:
-	virtual void GetElementShaderBindings(
+	void GetElementShaderBindings(
 		const class FSceneInterface* Scene,
 		const FSceneView* InView,
 		const class FMeshMaterialShader* Shader,
@@ -290,14 +258,14 @@ public:
 		const FMeshBatchElement& BatchElement,
 		class FMeshDrawSingleShaderBindings& ShaderBindings,
 		FVertexInputStreamArray& VertexStreams
-	) const override
+	) const
 	{
 		SCOPE_CYCLE_COUNTER(STAT_LandscapeVFDrawTimeVS);
 
 		const FLandscapeBatchElementParams* BatchElementParams = (const FLandscapeBatchElementParams*)BatchElement.UserData;
 		check(BatchElementParams);
 		const FLandscapeComponentSceneProxyMobile* SceneProxy = (const FLandscapeComponentSceneProxyMobile*)BatchElementParams->SceneProxy;
-		ShaderBindings.Add(Shader->GetUniformBufferParameter<FLandscapeUniformShaderParameters>(),*BatchElementParams->LandscapeUniformShaderParametersResource);
+		ShaderBindings.Add(Shader->GetUniformBufferParameter<FLandscapeUniformShaderParameters>(), *BatchElementParams->LandscapeUniformShaderParametersResource);
 
 		if (LodValuesParameter.IsBound())
 		{
@@ -316,31 +284,27 @@ public:
 	}
 };
 
-void FLandscapeFixedGridVertexFactoryMobile::ModifyCompilationEnvironment(const FVertexFactoryType* Type, EShaderPlatform Platform, const FMaterial* Material, FShaderCompilerEnvironment& OutEnvironment)
-{
-	FLandscapeVertexFactoryMobile::ModifyCompilationEnvironment(Type, Platform, Material, OutEnvironment);
-	OutEnvironment.SetDefine(TEXT("FIXED_GRID"), TEXT("1"));
-}
+IMPLEMENT_VERTEX_FACTORY_PARAMETER_TYPE(FLandscapeVertexFactoryMobile, SF_Vertex, FLandscapeVertexFactoryMobileVertexShaderParameters);
+IMPLEMENT_VERTEX_FACTORY_PARAMETER_TYPE(FLandscapeVertexFactoryMobile, SF_Pixel, FLandscapeVertexFactoryMobilePixelShaderParameters);
 
-FVertexFactoryShaderParameters* FLandscapeFixedGridVertexFactoryMobile::ConstructShaderParameters(EShaderFrequency ShaderFrequency)
+IMPLEMENT_VERTEX_FACTORY_PARAMETER_TYPE(FLandscapeFixedGridVertexFactoryMobile, SF_Vertex, FLandscapeFixedGridVertexFactoryMobileVertexShaderParameters);
+IMPLEMENT_VERTEX_FACTORY_PARAMETER_TYPE(FLandscapeFixedGridVertexFactoryMobile, SF_Pixel, FLandscapeVertexFactoryMobilePixelShaderParameters);
+
+IMPLEMENT_VERTEX_FACTORY_TYPE(FLandscapeVertexFactoryMobile, "/Engine/Private/LandscapeVertexFactory.ush", true, true, true, false, false);
+IMPLEMENT_VERTEX_FACTORY_TYPE_EX(FLandscapeFixedGridVertexFactoryMobile, "/Engine/Private/LandscapeVertexFactory.ush", true, true, true, false, false, true, false);
+
+void FLandscapeFixedGridVertexFactoryMobile::ModifyCompilationEnvironment(const FVertexFactoryShaderPermutationParameters& Parameters, FShaderCompilerEnvironment& OutEnvironment)
 {
-	switch (ShaderFrequency)
-	{
-	case SF_Vertex:
-		return new FLandscapeFixedGridVertexFactoryMobileVertexShaderParameters();
-	case SF_Pixel:
-		return new FLandscapeVertexFactoryMobilePixelShaderParameters();
-	}
-	return nullptr;
+	FLandscapeVertexFactoryMobile::ModifyCompilationEnvironment(Parameters, OutEnvironment);
+	OutEnvironment.SetDefine(TEXT("FIXED_GRID"), TEXT("1"));
 }
 	
-bool FLandscapeFixedGridVertexFactoryMobile::ShouldCompilePermutation(EShaderPlatform Platform, const FMaterial* Material, const FShaderType* ShaderType)
+bool FLandscapeFixedGridVertexFactoryMobile::ShouldCompilePermutation(const FVertexFactoryShaderPermutationParameters& Parameters)
 {
-	return GetMaxSupportedFeatureLevel(Platform) == ERHIFeatureLevel::ES3_1 &&
-		(Material->IsUsedWithLandscape() || Material->IsSpecialEngineMaterial());
+	return GetMaxSupportedFeatureLevel(Parameters.Platform) == ERHIFeatureLevel::ES3_1 &&
+		(Parameters.MaterialParameters.bIsUsedWithLandscape || Parameters.MaterialParameters.bIsSpecialEngineMaterial);
 }
 
-
 /**
 * Initialize the RHI for this rendering resource
 */
diff --git a/Engine/Source/Runtime/Landscape/Public/LandscapeRender.h b/Engine/Source/Runtime/Landscape/Public/LandscapeRender.h
--- a/Engine/Source/Runtime/Landscape/Public/LandscapeRender.h
+++ b/Engine/Source/Runtime/Landscape/Public/LandscapeRender.h
@@ -151,11 +151,9 @@ public:
 /** Pixel shader parameters for use with FLandscapeVertexFactory */
 class FLandscapeVertexFactoryPixelShaderParameters : public FVertexFactoryShaderParameters
 {
+	DECLARE_TYPE_LAYOUT(FLandscapeVertexFactoryPixelShaderParameters, NonVirtual);
 public:
-	virtual void Bind(const class FShaderParameterMap& ParameterMap) {}
-	virtual void Serialize(FArchive& Ar) {}
-
-	virtual void GetElementShaderBindings(
+	void GetElementShaderBindings(
 		const class FSceneInterface* Scene,
 		const FSceneView* InView,
 		const class FMeshMaterialShader* Shader,
@@ -165,12 +163,10 @@ public:
 		const FMeshBatchElement& BatchElement,
 		class FMeshDrawSingleShaderBindings& ShaderBindings,
 		FVertexInputStreamArray& VertexStreams
-	) const override;
+	) const;
 
-	virtual uint32 GetSize() const override
-	{
-		return sizeof(*this);
-	}
+	
+	
 };
 
 /** vertex factory for VTF-heightmap terrain  */
@@ -188,8 +184,6 @@ public:
 		ReleaseResource();
 	}
 
-	static FVertexFactoryShaderParameters* ConstructShaderParameters(EShaderFrequency ShaderFrequency);
-
 	struct FDataType
 	{
 		/** The stream to read the vertex position from. */
@@ -199,18 +193,12 @@ public:
 	/**
 	* Should we cache the material's shadertype on this platform with this vertex factory?
 	*/
-	static bool ShouldCompilePermutation(EShaderPlatform Platform, const FMaterial* Material, const FShaderType* ShaderType)
-	{
-		// only compile landscape materials for landscape vertex factory
-		// The special engine materials must be compiled for the landscape vertex factory because they are used with it for wireframe, etc.
-		return IsFeatureLevelSupported(Platform, ERHIFeatureLevel::SM5) &&
-			(Material->IsUsedWithLandscape() || Material->IsSpecialEngineMaterial());
-	}
+	static bool ShouldCompilePermutation(const FVertexFactoryShaderPermutationParameters& Parameters);
 
 	/**
 	* Can be overridden by FVertexFactory subclasses to modify their compile environment just before compilation occurs.
 	*/
-	static void ModifyCompilationEnvironment(const FVertexFactoryType* Type, EShaderPlatform Platform, const FMaterial* Material, FShaderCompilerEnvironment& OutEnvironment);
+	static void ModifyCompilationEnvironment(const FVertexFactoryShaderPermutationParameters& Parameters, FShaderCompilerEnvironment& OutEnvironment);
 
 	/**
 	* Copy the data from another vertex factory
@@ -252,7 +240,7 @@ public:
 
 	virtual ~FLandscapeXYOffsetVertexFactory() {}
 
-	static void ModifyCompilationEnvironment(const FVertexFactoryType* Type, EShaderPlatform Platform, const FMaterial* Material, FShaderCompilerEnvironment& OutEnvironment);
+	static void ModifyCompilationEnvironment(const FVertexFactoryShaderPermutationParameters& Parameters, FShaderCompilerEnvironment& OutEnvironment);
 };
 
 
@@ -267,8 +255,7 @@ public:
 	{
 	}
 
-	static void ModifyCompilationEnvironment(const FVertexFactoryType* Type, EShaderPlatform Platform, const FMaterial* Material, FShaderCompilerEnvironment& OutEnvironment);
-	static FVertexFactoryShaderParameters* ConstructShaderParameters(EShaderFrequency ShaderFrequency);
+	static void ModifyCompilationEnvironment(const FVertexFactoryShaderPermutationParameters& Parameters, FShaderCompilerEnvironment& OutEnvironment);
 };
 
 
@@ -1098,7 +1085,7 @@ public:
 		return Parent->GetMaterialWithFallback(InFeatureLevel, OutFallbackMaterialRenderProxy);
 	}
 
-	virtual bool GetVectorValue(const FMaterialParameterInfo& ParameterInfo, FLinearColor* OutValue, const FMaterialRenderContext& Context) const
+	virtual bool GetVectorValue(const FHashedMaterialParameterInfo& ParameterInfo, FLinearColor* OutValue, const FMaterialRenderContext& Context) const override
 	{
 		if (ParameterInfo.Name == FName(TEXT("Landscape_RedMask")))
 		{
@@ -1120,11 +1107,11 @@ public:
 			return Parent->GetVectorValue(ParameterInfo, OutValue, Context);
 		}
 	}
-	virtual bool GetScalarValue(const FMaterialParameterInfo& ParameterInfo, float* OutValue, const FMaterialRenderContext& Context) const
+	virtual bool GetScalarValue(const FHashedMaterialParameterInfo& ParameterInfo, float* OutValue, const FMaterialRenderContext& Context) const override
 	{
 		return Parent->GetScalarValue(ParameterInfo, OutValue, Context);
 	}
-	virtual bool GetTextureValue(const FMaterialParameterInfo& ParameterInfo, const UTexture** OutValue, const FMaterialRenderContext& Context) const
+	virtual bool GetTextureValue(const FHashedMaterialParameterInfo& ParameterInfo, const UTexture** OutValue, const FMaterialRenderContext& Context) const override
 	{
 		// NOTE: These should be returning black textures when NULL. The material will
 		// use a white texture if they are.
@@ -1148,7 +1135,7 @@ public:
 			return Parent->GetTextureValue(ParameterInfo, OutValue, Context);
 		}
 	}
-	virtual bool GetTextureValue(const FMaterialParameterInfo& ParameterInfo, const URuntimeVirtualTexture** OutValue, const FMaterialRenderContext& Context) const
+	virtual bool GetTextureValue(const FHashedMaterialParameterInfo& ParameterInfo, const URuntimeVirtualTexture** OutValue, const FMaterialRenderContext& Context) const
 	{
 		return Parent->GetTextureValue(ParameterInfo, OutValue, Context);
 	}
@@ -1171,7 +1158,7 @@ public:
 	{
 		return Parent->GetMaterialWithFallback(InFeatureLevel, OutFallbackMaterialRenderProxy);
 	}
-	virtual bool GetVectorValue(const FMaterialParameterInfo& ParameterInfo, FLinearColor* OutValue, const FMaterialRenderContext& Context) const
+	virtual bool GetVectorValue(const FHashedMaterialParameterInfo& ParameterInfo, FLinearColor* OutValue, const FMaterialRenderContext& Context) const override
 	{
 		if (ParameterInfo.Name == FName(TEXT("HighlightColor")))
 		{
@@ -1183,11 +1170,11 @@ public:
 			return Parent->GetVectorValue(ParameterInfo, OutValue, Context);
 		}
 	}
-	virtual bool GetScalarValue(const FMaterialParameterInfo& ParameterInfo, float* OutValue, const FMaterialRenderContext& Context) const
+	virtual bool GetScalarValue(const FHashedMaterialParameterInfo& ParameterInfo, float* OutValue, const FMaterialRenderContext& Context) const override
 	{
 		return Parent->GetScalarValue(ParameterInfo, OutValue, Context);
 	}
-	virtual bool GetTextureValue(const FMaterialParameterInfo& ParameterInfo, const UTexture** OutValue, const FMaterialRenderContext& Context) const
+	virtual bool GetTextureValue(const FHashedMaterialParameterInfo& ParameterInfo, const UTexture** OutValue, const FMaterialRenderContext& Context) const override
 	{
 		if (ParameterInfo.Name == FName(TEXT("SelectedData")))
 		{
@@ -1199,7 +1186,7 @@ public:
 			return Parent->GetTextureValue(ParameterInfo, OutValue, Context);
 		}
 	}
-	virtual bool GetTextureValue(const FMaterialParameterInfo& ParameterInfo, const URuntimeVirtualTexture** OutValue, const FMaterialRenderContext& Context) const
+	virtual bool GetTextureValue(const FHashedMaterialParameterInfo& ParameterInfo, const URuntimeVirtualTexture** OutValue, const FMaterialRenderContext& Context) const
 	{
 		return Parent->GetTextureValue(ParameterInfo, OutValue, Context);
 	}
@@ -1224,11 +1211,11 @@ public:
 	{
 		return Parent->GetMaterialWithFallback(InFeatureLevel, OutFallbackMaterialRenderProxy);
 	}
-	virtual bool GetVectorValue(const FMaterialParameterInfo& ParameterInfo, FLinearColor* OutValue, const FMaterialRenderContext& Context) const
+	virtual bool GetVectorValue(const FHashedMaterialParameterInfo& ParameterInfo, FLinearColor* OutValue, const FMaterialRenderContext& Context) const override
 	{
 		return Parent->GetVectorValue(ParameterInfo, OutValue, Context);
 	}
-	virtual bool GetScalarValue(const FMaterialParameterInfo& ParameterInfo, float* OutValue, const FMaterialRenderContext& Context) const
+	virtual bool GetScalarValue(const FHashedMaterialParameterInfo& ParameterInfo, float* OutValue, const FMaterialRenderContext& Context) const override
 	{
 		if (ParameterInfo.Name == FName(TEXT("bInverted")))
 		{
@@ -1237,7 +1224,7 @@ public:
 		}
 		return Parent->GetScalarValue(ParameterInfo, OutValue, Context);
 	}
-	virtual bool GetTextureValue(const FMaterialParameterInfo& ParameterInfo, const UTexture** OutValue, const FMaterialRenderContext& Context) const
+	virtual bool GetTextureValue(const FHashedMaterialParameterInfo& ParameterInfo, const UTexture** OutValue, const FMaterialRenderContext& Context) const override
 	{
 		if (ParameterInfo.Name == FName(TEXT("SelectedData")))
 		{
@@ -1249,7 +1236,7 @@ public:
 			return Parent->GetTextureValue(ParameterInfo, OutValue, Context);
 		}
 	}
-	virtual bool GetTextureValue(const FMaterialParameterInfo& ParameterInfo, const URuntimeVirtualTexture** OutValue, const FMaterialRenderContext& Context) const
+	virtual bool GetTextureValue(const FHashedMaterialParameterInfo& ParameterInfo, const URuntimeVirtualTexture** OutValue, const FMaterialRenderContext& Context) const
 	{
 		return Parent->GetTextureValue(ParameterInfo, OutValue, Context);
 	}
@@ -1275,7 +1262,7 @@ public:
 	{
 		return Parent->GetMaterialWithFallback(InFeatureLevel, OutFallbackMaterialRenderProxy);
 	}
-	virtual bool GetVectorValue(const FMaterialParameterInfo& ParameterInfo, FLinearColor* OutValue, const FMaterialRenderContext& Context) const
+	virtual bool GetVectorValue(const FHashedMaterialParameterInfo& ParameterInfo, FLinearColor* OutValue, const FMaterialRenderContext& Context) const override
 	{
 		static FName ColorNames[] =
 		{
@@ -1307,7 +1294,7 @@ public:
 		}
 		return Parent->GetVectorValue(ParameterInfo, OutValue, Context);
 	}
-	virtual bool GetScalarValue(const FMaterialParameterInfo& ParameterInfo, float* OutValue, const FMaterialRenderContext& Context) const
+	virtual bool GetScalarValue(const FHashedMaterialParameterInfo& ParameterInfo, float* OutValue, const FMaterialRenderContext& Context) const override
 	{
 		if (ParameterInfo.Name == FName(TEXT("Rotation")))
 		{
@@ -1326,11 +1313,11 @@ public:
 		}		
 		return Parent->GetScalarValue(ParameterInfo, OutValue, Context);
 	}
-	virtual bool GetTextureValue(const FMaterialParameterInfo& ParameterInfo, const UTexture** OutValue, const FMaterialRenderContext& Context) const
+	virtual bool GetTextureValue(const FHashedMaterialParameterInfo& ParameterInfo, const UTexture** OutValue, const FMaterialRenderContext& Context) const override
 	{
 		return Parent->GetTextureValue(ParameterInfo, OutValue, Context);
 	}
-	virtual bool GetTextureValue(const FMaterialParameterInfo& ParameterInfo, const URuntimeVirtualTexture** OutValue, const FMaterialRenderContext& Context) const
+	virtual bool GetTextureValue(const FHashedMaterialParameterInfo& ParameterInfo, const URuntimeVirtualTexture** OutValue, const FMaterialRenderContext& Context) const
 	{
 		return Parent->GetTextureValue(ParameterInfo, OutValue, Context);
 	}
diff --git a/Engine/Source/Runtime/Landscape/Public/LandscapeRenderMobile.h b/Engine/Source/Runtime/Landscape/Public/LandscapeRenderMobile.h
--- a/Engine/Source/Runtime/Landscape/Public/LandscapeRenderMobile.h
+++ b/Engine/Source/Runtime/Landscape/Public/LandscapeRenderMobile.h
@@ -46,21 +46,14 @@ public:
 		ReleaseResource();
 	}
 
-	static FVertexFactoryShaderParameters* ConstructShaderParameters(EShaderFrequency ShaderFrequency);
-
 	/**
 	* Should we cache the material's shadertype on this platform with this vertex factory? 
 	*/
-	static bool ShouldCompilePermutation(EShaderPlatform Platform, const class FMaterial* Material, const class FShaderType* ShaderType)
-	{
-		auto FeatureLevel = GetMaxSupportedFeatureLevel(Platform);
-		return (FeatureLevel == ERHIFeatureLevel::ES2 || FeatureLevel == ERHIFeatureLevel::ES3_1) &&
-			(Material->IsUsedWithLandscape() || Material->IsSpecialEngineMaterial());
-	}
+	static bool ShouldCompilePermutation(const FVertexFactoryShaderPermutationParameters& Parameters);
 
-	static void ModifyCompilationEnvironment( const FVertexFactoryType* Type, EShaderPlatform Platform, const FMaterial* Material, FShaderCompilerEnvironment& OutEnvironment )
+	static void ModifyCompilationEnvironment(const FVertexFactoryShaderPermutationParameters& Parameters, FShaderCompilerEnvironment& OutEnvironment )
 	{
-		FVertexFactory::ModifyCompilationEnvironment(Type, Platform, Material, OutEnvironment);
+		FVertexFactory::ModifyCompilationEnvironment(Parameters, OutEnvironment);
 		OutEnvironment.SetDefine(TEXT("NUM_VF_PACKED_INTERPOLANTS"), TEXT("1"));
 	}
 
@@ -148,7 +141,6 @@ public:
 	{
 	}
 
-	static void ModifyCompilationEnvironment(const FVertexFactoryType* Type, EShaderPlatform Platform, const FMaterial* Material, FShaderCompilerEnvironment& OutEnvironment);
-	static FVertexFactoryShaderParameters* ConstructShaderParameters(EShaderFrequency ShaderFrequency);
-	static bool ShouldCompilePermutation(EShaderPlatform Platform, const FMaterial* Material, const FShaderType* ShaderType);
+	static void ModifyCompilationEnvironment(const FVertexFactoryShaderPermutationParameters& Parameters, FShaderCompilerEnvironment& OutEnvironment);
+	static bool ShouldCompilePermutation(const FVertexFactoryShaderPermutationParameters& Parameters);
 };
diff --git a/Engine/Source/Runtime/Launch/Private/LaunchEngineLoop.cpp b/Engine/Source/Runtime/Launch/Private/LaunchEngineLoop.cpp
--- a/Engine/Source/Runtime/Launch/Private/LaunchEngineLoop.cpp
+++ b/Engine/Source/Runtime/Launch/Private/LaunchEngineLoop.cpp
@@ -1313,6 +1313,8 @@ DECLARE_CYCLE_STAT(TEXT("FEngineLoop::PreInitPostStartupScreen.AfterStats"), STA
 
 int32 FEngineLoop::PreInitPreStartupScreen(const TCHAR* CmdLine)
 {
+	FDelayedAutoRegisterHelper::RunAndClearDelayedAutoRegisterDelegates(EDelayedRegisterRunPhase::StartOfEnginePreInit);
+
 #if UE_TRACE_ENABLED
 	{
 		Trace::Initialize();
@@ -1333,7 +1335,7 @@ int32 FEngineLoop::PreInitPreStartupScreen(const TCHAR* CmdLine)
 	}
 #endif
 
-	SCOPED_BOOT_TIMING("FEngineLoop::PreInit");
+	SCOPED_BOOT_TIMING("FEngineLoop::PreInitPreStartupScreen");
 
 	// The GLog singleton is lazy initialised and by default will assume that
 	// its "master thread" is the one it was created on. This lazy initialisation
@@ -1657,6 +1659,15 @@ int32 FEngineLoop::PreInitPreStartupScreen(const TCHAR* CmdLine)
 		IFileManager::Get().ProcessCommandLineOptions();
 	}
 
+#if WITH_COREUOBJECT
+	{
+		SCOPED_BOOT_TIMING("InitializeNewAsyncIO");
+		FPlatformFileManager::Get().InitializeNewAsyncIO();
+	}
+#endif
+
+	FDelayedAutoRegisterHelper::RunAndClearDelayedAutoRegisterDelegates(EDelayedRegisterRunPhase::FileSystemReady);
+
 	if (GIsGameAgnosticExe)
 	{
 		// If we launched without a project file, but with a game name that is incomplete, warn about the improper use of a Game suffix
@@ -1994,10 +2005,14 @@ int32 FEngineLoop::PreInitPreStartupScreen(const TCHAR* CmdLine)
 		FTaskGraphInterface::Get().AttachToThread(ENamedThreads::GameThread);
 	}
 
+	FDelayedAutoRegisterHelper::RunAndClearDelayedAutoRegisterDelegates(EDelayedRegisterRunPhase::TaskGraphSystemReady);
+
 #if STATS
 	FThreadStats::StartThread();
 #endif
 
+	FDelayedAutoRegisterHelper::RunAndClearDelayedAutoRegisterDelegates(EDelayedRegisterRunPhase::StatSystemReady);
+
 	FScopeCycleCounter CycleCount_AfterStats(GET_STATID(STAT_FEngineLoop_PreInitPreStartupScreen_AfterStats));
 
 	// Load Core modules required for everything else to work (needs to be loaded before InitializeRenderingCVarsCaching)
@@ -2153,13 +2168,6 @@ int32 FEngineLoop::PreInitPreStartupScreen(const TCHAR* CmdLine)
 		}
 	}
 
-#if WITH_COREUOBJECT
-	{
-		SCOPED_BOOT_TIMING("InitializeNewAsyncIO");
-		FPlatformFileManager::Get().InitializeNewAsyncIO();
-	}
-#endif
-
 	if (FPlatformProcess::SupportsMultithreading())
 	{
 		{
@@ -2591,6 +2599,8 @@ int32 FEngineLoop::PreInitPreStartupScreen(const TCHAR* CmdLine)
 			InitializeShaderTypes();
 		}
 
+		FDelayedAutoRegisterHelper::RunAndClearDelayedAutoRegisterDelegates(EDelayedRegisterRunPhase::ShaderTypesReady);
+
 		SlowTask.EnterProgressFrame(30);
 
 		// Load the global shaders.
@@ -2693,9 +2703,11 @@ int32 FEngineLoop::PreInitPreStartupScreen(const TCHAR* CmdLine)
 
 					if (FPreLoadScreenManager::Get())
 					{
-						SCOPED_BOOT_TIMING("PlayFirstPreLoadScreen - FPreLoadScreenManager::Get()->Initialize");
-						// initialize and present custom splash screen
-						FPreLoadScreenManager::Get()->Initialize(SlateRendererSharedRef.Get());
+						{
+							SCOPED_BOOT_TIMING("PlayFirstPreLoadScreen - FPreLoadScreenManager::Get()->Initialize");
+							// initialize and present custom splash screen
+							FPreLoadScreenManager::Get()->Initialize(SlateRendererSharedRef.Get());
+						}
 
 						if (FPreLoadScreenManager::Get()->HasRegisteredPreLoadScreenType(EPreLoadScreenTypes::CustomSplashScreen))
 						{
@@ -2739,6 +2751,8 @@ int32 FEngineLoop::PreInitPreStartupScreen(const TCHAR* CmdLine)
 
 int32 FEngineLoop::PreInitPostStartupScreen(const TCHAR* CmdLine)
 {
+	SCOPED_BOOT_TIMING("FEngineLoop::PreInitPostStartupScreen");
+
 	if (IsEngineExitRequested())
 	{
 		return 0;
@@ -2769,6 +2783,8 @@ int32 FEngineLoop::PreInitPostStartupScreen(const TCHAR* CmdLine)
 #if !UE_SERVER// && !UE_EDITOR
 		if (!IsRunningDedicatedServer() && !IsRunningCommandlet())
 		{
+			SCOPED_BOOT_TIMING("PreInitPostStartupScreen_StartupGraphics");
+
 			TSharedPtr<FSlateRenderer> SlateRenderer = PreInitContext.SlateRenderer;
 			TSharedRef<FSlateRenderer> SlateRendererSharedRef = SlateRenderer.ToSharedRef();
 			{
@@ -2966,9 +2982,15 @@ int32 FEngineLoop::PreInitPostStartupScreen(const TCHAR* CmdLine)
 
 		FEmbeddedCommunication::ForceTick(5);
 
+		// for any auto-registered functions that want to wait until main(), run them now
+		// @todo loadtime: this should have phases, so caller can decide when auto-register runs [use plugin phases probably?]
+		FDelayedAutoRegisterHelper::RunAndClearDelayedAutoRegisterDelegates(EDelayedRegisterRunPhase::PreObjectSystemReady);
+
 		// Make sure all UObject classes are registered and default properties have been initialized
 		ProcessNewlyLoadedUObjects();
 
+		FDelayedAutoRegisterHelper::RunAndClearDelayedAutoRegisterDelegates(EDelayedRegisterRunPhase::ObjectSystemReady);
+
 		FEmbeddedCommunication::ForceTick(6);
 
 #if WITH_EDITOR
@@ -4050,6 +4072,8 @@ int32 FEngineLoop::Init()
 	}
 #endif		// UE_EXTERNAL_PROFILING_ENABLED
 
+	FDelayedAutoRegisterHelper::RunAndClearDelayedAutoRegisterDelegates(EDelayedRegisterRunPhase::EndOfEngineInit);
+
 	return 0;
 }
 
@@ -5278,8 +5302,8 @@ bool FEngineLoop::AppInit( )
 
 #if PLATFORM_WINDOWS
 
-	// make sure that the log directory exists
-	IFileManager::Get().MakeDirectory( *FPaths::ProjectLogDir() );
+	// make sure that the log directory tree exists
+	IFileManager::Get().MakeDirectory( *FPaths::ProjectLogDir(), true );
 
 	// update the mini dump filename now that we have enough info to point it to the log folder even in installed builds
 	FCString::Strcpy(MiniDumpFilenameW, *IFileManager::Get().ConvertToAbsolutePathForExternalAppForWrite(*FString::Printf(TEXT("%sunreal-v%i-%s.dmp"), *FPaths::ProjectLogDir(), FEngineVersion::Current().GetChangelist(), *FDateTime::Now().ToString())));
@@ -5309,6 +5333,8 @@ bool FEngineLoop::AppInit( )
 		FConfigCacheIni::InitializeConfigSystem();
 	}
 
+	FDelayedAutoRegisterHelper::RunAndClearDelayedAutoRegisterDelegates(EDelayedRegisterRunPhase::IniSystemReady);
+
 	// Load "asap" plugin modules
 	IPluginManager&  PluginManager = IPluginManager::Get();
 	IProjectManager& ProjectManager = IProjectManager::Get();
diff --git a/Engine/Source/Runtime/MRMesh/Private/MRMeshComponent.cpp b/Engine/Source/Runtime/MRMesh/Private/MRMeshComponent.cpp
--- a/Engine/Source/Runtime/MRMesh/Private/MRMeshComponent.cpp
+++ b/Engine/Source/Runtime/MRMesh/Private/MRMeshComponent.cpp
@@ -468,7 +468,7 @@ private:
 		Result.bRenderInDepthPass = bEnableOcclusion;
 		Result.bUsesLightingChannels = GetLightingChannelMask() != GetDefaultLightingChannelMask();
 		Result.bRenderCustomDepth = ShouldRenderCustomDepth();
-		UMaterialInterface::TMicRecursionGuard RecursionGuard;
+		TMicRecursionGuard RecursionGuard;
 		Result.bSeparateTranslucency = MaterialToUse->GetMaterial_Concurrent(RecursionGuard)->bEnableSeparateTranslucency;
 		//MaterialRelevance.SetPrimitiveViewRelevance(Result);
 		return Result;
diff --git a/Engine/Source/Runtime/MovieSceneCapture/Private/FrameGrabber.cpp b/Engine/Source/Runtime/MovieSceneCapture/Private/FrameGrabber.cpp
--- a/Engine/Source/Runtime/MovieSceneCapture/Private/FrameGrabber.cpp
+++ b/Engine/Source/Runtime/MovieSceneCapture/Private/FrameGrabber.cpp
@@ -139,13 +139,13 @@ void FViewportSurfaceReader::ResolveRenderTarget(FViewportSurfaceReader* RenderT
 
 			const ERHIFeatureLevel::Type FeatureLevel = GMaxRHIFeatureLevel;
 
-			TShaderMap<FGlobalShaderType>* ShaderMap = GetGlobalShaderMap(FeatureLevel);
+			FGlobalShaderMap* ShaderMap = GetGlobalShaderMap(FeatureLevel);
 			TShaderMapRef<FScreenVS> VertexShader(ShaderMap);
 			TShaderMapRef<FScreenPS> PixelShader(ShaderMap);
 
 			GraphicsPSOInit.BoundShaderState.VertexDeclarationRHI = GFilterVertexDeclaration.VertexDeclarationRHI;
-			GraphicsPSOInit.BoundShaderState.VertexShaderRHI = GETSAFERHISHADER_VERTEX(*VertexShader);
-			GraphicsPSOInit.BoundShaderState.PixelShaderRHI = GETSAFERHISHADER_PIXEL(*PixelShader);
+			GraphicsPSOInit.BoundShaderState.VertexShaderRHI = VertexShader.GetVertexShader();
+			GraphicsPSOInit.BoundShaderState.PixelShaderRHI = PixelShader.GetPixelShader();
 			GraphicsPSOInit.PrimitiveType = PT_TriangleList;
 
 			SetGraphicsPipelineState(RHICmdList, GraphicsPSOInit);
@@ -176,7 +176,7 @@ void FViewportSurfaceReader::ResolveRenderTarget(FViewportSurfaceReader* RenderT
 				SizeU, SizeV,							// Source USize, VSize
 				CaptureRect.Max - CaptureRect.Min,		// Target buffer size
 				FIntPoint(1, 1),						// Source texture size
-				*VertexShader,
+				VertexShader,
 				EDRF_Default);
 		}
 		RHICmdList.EndRenderPass();
diff --git a/Engine/Source/Runtime/NavigationSystem/Private/NavMesh/NavMeshRenderingComponent.cpp b/Engine/Source/Runtime/NavigationSystem/Private/NavMesh/NavMeshRenderingComponent.cpp
--- a/Engine/Source/Runtime/NavigationSystem/Private/NavMesh/NavMeshRenderingComponent.cpp
+++ b/Engine/Source/Runtime/NavigationSystem/Private/NavMesh/NavMeshRenderingComponent.cpp
@@ -1351,9 +1351,9 @@ FPrimitiveSceneProxy* UNavMeshRenderingComponent::CreateSceneProxy()
 #endif //WITH_RECAST && !UE_BUILD_SHIPPING && !UE_BUILD_TEST
 }
 
-void UNavMeshRenderingComponent::CreateRenderState_Concurrent()
+void UNavMeshRenderingComponent::CreateRenderState_Concurrent(FRegisterComponentContext* Context)
 {
-	Super::CreateRenderState_Concurrent();
+	Super::CreateRenderState_Concurrent(Context);
 
 #if WITH_RECAST && !UE_BUILD_SHIPPING && !UE_BUILD_TEST
 	NavMeshDebugDrawDelgateManager.RegisterDebugDrawDelgate();
diff --git a/Engine/Source/Runtime/NavigationSystem/Private/NavMesh/NavTestRenderingComponent.cpp b/Engine/Source/Runtime/NavigationSystem/Private/NavMesh/NavTestRenderingComponent.cpp
--- a/Engine/Source/Runtime/NavigationSystem/Private/NavMesh/NavTestRenderingComponent.cpp
+++ b/Engine/Source/Runtime/NavigationSystem/Private/NavMesh/NavTestRenderingComponent.cpp
@@ -461,9 +461,9 @@ FBoxSphereBounds UNavTestRenderingComponent::CalcBounds(const FTransform& LocalT
 	return FBoxSphereBounds(BoundingBox);
 }
 
-void UNavTestRenderingComponent::CreateRenderState_Concurrent()
+void UNavTestRenderingComponent::CreateRenderState_Concurrent(FRegisterComponentContext* Context)
 {
-	Super::CreateRenderState_Concurrent();
+	Super::CreateRenderState_Concurrent(Context);
 
 #if WITH_RECAST && WITH_EDITOR
 	NavTestDebugDrawDelegateHelper.RegisterDebugDrawDelgate();
diff --git a/Engine/Source/Runtime/NonRealtimeAudioRenderer/Private/AudioMixerPlatformNonRealtime.cpp b/Engine/Source/Runtime/NonRealtimeAudioRenderer/Private/AudioMixerPlatformNonRealtime.cpp
--- a/Engine/Source/Runtime/NonRealtimeAudioRenderer/Private/AudioMixerPlatformNonRealtime.cpp
+++ b/Engine/Source/Runtime/NonRealtimeAudioRenderer/Private/AudioMixerPlatformNonRealtime.cpp
@@ -317,13 +317,13 @@ namespace Audio
 		static FName NAME_XMA(TEXT("XMA"));
 
 #if WITH_XMA2 && USE_XMA2_FOR_STREAMING
-		if (InSoundWave->IsStreaming() && InSoundWave->NumChannels <= 2)
+		if (InSoundWave->IsStreaming(nullptr) && InSoundWave->NumChannels <= 2)
 		{
 			return NAME_XMA;
 		}
 #endif
 
-		if (InSoundWave->IsStreaming())
+		if (InSoundWave->IsStreaming(nullptr))
 		{
 #if USE_VORBIS_FOR_STREAMING
 			return NAME_OGG;
@@ -364,7 +364,7 @@ namespace Audio
 		}
 #endif
 
-		if (InSoundWave->IsStreaming())
+		if (InSoundWave->IsStreaming(nullptr))
 		{
 #if USE_VORBIS_FOR_STREAMING
 			return new FVorbisAudioInfo();
diff --git a/Engine/Source/Runtime/NullDrv/Public/NullRHI.h b/Engine/Source/Runtime/NullDrv/Public/NullRHI.h
--- a/Engine/Source/Runtime/NullDrv/Public/NullRHI.h
+++ b/Engine/Source/Runtime/NullDrv/Public/NullRHI.h
@@ -42,32 +42,32 @@ public:
 		return new FRHIVertexDeclaration(); 
 	}
 
-	virtual FPixelShaderRHIRef RHICreatePixelShader(const TArray<uint8>& Code) final override
+	virtual FPixelShaderRHIRef RHICreatePixelShader(TArrayView<const uint8> Code, const FSHAHash& Hash) final override
 	{ 
 		return new FRHIPixelShader(); 
 	}
 
-	virtual FVertexShaderRHIRef RHICreateVertexShader(const TArray<uint8>& Code) final override
+	virtual FVertexShaderRHIRef RHICreateVertexShader(TArrayView<const uint8> Code, const FSHAHash& Hash) final override
 	{ 
 		return new FRHIVertexShader(); 
 	}
 
-	virtual FHullShaderRHIRef RHICreateHullShader(const TArray<uint8>& Code) final override
+	virtual FHullShaderRHIRef RHICreateHullShader(TArrayView<const uint8> Code, const FSHAHash& Hash) final override
 	{ 
 		return new FRHIHullShader(); 
 	}
 
-	virtual FDomainShaderRHIRef RHICreateDomainShader(const TArray<uint8>& Code) final override
+	virtual FDomainShaderRHIRef RHICreateDomainShader(TArrayView<const uint8> Code, const FSHAHash& Hash) final override
 	{ 
 		return new FRHIDomainShader(); 
 	}
 
-	virtual FGeometryShaderRHIRef RHICreateGeometryShader(const TArray<uint8>& Code) final override
+	virtual FGeometryShaderRHIRef RHICreateGeometryShader(TArrayView<const uint8> Code, const FSHAHash& Hash) final override
 	{ 
 		return new FRHIGeometryShader(); 
 	}
 
-	virtual FComputeShaderRHIRef RHICreateComputeShader(const TArray<uint8>& Code) final override
+	virtual FComputeShaderRHIRef RHICreateComputeShader(TArrayView<const uint8> Code, const FSHAHash& Hash) final override
 	{ 
 		return new FRHIComputeShader(); 
 	}
diff --git a/Engine/Source/Runtime/Online/HTTPServer/Private/HttpListener.cpp b/Engine/Source/Runtime/Online/HTTPServer/Private/HttpListener.cpp
--- a/Engine/Source/Runtime/Online/HTTPServer/Private/HttpListener.cpp
+++ b/Engine/Source/Runtime/Online/HTTPServer/Private/HttpListener.cpp
@@ -9,6 +9,10 @@
 #include "SocketSubsystem.h"
 #include "IPAddress.h"
 
+constexpr uint32 FHttpListener::MaxConnectionsToAcceptPerFrame;
+constexpr uint32 FHttpListener::ListenerConnectionBacklogSize;
+constexpr uint32 FHttpListener::ListenerBufferSize;
+
 DEFINE_LOG_CATEGORY(LogHttpListener)
 
 FHttpListener::FHttpListener(uint32 Port)
diff --git a/Engine/Source/Runtime/OpenGLDrv/Private/OpenGLCommands.cpp b/Engine/Source/Runtime/OpenGLDrv/Private/OpenGLCommands.cpp
--- a/Engine/Source/Runtime/OpenGLDrv/Private/OpenGLCommands.cpp
+++ b/Engine/Source/Runtime/OpenGLDrv/Private/OpenGLCommands.cpp
@@ -276,29 +276,29 @@ static FORCEINLINE GLint ModifyFilterByMips(GLint Filter, bool bHasMips)
 
 static FORCEINLINE EShaderFrequency GetShaderFrequency(FRHIGraphicsShader* ShaderRHI)
 {
-	switch (ShaderRHI->GetType())
+	switch (ShaderRHI->GetFrequency())
 	{
-	case FRHIShader::EType::Vertex:
+	case SF_Vertex:
 		VALIDATE_BOUND_SHADER(ShaderRHI, Vertex);
 		return SF_Vertex;
 #if PLATFORM_SUPPORTS_TESSELLATION_SHADERS
-	case FRHIShader::EType::Hull:
+	case SF_Hull:
 		VALIDATE_BOUND_SHADER(ShaderRHI, Hull);
 		return SF_Hull;
-	case FRHIShader::EType::Domain:
+	case SF_Domain:
 		VALIDATE_BOUND_SHADER(ShaderRHI, Domain);
 		return SF_Domain;
 #endif
 #if PLATFORM_SUPPORTS_GEOMETRY_SHADERS
-	case FRHIShader::EType::Geometry:
+	case SF_Geometry:
 		VALIDATE_BOUND_SHADER(ShaderRHI, Geometry);
 		return SF_Geometry;
 #endif
-	case FRHIShader::EType::Pixel:
+	case SF_Pixel:
 		VALIDATE_BOUND_SHADER(ShaderRHI, Pixel);
 		return SF_Pixel;
 	default:
-		checkf(0, TEXT("Undefined FRHIShader Type %d!"), (int32)ShaderRHI->GetType());
+		checkf(0, TEXT("Undefined FRHIShader Type %d!"), (int32)ShaderRHI->GetFrequency());
 	}
 
 	return SF_NumFrequencies;
@@ -306,29 +306,29 @@ static FORCEINLINE EShaderFrequency GetShaderFrequency(FRHIGraphicsShader* Shade
 
 static FORCEINLINE CrossCompiler::EShaderStage GetShaderCrossCompilerStage(FRHIGraphicsShader* ShaderRHI)
 {
-	switch (ShaderRHI->GetType())
+	switch (ShaderRHI->GetFrequency())
 	{
-	case FRHIShader::EType::Vertex:
+	case SF_Vertex:
 		VALIDATE_BOUND_SHADER(ShaderRHI, Vertex);
 		return CrossCompiler::SHADER_STAGE_VERTEX;
 #if PLATFORM_SUPPORTS_TESSELLATION_SHADERS
-	case FRHIShader::EType::Hull:
+	case SF_Hull:
 		VALIDATE_BOUND_SHADER(ShaderRHI, Hull);
 		return CrossCompiler::SHADER_STAGE_HULL;
-	case FRHIShader::EType::Domain:
+	case SF_Domain:
 		VALIDATE_BOUND_SHADER(ShaderRHI, Domain);
 		return CrossCompiler::SHADER_STAGE_DOMAIN;
 #endif
 #if PLATFORM_SUPPORTS_GEOMETRY_SHADERS
-	case FRHIShader::EType::Geometry:
+	case SF_Geometry:
 		VALIDATE_BOUND_SHADER(ShaderRHI, Geometry);
 		return CrossCompiler::SHADER_STAGE_GEOMETRY;
 #endif
-	case FRHIShader::EType::Pixel:
+	case SF_Pixel:
 		VALIDATE_BOUND_SHADER(ShaderRHI, Pixel);
 		return CrossCompiler::SHADER_STAGE_PIXEL;
 	default:
-		checkf(0, TEXT("Undefined FRHIShader Type %d!"), (int32)ShaderRHI->GetType());
+		checkf(0, TEXT("Undefined FRHIShader Type %d!"), (int32)ShaderRHI->GetFrequency());
 	}
 
 	return CrossCompiler::NUM_SHADER_STAGES;
@@ -336,39 +336,39 @@ static FORCEINLINE CrossCompiler::EShaderStage GetShaderCrossCompilerStage(FRHIG
 
 static FORCEINLINE void GetShaderStageIndexAndMaxUnits(FRHIGraphicsShader* ShaderRHI, GLint& OutIndex, GLint& OutMaxUnits)
 {
-	switch (ShaderRHI->GetType())
+	switch (ShaderRHI->GetFrequency())
 	{
-	case FRHIShader::EType::Vertex:
+	case SF_Vertex:
 		VALIDATE_BOUND_SHADER(ShaderRHI, Vertex);
 		OutIndex = FOpenGL::GetFirstVertexTextureUnit();
 		OutMaxUnits = FOpenGL::GetMaxVertexTextureImageUnits();
 		break;
 #if PLATFORM_SUPPORTS_TESSELLATION_SHADERS
-	case FRHIShader::EType::Hull:
+	case SF_Hull:
 		VALIDATE_BOUND_SHADER(ShaderRHI, Hull);
 		OutIndex = FOpenGL::GetFirstHullTextureUnit();
 		OutMaxUnits = FOpenGL::GetMaxHullTextureImageUnits();
 		break;
-	case FRHIShader::EType::Domain:
+	case SF_Domain:
 		VALIDATE_BOUND_SHADER(ShaderRHI, Domain);
 		OutIndex = FOpenGL::GetFirstDomainTextureUnit();
 		OutMaxUnits = FOpenGL::GetMaxDomainTextureImageUnits();
 		break;
 #endif
 #if PLATFORM_SUPPORTS_GEOMETRY_SHADERS
-	case FRHIShader::EType::Geometry:
+	case SF_Geometry:
 		VALIDATE_BOUND_SHADER(ShaderRHI, Geometry);
 		OutIndex = FOpenGL::GetFirstGeometryTextureUnit();
 		OutMaxUnits = FOpenGL::GetMaxGeometryTextureImageUnits();
 		break;
 #endif
-	case FRHIShader::EType::Pixel:
+	case SF_Pixel:
 		VALIDATE_BOUND_SHADER(ShaderRHI, Pixel);
 		OutIndex = FOpenGL::GetFirstPixelTextureUnit();
 		OutMaxUnits = FOpenGL::GetMaxTextureImageUnits();
 		break;
 	default:
-		checkf(0, TEXT("Undefined FRHIShader Type %d!"), (int32)ShaderRHI->GetType());
+		checkf(0, TEXT("Undefined FRHIShader Type %d!"), (int32)ShaderRHI->GetFrequency());
 	}
 }
 
@@ -1007,7 +1007,7 @@ void FOpenGLDynamicRHI::RHISetShaderResourceViewParameter(FRHIGraphicsShader* Sh
 	GLint MaxUnits = 0;
 	GetShaderStageIndexAndMaxUnits(ShaderRHI, Index, MaxUnits);
 
-	ensureMsgf((int32)TextureIndex < MaxUnits, TEXT("Using more texture units (%d) than allowed (%d) on Frequency %d!"), TextureIndex, MaxUnits, (int32)ShaderRHI->GetType());
+	ensureMsgf((int32)TextureIndex < MaxUnits, TEXT("Using more texture units (%d) than allowed (%d) on Frequency %d!"), TextureIndex, MaxUnits, (int32)ShaderRHI->GetFrequency());
 	FOpenGLShaderResourceView* SRV = ResourceCast(SRVRHI);
 	GLuint Resource = 0;
 	GLenum Target = GL_TEXTURE_BUFFER;
@@ -1054,7 +1054,7 @@ void FOpenGLDynamicRHI::RHISetShaderTexture(FRHIGraphicsShader* ShaderRHI,uint32
 	GLint MaxUnits = 0;
 	GetShaderStageIndexAndMaxUnits(ShaderRHI, Index, MaxUnits);
 
-	ensureMsgf((int32)TextureIndex < MaxUnits, TEXT("Using more texture units (%d) than allowed (%d) on Frequency %d!"), TextureIndex, MaxUnits, (int32)ShaderRHI->GetType());
+	ensureMsgf((int32)TextureIndex < MaxUnits, TEXT("Using more texture units (%d) than allowed (%d) on Frequency %d!"), TextureIndex, MaxUnits, (int32)ShaderRHI->GetFrequency());
 	if (NewTexture)
 	{
 		InternalSetShaderTexture(NewTexture, nullptr, Index + TextureIndex, NewTexture->Target, NewTexture->Resource, NewTextureRHI->GetNumMips(), -1);
diff --git a/Engine/Source/Runtime/OpenGLDrv/Private/OpenGLDevice.cpp b/Engine/Source/Runtime/OpenGLDrv/Private/OpenGLDevice.cpp
--- a/Engine/Source/Runtime/OpenGLDrv/Private/OpenGLDevice.cpp
+++ b/Engine/Source/Runtime/OpenGLDrv/Private/OpenGLDevice.cpp
@@ -1390,7 +1390,7 @@ static void CheckVaryingLimit()
 
 		// Try to compile test shaders
 		GOpenGLShaderHackLastCompileSuccess = false;
-		FVertexShaderRHIRef VertexShaderRHI = RHICreateVertexShader(VertexShaderCode.GetReadAccess());
+		FVertexShaderRHIRef VertexShaderRHI = RHICreateVertexShader(VertexShaderCode.GetReadAccess(), FSHAHash());
 		FOpenGLVertexShader* OGLVertexShader = (FOpenGLVertexShader*)VertexShaderRHI.GetReference();
 
 		if (IsRunningRHIInSeparateThread())
@@ -1406,7 +1406,7 @@ static void CheckVaryingLimit()
 		}
 
 		GOpenGLShaderHackLastCompileSuccess = false;
-		FPixelShaderRHIRef PixelShaderRHI = RHICreatePixelShader(FragmentShaderCode.GetReadAccess());
+		FPixelShaderRHIRef PixelShaderRHI = RHICreatePixelShader(FragmentShaderCode.GetReadAccess(), FSHAHash());
 		FOpenGLPixelShader* OGLPixelShader = (FOpenGLPixelShader*)PixelShaderRHI.GetReference();
 
 		if (IsRunningRHIInSeparateThread())
@@ -1495,7 +1495,7 @@ static void CheckTextureCubeLodSupport()
 
 		// try to compile without any hacks
 		GOpenGLShaderHackLastCompileSuccess = false;
-		FPixelShaderRHIRef PixelShaderRHI = RHICreatePixelShader(Code);
+		FPixelShaderRHIRef PixelShaderRHI = RHICreatePixelShader(Code, FSHAHash());
 		FOpenGLPixelShader* OGLPixelShader = (FOpenGLPixelShader*)(PixelShaderRHI.GetReference());
 
 		if (IsRunningRHIInSeparateThread())
@@ -1518,7 +1518,7 @@ static void CheckTextureCubeLodSupport()
 		// try to compile without using precision for texture samplers
 		// Samsung Galaxy Express	Samsung Galaxy S3	Samsung Galaxy S3 mini	Samsung Galaxy Tab GT-P1000	Samsung Galaxy Tab 2
 		GOpenGLShaderHackLastCompileSuccess = false;
-		PixelShaderRHI = RHICreatePixelShader(Code);
+		PixelShaderRHI = RHICreatePixelShader(Code, FSHAHash());
 		OGLPixelShader = (FOpenGLPixelShader*)(PixelShaderRHI.GetReference());
 
 		if (IsRunningRHIInSeparateThread())
@@ -1540,7 +1540,7 @@ static void CheckTextureCubeLodSupport()
 
 		// third most likely Samsung Galaxy Tab GT-P1000
 		GOpenGLShaderHackLastCompileSuccess = false;
-		PixelShaderRHI = RHICreatePixelShader(Code);
+		PixelShaderRHI = RHICreatePixelShader(Code, FSHAHash());
 		OGLPixelShader = (FOpenGLPixelShader*)(PixelShaderRHI.GetReference());
 
 		if (IsRunningRHIInSeparateThread())
@@ -1561,7 +1561,7 @@ static void CheckTextureCubeLodSupport()
 
 		// try both hacks
 		GOpenGLShaderHackLastCompileSuccess = false;
-		PixelShaderRHI = RHICreatePixelShader(Code);
+		PixelShaderRHI = RHICreatePixelShader(Code, FSHAHash());
 		OGLPixelShader = (FOpenGLPixelShader*)(PixelShaderRHI.GetReference());
 		
 		if (IsRunningRHIInSeparateThread())
@@ -1617,7 +1617,7 @@ static void CheckRoundFunction()
 
 		// Try to compile test shaders
 		GOpenGLShaderHackLastCompileSuccess = false;
-		FVertexShaderRHIRef VertexShaderRHI = RHICreateVertexShader(VertexShaderCode.GetReadAccess());
+		FVertexShaderRHIRef VertexShaderRHI = RHICreateVertexShader(VertexShaderCode.GetReadAccess(), FSHAHash());
 		FOpenGLVertexShader* OGLVertexShader = (FOpenGLVertexShader*)VertexShaderRHI.GetReference();
 
 		if (IsRunningRHIInSeparateThread())
@@ -1747,14 +1747,7 @@ void FOpenGLDynamicRHI::Cleanup()
 		GPUProfilingData.Cleanup();
 
 		// Ask all initialized FRenderResources to release their RHI resources.
-		for(TLinkedList<FRenderResource*>::TIterator ResourceIt(FRenderResource::GetResourceList());ResourceIt;ResourceIt.Next())
-		{
-			ResourceIt->ReleaseRHI();
-		}
-		for(TLinkedList<FRenderResource*>::TIterator ResourceIt(FRenderResource::GetResourceList());ResourceIt;ResourceIt.Next())
-		{
-			ResourceIt->ReleaseDynamicRHI();
-		}
+		FRenderResource::ReleaseRHIForAllResources();
 	}
 
 	// Release dynamic vertex and index buffers.
diff --git a/Engine/Source/Runtime/OpenGLDrv/Private/OpenGLShaders.cpp b/Engine/Source/Runtime/OpenGLDrv/Private/OpenGLShaders.cpp
--- a/Engine/Source/Runtime/OpenGLDrv/Private/OpenGLShaders.cpp
+++ b/Engine/Source/Runtime/OpenGLDrv/Private/OpenGLShaders.cpp
@@ -147,7 +147,7 @@ void FOpenGLDynamicRHI::SetupRecursiveResources()
 	auto ShaderMap = GetGlobalShaderMap(GMaxRHIFeatureLevel);
 	{
 		TShaderMapRef<FNULLPS> PixelShader(ShaderMap);
-		PixelShader->GetPixelShader();
+		PixelShader.GetPixelShader();
 	}
 }
 
@@ -710,16 +710,46 @@ static typename TEnableIf<!TPointerIsConvertibleFromTo<TRHIType, const FRHIShade
  * @returns the compiled shader upon success.
  */
 template <typename ShaderType>
-ShaderType* CompileOpenGLShader(const TArray<uint8>& InShaderCode, const FSHAHash& LibraryHash, FRHIShader* RHIShader = nullptr)
+ShaderType* CompileOpenGLShader(TArrayView<const uint8> InShaderCode, const FSHAHash& LibraryHash, FRHIShader* RHIShader = nullptr)
 {
 	SCOPE_CYCLE_COUNTER(STAT_OpenGLShaderCompileTime);
 	VERIFY_GL_SCOPE();
 
+	ShaderType* Shader = nullptr;
+	{
+		FLibraryShaderCacheValue *Val = GetOpenGLCompiledLibraryShaderCache().Find(LibraryHash);
+		if (Val)
+		{
+			Shader = new ShaderType();
+			Shader->Resource = Val->GLShader;
+			Shader->Bindings = Val->Header->Bindings;
+			Shader->UniformBuffersCopyInfo = Val->Header->UniformBuffersCopyInfo;
+			if (FOpenGL::SupportsSeparateShaderObjects())
+			{
+				FSHAHash Hash;
+				// Just use the CRC - if it isn't being cached & logged we'll be dependent on the CRC alone anyway
+				FMemory::Memcpy(Hash.Hash, &Val->ShaderCrc, sizeof(uint32));
+				if (RHIShader)
+				{
+					SetShaderHash(Hash, RHIShader);
+				}
+				else
+				{
+					SetShaderHash(Hash, Shader);
+				}
+			}
+#if DEBUG_GL_SHADERS
+			Shader->GlslCode = Val->GlslCode;
+			Shader->GlslCodeString = (ANSICHAR*)Shader->GlslCode.GetData();
+#endif
+			return Shader;
+		}
+	}
+
 	FShaderCodeReader ShaderCode(InShaderCode);
 
-	ShaderType* Shader = nullptr;
 	const GLenum TypeEnum = ShaderType::TypeEnum;
-	FMemoryReader Ar(InShaderCode, true);
+	FMemoryReaderView Ar(InShaderCode, true);
 
 	Ar.SetLimitSize(ShaderCode.GetActualShaderCodeSize());
 
@@ -888,47 +918,6 @@ ShaderType* CompileOpenGLShader(const TArray<uint8>& InShaderCode, const FSHAHas
 	return Shader;
 }
 
-template <typename ShaderType>
-ShaderType* CompileOpenGLShader(FRHIShaderLibrary* Library, FSHAHash LibraryHash, FRHIShader* RHIShader = nullptr)
-{
-	FLibraryShaderCacheValue *Val = GetOpenGLCompiledLibraryShaderCache().Find(LibraryHash);
-	ShaderType* Shader = nullptr;
-	if (Val)
-	{
-		Shader = new ShaderType();
-		Shader->Resource = Val->GLShader;
-		Shader->Bindings = Val->Header->Bindings;
-		Shader->UniformBuffersCopyInfo = Val->Header->UniformBuffersCopyInfo;
-		if (FOpenGL::SupportsSeparateShaderObjects())
-		{
-			FSHAHash Hash;
-			// Just use the CRC - if it isn't being cached & logged we'll be dependent on the CRC alone anyway
-			FMemory::Memcpy(Hash.Hash, &Val->ShaderCrc, sizeof(uint32));
-			if (RHIShader)
-			{
-				SetShaderHash(Hash, RHIShader);
-			}
-			else
-			{
-				SetShaderHash(Hash, Shader);
-			}
-		}
-#if DEBUG_GL_SHADERS
-		Shader->GlslCode = Val->GlslCode;
-		Shader->GlslCodeString = (ANSICHAR*)Shader->GlslCode.GetData();
-#endif
-	}
-	else
-	{
-		TArray<uint8> InShaderCode;
-		bool bFound = Library->RequestEntry(LibraryHash, InShaderCode);
-		UE_CLOG(!bFound, LogRHI, Fatal, TEXT("Shader %s was supposed to be in a shader code library, however we looked for it later and it was not found."), *LibraryHash.ToString());
-		Shader = CompileOpenGLShader<ShaderType>(InShaderCode, LibraryHash, RHIShader);
-	}
-	return Shader;
-}
-
-
 void OPENGLDRV_API GetCurrentOpenGLShaderDeviceCapabilities(FOpenGLShaderDeviceCapabilities& Capabilities)
 {
 	FMemory::Memzero(Capabilities);
@@ -1443,52 +1432,52 @@ static ANSICHAR* SetIndex(ANSICHAR* Str, int32 Offset, int32 Index)
 }
 
 template<typename RHIType, typename TOGLProxyType>
-RHIType* CreateProxyShader(const TArray<uint8>& Code)
+RHIType* CreateProxyShader(TArrayView<const uint8> Code, const FSHAHash& Hash)
 {
 	FRHICommandListImmediate& RHICmdList = FRHICommandListExecutor::GetImmediateCommandList();
 	if (ShouldRunGLRenderContextOpOnThisThread(RHICmdList))
 	{
 		return new TOGLProxyType([&](RHIType* OwnerRHI)
 		{
-			return CompileOpenGLShader<typename TOGLProxyType::ContainedGLType>(Code, FSHAHash(), OwnerRHI);
+			return CompileOpenGLShader<typename TOGLProxyType::ContainedGLType>(Code, Hash, OwnerRHI);
 		});
 	}
 	else
 	{
 		// take a copy of the code for RHIT version.
-		TArray<uint8> CodeCopy = Code;
-		return new TOGLProxyType([Code = MoveTemp(CodeCopy)](RHIType* OwnerRHI)
+		TArray<uint8> CodeCopy(Code);
+		return new TOGLProxyType([Code = MoveTemp(CodeCopy), Hash](RHIType* OwnerRHI)
 		{
-			return CompileOpenGLShader<typename TOGLProxyType::ContainedGLType>(Code, FSHAHash(), OwnerRHI);
+			return CompileOpenGLShader<typename TOGLProxyType::ContainedGLType>(Code, Hash, OwnerRHI);
 		});
 	}
 }
 
-FVertexShaderRHIRef FOpenGLDynamicRHI::RHICreateVertexShader(const TArray<uint8>& Code)
+FVertexShaderRHIRef FOpenGLDynamicRHI::RHICreateVertexShader(TArrayView<const uint8> Code, const FSHAHash& Hash)
 {
-	return CreateProxyShader<FRHIVertexShader, FOpenGLVertexShaderProxy>(Code);
+	return CreateProxyShader<FRHIVertexShader, FOpenGLVertexShaderProxy>(Code, Hash);
 }
 
-FPixelShaderRHIRef FOpenGLDynamicRHI::RHICreatePixelShader(const TArray<uint8>& Code)
+FPixelShaderRHIRef FOpenGLDynamicRHI::RHICreatePixelShader(TArrayView<const uint8> Code, const FSHAHash& Hash)
 {
-	return CreateProxyShader<FRHIPixelShader, FOpenGLPixelShaderProxy>(Code);
+	return CreateProxyShader<FRHIPixelShader, FOpenGLPixelShaderProxy>(Code, Hash);
 }
 
-FGeometryShaderRHIRef FOpenGLDynamicRHI::RHICreateGeometryShader(const TArray<uint8>& Code)
+FGeometryShaderRHIRef FOpenGLDynamicRHI::RHICreateGeometryShader(TArrayView<const uint8> Code, const FSHAHash& Hash)
 {
-	return CreateProxyShader<FRHIGeometryShader, FOpenGLGeometryShaderProxy>(Code);
+	return CreateProxyShader<FRHIGeometryShader, FOpenGLGeometryShaderProxy>(Code, Hash);
 }
 
-FHullShaderRHIRef FOpenGLDynamicRHI::RHICreateHullShader(const TArray<uint8>& Code)
+FHullShaderRHIRef FOpenGLDynamicRHI::RHICreateHullShader(TArrayView<const uint8> Code, const FSHAHash& Hash)
 {
 	check(GMaxRHIFeatureLevel >= ERHIFeatureLevel::SM5);
-	return CreateProxyShader<FRHIHullShader, FOpenGLHullShaderProxy>(Code);
+	return CreateProxyShader<FRHIHullShader, FOpenGLHullShaderProxy>(Code, Hash);
 }
 
-FDomainShaderRHIRef FOpenGLDynamicRHI::RHICreateDomainShader(const TArray<uint8>& Code)
+FDomainShaderRHIRef FOpenGLDynamicRHI::RHICreateDomainShader(TArrayView<const uint8> Code, const FSHAHash& Hash)
 {
 	check(GMaxRHIFeatureLevel >= ERHIFeatureLevel::SM5);
-	return CreateProxyShader<FRHIDomainShader, FOpenGLDomainShaderProxy>(Code);
+	return CreateProxyShader<FRHIDomainShader, FOpenGLDomainShaderProxy>(Code, Hash);
 }
 
 template<typename RHIType, typename TOGLProxyType>
@@ -1512,33 +1501,6 @@ RHIType* CreateProxyShader(FRHIShaderLibrary* Library, FSHAHash Hash)
 	}
 }
 
-FVertexShaderRHIRef FOpenGLDynamicRHI::RHICreateVertexShader(FRHIShaderLibrary* Library, FSHAHash Hash)
-{
-	return CreateProxyShader<FRHIVertexShader, FOpenGLVertexShaderProxy>(Library, Hash);
-}
-
-FPixelShaderRHIRef FOpenGLDynamicRHI::RHICreatePixelShader(FRHIShaderLibrary* Library, FSHAHash Hash)
-{
-	return CreateProxyShader<FRHIPixelShader, FOpenGLPixelShaderProxy>(Library, Hash);
-}
-
-FGeometryShaderRHIRef FOpenGLDynamicRHI::RHICreateGeometryShader(FRHIShaderLibrary* Library, FSHAHash Hash)
-{
-	return CreateProxyShader<FRHIGeometryShader, FOpenGLGeometryShaderProxy>(Library, Hash);
-}
-
-FHullShaderRHIRef FOpenGLDynamicRHI::RHICreateHullShader(FRHIShaderLibrary* Library, FSHAHash Hash)
-{
-	check(GMaxRHIFeatureLevel >= ERHIFeatureLevel::SM5);
-	return CreateProxyShader<FRHIHullShader, FOpenGLHullShaderProxy>(Library, Hash);
-}
-
-FDomainShaderRHIRef FOpenGLDynamicRHI::RHICreateDomainShader(FRHIShaderLibrary* Library, FSHAHash Hash)
-{
-	check(GMaxRHIFeatureLevel >= ERHIFeatureLevel::SM5);
-	return CreateProxyShader<FRHIDomainShader, FOpenGLDomainShaderProxy>(Library, Hash);
-}
-
 static void MarkShaderParameterCachesDirty(FOpenGLShaderParameterCache* ShaderParameters, bool UpdateCompute)
 {
 	VERIFY_GL_SCOPE();
@@ -3285,16 +3247,10 @@ FOpenGLLinkedProgram* FOpenGLDynamicRHI::GetLinkedComputeProgram(FRHIComputeShad
 	return LinkedProgram;
 }
 
-FComputeShaderRHIRef FOpenGLDynamicRHI::RHICreateComputeShader(FRHIShaderLibrary* Library, FSHAHash Hash)
-{
-	check(RHISupportsComputeShaders(GMaxRHIShaderPlatform));
-	return CreateProxyShader<FRHIComputeShader, FOpenGLComputeShaderProxy>(Library, Hash);
-}
-
-FComputeShaderRHIRef FOpenGLDynamicRHI::RHICreateComputeShader(const TArray<uint8>& Code)
+FComputeShaderRHIRef FOpenGLDynamicRHI::RHICreateComputeShader(TArrayView<const uint8> Code, const FSHAHash& Hash)
 {
 	check(RHISupportsComputeShaders(GMaxRHIShaderPlatform));
-	return CreateProxyShader<FRHIComputeShader, FOpenGLComputeShaderProxy>(Code);
+	return CreateProxyShader<FRHIComputeShader, FOpenGLComputeShaderProxy>(Code, Hash);
 }
 
 template<class TOpenGLStage>
@@ -3582,7 +3538,7 @@ FBoundShaderStateRHIRef FOpenGLDynamicRHI::RHICreateBoundShaderState_OnThisThrea
 	if(!PixelShaderRHI)
 	{
 		// use special null pixel shader when PixelShader was set to NULL
-		PixelShaderRHI = TShaderMapRef<FNULLPS>(GetGlobalShaderMap(GMaxRHIFeatureLevel))->GetPixelShader();
+		PixelShaderRHI = TShaderMapRef<FNULLPS>(GetGlobalShaderMap(GMaxRHIFeatureLevel)).GetPixelShader();
 	}
 
 	// Check for an existing bound shader state which matches the parameters
diff --git a/Engine/Source/Runtime/OpenGLDrv/Private/OpenGLUniformBuffer.cpp b/Engine/Source/Runtime/OpenGLDrv/Private/OpenGLUniformBuffer.cpp
--- a/Engine/Source/Runtime/OpenGLDrv/Private/OpenGLUniformBuffer.cpp
+++ b/Engine/Source/Runtime/OpenGLDrv/Private/OpenGLUniformBuffer.cpp
@@ -747,7 +747,7 @@ void FOpenGLDynamicRHI::RHIUpdateUniformBuffer(FRHIUniformBuffer* UniformBufferR
 				&& UniformBufferValidation == EUniformBufferValidation::ValidateResources)
 			{
 				checkf(Resource, TEXT("Invalid resource entry creating uniform buffer, %s.Resources[%u], ResourceType 0x%x."),
-					*Layout.GetDebugName().ToString(),
+					*Layout.GetDebugName(),
 					ResourceIndex,
 					Layout.Resources[ResourceIndex].MemberType);
 			}
@@ -775,7 +775,7 @@ void FOpenGLDynamicRHI::RHIUpdateUniformBuffer(FRHIUniformBuffer* UniformBufferR
 					&& UniformBufferValidation == EUniformBufferValidation::ValidateResources)
 				{
 					checkf(Resource, TEXT("Invalid resource entry creating uniform buffer, %s.Resources[%u], ResourceType 0x%x."),
-						*Layout.GetDebugName().ToString(),
+						*Layout.GetDebugName(),
 						ResourceIndex,
 						(uint8)Layout.Resources[ResourceIndex].MemberType);
 				}
diff --git a/Engine/Source/Runtime/OpenGLDrv/Public/OpenGLDrv.h b/Engine/Source/Runtime/OpenGLDrv/Public/OpenGLDrv.h
--- a/Engine/Source/Runtime/OpenGLDrv/Public/OpenGLDrv.h
+++ b/Engine/Source/Runtime/OpenGLDrv/Public/OpenGLDrv.h
@@ -343,19 +343,12 @@ public:
 	virtual FBlendStateRHIRef RHICreateBlendState(const FBlendStateInitializerRHI& Initializer) final override;
 	virtual FVertexDeclarationRHIRef RHICreateVertexDeclaration(const FVertexDeclarationElementList& Elements) final override;
 
-	virtual FPixelShaderRHIRef RHICreatePixelShader(const TArray<uint8>& Code) final override;
-	virtual FVertexShaderRHIRef RHICreateVertexShader(const TArray<uint8>& Code) final override;
-	virtual FHullShaderRHIRef RHICreateHullShader(const TArray<uint8>& Code) final override;
-	virtual FDomainShaderRHIRef RHICreateDomainShader(const TArray<uint8>& Code) final override;
-	virtual FGeometryShaderRHIRef RHICreateGeometryShader(const TArray<uint8>& Code) final override;
-	virtual FComputeShaderRHIRef RHICreateComputeShader(const TArray<uint8>& Code) final override;
-
-	virtual FPixelShaderRHIRef RHICreatePixelShader(FRHIShaderLibrary* Library, FSHAHash Hash) final override;
-	virtual FVertexShaderRHIRef RHICreateVertexShader(FRHIShaderLibrary* Library, FSHAHash Hash) final override;
-	virtual FHullShaderRHIRef RHICreateHullShader(FRHIShaderLibrary* Library, FSHAHash Hash) final override;
-	virtual FDomainShaderRHIRef RHICreateDomainShader(FRHIShaderLibrary* Library, FSHAHash Hash) final override;
-	virtual FGeometryShaderRHIRef RHICreateGeometryShader(FRHIShaderLibrary* Library, FSHAHash Hash) final override;
-	virtual FComputeShaderRHIRef RHICreateComputeShader(FRHIShaderLibrary* Library, FSHAHash Hash) final override;
+	virtual FPixelShaderRHIRef RHICreatePixelShader(TArrayView<const uint8> Code, const FSHAHash& Hash) final override;
+	virtual FVertexShaderRHIRef RHICreateVertexShader(TArrayView<const uint8> Code, const FSHAHash& Hash) final override;
+	virtual FHullShaderRHIRef RHICreateHullShader(TArrayView<const uint8> Code, const FSHAHash& Hash) final override;
+	virtual FDomainShaderRHIRef RHICreateDomainShader(TArrayView<const uint8> Code, const FSHAHash& Hash) final override;
+	virtual FGeometryShaderRHIRef RHICreateGeometryShader(TArrayView<const uint8> Code, const FSHAHash& Hash) final override;
+	virtual FComputeShaderRHIRef RHICreateComputeShader(TArrayView<const uint8> Code, const FSHAHash& Hash) final override;
 
 	virtual FUniformBufferRHIRef RHICreateUniformBuffer(const void* Contents, const FRHIUniformBufferLayout& Layout, EUniformBufferUsage Usage, EUniformBufferValidation Validation) final override;
 	virtual void RHIUpdateUniformBuffer(FRHIUniformBuffer* UniformBufferRHI, const void* Contents) final override;
@@ -919,29 +912,29 @@ public:
 		return this->RHICancelAsyncReallocateTexture2D(Texture2D, bBlockUntilCompleted);
 	}
 
-	virtual FVertexShaderRHIRef CreateVertexShader_RenderThread(class FRHICommandListImmediate& RHICmdList, const TArray<uint8>& Code) final override
+	virtual FVertexShaderRHIRef CreateVertexShader_RenderThread(class FRHICommandListImmediate& RHICmdList, TArrayView<const uint8> Code, const FSHAHash& Hash) final override
 	{
-		return this->RHICreateVertexShader(Code);
+		return this->RHICreateVertexShader(Code, Hash);
 	}
 
-	virtual FPixelShaderRHIRef CreatePixelShader_RenderThread(class FRHICommandListImmediate& RHICmdList, const TArray<uint8>& Code) final override
+	virtual FPixelShaderRHIRef CreatePixelShader_RenderThread(class FRHICommandListImmediate& RHICmdList, TArrayView<const uint8> Code, const FSHAHash& Hash) final override
 	{
-		return this->RHICreatePixelShader(Code);
+		return this->RHICreatePixelShader(Code, Hash);
 	}
 
-	virtual FGeometryShaderRHIRef CreateGeometryShader_RenderThread(class FRHICommandListImmediate& RHICmdList, const TArray<uint8>& Code) final override
+	virtual FGeometryShaderRHIRef CreateGeometryShader_RenderThread(class FRHICommandListImmediate& RHICmdList, TArrayView<const uint8> Code, const FSHAHash& Hash) final override
 	{
-		return this->RHICreateGeometryShader(Code);
+		return this->RHICreateGeometryShader(Code, Hash);
 	}
 
-	virtual FComputeShaderRHIRef CreateComputeShader_RenderThread(class FRHICommandListImmediate& RHICmdList, const TArray<uint8>& Code) final override
+	virtual FComputeShaderRHIRef CreateComputeShader_RenderThread(class FRHICommandListImmediate& RHICmdList, TArrayView<const uint8> Code, const FSHAHash& Hash) final override
 	{
-		return this->RHICreateComputeShader(Code);
+		return this->RHICreateComputeShader(Code, Hash);
 	}
 
-	virtual FHullShaderRHIRef CreateHullShader_RenderThread(class FRHICommandListImmediate& RHICmdList, const TArray<uint8>& Code) final override
+	virtual FHullShaderRHIRef CreateHullShader_RenderThread(class FRHICommandListImmediate& RHICmdList, TArrayView<const uint8> Code, const FSHAHash& Hash) final override
 	{
-		return this->RHICreateHullShader(Code);
+		return this->RHICreateHullShader(Code, Hash);
 	}
 
 	virtual FSamplerStateRHIRef RHICreateSamplerState(const FSamplerStateInitializerRHI& Initializer) final override;
diff --git a/Engine/Source/Runtime/OpenGLDrv/Public/OpenGLShaderResources.h b/Engine/Source/Runtime/OpenGLDrv/Public/OpenGLShaderResources.h
--- a/Engine/Source/Runtime/OpenGLDrv/Public/OpenGLShaderResources.h
+++ b/Engine/Source/Runtime/OpenGLDrv/Public/OpenGLShaderResources.h
@@ -266,9 +266,9 @@ class TOpenGLShader : public RHIResourceType
 public:
 	enum
 	{
-		StaticFrequency = FrequencyT
+		StaticFrequency = FrequencyT,
+		TypeEnum = GLTypeEnum,
 	};
-	static const GLenum TypeEnum = GLTypeEnum;
 
 	/** The OpenGL resource ID. */
 	GLuint Resource;
diff --git a/Engine/Source/Runtime/OpenGLDrv/Public/OpenGLUtil.h b/Engine/Source/Runtime/OpenGLDrv/Public/OpenGLUtil.h
--- a/Engine/Source/Runtime/OpenGLDrv/Public/OpenGLUtil.h
+++ b/Engine/Source/Runtime/OpenGLDrv/Public/OpenGLUtil.h
@@ -166,7 +166,11 @@ extern bool PlatformOpenGLContextValid();
 	#define INITIATE_GL_FRAME_DUMP_EVERY_X_CALLS( a )
 #endif
 
-struct FRHICommandGLCommand final : public FRHICommand<FRHICommandGLCommand>
+struct FRHICommandGLCommandString
+{
+	static const TCHAR* TStr() { return TEXT("FRHICommandGLCommand"); }
+};
+struct FRHICommandGLCommand final : public FRHICommand<FRHICommandGLCommand, FRHICommandGLCommandString>
 {
 	TFunction<void()> GLFunction;
 
