Fix taking a local variables address and storing it in an array out of scope, then using that address later

diff --git a/Engine/Source/Runtime/Landscape/Private/Landscape.cpp b/Engine/Source/Runtime/Landscape/Private/Landscape.cpp
--- a/Engine/Source/Runtime/Landscape/Private/Landscape.cpp
+++ b/Engine/Source/Runtime/Landscape/Private/Landscape.cpp
@@ -345,11 +345,11 @@ void ULandscapeComponent::Serialize(FArchive& Ar)
 		TArray<UObject**> TexturesAndMaterials;
 		TexturesAndMaterials.Add((UObject**)&HeightmapTexture);
 		TexturesAndMaterials.Add((UObject**)&XYOffsetmapTexture);
-		for (UTexture2D* WeightmapTexture : WeightmapTextures)
+		for (UTexture2D*& WeightmapTexture : WeightmapTextures)
 		{
 			TexturesAndMaterials.Add((UObject**)&WeightmapTexture);
 		}
-		for (UTexture2D* MobileWeightmapTexture : MobileWeightmapTextures)
+		for (UTexture2D*& MobileWeightmapTexture : MobileWeightmapTextures)
 		{
 			TexturesAndMaterials.Add((UObject**)&MobileWeightmapTexture);
 		}
@@ -362,15 +362,15 @@ void ULandscapeComponent::Serialize(FArchive& Ar)
 				TexturesAndMaterials.Add((UObject**)&WeightmapTexture);
 			}
 		}
-		for (UMaterialInstance* MaterialInstance : MaterialInstances)
+		for (UMaterialInstanceConstant*& MaterialInstance : MaterialInstances)
 		{
 			TexturesAndMaterials.Add((UObject**)&MaterialInstance);
 		}
-		for (UMaterialInterface* MobileMaterialInterface : MobileMaterialInterfaces)
+		for (UMaterialInterface*& MobileMaterialInterface : MobileMaterialInterfaces)
 		{
 			TexturesAndMaterials.Add((UObject**)(&MobileMaterialInterface));
 		}
-		for (UMaterialInstance* MobileCombinationMaterialInstance : MobileCombinationMaterialInstances)
+		for (UMaterialInstanceConstant*& MobileCombinationMaterialInstance : MobileCombinationMaterialInstances)
 		{
 			TexturesAndMaterials.Add((UObject**)&MobileCombinationMaterialInstance);
 		}
