Fix vertex factory parameter association with FLandscapeXYOffsetVertexFactory

diff --git a/Engine/Source/Runtime/Landscape/Private/LandscapeRender.cpp b/Engine/Source/Runtime/Landscape/Private/LandscapeRender.cpp
--- a/Engine/Source/Runtime/Landscape/Private/LandscapeRender.cpp
+++ b/Engine/Source/Runtime/Landscape/Private/LandscapeRender.cpp
@@ -4535,6 +4535,13 @@ void FLandscapeXYOffsetVertexFactory::ModifyCompilationEnvironment(const FVertex
 	OutEnvironment.SetDefine(TEXT("LANDSCAPE_XYOFFSET"), TEXT("1"));
 }
 
+IMPLEMENT_VERTEX_FACTORY_PARAMETER_TYPE(FLandscapeXYOffsetVertexFactory, SF_Vertex, FLandscapeVertexFactoryVertexShaderParameters);
+#if RHI_RAYTRACING
+IMPLEMENT_VERTEX_FACTORY_PARAMETER_TYPE(FLandscapeXYOffsetVertexFactory, SF_Compute, FLandscapeVertexFactoryVertexShaderParameters);
+IMPLEMENT_VERTEX_FACTORY_PARAMETER_TYPE(FLandscapeXYOffsetVertexFactory, SF_RayHitGroup, FLandscapeVertexFactoryVertexShaderParameters);
+#endif // RHI_RAYTRACING
+IMPLEMENT_VERTEX_FACTORY_PARAMETER_TYPE(FLandscapeXYOffsetVertexFactory, SF_Pixel, FLandscapeVertexFactoryPixelShaderParameters);
+
 IMPLEMENT_VERTEX_FACTORY_TYPE_EX(FLandscapeXYOffsetVertexFactory, "/Engine/Private/LandscapeVertexFactory.ush", true, true, true, false, false, true, false);
 
 //
