Optimize DDC usage in the landscape collision components

Only load cooked data from DDC once unless it is being rebuilt, and only write cooked data to DDC when it is built. Without these changes, cooking these components will perform six gets and five puts for each instance.

diff --git a/Engine/Source/Runtime/Landscape/Private/LandscapeCollision.cpp b/Engine/Source/Runtime/Landscape/Private/LandscapeCollision.cpp
--- a/Engine/Source/Runtime/Landscape/Private/LandscapeCollision.cpp
+++ b/Engine/Source/Runtime/Landscape/Private/LandscapeCollision.cpp
@@ -802,6 +802,12 @@ TArray<PxHeightFieldSample> ConvertHeightfieldDataForPhysx(const ULandscapeHeigh
 
 bool ULandscapeHeightfieldCollisionComponent::CookCollisionData(const FName& Format, bool bUseDefMaterial, bool bCheckDDC, TArray<uint8>& OutCookedData, TArray<UPhysicalMaterial*>& InOutMaterials) const
 {
+	// Use existing cooked data unless !bCheckDDC in which case the data must be rebuilt.
+	if (bCheckDDC && OutCookedData.Num() > 0)
+	{
+		return true;
+	}
+
 	COOK_STAT(auto Timer = LandscapeCollisionCookStats::HeightfieldUsageStats.TimeSyncWork());
 	 
 	bool Succeeded = false;
@@ -1006,6 +1012,12 @@ bool ULandscapeHeightfieldCollisionComponent::CookCollisionData(const FName& For
 
 bool ULandscapeMeshCollisionComponent::CookCollisionData(const FName& Format, bool bUseDefMaterial, bool bCheckDDC, TArray<uint8>& OutCookedData, TArray<UPhysicalMaterial*>& InOutMaterials) const
 {
+	// Use existing cooked data unless !bCheckDDC in which case the data must be rebuilt.
+	if (bCheckDDC && OutCookedData.Num() > 0)
+	{
+		return true;
+	}
+
 #if WITH_PHYSX
 	COOK_STAT(auto Timer = LandscapeCollisionCookStats::MeshUsageStats.TimeSyncWork());
 	// we have 2 versions of collision objects
@@ -1776,10 +1788,6 @@ void ULandscapeHeightfieldCollisionComponent::Serialize(FArchive& Ar)
 		{
 			FName Format = Ar.CookingTarget()->GetPhysicsFormat(nullptr);
 			CookCollisionData(Format, false, true, CookedCollisionData, CookedPhysicalMaterials);
-			if (!GLandscapeCollisionSkipDDC && HeightfieldGuid.IsValid())
-			{
-				GetDerivedDataCacheRef().Put(*GetHFDDCKeyString(Format, false, HeightfieldGuid, CookedPhysicalMaterials), CookedCollisionData, GetPathName());
-			}
 		}
 	}
 #endif// WITH_EDITOR
@@ -2111,20 +2119,6 @@ void ULandscapeHeightfieldCollisionComponent::PreSave(const class ITargetPlatfor
 				RenderComponent->RenderGrassMap();
 			}
 		}
-
-		if(!GLandscapeCollisionSkipDDC)
-		{
-			static const FName PhysicsFormatName(FPlatformProperties::GetPhysicsFormat());
-			if(CookedCollisionData.Num() && HeightfieldGuid.IsValid())
-			{
-				GetDerivedDataCacheRef().Put(*GetHFDDCKeyString(PhysicsFormatName, false, HeightfieldGuid, CookedPhysicalMaterials), CookedCollisionData, GetPathName());
-			}
-
-			if(CookedCollisionDataEd.Num() && HeightfieldGuid.IsValid())
-			{
-				GetDerivedDataCacheRef().Put(*GetHFDDCKeyString(PhysicsFormatName, true, HeightfieldGuid, TArray<UPhysicalMaterial*>()), CookedCollisionDataEd, GetPathName());
-			}
-		}
 #endif// WITH_EDITOR
 	}
 }
