Fix a Software Occlusion bug that Landscape doesn't generate correct VB and IB data for multi-sections component.

diff --git a/Engine/Source/Runtime/Landscape/Private/LandscapeEdit.cpp b/Engine/Source/Runtime/Landscape/Private/LandscapeEdit.cpp
--- a/Engine/Source/Runtime/Landscape/Private/LandscapeEdit.cpp
+++ b/Engine/Source/Runtime/Landscape/Private/LandscapeEdit.cpp
@@ -6025,7 +6025,7 @@ void ULandscapeComponent::GeneratePlatformVertexData(const ITargetPlatform* Targ
 						FColor* CurrentMipSrcRow = HeightmapMipData[OcclusionMeshMip] + (CurrentMipOfsY + y) * MipSizeX + CurrentMipOfsX;
 						uint16 Height = CurrentMipSrcRow[x].R << 8 | CurrentMipSrcRow[x].G;
 
-						FVector VtxPos = FVector(x*MipRatio, y*MipRatio, ((float)Height - 32768.f) * LANDSCAPE_ZSCALE);
+						FVector VtxPos = FVector(x*MipRatio + SubX * SubsectionSizeQuads, y*MipRatio + SubY * SubsectionSizeQuads, ((float)Height - 32768.f) * LANDSCAPE_ZSCALE);
 						OccluderVertices.Add(VtxPos);
 					}
 				}
diff --git a/Engine/Source/Runtime/Landscape/Private/LandscapeRender.cpp b/Engine/Source/Runtime/Landscape/Private/LandscapeRender.cpp
--- a/Engine/Source/Runtime/Landscape/Private/LandscapeRender.cpp
+++ b/Engine/Source/Runtime/Landscape/Private/LandscapeRender.cpp
@@ -3887,35 +3887,44 @@ void FLandscapeSharedBuffers::CreateOccluderIndexBuffer(int32 NumOccluderVertice
 		return;
 	}
 
-	uint16 NumLineQuads = ((uint16)FMath::Sqrt(NumOccluderVertices) - 1);
-	uint16 NumLineVtx = NumLineQuads + 1;
-	check(NumLineVtx*NumLineVtx == NumOccluderVertices);
+	const uint16 NumLineQuads = ((uint16)(FMath::Sqrt(NumOccluderVertices)/NumSubsections) - 1);
+	const uint16 NumLineVtx = NumLineQuads + 1;
+	check(NumLineVtx*NumLineVtx*NumSubsections*NumSubsections == NumOccluderVertices);
 
-	int32 NumTris = NumLineQuads*NumLineQuads * 2;
-	int32 NumIndices = NumTris * 3;
+	const int32 NumTrisPerSubsection = NumLineQuads*NumLineQuads * 2;
+	const int32 NumIndicesPerSubsection = NumTrisPerSubsection * 3;
 	OccluderIndicesSP = MakeShared<FOccluderIndexArray, ESPMode::ThreadSafe>();
-	OccluderIndicesSP->SetNumUninitialized(NumIndices, false);
+	OccluderIndicesSP->SetNumUninitialized(NumIndicesPerSubsection*NumSubsections*NumSubsections, false);
 
 	uint16* OcclusionIndices = OccluderIndicesSP->GetData();
 	const uint16 NumLineVtxPlusOne = NumLineVtx + 1;
 	const uint16 QuadIndices[2][3] = { {0, NumLineVtx, NumLineVtxPlusOne}, {0, NumLineVtxPlusOne, 1} };
 	uint16 QuadOffset = 0;
 	int32 Index = 0;
-	for (int32 y = 0; y < NumLineQuads; y++)
+
+	for (int32 SubY = 0; SubY < NumSubsections; SubY++)
 	{
-		for (int32 x = 0; x < NumLineQuads; x++)
+		for (int32 SubX = 0; SubX < NumSubsections; SubX++)
 		{
-			for (int32 i = 0; i < 2; i++)
+			QuadOffset = (SubY*NumSubsections + SubX) * NumLineVtx * NumLineVtx;
+
+			for (int32 y = 0; y < NumLineQuads; y++)
 			{
-				OcclusionIndices[Index++] = QuadIndices[i][0] + QuadOffset;
-				OcclusionIndices[Index++] = QuadIndices[i][1] + QuadOffset;
-				OcclusionIndices[Index++] = QuadIndices[i][2] + QuadOffset;
+				for (int32 x = 0; x < NumLineQuads; x++)
+				{
+					for (int32 i = 0; i < 2; i++)
+					{
+						OcclusionIndices[Index++] = QuadIndices[i][0] + QuadOffset;
+						OcclusionIndices[Index++] = QuadIndices[i][1] + QuadOffset;
+						OcclusionIndices[Index++] = QuadIndices[i][2] + QuadOffset;
+					}
+					QuadOffset++;
+				}
+				QuadOffset++;
 			}
-			QuadOffset++;
 		}
-		QuadOffset++;
 	}
-
+	
 	INC_DWORD_STAT_BY(STAT_LandscapeOccluderMem, OccluderIndicesSP->GetAllocatedSize());
 }
 
