Landscape post merge fixes

* Enable cached mesh draw commands that got accidentally disabled by merge
* Fix seams in multi-subsection landscapes - the subsections are treated as one big subsection now
* Small refactoring to the FLandscapeRenderSystem code

diff --git a/Engine/Source/Runtime/Landscape/Private/LandscapeRender.cpp b/Engine/Source/Runtime/Landscape/Private/LandscapeRender.cpp
--- a/Engine/Source/Runtime/Landscape/Private/LandscapeRender.cpp
+++ b/Engine/Source/Runtime/Landscape/Private/LandscapeRender.cpp
@@ -646,6 +646,69 @@ IMPLEMENT_GLOBAL_SHADER_PARAMETER_STRUCT(FLandscapeSectionLODUniformParameters,
 TGlobalResource<FNullLandscapeRenderSystemResources> GNullLandscapeRenderSystemResources;
 TMap<FLandscapeNeighborInfo::FLandscapeKey, FLandscapeRenderSystem*> LandscapeRenderSystems;
 
+void FLandscapeRenderSystem::RegisterEntity(FLandscapeComponentSceneProxy* SceneProxy)
+{
+	check(IsInRenderingThread());
+	check(SceneProxy != nullptr);
+
+	if (NumRegisteredEntities > 0)
+	{
+		// Calculate new bounding rect of landscape components
+		FIntPoint OriginalMin = Min;
+		FIntPoint OriginalMax = Min + Size - FIntPoint(1, 1);
+		FIntPoint NewMin(FMath::Min(Min.X, SceneProxy->ComponentBase.X), FMath::Min(Min.Y, SceneProxy->ComponentBase.Y));
+		FIntPoint NewMax(FMath::Max(OriginalMax.X, SceneProxy->ComponentBase.X), FMath::Max(OriginalMax.Y, SceneProxy->ComponentBase.Y));
+
+		FIntPoint SizeRequired = (NewMax - NewMin) + FIntPoint(1, 1);
+
+		if (NewMin != Min || Size != SizeRequired)
+		{
+			ResizeAndMoveTo(NewMin, SizeRequired);
+			RecreateBuffers();
+		}
+
+		// Validate system-wide global parameters
+		check(TessellationFalloffSettings.UseTessellationComponentScreenSizeFalloff == SceneProxy->UseTessellationComponentScreenSizeFalloff);
+		check(TessellationFalloffSettings.TessellationComponentSquaredScreenSize == SceneProxy->TessellationComponentSquaredScreenSize);
+		check(TessellationFalloffSettings.TessellationComponentScreenSizeFalloff == SceneProxy->TessellationComponentScreenSizeFalloff);
+
+		if (SceneProxy->MaterialHasTessellationEnabled.Find(true) != INDEX_NONE)
+		{
+			NumEntitiesWithTessellation++;
+		}
+	}
+	else
+	{
+		TessellationFalloffSettings.UseTessellationComponentScreenSizeFalloff = SceneProxy->UseTessellationComponentScreenSizeFalloff;
+		TessellationFalloffSettings.TessellationComponentSquaredScreenSize = SceneProxy->TessellationComponentSquaredScreenSize;
+		TessellationFalloffSettings.TessellationComponentScreenSizeFalloff = SceneProxy->TessellationComponentScreenSizeFalloff;
+
+		ResizeAndMoveTo(SceneProxy->ComponentBase, FIntPoint(1, 1));
+		RecreateBuffers();
+	}
+
+	NumRegisteredEntities++;
+	SetSectionLODSettings(SceneProxy->ComponentBase, SceneProxy->LODSettings);
+	SetSectionOriginAndRadius(SceneProxy->ComponentBase, FVector4(SceneProxy->GetBounds().Origin, SceneProxy->GetBounds().SphereRadius));
+	SetSceneProxy(SceneProxy->ComponentBase, SceneProxy);
+}
+
+void FLandscapeRenderSystem::UnregisterEntity(FLandscapeComponentSceneProxy* SceneProxy)
+{
+	check(IsInRenderingThread());
+	check(SceneProxy != nullptr);
+
+	SetSceneProxy(SceneProxy->ComponentBase, nullptr);
+	SetSectionOriginAndRadius(SceneProxy->ComponentBase, FVector4(ForceInitToZero));
+
+	if (SceneProxy->MaterialHasTessellationEnabled.Find(true) != INDEX_NONE)
+	{
+		NumEntitiesWithTessellation--;
+	}
+
+	NumRegisteredEntities--;
+}
+
 void FLandscapeRenderSystem::ResizeAndMoveTo(FIntPoint NewMin, FIntPoint NewSize)
 {
 	SectionLODBuffer.SafeRelease();
@@ -677,7 +740,7 @@ void FLandscapeRenderSystem::ResizeAndMoveTo(FIntPoint NewMin, FIntPoint NewSize
 			int32 NewLinearIndex = (Y + (Min.Y - NewMin.Y)) * NewSize.X + (X + (Min.X - NewMin.X));
 			NewSectionLODValues[NewLinearIndex] = SectionLODValues[LinearIndex];
 			NewSectionLODBiases[NewLinearIndex] = SectionLODBiases[LinearIndex];
-			if (TessellationFalloffSettings.UseTessellationComponentScreenSizeFalloff)
+			if (TessellationFalloffSettings.UseTessellationComponentScreenSizeFalloff && NumEntitiesWithTessellation > 0)
 			{
 				NewSectionTessellationFalloffC[NewLinearIndex] = SectionTessellationFalloffC[LinearIndex];
 				NewSectionTessellationFalloffK[NewLinearIndex] = SectionTessellationFalloffK[LinearIndex];
@@ -698,7 +761,7 @@ void FLandscapeRenderSystem::ResizeAndMoveTo(FIntPoint NewMin, FIntPoint NewSize
 	SectionOriginAndRadius = NewSectionOriginAndRadius;
 	SceneProxies = NewSceneProxies;
 
-	if (!TessellationFalloffSettings.UseTessellationComponentScreenSizeFalloff)
+	if (!(TessellationFalloffSettings.UseTessellationComponentScreenSizeFalloff && NumEntitiesWithTessellation > 0))
 	{
 		for (float& Value : SectionTessellationFalloffC)
 		{
@@ -717,18 +780,43 @@ void FLandscapeRenderSystem::ResizeAndMoveTo(FIntPoint NewMin, FIntPoint NewSize
 	SectionTessellationFalloffK.SetAllowCPUAccess(true);
 }
 
-void FLandscapeRenderSystem::ComputeSectionPerViewParameters(const FSceneView* View)
+void FLandscapeRenderSystem::PrepareView(const FSceneView* View)
+{
+	const bool bExecuteInParallel = FApp::ShouldUseThreadingForPerformance()
+		&& GRenderingThread; // Rendering thread is required to safely use rendering resources in parallel.
+
+	if (bExecuteInParallel)
+	{
+		PerViewParametersTasks.Add(View, TGraphTask<FComputeSectionPerViewParametersTask>::CreateTask(
+			nullptr, ENamedThreads::GetRenderThread()).ConstructAndDispatchWhenReady(*this, View));
+	}
+}
+
 {
 	if (!CachedSectionLODValues.Contains(View))
 	{
 		TRACE_CPUPROFILER_EVENT_SCOPE(FLandscapeRenderSystem::ComputeSectionPerViewParameters());
 
-		CachedSectionLODValues.Add(View, TResourceArray<float>{});
+	if (PerViewParametersTasks.Contains(View))
+	{
+		FTaskGraphInterface::Get().WaitUntilTaskCompletes(PerViewParametersTasks[View], ENamedThreads::GetRenderThread_Local());
+		PerViewParametersTasks.Remove(View);
+	}
+	else
+	{
+		FComputeSectionPerViewParametersTask Task(*this, View);
+		Task.AnyThreadTask();
+	}
 
 		if (TessellationFalloffSettings.UseTessellationComponentScreenSizeFalloff)
 		{
-			CachedSectionTessellationFalloffC.Add(View, TResourceArray<float>{});
-			CachedSectionTessellationFalloffK.Add(View, TResourceArray<float>{});
+			CachedSectionLODValues.Add(ViewPtrAsIdentifier, TResourceArray<float>{});
+
+			if (TessellationFalloffSettings.UseTessellationComponentScreenSizeFalloff && NumEntitiesWithTessellation > 0)
+			{
+				CachedSectionTessellationFalloffC.Add(ViewPtrAsIdentifier, TResourceArray<float>{});
+				CachedSectionTessellationFalloffK.Add(ViewPtrAsIdentifier, TResourceArray<float>{});
+			}
 		}
 
 		int32 ForcedLODLevel = View->Family->EngineShowFlags.LOD ? GetCVarForceLOD() : -1;
@@ -743,7 +831,7 @@ void FLandscapeRenderSystem::ComputeSectionPerViewParameters(const FSceneView* V
 
 			SectionLODValues[EntityIndex] = ForcedLODLevel >= 0 ? ForcedLODLevel : FractionalLOD;
 
-			if (TessellationFalloffSettings.UseTessellationComponentScreenSizeFalloff)
+			if (TessellationFalloffSettings.UseTessellationComponentScreenSizeFalloff && NumEntitiesWithTessellation > 0)
 			{
 				float MaxTesselationDistance = ComputeBoundsDrawDistance(FMath::Sqrt(TessellationFalloffSettings.TessellationComponentSquaredScreenSize), SectionOriginAndRadius[EntityIndex].W / 2.0f, View->ViewMatrices.GetProjectionMatrix());
 				float FallOffStartingDistance = FMath::Min(
@@ -761,8 +849,15 @@ void FLandscapeRenderSystem::ComputeSectionPerViewParameters(const FSceneView* V
 
 		if (TessellationFalloffSettings.UseTessellationComponentScreenSizeFalloff)
 		{
-			CachedSectionTessellationFalloffC[View] = SectionTessellationFalloffC;
-			CachedSectionTessellationFalloffK[View] = SectionTessellationFalloffK;
+			FScopeLock Lock(&CachedValuesCS);
+
+			CachedSectionLODValues[ViewPtrAsIdentifier] = SectionLODValues;
+
+			if (TessellationFalloffSettings.UseTessellationComponentScreenSizeFalloff && NumEntitiesWithTessellation > 0)
+			{
+				CachedSectionTessellationFalloffC[ViewPtrAsIdentifier] = SectionTessellationFalloffC;
+				CachedSectionTessellationFalloffK[ViewPtrAsIdentifier] = SectionTessellationFalloffK;
+			}
 		}
 	}
 	else
@@ -771,8 +866,15 @@ void FLandscapeRenderSystem::ComputeSectionPerViewParameters(const FSceneView* V
 
 		if (TessellationFalloffSettings.UseTessellationComponentScreenSizeFalloff)
 		{
-			SectionTessellationFalloffC = CachedSectionTessellationFalloffC[View];
-			SectionTessellationFalloffK = CachedSectionTessellationFalloffK[View];
+			FScopeLock Lock(&CachedValuesCS);
+
+			SectionLODValues = CachedSectionLODValues[ViewPtrAsIdentifier];
+
+			if (TessellationFalloffSettings.UseTessellationComponentScreenSizeFalloff && NumEntitiesWithTessellation > 0)
+			{
+				SectionTessellationFalloffC = CachedSectionTessellationFalloffC[ViewPtrAsIdentifier];
+				SectionTessellationFalloffK = CachedSectionTessellationFalloffK[ViewPtrAsIdentifier];
+			}
 		}
 	}
 }
@@ -845,7 +947,7 @@ void FLandscapeRenderSystem::RecreateBuffers(const FSceneView* InView /* = nullp
 			else
 			{
 				// If we use tessellation falloff, update the buffer, otherwise use the one already filled with default parameters
-				if (TessellationFalloffSettings.UseTessellationComponentScreenSizeFalloff)
+				if (TessellationFalloffSettings.UseTessellationComponentScreenSizeFalloff && NumEntitiesWithTessellation > 0)
 				{
 					float* Data = (float*)RHILockVertexBuffer(SectionTessellationFalloffCBuffer, 0, SectionTessellationFalloffC.GetResourceDataSize(), RLM_WriteOnly);
 					FMemory::Memcpy(Data, SectionTessellationFalloffC.GetData(), SectionTessellationFalloffC.GetResourceDataSize());
@@ -864,7 +966,7 @@ void FLandscapeRenderSystem::RecreateBuffers(const FSceneView* InView /* = nullp
 			else
 			{
 				// If we use tessellation falloff, update the buffer, otherwise use the one already filled with default parameters
-				if (TessellationFalloffSettings.UseTessellationComponentScreenSizeFalloff)
+				if (TessellationFalloffSettings.UseTessellationComponentScreenSizeFalloff && NumEntitiesWithTessellation > 0)
 				{
 					float* Data = (float*)RHILockVertexBuffer(SectionTessellationFalloffKBuffer, 0, SectionTessellationFalloffK.GetResourceDataSize(), RLM_WriteOnly);
 					FMemory::Memcpy(Data, SectionTessellationFalloffK.GetData(), SectionTessellationFalloffK.GetResourceDataSize());
@@ -900,7 +1002,7 @@ void FLandscapeRenderSystem::BeginFrame()
 
 	CachedSectionLODValues.Empty();
 
-	if (TessellationFalloffSettings.UseTessellationComponentScreenSizeFalloff)
+	if (TessellationFalloffSettings.UseTessellationComponentScreenSizeFalloff && NumEntitiesWithTessellation > 0)
 	{
 		CachedSectionTessellationFalloffC.Empty();
 		CachedSectionTessellationFalloffK.Empty();
@@ -921,6 +1023,24 @@ void FLandscapeRenderSystem::BeginFrame()
 	}
 }
 
+void FLandscapeRenderSystem::EndFrame()
+{
+	// Finalize any outstanding jobs before ~FSceneRenderer() so we don't have corrupted accesses
+	if (FetchHeightmapLODBiasesEventRef.IsValid())
+	{
+		FTaskGraphInterface::Get().WaitUntilTaskCompletes(FetchHeightmapLODBiasesEventRef, ENamedThreads::GetRenderThread_Local());
+		FetchHeightmapLODBiasesEventRef.SafeRelease();
+	}
+
+	for (auto& Pair : PerViewParametersTasks)
+	{
+		const FSceneView* View = Pair.Key;
+		FTaskGraphInterface::Get().WaitUntilTaskCompletes(PerViewParametersTasks[View], ENamedThreads::GetRenderThread_Local());
+	}
+
+	PerViewParametersTasks.Empty();
+}
+
 class FLandscapePersistentViewUniformBufferExtension : public IPersistentViewUniformBufferExtension
 {
 public:
@@ -1238,9 +1358,7 @@ void FLandscapeComponentSceneProxy::CreateRenderThreadResources()
 
 	if (IsComponentLevelVisible())
 	{
-		OriginAndSphereRadius = FVector4(GetBounds().Origin, GetBounds().SphereRadius);
-		LandscapeSceneProxy = this;
-		RegisterNeighbors();
+		RegisterNeighbors(this);
 	}
 
 	auto FeatureLevel = GetScene().GetFeatureLevel();
@@ -1403,14 +1521,12 @@ void FLandscapeComponentSceneProxy::CreateRenderThreadResources()
 void FLandscapeComponentSceneProxy::DestroyRenderThreadResources()
 {
 	FPrimitiveSceneProxy::DestroyRenderThreadResources();
-	UnregisterNeighbors();
+	UnregisterNeighbors(this);
 }
 
 void FLandscapeComponentSceneProxy::OnLevelAddedToWorld()
 {
-	OriginAndSphereRadius = FVector4(GetBounds().Origin, GetBounds().SphereRadius);
-	LandscapeSceneProxy = this;
-	RegisterNeighbors();
+	RegisterNeighbors(this);
 }
 
 FLandscapeComponentSceneProxy::~FLandscapeComponentSceneProxy()
@@ -1686,6 +1802,7 @@ void FLandscapeComponentSceneProxy::OnTransformChanged()
 	LandscapeParams.ComponentBaseX = ComponentBase.X;
 	LandscapeParams.ComponentBaseY = ComponentBase.Y;
 	LandscapeParams.SubsectionSizeVerts = SubsectionSizeVerts;
+	LandscapeParams.NumSubsections = NumSubsections;
 	LandscapeParams.LastLOD = LastLOD;
 	LandscapeParams.HeightmapUVScaleBias = HeightmapScaleBias;
 	LandscapeParams.WeightmapUVScaleBias = WeightmapScaleBias;
@@ -1752,7 +1869,7 @@ void FLandscapeComponentSceneProxy::OnTransformChanged()
 
 	if (bRegistered)
 	{
-		OriginAndSphereRadius = FVector4(GetBounds().Origin, GetBounds().SphereRadius);
+		FVector4 OriginAndSphereRadius(GetBounds().Origin, GetBounds().SphereRadius);
 
 		FLandscapeRenderSystem& RenderSystem = *LandscapeRenderSystems.FindChecked(LandscapeKey);
 		RenderSystem.SetSectionOriginAndRadius(ComponentBase, OriginAndSphereRadius);
@@ -4236,111 +4353,6 @@ public:
 		{
 			ShaderBindings.Add(ForcedLodParameter, BatchElementParams->ForcedLOD);
 		}
-
-		int32 SubSectionIndex = BatchElementParams->SubX + BatchElementParams->SubY * SceneProxy->NumSubsections;
-		// If we have no custom data for this primitive we will compute of the fly the proper values, this will happen if the shader is not used for normal landscape rendering(i.e grass rendering)
-		FLandscapeComponentSceneProxy::FViewCustomDataLOD* LODData = (FLandscapeComponentSceneProxy::FViewCustomDataLOD*)InView->GetCustomData(SceneProxy->GetPrimitiveSceneInfo()->GetIndex());
-
-		if (LODData != nullptr)
-		{
-			SceneProxy->PostInitViewCustomData(*InView, LODData);
-
-			if (LodTessellationParameter.IsBound())
-			{
-				ShaderBindings.Add(LodTessellationParameter, LODData->LodTessellationParams);
-			}
-				
-			if (LodBiasParameter.IsBound())
-			{
-				check(LODData->LodBias == SceneProxy->GetShaderLODBias());
-				ShaderBindings.Add(LodBiasParameter, LODData->LodBias);
-			}
-
-			if (SectionLodsParameter.IsBound())
-			{
-				if (LODData->UseCombinedMeshBatch)
-				{
-					ShaderBindings.Add(SectionLodsParameter, LODData->ShaderCurrentLOD);
-				}
-				else // in non combined, only the one representing us as we'll be called 4 times (once per sub section)
-				{
-					check(SubSectionIndex >= 0);
-					FVector4 ShaderCurrentLOD(ForceInitToZero);
-					ShaderCurrentLOD.Component(SubSectionIndex) = LODData->ShaderCurrentLOD.Component(SubSectionIndex);
-
-					ShaderBindings.Add(SectionLodsParameter, ShaderCurrentLOD);
-				}
-			}
-
-			if (NeighborSectionLodParameter.IsBound())
-			{
-				FVector4 ShaderCurrentNeighborLOD[FLandscapeComponentSceneProxy::NEIGHBOR_COUNT] = { FVector4(ForceInitToZero), FVector4(ForceInitToZero), FVector4(ForceInitToZero), FVector4(ForceInitToZero) };
-
-				if (LODData->UseCombinedMeshBatch)
-				{					
-					int32 SubSectionCount = SceneProxy->NumSubsections == 1 ? 1 : FLandscapeComponentSceneProxy::MAX_SUBSECTION_COUNT;
-
-					for (int32 NeighborSubSectionIndex = 0; NeighborSubSectionIndex < SubSectionCount; ++NeighborSubSectionIndex)
-					{
-						ShaderCurrentNeighborLOD[NeighborSubSectionIndex] = LODData->SubSections[NeighborSubSectionIndex].ShaderCurrentNeighborLOD;
-						check(ShaderCurrentNeighborLOD[NeighborSubSectionIndex].X != -1.0f); // they should all match so only check the 1st one for simplicity
-					}
-
-					ShaderBindings.Add(NeighborSectionLodParameter, ShaderCurrentNeighborLOD);
-				}
-				else // in non combined, only the one representing us as we'll be called 4 times (once per sub section)
-				{
-					check(SubSectionIndex >= 0);
-					ShaderCurrentNeighborLOD[SubSectionIndex] = LODData->SubSections[SubSectionIndex].ShaderCurrentNeighborLOD;
-					check(ShaderCurrentNeighborLOD[SubSectionIndex].X != -1.0f); // they should all match so only check the 1st one for simplicity
-
-					ShaderBindings.Add(NeighborSectionLodParameter, ShaderCurrentNeighborLOD);
-				}
-			}			
-		}
-		else
-		{
-			float ComponentScreenSize = SceneProxy->GetComponentScreenSize(InView, SceneProxy->GetBounds().Origin, SceneProxy->ComponentMaxExtend, SceneProxy->GetBounds().SphereRadius);
-			
-			FLandscapeComponentSceneProxy::FViewCustomDataLOD CurrentLODData;
-			SceneProxy->CalculateBatchElementLOD(*InView, ComponentScreenSize, InView->LODDistanceFactor, CurrentLODData, true);
-			check(CurrentLODData.UseCombinedMeshBatch);
-
-			if (LodBiasParameter.IsBound())
-			{
-				ShaderBindings.Add(LodBiasParameter, SceneProxy->GetShaderLODBias());
-			}
-
-			if (LodTessellationParameter.IsBound())
-			{
-				FVector4 LodTessellationParams(ForceInitToZero);
-				SceneProxy->ComputeTessellationFalloffShaderValues(CurrentLODData, InView->ViewMatrices.GetProjectionMatrix(), LodTessellationParams.X, LodTessellationParams.Y);
-
-				ShaderBindings.Add(LodTessellationParameter, LodTessellationParams);
-			}
-
-			if (SectionLodsParameter.IsBound())
-			{
-				ShaderBindings.Add(SectionLodsParameter, CurrentLODData.ShaderCurrentLOD);
-			}
-
-			if (NeighborSectionLodParameter.IsBound())
-			{
-				FVector4 CurrentNeighborLOD[4] = { FVector4(ForceInitToZero), FVector4(ForceInitToZero), FVector4(ForceInitToZero), FVector4(ForceInitToZero) };
-
-				for (int32 SubY = 0; SubY < SceneProxy->NumSubsections; SubY++)
-				{
-					for (int32 SubX = 0; SubX < SceneProxy->NumSubsections; SubX++)
-					{
-						int32 NeighborSubSectionIndex = SubX + SubY * SceneProxy->NumSubsections;
-						SceneProxy->GetShaderCurrentNeighborLOD(*InView, CurrentLODData.SubSections[NeighborSubSectionIndex].fBatchElementCurrentLOD, SceneProxy->NumSubsections > 1 ? SubX : INDEX_NONE, SceneProxy->NumSubsections > 1 ? SubY : INDEX_NONE, NeighborSubSectionIndex, CurrentNeighborLOD[NeighborSubSectionIndex]);
-						check(CurrentNeighborLOD[NeighborSubSectionIndex].X != -1.0f); // they should all match so only check the 1st one for simplicity
-					}
-				}
-
-				ShaderBindings.Add(NeighborSectionLodParameter, CurrentNeighborLOD);
-			}				
-		}
 	}
 
 	virtual uint32 GetSize() const override
@@ -4474,7 +4486,7 @@ void FLandscapeVertexFactory::ModifyCompilationEnvironment(const FVertexFactoryT
 }
 
 
-IMPLEMENT_VERTEX_FACTORY_TYPE(FLandscapeVertexFactory, "/Engine/Private/LandscapeVertexFactory.ush", true, true, true, false, false);
+IMPLEMENT_VERTEX_FACTORY_TYPE_EX(FLandscapeVertexFactory, "/Engine/Private/LandscapeVertexFactory.ush", true, true, true, false, false, true, false);
 
 /**
 * Copy the data from another vertex factory
@@ -4503,7 +4515,7 @@ void FLandscapeXYOffsetVertexFactory::ModifyCompilationEnvironment(const FVertex
 	OutEnvironment.SetDefine(TEXT("LANDSCAPE_XYOFFSET"), TEXT("1"));
 }
 
-IMPLEMENT_VERTEX_FACTORY_TYPE(FLandscapeXYOffsetVertexFactory, "/Engine/Private/LandscapeVertexFactory.ush", true, true, true, false, false);
+IMPLEMENT_VERTEX_FACTORY_TYPE_EX(FLandscapeXYOffsetVertexFactory, "/Engine/Private/LandscapeVertexFactory.ush", true, true, true, false, false, true, false);
 
 //
 // FLandscapeFixedGridVertexFactory
@@ -5324,7 +5336,7 @@ void FLandscapeComponentSceneProxy::GetLCIs(FLCIArray& LCIs)
 //
 // FLandscapeNeighborInfo
 //
-void FLandscapeNeighborInfo::RegisterNeighbors()
+void FLandscapeNeighborInfo::RegisterNeighbors(FLandscapeComponentSceneProxy* SceneProxy /* = nullptr */)
 {
 	check(IsInRenderingThread());
 	if (!bRegistered)
@@ -5369,43 +5381,10 @@ void FLandscapeNeighborInfo::RegisterNeighbors()
 				Neighbors[3]->Neighbors[0] = this;
 			}
 
+			if (SceneProxy != nullptr)
 			{
 				FLandscapeRenderSystem& RenderSystem = *LandscapeRenderSystems.FindChecked(LandscapeKey);
-
-				if (RenderSystem.NumRegisteredEntities > 0)
-				{
-					FIntPoint OriginalMin = RenderSystem.Min;
-					FIntPoint OriginalMax = RenderSystem.Min + RenderSystem.Size - FIntPoint(1, 1);
-					FIntPoint NewMin(FMath::Min(RenderSystem.Min.X, ComponentBase.X), FMath::Min(RenderSystem.Min.Y, ComponentBase.Y));
-					FIntPoint NewMax(FMath::Max(OriginalMax.X, ComponentBase.X), FMath::Max(OriginalMax.Y, ComponentBase.Y));
-
-					FIntPoint SizeRequired = (NewMax - NewMin) + FIntPoint(1, 1);
-
-					if (NewMin != RenderSystem.Min || RenderSystem.Size != SizeRequired)
-					{
-						RenderSystem.ResizeAndMoveTo(NewMin, SizeRequired);
-						RenderSystem.RecreateBuffers();
-					}
-
-					// Validate system-wide global parameters
-					check(RenderSystem.TessellationFalloffSettings.UseTessellationComponentScreenSizeFalloff == LandscapeSceneProxy->UseTessellationComponentScreenSizeFalloff);
-					check(RenderSystem.TessellationFalloffSettings.TessellationComponentSquaredScreenSize == LandscapeSceneProxy->TessellationComponentSquaredScreenSize);
-					check(RenderSystem.TessellationFalloffSettings.TessellationComponentScreenSizeFalloff == LandscapeSceneProxy->TessellationComponentScreenSizeFalloff);
-				}
-				else
-				{
-					RenderSystem.TessellationFalloffSettings.UseTessellationComponentScreenSizeFalloff = LandscapeSceneProxy->UseTessellationComponentScreenSizeFalloff;
-					RenderSystem.TessellationFalloffSettings.TessellationComponentSquaredScreenSize = LandscapeSceneProxy->TessellationComponentSquaredScreenSize;
-					RenderSystem.TessellationFalloffSettings.TessellationComponentScreenSizeFalloff = LandscapeSceneProxy->TessellationComponentScreenSizeFalloff;
-
-					RenderSystem.ResizeAndMoveTo(ComponentBase, FIntPoint(1, 1));
-					RenderSystem.RecreateBuffers();
-				}
-
-				RenderSystem.NumRegisteredEntities++;
-				RenderSystem.SetSectionLODSettings(ComponentBase, LODSettings);
-				RenderSystem.SetSectionOriginAndRadius(ComponentBase, OriginAndSphereRadius);
-				RenderSystem.SetSceneProxy(ComponentBase, LandscapeSceneProxy);
+				RenderSystem.RegisterEntity(SceneProxy);
 			}
 		}
 		else
@@ -5415,7 +5394,7 @@ void FLandscapeNeighborInfo::RegisterNeighbors()
 	}
 }
 
-void FLandscapeNeighborInfo::UnregisterNeighbors()
+void FLandscapeNeighborInfo::UnregisterNeighbors(FLandscapeComponentSceneProxy* SceneProxy /* = nullptr */)
 {
 	check(IsInRenderingThread());
 	
@@ -5430,6 +5409,12 @@ void FLandscapeNeighborInfo::UnregisterNeighbors()
 		{
 			SceneProxyMap->Remove(ComponentBase);
 
+			if (SceneProxy != nullptr)
+			{
+				FLandscapeRenderSystem& RenderSystem = *LandscapeRenderSystems.FindChecked(LandscapeKey);
+				RenderSystem.UnregisterEntity(SceneProxy);
+			}
+
 			if (SceneProxyMap->Num() == 0)
 			{
 				// remove the entire LandscapeKey entry as this is the last scene proxy
@@ -5441,11 +5426,6 @@ void FLandscapeNeighborInfo::UnregisterNeighbors()
 			}
 			else
 			{
-				FLandscapeRenderSystem* RenderSystem = LandscapeRenderSystems.FindChecked(LandscapeKey);
-				RenderSystem->SetSceneProxy(ComponentBase, nullptr);
-				RenderSystem->SetSectionOriginAndRadius(ComponentBase, FVector4(ForceInitToZero));
-				RenderSystem->NumRegisteredEntities--;
-
 				// remove reference to us from our neighbors
 				if (Neighbors[0])
 				{
diff --git a/Engine/Source/Runtime/Landscape/Private/LandscapeRenderMobile.cpp b/Engine/Source/Runtime/Landscape/Private/LandscapeRenderMobile.cpp
--- a/Engine/Source/Runtime/Landscape/Private/LandscapeRenderMobile.cpp
+++ b/Engine/Source/Runtime/Landscape/Private/LandscapeRenderMobile.cpp
@@ -345,9 +345,7 @@ void FLandscapeComponentSceneProxyMobile::CreateRenderThreadResources()
 
 	if (IsComponentLevelVisible())
 	{
-		OriginAndSphereRadius = FVector4(GetBounds().Origin, GetBounds().SphereRadius);
-		LandscapeSceneProxy = this;
-		RegisterNeighbors();
+		RegisterNeighbors(this);
 	}
 	
 	auto FeatureLevel = GetScene().GetFeatureLevel();
diff --git a/Engine/Source/Runtime/Landscape/Public/LandscapeRender.h b/Engine/Source/Runtime/Landscape/Public/LandscapeRender.h
--- a/Engine/Source/Runtime/Landscape/Public/LandscapeRender.h
+++ b/Engine/Source/Runtime/Landscape/Public/LandscapeRender.h
@@ -92,6 +92,7 @@ BEGIN_GLOBAL_SHADER_PARAMETER_STRUCT(FLandscapeUniformShaderParameters, LANDSCAP
 SHADER_PARAMETER(int32, ComponentBaseX)
 SHADER_PARAMETER(int32, ComponentBaseY)
 SHADER_PARAMETER(int32, SubsectionSizeVerts)
+SHADER_PARAMETER(int32, NumSubsections)
 SHADER_PARAMETER(int32, LastLOD)
 SHADER_PARAMETER(FVector4, HeightmapUVScaleBias)
 SHADER_PARAMETER(FVector4, WeightmapUVScaleBias)
@@ -381,15 +382,6 @@ public:
 // FLandscapeNeighborInfo
 //
 
-struct LODSettingsComponent
-{
-	float LOD0ScreenSizeSquared;
-	float LOD1ScreenSizeSquared;
-	float LODOnePlusDistributionScalarSquared;
-	float LastLODScreenSizeSquared;
-	int8 LastLODIndex;
-};
-
 class FLandscapeNeighborInfo
 {
 public:
@@ -450,10 +442,6 @@ protected:
 	bool					bRegistered;
 	int32					PrimitiveCustomDataIndex;
 
-	LODSettingsComponent LODSettings;
-	FVector4 OriginAndSphereRadius;
-	FLandscapeComponentSceneProxy* LandscapeSceneProxy;
-
 	friend class FLandscapeComponentSceneProxy;
 
 public:
@@ -465,7 +453,6 @@ public:
 	, LODBias(InLODBias)
 	, bRegistered(false)
 	, PrimitiveCustomDataIndex(INDEX_NONE)
-	, LandscapeSceneProxy(nullptr)
 	{
 		//       -Y       
 		//    - - 0 - -   
@@ -481,8 +468,8 @@ public:
 		Neighbors[3] = nullptr;
 	}
 
-	void RegisterNeighbors();
-	void UnregisterNeighbors();
+	void RegisterNeighbors(FLandscapeComponentSceneProxy* SceneProxy = nullptr);
+	void UnregisterNeighbors(FLandscapeComponentSceneProxy* SceneProxy = nullptr);
 };
 
 
@@ -525,6 +512,15 @@ extern RENDERER_API TAutoConsoleVariable<float> CVarStaticMeshLODDistanceScale;
 
 struct FLandscapeRenderSystem
 {
+	struct LODSettingsComponent
+	{
+		float LOD0ScreenSizeSquared;
+		float LOD1ScreenSizeSquared;
+		float LODOnePlusDistributionScalarSquared;
+		float LastLODScreenSizeSquared;
+		int8 LastLODIndex;
+	};
+
 	static int8 GetLODFromScreenSize(LODSettingsComponent LODSettings, float InScreenSizeSquared, float InViewLODScale, float& OutFractionalLOD)
 	{
 		float ScreenSizeSquared = InScreenSizeSquared / InViewLODScale;
@@ -548,6 +544,7 @@ struct FLandscapeRenderSystem
 	}
 
 	int32 NumRegisteredEntities;
+	int32 NumEntitiesWithTessellation;
 
 	FIntPoint Min;
 	FIntPoint Size;
@@ -583,7 +580,7 @@ struct FLandscapeRenderSystem
 	TMap<const FSceneView*, TResourceArray<float>> CachedSectionTessellationFalloffK;
 	const FSceneView* CachedView;
 
-	TMap<const FSceneView*, FGraphEventRef> TaskEventRef;
+	TMap<const FSceneView*, FGraphEventRef> PerViewParametersTasks;
 	FGraphEventRef FetchHeightmapLODBiasesEventRef;
 
 	struct FComputeSectionPerViewParametersTask
@@ -660,6 +657,7 @@ struct FLandscapeRenderSystem
 
 	FLandscapeRenderSystem()
 		: NumRegisteredEntities(0)
+		, NumEntitiesWithTessellation(0)
 		, Min(MAX_int32, MAX_int32)
 		, Size(EForceInit::ForceInitToZero)
 		, CachedView(nullptr)
@@ -670,6 +668,10 @@ struct FLandscapeRenderSystem
 		SectionTessellationFalloffK.SetAllowCPUAccess(true);
 	}
 
+	void RegisterEntity(FLandscapeComponentSceneProxy* SceneProxy);
+
+	void UnregisterEntity(FLandscapeComponentSceneProxy* SceneProxy);
+
 	int32 GetComponentLinearIndex(FIntPoint ComponentBase)
 	{
 		return (ComponentBase.Y - Min.Y) * Size.X + ComponentBase.X - Min.X;
@@ -703,17 +705,7 @@ struct FLandscapeRenderSystem
 
 	void ComputeSectionPerViewParameters(const FSceneView* View);
 
-	void PrepareView(const FSceneView* View)
-	{
-		const bool bExecuteInParallel = FApp::ShouldUseThreadingForPerformance()
-			&& GRenderingThread; // Rendering thread is required to safely use rendering resources in parallel.
-
-		if (bExecuteInParallel)
-		{
-			TaskEventRef.Add(View, TGraphTask<FComputeSectionPerViewParametersTask>::CreateTask(
-				nullptr, ENamedThreads::GetRenderThread()).ConstructAndDispatchWhenReady(*this, View));
-		}
-	}
+	void PrepareView(const FSceneView* View);
 
 	void BeginRenderView(const FSceneView* View)
 	{
@@ -863,7 +855,7 @@ public:
 	TStaticArray<FLandscapeSectionRayTracingState, MAX_SUBSECTION_COUNT> SectionRayTracingStates;
 #endif
 
-	friend FLandscapeNeighborInfo;
+	friend FLandscapeRenderSystem;
 
 protected:
 	int8						MaxLOD;		// Maximum LOD level, user override possible
@@ -884,6 +876,8 @@ protected:
 	float						TessellationComponentSquaredScreenSize;	// Screen size of the component at which we start to apply tessellation
 	float						TessellationComponentScreenSizeFalloff;	// Min Component screen size before we start applying the tessellation falloff
 
+	FLandscapeRenderSystem::LODSettingsComponent LODSettings;
+
 	/** 
 	 * Number of subsections within the component in each dimension, this can be 1 or 2.
 	 * Subsections exist to improve the speed at which LOD transitions can take place over distance.
diff --git a/Engine/Source/Runtime/Renderer/Private/RendererScene.cpp b/Engine/Source/Runtime/Renderer/Private/RendererScene.cpp
--- a/Engine/Source/Runtime/Renderer/Private/RendererScene.cpp
+++ b/Engine/Source/Runtime/Renderer/Private/RendererScene.cpp
@@ -3336,6 +3336,7 @@ void FScene::OnLevelAddedToWorld(FName LevelAddedName, UWorld* InWorld, bool bIs
 	ENQUEUE_RENDER_COMMAND(FLevelAddedToWorld)(
 		[Scene, LevelAddedName](FRHICommandListImmediate& RHICmdList)
 		{
+			Scene->UpdateAllPrimitiveSceneInfos(RHICmdList);
 			Scene->OnLevelAddedToWorld_RenderThread(LevelAddedName);
 		});
 }
