DRPS - Linker fixes
#rnx

diff --git a/Engine/Source/Runtime/Landscape/Private/LandscapeRender.cpp b/Engine/Source/Runtime/Landscape/Private/LandscapeRender.cpp
--- a/Engine/Source/Runtime/Landscape/Private/LandscapeRender.cpp
+++ b/Engine/Source/Runtime/Landscape/Private/LandscapeRender.cpp
@@ -1903,6 +1903,274 @@
 	}
 }
 
+float FLandscapeComponentSceneProxy::GetComponentScreenSize(const FSceneView* View, const FVector& Origin, float MaxExtend, float ElementRadius) const
+{
+	float SquaredScreenRadius = 0.f;
+
+	if (!View->ViewMatrices.IsPerspectiveProjection())
+	{
+		FMatrix ProjMatrix = View->ViewMatrices.GetProjectionMatrix();
+		const float ScreenMultiple = FMath::Max(0.5f * ProjMatrix.M[0][0], 0.5f * ProjMatrix.M[1][1]);
+		SquaredScreenRadius = FMath::Square(ScreenMultiple * ElementRadius);
+	}
+	else
+	{
+		FVector CameraOrigin = View->ViewMatrices.GetViewOrigin();
+		FMatrix ProjMatrix = View->ViewMatrices.GetProjectionMatrix();
+
+		const FVector OriginToCamera = (CameraOrigin - Origin).GetAbs();
+		const FVector ClosestPoint = OriginToCamera.ComponentMin(FVector(MaxExtend));
+		const float DistSquared = (OriginToCamera - ClosestPoint).SizeSquared();
+
+		// Get projection multiple accounting for view scaling.
+		const float ScreenMultiple = FMath::Max(0.5f * ProjMatrix.M[0][0], 0.5f * ProjMatrix.M[1][1]);
+
+		// Calculate screen-space projected radius
+		SquaredScreenRadius = FMath::Square(ScreenMultiple * ElementRadius) / FMath::Max(1.0f, DistSquared);
+	}
+
+	// If we hit NaN check see UE-64538
+	check(FMath::IsFinite(SquaredScreenRadius));
+	return SquaredScreenRadius * 2.0f;
+}
+
+void FLandscapeComponentSceneProxy::BuildDynamicMeshElement(const FViewCustomDataLOD* InPrimitiveCustomData, bool InToolMesh, bool InHasTessellation, bool InDisableTessellation, FMeshBatch& OutMeshBatch, TArray<FLandscapeBatchElementParams, SceneRenderingAllocator>& OutStaticBatchParamArray) const
+{
+	if (AvailableMaterials.Num() == 0 || InPrimitiveCustomData == nullptr || InPrimitiveCustomData->SubSections.Num() == 0)
+	{
+		return;
+	}
+
+	const int8 CurrentLODIndex = InPrimitiveCustomData->SubSections[0].BatchElementCurrentLOD;
+	int8 MaterialIndex = LODIndexToMaterialIndex.IsValidIndex(CurrentLODIndex) ? LODIndexToMaterialIndex[CurrentLODIndex] : INDEX_NONE;
+	UMaterialInterface* SelectedMaterial = MaterialIndex != INDEX_NONE ? AvailableMaterials[MaterialIndex] : nullptr;
+
+	if (InHasTessellation && MaterialIndex != INDEX_NONE)
+	{		
+		if (InDisableTessellation && MaterialIndexToDisabledTessellationMaterial.IsValidIndex(MaterialIndex))
+		{
+			SelectedMaterial = AvailableMaterials[MaterialIndexToDisabledTessellationMaterial[MaterialIndex]];
+		}
+	}
+
+	// this is really not normal that we have no material at this point, so do not continue
+	if (SelectedMaterial == nullptr)
+	{
+		return;
+	}
+
+	// Could be different from bRequiresAdjacencyInformation during shader compilation
+	bool bCurrentRequiresAdjacencyInformation = !InToolMesh && MaterialRenderingRequiresAdjacencyInformation_RenderingThread(SelectedMaterial, VertexFactory->GetType(), GetScene().GetFeatureLevel());
+
+	if (bCurrentRequiresAdjacencyInformation)
+	{
+		check(SharedBuffers->AdjacencyIndexBuffers);
+	}
+
+	OutMeshBatch.VertexFactory = VertexFactory;
+	OutMeshBatch.MaterialRenderProxy = SelectedMaterial->GetRenderProxy();
+	OutMeshBatch.LCI = ComponentLightInfo.Get();
+	OutMeshBatch.ReverseCulling = IsLocalToWorldDeterminantNegative();
+	OutMeshBatch.CastShadow = InToolMesh ? false : true;
+	OutMeshBatch.bUseAsOccluder = ShouldUseAsOccluder() && GetScene().GetShadingPath() == EShadingPath::Deferred && !IsMovable();
+	OutMeshBatch.bUseForMaterial = true;
+	OutMeshBatch.Type = bCurrentRequiresAdjacencyInformation ? PT_12_ControlPointPatchList : PT_TriangleList;
+	OutMeshBatch.LODIndex = 0;
+
+	OutMeshBatch.Elements.Empty();
+
+	if (NumSubsections > 1 && !InPrimitiveCustomData->UseCombinedMeshBatch)
+	{
+		for (int32 SubY = 0; SubY < NumSubsections; SubY++)
+		{
+			for (int32 SubX = 0; SubX < NumSubsections; SubX++)
+			{
+				const int8 SubSectionIdx = SubX + SubY * NumSubsections;
+				const int8 CurrentLOD = InPrimitiveCustomData->SubSections[SubSectionIdx].BatchElementCurrentLOD;
+
+				FMeshBatchElement BatchElement;
+				FLandscapeBatchElementParams& BatchElementParams = OutStaticBatchParamArray[SubSectionIdx];
+
+				if (!InToolMesh)
+				{
+					BatchElementParams.LocalToWorldNoScalingPtr = &LocalToWorldNoScaling;
+					BatchElementParams.LandscapeUniformShaderParametersResource = &LandscapeUniformShaderParameters;
+					BatchElementParams.SceneProxy = this;
+					BatchElementParams.SubX = SubX;
+					BatchElementParams.SubY = SubY;
+					BatchElementParams.CurrentLOD = CurrentLOD;
+				}
+
+				BatchElement.UserData = &BatchElementParams;
+				BatchElement.PrimitiveUniformBuffer = GetUniformBuffer();
+
+				int32 LodSubsectionSizeVerts = (SubsectionSizeVerts >> CurrentLOD);
+				uint32 NumPrimitives = FMath::Square((LodSubsectionSizeVerts - 1)) * 2;
+
+				if (bCurrentRequiresAdjacencyInformation)
+				{
+					check(SharedBuffers->AdjacencyIndexBuffers);
+					BatchElement.IndexBuffer = SharedBuffers->AdjacencyIndexBuffers->IndexBuffers[CurrentLOD];
+					BatchElement.FirstIndex = (SubX + SubY * NumSubsections) * NumPrimitives * 12;
+				}
+				else
+				{
+					BatchElement.IndexBuffer = SharedBuffers->IndexBuffers[CurrentLOD];
+					BatchElement.FirstIndex = (SubX + SubY * NumSubsections) * NumPrimitives * 3;
+				}
+				BatchElement.NumPrimitives = NumPrimitives;
+				BatchElement.MinVertexIndex = SharedBuffers->IndexRanges[CurrentLOD].MinIndex[SubX][SubY];
+				BatchElement.MaxVertexIndex = SharedBuffers->IndexRanges[CurrentLOD].MaxIndex[SubX][SubY];
+
+
+#if !(UE_BUILD_SHIPPING || UE_BUILD_TEST)
+				// We simplify this by considering only the biggest LOD index for this mesh element.
+				OutMeshBatch.VisualizeLODIndex = FMath::Max(OutMeshBatch.VisualizeLODIndex, CurrentLOD);
+#endif
+
+				OutMeshBatch.Elements.Add(BatchElement);
+			}
+		}
+	}
+	else
+	{
+		FMeshBatchElement BatchElement;
+
+		if (InToolMesh)
+		{
+			// Reuse the params for the tool mesh
+			BatchElement.UserData = &OutStaticBatchParamArray[0];
+		}
+		else
+		{
+			FLandscapeBatchElementParams& BatchElementParams = OutStaticBatchParamArray[0];
+			BatchElementParams.LandscapeUniformShaderParametersResource = &LandscapeUniformShaderParameters;
+			BatchElementParams.LocalToWorldNoScalingPtr = &LocalToWorldNoScaling;
+			BatchElementParams.SceneProxy = this;
+			BatchElementParams.SubX = -1;
+			BatchElementParams.SubY = -1;
+			BatchElementParams.CurrentLOD = CurrentLODIndex;
+
+			BatchElement.UserData = &BatchElementParams;
+		}
+
+		// Combined batch element
+		int32 LodSubsectionSizeVerts = SubsectionSizeVerts >> CurrentLODIndex;
+
+		BatchElement.PrimitiveUniformBuffer = GetUniformBuffer();
+		BatchElement.IndexBuffer = bCurrentRequiresAdjacencyInformation ? SharedBuffers->AdjacencyIndexBuffers->IndexBuffers[CurrentLODIndex] : SharedBuffers->IndexBuffers[CurrentLODIndex];
+		BatchElement.NumPrimitives = FMath::Square((LodSubsectionSizeVerts - 1)) * FMath::Square(NumSubsections) * 2;
+		BatchElement.FirstIndex = 0;
+		BatchElement.MinVertexIndex = SharedBuffers->IndexRanges[CurrentLODIndex].MinIndexFull;
+		BatchElement.MaxVertexIndex = SharedBuffers->IndexRanges[CurrentLODIndex].MaxIndexFull;
+
+		OutMeshBatch.Elements.Add(BatchElement);
+	}
+}
+
+bool FLandscapeComponentSceneProxy::GetMeshElement(bool UseSeperateBatchForShadow, bool ShadowOnly, bool HasTessellation, int8 InLODIndex, UMaterialInterface* InMaterialInterface, FMeshBatch& OutMeshBatch, TArray<FLandscapeBatchElementParams>& OutStaticBatchParamArray) const
+{
+	if (InMaterialInterface == nullptr)
+	{
+		return false;
+	}
+
+	// Could be different from bRequiresAdjacencyInformation during shader compilation
+	bool bCurrentRequiresAdjacencyInformation = MaterialRenderingRequiresAdjacencyInformation_RenderingThread(InMaterialInterface, VertexFactory->GetType(), GetScene().GetFeatureLevel());
+
+	if (bCurrentRequiresAdjacencyInformation)
+	{
+		check(SharedBuffers->AdjacencyIndexBuffers);
+	}
+
+	OutMeshBatch.VertexFactory = VertexFactory;
+	OutMeshBatch.MaterialRenderProxy = InMaterialInterface->GetRenderProxy();
+
+	OutMeshBatch.LCI = ComponentLightInfo.Get();
+	OutMeshBatch.ReverseCulling = IsLocalToWorldDeterminantNegative();
+	OutMeshBatch.CastShadow = UseSeperateBatchForShadow ? ShadowOnly : true;
+	OutMeshBatch.bUseForDepthPass = (UseSeperateBatchForShadow ? !ShadowOnly : true);
+	OutMeshBatch.bUseAsOccluder = (UseSeperateBatchForShadow ? !ShadowOnly : true) && ShouldUseAsOccluder() && GetScene().GetShadingPath() == EShadingPath::Deferred && !IsMovable();
+	OutMeshBatch.bUseForMaterial = (UseSeperateBatchForShadow ? !ShadowOnly : true);
+	OutMeshBatch.Type = bCurrentRequiresAdjacencyInformation ? PT_12_ControlPointPatchList : PT_TriangleList;
+	OutMeshBatch.DepthPriorityGroup = SDPG_World;
+	OutMeshBatch.LODIndex = InLODIndex;
+	OutMeshBatch.bDitheredLODTransition = false;
+
+	int32 BatchElementSize = NumSubsections == 1 ? 1 : MAX_SUBSECTION_COUNT + 1;
+	OutMeshBatch.Elements.Empty(FMath::Max(LastLOD - FirstLOD, 1) * BatchElementSize);
+
+	for (int32 i = FirstLOD; i <= LastLOD; ++i)
+	{
+		int32 LodSubsectionSizeVerts = SubsectionSizeVerts >> i;
+
+		if (NumSubsections > 1 && ForcedLOD < 0)
+		{
+			uint32 NumPrimitivesPerSection = FMath::Square((LodSubsectionSizeVerts - 1)) * 2;
+
+			// Per-subsection batch elements
+			for (int32 SubY = 0; SubY < NumSubsections; SubY++)
+			{
+				for (int32 SubX = 0; SubX < NumSubsections; SubX++)
+				{
+					FLandscapeBatchElementParams* BatchElementParams = new(OutStaticBatchParamArray) FLandscapeBatchElementParams;
+					BatchElementParams->LandscapeUniformShaderParametersResource = &LandscapeUniformShaderParameters;
+					BatchElementParams->LocalToWorldNoScalingPtr = &LocalToWorldNoScaling;
+					BatchElementParams->SceneProxy = this;
+					BatchElementParams->SubX = SubX;
+					BatchElementParams->SubY = SubY;
+					BatchElementParams->CurrentLOD = i;
+
+					FMeshBatchElement BatchElement;
+					BatchElement.UserData = BatchElementParams;
+					BatchElement.PrimitiveUniformBuffer = GetUniformBuffer();
+
+					if (bCurrentRequiresAdjacencyInformation)
+					{
+						BatchElement.IndexBuffer = SharedBuffers->AdjacencyIndexBuffers->IndexBuffers[i];
+						BatchElement.FirstIndex = (SubX + SubY * NumSubsections) * NumPrimitivesPerSection * 12;
+					}
+					else
+					{
+						BatchElement.IndexBuffer = SharedBuffers->IndexBuffers[i];
+						BatchElement.FirstIndex = (SubX + SubY * NumSubsections) * NumPrimitivesPerSection * 3;
+					}
+					BatchElement.NumPrimitives = NumPrimitivesPerSection;
+					BatchElement.MinVertexIndex = SharedBuffers->IndexRanges[i].MinIndex[SubX][SubY];
+					BatchElement.MaxVertexIndex = SharedBuffers->IndexRanges[i].MaxIndex[SubX][SubY];
+
+					OutMeshBatch.Elements.Add(BatchElement);
+				}
+			}
+		}
+
+		// Combined batch element
+		FLandscapeBatchElementParams* BatchElementParams = new(OutStaticBatchParamArray) FLandscapeBatchElementParams;
+		BatchElementParams->LandscapeUniformShaderParametersResource = &LandscapeUniformShaderParameters;
+		BatchElementParams->LocalToWorldNoScalingPtr = &LocalToWorldNoScaling;
+		BatchElementParams->SceneProxy = this;
+		BatchElementParams->SubX = -1;
+		BatchElementParams->SubY = -1;
+		BatchElementParams->CurrentLOD = i;
+
+		FMeshBatchElement BatchElement;
+
+		BatchElement.UserData = BatchElementParams;
+		BatchElement.PrimitiveUniformBuffer = GetUniformBuffer();
+		BatchElement.IndexBuffer = bCurrentRequiresAdjacencyInformation ? SharedBuffers->AdjacencyIndexBuffers->IndexBuffers[i] : SharedBuffers->IndexBuffers[i];
+		BatchElement.NumPrimitives = FMath::Square((LodSubsectionSizeVerts - 1)) * FMath::Square(NumSubsections) * 2;
+		BatchElement.FirstIndex = 0;
+		BatchElement.MinVertexIndex = SharedBuffers->IndexRanges[i].MinIndexFull;
+		BatchElement.MaxVertexIndex = SharedBuffers->IndexRanges[i].MaxIndexFull;
+
+		OutMeshBatch.Elements.Add(BatchElement);
+	}
+
+	OutMeshBatch.Elements.Shrink();
+
+	return true;
+}
+
 /** Creates a mesh batch for virtual texture rendering. Will render a simple fixed grid with combined subsections. */
 bool FLandscapeComponentSceneProxy::GetMeshElementForVirtualTexture(int32 InLodIndex, ERuntimeVirtualTextureMaterialType MaterialType, UMaterialInterface* InMaterialInterface, FMeshBatch& OutMeshBatch, TArray<FLandscapeBatchElementParams>& OutStaticBatchParamArray) const
 {
@@ -2104,6 +2372,258 @@
 	check(StaticBatchParamArray.Num() <= TotalBatchCount);
 }
 
+void FLandscapeComponentSceneProxy::CalculateLODFromScreenSize(const FSceneView& InView, float InMeshScreenSizeSquared, float InViewLODScale, int32 InSubSectionIndex, FViewCustomDataLOD& InOutLODData) const
+{
+	// Handle general LOD override
+	float PreferedLOD = (float)GetCVarForceLOD();
+
+#if WITH_EDITOR
+	if (InView.Family->LandscapeLODOverride >= 0)
+	{
+		PreferedLOD = (float)InView.Family->LandscapeLODOverride;
+	}
+#endif
+
+#if WITH_EDITOR || !(UE_BUILD_SHIPPING || UE_BUILD_TEST)
+	if (InView.Family->EngineShowFlags.CollisionVisibility || InView.Family->EngineShowFlags.CollisionPawn)
+	{
+		const bool bDrawSimpleCollision = InView.Family->EngineShowFlags.CollisionPawn        && CollisionResponse.GetResponse(ECC_Pawn) != ECR_Ignore;
+		const bool bDrawComplexCollision = InView.Family->EngineShowFlags.CollisionVisibility && CollisionResponse.GetResponse(ECC_Visibility) != ECR_Ignore;
+
+		if (bDrawSimpleCollision)
+		{
+			PreferedLOD = FMath::Max(CollisionMipLevel, SimpleCollisionMipLevel);
+		}
+		else if (bDrawComplexCollision)
+		{
+			PreferedLOD = CollisionMipLevel;
+		}
+	}
+#endif
+
+	if (ForcedLOD >= 0)
+	{
+		PreferedLOD = (float)ForcedLOD;
+	}
+
+	int8 MinStreamedLOD = HeightmapTexture ? FMath::Min<int8>(((FTexture2DResource*)HeightmapTexture->Resource)->GetCurrentFirstMip(), FMath::CeilLogTwo(SubsectionSizeVerts) - 1) : 0;
+	MinStreamedLOD = FMath::Min(MinStreamedLOD, MaxLOD); // We can't go above MaxLOD even for texture streaming
+
+	int8 LocalLODBias = LODBias + (int8)GLandscapeMeshLODBias;
+	FViewCustomDataSubSectionLOD& SubSectionLODData = InOutLODData.SubSections[InSubSectionIndex];
+
+	if (PreferedLOD >= 0.0f)
+	{
+		PreferedLOD = FMath::Clamp<float>(PreferedLOD + LocalLODBias, FMath::Max((float)MinStreamedLOD, MinValidLOD), FMath::Min((float)LastLOD, MaxValidLOD));
+	}
+	else
+	{
+		PreferedLOD = FMath::Clamp<float>(ComputeBatchElementCurrentLOD(GetLODFromScreenSize(InMeshScreenSizeSquared, InViewLODScale), InMeshScreenSizeSquared, InViewLODScale) + LocalLODBias, FMath::Max((float)MinStreamedLOD, MinValidLOD), FMath::Min((float)LastLOD, MaxValidLOD));
+	}
+
+	check(PreferedLOD != -1.0f && PreferedLOD <= MaxLOD);
+	SubSectionLODData.fBatchElementCurrentLOD = PreferedLOD;
+	SubSectionLODData.BatchElementCurrentLOD = FMath::FloorToInt(PreferedLOD);
+}
+
+uint64 FLandscapeVertexFactory::GetStaticBatchElementVisibility(const FSceneView& InView, const FMeshBatch* InBatch, const void* InViewCustomData) const
+{
+	const FLandscapeComponentSceneProxy* SceneProxy = ((FLandscapeBatchElementParams*)InBatch->Elements[0].UserData)->SceneProxy;
+	return SceneProxy->GetStaticBatchElementVisibility(InView, InBatch, InViewCustomData);
+}
+
+uint64 FLandscapeComponentSceneProxy::GetStaticBatchElementVisibility(const FSceneView& InView, const FMeshBatch* InBatch, const void* InViewCustomData) const
+{
+	uint64 BatchesToRenderMask = 0;
+
+	SCOPE_CYCLE_COUNTER(STAT_LandscapeStaticDrawLODTime);
+
+	const void* ViewCustomData = InViewCustomData != nullptr ? InViewCustomData : InView.GetCustomData(GetPrimitiveSceneInfo()->GetIndex());
+
+	if (ViewCustomData != nullptr)
+	{
+		const FViewCustomDataLOD* CurrentLODData = (const FViewCustomDataLOD*)ViewCustomData;
+		const auto FeatureLevel = InView.GetFeatureLevel();
+
+		if (FeatureLevel >= ERHIFeatureLevel::SM5)
+		{
+			const int8 CurrentLODIndex = CurrentLODData->SubSections[0].BatchElementCurrentLOD;
+			ULandscapeMaterialInstanceConstant* LandscapeMIC = Cast<ULandscapeMaterialInstanceConstant>(AvailableMaterials[LODIndexToMaterialIndex[CurrentLODIndex]]);
+
+			if (LandscapeMIC == nullptr)
+			{
+				UMaterialInstanceDynamic* LandscapeMID = Cast<UMaterialInstanceDynamic>(AvailableMaterials[LODIndexToMaterialIndex[CurrentLODIndex]]);
+
+				if (LandscapeMID != nullptr)
+				{
+					LandscapeMIC = Cast<ULandscapeMaterialInstanceConstant>(LandscapeMID->Parent);
+				}
+			}
+
+			bool HasTessellationEnabled = MaterialHasTessellationEnabled[LODIndexToMaterialIndex[CurrentLODIndex]] && (LandscapeMIC != nullptr && !LandscapeMIC->bDisableTessellation);
+
+			if (HasTessellationEnabled)
+			{
+				INC_DWORD_STAT(STAT_LandscapeTessellatedComponents);
+			}
+		}
+
+		if (NumSubsections > 1 && !CurrentLODData->UseCombinedMeshBatch)
+		{
+			INC_DWORD_STAT(STAT_LandscapeComponentUsingSubSectionDrawCalls);
+
+			for (int32 SubSectionIndex = 0; SubSectionIndex < MAX_SUBSECTION_COUNT; ++SubSectionIndex)
+			{
+				const FViewCustomDataSubSectionLOD& SubSectionLODData = CurrentLODData->SubSections[SubSectionIndex];
+				check(SubSectionLODData.StaticBatchElementIndexToRender != INDEX_NONE);
+
+				BatchesToRenderMask |= (((uint64)1) << SubSectionLODData.StaticBatchElementIndexToRender);
+				INC_DWORD_STAT(STAT_LandscapeDrawCalls);
+				INC_DWORD_STAT_BY(STAT_LandscapeTriangles, InBatch->Elements[SubSectionLODData.StaticBatchElementIndexToRender].NumPrimitives);
+			}
+		}
+		else
+		{
+			int32 SubSectionIndex = 0;
+			const FViewCustomDataSubSectionLOD& SubSectionLODData = CurrentLODData->SubSections[SubSectionIndex];
+			check(SubSectionLODData.StaticBatchElementIndexToRender != INDEX_NONE);
+
+			BatchesToRenderMask |= (((uint64)1) << SubSectionLODData.StaticBatchElementIndexToRender);
+			INC_DWORD_STAT(STAT_LandscapeDrawCalls);
+			INC_DWORD_STAT_BY(STAT_LandscapeTriangles, InBatch->Elements[SubSectionLODData.StaticBatchElementIndexToRender].NumPrimitives);
+		}
+	}
+
+	INC_DWORD_STAT(STAT_LandscapeComponentRenderPasses);
+
+	return BatchesToRenderMask;
+}
+
+void FLandscapeComponentSceneProxy::CalculateBatchElementLOD(const FSceneView& InView, float InMeshScreenSizeSquared, float InViewLODScale, FViewCustomDataLOD& InOutLODData, bool InForceCombined) const
+{
+	float SquaredViewLODScale = FMath::Square(InViewLODScale);
+
+	check(InMeshScreenSizeSquared >= 0.0f);
+	float ComponentScreenSize = InMeshScreenSizeSquared;
+
+	if (NumSubsections > 1)
+	{
+		InOutLODData.UseCombinedMeshBatch = false; // default to individual batch render
+
+		float SubSectionMaxExtend = ComponentMaxExtend / 2.0f;
+		float SubSectionRadius = GetBounds().SphereRadius / 2.0f;
+		float CombinedScreenRatio = 0.0f;
+		bool AllSubSectionHaveSameScreenSize = true;
+
+		// Compute screen size of each sub section to determine if we should use the combined logic or the individual logic
+		for (int32 SubY = 0; SubY < NumSubsections; ++SubY)
+		{
+			for (int32 SubX = 0; SubX < NumSubsections; ++SubX)
+			{
+				int32 SubSectionIndex = SubX + SubY * NumSubsections;
+				FViewCustomDataSubSectionLOD& SubSectionLODData = InOutLODData.SubSections[SubSectionIndex];
+
+				SubSectionLODData.ScreenSizeSquared = GetComponentScreenSize(&InView, SubSectionScreenSizeTestingPosition[SubSectionIndex], SubSectionMaxExtend, SubSectionRadius);
+
+				check(SubSectionLODData.ScreenSizeSquared > 0.0f);
+
+				CalculateLODFromScreenSize(InView, SubSectionLODData.ScreenSizeSquared, InViewLODScale, SubSectionIndex, InOutLODData);
+				check(SubSectionLODData.fBatchElementCurrentLOD != -1.0f);
+
+				InOutLODData.ShaderCurrentLOD.Component(SubSectionIndex) = SubSectionLODData.fBatchElementCurrentLOD;
+
+				// Determine if we should use the combined batch or not
+				if (ComponentScreenSize > ComponentSquaredScreenSizeToUseSubSections * SquaredViewLODScale)
+				{
+					if (AllSubSectionHaveSameScreenSize)
+					{
+						float CurrentScreenRadiusSquared = SubSectionLODData.ScreenSizeSquared * SquaredViewLODScale;
+
+						if (CombinedScreenRatio > 0.0f && !FMath::IsNearlyEqual(CombinedScreenRatio, CurrentScreenRadiusSquared, KINDA_SMALL_NUMBER))
+						{
+							AllSubSectionHaveSameScreenSize = false;
+						}
+
+						CombinedScreenRatio += CurrentScreenRadiusSquared;
+
+						if (SubSectionIndex > 0)
+						{
+							CombinedScreenRatio *= 0.5f;
+						}
+					}
+				}
+			}
+		}
+
+		if (!GLandscapeDebugOptions.IsCombinedDisabled() && (AllSubSectionHaveSameScreenSize || GLandscapeDebugOptions.IsCombinedAll() || ForcedLOD != INDEX_NONE || InForceCombined))
+		{
+			InOutLODData.UseCombinedMeshBatch = true;
+
+			int32 MinLOD = FMath::Min(InOutLODData.SubSections[0].BatchElementCurrentLOD, FMath::Min3(InOutLODData.SubSections[1].BatchElementCurrentLOD, InOutLODData.SubSections[2].BatchElementCurrentLOD, InOutLODData.SubSections[3].BatchElementCurrentLOD));
+
+			for (int32 SubSectionIndex = 0; SubSectionIndex < MAX_SUBSECTION_COUNT; ++SubSectionIndex)
+			{
+				InOutLODData.SubSections[SubSectionIndex].BatchElementCurrentLOD = MinLOD;
+			}
+		}
+	}
+	else
+	{
+		const int32 SubSectionIndex = 0;
+		InOutLODData.UseCombinedMeshBatch = true;
+		FViewCustomDataSubSectionLOD& SubSectionLODData = InOutLODData.SubSections[SubSectionIndex];
+
+		SubSectionLODData.ScreenSizeSquared = ComponentScreenSize;
+		CalculateLODFromScreenSize(InView, SubSectionLODData.ScreenSizeSquared, InViewLODScale, SubSectionIndex, InOutLODData);
+		check(SubSectionLODData.fBatchElementCurrentLOD != -1.0f);
+
+		InOutLODData.ShaderCurrentLOD.Component(SubSectionIndex) = SubSectionLODData.fBatchElementCurrentLOD;
+	}
+}
+
+int32 FLandscapeComponentSceneProxy::ConvertBatchElementLODToBatchElementIndex(int8 BatchElementLOD, bool UseCombinedMeshBatch) const
+{
+	int32 BatchElementIndex = BatchElementLOD;
+
+	if (NumSubsections > 1 && ForcedLOD < 0)
+	{
+		BatchElementIndex = BatchElementLOD * (MAX_SUBSECTION_COUNT + 1);
+
+		if (UseCombinedMeshBatch)
+		{
+			BatchElementIndex += MAX_SUBSECTION_COUNT;
+		}
+	}
+
+	return BatchElementIndex;
+}
+
+float FLandscapeComponentSceneProxy::ComputeBatchElementCurrentLOD(int32 InSelectedLODIndex, float InComponentScreenSize, float InViewLODScale) const
+{
+	check(LODScreenRatioSquared.IsValidIndex(InSelectedLODIndex));
+
+	bool LastElement = InSelectedLODIndex == LODScreenRatioSquared.Num() - 1;
+	float CurrentLODScreenRatio = LODScreenRatioSquared[InSelectedLODIndex];
+	float NextLODScreenRatio = LastElement ? 0 : LODScreenRatioSquared[InSelectedLODIndex + 1];
+
+	float LODScreenRatioRange = CurrentLODScreenRatio - NextLODScreenRatio;
+	float ScreenSizeWithLODScale = FMath::Clamp(InComponentScreenSize / InViewLODScale, 0.0f, LODScreenRatioSquared[0]);
+
+	if (ScreenSizeWithLODScale > CurrentLODScreenRatio || ScreenSizeWithLODScale < NextLODScreenRatio)
+	{
+		// Find corresponding LODIndex to appropriately calculate Ratio and apply it to new LODIndex
+		int32 LODFromScreenSize = GetLODFromScreenSize(InComponentScreenSize, InViewLODScale);
+		CurrentLODScreenRatio = LODScreenRatioSquared[LODFromScreenSize];
+		NextLODScreenRatio = LODFromScreenSize == LODScreenRatioSquared.Num() - 1 ? 0 : LODScreenRatioSquared[LODFromScreenSize + 1];
+		LODScreenRatioRange = CurrentLODScreenRatio - NextLODScreenRatio;
+	}
+
+	float CurrentLODRangeRatio = (ScreenSizeWithLODScale - NextLODScreenRatio) / LODScreenRatioRange;
+	float fLOD = (float)InSelectedLODIndex + (1.0f - CurrentLODRangeRatio);
+
+	return fLOD;
+}
+
 int8 FLandscapeComponentSceneProxy::GetLODFromScreenSize(float InScreenSizeSquared, float InViewLODScale) const
 {
 	float FractionalLOD;
@@ -2144,6 +2664,98 @@
 #endif
 }
 
+void* FLandscapeComponentSceneProxy::InitViewCustomData(const FSceneView& InView, float InViewLODScale, FMemStackBase& InCustomDataMemStack, bool InIsStaticRelevant, bool InIsShadowOnly, const FLODMask* InVisiblePrimitiveLODMask, float InMeshScreenSizeSquared)
+{
+	SCOPE_CYCLE_COUNTER(STAT_LandscapeInitViewCustomData);
+
+	// NOTE: we can't access other proxy here as this can be run in parallel we need to wait for the PostInitViewCustomData which is run in synchronous	
+
+	PrimitiveCustomDataIndex = GetPrimitiveSceneInfo()->GetIndex();
+
+	FViewCustomDataLOD* LODData = (FViewCustomDataLOD*)new(InCustomDataMemStack) FViewCustomDataLOD();
+
+	LODData->ComponentScreenSize = InMeshScreenSizeSquared;
+
+	// If a valid screen size was provided, we use it instead of recomputing it
+	if (InMeshScreenSizeSquared < 0.0f)
+	{
+		LODData->ComponentScreenSize = GetComponentScreenSize(&InView, GetBounds().Origin, ComponentMaxExtend, GetBounds().SphereRadius);
+	}
+
+	CalculateBatchElementLOD(InView, LODData->ComponentScreenSize, InViewLODScale, *LODData, false);
+
+	if (InIsStaticRelevant)
+	{
+		check(InVisiblePrimitiveLODMask != nullptr);
+		LODData->StaticMeshBatchLOD = InVisiblePrimitiveLODMask->DitheredLODIndices[0];
+
+		if (LODData->UseCombinedMeshBatch)
+		{
+			ComputeStaticBatchIndexToRender(*LODData, 0);
+		}
+		else
+		{
+			for (int32 i = 0; i < MAX_SUBSECTION_COUNT; ++i)
+			{
+				ComputeStaticBatchIndexToRender(*LODData, i);
+			}
+		}
+	}
+
+	LODData->IsShadowOnly = InIsShadowOnly;
+
+	// Mobile use a different way of calculating the Bias
+	if (GetScene().GetFeatureLevel() >= ERHIFeatureLevel::SM5)
+	{
+		LODData->LodBias = GetShaderLODBias();
+	}
+
+	ComputeTessellationFalloffShaderValues(*LODData, InView.ViewMatrices.GetProjectionMatrix(), LODData->LodTessellationParams.X, LODData->LodTessellationParams.Y);
+
+	return LODData;
+}
+
+void FLandscapeComponentSceneProxy::ComputeTessellationFalloffShaderValues(const FViewCustomDataLOD& InLODData, const FMatrix& InViewProjectionMatrix, float& OutC, float& OutK) const
+{
+	// No Falloff
+	OutC = 1.0f;
+	OutK = 0.0f;
+
+	const auto FeatureLevel = GetScene().GetFeatureLevel();
+	bool HasTessellationEnabled = false;
+
+	if (FeatureLevel >= ERHIFeatureLevel::SM5)
+	{
+		const int8 CurrentLODIndex = InLODData.SubSections[0].BatchElementCurrentLOD;
+		ULandscapeMaterialInstanceConstant* LandscapeMIC = Cast<ULandscapeMaterialInstanceConstant>(AvailableMaterials[LODIndexToMaterialIndex[CurrentLODIndex]]);
+
+		if (LandscapeMIC == nullptr)
+		{
+			UMaterialInstanceDynamic* LandscapeMID = Cast<UMaterialInstanceDynamic>(AvailableMaterials[LODIndexToMaterialIndex[CurrentLODIndex]]);
+
+			if (LandscapeMID != nullptr)
+			{
+				LandscapeMIC = Cast<ULandscapeMaterialInstanceConstant>(LandscapeMID->Parent);
+			}
+		}
+
+		HasTessellationEnabled = MaterialHasTessellationEnabled[LODIndexToMaterialIndex[CurrentLODIndex]] && (LandscapeMIC != nullptr && !LandscapeMIC->bDisableTessellation);
+	}
+
+	if (HasTessellationEnabled && (InLODData.StaticMeshBatchLOD == INDEX_NONE || (InLODData.StaticMeshBatchLOD == 0 || InLODData.StaticMeshBatchLOD == 2))) // Tess batch will be used
+	{
+		if (UseTessellationComponentScreenSizeFalloff)
+		{
+			float MaxTesselationDistance = ComputeBoundsDrawDistance(FMath::Sqrt(TessellationComponentSquaredScreenSize), GetBounds().SphereRadius / 2.0f, InViewProjectionMatrix);
+			float FallOffStartingDistance = FMath::Min(ComputeBoundsDrawDistance(FMath::Sqrt(FMath::Min(FMath::Square(TessellationComponentScreenSizeFalloff), TessellationComponentSquaredScreenSize)), GetBounds().SphereRadius / 2.0f, InViewProjectionMatrix) - MaxTesselationDistance, MaxTesselationDistance);
+
+			// Calculate the falloff using a = C - K * d by sending C & K into the shader
+			OutC = MaxTesselationDistance / (MaxTesselationDistance - FallOffStartingDistance);
+			OutK = -(1 / (-MaxTesselationDistance + FallOffStartingDistance));
+		}
+	}
+}
+
 FVector4 FLandscapeComponentSceneProxy::GetShaderLODBias() const
 {
 	return FVector4(
@@ -2162,4 +2774,407 @@
 		1.f / (float)((SubsectionSizeVerts >> InBatchElementCurrentLOD) - 1));
 }
 
+void FLandscapeComponentSceneProxy::GetShaderCurrentNeighborLOD(const FSceneView& InView, float InBatchElementCurrentLOD, int8 InSubSectionX, int8 InSubSectionY, int8 InCurrentSubSectionIndex, FVector4& OutShaderCurrentNeighborLOD) const
+{
+	for (int32 NeighborIndex = 0; NeighborIndex < NEIGHBOR_COUNT; ++NeighborIndex)
+	{
+		OutShaderCurrentNeighborLOD.Component(NeighborIndex) = GetNeighborLOD(InView, InBatchElementCurrentLOD, NeighborIndex, InSubSectionX, InSubSectionY, InCurrentSubSectionIndex);
+		check(OutShaderCurrentNeighborLOD.Component(NeighborIndex) != -1.0f);
+	}
+}
+
+struct SubSectionData
+{
+	SubSectionData(int8 InSubSectionOffsetX, int8 InSubSectionOffsetY, bool InInsideComponent)
+		: SubSectionOffsetX(InSubSectionOffsetX)
+		, SubSectionOffsetY(InSubSectionOffsetY)
+		, InsideComponent(InInsideComponent)
+	{}
+
+	int8 SubSectionOffsetX;
+	int8 SubSectionOffsetY;
+	bool InsideComponent;
+};
+
+// SubSectionIndex, NeighborIndex
+static const SubSectionData SubSectionValues[4][4] = { { SubSectionData(0, 1, false), SubSectionData(1, 0, false), SubSectionData(1, 0, true), SubSectionData(0, 1, true) },
+{ SubSectionData(0, 1, false), SubSectionData(-1, 0, true), SubSectionData(-1, 0, false), SubSectionData(0, 1, true) },
+{ SubSectionData(0, -1, true), SubSectionData(1, 0, false), SubSectionData(1, 0, true), SubSectionData(0, -1, false) },
+{ SubSectionData(0, -1, true), SubSectionData(-1, 0, true), SubSectionData(-1, 0, false), SubSectionData(0, -1, false) } };
+
+
+float FLandscapeComponentSceneProxy::GetNeighborLOD(const FSceneView& InView, float InBatchElementCurrentLOD, int8 InNeighborIndex, int8 InSubSectionX, int8 InSubSectionY, int8 InCurrentSubSectionIndex) const
+{
+	float NeighborLOD = InBatchElementCurrentLOD;
+
+	// Assume no sub section initialization
+	bool InsideComponent = false;
+	int32 PrimitiveDataIndex = PrimitiveCustomDataIndex;
+	int32 DesiredSubSectionIndex = 0;
+	int32 DesiredSubSectionX = 0;
+	int32 DesiredSubSectionY = 0;
+	int32 CurrentSubSectionIndex = InSubSectionX != INDEX_NONE && InSubSectionY != INDEX_NONE ? InSubSectionX + InSubSectionY * NumSubsections : 0;
+
+	// Handle subsection
+	if (InSubSectionX != INDEX_NONE && InSubSectionY != INDEX_NONE)
+	{
+		const SubSectionData& Data = SubSectionValues[CurrentSubSectionIndex][InNeighborIndex];
+		DesiredSubSectionX = InSubSectionX + Data.SubSectionOffsetX;
+		DesiredSubSectionY = InSubSectionY + Data.SubSectionOffsetY;
+		DesiredSubSectionIndex = DesiredSubSectionX + DesiredSubSectionY * NumSubsections;
+		InsideComponent = Data.InsideComponent;
+	}
+
+	const FLandscapeNeighborInfo* Neighbor = nullptr;
+
+	if (!InsideComponent)
+	{
+		Neighbor = Neighbors[InNeighborIndex];
+
+		if (Neighbor != nullptr)
+		{
+			PrimitiveDataIndex = Neighbor->PrimitiveCustomDataIndex;
+		}
+		else
+		{
+			DesiredSubSectionX = InSubSectionX;
+			DesiredSubSectionY = InSubSectionY;
+			DesiredSubSectionIndex = CurrentSubSectionIndex;
+		}
+	}
+
+	bool ComputeNeighborCustomDataLOD = true;
+
+	void* CustomData = InView.GetCustomData(PrimitiveDataIndex);
+
+	if (CustomData != nullptr)
+	{
+		const FViewCustomDataLOD* LODData = (const FViewCustomDataLOD*)CustomData;
+		// Don't use the custom data for neighbor calculation when it is marked shadow only (ie it is not visible in the view)
+		// See UE-69785 for more information
+		if (!LODData->IsShadowOnly)
+		{
+			ComputeNeighborCustomDataLOD = false;
+			NeighborLOD = FMath::Max(LODData->SubSections[DesiredSubSectionIndex].fBatchElementCurrentLOD, InBatchElementCurrentLOD);
+		}
+	}
+
+	if (ComputeNeighborCustomDataLOD)
+	{
+		FBoxSphereBounds NeighborBounds = GetBounds();
+		float NeighborMaxExtends = ComponentMaxExtend;
+		FLandscapeComponentSceneProxy* NeighborSceneProxy = nullptr;
+
+		if (Neighbor != nullptr)
+		{
+			const ULandscapeComponent* NeighborComponent = Neighbor->GetLandscapeComponent();
+			if (NeighborComponent != nullptr && NeighborComponent->SceneProxy != nullptr)
+			{
+				NeighborSceneProxy = (FLandscapeComponentSceneProxy*)NeighborComponent->SceneProxy;
+				NeighborBounds = NeighborComponent->Bounds;
+				NeighborMaxExtends = NeighborSceneProxy->SubsectionSizeQuads * FMath::Max(NeighborComponent->GetRenderMatrix().GetScaleVector().X, NeighborComponent->GetRenderMatrix().GetScaleVector().Y);
+			}
+		}
+
+		if (NumSubsections > 1)
+		{
+			float NeighborSubsectionMaxExtends = NeighborMaxExtends / 2.0f;
+			FVector ComponentTopLeftCorner = NeighborBounds.Origin - FVector(NeighborSubsectionMaxExtends, NeighborSubsectionMaxExtends, 0.0f);
+
+			FVector SubSectionOrigin = ComponentTopLeftCorner + FVector(NeighborMaxExtends * DesiredSubSectionX, NeighborMaxExtends * DesiredSubSectionY, 0.0f);
+			float MeshBatchScreenSizeSquared = GetComponentScreenSize(&InView, SubSectionOrigin, NeighborSubsectionMaxExtends, NeighborBounds.SphereRadius / 2.0f);
+
+			FViewCustomDataLOD NeighborLODData;
+			if (NeighborSceneProxy != nullptr)
+			{
+				// Needs to pull some data like per-component LOD bias from neighbor scene proxy
+				NeighborSceneProxy->CalculateLODFromScreenSize(InView, MeshBatchScreenSizeSquared, InView.LODDistanceFactor, DesiredSubSectionIndex, NeighborLODData);
+			}
+			else
+			{
+				CalculateLODFromScreenSize(InView, MeshBatchScreenSizeSquared, InView.LODDistanceFactor, DesiredSubSectionIndex, NeighborLODData);
+			}
+
+			const FViewCustomDataSubSectionLOD& SubSectionData = NeighborLODData.SubSections[DesiredSubSectionIndex];
+			check(SubSectionData.fBatchElementCurrentLOD != -1.0f);
+
+			NeighborLOD = FMath::Max(SubSectionData.fBatchElementCurrentLOD, InBatchElementCurrentLOD);
+		}
+		else
+		{
+			float MeshBatchScreenSizeSquared = GetComponentScreenSize(&InView, NeighborBounds.Origin, NeighborMaxExtends, NeighborBounds.SphereRadius);
+
+			FViewCustomDataLOD NeighborLODData;
+			if (NeighborSceneProxy != nullptr)
+			{
+				NeighborSceneProxy->CalculateLODFromScreenSize(InView, MeshBatchScreenSizeSquared, InView.LODDistanceFactor, DesiredSubSectionIndex, NeighborLODData);
+			}
+			else
+			{
+				CalculateLODFromScreenSize(InView, MeshBatchScreenSizeSquared, InView.LODDistanceFactor, DesiredSubSectionIndex, NeighborLODData);
+			}
+
+			FViewCustomDataSubSectionLOD& SubSectionLODData = NeighborLODData.SubSections[0];
+			check(SubSectionLODData.fBatchElementCurrentLOD != -1.0f);
+
+			if (SubSectionLODData.fBatchElementCurrentLOD > InBatchElementCurrentLOD)
+			{
+				NeighborLOD = SubSectionLODData.fBatchElementCurrentLOD;
+			}
+		}
+	}
+
+	return NeighborLOD;
+}
+
+void FLandscapeComponentSceneProxy::ComputeStaticBatchIndexToRender(FViewCustomDataLOD& OutLODData, int32 SubSectionIndex)
+{
+	FViewCustomDataSubSectionLOD& SubSectionLODData = OutLODData.SubSections[SubSectionIndex];
+
+	SubSectionLODData.StaticBatchElementIndexToRender = INDEX_NONE;
+	SubSectionLODData.StaticBatchElementIndexToRender = ConvertBatchElementLODToBatchElementIndex(SubSectionLODData.BatchElementCurrentLOD, OutLODData.UseCombinedMeshBatch) + SubSectionIndex;
+	check(SubSectionLODData.StaticBatchElementIndexToRender != INDEX_NONE);
+}
+
+void FLandscapeComponentSceneProxy::PostInitViewCustomData(const FSceneView& InView, void* InViewCustomData) const
+{
+	SCOPE_CYCLE_COUNTER(STAT_LandscapePostInitViewCustomData);
+
+	FViewCustomDataLOD* CurrentLODData = (FViewCustomDataLOD*)InViewCustomData;
+	check(CurrentLODData != nullptr);
+
+	for (int8 SubY = 0; SubY < NumSubsections; ++SubY)
+	{
+		for (int8 SubX = 0; SubX < NumSubsections; ++SubX)
+		{
+			int8 SubSectionIndex = SubX + SubY * NumSubsections;
+			FViewCustomDataSubSectionLOD& SubSectionLODData = CurrentLODData->SubSections[SubSectionIndex];
+			GetShaderCurrentNeighborLOD(InView, SubSectionLODData.fBatchElementCurrentLOD, NumSubsections > 1 ? SubX : INDEX_NONE, NumSubsections > 1 ? SubY : INDEX_NONE, SubSectionIndex, SubSectionLODData.ShaderCurrentNeighborLOD);
+		}
+	}
+
+#if !UE_BUILD_SHIPPING
+	if (GVarDumpLandscapeLODs && GFrameNumberRenderThread == GVarDumpLandscapeLODsCurrentFrame)
+	{
+		if (NumSubsections == 1)
+		{
+			UE_LOG(LogLandscape, Warning, TEXT("\nComponent: [%s] -> MeshBatchLOD: %d, ComponentScreenSize: %f, ShaderCurrentLOD: %s, LODTessellation: %s\nHeightmap Texture Name: %s, Heightmap Streamed Mip: %d\nSubSections:\n|0| IndexToRender: %d, fLOD: %f, LOD: %d, NeighborLOD: %s\n"), *SectionBase.ToString(),
+				CurrentLODData->StaticMeshBatchLOD, CurrentLODData->ComponentScreenSize, *CurrentLODData->ShaderCurrentLOD.ToString(), *CurrentLODData->LodTessellationParams.ToString(),
+				HeightmapTexture != nullptr ? *HeightmapTexture->GetFullName() : TEXT("Invalid"), HeightmapTexture != nullptr ? ((FTexture2DResource*)HeightmapTexture->Resource)->GetCurrentFirstMip() : INDEX_NONE,
+				CurrentLODData->SubSections[0].StaticBatchElementIndexToRender, CurrentLODData->SubSections[0].fBatchElementCurrentLOD, CurrentLODData->SubSections[0].BatchElementCurrentLOD, *CurrentLODData->SubSections[0].ShaderCurrentNeighborLOD.ToString());
+		}
+		else
+		{
+			UE_LOG(LogLandscape, Warning, TEXT("\nComponent: [%s] -> MeshBatchLOD: %d, ComponentScreenSize: %f, UseCombinedMeshBatch: %d, ShaderCurrentLOD: %s, LODTessellation: %s\nHeightmap Texture Name: %s, Heightmap Streamed Mip: %d\nSubSections:\n|0| IndexToRender: %d, fLOD: %f, LOD: %d, NeighborLOD: %s\n|1| IndexToRender: %d, fLOD: %f, LOD: %d, NeighborLOD: %s\n|2| IndexToRender: %d, fLOD: %f, LOD: %d, NeighborLOD: %s\n|3| IndexToRender: %d, fLOD: %f, LOD: %d, NeighborLOD: %s\n"),
+				*SectionBase.ToString(), CurrentLODData->StaticMeshBatchLOD, CurrentLODData->ComponentScreenSize, CurrentLODData->UseCombinedMeshBatch, *CurrentLODData->ShaderCurrentLOD.ToString(), *CurrentLODData->LodTessellationParams.ToString(),
+				HeightmapTexture != nullptr ? *HeightmapTexture->GetFullName() : TEXT("Invalid"), HeightmapTexture != nullptr ? ((FTexture2DResource*)HeightmapTexture->Resource)->GetCurrentFirstMip() : INDEX_NONE,
+				CurrentLODData->SubSections[0].StaticBatchElementIndexToRender, CurrentLODData->SubSections[0].fBatchElementCurrentLOD, CurrentLODData->SubSections[0].BatchElementCurrentLOD, *CurrentLODData->SubSections[0].ShaderCurrentNeighborLOD.ToString(),
+				CurrentLODData->SubSections[1].StaticBatchElementIndexToRender, CurrentLODData->SubSections[1].fBatchElementCurrentLOD, CurrentLODData->SubSections[1].BatchElementCurrentLOD, *CurrentLODData->SubSections[1].ShaderCurrentNeighborLOD.ToString(),
+				CurrentLODData->SubSections[2].StaticBatchElementIndexToRender, CurrentLODData->SubSections[2].fBatchElementCurrentLOD, CurrentLODData->SubSections[2].BatchElementCurrentLOD, *CurrentLODData->SubSections[2].ShaderCurrentNeighborLOD.ToString(),
+				CurrentLODData->SubSections[3].StaticBatchElementIndexToRender, CurrentLODData->SubSections[3].fBatchElementCurrentLOD, CurrentLODData->SubSections[3].BatchElementCurrentLOD, *CurrentLODData->SubSections[3].ShaderCurrentNeighborLOD.ToString());
+		}
+	}
+#endif
+}
+
+bool FLandscapeComponentSceneProxy::CanUseMeshBatchForShadowCascade(int8 InLODIndex, float InShadowMapTextureResolution, float InShadowMapCascadeSize) const
+{
+#if 0
+	const FStaticMeshBatch* MeshBatch = nullptr;
+	const TArray<FStaticMeshBatch>& PrimitiveStaticMeshes = GetPrimitiveSceneInfo()->StaticMeshes;
+
+	if (!PrimitiveStaticMeshes.IsValidIndex(InLODIndex))
+	{
+		return false;
+	}
+
+	check(PrimitiveStaticMeshes[InLODIndex].CastShadow);
+	check(InLODIndex == PrimitiveStaticMeshes[InLODIndex].LODIndex);
+	MeshBatch = &PrimitiveStaticMeshes[InLODIndex];
+
+	check(MeshBatch != nullptr);
+
+	if (!MeshBatch->CastShadow || MeshBatch->TessellationDisablingShadowMapMeshSize == 0)
+	{
+		return true;
+	}
+
+	float WorldUnitsForOneTexel = InShadowMapCascadeSize / InShadowMapTextureResolution; // We assume Shadow Map texture to be squared
+	return MeshBatch->TessellationDisablingShadowMapMeshSize >= WorldUnitsForOneTexel * (GShadowMapWorldUnitsToTexelFactor != -1.0f ? GShadowMapWorldUnitsToTexelFactor : 1.0f);
+#endif
+
+	return true;
+}
+
+FLODMask FLandscapeComponentSceneProxy::GetCustomLOD(const FSceneView& InView, float InViewLODScale, int32 InForcedLODLevel, float& OutScreenSizeSquared) const
+{
+	SCOPE_CYCLE_COUNTER(STAT_LandscapeComputeCustomMeshBatchLOD);
+	FLODMask LODToRender;
+	OutScreenSizeSquared = 0.0f;
+
+	// Handle forced LOD level first
+	if (InForcedLODLevel >= 0)
+	{
+		int8 MinMeshLOD = MAX_int8;
+		int8 MaxMeshLOD = 0;
+		GetPrimitiveSceneInfo()->GetStaticMeshesLODRange(MinMeshLOD, MaxMeshLOD);
+
+		LODToRender.SetLOD(FMath::Clamp<int8>(InForcedLODLevel, MinMeshLOD, MaxMeshLOD));
+	}
+	else if (InView.Family->EngineShowFlags.LOD)
+	{
+		int8 PotentialLOD = INDEX_NONE;
+		OutScreenSizeSquared = GetComponentScreenSize(&InView, GetBounds().Origin, ComponentMaxExtend, GetBounds().SphereRadius);
+
+		if (NumSubsections > 1)
+		{
+			float SubSectionMaxExtend = ComponentMaxExtend / 2.0f;
+			float SubSectionRadius = GetBounds().SphereRadius / 2.0f;
+
+			// Compute screen size of each sub section to determine if we should use the combined logic or the individual logic
+			float ScreenSizeSquared = GetComponentScreenSize(&InView, SubSectionScreenSizeTestingPosition[0], SubSectionMaxExtend, SubSectionRadius);
+			PotentialLOD = GetLODFromScreenSize(ScreenSizeSquared, InViewLODScale);
+		}
+		else
+		{
+			PotentialLOD = GetLODFromScreenSize(OutScreenSizeSquared, InViewLODScale);
+		}
+
+		const auto FeatureLevel = InView.GetFeatureLevel();
+		bool HasTessellationEnabled = false;
+
+		if (FeatureLevel >= ERHIFeatureLevel::SM5)
+		{
+			ULandscapeMaterialInstanceConstant* LandscapeMIC = Cast<ULandscapeMaterialInstanceConstant>(AvailableMaterials[LODIndexToMaterialIndex[PotentialLOD]]);
+
+			if (LandscapeMIC == nullptr)
+			{
+				UMaterialInstanceDynamic* LandscapeMID = Cast<UMaterialInstanceDynamic>(AvailableMaterials[LODIndexToMaterialIndex[PotentialLOD]]);
+
+				if (LandscapeMID != nullptr)
+				{
+					LandscapeMIC = Cast<ULandscapeMaterialInstanceConstant>(LandscapeMID->Parent);
+				}
+			}
+
+			HasTessellationEnabled = MaterialHasTessellationEnabled[LODIndexToMaterialIndex[PotentialLOD]] && (LandscapeMIC != nullptr && !LandscapeMIC->bDisableTessellation);
+		}
+
+		int32 BaseMeshBatchIndex = MaterialIndexToStaticMeshBatchLOD[LODIndexToMaterialIndex[PotentialLOD]];
+
+		if (HasTessellationEnabled)
+		{
+			const int8 TessellatedMeshBatchLODIndex = 0;
+			const int8 NonTessellatedMeshBatchLODIndex = 1;
+			LODToRender.SetLOD(OutScreenSizeSquared >= TessellationComponentSquaredScreenSize * InViewLODScale ? BaseMeshBatchIndex + TessellatedMeshBatchLODIndex : BaseMeshBatchIndex + NonTessellatedMeshBatchLODIndex);
+		}
+		else
+		{
+			LODToRender.SetLOD(BaseMeshBatchIndex);
+		}
+	}
+
+	return LODToRender;
+}
+
+FLODMask FLandscapeComponentSceneProxy::GetCustomWholeSceneShadowLOD(const FSceneView& InView, float InViewLODScale, int32 InForcedLODLevel, const FLODMask& InVisibilePrimitiveLODMask, float InShadowMapTextureResolution, float InShadowMapCascadeSize, int8 InShadowCascadeId, bool InHasSelfShadow) const
+{
+	SCOPE_CYCLE_COUNTER(STAT_LandscapeComputeCustomShadowMeshBatchLOD);
+
+	FLODMask LODToRender;
+
+	// Handle forced LOD level first
+	if (InForcedLODLevel >= 0)
+	{
+		int8 MinMeshLOD = MAX_int8;
+		int8 MaxMeshLOD = 0;
+		GetPrimitiveSceneInfo()->GetStaticMeshesLODRange(MinMeshLOD, MaxMeshLOD);
+
+		LODToRender.SetLOD(FMath::Clamp<int8>(InForcedLODLevel, MinMeshLOD, MaxMeshLOD));
+	}
+	else if (!InHasSelfShadow) // Force lowest valid LOD
+	{
+		int8 MinMeshLOD = MAX_int8;
+		int8 MaxMeshLOD = 0;
+		GetPrimitiveSceneInfo()->GetStaticMeshesLODRange(MinMeshLOD, MaxMeshLOD);
+		LODToRender.SetLOD(MinMeshLOD);
+	}
+	else
+	{
+		const FViewCustomDataLOD* PrimitiveCustomData = (const FViewCustomDataLOD*)InView.GetCustomData(GetPrimitiveSceneInfo()->GetIndex());
+		int8 PotentialLOD = INDEX_NONE;
+		float ScreenSizeSquared = 0.0f;
+
+		if (PrimitiveCustomData == nullptr)
+		{
+			ScreenSizeSquared = GetComponentScreenSize(&InView, GetBounds().Origin, ComponentMaxExtend, GetBounds().SphereRadius);
+
+			if (NumSubsections > 1)
+			{
+				float SubSectionMaxExtend = ComponentMaxExtend / 2.0f;
+				float SubSectionRadius = GetBounds().SphereRadius / 2.0f;
+
+				// Compute screen size of each sub section to determine if we should use the combined logic or the individual logic
+				float SubSectionScreenSizeSquared = GetComponentScreenSize(&InView, SubSectionScreenSizeTestingPosition[0], SubSectionMaxExtend, SubSectionRadius);
+				PotentialLOD = GetLODFromScreenSize(SubSectionScreenSizeSquared, InViewLODScale);
+			}
+			else
+			{
+				PotentialLOD = GetLODFromScreenSize(ScreenSizeSquared, InViewLODScale);
+			}
+		}
+		else
+		{
+			ScreenSizeSquared = PrimitiveCustomData->ComponentScreenSize;
+			PotentialLOD = PrimitiveCustomData->SubSections[0].BatchElementCurrentLOD;
+		}
+
+		const auto FeatureLevel = InView.GetFeatureLevel();
+		bool HasTessellationEnabled = false;
+
+		if (FeatureLevel >= ERHIFeatureLevel::SM5)
+		{
+			ULandscapeMaterialInstanceConstant* LandscapeMIC = Cast<ULandscapeMaterialInstanceConstant>(AvailableMaterials[LODIndexToMaterialIndex[PotentialLOD]]);
+
+			if (LandscapeMIC == nullptr)
+			{
+				UMaterialInstanceDynamic* LandscapeMID = Cast<UMaterialInstanceDynamic>(AvailableMaterials[LODIndexToMaterialIndex[PotentialLOD]]);
+
+				if (LandscapeMID != nullptr)
+				{
+					LandscapeMIC = Cast<ULandscapeMaterialInstanceConstant>(LandscapeMID->Parent);
+				}
+			}
+
+			HasTessellationEnabled = MaterialHasTessellationEnabled[LODIndexToMaterialIndex[PotentialLOD]] && (LandscapeMIC != nullptr && !LandscapeMIC->bDisableTessellation);
+		}
+
+		int32 BaseMeshBatchIndex = MaterialIndexToStaticMeshBatchLOD[LODIndexToMaterialIndex[PotentialLOD]];
+
+		if (HasTessellationEnabled)
+		{
+			const int8 ShadowTessellatedMeshBatchLODIndex = 2;
+			const int8 ShadowNonTessellatedMeshBatchLODIndex = 3;
+
+			bool UseTessellationMeshBatch = ScreenSizeSquared >= TessellationComponentSquaredScreenSize * InViewLODScale;
+
+			if (UseTessellationMeshBatch)
+			{
+				if (!CanUseMeshBatchForShadowCascade(ShadowTessellatedMeshBatchLODIndex, InShadowMapTextureResolution, InShadowMapCascadeSize))
+				{
+					UseTessellationMeshBatch = false;
+				}
+			}
+
+			LODToRender.SetLOD(UseTessellationMeshBatch ? BaseMeshBatchIndex + ShadowTessellatedMeshBatchLODIndex : BaseMeshBatchIndex + ShadowNonTessellatedMeshBatchLODIndex);
+		}
+		else
+		{
+			LODToRender.SetLOD(BaseMeshBatchIndex);
+		}
+	}
+
+	return LODToRender;
+}
+
 namespace
@@ -3024,42 +4039,33 @@
 		return;
 	}
 
-	const uint16 NumLineQuads = ((uint16)(FMath::Sqrt(NumOccluderVertices)/NumSubsections) - 1);
-	const uint16 NumLineVtx = NumLineQuads + 1;
-	check(NumLineVtx*NumLineVtx*NumSubsections*NumSubsections == NumOccluderVertices);
+	uint16 NumLineQuads = ((uint16)FMath::Sqrt(NumOccluderVertices) - 1);
+	uint16 NumLineVtx = NumLineQuads + 1;
+	check(NumLineVtx*NumLineVtx == NumOccluderVertices);
 
-	const int32 NumTrisPerSubsection = NumLineQuads*NumLineQuads * 2;
-	const int32 NumIndicesPerSubsection = NumTrisPerSubsection * 3;
+	int32 NumTris = NumLineQuads*NumLineQuads * 2;
+	int32 NumIndices = NumTris * 3;
 	OccluderIndicesSP = MakeShared<FOccluderIndexArray, ESPMode::ThreadSafe>();
-	OccluderIndicesSP->SetNumUninitialized(NumIndicesPerSubsection*NumSubsections*NumSubsections, false);
+	OccluderIndicesSP->SetNumUninitialized(NumIndices, false);
 
 	uint16* OcclusionIndices = OccluderIndicesSP->GetData();
 	const uint16 NumLineVtxPlusOne = NumLineVtx + 1;
 	const uint16 QuadIndices[2][3] = { {0, NumLineVtx, NumLineVtxPlusOne}, {0, NumLineVtxPlusOne, 1} };
 	uint16 QuadOffset = 0;
 	int32 Index = 0;
-
-	for (int32 SubY = 0; SubY < NumSubsections; SubY++)
+	for (int32 y = 0; y < NumLineQuads; y++)
 	{
-		for (int32 SubX = 0; SubX < NumSubsections; SubX++)
+		for (int32 x = 0; x < NumLineQuads; x++)
 		{
-			QuadOffset = (SubY*NumSubsections + SubX) * NumLineVtx * NumLineVtx;
-
-	        for (int32 y = 0; y < NumLineQuads; y++)
-	        {
-		        for (int32 x = 0; x < NumLineQuads; x++)
-		        {
-			        for (int32 i = 0; i < 2; i++)
-			        {
-				        OcclusionIndices[Index++] = QuadIndices[i][0] + QuadOffset;
-				        OcclusionIndices[Index++] = QuadIndices[i][1] + QuadOffset;
-				        OcclusionIndices[Index++] = QuadIndices[i][2] + QuadOffset;
-			        }
-			        QuadOffset++;
-		        }
-		        QuadOffset++;
-	        }
+			for (int32 i = 0; i < 2; i++)
+			{
+				OcclusionIndices[Index++] = QuadIndices[i][0] + QuadOffset;
+				OcclusionIndices[Index++] = QuadIndices[i][1] + QuadOffset;
+				OcclusionIndices[Index++] = QuadIndices[i][2] + QuadOffset;
+			}
+			QuadOffset++;
 		}
+		QuadOffset++;
 	}
 
 	INC_DWORD_STAT_BY(STAT_LandscapeOccluderMem, OccluderIndicesSP->GetAllocatedSize());
diff --git a/Engine/Source/Runtime/VulkanRHI/Private/VulkanCommands.cpp b/Engine/Source/Runtime/VulkanRHI/Private/VulkanCommands.cpp
--- a/Engine/Source/Runtime/VulkanRHI/Private/VulkanCommands.cpp
+++ b/Engine/Source/Runtime/VulkanRHI/Private/VulkanCommands.cpp
@@ -501,7 +501,7 @@ inline void /*FVulkanCommandListContext::*/SetShaderUniformBufferResources(FVulk
 				|| DescriptorType == VK_DESCRIPTOR_TYPE_STORAGE_IMAGE
 				|| DescriptorType == VK_DESCRIPTOR_TYPE_STORAGE_TEXEL_BUFFER);
 			FRHIUnorderedAccessView* CurrentUAV = (FRHIUnorderedAccessView*)(ResourceArray[ResourceInfo.SourceUBResourceIndex].GetReference());
-			if (UAV)
+			if (CurrentUAV)
 			{
 				FVulkanUnorderedAccessView* UAV = ResourceCast(CurrentUAV);
 				State->SetUAVForUBResource(GlobalRemappingInfo[ResourceInfo.GlobalIndex].NewDescriptorSet, GlobalRemappingInfo[ResourceInfo.GlobalIndex].NewBindingIndex, UAV);
