Same condition to FlushGrassMaps then to Render them so we don't flush and fail to render.

diff --git a/Engine/Source/Runtime/Landscape/Private/LandscapeGrass.cpp b/Engine/Source/Runtime/Landscape/Private/LandscapeGrass.cpp
--- a/Engine/Source/Runtime/Landscape/Private/LandscapeGrass.cpp
+++ b/Engine/Source/Runtime/Landscape/Private/LandscapeGrass.cpp
@@ -1975,11 +1975,14 @@ void ALandscapeProxy::FlushGrassComponents(const TSet<ULandscapeComponent*>* Onl
 			}
 		}
 #if WITH_EDITOR
-		if (GIsEditor && bFlushGrassMaps && GetWorld() && GetWorld()->FeatureLevel >= ERHIFeatureLevel::SM5)
+		if (bFlushGrassMaps)
 		{
 			for (ULandscapeComponent* Component : *OnlyForComponents)
 			{
-				Component->RemoveGrassMap();
+				if (Component->CanRenderGrassMap())
+				{
+					Component->RemoveGrassMap();
+				}
 			}
 		}
 #endif
@@ -2020,14 +2023,17 @@ void ALandscapeProxy::FlushGrassComponents(const TSet<ULandscapeComponent*>* Onl
 #if WITH_EDITOR
 		UWorld* World = GetWorld();
 
-		if (GIsEditor && bFlushGrassMaps && World != nullptr && World->Scene != nullptr && World->Scene->GetFeatureLevel() >= ERHIFeatureLevel::SM5)
+		if (bFlushGrassMaps)
 		{
 			// Clear GrassMaps
 			for (UActorComponent* Component : GetComponents())
 			{
 				if (ULandscapeComponent* LandscapeComp = Cast<ULandscapeComponent>(Component))
 				{
-					LandscapeComp->RemoveGrassMap();
+					if (LandscapeComp->CanRenderGrassMap())
+					{
+						LandscapeComp->RemoveGrassMap();
+					}
 				}
 			}
 		}
