- Restored parallel primtive info fix for Landscapes

diff --git a/Engine/Source/Runtime/Landscape/Private/LandscapeRender.cpp b/Engine/Source/Runtime/Landscape/Private/LandscapeRender.cpp
--- a/Engine/Source/Runtime/Landscape/Private/LandscapeRender.cpp
+++ b/Engine/Source/Runtime/Landscape/Private/LandscapeRender.cpp
@@ -1316,7 +1316,7 @@ FLandscapeComponentSceneProxy::FLandscapeComponentSceneProxy(ULandscapeComponent
 				MaterialRelevances.Add(MaterialInterface->GetRelevance_Concurrent(FeatureLevel));
 			}
 
-			bRequiresAdjacencyInformation |= MaterialSettingsRequireAdjacencyInformation_GameThread(MaterialInterface, XYOffsetmapTexture == nullptr ? &FLandscapeVertexFactory::StaticType : &FLandscapeXYOffsetVertexFactory::StaticType, InComponent->GetWorld()->FeatureLevel.GetValue());
+			bRequiresAdjacencyInformation |= RequiresAdjacencyInformation(MaterialInterface, XYOffsetmapTexture == nullptr ? &FLandscapeVertexFactory::StaticType : &FLandscapeXYOffsetVertexFactory::StaticType, InComponent->GetWorld()->FeatureLevel.GetValue());
 
 			bool HasTessellationEnabled = false;
 
