Fix bad buffersize calculation leading to corrupt landscape edit (Case: heightmap Width != Height)

#rnx

diff --git a/Engine/Source/Runtime/Landscape/Private/LandscapeEditInterface.cpp b/Engine/Source/Runtime/Landscape/Private/LandscapeEditInterface.cpp
--- a/Engine/Source/Runtime/Landscape/Private/LandscapeEditInterface.cpp
+++ b/Engine/Source/Runtime/Landscape/Private/LandscapeEditInterface.cpp
@@ -5689,9 +5689,10 @@ bool FLandscapeTextureDataInfo::UpdateTextureData()
 				break;
 			}
 
-			const uint32 SrcSize = (Texture->Source.GetSizeX()) >> i;
-			const uint32 SrcPitch = (SrcSize * DataSize);
-			const uint32 BufferSize = SrcSize * SrcSize * DataSize;
+			const uint32 SrcSizeX = (Texture->Source.GetSizeX()) >> i;
+			const uint32 SrcSizeY = (Texture->Source.GetSizeY()) >> i;
+			const uint32 SrcPitch = (SrcSizeX * DataSize);
+			const uint32 BufferSize = SrcSizeX * SrcSizeY * DataSize;
 
 			// Copy Mip update data so we can avoid waiting for Render thread in calling method
 			FMipInfo* CopyMipInfo = new FMipInfo();
