Fixed crash when creating a landscape from scratch with edit layers enabled

patrick.enfedaque

diff --git a/Engine/Source/Runtime/Landscape/Private/LandscapeEdit.cpp b/Engine/Source/Runtime/Landscape/Private/LandscapeEdit.cpp
--- a/Engine/Source/Runtime/Landscape/Private/LandscapeEdit.cpp
+++ b/Engine/Source/Runtime/Landscape/Private/LandscapeEdit.cpp
@@ -2882,19 +2882,18 @@ LANDSCAPE_API void ALandscapeProxy::Import(const FGuid& InGuid, int32 InMinX, in
 
 	if (CanHaveLayersContent())
 	{
-		// Components need to be registered to be able to import the layer content and we will remove them if they should have not been visible
-		bool ShouldComponentBeRegistered = GetLevel()->bIsVisible;
-		RegisterAllComponents();
-
+		// Create the default layer first
 		ALandscape* LandscapeActor = GetLandscapeActor();
 		check(LandscapeActor != nullptr);
-
-		// Create the default layer
 		if (LandscapeActor->GetLayerCount() == 0 && InImportLayers == nullptr)
 		{
 			LandscapeActor->CreateDefaultLayer();
 		}
 
+		// Components need to be registered to be able to import the layer content and we will remove them if they should have not been visible
+		bool ShouldComponentBeRegistered = GetLevel()->bIsVisible;
+		RegisterAllComponents();
+
 		TSet<ULandscapeComponent*> ComponentsToProcess;
 
 		struct FLayerImportSettings
