Fix assert geting Landscape render data on mobile
#rnx

diff --git a/Engine/Source/Runtime/Landscape/Private/LandscapeRenderMobile.cpp b/Engine/Source/Runtime/Landscape/Private/LandscapeRenderMobile.cpp
--- a/Engine/Source/Runtime/Landscape/Private/LandscapeRenderMobile.cpp
+++ b/Engine/Source/Runtime/Landscape/Private/LandscapeRenderMobile.cpp
@@ -565,7 +565,8 @@ void FLandscapeComponentSceneProxyMobile::CreateRenderThreadResources()
 
 TSharedPtr<FLandscapeMobileRenderData, ESPMode::ThreadSafe> FLandscapeComponentDerivedData::GetRenderData()
 {
-	check(IsInGameThread());
+	// This function is expected to be called from either the GameThread or via ParallelFor from the GameThread
+	check(!IsInRenderingThread());
 
 	if (FPlatformProperties::RequiresCookedData() && CachedRenderData.IsValid())
 	{
