Fix bad index buffer generation on mobile landscape.
Vertex Key generation was not unique for a 255*255 grid.
Removed it along with its use as a TMap key and replaced with simpler array lookup.
Cherry pick integration from Dev-Rendering

diff --git a/Engine/Source/Runtime/Landscape/Classes/LandscapeComponent.h b/Engine/Source/Runtime/Landscape/Classes/LandscapeComponent.h
--- a/Engine/Source/Runtime/Landscape/Classes/LandscapeComponent.h
+++ b/Engine/Source/Runtime/Landscape/Classes/LandscapeComponent.h
@@ -150,25 +150,25 @@ public:
 	friend FArchive& operator<<(FArchive& Ar, FLandscapeComponentDerivedData& Data);
 };
 
-/* Used to uniquely reference a landscape vertex in a component, and generate a key suitable for a TMap. */
+/* Used to uniquely reference a landscape vertex in a component. */
 struct FLandscapeVertexRef
 {
 	FLandscapeVertexRef(int16 InX, int16 InY, int8 InSubX, int8 InSubY)
-	: X(InX)
-	, Y(InY)
-	, SubX(InSubX)
-	, SubY(InSubY)
+		: X(InX)
+		, Y(InY)
+		, SubX(InSubX)
+		, SubY(InSubY)
 	{}
-	int16 X;
-	int16 Y;
-	int8 SubX;
-	int8 SubY;
 
-	uint64 MakeKey() const
+	uint32 X : 8;
+	uint32 Y : 8;
+	uint32 SubX : 8;
+	uint32 SubY : 8;
+
+	/** Helper to provide a standard ordering for vertex arrays. */
+	static int32 GetVertexIndex(FLandscapeVertexRef Vert, int32 SubsectionCount, int32 SubsectionVerts)
 	{
-		// this is very bad for TMap
-		//return (uint64)X << 32 | (uint64)Y << 16 | (uint64)SubX << 8 | (uint64)SubY;
-		return HashCombine((uint32(X) << 8) | uint32(SubY), (uint32(SubX) << 24) | uint32(Y));
+		return (Vert.SubY * SubsectionVerts + Vert.Y) * SubsectionVerts * SubsectionCount + Vert.SubX * SubsectionVerts + Vert.X;
 	}
 };
 
diff --git a/Engine/Source/Runtime/Landscape/Private/LandscapeEdit.cpp b/Engine/Source/Runtime/Landscape/Private/LandscapeEdit.cpp
--- a/Engine/Source/Runtime/Landscape/Private/LandscapeEdit.cpp
+++ b/Engine/Source/Runtime/Landscape/Private/LandscapeEdit.cpp
@@ -6535,7 +6535,7 @@ void SerializeHoleRenderData(FMemoryArchive& Ar, FLandscapeHoleRenderData<INDEX_
 
 // Take the processed hole map and generate the hole render data.
 template <typename INDEX_TYPE>
-void BuildHoleRenderData(int32 InNumSubsections, int32 InSubsectionSizeVerts, TArray<uint8> const& InVisibilityData, TMap<uint64, int32>& InVertexMap, FLandscapeHoleRenderData<INDEX_TYPE>& OutHoleRenderData)
+void BuildHoleRenderData(int32 InNumSubsections, int32 InSubsectionSizeVerts, TArray<uint8> const& InVisibilityData, TArray<uint32>& InVertexToIndexMap, FLandscapeHoleRenderData<INDEX_TYPE>& OutHoleRenderData)
 {
 	const int32 SizeVerts = InNumSubsections * InSubsectionSizeVerts;
 	const int32 SubsectionSizeQuads = InSubsectionSizeVerts - 1;
@@ -6561,10 +6561,10 @@ void BuildHoleRenderData(int32 InNumSubsections, int32 InSubsectionSizeVerts, TA
 					const bool bIsHole = InVisibilityData[VertexIndex] < VisThreshold;
 					if (bIsHole)
 					{
-						INDEX_TYPE i00 = *(InVertexMap.Find(FLandscapeVertexRef(x0, y0, SubX, SubY).MakeKey()));
-						INDEX_TYPE i10 = *(InVertexMap.Find(FLandscapeVertexRef(x1, y0, SubX, SubY).MakeKey()));
-						INDEX_TYPE i11 = *(InVertexMap.Find(FLandscapeVertexRef(x1, y1, SubX, SubY).MakeKey()));
-						INDEX_TYPE i01 = *(InVertexMap.Find(FLandscapeVertexRef(x0, y1, SubX, SubY).MakeKey()));
+						INDEX_TYPE i00 = InVertexToIndexMap[FLandscapeVertexRef::GetVertexIndex(FLandscapeVertexRef(x0, y0, SubX, SubY), InNumSubsections, InSubsectionSizeVerts)];
+						INDEX_TYPE i10 = InVertexToIndexMap[FLandscapeVertexRef::GetVertexIndex(FLandscapeVertexRef(x1, y0, SubX, SubY), InNumSubsections, InSubsectionSizeVerts)];
+						INDEX_TYPE i11 = InVertexToIndexMap[FLandscapeVertexRef::GetVertexIndex(FLandscapeVertexRef(x1, y1, SubX, SubY), InNumSubsections, InSubsectionSizeVerts)];
+						INDEX_TYPE i01 = InVertexToIndexMap[FLandscapeVertexRef::GetVertexIndex(FLandscapeVertexRef(x0, y1, SubX, SubY), InNumSubsections, InSubsectionSizeVerts)];
 
 						OutHoleRenderData.HoleIndices.Add(i00);
 						OutHoleRenderData.HoleIndices.Add(i11);
@@ -6650,9 +6650,14 @@ void ULandscapeComponent::GeneratePlatformVertexData(const ITargetPlatform* Targ
 
 	// Layout index buffer to determine best vertex order.
 	// This vertex layout code is duplicated in FLandscapeSharedBuffers::CreateIndexBuffers() to create matching index buffers at runtime.
-	TMap<uint64, int32> VertexMap;
+	const int32 NumVertices = FMath::Square(SubsectionSizeVerts * NumSubsections);
+	
+	TArray<uint32> VertexToIndexMap;
+	VertexToIndexMap.AddUninitialized(NumVertices);
+	FMemory::Memset(VertexToIndexMap.GetData(), 0xFF, NumVertices * sizeof(uint32));
+	
 	TArray<FLandscapeVertexRef> VertexOrder;
-	VertexOrder.Empty(FMath::Square(SubsectionSizeVerts * NumSubsections));
+	VertexOrder.Empty(NumVertices);
 
 	for (int32 Mip = MaxLOD; Mip >= 0; Mip--)
 	{
@@ -6663,43 +6668,22 @@ void ULandscapeComponent::GeneratePlatformVertexData(const ITargetPlatform* Targ
 		{
 			for (int32 SubX = 0; SubX < NumSubsections; SubX++)
 			{
-				for (int32 y = 0; y < LodSubsectionSizeQuads; y++)
+				for (int32 Y = 0; Y < LodSubsectionSizeQuads; Y++)
 				{
-					for (int32 x = 0; x < LodSubsectionSizeQuads; x++)
+					for (int32 X = 0; X < LodSubsectionSizeQuads; X++)
 					{
-						int32 x0 = FMath::RoundToInt((float)x * MipRatio);
-						int32 y0 = FMath::RoundToInt((float)y * MipRatio);
-						int32 x1 = FMath::RoundToInt((float)(x + 1) * MipRatio);
-						int32 y1 = FMath::RoundToInt((float)(y + 1) * MipRatio);
-
-						FLandscapeVertexRef V1(x0, y0, SubX, SubY);
-						FLandscapeVertexRef V2(x1, y0, SubX, SubY);
-						FLandscapeVertexRef V3(x1, y1, SubX, SubY);
-						FLandscapeVertexRef V4(x0, y1, SubX, SubY);
-
-						uint64 Key1 = V1.MakeKey();
-						if (VertexMap.Find(Key1) == nullptr)
-						{
-							VertexMap.Add(Key1, VertexOrder.Num());
-							VertexOrder.Add(V1);
-						}
-						uint64 Key2 = V2.MakeKey();
-						if (VertexMap.Find(Key2) == nullptr)
-						{
-							VertexMap.Add(Key2, VertexOrder.Num());
-							VertexOrder.Add(V2);
-						}
-						uint64 Key3 = V3.MakeKey();
-						if (VertexMap.Find(Key3) == nullptr)
+						for (int32 CornerId = 0; CornerId < 4; CornerId++)
 						{
-							VertexMap.Add(Key3, VertexOrder.Num());
-							VertexOrder.Add(V3);
-						}
-						uint64 Key4 = V4.MakeKey();
-						if (VertexMap.Find(Key4) == nullptr)
-						{
-							VertexMap.Add(Key4, VertexOrder.Num());
-							VertexOrder.Add(V4);
+							const int32 CornerX = FMath::RoundToInt((float)(X + (CornerId & 1)) * MipRatio);
+							const int32 CornerY = FMath::RoundToInt((float)(Y + (CornerId >> 1)) * MipRatio);
+							const FLandscapeVertexRef VertexRef(CornerX, CornerY, SubX, SubY);
+
+							const int32 VertexIndex = FLandscapeVertexRef::GetVertexIndex(VertexRef, NumSubsections, SubsectionSizeVerts);
+							if (VertexToIndexMap[VertexIndex] == 0xFFFFFFFF)
+							{
+								VertexToIndexMap[VertexIndex] = VertexOrder.Num();
+								VertexOrder.Add(VertexRef);
+							}
 						}
 					}
 				}
@@ -6707,10 +6691,9 @@ void ULandscapeComponent::GeneratePlatformVertexData(const ITargetPlatform* Targ
 		}
 	}
 
-	if (VertexOrder.Num() != FMath::Square(SubsectionSizeVerts) * FMath::Square(NumSubsections)) 
+	if (VertexOrder.Num() != NumVertices)
 	{
-		UE_LOG(LogLandscape, Warning, TEXT("VertexOrder count of %d did not match expected size of %d"), 
-			VertexOrder.Num(), FMath::Square(SubsectionSizeVerts) * FMath::Square(NumSubsections));
+		UE_LOG(LogLandscape, Warning, TEXT("VertexOrder count of %d did not match expected size of %d"), VertexOrder.Num(), NumVertices);
 	}
 
 	// Build and serialize hole render data which includes a unique index buffer with the holes missing.
@@ -6723,16 +6706,16 @@ void ULandscapeComponent::GeneratePlatformVertexData(const ITargetPlatform* Targ
 		GetHoleBounds(SubsectionSizeVerts * NumSubsections, VisibilityData, HoleBounds);
 		BuildHoleVertexLods(SubsectionSizeVerts * NumSubsections, NumHoleLods, HoleBounds, HoleVertexLods);
 
-		if (VertexOrder.Num() <= UINT16_MAX)
+		if (NumVertices <= UINT16_MAX)
 		{
 			FLandscapeHoleRenderData<uint16> HoleRenderData;
-			BuildHoleRenderData(NumSubsections, SubsectionSizeVerts, VisibilityData, VertexMap, HoleRenderData);
+			BuildHoleRenderData(NumSubsections, SubsectionSizeVerts, VisibilityData, VertexToIndexMap, HoleRenderData);
 			SerializeHoleRenderData(PlatformAr, HoleRenderData);
 		}
 		else
 		{
 			FLandscapeHoleRenderData<uint32> HoleRenderData;
-			BuildHoleRenderData(NumSubsections, SubsectionSizeVerts, VisibilityData, VertexMap, HoleRenderData);
+			BuildHoleRenderData(NumSubsections, SubsectionSizeVerts, VisibilityData, VertexToIndexMap, HoleRenderData);
 			SerializeHoleRenderData(PlatformAr, HoleRenderData);
 		}
 	}
@@ -6744,7 +6727,7 @@ void ULandscapeComponent::GeneratePlatformVertexData(const ITargetPlatform* Targ
 	MobileVertices.AddZeroed(NumMobileVertices);
 	FLandscapeMobileVertex* DstVert = MobileVertices.GetData();
 
-	for (int32 Idx = 0; Idx < VertexOrder.Num(); Idx++)
+	for (int32 Idx = 0; Idx < NumVertices; Idx++)
 	{
 		// Store XY position info
 		const int32 X = VertexOrder[Idx].X;
diff --git a/Engine/Source/Runtime/Landscape/Private/LandscapeRender.cpp b/Engine/Source/Runtime/Landscape/Private/LandscapeRender.cpp
--- a/Engine/Source/Runtime/Landscape/Private/LandscapeRender.cpp
+++ b/Engine/Source/Runtime/Landscape/Private/LandscapeRender.cpp
@@ -3802,7 +3802,10 @@ void FLandscapeSharedBuffers::CreateIndexBuffers(ERHIFeatureLevel::Type InFeatur
 		}
 	}
 
-	TMap<uint64, INDEX_TYPE> VertexMap;
+	TArray<INDEX_TYPE> VertexToIndexMap;
+	VertexToIndexMap.AddUninitialized(FMath::Square(SubsectionSizeVerts * NumSubsections));
+	FMemory::Memset(VertexToIndexMap.GetData(), 0xff, NumVertices * sizeof(INDEX_TYPE));
+
 	INDEX_TYPE VertexCount = 0;
 	int32 SubsectionSizeQuads = SubsectionSizeVerts - 1;
 
@@ -3823,7 +3826,7 @@ void FLandscapeSharedBuffers::CreateIndexBuffers(ERHIFeatureLevel::Type InFeatur
 
 		if (InFeatureLevel <= ERHIFeatureLevel::ES3_1)
 		{
-			// ES2 version
+			// mobile version shares vertices across LODs to save memory
 			float MipRatio = (float)SubsectionSizeQuads / (float)LodSubsectionSizeQuads; // Morph current MIP to base MIP
 
 			for (int32 SubY = 0; SubY < NumSubsections; SubY++)
@@ -3838,91 +3841,40 @@ void FLandscapeSharedBuffers::CreateIndexBuffers(ERHIFeatureLevel::Type InFeatur
 					MaxIndex = 0;
 					MinIndex = MAX_int32;
 
-					for (int32 y = 0; y < LodSubsectionSizeQuads; y++)
+					for (int32 Y = 0; Y < LodSubsectionSizeQuads; Y++)
 					{
-						for (int32 x = 0; x < LodSubsectionSizeQuads; x++)
+						for (int32 X = 0; X < LodSubsectionSizeQuads; X++)
 						{
-							int32 x0 = FMath::RoundToInt((float)x * MipRatio);
-							int32 y0 = FMath::RoundToInt((float)y * MipRatio);
-							int32 x1 = FMath::RoundToInt((float)(x + 1) * MipRatio);
-							int32 y1 = FMath::RoundToInt((float)(y + 1) * MipRatio);
-
-							FLandscapeVertexRef V00(x0, y0, SubX, SubY);
-							FLandscapeVertexRef V10(x1, y0, SubX, SubY);
-							FLandscapeVertexRef V11(x1, y1, SubX, SubY);
-							FLandscapeVertexRef V01(x0, y1, SubX, SubY);
-
-							uint64 Key00 = V00.MakeKey();
-							uint64 Key10 = V10.MakeKey();
-							uint64 Key11 = V11.MakeKey();
-							uint64 Key01 = V01.MakeKey();
-
-							INDEX_TYPE i00;
-							INDEX_TYPE i10;
-							INDEX_TYPE i11;
-							INDEX_TYPE i01;
-
-							INDEX_TYPE* KeyPtr = VertexMap.Find(Key00);
-							if (KeyPtr == nullptr)
-							{
-								i00 = VertexCount++;
-								VertexMap.Add(Key00, i00);
-							}
-							else
-							{
-								i00 = *KeyPtr;
-							}
+							INDEX_TYPE QuadIndices[4];
 
-							KeyPtr = VertexMap.Find(Key10);
-							if (KeyPtr == nullptr)
+							for (int32 CornerId = 0; CornerId < 4; CornerId++)
 							{
-								i10 = VertexCount++;
-								VertexMap.Add(Key10, i10);
+								const int32 CornerX = FMath::RoundToInt((float)(X + (CornerId & 1)) * MipRatio);
+								const int32 CornerY = FMath::RoundToInt((float)(Y + (CornerId >> 1)) * MipRatio);
+								const FLandscapeVertexRef VertexRef(CornerX, CornerY, SubX, SubY);
+
+								const INDEX_TYPE VertexIndex = FLandscapeVertexRef::GetVertexIndex(VertexRef, NumSubsections, SubsectionSizeVerts);
+								if (VertexToIndexMap[VertexIndex] == INDEX_TYPE(-1))
+								{
+									VertexToIndexMap[VertexIndex] = QuadIndices[CornerId] = VertexCount++;
+								}
+								else
+								{
+									QuadIndices[CornerId] = VertexToIndexMap[VertexIndex];
+								}
+
+								// update the min/max index ranges
+								MaxIndex = FMath::Max<int32>(MaxIndex, QuadIndices[CornerId]);
+								MinIndex = FMath::Min<int32>(MinIndex, QuadIndices[CornerId]);
 							}
-							else
-							{
-								i10 = *KeyPtr;
-							}
-
-							KeyPtr = VertexMap.Find(Key11);
-							if (KeyPtr == nullptr)
-							{
-								i11 = VertexCount++;
-								VertexMap.Add(Key11, i11);
-							}
-							else
-							{
-								i11 = *KeyPtr;
-							}
-
-							KeyPtr = VertexMap.Find(Key01);
-							if (KeyPtr == nullptr)
-							{
-								i01 = VertexCount++;
-								VertexMap.Add(Key01, i01);
-							}
-							else
-							{
-								i01 = *KeyPtr;
-							}
-
-							// Update the min/max index ranges
-							MaxIndex = FMath::Max<int32>(MaxIndex, i00);
-							MinIndex = FMath::Min<int32>(MinIndex, i00);
-							MaxIndex = FMath::Max<int32>(MaxIndex, i10);
-							MinIndex = FMath::Min<int32>(MinIndex, i10);
-							MaxIndex = FMath::Max<int32>(MaxIndex, i11);
-							MinIndex = FMath::Min<int32>(MinIndex, i11);
-							MaxIndex = FMath::Max<int32>(MaxIndex, i01);
-							MinIndex = FMath::Min<int32>(MinIndex, i01);
 
-							SubIndices.Add(i00);
-							SubIndices.Add(i11);
-							SubIndices.Add(i10);
+							SubIndices.Add(QuadIndices[0]);
+							SubIndices.Add(QuadIndices[3]);
+							SubIndices.Add(QuadIndices[1]);
 
-							SubIndices.Add(i00);
-							SubIndices.Add(i01);
-							SubIndices.Add(i11);
+							SubIndices.Add(QuadIndices[0]);
+							SubIndices.Add(QuadIndices[2]);
+							SubIndices.Add(QuadIndices[3]);
 						}
 					}
 
@@ -3938,7 +3890,7 @@ void FLandscapeSharedBuffers::CreateIndexBuffers(ERHIFeatureLevel::Type InFeatur
 		}
 		else
 		{
-			// non-ES2 version
+			// non-mobile version
 			int32 SubOffset = 0;
 			for (int32 SubY = 0; SubY < NumSubsections; SubY++)
 			{
