Fix mobile runtime virtual texture.
Cherry pick integration from Dev-Rendering
#rnx

diff --git a/Engine/Shaders/Private/MobileBasePassPixelShader.usf b/Engine/Shaders/Private/MobileBasePassPixelShader.usf
--- a/Engine/Shaders/Private/MobileBasePassPixelShader.usf
+++ b/Engine/Shaders/Private/MobileBasePassPixelShader.usf
@@ -107,14 +107,6 @@ float4 SpotLightDirectionAndSpecularScale[MAX_DYNAMIC_POINT_LIGHTS];
 	#endif
 #endif
 
-#if LIGHTMAP_VT_ENABLED || NUM_VIRTUALTEXTURE_SAMPLES
-#if COMPILER_HLSLCC
-	#define VTFeedbackBuffer MobileSceneTextures.VirtualTextureFeedbackUAV
-#else
-	RWBuffer<uint> VTFeedbackBuffer : register(VT_FEEDBACK_REGISTER);
-#endif
-#endif
-
 half PhongApprox( half Roughness, half RoL )
 {
 	half a = Roughness * Roughness;			// 1 mul
@@ -841,9 +833,9 @@ void Main(
 	FinalizeVirtualTextureFeedback(
 		MaterialParameters.VirtualTextureFeedback,
 		MaterialParameters.SvPosition,
-		Opacity,
+		ShadingModelContext.Opacity,
 		View.FrameNumber,
-		VTFeedbackBuffer
+		View.VTFeedbackBuffer
 	);
 #endif
 }
diff --git a/Engine/Source/Runtime/Landscape/Private/LandscapeRenderMobile.cpp b/Engine/Source/Runtime/Landscape/Private/LandscapeRenderMobile.cpp
--- a/Engine/Source/Runtime/Landscape/Private/LandscapeRenderMobile.cpp
+++ b/Engine/Source/Runtime/Landscape/Private/LandscapeRenderMobile.cpp
@@ -399,7 +399,8 @@ void FLandscapeComponentSceneProxyMobile::CreateRenderThreadResources()
 	}
 	
 	auto FeatureLevel = GetScene().GetFeatureLevel();
-	// Use only Index buffers
+	
+	// Use only index buffers from the shared buffers since the vertex buffers are unique per proxy on mobile
 	SharedBuffers = FLandscapeComponentSceneProxy::SharedBuffersMap.FindRef(SharedBuffersKey);
 	if (SharedBuffers == nullptr)
 	{
@@ -410,44 +411,58 @@ void FLandscapeComponentSceneProxyMobile::CreateRenderThreadResources()
 			GetScene().GetFeatureLevel(), false, NumOcclusionVertices);
 
 		FLandscapeComponentSceneProxy::SharedBuffersMap.Add(SharedBuffersKey, SharedBuffers);
-				
-		if (UseVirtualTexturing(FeatureLevel))
-		{
-			//todo[vt]: We will need a version of this to support XYOffsetmapTexture
-			FLandscapeFixedGridVertexFactoryMobile* LandscapeVertexFactory = new FLandscapeFixedGridVertexFactoryMobile(FeatureLevel);
-			LandscapeVertexFactory->MobileData.PositionComponent = FVertexStreamComponent(MobileRenderData->VertexBuffer, STRUCT_OFFSET(FLandscapeMobileVertex,Position), sizeof(FLandscapeMobileVertex), VET_UByte4N);
-			for( uint32 Index = 0; Index < LANDSCAPE_MAX_ES_LOD_COMP; ++Index )
-			{
-				LandscapeVertexFactory->MobileData.LODHeightsComponent.Add
-				(FVertexStreamComponent(MobileRenderData->VertexBuffer, STRUCT_OFFSET(FLandscapeMobileVertex,LODHeights) + sizeof(uint8) * 4 * Index, sizeof(FLandscapeMobileVertex), VET_UByte4N));
-			}
-			LandscapeVertexFactory->InitResource();
-			SharedBuffers->FixedGridVertexFactory = LandscapeVertexFactory;
-		}
 	}
-
-	//
-	FixedGridVertexFactory = SharedBuffers->FixedGridVertexFactory;
-	
 	SharedBuffers->AddRef();
 
 	// Init vertex buffer
-	check(MobileRenderData->VertexBuffer);
-	MobileRenderData->VertexBuffer->InitResource();
-
-	FLandscapeVertexFactoryMobile* LandscapeVertexFactory = new FLandscapeVertexFactoryMobile(FeatureLevel);
-	LandscapeVertexFactory->MobileData.PositionComponent = FVertexStreamComponent(MobileRenderData->VertexBuffer, STRUCT_OFFSET(FLandscapeMobileVertex,Position), sizeof(FLandscapeMobileVertex), VET_UByte4N);
-	for( uint32 Index = 0; Index < LANDSCAPE_MAX_ES_LOD_COMP; ++Index )
 	{
-		LandscapeVertexFactory->MobileData.LODHeightsComponent.Add
-			(FVertexStreamComponent(MobileRenderData->VertexBuffer, STRUCT_OFFSET(FLandscapeMobileVertex,LODHeights) + sizeof(uint8) * 4 * Index, sizeof(FLandscapeMobileVertex), VET_UByte4N));
+		check(MobileRenderData->VertexBuffer);
+		MobileRenderData->VertexBuffer->InitResource();
+
+		FLandscapeVertexFactoryMobile* LandscapeVertexFactory = new FLandscapeVertexFactoryMobile(FeatureLevel);
+		LandscapeVertexFactory->MobileData.PositionComponent = FVertexStreamComponent(MobileRenderData->VertexBuffer, STRUCT_OFFSET(FLandscapeMobileVertex, Position), sizeof(FLandscapeMobileVertex), VET_UByte4N);
+		for (uint32 Index = 0; Index < LANDSCAPE_MAX_ES_LOD_COMP; ++Index)
+		{
+			LandscapeVertexFactory->MobileData.LODHeightsComponent.Add
+			(FVertexStreamComponent(MobileRenderData->VertexBuffer, STRUCT_OFFSET(FLandscapeMobileVertex, LODHeights) + sizeof(uint8) * 4 * Index, sizeof(FLandscapeMobileVertex), VET_UByte4N));
+		}
+
+		LandscapeVertexFactory->InitResource();
+		VertexFactory = LandscapeVertexFactory;
 	}
 
-	LandscapeVertexFactory->InitResource();
-	VertexFactory = LandscapeVertexFactory;
+	// Init vertex buffer for rendering to virtual texture
+	if (UseVirtualTexturing(FeatureLevel))
+	{
+		FLandscapeFixedGridVertexFactoryMobile* LandscapeVertexFactory = new FLandscapeFixedGridVertexFactoryMobile(FeatureLevel);
+		LandscapeVertexFactory->MobileData.PositionComponent = FVertexStreamComponent(MobileRenderData->VertexBuffer, STRUCT_OFFSET(FLandscapeMobileVertex, Position), sizeof(FLandscapeMobileVertex), VET_UByte4N);
+		
+		for (uint32 Index = 0; Index < LANDSCAPE_MAX_ES_LOD_COMP; ++Index)
+		{
+			LandscapeVertexFactory->MobileData.LODHeightsComponent.Add
+			(FVertexStreamComponent(MobileRenderData->VertexBuffer, STRUCT_OFFSET(FLandscapeMobileVertex, LODHeights) + sizeof(uint8) * 4 * Index, sizeof(FLandscapeMobileVertex), VET_UByte4N));
+		}
+		
+		LandscapeVertexFactory->InitResource();
+		FixedGridVertexFactory = LandscapeVertexFactory;
+	}
 
 	// Assign LandscapeUniformShaderParameters
 	LandscapeUniformShaderParameters.InitResource();
+
+	// Create per Lod uniform buffers
+	LandscapeFixedGridUniformShaderParameters.AddDefaulted(MaxLOD + 1);
+	for (int32 LodIndex = 0; LodIndex <= MaxLOD; ++LodIndex)
+	{
+		LandscapeFixedGridUniformShaderParameters[LodIndex].InitResource();
+		FLandscapeFixedGridUniformShaderParameters Parameters;
+		Parameters.LodValues = FVector4(
+			LodIndex,
+			0.f,
+			(float)((SubsectionSizeVerts >> LodIndex) - 1),
+			1.f / (float)((SubsectionSizeVerts >> LodIndex) - 1));
+		LandscapeFixedGridUniformShaderParameters[LodIndex].SetContents(Parameters);
+	}
 }
 
 TSharedPtr<FLandscapeMobileRenderData, ESPMode::ThreadSafe> FLandscapeComponentDerivedData::GetRenderData()
