Simplified ULandscapeSubsystem ticking boilerplate.
Getting familiar with it to use a similar approach for batched pawn ticking.

diff --git a/Engine/Source/Runtime/Landscape/Private/LandscapeSubsystem.cpp b/Engine/Source/Runtime/Landscape/Private/LandscapeSubsystem.cpp
--- a/Engine/Source/Runtime/Landscape/Private/LandscapeSubsystem.cpp
+++ b/Engine/Source/Runtime/Landscape/Private/LandscapeSubsystem.cpp
@@ -17,42 +17,45 @@ static FAutoConsoleVariableRef CVarUseStreamingManagerForCameras(
 DECLARE_CYCLE_STAT(TEXT("LandscapeSubsystem Tick"), STAT_LandscapeSubsystemTick, STATGROUP_Landscape);
 
 ULandscapeSubsystem::ULandscapeSubsystem()
-	: TickFunction(this)
-{}
+{
+}
 
 ULandscapeSubsystem::~ULandscapeSubsystem()
 {
-
 }
 
-void ULandscapeSubsystem::Initialize(FSubsystemCollectionBase& Collection)
+void ULandscapeSubsystem::RegisterActor(ALandscapeProxy* Proxy)
 {
-	// Register Tick 
-	TickFunction.bCanEverTick = true;
-	TickFunction.bTickEvenWhenPaused = true;
-	TickFunction.bStartWithTickEnabled = true;
-	TickFunction.TickGroup = TG_DuringPhysics;
-	TickFunction.bAllowTickOnDedicatedServer = false;
-	TickFunction.RegisterTickFunction(GetWorld()->PersistentLevel);
+	Proxies.AddUnique(Proxy);
 }
 
-void ULandscapeSubsystem::Deinitialize()
+void ULandscapeSubsystem::UnregisterActor(ALandscapeProxy* Proxy)
 {
-	TickFunction.UnRegisterTickFunction();
-	Proxies.Empty();
+	Proxies.Remove(Proxy);
 }
 
-void ULandscapeSubsystem::RegisterActor(ALandscapeProxy* Proxy)
+void ULandscapeSubsystem::Initialize(FSubsystemCollectionBase& Collection)
 {
-	Proxies.AddUnique(Proxy);
+	Super::Initialize(Collection);
+
+	// Register Tick 
+	bCanEverTick = true;
+	bTickEvenWhenPaused = true;
+	bStartWithTickEnabled = true;
+	TickGroup = TG_DuringPhysics;
+	bAllowTickOnDedicatedServer = false;
+	RegisterTickFunction(GetWorld()->PersistentLevel);
 }
 
-void ULandscapeSubsystem::UnregisterActor(ALandscapeProxy* Proxy)
+void ULandscapeSubsystem::Deinitialize()
 {
-	Proxies.Remove(Proxy);
+	UnRegisterTickFunction();
+	Proxies.Empty();
+
+	Super::Deinitialize();
 }
 
-void ULandscapeSubsystem::Tick(float DeltaTime, enum ELevelTick TickType, ENamedThreads::Type CurrentThread, const FGraphEventRef& MyCompletionGraphEvent)
+void ULandscapeSubsystem::ExecuteTick(float DeltaTime, ELevelTick TickType, ENamedThreads::Type CurrentThread, const FGraphEventRef& MyCompletionGraphEvent)
 {
 	SCOPE_CYCLE_COUNTER(STAT_LandscapeSubsystemTick);
 	TRACE_CPUPROFILER_EVENT_SCOPE(ULandscapeSubsystem::Tick);
@@ -93,7 +96,7 @@ void ULandscapeSubsystem::Tick(float DeltaTime, enum ELevelTick TickType, ENamed
 	}
 
 	int32 InOutNumComponentsCreated = 0;
-	for(ALandscapeProxy* Proxy : Proxies)
+	for (ALandscapeProxy* Proxy : Proxies)
 	{
 #if WITH_EDITOR
 		if (GIsEditor)
@@ -117,19 +120,14 @@ void ULandscapeSubsystem::Tick(float DeltaTime, enum ELevelTick TickType, ENamed
 	}
 }
 
-void FLandscapeSubsystemTickFunction::ExecuteTick(float DeltaTime, enum ELevelTick TickType, ENamedThreads::Type CurrentThread, const FGraphEventRef& MyCompletionGraphEvent)
-{
-	Subsystem->Tick(DeltaTime, TickType, CurrentThread, MyCompletionGraphEvent);
-}
-
-FString FLandscapeSubsystemTickFunction::DiagnosticMessage()
+FString ULandscapeSubsystem::DiagnosticMessage()
 {
-	static const FString Message(TEXT("FLandscapeSubsystemTickFunction"));
+	static const FString Message(TEXT("ULandscapeSubsystem"));
 	return Message;
 }
 
-FName FLandscapeSubsystemTickFunction::DiagnosticContext(bool bDetailed)
+FName ULandscapeSubsystem::DiagnosticContext(bool bDetailed)
 {
-	static const FName Context(TEXT("FLandscapeSubsystemTickFunction"));
+	static const FName Context(TEXT("ULandscapeSubsystem"));
 	return Context;
-}
\ No newline at end of file
+}
diff --git a/Engine/Source/Runtime/Landscape/Public/LandscapeSubsystem.h b/Engine/Source/Runtime/Landscape/Public/LandscapeSubsystem.h
--- a/Engine/Source/Runtime/Landscape/Public/LandscapeSubsystem.h
+++ b/Engine/Source/Runtime/Landscape/Public/LandscapeSubsystem.h
@@ -10,44 +10,29 @@
 
 class ALandscapeProxy;
 
-struct FLandscapeSubsystemTickFunction : public FTickFunction
-{
-	FLandscapeSubsystemTickFunction() : FLandscapeSubsystemTickFunction(nullptr) {}
-	FLandscapeSubsystemTickFunction(ULandscapeSubsystem* InSubsystem) : Subsystem(InSubsystem) {}
-	virtual ~FLandscapeSubsystemTickFunction() {}
-
-	// Begin FTickFunction overrides
-	virtual void ExecuteTick(float DeltaTime, enum ELevelTick TickType, ENamedThreads::Type CurrentThread, const FGraphEventRef& MyCompletionGraphEvent) override;
-	virtual FString DiagnosticMessage() override;
-	virtual FName DiagnosticContext(bool bDetailed) override;
-	// End FTickFunction overrides
-
-	class ULandscapeSubsystem* Subsystem;
-};
-
 UCLASS()
 class ULandscapeSubsystem : public UWorldSubsystem, public FTickFunction
 {
 	GENERATED_BODY()
 
+public:
 	ULandscapeSubsystem();
 	virtual ~ULandscapeSubsystem();
 
-public:
-	// Begin USubsystem
-	virtual void Initialize(FSubsystemCollectionBase& Collection) override;
-	virtual void Deinitialize() override;
-	// End USubsystem
-		
 	void RegisterActor(ALandscapeProxy* Proxy);
 	void UnregisterActor(ALandscapeProxy* Proxy);
 
 private:
-	void Tick(float DeltaTime, enum ELevelTick TickType, ENamedThreads::Type CurrentThread, const FGraphEventRef& MyCompletionGraphEvent);
+	// Begin USubsystem
+	virtual void Initialize(FSubsystemCollectionBase& Collection) override;
+	virtual void Deinitialize() override;
+	// End USubsystem
 
-	friend struct FLandscapeSubsystemTickFunction;
-	FLandscapeSubsystemTickFunction TickFunction;
+	// Begin FTickFunction overrides
+	virtual void ExecuteTick(float DeltaTime, ELevelTick TickType, ENamedThreads::Type CurrentThread, const FGraphEventRef& MyCompletionGraphEvent) override;
+	virtual FString DiagnosticMessage() override;
+	virtual FName DiagnosticContext(bool bDetailed) override;
+	// End FTickFunction overrides
 
 	TArray<ALandscapeProxy*> Proxies;
 };
-
