Fix some places in landscape code where LOD overrides are not respected.
Cherry pick integration from Dev-Rendering

diff --git a/Engine/Source/Runtime/Landscape/Private/LandscapeRender.cpp b/Engine/Source/Runtime/Landscape/Private/LandscapeRender.cpp
--- a/Engine/Source/Runtime/Landscape/Private/LandscapeRender.cpp
+++ b/Engine/Source/Runtime/Landscape/Private/LandscapeRender.cpp
@@ -636,6 +636,24 @@ void ULandscapeComponent::GetUsedMaterials(TArray<UMaterialInterface*>& OutMater
 #endif
 }
 
+/** 
+ * Return any global Lod override for landscape. 
+ * A return value less than 0 means no override.
+ * Any positive value must still be clamped into the valid Lod range for the landscape.
+ */
+static int32 GetLandscapeLodOverride(FSceneView const& View)
+{
+	// Apply r.ForceLOD override
+	int32 LodOverride = GetCVarForceLOD();
+#if WITH_EDITOR
+	// Apply editor landscape lod override
+	LodOverride = View.Family->LandscapeLODOverride >= 0 ? View.Family->LandscapeLODOverride : LodOverride;
+#endif
+	// Use lod 0 if lodding is disabled
+	LodOverride = View.Family->EngineShowFlags.LOD == 0 ? 0 : LodOverride;
+	return LodOverride;
+}
+
 //
 // FLandscapeComponentSceneProxy
 //
@@ -833,7 +851,7 @@ void FLandscapeRenderSystem::BeginRenderView(const FSceneView* View)
 
 void FLandscapeRenderSystem::ComputeSectionPerViewParameters(
 	const FSceneView* ViewPtrAsIdentifier,
-	bool ViewEngineShowFlagLOD,
+	int32 ViewLODOverride,
 	float ViewLODDistanceFactor,
 	FVector ViewOrigin,
 	FMatrix ViewProjectionMarix
@@ -853,7 +871,6 @@ void FLandscapeRenderSystem::ComputeSectionPerViewParameters(
 		NewSectionTessellationFalloffK.AddZeroed(SectionLODSettings.Num());
 	}
 
-	int32 ForcedLODLevel = ViewEngineShowFlagLOD ? GetCVarForceLOD() : -1;
 	float LODScale = ViewLODDistanceFactor * CVarStaticMeshLODDistanceScale.GetValueOnRenderThread();
 
 	for (int32 EntityIndex = 0; EntityIndex < SectionLODSettings.Num(); EntityIndex++)
@@ -863,6 +880,7 @@ void FLandscapeRenderSystem::ComputeSectionPerViewParameters(
 		float FractionalLOD;
 		GetLODFromScreenSize(SectionLODSettings[EntityIndex], MeshScreenSizeSquared, LODScale * LODScale, FractionalLOD);
 
+		int32 ForcedLODLevel = FMath::Min<int32>(ViewLODOverride, SectionLODSettings[EntityIndex].LastLODIndex);
 		NewSectionLODValues[EntityIndex] = ForcedLODLevel >= 0 ? ForcedLODLevel : FractionalLOD;
 
 		if (TessellationFalloffSettings.UseTessellationComponentScreenSizeFalloff && NumEntitiesWithTessellation > 0)
@@ -1049,6 +1067,16 @@ void FLandscapeRenderSystem::EndFrame()
 	PerViewParametersTasks.Empty();
 }
 
+FLandscapeRenderSystem::FComputeSectionPerViewParametersTask::FComputeSectionPerViewParametersTask(FLandscapeRenderSystem& InRenderSystem, const FSceneView* InView)
+	: RenderSystem(InRenderSystem)
+	, ViewPtrAsIdentifier(InView)
+	, ViewLODOverride(GetLandscapeLodOverride(*InView))
+	, ViewLODDistanceFactor(InView->LODDistanceFactor)
+	, ViewOrigin(InView->ViewMatrices.GetViewOrigin())
+	, ViewProjectionMatrix(InView->ViewMatrices.GetProjectionMatrix())
+{
+}
+
 class FLandscapePersistentViewUniformBufferExtension : public IPersistentViewUniformBufferExtension
 {
 public:
@@ -2386,14 +2414,7 @@ void FLandscapeComponentSceneProxy::DrawStaticElements(FStaticPrimitiveDrawInter
 void FLandscapeComponentSceneProxy::CalculateLODFromScreenSize(const FSceneView& InView, float InMeshScreenSizeSquared, float InViewLODScale, int32 InSubSectionIndex, FViewCustomDataLOD& InOutLODData) const
 {
 	// Handle general LOD override
-	float PreferedLOD = (float)GetCVarForceLOD();
-
-#if WITH_EDITOR
-	if (InView.Family->LandscapeLODOverride >= 0)
-	{
-		PreferedLOD = (float)InView.Family->LandscapeLODOverride;
-	}
-#endif
+	float PreferedLOD = GetLandscapeLodOverride(InView);
 
 #if WITH_EDITOR || !(UE_BUILD_SHIPPING || UE_BUILD_TEST)
 	if (InView.Family->EngineShowFlags.CollisionVisibility || InView.Family->EngineShowFlags.CollisionPawn)
@@ -3239,10 +3260,7 @@ void FLandscapeComponentSceneProxy::GetDynamicMeshElements(const TArray<const FS
 
 			const FSceneView* View = Views[ViewIndex];
 
-			int32 ForcedLODLevel = (View->Family->EngineShowFlags.LOD) ? GetCVarForceLOD() : -1;
-#if WITH_EDITOR
-			ForcedLODLevel = View->Family->LandscapeLODOverride >= 0 ? View->Family->LandscapeLODOverride : ForcedLODLevel;
-#endif
+			int32 ForcedLODLevel = GetLandscapeLodOverride(*View);
 
 #if WITH_EDITOR || !(UE_BUILD_SHIPPING || UE_BUILD_TEST)
 			ForcedLODLevel = CollisionLODLevel >= 0 ? CollisionLODLevel : ForcedLODLevel;
@@ -3594,7 +3612,7 @@ void FLandscapeComponentSceneProxy::GetDynamicRayTracingInstances(FRayTracingMat
 		return;
 	}
 	float MeshScreenSizeSquared = ComputeBoundsScreenRadiusSquared(GetBounds().Origin, GetBounds().SphereRadius, *Context.ReferenceView);
-	int32 ForcedLODLevel = (Context.ReferenceView->Family->EngineShowFlags.LOD) ? GetCVarForceLOD() : 0;
+	int32 ForcedLODLevel = GetLandscapeLodOverride(*Context.ReferenceView);
 
 	int32 LODToRender = ForcedLODLevel >= 0 ? ForcedLODLevel : GetLODFromScreenSize(MeshScreenSizeSquared, Context.ReferenceView->LODDistanceFactor);
 	
diff --git a/Engine/Source/Runtime/Landscape/Public/LandscapeRender.h b/Engine/Source/Runtime/Landscape/Public/LandscapeRender.h
--- a/Engine/Source/Runtime/Landscape/Public/LandscapeRender.h
+++ b/Engine/Source/Runtime/Landscape/Public/LandscapeRender.h
@@ -572,20 +572,12 @@ struct FLandscapeRenderSystem
 	{
 		FLandscapeRenderSystem& RenderSystem;
 		const FSceneView* ViewPtrAsIdentifier;
-		bool ViewEngineShowFlagLOD;
+		int32 ViewLODOverride;
 		float ViewLODDistanceFactor;
 		FVector ViewOrigin;
 		FMatrix ViewProjectionMatrix;
 
-		FComputeSectionPerViewParametersTask(FLandscapeRenderSystem& InRenderSystem, const FSceneView* InView)
-			: RenderSystem(InRenderSystem)
-			, ViewPtrAsIdentifier(InView)
-			, ViewEngineShowFlagLOD(InView->Family->EngineShowFlags.LOD)
-			, ViewLODDistanceFactor(InView->LODDistanceFactor)
-			, ViewOrigin(InView->ViewMatrices.GetViewOrigin())
-			, ViewProjectionMatrix(InView->ViewMatrices.GetProjectionMatrix())
-		{
-		}
+		FComputeSectionPerViewParametersTask(FLandscapeRenderSystem& InRenderSystem, const FSceneView* InView);
 
 		FORCEINLINE TStatId GetStatId() const
 		{
@@ -604,7 +596,7 @@ struct FLandscapeRenderSystem
 
 		void AnyThreadTask()
 		{
-			RenderSystem.ComputeSectionPerViewParameters(ViewPtrAsIdentifier, ViewEngineShowFlagLOD, ViewLODDistanceFactor, ViewOrigin, ViewProjectionMatrix);
+			RenderSystem.ComputeSectionPerViewParameters(ViewPtrAsIdentifier, ViewLODOverride, ViewLODDistanceFactor, ViewOrigin, ViewProjectionMatrix);
 		}
 
 		void DoTask(ENamedThreads::Type CurrentThread, const FGraphEventRef& MyCompletionGraphEvent)
@@ -698,7 +690,7 @@ struct FLandscapeRenderSystem
 
 	void ComputeSectionPerViewParameters(
 		const FSceneView* ViewPtrAsIdentifier,
-		bool ViewEngineShowFlagLOD,
+		int32 ViewLODOverride,
 		float ViewLODDistanceFactor,
 		FVector ViewOrigin,
 		FMatrix ViewProjectionMarix);
