Fix SimpleCollision null DominantLayer crash

diff --git a/Engine/Source/Runtime/Landscape/Private/LandscapeCollision.cpp b/Engine/Source/Runtime/Landscape/Private/LandscapeCollision.cpp
--- a/Engine/Source/Runtime/Landscape/Private/LandscapeCollision.cpp
+++ b/Engine/Source/Runtime/Landscape/Private/LandscapeCollision.cpp
@@ -897,7 +897,8 @@ bool ULandscapeHeightfieldCollisionComponent::CookCollisionData(const FName& For
 
 	if (bGenerateSimpleCollision)
 	{
-		SimpleSamples = ConvertHeightfieldDataForPhysx(this, SimpleCollisionSizeVerts, bIsMirrored, Heights + NumSamples, bUseDefMaterial, DominantLayers + NumSamples, DefMaterial, InOutMaterials);
+		const uint8* SimpleDominantLayers = DominantLayers ? DominantLayers + NumSamples : nullptr;
+		SimpleSamples = ConvertHeightfieldDataForPhysx(this, SimpleCollisionSizeVerts, bIsMirrored, Heights + NumSamples, bUseDefMaterial, SimpleDominantLayers, DefMaterial, InOutMaterials);
 	}
 
 	CollisionHeightData.Unlock();
