Modify default for Landscape VirtualTextureNumLods from 0->6
This default value reduces likelyhood of vertex interpolation artifacts when writing to a virtual texture.

diff --git a/Engine/Source/Runtime/Landscape/Classes/LandscapeProxy.h b/Engine/Source/Runtime/Landscape/Classes/LandscapeProxy.h
--- a/Engine/Source/Runtime/Landscape/Classes/LandscapeProxy.h
+++ b/Engine/Source/Runtime/Landscape/Classes/LandscapeProxy.h
@@ -550,13 +550,15 @@ public:
 
 	/** 
 	 * Number of mesh levels to use when rendering landscape into runtime virtual texture.
-	 * Set this only if the material used to render the virtual texture requires interpolated vertex data such as height.
-	 * Higher values use more tessellated meshes and are expensive when rendering the runtime virtual texture.
+	 * Lower values reduce vertex count when rendering to the runtime virtual texture but decrease accuracy when using values that require vertex interpolation.
 	 */
 	UPROPERTY(EditAnywhere, AdvancedDisplay, BlueprintReadOnly, Category = VirtualTexture, meta = (DisplayName = "Virtual Texture Num LODs", UIMin = "0", UIMax = "7"))
-	int32 VirtualTextureNumLods = 0;
+	int32 VirtualTextureNumLods = 6;
 
-	/** Bias to the LOD selected for rendering to runtime virtual textures. */
+	/** 
+	 * Bias to the LOD selected for rendering to runtime virtual textures.
+	 * Higher values reduce vertex count when rendering to the runtime virtual texture.
+	 */
 	UPROPERTY(EditAnywhere, AdvancedDisplay, BlueprintReadOnly, Category = VirtualTexture, meta = (DisplayName = "Virtual Texture LOD Bias", UIMin = "0", UIMax = "7"))
 	int32 VirtualTextureLodBias = 0;
 
