Fixed crash with landscape grass mesh when there are more mips than the landscape currently has

diff --git a/Engine/Source/Runtime/Landscape/Private/LandscapeRender.cpp b/Engine/Source/Runtime/Landscape/Private/LandscapeRender.cpp
--- a/Engine/Source/Runtime/Landscape/Private/LandscapeRender.cpp
+++ b/Engine/Source/Runtime/Landscape/Private/LandscapeRender.cpp
@@ -1491,8 +1491,10 @@ void FLandscapeComponentSceneProxy::CreateRenderThreadResources()
 	LandscapeUniformShaderParameters.InitResource();
 
 	// Create per Lod uniform buffers
-	LandscapeFixedGridUniformShaderParameters.AddDefaulted(MaxLOD+1);
-	for (int32 LodIndex = 0; LodIndex <= MaxLOD; ++LodIndex)
+	const int32 NumMips = FMath::CeilLogTwo(SubsectionSizeVerts); 
+	// create as many as there are potential mips (even if MaxLOD can be inferior than that), because the grass could need that much :
+	LandscapeFixedGridUniformShaderParameters.AddDefaulted(NumMips);
+	for (int32 LodIndex = 0; LodIndex < NumMips; ++LodIndex)
 	{
 		LandscapeFixedGridUniformShaderParameters[LodIndex].InitResource();
 		FLandscapeFixedGridUniformShaderParameters Parameters;
@@ -1508,7 +1510,6 @@ void FLandscapeComponentSceneProxy::CreateRenderThreadResources()
 	// Create MeshBatch for grass rendering
 	if (SharedBuffers->GrassIndexBuffer)
 	{
-		const int32 NumMips = FMath::CeilLogTwo(SubsectionSizeVerts);
 		GrassMeshBatch.Elements.Empty(NumMips);
 		GrassMeshBatch.Elements.AddDefaulted(NumMips);
 		GrassBatchParams.Empty(NumMips);
@@ -2663,39 +2664,6 @@ int8 FLandscapeComponentSceneProxy::GetLODFromScreenSize(float InScreenSizeSquar
 	float FractionalLOD;
 
 	return FLandscapeRenderSystem::GetLODFromScreenSize(LODSettings, InScreenSizeSquared, InViewLODScale, FractionalLOD);
-
-#if 0
-	int32 LODScreenRatioSquaredCount = LODScreenRatioSquared.Num();
-	float ScreenSizeSquared = InScreenSizeSquared / InViewLODScale;
-
-	if (ScreenSizeSquared <= LODScreenRatioSquared[LODScreenRatioSquaredCount - 1])
-	{
-		return LODScreenRatioSquaredCount - 1;
-	}
-	else if (ScreenSizeSquared > LODScreenRatioSquared[1])
-	{
-		return 0;
-	}
-	else
-	{
-		int32 HalfPointIndex = (LODScreenRatioSquaredCount - 1) / 2;
-		int32 StartingIndex = ScreenSizeSquared < LODScreenRatioSquared[HalfPointIndex] ? HalfPointIndex : 1;
-		int8 SelectedLODIndex = INDEX_NONE;
-
-		for (int32 i = StartingIndex; i < LODScreenRatioSquaredCount - 1; ++i)
-		{
-			if (ScreenSizeSquared > LODScreenRatioSquared[i + 1])
-			{
-				SelectedLODIndex = i;
-				break;
-			}
-		}
-
-		return SelectedLODIndex;
-	}
-
-	return INDEX_NONE;
-#endif
 }
 
 void* FLandscapeComponentSceneProxy::InitViewCustomData(const FSceneView& InView, float InViewLODScale, FMemStackBase& InCustomDataMemStack, bool InIsStaticRelevant, bool InIsShadowOnly, const FLODMask* InVisiblePrimitiveLODMask, float InMeshScreenSizeSquared)
