Fix Grass not being updated when adding GrassOutput to LandscapeMaterial

#rnx

diff --git a/Engine/Source/Runtime/Landscape/Classes/LandscapeProxy.h b/Engine/Source/Runtime/Landscape/Classes/LandscapeProxy.h
--- a/Engine/Source/Runtime/Landscape/Classes/LandscapeProxy.h
+++ b/Engine/Source/Runtime/Landscape/Classes/LandscapeProxy.h
@@ -846,7 +846,9 @@ public:
 	/* Per-frame call to update dynamic grass placement and render grassmaps */
 	FORCEINLINE bool ShouldTickGrass() const
 	{
-		if (!bHasLandscapeGrass)
+		// At runtime if we don't have grass we will never have any so avoid ticking it
+		// In editor we might have a material that didn't have grass and now does so we can't rely on bHasLandscapeGrass.
+		if (!GIsEditor && !bHasLandscapeGrass)
 		{
 			return false;
 		}
