Fix for crash when loading a cooked landscape in the editor

diff --git a/Engine/Source/Runtime/Landscape/Private/Landscape.cpp b/Engine/Source/Runtime/Landscape/Private/Landscape.cpp
--- a/Engine/Source/Runtime/Landscape/Private/Landscape.cpp
+++ b/Engine/Source/Runtime/Landscape/Private/Landscape.cpp
@@ -1894,6 +1894,13 @@ void ALandscapeProxy::PostRegisterAllComponents()
 #if WITH_EDITOR
 			if (GIsEditor && !GetWorld()->IsGameWorld())
 			{
+				// Note: This can happen when loading certain cooked assets in an editor
+				// Todo: Determine the root cause of this and fix it at a higher level!
+				if (LandscapeComponents.Num() > 0 && LandscapeComponents[0] == nullptr)
+				{
+					LandscapeComponents.Empty();
+				}
+
 				UpdateCachedHasLayersContent(true);
 
 				// Cache the value at this point as CreateLandscapeInfo (-> RegisterActor) might create/destroy layers content if there was a mismatch between landscape & proxy
