Fixed crash after moving Landscape Actor to another Level

diff --git a/Engine/Source/Runtime/Landscape/Private/LandscapeEdit.cpp b/Engine/Source/Runtime/Landscape/Private/LandscapeEdit.cpp
--- a/Engine/Source/Runtime/Landscape/Private/LandscapeEdit.cpp
+++ b/Engine/Source/Runtime/Landscape/Private/LandscapeEdit.cpp
@@ -3764,6 +3764,8 @@ void ALandscape::PostDuplicate(bool bDuplicateForPIE)
 	{
 		// Need to generate new GUID when duplicating
 		LandscapeGuid = FGuid::NewGuid();
+		// This makes sure at least we have a LandscapeInfo mapped for this GUID.
+		CreateLandscapeInfo();
 	}
 
 	Super::PostDuplicate(bDuplicateForPIE);
