Fixed warnings about invalid material for landscape thumbnail renders

diff --git a/Engine/Source/Runtime/Landscape/Private/LandscapeRender.cpp b/Engine/Source/Runtime/Landscape/Private/LandscapeRender.cpp
--- a/Engine/Source/Runtime/Landscape/Private/LandscapeRender.cpp
+++ b/Engine/Source/Runtime/Landscape/Private/LandscapeRender.cpp
@@ -4706,7 +4706,8 @@ public:
 			FName(TEXT("TBasePassVSFSimpleNoLightmapLightingPolicyAtmosphericFog")),
 			FName(TEXT("TDepthOnlyVS<false>")),
 			FName(TEXT("TDepthOnlyVS<true>")),
-			FName(TEXT("FDepthOnlyPS")),
+			FName(TEXT("FDepthOnlyPS<true>")),
+			FName(TEXT("FDepthOnlyPS<false>")),
 			// UE-44519, masked material with landscape layers requires FHitProxy shaders.
 			FName(TEXT("FHitProxyVS")),
 			FName(TEXT("FHitProxyPS")),
@@ -4875,6 +4876,7 @@ public:
 			FName(TEXT("TTranslucencyShadowDepthVS<TranslucencyShadowDepth_PerspectiveCorrect>")),
 			FName(TEXT("TTranslucencyShadowDepthPS<TranslucencyShadowDepth_PerspectiveCorrect>")),
 
+			FName(TEXT("TShadowDepthVSForGSVertexShadowDepth_OnePassPointLight")),
 			FName(TEXT("TShadowDepthVSForGSVertexShadowDepth_OnePassPointLightPositionOnly")),
 			FName(TEXT("TShadowDepthVSVertexShadowDepth_OnePassPointLightPositionOnly")),
 			FName(TEXT("TShadowDepthVSVertexShadowDepth_OutputDepthPositionOnly")),
