Fixed crash with grass scattering when virtual texturing is not enabled : we use the fixed grid vertex factory for both the grass and virtual texturing now.

diff --git a/Engine/Source/Runtime/Landscape/Private/LandscapeGrass.cpp b/Engine/Source/Runtime/Landscape/Private/LandscapeGrass.cpp
--- a/Engine/Source/Runtime/Landscape/Private/LandscapeGrass.cpp
+++ b/Engine/Source/Runtime/Landscape/Private/LandscapeGrass.cpp
@@ -377,7 +377,7 @@ void FLandscapeGrassWeightMeshProcessor::AddMeshBatch(const FMeshBatch& RESTRICT
 	const FMaterial& Material = MeshBatch.MaterialRenderProxy->GetMaterialWithFallback(FeatureLevel, FallbackMaterialRenderProxyPtr);
 
 	const FMaterialRenderProxy& MaterialRenderProxy = FallbackMaterialRenderProxyPtr ? *FallbackMaterialRenderProxyPtr : *MeshBatch.MaterialRenderProxy;
-
+	check(MeshBatch.VertexFactory != nullptr);
 	Process(MeshBatch, BatchElementMask, PrimitiveSceneProxy, MaterialRenderProxy, Material, NumPasses, ViewOffset, PassOffsetX, FirstHeightMipsPassIndex, HeightMips);
 }
 
diff --git a/Engine/Source/Runtime/Landscape/Private/LandscapeRender.cpp b/Engine/Source/Runtime/Landscape/Private/LandscapeRender.cpp
--- a/Engine/Source/Runtime/Landscape/Private/LandscapeRender.cpp
+++ b/Engine/Source/Runtime/Landscape/Private/LandscapeRender.cpp
@@ -1443,7 +1443,8 @@ void FLandscapeComponentSceneProxy::CreateRenderThreadResources()
 			SharedBuffers->VertexFactory = LandscapeXYOffsetVertexFactory;
 		}
 
-		if (UseVirtualTexturing(FeatureLevel))
+		// we need the fixed grid vertex factory for both virtual texturing and grass : 
+		if (SharedBuffers->GrassIndexBuffer || UseVirtualTexturing(FeatureLevel))
 		{
 			//todo[vt]: We will need a version of this to support XYOffsetmapTexture
 			FLandscapeFixedGridVertexFactory* LandscapeVertexFactory = new FLandscapeFixedGridVertexFactory(FeatureLevel);
@@ -1510,6 +1511,8 @@ void FLandscapeComponentSceneProxy::CreateRenderThreadResources()
 	// Create MeshBatch for grass rendering
 	if (SharedBuffers->GrassIndexBuffer)
 	{
+		check(FixedGridVertexFactory != nullptr);
+
 		GrassMeshBatch.Elements.Empty(NumMips);
 		GrassMeshBatch.Elements.AddDefaulted(NumMips);
 		GrassBatchParams.Empty(NumMips);
@@ -4073,6 +4076,7 @@ FLandscapeSharedBuffers::FLandscapeSharedBuffers(const int32 InSharedBuffersKey,
 	, SubsectionSizeVerts(InSubsectionSizeQuads + 1)
 	, NumSubsections(InNumSubsections)
 	, VertexFactory(nullptr)
+	, FixedGridVertexFactory(nullptr)
 	, VertexBuffer(nullptr)
 	, AdjacencyIndexBuffers(nullptr)
 	, bUse32BitIndices(false)
