Revert back from High to Default compression quality for landscape weightmaps (ASTC_4x4 -> ASTC_8x8)

diff --git a/Engine/Source/Runtime/Landscape/Private/LandscapeEdit.cpp b/Engine/Source/Runtime/Landscape/Private/LandscapeEdit.cpp
--- a/Engine/Source/Runtime/Landscape/Private/LandscapeEdit.cpp
+++ b/Engine/Source/Runtime/Landscape/Private/LandscapeEdit.cpp
@@ -6856,12 +6856,6 @@ UTexture2D* ALandscapeProxy::CreateLandscapeTexture(int32 InSizeX, int32 InSizeY
 	NewTexture->AddressY = TA_Clamp;
 	NewTexture->LODGroup = InLODGroup;
 
-	if (bCompress && InLODGroup == TEXTUREGROUP_Terrain_Weightmap)
-	{
-		// Compress weightmaps with a best quality to reduce blocky artifacts on a landscape 
-		NewTexture->CompressionQuality = TCQ_Highest;
-	}
-
 	return NewTexture;
 }
 
