Fixed compile error on non-editor targets

diff --git a/Engine/Source/Runtime/Landscape/Private/LandscapeRender.cpp b/Engine/Source/Runtime/Landscape/Private/LandscapeRender.cpp
--- a/Engine/Source/Runtime/Landscape/Private/LandscapeRender.cpp
+++ b/Engine/Source/Runtime/Landscape/Private/LandscapeRender.cpp
@@ -1444,7 +1444,13 @@ void FLandscapeComponentSceneProxy::CreateRenderThreadResources()
 		}
 
 		// we need the fixed grid vertex factory for both virtual texturing and grass : 
-		if (SharedBuffers->GrassIndexBuffer || UseVirtualTexturing(FeatureLevel))
+		bool bNeedsFixedGridVertexFactory = UseVirtualTexturing(FeatureLevel);
+
+#if WITH_EDITOR
+		bNeedsFixedGridVertexFactory |= (SharedBuffers->GrassIndexBuffer != nullptr);
+#endif // WITH_EDITOR
+
+		if (bNeedsFixedGridVertexFactory)
 		{
 			//todo[vt]: We will need a version of this to support XYOffsetmapTexture
 			FLandscapeFixedGridVertexFactory* LandscapeVertexFactory = new FLandscapeFixedGridVertexFactory(FeatureLevel);
