DebugName for landscape

diff --git a/Engine/Source/Runtime/Landscape/Private/LandscapeCollision.cpp b/Engine/Source/Runtime/Landscape/Private/LandscapeCollision.cpp
--- a/Engine/Source/Runtime/Landscape/Private/LandscapeCollision.cpp
+++ b/Engine/Source/Runtime/Landscape/Private/LandscapeCollision.cpp
@@ -44,6 +44,7 @@
 #include "Materials/MaterialInstanceConstant.h"
 #include "Physics/PhysicsInterfaceCore.h"
 #include "Physics/PhysicsInterfaceUtils.h"
+#include "Containers/StringConv.h"
 
 #if WITH_EDITOR && WITH_PHYSX
 	#include "IPhysXCooking.h"
@@ -400,6 +401,11 @@ void ULandscapeHeightfieldCollisionComponent::OnCreatePhysicsState()
 				Params.InitialTM.SetScale3D(FVector(0));
 				Params.bQueryOnly = true;
 				Params.Scene = GetWorld()->GetPhysicsScene();
+#if USE_BODYINSTANCE_DEBUG_NAMES
+				const FString DebugName = FString::Printf(TEXT("LandscapeCollision: '%s' "), *GetPathName());
+				BodyInstance.CharDebugName = MakeShareable(new TArray<ANSICHAR>(StringToArray<ANSICHAR>(*DebugName, DebugName.Len() + 1)));
+				Params.DebugName = BodyInstance.CharDebugName->GetData();
+#endif
 				FPhysicsActorHandle PhysHandle;
 				FPhysicsInterface::CreateActor(Params, PhysHandle);
 
