Fix crash when copying offset regions, reported in a UDN.

diff --git a/Engine/Source/Runtime/Landscape/Private/LandscapeRender.cpp b/Engine/Source/Runtime/Landscape/Private/LandscapeRender.cpp
--- a/Engine/Source/Runtime/Landscape/Private/LandscapeRender.cpp
+++ b/Engine/Source/Runtime/Landscape/Private/LandscapeRender.cpp
@@ -759,16 +759,20 @@ void FLandscapeRenderSystem::ResizeAndMoveTo(FIntPoint NewMin, FIntPoint NewSize
 		{
 			int32 LinearIndex = Y * Size.X + X;
 			int32 NewLinearIndex = (Y + (Min.Y - NewMin.Y)) * NewSize.X + (X + (Min.X - NewMin.X));
-			NewSectionLODValues[NewLinearIndex] = SectionLODValues[LinearIndex];
-			NewSectionLODBiases[NewLinearIndex] = SectionLODBiases[LinearIndex];
-			if (TessellationFalloffSettings.UseTessellationComponentScreenSizeFalloff && NumEntitiesWithTessellation > 0)
+
+			if (NewLinearIndex >= 0 && NewLinearIndex < NewSize.X * NewSize.Y)
 			{
-				NewSectionTessellationFalloffC[NewLinearIndex] = SectionTessellationFalloffC[LinearIndex];
-				NewSectionTessellationFalloffK[NewLinearIndex] = SectionTessellationFalloffK[LinearIndex];
+				NewSectionLODValues[NewLinearIndex] = SectionLODValues[LinearIndex];
+				NewSectionLODBiases[NewLinearIndex] = SectionLODBiases[LinearIndex];
+				if (TessellationFalloffSettings.UseTessellationComponentScreenSizeFalloff && NumEntitiesWithTessellation > 0)
+				{
+					NewSectionTessellationFalloffC[NewLinearIndex] = SectionTessellationFalloffC[LinearIndex];
+					NewSectionTessellationFalloffK[NewLinearIndex] = SectionTessellationFalloffK[LinearIndex];
+				}
+				NewSectionLODSettings[NewLinearIndex] = SectionLODSettings[LinearIndex];
+				NewSectionOriginAndRadius[NewLinearIndex] = SectionOriginAndRadius[LinearIndex];
+				NewSceneProxies[NewLinearIndex] = SceneProxies[LinearIndex];
 			}
-			NewSectionLODSettings[NewLinearIndex] = SectionLODSettings[LinearIndex];
-			NewSectionOriginAndRadius[NewLinearIndex] = SectionOriginAndRadius[LinearIndex];
-			NewSceneProxies[NewLinearIndex] = SceneProxies[LinearIndex];
 		}
 	}
 
