ALandscape::LandscapeLayers is no more 'TextExportTransient'.

diff --git a/Engine/Source/Runtime/Landscape/Classes/Landscape.h b/Engine/Source/Runtime/Landscape/Classes/Landscape.h
--- a/Engine/Source/Runtime/Landscape/Classes/Landscape.h
+++ b/Engine/Source/Runtime/Landscape/Classes/Landscape.h
@@ -431,7 +431,7 @@ public:
 	/** Used to temporarily disable Grass Update in Editor */
 	bool bGrassUpdateEnabled;
 
-	UPROPERTY(TextExportTransient)
+	UPROPERTY()
 	TArray<FLandscapeLayer> LandscapeLayers;
 
 	UPROPERTY(Transient)
