Remove deprecated render target binding used only by landscape grass system

diff --git a/Engine/Source/Runtime/Landscape/Private/LandscapeGrass.cpp b/Engine/Source/Runtime/Landscape/Private/LandscapeGrass.cpp
--- a/Engine/Source/Runtime/Landscape/Private/LandscapeGrass.cpp
+++ b/Engine/Source/Runtime/Landscape/Private/LandscapeGrass.cpp
@@ -490,7 +490,7 @@ public:
 
 	void RenderLandscapeComponentToTexture_RenderThread(FRHICommandListImmediate& RHICmdList)
 	{
-		FSceneViewFamilyContext ViewFamily(FSceneViewFamily::ConstructionValues(RenderTargetResource, NULL, FEngineShowFlags(ESFIM_Game)).SetWorldTimes(FApp::GetCurrentTime() - GStartTime, FApp::GetDeltaTime(), FApp::GetCurrentTime() - GStartTime));
+		FSceneViewFamilyContext ViewFamily(FSceneViewFamily::ConstructionValues(nullptr, nullptr, FEngineShowFlags(ESFIM_Game)).SetWorldTimes(FApp::GetCurrentTime() - GStartTime, FApp::GetDeltaTime(), FApp::GetCurrentTime() - GStartTime));
 
 		ViewFamily.LandscapeLODOverride = 0; // Force LOD render
 
@@ -504,30 +504,36 @@ public:
 		GetRendererModule().CreateAndInitSingleView(RHICmdList, &ViewFamily, &ViewInitOptions);
 		
 		const FSceneView* View = ViewFamily.Views[0];
-		RHICmdList.SetViewport(View->UnscaledViewRect.Min.X, View->UnscaledViewRect.Min.Y, 0.0f, View->UnscaledViewRect.Max.X, View->UnscaledViewRect.Max.Y, 1.0f);
 
+		FRHIRenderPassInfo RPInfo(RenderTargetResource->GetTextureRHI(), ERenderTargetActions::Clear_Store);
+		RHICmdList.BeginRenderPass(RPInfo, TEXT("LandscapeGrass"));
+		RHICmdList.SetViewport(View->UnscaledViewRect.Min.X, View->UnscaledViewRect.Min.Y, 0.0f, View->UnscaledViewRect.Max.X, View->UnscaledViewRect.Max.Y, 1.0f);
 		RHICmdList.SetScissorRect(false, 0, 0, 0, 0);
 
-		FMemMark Mark(FMemStack::Get());
-
-		DrawDynamicMeshPass(*View, RHICmdList,
-			[View, PassOffsetX = PassOffsetX, &ComponentInfos = ComponentInfos, NumPasses = NumPasses, FirstHeightMipsPassIndex = FirstHeightMipsPassIndex, HeightMips = HeightMips](FDynamicPassMeshDrawListContext* DynamicMeshPassContext)
 		{
-			FLandscapeGrassWeightMeshProcessor PassMeshProcessor(
-				nullptr,
-				View,
-				DynamicMeshPassContext);
-
-			const uint64 DefaultBatchElementMask = ~0ul;
+			FMemMark Mark(FMemStack::Get());
 
-			for (auto& ComponentInfo : ComponentInfos)
+			DrawDynamicMeshPass(*View, RHICmdList,
+				[View, PassOffsetX = PassOffsetX, &ComponentInfos = ComponentInfos, NumPasses = NumPasses, FirstHeightMipsPassIndex = FirstHeightMipsPassIndex, HeightMips = HeightMips](FDynamicPassMeshDrawListContext* DynamicMeshPassContext)
 			{
-				const FMeshBatch& Mesh = ComponentInfo.SceneProxy->GetGrassMeshBatch();
-				Mesh.MaterialRenderProxy->UpdateUniformExpressionCacheIfNeeded(View->GetFeatureLevel());
+				FLandscapeGrassWeightMeshProcessor PassMeshProcessor(
+					nullptr,
+					View,
+					DynamicMeshPassContext);
 
-				PassMeshProcessor.AddMeshBatch(Mesh, DefaultBatchElementMask, NumPasses, ComponentInfo.ViewOffset, PassOffsetX, FirstHeightMipsPassIndex, HeightMips, ComponentInfo.SceneProxy);
-			}
-		});
+				const uint64 DefaultBatchElementMask = ~0ul;
+
+				for (auto& ComponentInfo : ComponentInfos)
+				{
+					const FMeshBatch& Mesh = ComponentInfo.SceneProxy->GetGrassMeshBatch();
+					Mesh.MaterialRenderProxy->UpdateUniformExpressionCacheIfNeeded(View->GetFeatureLevel());
+
+					PassMeshProcessor.AddMeshBatch(Mesh, DefaultBatchElementMask, NumPasses, ComponentInfo.ViewOffset, PassOffsetX, FirstHeightMipsPassIndex, HeightMips, ComponentInfo.SceneProxy);
+				}
+			});
+		}
+
+		RHICmdList.EndRenderPass();
 	}
 };
 
diff --git a/Engine/Source/Runtime/Renderer/Private/SceneRendering.cpp b/Engine/Source/Runtime/Renderer/Private/SceneRendering.cpp
--- a/Engine/Source/Runtime/Renderer/Private/SceneRendering.cpp
+++ b/Engine/Source/Runtime/Renderer/Private/SceneRendering.cpp
@@ -2358,7 +2358,7 @@ void FSceneRenderer::PrepareViewRectsForRendering()
 			TEXT("ISceneViewFamilyScreenPercentage::GetPrimaryResolutionFractionUpperBound() should not lie to the renderer."));
 		
 		check(FSceneViewScreenPercentageConfig::IsValidResolutionFraction(PrimaryResolutionFraction));
-
+		
 		// Compute final resolution fraction.
 		float ResolutionFraction = PrimaryResolutionFraction * ViewFamily.SecondaryViewFraction;
 
@@ -2497,7 +2497,7 @@ void FSceneRenderer::ComputeViewGPUMasks(FRHIGPUMask RenderTargetGPUMask)
 				{
 					ViewInfo.GPUMask = FRHIGPUMask::FromIndex(*GPUIterator);
 					ViewFamily.bMultiGPUForkAndJoin |= (ViewInfo.GPUMask != RenderTargetGPUMask);
-					
+
 					// Increment and wrap around if we reach the last index.
 					++GPUIterator;
 					if (!GPUIterator)
@@ -2733,7 +2733,7 @@ void FSceneRenderer::RenderFinish(FRHICommandListImmediate& RHICmdList)
 				bMobileShowVertexFogWarning = true;
 			}
 		}
-		
+
 		const bool bSingleLayerWaterWarning = ShouldRenderSingleLayerWaterSkippedRenderEditorNotification(Views);
 		
 		const bool bAnyWarning = bShowPrecomputedVisibilityWarning || bShowGlobalClipPlaneWarning || bShowAtmosphericFogWarning || bShowSkylightWarning || bShowPointLightWarning 
@@ -3209,7 +3209,7 @@ void FSceneRenderer::RenderCustomDepthPass(FRHICommandListImmediate& RHICmdList)
 						Scene->UniformBuffers.InstancedCustomDepthViewUniformBuffer.UpdateUniformBufferImmediate(reinterpret_cast<FInstancedViewUniformShaderParameters&>(*View.CachedViewUniformShaderParameters));
 					}
 				}
-	
+
 				extern TSet<IPersistentViewUniformBufferExtension*> PersistentViewUniformBufferExtensions;
 				
 				for (IPersistentViewUniformBufferExtension* Extension : PersistentViewUniformBufferExtensions)
@@ -3621,10 +3621,6 @@ void FRendererModule::CreateAndInitSingleView(FRHICommandListImmediate& RHICmdLi
 	// Create and add the new view
 	FViewInfo* NewView = new FViewInfo(*ViewInitOptions);
 	ViewFamily->Views.Add(NewView);
-	PRAGMA_DISABLE_DEPRECATION_WARNINGS
-	FRHIRenderTargetView RTV(ViewFamily->RenderTarget->GetRenderTargetTexture(), ERenderTargetLoadAction::EClear);
-	RHICmdList.SetRenderTargets(1, &RTV, nullptr);
-	PRAGMA_ENABLE_DEPRECATION_WARNINGS
 	FViewInfo* View = (FViewInfo*)ViewFamily->Views[0];
 	View->ViewRect = View->UnscaledViewRect;
 	View->InitRHIResources();
