Added a new material expression which is a physical material custom output node.
This will be (optionally) used to generate the dominant physical material for each point on a landscape.

diff --git a/Engine/Source/Runtime/Landscape/Classes/Materials/MaterialExpressionLandscapePhysicalMaterialOutput.h b/Engine/Source/Runtime/Landscape/Classes/Materials/MaterialExpressionLandscapePhysicalMaterialOutput.h
new file mode 100644
--- /dev/null
+++ b/Engine/Source/Runtime/Landscape/Classes/Materials/MaterialExpressionLandscapePhysicalMaterialOutput.h
@@ -0,0 +1,60 @@
+// Copyright Epic Games, Inc. All Rights Reserved.
+
+#pragma once
+
+#include "CoreMinimal.h"
+#include "Materials/MaterialExpressionCustomOutput.h"
+
+#include "MaterialExpressionLandscapePhysicalMaterialOutput.generated.h"
+
+/** Structure linking a material expression input with a physical material. For use by UMaterialExpressionLandscapePhysicalMaterialOutput. */
+USTRUCT()
+struct FPhysicalMaterialInput
+{
+	GENERATED_USTRUCT_BODY()
+
+	UPROPERTY(EditAnywhere, Category = PhysicalMaterial)
+	class UPhysicalMaterial* PhysicalMaterial;
+
+	UPROPERTY(meta = (RequiredInput = "true"))
+	FExpressionInput Input;
+
+	FPhysicalMaterialInput()
+		: PhysicalMaterial(nullptr)
+	{}
+};
+
+/** 
+ * Custom output node to write out physical material weights.
+ * This can be used to generate the dominant physical material for each point on a landscape.
+ * Note that the use of a material output node to generate this information is optional and when a node of this type is not present we fall back on a CPU path which analyzes landscape layer data.
+ */
+UCLASS(collapsecategories, hidecategories=Object)
+class LANDSCAPE_API UMaterialExpressionLandscapePhysicalMaterialOutput : public UMaterialExpressionCustomOutput
+{
+	GENERATED_UCLASS_BODY()
+
+	/** Array of physical material inputs. */
+	UPROPERTY(EditAnywhere, Category = PhysicalMaterial)
+	TArray< FPhysicalMaterialInput > Inputs;
+
+#if WITH_EDITOR
+	//~ Begin UObject Interface
+	virtual void PostEditChangeProperty(FPropertyChangedEvent& PropertyChangedEvent) override;
+	//~ End UObject Interface
+
+	//~ Begin UMaterialExpression Interface
+	virtual void GetCaption(TArray<FString>& OutCaptions) const override;
+	virtual const TArray<FExpressionInput*> GetInputs() override;
+	virtual FExpressionInput* GetInput(int32 InputIndex) override;
+	virtual FName GetInputName(int32 InputIndex) const override;
+	virtual uint32 GetInputType(int32 InputIndex) override { return MCT_Float; }
+	virtual int32 Compile(class FMaterialCompiler* Compiler, int32 OutputIndex) override;
+	//~ End UMaterialExpression Interface
+#endif
+
+	//~ Begin UMaterialExpressionCustomOutput Interface
+	virtual int32 GetNumOutputs() const override { return Inputs.Num(); }
+	virtual FString GetFunctionName() const override { return TEXT("GetPhysicalMaterial"); }
+	//~ End UMaterialExpressionCustomOutput Interface
+};
diff --git a/Engine/Source/Runtime/Landscape/Private/Materials/MaterialExpressionLandscapePhysicalMaterialOutput.cpp b/Engine/Source/Runtime/Landscape/Private/Materials/MaterialExpressionLandscapePhysicalMaterialOutput.cpp
new file mode 100644
--- /dev/null
+++ b/Engine/Source/Runtime/Landscape/Private/Materials/MaterialExpressionLandscapePhysicalMaterialOutput.cpp
@@ -0,0 +1,107 @@
+// Copyright Epic Games, Inc. All Rights Reserved.
+
+#include "Materials/MaterialExpressionLandscapePhysicalMaterialOutput.h"
+
+#include "Engine/Engine.h"
+#include "EngineGlobals.h"
+#include "MaterialCompiler.h"
+#include "PhysicalMaterials/PhysicalMaterial.h"
+
+#define LOCTEXT_NAMESPACE "Landscape"
+
+UMaterialExpressionLandscapePhysicalMaterialOutput::UMaterialExpressionLandscapePhysicalMaterialOutput(const FObjectInitializer& ObjectInitializer)
+	: Super(ObjectInitializer)
+{
+	// Structure to hold one-time initialization
+	struct FConstructorStatics
+	{
+		FText NAME_Landscape;
+		FConstructorStatics()
+			: NAME_Landscape(LOCTEXT("Landscape", "Landscape"))
+		{
+		}
+	};
+	static FConstructorStatics ConstructorStatics;
+
+#if WITH_EDITORONLY_DATA
+	MenuCategories.Add(ConstructorStatics.NAME_Landscape);
+
+	// No outputs
+	Outputs.Reset();
+#endif
+}
+
+#if WITH_EDITOR
+
+void UMaterialExpressionLandscapePhysicalMaterialOutput::GetCaption(TArray<FString>& OutCaptions) const
+{
+	OutCaptions.Add(FString(TEXT("Landscape Physical Material Output")));
+}
+
+const TArray<FExpressionInput*> UMaterialExpressionLandscapePhysicalMaterialOutput::GetInputs()
+{
+	TArray<FExpressionInput*> OutInputs;
+	for (auto& Input : Inputs)
+	{
+		OutInputs.Add(&Input.Input);
+	}
+	return OutInputs;
+}
+
+FExpressionInput* UMaterialExpressionLandscapePhysicalMaterialOutput::GetInput(int32 InputIndex)
+{
+	return &Inputs[InputIndex].Input;
+}
+
+
+FName UMaterialExpressionLandscapePhysicalMaterialOutput::GetInputName(int32 InputIndex) const
+{
+	UPhysicalMaterial* PhysicalMaterial = Inputs[InputIndex].PhysicalMaterial;
+	if (PhysicalMaterial != nullptr)
+	{
+		return PhysicalMaterial->GetFName();
+	}
+
+	return NAME_None;
+}
+
+int32 UMaterialExpressionLandscapePhysicalMaterialOutput::Compile(class FMaterialCompiler* Compiler, int32 OutputIndex)
+{
+	if (Inputs.IsValidIndex(OutputIndex))
+	{
+		if (Inputs[OutputIndex].PhysicalMaterial == nullptr)
+		{
+			return CompilerError(Compiler, TEXT("Physical material not set."));
+		}
+		
+		if (Inputs[OutputIndex].Input.Expression == nullptr)
+		{
+			return CompilerError(Compiler, TEXT("Input missing."));
+		}
+		
+		return Compiler->CustomOutput(this, OutputIndex, Inputs[OutputIndex].Input.Compile(Compiler));
+	}
+
+	return INDEX_NONE;
+}
+
+void UMaterialExpressionLandscapePhysicalMaterialOutput::PostEditChangeProperty(FPropertyChangedEvent& PropertyChangedEvent)
+{
+	Super::PostEditChangeProperty(PropertyChangedEvent);
+
+	if (PropertyChangedEvent.MemberProperty)
+	{
+		const FName PropertyName = PropertyChangedEvent.MemberProperty->GetFName();
+		if (PropertyName == GET_MEMBER_NAME_CHECKED(UMaterialExpressionLandscapePhysicalMaterialOutput, Inputs))
+		{
+			if (GraphNode)
+			{
+				GraphNode->ReconstructNode();
+			}
+		}
+	}
+}
+
+#endif // WITH_EDITOR
+
+#undef LOCTEXT_NAMESPACE
