Add render task for writing landscape physical material custom outputs.
Task is designed to be non-blocking so that a landscape component can wait and apply results when ready.

diff --git a/Engine/Shaders/Private/LandscapePhysicalMaterial.usf b/Engine/Shaders/Private/LandscapePhysicalMaterial.usf
new file mode 100644
--- /dev/null
+++ b/Engine/Shaders/Private/LandscapePhysicalMaterial.usf
@@ -0,0 +1,121 @@
+// Copyright Epic Games, Inc. All Rights Reserved.
+
+/**
+ * Material shader used to output the dominant landscape physical materials.
+ */
+
+#include "Common.ush"
+#include "/Engine/Generated/Material.ush"
+#include "/Engine/Generated/VertexFactory.ush"
+
+struct FLandscapeGrassWeightInterpolantsVSToPS
+{
+	FVertexFactoryInterpolantsVSToPS FactoryInterpolants;
+	float4 Position : SV_POSITION;
+};
+
+#if VERTEXSHADER
+
+/** Simple default vertex shader. */
+void VSMain(
+	FVertexFactoryInput Input,
+	out FLandscapeGrassWeightInterpolantsVSToPS Output
+)
+{
+	ResolvedView = ResolveView();
+
+	FVertexFactoryIntermediates VFIntermediates = GetVertexFactoryIntermediates(Input);
+	float4 WorldPosition = VertexFactoryGetWorldPosition(Input, VFIntermediates);
+	float4 RasterizedWorldPosition = VertexFactoryGetRasterizedWorldPosition(Input, VFIntermediates, WorldPosition);
+	Output.Position = mul(RasterizedWorldPosition, ResolvedView.TranslatedWorldToClip);
+
+	half3x3 TangentToLocal = VertexFactoryGetTangentToLocal(Input, VFIntermediates);
+	FMaterialVertexParameters VertexParameters = GetMaterialVertexParameters(Input, VFIntermediates, WorldPosition.xyz, TangentToLocal);
+	Output.FactoryInterpolants = VertexFactoryGetInterpolantsVSToPS(Input, VFIntermediates, VertexParameters);
+}
+
+#elif PIXELSHADER
+
+float2 UpdateDominantMaterial(in float Index, in float Contribution, in float2 DominantMaterial)
+{
+	if (Contribution > DominantMaterial.y)
+	{
+		DominantMaterial.x = Index + 1; // Store indices as 1-based so that 0 means empty/default.
+		DominantMaterial.y = Contribution;
+	}
+	return DominantMaterial;
+}
+
+/** Pixel shader outputs index of physical material. */
+void PSMain(
+	FLandscapeGrassWeightInterpolantsVSToPS Interpolants
+	OPTIONAL_IsFrontFace,
+	out half4 OutColor : SV_Target0
+)
+{
+	float4 SvPosition = Interpolants.Position;
+	FMaterialPixelParameters MaterialParameters = GetMaterialPixelParameters(Interpolants.FactoryInterpolants, SvPosition);
+
+	FPixelMaterialInputs PixelMaterialInputs;
+	CalcMaterialParameters(MaterialParameters, PixelMaterialInputs, SvPosition, bIsFrontFace);
+
+	float2 DomininantMaterial = float2(0, 0); // float2(MaterialIndex, Contribution)
+
+#if NUM_MATERIAL_OUTPUTS_GETPHYSICALMATERIAL > 0
+	DomininantMaterial = UpdateDominantMaterial(0, GetPhysicalMaterial0(MaterialParameters), DomininantMaterial);
+#endif
+#if NUM_MATERIAL_OUTPUTS_GETPHYSICALMATERIAL > 1
+	DomininantMaterial = UpdateDominantMaterial(1, GetPhysicalMaterial1(MaterialParameters), DomininantMaterial);
+#endif
+#if NUM_MATERIAL_OUTPUTS_GETPHYSICALMATERIAL > 2
+	DomininantMaterial = UpdateDominantMaterial(2, GetPhysicalMaterial2(MaterialParameters), DomininantMaterial);
+#endif
+#if NUM_MATERIAL_OUTPUTS_GETPHYSICALMATERIAL > 3
+	DomininantMaterial = UpdateDominantMaterial(3, GetPhysicalMaterial3(MaterialParameters), DomininantMaterial);
+#endif
+#if NUM_MATERIAL_OUTPUTS_GETPHYSICALMATERIAL > 4
+	DomininantMaterial = UpdateDominantMaterial(4, GetPhysicalMaterial4(MaterialParameters), DomininantMaterial);
+#endif
+#if NUM_MATERIAL_OUTPUTS_GETPHYSICALMATERIAL > 5
+	DomininantMaterial = UpdateDominantMaterial(5, GetPhysicalMaterial5(MaterialParameters), DomininantMaterial);
+#endif
+#if NUM_MATERIAL_OUTPUTS_GETPHYSICALMATERIAL > 6
+	DomininantMaterial = UpdateDominantMaterial(6, GetPhysicalMaterial6(MaterialParameters), DomininantMaterial);
+#endif
+#if NUM_MATERIAL_OUTPUTS_GETPHYSICALMATERIAL > 7
+	DomininantMaterial = UpdateDominantMaterial(7, GetPhysicalMaterial7(MaterialParameters), DomininantMaterial);
+#endif
+#if NUM_MATERIAL_OUTPUTS_GETPHYSICALMATERIAL > 8
+	DomininantMaterial = UpdateDominantMaterial(8, GetPhysicalMaterial8(MaterialParameters), DomininantMaterial);
+#endif
+#if NUM_MATERIAL_OUTPUTS_GETPHYSICALMATERIAL > 9
+	DomininantMaterial = UpdateDominantMaterial(9, GetPhysicalMaterial9(MaterialParameters), DomininantMaterial);
+#endif
+#if NUM_MATERIAL_OUTPUTS_GETPHYSICALMATERIAL > 10
+	DomininantMaterial = UpdateDominantMaterial(10, GetPhysicalMaterial10(MaterialParameters), DomininantMaterial);
+#endif
+#if NUM_MATERIAL_OUTPUTS_GETPHYSICALMATERIAL > 11
+	DomininantMaterial = UpdateDominantMaterial(11, GetPhysicalMaterial11(MaterialParameters), DomininantMaterial);
+#endif
+#if NUM_MATERIAL_OUTPUTS_GETPHYSICALMATERIAL > 12
+	DomininantMaterial = UpdateDominantMaterial(12, GetPhysicalMaterial12(MaterialParameters), DomininantMaterial);
+#endif
+#if NUM_MATERIAL_OUTPUTS_GETPHYSICALMATERIAL > 13
+	DomininantMaterial = UpdateDominantMaterial(13, GetPhysicalMaterial13(MaterialParameters), DomininantMaterial);
+#endif
+#if NUM_MATERIAL_OUTPUTS_GETPHYSICALMATERIAL > 14
+	DomininantMaterial = UpdateDominantMaterial(14, GetPhysicalMaterial14(MaterialParameters), DomininantMaterial);
+#endif
+#if NUM_MATERIAL_OUTPUTS_GETPHYSICALMATERIAL > 15
+	DomininantMaterial = UpdateDominantMaterial(15, GetPhysicalMaterial15(MaterialParameters), DomininantMaterial);
+#endif
+#if NUM_MATERIAL_OUTPUTS_GETPHYSICALMATERIAL > 16
+	// We pick an arbitrary limit to support here, but the only cost of extending this is more typing.
+#error Too many physical materials
+#endif
+
+	// Write index assuming R8_UNORM
+	OutColor = float4(DomininantMaterial.x / 255.f, 0, 0, 0);
+}
+
+#endif
diff --git a/Engine/Source/Runtime/Landscape/Classes/LandscapePhysicalMaterial.h b/Engine/Source/Runtime/Landscape/Classes/LandscapePhysicalMaterial.h
new file mode 100644
--- /dev/null
+++ b/Engine/Source/Runtime/Landscape/Classes/LandscapePhysicalMaterial.h
@@ -0,0 +1,56 @@
+// Copyright Epic Games, Inc. All Rights Reserved.
+
+#pragma once
+
+#include "CoreMinimal.h"
+
+#if WITH_EDITOR
+
+/** 
+ * A single task for calculating the landscape physical material.
+ * The work uses the GPU so this object handles reading back the data without stalling.
+ * We need to call Tick() until IsComplete() returns true.
+ * We need to call Release() after use to prevent leaking of internal objects.
+ * All calls are expected to be made on the game thread only.
+ */
+class FLandscapePhysicalMaterialRenderTask
+{
+public:
+	/* Task constructs with IsValid() false. */
+	FLandscapePhysicalMaterialRenderTask()
+		: PoolHandle(-1)
+	{}
+
+	~FLandscapePhysicalMaterialRenderTask()
+	{
+		Release();
+	}
+
+	/** Initialize the task for a component. A task in progress can safely call Init() again to rekick the task. */
+	void Init(class ULandscapeComponent const* LandscapeComponent);
+	/** Release the task. After calling this IsValid() will return false until we call Init() again. */
+	void Release();
+
+	/** Returns true if Init() has been called without a Release(). */
+	bool IsValid() const;
+	/** Returns true if the task is complete and the result is available. */
+	bool IsComplete() const;
+
+	/** Update a task. Does nothing if the task is complete. */
+	void Tick();
+	/** Flush a task to completion. This may stall waiting for completion. */
+	void Flush();
+	
+	/** Get the result data. Assumes that IsComplete() is true. This doesn't Release() the data. */
+	TArray<uint8> const& GetResultIds() const;
+	/** Get the result data. Assumes that IsComplete() is true. This doesn't Release() the data. */
+	TArray <class UPhysicalMaterial* > const& GetResultMaterials() const;
+
+private:
+	int32 PoolHandle;
+	
+	friend class FLandscapePhysicalMaterialRenderTaskPool;
+	FLandscapePhysicalMaterialRenderTask(FLandscapePhysicalMaterialRenderTask const&);
+};
+
+#endif // WITH_EDITOR
diff --git a/Engine/Source/Runtime/Landscape/Private/LandscapePhysicalMaterial.cpp b/Engine/Source/Runtime/Landscape/Private/LandscapePhysicalMaterial.cpp
new file mode 100644
--- /dev/null
+++ b/Engine/Source/Runtime/Landscape/Private/LandscapePhysicalMaterial.cpp
@@ -0,0 +1,608 @@
+// Copyright Epic Games, Inc. All Rights Reserved.
+
+#include "LandscapePhysicalMaterial.h"
+
+#if WITH_EDITOR
+
+#include "EngineModule.h"
+#include "LandscapeComponent.h"
+#include "LandscapeRender.h"
+#include "Materials/MaterialExpressionLandscapePhysicalMaterialOutput.h"
+#include "Materials/MaterialInstanceConstant.h"
+#include "MeshMaterialShader.h"
+#include "MeshPassProcessor.h"
+#include "MeshPassProcessor.inl"
+#include "RenderCommandFence.h"
+#include "RenderTargetPool.h"
+#include "RHIResources.h"
+#include "SceneRenderTargetParameters.h"
+
+DECLARE_GPU_STAT_NAMED(LandscapePhysicalMaterial_Draw, TEXT("LandscapePhysicalMaterial"));
+
+
+namespace
+{
+	/** Get the landscape material used by the landscape component. */
+	UMaterialInterface* GetLandscapeMaterial(ULandscapeComponent const* InLandscapeComponent)
+	{
+		UMaterialInterface* Material = InLandscapeComponent->GetLandscapeMaterial();
+		for (UMaterialInstanceConstant* MIC = Cast<UMaterialInstanceConstant>(Material); MIC; MIC = Cast<UMaterialInstanceConstant>(Material))
+		{
+			Material = MIC->Parent;
+		}
+		return Material;
+	}
+
+	/** 
+	 * Get the physical materials that are configured by the landscape component graphical material. 
+	 * Returns false if there are no non-null physical materials. (We probably don't want to use if no physical material connections are bound.)
+	 */
+	bool GetPhysicalMaterials(ULandscapeComponent const* InLandscapeComponent, TArray<UPhysicalMaterial*>& OutPhysicalMaterials)
+	{
+		bool bReturnValue = false;
+		OutPhysicalMaterials.Reset();
+
+		UMaterialInterface* Material = GetLandscapeMaterial(InLandscapeComponent);
+		if (Material != nullptr)
+		{
+			ERHIFeatureLevel::Type FeatureLevel = GMaxRHIFeatureLevel;
+			{
+				TArray<const UMaterialExpressionLandscapePhysicalMaterialOutput*> Expressions;
+				Material->GetMaterial()->GetAllExpressionsOfType<UMaterialExpressionLandscapePhysicalMaterialOutput>(Expressions);
+				if (Expressions.Num() > 0)
+				{
+					// Assume only one valid physical material output material node
+					for (const FPhysicalMaterialInput& Input : Expressions[0]->Inputs)
+					{
+						OutPhysicalMaterials.Add(Input.PhysicalMaterial);
+						bReturnValue |= (Input.PhysicalMaterial != nullptr);
+					}
+				}
+			}
+		}
+
+		return bReturnValue;
+	}
+}
+
+
+/** Material shader for rendering physical material IDs. */
+class FLandscapePhysicalMaterial : public FMeshMaterialShader
+{
+public:
+	FLandscapePhysicalMaterial()
+	{}
+
+	FLandscapePhysicalMaterial(const FMeshMaterialShaderType::CompiledShaderInitializerType& Initializer)
+		: FMeshMaterialShader(Initializer)
+	{
+		PassUniformBuffer.Bind(Initializer.ParameterMap, FSceneTexturesUniformParameters::StaticStructMetadata.GetShaderVariableName());
+	}
+
+	static bool ShouldCompilePermutation(const FMeshMaterialShaderPermutationParameters& Parameters)
+	{
+		return (Parameters.Material->IsUsedWithLandscape() || Parameters.Material->IsSpecialEngineMaterial()) &&
+			IsFeatureLevelSupported(Parameters.Platform, ERHIFeatureLevel::SM5) &&
+			((Parameters.VertexFactoryType == FindVertexFactoryType(FName(TEXT("FLandscapeVertexFactory"), FNAME_Find))) || (Parameters.VertexFactoryType == FindVertexFactoryType(FName(TEXT("FLandscapeXYOffsetVertexFactory"), FNAME_Find)))) &&
+			!IsConsolePlatform(Parameters.Platform);
+	}
+};
+
+class FLandscapePhysicalMaterialVS : public FLandscapePhysicalMaterial
+{
+	DECLARE_SHADER_TYPE(FLandscapePhysicalMaterialVS, MeshMaterial);
+
+public:
+	FLandscapePhysicalMaterialVS()
+	{}
+
+	FLandscapePhysicalMaterialVS(const FMeshMaterialShaderType::CompiledShaderInitializerType& Initializer)
+		: FLandscapePhysicalMaterial(Initializer)
+	{}
+};
+
+IMPLEMENT_MATERIAL_SHADER_TYPE(, FLandscapePhysicalMaterialVS, TEXT("/Engine/Private/LandscapePhysicalMaterial.usf"), TEXT("VSMain"), SF_Vertex);
+
+class FLandscapePhysicalMaterialPS : public FLandscapePhysicalMaterial
+{
+	DECLARE_SHADER_TYPE(FLandscapePhysicalMaterialPS, MeshMaterial);
+
+public:
+	FLandscapePhysicalMaterialPS()
+	{}
+
+	FLandscapePhysicalMaterialPS(const ShaderMetaType::CompiledShaderInitializerType& Initializer)
+		: FLandscapePhysicalMaterial(Initializer)
+	{}
+};
+
+IMPLEMENT_MATERIAL_SHADER_TYPE(, FLandscapePhysicalMaterialPS, TEXT("/Engine/Private/LandscapePhysicalMaterial.usf"), TEXT("PSMain"), SF_Pixel);
+
+
+/** Simple mesh processor implementation to draw using the FLandscapePhysicalMaterial mesh material shader. */
+class FLandscapePhysicalMaterialMeshProcessor : public FMeshPassProcessor
+{
+public:
+	FLandscapePhysicalMaterialMeshProcessor(
+		const FScene* Scene, 
+		const FSceneView* InViewIfDynamicMeshCommand, 
+		FMeshPassDrawListContext* InDrawListContext);
+
+	virtual void AddMeshBatch(
+		const FMeshBatch& RESTRICT MeshBatch,
+		uint64 BatchElementMask,
+		const FPrimitiveSceneProxy* RESTRICT PrimitiveSceneProxy,
+		int32 StaticMeshId = -1) override final;
+
+private:
+	void Process(
+		const FMeshBatch& MeshBatch,
+		uint64 BatchElementMask,
+		const FPrimitiveSceneProxy* RESTRICT PrimitiveSceneProxy,
+		const FMaterialRenderProxy& RESTRICT MaterialRenderProxy,
+		const FMaterial& RESTRICT MaterialResource);
+
+	FMeshPassProcessorRenderState PassDrawRenderState;
+};
+
+FLandscapePhysicalMaterialMeshProcessor::FLandscapePhysicalMaterialMeshProcessor(
+	const FScene* Scene, 
+	const FSceneView* InViewIfDynamicMeshCommand, 
+	FMeshPassDrawListContext* InDrawListContext)
+	: FMeshPassProcessor(Scene, InViewIfDynamicMeshCommand->GetFeatureLevel(), InViewIfDynamicMeshCommand, InDrawListContext)
+{
+	PassDrawRenderState.SetBlendState(TStaticBlendState<>::GetRHI());
+	PassDrawRenderState.SetDepthStencilState(TStaticDepthStencilState<false, CF_Always>::GetRHI());
+
+	PassDrawRenderState.SetViewUniformBuffer(InViewIfDynamicMeshCommand->ViewUniformBuffer);
+	PassDrawRenderState.SetPassUniformBuffer(nullptr);
+}
+
+void FLandscapePhysicalMaterialMeshProcessor::AddMeshBatch(
+	const FMeshBatch& RESTRICT MeshBatch,
+	uint64 BatchElementMask,
+	const FPrimitiveSceneProxy* RESTRICT PrimitiveSceneProxy,
+	int32 StaticMeshId)
+{
+	const FMaterialRenderProxy* FallbackMaterialRenderProxyPtr = nullptr;
+	const FMaterial& MaterialResource = MeshBatch.MaterialRenderProxy->GetMaterialWithFallback(FeatureLevel, FallbackMaterialRenderProxyPtr);
+	const FMaterialRenderProxy& MaterialRenderProxy = FallbackMaterialRenderProxyPtr ? *FallbackMaterialRenderProxyPtr : *MeshBatch.MaterialRenderProxy;
+
+	Process(MeshBatch, BatchElementMask, PrimitiveSceneProxy, MaterialRenderProxy, MaterialResource);
+}
+
+void FLandscapePhysicalMaterialMeshProcessor::Process(
+	const FMeshBatch& MeshBatch,
+	uint64 BatchElementMask,
+	const FPrimitiveSceneProxy* RESTRICT PrimitiveSceneProxy,
+	const FMaterialRenderProxy& RESTRICT MaterialRenderProxy,
+	const FMaterial& RESTRICT MaterialResource)
+{
+	const FVertexFactory* VertexFactory = MeshBatch.VertexFactory;
+
+	TMeshProcessorShaders<
+		FLandscapePhysicalMaterialVS,
+		FMeshMaterialShader,
+		FMeshMaterialShader,
+		FLandscapePhysicalMaterialPS> PassShaders;
+
+	PassShaders.PixelShader = MaterialResource.GetShader<FLandscapePhysicalMaterialPS>(VertexFactory->GetType());
+	PassShaders.VertexShader = MaterialResource.GetShader<FLandscapePhysicalMaterialVS>(VertexFactory->GetType());
+
+	const ERasterizerFillMode MeshFillMode = ComputeMeshFillMode(MeshBatch, MaterialResource);
+	const ERasterizerCullMode MeshCullMode = CM_None;
+
+	FMeshMaterialShaderElementData ShaderElementData;
+	ShaderElementData.InitializeMeshMaterialData(ViewIfDynamicMeshCommand, PrimitiveSceneProxy, MeshBatch, -1, true);
+
+	const FMeshDrawCommandSortKey SortKey = CalculateMeshStaticSortKey(PassShaders.VertexShader, PassShaders.PixelShader);
+
+	BuildMeshDrawCommands(
+		MeshBatch,
+		BatchElementMask,
+		PrimitiveSceneProxy,
+		MaterialRenderProxy,
+		MaterialResource,
+		PassDrawRenderState,
+		PassShaders,
+		MeshFillMode,
+		MeshCullMode,
+		SortKey,
+		EMeshPassFeatures::Default,
+		ShaderElementData);
+}
+
+
+namespace
+{
+	/** A description of a mesh to render. */
+	struct FMeshInfo
+	{
+		FPrimitiveSceneProxy const* Proxy;
+		FMeshBatch const* MeshBatch;
+		uint64 MeshBatchElementMask;
+	};
+	typedef TArray< FMeshInfo, TInlineAllocator<1> > FMeshInfoArray;
+
+	/**
+	 * Collect the meshes to render for a landscape component.
+	 * Initially we only collect the base landscape mesh, but potentially we could gather other objects like roads?
+	 * WARNING: This gets the SceneProxy pointer from the UComponent on the render thread. This doesn't feel safe but it's what the grass renderer does...
+	 */
+	void FillMeshInfos_RenderThread(ULandscapeComponent const* InLandscapeComponent, FMeshInfoArray& OutMeshInfos)
+	{
+		const int32 MeshInfoIndex = OutMeshInfos.AddUninitialized();
+		OutMeshInfos[MeshInfoIndex].Proxy = InLandscapeComponent->SceneProxy;
+		OutMeshInfos[MeshInfoIndex].MeshBatch = &((FLandscapeComponentSceneProxy*)InLandscapeComponent->SceneProxy)->GetGrassMeshBatch();
+		OutMeshInfos[MeshInfoIndex].MeshBatchElementMask = 1 << 0; // LOD 0 only
+	}
+
+	/** Render the landscape physical material IDs and copy to the read back texture. */
+	void Render_RenderThread(
+		FRHICommandListImmediate& RHICmdList,
+		FMeshInfoArray const& MeshInfos,
+		FIntPoint TargetSize,
+		FVector ViewOrigin,
+		FMatrix ViewRotationMatrix,
+		FMatrix ProjectionMatrix,
+		FTextureRHIRef ReadbackTexture,
+		FGPUFenceRHIRef ReadbackFence)
+	{
+		// Allocate temporary render target
+		TRefCountPtr<IPooledRenderTarget> PooledRenderTarget;
+		FPooledRenderTargetDesc Desc(FPooledRenderTargetDesc::Create2DDesc(TargetSize, PF_G8, FClearValueBinding::Black, TexCreate_None, TexCreate_RenderTargetable | TexCreate_ShaderResource, false));
+		GRenderTargetPool.FindFreeElement(RHICmdList, Desc, PooledRenderTarget, TEXT("LandscapePhysicalMaterialTarget"), true, ERenderTargetTransience::Transient);
+		FTextureRHIRef RenderTargetTexture = PooledRenderTarget->GetRenderTargetItem().TargetableTexture;
+
+		// Create the view
+		FSceneViewFamily::ConstructionValues ViewFamilyInit(nullptr, nullptr, FEngineShowFlags(ESFIM_Game));
+		ViewFamilyInit.SetWorldTimes(0.0f, 0.0f, 0.0f);
+		FSceneViewFamilyContext ViewFamily(ViewFamilyInit);
+		ViewFamily.LandscapeLODOverride = 0; // Force LOD 0 render
+
+		FSceneViewInitOptions ViewInitOptions;
+		ViewInitOptions.SetViewRectangle(FIntRect(0, 0, TargetSize.X, TargetSize.Y));
+		ViewInitOptions.ViewOrigin = ViewOrigin;
+		ViewInitOptions.ViewRotationMatrix = ViewRotationMatrix;
+		ViewInitOptions.ProjectionMatrix = ProjectionMatrix;
+		ViewInitOptions.ViewFamily = &ViewFamily;
+
+		GetRendererModule().CreateAndInitSingleView(RHICmdList, &ViewFamily, &ViewInitOptions);
+		const FSceneView* View = ViewFamily.Views[0];
+
+		// Draw
+		{
+			SCOPED_GPU_STAT(RHICmdList, LandscapePhysicalMaterial_Draw);
+
+			FRHIRenderPassInfo RPInfo(RenderTargetTexture, ERenderTargetActions::Clear_Store);
+			RHICmdList.BeginRenderPass(RPInfo, TEXT("LandscapePhysicalMaterial"));
+			RHICmdList.SetViewport(View->UnscaledViewRect.Min.X, View->UnscaledViewRect.Min.Y, 0.0f, View->UnscaledViewRect.Max.X, View->UnscaledViewRect.Max.Y, 1.0f);
+			RHICmdList.SetScissorRect(false, 0, 0, 0, 0);
+
+			{
+				FMemMark Mark(FMemStack::Get());
+			
+				DrawDynamicMeshPass(*View, RHICmdList,
+					[View, &MeshInfos = MeshInfos](FDynamicPassMeshDrawListContext* DynamicMeshPassContext)
+					{
+						FLandscapePhysicalMaterialMeshProcessor PassMeshProcessor(
+							nullptr,
+							View,
+							DynamicMeshPassContext);
+
+						for (auto& MeshInfo : MeshInfos)
+						{
+							const FMeshBatch& Mesh = *MeshInfo.MeshBatch;
+							if (Mesh.MaterialRenderProxy != nullptr)
+							{
+								Mesh.MaterialRenderProxy->UpdateUniformExpressionCacheIfNeeded(View->GetFeatureLevel());
+								PassMeshProcessor.AddMeshBatch(Mesh, MeshInfo.MeshBatchElementMask, MeshInfo.Proxy);
+							}
+						}
+					});
+			}
+
+			RHICmdList.EndRenderPass();
+
+			// Copy to the read back texture
+			RHICmdList.TransitionResource(EResourceTransitionAccess::EReadable, RenderTargetTexture);
+			RHICmdList.CopyTexture(RenderTargetTexture, ReadbackTexture, FRHICopyTextureInfo());
+			RHICmdList.WriteGPUFence(ReadbackFence);
+		}
+	}
+
+	/** Fetch the landscape physical material IDs from a read back texture. */
+	void FetchResults_RenderThread(
+		FRHICommandListImmediate& RHICmdList,
+		FIntPoint TargetSize,
+		FTextureRHIRef InReadbackTexture,
+		FGPUFenceRHIRef InReadbackFence,
+		TArray<uint8>& OutPhysicalMaterialIds)
+	{
+		void* Data = nullptr;
+		int32 TargetWidth, TargetHeight;
+		RHICmdList.MapStagingSurface(InReadbackTexture, InReadbackFence.GetReference(), Data, TargetWidth, TargetHeight);
+		check(Data != nullptr);
+		check(TargetSize.X <= TargetWidth && TargetSize.Y <= TargetHeight);
+
+		OutPhysicalMaterialIds.Empty(TargetSize.X * TargetSize.Y);
+		OutPhysicalMaterialIds.AddUninitialized(TargetSize.X * TargetSize.Y);
+
+		uint8 const* ReadPtr = (uint8*)Data;
+		uint8* WritePtr = OutPhysicalMaterialIds.GetData();
+		for (int32 Y = 0; Y < TargetSize.Y; ++Y)
+		{
+			FMemory::Memcpy(WritePtr, ReadPtr, TargetSize.X);
+			ReadPtr += TargetWidth;
+			WritePtr += TargetSize.X;
+		}
+
+		RHICmdList.UnmapStagingSurface(InReadbackTexture);
+	}
+}
+
+
+namespace
+{
+	/** Completion state for a physical material render task. */
+	enum class ECompletionState : uint64
+	{
+		None = 0,		// Draw not submitted
+		Pending = 1,	// Draw submitted, waiting for GPU
+		Complete = 2	// Result copied back from GPU
+	};
+
+	/** Data for a physical material render task. */
+	struct FLandscapePhysicalMaterialRenderTaskImpl
+	{
+		// Create on game thread
+		ULandscapeComponent const* LandscapeComponent;
+		FIntPoint TargetSize;
+		FVector ViewOrigin;
+		FMatrix ViewRotationMatrix;
+		FMatrix ProjectionMatrix;
+
+		// Result written on game thread and read on game thread
+		TArray<UPhysicalMaterial*> ResultMaterials;
+
+		// Create on render thread
+		FTextureRHIRef ReadbackTexture;
+		FGPUFenceRHIRef ReadbackFence;
+
+		// Result written on render thread and read on game thread
+		ECompletionState CompletionState;
+		TArray<uint8> ResultIds;
+	};
+
+	/** Initialize the physical material render task data. */
+	bool InitTask(FLandscapePhysicalMaterialRenderTaskImpl& Task, ULandscapeComponent const* InLandscapeComponent)
+	{
+		if (InLandscapeComponent != nullptr)
+		{
+			if (GetPhysicalMaterials(InLandscapeComponent, Task.ResultMaterials))
+			{
+				Task.LandscapeComponent = InLandscapeComponent;
+				Task.CompletionState = ECompletionState::None;
+
+				const FTransform& ComponentTransform = InLandscapeComponent->GetComponentTransform();
+				const int32 ComponentSizeVerts = InLandscapeComponent->SubsectionSizeQuads * InLandscapeComponent->NumSubsections + 1;
+				const FIntPoint TargetSize(ComponentSizeVerts, ComponentSizeVerts);
+				const FIntPoint TargetSizeMinusOne(TargetSize - FIntPoint(1, 1));
+				const FVector TargetCenter = ComponentTransform.TransformPosition(FVector(TargetSizeMinusOne, 0.f) * 0.5f);
+				const FVector TargetExtent = FVector(TargetSize, 0.0f) * ComponentTransform.GetScale3D() * 0.5f;
+				const FMatrix ViewRotationMatrix = FInverseRotationMatrix(ComponentTransform.Rotator()) * FMatrix(FPlane(1, 0, 0, 0), FPlane(0, -1, 0, 0), FPlane(0, 0, -1, 0), FPlane(0, 0, 0, 1));
+				const float ZOffset = WORLD_MAX;
+				const FMatrix ProjectionMatrix = FReversedZOrthoMatrix(TargetExtent.X, TargetExtent.Y, 0.5f / ZOffset, ZOffset);
+
+				Task.TargetSize = TargetSize;
+				Task.ViewOrigin = TargetCenter;
+				Task.ViewRotationMatrix = ViewRotationMatrix;
+				Task.ProjectionMatrix = ProjectionMatrix;
+
+				return true;
+			}
+		}
+
+		return false;
+	}
+
+	/** Initialize the physical material render task read back resources. */
+	bool InitTaskRenderResources(FLandscapePhysicalMaterialRenderTaskImpl& Task)
+	{
+		//todo: Consider pooling these and throttling to the pool size?
+		if (!Task.ReadbackTexture.IsValid() || Task.ReadbackTexture->GetSizeXYZ() != FIntVector(Task.TargetSize.X, Task.TargetSize.Y, 1))
+		{
+			FRHIResourceCreateInfo CreateInfo(TEXT(""));
+			Task.ReadbackTexture = RHICreateTexture2D(Task.TargetSize.X, Task.TargetSize.Y, PF_G8, 1, 1, TexCreate_CPUReadback | TexCreate_HideInVisualizeTexture, CreateInfo);
+		}
+		if (!Task.ReadbackFence.IsValid())
+		{
+			Task.ReadbackFence = RHICreateGPUFence(TEXT(""));
+		}
+
+		return true;
+	}
+
+	/** Update the physical material render task on the render thread. */
+	void UpdateTask_RenderThread(FRHICommandListImmediate& RHICmdList, FLandscapePhysicalMaterialRenderTaskImpl& Task, bool bFlush)
+	{
+		// WARNING: We access the UComponent to get in SceneProxy for FillMeshInfos_RenderThread(). 
+		// This isn't good style but probably works since the UComponent owns the update task and is guaranteed to be valid.
+		if (Task.CompletionState == ECompletionState::None && Task.LandscapeComponent->SceneProxy)
+		{
+			// Allocate read back resources
+			if (InitTaskRenderResources(Task))
+			{
+				// Render the pending item.
+				FMeshInfoArray MeshInfos;
+				FillMeshInfos_RenderThread(Task.LandscapeComponent, MeshInfos);
+
+				Render_RenderThread(
+					RHICmdList,
+					MeshInfos,
+					Task.TargetSize,
+					Task.ViewOrigin,
+					Task.ViewRotationMatrix,
+					Task.ProjectionMatrix,
+					Task.ReadbackTexture,
+					Task.ReadbackFence);
+
+				FPlatformMisc::MemoryBarrier();
+				Task.CompletionState = ECompletionState::Pending;
+			}
+		}
+
+		if (Task.CompletionState == ECompletionState::Pending)
+		{
+			if (bFlush || Task.ReadbackFence->Poll())
+			{
+				FetchResults_RenderThread(RHICmdList, Task.TargetSize, Task.ReadbackTexture, Task.ReadbackFence, Task.ResultIds);
+
+				FPlatformMisc::MemoryBarrier();
+				Task.CompletionState = ECompletionState::Complete;
+			}
+		}
+	}
+}
+
+
+/** 
+ * Pool for storing physical material render task data. 
+ * todo: Would be good to garbage collect the render resources from the pool but that would require some global tick function.
+ */
+class FLandscapePhysicalMaterialRenderTaskPool
+{
+public:
+	/** Pool uses chunked array to avoid task data being moved by a realloc. */
+	TChunkedArray< FLandscapePhysicalMaterialRenderTaskImpl > Pool;
+
+	/** Allocate task data from the pool. */
+	void Allocate(FLandscapePhysicalMaterialRenderTask& InTask, ULandscapeComponent const* LandscapeComponent)
+	{
+		check(InTask.PoolHandle == -1);
+
+		int32 Index = 0;
+		auto ItEnd = Pool.end();
+		for (auto It = Pool.begin(); It != ItEnd; ++It, ++Index)
+		{
+			if ((*It).LandscapeComponent == nullptr)
+			{
+				break;
+			}
+		}
+
+		if (Index == Pool.Num())
+		{
+			Pool.Add();
+		}
+
+		const bool bSuccess = InitTask(Pool[Index], LandscapeComponent);
+		InTask.PoolHandle = bSuccess ? Index : -1;
+	}
+
+	/** Return task data to the pool. */
+	void Free(FLandscapePhysicalMaterialRenderTask& InTask)
+	{
+		check(InTask.PoolHandle != -1);
+		
+		FLandscapePhysicalMaterialRenderTaskImpl* Task = &Pool[InTask.PoolHandle];
+		
+		// Invalidate the task object.
+		InTask.PoolHandle = -1;
+
+		// Submit render thread command to mark pooled task as free.
+		ENQUEUE_RENDER_COMMAND(FLandscapePhysicalMaterialUpdaterTick)(
+			[Task](FRHICommandListImmediate& RHICmdList)
+			{
+				Task->LandscapeComponent = nullptr;
+			});
+	}
+};
+
+/** Static global pool object. */
+static FLandscapePhysicalMaterialRenderTaskPool GTaskPool;
+
+
+void FLandscapePhysicalMaterialRenderTask::Init(ULandscapeComponent const* LandscapeComponent)
+{
+	check(IsInGameThread());
+	if (IsValid())
+	{
+		GTaskPool.Free(*this);
+	}
+	GTaskPool.Allocate(*this, LandscapeComponent);
+}
+
+void FLandscapePhysicalMaterialRenderTask::Release()
+{
+	check(IsInGameThread());
+	if (IsValid())
+	{
+		GTaskPool.Free(*this);
+	}
+}
+
+bool FLandscapePhysicalMaterialRenderTask::IsValid() const
+{
+	check(IsInGameThread());
+	return PoolHandle != -1;
+}
+
+bool FLandscapePhysicalMaterialRenderTask::IsComplete() const
+{
+	check(IsInGameThread());
+	check(IsValid());
+	return GTaskPool.Pool[PoolHandle].CompletionState == ECompletionState::Complete;
+}
+
+void FLandscapePhysicalMaterialRenderTask::Tick()
+{
+	check(IsInGameThread());
+	if (IsValid() && !IsComplete())
+	{
+		FLandscapePhysicalMaterialRenderTaskImpl* Task = &GTaskPool.Pool[PoolHandle];
+		
+		ENQUEUE_RENDER_COMMAND(FLandscapePhysicalMaterialUpdaterTick)(
+			[Task](FRHICommandListImmediate& RHICmdList)
+			{
+				check(Task->LandscapeComponent);
+				UpdateTask_RenderThread(RHICmdList, *Task, false);
+			});
+	}
+}
+
+// Note: We could add a global function that calls Flush() on multiple tasks. 
+// That could enqueue a single render thread function and use a single rendering command flush.
+// It could be useful if we see performance issue with any path that calls Flush individually for each landscape component.
+void FLandscapePhysicalMaterialRenderTask::Flush()
+{
+	check(IsInGameThread());
+	if (IsValid() && !IsComplete())
+	{
+		FLandscapePhysicalMaterialRenderTaskImpl* Task = &GTaskPool.Pool[PoolHandle];
+
+		ENQUEUE_RENDER_COMMAND(FLandscapePhysicalMaterialUpdaterTick)(
+			[Task](FRHICommandListImmediate& RHICmdList)
+			{
+				UpdateTask_RenderThread(RHICmdList, *Task, true);
+			});
+
+		FlushRenderingCommands();
+	}
+}
+
+TArray<uint8> const& FLandscapePhysicalMaterialRenderTask::GetResultIds() const
+{
+	check(IsInGameThread());
+	check(IsValid() && IsComplete());
+	return GTaskPool.Pool[PoolHandle].ResultIds;
+}
+
+TArray<UPhysicalMaterial*> const& FLandscapePhysicalMaterialRenderTask::GetResultMaterials() const
+{
+	check(IsInGameThread());
+	check(IsValid() && IsComplete());
+	return GTaskPool.Pool[PoolHandle].ResultMaterials;
+}
+
+#endif // WITH_EDITOR
