Merge fixes for new landscape physical material code
#rnx

diff --git a/Engine/Source/Runtime/Landscape/Private/LandscapePhysicalMaterial.cpp b/Engine/Source/Runtime/Landscape/Private/LandscapePhysicalMaterial.cpp
--- a/Engine/Source/Runtime/Landscape/Private/LandscapePhysicalMaterial.cpp
+++ b/Engine/Source/Runtime/Landscape/Private/LandscapePhysicalMaterial.cpp
@@ -81,9 +81,9 @@ public:
 
 	static bool ShouldCompilePermutation(const FMeshMaterialShaderPermutationParameters& Parameters)
 	{
-		return (Parameters.Material->IsUsedWithLandscape() || Parameters.Material->IsSpecialEngineMaterial()) &&
+		return (Parameters.MaterialParameters.bIsUsedWithLandscape || Parameters.MaterialParameters.bIsSpecialEngineMaterial) &&
 			IsFeatureLevelSupported(Parameters.Platform, ERHIFeatureLevel::SM5) &&
-			((Parameters.VertexFactoryType == FindVertexFactoryType(FName(TEXT("FLandscapeVertexFactory"), FNAME_Find))) || (Parameters.VertexFactoryType == FindVertexFactoryType(FName(TEXT("FLandscapeXYOffsetVertexFactory"), FNAME_Find)))) &&
+			Parameters.VertexFactoryType == FindVertexFactoryType(FName(TEXT("FLandscapeFixedGridVertexFactory"), FNAME_Find)) &&
 			!IsConsolePlatform(Parameters.Platform);
 	}
 };
@@ -189,7 +189,8 @@ void FLandscapePhysicalMaterialMeshProcessor::Process(
 	PassShaders.PixelShader = MaterialResource.GetShader<FLandscapePhysicalMaterialPS>(VertexFactory->GetType());
 	PassShaders.VertexShader = MaterialResource.GetShader<FLandscapePhysicalMaterialVS>(VertexFactory->GetType());
 
-	const ERasterizerFillMode MeshFillMode = ComputeMeshFillMode(MeshBatch, MaterialResource);
+	const FMeshDrawingPolicyOverrideSettings OverrideSettings = ComputeMeshOverrideSettings(MeshBatch);
+	const ERasterizerFillMode MeshFillMode = ComputeMeshFillMode(MeshBatch, MaterialResource, OverrideSettings);
 	const ERasterizerCullMode MeshCullMode = CM_None;
 
 	FMeshMaterialShaderElementData ShaderElementData;
@@ -281,7 +282,7 @@ namespace
 
 			{
 				FMemMark Mark(FMemStack::Get());
-			
+
 				DrawDynamicMeshPass(*View, RHICmdList,
 					[View, &MeshInfos = MeshInfos](FDynamicPassMeshDrawListContext* DynamicMeshPassContext)
 					{
@@ -504,7 +505,7 @@ public:
 	void Free(FLandscapePhysicalMaterialRenderTask& InTask)
 	{
 		check(InTask.PoolHandle != -1);
-		
+
 		FLandscapePhysicalMaterialRenderTaskImpl* Task = &Pool[InTask.PoolHandle];
 		
 		// Invalidate the task object.
@@ -561,7 +562,7 @@ void FLandscapePhysicalMaterialRenderTask::Tick()
 	if (IsValid() && !IsComplete())
 	{
 		FLandscapePhysicalMaterialRenderTaskImpl* Task = &GTaskPool.Pool[PoolHandle];
-		
+
 		ENQUEUE_RENDER_COMMAND(FLandscapePhysicalMaterialUpdaterTick)(
 			[Task](FRHICommandListImmediate& RHICmdList)
 			{
diff --git a/Engine/Source/Runtime/Landscape/Private/LandscapeRender.cpp b/Engine/Source/Runtime/Landscape/Private/LandscapeRender.cpp
--- a/Engine/Source/Runtime/Landscape/Private/LandscapeRender.cpp
+++ b/Engine/Source/Runtime/Landscape/Private/LandscapeRender.cpp
@@ -4969,6 +4969,8 @@ public:
 		{
 			FName(TEXT("FLandscapeGrassWeightVS")),
 			FName(TEXT("FLandscapeGrassWeightPS")),
+			FName(TEXT("FLandscapePhysicalMaterialVS")),
+			FName(TEXT("FLandscapePhysicalMaterialPS")),
 		};
 		return ShaderTypes;
 	}
