Fixed infinite loop when undoing a landscape change

diff --git a/Engine/Source/Runtime/Landscape/Classes/LandscapeComponent.h b/Engine/Source/Runtime/Landscape/Classes/LandscapeComponent.h
--- a/Engine/Source/Runtime/Landscape/Classes/LandscapeComponent.h
+++ b/Engine/Source/Runtime/Landscape/Classes/LandscapeComponent.h
@@ -341,9 +341,9 @@ enum ELandscapeLayerUpdateMode : uint32
 	Update_All_Editing = Update_Weightmap_Editing | Update_Heightmap_Editing,
 	Update_All_Editing_NoCollision = Update_Weightmap_Editing_NoCollision | Update_Heightmap_Editing_NoCollision,
 	// In cases where we couldn't update the clients right away this flag will be set in RegenerateLayersContent
-	Update_Client_Deferred = 1 << 4,
+	Update_Client_Deferred = 1 << 6,
 	// Update landscape component clients while editing
-	Update_Client_Editing = 1 << 5
+	Update_Client_Editing = 1 << 7
 };
 
 static const uint32 DefaultSplineHash = 0xFFFFFFFF;
