Garbage collect landscape physcial material render resources

diff --git a/Engine/Source/Runtime/Landscape/Classes/LandscapePhysicalMaterial.h b/Engine/Source/Runtime/Landscape/Classes/LandscapePhysicalMaterial.h
--- a/Engine/Source/Runtime/Landscape/Classes/LandscapePhysicalMaterial.h
+++ b/Engine/Source/Runtime/Landscape/Classes/LandscapePhysicalMaterial.h
@@ -53,4 +53,10 @@ private:
 	FLandscapePhysicalMaterialRenderTask(FLandscapePhysicalMaterialRenderTask const&);
 };
 
+namespace LandscapePhysicalMaterial
+{
+	/** Optional call to allow garbage collection of any FLandscapePhysicalMaterialRenderTask resources and check for leaks. */
+	void GarbageCollectTasks();
+}
+
 #endif // WITH_EDITOR
diff --git a/Engine/Source/Runtime/Landscape/Private/LandscapePhysicalMaterial.cpp b/Engine/Source/Runtime/Landscape/Private/LandscapePhysicalMaterial.cpp
--- a/Engine/Source/Runtime/Landscape/Private/LandscapePhysicalMaterial.cpp
+++ b/Engine/Source/Runtime/Landscape/Private/LandscapePhysicalMaterial.cpp
@@ -6,6 +6,7 @@
 
 #include "EngineModule.h"
 #include "LandscapeComponent.h"
+#include "LandscapePrivate.h"
 #include "LandscapeRender.h"
 #include "Materials/MaterialExpressionLandscapePhysicalMaterialOutput.h"
 #include "Materials/MaterialInstanceConstant.h"
@@ -358,6 +359,7 @@ namespace
 	{
 		// Create on game thread
 		ULandscapeComponent const* LandscapeComponent;
+		uint32 InitFrameId;
 		FIntPoint TargetSize;
 		FVector ViewOrigin;
 		FMatrix ViewRotationMatrix;
@@ -376,13 +378,14 @@ namespace
 	};
 
 	/** Initialize the physical material render task data. */
-	bool InitTask(FLandscapePhysicalMaterialRenderTaskImpl& Task, ULandscapeComponent const* InLandscapeComponent)
+	bool InitTask(FLandscapePhysicalMaterialRenderTaskImpl& Task, ULandscapeComponent const* InLandscapeComponent, uint32 InFrameId)
 	{
 		if (InLandscapeComponent != nullptr)
 		{
 			if (GetPhysicalMaterials(InLandscapeComponent, Task.ResultMaterials))
 			{
 				Task.LandscapeComponent = InLandscapeComponent;
+				Task.InitFrameId = InFrameId;
 				Task.CompletionState = ECompletionState::None;
 
 				const FTransform& ComponentTransform = InLandscapeComponent->GetComponentTransform();
@@ -467,18 +470,19 @@ namespace
 }
 
 
-/** 
- * Pool for storing physical material render task data. 
- * todo: Would be good to garbage collect the render resources from the pool but that would require some global tick function.
+/**
+ * Pool for storing physical material render task data.
  */
 class FLandscapePhysicalMaterialRenderTaskPool
 {
 public:
 	/** Pool uses chunked array to avoid task data being moved by a realloc. */
 	TChunkedArray< FLandscapePhysicalMaterialRenderTaskImpl > Pool;
+	/** Frame count used to validate and garbage collect. */
+	uint32 FrameCount = 0;
 
 	/** Allocate task data from the pool. */
-	void Allocate(FLandscapePhysicalMaterialRenderTask& InTask, ULandscapeComponent const* LandscapeComponent)
+	void Allocate(FLandscapePhysicalMaterialRenderTask& InTask, ULandscapeComponent const* InLandscapeComponent)
 	{
 		check(InTask.PoolHandle == -1);
 
@@ -497,7 +501,7 @@ public:
 			Pool.Add();
 		}
 
-		const bool bSuccess = InitTask(Pool[Index], LandscapeComponent);
+		const bool bSuccess = InitTask(Pool[Index], InLandscapeComponent, FrameCount);
 		InTask.PoolHandle = bSuccess ? Index : -1;
 	}
 
@@ -507,7 +511,7 @@ public:
 		check(InTask.PoolHandle != -1);
 
 		FLandscapePhysicalMaterialRenderTaskImpl* Task = &Pool[InTask.PoolHandle];
-		
+
 		// Invalidate the task object.
 		InTask.PoolHandle = -1;
 
@@ -518,6 +522,41 @@ public:
 				Task->LandscapeComponent = nullptr;
 			});
 	}
+
+	/** Free render resources that have been unused for long enough. */
+	void GarbageCollect()
+	{
+		const uint32 PoolSize = Pool.Num();
+		if (PoolSize > 0)
+		{
+			// Garbage collect a maximum of one item per call to reduce overhead if pool has grown large.
+			FLandscapePhysicalMaterialRenderTaskImpl* Task = &Pool[FrameCount % PoolSize];
+			if (Task->InitFrameId + 100 < FrameCount)
+			{
+				if (Task->LandscapeComponent != nullptr)
+				{
+					// Task not completed after 100 updates. We are probably leaking tasks!
+					UE_LOG(LogLandscape, Warning, TEXT("Leaking landscape physical material tasks."))
+				}
+				else
+				{
+					// Free the array allocations
+					Task->ResultMaterials.Empty();
+					Task->ResultIds.Empty();
+
+					// Free the render resources (which may already be free)
+					ENQUEUE_RENDER_COMMAND(FLandscapePhysicalMaterialUpdaterTick)(
+						[Task](FRHICommandListImmediate& RHICmdList)
+						{
+							Task->ReadbackTexture.SafeRelease();
+							Task->ReadbackFence.SafeRelease();
+						});
+				}
+			}
+		}
+
+		FrameCount++;
+	}
 };
 
 /** Static global pool object. */
@@ -606,4 +645,14 @@ TArray<UPhysicalMaterial*> const& FLandscapePhysicalMaterialRenderTask::GetResul
 	return GTaskPool.Pool[PoolHandle].ResultMaterials;
 }
 
+
+namespace LandscapePhysicalMaterial
+{
+	void GarbageCollectTasks()
+	{
+		check(IsInGameThread());
+		GTaskPool.GarbageCollect();
+	}
+}
+
 #endif // WITH_EDITOR
