Support for storing and building rendered landscape physical materials for a ULandscapeHeightfieldCollisionComponent

diff --git a/Engine/Source/Runtime/Landscape/Classes/LandscapeHeightfieldCollisionComponent.h b/Engine/Source/Runtime/Landscape/Classes/LandscapeHeightfieldCollisionComponent.h
--- a/Engine/Source/Runtime/Landscape/Classes/LandscapeHeightfieldCollisionComponent.h
+++ b/Engine/Source/Runtime/Landscape/Classes/LandscapeHeightfieldCollisionComponent.h
@@ -119,6 +119,13 @@ class ULandscapeHeightfieldCollisionComponent : public UPrimitiveComponent
 	/** Indices into the ComponentLayers array for the per-vertex dominant layer. Stripped from cooked content */
 	FByteBulkData								DominantLayerData;
 
+	/** Indices for physical materials generated by the render material. Stripped from cooked content */
+	FByteBulkData								PhysicalMaterialRenderData;
+
+	/** Physical materials objects referenced by the indices in PhysicalMaterialRenderData. Stripped from cooked content */
+	UPROPERTY()
+	TArray<UPhysicalMaterial*>					PhysicalMaterialRenderObjects;
+
 	/*  Cooked editor specific heightfield data, never serialized  */
 	TArray<uint8>								CookedCollisionDataEd;
 
diff --git a/Engine/Source/Runtime/Landscape/Private/LandscapeCollision.cpp b/Engine/Source/Runtime/Landscape/Private/LandscapeCollision.cpp
--- a/Engine/Source/Runtime/Landscape/Private/LandscapeCollision.cpp
+++ b/Engine/Source/Runtime/Landscape/Private/LandscapeCollision.cpp
@@ -731,7 +731,17 @@ void ULandscapeHeightfieldCollisionComponent::SpeculativelyLoadAsyncDDCCollsionD
 }
 
 #if PHYSICS_INTERFACE_PHYSX
-TArray<PxHeightFieldSample> ConvertHeightfieldDataForPhysx(const ULandscapeHeightfieldCollisionComponent* const Component, const int32 CollisionSizeVerts, const bool bIsMirrored, const uint16* Heights, const bool bUseDefMaterial, const uint8* DominantLayers, UPhysicalMaterial* const DefMaterial, TArray<UPhysicalMaterial*> &InOutMaterials)
+TArray<PxHeightFieldSample> ConvertHeightfieldDataForPhysx(
+	const ULandscapeHeightfieldCollisionComponent* const Component,
+	const int32 CollisionSizeVerts,
+	const bool bIsMirrored,
+	const uint16* Heights,
+	const bool bUseDefMaterial,
+	UPhysicalMaterial* const DefMaterial,
+	const uint8* DominantLayers,
+	const uint8* RenderPhysicalMaterialIds,
+	TArray<UPhysicalMaterial*> const& PhysicalMaterialRenderObjects,
+	TArray<UPhysicalMaterial*>& InOutMaterials) 
 {
 	const int32 NumSamples = FMath::Square(CollisionSizeVerts);
 	check(DefMaterial);
@@ -759,24 +769,27 @@ TArray<PxHeightFieldSample> ConvertHeightfieldDataForPhysx(const ULandscapeHeigh
 				int32 MaterialIndex = DefaultMaterialIndex; // Default physical material.
 				if (!bUseDefMaterial && DominantLayers)
 				{
-					uint8 DominantLayerIdx = DominantLayers[SrcSampleIndex];
-					if (Component->ComponentLayerInfos.IsValidIndex(DominantLayerIdx))
-					{
+					uint8 DominantLayerIdx = DominantLayers ? DominantLayers[SrcSampleIndex] : -1;
+					ULandscapeLayerInfoObject* Layer = ComponentLayerInfos.IsValidIndex(DominantLayerIdx) ? ComponentLayerInfos[DominantLayerIdx] : nullptr;
+
 						ULandscapeLayerInfoObject* Layer = Component->ComponentLayerInfos[DominantLayerIdx];
 						if (Layer == ALandscapeProxy::VisibilityLayer)
 						{
 							// If it's a hole, override with the hole flag.
 							MaterialIndex = PxHeightFieldMaterial::eHOLE;
 						}
-						else if (Layer && Layer->PhysMaterial)
+					else if (RenderPhysicalMaterialIds)
+					{
+						uint8 RenderIdx = RenderPhysicalMaterialIds[SrcSampleIndex];
+						if (RenderIdx > 0)
 						{
-							MaterialIndex = InOutMaterials.AddUnique(Layer->PhysMaterial);
+							MaterialIndex = InOutMaterials.AddUnique(PhysicalMaterialRenderObjects[RenderIdx - 1]);
 						}
-						else
+					}
+						else if (Layer && Layer->PhysMaterial)
 						{
-							MaterialIndex = DefaultMaterialIndex;
+							MaterialIndex = InOutMaterials.AddUnique(Layer->PhysMaterial);
 						}
-					}
 				}
 
 				// Default Material but Def Material wasn't added yet...
@@ -875,13 +888,26 @@ bool ULandscapeHeightfieldCollisionComponent::CookCollisionData(const FName& For
 
 	const uint16* Heights = (const uint16*)CollisionHeightData.LockReadOnly();
 	check(CollisionHeightData.GetElementCount() == NumSamples + NumSimpleSamples);
+	const uint16* SimpleHeights = Heights + NumSamples;
 
-	//
+	// Physical material data from layer system
 	const uint8* DominantLayers = nullptr;
+	const uint8* SimpleDominantLayers = nullptr;
 	if (DominantLayerData.GetElementCount())
 	{
 		DominantLayers = (const uint8*)DominantLayerData.LockReadOnly();
 		check(DominantLayerData.GetElementCount() == NumSamples + NumSimpleSamples);
+		SimpleDominantLayers = DominantLayers + NumSamples;
+	}
+
+	// Physical material data from render material graph
+	const uint8* RenderPhysicalMaterialIds = nullptr;
+	const uint8* SimpleRenderPhysicalMaterialIds = nullptr;
+	if (PhysicalMaterialRenderData.GetElementCount())
+	{
+		RenderPhysicalMaterialIds = (const uint8*)PhysicalMaterialRenderData.LockReadOnly();
+		check(PhysicalMaterialRenderData.GetElementCount() == NumSamples + NumSimpleSamples);
+		SimpleRenderPhysicalMaterialIds = RenderPhysicalMaterialIds + NumSamples;
 	}
 
 #if WITH_PHYSX && PHYSICS_INTERFACE_PHYSX
@@ -890,12 +916,11 @@ bool ULandscapeHeightfieldCollisionComponent::CookCollisionData(const FName& For
 		
 	TArray<PxHeightFieldSample> Samples;
 	TArray<PxHeightFieldSample> SimpleSamples;
-	Samples = ConvertHeightfieldDataForPhysx(this, CollisionSizeVerts, bIsMirrored, Heights, bUseDefMaterial, DominantLayers, DefMaterial, InOutMaterials);
+	Samples = ConvertHeightfieldDataForPhysx(this, CollisionSizeVerts, bIsMirrored, Heights, bUseDefMaterial, DefMaterial, DominantLayers, RenderPhysicalMaterialIds, PhysicalMaterialRenderObjects, InOutMaterials);
 
 	if (bGenerateSimpleCollision)
 	{
-		const uint8* SimpleDominantLayers = DominantLayers ? DominantLayers + NumSamples : nullptr;
-		SimpleSamples = ConvertHeightfieldDataForPhysx(this, SimpleCollisionSizeVerts, bIsMirrored, Heights + NumSamples, bUseDefMaterial, SimpleDominantLayers, DefMaterial, InOutMaterials);
+		SimpleSamples = ConvertHeightfieldDataForPhysx(this, SimpleCollisionSizeVerts, bIsMirrored, SimpleHeights, bUseDefMaterial, DefMaterial, SimpleDominantLayers, SimpleRenderPhysicalMaterialIds, PhysicalMaterialRenderObjects, InOutMaterials);
 	}
 
 	CollisionHeightData.Unlock();
@@ -932,24 +957,28 @@ bool ULandscapeHeightfieldCollisionComponent::CookCollisionData(const FName& For
 				ColIndex < CollisionSizeVerts - 1)
 			{
 				int32 MaterialIndex = 0; // Default physical material.
-				if(!bUseDefMaterial && DominantLayers)
+				if(!bUseDefMaterial)
 				{
-					uint8 DominantLayerIdx = DominantLayers[SrcSampleIndex];
-					if(ComponentLayerInfos.IsValidIndex(DominantLayerIdx))
-					{
-						ULandscapeLayerInfoObject* Layer = ComponentLayerInfos[DominantLayerIdx];
+					uint8 DominantLayerIdx = DominantLayers ? DominantLayers[SrcSampleIndex] : -1;
+					ULandscapeLayerInfoObject* Layer = ComponentLayerInfos.IsValidIndex(DominantLayerIdx) ? ComponentLayerInfos[DominantLayerIdx] : nullptr;
+
 						if(Layer == ALandscapeProxy::VisibilityLayer)
 						{
 							// If it's a hole, use the final index
 							MaterialIndex = TNumericLimits<uint8>::Max();
 						}
+					else if (RenderPhysicalMaterialIds)
+					{
+						uint8 RenderIdx = RenderPhysicalMaterialIds[SrcSampleIndex];
+						UPhysicalMaterial* DominantMaterial = RenderIdx > 0 ? PhysicalMaterialRenderObjects[RenderIdx - 1] : DefMaterial;
+						MaterialIndex = InOutMaterials.AddUnique(DominantMaterial);
+					}
 						else
 						{
 							UPhysicalMaterial* DominantMaterial = Layer && Layer->PhysMaterial ? Layer->PhysMaterial : DefMaterial;
 							MaterialIndex = InOutMaterials.AddUnique(DominantMaterial);
 						}
 					}
-				}
 				MaterialIndices.Add(MaterialIndex);
 			}
 		}
@@ -975,15 +1004,21 @@ bool ULandscapeHeightfieldCollisionComponent::CookCollisionData(const FName& For
 		Ar << HeightfieldSimple;
 	}
 
+	Succeeded = true;
+#endif
+
+	if (CollisionHeightData.IsLocked())
+	{
 	CollisionHeightData.Unlock();
-	if(DominantLayers)
+	}
+	if (DominantLayerData.IsLocked())
 	{
 		DominantLayerData.Unlock();
 	}
-
-	Succeeded = true;
-
-#endif
+	if (PhysicalMaterialRenderData.IsLocked())
+	{
+		PhysicalMaterialRenderData.Unlock();
+	}
 
 	if (Succeeded)
 	{
@@ -1462,6 +1497,10 @@ uint32 ULandscapeHeightfieldCollisionComponent::ComputeCollisionHash() const
 	Hash = FCrc::MemCrc32(DominantBuffer, DominantLayerData.GetBulkDataSize(), Hash);
 	DominantLayerData.Unlock();
 
+	const void* PhysicalMaterialBuffer = PhysicalMaterialRenderData.LockReadOnly();
+	Hash = FCrc::MemCrc32(PhysicalMaterialBuffer, PhysicalMaterialRenderData.GetBulkDataSize(), Hash);
+	PhysicalMaterialRenderData.Unlock();
+
 	return Hash;
 }
 
