Fix landscape physical material staging buffer copy for D3D12.
Original code is good but it won't work until some D3D12 RHI fixes make their way across streams.
So changing to use a working RHI path.
#rnx

diff --git a/Engine/Source/Runtime/Landscape/Private/LandscapePhysicalMaterial.cpp b/Engine/Source/Runtime/Landscape/Private/LandscapePhysicalMaterial.cpp
--- a/Engine/Source/Runtime/Landscape/Private/LandscapePhysicalMaterial.cpp
+++ b/Engine/Source/Runtime/Landscape/Private/LandscapePhysicalMaterial.cpp
@@ -308,7 +308,8 @@ namespace
 
 			// Copy to the read back texture
 			RHICmdList.TransitionResource(EResourceTransitionAccess::EReadable, RenderTargetTexture);
-			RHICmdList.CopyTexture(RenderTargetTexture, ReadbackTexture, FRHICopyTextureInfo());
+			RHICmdList.TransitionResource(EResourceTransitionAccess::EWritable, ReadbackTexture);
+			RHICmdList.CopyToResolveTarget(RenderTargetTexture, ReadbackTexture, FResolveParams());
 			RHICmdList.WriteGPUFence(ReadbackFence);
 		}
 	}
@@ -416,7 +417,7 @@ namespace
 		//todo: Consider pooling these and throttling to the pool size?
 		if (!Task.ReadbackTexture.IsValid() || Task.ReadbackTexture->GetSizeXYZ() != FIntVector(Task.TargetSize.X, Task.TargetSize.Y, 1))
 		{
-			FRHIResourceCreateInfo CreateInfo(TEXT(""));
+			FRHIResourceCreateInfo CreateInfo(TEXT("LandscapePhysicalMaterialReadback"));
 			Task.ReadbackTexture = RHICreateTexture2D(Task.TargetSize.X, Task.TargetSize.Y, PF_G8, 1, 1, TexCreate_CPUReadback | TexCreate_HideInVisualizeTexture, CreateInfo);
 		}
 		if (!Task.ReadbackFence.IsValid())
