Hook up landscape physical material tasks to landscape subsystem.
The subsystem will tick the tasks and copy results to landscape collision whenever a task completes.

diff --git a/Engine/Source/Runtime/Landscape/Classes/LandscapeComponent.h b/Engine/Source/Runtime/Landscape/Classes/LandscapeComponent.h
--- a/Engine/Source/Runtime/Landscape/Classes/LandscapeComponent.h
+++ b/Engine/Source/Runtime/Landscape/Classes/LandscapeComponent.h
@@ -8,6 +8,7 @@
 #include "Engine/TextureStreamingTypes.h"
 #include "Components/PrimitiveComponent.h"
 #include "PerPlatformProperties.h"
+#include "LandscapePhysicalMaterial.h"
 #include "LandscapeWeightmapUsage.h"
 
 #include "LandscapeComponent.generated.h"
@@ -589,6 +590,10 @@ public:
 	TArray<FBox> ActiveExcludedBoxes;
 	uint32 ChangeTag;
 
+#if WITH_EDITOR
+	/** Physical material update task */
+	FLandscapePhysicalMaterialRenderTask PhysicalMaterialTask;
+#endif
 
 	//~ Begin UObject Interface.	
 	virtual void PostInitProperties() override;	
@@ -907,6 +912,11 @@ public:
 	 */
 	LANDSCAPE_API void UpdateCollisionLayerData();
 
+	/** Update physical material render tasks. */
+	void UpdatePhysicalMaterialTasks();
+	/** Update collision component physical materials from render task results. */
+	void UpdateCollisionPhysicalMaterialData(TArray<UPhysicalMaterial*> const& InPhysicalMaterials, TArray<uint8> const& InMaterialIds);
+
 	/**
 	 * Create weightmaps for this component for the layers specified in the WeightmapLayerAllocations array
 	 */
diff --git a/Engine/Source/Runtime/Landscape/Classes/LandscapeProxy.h b/Engine/Source/Runtime/Landscape/Classes/LandscapeProxy.h
--- a/Engine/Source/Runtime/Landscape/Classes/LandscapeProxy.h
+++ b/Engine/Source/Runtime/Landscape/Classes/LandscapeProxy.h
@@ -904,6 +904,9 @@ public:
 	/** Update any textures baked from the landscape as necessary */
 	void UpdateBakedTextures();
 
+	/** Update the landscape physical material render tasks */
+	void UpdatePhysicalMaterialTasks();
+
 	/** Frame counter to count down to the next time we check to update baked textures, so we don't check every frame */
 	int32 UpdateBakedTexturesCountdown;
 
diff --git a/Engine/Source/Runtime/Landscape/Private/Landscape.cpp b/Engine/Source/Runtime/Landscape/Private/Landscape.cpp
--- a/Engine/Source/Runtime/Landscape/Private/Landscape.cpp
+++ b/Engine/Source/Runtime/Landscape/Private/Landscape.cpp
@@ -1486,6 +1486,10 @@ void ULandscapeComponent::OnUnregister()
 {
 	Super::OnUnregister();
 
+#if WITH_EDITOR
+	PhysicalMaterialTask.Release();
+#endif
+
 	if (GetLandscapeProxy())
 	{
 		// Generate MID representing the MIC
@@ -3729,6 +3733,14 @@ void ALandscapeProxy::UpdateBakedTextures()
 		UpdateBakedTexturesCountdown = 60;
 	}
 }
+
+void ALandscapeProxy::UpdatePhysicalMaterialTasks()
+{
+	for (ULandscapeComponent* Component : LandscapeComponents)
+	{
+		Component->UpdatePhysicalMaterialTasks();
+	}
+}
 #endif
 
 template<class ContainerType>
diff --git a/Engine/Source/Runtime/Landscape/Private/LandscapeEdit.cpp b/Engine/Source/Runtime/Landscape/Private/LandscapeEdit.cpp
--- a/Engine/Source/Runtime/Landscape/Private/LandscapeEdit.cpp
+++ b/Engine/Source/Runtime/Landscape/Private/LandscapeEdit.cpp
@@ -1031,6 +1031,10 @@ void ULandscapeComponent::UpdateCollisionHeightData(const FColor* const Heightma
 	{
 		CollisionComp->RegisterComponent();
 	}
+
+	// Kick off asynchronous rendering task for physical materials
+	// This task will be updated in UpdatePhysicalMaterialTasks()
+	PhysicalMaterialTask.Init(this);
 }
 
 void ULandscapeComponent::DestroyCollisionData()
@@ -1434,6 +1438,10 @@ void ULandscapeComponent::UpdateCollisionLayerData(const FColor* const* const We
 			CollisionComp->DominantLayerData.Unlock();
 		}
 
+		// Kick off asynchronous rendering task for physical materials
+		// This task will be updated in UpdatePhysicalMaterialTasks()
+		PhysicalMaterialTask.Init(this);
+
 		// We do not force an update of the physics data here. We don't need the layer information in the editor and it
 		// causes problems if we update it multiple times in a single frame.
 	}
@@ -1473,8 +1481,67 @@ void ULandscapeComponent::UpdateCollisionLayerData()
 	UpdateCollisionLayerData(WeightmapTextureMipDataParam.GetData(), SimpleCollisionWeightmapMipDataParam.GetData());
 }
 
+void ULandscapeComponent::UpdatePhysicalMaterialTasks()
+{
+	if (PhysicalMaterialTask.IsValid())
+	{
+		if (PhysicalMaterialTask.IsComplete())
+		{
+			UpdateCollisionPhysicalMaterialData(PhysicalMaterialTask.GetResultMaterials(), PhysicalMaterialTask.GetResultIds());
+
+			PhysicalMaterialTask.Release();
+
+			// We do not force an update of the physics data here. We don't need the information in the editor.
+		}
+		else
+		{
+			PhysicalMaterialTask.Tick();
+		}
+	}
+}
+
+void ULandscapeComponent::UpdateCollisionPhysicalMaterialData(TArray<UPhysicalMaterial*> const& InPhysicalMaterials, TArray<uint8> const& InMaterialIds)
+{
+	// Copy the physical material array
+	CollisionComponent->PhysicalMaterialRenderObjects = InPhysicalMaterials;
+
+	// Copy the physical material IDs for both the full and (optional) simple collision.
+	const int32 SizeVerts = SubsectionSizeQuads * NumSubsections + 1;
+	check(InMaterialIds.Num() == SizeVerts * SizeVerts);
+	const int32 FullCollisionSizeVerts = CollisionComponent->CollisionSizeQuads + 1;
+	const int32 SimpleCollisionSizeVerts = CollisionComponent->SimpleCollisionSizeQuads > 0 ? CollisionComponent->SimpleCollisionSizeQuads + 1 : 0;
+	const int32 BulkDataSize = FullCollisionSizeVerts * FullCollisionSizeVerts + SimpleCollisionSizeVerts * SimpleCollisionSizeVerts;
 
+	void* Data = CollisionComponent->PhysicalMaterialRenderData.Lock(LOCK_READ_WRITE);
+	Data = CollisionComponent->PhysicalMaterialRenderData.Realloc(BulkDataSize);
+	uint8* WritePtr = (uint8*)Data;
+
+	const int32 CollisionSizes[2] = { FullCollisionSizeVerts, SimpleCollisionSizeVerts };
+	for (int32 i = 0; i < 2; ++i)
+	{
+		const int32 CollisionSizeVerts = CollisionSizes[i];
+		if (CollisionSizeVerts == SizeVerts)
+		{
+			FMemory::Memcpy(WritePtr, InMaterialIds.GetData(), SizeVerts * SizeVerts);
+			WritePtr += SizeVerts * SizeVerts;
+		}
+		else if (CollisionSizeVerts > 0)
+		{
+			const int32 StepSize = SizeVerts / CollisionSizeVerts;
+			check(CollisionSizeVerts * StepSize == SizeVerts);
+			for (int32 y = 0; y < SizeVerts; y += StepSize)
+			{
+				for (int32 x = 0; x < SizeVerts; x += StepSize)
+				{
+					*WritePtr++ = InMaterialIds[y * SizeVerts + x];
+				}
+			}
+		}
+	}
 
+	check(WritePtr - (uint8*)Data == BulkDataSize);
+	CollisionComponent->PhysicalMaterialRenderData.Unlock();
+}
 
 void ULandscapeComponent::GenerateHeightmapMips(TArray<FColor*>& HeightmapTextureMipData, int32 ComponentX1/*=0*/, int32 ComponentY1/*=0*/, int32 ComponentX2/*=MAX_int32*/, int32 ComponentY2/*=MAX_int32*/, FLandscapeTextureDataInfo* TextureDataInfo/*=nullptr*/)
 {
diff --git a/Engine/Source/Runtime/Landscape/Private/LandscapeSubsystem.cpp b/Engine/Source/Runtime/Landscape/Private/LandscapeSubsystem.cpp
--- a/Engine/Source/Runtime/Landscape/Private/LandscapeSubsystem.cpp
+++ b/Engine/Source/Runtime/Landscape/Private/LandscapeSubsystem.cpp
@@ -110,6 +110,7 @@ void ULandscapeSubsystem::ExecuteTick(float DeltaTime, ELevelTick TickType, ENam
 			if (!World->IsPlayInEditor())
 			{
 				Proxy->UpdateBakedTextures();
+				Proxy->UpdatePhysicalMaterialTasks();
 			}
 		}
 #endif
@@ -118,6 +119,13 @@ void ULandscapeSubsystem::ExecuteTick(float DeltaTime, ELevelTick TickType, ENam
 			Proxy->TickGrass(*Cameras, InOutNumComponentsCreated);
 		}
 	}
+
+#if WITH_EDITOR
+	if (GIsEditor && !World->IsPlayInEditor())
+	{
+		LandscapePhysicalMaterial::GarbageCollectTasks();
+	}
+#endif
 }
 
 FString ULandscapeSubsystem::DiagnosticMessage()
