Fixed Grass Align to Surface normal computation when using Halton sequence placement.
#tests Validated at specified locations that rocks are no more floating and well aligned on landscape surface

diff --git a/Engine/Source/Runtime/Landscape/Private/LandscapeGrass.cpp b/Engine/Source/Runtime/Landscape/Private/LandscapeGrass.cpp
--- a/Engine/Source/Runtime/Landscape/Private/LandscapeGrass.cpp
+++ b/Engine/Source/Runtime/Landscape/Private/LandscapeGrass.cpp
@@ -1722,7 +1722,9 @@ struct FAsyncGrassBuilder : public FGrassBuilderBase
 				float HaltonY = Halton(InstanceIndex + HaltonBaseIndex, 3);
 				FVector Location(Origin.X + HaltonX * Extent.X, Origin.Y + HaltonY * Extent.Y, 0.0f);
 				FVector LocationWithHeight;
-				float Weight = GetLayerWeightAtLocationLocal(Location, &LocationWithHeight);
+				FVector ComputedNormal;
+				float Weight = 0.f;
+				SampleLandscapeAtLocationLocal(Location, LocationWithHeight, Weight, AlignToSurface ? &ComputedNormal : nullptr);
 				bool bKeep = Weight > 0.0f && Weight >= RandomStream.GetFraction() && !IsExcluded(LocationWithHeight);
 				if (bKeep)
 				{
@@ -1730,33 +1732,13 @@ struct FAsyncGrassBuilder : public FGrassBuilderBase
 					const float Rot = RandomRotation ? RandomStream.GetFraction() * 360.0f : 0.0f;
 					const FMatrix BaseXForm = FScaleRotationTranslationMatrix(Scale, FRotator(0.0f, Rot, 0.0f), FVector::ZeroVector);
 					FMatrix OutXForm;
-					if (AlignToSurface)
+					if (AlignToSurface && !ComputedNormal.IsNearlyZero())
 					{
-						FVector LocationWithHeightDX;
-						FVector LocationDX(Location);
-						LocationDX.X = FMath::Clamp<float>(LocationDX.X + (HaltonX < 0.5f ? DivExtent.X : -DivExtent.X), Origin.X, Origin.X + Extent.X);
-						GetLayerWeightAtLocationLocal(LocationDX, &LocationWithHeightDX, false);
-
-						FVector LocationWithHeightDY;
-						FVector LocationDY(Location);
-						LocationDY.Y = FMath::Clamp<float>(LocationDX.Y + (HaltonY < 0.5f ? DivExtent.Y : -DivExtent.Y), Origin.Y, Origin.Y + Extent.Y);
-						GetLayerWeightAtLocationLocal(LocationDY, &LocationWithHeightDY, false);
-
-						if (LocationWithHeight != LocationWithHeightDX && LocationWithHeight != LocationWithHeightDY)
-						{
-							FVector NewZ = ((LocationWithHeight - LocationWithHeightDX) ^ (LocationWithHeight - LocationWithHeightDY)).GetSafeNormal();
-							NewZ *= FMath::Sign(NewZ.Z);
-
-							const FVector NewX = (FVector(0, -1, 0) ^ NewZ).GetSafeNormal();
-							const FVector NewY = NewZ ^ NewX;
-
-							FMatrix Align = FMatrix(NewX, NewY, NewZ, FVector::ZeroVector);
-							OutXForm = (BaseXForm * Align).ConcatTranslation(LocationWithHeight) * XForm;
-						}
-						else
-						{
-							OutXForm = BaseXForm.ConcatTranslation(LocationWithHeight) * XForm;
-						}
+						const FVector NewZ = ComputedNormal * FMath::Sign(ComputedNormal.Z);
+						const FVector NewX = (FVector(0, -1, 0) ^ NewZ).GetSafeNormal();
+						const FVector NewY = NewZ ^ NewX;
+						const FMatrix Align = FMatrix(NewX, NewY, NewZ, FVector::ZeroVector);
+						OutXForm = (BaseXForm * Align).ConcatTranslation(LocationWithHeight) * XForm;
 					}
 					else
 					{
@@ -1806,7 +1788,8 @@ struct FAsyncGrassBuilder : public FGrassBuilderBase
 						Location += FVector(FirstRandom * 2.0f - 1.0f, SecondRandom * 2.0f - 1.0f, 0.0f) * MaxJitter;
 
 						FInstanceLocal& Instance = Instances[InstanceIndex];
-						float Weight = GetLayerWeightAtLocationLocal(Location, &Instance.Pos);
+						float Weight = 0.f;
+						SampleLandscapeAtLocationLocal(Location, Instance.Pos, Weight);
 						Instance.bKeep = Weight > 0.0f && Weight >= RandomStream.GetFraction() && !IsExcluded(Instance.Pos);
 						if (Instance.bKeep)
 						{
@@ -1903,59 +1886,70 @@ struct FAsyncGrassBuilder : public FGrassBuilderBase
 		}
 		BuildTime = FPlatformTime::Seconds() - StartTime;
 	}
-	FORCEINLINE_DEBUGGABLE float GetLayerWeightAtLocationLocal(const FVector& InLocation, FVector* OutLocation, bool bWeight = true)
+
+	FORCEINLINE_DEBUGGABLE void SampleLandscapeAtLocationLocal(const FVector& InLocation, FVector& OutLocation, float& OutLayerWeight, FVector* OutNormal = nullptr)
 	{
 		// Find location
-		float TestX = InLocation.X / DrawScale.X - (float)SectionBase.X;
-		float TestY = InLocation.Y / DrawScale.Y - (float)SectionBase.Y;
+		const float TestX = InLocation.X / DrawScale.X - (float)SectionBase.X;
+		const float TestY = InLocation.Y / DrawScale.Y - (float)SectionBase.Y;
 
 		// Find data
-		int32 X1 = FMath::FloorToInt(TestX);
-		int32 Y1 = FMath::FloorToInt(TestY);
-		int32 X2 = FMath::CeilToInt(TestX);
-		int32 Y2 = FMath::CeilToInt(TestY);
+		const int32 X1 = FMath::FloorToInt(TestX);
+		const int32 Y1 = FMath::FloorToInt(TestY);
+		const int32 X2 = FMath::CeilToInt(TestX);
+		const int32 Y2 = FMath::CeilToInt(TestY);
 
 		// Clamp to prevent the sampling of the final columns from overflowing
-		int32 IdxX1 = FMath::Clamp<int32>(X1, 0, GrassData.GetStride() - 1);
-		int32 IdxY1 = FMath::Clamp<int32>(Y1, 0, GrassData.GetStride() - 1);
-		int32 IdxX2 = FMath::Clamp<int32>(X2, 0, GrassData.GetStride() - 1);
-		int32 IdxY2 = FMath::Clamp<int32>(Y2, 0, GrassData.GetStride() - 1);
-
-		float LerpX = FMath::Fractional(TestX);
-		float LerpY = FMath::Fractional(TestY);
-
-		float Result = 0.0f;
-		if (bWeight)
-		{
-			// sample
-			float Sample11 = GrassData.GetWeight(IdxX1, IdxY1);
-			float Sample21 = GrassData.GetWeight(IdxX2, IdxY1);
-			float Sample12 = GrassData.GetWeight(IdxX1, IdxY2);
-			float Sample22 = GrassData.GetWeight(IdxX2, IdxY2);
-
-			// Bilinear interpolate
-			Result = FMath::Lerp(
-				FMath::Lerp(Sample11, Sample21, LerpX),
-				FMath::Lerp(Sample12, Sample22, LerpX),
-				LerpY);
-		}
-
+		const int32 IdxX1 = FMath::Clamp<int32>(X1, 0, GrassData.GetStride() - 1);
+		const int32 IdxY1 = FMath::Clamp<int32>(Y1, 0, GrassData.GetStride() - 1);
+		const int32 IdxX2 = FMath::Clamp<int32>(X2, 0, GrassData.GetStride() - 1);
+		const int32 IdxY2 = FMath::Clamp<int32>(Y2, 0, GrassData.GetStride() - 1);
+
+		const float LerpX = FMath::Fractional(TestX);
+		const float LerpY = FMath::Fractional(TestY);
+
+		// Bilinear interpolate sampled weights
+		const float SampleWeight11 = GrassData.GetWeight(IdxX1, IdxY1);
+		const float SampleWeight21 = GrassData.GetWeight(IdxX2, IdxY1);
+		const float SampleWeight12 = GrassData.GetWeight(IdxX1, IdxY2);
+		const float SampleWeight22 = GrassData.GetWeight(IdxX2, IdxY2);
+
+		OutLayerWeight = FMath::Lerp(
+			FMath::Lerp(SampleWeight11, SampleWeight21, LerpX),
+			FMath::Lerp(SampleWeight12, SampleWeight22, LerpX),
+			LerpY);
+
+		// Bilinear interpolate sampled heights
+		const float SampleHeight11 = GrassData.GetHeight(IdxX1, IdxY1);
+		const float SampleHeight21 = GrassData.GetHeight(IdxX2, IdxY1);
+		const float SampleHeight12 = GrassData.GetHeight(IdxX1, IdxY2);
+		const float SampleHeight22 = GrassData.GetHeight(IdxX2, IdxY2);
+
+		OutLocation.X = InLocation.X - DrawScale.X * float(LandscapeSectionOffset.X);
+		OutLocation.Y = InLocation.Y - DrawScale.Y * float(LandscapeSectionOffset.Y);
+		OutLocation.Z = DrawScale.Z * FMath::Lerp(
+			FMath::Lerp(SampleHeight11, SampleHeight21, LerpX),
+			FMath::Lerp(SampleHeight12, SampleHeight22, LerpX),
+			LerpY);
+		
+		// Compute normal
+		if (OutNormal)
 		{
-			// sample
-			float Sample11 = GrassData.GetHeight(IdxX1, IdxY1);
-			float Sample21 = GrassData.GetHeight(IdxX2, IdxY1);
-			float Sample12 = GrassData.GetHeight(IdxX1, IdxY2);
-			float Sample22 = GrassData.GetHeight(IdxX2, IdxY2);
-
-			OutLocation->X = InLocation.X - DrawScale.X * float(LandscapeSectionOffset.X);
-			OutLocation->Y = InLocation.Y - DrawScale.Y * float(LandscapeSectionOffset.Y);
-			// Bilinear interpolate
-			OutLocation->Z = DrawScale.Z * FMath::Lerp(
-				FMath::Lerp(Sample11, Sample21, LerpX),
-				FMath::Lerp(Sample12, Sample22, LerpX),
-				LerpY);
+			const FVector P11((float)X1, (float)Y1, SampleHeight11);
+			const FVector P22((float)X2, (float)Y2, SampleHeight22);
+			
+			// Choose triangle and compute normal
+			if (LerpX > LerpY)
+			{
+				const FVector P21((float)X2, (float)Y1, SampleHeight21);
+				*OutNormal = (P11 != P21 && P22 != P21) ? FVector(((P22 - P21) ^ (P11 - P21)).GetSafeNormal()) : FVector::ZeroVector;
+			}
+			else 
+			{
+				const FVector P12((float)X1, (float)Y2, SampleHeight12);
+				*OutNormal = (P11 != P12 && P22 != P12) ? FVector(((P11 - P12) ^ (P22 - P12)).GetSafeNormal()) : FVector::ZeroVector;
+			}
 		}
-		return Result;
 	}
 };
 
