Avoid the overhead of calling GetCurrentFirstLODIdx_RenderThread on platforms that don't support landscape visual mesh LOD streaming

diff --git a/Engine/Source/Runtime/Core/Public/Android/AndroidPlatform.h b/Engine/Source/Runtime/Core/Public/Android/AndroidPlatform.h
--- a/Engine/Source/Runtime/Core/Public/Android/AndroidPlatform.h
+++ b/Engine/Source/Runtime/Core/Public/Android/AndroidPlatform.h
@@ -61,6 +61,7 @@ typedef FAndroidTypes FPlatformTypes;
 #define PLATFORM_SUPPORTS_GEOMETRY_SHADERS			0
 #define PLATFORM_SUPPORTS_TESSELLATION_SHADERS		0
 #define PLATFORM_SUPPORTS_VIRTUAL_TEXTURE_STREAMING	1
+#define PLATFORM_SUPPORTS_LANDSCAPE_VISUAL_MESH_LOD_STREAMING 1
 
 #define PLATFORM_CODE_SECTION(Name)					__attribute__((section(Name)))
 
diff --git a/Engine/Source/Runtime/Core/Public/HAL/Platform.h b/Engine/Source/Runtime/Core/Public/HAL/Platform.h
--- a/Engine/Source/Runtime/Core/Public/HAL/Platform.h
+++ b/Engine/Source/Runtime/Core/Public/HAL/Platform.h
@@ -572,6 +572,10 @@
 	#define	PLATFORM_HAS_FPlatformVirtualMemoryBlock 1
 #endif
 
+#ifndef PLATFORM_SUPPORTS_LANDSCAPE_VISUAL_MESH_LOD_STREAMING
+	#define PLATFORM_SUPPORTS_LANDSCAPE_VISUAL_MESH_LOD_STREAMING 0
+#endif
+
 #ifndef PLATFORM_USE_GENERIC_LAUNCH_IMPLEMENTATION
 	#define PLATFORM_USE_GENERIC_LAUNCH_IMPLEMENTATION			0
 #endif
diff --git a/Engine/Source/Runtime/Core/Public/IOS/IOSPlatform.h b/Engine/Source/Runtime/Core/Public/IOS/IOSPlatform.h
--- a/Engine/Source/Runtime/Core/Public/IOS/IOSPlatform.h
+++ b/Engine/Source/Runtime/Core/Public/IOS/IOSPlatform.h
@@ -74,6 +74,7 @@ typedef FIOSPlatformTypes FPlatformTypes;
 #define PLATFORM_SUPPORTS_GEOMETRY_SHADERS				0
 #define PLATFORM_SUPPORTS_TESSELLATION_SHADERS			0
 #define PLATFORM_SUPPORTS_VIRTUAL_TEXTURE_STREAMING		1
+#define PLATFORM_SUPPORTS_LANDSCAPE_VISUAL_MESH_LOD_STREAMING 1
 
 #define PLATFORM_GLOBAL_LOG_CATEGORY					LogIOS
 
diff --git a/Engine/Source/Runtime/Landscape/Private/LandscapeRender.cpp b/Engine/Source/Runtime/Landscape/Private/LandscapeRender.cpp
--- a/Engine/Source/Runtime/Landscape/Private/LandscapeRender.cpp
+++ b/Engine/Source/Runtime/Landscape/Private/LandscapeRender.cpp
@@ -807,6 +807,7 @@ void FLandscapeRenderSystem::ResizeAndMoveTo(FIntPoint NewMin, FIntPoint NewSize
 
 void FLandscapeRenderSystem::PrepareView(const FSceneView* View)
 {
+#if PLATFORM_SUPPORTS_LANDSCAPE_VISUAL_MESH_LOD_STREAMING
 	const int32 NumSceneProxies = SceneProxies.Num();
 	SectionCurrentFirstLODIndices.Empty(NumSceneProxies);
 	SectionCurrentFirstLODIndices.AddUninitialized(NumSceneProxies);
@@ -816,6 +817,7 @@ void FLandscapeRenderSystem::PrepareView(const FSceneView* View)
 		const FLandscapeComponentSceneProxy* Proxy = SceneProxies[Idx];
 		SectionCurrentFirstLODIndices[Idx] = Proxy ? Proxy->GetCurrentFirstLODIdx_RenderThread() : 0;
 	}
+#endif
 	
 	const bool bExecuteInParallel = FApp::ShouldUseThreadingForPerformance()
 		&& GRenderingThread; // Rendering thread is required to safely use rendering resources in parallel.
@@ -894,7 +896,11 @@ void FLandscapeRenderSystem::ComputeSectionPerViewParameters(
 		float FractionalLOD;
 		GetLODFromScreenSize(SectionLODSettings[EntityIndex], MeshScreenSizeSquared, LODScale * LODScale, FractionalLOD);
 
+#if PLATFORM_SUPPORTS_LANDSCAPE_VISUAL_MESH_LOD_STREAMING
 		const float CurFirstLODIdx = (float)SectionCurrentFirstLODIndices[EntityIndex];
+#else
+		constexpr float CurFirstLODIdx = 0.f;
+#endif
 		int32 ForcedLODLevel = FMath::Min<int32>(ViewLODOverride, SectionLODSettings[EntityIndex].LastLODIndex);
 		NewSectionLODValues[EntityIndex] = FMath::Max(ForcedLODLevel >= 0 ? ForcedLODLevel : FractionalLOD, CurFirstLODIdx);
 
@@ -3267,8 +3273,9 @@ void FLandscapeComponentSceneProxy::GetDynamicMeshElements(const TArray<const FS
 			const float MeshScreenSizeSquared = ComputeBoundsScreenRadiusSquared(GetBounds().Origin, GetBounds().SphereRadius, *View);
 			int32 LODToRender = ForcedLODLevel >= 0 ? ForcedLODLevel : GetLODFromScreenSize(MeshScreenSizeSquared, LODScale * LODScale);
 
+#if PLATFORM_SUPPORTS_LANDSCAPE_VISUAL_MESH_LOD_STREAMING
 			LODToRender = FMath::Max<int32>(LODToRender, GetCurrentFirstLODIdx_RenderThread());
-
+#endif
 			FMeshBatch& Mesh = Collector.AllocateMesh();
 			GetStaticMeshElement(LODToRender, false, ForcedLODLevel >= 0, Mesh, ParameterArray.ElementParams);
 
diff --git a/Engine/Source/Runtime/Landscape/Private/LandscapeRenderMobile.cpp b/Engine/Source/Runtime/Landscape/Private/LandscapeRenderMobile.cpp
--- a/Engine/Source/Runtime/Landscape/Private/LandscapeRenderMobile.cpp
+++ b/Engine/Source/Runtime/Landscape/Private/LandscapeRenderMobile.cpp
@@ -550,8 +550,10 @@ void FLandscapeComponentSceneProxyMobile::ApplyMeshElementModifier(FMeshBatchEle
 	}
 }
 
+#if PLATFORM_SUPPORTS_LANDSCAPE_VISUAL_MESH_LOD_STREAMING
 uint8 FLandscapeComponentSceneProxyMobile::GetCurrentFirstLODIdx_RenderThread() const
 {
 	check(MobileRenderData.IsValid());
 	return MobileRenderData->CurrentFirstLODIdx;
 }
+#endif
diff --git a/Engine/Source/Runtime/Landscape/Public/LandscapeRenderMobile.h b/Engine/Source/Runtime/Landscape/Public/LandscapeRenderMobile.h
--- a/Engine/Source/Runtime/Landscape/Public/LandscapeRenderMobile.h
+++ b/Engine/Source/Runtime/Landscape/Public/LandscapeRenderMobile.h
@@ -145,7 +145,9 @@ public:
 
 	virtual void ApplyMeshElementModifier(FMeshBatchElement& InOutMeshElement, int32 InLodIndex) const override;
 
+#if PLATFORM_SUPPORTS_LANDSCAPE_VISUAL_MESH_LOD_STREAMING
 	virtual uint8 GetCurrentFirstLODIdx_RenderThread() const override;
+#endif
 };
 
 
