Build fix for PhysX builds
#rnx

diff --git a/Engine/Source/Runtime/Landscape/Private/LandscapeCollision.cpp b/Engine/Source/Runtime/Landscape/Private/LandscapeCollision.cpp
--- a/Engine/Source/Runtime/Landscape/Private/LandscapeCollision.cpp
+++ b/Engine/Source/Runtime/Landscape/Private/LandscapeCollision.cpp
@@ -770,14 +770,12 @@ TArray<PxHeightFieldSample> ConvertHeightfieldDataForPhysx(
 				if (!bUseDefMaterial && DominantLayers)
 				{
 					uint8 DominantLayerIdx = DominantLayers ? DominantLayers[SrcSampleIndex] : -1;
-					ULandscapeLayerInfoObject* Layer = ComponentLayerInfos.IsValidIndex(DominantLayerIdx) ? ComponentLayerInfos[DominantLayerIdx] : nullptr;
-
-						ULandscapeLayerInfoObject* Layer = Component->ComponentLayerInfos[DominantLayerIdx];
-						if (Layer == ALandscapeProxy::VisibilityLayer)
-						{
-							// If it's a hole, override with the hole flag.
-							MaterialIndex = PxHeightFieldMaterial::eHOLE;
-						}
+					ULandscapeLayerInfoObject* Layer = Component->ComponentLayerInfos.IsValidIndex(DominantLayerIdx) ? Component->ComponentLayerInfos[DominantLayerIdx] : nullptr;
+					if (Layer == ALandscapeProxy::VisibilityLayer)
+					{
+						// If it's a hole, override with the hole flag.
+						MaterialIndex = PxHeightFieldMaterial::eHOLE;
+					}
 					else if (RenderPhysicalMaterialIds)
 					{
 						uint8 RenderIdx = RenderPhysicalMaterialIds[SrcSampleIndex];
@@ -786,10 +784,10 @@ TArray<PxHeightFieldSample> ConvertHeightfieldDataForPhysx(
 							MaterialIndex = InOutMaterials.AddUnique(PhysicalMaterialRenderObjects[RenderIdx - 1]);
 						}
 					}
-						else if (Layer && Layer->PhysMaterial)
-						{
-							MaterialIndex = InOutMaterials.AddUnique(Layer->PhysMaterial);
-						}
+					else if (Layer && Layer->PhysMaterial)
+					{
+						MaterialIndex = InOutMaterials.AddUnique(Layer->PhysMaterial);
+					}
 				}
 
 				// Default Material but Def Material wasn't added yet...
@@ -962,23 +960,23 @@ bool ULandscapeHeightfieldCollisionComponent::CookCollisionData(const FName& For
 					uint8 DominantLayerIdx = DominantLayers ? DominantLayers[SrcSampleIndex] : -1;
 					ULandscapeLayerInfoObject* Layer = ComponentLayerInfos.IsValidIndex(DominantLayerIdx) ? ComponentLayerInfos[DominantLayerIdx] : nullptr;
 
-						if(Layer == ALandscapeProxy::VisibilityLayer)
-						{
-							// If it's a hole, use the final index
-							MaterialIndex = TNumericLimits<uint8>::Max();
-						}
+					if(Layer == ALandscapeProxy::VisibilityLayer)
+					{
+						// If it's a hole, use the final index
+						MaterialIndex = TNumericLimits<uint8>::Max();
+					}
 					else if (RenderPhysicalMaterialIds)
 					{
 						uint8 RenderIdx = RenderPhysicalMaterialIds[SrcSampleIndex];
 						UPhysicalMaterial* DominantMaterial = RenderIdx > 0 ? PhysicalMaterialRenderObjects[RenderIdx - 1] : DefMaterial;
 						MaterialIndex = InOutMaterials.AddUnique(DominantMaterial);
 					}
-						else
-						{
-							UPhysicalMaterial* DominantMaterial = Layer && Layer->PhysMaterial ? Layer->PhysMaterial : DefMaterial;
-							MaterialIndex = InOutMaterials.AddUnique(DominantMaterial);
-						}
+					else
+					{
+						UPhysicalMaterial* DominantMaterial = Layer && Layer->PhysMaterial ? Layer->PhysMaterial : DefMaterial;
+						MaterialIndex = InOutMaterials.AddUnique(DominantMaterial);
 					}
+				}
 				MaterialIndices.Add(MaterialIndex);
 			}
 		}
@@ -1009,7 +1007,7 @@ bool ULandscapeHeightfieldCollisionComponent::CookCollisionData(const FName& For
 
 	if (CollisionHeightData.IsLocked())
 	{
-	CollisionHeightData.Unlock();
+		CollisionHeightData.Unlock();
 	}
 	if (DominantLayerData.IsLocked())
 	{
