Track ray tracing dynamic geometry video memory usage under 'stat d3d12raytracing'

diff --git a/Engine/Plugins/FX/Niagara/Source/Niagara/Private/NiagaraRendererRibbons.cpp b/Engine/Plugins/FX/Niagara/Source/Niagara/Private/NiagaraRendererRibbons.cpp
--- a/Engine/Plugins/FX/Niagara/Source/Niagara/Private/NiagaraRendererRibbons.cpp
+++ b/Engine/Plugins/FX/Niagara/Source/Niagara/Private/NiagaraRendererRibbons.cpp
@@ -230,6 +230,11 @@ FNiagaraRendererRibbons::~FNiagaraRendererRibbons()
 {
 }
 
+
+DECLARE_STATS_GROUP(TEXT("D3D12RHI: Ray Tracing"), STATGROUP_D3D12RayTracing, STATCAT_Advanced);
+DECLARE_MEMORY_STAT(TEXT("Total Used Video Memory"), STAT_D3D12RayTracingUsedVideoMemory, STATGROUP_D3D12RayTracing);
+DECLARE_MEMORY_STAT(TEXT("Dynamic Vertex Buffer Memory"), STAT_D3D12RayTracingDynamicVertexBufferMemory, STATGROUP_D3D12RayTracing);
+
 void FNiagaraRendererRibbons::ReleaseRenderThreadResources()
 {
 	FNiagaraRenderer::ReleaseRenderThreadResources();
@@ -238,6 +243,9 @@ void FNiagaraRendererRibbons::ReleaseRenderThreadResources()
 	{
 		RayTracingGeometry.ReleaseResource();
 		RayTracingDynamicVertexBuffer.Release();
+
+		DEC_MEMORY_STAT_BY(STAT_D3D12RayTracingUsedVideoMemory, RayTracingDynamicVertexBuffer.NumBytes);
+		DEC_MEMORY_STAT_BY(STAT_D3D12RayTracingDynamicVertexBufferMemory, RayTracingDynamicVertexBuffer.NumBytes);
 	}
 #endif
 }
@@ -251,6 +259,9 @@ void FNiagaraRendererRibbons::CreateRenderThreadResources(NiagaraEmitterInstance
 	{
 		RayTracingDynamicVertexBuffer.Initialize(4, 256, PF_R32_FLOAT, BUF_UnorderedAccess | BUF_ShaderResource, TEXT("FNiagaraRendererRibbons::RayTracingDynamicVertexBuffer"));
 
+		INC_MEMORY_STAT_BY(STAT_D3D12RayTracingUsedVideoMemory, RayTracingDynamicVertexBuffer.NumBytes);
+		INC_MEMORY_STAT_BY(STAT_D3D12RayTracingDynamicVertexBufferMemory, RayTracingDynamicVertexBuffer.NumBytes);
+
 		FRayTracingGeometryInitializer Initializer;
 		Initializer.IndexBuffer = nullptr;
 		Initializer.TotalPrimitiveCount = 0;
diff --git a/Engine/Plugins/FX/Niagara/Source/Niagara/Private/NiagaraRendererSprites.cpp b/Engine/Plugins/FX/Niagara/Source/Niagara/Private/NiagaraRendererSprites.cpp
--- a/Engine/Plugins/FX/Niagara/Source/Niagara/Private/NiagaraRendererSprites.cpp
+++ b/Engine/Plugins/FX/Niagara/Source/Niagara/Private/NiagaraRendererSprites.cpp
@@ -154,6 +154,10 @@ FNiagaraRendererSprites::~FNiagaraRendererSprites()
 {
 }
 
+DECLARE_STATS_GROUP(TEXT("D3D12RHI: Ray Tracing"), STATGROUP_D3D12RayTracing, STATCAT_Advanced);
+DECLARE_MEMORY_STAT(TEXT("Total Used Video Memory"), STAT_D3D12RayTracingUsedVideoMemory, STATGROUP_D3D12RayTracing);
+DECLARE_MEMORY_STAT(TEXT("Dynamic Vertex Buffer Memory"), STAT_D3D12RayTracingDynamicVertexBufferMemory, STATGROUP_D3D12RayTracing);
+
 void FNiagaraRendererSprites::ReleaseRenderThreadResources()
 {
 	FNiagaraRenderer::ReleaseRenderThreadResources();
@@ -164,6 +168,9 @@ void FNiagaraRendererSprites::ReleaseRenderThreadResources()
 	{
 		RayTracingGeometry.ReleaseResource();
 		RayTracingDynamicVertexBuffer.Release();
+
+		DEC_MEMORY_STAT_BY(STAT_D3D12RayTracingUsedVideoMemory, RayTracingDynamicVertexBuffer.NumBytes);
+		DEC_MEMORY_STAT_BY(STAT_D3D12RayTracingDynamicVertexBufferMemory, RayTracingDynamicVertexBuffer.NumBytes);
 	}
 #endif
 }
@@ -178,6 +185,9 @@ void FNiagaraRendererSprites::CreateRenderThreadResources(NiagaraEmitterInstance
 	{
 		RayTracingDynamicVertexBuffer.Initialize(4, 256, PF_R32_FLOAT, BUF_UnorderedAccess | BUF_ShaderResource, TEXT("FNiagaraRendererSprites::RayTracingDynamicVertexBuffer"));
 
+		INC_MEMORY_STAT_BY(STAT_D3D12RayTracingUsedVideoMemory, RayTracingDynamicVertexBuffer.NumBytes);
+		INC_MEMORY_STAT_BY(STAT_D3D12RayTracingDynamicVertexBufferMemory, RayTracingDynamicVertexBuffer.NumBytes);
+
 		FRayTracingGeometryInitializer Initializer;
 		Initializer.IndexBuffer = nullptr;
 		Initializer.GeometryType = RTGT_Triangles;
diff --git a/Engine/Source/Runtime/Engine/Private/SkeletalRenderGPUSkin.cpp b/Engine/Source/Runtime/Engine/Private/SkeletalRenderGPUSkin.cpp
--- a/Engine/Source/Runtime/Engine/Private/SkeletalRenderGPUSkin.cpp
+++ b/Engine/Source/Runtime/Engine/Private/SkeletalRenderGPUSkin.cpp
@@ -190,6 +190,10 @@ void FSkeletalMeshObjectGPUSkin::InitResources(USkinnedMeshComponent* InMeshComp
 #endif
 }
 
+DECLARE_STATS_GROUP(TEXT("D3D12RHI: Ray Tracing"), STATGROUP_D3D12RayTracing, STATCAT_Advanced);
+DECLARE_MEMORY_STAT(TEXT("Total Used Video Memory"), STAT_D3D12RayTracingUsedVideoMemory, STATGROUP_D3D12RayTracing);
+DECLARE_MEMORY_STAT(TEXT("Dynamic Vertex Buffer Memory"), STAT_D3D12RayTracingDynamicVertexBufferMemory, STATGROUP_D3D12RayTracing);
+
 void FSkeletalMeshObjectGPUSkin::ReleaseResources()
 {
 	for( int32 LODIndex=0;LODIndex < LODs.Num();LODIndex++ )
@@ -224,6 +228,8 @@ void FSkeletalMeshObjectGPUSkin::ReleaseResources()
 		[RayTracingDynamicVertexBuffer = MoveTemp(RayTracingDynamicVertexBuffer)](FRHICommandListImmediate& RHICmdList) mutable
 	{
 		RayTracingDynamicVertexBuffer.Release();
+		DEC_MEMORY_STAT_BY(STAT_D3D12RayTracingUsedVideoMemory, RayTracingDynamicVertexBuffer.NumBytes);
+		DEC_MEMORY_STAT_BY(STAT_D3D12RayTracingDynamicVertexBufferMemory, RayTracingDynamicVertexBuffer.NumBytes);
 	});
 
 	check(!RayTracingDynamicVertexBuffer.Buffer.IsValid());
diff --git a/Engine/Source/Runtime/Landscape/Private/LandscapeRender.cpp b/Engine/Source/Runtime/Landscape/Private/LandscapeRender.cpp
--- a/Engine/Source/Runtime/Landscape/Private/LandscapeRender.cpp
+++ b/Engine/Source/Runtime/Landscape/Private/LandscapeRender.cpp
@@ -42,6 +42,10 @@ LandscapeRender.cpp: New terrain rendering
 #include "RayTracingInstance.h"
 #include "ProfilingDebugging/LoadTimeTracker.h"
 
+DECLARE_STATS_GROUP(TEXT("D3D12RHI: Ray Tracing"), STATGROUP_D3D12RayTracing, STATCAT_Advanced);
+DECLARE_MEMORY_STAT(TEXT("Total Used Video Memory"), STAT_D3D12RayTracingUsedVideoMemory, STATGROUP_D3D12RayTracing);
+DECLARE_MEMORY_STAT(TEXT("Dynamic Vertex Buffer Memory"), STAT_D3D12RayTracingDynamicVertexBufferMemory, STATGROUP_D3D12RayTracing);
+
 IMPLEMENT_GLOBAL_SHADER_PARAMETER_STRUCT(FLandscapeUniformShaderParameters, "LandscapeParameters");
 IMPLEMENT_GLOBAL_SHADER_PARAMETER_STRUCT(FLandscapeFixedGridUniformShaderParameters, "LandscapeFixedGrid");
 IMPLEMENT_TYPE_LAYOUT(FLandscapeVertexFactoryPixelShaderParameters);
@@ -1660,6 +1664,9 @@ FLandscapeComponentSceneProxy::~FLandscapeComponentSceneProxy()
 			const int8 SubSectionIdx = SubX + SubY * NumSubsections;
 			SectionRayTracingStates[SubSectionIdx].Geometry.ReleaseResource();
 			SectionRayTracingStates[SubSectionIdx].RayTracingDynamicVertexBuffer.Release();
+
+			DEC_MEMORY_STAT_BY(STAT_D3D12RayTracingUsedVideoMemory, SectionRayTracingStates[SubSectionIdx].RayTracingDynamicVertexBuffer.NumBytes);
+			DEC_MEMORY_STAT_BY(STAT_D3D12RayTracingDynamicVertexBufferMemory, SectionRayTracingStates[SubSectionIdx].RayTracingDynamicVertexBuffer.NumBytes);
 		}
 	}
 #endif
@@ -3702,6 +3709,8 @@ void FLandscapeComponentSceneProxy::GetDynamicRayTracingInstances(FRayTracingMat
 					bNeedsRayTracingGeometryUpdate = true;
 					SectionRayTracingStates[SubSectionIdx].CurrentLOD = CurrentLOD;
 					SectionRayTracingStates[SubSectionIdx].RayTracingDynamicVertexBuffer.Release();
+					DEC_MEMORY_STAT_BY(STAT_D3D12RayTracingUsedVideoMemory, SectionRayTracingStates[SubSectionIdx].RayTracingDynamicVertexBuffer.NumBytes);
+					DEC_MEMORY_STAT_BY(STAT_D3D12RayTracingDynamicVertexBufferMemory, SectionRayTracingStates[SubSectionIdx].RayTracingDynamicVertexBuffer.NumBytes);
 				}
 				if (SectionRayTracingStates[SubSectionIdx].HeightmapLODBias != RenderSystem.GetSectionLODBias(ComponentBase))
 				{
diff --git a/Engine/Source/Runtime/Renderer/Private/RayTracing/RayTracingDynamicGeometry.cpp b/Engine/Source/Runtime/Renderer/Private/RayTracing/RayTracingDynamicGeometry.cpp
--- a/Engine/Source/Runtime/Renderer/Private/RayTracing/RayTracingDynamicGeometry.cpp
+++ b/Engine/Source/Runtime/Renderer/Private/RayTracing/RayTracingDynamicGeometry.cpp
@@ -6,6 +6,9 @@
 
 #if RHI_RAYTRACING
 
+DECLARE_STATS_GROUP(TEXT("D3D12RHI: Ray Tracing"), STATGROUP_D3D12RayTracing, STATCAT_Advanced);
+DECLARE_MEMORY_STAT(TEXT("Total Used Video Memory"), STAT_D3D12RayTracingUsedVideoMemory, STATGROUP_D3D12RayTracing);
+DECLARE_MEMORY_STAT(TEXT("Dynamic Vertex Buffer Memory"), STAT_D3D12RayTracingDynamicVertexBufferMemory, STATGROUP_D3D12RayTracing);
 
 static bool IsSupportedDynamicVertexFactoryType(const FVertexFactoryType* VertexFactoryType)
 {
@@ -157,7 +160,12 @@ void FRayTracingDynamicGeometryCollection::AddDynamicMeshBatchForGeometryUpdate(
 	uint32 DesiredVertexBufferSize = UpdateParams.VertexBufferSize;
 	if (Buffer.NumBytes != DesiredVertexBufferSize)
 	{
+		int32 OriginalSize = Buffer.NumBytes;
+		DEC_MEMORY_STAT_BY(STAT_D3D12RayTracingUsedVideoMemory, Buffer.NumBytes);
+		DEC_MEMORY_STAT_BY(STAT_D3D12RayTracingDynamicVertexBufferMemory, Buffer.NumBytes);
 		Buffer.Initialize(sizeof(float), DesiredVertexBufferSize / sizeof(float), PF_R32_FLOAT, BUF_UnorderedAccess | BUF_ShaderResource, TEXT("FRayTracingDynamicGeometryCollection::RayTracingDynamicVertexBuffer"));
+		INC_MEMORY_STAT_BY(STAT_D3D12RayTracingUsedVideoMemory, Buffer.NumBytes);
+		INC_MEMORY_STAT_BY(STAT_D3D12RayTracingDynamicVertexBufferMemory, Buffer.NumBytes);
 		bRefit = false;
 	}
 
