Submitting Eric Soulvie's Native Brush (not cleaned up yet, not used, straight BP port)

diff --git a/Engine/Plugins/Experimental/Landmass/Source/Runtime/Private/BrushEffectsList.cpp b/Engine/Plugins/Experimental/Landmass/Source/Runtime/Private/BrushEffectsList.cpp
--- a/Engine/Plugins/Experimental/Landmass/Source/Runtime/Private/BrushEffectsList.cpp
+++ b/Engine/Plugins/Experimental/Landmass/Source/Runtime/Private/BrushEffectsList.cpp
@@ -2,7 +2,3 @@
 
 #include "BrushEffectsList.h"
 
-FLandmassBrushEffectsList::FLandmassBrushEffectsList()
-{
-
-}
\ No newline at end of file
diff --git a/Engine/Plugins/Experimental/Landmass/Source/Runtime/Public/BrushEffectsList.h b/Engine/Plugins/Experimental/Landmass/Source/Runtime/Public/BrushEffectsList.h
--- a/Engine/Plugins/Experimental/Landmass/Source/Runtime/Public/BrushEffectsList.h
+++ b/Engine/Plugins/Experimental/Landmass/Source/Runtime/Public/BrushEffectsList.h
@@ -163,8 +163,6 @@ struct FLandmassBrushEffectsList
 {
 	GENERATED_BODY()
 
-	FLandmassBrushEffectsList();
-
 	UPROPERTY(EditAnywhere, BlueprintReadWrite, Category = BrushEffects)
 	FBrushEffectBlurring Blurring;
 
diff --git a/Engine/Source/Runtime/Landscape/Private/LandscapeBlueprintBrushBase.cpp b/Engine/Source/Runtime/Landscape/Private/LandscapeBlueprintBrushBase.cpp
--- a/Engine/Source/Runtime/Landscape/Private/LandscapeBlueprintBrushBase.cpp
+++ b/Engine/Source/Runtime/Landscape/Private/LandscapeBlueprintBrushBase.cpp
@@ -37,6 +37,16 @@ ALandscapeBlueprintBrushBase::ALandscapeBlueprintBrushBase(const FObjectInitiali
 #endif // WITH_EDITOR
 }
 
+UTextureRenderTarget2D* ALandscapeBlueprintBrushBase::Render_Implementation(bool InIsHeightmap, UTextureRenderTarget2D* InCombinedResult, const FName& InWeightmapLayerName)
+{
+	return Render_Native(InIsHeightmap, InCombinedResult, InWeightmapLayerName);
+}
+
+void ALandscapeBlueprintBrushBase::Initialize_Implementation(const FTransform& InLandscapeTransform, const FIntPoint& InLandscapeSize, const FIntPoint& InLandscapeRenderTargetSize)
+{
+	Initialize_Native(InLandscapeTransform, InLandscapeSize, InLandscapeRenderTargetSize);
+}
+
 void ALandscapeBlueprintBrushBase::RequestLandscapeUpdate()
 {
 #if WITH_EDITORONLY_DATA
diff --git a/Engine/Source/Runtime/Landscape/Public/LandscapeBlueprintBrushBase.h b/Engine/Source/Runtime/Landscape/Public/LandscapeBlueprintBrushBase.h
--- a/Engine/Source/Runtime/Landscape/Public/LandscapeBlueprintBrushBase.h
+++ b/Engine/Source/Runtime/Landscape/Public/LandscapeBlueprintBrushBase.h
@@ -8,6 +8,8 @@
 
 #include "LandscapeBlueprintBrushBase.generated.h"
 
+class UTextureRenderTarget2D;
+
 UCLASS(Abstract, NotBlueprintable)
 class LANDSCAPE_API ALandscapeBlueprintBrushBase : public AActor
 {
@@ -34,10 +36,13 @@ protected:
 #endif
 
 public:
-	UFUNCTION(BlueprintImplementableEvent)
+	virtual UTextureRenderTarget2D* Render_Native(bool InIsHeightmap, UTextureRenderTarget2D* InCombinedResult, const FName& InWeightmapLayerName) {return nullptr;}
+	virtual void Initialize_Native(const FTransform& InLandscapeTransform, const FIntPoint& InLandscapeSize, const FIntPoint& InLandscapeRenderTargetSize) {}
+
+	UFUNCTION(BlueprintNativeEvent)
 	UTextureRenderTarget2D* Render(bool InIsHeightmap, UTextureRenderTarget2D* InCombinedResult, const FName& InWeightmapLayerName);
 
-	UFUNCTION(BlueprintImplementableEvent)
+	UFUNCTION(BlueprintNativeEvent)
 	void Initialize(const FTransform& InLandscapeTransform, const FIntPoint& InLandscapeSize, const FIntPoint& InLandscapeRenderTargetSize);
 
 	UFUNCTION(BlueprintCallable, Category = "Landscape")
