Fix rendering of landscape collision.
Needs fixed lod behaviour that is defined by both view and component, so extra information needs to be stored for the FComputeSectionPerViewParametersTask

diff --git a/Engine/Source/Runtime/Landscape/Private/LandscapeRender.cpp b/Engine/Source/Runtime/Landscape/Private/LandscapeRender.cpp
--- a/Engine/Source/Runtime/Landscape/Private/LandscapeRender.cpp
+++ b/Engine/Source/Runtime/Landscape/Private/LandscapeRender.cpp
@@ -641,7 +641,7 @@ void ULandscapeComponent::GetUsedMaterials(TArray<UMaterialInterface*>& OutMater
  * A return value less than 0 means no override.
  * Any positive value must still be clamped into the valid Lod range for the landscape.
  */
-static int32 GetLandscapeLodOverride(FSceneView const& View)
+static int32 GetViewLodOverride(FSceneView const& View)
 {
 	// Apply r.ForceLOD override
 	int32 LodOverride = GetCVarForceLOD();
@@ -654,6 +654,28 @@ static int32 GetLandscapeLodOverride(FSceneView const& View)
 	return LodOverride;
 }
 
+static int32 GetDrawCollisionLodOverride(bool bShowCollisionPawn, bool bShowCollisionVisibility, int32 DrawCollisionPawnLOD, int32 DrawCollisionVisibilityLOD)
+{
+#if WITH_EDITOR || !(UE_BUILD_SHIPPING || UE_BUILD_TEST)
+	return bShowCollisionPawn ? FMath::Max(DrawCollisionPawnLOD, DrawCollisionVisibilityLOD) : bShowCollisionVisibility ? DrawCollisionVisibilityLOD : -1;
+#else
+	return -1;
+#endif
+}
+
+static int32 GetDrawCollisionLodOverride(FSceneView const& View, FCollisionResponseContainer const& CollisionResponse, int32 CollisionLod, int32 SimpleCollisionLod)
+{
+#if WITH_EDITOR || !(UE_BUILD_SHIPPING || UE_BUILD_TEST)
+	bool bShowCollisionPawn = View.Family->EngineShowFlags.CollisionPawn;
+	bool bShowCollisionVisibility = View.Family->EngineShowFlags.CollisionVisibility;
+	int32 DrawCollisionPawnLOD = CollisionResponse.GetResponse(ECC_Pawn) == ECR_Ignore ? -1 : SimpleCollisionLod;
+	int32 DrawCollisionVisibilityLOD = CollisionResponse.GetResponse(ECC_Visibility) == ECR_Ignore ? -1 : CollisionLod;
+	return GetDrawCollisionLodOverride(bShowCollisionPawn, bShowCollisionVisibility, DrawCollisionPawnLOD, DrawCollisionVisibilityLOD);
+#else
+	return -1;
+#endif
+}
+
 //
 // FLandscapeComponentSceneProxy
 //
@@ -722,6 +744,10 @@ void FLandscapeRenderSystem::UnregisterEntity(FLandscapeComponentSceneProxy* Sce
 	SetSceneProxy(SceneProxy->ComponentBase, nullptr);
 	SetSectionOriginAndRadius(SceneProxy->ComponentBase, FVector4(ForceInitToZero));
 
+	LODSettingsComponent LODSettings;
+	FMemory::Memzero(&LODSettings, sizeof(LODSettingsComponent));
+	SetSectionLODSettings(SceneProxy->ComponentBase, LODSettings);
+
 	if (SceneProxy->MaterialHasTessellationEnabled.Find(true) != INDEX_NONE)
 	{
 		NumEntitiesWithTessellation--;
@@ -869,6 +895,8 @@ void FLandscapeRenderSystem::ComputeSectionPerViewParameters(
 	const FSceneView* ViewPtrAsIdentifier,
 	int32 ViewLODOverride,
 	float ViewLODDistanceFactor,
+	bool bDrawCollisionPawn,
+	bool bDrawCollisionCollision,
 	FVector ViewOrigin,
 	FMatrix ViewProjectionMarix
 )
@@ -896,12 +924,16 @@ void FLandscapeRenderSystem::ComputeSectionPerViewParameters(
 		float FractionalLOD;
 		GetLODFromScreenSize(SectionLODSettings[EntityIndex], MeshScreenSizeSquared, LODScale * LODScale, FractionalLOD);
 
+		const int32 DrawCollisionLODOverride = GetDrawCollisionLodOverride(bDrawCollisionPawn, bDrawCollisionCollision, SectionLODSettings[EntityIndex].DrawCollisionPawnLOD, SectionLODSettings[EntityIndex].DrawCollisionVisibilityLOD);
+		int32 ForcedLODLevel = DrawCollisionLODOverride >= 0 ? DrawCollisionLODOverride : ViewLODOverride;
+		ForcedLODLevel = FMath::Min<int32>(ForcedLODLevel, SectionLODSettings[EntityIndex].LastLODIndex);
+
 #if PLATFORM_SUPPORTS_LANDSCAPE_VISUAL_MESH_LOD_STREAMING
 		const float CurFirstLODIdx = (float)SectionCurrentFirstLODIndices[EntityIndex];
 #else
 		constexpr float CurFirstLODIdx = 0.f;
 #endif
-		int32 ForcedLODLevel = FMath::Min<int32>(ViewLODOverride, SectionLODSettings[EntityIndex].LastLODIndex);
+
 		NewSectionLODValues[EntityIndex] = FMath::Max(ForcedLODLevel >= 0 ? ForcedLODLevel : FractionalLOD, CurFirstLODIdx);
 
 		if (TessellationFalloffSettings.UseTessellationComponentScreenSizeFalloff && NumEntitiesWithTessellation > 0)
@@ -1091,8 +1123,10 @@ void FLandscapeRenderSystem::EndFrame()
 FLandscapeRenderSystem::FComputeSectionPerViewParametersTask::FComputeSectionPerViewParametersTask(FLandscapeRenderSystem& InRenderSystem, const FSceneView* InView)
 	: RenderSystem(InRenderSystem)
 	, ViewPtrAsIdentifier(InView)
-	, ViewLODOverride(GetLandscapeLodOverride(*InView))
+	, ViewLODOverride(GetViewLodOverride(*InView))
 	, ViewLODDistanceFactor(InView->LODDistanceFactor)
+	, ViewEngineShowFlagCollisionPawn(InView->Family->EngineShowFlags.CollisionPawn)
+	, ViewEngineShowFlagCollisionVisibility(InView->Family->EngineShowFlags.CollisionVisibility)
 	, ViewOrigin(InView->ViewMatrices.GetViewOrigin())
 	, ViewProjectionMatrix(InView->ViewMatrices.GetProjectionMatrix())
 {
@@ -1299,6 +1333,9 @@ FLandscapeComponentSceneProxy::FLandscapeComponentSceneProxy(ULandscapeComponent
 	FirstVirtualTextureLOD = FMath::Max(MaxLOD - InComponent->GetLandscapeProxy()->VirtualTextureNumLods, 0);
 	VirtualTextureLodBias = InComponent->GetLandscapeProxy()->VirtualTextureLodBias;
 
+	LODSettings.DrawCollisionPawnLOD = CollisionResponse.GetResponse(ECC_Pawn) == ECR_Ignore ? -1 : SimpleCollisionMipLevel;
+	LODSettings.DrawCollisionVisibilityLOD = CollisionResponse.GetResponse(ECC_Visibility) == ECR_Ignore ? -1 : CollisionMipLevel;
+
 	ComponentMaxExtend = SubsectionSizeQuads * FMath::Max(InComponent->GetComponentTransform().GetScale3D().X, InComponent->GetComponentTransform().GetScale3D().Y);
 
 	if (NumSubsections > 1)
@@ -2447,30 +2484,13 @@ void FLandscapeComponentSceneProxy::DrawStaticElements(FStaticPrimitiveDrawInter
 
 void FLandscapeComponentSceneProxy::CalculateLODFromScreenSize(const FSceneView& InView, float InMeshScreenSizeSquared, float InViewLODScale, int32 InSubSectionIndex, FViewCustomDataLOD& InOutLODData) const
 {
-	// Handle general LOD override
-	float PreferedLOD = GetLandscapeLodOverride(InView);
-
-#if WITH_EDITOR || !(UE_BUILD_SHIPPING || UE_BUILD_TEST)
-	if (InView.Family->EngineShowFlags.CollisionVisibility || InView.Family->EngineShowFlags.CollisionPawn)
-	{
-		const bool bDrawSimpleCollision = InView.Family->EngineShowFlags.CollisionPawn        && CollisionResponse.GetResponse(ECC_Pawn) != ECR_Ignore;
-		const bool bDrawComplexCollision = InView.Family->EngineShowFlags.CollisionVisibility && CollisionResponse.GetResponse(ECC_Visibility) != ECR_Ignore;
+	// Handle LOD overrides
+	const int32 ViewLodOverride = GetViewLodOverride(InView);
+	const int32 DrawCollisionLodOverride = GetDrawCollisionLodOverride(InView, CollisionResponse, CollisionMipLevel, SimpleCollisionMipLevel);
 
-		if (bDrawSimpleCollision)
-		{
-			PreferedLOD = FMath::Max(CollisionMipLevel, SimpleCollisionMipLevel);
-		}
-		else if (bDrawComplexCollision)
-		{
-			PreferedLOD = CollisionMipLevel;
-		}
-	}
-#endif
-
-	if (ForcedLOD >= 0)
-	{
-		PreferedLOD = (float)ForcedLOD;
-	}
+	float PreferedLOD = (float)ViewLodOverride;
+	PreferedLOD = DrawCollisionLodOverride >= 0 ? (float)DrawCollisionLodOverride : PreferedLOD;
+	PreferedLOD = ForcedLOD >= 0 ? (float)ForcedLOD : PreferedLOD;
 
 	int8 MinStreamedLOD = HeightmapTexture ? FMath::Min<int8>(((FTexture2DResource*)HeightmapTexture->Resource)->GetCurrentFirstMip(), FMath::CeilLogTwo(SubsectionSizeVerts) - 1) : 0;
 	MinStreamedLOD = FMath::Min(MinStreamedLOD, MaxLOD); // We can't go above MaxLOD even for texture streaming
@@ -3240,13 +3260,6 @@ void FLandscapeComponentSceneProxy::GetDynamicMeshElements(const TArray<const FS
 	QUICK_SCOPE_CYCLE_COUNTER(STAT_FLandscapeComponentSceneProxy_GetMeshElements);
 	SCOPE_CYCLE_COUNTER(STAT_LandscapeDynamicDrawTime);
 
-#if WITH_EDITOR || !(UE_BUILD_SHIPPING || UE_BUILD_TEST)
-	const bool bInCollisionView = ViewFamily.EngineShowFlags.CollisionVisibility || ViewFamily.EngineShowFlags.CollisionPawn;
-	const bool bDrawSimpleCollision = ViewFamily.EngineShowFlags.CollisionPawn       && CollisionResponse.GetResponse(ECC_Pawn) != ECR_Ignore;
-	const bool bDrawComplexCollision = ViewFamily.EngineShowFlags.CollisionVisibility && CollisionResponse.GetResponse(ECC_Visibility) != ECR_Ignore;
-	const int32 CollisionLODLevel = bDrawSimpleCollision ? FMath::Max(CollisionMipLevel, SimpleCollisionMipLevel) : bDrawComplexCollision ? CollisionMipLevel : -1;
-#endif
-
 	int32 NumPasses = 0;
 	int32 NumTriangles = 0;
 	int32 NumDrawCalls = 0;
@@ -3261,12 +3274,11 @@ void FLandscapeComponentSceneProxy::GetDynamicMeshElements(const TArray<const FS
 
 			const FSceneView* View = Views[ViewIndex];
 
-			int32 ForcedLODLevel = GetLandscapeLodOverride(*View);
-
-#if WITH_EDITOR || !(UE_BUILD_SHIPPING || UE_BUILD_TEST)
-			ForcedLODLevel = CollisionLODLevel >= 0 ? CollisionLODLevel : ForcedLODLevel;
-#endif
+			const int32 ViewLodOverride = GetViewLodOverride(*View);
+			const int32 DrawCollisionLodOverride = GetDrawCollisionLodOverride(*View, CollisionResponse, CollisionMipLevel, SimpleCollisionMipLevel);
 
+			int32 ForcedLODLevel = ViewLodOverride;
+			ForcedLODLevel = DrawCollisionLodOverride >= 0 ? DrawCollisionLodOverride : ForcedLODLevel;
 			ForcedLODLevel = FMath::Min(ForcedLODLevel, (int32)LODSettings.LastLODIndex);
 
 			const float LODScale = View->LODDistanceFactor * CVarStaticMeshLODDistanceScale.GetValueOnRenderThread();
@@ -3433,8 +3445,11 @@ void FLandscapeComponentSceneProxy::GetDynamicMeshElements(const TArray<const FS
 #endif // WITH_EDITOR
 
 #if WITH_EDITOR || !(UE_BUILD_SHIPPING || UE_BUILD_TEST)
+				const bool bInCollisionView = View->Family->EngineShowFlags.CollisionVisibility || View->Family->EngineShowFlags.CollisionPawn;
 				if (AllowDebugViewmodes() && bInCollisionView)
 				{
+					const bool bDrawSimpleCollision = View->Family->EngineShowFlags.CollisionPawn && CollisionResponse.GetResponse(ECC_Pawn) != ECR_Ignore;
+					const bool bDrawComplexCollision = View->Family->EngineShowFlags.CollisionVisibility && CollisionResponse.GetResponse(ECC_Visibility) != ECR_Ignore;
 					if (bDrawSimpleCollision || bDrawComplexCollision)
 					{
 						// Override the mesh's material with our material that draws the collision color
@@ -3616,7 +3631,7 @@ void FLandscapeComponentSceneProxy::GetDynamicRayTracingInstances(FRayTracingMat
 		return;
 	}
 	float MeshScreenSizeSquared = ComputeBoundsScreenRadiusSquared(GetBounds().Origin, GetBounds().SphereRadius, *Context.ReferenceView);
-	int32 ForcedLODLevel = GetLandscapeLodOverride(*Context.ReferenceView);
+	int32 ForcedLODLevel = GetViewLodOverride(*Context.ReferenceView);
 
 	int32 LODToRender = ForcedLODLevel >= 0 ? ForcedLODLevel : GetLODFromScreenSize(MeshScreenSizeSquared, Context.ReferenceView->LODDistanceFactor);
 	
diff --git a/Engine/Source/Runtime/Landscape/Public/LandscapeRender.h b/Engine/Source/Runtime/Landscape/Public/LandscapeRender.h
--- a/Engine/Source/Runtime/Landscape/Public/LandscapeRender.h
+++ b/Engine/Source/Runtime/Landscape/Public/LandscapeRender.h
@@ -503,6 +503,8 @@ struct FLandscapeRenderSystem
 		float LODOnePlusDistributionScalarSquared;
 		float LastLODScreenSizeSquared;
 		int8 LastLODIndex;
+		int8 DrawCollisionPawnLOD;
+		int8 DrawCollisionVisibilityLOD;
 	};
 
 	static int8 GetLODFromScreenSize(LODSettingsComponent LODSettings, float InScreenSizeSquared, float InViewLODScale, float& OutFractionalLOD)
@@ -575,6 +577,8 @@ struct FLandscapeRenderSystem
 		const FSceneView* ViewPtrAsIdentifier;
 		int32 ViewLODOverride;
 		float ViewLODDistanceFactor;
+		bool ViewEngineShowFlagCollisionPawn;
+		bool ViewEngineShowFlagCollisionVisibility;
 		FVector ViewOrigin;
 		FMatrix ViewProjectionMatrix;
 
@@ -597,7 +601,10 @@ struct FLandscapeRenderSystem
 
 		void AnyThreadTask()
 		{
-			RenderSystem.ComputeSectionPerViewParameters(ViewPtrAsIdentifier, ViewLODOverride, ViewLODDistanceFactor, ViewOrigin, ViewProjectionMatrix);
+			RenderSystem.ComputeSectionPerViewParameters(
+				ViewPtrAsIdentifier, ViewLODOverride, ViewLODDistanceFactor, 
+				ViewEngineShowFlagCollisionPawn, ViewEngineShowFlagCollisionVisibility, 
+				ViewOrigin, ViewProjectionMatrix);
 		}
 
 		void DoTask(ENamedThreads::Type CurrentThread, const FGraphEventRef& MyCompletionGraphEvent)
@@ -693,6 +700,8 @@ struct FLandscapeRenderSystem
 		const FSceneView* ViewPtrAsIdentifier,
 		int32 ViewLODOverride,
 		float ViewLODDistanceFactor,
+		bool bDrawCollisionPawn,
+		bool bDrawCollisionCollision,
 		FVector ViewOrigin,
 		FMatrix ViewProjectionMarix);
 
