Set LandscapeComponent physics actor to static. Fixes CCD not being used with landscape.

diff --git a/Engine/Source/Runtime/Landscape/Private/LandscapeCollision.cpp b/Engine/Source/Runtime/Landscape/Private/LandscapeCollision.cpp
--- a/Engine/Source/Runtime/Landscape/Private/LandscapeCollision.cpp
+++ b/Engine/Source/Runtime/Landscape/Private/LandscapeCollision.cpp
@@ -399,6 +399,7 @@ void ULandscapeHeightfieldCollisionComponent::OnCreatePhysicsState()
 				Params.InitialTM = LandscapeComponentTransform;
 				Params.InitialTM.SetScale3D(FVector(0));
 				Params.bQueryOnly = true;
+				Params.bStatic = true;
 				Params.Scene = GetWorld()->GetPhysicsScene();
 				FPhysicsActorHandle PhysHandle;
 				FPhysicsInterface::CreateActor(Params, PhysHandle);
