- Fixing up compile errors after UE4/Main pull
#rnx

diff --git a/Engine/Source/Runtime/Landscape/Private/LandscapeRender.cpp b/Engine/Source/Runtime/Landscape/Private/LandscapeRender.cpp
--- a/Engine/Source/Runtime/Landscape/Private/LandscapeRender.cpp
+++ b/Engine/Source/Runtime/Landscape/Private/LandscapeRender.cpp
@@ -1333,8 +1333,10 @@ FLandscapeComponentSceneProxy::FLandscapeComponentSceneProxy(ULandscapeComponent
 	FirstVirtualTextureLOD = FMath::Max(MaxLOD - InComponent->GetLandscapeProxy()->VirtualTextureNumLods, 0);
 	VirtualTextureLodBias = InComponent->GetLandscapeProxy()->VirtualTextureLodBias;
 
+#if WITH_EDITOR || !(UE_BUILD_SHIPPING || UE_BUILD_TEST)
 	LODSettings.DrawCollisionPawnLOD = CollisionResponse.GetResponse(ECC_Pawn) == ECR_Ignore ? -1 : SimpleCollisionMipLevel;
 	LODSettings.DrawCollisionVisibilityLOD = CollisionResponse.GetResponse(ECC_Visibility) == ECR_Ignore ? -1 : CollisionMipLevel;
+#endif
 
 	ComponentMaxExtend = SubsectionSizeQuads * FMath::Max(InComponent->GetComponentTransform().GetScale3D().X, InComponent->GetComponentTransform().GetScale3D().Y);
 
@@ -2486,10 +2488,12 @@ void FLandscapeComponentSceneProxy::CalculateLODFromScreenSize(const FSceneView&
 {
 	// Handle LOD overrides
 	const int32 ViewLodOverride = GetViewLodOverride(InView);
-	const int32 DrawCollisionLodOverride = GetDrawCollisionLodOverride(InView, CollisionResponse, CollisionMipLevel, SimpleCollisionMipLevel);
 
 	float PreferedLOD = (float)ViewLodOverride;
+#if WITH_EDITOR || !(UE_BUILD_SHIPPING || UE_BUILD_TEST)
+	const int32 DrawCollisionLodOverride = GetDrawCollisionLodOverride(InView, CollisionResponse, CollisionMipLevel, SimpleCollisionMipLevel);
 	PreferedLOD = DrawCollisionLodOverride >= 0 ? (float)DrawCollisionLodOverride : PreferedLOD;
+#endif
 	PreferedLOD = ForcedLOD >= 0 ? (float)ForcedLOD : PreferedLOD;
 
 	int8 MinStreamedLOD = HeightmapTexture ? FMath::Min<int8>(((FTexture2DResource*)HeightmapTexture->Resource)->GetCurrentFirstMip(), FMath::CeilLogTwo(SubsectionSizeVerts) - 1) : 0;
@@ -3275,10 +3279,11 @@ void FLandscapeComponentSceneProxy::GetDynamicMeshElements(const TArray<const FS
 			const FSceneView* View = Views[ViewIndex];
 
 			const int32 ViewLodOverride = GetViewLodOverride(*View);
-			const int32 DrawCollisionLodOverride = GetDrawCollisionLodOverride(*View, CollisionResponse, CollisionMipLevel, SimpleCollisionMipLevel);
-
 			int32 ForcedLODLevel = ViewLodOverride;
+#if WITH_EDITOR || !(UE_BUILD_SHIPPING || UE_BUILD_TEST)
+			const int32 DrawCollisionLodOverride = GetDrawCollisionLodOverride(*View, CollisionResponse, CollisionMipLevel, SimpleCollisionMipLevel);
 			ForcedLODLevel = DrawCollisionLodOverride >= 0 ? DrawCollisionLodOverride : ForcedLODLevel;
+#endif
 			ForcedLODLevel = FMath::Min(ForcedLODLevel, (int32)LODSettings.LastLODIndex);
 
 			const float LODScale = View->LODDistanceFactor * CVarStaticMeshLODDistanceScale.GetValueOnRenderThread();
