Fix landscape component ForcedLOD option.
Clamps the FirstLOD and LastLOD to the ForcedLOD.

diff --git a/Engine/Source/Runtime/Landscape/Private/LandscapeRender.cpp b/Engine/Source/Runtime/Landscape/Private/LandscapeRender.cpp
--- a/Engine/Source/Runtime/Landscape/Private/LandscapeRender.cpp
+++ b/Engine/Source/Runtime/Landscape/Private/LandscapeRender.cpp
@@ -924,8 +924,10 @@ void FLandscapeRenderSystem::ComputeSectionPerViewParameters(
 		float FractionalLOD;
 		GetLODFromScreenSize(SectionLODSettings[EntityIndex], MeshScreenSizeSquared, LODScale * LODScale, FractionalLOD);
 
+		int32 ForcedLODLevel = SectionLODSettings[EntityIndex].ForcedLOD;
+		ForcedLODLevel = ViewLODOverride >= 0 ? ViewLODOverride : ForcedLODLevel;
 		const int32 DrawCollisionLODOverride = GetDrawCollisionLodOverride(bDrawCollisionPawn, bDrawCollisionCollision, SectionLODSettings[EntityIndex].DrawCollisionPawnLOD, SectionLODSettings[EntityIndex].DrawCollisionVisibilityLOD);
-		int32 ForcedLODLevel = DrawCollisionLODOverride >= 0 ? DrawCollisionLODOverride : ViewLODOverride;
+		ForcedLODLevel = DrawCollisionLODOverride >= 0 ? DrawCollisionLODOverride : ForcedLODLevel;
 		ForcedLODLevel = FMath::Min<int32>(ForcedLODLevel, SectionLODSettings[EntityIndex].LastLODIndex);
 
 #if PLATFORM_SUPPORTS_LANDSCAPE_VISUAL_MESH_LOD_STREAMING
@@ -1319,10 +1321,15 @@ FLandscapeComponentSceneProxy::FLandscapeComponentSceneProxy(ULandscapeComponent
 		LastLOD = FMath::Min<int32>(MaxLODLevel, LastLOD);
 	}
 
+	// Clamp ForcedLOD to the valid range and then apply
+	ForcedLOD = ForcedLOD >= 0 ? FMath::Clamp<int32>(ForcedLOD, FirstLOD, LastLOD) : ForcedLOD;
+	FirstLOD = ForcedLOD >= 0 ? ForcedLOD : FirstLOD;
+	LastLOD = ForcedLOD >= 0 ? ForcedLOD : LastLOD;
+
 	LODSettings.LastLODIndex = LastLOD;
 	LODSettings.LastLODScreenSizeSquared = LODScreenRatioSquared[LastLOD];
+	LODSettings.ForcedLOD = ForcedLOD;
 
-	ForcedLOD = ForcedLOD != INDEX_NONE ? FMath::Clamp<int32>(ForcedLOD, FirstLOD, LastLOD) : ForcedLOD;
 	LODBias = FMath::Clamp<int8>(LODBias, -MaxLOD, MaxLOD);
 
 	int8 LocalLODBias = LODBias + (int8)GLandscapeMeshLODBias;
@@ -1827,7 +1834,7 @@ FPrimitiveViewRelevance FLandscapeComponentSceneProxy::GetViewRelevance(const FS
 		(IsSelected() && !GLandscapeEditModeActive) ||
 		(GLandscapeViewMode != ELandscapeViewMode::Normal) ||
 		(CVarLandscapeShowDirty.GetValueOnAnyThread() && GLandscapeDirtyMaterial) ||
-		(View->Family->LandscapeLODOverride >= 0) ||
+		(GetViewLodOverride(*View) >= 0) ||
 #else
 		IsSelected() ||
 #endif
@@ -3281,8 +3288,10 @@ void FLandscapeComponentSceneProxy::GetDynamicMeshElements(const TArray<const FS
 
 			const FSceneView* View = Views[ViewIndex];
 
+			int32 ForcedLODLevel = ForcedLOD;
+
 			const int32 ViewLodOverride = GetViewLodOverride(*View);
-			int32 ForcedLODLevel = ViewLodOverride;
+			ForcedLODLevel = ViewLodOverride >= 0 ? ViewLodOverride : ForcedLODLevel;
 
 #if WITH_EDITOR || !(UE_BUILD_SHIPPING || UE_BUILD_TEST)
 			const int32 DrawCollisionLodOverride = GetDrawCollisionLodOverride(*View, CollisionResponse, CollisionMipLevel, SimpleCollisionMipLevel);
@@ -3509,7 +3518,7 @@ void FLandscapeComponentSceneProxy::GetDynamicMeshElements(const TArray<const FS
 						bIsWireframe ||
 #if WITH_EDITOR
 						(IsSelected() && !GLandscapeEditModeActive) ||
-						ViewFamily.LandscapeLODOverride >= 0 ||
+						(GetViewLodOverride(*View) >= 0) ||
 #else
 						IsSelected() ||
 #endif
@@ -3640,9 +3649,13 @@ void FLandscapeComponentSceneProxy::GetDynamicRayTracingInstances(FRayTracingMat
 	{
 		return;
 	}
-	float MeshScreenSizeSquared = ComputeBoundsScreenRadiusSquared(GetBounds().Origin, GetBounds().SphereRadius, *Context.ReferenceView);
-	int32 ForcedLODLevel = GetViewLodOverride(*Context.ReferenceView);
 
+	int32 ForcedLODLevel = ForcedLOD;
+
+	int32 ViewLodOveride = GetViewLodOverride(*Context.ReferenceView);
+	ForcedLODLevel = ViewLodOveride >= 0 ? ViewLodOveride : ForcedLODLevel;
+
+	float MeshScreenSizeSquared = ComputeBoundsScreenRadiusSquared(GetBounds().Origin, GetBounds().SphereRadius, *Context.ReferenceView);
 	int32 LODToRender = ForcedLODLevel >= 0 ? ForcedLODLevel : GetLODFromScreenSize(MeshScreenSizeSquared, Context.ReferenceView->LODDistanceFactor);
 	
 	FLandscapeElementParamArray& ParameterArray = Context.RayTracingMeshResourceCollector.AllocateOneFrameResource<FLandscapeElementParamArray>();
diff --git a/Engine/Source/Runtime/Landscape/Public/LandscapeRender.h b/Engine/Source/Runtime/Landscape/Public/LandscapeRender.h
--- a/Engine/Source/Runtime/Landscape/Public/LandscapeRender.h
+++ b/Engine/Source/Runtime/Landscape/Public/LandscapeRender.h
@@ -503,6 +503,7 @@ struct FLandscapeRenderSystem
 		float LODOnePlusDistributionScalarSquared;
 		float LastLODScreenSizeSquared;
 		int8 LastLODIndex;
+		int8 ForcedLOD;
 		int8 DrawCollisionPawnLOD;
 		int8 DrawCollisionVisibilityLOD;
 	};
@@ -510,7 +511,7 @@ struct FLandscapeRenderSystem
 	static int8 GetLODFromScreenSize(LODSettingsComponent LODSettings, float InScreenSizeSquared, float InViewLODScale, float& OutFractionalLOD)
 	{
 		float ScreenSizeSquared = InScreenSizeSquared / InViewLODScale;
-
+		
 		if (ScreenSizeSquared <= LODSettings.LastLODScreenSizeSquared)
 		{
 			OutFractionalLOD = LODSettings.LastLODIndex;
@@ -985,7 +986,7 @@ protected:
 	FORCEINLINE void ComputeStaticBatchIndexToRender(FViewCustomDataLOD& OutLODData, int32 InSubSectionIndex);
 	int8 GetLODFromScreenSize(float InScreenSizeSquared, float InViewLODScale) const;
 	FORCEINLINE_DEBUGGABLE float ComputeBatchElementCurrentLOD(int32 InSelectedLODIndex, float InComponentScreenSize, float InViewLODScale) const;
-	
+
 	FORCEINLINE void GetShaderCurrentNeighborLOD(const FSceneView& InView, float InBatchElementCurrentLOD, int8 InSubSectionX, int8 InSubSectionY, int8 InCurrentSubSectionIndex, FVector4& OutShaderCurrentNeighborLOD) const;
 	FORCEINLINE FVector4 GetShaderLODBias() const;
 	FVector4 GetShaderLODValues(int8 BatchElementCurrentLOD) const;
