Fix holes in navmesh when using lazy octree

#rnx

diff --git a/Engine/Source/Runtime/Landscape/Private/LandscapeCollision.cpp b/Engine/Source/Runtime/Landscape/Private/LandscapeCollision.cpp
--- a/Engine/Source/Runtime/Landscape/Private/LandscapeCollision.cpp
+++ b/Engine/Source/Runtime/Landscape/Private/LandscapeCollision.cpp
@@ -2033,15 +2033,16 @@ void ULandscapeHeightfieldCollisionComponent::PrepareGeometryExportSync()
 		{
 			HeightfieldRowsCount = HeightfieldRef->RBHeightfield->getNbRows();
 			HeightfieldColumnsCount = HeightfieldRef->RBHeightfield->getNbColumns();
+			const int32 SamplesCount = HeightfieldRowsCount * HeightfieldColumnsCount;
 				
-			if (CachedHeightFieldSamples.Heights.Num() != HeightfieldRowsCount * HeightfieldRowsCount)
+			if (CachedHeightFieldSamples.Heights.Num() != SamplesCount)
 			{
 				QUICK_SCOPE_CYCLE_COUNTER(STAT_NavMesh_ExportPxHeightField_saveCells);
 
-				CachedHeightFieldSamples.Heights.SetNumUninitialized(HeightfieldRowsCount * HeightfieldRowsCount);
+				CachedHeightFieldSamples.Heights.SetNumUninitialized(SamplesCount);
 
 				TArray<PxHeightFieldSample> HFSamples;
-				HFSamples.SetNumUninitialized(HeightfieldRowsCount * HeightfieldRowsCount);
+				HFSamples.SetNumUninitialized(SamplesCount);
 				HeightfieldRef->RBHeightfield->saveCells(HFSamples.GetData(), HFSamples.Num()*HFSamples.GetTypeSize());
 
 				for (int32 SampleIndex = 0; SampleIndex < HFSamples.Num(); ++SampleIndex)
@@ -2062,17 +2063,23 @@ void ULandscapeHeightfieldCollisionComponent::PrepareGeometryExportSync()
 		{
 			HeightfieldRowsCount = HeightfieldRef->Heightfield->GetNumRows();
 			HeightfieldColumnsCount = HeightfieldRef->Heightfield->GetNumCols();
+			const int32 HeightsCount = HeightfieldRowsCount * HeightfieldColumnsCount;
 
-			if(CachedHeightFieldSamples.Heights.Num() != HeightfieldRowsCount * HeightfieldRowsCount)
+			if(CachedHeightFieldSamples.Heights.Num() != HeightsCount)
 			{
 				QUICK_SCOPE_CYCLE_COUNTER(STAT_NavMesh_ExportChaosHeightField_saveCells);
 
-				CachedHeightFieldSamples.Heights.SetNumUninitialized(HeightfieldRowsCount * HeightfieldRowsCount);
-				
-				for(int32 SampleIndex = 0; SampleIndex < CachedHeightFieldSamples.Heights.Num(); ++SampleIndex)
+				CachedHeightFieldSamples.Heights.SetNumUninitialized(HeightsCount);
+				for(int32 Index = 0; Index < HeightsCount; ++Index)
+				{
+					CachedHeightFieldSamples.Heights[Index] = HeightfieldRef->Heightfield->GetHeight(Index);
+				}
+
+				const int32 HolesCount = (HeightfieldRowsCount-1) * (HeightfieldColumnsCount-1);
+				CachedHeightFieldSamples.Holes.SetNumUninitialized(HolesCount);
+				for(int32 Index = 0; Index < HolesCount; ++Index)
 				{
-					CachedHeightFieldSamples.Heights[SampleIndex] = HeightfieldRef->Heightfield->GetHeight(SampleIndex);
-					CachedHeightFieldSamples.Holes.Add(HeightfieldRef->Heightfield->IsHole(SampleIndex));
+					CachedHeightFieldSamples.Holes[Index] = HeightfieldRef->Heightfield->IsHole(Index);
 				}
 			}
 		}
diff --git a/Engine/Source/Runtime/NavigationSystem/Private/NavMesh/RecastNavMeshGenerator.cpp b/Engine/Source/Runtime/NavigationSystem/Private/NavMesh/RecastNavMeshGenerator.cpp
--- a/Engine/Source/Runtime/NavigationSystem/Private/NavMesh/RecastNavMeshGenerator.cpp
+++ b/Engine/Source/Runtime/NavigationSystem/Private/NavMesh/RecastNavMeshGenerator.cpp
@@ -872,12 +872,6 @@ void ExportChaosHeightFieldSlice(const FNavHeightfieldSamples& PrefetchedHeightf
 {
 	QUICK_SCOPE_CYCLE_COUNTER(STAT_NavMesh_ExportChaosHeightFieldSlice);
 
-#if PHYSICS_INTERFACE_PHYSX
-	static const uint32 SizeOfPx = sizeof(physx::PxI16);
-	static const uint32 SizeOfHeight = PrefetchedHeightfieldSamples.Heights.GetTypeSize();
-	ensure(SizeOfPx == SizeOfHeight);
-#endif // WITH_PHYSX
-
 	// calculate the actual start and number of columns we want
 	const FBox LocalBox = SliceBox.TransformBy(LocalToWorld.Inverse());
 	const bool bMirrored = (LocalToWorld.GetDeterminant() < 0.f);
@@ -926,7 +920,7 @@ void ExportChaosHeightFieldSlice(const FNavHeightfieldSamples& PrefetchedHeightf
 		{
 			const int32 CoordX = IdxX + MinX;
 			const int32 CoordY = IdxY + MinY;
-			const int32 SampleIdx = CoordY * NumCols + CoordX;  // #PHYSTODO bMirrored support
+			const int32 SampleIdx = CoordY * (NumCols-1) + CoordX;  // #PHYSTODO bMirrored support
 
 			const bool bIsHole = PrefetchedHeightfieldSamples.Holes[SampleIdx];
 			if (bIsHole)
