Fix Landscape Grass Scale not applied when min and max are same value

diff --git a/Engine/Source/Runtime/Landscape/Classes/LandscapeGrassType.h b/Engine/Source/Runtime/Landscape/Classes/LandscapeGrassType.h
--- a/Engine/Source/Runtime/Landscape/Classes/LandscapeGrassType.h
+++ b/Engine/Source/Runtime/Landscape/Classes/LandscapeGrassType.h
@@ -69,11 +69,11 @@ struct FGrassVariety
 	FFloatInterval ScaleX;
 
 	/** Specifies the range of scale, from minimum to maximum, to apply to a grass instance's Y Scale property */
-	UPROPERTY(EditAnywhere, Category=Grass)
+	UPROPERTY(EditAnywhere, Category=Grass, meta = (EditCondition = "Scaling == EGrassScaling::Free"))
 	FFloatInterval ScaleY;
 
 	/** Specifies the range of scale, from minimum to maximum, to apply to a grass instance's Z Scale property */
-	UPROPERTY(EditAnywhere, Category=Grass)
+	UPROPERTY(EditAnywhere, Category=Grass, meta = (EditCondition = "Scaling == EGrassScaling::Free || Scaling == EGrassScaling::LockXY"))
 	FFloatInterval ScaleZ;
 
 	/** Whether the grass instances should be placed at random rotation (true) or all at the same rotation (false) */
diff --git a/Engine/Source/Runtime/Landscape/Private/LandscapeGrass.cpp b/Engine/Source/Runtime/Landscape/Private/LandscapeGrass.cpp
--- a/Engine/Source/Runtime/Landscape/Private/LandscapeGrass.cpp
+++ b/Engine/Source/Runtime/Landscape/Private/LandscapeGrass.cpp
@@ -1539,7 +1539,7 @@ struct FAsyncGrassBuilder : public FGrassBuilderBase
 		, ScaleY(GrassVariety.ScaleY)
 		, ScaleZ(GrassVariety.ScaleZ)
 		, RandomRotation(GrassVariety.RandomRotation)
-		, RandomScale(GrassVariety.ScaleX.Size() > 0 || GrassVariety.ScaleY.Size() > 0 || GrassVariety.ScaleZ.Size() > 0)
+		, RandomScale(false)
 		, AlignToSurface(GrassVariety.AlignToSurface)
 		, PlacementJitter(GrassVariety.PlacementJitter)
 		, RandomStream(HierarchicalInstancedStaticMeshComponent->InstancingRandomSeed)
@@ -1564,7 +1564,23 @@ struct FAsyncGrassBuilder : public FGrassBuilderBase
 		, InstanceBuffer(/*bSupportsVertexHalfFloat*/ GVertexElementTypeSupport.IsSupported(VET_Half2))
 		, ClusterTree()
 		, OutOcclusionLayerNum(0)
-	{
+	{		
+
+		switch (Scaling)
+		{
+		case EGrassScaling::Uniform:
+			RandomScale = ScaleX.Size() > 0;
+			break;
+		case EGrassScaling::Free:
+			RandomScale = ScaleX.Size() > 0 || ScaleY.Size() > 0 || ScaleZ.Size() > 0;
+			break;
+		case EGrassScaling::LockXY:
+			RandomScale = ScaleX.Size() > 0 || ScaleZ.Size() > 0;
+			break;
+		default:
+			check(0);
+		}
+
 		if (InExcludedBoxes.Num())
 		{
 			FMatrix BoxXForm = HierarchicalInstancedStaticMeshComponent->GetComponentToWorld().ToMatrixWithScale().Inverse() * XForm.Inverse();
@@ -1653,6 +1669,28 @@ struct FAsyncGrassBuilder : public FGrassBuilderBase
 		}
 	}
 
+	FVector GetDefaultScale() const
+	{
+		FVector Result( ScaleX.Min > 0.0f && FMath::IsNearlyZero(ScaleX.Size()) ? ScaleX.Min : 1.0f,
+						ScaleY.Min > 0.0f && FMath::IsNearlyZero(ScaleY.Size()) ? ScaleY.Min : 1.0f,
+						ScaleZ.Min > 0.0f && FMath::IsNearlyZero(ScaleZ.Size()) ? ScaleZ.Min : 1.0f);
+		switch (Scaling)
+		{
+		case EGrassScaling::Uniform:
+			Result.Y = Result.X;
+			Result.Z = Result.X;
+			break;
+		case EGrassScaling::Free:
+			break;
+		case EGrassScaling::LockXY:
+			Result.Y = Result.X;
+			break;
+		default:
+			check(0);
+		}
+		return Result;
+	}
+
 	FVector GetRandomScale() const
 	{
 		FVector Result(1.0f);
@@ -1677,7 +1715,7 @@ struct FAsyncGrassBuilder : public FGrassBuilderBase
 		default:
 			check(0);
 		}
-
+	
 		return Result;
 	}
 
@@ -1698,7 +1736,7 @@ struct FAsyncGrassBuilder : public FGrassBuilderBase
 		SCOPE_CYCLE_COUNTER(STAT_FoliageGrassAsyncBuildTime);
 		check(bHaveValidData);
 		double StartTime = FPlatformTime::Seconds();
-
+		const FVector DefaultScale = GetDefaultScale();
 		float Div = 1.0f / float(SqrtMaxInstances);
 		TArray<FMatrix> InstanceTransforms;
 		if (HaltonBaseIndex)
@@ -1728,7 +1766,7 @@ struct FAsyncGrassBuilder : public FGrassBuilderBase
 				bool bKeep = Weight > 0.0f && Weight >= RandomStream.GetFraction() && !IsExcluded(LocationWithHeight);
 				if (bKeep)
 				{
-					const FVector Scale = RandomScale ? GetRandomScale() : FVector(1);
+					const FVector Scale = RandomScale ? GetRandomScale() : DefaultScale;
 					const float Rot = RandomRotation ? RandomStream.GetFraction() * 360.0f : 0.0f;
 					const FMatrix BaseXForm = FScaleRotationTranslationMatrix(Scale, FRotator(0.0f, Rot, 0.0f), FVector::ZeroVector);
 					FMatrix OutXForm;
@@ -1814,7 +1852,7 @@ struct FAsyncGrassBuilder : public FGrassBuilderBase
 							const FInstanceLocal& Instance = Instances[InstanceIndex];
 							if (Instance.bKeep)
 							{
-								const FVector Scale = RandomScale ? GetRandomScale() : FVector(1);
+								const FVector Scale = RandomScale ? GetRandomScale() : DefaultScale;
 								const float Rot = RandomRotation ? RandomStream.GetFraction() * 360.0f : 0.0f;
 								const FMatrix BaseXForm = FScaleRotationTranslationMatrix(Scale, FRotator(0.0f, Rot, 0.0f), FVector::ZeroVector);
 								FMatrix OutXForm;
