Fix Landscape rotation not being taken into account in ULandscapeComponent::GetLayerWeightAtLocation (Procedural Foliage)

diff --git a/Engine/Source/Runtime/Landscape/Private/LandscapeEdit.cpp b/Engine/Source/Runtime/Landscape/Private/LandscapeEdit.cpp
--- a/Engine/Source/Runtime/Landscape/Private/LandscapeEdit.cpp
+++ b/Engine/Source/Runtime/Landscape/Private/LandscapeEdit.cpp
@@ -2091,50 +2091,43 @@ float ULandscapeComponent::GetLayerWeightAtLocation(const FVector& InLocation, U
 	}
 
 	// Find location
-	// TODO: Root landscape isn't always loaded, would Proxy suffice?
-	if (ALandscape* Landscape = GetLandscapeActor())
+	const FVector TestLocation = GetComponentToWorld().InverseTransformPosition(InLocation);
+	
+	// Abort if the test location is not on this component
+	if (TestLocation.X < 0 || TestLocation.Y < 0 || TestLocation.X > ComponentSizeQuads || TestLocation.Y > ComponentSizeQuads)
 	{
-		const FVector DrawScale = Landscape->GetRootComponent()->GetRelativeScale3D();
-		float TestX = (InLocation.X - Landscape->GetActorLocation().X) / DrawScale.X - (float)GetSectionBase().X;
-		float TestY = (InLocation.Y - Landscape->GetActorLocation().Y) / DrawScale.Y - (float)GetSectionBase().Y;
+		return 0.0f;
+	}
 
-		// Abort if the test location is not on this component
-		if (TestX < 0 || TestY < 0 || TestX > ComponentSizeQuads || TestY > ComponentSizeQuads)
-		{
-			return 0.0f;
-		}
+	// Find data
+	int32 X1 = FMath::FloorToInt(TestLocation.X);
+	int32 Y1 = FMath::FloorToInt(TestLocation.Y);
+	int32 X2 = FMath::CeilToInt(TestLocation.X);
+	int32 Y2 = FMath::CeilToInt(TestLocation.Y);
 
-		// Find data
-		int32 X1 = FMath::FloorToInt(TestX);
-		int32 Y1 = FMath::FloorToInt(TestY);
-		int32 X2 = FMath::CeilToInt(TestX);
-		int32 Y2 = FMath::CeilToInt(TestY);
+	int32 Stride = (SubsectionSizeQuads + 1) * NumSubsections;
 
-		int32 Stride = (SubsectionSizeQuads + 1) * NumSubsections;
+	// Min is to prevent the sampling of the final column from overflowing
+	int32 IdxX1 = FMath::Min<int32>(((X1 / SubsectionSizeQuads) * (SubsectionSizeQuads + 1)) + (X1 % SubsectionSizeQuads), Stride - 1);
+	int32 IdxY1 = FMath::Min<int32>(((Y1 / SubsectionSizeQuads) * (SubsectionSizeQuads + 1)) + (Y1 % SubsectionSizeQuads), Stride - 1);
+	int32 IdxX2 = FMath::Min<int32>(((X2 / SubsectionSizeQuads) * (SubsectionSizeQuads + 1)) + (X2 % SubsectionSizeQuads), Stride - 1);
+	int32 IdxY2 = FMath::Min<int32>(((Y2 / SubsectionSizeQuads) * (SubsectionSizeQuads + 1)) + (Y2 % SubsectionSizeQuads), Stride - 1);
 
-		// Min is to prevent the sampling of the final column from overflowing
-		int32 IdxX1 = FMath::Min<int32>(((X1 / SubsectionSizeQuads) * (SubsectionSizeQuads + 1)) + (X1 % SubsectionSizeQuads), Stride - 1);
-		int32 IdxY1 = FMath::Min<int32>(((Y1 / SubsectionSizeQuads) * (SubsectionSizeQuads + 1)) + (Y1 % SubsectionSizeQuads), Stride - 1);
-		int32 IdxX2 = FMath::Min<int32>(((X2 / SubsectionSizeQuads) * (SubsectionSizeQuads + 1)) + (X2 % SubsectionSizeQuads), Stride - 1);
-		int32 IdxY2 = FMath::Min<int32>(((Y2 / SubsectionSizeQuads) * (SubsectionSizeQuads + 1)) + (Y2 % SubsectionSizeQuads), Stride - 1);
+	// sample
+	float Sample11 = (float)((*LayerCache)[IdxX1 + Stride * IdxY1]) / 255.0f;
+	float Sample21 = (float)((*LayerCache)[IdxX2 + Stride * IdxY1]) / 255.0f;
+	float Sample12 = (float)((*LayerCache)[IdxX1 + Stride * IdxY2]) / 255.0f;
+	float Sample22 = (float)((*LayerCache)[IdxX2 + Stride * IdxY2]) / 255.0f;
 
-		// sample
-		float Sample11 = (float)((*LayerCache)[IdxX1 + Stride*IdxY1]) / 255.0f;
-		float Sample21 = (float)((*LayerCache)[IdxX2 + Stride*IdxY1]) / 255.0f;
-		float Sample12 = (float)((*LayerCache)[IdxX1 + Stride*IdxY2]) / 255.0f;
-		float Sample22 = (float)((*LayerCache)[IdxX2 + Stride*IdxY2]) / 255.0f;
+	float LerpX = FMath::Fractional(TestLocation.X);
+	float LerpY = FMath::Fractional(TestLocation.Y);
 
-		float LerpX = FMath::Fractional(TestX);
-		float LerpY = FMath::Fractional(TestY);
+	// Bilinear interpolate
+	return FMath::Lerp(
+		FMath::Lerp(Sample11, Sample21, LerpX),
+		FMath::Lerp(Sample12, Sample22, LerpX),
+		LerpY);
 
-		// Bilinear interpolate
-		return FMath::Lerp(
-			FMath::Lerp(Sample11, Sample21, LerpX),
-			FMath::Lerp(Sample12, Sample22, LerpX),
-			LerpY);
-	}
-	
-	return 0.f;	//if landscape is null we just return 0 instead of crashing. Seen cases where this happens, seems like a bug?
 }
 
 void ULandscapeComponent::GetComponentExtent(int32& MinX, int32& MinY, int32& MaxX, int32& MaxY) const
