Fix ForeignWorld not being updated properly on SplineControlPoints (UDN)

diff --git a/Engine/Source/Runtime/Landscape/Private/LandscapeSplines.cpp b/Engine/Source/Runtime/Landscape/Private/LandscapeSplines.cpp
--- a/Engine/Source/Runtime/Landscape/Private/LandscapeSplines.cpp
+++ b/Engine/Source/Runtime/Landscape/Private/LandscapeSplines.cpp
@@ -1829,6 +1829,8 @@ void ULandscapeSplineControlPoint::UpdateSplinePoints(bool bUpdateCollision, boo
 	auto ForeignMeshComponentsMap = GetForeignMeshComponents();
 
 	ModificationKey = FGuid::NewGuid();
+	// Clear Foreign world (will get updated if MeshComponent lives outside of Spline world)
+	ForeignWorld = nullptr;
 
 	UControlPointMeshComponent* MeshComponent = LocalMeshComponent;
 	ULandscapeSplinesComponent* MeshComponentOuterSplines = OuterSplines;
@@ -1889,22 +1891,27 @@ void ULandscapeSplineControlPoint::UpdateSplinePoints(bool bUpdateCollision, boo
 			}
 		}
 
+		// If something changed add MeshComponent to proper map
 		if (bUpdateLocalForeign)
 		{
 			if (MeshComponentOuterSplines == OuterSplines)
 			{
 				UObject*& ValueRef = MeshComponentOuterSplines->MeshComponentLocalOwnersMap.FindOrAdd(MeshComponent);
 				ValueRef = this;
-
 				LocalMeshComponent = MeshComponent;
 			}
 			else
 			{
 				MeshComponentOuterSplines->AddForeignMeshComponent(this, MeshComponent);
-				ForeignWorld = MeshComponentOuterSplines->GetTypedOuter<UWorld>();
 			}
 		}
-				
+		
+		// Update Foreign World
+		if (MeshComponent && MeshComponentOuterSplines != OuterSplines)
+		{
+			ForeignWorld = MeshComponentOuterSplines->GetTypedOuter<UWorld>();
+		}
+
 		FVector MeshLocation = Location;
 		FRotator MeshRotation = Rotation;
 		if (MeshComponentOuterSplines != OuterSplines)
