Landscape: Exclude Shader to fix log spam

#rnx

diff --git a/Engine/Source/Runtime/Landscape/Private/LandscapeRender.cpp b/Engine/Source/Runtime/Landscape/Private/LandscapeRender.cpp
--- a/Engine/Source/Runtime/Landscape/Private/LandscapeRender.cpp
+++ b/Engine/Source/Runtime/Landscape/Private/LandscapeRender.cpp
@@ -4995,6 +4995,7 @@ public:
 				FName(TEXT("TMaterialCHSFPrecomputedVolumetricLightmapLightingPolicy")),
 				FName(TEXT("TMaterialCHSFNoLightMapPolicy")),
 				FName(TEXT("FRayTracingDynamicGeometryConverterCS")),
+				FName(TEXT("FTrivialMaterialCHS"))
 #endif // RHI_RAYTRACING
 		};
 		return ExcludedShaderTypes;
