Initialize landscape physical material tasks.
Without this some uninitialized tasks are considered "live" and flagged up as leaking.

diff --git a/Engine/Source/Runtime/Landscape/Private/LandscapePhysicalMaterial.cpp b/Engine/Source/Runtime/Landscape/Private/LandscapePhysicalMaterial.cpp
--- a/Engine/Source/Runtime/Landscape/Private/LandscapePhysicalMaterial.cpp
+++ b/Engine/Source/Runtime/Landscape/Private/LandscapePhysicalMaterial.cpp
@@ -359,12 +359,12 @@ namespace
 	struct FLandscapePhysicalMaterialRenderTaskImpl
 	{
 		// Create on game thread
-		ULandscapeComponent const* LandscapeComponent;
-		uint32 InitFrameId;
-		FIntPoint TargetSize;
-		FVector ViewOrigin;
-		FMatrix ViewRotationMatrix;
-		FMatrix ProjectionMatrix;
+		ULandscapeComponent const* LandscapeComponent = nullptr;
+		uint32 InitFrameId = 0;
+		FIntPoint TargetSize = FIntPoint(ForceInitToZero);
+		FVector ViewOrigin = FVector(ForceInitToZero);
+		FMatrix ViewRotationMatrix = FMatrix(ForceInitToZero);
+		FMatrix ProjectionMatrix = FMatrix(ForceInitToZero);
 
 		// Result written on game thread and read on game thread
 		TArray<UPhysicalMaterial*> ResultMaterials;
@@ -374,7 +374,7 @@ namespace
 		FGPUFenceRHIRef ReadbackFence;
 
 		// Result written on render thread and read on game thread
-		ECompletionState CompletionState;
+		ECompletionState CompletionState = ECompletionState::None;
 		TArray<uint8> ResultIds;
 	};
 
