Deprecate ViewCustomData and CustomLOD frameworks

diff --git a/Engine/Plugins/Enterprise/LidarPointCloud/Source/LidarPointCloudRuntime/Private/Rendering/LidarPointCloudRendering.cpp b/Engine/Plugins/Enterprise/LidarPointCloud/Source/LidarPointCloudRuntime/Private/Rendering/LidarPointCloudRendering.cpp
--- a/Engine/Plugins/Enterprise/LidarPointCloud/Source/LidarPointCloudRuntime/Private/Rendering/LidarPointCloudRendering.cpp
+++ b/Engine/Plugins/Enterprise/LidarPointCloud/Source/LidarPointCloudRuntime/Private/Rendering/LidarPointCloudRendering.cpp
@@ -265,7 +265,8 @@ public:
 
 		return Result;
 	}
-
+	
+	PRAGMA_DISABLE_DEPRECATION_WARNINGS
 	virtual void* InitViewCustomData(const FSceneView& InView, float InViewLODScale, FMemStackBase& InCustomDataMemStack, bool InIsStaticRelevant, bool InIsShadowOnly, const FLODMask* InVisiblePrimitiveLODMask = nullptr, float InMeshScreenSizeSquared = -1.0f) override
 	{
 		// Don't process for shadow views
@@ -276,6 +277,7 @@ public:
 
 		return nullptr;
 	}
+	PRAGMA_ENABLE_DEPRECATION_WARNINGS
 
 	bool PrepareCollisionWireframe(FMeshBatch& MeshBatch, FColoredMaterialRenderProxy* CollisionMaterialInstance) const
 	{
diff --git a/Engine/Source/Runtime/Engine/Public/PrimitiveSceneProxy.h b/Engine/Source/Runtime/Engine/Public/PrimitiveSceneProxy.h
--- a/Engine/Source/Runtime/Engine/Public/PrimitiveSceneProxy.h
+++ b/Engine/Source/Runtime/Engine/Public/PrimitiveSceneProxy.h
@@ -723,6 +723,7 @@ public:
    	 * @param InVisiblePrimitiveLODMask - Calculated LODMask for visibile primitive in static relevancy
    	 * @param InMeshScreenSizeSquared - Computed mesh batch screen size, passed to prevent recalculation
 	 */
+	UE_DEPRECATED(4.25, "The entire ViewCustomData is deprecated. You need to reimplement your feature in other ways, like IPersistentViewUniformBufferExtension (use landscape as an example).")
 	ENGINE_API virtual void* InitViewCustomData(const FSceneView& InView, float InViewLODScale, FMemStackBase& InCustomDataMemStack, bool InIsStaticRelevant, bool InIsShadowOnly, const struct FLODMask* InVisiblePrimitiveLODMask = nullptr, float InMeshScreenSizeSquared = -1.0f) { return nullptr; }
 
 	/** Tell us if this proxy is drawn in game.*/
@@ -739,9 +740,11 @@ public:
   	 * @param InForcedLODLevel - Engine Forced LOD value
    	 * @param OutScreenSizeSquared - Computed screen size from the function
 	 */
+	UE_DEPRECATED(4.25, "We no longer support custom LOD rules.")
 	ENGINE_API virtual struct FLODMask GetCustomLOD(const FSceneView& InView, float InViewLODScale, int32 InForcedLODLevel, float& OutScreenSizeSquared) const;
 
 	/** Tell us if we should rely on the default shadow LOD computing rules or not for generating whole scene shadow.*/
+	UE_DEPRECATED(4.25, "We no longer support custom LOD rules.")
 	ENGINE_API virtual bool IsUsingCustomWholeSceneShadowLODRules() const { return false; }
 	
 	/** 
@@ -755,6 +758,7 @@ public:
    	 * @param InShadowCascadeId - Shadow cascade Id
    	 * @param InHasSelfShadow - Indicate if we have self shadow, as it can impact which LODMask we choose
 	 */
+	UE_DEPRECATED(4.25, "We no longer support custom shadow LOD rules.")
 	ENGINE_API virtual struct FLODMask GetCustomWholeSceneShadowLOD(const FSceneView& InView, float InViewLODScale, int32 InForcedLODLevel, const struct FLODMask& InVisibilePrimitiveLODMask, float InShadowMapTextureResolution, float InShadowMapCascadeSize, int8 InShadowCascadeId, bool InHasSelfShadow) const;
 
 	virtual uint8 GetCurrentFirstLODIdx_RenderThread() const { return 0; }
diff --git a/Engine/Source/Runtime/Landscape/Private/LandscapeRender.cpp b/Engine/Source/Runtime/Landscape/Private/LandscapeRender.cpp
--- a/Engine/Source/Runtime/Landscape/Private/LandscapeRender.cpp
+++ b/Engine/Source/Runtime/Landscape/Private/LandscapeRender.cpp
@@ -2730,57 +2730,6 @@ int8 FLandscapeComponentSceneProxy::GetLODFromScreenSize(float InScreenSizeSquar
 	return FLandscapeRenderSystem::GetLODFromScreenSize(LODSettings, InScreenSizeSquared, InViewLODScale, FractionalLOD);
 }
 
-void* FLandscapeComponentSceneProxy::InitViewCustomData(const FSceneView& InView, float InViewLODScale, FMemStackBase& InCustomDataMemStack, bool InIsStaticRelevant, bool InIsShadowOnly, const FLODMask* InVisiblePrimitiveLODMask, float InMeshScreenSizeSquared)
-{
-	SCOPE_CYCLE_COUNTER(STAT_LandscapeInitViewCustomData);
-
-	// NOTE: we can't access other proxy here as this can be run in parallel we need to wait for the PostInitViewCustomData which is run in synchronous	
-
-	PrimitiveCustomDataIndex = GetPrimitiveSceneInfo()->GetIndex();
-
-	FViewCustomDataLOD* LODData = (FViewCustomDataLOD*)new(InCustomDataMemStack) FViewCustomDataLOD();
-
-	LODData->ComponentScreenSize = InMeshScreenSizeSquared;
-
-	// If a valid screen size was provided, we use it instead of recomputing it
-	if (InMeshScreenSizeSquared < 0.0f)
-	{
-		LODData->ComponentScreenSize = GetComponentScreenSize(&InView, GetBounds().Origin, ComponentMaxExtend, GetBounds().SphereRadius);
-	}
-
-	CalculateBatchElementLOD(InView, LODData->ComponentScreenSize, InViewLODScale, *LODData, false);
-
-	if (InIsStaticRelevant)
-	{
-		check(InVisiblePrimitiveLODMask != nullptr);
-		LODData->StaticMeshBatchLOD = InVisiblePrimitiveLODMask->DitheredLODIndices[0];
-
-		if (LODData->UseCombinedMeshBatch)
-		{
-			ComputeStaticBatchIndexToRender(*LODData, 0);
-		}
-		else
-		{
-			for (int32 i = 0; i < MAX_SUBSECTION_COUNT; ++i)
-			{
-				ComputeStaticBatchIndexToRender(*LODData, i);
-			}
-		}
-	}
-
-	LODData->IsShadowOnly = InIsShadowOnly;
-
-	// Mobile use a different way of calculating the Bias
-	if (GetScene().GetFeatureLevel() >= ERHIFeatureLevel::SM5)
-	{
-		LODData->LodBias = GetShaderLODBias();
-	}
-
-	ComputeTessellationFalloffShaderValues(*LODData, InView.ViewMatrices.GetProjectionMatrix(), LODData->LodTessellationParams.X, LODData->LodTessellationParams.Y);
-
-	return LODData;
-}
-
 void FLandscapeComponentSceneProxy::ComputeTessellationFalloffShaderValues(const FViewCustomDataLOD& InLODData, const FMatrix& InViewProjectionMatrix, float& OutC, float& OutK) const
 {
 	// No Falloff
@@ -3002,47 +2951,6 @@ void FLandscapeComponentSceneProxy::ComputeStaticBatchIndexToRender(FViewCustomD
 	check(SubSectionLODData.StaticBatchElementIndexToRender != INDEX_NONE);
 }
 
-void FLandscapeComponentSceneProxy::PostInitViewCustomData(const FSceneView& InView, void* InViewCustomData) const
-{
-	SCOPE_CYCLE_COUNTER(STAT_LandscapePostInitViewCustomData);
-
-	FViewCustomDataLOD* CurrentLODData = (FViewCustomDataLOD*)InViewCustomData;
-	check(CurrentLODData != nullptr);
-
-	for (int8 SubY = 0; SubY < NumSubsections; ++SubY)
-	{
-		for (int8 SubX = 0; SubX < NumSubsections; ++SubX)
-		{
-			int8 SubSectionIndex = SubX + SubY * NumSubsections;
-			FViewCustomDataSubSectionLOD& SubSectionLODData = CurrentLODData->SubSections[SubSectionIndex];
-			GetShaderCurrentNeighborLOD(InView, SubSectionLODData.fBatchElementCurrentLOD, NumSubsections > 1 ? SubX : INDEX_NONE, NumSubsections > 1 ? SubY : INDEX_NONE, SubSectionIndex, SubSectionLODData.ShaderCurrentNeighborLOD);
-		}
-	}
-
-#if !UE_BUILD_SHIPPING
-	if (GVarDumpLandscapeLODs && GFrameNumberRenderThread == GVarDumpLandscapeLODsCurrentFrame)
-	{
-		if (NumSubsections == 1)
-		{
-			UE_LOG(LogLandscape, Warning, TEXT("\nComponent: [%s] -> MeshBatchLOD: %d, ComponentScreenSize: %f, ShaderCurrentLOD: %s, LODTessellation: %s\nHeightmap Texture Name: %s, Heightmap Streamed Mip: %d\nSubSections:\n|0| IndexToRender: %d, fLOD: %f, LOD: %d, NeighborLOD: %s\n"), *SectionBase.ToString(),
-				CurrentLODData->StaticMeshBatchLOD, CurrentLODData->ComponentScreenSize, *CurrentLODData->ShaderCurrentLOD.ToString(), *CurrentLODData->LodTessellationParams.ToString(),
-				HeightmapTexture != nullptr ? *HeightmapTexture->GetFullName() : TEXT("Invalid"), HeightmapTexture != nullptr ? ((FTexture2DResource*)HeightmapTexture->Resource)->GetCurrentFirstMip() : INDEX_NONE,
-				CurrentLODData->SubSections[0].StaticBatchElementIndexToRender, CurrentLODData->SubSections[0].fBatchElementCurrentLOD, CurrentLODData->SubSections[0].BatchElementCurrentLOD, *CurrentLODData->SubSections[0].ShaderCurrentNeighborLOD.ToString());
-		}
-		else
-		{
-			UE_LOG(LogLandscape, Warning, TEXT("\nComponent: [%s] -> MeshBatchLOD: %d, ComponentScreenSize: %f, UseCombinedMeshBatch: %d, ShaderCurrentLOD: %s, LODTessellation: %s\nHeightmap Texture Name: %s, Heightmap Streamed Mip: %d\nSubSections:\n|0| IndexToRender: %d, fLOD: %f, LOD: %d, NeighborLOD: %s\n|1| IndexToRender: %d, fLOD: %f, LOD: %d, NeighborLOD: %s\n|2| IndexToRender: %d, fLOD: %f, LOD: %d, NeighborLOD: %s\n|3| IndexToRender: %d, fLOD: %f, LOD: %d, NeighborLOD: %s\n"),
-				*SectionBase.ToString(), CurrentLODData->StaticMeshBatchLOD, CurrentLODData->ComponentScreenSize, CurrentLODData->UseCombinedMeshBatch, *CurrentLODData->ShaderCurrentLOD.ToString(), *CurrentLODData->LodTessellationParams.ToString(),
-				HeightmapTexture != nullptr ? *HeightmapTexture->GetFullName() : TEXT("Invalid"), HeightmapTexture != nullptr ? ((FTexture2DResource*)HeightmapTexture->Resource)->GetCurrentFirstMip() : INDEX_NONE,
-				CurrentLODData->SubSections[0].StaticBatchElementIndexToRender, CurrentLODData->SubSections[0].fBatchElementCurrentLOD, CurrentLODData->SubSections[0].BatchElementCurrentLOD, *CurrentLODData->SubSections[0].ShaderCurrentNeighborLOD.ToString(),
-				CurrentLODData->SubSections[1].StaticBatchElementIndexToRender, CurrentLODData->SubSections[1].fBatchElementCurrentLOD, CurrentLODData->SubSections[1].BatchElementCurrentLOD, *CurrentLODData->SubSections[1].ShaderCurrentNeighborLOD.ToString(),
-				CurrentLODData->SubSections[2].StaticBatchElementIndexToRender, CurrentLODData->SubSections[2].fBatchElementCurrentLOD, CurrentLODData->SubSections[2].BatchElementCurrentLOD, *CurrentLODData->SubSections[2].ShaderCurrentNeighborLOD.ToString(),
-				CurrentLODData->SubSections[3].StaticBatchElementIndexToRender, CurrentLODData->SubSections[3].fBatchElementCurrentLOD, CurrentLODData->SubSections[3].BatchElementCurrentLOD, *CurrentLODData->SubSections[3].ShaderCurrentNeighborLOD.ToString());
-		}
-	}
-#endif
-}
-
 bool FLandscapeComponentSceneProxy::CanUseMeshBatchForShadowCascade(int8 InLODIndex, float InShadowMapTextureResolution, float InShadowMapCascadeSize) const
 {
 #if 0
@@ -3072,177 +2980,6 @@ bool FLandscapeComponentSceneProxy::CanUseMeshBatchForShadowCascade(int8 InLODIn
 	return true;
 }
 
-FLODMask FLandscapeComponentSceneProxy::GetCustomLOD(const FSceneView& InView, float InViewLODScale, int32 InForcedLODLevel, float& OutScreenSizeSquared) const
-{
-	SCOPE_CYCLE_COUNTER(STAT_LandscapeComputeCustomMeshBatchLOD);
-	FLODMask LODToRender;
-	OutScreenSizeSquared = 0.0f;
-
-	// Handle forced LOD level first
-	if (InForcedLODLevel >= 0)
-	{
-		int8 MinMeshLOD = MAX_int8;
-		int8 MaxMeshLOD = 0;
-		GetPrimitiveSceneInfo()->GetStaticMeshesLODRange(MinMeshLOD, MaxMeshLOD);
-
-		LODToRender.SetLOD(FMath::Clamp<int8>(InForcedLODLevel, MinMeshLOD, MaxMeshLOD));
-	}
-	else if (InView.Family->EngineShowFlags.LOD)
-	{
-		int8 PotentialLOD = INDEX_NONE;
-		OutScreenSizeSquared = GetComponentScreenSize(&InView, GetBounds().Origin, ComponentMaxExtend, GetBounds().SphereRadius);
-
-		if (NumSubsections > 1)
-		{
-			float SubSectionMaxExtend = ComponentMaxExtend / 2.0f;
-			float SubSectionRadius = GetBounds().SphereRadius / 2.0f;
-
-			// Compute screen size of each sub section to determine if we should use the combined logic or the individual logic
-			float ScreenSizeSquared = GetComponentScreenSize(&InView, SubSectionScreenSizeTestingPosition[0], SubSectionMaxExtend, SubSectionRadius);
-			PotentialLOD = GetLODFromScreenSize(ScreenSizeSquared, InViewLODScale);
-		}
-		else
-		{
-			PotentialLOD = GetLODFromScreenSize(OutScreenSizeSquared, InViewLODScale);
-		}
-
-		const auto FeatureLevel = InView.GetFeatureLevel();
-		bool HasTessellationEnabled = false;
-
-		if (FeatureLevel >= ERHIFeatureLevel::SM5)
-		{
-			ULandscapeMaterialInstanceConstant* LandscapeMIC = Cast<ULandscapeMaterialInstanceConstant>(AvailableMaterials[LODIndexToMaterialIndex[PotentialLOD]]);
-
-			if (LandscapeMIC == nullptr)
-			{
-				UMaterialInstanceDynamic* LandscapeMID = Cast<UMaterialInstanceDynamic>(AvailableMaterials[LODIndexToMaterialIndex[PotentialLOD]]);
-
-				if (LandscapeMID != nullptr)
-				{
-					LandscapeMIC = Cast<ULandscapeMaterialInstanceConstant>(LandscapeMID->Parent);
-				}
-			}
-
-			HasTessellationEnabled = MaterialHasTessellationEnabled[LODIndexToMaterialIndex[PotentialLOD]] && (LandscapeMIC != nullptr && !LandscapeMIC->bDisableTessellation);
-		}
-
-		int32 BaseMeshBatchIndex = MaterialIndexToStaticMeshBatchLOD[LODIndexToMaterialIndex[PotentialLOD]];
-
-		if (HasTessellationEnabled)
-		{
-			const int8 TessellatedMeshBatchLODIndex = 0;
-			const int8 NonTessellatedMeshBatchLODIndex = 1;
-			LODToRender.SetLOD(OutScreenSizeSquared >= TessellationComponentSquaredScreenSize * InViewLODScale ? BaseMeshBatchIndex + TessellatedMeshBatchLODIndex : BaseMeshBatchIndex + NonTessellatedMeshBatchLODIndex);
-		}
-		else
-		{
-			LODToRender.SetLOD(BaseMeshBatchIndex);
-		}
-	}
-
-	return LODToRender;
-}
-
-FLODMask FLandscapeComponentSceneProxy::GetCustomWholeSceneShadowLOD(const FSceneView& InView, float InViewLODScale, int32 InForcedLODLevel, const FLODMask& InVisibilePrimitiveLODMask, float InShadowMapTextureResolution, float InShadowMapCascadeSize, int8 InShadowCascadeId, bool InHasSelfShadow) const
-{
-	SCOPE_CYCLE_COUNTER(STAT_LandscapeComputeCustomShadowMeshBatchLOD);
-
-	FLODMask LODToRender;
-
-	// Handle forced LOD level first
-	if (InForcedLODLevel >= 0)
-	{
-		int8 MinMeshLOD = MAX_int8;
-		int8 MaxMeshLOD = 0;
-		GetPrimitiveSceneInfo()->GetStaticMeshesLODRange(MinMeshLOD, MaxMeshLOD);
-
-		LODToRender.SetLOD(FMath::Clamp<int8>(InForcedLODLevel, MinMeshLOD, MaxMeshLOD));
-	}
-	else if (!InHasSelfShadow) // Force lowest valid LOD
-	{
-		int8 MinMeshLOD = MAX_int8;
-		int8 MaxMeshLOD = 0;
-		GetPrimitiveSceneInfo()->GetStaticMeshesLODRange(MinMeshLOD, MaxMeshLOD);
-		LODToRender.SetLOD(MinMeshLOD);
-	}
-	else
-	{
-		const FViewCustomDataLOD* PrimitiveCustomData = (const FViewCustomDataLOD*)InView.GetCustomData(GetPrimitiveSceneInfo()->GetIndex());
-		int8 PotentialLOD = INDEX_NONE;
-		float ScreenSizeSquared = 0.0f;
-
-		if (PrimitiveCustomData == nullptr)
-		{
-			ScreenSizeSquared = GetComponentScreenSize(&InView, GetBounds().Origin, ComponentMaxExtend, GetBounds().SphereRadius);
-
-			if (NumSubsections > 1)
-			{
-				float SubSectionMaxExtend = ComponentMaxExtend / 2.0f;
-				float SubSectionRadius = GetBounds().SphereRadius / 2.0f;
-
-				// Compute screen size of each sub section to determine if we should use the combined logic or the individual logic
-				float SubSectionScreenSizeSquared = GetComponentScreenSize(&InView, SubSectionScreenSizeTestingPosition[0], SubSectionMaxExtend, SubSectionRadius);
-				PotentialLOD = GetLODFromScreenSize(SubSectionScreenSizeSquared, InViewLODScale);
-			}
-			else
-			{
-				PotentialLOD = GetLODFromScreenSize(ScreenSizeSquared, InViewLODScale);
-			}
-		}
-		else
-		{
-			ScreenSizeSquared = PrimitiveCustomData->ComponentScreenSize;
-			PotentialLOD = PrimitiveCustomData->SubSections[0].BatchElementCurrentLOD;
-		}
-
-		const auto FeatureLevel = InView.GetFeatureLevel();
-		bool HasTessellationEnabled = false;
-
-		if (FeatureLevel >= ERHIFeatureLevel::SM5)
-		{
-			ULandscapeMaterialInstanceConstant* LandscapeMIC = Cast<ULandscapeMaterialInstanceConstant>(AvailableMaterials[LODIndexToMaterialIndex[PotentialLOD]]);
-
-			if (LandscapeMIC == nullptr)
-			{
-				UMaterialInstanceDynamic* LandscapeMID = Cast<UMaterialInstanceDynamic>(AvailableMaterials[LODIndexToMaterialIndex[PotentialLOD]]);
-
-				if (LandscapeMID != nullptr)
-				{
-					LandscapeMIC = Cast<ULandscapeMaterialInstanceConstant>(LandscapeMID->Parent);
-				}
-			}
-
-			HasTessellationEnabled = MaterialHasTessellationEnabled[LODIndexToMaterialIndex[PotentialLOD]] && (LandscapeMIC != nullptr && !LandscapeMIC->bDisableTessellation);
-		}
-
-		int32 BaseMeshBatchIndex = MaterialIndexToStaticMeshBatchLOD[LODIndexToMaterialIndex[PotentialLOD]];
-
-		if (HasTessellationEnabled)
-		{
-			const int8 ShadowTessellatedMeshBatchLODIndex = 2;
-			const int8 ShadowNonTessellatedMeshBatchLODIndex = 3;
-
-			bool UseTessellationMeshBatch = ScreenSizeSquared >= TessellationComponentSquaredScreenSize * InViewLODScale;
-
-			if (UseTessellationMeshBatch)
-			{
-				if (!CanUseMeshBatchForShadowCascade(ShadowTessellatedMeshBatchLODIndex, InShadowMapTextureResolution, InShadowMapCascadeSize))
-				{
-					UseTessellationMeshBatch = false;
-				}
-			}
-
-			LODToRender.SetLOD(UseTessellationMeshBatch ? BaseMeshBatchIndex + ShadowTessellatedMeshBatchLODIndex : BaseMeshBatchIndex + ShadowNonTessellatedMeshBatchLODIndex);
-		}
-		else
-		{
-			LODToRender.SetLOD(BaseMeshBatchIndex);
-		}
-	}
-
-	return LODToRender;
-}
-
 namespace
 {
 	FLinearColor GetColorForLod(int32 CurrentLOD, int32 ForcedLOD, bool DisplayCombinedBatch)
diff --git a/Engine/Source/Runtime/Landscape/Public/LandscapeRender.h b/Engine/Source/Runtime/Landscape/Public/LandscapeRender.h
--- a/Engine/Source/Runtime/Landscape/Public/LandscapeRender.h
+++ b/Engine/Source/Runtime/Landscape/Public/LandscapeRender.h
@@ -1017,10 +1017,6 @@ public:
 	virtual void CreateRenderThreadResources() override;
 	virtual void DestroyRenderThreadResources() override;
 	virtual void OnLevelAddedToWorld() override;
-	virtual void* InitViewCustomData(const FSceneView& InView, float InViewLODScale, FMemStackBase& InCustomDataMemStack, bool InIsStaticRelevant, bool InIsShadowOnly, const FLODMask* InVisiblePrimitiveLODMask = nullptr, float InMeshScreenSizeSquared = -1.0f) override;
-	void PostInitViewCustomData(const FSceneView& InView, void* InViewCustomData) const;
-	virtual FLODMask GetCustomLOD(const FSceneView& InView, float InViewLODScale, int32 InForcedLODLevel, float& OutScreenSizeSquared) const override;
-	virtual FLODMask GetCustomWholeSceneShadowLOD(const FSceneView& InView, float InViewLODScale, int32 InForcedLODLevel, const struct FLODMask& InVisibilePrimitiveLODMask, float InShadowMapTextureResolution, float InShadowMapCascadeSize, int8 InShadowCascadeId, bool InHasSelfShadow) const override;
 	
 	friend class ULandscapeComponent;
 	friend class FLandscapeVertexFactoryVertexShaderParameters;
diff --git a/Engine/Source/Runtime/Renderer/Private/DeferredShadingRenderer.cpp b/Engine/Source/Runtime/Renderer/Private/DeferredShadingRenderer.cpp
--- a/Engine/Source/Runtime/Renderer/Private/DeferredShadingRenderer.cpp
+++ b/Engine/Source/Runtime/Renderer/Private/DeferredShadingRenderer.cpp
@@ -866,9 +866,11 @@ bool FDeferredShadingSceneRenderer::GatherRayTracingWorldInstances(FRHICommandLi
 					float MeshScreenSizeSquared = 0;
 					if (SceneInfo->bIsUsingCustomLODRules)
 					{
+						PRAGMA_DISABLE_DEPRECATION_WARNINGS
 						FPrimitiveSceneProxy* SceneProxy = Scene->PrimitiveSceneProxies[PrimitiveIndex];
 						LODToRender = SceneProxy->GetCustomLOD(View, View.LODDistanceFactor, ForcedLODLevel, MeshScreenSizeSquared);
 						LODToRender.ClampToFirstLOD(CurFirstLODIdx);
+						PRAGMA_ENABLE_DEPRECATION_WARNINGS
 					}
 					else
 					{
diff --git a/Engine/Source/Runtime/Renderer/Private/PrimitiveSceneInfo.cpp b/Engine/Source/Runtime/Renderer/Private/PrimitiveSceneInfo.cpp
--- a/Engine/Source/Runtime/Renderer/Private/PrimitiveSceneInfo.cpp
+++ b/Engine/Source/Runtime/Renderer/Private/PrimitiveSceneInfo.cpp
@@ -132,8 +132,10 @@ FPrimitiveSceneInfo::FPrimitiveSceneInfo(UPrimitiveComponent* InComponent,FScene
 	LastRenderTime(-FLT_MAX),
 	Scene(InScene),
 	NumMobileMovablePointLights(0),
+PRAGMA_DISABLE_DEPRECATION_WARNINGS
 	bIsUsingCustomLODRules(Proxy->IsUsingCustomLODRules()),
 	bIsUsingCustomWholeSceneShadowLODRules(Proxy->IsUsingCustomWholeSceneShadowLODRules()),
+PRAGMA_ENABLE_DEPRECATION_WARNINGS
 #if RHI_RAYTRACING
 	bDrawInGame(Proxy->IsDrawnInGame()),
 	bShouldRenderInMainPass(InComponent->SceneProxy->ShouldRenderInMainPass()),
diff --git a/Engine/Source/Runtime/Renderer/Private/SceneVisibility.cpp b/Engine/Source/Runtime/Renderer/Private/SceneVisibility.cpp
--- a/Engine/Source/Runtime/Renderer/Private/SceneVisibility.cpp
+++ b/Engine/Source/Runtime/Renderer/Private/SceneVisibility.cpp
@@ -2169,8 +2169,10 @@ struct FRelevancePacket
 
 			if (PrimitiveSceneInfo->bIsUsingCustomLODRules)
 			{
+				PRAGMA_DISABLE_DEPRECATION_WARNINGS
 				LODToRender = PrimitiveSceneInfo->Proxy->GetCustomLOD(View, View.LODDistanceFactor, ViewData.ForcedLODLevel, MeshScreenSizeSquared);
 				LODToRender.ClampToFirstLOD(CurFirstLODIdx);
+				PRAGMA_ENABLE_DEPRECATION_WARNINGS
 			}
 			else
 			{
@@ -2183,12 +2185,14 @@ struct FRelevancePacket
 
 			if (OutHasViewCustomDataMasks[PrimitiveIndex] != 0) // Has a relevance for this view
 			{
+				PRAGMA_DISABLE_DEPRECATION_WARNINGS
 				UserViewCustomData = PrimitiveSceneInfo->Proxy->InitViewCustomData(View, View.LODDistanceFactor, PrimitiveCustomDataMemStack, true, false, &LODToRender, MeshScreenSizeSquared);
 
 				if (UserViewCustomData != nullptr)
 				{
 					PrimitivesCustomData.AddPrim(FRelevancePacket::FViewCustomData(PrimitiveSceneInfo, UserViewCustomData));
 				}
+				PRAGMA_ENABLE_DEPRECATION_WARNINGS
 			}
 
 			const bool bIsHLODFading = HLODState ? HLODState->IsNodeFading(PrimitiveIndex) : false;
@@ -2706,10 +2710,12 @@ static void SetDynamicMeshElementViewCustomData(TArray<FViewInfo>& InViews, cons
 		{
 			FViewInfo& ViewInfo = InViews[ViewIndex];
 
+			PRAGMA_DISABLE_DEPRECATION_WARNINGS
 			if (InHasViewCustomDataMasks[PrimitiveIndex] & (1 << ViewIndex) && ViewInfo.GetCustomData(InPrimitiveSceneInfo->GetIndex()) == nullptr)
 			{
 				ViewInfo.SetCustomData(InPrimitiveSceneInfo, InPrimitiveSceneInfo->Proxy->InitViewCustomData(ViewInfo, ViewInfo.LODDistanceFactor, ViewInfo.GetCustomDataGlobalMemStack(), false, false));
 			}
+			PRAGMA_ENABLE_DEPRECATION_WARNINGS
 		}
 	}
 }
diff --git a/Engine/Source/Runtime/Renderer/Private/ShadowSetup.cpp b/Engine/Source/Runtime/Renderer/Private/ShadowSetup.cpp
--- a/Engine/Source/Runtime/Renderer/Private/ShadowSetup.cpp
+++ b/Engine/Source/Runtime/Renderer/Private/ShadowSetup.cpp
@@ -1129,8 +1129,10 @@ bool FProjectedShadowInfo::ShouldDrawStaticMeshes(FViewInfo& InCurrentView, bool
 
 			if (InPrimitiveSceneInfo->bIsUsingCustomLODRules)
 			{
+				PRAGMA_DISABLE_DEPRECATION_WARNINGS
 				ViewLODToRender = InPrimitiveSceneInfo->Proxy->GetCustomLOD(InCurrentView, InCurrentView.LODDistanceFactor, ForcedLOD, MeshScreenSizeSquared);
 				ViewLODToRender.ClampToFirstLOD(CurFirstLODIdx);
+				PRAGMA_ENABLE_DEPRECATION_WARNINGS
 			}
 			else
 			{
@@ -1145,7 +1147,9 @@ bool FProjectedShadowInfo::ShouldDrawStaticMeshes(FViewInfo& InCurrentView, bool
 			{
 				if (InCurrentView.GetCustomData(InPrimitiveSceneInfo->GetIndex()) == nullptr)
 				{
+					PRAGMA_DISABLE_DEPRECATION_WARNINGS
 					InCurrentView.SetCustomData(InPrimitiveSceneInfo, InPrimitiveSceneInfo->Proxy->InitViewCustomData(InCurrentView, InCurrentView.LODDistanceFactor, InCurrentView.GetCustomDataGlobalMemStack(), true, true, &ViewLODToRender, MeshScreenSizeSquared));
+					PRAGMA_ENABLE_DEPRECATION_WARNINGS
 				}
 			}
 		}
@@ -1188,12 +1192,14 @@ bool FProjectedShadowInfo::ShouldDrawStaticMeshes(FViewInfo& InCurrentView, bool
 
 		if (WholeSceneDirectionalShadow)
 		{
+			PRAGMA_DISABLE_DEPRECATION_WARNINGS
 			if (InPrimitiveSceneInfo->bIsUsingCustomWholeSceneShadowLODRules)
 			{
 				const bool bHasShelfShadow = !bReflectiveShadowmap && !bPreShadow;
 				const float ShadowMapTextureResolution = BorderSize * 2 + ResolutionX; // We assume Shadow Map texture to be squared (current design)
 				ShadowLODToRender = InPrimitiveSceneInfo->Proxy->GetCustomWholeSceneShadowLOD(InCurrentView, InCurrentView.LODDistanceFactor, ForcedLOD, *VisibilePrimitiveLODMask, ShadowMapTextureResolution, ShadowBounds.W * 2.0f, ShadowId, bHasShelfShadow);
 			}
+			PRAGMA_ENABLE_DEPRECATION_WARNINGS
 
 			// Don't cache if it requires per view per mesh state for distance cull fade.
 			const bool bIsPrimitiveDistanceCullFading = InCurrentView.PotentiallyFadingPrimitiveMap[InPrimitiveSceneInfo->GetIndex()];
@@ -2178,6 +2184,7 @@ void FProjectedShadowInfo::GatherDynamicMeshElementsArray(
 
 			if (ViewRelevance.bUseCustomViewData && DependentView != nullptr)
 			{
+				PRAGMA_DISABLE_DEPRECATION_WARNINGS
 				if (DependentView->GetCustomData(PrimitiveSceneInfo->GetIndex()) == nullptr)
 				{
 					void* CustomData = PrimitiveSceneInfo->Proxy->InitViewCustomData(*DependentView, DependentView->LODDistanceFactor, DependentView->GetCustomDataGlobalMemStack(), false, true);
@@ -2186,6 +2193,7 @@ void FProjectedShadowInfo::GatherDynamicMeshElementsArray(
 					// This is required as GetDynamicMeshElements will received ReusedViewsArray which contains only FoundView
 					FoundView->SetCustomData(PrimitiveSceneInfo, CustomData);
 				}
+				PRAGMA_ENABLE_DEPRECATION_WARNINGS
 			}
 
 			PrimitiveSceneInfo->Proxy->GetDynamicMeshElements(ReusedViewsArray, Renderer.ViewFamily, 0x1, Renderer.MeshCollector);
