Mobile Landscape DDC bump

diff --git a/Engine/Source/Runtime/Landscape/Private/Landscape.cpp b/Engine/Source/Runtime/Landscape/Private/Landscape.cpp
--- a/Engine/Source/Runtime/Landscape/Private/Landscape.cpp
+++ b/Engine/Source/Runtime/Landscape/Private/Landscape.cpp
@@ -120,7 +120,7 @@ namespace LandscapeCookStats
 // differences, etc.) replace the version GUID below with a new one.
 // In case of merge conflicts with DDC versions, you MUST generate a new GUID
 // and set this new GUID as the version.                                       
-#define LANDSCAPE_MOBILE_COOK_VERSION TEXT("4CBCBE3102304395B2BF46133642670B")
+#define LANDSCAPE_MOBILE_COOK_VERSION TEXT("A7E61E6494A6450F950333B3BD0DF7A9")
 
 #define LOCTEXT_NAMESPACE "Landscape"
 
