Expose Material Overrides to GrassVariety

diff --git a/Engine/Source/Runtime/Landscape/Classes/LandscapeGrassType.h b/Engine/Source/Runtime/Landscape/Classes/LandscapeGrassType.h
--- a/Engine/Source/Runtime/Landscape/Classes/LandscapeGrassType.h
+++ b/Engine/Source/Runtime/Landscape/Classes/LandscapeGrassType.h
@@ -31,6 +31,9 @@ struct FGrassVariety
 	UPROPERTY(EditAnywhere, Category=Grass)
 	UStaticMesh* GrassMesh;
 
+	UPROPERTY(EditAnywhere, Category=Grass, meta = (ToolTip = "Material Overrides."))
+	TArray<class UMaterialInterface*> OverrideMaterials;
+
 	/* Instances per 10 square meters. */
 	UPROPERTY(EditAnywhere, Category=Grass, meta = (UIMin = 0, ClampMin = 0, UIMax = 1000, ClampMax = 1000))
 	FPerPlatformFloat GrassDensity;
diff --git a/Engine/Source/Runtime/Landscape/Private/LandscapeGrass.cpp b/Engine/Source/Runtime/Landscape/Private/LandscapeGrass.cpp
--- a/Engine/Source/Runtime/Landscape/Private/LandscapeGrass.cpp
+++ b/Engine/Source/Runtime/Landscape/Private/LandscapeGrass.cpp
@@ -2644,7 +2644,8 @@ void ALandscapeProxy::UpdateGrass(const TArray<FVector>& Cameras, int32& InOutNu
 										HierarchicalInstancedStaticMeshComponent->bCastStaticShadow = false;
 										HierarchicalInstancedStaticMeshComponent->CastShadow = GrassVariety.bCastDynamicShadow && !bDisableDynamicShadows;
 										HierarchicalInstancedStaticMeshComponent->bCastDynamicShadow = GrassVariety.bCastDynamicShadow && !bDisableDynamicShadows;
-										
+										HierarchicalInstancedStaticMeshComponent->OverrideMaterials = GrassVariety.OverrideMaterials;
+
 										const FMeshMapBuildData* MeshMapBuildData = Component->GetMeshMapBuildData();
 
 										if (GrassVariety.bUseLandscapeLightmap
