Strip unused code from LandscapeRender

diff --git a/Engine/Source/Runtime/Landscape/Private/LandscapeRender.cpp b/Engine/Source/Runtime/Landscape/Private/LandscapeRender.cpp
--- a/Engine/Source/Runtime/Landscape/Private/LandscapeRender.cpp
+++ b/Engine/Source/Runtime/Landscape/Private/LandscapeRender.cpp
@@ -1735,7 +1735,6 @@ FPrimitiveViewRelevance FLandscapeComponentSceneProxy::GetViewRelevance(const FS
 	Result.bRenderCustomDepth = ShouldRenderCustomDepth();
 	Result.bUsesLightingChannels = GetLightingChannelMask() != GetDefaultLightingChannelMask();
 	Result.bTranslucentSelfShadow = bCastVolumetricTranslucentShadow;
-	Result.bUseCustomViewData = true;
 
 	auto FeatureLevel = View->GetFeatureLevel();
 
@@ -2043,262 +2042,6 @@ void FLandscapeComponentSceneProxy::OnTransformChanged()
 	GetScene().UpdateCachedRenderStates(this);
 }
 
-float FLandscapeComponentSceneProxy::GetComponentScreenSize(const FSceneView* View, const FVector& Origin, float MaxExtend, float ElementRadius) const
-{
-	float SquaredScreenRadius = 0.f;
-
-	if (!View->ViewMatrices.IsPerspectiveProjection())
-	{
-		FMatrix ProjMatrix = View->ViewMatrices.GetProjectionMatrix();
-		const float ScreenMultiple = FMath::Max(0.5f * ProjMatrix.M[0][0], 0.5f * ProjMatrix.M[1][1]);
-		SquaredScreenRadius = FMath::Square(ScreenMultiple * ElementRadius);
-	}
-	else
-	{
-		FVector CameraOrigin = View->ViewMatrices.GetViewOrigin();
-		FMatrix ProjMatrix = View->ViewMatrices.GetProjectionMatrix();
-
-		const FVector OriginToCamera = (CameraOrigin - Origin).GetAbs();
-		const FVector ClosestPoint = OriginToCamera.ComponentMin(FVector(MaxExtend));
-		const float DistSquared = (OriginToCamera - ClosestPoint).SizeSquared();
-
-		// Get projection multiple accounting for view scaling.
-		const float ScreenMultiple = FMath::Max(0.5f * ProjMatrix.M[0][0], 0.5f * ProjMatrix.M[1][1]);
-
-		// Calculate screen-space projected radius
-		SquaredScreenRadius = FMath::Square(ScreenMultiple * ElementRadius) / FMath::Max(1.0f, DistSquared);
-	}
-
-	// If we hit NaN check see UE-64538
-	check(FMath::IsFinite(SquaredScreenRadius));
-	return SquaredScreenRadius * 2.0f;
-}
-
-void FLandscapeComponentSceneProxy::BuildDynamicMeshElement(const FViewCustomDataLOD* InPrimitiveCustomData, bool InToolMesh, bool InHasTessellation, bool InDisableTessellation, FMeshBatch& OutMeshBatch, TArray<FLandscapeBatchElementParams, SceneRenderingAllocator>& OutStaticBatchParamArray) const
-{
-	if (AvailableMaterials.Num() == 0 || InPrimitiveCustomData == nullptr || InPrimitiveCustomData->SubSections.Num() == 0)
-	{
-		return;
-	}
-
-	const int8 CurrentLODIndex = InPrimitiveCustomData->SubSections[0].BatchElementCurrentLOD;
-	int8 MaterialIndex = LODIndexToMaterialIndex.IsValidIndex(CurrentLODIndex) ? LODIndexToMaterialIndex[CurrentLODIndex] : INDEX_NONE;
-	UMaterialInterface* SelectedMaterial = MaterialIndex != INDEX_NONE ? AvailableMaterials[MaterialIndex] : nullptr;
-
-	if (InHasTessellation && MaterialIndex != INDEX_NONE)
-	{		
-		if (InDisableTessellation && MaterialIndexToDisabledTessellationMaterial.IsValidIndex(MaterialIndex))
-		{
-			SelectedMaterial = AvailableMaterials[MaterialIndexToDisabledTessellationMaterial[MaterialIndex]];
-		}
-	}
-
-	// this is really not normal that we have no material at this point, so do not continue
-	if (SelectedMaterial == nullptr)
-	{
-		return;
-	}
-
-	// Could be different from bRequiresAdjacencyInformation during shader compilation
-	bool bCurrentRequiresAdjacencyInformation = !InToolMesh && MaterialRenderingRequiresAdjacencyInformation_RenderingThread(SelectedMaterial, VertexFactory->GetType(), GetScene().GetFeatureLevel());
-
-	if (bCurrentRequiresAdjacencyInformation)
-	{
-		check(SharedBuffers->AdjacencyIndexBuffers);
-	}
-
-	OutMeshBatch.VertexFactory = VertexFactory;
-	OutMeshBatch.MaterialRenderProxy = SelectedMaterial->GetRenderProxy();
-	OutMeshBatch.LCI = ComponentLightInfo.Get();
-	OutMeshBatch.ReverseCulling = IsLocalToWorldDeterminantNegative();
-	OutMeshBatch.CastShadow = InToolMesh ? false : true;
-	OutMeshBatch.bUseAsOccluder = ShouldUseAsOccluder() && GetScene().GetShadingPath() == EShadingPath::Deferred && !IsMovable();
-	OutMeshBatch.bUseForMaterial = true;
-	OutMeshBatch.Type = bCurrentRequiresAdjacencyInformation ? PT_12_ControlPointPatchList : PT_TriangleList;
-	OutMeshBatch.LODIndex = 0;
-
-	OutMeshBatch.Elements.Empty();
-
-	if (NumSubsections > 1 && !InPrimitiveCustomData->UseCombinedMeshBatch)
-	{
-		for (int32 SubY = 0; SubY < NumSubsections; SubY++)
-		{
-			for (int32 SubX = 0; SubX < NumSubsections; SubX++)
-			{
-				const int8 SubSectionIdx = SubX + SubY * NumSubsections;
-				const int8 CurrentLOD = InPrimitiveCustomData->SubSections[SubSectionIdx].BatchElementCurrentLOD;
-
-				FMeshBatchElement BatchElement;
-				FLandscapeBatchElementParams& BatchElementParams = OutStaticBatchParamArray[SubSectionIdx];
-
-				if (!InToolMesh)
-				{
-					BatchElementParams.LandscapeUniformShaderParametersResource = &LandscapeUniformShaderParameters;
-					BatchElementParams.SceneProxy = this;
-					BatchElementParams.CurrentLOD = CurrentLOD;
-				}
-
-				BatchElement.UserData = &BatchElementParams;
-				BatchElement.PrimitiveUniformBuffer = GetUniformBuffer();
-
-				int32 LodSubsectionSizeVerts = (SubsectionSizeVerts >> CurrentLOD);
-				uint32 NumPrimitives = FMath::Square((LodSubsectionSizeVerts - 1)) * 2;
-
-				if (bCurrentRequiresAdjacencyInformation)
-				{
-					check(SharedBuffers->AdjacencyIndexBuffers);
-					BatchElement.IndexBuffer = SharedBuffers->AdjacencyIndexBuffers->IndexBuffers[CurrentLOD];
-					BatchElement.FirstIndex = (SubX + SubY * NumSubsections) * NumPrimitives * 12;
-				}
-				else
-				{
-					BatchElement.IndexBuffer = SharedBuffers->IndexBuffers[CurrentLOD];
-					BatchElement.FirstIndex = (SubX + SubY * NumSubsections) * NumPrimitives * 3;
-				}
-				BatchElement.NumPrimitives = NumPrimitives;
-				BatchElement.MinVertexIndex = SharedBuffers->IndexRanges[CurrentLOD].MinIndex[SubX][SubY];
-				BatchElement.MaxVertexIndex = SharedBuffers->IndexRanges[CurrentLOD].MaxIndex[SubX][SubY];
-
-
-#if !(UE_BUILD_SHIPPING || UE_BUILD_TEST)
-				// We simplify this by considering only the biggest LOD index for this mesh element.
-				OutMeshBatch.VisualizeLODIndex = FMath::Max(OutMeshBatch.VisualizeLODIndex, CurrentLOD);
-#endif
-
-				OutMeshBatch.Elements.Add(BatchElement);
-			}
-		}
-	}
-	else
-	{
-		FMeshBatchElement BatchElement;
-
-		if (InToolMesh)
-		{
-			// Reuse the params for the tool mesh
-			BatchElement.UserData = &OutStaticBatchParamArray[0];
-		}
-		else
-		{
-			FLandscapeBatchElementParams& BatchElementParams = OutStaticBatchParamArray[0];
-			BatchElementParams.LandscapeUniformShaderParametersResource = &LandscapeUniformShaderParameters;
-			BatchElementParams.SceneProxy = this;
-			BatchElementParams.CurrentLOD = CurrentLODIndex;
-
-			BatchElement.UserData = &BatchElementParams;
-		}
-
-		// Combined batch element
-		int32 LodSubsectionSizeVerts = SubsectionSizeVerts >> CurrentLODIndex;
-
-		BatchElement.PrimitiveUniformBuffer = GetUniformBuffer();
-		BatchElement.IndexBuffer = bCurrentRequiresAdjacencyInformation ? SharedBuffers->AdjacencyIndexBuffers->IndexBuffers[CurrentLODIndex] : SharedBuffers->IndexBuffers[CurrentLODIndex];
-		BatchElement.NumPrimitives = FMath::Square((LodSubsectionSizeVerts - 1)) * FMath::Square(NumSubsections) * 2;
-		BatchElement.FirstIndex = 0;
-		BatchElement.MinVertexIndex = SharedBuffers->IndexRanges[CurrentLODIndex].MinIndexFull;
-		BatchElement.MaxVertexIndex = SharedBuffers->IndexRanges[CurrentLODIndex].MaxIndexFull;
-
-		OutMeshBatch.Elements.Add(BatchElement);
-	}
-}
-
-bool FLandscapeComponentSceneProxy::GetMeshElement(bool UseSeperateBatchForShadow, bool ShadowOnly, bool HasTessellation, int8 InLODIndex, UMaterialInterface* InMaterialInterface, FMeshBatch& OutMeshBatch, TArray<FLandscapeBatchElementParams>& OutStaticBatchParamArray) const
-{
-	if (InMaterialInterface == nullptr)
-	{
-		return false;
-	}
-
-	// Could be different from bRequiresAdjacencyInformation during shader compilation
-	bool bCurrentRequiresAdjacencyInformation = RequiresAdjacencyInformation(InMaterialInterface, VertexFactory->GetType(), GetScene().GetFeatureLevel());
-
-	if (bCurrentRequiresAdjacencyInformation)
-	{
-		check(SharedBuffers->AdjacencyIndexBuffers);
-	}
-
-	OutMeshBatch.VertexFactory = VertexFactory;
-	OutMeshBatch.MaterialRenderProxy = InMaterialInterface->GetRenderProxy();
-
-	OutMeshBatch.LCI = ComponentLightInfo.Get();
-	OutMeshBatch.ReverseCulling = IsLocalToWorldDeterminantNegative();
-	OutMeshBatch.CastShadow = UseSeperateBatchForShadow ? ShadowOnly : true;
-	OutMeshBatch.bUseForDepthPass = (UseSeperateBatchForShadow ? !ShadowOnly : true);
-	OutMeshBatch.bUseAsOccluder = (UseSeperateBatchForShadow ? !ShadowOnly : true) && ShouldUseAsOccluder() && GetScene().GetShadingPath() == EShadingPath::Deferred && !IsMovable();
-	OutMeshBatch.bUseForMaterial = (UseSeperateBatchForShadow ? !ShadowOnly : true);
-	OutMeshBatch.Type = bCurrentRequiresAdjacencyInformation ? PT_12_ControlPointPatchList : PT_TriangleList;
-	OutMeshBatch.DepthPriorityGroup = SDPG_World;
-	OutMeshBatch.LODIndex = InLODIndex;
-	OutMeshBatch.bDitheredLODTransition = false;
-
-	int32 BatchElementSize = NumSubsections == 1 ? 1 : MAX_SUBSECTION_COUNT + 1;
-	OutMeshBatch.Elements.Empty(FMath::Max(LastLOD - FirstLOD, 1) * BatchElementSize);
-
-	for (int32 i = FirstLOD; i <= LastLOD; ++i)
-	{
-		int32 LodSubsectionSizeVerts = SubsectionSizeVerts >> i;
-
-		if (NumSubsections > 1 && ForcedLOD < 0)
-		{
-			uint32 NumPrimitivesPerSection = FMath::Square((LodSubsectionSizeVerts - 1)) * 2;
-
-			// Per-subsection batch elements
-			for (int32 SubY = 0; SubY < NumSubsections; SubY++)
-			{
-				for (int32 SubX = 0; SubX < NumSubsections; SubX++)
-				{
-					FLandscapeBatchElementParams* BatchElementParams = new(OutStaticBatchParamArray) FLandscapeBatchElementParams;
-					BatchElementParams->LandscapeUniformShaderParametersResource = &LandscapeUniformShaderParameters;
-					BatchElementParams->SceneProxy = this;
-					BatchElementParams->CurrentLOD = i;
-
-					FMeshBatchElement BatchElement;
-					BatchElement.UserData = BatchElementParams;
-					BatchElement.PrimitiveUniformBuffer = GetUniformBuffer();
-
-					if (bCurrentRequiresAdjacencyInformation)
-					{
-						BatchElement.IndexBuffer = SharedBuffers->AdjacencyIndexBuffers->IndexBuffers[i];
-						BatchElement.FirstIndex = (SubX + SubY * NumSubsections) * NumPrimitivesPerSection * 12;
-					}
-					else
-					{
-						BatchElement.IndexBuffer = SharedBuffers->IndexBuffers[i];
-						BatchElement.FirstIndex = (SubX + SubY * NumSubsections) * NumPrimitivesPerSection * 3;
-					}
-					BatchElement.NumPrimitives = NumPrimitivesPerSection;
-					BatchElement.MinVertexIndex = SharedBuffers->IndexRanges[i].MinIndex[SubX][SubY];
-					BatchElement.MaxVertexIndex = SharedBuffers->IndexRanges[i].MaxIndex[SubX][SubY];
-
-					OutMeshBatch.Elements.Add(BatchElement);
-				}
-			}
-		}
-
-		// Combined batch element
-		FLandscapeBatchElementParams* BatchElementParams = new(OutStaticBatchParamArray) FLandscapeBatchElementParams;
-		BatchElementParams->LandscapeUniformShaderParametersResource = &LandscapeUniformShaderParameters;
-		BatchElementParams->SceneProxy = this;
-		BatchElementParams->CurrentLOD = i;
-
-		FMeshBatchElement BatchElement;
-
-		BatchElement.UserData = BatchElementParams;
-		BatchElement.PrimitiveUniformBuffer = GetUniformBuffer();
-		BatchElement.IndexBuffer = bCurrentRequiresAdjacencyInformation ? SharedBuffers->AdjacencyIndexBuffers->IndexBuffers[i] : SharedBuffers->IndexBuffers[i];
-		BatchElement.NumPrimitives = FMath::Square((LodSubsectionSizeVerts - 1)) * FMath::Square(NumSubsections) * 2;
-		BatchElement.FirstIndex = 0;
-		BatchElement.MinVertexIndex = SharedBuffers->IndexRanges[i].MinIndexFull;
-		BatchElement.MaxVertexIndex = SharedBuffers->IndexRanges[i].MaxIndexFull;
-
-		OutMeshBatch.Elements.Add(BatchElement);
-	}
-
-	OutMeshBatch.Elements.Shrink();
-
-	return true;
-}
-
 /** Creates a mesh batch for virtual texture rendering. Will render a simple fixed grid with combined subsections. */
 bool FLandscapeComponentSceneProxy::GetMeshElementForVirtualTexture(int32 InLodIndex, ERuntimeVirtualTextureMaterialType MaterialType, UMaterialInterface* InMaterialInterface, FMeshBatch& OutMeshBatch, TArray<FLandscapeBatchElementParams>& OutStaticBatchParamArray) const
 {
@@ -2317,7 +2060,6 @@ bool FLandscapeComponentSceneProxy::GetMeshElementForVirtualTexture(int32 InLodI
 	OutMeshBatch.Type = PT_TriangleList;
 	OutMeshBatch.DepthPriorityGroup = SDPG_World;
 	OutMeshBatch.LODIndex = InLodIndex;
-	OutMeshBatch.bRequiresPerElementVisibility = false;
 	OutMeshBatch.bDitheredLODTransition = false;
 	OutMeshBatch.bRenderToVirtualTexture = true;
 	OutMeshBatch.RuntimeVirtualTextureMaterialType = (uint32)MaterialType;
@@ -2493,236 +2235,6 @@ void FLandscapeComponentSceneProxy::DrawStaticElements(FStaticPrimitiveDrawInter
 	check(StaticBatchParamArray.Num() <= TotalBatchCount);
 }
 
-void FLandscapeComponentSceneProxy::CalculateLODFromScreenSize(const FSceneView& InView, float InMeshScreenSizeSquared, float InViewLODScale, int32 InSubSectionIndex, FViewCustomDataLOD& InOutLODData) const
-{
-	// Handle LOD overrides
-	const int32 ViewLodOverride = GetViewLodOverride(InView);
-	float PreferedLOD = (float)ViewLodOverride;
-
-#if WITH_EDITOR || !(UE_BUILD_SHIPPING || UE_BUILD_TEST)
-	const int32 DrawCollisionLodOverride = GetDrawCollisionLodOverride(InView, CollisionResponse, CollisionMipLevel, SimpleCollisionMipLevel);
-	PreferedLOD = DrawCollisionLodOverride >= 0 ? (float)DrawCollisionLodOverride : PreferedLOD;
-#endif
-
-	PreferedLOD = ForcedLOD >= 0 ? (float)ForcedLOD : PreferedLOD;
-
-	int8 MinStreamedLOD = HeightmapTexture ? FMath::Min<int8>(((FTexture2DResource*)HeightmapTexture->Resource)->GetCurrentFirstMip(), FMath::CeilLogTwo(SubsectionSizeVerts) - 1) : 0;
-	MinStreamedLOD = FMath::Min(MinStreamedLOD, MaxLOD); // We can't go above MaxLOD even for texture streaming
-
-	int8 LocalLODBias = LODBias + (int8)GLandscapeMeshLODBias;
-	FViewCustomDataSubSectionLOD& SubSectionLODData = InOutLODData.SubSections[InSubSectionIndex];
-
-	if (PreferedLOD >= 0.0f)
-	{
-		PreferedLOD = FMath::Clamp<float>(PreferedLOD + LocalLODBias, FMath::Max((float)MinStreamedLOD, MinValidLOD), FMath::Min((float)LastLOD, MaxValidLOD));
-	}
-	else
-	{
-		PreferedLOD = FMath::Clamp<float>(ComputeBatchElementCurrentLOD(GetLODFromScreenSize(InMeshScreenSizeSquared, InViewLODScale), InMeshScreenSizeSquared, InViewLODScale) + LocalLODBias, FMath::Max((float)MinStreamedLOD, MinValidLOD), FMath::Min((float)LastLOD, MaxValidLOD));
-	}
-
-	check(PreferedLOD != -1.0f && PreferedLOD <= MaxLOD);
-	SubSectionLODData.fBatchElementCurrentLOD = PreferedLOD;
-	SubSectionLODData.BatchElementCurrentLOD = FMath::FloorToInt(PreferedLOD);
-}
-
-uint64 FLandscapeVertexFactory::GetStaticBatchElementVisibility(const FSceneView& InView, const FMeshBatch* InBatch, const void* InViewCustomData) const
-{
-	const FLandscapeComponentSceneProxy* SceneProxy = ((FLandscapeBatchElementParams*)InBatch->Elements[0].UserData)->SceneProxy;
-	return SceneProxy->GetStaticBatchElementVisibility(InView, InBatch, InViewCustomData);
-}
-
-uint64 FLandscapeComponentSceneProxy::GetStaticBatchElementVisibility(const FSceneView& InView, const FMeshBatch* InBatch, const void* InViewCustomData) const
-{
-	uint64 BatchesToRenderMask = 0;
-
-	SCOPE_CYCLE_COUNTER(STAT_LandscapeStaticDrawLODTime);
-
-	const void* ViewCustomData = InViewCustomData != nullptr ? InViewCustomData : InView.GetCustomData(GetPrimitiveSceneInfo()->GetIndex());
-
-	if (ViewCustomData != nullptr)
-	{
-		const FViewCustomDataLOD* CurrentLODData = (const FViewCustomDataLOD*)ViewCustomData;
-		const auto FeatureLevel = InView.GetFeatureLevel();
-
-		if (FeatureLevel >= ERHIFeatureLevel::SM5)
-		{
-			const int8 CurrentLODIndex = CurrentLODData->SubSections[0].BatchElementCurrentLOD;
-			ULandscapeMaterialInstanceConstant* LandscapeMIC = Cast<ULandscapeMaterialInstanceConstant>(AvailableMaterials[LODIndexToMaterialIndex[CurrentLODIndex]]);
-
-			if (LandscapeMIC == nullptr)
-			{
-				UMaterialInstanceDynamic* LandscapeMID = Cast<UMaterialInstanceDynamic>(AvailableMaterials[LODIndexToMaterialIndex[CurrentLODIndex]]);
-
-				if (LandscapeMID != nullptr)
-				{
-					LandscapeMIC = Cast<ULandscapeMaterialInstanceConstant>(LandscapeMID->Parent);
-				}
-			}
-
-			bool HasTessellationEnabled = MaterialHasTessellationEnabled[LODIndexToMaterialIndex[CurrentLODIndex]] && (LandscapeMIC != nullptr && !LandscapeMIC->bDisableTessellation);
-
-			if (HasTessellationEnabled)
-			{
-				INC_DWORD_STAT(STAT_LandscapeTessellatedComponents);
-			}
-		}
-
-		if (NumSubsections > 1 && !CurrentLODData->UseCombinedMeshBatch)
-		{
-			INC_DWORD_STAT(STAT_LandscapeComponentUsingSubSectionDrawCalls);
-
-			for (int32 SubSectionIndex = 0; SubSectionIndex < MAX_SUBSECTION_COUNT; ++SubSectionIndex)
-			{
-				const FViewCustomDataSubSectionLOD& SubSectionLODData = CurrentLODData->SubSections[SubSectionIndex];
-				check(SubSectionLODData.StaticBatchElementIndexToRender != INDEX_NONE);
-
-				BatchesToRenderMask |= (((uint64)1) << SubSectionLODData.StaticBatchElementIndexToRender);
-				INC_DWORD_STAT(STAT_LandscapeDrawCalls);
-				INC_DWORD_STAT_BY(STAT_LandscapeTriangles, InBatch->Elements[SubSectionLODData.StaticBatchElementIndexToRender].NumPrimitives);
-			}
-		}
-		else
-		{
-			int32 SubSectionIndex = 0;
-			const FViewCustomDataSubSectionLOD& SubSectionLODData = CurrentLODData->SubSections[SubSectionIndex];
-			check(SubSectionLODData.StaticBatchElementIndexToRender != INDEX_NONE);
-
-			BatchesToRenderMask |= (((uint64)1) << SubSectionLODData.StaticBatchElementIndexToRender);
-			INC_DWORD_STAT(STAT_LandscapeDrawCalls);
-			INC_DWORD_STAT_BY(STAT_LandscapeTriangles, InBatch->Elements[SubSectionLODData.StaticBatchElementIndexToRender].NumPrimitives);
-		}
-	}
-
-	INC_DWORD_STAT(STAT_LandscapeComponentRenderPasses);
-
-	return BatchesToRenderMask;
-}
-
-void FLandscapeComponentSceneProxy::CalculateBatchElementLOD(const FSceneView& InView, float InMeshScreenSizeSquared, float InViewLODScale, FViewCustomDataLOD& InOutLODData, bool InForceCombined) const
-{
-	float SquaredViewLODScale = FMath::Square(InViewLODScale);
-
-	check(InMeshScreenSizeSquared >= 0.0f);
-	float ComponentScreenSize = InMeshScreenSizeSquared;
-
-	if (NumSubsections > 1)
-	{
-		InOutLODData.UseCombinedMeshBatch = false; // default to individual batch render
-
-		float SubSectionMaxExtend = ComponentMaxExtend / 2.0f;
-		float SubSectionRadius = FMath::Max(GetBounds().SphereRadius / 2.0f, 1.0f); // clamp radius to avoid any div by 0 later
-		float CombinedScreenRatio = 0.0f;
-		bool AllSubSectionHaveSameScreenSize = true;
-
-		// Compute screen size of each sub section to determine if we should use the combined logic or the individual logic
-		for (int32 SubY = 0; SubY < NumSubsections; ++SubY)
-		{
-			for (int32 SubX = 0; SubX < NumSubsections; ++SubX)
-			{
-				int32 SubSectionIndex = SubX + SubY * NumSubsections;
-				FViewCustomDataSubSectionLOD& SubSectionLODData = InOutLODData.SubSections[SubSectionIndex];
-
-				SubSectionLODData.ScreenSizeSquared = GetComponentScreenSize(&InView, SubSectionScreenSizeTestingPosition[SubSectionIndex], SubSectionMaxExtend, SubSectionRadius);
-				check(SubSectionLODData.ScreenSizeSquared > 0.0f);
-
-				CalculateLODFromScreenSize(InView, SubSectionLODData.ScreenSizeSquared, InViewLODScale, SubSectionIndex, InOutLODData);
-				check(SubSectionLODData.fBatchElementCurrentLOD != -1.0f);
-
-				InOutLODData.ShaderCurrentLOD.Component(SubSectionIndex) = SubSectionLODData.fBatchElementCurrentLOD;
-
-				// Determine if we should use the combined batch or not
-				if (ComponentScreenSize > ComponentSquaredScreenSizeToUseSubSections * SquaredViewLODScale)
-				{
-					if (AllSubSectionHaveSameScreenSize)
-					{
-						float CurrentScreenRadiusSquared = SubSectionLODData.ScreenSizeSquared * SquaredViewLODScale;
-
-						if (CombinedScreenRatio > 0.0f && !FMath::IsNearlyEqual(CombinedScreenRatio, CurrentScreenRadiusSquared, KINDA_SMALL_NUMBER))
-						{
-							AllSubSectionHaveSameScreenSize = false;
-						}
-
-						CombinedScreenRatio += CurrentScreenRadiusSquared;
-
-						if (SubSectionIndex > 0)
-						{
-							CombinedScreenRatio *= 0.5f;
-						}
-					}
-				}
-			}
-		}
-
-		if (!GLandscapeDebugOptions.IsCombinedDisabled() && (AllSubSectionHaveSameScreenSize || GLandscapeDebugOptions.IsCombinedAll() || ForcedLOD != INDEX_NONE || InForceCombined))
-		{
-			InOutLODData.UseCombinedMeshBatch = true;
-
-			int32 MinLOD = FMath::Min(InOutLODData.SubSections[0].BatchElementCurrentLOD, FMath::Min3(InOutLODData.SubSections[1].BatchElementCurrentLOD, InOutLODData.SubSections[2].BatchElementCurrentLOD, InOutLODData.SubSections[3].BatchElementCurrentLOD));
-
-			for (int32 SubSectionIndex = 0; SubSectionIndex < MAX_SUBSECTION_COUNT; ++SubSectionIndex)
-			{
-				InOutLODData.SubSections[SubSectionIndex].BatchElementCurrentLOD = MinLOD;
-			}
-		}
-	}
-	else
-	{
-		const int32 SubSectionIndex = 0;
-		InOutLODData.UseCombinedMeshBatch = true;
-		FViewCustomDataSubSectionLOD& SubSectionLODData = InOutLODData.SubSections[SubSectionIndex];
-
-		SubSectionLODData.ScreenSizeSquared = ComponentScreenSize;
-		CalculateLODFromScreenSize(InView, SubSectionLODData.ScreenSizeSquared, InViewLODScale, SubSectionIndex, InOutLODData);
-		check(SubSectionLODData.fBatchElementCurrentLOD != -1.0f);
-
-		InOutLODData.ShaderCurrentLOD.Component(SubSectionIndex) = SubSectionLODData.fBatchElementCurrentLOD;
-	}
-}
-
-int32 FLandscapeComponentSceneProxy::ConvertBatchElementLODToBatchElementIndex(int8 BatchElementLOD, bool UseCombinedMeshBatch) const
-{
-	int32 BatchElementIndex = BatchElementLOD;
-
-	if (NumSubsections > 1 && ForcedLOD < 0)
-	{
-		BatchElementIndex = BatchElementLOD * (MAX_SUBSECTION_COUNT + 1);
-
-		if (UseCombinedMeshBatch)
-		{
-			BatchElementIndex += MAX_SUBSECTION_COUNT;
-		}
-	}
-
-	return BatchElementIndex;
-}
-
-float FLandscapeComponentSceneProxy::ComputeBatchElementCurrentLOD(int32 InSelectedLODIndex, float InComponentScreenSize, float InViewLODScale) const
-{
-	check(LODScreenRatioSquared.IsValidIndex(InSelectedLODIndex));
-
-	bool LastElement = InSelectedLODIndex == LODScreenRatioSquared.Num() - 1;
-	float CurrentLODScreenRatio = LODScreenRatioSquared[InSelectedLODIndex];
-	float NextLODScreenRatio = LastElement ? 0 : LODScreenRatioSquared[InSelectedLODIndex + 1];
-
-	float LODScreenRatioRange = CurrentLODScreenRatio - NextLODScreenRatio;
-	float ScreenSizeWithLODScale = FMath::Clamp(InComponentScreenSize / InViewLODScale, 0.0f, LODScreenRatioSquared[0]);
-
-	if (ScreenSizeWithLODScale > CurrentLODScreenRatio || ScreenSizeWithLODScale < NextLODScreenRatio)
-	{
-		// Find corresponding LODIndex to appropriately calculate Ratio and apply it to new LODIndex
-		int32 LODFromScreenSize = GetLODFromScreenSize(InComponentScreenSize, InViewLODScale);
-		CurrentLODScreenRatio = LODScreenRatioSquared[LODFromScreenSize];
-		NextLODScreenRatio = LODFromScreenSize == LODScreenRatioSquared.Num() - 1 ? 0 : LODScreenRatioSquared[LODFromScreenSize + 1];
-		LODScreenRatioRange = CurrentLODScreenRatio - NextLODScreenRatio;
-	}
-
-	float CurrentLODRangeRatio = (ScreenSizeWithLODScale - NextLODScreenRatio) / LODScreenRatioRange;
-	float fLOD = (float)InSelectedLODIndex + (1.0f - CurrentLODRangeRatio);
-
-	return fLOD;
-}
-
 int8 FLandscapeComponentSceneProxy::GetLODFromScreenSize(float InScreenSizeSquared, float InViewLODScale) const
 {
 	float FractionalLOD;
@@ -2730,256 +2242,6 @@ int8 FLandscapeComponentSceneProxy::GetLODFromScreenSize(float InScreenSizeSquar
 	return FLandscapeRenderSystem::GetLODFromScreenSize(LODSettings, InScreenSizeSquared, InViewLODScale, FractionalLOD);
 }
 
-void FLandscapeComponentSceneProxy::ComputeTessellationFalloffShaderValues(const FViewCustomDataLOD& InLODData, const FMatrix& InViewProjectionMatrix, float& OutC, float& OutK) const
-{
-	// No Falloff
-	OutC = 1.0f;
-	OutK = 0.0f;
-
-	const auto FeatureLevel = GetScene().GetFeatureLevel();
-	bool HasTessellationEnabled = false;
-
-	if (FeatureLevel >= ERHIFeatureLevel::SM5)
-	{
-		const int8 CurrentLODIndex = InLODData.SubSections[0].BatchElementCurrentLOD;
-		ULandscapeMaterialInstanceConstant* LandscapeMIC = Cast<ULandscapeMaterialInstanceConstant>(AvailableMaterials[LODIndexToMaterialIndex[CurrentLODIndex]]);
-
-		if (LandscapeMIC == nullptr)
-		{
-			UMaterialInstanceDynamic* LandscapeMID = Cast<UMaterialInstanceDynamic>(AvailableMaterials[LODIndexToMaterialIndex[CurrentLODIndex]]);
-
-			if (LandscapeMID != nullptr)
-			{
-				LandscapeMIC = Cast<ULandscapeMaterialInstanceConstant>(LandscapeMID->Parent);
-			}
-		}
-
-		HasTessellationEnabled = MaterialHasTessellationEnabled[LODIndexToMaterialIndex[CurrentLODIndex]] && (LandscapeMIC != nullptr && !LandscapeMIC->bDisableTessellation);
-	}
-
-	if (HasTessellationEnabled && (InLODData.StaticMeshBatchLOD == INDEX_NONE || (InLODData.StaticMeshBatchLOD == 0 || InLODData.StaticMeshBatchLOD == 2))) // Tess batch will be used
-	{
-		if (UseTessellationComponentScreenSizeFalloff)
-		{
-			float MaxTesselationDistance = ComputeBoundsDrawDistance(FMath::Sqrt(TessellationComponentSquaredScreenSize), GetBounds().SphereRadius / 2.0f, InViewProjectionMatrix);
-			float FallOffStartingDistance = FMath::Min(ComputeBoundsDrawDistance(FMath::Sqrt(FMath::Min(FMath::Square(TessellationComponentScreenSizeFalloff), TessellationComponentSquaredScreenSize)), GetBounds().SphereRadius / 2.0f, InViewProjectionMatrix) - MaxTesselationDistance, MaxTesselationDistance);
-
-			// Calculate the falloff using a = C - K * d by sending C & K into the shader
-			OutC = MaxTesselationDistance / (MaxTesselationDistance - FallOffStartingDistance);
-			OutK = -(1 / (-MaxTesselationDistance + FallOffStartingDistance));
-		}
-	}
-}
-
-FVector4 FLandscapeComponentSceneProxy::GetShaderLODBias() const
-{
-	return FVector4(
-		0.0f, // unused
-		0.0f, // unused
-		((FTexture2DResource*)HeightmapTexture->Resource)->GetCurrentFirstMip(),
-		XYOffsetmapTexture ? ((FTexture2DResource*)XYOffsetmapTexture->Resource)->GetCurrentFirstMip() : 0.0f);
-}
-
-FVector4 FLandscapeComponentSceneProxy::GetShaderLODValues(int8 InBatchElementCurrentLOD) const
-{
-	return FVector4(
-		InBatchElementCurrentLOD,
-		0.0f, // unused
-		(float)((SubsectionSizeVerts >> InBatchElementCurrentLOD) - 1),
-		1.f / (float)((SubsectionSizeVerts >> InBatchElementCurrentLOD) - 1));
-}
-
-void FLandscapeComponentSceneProxy::GetShaderCurrentNeighborLOD(const FSceneView& InView, float InBatchElementCurrentLOD, int8 InSubSectionX, int8 InSubSectionY, int8 InCurrentSubSectionIndex, FVector4& OutShaderCurrentNeighborLOD) const
-{
-	for (int32 NeighborIndex = 0; NeighborIndex < NEIGHBOR_COUNT; ++NeighborIndex)
-	{
-		OutShaderCurrentNeighborLOD.Component(NeighborIndex) = GetNeighborLOD(InView, InBatchElementCurrentLOD, NeighborIndex, InSubSectionX, InSubSectionY, InCurrentSubSectionIndex);
-		check(OutShaderCurrentNeighborLOD.Component(NeighborIndex) != -1.0f);
-	}
-}
-
-struct SubSectionData
-{
-	SubSectionData(int8 InSubSectionOffsetX, int8 InSubSectionOffsetY, bool InInsideComponent)
-		: SubSectionOffsetX(InSubSectionOffsetX)
-		, SubSectionOffsetY(InSubSectionOffsetY)
-		, InsideComponent(InInsideComponent)
-	{}
-
-	int8 SubSectionOffsetX;
-	int8 SubSectionOffsetY;
-	bool InsideComponent;
-};
-
-// SubSectionIndex, NeighborIndex
-static const SubSectionData SubSectionValues[4][4] = { { SubSectionData(0, 1, false), SubSectionData(1, 0, false), SubSectionData(1, 0, true), SubSectionData(0, 1, true) },
-{ SubSectionData(0, 1, false), SubSectionData(-1, 0, true), SubSectionData(-1, 0, false), SubSectionData(0, 1, true) },
-{ SubSectionData(0, -1, true), SubSectionData(1, 0, false), SubSectionData(1, 0, true), SubSectionData(0, -1, false) },
-{ SubSectionData(0, -1, true), SubSectionData(-1, 0, true), SubSectionData(-1, 0, false), SubSectionData(0, -1, false) } };
-
-
-float FLandscapeComponentSceneProxy::GetNeighborLOD(const FSceneView& InView, float InBatchElementCurrentLOD, int8 InNeighborIndex, int8 InSubSectionX, int8 InSubSectionY, int8 InCurrentSubSectionIndex) const
-{
-	float NeighborLOD = InBatchElementCurrentLOD;
-
-	// Assume no sub section initialization
-	bool InsideComponent = false;
-	int32 PrimitiveDataIndex = PrimitiveCustomDataIndex;
-	int32 DesiredSubSectionIndex = 0;
-	int32 DesiredSubSectionX = 0;
-	int32 DesiredSubSectionY = 0;
-	int32 CurrentSubSectionIndex = InSubSectionX != INDEX_NONE && InSubSectionY != INDEX_NONE ? InSubSectionX + InSubSectionY * NumSubsections : 0;
-
-	// Handle subsection
-	if (InSubSectionX != INDEX_NONE && InSubSectionY != INDEX_NONE)
-	{
-		const SubSectionData& Data = SubSectionValues[CurrentSubSectionIndex][InNeighborIndex];
-		DesiredSubSectionX = InSubSectionX + Data.SubSectionOffsetX;
-		DesiredSubSectionY = InSubSectionY + Data.SubSectionOffsetY;
-		DesiredSubSectionIndex = DesiredSubSectionX + DesiredSubSectionY * NumSubsections;
-		InsideComponent = Data.InsideComponent;
-	}
-
-	const FLandscapeNeighborInfo* Neighbor = nullptr;
-
-	if (!InsideComponent)
-	{
-		Neighbor = Neighbors[InNeighborIndex];
-
-		if (Neighbor != nullptr)
-		{
-			PrimitiveDataIndex = Neighbor->PrimitiveCustomDataIndex;
-		}
-		else
-		{
-			DesiredSubSectionX = InSubSectionX;
-			DesiredSubSectionY = InSubSectionY;
-			DesiredSubSectionIndex = CurrentSubSectionIndex;
-		}
-	}
-
-	bool ComputeNeighborCustomDataLOD = true;
-
-	void* CustomData = InView.GetCustomData(PrimitiveDataIndex);
-
-	if (CustomData != nullptr)
-	{
-		const FViewCustomDataLOD* LODData = (const FViewCustomDataLOD*)CustomData;
-		// Don't use the custom data for neighbor calculation when it is marked shadow only (ie it is not visible in the view)
-		// See UE-69785 for more information
-		if (!LODData->IsShadowOnly)
-		{
-			ComputeNeighborCustomDataLOD = false;
-			NeighborLOD = FMath::Max(LODData->SubSections[DesiredSubSectionIndex].fBatchElementCurrentLOD, InBatchElementCurrentLOD);
-		}
-	}
-
-	if (ComputeNeighborCustomDataLOD)
-	{
-		FBoxSphereBounds NeighborBounds = GetBounds();
-		float NeighborMaxExtends = ComponentMaxExtend;
-		FLandscapeComponentSceneProxy* NeighborSceneProxy = nullptr;
-
-		if (Neighbor != nullptr)
-		{
-			const ULandscapeComponent* NeighborComponent = Neighbor->GetLandscapeComponent();
-			if (NeighborComponent != nullptr && NeighborComponent->SceneProxy != nullptr)
-			{
-				NeighborSceneProxy = (FLandscapeComponentSceneProxy*)NeighborComponent->SceneProxy;
-				NeighborBounds = NeighborComponent->Bounds;
-				NeighborMaxExtends = NeighborSceneProxy->SubsectionSizeQuads * FMath::Max(NeighborComponent->GetRenderMatrix().GetScaleVector().X, NeighborComponent->GetRenderMatrix().GetScaleVector().Y);
-			}
-		}
-
-		if (NumSubsections > 1)
-		{
-			float NeighborSubsectionMaxExtends = NeighborMaxExtends / 2.0f;
-			FVector ComponentTopLeftCorner = NeighborBounds.Origin - FVector(NeighborSubsectionMaxExtends, NeighborSubsectionMaxExtends, 0.0f);
-
-			FVector SubSectionOrigin = ComponentTopLeftCorner + FVector(NeighborMaxExtends * DesiredSubSectionX, NeighborMaxExtends * DesiredSubSectionY, 0.0f);
-			float MeshBatchScreenSizeSquared = GetComponentScreenSize(&InView, SubSectionOrigin, NeighborSubsectionMaxExtends, NeighborBounds.SphereRadius / 2.0f);
-
-			FViewCustomDataLOD NeighborLODData;
-			if (NeighborSceneProxy != nullptr)
-			{
-				// Needs to pull some data like per-component LOD bias from neighbor scene proxy
-				NeighborSceneProxy->CalculateLODFromScreenSize(InView, MeshBatchScreenSizeSquared, InView.LODDistanceFactor, DesiredSubSectionIndex, NeighborLODData);
-			}
-			else
-			{
-				CalculateLODFromScreenSize(InView, MeshBatchScreenSizeSquared, InView.LODDistanceFactor, DesiredSubSectionIndex, NeighborLODData);
-			}
-
-			const FViewCustomDataSubSectionLOD& SubSectionData = NeighborLODData.SubSections[DesiredSubSectionIndex];
-			check(SubSectionData.fBatchElementCurrentLOD != -1.0f);
-
-			NeighborLOD = FMath::Max(SubSectionData.fBatchElementCurrentLOD, InBatchElementCurrentLOD);
-		}
-		else
-		{
-			float MeshBatchScreenSizeSquared = GetComponentScreenSize(&InView, NeighborBounds.Origin, NeighborMaxExtends, NeighborBounds.SphereRadius);
-
-			FViewCustomDataLOD NeighborLODData;
-			if (NeighborSceneProxy != nullptr)
-			{
-				NeighborSceneProxy->CalculateLODFromScreenSize(InView, MeshBatchScreenSizeSquared, InView.LODDistanceFactor, DesiredSubSectionIndex, NeighborLODData);
-			}
-			else
-			{
-				CalculateLODFromScreenSize(InView, MeshBatchScreenSizeSquared, InView.LODDistanceFactor, DesiredSubSectionIndex, NeighborLODData);
-			}
-
-			FViewCustomDataSubSectionLOD& SubSectionLODData = NeighborLODData.SubSections[0];
-			check(SubSectionLODData.fBatchElementCurrentLOD != -1.0f);
-
-			if (SubSectionLODData.fBatchElementCurrentLOD > InBatchElementCurrentLOD)
-			{
-				NeighborLOD = SubSectionLODData.fBatchElementCurrentLOD;
-			}
-		}
-	}
-
-	return NeighborLOD;
-}
-
-void FLandscapeComponentSceneProxy::ComputeStaticBatchIndexToRender(FViewCustomDataLOD& OutLODData, int32 SubSectionIndex)
-{
-	FViewCustomDataSubSectionLOD& SubSectionLODData = OutLODData.SubSections[SubSectionIndex];
-
-	SubSectionLODData.StaticBatchElementIndexToRender = INDEX_NONE;
-	SubSectionLODData.StaticBatchElementIndexToRender = ConvertBatchElementLODToBatchElementIndex(SubSectionLODData.BatchElementCurrentLOD, OutLODData.UseCombinedMeshBatch) + SubSectionIndex;
-	check(SubSectionLODData.StaticBatchElementIndexToRender != INDEX_NONE);
-}
-
-bool FLandscapeComponentSceneProxy::CanUseMeshBatchForShadowCascade(int8 InLODIndex, float InShadowMapTextureResolution, float InShadowMapCascadeSize) const
-{
-#if 0
-	const FStaticMeshBatch* MeshBatch = nullptr;
-	const TArray<FStaticMeshBatch>& PrimitiveStaticMeshes = GetPrimitiveSceneInfo()->StaticMeshes;
-
-	if (!PrimitiveStaticMeshes.IsValidIndex(InLODIndex))
-	{
-		return false;
-	}
-
-	check(PrimitiveStaticMeshes[InLODIndex].CastShadow);
-	check(InLODIndex == PrimitiveStaticMeshes[InLODIndex].LODIndex);
-	MeshBatch = &PrimitiveStaticMeshes[InLODIndex];
-
-	check(MeshBatch != nullptr);
-
-	if (!MeshBatch->CastShadow || MeshBatch->TessellationDisablingShadowMapMeshSize == 0)
-	{
-		return true;
-	}
-
-	float WorldUnitsForOneTexel = InShadowMapCascadeSize / InShadowMapTextureResolution; // We assume Shadow Map texture to be squared
-	return MeshBatch->TessellationDisablingShadowMapMeshSize >= WorldUnitsForOneTexel * (GShadowMapWorldUnitsToTexelFactor != -1.0f ? GShadowMapWorldUnitsToTexelFactor : 1.0f);
-#endif
-
-	return true;
-}
-
 namespace
 {
 	FLinearColor GetColorForLod(int32 CurrentLOD, int32 ForcedLOD, bool DisplayCombinedBatch)
diff --git a/Engine/Source/Runtime/Landscape/Public/LandscapeRender.h b/Engine/Source/Runtime/Landscape/Public/LandscapeRender.h
--- a/Engine/Source/Runtime/Landscape/Public/LandscapeRender.h
+++ b/Engine/Source/Runtime/Landscape/Public/LandscapeRender.h
@@ -217,8 +217,6 @@ public:
 		UpdateRHI();
 	}
 
-	virtual uint64 GetStaticBatchElementVisibility(const FSceneView& InView, const FMeshBatch* InBatch, const void* InViewCustomData = nullptr) const override;
-
 	/** stream component data bound to this vertex factory */
 	FDataType Data;
 };
@@ -424,7 +422,6 @@ protected:
 	int8					ForcedLOD;
 	int8					LODBias;
 	bool					bRegistered;
-	int32					PrimitiveCustomDataIndex;
 
 	friend class FLandscapeComponentSceneProxy;
 
@@ -436,7 +433,6 @@ public:
 	, ForcedLOD(InForcedLOD)
 	, LODBias(InLODBias)
 	, bRegistered(false)
-	, PrimitiveCustomDataIndex(INDEX_NONE)
 	{
 		//       -Y       
 		//    - - 0 - -   
@@ -773,49 +769,6 @@ class FLandscapeComponentSceneProxy : public FPrimitiveSceneProxy, public FLands
 public:
 	static const int8 MAX_SUBSECTION_COUNT = 2*2;
 
-	// NOTE: CustomData is added in a FMemStack of the render thread, so no destructor will be called on any of the elements
-	struct FViewCustomDataSubSectionLOD
-	{
-		FViewCustomDataSubSectionLOD()
-			: StaticBatchElementIndexToRender(INDEX_NONE)
-			, fBatchElementCurrentLOD(-1.0f)
-			, BatchElementCurrentLOD(INDEX_NONE)
-			, ScreenSizeSquared(-1.0f)
-			, ShaderCurrentNeighborLOD(FVector4(-1.0f, -1.0f, -1.0f, -1.0f))
-		{}
-
-		int8 StaticBatchElementIndexToRender;
-		float fBatchElementCurrentLOD;
-		int8 BatchElementCurrentLOD;
-		float ScreenSizeSquared;
-		FVector4 ShaderCurrentNeighborLOD;
-	};
-
-	// NOTE: CustomData is added in a FMemStack of the render thread, so no destructor will be called on any of the elements
-	struct FViewCustomDataLOD
-	{
-		FViewCustomDataLOD()
-			: StaticMeshBatchLOD(INDEX_NONE)
-			, UseCombinedMeshBatch(true)
-			, IsShadowOnly(false)
-			, ComponentScreenSize(0.0f)
-			, ShaderCurrentLOD(ForceInitToZero)
-			, LodBias(ForceInitToZero)
-			, LodTessellationParams(ForceInitToZero)
-		{}
-
-		int8 StaticMeshBatchLOD;
-		bool UseCombinedMeshBatch;
-		bool IsShadowOnly;
-		float ComponentScreenSize;
-		TStaticArray<FViewCustomDataSubSectionLOD, MAX_SUBSECTION_COUNT> SubSections; // We always have at least 1 subsections
-
-		// Shaders pre calculated params
-		FVector4 ShaderCurrentLOD;
-		FVector4 LodBias;
-		FVector4 LodTessellationParams;
-	};
-
 #if RHI_RAYTRACING
 	struct FLandscapeSectionRayTracingState
 	{
@@ -890,6 +843,7 @@ protected:
 	// Storage for static draw list batch params
 	TArray<FLandscapeBatchElementParams> StaticBatchParamArray;
 
+
 #if WITH_EDITOR
 	// Precomputed grass rendering MeshBatch and per-LOD params
 	FMeshBatch                           GrassMeshBatch;
@@ -977,29 +931,13 @@ protected:
 	virtual ~FLandscapeComponentSceneProxy();
 	
 	virtual const ULandscapeComponent* GetLandscapeComponent() const { return LandscapeComponent; }
-	FORCEINLINE void ComputeTessellationFalloffShaderValues(const FViewCustomDataLOD& InLODData, const FMatrix& InViewProjectionMatrix, float& OutC, float& OutK) const;
-	bool CanUseMeshBatchForShadowCascade(int8 InLODIndex, float InShadowMapTextureResolution, float InShadowMapCascadeSize) const;
-	FORCEINLINE int32 ConvertBatchElementLODToBatchElementIndex(int8 InBatchElementLOD, bool InUseCombinedMeshBatch) const;
-	float GetNeighborLOD(const FSceneView& InView, float InBatchElementCurrentLOD, int8 InNeighborIndex, int8 InSubSectionX, int8 InSubSectionY, int8 InCurrentSubSectionIndex) const;
-	void CalculateBatchElementLOD(const FSceneView& InView, float InMeshScreenSizeSquared, float InViewLODScale, FViewCustomDataLOD& InOutLODData, bool InForceCombined) const;
-	void CalculateLODFromScreenSize(const FSceneView& InView, float InMeshScreenSizeSquared, float InViewLODScale, int32 InSubSectionIndex, FViewCustomDataLOD& InOutLODData) const;
-	FORCEINLINE void ComputeStaticBatchIndexToRender(FViewCustomDataLOD& OutLODData, int32 InSubSectionIndex);
 	int8 GetLODFromScreenSize(float InScreenSizeSquared, float InViewLODScale) const;
-	FORCEINLINE_DEBUGGABLE float ComputeBatchElementCurrentLOD(int32 InSelectedLODIndex, float InComponentScreenSize, float InViewLODScale) const;
 
-	FORCEINLINE void GetShaderCurrentNeighborLOD(const FSceneView& InView, float InBatchElementCurrentLOD, int8 InSubSectionX, int8 InSubSectionY, int8 InCurrentSubSectionIndex, FVector4& OutShaderCurrentNeighborLOD) const;
-	FORCEINLINE FVector4 GetShaderLODBias() const;
-	FVector4 GetShaderLODValues(int8 BatchElementCurrentLOD) const;
-
-	bool GetMeshElement(bool UseSeperateBatchForShadow, bool ShadowOnly, bool HasTessellation, int8 InLODIndex, UMaterialInterface* InMaterialInterface, FMeshBatch& OutMeshBatch, TArray<FLandscapeBatchElementParams>& OutStaticBatchParamArray) const;
 	bool GetMeshElementForVirtualTexture(int32 InLodIndex, ERuntimeVirtualTextureMaterialType MaterialType, UMaterialInterface* InMaterialInterface, FMeshBatch& OutMeshBatch, TArray<FLandscapeBatchElementParams>& OutStaticBatchParamArray) const;
 	template<class ArrayType> bool GetStaticMeshElement(int32 LODIndex, bool bForToolMesh, bool bForcedLOD, FMeshBatch& MeshBatch, ArrayType& OutStaticBatchParamArray) const;
-	void BuildDynamicMeshElement(const FViewCustomDataLOD* InPrimitiveCustomData, bool InToolMesh, bool InHasTessellation, bool InDisableTessellation, FMeshBatch& OutMeshBatch, TArray<FLandscapeBatchElementParams, SceneRenderingAllocator>& OutStaticBatchParamArray) const;
-
+	
 	virtual void ApplyMeshElementModifier(FMeshBatchElement& InOutMeshElement, int32 InLodIndex) const {}
 
-	float GetComponentScreenSize(const class FSceneView* View, const FVector& Origin,  float MaxExtend, float ElementRadius) const;
-
 public:
 	// constructor
 	FLandscapeComponentSceneProxy(ULandscapeComponent* InComponent);
@@ -1028,8 +966,6 @@ public:
 	friend class FLandscapeFixedGridVertexFactoryVertexShaderParameters;
 	friend class FLandscapeFixedGridVertexFactoryMobileVertexShaderParameters;
 
-	// FLandscapeComponentSceneProxy interface.
-	uint64 GetStaticBatchElementVisibility(const FSceneView& InView, const FMeshBatch* InBatch, const void* InViewCustomData) const;
 #if WITH_EDITOR
 	const FMeshBatch& GetGrassMeshBatch() const { return GrassMeshBatch; }
 #endif
