Change DDC key used for compressed blob of mobile landscape data.

diff --git a/Engine/Source/Runtime/Landscape/Private/Landscape.cpp b/Engine/Source/Runtime/Landscape/Private/Landscape.cpp
--- a/Engine/Source/Runtime/Landscape/Private/Landscape.cpp
+++ b/Engine/Source/Runtime/Landscape/Private/Landscape.cpp
@@ -3352,7 +3352,7 @@ ULandscapeWeightmapUsage::ULandscapeWeightmapUsage(const FObjectInitializer& Obj
 }
 
 // Generate a new guid to force a recache of all landscape derived data
-#define LANDSCAPE_FULL_DERIVEDDATA_VER			TEXT("89D752865B0642E89CC8A5A2FED808A2")
+#define LANDSCAPE_FULL_DERIVEDDATA_VER			TEXT("746F41F382904AC9BA4C6449C3BF4A38")
 
 FString FLandscapeComponentDerivedData::GetDDCKeyString(const FGuid& StateId)
 {
