Replace landscape normal texture sampler with static bilinear sampler.
This fixes broken landscape visuals.

diff --git a/Engine/Source/Runtime/Landscape/Private/LandscapeRender.cpp b/Engine/Source/Runtime/Landscape/Private/LandscapeRender.cpp
--- a/Engine/Source/Runtime/Landscape/Private/LandscapeRender.cpp
+++ b/Engine/Source/Runtime/Landscape/Private/LandscapeRender.cpp
@@ -2020,7 +2020,7 @@ void FLandscapeComponentSceneProxy::OnTransformChanged()
 	if (NormalmapTexture)
 	{
 		LandscapeParams.NormalmapTexture = NormalmapTexture->TextureReference.TextureReferenceRHI;
-		LandscapeParams.NormalmapTextureSampler = NormalmapTexture->Resource->SamplerStateRHI;
+		LandscapeParams.NormalmapTextureSampler = TStaticSamplerState<SF_Bilinear>::GetRHI();
 	}
 	else
 	{
