Remove unneeded chaos header

diff --git a/Engine/Source/Runtime/Landscape/Classes/LandscapeHeightfieldCollisionComponent.h b/Engine/Source/Runtime/Landscape/Classes/LandscapeHeightfieldCollisionComponent.h
--- a/Engine/Source/Runtime/Landscape/Classes/LandscapeHeightfieldCollisionComponent.h
+++ b/Engine/Source/Runtime/Landscape/Classes/LandscapeHeightfieldCollisionComponent.h
@@ -12,10 +12,6 @@
 #include "Components/PrimitiveComponent.h"
 #include "Serialization/BulkData.h"
 
-#define TEMP_HEADER_CHAOS_LEVEL_1
-#include "Chaos/HeightField.h"
-#undef TEMP_HEADER_CHAOS_LEVEL_1
-
 #include "LandscapeHeightfieldCollisionComponent.generated.h"
 
 class ALandscapeProxy;
@@ -28,6 +24,11 @@ struct FConvexVolume;
 struct FEngineShowFlags;
 struct FNavigableGeometryExport;
 
+namespace Chaos
+{
+class FHeightField;
+}
+
 #if WITH_PHYSX
 namespace physx
 {
diff --git a/Engine/Source/Runtime/Landscape/Private/LandscapeCollision.cpp b/Engine/Source/Runtime/Landscape/Private/LandscapeCollision.cpp
--- a/Engine/Source/Runtime/Landscape/Private/LandscapeCollision.cpp
+++ b/Engine/Source/Runtime/Landscape/Private/LandscapeCollision.cpp
@@ -54,6 +54,7 @@
 #include "Chaos/ParticleHandle.h"
 #include "Chaos/Vector.h"
 #include "Chaos/Core.h"
+#include "Chaos/HeightField.h"
 #endif
 
 using namespace PhysicsInterfaceTypes;
