Temp reinclude to fix CIS

diff --git a/Engine/Source/Runtime/Landscape/Classes/LandscapeHeightfieldCollisionComponent.h b/Engine/Source/Runtime/Landscape/Classes/LandscapeHeightfieldCollisionComponent.h
--- a/Engine/Source/Runtime/Landscape/Classes/LandscapeHeightfieldCollisionComponent.h
+++ b/Engine/Source/Runtime/Landscape/Classes/LandscapeHeightfieldCollisionComponent.h
@@ -12,6 +12,11 @@
 #include "Components/PrimitiveComponent.h"
 #include "Serialization/BulkData.h"
 
+#if WITH_CHAOS
+//TODO: remove this, just putting back in to unblock CIS
+#include "Chaos/HeightField.h"
+#endif
+
 #include "LandscapeHeightfieldCollisionComponent.generated.h"
 
 class ALandscapeProxy;
