Remove unneeded header. Removed default initialize that was causing clang issues

diff --git a/Engine/Source/Runtime/Landscape/Classes/LandscapeHeightfieldCollisionComponent.h b/Engine/Source/Runtime/Landscape/Classes/LandscapeHeightfieldCollisionComponent.h
--- a/Engine/Source/Runtime/Landscape/Classes/LandscapeHeightfieldCollisionComponent.h
+++ b/Engine/Source/Runtime/Landscape/Classes/LandscapeHeightfieldCollisionComponent.h
@@ -12,11 +12,6 @@
 #include "Components/PrimitiveComponent.h"
 #include "Serialization/BulkData.h"
 
-#if WITH_CHAOS
-//TODO: remove this, just putting back in to unblock CIS
-#include "Chaos/HeightField.h"
-#endif
-
 #include "LandscapeHeightfieldCollisionComponent.generated.h"
 
 class ALandscapeProxy;
@@ -109,10 +104,10 @@ class ULandscapeHeightfieldCollisionComponent : public UPrimitiveComponent
 
 #if WITH_CHAOS
 		TArray<Chaos::FMaterialHandle> UsedChaosMaterials;
-		TUniquePtr<Chaos::FHeightField> Heightfield = nullptr;
-	    TUniquePtr<Chaos::FHeightField> HeightfieldSimple = nullptr;
+		TUniquePtr<Chaos::FHeightField> Heightfield;
+	    TUniquePtr<Chaos::FHeightField> HeightfieldSimple;
 #if WITH_EDITOR
-		TUniquePtr<Chaos::FHeightField> EditorHeightfield = nullptr;
+		TUniquePtr<Chaos::FHeightField> EditorHeightfield;
 #endif
 #endif
 
