Fix missing includes

diff --git a/Engine/Source/Runtime/Landscape/Private/LandscapeEditLayers.cpp b/Engine/Source/Runtime/Landscape/Private/LandscapeEditLayers.cpp
--- a/Engine/Source/Runtime/Landscape/Private/LandscapeEditLayers.cpp
+++ b/Engine/Source/Runtime/Landscape/Private/LandscapeEditLayers.cpp
@@ -41,6 +41,7 @@ LandscapeEditLayers.cpp: Landscape editing layers mode
 #include "Misc/UObjectToken.h"
 #include "Misc/ScopedSlowTask.h"
 #endif
+#include "Logging/MessageLog.h"
 
 #define LOCTEXT_NAMESPACE "Landscape"
 
