ALandscapeProxy: Don't allow edition of properties that are shared with the parent landscape properties

diff --git a/Engine/Source/Runtime/Landscape/Classes/LandscapeProxy.h b/Engine/Source/Runtime/Landscape/Classes/LandscapeProxy.h
--- a/Engine/Source/Runtime/Landscape/Classes/LandscapeProxy.h
+++ b/Engine/Source/Runtime/Landscape/Classes/LandscapeProxy.h
@@ -934,6 +934,7 @@ public:
 #if WITH_EDITOR
 	LANDSCAPE_API void CreateSplineComponent(const FVector& Scale3D);
 
+	virtual bool CanEditChange(const FProperty* InProperty) const override;
 	virtual void PostEditChangeProperty(FPropertyChangedEvent& PropertyChangedEvent) override;
 	virtual void PostEditImport() override;
 	//~ End UObject Interface
diff --git a/Engine/Source/Runtime/Landscape/Private/LandscapeEdit.cpp b/Engine/Source/Runtime/Landscape/Private/LandscapeEdit.cpp
--- a/Engine/Source/Runtime/Landscape/Private/LandscapeEdit.cpp
+++ b/Engine/Source/Runtime/Landscape/Private/LandscapeEdit.cpp
@@ -4505,6 +4505,38 @@ namespace
 	}
 };
 
+bool ALandscapeProxy::CanEditChange(const FProperty* InProperty) const
+{
+	if (!Super::CanEditChange(InProperty))
+	{
+		return false;
+	}
+
+	FName PropertyName = InProperty ? InProperty->GetFName() : NAME_None;
+
+	// Don't allow edition of properties that are shared with the parent landscape properties
+	// See  ALandscapeProxy::FixupSharedData(ALandscape* Landscape)
+	if (GetLandscapeActor())
+	{
+		if (PropertyName == GET_MEMBER_NAME_CHECKED(ALandscapeProxy, MaxLODLevel) ||
+			PropertyName == GET_MEMBER_NAME_CHECKED(ALandscapeProxy, TessellationComponentScreenSize) ||
+			PropertyName == GET_MEMBER_NAME_CHECKED(ALandscapeProxy, ComponentScreenSizeToUseSubSections) ||
+			PropertyName == GET_MEMBER_NAME_CHECKED(ALandscapeProxy, UseTessellationComponentScreenSizeFalloff) ||
+			PropertyName == GET_MEMBER_NAME_CHECKED(ALandscapeProxy, TessellationComponentScreenSizeFalloff) ||
+			PropertyName == GET_MEMBER_NAME_CHECKED(ALandscapeProxy, LODDistributionSetting) ||
+			PropertyName == GET_MEMBER_NAME_CHECKED(ALandscapeProxy, LOD0DistributionSetting) ||
+			PropertyName == GET_MEMBER_NAME_CHECKED(ALandscapeProxy, LOD0ScreenSize) ||
+			PropertyName == GET_MEMBER_NAME_CHECKED(ALandscapeProxy, OccluderGeometryLOD) ||
+			PropertyName == GET_MEMBER_NAME_CHECKED(ALandscapeProxy, TargetDisplayOrder) ||
+			PropertyName == GET_MEMBER_NAME_CHECKED(ALandscapeProxy, TargetDisplayOrderList))
+		{
+			return false;
+		}
+	}
+
+	return true;
+}
+
 void ALandscapeProxy::PostEditChangeProperty(FPropertyChangedEvent& PropertyChangedEvent)
 {
 	const FName PropertyName = PropertyChangedEvent.MemberProperty ? PropertyChangedEvent.MemberProperty->GetFName() : NAME_None;
