Fixed bug in Landscape BP Brush where Tick was never processing an update of all landscape height/weight (including collision) after an undo.

diff --git a/Engine/Source/Runtime/Landscape/Private/LandscapeBlueprintBrushBase.cpp b/Engine/Source/Runtime/Landscape/Private/LandscapeBlueprintBrushBase.cpp
--- a/Engine/Source/Runtime/Landscape/Private/LandscapeBlueprintBrushBase.cpp
+++ b/Engine/Source/Runtime/Landscape/Private/LandscapeBlueprintBrushBase.cpp
@@ -79,7 +79,7 @@ void ALandscapeBlueprintBrushBase::Tick(float DeltaSeconds)
 	// Wait until we didn't trigger any more landscape update requests (padding of a couple of frames)
 	if (OwningLandscape != nullptr && 
 		LastRequestLayersContentUpdateFrameNumber != InvalidLastRequestLayersContentUpdateFrameNumber &&
-		LastRequestLayersContentUpdateFrameNumber + CVarLandscapeBrushPadding.GetValueOnAnyThread() == GFrameNumber)
+		LastRequestLayersContentUpdateFrameNumber + CVarLandscapeBrushPadding.GetValueOnAnyThread() <= GFrameNumber)
 	{
 		uint32 ModeMask = 0;
 		if (AffectHeightmap)
@@ -94,6 +94,7 @@ void ALandscapeBlueprintBrushBase::Tick(float DeltaSeconds)
 		{
 			OwningLandscape->RequestLayersContentUpdateForceAll((ELandscapeLayerUpdateMode)ModeMask);
 		}
+		LastRequestLayersContentUpdateFrameNumber = InvalidLastRequestLayersContentUpdateFrameNumber;
 	}
 #endif
 
