Fixed assert with grass components trying to be detached while not attached after undoing a location change

diff --git a/Engine/Source/Runtime/Landscape/Private/LandscapeGrass.cpp b/Engine/Source/Runtime/Landscape/Private/LandscapeGrass.cpp
--- a/Engine/Source/Runtime/Landscape/Private/LandscapeGrass.cpp
+++ b/Engine/Source/Runtime/Landscape/Private/LandscapeGrass.cpp
@@ -2295,6 +2295,13 @@ void ALandscapeProxy::UpdateGrass(const TArray<FVector>& Cameras, int32& InOutNu
 {
 	SCOPE_CYCLE_COUNTER(STAT_GrassUpdate);
 	TRACE_CPUPROFILER_EVENT_SCOPE(ALandscapeProxy::UpdateGrass);
+
+	// don't let it create grass actors inside a transaction: 
+	if (GUndo != nullptr)
+	{
+		return;
+	}
+
 	if (GFrameNumberLastStaleCheck != GFrameNumber && GIgnoreExcludeBoxes == 0)
 	{
 		GFrameNumberLastStaleCheck = GFrameNumber;
