Use null vertex buffers for landscape ray tracing geometry initialization to save some memory on startup

diff --git a/Engine/Source/Runtime/Landscape/Private/LandscapeRender.cpp b/Engine/Source/Runtime/Landscape/Private/LandscapeRender.cpp
--- a/Engine/Source/Runtime/Landscape/Private/LandscapeRender.cpp
+++ b/Engine/Source/Runtime/Landscape/Private/LandscapeRender.cpp
@@ -1644,21 +1644,16 @@ void FLandscapeComponentSceneProxy::CreateRenderThreadResources()
 			{
 				const int8 SubSectionIdx = SubX + SubY * NumSubsections;
 
-				int32 LodSubsectionSizeVerts = (SubsectionSizeVerts >> 0);
-				uint32 NumPrimitives = FMath::Square((LodSubsectionSizeVerts - 1)) * 2;
-
 				FRayTracingGeometryInitializer Initializer;
 				FRHIResourceCreateInfo CreateInfo;
 				Initializer.IndexBuffer = nullptr;
-				Initializer.TotalPrimitiveCount = NumPrimitives;
 				Initializer.GeometryType = RTGT_Triangles;
 				Initializer.bFastBuild = true;
 				Initializer.bAllowUpdate = true;
 				FRayTracingGeometrySegment Segment;
-				Segment.VertexBuffer = RHICreateVertexBuffer(sizeof(FVector4) * NumPrimitives * 3, BUF_UnorderedAccess | BUF_ShaderResource, CreateInfo);
+				Segment.VertexBuffer = nullptr;
 				Segment.VertexBufferStride = sizeof(FVector);
 				Segment.VertexBufferElementType = VET_Float3;
-				Segment.NumPrimitives = NumPrimitives;
 				Initializer.Segments.Add(Segment);
 				SectionRayTracingStates[SubSectionIdx].Geometry.SetInitializer(Initializer);
 				SectionRayTracingStates[SubSectionIdx].Geometry.InitResource();
