fix CIS - avoid using editor only method

diff --git a/Engine/Source/Runtime/Landscape/Private/LandscapeCollision.cpp b/Engine/Source/Runtime/Landscape/Private/LandscapeCollision.cpp
--- a/Engine/Source/Runtime/Landscape/Private/LandscapeCollision.cpp
+++ b/Engine/Source/Runtime/Landscape/Private/LandscapeCollision.cpp
@@ -2689,7 +2689,11 @@ void ALandscapeProxy::GetHeightValues(int32& SizeX, int32& SizeY, TArray<float>
 			return;
 		}
 
-		LandscapeComponent->GetComponentExtent(MinX, MinY, MaxX, MaxY);
+		// #todo(dmp): should we be using ULandscapeHeightfieldCollisionComponent.CollisionSizeQuads (or HeightFieldData->GetNumCols)
+		MinX = FMath::Min(LandscapeComponent->SectionBaseX, MinX);
+		MinY = FMath::Min(LandscapeComponent->SectionBaseY, MinY);
+		MaxX = FMath::Max(LandscapeComponent->SectionBaseX + LandscapeComponent->ComponentSizeQuads, MaxX);
+		MaxY = FMath::Max(LandscapeComponent->SectionBaseY + LandscapeComponent->ComponentSizeQuads, MaxY);		
 	}
 
 	if (MinX == MAX_int32)
