Foliage instances snapping was dirtying/invalidating foliage component's lighting cache even when result didn't change

diff --git a/Engine/Source/Runtime/Landscape/Private/LandscapeCollision.cpp b/Engine/Source/Runtime/Landscape/Private/LandscapeCollision.cpp
--- a/Engine/Source/Runtime/Landscape/Private/LandscapeCollision.cpp
+++ b/Engine/Source/Runtime/Landscape/Private/LandscapeCollision.cpp
@@ -1720,7 +1720,8 @@ void ULandscapeHeightfieldCollisionComponent::SnapFoliageInstances(const FBox& I
 
 				TArray<int32> InstancesToRemove;
 				TSet<UHierarchicalInstancedStaticMeshComponent*> AffectedFoliageComponents;
-
+				
+				bool bIsMeshInfoDirty = false;
 				for (int32 InstanceIndex : *InstanceSet)
 				{
 					FFoliageInstance& Instance = MeshInfo.Instances[InstanceIndex];
@@ -1751,6 +1752,8 @@ void ULandscapeHeightfieldCollisionComponent::SnapFoliageInstances(const FBox& I
 								{
 									IFA->Modify();
 
+									bIsMeshInfoDirty = true;
+
 									// Remove instance location from the hash. Do not need to update ComponentHash as we re-add below.
 									MeshInfo.InstanceHash->RemoveInstance(Instance.Location, InstanceIndex);
 
@@ -1784,13 +1787,13 @@ void ULandscapeHeightfieldCollisionComponent::SnapFoliageInstances(const FBox& I
 						{
 							// Couldn't find new spot - remove instance
 							InstancesToRemove.Add(InstanceIndex);
+							bIsMeshInfoDirty = true;
 						}
 
-						if (MeshInfo.GetComponent() != nullptr)
+						if (bIsMeshInfoDirty && (MeshInfo.GetComponent() != nullptr))
 						{
 							AffectedFoliageComponents.Add(MeshInfo.GetComponent());
 						}
-						
 					}
 				}
 
