- Grass Maps updated during Tick are no longer dirtying packages. Instead, they mark themselves as dirty/need rebuilt.
- New map warning will display "GRASS MAPS NEED TO BE REBUILT" whenever grass maps are outdated.
- New "Build Grass Maps" command added to the Build menu options.

diff --git a/Engine/Source/Editor/LevelEditor/Private/LevelEditor.cpp b/Engine/Source/Editor/LevelEditor/Private/LevelEditor.cpp
--- a/Engine/Source/Editor/LevelEditor/Private/LevelEditor.cpp
+++ b/Engine/Source/Editor/LevelEditor/Private/LevelEditor.cpp
@@ -1536,6 +1536,9 @@ void FLevelEditorModule::BindGlobalLevelEditorCommands()
 	ActionList.MapAction(Commands.BuildVirtualTextureOnly,
 		FExecuteAction::CreateStatic(&FLevelEditorActionCallbacks::BuildVirtualTextureOnly_Execute));
 
+	ActionList.MapAction(Commands.BuildGrassMapsOnly,
+		FExecuteAction::CreateStatic(&FLevelEditorActionCallbacks::BuildGrassMapsOnly_Execute));
+
 	ActionList.MapAction( 
 		Commands.LightingQuality_Production, 
 		FExecuteAction::CreateStatic( &FLevelEditorActionCallbacks::SetLightingQuality, (ELightingBuildQuality)Quality_Production ),
diff --git a/Engine/Source/Editor/LevelEditor/Private/LevelEditorActions.cpp b/Engine/Source/Editor/LevelEditor/Private/LevelEditorActions.cpp
--- a/Engine/Source/Editor/LevelEditor/Private/LevelEditorActions.cpp
+++ b/Engine/Source/Editor/LevelEditor/Private/LevelEditorActions.cpp
@@ -777,6 +777,12 @@ void FLevelEditorActionCallbacks::BuildVirtualTextureOnly_Execute()
 	GEngine->DeferredCommands.AddUnique(TEXT("MAP CHECK NOTIFYRESULTS"));
 }
 
+void FLevelEditorActionCallbacks::BuildGrassMapsOnly_Execute()
+{
+	// Build Grass maps
+	FEditorBuildUtils::EditorBuild(GetWorld(), FBuildOptions::BuildGrassMaps);
+}
+
 bool FLevelEditorActionCallbacks::IsLightingQualityChecked( ELightingBuildQuality TestQuality )
 {
 	int32 CurrentQualityLevel;
@@ -3093,7 +3099,8 @@ void FLevelEditorCommands::RegisterCommands()
 	UI_COMMAND( BuildPathsOnly, "Build Paths", "Only builds paths (all levels.)", EUserInterfaceActionType::Button, FInputChord() );
 	UI_COMMAND( BuildLODsOnly, "Build LODs", "Only builds LODs (all levels.)", EUserInterfaceActionType::Button, FInputChord() );
 	UI_COMMAND( BuildTextureStreamingOnly, "Build Texture Streaming", "Build texture streaming data", EUserInterfaceActionType::Button, FInputChord() );
-	UI_COMMAND(BuildVirtualTextureOnly, "Build Virtual Textures", "Build runtime virtual texture low mips streaming data", EUserInterfaceActionType::Button, FInputChord());
+	UI_COMMAND( BuildVirtualTextureOnly, "Build Virtual Textures", "Build runtime virtual texture low mips streaming data", EUserInterfaceActionType::Button, FInputChord());
+	UI_COMMAND( BuildGrassMapsOnly, "Build Grass Maps", "Build landscape grass maps", EUserInterfaceActionType::Button, FInputChord());
 
 	UI_COMMAND( LightingQuality_Production, "Production", "Sets precomputed lighting quality to highest possible quality (slowest computation time.)", EUserInterfaceActionType::RadioButton, FInputChord() );
 	UI_COMMAND( LightingQuality_High, "High", "Sets precomputed lighting quality to high quality", EUserInterfaceActionType::RadioButton, FInputChord() );
diff --git a/Engine/Source/Editor/LevelEditor/Private/LevelEditorToolBar.cpp b/Engine/Source/Editor/LevelEditor/Private/LevelEditorToolBar.cpp
--- a/Engine/Source/Editor/LevelEditor/Private/LevelEditorToolBar.cpp
+++ b/Engine/Source/Editor/LevelEditor/Private/LevelEditorToolBar.cpp
@@ -1873,6 +1873,11 @@ void FLevelEditorToolBar::RegisterBuildMenu()
 		Section.AddMenuEntry(FLevelEditorCommands::Get().BuildVirtualTextureOnly);
 	}
 
+	{
+		FToolMenuSection& Section = Menu->AddSection("LevelEditorGrass", LOCTEXT("GrassHeading", "Grass Maps"));
+		Section.AddMenuEntry(FLevelEditorCommands::Get().BuildGrassMapsOnly);
+	}
+
 	
 	{
 		FToolMenuSection& Section = Menu->AddSection("LevelEditorAutomation", LOCTEXT( "AutomationHeading", "Automation" ) );
diff --git a/Engine/Source/Editor/LevelEditor/Public/LevelEditorActions.h b/Engine/Source/Editor/LevelEditor/Public/LevelEditorActions.h
--- a/Engine/Source/Editor/LevelEditor/Public/LevelEditorActions.h
+++ b/Engine/Source/Editor/LevelEditor/Public/LevelEditorActions.h
@@ -88,6 +88,7 @@ public:
 	TSharedPtr< FUICommandInfo > BuildLODsOnly;
 	TSharedPtr< FUICommandInfo > BuildTextureStreamingOnly;
 	TSharedPtr< FUICommandInfo > BuildVirtualTextureOnly;
+	TSharedPtr< FUICommandInfo > BuildGrassMapsOnly;
 	TSharedPtr< FUICommandInfo > LightingQuality_Production;
 	TSharedPtr< FUICommandInfo > LightingQuality_High;
 	TSharedPtr< FUICommandInfo > LightingQuality_Medium;
@@ -772,6 +773,7 @@ public:
 	static void BuildLODsOnly_Execute();
 	static void BuildTextureStreamingOnly_Execute();
 	static void BuildVirtualTextureOnly_Execute();
+	static void BuildGrassMapsOnly_Execute();
 	static void SetLightingQuality( ELightingBuildQuality NewQuality );
 	static bool IsLightingQualityChecked( ELightingBuildQuality TestQuality );
 	static float GetLightingDensityIdeal();
diff --git a/Engine/Source/Editor/UnrealEd/Private/EditorBuildUtils.cpp b/Engine/Source/Editor/UnrealEd/Private/EditorBuildUtils.cpp
--- a/Engine/Source/Editor/UnrealEd/Private/EditorBuildUtils.cpp
+++ b/Engine/Source/Editor/UnrealEd/Private/EditorBuildUtils.cpp
@@ -43,6 +43,7 @@
 #include "Materials/MaterialInstance.h"
 #include "VirtualTexturingEditorModule.h"
 #include "Components/RuntimeVirtualTextureComponent.h"
+#include "LandscapeSubsystem.h"
 
 DEFINE_LOG_CATEGORY_STATIC(LogEditorBuildUtils, Log, All);
 
@@ -63,6 +64,7 @@ const FName FBuildOptions::BuildAllOnlySelectedPaths(TEXT("BuildAllOnlySelectedP
 const FName FBuildOptions::BuildHierarchicalLOD(TEXT("BuildHierarchicalLOD"));
 const FName FBuildOptions::BuildTextureStreaming(TEXT("BuildTextureStreaming"));
 const FName FBuildOptions::BuildVirtualTexture(TEXT("BuildVirtualTexture"));
+const FName FBuildOptions::BuildGrassMaps(TEXT("BuildGrassMaps"));
 
 bool FEditorBuildUtils::bBuildingNavigationFromUserRequest = false;
 TMap<FName, FEditorBuildUtils::FCustomBuildType> FEditorBuildUtils::CustomBuildTypes;
@@ -304,6 +306,10 @@ bool FEditorBuildUtils::EditorBuild( UWorld* InWorld, FName Id, const bool bAllo
 	{
 		BuildType = SBuildProgressWidget::BUILDTYPE_VirtualTexture;
 	}
+	else if (Id == FBuildOptions::BuildGrassMaps)
+	{
+		BuildType = SBuildProgressWidget::BUILDTYPE_GrassMaps;
+	}
 	else
 	{
 		BuildType = SBuildProgressWidget::BUILDTYPE_Unknown;	
@@ -425,6 +431,15 @@ bool FEditorBuildUtils::EditorBuild( UWorld* InWorld, FName Id, const bool bAllo
 			TriggerHierarchicalLODBuilder(InWorld, Id);
 		}
 	}
+	else if (Id == FBuildOptions::BuildGrassMaps)
+	{
+		bDoBuild = GEditor->WarnAboutHiddenLevels(InWorld, false);
+		if (bDoBuild)
+		{
+			GEditor->ResetTransaction(NSLOCTEXT("UnrealEd", "BuildGrassMaps", "Building Grass Maps"));
+			EditorBuildGrassMaps(InWorld);
+		}
+	}
 	else if (Id == FBuildOptions::BuildAll || Id == FBuildOptions::BuildAllSubmit)
 	{
 		// TODO: WarnIfLightingBuildIsCurrentlyRunning should check with FBuildAllHandler
@@ -912,6 +927,7 @@ FBuildAllHandler::FBuildAllHandler()
 	: CurrentStep(0)
 {
 	// Add built in build steps.
+	BuildSteps.Add(FBuildOptions::BuildGrassMaps);
 	BuildSteps.Add(FBuildOptions::BuildGeometry);
 	BuildSteps.Add(FBuildOptions::BuildHierarchicalLOD);
 	BuildSteps.Add(FBuildOptions::BuildAIPaths);
@@ -1013,6 +1029,11 @@ void FBuildAllHandler::ProcessBuild(const TWeakPtr<SBuildProgressWidget>& BuildP
 			BuildProgressWidget.Pin()->SetBuildType(SBuildProgressWidget::BUILDTYPE_VirtualTexture);
 			FEditorBuildUtils::EditorBuildVirtualTexture(CurrentWorld);
 		}
+		else if (StepId == FBuildOptions::BuildGrassMaps)
+		{
+			BuildProgressWidget.Pin()->SetBuildType(SBuildProgressWidget::BUILDTYPE_GrassMaps);
+			FEditorBuildUtils::EditorBuildGrassMaps(CurrentWorld);
+		}
 		else if (StepId == FBuildOptions::BuildAIPaths)
 		{
 			BuildProgressWidget.Pin()->SetBuildType(SBuildProgressWidget::BUILDTYPE_Paths);
@@ -1408,6 +1429,17 @@ bool FEditorBuildUtils::EditorBuildVirtualTexture(UWorld* InWorld)
 	return true;
 }
 
+void FEditorBuildUtils::EditorBuildGrassMaps(UWorld* InWorld)
+{
+	if (InWorld)
+	{
+		if (ULandscapeSubsystem* LandscapeSubsystem = InWorld->GetSubsystem<ULandscapeSubsystem>())
+		{
+			LandscapeSubsystem->BuildGrassMaps();
+		}
+	}
+}
+
 /** classed used to compile shaders for a specific (mobile) platform and copy the number of instruction to the editor-emulated (mobile) platform */
 class FMaterialOfflineCompilation : public FMaterialResource
 {
diff --git a/Engine/Source/Editor/UnrealEd/Public/Dialogs/SBuildProgress.h b/Engine/Source/Editor/UnrealEd/Public/Dialogs/SBuildProgress.h
--- a/Engine/Source/Editor/UnrealEd/Public/Dialogs/SBuildProgress.h
+++ b/Engine/Source/Editor/UnrealEd/Public/Dialogs/SBuildProgress.h
@@ -46,6 +46,8 @@ public:
 		BUILDTYPE_TextureStreaming, 
 		/** Virtual texture data is being built */
 		BUILDTYPE_VirtualTexture,
+		/** Grass maps are being built */
+		BUILDTYPE_GrassMaps,
 	};
 
 	/** The various issues that can occur. */
diff --git a/Engine/Source/Editor/UnrealEd/Public/EditorBuildUtils.h b/Engine/Source/Editor/UnrealEd/Public/EditorBuildUtils.h
--- a/Engine/Source/Editor/UnrealEd/Public/EditorBuildUtils.h
+++ b/Engine/Source/Editor/UnrealEd/Public/EditorBuildUtils.h
@@ -36,6 +36,8 @@ struct FBuildOptions
 	UNREALED_API static const FName BuildTextureStreaming;
 	/** Build virtual textures */
 	UNREALED_API static const FName BuildVirtualTexture;
+	/** Build grass maps */
+	UNREALED_API static const FName BuildGrassMaps;
 };
 /**
  * Result of a custom editor build.
@@ -162,6 +164,13 @@ public:
 	 */
 	static UNREALED_API bool EditorBuildVirtualTexture(UWorld* InWorld);
 
+	/**
+	 * Perform an editor build for grass maps
+	 *
+	 * @param	InWorld				WorldContext
+	 */
+	static UNREALED_API void EditorBuildGrassMaps(UWorld* InWorld);
+
 	/** 
 	* check if navigation build was was triggered from editor as user request
 	*
diff --git a/Engine/Source/Runtime/Engine/Private/UnrealEngine.cpp b/Engine/Source/Runtime/Engine/Private/UnrealEngine.cpp
--- a/Engine/Source/Runtime/Engine/Private/UnrealEngine.cpp
+++ b/Engine/Source/Runtime/Engine/Private/UnrealEngine.cpp
@@ -146,6 +146,7 @@ UnrealEngine.cpp: Implements the UEngine class and helpers.
 
 #if WITH_EDITOR
 #include "Settings/LevelEditorPlaySettings.h"
+#include "LandscapeSubsystem.h"
 #endif
 // @todo this is here only due to circular dependency to AIModule. To be removed
 
@@ -9985,6 +9986,19 @@ float DrawMapWarnings(UWorld* World, FViewport* Viewport, FCanvas* Canvas, UCanv
 #endif
 	}
 
+#if WITH_EDITOR && !(UE_BUILD_SHIPPING || UE_BUILD_TEST)
+	if (ULandscapeSubsystem* LandscapeSubsystem = World->GetSubsystem<ULandscapeSubsystem>())
+	{
+		if (int32 OutdatedGrassMapCount = LandscapeSubsystem->GetOutdatedGrassMapCount())
+		{
+			SmallTextItem.SetColor(FLinearColor::Red);
+			SmallTextItem.Text = FText::Format(LOCTEXT("GRASS_MAPS_NEED_TO_BE_REBUILT_FMT", "GRASS MAPS NEED TO BE REBUILT ({0} {0}|plural(one=object,other=objects))"), OutdatedGrassMapCount);
+			Canvas->DrawItem(SmallTextItem, FVector2D(MessageX, MessageY));
+			MessageY += FontSizeY;
+		}
+	}
+#endif
+
 	if (World->NumTextureStreamingUnbuiltComponents > 0 || World->NumTextureStreamingDirtyResources > 0)
 	{
 		SmallTextItem.SetColor(FLinearColor::Red);
diff --git a/Engine/Source/Runtime/Landscape/Classes/LandscapeComponent.h b/Engine/Source/Runtime/Landscape/Classes/LandscapeComponent.h
--- a/Engine/Source/Runtime/Landscape/Classes/LandscapeComponent.h
+++ b/Engine/Source/Runtime/Landscape/Classes/LandscapeComponent.h
@@ -252,10 +252,17 @@ struct FLandscapeComponentGrassData
 #if WITH_EDITORONLY_DATA
 	// Height data for LODs 1+, keyed on LOD index
 	TMap<int32, TArray<uint16>> HeightMipData;
+
+	// Grass data was updated but not saved yet
+	bool bIsDirty;
 #endif
 	TMap<ULandscapeGrassType*, TArray<uint8>> WeightData;
 
-	FLandscapeComponentGrassData() {}
+	FLandscapeComponentGrassData()
+#if WITH_EDITORONLY_DATA
+		: bIsDirty(false) 
+#endif
+	{}
 
 #if WITH_EDITOR
 	FLandscapeComponentGrassData(ULandscapeComponent* Component);
diff --git a/Engine/Source/Runtime/Landscape/Classes/LandscapeProxy.h b/Engine/Source/Runtime/Landscape/Classes/LandscapeProxy.h
--- a/Engine/Source/Runtime/Landscape/Classes/LandscapeProxy.h
+++ b/Engine/Source/Runtime/Landscape/Classes/LandscapeProxy.h
@@ -899,7 +899,7 @@ public:
 
 #if WITH_EDITOR
 	/** Update Grass maps */
-	void UpdateGrassData();
+	void UpdateGrassData(bool bInShouldMarkDirty = false, struct FScopedSlowTask* InSlowTask = nullptr);
 
 	/** Render grass maps for the specified components */
 	void RenderGrassMaps(const TArray<ULandscapeComponent*>& LandscapeComponents, const TArray<ULandscapeGrassType*>& GrassTypes);
@@ -935,6 +935,8 @@ public:
 	/** Get landscape position in section space */
 	LANDSCAPE_API FIntPoint GetSectionBaseOffset() const;
 #if WITH_EDITOR
+	LANDSCAPE_API int32 GetOutdatedGrassMapCount() const;
+	LANDSCAPE_API void BuildGrassMaps(struct FScopedSlowTask* InSlowTask = nullptr);
 	LANDSCAPE_API void CreateSplineComponent(const FVector& Scale3D);
 
 	virtual bool CanEditChange(const FProperty* InProperty) const override;
@@ -1135,7 +1137,7 @@ private:
 	static int32 GrassUpdateInterval;
 
 #if WITH_EDITOR
-	void UpdateGrassDataStatus(TSet<UTexture2D*>& OutCurrentForcedStreamedTextures, TSet<UTexture2D*>* OutDesiredForcedStreamedTextures, TSet<ULandscapeComponent*>& OutComponentsNeedingGrassMapRender, TSet<ULandscapeComponent*>* OutOutdatedComponents, bool bInEnableForceResidentFlag);
+	void UpdateGrassDataStatus(TSet<UTexture2D*>* OutCurrentForcedStreamedTextures, TSet<UTexture2D*>* OutDesiredForcedStreamedTextures, TSet<ULandscapeComponent*>* OutComponentsNeedingGrassMapRender, TSet<ULandscapeComponent*>* OutOutdatedComponents, bool bInEnableForceResidentFlag, int32* OutOutdatedGrassMaps = nullptr) const;
 #endif
 
 #if WITH_EDITORONLY_DATA
@@ -1147,3 +1149,21 @@ private:
 	static TArray<ALandscapeProxy*> LandscapeProxies;
 #endif
 };
+
+
+#if WITH_EDITOR
+/**
+ * Helper class used to Build or monitor outdated Grass maps of a world
+ */
+class LANDSCAPE_API FLandscapeGrassMapsBuilder
+{
+public:
+	FLandscapeGrassMapsBuilder(UWorld* InWorld);
+	void Build();
+	int32 GetOutdatedGrassMapCount(bool bInForceUpdate = true) const;
+private:
+	UWorld* World;
+	mutable int32 OutdatedGrassMapCount;
+	mutable double GrassMapsLastCheckTime;
+};
+#endif
\ No newline at end of file
diff --git a/Engine/Source/Runtime/Landscape/Private/Landscape.cpp b/Engine/Source/Runtime/Landscape/Private/Landscape.cpp
--- a/Engine/Source/Runtime/Landscape/Private/Landscape.cpp
+++ b/Engine/Source/Runtime/Landscape/Private/Landscape.cpp
@@ -2099,11 +2099,12 @@ void ALandscapeProxy::PreSave(const class ITargetPlatform* TargetPlatform)
 
 #if WITH_EDITOR
 	// Work out whether we have grass or not for the next game run
-	if (!HasAnyFlags(RF_ClassDefaultObject))
-	{
-		bHasLandscapeGrass = LandscapeComponents.ContainsByPredicate([](ULandscapeComponent* Component) { return Component->MaterialHasGrass(); });
+	BuildGrassMaps();
 
-		UpdateGrassData();
+	for (ULandscapeComponent* Component : LandscapeComponents)
+	{
+		// Reset flag
+		Component->GrassData->bIsDirty = false;
 	}
 
 	if (ALandscape* Landscape = GetLandscapeActor())
diff --git a/Engine/Source/Runtime/Landscape/Private/LandscapeGrass.cpp b/Engine/Source/Runtime/Landscape/Private/LandscapeGrass.cpp
--- a/Engine/Source/Runtime/Landscape/Private/LandscapeGrass.cpp
+++ b/Engine/Source/Runtime/Landscape/Private/LandscapeGrass.cpp
@@ -55,6 +55,8 @@
 #include "MeshPassProcessor.h"
 #include "MeshPassProcessor.inl"
 #include "Math/Halton.h"
+#include "EngineUtils.h"
+#include "Misc/ScopedSlowTask.h"
 
 #define LOCTEXT_NAMESPACE "Landscape"
 
@@ -774,6 +776,7 @@ public:
 
 			// Assign the new data (thread-safe)
 			Component->GrassData = MakeShareable(ComponentGrassData);
+			Component->GrassData->bIsDirty = true;
 
 			if (Proxy->bBakeMaterialPositionOffsetIntoCollision)
 			{
@@ -815,6 +818,7 @@ public:
 
 FLandscapeComponentGrassData::FLandscapeComponentGrassData(ULandscapeComponent* Component)
 	: RotationForWPO(Component->GetLandscapeMaterial()->GetMaterial()->WorldPositionOffset.IsConnected() ? Component->GetComponentTransform().GetRotation() : FQuat(0, 0, 0, 0))
+	, bIsDirty(false)
 {
 	UMaterialInterface* Material = Component->GetLandscapeMaterial();
 	for (UMaterialInstanceConstant* MIC = Cast<UMaterialInstanceConstant>(Material); MIC; MIC = Cast<UMaterialInstanceConstant>(Material))
@@ -1005,6 +1009,7 @@ TArray<uint16> ULandscapeComponent::RenderWPOHeightmap(int32 LOD)
 void ULandscapeComponent::RemoveGrassMap()
 {
 	GrassData = MakeShareable(new FLandscapeComponentGrassData());
+	GrassData->bIsDirty = true;
 }
 
 void ALandscapeProxy::RenderGrassMaps(const TArray<ULandscapeComponent*>& InLandscapeComponents, const TArray<ULandscapeGrassType*>& GrassTypes)
@@ -2131,23 +2136,111 @@ void ALandscapeProxy::RemoveAllExclusionBoxes()
 
 
 #if WITH_EDITOR
+
+FLandscapeGrassMapsBuilder::FLandscapeGrassMapsBuilder(UWorld* InOwner)
+	: World(InOwner)
+	, OutdatedGrassMapCount(0)
+	, GrassMapsLastCheckTime(0)
+{}
+
+
+void FLandscapeGrassMapsBuilder::Build()
+{
+	if (World)
+	{
+		int32 Count = GetOutdatedGrassMapCount();
+		FScopedSlowTask SlowTask(Count, (LOCTEXT("GrassMaps_BuildGrassMaps", "Building Grass maps")));
+		SlowTask.MakeDialog();
+
+		for (TActorIterator<ALandscapeProxy> ProxyIt(World); ProxyIt; ++ProxyIt)
+		{
+			ProxyIt->BuildGrassMaps(&SlowTask);
+		}
+	}
+}
+
+int32 FLandscapeGrassMapsBuilder::GetOutdatedGrassMapCount(bool bInForceUpdate) const
+{
+	if (World)
+	{
+		bool bUpdate = bInForceUpdate || GLandscapeEditModeActive;
+		if (!bUpdate)
+		{
+			double GrassMapsTimeNow = FPlatformTime::Seconds();
+			// Recheck every 20 secs to handle the case where levels may have been Streamed in/out
+			if ((GrassMapsTimeNow - GrassMapsLastCheckTime) > 20)
+			{
+				GrassMapsLastCheckTime = GrassMapsTimeNow;
+				bUpdate = true;
+			}
+		}
+
+		if (bUpdate)
+		{
+			OutdatedGrassMapCount = 0;
+			for (TActorIterator<ALandscapeProxy> ProxyIt(World); ProxyIt; ++ProxyIt)
+			{
+				OutdatedGrassMapCount += ProxyIt->GetOutdatedGrassMapCount();
+			}
+		}
+	}
+	return OutdatedGrassMapCount;
+}
+
+void ALandscapeProxy::BuildGrassMaps(FScopedSlowTask* InSlowTask)
+{
+	if (!HasAnyFlags(RF_ClassDefaultObject))
+	{
+		bool bNewHasLandscapeGrass = LandscapeComponents.ContainsByPredicate([](ULandscapeComponent* Component) { return Component->MaterialHasGrass(); });
+		if (bHasLandscapeGrass != bNewHasLandscapeGrass)
+		{
+			bHasLandscapeGrass = bNewHasLandscapeGrass;
+			MarkPackageDirty();
+		}
+
+		// Generate an mark dirty if changed
+		UpdateGrassData(true, InSlowTask);
+	}
+}
+
+int32 ALandscapeProxy::GetOutdatedGrassMapCount() const
+{
+	int32 OutdatedGrassMaps = 0;
+	if (GGrassEnable > 0)
+	{
+		UpdateGrassDataStatus(nullptr, nullptr, nullptr, nullptr, false, &OutdatedGrassMaps);
+	}
+	return OutdatedGrassMaps;
+}
+
 int32 ALandscapeProxy::TotalComponentsNeedingGrassMapRender = 0;
 int32 ALandscapeProxy::TotalTexturesToStreamForVisibleGrassMapRender = 0;
 int32 ALandscapeProxy::TotalComponentsNeedingTextureBaking = 0;
 
-void ALandscapeProxy::UpdateGrassDataStatus(TSet<UTexture2D*>& OutCurrentForcedStreamedTextures, TSet<UTexture2D*>* OutDesiredForcedStreamedTextures, TSet<ULandscapeComponent*>& OutComponentsNeedingGrassMapRender, TSet<ULandscapeComponent*>* OutOutdatedComponents, bool bInEnableForceResidentFlag)
+void ALandscapeProxy::UpdateGrassDataStatus(TSet<UTexture2D*>* OutCurrentForcedStreamedTextures, TSet<UTexture2D*>* OutDesiredForcedStreamedTextures, TSet<ULandscapeComponent*>* OutComponentsNeedingGrassMapRender, TSet<ULandscapeComponent*>* OutOutdatedComponents, bool bInEnableForceResidentFlag, int32* OutOutdatedGrassMaps) const
 {
-	OutCurrentForcedStreamedTextures.Empty();
-	OutComponentsNeedingGrassMapRender.Empty();
-	
+	if (OutCurrentForcedStreamedTextures)
+	{
+		OutCurrentForcedStreamedTextures->Empty();
+	}
+	if (OutDesiredForcedStreamedTextures)
+	{
+		OutDesiredForcedStreamedTextures->Empty();
+	}
+
+	if (OutComponentsNeedingGrassMapRender)
+	{
+		OutComponentsNeedingGrassMapRender->Empty();
+	}
+
 	if (OutOutdatedComponents)
 	{
 		OutOutdatedComponents->Empty();
 	}
 
-	if (OutDesiredForcedStreamedTextures)
+	if (OutOutdatedGrassMaps)
 	{
-		OutDesiredForcedStreamedTextures->Empty();
+		*OutOutdatedGrassMaps = 0;
 	}
 		
 	const UWorld* World = GetWorld();
@@ -2164,27 +2257,34 @@ void ALandscapeProxy::UpdateGrassDataStatus(TSet<UTexture2D*>& OutCurrentForcedS
 	GetGrassTypes(World, LandscapeMaterial, GrassTypes, OutMaxSquareDiscardDistance);
 	const bool bHasGrassTypes = GrassTypes.Num() > 0;
 
+	const bool bIsOutermostPackageDirty = GetOutermost()->IsDirty();
+
+	int32 OutdatedGrassMaps = 0;
 	for (auto Component : LandscapeComponents)
 	{
 		if (Component != nullptr)
 		{
 			UTexture2D* Heightmap = Component->GetHeightmap();
-			// check textures currently needing force streaming
-			if (Heightmap->bForceMiplevelsToBeResident)
-			{
-				OutCurrentForcedStreamedTextures.Add(Heightmap);
-			}
-			TArray<UTexture2D*>& ComponentWeightmapTextures = Component->GetWeightmapTextures();
+			const TArray<UTexture2D*>& ComponentWeightmapTextures = Component->GetWeightmapTextures();
 
-			for (auto WeightmapTexture : ComponentWeightmapTextures)
+			if (OutCurrentForcedStreamedTextures)
 			{
-				if (WeightmapTexture->bForceMiplevelsToBeResident)
+				// check textures currently needing force streaming
+				if (Heightmap->bForceMiplevelsToBeResident)
+				{
+					OutCurrentForcedStreamedTextures->Add(Heightmap);
+				}
+
+				for (auto WeightmapTexture : ComponentWeightmapTextures)
 				{
-					OutCurrentForcedStreamedTextures.Add(WeightmapTexture);
+					if (WeightmapTexture->bForceMiplevelsToBeResident)
+					{
+						OutCurrentForcedStreamedTextures->Add(WeightmapTexture);
+					}
 				}
 			}
 
-			if (Component->IsGrassMapOutdated() && OutOutdatedComponents)
+			if (OutOutdatedComponents && Component->IsGrassMapOutdated())
 			{
 				OutOutdatedComponents->Add(Component);
 			}
@@ -2193,7 +2293,11 @@ void ALandscapeProxy::UpdateGrassDataStatus(TSet<UTexture2D*>& OutCurrentForcedS
 			{
 				if (Component->IsGrassMapOutdated() || !Component->GrassData->HasData())
 				{
-					OutComponentsNeedingGrassMapRender.Add(Component);
+					if (OutComponentsNeedingGrassMapRender)
+					{
+						OutComponentsNeedingGrassMapRender->Add(Component);
+					}
+					++OutdatedGrassMaps;
 
 					if (bCheckStreamingState && !Component->AreTexturesStreamedForGrassMapRender())
 					{
@@ -2221,20 +2325,26 @@ void ALandscapeProxy::UpdateGrassDataStatus(TSet<UTexture2D*>& OutCurrentForcedS
 						}
 					}
 				}
+				// Don't count dirty component's if package is already dirty. 
+				// Saving the dirty package will also save the updated grass maps.
+				// This counter is used to know if grass maps need to be rebuilt.
+				else if (Component->GrassData->bIsDirty && !bIsOutermostPackageDirty)
+				{
+					++OutdatedGrassMaps;
+				}
 			}
 		}
 	}
-}
 
-void ALandscapeProxy::UpdateGrassData()
-{
-	UWorld* World = GetWorld();
-	if (!World || World->IsGameWorld())
+	if (OutOutdatedGrassMaps)
 	{
-		return;
+		*OutOutdatedGrassMaps = OutdatedGrassMaps;
 	}
+}
 
-	if (!GGrassEnable)
+void ALandscapeProxy::UpdateGrassData(bool bInShouldMarkDirty, FScopedSlowTask* InSlowTask)
+{
+	if (!GGrassEnable || !GetWorld() || GetWorld()->IsGameWorld())
 	{
 		return;
 	}
@@ -2244,8 +2354,10 @@ void ALandscapeProxy::UpdateGrassData()
 	TSet<UTexture2D*> CurrentForcedStreamedTextures;
 	TSet<ULandscapeComponent*> ComponentsNeedingGrassMapRender;
 	TSet<ULandscapeComponent*> OutdatedComponents;
+	int32 TotalOutdatedGrassMaps = 0;
 	const bool bEnableForceResidentFlag = true;
-	UpdateGrassDataStatus(CurrentForcedStreamedTextures, &DesiredForcedStreamedTextures, ComponentsNeedingGrassMapRender, &OutdatedComponents, bEnableForceResidentFlag);
+	UpdateGrassDataStatus(&CurrentForcedStreamedTextures, &DesiredForcedStreamedTextures, &ComponentsNeedingGrassMapRender, &OutdatedComponents, bEnableForceResidentFlag, &TotalOutdatedGrassMaps);
+
 	if (OutdatedComponents.Num())
 	{
 		FlushGrassComponents(&OutdatedComponents);
@@ -2265,17 +2377,32 @@ void ALandscapeProxy::UpdateGrassData()
 		TextureToStream->WaitForStreaming();
 		CurrentForcedStreamedTextures.Add(TextureToStream);
 	}
-		
-	// render grass data if we don't have any
+	
+	auto UpdateProgress = [InSlowTask](int Increment)
+	{
+		if (InSlowTask && Increment && ((InSlowTask->CompletedWork + Increment) <= InSlowTask->TotalAmountOfWork))
+		{
+			InSlowTask->EnterProgressFrame(Increment, FText::Format(LOCTEXT("GrassMaps_BuildGrassMapsProgress", "Building Grass Map {0} of {1})"), FText::AsNumber(InSlowTask->CompletedWork), FText::AsNumber(InSlowTask->TotalAmountOfWork)));
+		}
+	};
+
+	// Take into consideration already dirty/rendered grass maps to update progress accordingly
+	int32 AlreadyRenderedGrassMap = TotalOutdatedGrassMaps - ComponentsNeedingGrassMapRender.Num();
+	UpdateProgress(AlreadyRenderedGrassMap);
+
+	// Render grass data
 	for(ULandscapeComponent* Component : ComponentsNeedingGrassMapRender)
 	{
-		if (!Component->CanRenderGrassMap())
+		if (Component->CanRenderGrassMap())
 		{
-			// we can't currently render grassmaps (eg shaders not compiled)
-			continue;
+			Component->RenderGrassMap();
+			UpdateProgress(1);
 		}
-		
-		Component->RenderGrassMap();
+	}
+
+	if (bInShouldMarkDirty && GetOutdatedGrassMapCount() > 0)
+	{
+		MarkPackageDirty();
 	}
 		
 	for (UTexture2D* TextureToStream : CurrentForcedStreamedTextures)
@@ -2354,7 +2481,7 @@ void ALandscapeProxy::UpdateGrass(const TArray<FVector>& Cameras, int32& InOutNu
 			const bool bEnableForceResidentFlag = false;
 			
 			// Do not pass in DesiredForceStreamedTextures because we build our own list
-			UpdateGrassDataStatus(CurrentForcedStreamedTextures, nullptr, ComponentsNeedingGrassMapRender, &OutdatedComponents, bEnableForceResidentFlag);
+			UpdateGrassDataStatus(&CurrentForcedStreamedTextures, nullptr, &ComponentsNeedingGrassMapRender, &OutdatedComponents, bEnableForceResidentFlag);
 			
 			if(OutdatedComponents.Num())
 			{
@@ -2790,7 +2917,6 @@ void ALandscapeProxy::UpdateGrass(const TArray<FVector>& Cameras, int32& InOutNu
 					if (ComponentsToRender.Num())
 					{
 						RenderGrassMaps(ComponentsToRender, LandscapeGrassTypes);
-						MarkPackageDirty();
 					}
 				}
 
diff --git a/Engine/Source/Runtime/Landscape/Private/LandscapeSubsystem.cpp b/Engine/Source/Runtime/Landscape/Private/LandscapeSubsystem.cpp
--- a/Engine/Source/Runtime/Landscape/Private/LandscapeSubsystem.cpp
+++ b/Engine/Source/Runtime/Landscape/Private/LandscapeSubsystem.cpp
@@ -6,6 +6,7 @@
 #include "Engine/World.h"
 #include "ContentStreaming.h"
 #include "Landscape.h"
+#include "LandscapeProxy.h"
 #include "ProfilingDebugging/CsvProfiler.h"
 
 static int32 GUseStreamingManagerForCameras = 1;
@@ -45,6 +46,10 @@ void ULandscapeSubsystem::Initialize(FSubsystemCollectionBase& Collection)
 	TickGroup = TG_DuringPhysics;
 	bAllowTickOnDedicatedServer = false;
 	RegisterTickFunction(GetWorld()->PersistentLevel);
+
+#if WITH_EDITOR
+	GrassMapsBuilder = MakeUnique<FLandscapeGrassMapsBuilder>(GetWorld());
+#endif
 }
 
 void ULandscapeSubsystem::Deinitialize()
@@ -139,3 +144,15 @@ FName ULandscapeSubsystem::DiagnosticContext(bool bDetailed)
 	static const FName Context(TEXT("ULandscapeSubsystem"));
 	return Context;
 }
+
+#if WITH_EDITOR
+void ULandscapeSubsystem::BuildGrassMaps()
+{
+	GrassMapsBuilder->Build();
+}
+
+int32 ULandscapeSubsystem::GetOutdatedGrassMapCount()
+{
+	return GrassMapsBuilder->GetOutdatedGrassMapCount(/*bInForceUpdate*/false);
+}
+#endif
\ No newline at end of file
diff --git a/Engine/Source/Runtime/Landscape/Public/LandscapeSubsystem.h b/Engine/Source/Runtime/Landscape/Public/LandscapeSubsystem.h
--- a/Engine/Source/Runtime/Landscape/Public/LandscapeSubsystem.h
+++ b/Engine/Source/Runtime/Landscape/Public/LandscapeSubsystem.h
@@ -10,7 +10,7 @@
 
 class ALandscapeProxy;
 
-UCLASS()
+UCLASS(MinimalAPI)
 class ULandscapeSubsystem : public UWorldSubsystem, public FTickFunction
 {
 	GENERATED_BODY()
@@ -22,6 +22,11 @@ public:
 	void RegisterActor(ALandscapeProxy* Proxy);
 	void UnregisterActor(ALandscapeProxy* Proxy);
 
+#if WITH_EDITOR
+	LANDSCAPE_API void BuildGrassMaps();
+	LANDSCAPE_API int32 GetOutdatedGrassMapCount();
+#endif
+
 private:
 	// Begin USubsystem
 	virtual void Initialize(FSubsystemCollectionBase& Collection) override;
@@ -35,4 +40,8 @@ private:
 	// End FTickFunction overrides
 
 	TArray<ALandscapeProxy*> Proxies;
+
+#if WITH_EDITOR
+	TUniquePtr<class FLandscapeGrassMapsBuilder> GrassMapsBuilder;
+#endif
 };
