Fixed build

diff --git a/Engine/Source/Runtime/Landscape/Private/LandscapeSubsystem.cpp b/Engine/Source/Runtime/Landscape/Private/LandscapeSubsystem.cpp
--- a/Engine/Source/Runtime/Landscape/Private/LandscapeSubsystem.cpp
+++ b/Engine/Source/Runtime/Landscape/Private/LandscapeSubsystem.cpp
@@ -48,12 +48,18 @@ void ULandscapeSubsystem::Initialize(FSubsystemCollectionBase& Collection)
 	RegisterTickFunction(GetWorld()->PersistentLevel);
 
 #if WITH_EDITOR
-	GrassMapsBuilder = MakeUnique<FLandscapeGrassMapsBuilder>(GetWorld());
+	GrassMapsBuilder = new FLandscapeGrassMapsBuilder(GetWorld());
 #endif
 }
 
 void ULandscapeSubsystem::Deinitialize()
 {
+#if WITH_EDITOR
+	if (GrassMapsBuilder)
+	{
+		delete GrassMapsBuilder;
+	}
+#endif
 	UnRegisterTickFunction();
 	Proxies.Empty();
 
diff --git a/Engine/Source/Runtime/Landscape/Public/LandscapeSubsystem.h b/Engine/Source/Runtime/Landscape/Public/LandscapeSubsystem.h
--- a/Engine/Source/Runtime/Landscape/Public/LandscapeSubsystem.h
+++ b/Engine/Source/Runtime/Landscape/Public/LandscapeSubsystem.h
@@ -42,6 +42,6 @@ private:
 	TArray<ALandscapeProxy*> Proxies;
 
 #if WITH_EDITOR
-	TUniquePtr<class FLandscapeGrassMapsBuilder> GrassMapsBuilder;
+	class FLandscapeGrassMapsBuilder* GrassMapsBuilder;
 #endif
 };
