Fixed LandscapeExportHeightmapToRenderTarget to take exported components overall extent into consideration.
Added option to export or not all landscape proxies when LandscapeExportHeightmapToRenderTarget is used with Landscape Actor.

diff --git a/Engine/Source/Runtime/Landscape/Classes/LandscapeProxy.h b/Engine/Source/Runtime/Landscape/Classes/LandscapeProxy.h
--- a/Engine/Source/Runtime/Landscape/Classes/LandscapeProxy.h
+++ b/Engine/Source/Runtime/Landscape/Classes/LandscapeProxy.h
@@ -932,6 +932,16 @@ public:
 	/** Get the LandcapeActor-to-world transform with respect to landscape section offset*/
 	LANDSCAPE_API FTransform LandscapeActorToWorld() const;
 
+	/**
+	* Output a landscape heightmap to a render target
+	* @param InRenderTarget - Valid render target with a format of RTF_RGBA16f, RTF_RGBA32f or RTF_RGBA8
+	* @param InExportHeightIntoRGChannel - Tell us if we should export the height that is internally stored as R & G (for 16 bits) to a single R channel of the render target (the format need to be RTF_RGBA16f or RTF_RGBA32f)
+	*									   Note that using RTF_RGBA16f with InExportHeightIntoRGChannel == false, could have precision loss.
+	* @param InExportLandscapeProxies - Option to also export components of all proxies of Landscape actor (if LandscapeProxy is the Landscape Actor)
+	*/
+	UFUNCTION(BlueprintCallable, meta = (DisplayName = "Landscape Export Heightmap to RenderTarget", Keywords = "Push Landscape Heightmap to RenderTarget", UnsafeDuringActorConstruction = "true"), Category = Rendering)
+	bool LandscapeExportHeightmapToRenderTarget(UTextureRenderTarget2D* InRenderTarget, bool InExportHeightIntoRGChannel = false, bool InExportLandscapeProxies = true);
+
 	/** Get landscape position in section space */
 	LANDSCAPE_API FIntPoint GetSectionBaseOffset() const;
 #if WITH_EDITOR
@@ -1062,15 +1072,6 @@ public:
 	UFUNCTION(BlueprintCallable, meta = (DisplayName = "Landscape Import Heightmap from RenderTarget", Keywords = "Push RenderTarget to Landscape Heightmap", UnsafeDuringActorConstruction = "true"), Category = Rendering)
 	bool LandscapeImportHeightmapFromRenderTarget(UTextureRenderTarget2D* InRenderTarget, bool InImportHeightFromRGChannel = false);
 
-	/**
-	* Output a landscape heightmap to a render target
-	* @param InRenderTarget - Valid render target with a format of RTF_RGBA16f, RTF_RGBA32f or RTF_RGBA8
-	* @param InExportHeightIntoRGChannel - Tell us if we should export the height that is internally stored as R & G (for 16 bits) to a single R channel of the render target (the format need to be RTF_RGBA16f or RTF_RGBA32f)
-	*									   Note that using RTF_RGBA16f with InExportHeightIntoRGChannel == false, could have precision loss.
-	*/
-	UFUNCTION(BlueprintCallable, meta = (DisplayName = "Landscape Export Heightmap to RenderTarget", Keywords = "Push Landscape Heightmap to RenderTarget", UnsafeDuringActorConstruction = "true"), Category = Rendering)
-	bool LandscapeExportHeightmapToRenderTarget(UTextureRenderTarget2D* InRenderTarget, bool InExportHeightIntoRGChannel = false);
-	
 	/**
 	* Overwrites a landscape weightmap with render target data
 	* Only works in the editor
diff --git a/Engine/Source/Runtime/Landscape/Private/LandscapeEdit.cpp b/Engine/Source/Runtime/Landscape/Private/LandscapeEdit.cpp
--- a/Engine/Source/Runtime/Landscape/Private/LandscapeEdit.cpp
+++ b/Engine/Source/Runtime/Landscape/Private/LandscapeEdit.cpp
@@ -7166,24 +7166,37 @@ bool ALandscapeProxy::LandscapeImportHeightmapFromRenderTarget(UTextureRenderTar
 
 	return true;
 }
+#endif
 
-bool ALandscapeProxy::LandscapeExportHeightmapToRenderTarget(UTextureRenderTarget2D* InRenderTarget, bool InExportHeightIntoRGChannel)
+bool ALandscapeProxy::LandscapeExportHeightmapToRenderTarget(UTextureRenderTarget2D* InRenderTarget, bool bInExportHeightIntoRGChannel, bool InExportLandscapeProxies)
 {
+#if WITH_EDITOR
 	uint64 StartCycle = FPlatformTime::Cycles64();
 
-	ALandscape* Landscape = GetLandscapeActor();
-	if (Landscape == nullptr)
+	UMaterial* HeightmapRenderMaterial = LoadObject<UMaterial>(nullptr, TEXT("/Engine/EditorLandscapeResources/Landscape_Heightmap_To_RenderTarget2D.Landscape_Heightmap_To_RenderTarget2D"));
+	if (HeightmapRenderMaterial == nullptr)
 	{
-		FMessageLog("Blueprint").Error(LOCTEXT("LandscapeExportHeightmapToRenderTarget_NullLandscape.", "LandscapeExportHeightmapToRenderTarget: Landscape must be non-null."));
+		FMessageLog("Blueprint").Error(LOCTEXT("LandscapeExportHeightmapToRenderTarget_Landscape_Heightmap_To_RenderTarget2D.", "LandscapeExportHeightmapToRenderTarget: Material Landscape_Heightmap_To_RenderTarget2D not found in engine content."));
 		return false;
 	}
 
-	UMaterial* HeightmapRenderMaterial = LoadObject<UMaterial>(nullptr, TEXT("/Engine/EditorLandscapeResources/Landscape_Heightmap_To_RenderTarget2D.Landscape_Heightmap_To_RenderTarget2D"));
+	TArray<ULandscapeComponent*> LandscapeComponentsToExport;
+	//  Export the component of the specified proxy
+	LandscapeComponentsToExport.Append(LandscapeComponents);
 
-	if (HeightmapRenderMaterial == nullptr)
+	// If requested, export all proxies
+	if (InExportLandscapeProxies && (GetLandscapeActor() == this))
 	{
-		FMessageLog("Blueprint").Error(LOCTEXT("LandscapeExportHeightmapToRenderTarget_Landscape_Heightmap_To_RenderTarget2D.", "LandscapeExportHeightmapToRenderTarget: Material Landscape_Heightmap_To_RenderTarget2D not found in engine content."));
-		return false;
+		ULandscapeInfo* LandscapeInfo = GetLandscapeInfo();
+		for (ALandscapeProxy* Proxy : LandscapeInfo->Proxies)
+		{
+			LandscapeComponentsToExport.Append(Proxy->LandscapeComponents);
+		}
+	}
+
+	if (LandscapeComponentsToExport.Num() == 0)
+	{
+		return true;
 	}
 
 	UWorld* World = GEditor->GetEditorWorldContext().World();
@@ -7193,19 +7206,13 @@ bool ALandscapeProxy::LandscapeExportHeightmapToRenderTarget(UTextureRenderTarge
 	FCanvas Canvas(RenderTargetResource, nullptr, 0, 0, 0, World->FeatureLevel);
 	Canvas.Clear(FLinearColor::Black);
 
-	ULandscapeInfo* LandscapeInfo = Landscape->GetLandscapeInfo();
-
-	TArray<ULandscapeComponent*> LandscapeComponentsToExport;
-	LandscapeComponentsToExport.Append(Landscape->LandscapeComponents);
-
-	// When run on the landscape actor, we will also export all proxies otherwise, we only export the component linked to the specified proxy
-	if (Landscape == LandscapeInfo->LandscapeActor.Get())
+	// Find exported component's base offset
+	FIntRect ComponentsExtent(MAX_int32, MAX_int32, MIN_int32, MIN_int32);
+	for (ULandscapeComponent* Component : LandscapeComponentsToExport)
 	{
-		for (ALandscapeProxy* Proxy : LandscapeInfo->Proxies)
-		{
-			LandscapeComponentsToExport.Append(Proxy->LandscapeComponents);
-		}
+		Component->GetComponentExtent(ComponentsExtent.Min.X, ComponentsExtent.Min.Y, ComponentsExtent.Max.X, ComponentsExtent.Max.Y);
 	}
+	FIntPoint ExportBaseOffset = ComponentsExtent.Min;
 
 	struct FTrianglePerMID
 	{
@@ -7224,7 +7231,7 @@ bool ALandscapeProxy::LandscapeExportHeightmapToRenderTarget(UTextureRenderTarge
 			FTrianglePerMID Data;
 			Data.HeightmapMID = UMaterialInstanceDynamic::Create(HeightmapRenderMaterial, this);
 			Data.HeightmapMID->SetTextureParameterValue(TEXT("Heightmap"), Component->GetHeightmap());
-			Data.HeightmapMID->SetScalarParameterValue(TEXT("ExportHeightIntoRGChannel"), InExportHeightIntoRGChannel);
+			Data.HeightmapMID->SetScalarParameterValue(TEXT("ExportHeightIntoRGChannel"), bInExportHeightIntoRGChannel);
 			TrianglesPerMID = &TrianglesPerHeightmap.Add(Component->GetHeightmap(), Data);
 		}
 
@@ -7238,9 +7245,9 @@ bool ALandscapeProxy::LandscapeExportHeightmapToRenderTarget(UTextureRenderTarge
 		{
 			for (int8 SubX = 0; SubX < NumSubsections; ++SubX)
 			{
-				FIntPoint SubSectionSectionBase = ComponentSectionBase;
-				SubSectionSectionBase.X = ComponentSectionBase.X + Component->SubsectionSizeQuads * SubX;
-				SubSectionSectionBase.Y = ComponentSectionBase.Y + Component->SubsectionSizeQuads * SubY;
+				FIntPoint SubSectionSectionBase = ComponentSectionBase - ExportBaseOffset;
+				SubSectionSectionBase.X += Component->SubsectionSizeQuads * SubX;
+				SubSectionSectionBase.Y += Component->SubsectionSizeQuads * SubY;
 
 				// Offset for this component's data in heightmap texture
 				float HeightmapOffsetU = Component->HeightmapScaleBias.Z + HeightmapSubsectionOffsetU * (float)SubX;
@@ -7300,10 +7307,11 @@ bool ALandscapeProxy::LandscapeExportHeightmapToRenderTarget(UTextureRenderTarge
 
 	double SecondsTaken = FPlatformTime::ToSeconds64(FPlatformTime::Cycles64() - StartCycle);
 	UE_LOG(LogLandscapeBP, Display, TEXT("Took %f seconds to export heightmap to render target."), SecondsTaken);
-
+#endif
 	return true;
 }
 
+#if WITH_EDITOR
 bool ALandscapeProxy::LandscapeImportWeightmapFromRenderTarget(UTextureRenderTarget2D* InRenderTarget, FName InLayerName)
 {
 	ALandscape* Landscape = GetLandscapeActor();
