Tied Chaos core version into landscape DDC key build as landscape DDC directly serializes chaos types (FHeightfield). This fixes serialization breakage after base FImplicitObject serializer changed without handling DDC.
#rnx

diff --git a/Engine/Source/Runtime/Landscape/Private/LandscapeCollision.cpp b/Engine/Source/Runtime/Landscape/Private/LandscapeCollision.cpp
--- a/Engine/Source/Runtime/Landscape/Private/LandscapeCollision.cpp
+++ b/Engine/Source/Runtime/Landscape/Private/LandscapeCollision.cpp
@@ -166,7 +166,7 @@ static FString GetHFDDCKeyString(const FName& Format, bool bDefMaterial, const F
 #if PHYSICS_INTERFACE_PHYSX
 	const FString InterfacePrefix = TEXT("PHYSX");
 #elif WITH_CHAOS
-	const FString InterfacePrefix = TEXT("CHAOS");
+	const FString InterfacePrefix = FString::Printf(TEXT("%s_%s"), TEXT("CHAOS"), *Chaos::ChaosVersionString);
 #else
 	const FString InterfacePrefix = TEXT("UNDEFINED");
 #endif
