PR #6748: Made GetAllParameterInfo virtual, so UMaterialExpressionLandscapeLayerBlend can be subclassed (Contributed by [REDACTED_PERSON])

diff --git a/Engine/Source/Runtime/Landscape/Classes/Materials/MaterialExpressionLandscapeLayerBlend.h b/Engine/Source/Runtime/Landscape/Classes/Materials/MaterialExpressionLandscapeLayerBlend.h
--- a/Engine/Source/Runtime/Landscape/Classes/Materials/MaterialExpressionLandscapeLayerBlend.h
+++ b/Engine/Source/Runtime/Landscape/Classes/Materials/MaterialExpressionLandscapeLayerBlend.h
@@ -97,7 +97,7 @@ class LANDSCAPE_API UMaterialExpressionLandscapeLayerBlend : public UMaterialExp
 	/**
 	 * Get list of parameter names for static parameter sets
 	 */
-	void GetAllParameterInfo(TArray<FMaterialParameterInfo> &OutParameterInfo, TArray<FGuid> &OutParameterIds, const FMaterialParameterInfo& InBaseParameterInfo) const;
+	virtual void GetAllParameterInfo(TArray<FMaterialParameterInfo> &OutParameterInfo, TArray<FGuid> &OutParameterIds, const FMaterialParameterInfo& InBaseParameterInfo) const;
 };
 
 
