Make per-platform property preview follow the correct preview platform in the editor

diff --git a/Engine/Plugins/Experimental/PanoramicCapture/Source/PanoramicCapture/Private/StereoStaticMeshComponent.cpp b/Engine/Plugins/Experimental/PanoramicCapture/Source/PanoramicCapture/Private/StereoStaticMeshComponent.cpp
--- a/Engine/Plugins/Experimental/PanoramicCapture/Source/PanoramicCapture/Private/StereoStaticMeshComponent.cpp
+++ b/Engine/Plugins/Experimental/PanoramicCapture/Source/PanoramicCapture/Private/StereoStaticMeshComponent.cpp
@@ -54,12 +54,10 @@ public:
 
 FPrimitiveSceneProxy* UStereoStaticMeshComponent::CreateSceneProxy()
 {
-	const auto FeatureLevel = GetScene()->GetFeatureLevel();
-
     if ((GetStaticMesh() == nullptr) ||
 		(GetStaticMesh()->RenderData == nullptr) ||
 		(GetStaticMesh()->RenderData->LODResources.Num() == 0) ||
-		(GetStaticMesh()->RenderData->LODResources[GetStaticMesh()->MinLOD.GetValueForFeatureLevel(FeatureLevel)].VertexBuffers.PositionVertexBuffer.GetNumVertices() == 0))
+		(GetStaticMesh()->RenderData->LODResources[GetStaticMesh()->MinLOD.GetValue()].VertexBuffers.PositionVertexBuffer.GetNumVertices() == 0))
     {
         return nullptr;
     }
diff --git a/Engine/Plugins/FX/Niagara/Source/Niagara/Private/NiagaraDataInterfaceStaticMesh.cpp b/Engine/Plugins/FX/Niagara/Source/Niagara/Private/NiagaraDataInterfaceStaticMesh.cpp
--- a/Engine/Plugins/FX/Niagara/Source/Niagara/Private/NiagaraDataInterfaceStaticMesh.cpp
+++ b/Engine/Plugins/FX/Niagara/Source/Niagara/Private/NiagaraDataInterfaceStaticMesh.cpp
@@ -307,7 +307,7 @@ bool FNDIStaticMesh_InstanceData::Init(UNiagaraDataInterfaceStaticMesh* Interfac
 	ValidSections.Empty();
 	if (Mesh)
 	{
-	    MinLOD = Mesh->MinLOD.GetValueForFeatureLevel(SystemInstance->GetFeatureLevel());
+	    MinLOD = Mesh->MinLOD.GetValue();
 	    CachedLODIdx = Mesh->RenderData->GetCurrentFirstLODIdx(MinLOD);
 
 		bMeshAllowsCpuAccess = Mesh->bAllowCPUAccess;
diff --git a/Engine/Plugins/FX/Niagara/Source/Niagara/Private/NiagaraRendererMeshes.cpp b/Engine/Plugins/FX/Niagara/Source/Niagara/Private/NiagaraRendererMeshes.cpp
--- a/Engine/Plugins/FX/Niagara/Source/Niagara/Private/NiagaraRendererMeshes.cpp
+++ b/Engine/Plugins/FX/Niagara/Source/Niagara/Private/NiagaraRendererMeshes.cpp
@@ -161,7 +161,7 @@ FNiagaraRendererMeshes::FNiagaraRendererMeshes(ERHIFeatureLevel::Type FeatureLev
 	MaterialParamValidMask |= SetVertexFactoryVariable(Data, Properties->DynamicMaterial2Binding.DataSetVariable, ENiagaraMeshVFLayout::DynamicParam2) ? 0x4 : 0;
 	MaterialParamValidMask |= SetVertexFactoryVariable(Data, Properties->DynamicMaterial3Binding.DataSetVariable, ENiagaraMeshVFLayout::DynamicParam3) ? 0x8 : 0;
 
-	MeshMinimumLOD = Properties->ParticleMesh->MinLOD.GetValueForFeatureLevel(FeatureLevel);
+	MeshMinimumLOD = Properties->ParticleMesh->MinLOD.GetValue();
 
 	const int32 LODCount = MeshRenderData->LODResources.Num();
 	IndexInfoPerSection.SetNum(LODCount);
diff --git a/Engine/Source/Editor/StaticMeshEditor/Private/StaticMeshEditorViewportClient.cpp b/Engine/Source/Editor/StaticMeshEditor/Private/StaticMeshEditorViewportClient.cpp
--- a/Engine/Source/Editor/StaticMeshEditor/Private/StaticMeshEditorViewportClient.cpp
+++ b/Engine/Source/Editor/StaticMeshEditor/Private/StaticMeshEditorViewportClient.cpp
@@ -815,8 +815,7 @@ void FStaticMeshEditorViewportClient::DrawCanvas( FViewport& InViewport, FSceneV
 			LOD - 1;
 	}();
 
-	const ERHIFeatureLevel::Type FeatureLevel = GEditor->PreviewPlatform.GetEffectivePreviewFeatureLevel();
-	const int32 CurrentMinLODLevel = StaticMesh->MinLOD.GetValueForFeatureLevel(FeatureLevel);
+	const int32 CurrentMinLODLevel = StaticMesh->MinLOD.GetValue();
 	const bool bBelowMinLOD = CurrentLODLevel < CurrentMinLODLevel;
 	TextItems.Add(SStaticMeshEditorViewport::FOverlayTextItem(
 		FText::Format(NSLOCTEXT("UnrealEd", "LOD_F", "LOD:  {0}"), FText::AsNumber(CurrentLODLevel)),
diff --git a/Engine/Source/Editor/UnrealEd/Classes/Editor/EditorEngine.h b/Engine/Source/Editor/UnrealEd/Classes/Editor/EditorEngine.h
--- a/Engine/Source/Editor/UnrealEd/Classes/Editor/EditorEngine.h
+++ b/Engine/Source/Editor/UnrealEd/Classes/Editor/EditorEngine.h
@@ -805,6 +805,7 @@ public:
 	virtual void SetMapBuildCancelled(bool InCancelled) override { /* Intentionally empty. */ }
 	virtual void HandleNetworkFailure(UWorld *World, UNetDriver *NetDriver, ENetworkFailure::Type FailureType, const FString& ErrorString) override;
 	virtual ERHIFeatureLevel::Type GetDefaultWorldFeatureLevel() const override { return DefaultWorldFeatureLevel; }
+	virtual bool GetPreviewPlatformName(FName& PlatformGroupName, FName& VanillaPlatformName) const override;
 
 protected:
 	virtual void InitializeObjectReferences() override;
diff --git a/Engine/Source/Editor/UnrealEd/Private/EditorEngine.cpp b/Engine/Source/Editor/UnrealEd/Private/EditorEngine.cpp
--- a/Engine/Source/Editor/UnrealEd/Private/EditorEngine.cpp
+++ b/Engine/Source/Editor/UnrealEd/Private/EditorEngine.cpp
@@ -91,6 +91,7 @@
 #include "HAL/PlatformApplicationMisc.h"
 #include "IMediaModule.h"
 #include "Scalability.h"
+#include "PlatformInfo.h"
 
 // needed for the RemotePropagator
 #include "AudioDevice.h"
@@ -7673,6 +7674,20 @@ void UEditorEngine::SaveEditorFeatureLevel()
 	Settings->PostEditChange();
 }
 
+bool UEditorEngine::GetPreviewPlatformName(FName& PlatformGroupName, FName& VanillaPlatformName) const
+{
+	FName PlatformName = PreviewPlatform.GetEffectivePreviewPlatformName();
+	const PlatformInfo::FPlatformInfo* PlatInfo;
+	if (PlatformName != NAME_None && (PlatInfo = PlatformInfo::FindPlatformInfo(PlatformName)) != nullptr)
+	{
+		VanillaPlatformName = PlatInfo->VanillaPlatformName;
+		PlatformGroupName = PlatInfo->PlatformGroupName;
+		return true;
+	}
+
+	return false;
+}
+
 void UEditorEngine::AddReferencedObjects(FReferenceCollector& Collector)
 {
 	if (PlayInEditorSessionInfo.IsSet())
diff --git a/Engine/Source/Runtime/Core/Public/UObject/UnrealNames.inl b/Engine/Source/Runtime/Core/Public/UObject/UnrealNames.inl
--- a/Engine/Source/Runtime/Core/Public/UObject/UnrealNames.inl
+++ b/Engine/Source/Runtime/Core/Public/UObject/UnrealNames.inl
@@ -138,6 +138,7 @@ REGISTER_NAME(256,Voice)
 REGISTER_NAME(257, Zlib)
 REGISTER_NAME(258, Gzip)
 REGISTER_NAME(259, LZ4)
+REGISTER_NAME(260, Mobile)
 
 // Online
 REGISTER_NAME(280,DGram)
diff --git a/Engine/Source/Runtime/Engine/Classes/Engine/Engine.h b/Engine/Source/Runtime/Engine/Classes/Engine/Engine.h
--- a/Engine/Source/Runtime/Engine/Classes/Engine/Engine.h
+++ b/Engine/Source/Runtime/Engine/Classes/Engine/Engine.h
@@ -16,7 +16,6 @@
 #include "Misc/FrameRate.h"
 #include "Subsystems/SubsystemCollection.h"
 #include "Subsystems/EngineSubsystem.h"
-#include "RHI.h"
 #include "AudioDeviceManager.h"
 #include "Engine.generated.h"
 
@@ -1832,9 +1831,11 @@ public:
 	virtual bool IsInitialized() const { return bIsInitialized; }
 
 	/** The feature used to create new worlds, by default. Overridden for feature level preview in the editor */
-	virtual ERHIFeatureLevel::Type GetDefaultWorldFeatureLevel() const { return GMaxRHIFeatureLevel;  }
+	virtual ERHIFeatureLevel::Type GetDefaultWorldFeatureLevel() const;
 
 #if WITH_EDITOR
+	/** Return the platform group name and vanilla platform name the current preview platform, or false if there is no preview platform. */
+	virtual bool GetPreviewPlatformName(FName& PlatformGroupName, FName& VanillaPlatformName) const;
 
 	/** Editor-only event triggered when the actor list of the world has changed */
 	DECLARE_EVENT( UEngine, FLevelActorListChangedEvent );
diff --git a/Engine/Source/Runtime/Engine/Private/Components/SkinnedMeshComponent.cpp b/Engine/Source/Runtime/Engine/Private/Components/SkinnedMeshComponent.cpp
--- a/Engine/Source/Runtime/Engine/Private/Components/SkinnedMeshComponent.cpp
+++ b/Engine/Source/Runtime/Engine/Private/Components/SkinnedMeshComponent.cpp
@@ -2948,7 +2948,7 @@ bool USkinnedMeshComponent::UpdateLODStatus_Internal(int32 InMasterPoseComponent
 		const int32 LODBias = GSkeletalMeshLODBias;
 #endif
 
-		int32 MinLodIndex = bOverrideMinLod ? MinLodModel : SkeletalMesh->MinLod.GetValueForFeatureLevel(CachedSceneFeatureLevel);
+		int32 MinLodIndex = bOverrideMinLod ? MinLodModel : SkeletalMesh->MinLod.GetValue();
 
 		int32 MaxLODIndex = 0;
 		if (MeshObject)
diff --git a/Engine/Source/Runtime/Engine/Private/Components/StaticMeshComponent.cpp b/Engine/Source/Runtime/Engine/Private/Components/StaticMeshComponent.cpp
--- a/Engine/Source/Runtime/Engine/Private/Components/StaticMeshComponent.cpp
+++ b/Engine/Source/Runtime/Engine/Private/Components/StaticMeshComponent.cpp
@@ -2057,8 +2057,7 @@ bool UStaticMeshComponent::HasLightmapTextureCoordinates() const
 		Mesh->RenderData->LODResources.Num() > 0 &&
 		Mesh->LightMapCoordinateIndex >= 0)
 	{
-		const ERHIFeatureLevel::Type FeatureLevel = GetScene() ? GetScene()->GetFeatureLevel() : GMaxRHIFeatureLevel;
-		int32 MeshMinLOD = Mesh->MinLOD.GetValueForFeatureLevel(FeatureLevel);
+		int32 MeshMinLOD = Mesh->MinLOD.GetValue();
 		MeshMinLOD = FMath::Min(MeshMinLOD,  Mesh->RenderData->LODResources.Num() - 1);
 		
 		return ((uint32)Mesh->LightMapCoordinateIndex < Mesh->RenderData->LODResources[MeshMinLOD].VertexBuffers.StaticMeshVertexBuffer.GetNumTexCoords());
diff --git a/Engine/Source/Runtime/Engine/Private/HierarchicalInstancedStaticMesh.cpp b/Engine/Source/Runtime/Engine/Private/HierarchicalInstancedStaticMesh.cpp
--- a/Engine/Source/Runtime/Engine/Private/HierarchicalInstancedStaticMesh.cpp
+++ b/Engine/Source/Runtime/Engine/Private/HierarchicalInstancedStaticMesh.cpp
@@ -1599,7 +1599,7 @@ void FHierarchicalStaticMeshSceneProxy::GetDynamicMeshElements(const TArray<cons
 
 				for (int32 LODIndex = 1; LODIndex < InstanceParams.LODs; LODIndex++)
 				{
-					float Distance = ComputeBoundsDrawDistance(RenderData->ScreenSize[LODIndex].GetValueForFeatureLevel(View->GetFeatureLevel()), SphereRadius, View->ViewMatrices.GetProjectionMatrix()) * LODScale;
+					float Distance = ComputeBoundsDrawDistance(RenderData->ScreenSize[LODIndex].GetValue(), SphereRadius, View->ViewMatrices.GetProjectionMatrix()) * LODScale;
 					InstanceParams.LODPlanesMin[LODIndex - 1] = FMath::Min(FinalCull - LODRandom, Distance - LODRandom);
 					InstanceParams.LODPlanesMax[LODIndex - 1] = FMath::Min(FinalCull, Distance);
 				}
@@ -1779,7 +1779,7 @@ void FHierarchicalStaticMeshSceneProxy::GetDynamicMeshElements(const TArray<cons
 
 						for (int32 LODIndex = 1; LODIndex < NumLODs; LODIndex++)
 						{
-							float Distance = ComputeBoundsDrawDistance(RenderData->ScreenSize[LODIndex].GetValueForFeatureLevel(View->GetFeatureLevel()), SphereRadius, View->ViewMatrices.GetProjectionMatrix()) * LODScale;
+							float Distance = ComputeBoundsDrawDistance(RenderData->ScreenSize[LODIndex].GetValue(), SphereRadius, View->ViewMatrices.GetProjectionMatrix()) * LODScale;
 							LODPlanesMin[LODIndex - 1] = FMath::Min(FinalCull - LODRandom, Distance - LODRandom);
 							LODPlanesMax[LODIndex - 1] = FMath::Min(FinalCull, Distance);
 						}
diff --git a/Engine/Source/Runtime/Engine/Private/InstancedStaticMesh.cpp b/Engine/Source/Runtime/Engine/Private/InstancedStaticMesh.cpp
--- a/Engine/Source/Runtime/Engine/Private/InstancedStaticMesh.cpp
+++ b/Engine/Source/Runtime/Engine/Private/InstancedStaticMesh.cpp
@@ -3268,14 +3268,14 @@ void FInstancedStaticMeshVertexFactoryShaderParameters::GetElementShaderBindings
 				InstancingViewZCompare.Z = FinalCull;
 				if (int(BatchElement.InstancedLODIndex) < InstancingUserData->MeshRenderData->LODResources.Num() - 1)
 				{
-					float NextCut = ComputeBoundsDrawDistance(InstancingUserData->MeshRenderData->ScreenSize[BatchElement.InstancedLODIndex + 1].GetValueForFeatureLevel(FeatureLevel), SphereRadius, View->ViewMatrices.GetProjectionMatrix()) * LODScale;
+					float NextCut = ComputeBoundsDrawDistance(InstancingUserData->MeshRenderData->ScreenSize[BatchElement.InstancedLODIndex + 1].GetValue(), SphereRadius, View->ViewMatrices.GetProjectionMatrix()) * LODScale;
 					InstancingViewZCompare.Z = FMath::Min(NextCut, FinalCull);
 				}
 
 				InstancingViewZCompare.X = MIN_flt;
 				if (int(BatchElement.InstancedLODIndex) > FirstLOD)
 				{
-					float CurCut = ComputeBoundsDrawDistance(InstancingUserData->MeshRenderData->ScreenSize[BatchElement.InstancedLODIndex].GetValueForFeatureLevel(FeatureLevel), SphereRadius, View->ViewMatrices.GetProjectionMatrix()) * LODScale;
+					float CurCut = ComputeBoundsDrawDistance(InstancingUserData->MeshRenderData->ScreenSize[BatchElement.InstancedLODIndex].GetValue(), SphereRadius, View->ViewMatrices.GetProjectionMatrix()) * LODScale;
 					if (CurCut < FinalCull)
 					{
 						InstancingViewZCompare.Y = CurCut;
diff --git a/Engine/Source/Runtime/Engine/Private/Particles/ParticleModules_Location.cpp b/Engine/Source/Runtime/Engine/Private/Particles/ParticleModules_Location.cpp
--- a/Engine/Source/Runtime/Engine/Private/Particles/ParticleModules_Location.cpp
+++ b/Engine/Source/Runtime/Engine/Private/Particles/ParticleModules_Location.cpp
@@ -2748,8 +2748,7 @@ void UParticleModuleLocationSkelVertSurface::GetSkeletalMeshComponentSource(FPar
 		int32 MinLOD = INDEX_NONE;
 		if (NewSkelMeshComp && NewSkelMeshComp->GetScene() && NewSkelMeshComp->SkeletalMesh)
 		{
-			ERHIFeatureLevel::Type FeatureLevel = NewSkelMeshComp->GetScene()->GetFeatureLevel();
-			MinLOD = NewSkelMeshComp->SkeletalMesh->MinLod.GetValueForFeatureLevel(FeatureLevel);
+			MinLOD = NewSkelMeshComp->SkeletalMesh->MinLod.GetValue();
 			FSkeletalMeshRenderData* SkelMeshResource = NewSkelMeshComp->GetSkeletalMeshRenderData();
 			FSkinWeightVertexBuffer* SkinWeightBuffer = NewSkelMeshComp->GetSkinWeightBuffer(MinLOD);
 			FSkeletalMeshLODRenderData& LODData = SkelMeshResource->LODRenderData[MinLOD];
diff --git a/Engine/Source/Runtime/Engine/Private/Particles/ParticleSystemRender.cpp b/Engine/Source/Runtime/Engine/Private/Particles/ParticleSystemRender.cpp
--- a/Engine/Source/Runtime/Engine/Private/Particles/ParticleSystemRender.cpp
+++ b/Engine/Source/Runtime/Engine/Private/Particles/ParticleSystemRender.cpp
@@ -1571,8 +1571,7 @@ uint32 FDynamicMeshEmitterData::GetMeshLODIndexFromProxy(const FParticleSystemSc
 	{
 		return FirstAvailableLOD;
 	}
-	const auto FeatureLevel = InOwnerProxy->GetScene().GetFeatureLevel();
-	const int32 EffectiveMinLOD = StaticMesh->MinLOD.GetValueForFeatureLevel(FeatureLevel);
+	const int32 EffectiveMinLOD = StaticMesh->MinLOD.GetValue();
 	const int32 MaxLOD = StaticMesh->RenderData->LODResources.Num() - 1;
 	const int32 ClampedMinLOD = FMath::Clamp(EffectiveMinLOD, FirstAvailableLOD, MaxLOD);
 	const int32 MeshLOD = (InOwnerProxy->MeshEmitterLODIndices.IsValidIndex(EmitterIndex)) ? InOwnerProxy->MeshEmitterLODIndices[EmitterIndex] : 0;
diff --git a/Engine/Source/Runtime/Engine/Private/SceneManagement.cpp b/Engine/Source/Runtime/Engine/Private/SceneManagement.cpp
--- a/Engine/Source/Runtime/Engine/Private/SceneManagement.cpp
+++ b/Engine/Source/Runtime/Engine/Private/SceneManagement.cpp
@@ -507,7 +507,7 @@ int8 ComputeTemporalStaticMeshLOD( const FStaticMeshRenderData* RenderData, cons
 	// Walk backwards and return the first matching LOD
 	for(int32 LODIndex = NumLODs - 1 ; LODIndex >= 0 ; --LODIndex)
 	{
-		if(FMath::Square(RenderData->ScreenSize[LODIndex].GetValueForFeatureLevel(View.GetFeatureLevel()) * 0.5f) > ScreenRadiusSquared)
+		if(FMath::Square(RenderData->ScreenSize[LODIndex].GetValue() * 0.5f) > ScreenRadiusSquared)
 		{
 			return FMath::Max(LODIndex, MinLOD);
 		}
@@ -542,7 +542,7 @@ int8 ComputeStaticMeshLOD( const FStaticMeshRenderData* RenderData, const FVecto
 		for (int32 LODIndex = NumLODs - 1; LODIndex >= 0; --LODIndex)
 		{
 			float ScreenSizeScale = FactorScale * LODView.LODDistanceFactor;
-			float MeshScreenSize = RenderData->ScreenSize[LODIndex].GetValueForFeatureLevel(View.GetFeatureLevel()) * ScreenSizeScale;
+			float MeshScreenSize = RenderData->ScreenSize[LODIndex].GetValue() * ScreenSizeScale;
 
 			if (FMath::Square(MeshScreenSize * 0.5f) > ScreenRadiusSquared)
 			{
diff --git a/Engine/Source/Runtime/Engine/Private/StaticMesh.cpp b/Engine/Source/Runtime/Engine/Private/StaticMesh.cpp
--- a/Engine/Source/Runtime/Engine/Private/StaticMesh.cpp
+++ b/Engine/Source/Runtime/Engine/Private/StaticMesh.cpp
@@ -2158,7 +2158,7 @@ static void SerializeBuildSettingsForDDC(FArchive& Ar, FMeshBuildSettings& Build
 // If static mesh derived data needs to be rebuilt (new format, serialization
 // differences, etc.) replace the version GUID below with a new one.
 // In case of merge conflicts with DDC versions, you MUST generate a new GUID
-// and set this new GUID as the version.                                       
+// and set this new GUID as the version.
 #define STATICMESH_DERIVEDDATA_VER TEXT("F9378A28F161444987BACE79B2590E57")
 
 static const FString& GetStaticMeshDerivedDataVersion()
@@ -2803,7 +2803,7 @@ bool UStaticMesh::HasValidRenderData(bool bCheckLODForVerts, int32 LODIndex) con
 		{
 		    if (LODIndex == INDEX_NONE)
 		    {
-			    LODIndex = FMath::Clamp<int32>(MinLOD.GetValueForFeatureLevel(GMaxRHIFeatureLevel), 0, RenderData->LODResources.Num() - 1);
+			    LODIndex = FMath::Clamp<int32>(MinLOD.GetValue(), 0, RenderData->LODResources.Num() - 1);
 		    }
 			return (RenderData->LODResources[LODIndex].VertexBuffers.StaticMeshVertexBuffer.GetNumVertices() > 0);
 		}
diff --git a/Engine/Source/Runtime/Engine/Private/StaticMeshRender.cpp b/Engine/Source/Runtime/Engine/Private/StaticMeshRender.cpp
--- a/Engine/Source/Runtime/Engine/Private/StaticMeshRender.cpp
+++ b/Engine/Source/Runtime/Engine/Private/StaticMeshRender.cpp
@@ -149,7 +149,7 @@ FStaticMeshSceneProxy::FStaticMeshSceneProxy(UStaticMeshComponent* InComponent,
 
 	const auto FeatureLevel = GetScene().GetFeatureLevel();
 
-	const int32 SMCurrentMinLOD = InComponent->GetStaticMesh()->MinLOD.GetValueForFeatureLevel(FeatureLevel);
+	const int32 SMCurrentMinLOD = InComponent->GetStaticMesh()->MinLOD.GetValue();
 	int32 EffectiveMinLOD = InComponent->bOverrideMinLOD ? InComponent->MinLOD : SMCurrentMinLOD;
 
 #if WITH_EDITOR
@@ -2088,7 +2088,7 @@ FLightInteraction FStaticMeshSceneProxy::FLODInfo::GetInteraction(const FLightSc
 
 float FStaticMeshSceneProxy::GetScreenSize( int32 LODIndex ) const
 {
-	return RenderData->ScreenSize[LODIndex].GetValueForFeatureLevel(GetScene().GetFeatureLevel());
+	return RenderData->ScreenSize[LODIndex].GetValue();
 }
 
 /**
diff --git a/Engine/Source/Runtime/Engine/Private/Streaming/StreamingTexture.cpp b/Engine/Source/Runtime/Engine/Private/Streaming/StreamingTexture.cpp
--- a/Engine/Source/Runtime/Engine/Private/Streaming/StreamingTexture.cpp
+++ b/Engine/Source/Runtime/Engine/Private/Streaming/StreamingTexture.cpp
@@ -82,7 +82,7 @@ void FStreamingRenderAsset::UpdateStaticData(const FRenderAssetStreamingSettings
 					// Screen sizes stored on assets are 2R/D where R is the radius of bounding spheres and D is the
 					// distance from view origins to bounds origins. The factor calculated by the streamer, however,
 					// is R/D so multiply 0.5 here
-					LODScreenSizes[Idx] = StaticMesh->RenderData->ScreenSize[LODIdx].GetValueForFeatureLevel(GMaxRHIFeatureLevel) * 0.5f;
+					LODScreenSizes[Idx] = StaticMesh->RenderData->ScreenSize[LODIdx].GetValue() * 0.5f;
 				}
 			}
 			else if (RenderAssetType == AT_SkeletalMesh)
@@ -92,7 +92,7 @@ void FStreamingRenderAsset::UpdateStaticData(const FRenderAssetStreamingSettings
 				for (int32 Idx = 0; Idx < MaxNumMeshLODs; ++Idx)
 				{
 					const int32 LODIdx = FMath::Max(MipCount - Idx - 1, 0);
-					LODScreenSizes[Idx] = LODInfos[LODIdx].ScreenSize.GetValueForFeatureLevel(GMaxRHIFeatureLevel) * 0.5f;
+					LODScreenSizes[Idx] = LODInfos[LODIdx].ScreenSize.GetValue() * 0.5f;
 				}
 			}
 			else
diff --git a/Engine/Source/Runtime/Engine/Private/UnrealEngine.cpp b/Engine/Source/Runtime/Engine/Private/UnrealEngine.cpp
--- a/Engine/Source/Runtime/Engine/Private/UnrealEngine.cpp
+++ b/Engine/Source/Runtime/Engine/Private/UnrealEngine.cpp
@@ -16145,6 +16145,11 @@ int32 UEngine::RenderStatSlateBatches(UWorld* World, FViewport* Viewport, FCanva
 }
 #endif
 
+ERHIFeatureLevel::Type UEngine::GetDefaultWorldFeatureLevel() const
+{
+	return GMaxRHIFeatureLevel;
+}
+
 #if WITH_EDITOR
 void UEngine::PostEditChangeProperty(struct FPropertyChangedEvent& PropertyChangedEvent)
 {
@@ -16177,6 +16182,34 @@ void UEngine::PostEditChangeProperty(struct FPropertyChangedEvent& PropertyChang
 
 	Super::PostEditChangeProperty(PropertyChangedEvent);
 }
+
+bool UEngine::GetPreviewPlatformName(FName& PlatformGroupName, FName& VanillaPlatformName) const
+{
+	// Support returning "Mobile" when started with -featureleveles31 etc
+
+	static bool bHasCachedResult = false;
+	static bool bCachedMobile = false;
+
+	if (!bHasCachedResult)
+	{
+		bHasCachedResult = true;
+		ERHIFeatureLevel::Type FeatureLevel;
+
+		bCachedMobile = (RHIGetPreviewFeatureLevel(FeatureLevel) && FeatureLevel <= ERHIFeatureLevel::ES3_1);
+	}
+
+	if (bCachedMobile)
+	{
+		PlatformGroupName = NAME_Mobile;
+		PlatformGroupName = NAME_None;
+		return true;
+	}
+	else
+	{
+		return false;
+	}
+}
+
 #endif
 
 #undef LOCTEXT_NAMESPACE
diff --git a/Engine/Source/Runtime/Engine/Public/PerPlatformProperties.h b/Engine/Source/Runtime/Engine/Public/PerPlatformProperties.h
--- a/Engine/Source/Runtime/Engine/Public/PerPlatformProperties.h
+++ b/Engine/Source/Runtime/Engine/Public/PerPlatformProperties.h
@@ -59,22 +59,38 @@ struct ENGINE_API TPerPlatformProperty
 
 #endif
 
+	_ValueType GetValue() const
+	{
+#if WITH_EDITORONLY_DATA && WITH_EDITOR
+		FName PlatformGroupName, VanillaPlatformName;
+		// Lookup the override preview platform info, if any
+		if (GEngine && GEngine->GetPreviewPlatformName(PlatformGroupName, VanillaPlatformName))
+		{
+			return GetValueForPlatformIdentifiers(PlatformGroupName, VanillaPlatformName);
+		}
+		else		
+#endif
+		{
+			const _StructType* This = StaticCast<const _StructType*>(this);
+			return This->Default;
+		}
+	}
+
+	UE_DEPRECATED(4.26, "GetValueForFeatureLevel is not needed for platform previewing and GetValue() should be used instead.")
 	_ValueType GetValueForFeatureLevel(ERHIFeatureLevel::Type FeatureLevel) const
 	{
 #if WITH_EDITORONLY_DATA
-		/* Temporary, we should replace this with something better using editor feature level preview system */
 		FName PlatformGroupName;
 		switch (FeatureLevel)
 		{
-			case ERHIFeatureLevel::ES3_1:
-			{
-				static FName NAME_Mobile("Mobile");
-				PlatformGroupName = NAME_Mobile;
-				break;
-			}
-			default:
-				PlatformGroupName = NAME_None;
-				break;
+		    case ERHIFeatureLevel::ES3_1:
+		    {
+			    PlatformGroupName = NAME_Mobile;
+			    break;
+		    }
+		    default:
+			    PlatformGroupName = NAME_None;
+			    break;
 		}
 		return GetValueForPlatformIdentifiers(PlatformGroupName);
 #else
diff --git a/Engine/Source/Runtime/Landscape/Private/Landscape.cpp b/Engine/Source/Runtime/Landscape/Private/Landscape.cpp
--- a/Engine/Source/Runtime/Landscape/Private/Landscape.cpp
+++ b/Engine/Source/Runtime/Landscape/Private/Landscape.cpp
@@ -570,7 +570,7 @@ UMaterialInterface* ULandscapeComponent::GetLandscapeMaterial(int8 InLODIndex) c
 		{
 			for (const FLandscapeComponentMaterialOverride& Material : OverrideMaterials)
 			{
-				if (Material.LODIndex.GetValueForFeatureLevel(World->FeatureLevel) == InLODIndex)
+				if (Material.LODIndex.GetValue() == InLODIndex)
 				{
 					if (Material.Material != nullptr)
 					{
@@ -2739,7 +2739,7 @@ UMaterialInterface* ALandscapeProxy::GetLandscapeMaterial(int8 InLODIndex) const
 		{
 			for (const FLandscapeProxyMaterialOverride& OverrideMaterial : LandscapeMaterialsOverride)
 			{
-				if (OverrideMaterial.LODIndex.GetValueForFeatureLevel(World->FeatureLevel) == InLODIndex)
+				if (OverrideMaterial.LODIndex.GetValue() == InLODIndex)
 				{
 					if (OverrideMaterial.Material != nullptr)
 					{
@@ -2774,7 +2774,7 @@ UMaterialInterface* ALandscapeStreamingProxy::GetLandscapeMaterial(int8 InLODInd
 		{
 			for (const FLandscapeProxyMaterialOverride& OverrideMaterial : LandscapeMaterialsOverride)
 			{
-				if (OverrideMaterial.LODIndex.GetValueForFeatureLevel(World->FeatureLevel) == InLODIndex)
+				if (OverrideMaterial.LODIndex.GetValue() == InLODIndex)
 				{
 					if (OverrideMaterial.Material != nullptr)
 					{
diff --git a/Engine/Source/Runtime/Landscape/Private/LandscapeGrass.cpp b/Engine/Source/Runtime/Landscape/Private/LandscapeGrass.cpp
--- a/Engine/Source/Runtime/Landscape/Private/LandscapeGrass.cpp
+++ b/Engine/Source/Runtime/Landscape/Private/LandscapeGrass.cpp
@@ -1422,7 +1422,7 @@ struct FGrassBuilderBase
 	FGrassBuilderBase(ALandscapeProxy* Landscape, ULandscapeComponent* Component, const FGrassVariety& GrassVariety, ERHIFeatureLevel::Type FeatureLevel, int32 SqrtSubsections = 1, int32 SubX = 0, int32 SubY = 0, bool bEnableDensityScaling = true)
 	{
 		const float DensityScale = bEnableDensityScaling ? GGrassDensityScale : 1.0f;
-		GrassDensity = GrassVariety.GrassDensity.GetValueForFeatureLevel(FeatureLevel) * DensityScale;
+		GrassDensity = GrassVariety.GrassDensity.GetValue() * DensityScale;
 
 		DrawScale = Landscape->GetRootComponent()->GetRelativeScale3D();
 		DrawLoc = Landscape->GetActorLocation();
@@ -2087,7 +2087,6 @@ void ALandscapeProxy::FlushGrassComponents(const TSet<ULandscapeComponent*>* Onl
 void ALandscapeProxy::GetGrassTypes(const UWorld* World, UMaterialInterface* LandscapeMat, TArray<ULandscapeGrassType*>& GrassTypesOut, float& OutMaxSquareDiscardDistance)
 {
 	float MaxDiscardDistance = 0.0f;
-	ERHIFeatureLevel::Type FeatureLevel = World->Scene != nullptr? World->Scene->GetFeatureLevel() : GMaxRHIFeatureLevel;
 	if (LandscapeMat)
 	{
 		GrassTypesOut.Append(LandscapeMat->GetMaterial()->GetCachedExpressionData().GrassTypes);
@@ -2098,7 +2097,7 @@ void ALandscapeProxy::GetGrassTypes(const UWorld* World, UMaterialInterface* Lan
 			{
 				for (auto& GrassVariety : GrassType->GrassVarieties)
 				{
-					const int32 EndCullDistance = GrassVariety.EndCullDistance.GetValueForFeatureLevel(FeatureLevel);
+					const int32 EndCullDistance = GrassVariety.EndCullDistance.GetValue();
 					if (EndCullDistance > MaxDiscardDistance)
 					{
 						MaxDiscardDistance = EndCullDistance;
@@ -2572,8 +2571,8 @@ void ALandscapeProxy::UpdateGrass(const TArray<FVector>& Cameras, int32& InOutNu
 						for (auto& GrassVariety : GrassType->GrassVarieties)
 						{
 							GrassVarietyIndex++;
-							int32 EndCullDistance = GrassVariety.EndCullDistance.GetValueForFeatureLevel(FeatureLevel);
-							if (GrassVariety.GrassMesh && GrassVariety.GrassDensity.GetValueForFeatureLevel(FeatureLevel) > 0.0f && EndCullDistance > 0)
+							int32 EndCullDistance = GrassVariety.EndCullDistance.GetValue();
+							if (GrassVariety.GrassMesh && GrassVariety.GrassDensity.GetValue() > 0.0f && EndCullDistance > 0)
 							{
 								float MustHaveDistance = GuardBand * (float)EndCullDistance * CullDistanceScale;
 								float DiscardDistance = DiscardGuardBand * (float)EndCullDistance * CullDistanceScale;
@@ -2809,8 +2808,8 @@ void ALandscapeProxy::UpdateGrass(const TArray<FVector>& Cameras, int32& InOutNu
 										}
 										else
 										{
-											HierarchicalInstancedStaticMeshComponent->InstanceStartCullDistance = GrassVariety.StartCullDistance.GetValueForFeatureLevel(FeatureLevel) * CullDistanceScale;
-											HierarchicalInstancedStaticMeshComponent->InstanceEndCullDistance = GrassVariety.EndCullDistance.GetValueForFeatureLevel(FeatureLevel) * CullDistanceScale;
+											HierarchicalInstancedStaticMeshComponent->InstanceStartCullDistance = GrassVariety.StartCullDistance.GetValue() * CullDistanceScale;
+											HierarchicalInstancedStaticMeshComponent->InstanceEndCullDistance = GrassVariety.EndCullDistance.GetValue() * CullDistanceScale;
 										}
 
 										//@todo - take the settings from a UFoliageType object.  For now, disable distance field lighting on grass so we don't hitch.
