Bump ddc version to force landscape recook. Temp fix

diff --git a/Engine/Source/Runtime/Landscape/Private/Landscape.cpp b/Engine/Source/Runtime/Landscape/Private/Landscape.cpp
--- a/Engine/Source/Runtime/Landscape/Private/Landscape.cpp
+++ b/Engine/Source/Runtime/Landscape/Private/Landscape.cpp
@@ -120,7 +120,7 @@ namespace LandscapeCookStats
 // differences, etc.) replace the version GUID below with a new one.
 // In case of merge conflicts with DDC versions, you MUST generate a new GUID
 // and set this new GUID as the version.                                       
-#define LANDSCAPE_MOBILE_COOK_VERSION TEXT("A7E61E6494A6450F950333B3BD0DF7A9")
+#define LANDSCAPE_MOBILE_COOK_VERSION TEXT("F96002C1787F44878795B534CEE2F902")
 
 #define LOCTEXT_NAMESPACE "Landscape"
 
@@ -3353,7 +3353,7 @@ ULandscapeWeightmapUsage::ULandscapeWeightmapUsage(const FObjectInitializer& Obj
 }
 
 // Generate a new guid to force a recache of all landscape derived data
-#define LANDSCAPE_FULL_DERIVEDDATA_VER			TEXT("746F41F382904AC9BA4C6449C3BF4A38")
+#define LANDSCAPE_FULL_DERIVEDDATA_VER			TEXT("3000901CF3B24F028854C2DB986E5B3B")
 
 FString FLandscapeComponentDerivedData::GetDDCKeyString(const FGuid& StateId)
 {
