Fix landscape ray tracing LOD calculation not considering View LOD scale properly

diff --git a/Engine/Source/Runtime/Landscape/Private/LandscapeRender.cpp b/Engine/Source/Runtime/Landscape/Private/LandscapeRender.cpp
--- a/Engine/Source/Runtime/Landscape/Private/LandscapeRender.cpp
+++ b/Engine/Source/Runtime/Landscape/Private/LandscapeRender.cpp
@@ -2650,7 +2650,8 @@ void FLandscapeComponentSceneProxy::GetDynamicRayTracingInstances(FRayTracingMat
 	ForcedLODLevel = ViewLodOveride >= 0 ? ViewLodOveride : ForcedLODLevel;
 
 	float MeshScreenSizeSquared = ComputeBoundsScreenRadiusSquared(GetBounds().Origin, GetBounds().SphereRadius, *Context.ReferenceView);
-	int32 LODToRender = ForcedLODLevel >= 0 ? ForcedLODLevel : GetLODFromScreenSize(MeshScreenSizeSquared, Context.ReferenceView->LODDistanceFactor);
+	float LODScale = Context.ReferenceView->LODDistanceFactor * CVarStaticMeshLODDistanceScale.GetValueOnRenderThread();
+	int32 LODToRender = ForcedLODLevel >= 0 ? ForcedLODLevel : GetLODFromScreenSize(MeshScreenSizeSquared, LODScale * LODScale);
 	
 	FLandscapeElementParamArray& ParameterArray = Context.RayTracingMeshResourceCollector.AllocateOneFrameResource<FLandscapeElementParamArray>();
 	ParameterArray.ElementParams.AddDefaulted(NumSubsections * NumSubsections);
