Edigrate CL14101450 from //UE4/Main

diff --git a/Engine/Source/Runtime/Landscape/Private/Landscape.cpp b/Engine/Source/Runtime/Landscape/Private/Landscape.cpp
--- a/Engine/Source/Runtime/Landscape/Private/Landscape.cpp
+++ b/Engine/Source/Runtime/Landscape/Private/Landscape.cpp
@@ -2507,6 +2507,15 @@ void ALandscapeProxy::PostLoad()
 		if (Comp)
 		{
 			Comp->UpdateRejectNavmeshUnderneath();
+
+			// Store the layer combination in the MaterialInstanceConstantMap
+			if (UMaterialInstance* MaterialInstance = Comp->GetMaterialInstance(0, false))
+			{
+				if(UMaterialInstanceConstant* CombinationMaterialInstance = Cast<UMaterialInstanceConstant>(MaterialInstance->Parent))
+				{
+					MaterialInstanceConstantMap.Add(*ULandscapeComponent::GetLayerAllocationKey(Comp->GetWeightmapLayerAllocations(), CombinationMaterialInstance->Parent), CombinationMaterialInstance);
+				}
+			}
 		}
 	}
 
diff --git a/Engine/Source/Runtime/Landscape/Private/LandscapeEdit.cpp b/Engine/Source/Runtime/Landscape/Private/LandscapeEdit.cpp
--- a/Engine/Source/Runtime/Landscape/Private/LandscapeEdit.cpp
+++ b/Engine/Source/Runtime/Landscape/Private/LandscapeEdit.cpp
@@ -756,16 +756,6 @@ void ULandscapeComponent::FixupWeightmaps()
 			}
 
 			RemoveInvalidWeightmaps();
-
-			// Store the layer combination in the MaterialInstanceConstantMap
-			if (GetMaterialInstance(0, false) != nullptr)
-			{
-				UMaterialInstanceConstant* CombinationMaterialInstance = Cast<UMaterialInstanceConstant>(GetMaterialInstance(0, false)->Parent);
-				if (CombinationMaterialInstance)
-				{
-					Proxy->MaterialInstanceConstantMap.Add(*GetLayerAllocationKey(WeightmapLayerAllocations, CombinationMaterialInstance->Parent), CombinationMaterialInstance);
-				}
-			}
 		}
 	}
 }
