Landscape: Fix Async Texture Compile Crashes/Bugs

diff --git a/Engine/Source/Runtime/Landscape/Classes/LandscapeComponent.h b/Engine/Source/Runtime/Landscape/Classes/LandscapeComponent.h
--- a/Engine/Source/Runtime/Landscape/Classes/LandscapeComponent.h
+++ b/Engine/Source/Runtime/Landscape/Classes/LandscapeComponent.h
@@ -1001,7 +1001,7 @@ public:
 	/**
 	 * Create weightmaps for this component for the layers specified in the WeightmapLayerAllocations array
 	 */
-	LANDSCAPE_API void ReallocateWeightmaps(FLandscapeEditDataInterface* DataInterface = nullptr, bool InCanUseEditingWeightmap = true, bool InSaveToTransactionBuffer = true, bool InInitPlatformDataAsync = false, bool InForceReallocate = false, ALandscapeProxy* InTargetProxy = nullptr, TArray<UTexture2D*>* OutNewCreatedTextures = nullptr);
+	LANDSCAPE_API void ReallocateWeightmaps(FLandscapeEditDataInterface* DataInterface = nullptr, bool InCanUseEditingWeightmap = true, bool InSaveToTransactionBuffer = true, bool InForceReallocate = false, ALandscapeProxy* InTargetProxy = nullptr, TArray<UTexture*>* OutNewCreatedTextures = nullptr);
 
 	/** Returns the component's LandscapeMaterial, or the Component's OverrideLandscapeMaterial if set */
 	LANDSCAPE_API UMaterialInterface* GetLandscapeMaterial(int8 InLODIndex = INDEX_NONE) const;
diff --git a/Engine/Source/Runtime/Landscape/Private/Landscape.cpp b/Engine/Source/Runtime/Landscape/Private/Landscape.cpp
--- a/Engine/Source/Runtime/Landscape/Private/Landscape.cpp
+++ b/Engine/Source/Runtime/Landscape/Private/Landscape.cpp
@@ -1823,8 +1823,7 @@ void ULandscapeComponent::AddDefaultLayerData(const FGuid& InLayerGuid, const TA
 
 			NewLayerHeightmap->Source.UnlockMip(0);
 
-			NewLayerHeightmap->BeginCachePlatformData();
-			NewLayerHeightmap->ClearAllCachedCookedPlatformData();
+			NewLayerHeightmap->UpdateResource();
 		}
 
 		NewData.HeightmapData.Texture = *LayerHeightmap;
diff --git a/Engine/Source/Runtime/Landscape/Private/LandscapeEdit.cpp b/Engine/Source/Runtime/Landscape/Private/LandscapeEdit.cpp
--- a/Engine/Source/Runtime/Landscape/Private/LandscapeEdit.cpp
+++ b/Engine/Source/Runtime/Landscape/Private/LandscapeEdit.cpp
@@ -3065,8 +3065,7 @@ LANDSCAPE_API void ALandscapeProxy::Import(const FGuid& InGuid, int32 InMinX, in
 		// Retrigger a caching of the platform data as we wrote again in the textures
 		for (UTexture2D* Texture : LayersTextures)
 		{
-			Texture->ClearAllCachedCookedPlatformData();
-			Texture->BeginCachePlatformData();
+			Texture->UpdateResource();
 		}
 
 		LandscapeActor->RequestLayersContentUpdateForceAll();
@@ -4257,11 +4256,11 @@ ALandscapeProxy* ULandscapeInfo::MoveComponentsToLevel(const TArray<ULandscapeCo
 		FLandscapeEditDataInterface LandscapeEdit(this);
 		for (ULandscapeComponent* Component : TargetSelectedComponents)
 		{
-			Component->ReallocateWeightmaps(&LandscapeEdit, false, true, false, true, LandscapeProxy);
+			Component->ReallocateWeightmaps(&LandscapeEdit, false, true, true, LandscapeProxy);
 			Component->ForEachLayer([&](const FGuid& LayerGuid, FLandscapeLayerComponentData& LayerData)
 			{
 				FScopedSetLandscapeEditingLayer Scope(Landscape, LayerGuid);
-				Component->ReallocateWeightmaps(&LandscapeEdit, true, true, false, true, LandscapeProxy);
+				Component->ReallocateWeightmaps(&LandscapeEdit, true, true, true, LandscapeProxy);
 			});
 			Landscape->RequestLayersContentUpdateForceAll();
 		}
@@ -4359,6 +4358,8 @@ void ALandscape::SplitHeightmap(ULandscapeComponent* Comp, ALandscapeProxy* Targ
 		// Create the new heightmap texture
 		NewHeightmapTexture = DstProxy->CreateLandscapeTexture(HeightmapSizeU, HeightmapSizeV, TEXTUREGROUP_Terrain_Heightmap, TSF_BGRA8);
 		ULandscapeComponent::CreateEmptyTextureMips(NewHeightmapTexture, true);
+		NewHeightmapTexture->PostEditChange();
+
 		Comp->HeightmapScaleBias = NewHeightmapScaleBias;
 		Comp->SetHeightmap(NewHeightmapTexture);
 
@@ -4379,8 +4380,6 @@ void ALandscape::SplitHeightmap(ULandscapeComponent* Comp, ALandscapeProxy* Targ
 	// We disable automatic material update context, to manage it manually if we have a custom update context specified
 	GDisableAutomaticTextureMaterialUpdateDependencies = (InOutUpdateContext != nullptr);
 
-	NewHeightmapTexture->PostEditChange();
-
 	if (InOutUpdateContext != nullptr)
 	{
 		// Build a list of all unique materials the landscape uses
@@ -5364,7 +5363,7 @@ void ULandscapeInfo::ClearSelectedRegion(bool bIsComponentwise /*= true*/)
 	}
 }
 
-void ULandscapeComponent::ReallocateWeightmaps(FLandscapeEditDataInterface* DataInterface, bool InCanUseEditingWeightmap, bool InSaveToTransactionBuffer, bool InInitPlatformDataAsync, bool InForceReallocate, ALandscapeProxy* InTargetProxy, TArray<UTexture2D*>* OutNewCreatedTextures)
+void ULandscapeComponent::ReallocateWeightmaps(FLandscapeEditDataInterface* DataInterface, bool InCanUseEditingWeightmap, bool InSaveToTransactionBuffer, bool InForceReallocate, ALandscapeProxy* InTargetProxy, TArray<UTexture*>* OutNewCreatedTextures)
 {
 	int32 NeededNewChannels = 0;
 	ALandscapeProxy* TargetProxy = InTargetProxy ? InTargetProxy : GetLandscapeProxy();
@@ -5548,15 +5547,7 @@ void ULandscapeComponent::ReallocateWeightmaps(FLandscapeEditDataInterface* Data
 			// Alloc dummy mips
 			CreateEmptyTextureMips(CurrentWeightmapTexture, true);
 
-			if (InInitPlatformDataAsync)
-			{
-				CurrentWeightmapTexture->BeginCachePlatformData();
-				CurrentWeightmapTexture->ClearAllCachedCookedPlatformData();
-			}
-			else
-			{
-				CurrentWeightmapTexture->PostEditChange();
-			}
+			CurrentWeightmapTexture->PostEditChange();
 
 			if (OutNewCreatedTextures != nullptr)
 			{
diff --git a/Engine/Source/Runtime/Landscape/Private/LandscapeEditInterface.cpp b/Engine/Source/Runtime/Landscape/Private/LandscapeEditInterface.cpp
--- a/Engine/Source/Runtime/Landscape/Private/LandscapeEditInterface.cpp
+++ b/Engine/Source/Runtime/Landscape/Private/LandscapeEditInterface.cpp
@@ -22,6 +22,7 @@ LandscapeEditInterface.cpp: Landscape editing interface
 #include "LandscapeRender.h"
 #include "ComponentReregisterContext.h"
 #include "Algo/Transform.h"
+#include "TextureCompiler.h"
 
 // Channel remapping
 extern const size_t ChannelOffsets[4] = {STRUCT_OFFSET(FColor,R), STRUCT_OFFSET(FColor,G), STRUCT_OFFSET(FColor,B), STRUCT_OFFSET(FColor,A)};
@@ -5663,7 +5664,6 @@ void FLandscapeEditDataInterface::GetXYOffsetDataFast(const int32 X1, const int3
 FLandscapeTextureDataInfo::FLandscapeTextureDataInfo(UTexture2D* InTexture, bool bShouldDirtyPackage)
 :	Texture(InTexture)
 {
-	check(InTexture->IsAsyncCacheComplete());
 	MipInfo.AddZeroed(Texture->Source.GetNumMips());
 	Texture->SetFlags(RF_Transactional);
 	Texture->TemporarilyDisableStreaming();
@@ -5681,6 +5681,9 @@ bool FLandscapeTextureDataInfo::UpdateTextureData()
 		DataSize = sizeof(uint8);
 	}
 
+	// Only wait once
+	bool bNeedToFinishCompilation = true;
+
 	for (int32 i = 0; i < MipInfo.Num(); i++)
 	{
 		if (MipInfo[i].MipData && MipInfo[i].MipUpdateRegions.Num() > 0)
@@ -5692,6 +5695,13 @@ bool FLandscapeTextureDataInfo::UpdateTextureData()
 				break;
 			}
 
+			if (bNeedToFinishCompilation)
+			{
+				// Need to make sure we have a valid Resource
+				FTextureCompilingManager::Get().FinishCompilation({ Texture });
+				bNeedToFinishCompilation = false;
+			}
+
 			const uint32 SrcSizeX = (Texture->Source.GetSizeX()) >> i;
 			const uint32 SrcSizeY = (Texture->Source.GetSizeY()) >> i;
 			const uint32 SrcPitch = (SrcSizeX * DataSize);
diff --git a/Engine/Source/Runtime/Landscape/Private/LandscapeEditLayers.cpp b/Engine/Source/Runtime/Landscape/Private/LandscapeEditLayers.cpp
--- a/Engine/Source/Runtime/Landscape/Private/LandscapeEditLayers.cpp
+++ b/Engine/Source/Runtime/Landscape/Private/LandscapeEditLayers.cpp
@@ -1277,7 +1277,7 @@ typedef FLandscapeLayersRender_RenderThread<FLandscapeLayersWeightmapShaderParam
 
 #if WITH_EDITOR
 
-struct FLandscapeIsTextureFullyStreamedIn
+struct FLandscapeIsTextureReady
 {
 	bool operator()(UTexture2D* InTexture, bool bInWaitForStreaming)
 	{
@@ -1285,7 +1285,6 @@ struct FLandscapeIsTextureFullyStreamedIn
 		InTexture->bForceMiplevelsToBeResident = true;
 		if (bInWaitForStreaming)
 		{
-			FTextureCompilingManager::Get().FinishCompilation({InTexture});
 			InTexture->WaitForStreaming();
 		}
 		return !InTexture->IsDefaultTexture() && InTexture->IsFullyStreamedIn();
@@ -1668,8 +1667,7 @@ void ALandscape::CopyOldDataToDefaultLayer(ALandscapeProxy* InProxy)
 				FMemory::Memcpy(Mip0Data, ExistingMip0Data.GetData(), ExistingMip0Data.Num());
 				DefaultLayerHeightmap->Source.UnlockMip(0);
 
-				DefaultLayerHeightmap->BeginCachePlatformData();
-				DefaultLayerHeightmap->ClearAllCachedCookedPlatformData();
+				DefaultLayerHeightmap->UpdateResource();
 			}
 
 			// Weightmaps
@@ -1748,8 +1746,7 @@ void ALandscape::CopyOldDataToDefaultLayer(ALandscapeProxy* InProxy)
 
 					ProcessedWeightmaps.Add(ComponentWeightmap, NewTextureData);
 
-					NewLayerWeightmapTexture->BeginCachePlatformData();
-					NewLayerWeightmapTexture->ClearAllCachedCookedPlatformData();
+					NewLayerWeightmapTexture->UpdateResource();
 				}
 			}
 		}
@@ -2671,10 +2668,10 @@ bool ALandscape::PrepareLayersBrushTextureResources(bool bInWaitForStreaming, bo
 	}
 
 	bool bReady = true;
-	FLandscapeIsTextureFullyStreamedIn IsTextureFullyStreamedIn;
+	FLandscapeIsTextureReady IsTextureReady;
 	for (UTexture2D* Texture : StreamableTextures)
 	{
-		bReady &= IsTextureFullyStreamedIn(Texture, bInWaitForStreaming);
+		bReady &= IsTextureReady(Texture, bInWaitForStreaming);
 	}
 
 	return bReady;
@@ -2694,38 +2691,22 @@ bool ALandscape::PrepareLayersHeightmapTextureResources(bool bInWaitForStreaming
 
 	Info->ForAllLandscapeProxies([&](ALandscapeProxy* Proxy)
 	{
-		FLandscapeIsTextureFullyStreamedIn IsTextureFullyStreamedIn;
+		FLandscapeIsTextureReady IsTextureReady;
 				
 		for (ULandscapeComponent* Component : Proxy->LandscapeComponents)
 		{
 			UTexture2D* ComponentHeightmap = Component->GetHeightmap();
 
-			IsReady &= IsTextureFullyStreamedIn(ComponentHeightmap, bInWaitForStreaming);
+			IsReady &= IsTextureReady(ComponentHeightmap, bInWaitForStreaming);
 
 			for (const FLandscapeLayer& Layer : LandscapeLayers)
 			{
 				FLandscapeLayerComponentData* ComponentLayerData = Component->GetLayerData(Layer.Guid);
 
 				IsReady &= ComponentLayerData != nullptr;
-
 				if (ComponentLayerData != nullptr)
 				{
-					UTexture2D* LayerHeightmap = ComponentLayerData->HeightmapData.Texture;
-
-					if ((LayerHeightmap->IsAsyncCacheComplete() || bInWaitForStreaming) && LayerHeightmap->Resource == nullptr)
-					{
-						LayerHeightmap->FinishCachePlatformData();
-
-						LayerHeightmap->Resource = LayerHeightmap->CreateResource();
-
-						if (LayerHeightmap->Resource != nullptr)
-						{
-							BeginInitResource(LayerHeightmap->Resource);
-						}
-					}
-
-					IsReady &= IsTextureFullyStreamedIn(LayerHeightmap, bInWaitForStreaming);
-					IsReady &= bInWaitForStreaming || (LayerHeightmap->Resource != nullptr && LayerHeightmap->Resource->IsInitialized());
+					IsReady &= IsTextureReady(ComponentLayerData->HeightmapData.Texture, bInWaitForStreaming);
 				}
 			}
 		}
@@ -3544,20 +3525,19 @@ void ALandscape::ReallocateLayersWeightmaps(const TArray<ULandscapeComponent*>&
 	}
 
 	// Realloc the weightmap so it will create proper texture (if needed) and will set the allocations information
-	TArray<UTexture2D*> NewCreatedTextures;
+	TArray<UTexture*> NewCreatedTextures;
 
 	for (ULandscapeComponent* Component : InLandscapeComponents)
 	{
-		Component->ReallocateWeightmaps(nullptr, false, false, true, false, nullptr, &NewCreatedTextures);
+		Component->ReallocateWeightmaps(nullptr, false, false, false, nullptr, &NewCreatedTextures);
 	}
 
 	// TODO: correctly only recreate what is required instead of everything..
 	//GDisableAutomaticTextureMaterialUpdateDependencies = true;
 
-	for (UTexture2D* Texture : NewCreatedTextures)
+	FTextureCompilingManager::Get().FinishCompilation(NewCreatedTextures);
+	for (UTexture* Texture : NewCreatedTextures)
 	{
-		Texture->FinishCachePlatformData();
-		Texture->PostEditChange();
 		Texture->bForceMiplevelsToBeResident = true;
 		Texture->WaitForStreaming();
 	}
@@ -3688,14 +3668,14 @@ bool ALandscape::PrepareLayersWeightmapTextureResources(bool bInWaitForStreaming
 
 	Info->ForAllLandscapeProxies([&](ALandscapeProxy* Proxy)
 	{
-		FLandscapeIsTextureFullyStreamedIn IsTextureFullyStreamedIn;
+		FLandscapeIsTextureReady IsTextureReady;
 		for (const FLandscapeLayer& Layer : LandscapeLayers)
 		{
 			for (ULandscapeComponent* Component : Proxy->LandscapeComponents)
 			{
 				for (UTexture2D* ComponentWeightmap : Component->GetWeightmapTextures())
 				{
-					if (!IsTextureFullyStreamedIn(ComponentWeightmap, bInWaitForStreaming))
+					if (!IsTextureReady(ComponentWeightmap, bInWaitForStreaming))
 					{
 						IsReady = false;
 						break;
@@ -3710,20 +3690,7 @@ bool ALandscape::PrepareLayersWeightmapTextureResources(bool bInWaitForStreaming
 				{
 					for (UTexture2D* LayerWeightmap : ComponentLayerData->WeightmapData.Textures)
 					{
-						if ((LayerWeightmap->IsAsyncCacheComplete() || bInWaitForStreaming) && LayerWeightmap->Resource == nullptr)
-						{
-							LayerWeightmap->FinishCachePlatformData();
-
-							LayerWeightmap->Resource = LayerWeightmap->CreateResource();
-
-							if (LayerWeightmap->Resource != nullptr)
-							{
-								BeginInitResource(LayerWeightmap->Resource);
-							}
-						}
-
-						IsReady &= IsTextureFullyStreamedIn(LayerWeightmap, bInWaitForStreaming);
-						IsReady &= bInWaitForStreaming || (LayerWeightmap->Resource != nullptr && LayerWeightmap->Resource->IsInitialized());
+						IsReady &= IsTextureReady(LayerWeightmap, bInWaitForStreaming);
 					}
 				}
 			}
@@ -5364,15 +5331,6 @@ void ALandscapeProxy::InitializeLayerWithEmptyContent(const FGuid& InLayerGuid)
 
 		Component->AddDefaultLayerData(InLayerGuid, *ComponentsUsingHeightmap, CreatedHeightmapTextures);
 	}
-
-	// Finish caching
-	for (TPair<UTexture2D*, UTexture2D*> Pair : CreatedHeightmapTextures)
-	{
-		if (Pair.Value != nullptr && !Pair.Value->IsAsyncCacheComplete())
-		{
-			Pair.Value->FinishCachePlatformData();
-		}
-	}
 }
 #endif
 
