Fix Landscape Visibility: Toolbar button not being there + Landscape Edit Layers not supporting Visibility when Landscape Material didn't sample any layers.

diff --git a/Engine/Source/Editor/LandscapeEditor/Private/SLandscapeEditor.cpp b/Engine/Source/Editor/LandscapeEditor/Private/SLandscapeEditor.cpp
--- a/Engine/Source/Editor/LandscapeEditor/Private/SLandscapeEditor.cpp
+++ b/Engine/Source/Editor/LandscapeEditor/Private/SLandscapeEditor.cpp
@@ -224,6 +224,7 @@ void FLandscapeToolKit::BuildToolPalette(FName PaletteName, class FToolBarBuilde
 		ToolBarBuilder.AddToolBarButton(Commands.HydroErosionTool);
 		ToolBarBuilder.AddToolBarButton(Commands.NoiseTool);
 		ToolBarBuilder.AddToolBarButton(Commands.RetopologizeTool);
+		ToolBarBuilder.AddToolBarButton(Commands.VisibilityTool);
 		ToolBarBuilder.AddToolBarButton(Commands.BlueprintBrushTool);
 		ToolBarBuilder.AddToolBarButton(Commands.MirrorTool);
 
diff --git a/Engine/Source/Runtime/Landscape/Private/LandscapeEditLayers.cpp b/Engine/Source/Runtime/Landscape/Private/LandscapeEditLayers.cpp
--- a/Engine/Source/Runtime/Landscape/Private/LandscapeEditLayers.cpp
+++ b/Engine/Source/Runtime/Landscape/Private/LandscapeEditLayers.cpp
@@ -3597,7 +3597,7 @@ void ALandscape::InitializeLayersWeightmapResources()
 {
 	ULandscapeInfo* Info = GetLandscapeInfo();
 
-	if (Info == nullptr || Info->Layers.Num() == 0)
+	if (Info == nullptr)
 	{
 		return;
 	}
@@ -3715,7 +3715,7 @@ int32 ALandscape::RegenerateLayersWeightmaps(const TArray<ULandscapeComponent*>&
 		return 0;
 	}
 
-	if (InLandscapeComponentsToResolve.Num() == 0 || Info == nullptr || Info->Layers.Num() == 0)
+	if (InLandscapeComponentsToResolve.Num() == 0 || Info == nullptr)
 	{
 		return WeightmapUpdateModes;
 	}
