Scaling Landscape Spline Widget to viewport based on distance from camera

diff --git a/Engine/Source/Runtime/Landscape/Private/LandscapeSplines.cpp b/Engine/Source/Runtime/Landscape/Private/LandscapeSplines.cpp
--- a/Engine/Source/Runtime/Landscape/Private/LandscapeSplines.cpp
+++ b/Engine/Source/Runtime/Landscape/Private/LandscapeSplines.cpp
@@ -21,6 +21,7 @@
 #include "LandscapeComponent.h"
 #include "LandscapeVersion.h"
 #include "Components/SplineMeshComponent.h"
+#include "Components/BillboardComponent.h"
 #include "Engine/StaticMesh.h"
 #include "LandscapeSplineProxies.h"
 #include "LandscapeSplinesComponent.h"
@@ -172,7 +173,7 @@ public:
 			ControlPointProxy.HitProxy = nullptr;
 			ControlPointProxy.Location = ControlPoint->Location;
 			ControlPointProxy.Points = ControlPoint->GetPoints();
-			ControlPointProxy.SpriteScale = FMath::Clamp<float>(ControlPoint->Width != 0 ? ControlPoint->Width / 2 : ControlPoint->SideFalloff / 4, 10, 1000);
+			ControlPointProxy.SpriteScale = FMath::Clamp<float>(ControlPoint->Width != 0 ? ControlPoint->Width / 8 : ControlPoint->SideFalloff / 4, 10, 1000);
 			ControlPointProxy.bSelected = ControlPoint->IsSplineSelected();
 			ControlPoints.Add(ControlPointProxy);
 		}
@@ -264,12 +265,20 @@ public:
 					// Draw Sprite
 					if (bDrawControlPointSprite)
 					{
-						const float ControlPointSpriteScale = MyLocalToWorld.GetScaleVector().X * ControlPoint.SpriteScale;
-						const FVector ControlPointSpriteLocation = ControlPointLocation + FVector(0, 0, ControlPointSpriteScale * 0.75f);
+						float ControlPointSpriteScale = MyLocalToWorld.GetScaleVector().X * ControlPoint.SpriteScale;
 						const FLinearColor ControlPointSpriteColor = ControlPoint.bSelected ? SelectedControlPointSpriteColor : FLinearColor::White;
-
 						PDI->SetHitProxy(ControlPoint.HitProxy);
+						float ZoomFactor = FMath::Min<float>(View->ViewMatrices.GetProjectionMatrix().M[0][0], View->ViewMatrices.GetProjectionMatrix().M[1][1]);
+						float Scale = View->WorldToScreen(ControlPointLocation).W * (4.0f / View->UnscaledViewRect.Width() / ZoomFactor);
+						ControlPointSpriteScale *= Scale;
 
+#if WITH_EDITORONLY_DATA
+						ControlPointSpriteScale *= UBillboardComponent::EditorScale;
+#endif
+
+						// Clamping the scale between 10 and the ControlPoint initial scale 
+						ControlPointSpriteScale = FMath::Clamp<float>(ControlPointSpriteScale, 10, ControlPoint.SpriteScale);
+						const FVector ControlPointSpriteLocation = ControlPointLocation + FVector(0, 0, ControlPointSpriteScale * 0.75f);
 						PDI->DrawSprite(
 							ControlPointSpriteLocation,
 							ControlPointSpriteScale,
@@ -282,7 +291,6 @@ public:
 							SE_BLEND_Masked);
 					}
 
-
 					// Draw Lines
 					const FLinearColor ControlPointColor = ControlPoint.bSelected ? SelectedSplineColor : SplineColor;
 
