Fix crash: Use GetResource (because Resource will get deprecated) and get a local copy of the pointer. In case resource poitner gets invalidated between two access. (and check for null)

note: This will avoid the crash but GetResource needs to be sorted out so it returns the proper pointer always (depending on which thread it runs on)

diff --git a/Engine/Source/Runtime/Landscape/Private/LandscapeRender.cpp b/Engine/Source/Runtime/Landscape/Private/LandscapeRender.cpp
--- a/Engine/Source/Runtime/Landscape/Private/LandscapeRender.cpp
+++ b/Engine/Source/Runtime/Landscape/Private/LandscapeRender.cpp
@@ -973,11 +973,11 @@ void FLandscapeRenderSystem::FetchHeightmapLODBiases()
 	for (int32 EntityIndex = 0; EntityIndex < SceneProxies.Num(); EntityIndex++)
 	{
 		FLandscapeComponentSceneProxy* SceneProxy = SceneProxies[EntityIndex];
-		if (SceneProxy)
+		if (SceneProxy && SceneProxy->HeightmapTexture)
 		{
-			if (SceneProxy->HeightmapTexture)
+			if (const FTexture2DResource* TextureResource = (const FTexture2DResource*)SceneProxy->HeightmapTexture->GetResource())
 			{
-				float SectionLODBias = ((FTexture2DResource*)SceneProxy->HeightmapTexture->Resource)->GetCurrentFirstMip();
+				float SectionLODBias = TextureResource->GetCurrentFirstMip();
 				SectionLODBiases[EntityIndex] = SectionLODBias;
 
 				// TODO: support mipmap LOD bias of XY offset map
