Fixed SimpleCollisions on Landscape when Simple Collision mip level is greater than Collision mip level.

diff --git a/Engine/Source/Runtime/Landscape/Private/LandscapeCollision.cpp b/Engine/Source/Runtime/Landscape/Private/LandscapeCollision.cpp
--- a/Engine/Source/Runtime/Landscape/Private/LandscapeCollision.cpp
+++ b/Engine/Source/Runtime/Landscape/Private/LandscapeCollision.cpp
@@ -435,8 +435,9 @@ void ULandscapeHeightfieldCollisionComponent::OnCreatePhysicsState()
 				if(bCreateSimpleCollision)
 				{
 					TUniquePtr<Chaos::FPerShapeData> NewSimpleShape = Chaos::FPerShapeData::CreatePerShapeData(ShapeArray.Num());
+					FVector FinalSimpleCollisionScale(LandscapeScale.X* SimpleCollisionScale, LandscapeScale.Y* SimpleCollisionScale, LandscapeScale.Z* LANDSCAPE_ZSCALE);
 
-					HeightfieldRef->HeightfieldSimple->SetScale(FinalScale);
+					HeightfieldRef->HeightfieldSimple->SetScale(FinalSimpleCollisionScale);
 					TUniquePtr<Chaos::TImplicitObjectTransformed<float, 3>> ChaosSimpleHeightFieldFromCooked = MakeUnique<Chaos::TImplicitObjectTransformed<float, 3>>(MakeSerializable(HeightfieldRef->HeightfieldSimple), Chaos::TRigidTransform<float, 3>(FTransform::Identity));
 
 					FCollisionFilterData QueryFilterDataSimple = QueryFilterData;
