Fix Build (Procedural Plugin + Houdini Engine)

diff --git a/Engine/Source/Runtime/Landscape/Classes/LandscapeInfo.h b/Engine/Source/Runtime/Landscape/Classes/LandscapeInfo.h
--- a/Engine/Source/Runtime/Landscape/Classes/LandscapeInfo.h
+++ b/Engine/Source/Runtime/Landscape/Classes/LandscapeInfo.h
@@ -209,6 +209,12 @@ public:
 	/** Resets all actors, proxies, components registrations */
 	LANDSCAPE_API void Reset();
 
+	/** Recreate all LandscapeInfo objects in given world
+	 *  @param  bMapCheck	Whether to warn about landscape errors
+	 */
+	UE_DEPRECATED(5.0, "Deprecated, please use ULandscapeSubsystem::RecreateLandscapeInfos instead")
+	LANDSCAPE_API static void RecreateLandscapeInfo(UWorld* InWorld, bool bMapCheck);
+
 	/** 
 	 *  Fixes up proxies relative position to landscape actor
 	 *  basically makes sure that each LandscapeProxy RootComponent transform reflects LandscapeSectionOffset value
diff --git a/Engine/Source/Runtime/Landscape/Private/Landscape.cpp b/Engine/Source/Runtime/Landscape/Private/Landscape.cpp
--- a/Engine/Source/Runtime/Landscape/Private/Landscape.cpp
+++ b/Engine/Source/Runtime/Landscape/Private/Landscape.cpp
@@ -2939,6 +2939,11 @@ void ULandscapeInfo::Reset()
 	//SelectedRegion.Empty();
 }
 
+void ULandscapeInfo::RecreateLandscapeInfo(UWorld* InWorld, bool bMapCheck)
+{
+	ULandscapeSubsystem::RecreateLandscapeInfos(InWorld, bMapCheck);
+}
+
 void ULandscapeInfo::FixupProxiesTransform()
 {
 	ALandscape* Landscape = LandscapeActor.Get();
