Landscape: Fix AddMipUpdateRegion

diff --git a/Engine/Source/Runtime/Landscape/Public/LandscapeEdit.h b/Engine/Source/Runtime/Landscape/Public/LandscapeEdit.h
--- a/Engine/Source/Runtime/Landscape/Public/LandscapeEdit.h
+++ b/Engine/Source/Runtime/Landscape/Public/LandscapeEdit.h
@@ -82,7 +82,7 @@ struct FLandscapeTextureDataInfo
 		uint32 Width = 1 + InX2 - InX1;
 		uint32 Height = 1 + InY2 - InY1;
 		// Catch situation where we are updating the whole texture to avoid adding redundant regions once the whole region as been included.
-		if (Width == GetMipSizeX(MipNum) && Width == GetMipSizeY(MipNum))
+		if (Width == GetMipSizeX(MipNum) && Height == GetMipSizeY(MipNum))
 		{
 			MipInfo[MipNum].bFull = true;
 			MipInfo[MipNum].MipUpdateRegions.Reset();
