Landscape: New helper class allows changing Landscape Configuration (Component Size, Num Subsection, WorldPartition GridSize)

diff --git a/Engine/Source/Editor/LandscapeEditor/Private/LandscapeEdModeComponentTools.cpp b/Engine/Source/Editor/LandscapeEditor/Private/LandscapeEdModeComponentTools.cpp
--- a/Engine/Source/Editor/LandscapeEditor/Private/LandscapeEdModeComponentTools.cpp
+++ b/Engine/Source/Editor/LandscapeEditor/Private/LandscapeEdModeComponentTools.cpp
@@ -578,7 +578,6 @@ public:
 						if (LandscapeProxy)
 						{
 							LandscapeComponent = NewObject<ULandscapeComponent>(LandscapeProxy, NAME_None, RF_Transactional);
-							LandscapeProxy->LandscapeComponents.Add(LandscapeComponent);
 							NewComponents.Add(LandscapeComponent);
 							LandscapeComponent->Init(
 								ComponentBase.X, ComponentBase.Y,
@@ -586,17 +585,9 @@ public:
 								LandscapeProxy->NumSubsections,
 								LandscapeProxy->SubsectionSizeQuads
 							);
-							LandscapeComponent->AttachToComponent(LandscapeProxy->GetRootComponent(), FAttachmentTransformRules::KeepRelativeTransform);
-
-							// Assign shared properties
-							LandscapeComponent->UpdatedSharedPropertiesFromActor();
-
-							int32 ComponentVerts = (LandscapeProxy->SubsectionSizeQuads + 1) * LandscapeProxy->NumSubsections;
-							// Update Weightmap Scale Bias
-							LandscapeComponent->WeightmapScaleBias = FVector4(1.0f / (float)ComponentVerts, 1.0f / (float)ComponentVerts, 0.5f / (float)ComponentVerts, 0.5f / (float)ComponentVerts);
-							LandscapeComponent->WeightmapSubsectionOffset = (float)(LandscapeComponent->SubsectionSizeQuads + 1) / (float)ComponentVerts;
 
 							TArray<FColor> HeightData;
+							const int32 ComponentVerts = (LandscapeComponent->SubsectionSizeQuads + 1) * LandscapeComponent->NumSubsections;
 							HeightData.Empty(FMath::Square(ComponentVerts));
 							HeightData.AddZeroed(FMath::Square(ComponentVerts));
 							LandscapeComponent->InitHeightmapData(HeightData, true);
diff --git a/Engine/Source/Editor/LandscapeEditor/Private/LandscapeEdModeSplineTools.cpp b/Engine/Source/Editor/LandscapeEditor/Private/LandscapeEdModeSplineTools.cpp
--- a/Engine/Source/Editor/LandscapeEditor/Private/LandscapeEdModeSplineTools.cpp
+++ b/Engine/Source/Editor/LandscapeEditor/Private/LandscapeEdModeSplineTools.cpp
@@ -98,11 +98,6 @@ public:
 	virtual void SetEditRenderType() override { GLandscapeEditRenderMode = ELandscapeEditRenderMode::None | (GLandscapeEditRenderMode & ELandscapeEditRenderMode::BitMaskForMask); }
 	virtual bool SupportsMask() override { return false; }
 
-	void CreateSplineComponent(ALandscapeProxy* Proxy, FVector Scale3D)
-	{
-		Proxy->CreateSplineComponent(Scale3D);
-	}
-
 	void UpdatePropertiesWindows()
 	{
 		if (GLevelEditorModeTools().IsModeActive(EdMode->GetID()))
@@ -981,7 +976,7 @@ public:
 			{
 				if (!Landscape->SplineComponent)
 				{
-					CreateSplineComponent(Landscape, FVector(1.0f) / Landscape->GetRootComponent()->GetRelativeScale3D());
+					Landscape->CreateSplineComponent();
 					check(Landscape->SplineComponent);
 				}
 				SplinesComponent = Landscape->SplineComponent;
@@ -2139,7 +2134,7 @@ public:
 		}
 		if (!Landscape->SplineComponent)
 		{
-			CreateSplineComponent(Landscape, FVector(1.0f) / Landscape->GetRootComponent()->GetRelativeScale3D());
+			Landscape->CreateSplineComponent();
 			check(Landscape->SplineComponent);
 		}
 		Landscape->SplineComponent->Modify();
diff --git a/Engine/Source/Editor/LandscapeEditor/Private/LandscapeEditorDetailCustomization_NewLandscape.cpp b/Engine/Source/Editor/LandscapeEditor/Private/LandscapeEditorDetailCustomization_NewLandscape.cpp
--- a/Engine/Source/Editor/LandscapeEditor/Private/LandscapeEditorDetailCustomization_NewLandscape.cpp
+++ b/Engine/Source/Editor/LandscapeEditor/Private/LandscapeEditorDetailCustomization_NewLandscape.cpp
@@ -890,7 +890,7 @@ FReply FLandscapeEditorDetailCustomization_NewLandscape::OnCreateButtonClicked()
 		LandscapeEdMode->SetCurrentTool("Sculpt"); // change to sculpting mode and tool
 		LandscapeEdMode->SetCurrentLayer(0);
 
-		LandscapeEdMode->GetWorld()->GetSubsystem<ULandscapeSubsystem>()->UpdateGrid(LandscapeInfo, UISettings->WorldPartitionGridSize);
+		LandscapeEdMode->GetWorld()->GetSubsystem<ULandscapeSubsystem>()->ChangeGridSize(LandscapeInfo, UISettings->WorldPartitionGridSize);
 
 		if (LandscapeEdMode->CurrentToolTarget.LandscapeInfo.IsValid())
 		{
diff --git a/Engine/Source/Editor/LandscapeEditor/Private/LandscapeEditorDetailCustomization_ResizeLandscape.cpp b/Engine/Source/Editor/LandscapeEditor/Private/LandscapeEditorDetailCustomization_ResizeLandscape.cpp
--- a/Engine/Source/Editor/LandscapeEditor/Private/LandscapeEditorDetailCustomization_ResizeLandscape.cpp
+++ b/Engine/Source/Editor/LandscapeEditor/Private/LandscapeEditorDetailCustomization_ResizeLandscape.cpp
@@ -19,8 +19,7 @@
 
 #define LOCTEXT_NAMESPACE "LandscapeEditor.ResizeLandscape"
 
-const int32 FLandscapeEditorDetailCustomization_ResizeLandscape::SectionSizes[] = {7, 15, 31, 63, 127, 255};
-const int32 FLandscapeEditorDetailCustomization_ResizeLandscape::NumSections[] = {1, 2};
+#include "LandscapeConfigHelper.h"
 
 TSharedRef<IDetailCustomization> FLandscapeEditorDetailCustomization_ResizeLandscape::MakeInstance()
 {
@@ -305,9 +304,9 @@ TSharedRef<SWidget> FLandscapeEditorDetailCustomization_ResizeLandscape::GetSect
 {
 	FMenuBuilder MenuBuilder(true, NULL);
 
-	for (int32 i = 0; i < UE_ARRAY_COUNT(SectionSizes); i++)
+	for (int32 i = 0; i < UE_ARRAY_COUNT(FLandscapeConfig::SubsectionSizeQuadsValues); i++)
 	{
-		MenuBuilder.AddMenuEntry(FText::Format(LOCTEXT("NxNQuads", "{0}x{0} Quads"), FText::AsNumber(SectionSizes[i])), FText::GetEmpty(), FSlateIcon(), FExecuteAction::CreateStatic(&OnChangeSectionSize, PropertyHandle, SectionSizes[i]));
+		MenuBuilder.AddMenuEntry(FText::Format(LOCTEXT("NxNQuads", "{0}x{0} Quads"), FText::AsNumber(FLandscapeConfig::SubsectionSizeQuadsValues[i])), FText::GetEmpty(), FSlateIcon(), FExecuteAction::CreateStatic(&OnChangeSectionSize, PropertyHandle, FLandscapeConfig::SubsectionSizeQuadsValues[i]));
 	}
 
 	return MenuBuilder.MakeWidget();
@@ -372,12 +371,12 @@ TSharedRef<SWidget> FLandscapeEditorDetailCustomization_ResizeLandscape::GetSect
 {
 	FMenuBuilder MenuBuilder(true, NULL);
 
-	for (int32 i = 0; i < UE_ARRAY_COUNT(NumSections); i++)
+	for (int32 i = 0; i < UE_ARRAY_COUNT(FLandscapeConfig::NumSectionValues); i++)
 	{
 		FFormatNamedArguments Args;
-		Args.Add(TEXT("Width"), NumSections[i]);
-		Args.Add(TEXT("Height"), NumSections[i]);
-		MenuBuilder.AddMenuEntry(FText::Format(NumSections[i] == 1 ? LOCTEXT("1x1Section", "{Width}\u00D7{Height} Section") : LOCTEXT("NxNSections", "{Width}\u00D7{Height} Sections"), Args), FText::GetEmpty(), FSlateIcon(), FExecuteAction::CreateStatic(&OnChangeSectionsPerComponent, PropertyHandle, NumSections[i]));
+		Args.Add(TEXT("Width"), FLandscapeConfig::NumSectionValues[i]);
+		Args.Add(TEXT("Height"), FLandscapeConfig::NumSectionValues[i]);
+		MenuBuilder.AddMenuEntry(FText::Format(FLandscapeConfig::NumSectionValues[i] == 1 ? LOCTEXT("1x1Section", "{Width}\u00D7{Height} Section") : LOCTEXT("NxNSections", "{Width}\u00D7{Height} Sections"), Args), FText::GetEmpty(), FSlateIcon(), FExecuteAction::CreateStatic(&OnChangeSectionsPerComponent, PropertyHandle, FLandscapeConfig::NumSectionValues[i]));
 	}
 
 	return MenuBuilder.MakeWidget();
diff --git a/Engine/Source/Runtime/Foliage/Private/InstancedFoliage.cpp b/Engine/Source/Runtime/Foliage/Private/InstancedFoliage.cpp
--- a/Engine/Source/Runtime/Foliage/Private/InstancedFoliage.cpp
+++ b/Engine/Source/Runtime/Foliage/Private/InstancedFoliage.cpp
@@ -2851,94 +2851,128 @@ void AInstancedFoliageActor::MoveInstancesForComponentToLevel(UActorComponent* I
 }
 
 void AInstancedFoliageActor::MoveInstancesToNewComponent(UPrimitiveComponent* InOldComponent, const FBox& InBoxWithInstancesToMove, UPrimitiveComponent* InNewComponent)
-{	
-	const auto OldBaseId = InstanceBaseCache.GetInstanceBaseId(InOldComponent);
-	if (OldBaseId == FFoliageInstanceBaseCache::InvalidBaseId)
+{
+	MoveInstancesToNewComponent(InOldComponent, InNewComponent, [InBoxWithInstancesToMove](const FFoliageInfo& FoliageInfo)
 	{
-		// This foliage actor has no instances with specified base
-		return;
-	}
-
-	AInstancedFoliageActor* TargetIFA = AInstancedFoliageActor::GetInstancedFoliageActorForLevel(InNewComponent->GetOwner()->GetLevel(), true);
-	TArray<int32> InstancesToMove;
+		return FoliageInfo.GetInstancesOverlappingBox(InBoxWithInstancesToMove);
+	});
+}
 
-	for (auto& Pair : FoliageInfos)
+void AInstancedFoliageActor::MoveInstancesToNewComponent(UPrimitiveComponent* InOldComponent, UPrimitiveComponent* InNewComponent)
+{
+	MoveInstancesToNewComponent(InOldComponent, InNewComponent, [](const FFoliageInfo& FoliageInfo)
 	{
-		InstancesToMove.Reset();
-
-		FFoliageInfo& Info = *Pair.Value;
-		
-		InstancesToMove = Info.GetInstancesOverlappingBox(InBoxWithInstancesToMove);
-		
-		FFoliageInfo* TargetMeshInfo = nullptr;
-		UFoliageType* TargetFoliageType = TargetIFA->AddFoliageType(Pair.Key, &TargetMeshInfo);
-
-		// Add the foliage to the new level
-		for (int32 InstanceIndex : InstancesToMove)
+		TArray<int32> InstanceIndices;
+		InstanceIndices.SetNum(FoliageInfo.Instances.Num());
+		for (int32 InstanceIndex = 0; InstanceIndex < FoliageInfo.Instances.Num(); ++InstanceIndex)
 		{
-			if (Info.Instances.IsValidIndex(InstanceIndex))
-			{
-				FFoliageInstance NewInstance = Info.Instances[InstanceIndex];
-				TargetMeshInfo->AddInstance(TargetIFA, TargetFoliageType, NewInstance, InNewComponent);
-			}
+			InstanceIndices[InstanceIndex] = InstanceIndex;
 		}
-
-		TargetMeshInfo->Refresh(TargetIFA, true, true);
-		
-		// Remove from old level
-		Info.RemoveInstances(this, InstancesToMove, true);
-	}
+		return MoveTemp(InstanceIndices);
+	});
 }
 
-void AInstancedFoliageActor::MoveInstancesToNewComponent(UPrimitiveComponent* InOldComponent, UPrimitiveComponent* InNewComponent)
-{
-	AInstancedFoliageActor* TargetIFA = AInstancedFoliageActor::GetInstancedFoliageActorForLevel(InNewComponent->GetOwner()->GetLevel(), true);
-
-	const auto OldBaseId = this->InstanceBaseCache.GetInstanceBaseId(InOldComponent);
+void AInstancedFoliageActor::MoveInstancesToNewComponent(UPrimitiveComponent* InOldComponent, UPrimitiveComponent* InNewComponent, TFunctionRef<TArray<int32>(const FFoliageInfo&)> GetInstancesToMoveFunc)
+{	
+	const auto OldBaseId = InstanceBaseCache.GetInstanceBaseId(InOldComponent);
 	if (OldBaseId == FFoliageInstanceBaseCache::InvalidBaseId)
 	{
-		// This foliage actor has no any instances with specified base
+		// This foliage actor has no instances with specified base
 		return;
 	}
+		
+	ULevel* IFALevel = InNewComponent->GetOwner()->GetLevel();
+	UWorld* IFAWorld = IFALevel->GetWorld();
+	
+	TArray<int32> InstancesToDelete;
+	TSet<int32> InstancesToUpdateBase;
+	TMap<AInstancedFoliageActor*, TArray<FFoliageInstance*>> PerIFAInstancesToMove;
 
-	const auto NewBaseId = TargetIFA->InstanceBaseCache.AddInstanceBaseId(InNewComponent);
-
+	// If Modify was called on this IFA
+	bool bModified = false;
+	
 	for (auto& Pair : FoliageInfos)
-	{
+	{			
 		FFoliageInfo& Info = *Pair.Value;
-
-		TSet<int32> InstanceSet;
-		if (Info.ComponentHash.RemoveAndCopyValue(OldBaseId, InstanceSet) && InstanceSet.Num())
+		// Make sure we indeed have Instances that have the OldComponent as a base for this specific FoliageType
+		TSet<int32>* OldInstanceSet = Info.ComponentHash.Find(OldBaseId);
+		if(OldInstanceSet && OldInstanceSet->Num())
 		{
-			// For same FoliageActor can just remap the instances, otherwise we have to do a more complex move
-			if (TargetIFA == this)
+			FFoliageInstanceBaseId NewBaseId = FFoliageInstanceBaseCache::InvalidBaseId;
+			TArray<int32> PotentialInstances = GetInstancesToMoveFunc(Info);
+
+			// Reset temp containers
+			PerIFAInstancesToMove.Reset();
+			InstancesToUpdateBase.Reset();
+			InstancesToDelete.Reset(PotentialInstances.Num());
+
+			// Cumulate Instances to move Per IFA
+			for (int32 InstanceIndex : PotentialInstances)
 			{
-				// Update the instances
-				for (int32 InstanceIndex : InstanceSet)
+				if (Info.Instances.IsValidIndex(InstanceIndex) && OldInstanceSet->Contains(InstanceIndex))
 				{
-					Info.Instances[InstanceIndex].BaseId = NewBaseId;
-				}
+					FFoliageInstance& InstanceToMove = Info.Instances[InstanceIndex];
+					const bool bCreate = true;
+					if (AInstancedFoliageActor* TargetIFA = AInstancedFoliageActor::Get(IFAWorld, bCreate, IFALevel, InstanceToMove.Location))
+					{
+						// Call Modify only once
+						if (!bModified)
+						{
+							Modify();
+							bModified = true;
+						}
 
-				// Update the hash
-				Info.ComponentHash.Add(NewBaseId, MoveTemp(InstanceSet));
+						// Same IFA just update the Base 
+						if (this == TargetIFA)
+						{
+							if (NewBaseId == FFoliageInstanceBaseCache::InvalidBaseId)
+							{
+								NewBaseId = InstanceBaseCache.AddInstanceBaseId(InNewComponent);
+							}
+							InstanceToMove.BaseComponent = InNewComponent;
+							InstanceToMove.BaseId = NewBaseId;
+							InstancesToUpdateBase.Add(InstanceIndex);
+							OldInstanceSet->Remove(InstanceIndex);
+						}
+						else
+						{
+							TArray<FFoliageInstance*>& InstancesToMove = PerIFAInstancesToMove.FindOrAdd(TargetIFA);
+							InstancesToMove.Add(&InstanceToMove);
+							InstancesToDelete.Add(InstanceIndex);
+						}
+					}
+				}
 			}
-			else
+
+			// Add Instances to IFAs
+			for (auto& IFAPair : PerIFAInstancesToMove)
 			{
+				AInstancedFoliageActor* TargetIFA = IFAPair.Key;
+				TargetIFA->Modify();
 				FFoliageInfo* TargetMeshInfo = nullptr;
 				UFoliageType* TargetFoliageType = TargetIFA->AddFoliageType(Pair.Key, &TargetMeshInfo);
-
-				// Add the foliage to the new level
-				for (int32 InstanceIndex : InstanceSet)
+				for (FFoliageInstance* InstanceToMove : IFAPair.Value)
 				{
-					FFoliageInstance NewInstance = Info.Instances[InstanceIndex];
-					NewInstance.BaseId = NewBaseId;
-					TargetMeshInfo->AddInstance(TargetIFA, TargetFoliageType, NewInstance);
+					TargetMeshInfo->AddInstance(TargetIFA, TargetFoliageType, *InstanceToMove, InNewComponent);
 				}
-
 				TargetMeshInfo->Refresh(TargetIFA, true, true);
-
-				// Remove from old level
-				Info.RemoveInstances(this, InstanceSet.Array(), true);
+			}
+									
+			// Remove old set if empty
+			if (!OldInstanceSet->Num())
+			{
+				Info.ComponentHash.Remove(OldBaseId);
+			}
+			// Add instances that are still in the old ifa but with a new base
+			if (InstancesToUpdateBase.Num())
+			{
+				check(NewBaseId != FFoliageInstanceBaseCache::InvalidBaseId);
+				Info.ComponentHash.Add(NewBaseId, InstancesToUpdateBase);
+			}
+			// Remove from old IFA
+			if (InstancesToDelete.Num())
+			{
+				Info.RemoveInstances(this, InstancesToDelete, true);
 			}
 		}
 	}
diff --git a/Engine/Source/Runtime/Foliage/Public/InstancedFoliageActor.h b/Engine/Source/Runtime/Foliage/Public/InstancedFoliageActor.h
--- a/Engine/Source/Runtime/Foliage/Public/InstancedFoliageActor.h
+++ b/Engine/Source/Runtime/Foliage/Public/InstancedFoliageActor.h
@@ -295,6 +295,7 @@ private:
 	void OnApplyLevelTransform(const FTransform& InTransform);
 	void OnPostApplyLevelOffset(ULevel* InLevel, UWorld* InWorld, const FVector& InOffset, bool bWorldShift);
 	void OnPostWorldInitialization(UWorld* World, const UWorld::InitializationValues IVS);
+	void MoveInstancesToNewComponent(UPrimitiveComponent* InOldComponent, UPrimitiveComponent* InNewComponent, TFunctionRef<TArray<int32>(const FFoliageInfo&)> GetInstancesToMoveFunc);
 #endif
 private:
 #if WITH_EDITOR
diff --git a/Engine/Source/Runtime/Landscape/Classes/LandscapeComponent.h b/Engine/Source/Runtime/Landscape/Classes/LandscapeComponent.h
--- a/Engine/Source/Runtime/Landscape/Classes/LandscapeComponent.h
+++ b/Engine/Source/Runtime/Landscape/Classes/LandscapeComponent.h
@@ -1030,6 +1030,8 @@ public:
 	/** Extends passed region with this component section size */
 	LANDSCAPE_API void GetComponentExtent(int32& MinX, int32& MinY, int32& MaxX, int32& MaxY) const;
 
+	LANDSCAPE_API FIntRect GetComponentExtent() const;
+
 	/** Updates navigation properties to match landscape's master switch */
 	void UpdateNavigationRelevance();
 
diff --git a/Engine/Source/Runtime/Landscape/Classes/LandscapeInfo.h b/Engine/Source/Runtime/Landscape/Classes/LandscapeInfo.h
--- a/Engine/Source/Runtime/Landscape/Classes/LandscapeInfo.h
+++ b/Engine/Source/Runtime/Landscape/Classes/LandscapeInfo.h
@@ -250,6 +250,12 @@ public:
 	/** Moves all Splines to target level. Creates ULandscapeSplineComponent if needed. */
 	LANDSCAPE_API void MoveSplinesToLevel(ULandscapeSplinesComponent* InSplineComponent, ULevel* TargetLevel);
 
+	/** Moves Splines connected to this control point to target Proxy. Creates ULandscapeSplineComponent if needed. */
+	LANDSCAPE_API void MoveSplineToProxy(ULandscapeSplineControlPoint* InControlPoint, ALandscapeProxy* InLandscapeProxy);
+
+	/** Moves all Splines to target Proxy. Creates ULandscapeSplineComponent if needed */
+	LANDSCAPE_API void MoveSplinesToProxy(ULandscapeSplinesComponent* InSplineComponent, ALandscapeProxy* InLandscapeProxy);
+
 	/** Will call UpdateAllComponentMaterialInstances on all LandscapeProxies */
 	LANDSCAPE_API void UpdateAllComponentMaterialInstances();
 
diff --git a/Engine/Source/Runtime/Landscape/Classes/LandscapeProxy.h b/Engine/Source/Runtime/Landscape/Classes/LandscapeProxy.h
--- a/Engine/Source/Runtime/Landscape/Classes/LandscapeProxy.h
+++ b/Engine/Source/Runtime/Landscape/Classes/LandscapeProxy.h
@@ -973,6 +973,7 @@ public:
 #if WITH_EDITOR
 	LANDSCAPE_API int32 GetOutdatedGrassMapCount() const;
 	LANDSCAPE_API void BuildGrassMaps(struct FScopedSlowTask* InSlowTask = nullptr);
+	LANDSCAPE_API void CreateSplineComponent();
 	LANDSCAPE_API void CreateSplineComponent(const FVector& Scale3D);
 
 	virtual bool CanEditChange(const FProperty* InProperty) const override;
diff --git a/Engine/Source/Runtime/Landscape/Private/LandscapeConfigHelper.cpp b/Engine/Source/Runtime/Landscape/Private/LandscapeConfigHelper.cpp
new file mode 100644
--- /dev/null
+++ b/Engine/Source/Runtime/Landscape/Private/LandscapeConfigHelper.cpp
@@ -0,0 +1,509 @@
+// Copyright Epic Games, Inc. All Rights Reserved.
+#include "LandscapeConfigHelper.h"
+
+#if WITH_EDITOR
+
+#include "Engine/World.h"
+#include "Misc/ScopeExit.h"
+#include "LandscapeInfo.h"
+#include "Landscape.h"
+#include "LandscapeProxy.h"
+#include "LandscapeStreamingProxy.h"
+#include "LandscapeSplinesComponent.h"
+#include "LandscapeEdit.h"
+#include "LandscapeDataAccess.h"
+#include "LandscapeSubsystem.h"
+#include "ActorPartition/ActorPartitionSubsystem.h"
+
+int32 FLandscapeConfig::NumSectionValues[2] = { 1, 2 };
+int32 FLandscapeConfig::SubsectionSizeQuadsValues[6] = { 7, 15, 31, 63, 127, 255 };
+
+FLandscapeConfig::FLandscapeConfig(ULandscapeInfo* InLandscapeInfo)
+	: FLandscapeConfig(InLandscapeInfo->ComponentNumSubsections, InLandscapeInfo->SubsectionSizeQuads, InLandscapeInfo->LandscapeActor->GridSize / InLandscapeInfo->ComponentSizeQuads)
+{
+
+}
+
+bool FLandscapeConfigChange::Validate() const
+{
+	bool bValidNumSections = false;
+	for (int32 Index = 0; Index < UE_ARRAY_COUNT(NumSectionValues); ++Index)
+	{
+		if (ComponentNumSubsections == NumSectionValues[Index])
+		{
+			bValidNumSections = true;
+			break;
+		}
+	}
+	bool bValidSectionSizeQuads = false;
+	for (int32 Index = 0; Index < UE_ARRAY_COUNT(SubsectionSizeQuadsValues); ++Index)
+	{
+		if (SubsectionSizeQuads == SubsectionSizeQuadsValues[Index])
+		{
+			bValidSectionSizeQuads = true;
+			break;
+		}
+	}
+	return bValidNumSections && bValidSectionSizeQuads;
+}
+
+ALandscapeProxy* FLandscapeConfigHelper::FindOrAddLandscapeStreamingProxy(ULandscapeInfo* InLandscapeInfo, const FIntPoint& InSectionBase)
+{
+	UWorld* World = InLandscapeInfo->LandscapeActor->GetWorld();
+	UActorPartitionSubsystem* ActorPartitionSubsystem = World->GetSubsystem<UActorPartitionSubsystem>();
+	const uint32 GridSize = InLandscapeInfo->LandscapeActor->GridSize;
+
+	UActorPartitionSubsystem::FCellCoord CellCoord = UActorPartitionSubsystem::FCellCoord::GetCellCoord(InSectionBase, World->PersistentLevel, GridSize);
+	return FLandscapeConfigHelper::FindOrAddLandscapeStreamingProxy(ActorPartitionSubsystem, InLandscapeInfo, CellCoord, GridSize);
+}
+
+ALandscapeProxy* FLandscapeConfigHelper::FindOrAddLandscapeStreamingProxy(UActorPartitionSubsystem* InActorPartitionSubsystem, ULandscapeInfo* InLandscapeInfo, const UActorPartitionSubsystem::FCellCoord& InCellCoord, uint32 InSearchGridSize)
+{
+	ALandscape* Landscape = InLandscapeInfo->LandscapeActor.Get();
+	check(Landscape);
+
+	auto LandscapeProxyCreated = [InCellCoord, Landscape](APartitionActor* PartitionActor)
+	{
+		const FIntPoint CellLocation(InCellCoord.X * Landscape->GridSize, InCellCoord.Y * Landscape->GridSize);
+
+		ALandscapeProxy* LandscapeProxy = CastChecked<ALandscapeProxy>(PartitionActor);
+		// copy shared properties to this new proxy
+		LandscapeProxy->GetSharedProperties(Landscape);
+		const FVector ProxyLocation = Landscape->GetActorLocation() + FVector(CellLocation.X * Landscape->GetActorRelativeScale3D().X, CellLocation.Y * Landscape->GetActorRelativeScale3D().Y, 0.0f);
+
+		LandscapeProxy->CreateLandscapeInfo();
+		LandscapeProxy->SetActorLocationAndRotation(ProxyLocation, Landscape->GetActorRotation());
+		LandscapeProxy->LandscapeSectionOffset = FIntPoint(CellLocation.X, CellLocation.Y);
+	};
+
+	const bool bCreate = true;
+	const bool bBoundsSearch = false;
+
+	// SearchGridSize might be different than Landscape->GridSize when we are trying to recycle Proxies if the GridSize is being changed.
+	ALandscapeProxy* LandscapeProxy = Cast<ALandscapeProxy>(InActorPartitionSubsystem->GetActor(ALandscapeStreamingProxy::StaticClass(), InCellCoord, bCreate, InLandscapeInfo->LandscapeGuid, InSearchGridSize, bBoundsSearch, LandscapeProxyCreated));
+	if (LandscapeProxy && LandscapeProxy->GridSize != Landscape->GridSize)
+	{
+		LandscapeProxy->Modify();
+		LandscapeProxy->GridSize = Landscape->GridSize;
+	}
+	return LandscapeProxy;
+}
+
+bool FLandscapeConfigHelper::ChangeGridSize(ULandscapeInfo* InLandscapeInfo, uint32 InNewGridSizeInComponents, TSet<AActor*>& OutActorsToDelete)
+{
+	check(InLandscapeInfo);
+
+	const uint32 GridSize = InLandscapeInfo->GetGridSize(InNewGridSizeInComponents);
+	const uint32 OldGridSize = InLandscapeInfo->LandscapeActor->GridSize;
+	
+	InLandscapeInfo->LandscapeActor->Modify();
+	InLandscapeInfo->LandscapeActor->GridSize = GridSize;
+
+	FIntRect Extent;
+	InLandscapeInfo->GetLandscapeExtent(Extent.Min.X, Extent.Min.Y, Extent.Max.X, Extent.Max.Y);
+	const FBox Bounds(FVector(Extent.Min), FVector(Extent.Max));
+
+	UWorld* World = InLandscapeInfo->LandscapeActor->GetWorld();
+	UActorPartitionSubsystem* ActorPartitionSubsystem = World->GetSubsystem<UActorPartitionSubsystem>();
+
+	TArray<ULandscapeComponent*> LandscapeComponents;
+	LandscapeComponents.Reserve(InLandscapeInfo->XYtoComponentMap.Num());
+	InLandscapeInfo->ForAllLandscapeComponents([&LandscapeComponents](ULandscapeComponent* LandscapeComponent)
+	{
+		LandscapeComponents.Add(LandscapeComponent);
+	});
+
+	TSet<ALandscapeProxy*> ProxiesToDelete;
+
+	FActorPartitionGridHelper::ForEachIntersectingCell(ALandscapeStreamingProxy::StaticClass(), Extent, World->PersistentLevel, [ActorPartitionSubsystem, InLandscapeInfo, InNewGridSizeInComponents, &LandscapeComponents, OldGridSize, &ProxiesToDelete](const UActorPartitionSubsystem::FCellCoord& CellCoord, const FIntRect& CellBounds)
+	{
+		TArray<ULandscapeComponent*> ComponentsToMove;
+		const int32 MaxComponents = (int32)(InNewGridSizeInComponents * InNewGridSizeInComponents);
+		ComponentsToMove.Reserve(MaxComponents);
+		for (int32 i = 0; i < LandscapeComponents.Num();)
+		{
+			ULandscapeComponent* LandscapeComponent = LandscapeComponents[i];
+			if (CellBounds.Contains(LandscapeComponent->GetSectionBase()))
+			{
+				ComponentsToMove.Add(LandscapeComponent);
+				LandscapeComponents.RemoveAtSwap(i);
+				ProxiesToDelete.Add(LandscapeComponent->GetTypedOuter<ALandscapeProxy>());
+			}
+			else
+			{
+				i++;
+			}
+		}
+
+		check(ComponentsToMove.Num() <= MaxComponents);
+		if (ComponentsToMove.Num())
+		{
+			ALandscapeProxy* LandscapeProxy = FLandscapeConfigHelper::FindOrAddLandscapeStreamingProxy(ActorPartitionSubsystem, InLandscapeInfo, CellCoord, OldGridSize);
+			check(LandscapeProxy);
+
+			// @todo_ow: handle mismatching materials between source and dest proxy (probably override material on component)
+			InLandscapeInfo->MoveComponentsToProxy(ComponentsToMove, LandscapeProxy);
+		}
+
+		return true;
+	}, GridSize);
+
+	// @todo_ow: This is probably not the final solution
+	MoveSplinesToLandscape(InLandscapeInfo, InLandscapeInfo->LandscapeActor.Get(), 1.0f);
+
+	// Only delete Proxies that where not reused
+	for (ALandscapeProxy* ProxyToDelete : ProxiesToDelete)
+	{
+		if (ProxyToDelete->LandscapeComponents.Num() > 0 || ProxyToDelete->IsA<ALandscape>())
+		{
+			check(ProxyToDelete->GridSize == GridSize);
+			continue;
+		}
+
+		OutActorsToDelete.Add(ProxyToDelete);
+	}
+
+	return true;
+}
+
+void FLandscapeConfigHelper::ExtractLandscapeData(ULandscapeInfo* InLandscapeInfo, const FIntRect& InRegion, const FGuid& InLayerGuid, TArray<uint16>& OutHeightData, TArray<FLandscapeImportLayerInfo>& OutImportMaterialLayerInfos)
+{
+	FScopedSetLandscapeEditingLayer LayerScope(InLandscapeInfo->LandscapeActor.Get(), InLayerGuid);
+	FLandscapeEditDataInterface LandscapeEdit(InLandscapeInfo);
+	
+	int32 VertsX = InRegion.Width() + 1;
+	int32 VertsY = InRegion.Height() + 1;
+	OutHeightData.Reset();
+	OutHeightData.AddZeroed(VertsX * VertsY);
+	
+	{
+		FIntRect CopyRegion(InRegion);
+		LandscapeEdit.GetHeightData(CopyRegion.Min.X, CopyRegion.Min.Y, CopyRegion.Max.X, CopyRegion.Max.Y, OutHeightData.GetData(), 0);
+	}
+
+	for (const FLandscapeInfoLayerSettings& LayerSettings : InLandscapeInfo->Layers)
+	{
+		if (LayerSettings.LayerInfoObj != NULL)
+		{
+			auto ImportLayerInfo = new(OutImportMaterialLayerInfos) FLandscapeImportLayerInfo(LayerSettings);
+			ImportLayerInfo->LayerData.Reset();
+			ImportLayerInfo->LayerData.AddZeroed(VertsX * VertsY);
+			FIntRect CopyRegion(InRegion);
+			LandscapeEdit.GetWeightData(LayerSettings.LayerInfoObj, CopyRegion.Min.X, CopyRegion.Min.Y, CopyRegion.Max.X, CopyRegion.Max.Y, ImportLayerInfo->LayerData.GetData(), 0);
+		}
+	}
+};
+
+void FLandscapeConfigHelper::CopyRegionToComponent(ULandscapeInfo* InLandscapeInfo, const FIntRect& InRegion, bool bInResample, ULandscapeComponent* InComponent)
+{
+	TArray<uint16> SrcHeightData;
+	TArray<FLandscapeImportLayerInfo> SrcWeightData;
+	
+	ULandscapeInfo* NewLandscapeInfo = InComponent->GetLandscapeInfo();
+	
+	FIntRect NewRegion = InComponent->GetComponentExtent();
+		
+	// Create Heightmap Texture
+	check(InComponent->GetHeightmap(false) == nullptr);
+	int32 HeightmapSize = (InComponent->SubsectionSizeQuads + 1) * InComponent->NumSubsections;
+	
+	TArray<uint16> NewHeightData;
+	TArray<uint8> NewWeightData;
+
+	if (!NewLandscapeInfo->LandscapeActor->bCanHaveLayersContent)
+	{
+		InComponent->HeightmapScaleBias = FVector4(1.0f / (float)HeightmapSize, 1.0f / (float)HeightmapSize, 0.0f, 0.0f);
+		InComponent->SetHeightmap(InComponent->GetLandscapeProxy()->CreateLandscapeTexture(HeightmapSize, HeightmapSize, TEXTUREGROUP_Terrain_Heightmap, TSF_BGRA8));
+
+		ExtractLandscapeData(InLandscapeInfo, InRegion, FGuid(), SrcHeightData, SrcWeightData);
+
+		CopyData(SrcHeightData, NewHeightData, InRegion, NewRegion, bInResample);
+
+		{
+			FLandscapeEditDataInterface LandscapeEdit(NewLandscapeInfo);
+			LandscapeEdit.SetHeightData(NewRegion.Min.X, NewRegion.Min.Y, NewRegion.Max.X, NewRegion.Max.Y, NewHeightData.GetData(), 0, false);
+		}
+
+		
+		for (const FLandscapeImportLayerInfo& Layer : SrcWeightData)
+		{
+			FLandscapeEditDataInterface LandscapeEdit(NewLandscapeInfo);
+			CopyData(Layer.LayerData, NewWeightData, InRegion, NewRegion, bInResample);
+			LandscapeEdit.SetAlphaData(Layer.LayerInfo, NewRegion.Min.X, NewRegion.Min.Y, NewRegion.Max.X, NewRegion.Max.Y, NewWeightData.GetData(), 0);
+		}
+	}
+	else
+	{
+		TArray<FColor> HeightInitData;
+		HeightInitData.SetNum(FMath::Square(HeightmapSize));
+		const uint16 DefaultHeight = LandscapeDataAccess::GetTexHeight(0.f);
+		FColor DefaultPackedHeight = LandscapeDataAccess::PackHeight(DefaultHeight);
+		for (int32 Index = 0; Index < HeightInitData.Num(); ++Index)
+		{
+			HeightInitData[Index] = DefaultPackedHeight;
+		}
+		InComponent->InitHeightmapData(HeightInitData, false);
+
+		for (const FLandscapeLayer& LandscapeLayer : InLandscapeInfo->LandscapeActor->LandscapeLayers)
+		{		
+			TMap<UTexture2D*, UTexture2D*> CreatedTextures;
+			InComponent->AddDefaultLayerData(LandscapeLayer.Guid, { InComponent }, CreatedTextures);
+
+			if (LandscapeLayer.Guid == InLandscapeInfo->LandscapeActor->LandscapeSplinesTargetLayerGuid)
+			{
+				continue;
+			}
+
+			SrcWeightData.Empty();
+			ExtractLandscapeData(InLandscapeInfo, InRegion, LandscapeLayer.Guid, SrcHeightData, SrcWeightData);
+			CopyData(SrcHeightData, NewHeightData, InRegion, NewRegion, bInResample);
+
+			FScopedSetLandscapeEditingLayer LayerScope(NewLandscapeInfo->LandscapeActor.Get(), LandscapeLayer.Guid);
+			{
+				FLandscapeEditDataInterface LandscapeEdit(NewLandscapeInfo);
+				LandscapeEdit.SetHeightData(NewRegion.Min.X, NewRegion.Min.Y, NewRegion.Max.X, NewRegion.Max.Y, NewHeightData.GetData(), 0, false);
+			}
+
+			for (const FLandscapeImportLayerInfo& Layer : SrcWeightData)
+			{
+				CopyData(Layer.LayerData, NewWeightData, InRegion, NewRegion, bInResample);
+				FLandscapeEditDataInterface LandscapeEdit(NewLandscapeInfo);
+				LandscapeEdit.SetAlphaData(Layer.LayerInfo, NewRegion.Min.X, NewRegion.Min.Y, NewRegion.Max.X, NewRegion.Max.Y, NewWeightData.GetData(), 0);
+			}
+		}
+	}
+}
+
+void FLandscapeConfigHelper::MoveSplinesToLandscape(ULandscapeInfo* InLandscapeInfo, ALandscapeProxy* InLandscape, float InScaleFactor)
+{
+	if (!InLandscape->SplineComponent)
+	{
+		InLandscape->CreateSplineComponent();
+	}
+
+	InLandscapeInfo->ForAllLandscapeProxies([InLandscapeInfo, InLandscape](ALandscapeProxy* LandscapeProxy)
+	{
+		if (InLandscape == LandscapeProxy || LandscapeProxy->SplineComponent == nullptr)
+		{
+			return;
+		}
+
+		InLandscapeInfo->MoveSplinesToProxy(LandscapeProxy->SplineComponent, InLandscape);
+	});
+}
+
+void FLandscapeConfigHelper::MoveFoliageToLandscape(ULandscapeInfo* InLandscapeInfo, ULandscapeInfo* InNewLandscapeInfo)
+{
+	// Move instances
+	for (const TPair<FIntPoint, ULandscapeComponent*>& OldEntry : InLandscapeInfo->XYtoComponentMap)
+	{
+		ULandscapeHeightfieldCollisionComponent* OldCollisionComponent = OldEntry.Value->CollisionComponent.Get();
+
+		if (OldCollisionComponent)
+		{
+			UWorld* World = OldCollisionComponent->GetWorld();
+
+			for (const TPair<FIntPoint, ULandscapeComponent*>& NewEntry : InNewLandscapeInfo->XYtoComponentMap)
+			{
+				ULandscapeHeightfieldCollisionComponent* NewCollisionComponent = NewEntry.Value->CollisionComponent.Get();
+
+				if (NewCollisionComponent && FBoxSphereBounds::BoxesIntersect(NewCollisionComponent->Bounds, OldCollisionComponent->Bounds))
+				{
+					FBox Box = NewCollisionComponent->Bounds.GetBox();
+					Box.Min.Z = -WORLD_MAX;
+					Box.Max.Z = WORLD_MAX;
+
+					AInstancedFoliageActor::MoveInstancesToNewComponent(World, OldCollisionComponent, Box, NewCollisionComponent);
+				}
+			}
+		}
+	}
+
+	// Snap them to the bounds
+	for (const TPair<FIntPoint, ULandscapeComponent*>& NewEntry : InNewLandscapeInfo->XYtoComponentMap)
+	{
+		ULandscapeHeightfieldCollisionComponent* NewCollisionComponent = NewEntry.Value->CollisionComponent.Get();
+
+		if (NewCollisionComponent)
+		{
+			FBox Box = NewCollisionComponent->Bounds.GetBox();
+			Box.Min.Z = -WORLD_MAX;
+			Box.Max.Z = WORLD_MAX;
+
+			NewCollisionComponent->SnapFoliageInstances(Box);
+		}
+	}
+}
+
+ULandscapeInfo* FLandscapeConfigHelper::ChangeConfiguration(ULandscapeInfo* InLandscapeInfo, const FLandscapeConfigChange& InNewConfig, TSet<AActor*>& OutActorsToDelete, TSet<AActor*>& OutModifiedActors)
+{
+	// @todo_ow: output a list of ModifiedActors
+
+	if (!InNewConfig.Validate())
+	{
+		return nullptr;
+	}
+
+	check(InLandscapeInfo);
+	ALandscape* OldLandscape = InLandscapeInfo->LandscapeActor.Get();
+	check(OldLandscape);
+
+	UWorld* World = OldLandscape->GetWorld();
+	UActorPartitionSubsystem* ActorPartitionSubsystem = World->GetSubsystem<UActorPartitionSubsystem>();
+	FLandscapeConfig OldConfig(InLandscapeInfo);
+
+	bool bGridSizeOnly = false;
+	if (OldConfig.ComponentNumSubsections == InNewConfig.ComponentNumSubsections &&
+		OldConfig.SubsectionSizeQuads == InNewConfig.SubsectionSizeQuads)
+	{
+		if (OldConfig.GridSizeInComponents == InNewConfig.GridSizeInComponents)
+		{
+			// Nothing to do
+			return nullptr;
+		}
+
+		// Simple redistribute
+		if (!FLandscapeConfigHelper::ChangeGridSize(InLandscapeInfo, InNewConfig.GridSizeInComponents, OutActorsToDelete))
+		{
+			return nullptr;
+		}
+
+		return InLandscapeInfo;
+	}
+		
+	int32 LSMinX, LSMinY, LSMaxX, LSMaxY;
+	if (!InLandscapeInfo->GetLandscapeExtent(LSMinX, LSMinY, LSMaxX, LSMaxY))
+	{
+		return nullptr;
+	}
+		
+	const int32 NewComponentSizeQuads = InNewConfig.GetComponentSizeQuads();
+	const int32 OldComponentSizeQuads = OldConfig.GetComponentSizeQuads();
+	
+	const bool bResample = InNewConfig.ResizeMode == ELandscapeResizeMode::Resample;
+	const float SizeFactor = (float)OldComponentSizeQuads / NewComponentSizeQuads;
+	const float NewScaleFactor = bResample ? SizeFactor : 1.0f;
+	const FVector OldScale = OldLandscape->GetActorScale();
+	const FVector NewScale = OldScale * NewScaleFactor;
+
+	const FIntPoint OldSize(LSMaxX - LSMinX, LSMaxY - LSMinY);
+	FIntPoint NewLSMin((LSMinX / OldComponentSizeQuads) * NewComponentSizeQuads, (LSMinY / OldComponentSizeQuads) * NewComponentSizeQuads);
+
+	FVector ActorOffset(0, 0, 0);
+	if (InNewConfig.bZeroBased)
+	{
+		ActorOffset = FVector(NewLSMin) * OldScale * SizeFactor;
+		NewLSMin = FIntPoint(0, 0);
+	}
+
+	FIntPoint NewLSMax(NewLSMin.X + (OldSize.X / OldComponentSizeQuads) * NewComponentSizeQuads, NewLSMin.Y + (OldSize.Y / OldComponentSizeQuads) * NewComponentSizeQuads);
+		
+	if (InNewConfig.ResizeMode == ELandscapeResizeMode::Expand)
+	{
+		NewLSMax.X = NewLSMin.X + FMath::DivideAndRoundUp(OldSize.X, NewComponentSizeQuads) * NewComponentSizeQuads;
+		NewLSMax.Y = NewLSMin.Y + FMath::DivideAndRoundUp(OldSize.Y, NewComponentSizeQuads) * NewComponentSizeQuads;
+	}
+	else if (InNewConfig.ResizeMode == ELandscapeResizeMode::Clip)
+	{
+		NewLSMax.X = NewLSMin.X + FMath::Max(1, OldSize.X / NewComponentSizeQuads) * NewComponentSizeQuads;
+		NewLSMax.Y = NewLSMin.Y + FMath::Max(1, OldSize.Y / NewComponentSizeQuads) * NewComponentSizeQuads;
+	}
+			
+	FActorSpawnParameters SpawnParams;
+	ALandscape* NewLandscape = World->SpawnActor<ALandscape>(OldLandscape->GetActorLocation() + ActorOffset, OldLandscape->GetActorRotation(), SpawnParams);
+	
+	NewLandscape->GetSharedProperties(OldLandscape);
+	NewLandscape->SetLandscapeGuid(FGuid::NewGuid());
+	NewLandscape->GridSize = InNewConfig.GetGridSizeQuads();
+	NewLandscape->ComponentSizeQuads = InNewConfig.GetComponentSizeQuads();
+	NewLandscape->NumSubsections = InNewConfig.ComponentNumSubsections;
+	NewLandscape->SubsectionSizeQuads = InNewConfig.SubsectionSizeQuads;	
+	NewLandscape->SetActorRelativeScale3D(NewScale);
+	
+	// Copy Layer Data
+	if (OldLandscape->HasLayersContent())
+	{
+		NewLandscape->bCanHaveLayersContent = true;
+		OldLandscape->ForEachLayer([NewLandscape](FLandscapeLayer& Layer)
+		{
+			NewLandscape->LandscapeLayers.Add(Layer);
+		});
+		NewLandscape->LandscapeSplinesTargetLayerGuid = OldLandscape->LandscapeSplinesTargetLayerGuid;
+	}
+		
+	ULandscapeInfo* NewLandscapeInfo = NewLandscape->CreateLandscapeInfo();
+		
+	for (int32 X = NewLSMin.X; X < NewLSMax.X; X+=NewComponentSizeQuads)
+	{
+		for (int32 Y = NewLSMin.Y; Y < NewLSMax.Y; Y+=NewComponentSizeQuads)
+		{
+			FIntPoint ComponentKey(X/NewComponentSizeQuads, Y/NewComponentSizeQuads);
+			FIntPoint NewSectionBase(X, Y);
+						
+			// Means we don't have more data to copy...
+			if (!bResample && ((NewSectionBase.X-NewLSMin.X) > OldSize.X || (NewSectionBase.Y-NewLSMin.Y) > OldSize.Y))
+			{
+				continue;
+			}
+
+			UActorPartitionSubsystem::FCellCoord CellCoord = UActorPartitionSubsystem::FCellCoord::GetCellCoord(NewSectionBase, World->PersistentLevel, NewLandscape->GridSize);
+			ALandscapeProxy* LandscapeProxy = FLandscapeConfigHelper::FindOrAddLandscapeStreamingProxy(ActorPartitionSubsystem, NewLandscapeInfo, CellCoord, NewLandscape->GridSize);
+			check(LandscapeProxy);
+
+			ULandscapeComponent* NewComponent = NewObject<ULandscapeComponent>(LandscapeProxy, NAME_None, RF_Transactional);
+			NewComponent->Init(NewSectionBase.X, NewSectionBase.Y, NewComponentSizeQuads, InNewConfig.ComponentNumSubsections, InNewConfig.SubsectionSizeQuads);
+
+			// Add to map before resample so that SetHeightData/SetWeightData can find the component. XYtoComponentMap is usually updated in RegisterComponent
+			NewLandscapeInfo->XYtoComponentMap.Add(ComponentKey, NewComponent);
+
+			FIntRect SrcRegion(0, 0, 0, 0);
+			if (bResample)
+			{
+				FIntPoint OldSectionBase(ComponentKey.X* OldComponentSizeQuads, ComponentKey.Y* OldComponentSizeQuads);
+				SrcRegion.Min = OldSectionBase;
+				SrcRegion.Max = OldSectionBase + OldComponentSizeQuads;
+				if (InNewConfig.bZeroBased)
+				{
+					SrcRegion += FIntPoint(LSMinX, LSMinY);
+				}
+			}
+			else
+			{
+				FIntPoint RegionOffset(NewLSMin.X - LSMinX, NewLSMin.Y - LSMinY);
+				SrcRegion.Min = NewSectionBase - RegionOffset;
+				SrcRegion.Max.X = FMath::Min(LSMaxX, SrcRegion.Min.X + NewComponentSizeQuads);
+				SrcRegion.Max.Y = FMath::Min(LSMaxY, SrcRegion.Min.Y + NewComponentSizeQuads);
+			}
+
+			CopyRegionToComponent(InLandscapeInfo, SrcRegion, bResample, NewComponent);
+
+			NewComponent->UpdateMaterialInstances();
+			NewComponent->UpdateCachedBounds();
+			NewComponent->UpdateBounds();
+
+			NewComponent->RegisterComponent();
+		}
+	}
+		
+	// @todo_ow: This is probably not the final solution
+	MoveSplinesToLandscape(InLandscapeInfo, NewLandscape, NewScaleFactor);
+		
+	if (NewLandscapeInfo->CanHaveLayersContent())
+	{
+		NewLandscape->RequestLayersContentUpdateForceAll();
+		NewLandscape->RequestSplineLayerUpdate();
+		NewLandscape->ForceUpdateLayersContent();
+	}
+	else
+	{
+		FLandscapeEditDataInterface LandscapeEdit(NewLandscapeInfo);
+		LandscapeEdit.RecalculateNormals();
+	}
+
+	// Do as last step so that components have a collision component
+	MoveFoliageToLandscape(InLandscapeInfo, NewLandscapeInfo);
+
+	return NewLandscapeInfo;
+}
+
+#endif
\ No newline at end of file
diff --git a/Engine/Source/Runtime/Landscape/Private/LandscapeEdit.cpp b/Engine/Source/Runtime/Landscape/Private/LandscapeEdit.cpp
--- a/Engine/Source/Runtime/Landscape/Private/LandscapeEdit.cpp
+++ b/Engine/Source/Runtime/Landscape/Private/LandscapeEdit.cpp
@@ -89,13 +89,24 @@ ULandscapeLayerInfoObject* ALandscapeProxy::VisibilityLayer = nullptr;
 
 void ULandscapeComponent::Init(int32 InBaseX, int32 InBaseY, int32 InComponentSizeQuads, int32 InNumSubsections, int32 InSubsectionSizeQuads)
 {
+	ALandscapeProxy* LandscapeProxy = GetLandscapeProxy();
+	check(LandscapeProxy && !LandscapeProxy->LandscapeComponents.Contains(this));
+	LandscapeProxy->LandscapeComponents.Add(this);
+
 	SetSectionBase(FIntPoint(InBaseX, InBaseY));
 	SetRelativeLocation(FVector(GetSectionBase() - GetLandscapeProxy()->LandscapeSectionOffset));
 	ComponentSizeQuads = InComponentSizeQuads;
 	NumSubsections = InNumSubsections;
 	SubsectionSizeQuads = InSubsectionSizeQuads;
 	check(NumSubsections * SubsectionSizeQuads == ComponentSizeQuads);
-	ULandscapeInfo* Info = GetLandscapeInfo();
+
+	AttachToComponent(LandscapeProxy->GetRootComponent(), FAttachmentTransformRules::KeepRelativeTransform);
+	const int32 ComponentVerts = (SubsectionSizeQuads + 1) * NumSubsections;
+	
+	WeightmapScaleBias = FVector4(1.0f / (float)ComponentVerts, 1.0f / (float)ComponentVerts, 0.5f / (float)ComponentVerts, 0.5f / (float)ComponentVerts);
+	WeightmapSubsectionOffset = (float)(SubsectionSizeQuads + 1) / (float)ComponentVerts;
+		
+	UpdatedSharedPropertiesFromActor();
 }
 
 void ULandscapeComponent::UpdateCachedBounds(bool bInApproximateBounds)
@@ -2163,6 +2174,13 @@ void ULandscapeComponent::GetComponentExtent(int32& MinX, int32& MinY, int32& Ma
 	MaxY = FMath::Max(SectionBaseY + ComponentSizeQuads, MaxY);
 }
 
+FIntRect ULandscapeComponent::GetComponentExtent() const
+{
+	int32 MinX = MAX_int32, MinY = MAX_int32, MaxX = MIN_int32, MaxY = MIN_int32;
+	GetComponentExtent(MinX, MinY, MaxX, MaxY);
+	return FIntRect(MinX, MinY, MaxX, MaxY);
+}
+
 //
 // ALandscape
 //
@@ -2377,13 +2395,7 @@ LANDSCAPE_API void ALandscapeProxy::Import(const FGuid& InGuid, int32 InMinX, in
 			const int32 BaseY = InMinY + Y * ComponentSizeQuads;
 
 			ULandscapeComponent* LandscapeComponent = NewObject<ULandscapeComponent>(this, NAME_None, RF_Transactional);
-			LandscapeComponent->SetRelativeLocation(FVector(BaseX, BaseY, 0));
-			LandscapeComponent->SetupAttachment(GetRootComponent(), NAME_None);
-			LandscapeComponents.Add(LandscapeComponent);
 			LandscapeComponent->Init(BaseX, BaseY, ComponentSizeQuads, NumSubsections, SubsectionSizeQuads);
-
-			// Assign shared properties
-			LandscapeComponent->UpdatedSharedPropertiesFromActor();
 		}
 	}
 
diff --git a/Engine/Source/Runtime/Landscape/Private/LandscapeSplines.cpp b/Engine/Source/Runtime/Landscape/Private/LandscapeSplines.cpp
--- a/Engine/Source/Runtime/Landscape/Private/LandscapeSplines.cpp
+++ b/Engine/Source/Runtime/Landscape/Private/LandscapeSplines.cpp
@@ -3270,6 +3270,11 @@ void ULandscapeSplineSegment::PostEditChangeProperty(FPropertyChangedEvent& Prop
 	}
 }
 
+void ALandscapeProxy::CreateSplineComponent()
+{
+	CreateSplineComponent(FVector(1.0f) / GetRootComponent()->GetRelativeScale3D());
+}
+
 void ALandscapeProxy::CreateSplineComponent(const FVector& Scale3D)
 {
 	check(SplineComponent == nullptr);
@@ -3441,27 +3446,53 @@ void ULandscapeInfo::MoveSplineToLevel(ULandscapeSplineControlPoint* InControlPo
 		return;
 	}
 	
+	MoveSplineToProxy(InControlPoint, ToProxy);
+}
+
+void ULandscapeInfo::MoveSplineToProxy(ULandscapeSplineControlPoint* InControlPoint, ALandscapeProxy* InLandscapeProxy)
+{
+	ALandscapeProxy* FromProxy = InControlPoint->GetTypedOuter<ALandscapeProxy>();
 	TSet<ULandscapeSplineControlPoint*> ControlPoints;
-	
+
 	FindStartingControlPoints(InControlPoint, ControlPoints);
-	
+
 	FromProxy->Modify();
 	FromProxy->SplineComponent->Modify();
 	FromProxy->SplineComponent->MarkRenderStateDirty();
 
 	for (ULandscapeSplineControlPoint* ControlPoint : ControlPoints)
 	{
-		MoveControlPointToLandscape(ControlPoint, FromProxy, ToProxy);
+		MoveControlPointToLandscape(ControlPoint, FromProxy, InLandscapeProxy);
 	}
 }
 
 void ULandscapeInfo::MoveSplinesToLevel(ULandscapeSplinesComponent * InSplineComponent, ULevel * TargetLevel)
 {
+	ALandscapeProxy* FromProxy = InSplineComponent->GetTypedOuter<ALandscapeProxy>();
+	if (FromProxy->GetLevel() == TargetLevel)
+	{
+		return;
+	}
+
+	ALandscapeProxy* ToProxy = GetLandscapeProxyForLevel(TargetLevel);
+	if (!ToProxy)
+	{
+		return;
+	}
+	
+	MoveSplinesToProxy(InSplineComponent, ToProxy);
+}
+
+/** Moves all Splines to target Proxy. Creates ULandscapeSplineComponent if needed */
+void ULandscapeInfo::MoveSplinesToProxy(ULandscapeSplinesComponent* InSplineComponent, ALandscapeProxy* InLandscapeProxy)
+{
+	check(InSplineComponent->GetTypedOuter<ALandscapeProxy>() != InLandscapeProxy);
+	
 	TArray<ULandscapeSplineControlPoint*> CopyControlPoints(InSplineComponent->ControlPoints);
-	for(ULandscapeSplineControlPoint* ControlPoint : CopyControlPoints)
+	for (ULandscapeSplineControlPoint* ControlPoint : CopyControlPoints)
 	{
 		// Even if ControlPoints are part of same connected spline they will be skipped if already moved
-		MoveSplineToLevel(ControlPoint, TargetLevel);
+		MoveSplineToProxy(ControlPoint, InLandscapeProxy);
 	}
 }
 #endif // WITH_EDITOR
diff --git a/Engine/Source/Runtime/Landscape/Private/LandscapeSubsystem.cpp b/Engine/Source/Runtime/Landscape/Private/LandscapeSubsystem.cpp
--- a/Engine/Source/Runtime/Landscape/Private/LandscapeSubsystem.cpp
+++ b/Engine/Source/Runtime/Landscape/Private/LandscapeSubsystem.cpp
@@ -14,6 +14,7 @@
 #include "ActorPartition/ActorPartitionSubsystem.h"
 #include "Engine/World.h"
 #include "Math/IntRect.h"
+#include "LandscapeConfigHelper.h"
 
 static int32 GUseStreamingManagerForCameras = 1;
 static FAutoConsoleVariableRef CVarUseStreamingManagerForCameras(
@@ -170,88 +171,17 @@ bool ULandscapeSubsystem::IsGridBased() const
 	return UWorld::HasSubsystem<UWorldPartitionSubsystem>(GetWorld());
 }
 
-namespace LandscapeSubsystemUtils
-{
-	ALandscapeProxy* FindOrAddLandscapeProxy(UActorPartitionSubsystem* ActorPartitionSubsystem, ULandscapeInfo* LandscapeInfo, const UActorPartitionSubsystem::FCellCoord& CellCoord)
-	{
-		ALandscape* Landscape = LandscapeInfo->LandscapeActor.Get();
-		check(Landscape);
-							
-		auto LandscapeProxyCreated = [CellCoord, Landscape](APartitionActor* PartitionActor)
-		{
-			const FIntPoint CellLocation(CellCoord.X * Landscape->GridSize, CellCoord.Y * Landscape->GridSize);
-
-			ALandscapeProxy* LandscapeProxy = CastChecked<ALandscapeProxy>(PartitionActor);
-			// copy shared properties to this new proxy
-			LandscapeProxy->GetSharedProperties(Landscape);
-			const FVector ProxyLocation = Landscape->GetActorLocation() + FVector(CellLocation.X * Landscape->GetActorRelativeScale3D().X, CellLocation.Y * Landscape->GetActorRelativeScale3D().Y, 0.0f);
-
-			LandscapeProxy->CreateLandscapeInfo();
-			LandscapeProxy->SetActorLocationAndRotation(ProxyLocation, Landscape->GetActorRotation());
-			LandscapeProxy->LandscapeSectionOffset = FIntPoint(CellLocation.X, CellLocation.Y);
-		};
-
-		const bool bCreate = true;
-		const bool bBoundsSearch = false;
-		return Cast<ALandscapeProxy>(ActorPartitionSubsystem->GetActor(ALandscapeStreamingProxy::StaticClass(), CellCoord, bCreate, LandscapeInfo->LandscapeGuid, Landscape->GridSize, bBoundsSearch, LandscapeProxyCreated));
-	}
-}
-
-void ULandscapeSubsystem::UpdateGrid(ULandscapeInfo* LandscapeInfo, uint32 GridSizeInComponents)
+void ULandscapeSubsystem::ChangeGridSize(ULandscapeInfo* LandscapeInfo, uint32 GridSizeInComponents)
 {
 	if (!IsGridBased())
 	{
 		return;
 	}
 
-	check(LandscapeInfo);
-
-	const uint32 GridSize = LandscapeInfo->GetGridSize(GridSizeInComponents);
-	LandscapeInfo->LandscapeActor->Modify();
-	LandscapeInfo->LandscapeActor->GridSize = GridSize;
-
-	FIntRect Extent;
-	LandscapeInfo->GetLandscapeExtent(Extent.Min.X, Extent.Min.Y, Extent.Max.X, Extent.Max.Y);
-	const FBox Bounds(FVector(Extent.Min), FVector(Extent.Max));
-
-	UActorPartitionSubsystem* ActorPartitionSubsystem = GetWorld()->GetSubsystem<UActorPartitionSubsystem>();
-
-	TArray<ULandscapeComponent*> LandscapeComponents;
-	LandscapeComponents.Reserve(LandscapeInfo->XYtoComponentMap.Num());
-	LandscapeInfo->ForAllLandscapeComponents([&LandscapeComponents](ULandscapeComponent* LandscapeComponent)
-	{
-		LandscapeComponents.Add(LandscapeComponent);
-	});
-
-	FActorPartitionGridHelper::ForEachIntersectingCell(ALandscapeStreamingProxy::StaticClass(), Extent, GetWorld()->PersistentLevel, [ActorPartitionSubsystem, LandscapeInfo, GridSizeInComponents, &LandscapeComponents, GridSize](const UActorPartitionSubsystem::FCellCoord& CellCoord, const FIntRect& CellBounds)
-	{		
-		TArray<ULandscapeComponent*> ComponentsToMove;
-		const int32 MaxComponents = (int32)(GridSizeInComponents * GridSizeInComponents);
-		ComponentsToMove.Reserve(MaxComponents);
-		for (int32 i = 0; i < LandscapeComponents.Num();)
-		{
-			ULandscapeComponent* LandscapeComponent = LandscapeComponents[i];
-			if (CellBounds.Contains(LandscapeComponent->GetSectionBase()))
-			{
-				ComponentsToMove.Add(LandscapeComponent);
-				LandscapeComponents.RemoveAtSwap(i);
-			}
-			else
-			{
-				i++;
-			}
-		}
-		
-		check(ComponentsToMove.Num() <= MaxComponents);
-		if (ComponentsToMove.Num())
-		{
-			ALandscapeProxy* LandscapeProxy = LandscapeSubsystemUtils::FindOrAddLandscapeProxy(ActorPartitionSubsystem, LandscapeInfo, CellCoord);
-			check(LandscapeProxy);
-			LandscapeInfo->MoveComponentsToProxy(ComponentsToMove, LandscapeProxy);
-		}
-
-		return true;
-	}, GridSize);
+	TSet<AActor*> ActorsToDelete;
+	FLandscapeConfigHelper::ChangeGridSize(LandscapeInfo, GridSizeInComponents, ActorsToDelete);
+	// This code path is used for converting a non grid based Landscape to a gridbased so it shouldn't delete any actors
+	check(!ActorsToDelete.Num());
 }
 
 ALandscapeProxy* ULandscapeSubsystem::FindOrAddLandscapeProxy(ULandscapeInfo* LandscapeInfo, const FIntPoint& SectionBase)
@@ -261,11 +191,7 @@ ALandscapeProxy* ULandscapeSubsystem::FindOrAddLandscapeProxy(ULandscapeInfo* La
 		return LandscapeInfo->GetCurrentLevelLandscapeProxy(true);
 	}
 
-	UActorPartitionSubsystem* ActorPartitionSubsystem = GetWorld()->GetSubsystem<UActorPartitionSubsystem>();
-	const uint32 GridSize = LandscapeInfo->LandscapeActor->GridSize;
-	
-	UActorPartitionSubsystem::FCellCoord CellCoord = UActorPartitionSubsystem::FCellCoord::GetCellCoord(SectionBase, GetWorld()->PersistentLevel, GridSize);
-	return LandscapeSubsystemUtils::FindOrAddLandscapeProxy(ActorPartitionSubsystem, LandscapeInfo, CellCoord);
+	return FLandscapeConfigHelper::FindOrAddLandscapeStreamingProxy(LandscapeInfo, SectionBase);
 }
 
 #endif
\ No newline at end of file
diff --git a/Engine/Source/Runtime/Landscape/Public/LandscapeConfigHelper.h b/Engine/Source/Runtime/Landscape/Public/LandscapeConfigHelper.h
new file mode 100644
--- /dev/null
+++ b/Engine/Source/Runtime/Landscape/Public/LandscapeConfigHelper.h
@@ -0,0 +1,157 @@
+// Copyright Epic Games, Inc. All Rights Reserved.
+
+#pragma once
+
+#include "CoreMinimal.h"
+#include "UObject/ObjectMacros.h"
+#include "ActorPartition/ActorPartitionSubsystem.h"
+#include "LandscapeConfigHelper.generated.h"
+
+UENUM()
+enum class ELandscapeResizeMode : uint8
+{
+	Resample = 0,
+	Clip = 1,
+	Expand = 2
+};
+
+#if WITH_EDITOR
+
+class ULandscapeInfo;
+class ULandscapeComponent;
+class ALandscapeProxy;
+struct FLandscapeImportLayerInfo;
+
+/**
+ * 
+ */
+struct LANDSCAPE_API FLandscapeConfig
+{
+	FLandscapeConfig(int32 InComponentNumSubSections, int32 InSubsectionSizeQuads, int32 InGridSizeInComponents)
+		: ComponentNumSubsections(InComponentNumSubSections)
+		, SubsectionSizeQuads(InSubsectionSizeQuads)
+		, GridSizeInComponents(InGridSizeInComponents)
+	{
+
+	}
+
+	FLandscapeConfig(ULandscapeInfo* InLandscapeInfo);
+
+	int32 GetComponentSizeQuads() const { return SubsectionSizeQuads * ComponentNumSubsections; }
+	int32 GetComponentSizeVerts() const { return (SubsectionSizeQuads + 1) * ComponentNumSubsections; }
+	int32 GetGridSizeQuads() const { return GetComponentSizeQuads() * GridSizeInComponents; }
+
+	int32 ComponentNumSubsections;
+	int32 SubsectionSizeQuads;
+	int32 GridSizeInComponents;
+
+	static int32 NumSectionValues[2];
+	static int32 SubsectionSizeQuadsValues[6];
+};
+
+struct LANDSCAPE_API FLandscapeConfigChange : public FLandscapeConfig
+{
+	FLandscapeConfigChange(int32 InComponentNumSubSections, int32 InSubsectionSizeQuads, int32 InGridSize, ELandscapeResizeMode InResizeMode, bool bInZeroBased)
+		: FLandscapeConfig(InComponentNumSubSections, InSubsectionSizeQuads, InGridSize)
+		, ResizeMode(InResizeMode)
+		, bZeroBased(bInZeroBased)
+	{
+	}
+
+	bool Validate() const;
+
+	ELandscapeResizeMode ResizeMode;
+	bool bZeroBased;
+};
+
+class FLandscapeConfigHelper
+{
+public:
+	LANDSCAPE_API static ALandscapeProxy* FindOrAddLandscapeStreamingProxy(ULandscapeInfo* InLandscapeInfo, const FIntPoint& InSectionBase);
+	LANDSCAPE_API static ULandscapeInfo* ChangeConfiguration(ULandscapeInfo* InLandscapeInfo, const FLandscapeConfigChange& InNewConfig, TSet<AActor*>& OutActorsToDelete, TSet<AActor*>& OutModifiedActors);
+	LANDSCAPE_API static bool ChangeGridSize(ULandscapeInfo* InLandscapeInfo, uint32 InNewGridSizeInComponents, TSet<AActor*>& OutActorsToDelete);
+
+private:
+	static ALandscapeProxy* FindOrAddLandscapeStreamingProxy(UActorPartitionSubsystem* ActorPartitionSubsystem, ULandscapeInfo* LandscapeInfo, const UActorPartitionSubsystem::FCellCoord& CellCoord, uint32 InSearchGridSize);
+	static void CopyRegionToComponent(ULandscapeInfo* LandscapeInfo, const FIntRect& Region, bool bResample, ULandscapeComponent* Component);
+	static void ExtractLandscapeData(ULandscapeInfo* LandscapeInfo, const FIntRect& Region, const FGuid& LayerGuid, TArray<uint16>& OutHeightData, TArray<FLandscapeImportLayerInfo>& OutImportMaterialLayerInfos);
+	static void MoveSplinesToLandscape(ULandscapeInfo* InLandscapeInfo, ALandscapeProxy* InLandscape, float InScaleFactor);
+	static void MoveFoliageToLandscape(ULandscapeInfo* InLandscapeInfo, ULandscapeInfo* InNewLandscapeInfo);
+
+	template<typename T>
+	static void CopyData(const TArray<T>& InData, TArray<T>& OutData, const FIntRect& InSrcRegion, const FIntRect& InDestRegion, bool bInResample)
+	{
+		if (bInResample)
+		{
+			ResampleData(InData, OutData, InSrcRegion, InDestRegion);
+		}
+		else
+		{
+			ExpandData(InData, OutData, InSrcRegion, InDestRegion);
+		}
+	}
+
+	template<typename T>
+	static void ExpandData(const TArray<T>& InData, TArray<T>& OutData, const FIntRect& InSrcRegion, const FIntRect& InDestRegion)
+	{
+		// Regions are in ComponentQuads and we want Vertices (+1)
+		const int32 SrcWidth = InSrcRegion.Width() + 1;
+		const int32 SrcHeight = InSrcRegion.Height() + 1;
+		const int32 DestWidth = InDestRegion.Width() + 1;
+		const int32 DestHeight = InDestRegion.Height() + 1;
+		
+		OutData.Empty(DestWidth * DestHeight);
+		OutData.AddUninitialized(DestWidth * DestHeight);
+
+		for (int32 Y = 0; Y < DestHeight; ++Y)
+		{
+			const int32 OldY = FMath::Clamp<int32>(Y, 0, SrcHeight - 1);
+						
+			// Copy one row of the old data
+			FMemory::Memcpy(&OutData[Y * DestWidth], &InData[OldY * SrcWidth], FMath::Min<int32>(SrcWidth, DestWidth) * sizeof(T));
+		
+			const T PadRight = InData[OldY * SrcWidth + SrcWidth - 1];
+			for (int32 X = SrcWidth; X < DestWidth; ++X)
+			{
+				OutData[Y * DestWidth + X] = PadRight;
+			}
+		}
+	}
+		
+	template<typename T>
+	static void ResampleData(const TArray<T>& InData, TArray<T>& OutData, const FIntRect& InSrcRegion, const FIntRect& InDestRegion)
+	{
+		// Regions are in ComponentQuads and we want Vertices (+1)
+		const int32 SrcWidth = InSrcRegion.Width() + 1;
+		const int32 SrcHeight = InSrcRegion.Height() + 1;
+		const int32 DestWidth = InDestRegion.Width() + 1;
+		const int32 DestHeight = InDestRegion.Height() + 1;
+
+		OutData.Empty(DestWidth * DestHeight);
+		OutData.AddUninitialized(DestWidth * DestHeight);
+
+		const float XScale = (float)(SrcWidth - 1) / (DestWidth - 1);
+		const float YScale = (float)(SrcHeight - 1) / (DestHeight - 1);
+		for (int32 Y = 0; Y < DestHeight; ++Y)
+		{
+			for (int32 X = 0; X < DestWidth; ++X)
+			{
+				const float OldY = Y * YScale;
+				const float OldX = X * XScale;
+				const int32 X0 = FMath::FloorToInt(OldX);
+				const int32 X1 = FMath::Min(FMath::FloorToInt(OldX) + 1, SrcWidth - 1);
+				const int32 Y0 = FMath::FloorToInt(OldY);
+				const int32 Y1 = FMath::Min(FMath::FloorToInt(OldY) + 1, SrcHeight - 1);
+				const T& Original00 = InData[Y0 * SrcWidth + X0];
+				const T& Original10 = InData[Y0 * SrcWidth + X1];
+				const T& Original01 = InData[Y1 * SrcWidth + X0];
+				const T& Original11 = InData[Y1 * SrcWidth + X1];
+				int32 Index = Y * DestWidth + X;
+				check(Index < OutData.Num());
+				OutData[Y * DestWidth + X] = FMath::BiLerp(Original00, Original10, Original01, Original11, FMath::Fractional(OldX), FMath::Fractional(OldY));
+			}
+		}
+	}
+};
+
+#endif
\ No newline at end of file
diff --git a/Engine/Source/Runtime/Landscape/Public/LandscapeDataAccess.h b/Engine/Source/Runtime/Landscape/Public/LandscapeDataAccess.h
--- a/Engine/Source/Runtime/Landscape/Public/LandscapeDataAccess.h
+++ b/Engine/Source/Runtime/Landscape/Public/LandscapeDataAccess.h
@@ -35,6 +35,14 @@ namespace LandscapeDataAccess
 	{
 		return FMath::RoundToInt(FMath::Clamp<float>(Height * LANDSCAPE_INV_ZSCALE + MidValue, 0.f, MaxValue));		
 	}
+
+	FORCEINLINE FColor PackHeight(uint16 Height)
+	{
+		FColor Color(ForceInit);
+		Color.R = Height >> 8;
+		Color.G = Height & 255;
+		return MoveTemp(Color);
+	}
 };
 
 #if WITH_EDITOR
diff --git a/Engine/Source/Runtime/Landscape/Public/LandscapeSubsystem.h b/Engine/Source/Runtime/Landscape/Public/LandscapeSubsystem.h
--- a/Engine/Source/Runtime/Landscape/Public/LandscapeSubsystem.h
+++ b/Engine/Source/Runtime/Landscape/Public/LandscapeSubsystem.h
@@ -36,7 +36,7 @@ public:
 	LANDSCAPE_API int32 GetOutdatedGrassMapCount();
 
 	LANDSCAPE_API bool IsGridBased() const;
-	LANDSCAPE_API void UpdateGrid(ULandscapeInfo* LandscapeInfo, uint32 GridSizeInComponents);
+	LANDSCAPE_API void ChangeGridSize(ULandscapeInfo* LandscapeInfo, uint32 NewGridSizeInComponents);
 	LANDSCAPE_API ALandscapeProxy* FindOrAddLandscapeProxy(ULandscapeInfo* LandscapeInfo, const FIntPoint& SectionBase);
 #endif
 
