Landscape: Conversion fix make sure LandscapeComponents/Proxies keep their material overrides

diff --git a/Engine/Source/Runtime/Landscape/Private/LandscapeConfigHelper.cpp b/Engine/Source/Runtime/Landscape/Private/LandscapeConfigHelper.cpp
--- a/Engine/Source/Runtime/Landscape/Private/LandscapeConfigHelper.cpp
+++ b/Engine/Source/Runtime/Landscape/Private/LandscapeConfigHelper.cpp
@@ -99,6 +99,10 @@ bool FLandscapeConfigHelper::ChangeGridSize(ULandscapeInfo* InLandscapeInfo, uin
 	InLandscapeInfo->LandscapeActor->Modify();
 	InLandscapeInfo->LandscapeActor->GridSize = GridSize;
 
+	// This needs to be done before moving components
+	InLandscapeInfo->LandscapeActor->InitializeLandscapeLayersWeightmapUsage();
+
+
 	FIntRect Extent;
 	InLandscapeInfo->GetLandscapeExtent(Extent.Min.X, Extent.Min.Y, Extent.Max.X, Extent.Max.Y);
 	const FBox Bounds(FVector(Extent.Min), FVector(Extent.Max));
@@ -114,9 +118,13 @@ bool FLandscapeConfigHelper::ChangeGridSize(ULandscapeInfo* InLandscapeInfo, uin
 	});
 
 	TSet<ALandscapeProxy*> ProxiesToDelete;
-
+	
 	FActorPartitionGridHelper::ForEachIntersectingCell(ALandscapeStreamingProxy::StaticClass(), Extent, World->PersistentLevel, [ActorPartitionSubsystem, InLandscapeInfo, InNewGridSizeInComponents, &LandscapeComponents, OldGridSize, &ProxiesToDelete](const UActorPartitionSubsystem::FCellCoord& CellCoord, const FIntRect& CellBounds)
 	{
+		TMap<ULandscapeComponent*, UMaterialInterface*> ComponentMaterials;
+		TMap<ULandscapeComponent*, UMaterialInterface*> ComponentHoleMaterials;
+		TMap <ULandscapeComponent*, TMap<int32, UMaterialInterface*>> ComponentLODMaterials;
+
 		TArray<ULandscapeComponent*> ComponentsToMove;
 		const int32 MaxComponents = (int32)(InNewGridSizeInComponents * InNewGridSizeInComponents);
 		ComponentsToMove.Reserve(MaxComponents);
@@ -125,6 +133,14 @@ bool FLandscapeConfigHelper::ChangeGridSize(ULandscapeInfo* InLandscapeInfo, uin
 			ULandscapeComponent* LandscapeComponent = LandscapeComponents[i];
 			if (CellBounds.Contains(LandscapeComponent->GetSectionBase()))
 			{
+				ComponentMaterials.FindOrAdd(LandscapeComponent, LandscapeComponent->GetLandscapeMaterial());
+				ComponentHoleMaterials.FindOrAdd(LandscapeComponent, LandscapeComponent->GetLandscapeHoleMaterial());
+				TMap<int32, UMaterialInterface*>& LODMaterials = ComponentLODMaterials.FindOrAdd(LandscapeComponent);
+				for (int32 LODIndex = 0; LODIndex <= 8; ++LODIndex)
+				{
+					LODMaterials.Add(LODIndex, LandscapeComponent->GetLandscapeMaterial(LODIndex));
+				}
+
 				ComponentsToMove.Add(LandscapeComponent);
 				LandscapeComponents.RemoveAtSwap(i);
 				ProxiesToDelete.Add(LandscapeComponent->GetTypedOuter<ALandscapeProxy>());
@@ -140,9 +156,61 @@ bool FLandscapeConfigHelper::ChangeGridSize(ULandscapeInfo* InLandscapeInfo, uin
 		{
 			ALandscapeProxy* LandscapeProxy = FLandscapeConfigHelper::FindOrAddLandscapeStreamingProxy(ActorPartitionSubsystem, InLandscapeInfo, CellCoord, OldGridSize);
 			check(LandscapeProxy);
-
-			// @todo_ow: handle mismatching materials between source and dest proxy (probably override material on component)
 			InLandscapeInfo->MoveComponentsToProxy(ComponentsToMove, LandscapeProxy);
+
+			// Make sure components retain their Materials if they don't match with their parent proxy
+			for (ULandscapeComponent* MovedComponent : ComponentsToMove)
+			{
+				UMaterialInterface* PreviousLandscapeMaterial = ComponentMaterials.FindChecked(MovedComponent);
+				UMaterialInterface* PreviousLandscapeHoleMaterial = ComponentHoleMaterials.FindChecked(MovedComponent);
+				TMap<int32, UMaterialInterface*> PreviousLandscapeLODMaterials = ComponentLODMaterials.FindChecked(MovedComponent);
+
+				MovedComponent->OverrideMaterial = nullptr;
+				if (PreviousLandscapeMaterial != nullptr && PreviousLandscapeMaterial != MovedComponent->GetLandscapeMaterial())
+				{
+					// If Proxy doesn't differ from Landscape override material there first
+					if(LandscapeProxy->GetLandscapeMaterial() == LandscapeProxy->GetLandscapeActor()->GetLandscapeMaterial())
+					{
+						LandscapeProxy->LandscapeMaterial = PreviousLandscapeMaterial; 
+					}
+					else // If it already differs it means that the component differs from it, override on component
+					{
+						MovedComponent->OverrideMaterial = PreviousLandscapeMaterial;
+					}
+				}
+
+				MovedComponent->OverrideHoleMaterial = nullptr;
+				if (PreviousLandscapeHoleMaterial != nullptr && PreviousLandscapeHoleMaterial != MovedComponent->GetLandscapeHoleMaterial())
+				{
+					// If Proxy doesn't differ from Landscape override material there first
+					if (LandscapeProxy->GetLandscapeHoleMaterial() == LandscapeProxy->GetLandscapeActor()->GetLandscapeHoleMaterial())
+					{
+						LandscapeProxy->LandscapeHoleMaterial = PreviousLandscapeHoleMaterial;
+					}
+					else // If it already differs it means that the component differs from it, override on component
+					{
+						MovedComponent->OverrideHoleMaterial = PreviousLandscapeHoleMaterial;
+					}
+				}
+
+				MovedComponent->OverrideMaterials.Reset();
+				for (int32 LODIndex = 0; LODIndex <= 8; ++LODIndex)
+				{
+					UMaterialInterface* PreviousLODMaterial = PreviousLandscapeLODMaterials.FindChecked(LODIndex);
+					// If Proxy doesn't differ from Landscape override material there first
+					if (PreviousLODMaterial != nullptr && PreviousLODMaterial != MovedComponent->GetLandscapeMaterial(LODIndex))
+					{
+						if (LandscapeProxy->GetLandscapeMaterial(LODIndex) == LandscapeProxy->GetLandscapeActor()->GetLandscapeMaterial(LODIndex))
+						{
+							LandscapeProxy->LandscapeMaterialsOverride.Add({ LODIndex, PreviousLODMaterial });
+						}
+						else // If it already differs it means that the component differs from it, override on component
+						{
+							MovedComponent->OverrideMaterials.Add({ LODIndex, PreviousLODMaterial });
+						}
+					}
+				}
+			}
 		}
 
 		return true;
