Landscape: Conversion fix Landscape Spline failed check when moving spline to new proxy

diff --git a/Engine/Source/Runtime/Landscape/Private/LandscapeSplines.cpp b/Engine/Source/Runtime/Landscape/Private/LandscapeSplines.cpp
--- a/Engine/Source/Runtime/Landscape/Private/LandscapeSplines.cpp
+++ b/Engine/Source/Runtime/Landscape/Private/LandscapeSplines.cpp
@@ -3451,17 +3451,19 @@ void ULandscapeInfo::MoveSplineToLevel(ULandscapeSplineControlPoint* InControlPo
 
 void ULandscapeInfo::MoveSplineToProxy(ULandscapeSplineControlPoint* InControlPoint, ALandscapeProxy* InLandscapeProxy)
 {
-	ALandscapeProxy* FromProxy = InControlPoint->GetTypedOuter<ALandscapeProxy>();
 	TSet<ULandscapeSplineControlPoint*> ControlPoints;
-
 	FindStartingControlPoints(InControlPoint, ControlPoints);
-
-	FromProxy->Modify();
-	FromProxy->SplineComponent->Modify();
-	FromProxy->SplineComponent->MarkRenderStateDirty();
-
+		
 	for (ULandscapeSplineControlPoint* ControlPoint : ControlPoints)
 	{
+		ALandscapeProxy* FromProxy = ControlPoint->GetTypedOuter<ALandscapeProxy>();
+		if (FromProxy == InLandscapeProxy)
+		{
+			continue;
+		}
+		FromProxy->Modify();
+		FromProxy->SplineComponent->Modify();
+		FromProxy->SplineComponent->MarkRenderStateDirty();
 		MoveControlPointToLandscape(ControlPoint, FromProxy, InLandscapeProxy);
 	}
 }
@@ -3494,6 +3496,8 @@ void ULandscapeInfo::MoveSplinesToProxy(ULandscapeSplinesComponent* InSplineComp
 		// Even if ControlPoints are part of same connected spline they will be skipped if already moved
 		MoveSplineToProxy(ControlPoint, InLandscapeProxy);
 	}
+
+	check(InSplineComponent->ControlPoints.Num() == 0 && InSplineComponent->Segments.Num() == 0);
 }
 #endif // WITH_EDITOR
 
