LandscapeGizmoActor: Fix transient class load warning

diff --git a/Engine/Source/Editor/LandscapeEditor/Private/LandscapeEdMode.cpp b/Engine/Source/Editor/LandscapeEditor/Private/LandscapeEdMode.cpp
--- a/Engine/Source/Editor/LandscapeEditor/Private/LandscapeEdMode.cpp
+++ b/Engine/Source/Editor/LandscapeEditor/Private/LandscapeEdMode.cpp
@@ -437,7 +437,9 @@ void FEdModeLandscape::Enter()
 
 	if (!CurrentGizmoActor.IsValid())
 	{
-		CurrentGizmoActor = GetWorld()->SpawnActor<ALandscapeGizmoActiveActor>();
+		FActorSpawnParameters SpawnParams;
+		SpawnParams.ObjectFlags |= RF_Transient;
+		CurrentGizmoActor = GetWorld()->SpawnActor<ALandscapeGizmoActiveActor>(SpawnParams);
 		CurrentGizmoActor->ImportFromClipboard();
 	}
 
diff --git a/Engine/Source/Runtime/Landscape/Classes/LandscapeGizmoActiveActor.h b/Engine/Source/Runtime/Landscape/Classes/LandscapeGizmoActiveActor.h
--- a/Engine/Source/Runtime/Landscape/Classes/LandscapeGizmoActiveActor.h
+++ b/Engine/Source/Runtime/Landscape/Classes/LandscapeGizmoActiveActor.h
@@ -50,7 +50,7 @@ struct FGizmoSelectData
 	
 };
 
-UCLASS(notplaceable, transient, MinimalAPI)
+UCLASS(notplaceable, MinimalAPI)
 class ALandscapeGizmoActiveActor : public ALandscapeGizmoActor
 {
 	GENERATED_UCLASS_BODY()
diff --git a/Engine/Source/Runtime/Landscape/Classes/LandscapeGizmoActor.h b/Engine/Source/Runtime/Landscape/Classes/LandscapeGizmoActor.h
--- a/Engine/Source/Runtime/Landscape/Classes/LandscapeGizmoActor.h
+++ b/Engine/Source/Runtime/Landscape/Classes/LandscapeGizmoActor.h
@@ -10,7 +10,7 @@
 
 class UBillboardComponent;
 
-UCLASS(notplaceable, transient, MinimalAPI, NotBlueprintable)
+UCLASS(notplaceable, MinimalAPI, NotBlueprintable)
 class ALandscapeGizmoActor : public AActor
 {
 	GENERATED_UCLASS_BODY()
diff --git a/Engine/Source/Runtime/Landscape/Private/LandscapeGizmoActor.cpp b/Engine/Source/Runtime/Landscape/Private/LandscapeGizmoActor.cpp
--- a/Engine/Source/Runtime/Landscape/Private/LandscapeGizmoActor.cpp
+++ b/Engine/Source/Runtime/Landscape/Private/LandscapeGizmoActor.cpp
@@ -637,7 +637,9 @@ void ALandscapeGizmoActiveActor::EditorApplyRotation(const FRotator& DeltaRotati
 ALandscapeGizmoActor* ALandscapeGizmoActiveActor::SpawnGizmoActor()
 {
 	// ALandscapeGizmoActor is history for ALandscapeGizmoActiveActor
-	ALandscapeGizmoActor* NewActor = GetWorld()->SpawnActor<ALandscapeGizmoActor>();
+	FActorSpawnParameters SpawnParams;
+	SpawnParams.ObjectFlags |= RF_Transient;
+	ALandscapeGizmoActor* NewActor = GetWorld()->SpawnActor<ALandscapeGizmoActor>(SpawnParams);
 	Duplicate(NewActor);
 	return NewActor;
 }
