Change notification message to read more clearly.

diff --git a/Engine/Source/Runtime/Landscape/Private/LandscapeEditLayers.cpp b/Engine/Source/Runtime/Landscape/Private/LandscapeEditLayers.cpp
--- a/Engine/Source/Runtime/Landscape/Private/LandscapeEditLayers.cpp
+++ b/Engine/Source/Runtime/Landscape/Private/LandscapeEditLayers.cpp
@@ -4881,7 +4881,7 @@ void ALandscape::UpdateLayersContent(bool bInWaitForStreaming, bool bInSkipMonit
 			if ((FSlateApplicationBase::Get().GetCurrentTime() - WaitingForResourcesStartTime) > TimeBeforeDisplayingNotification)
 			{
 				// let the user know we are waiting for resources :
-				static const FText NotificationText(LOCTEXT("WaitForLayersResources", "Landscape edition waiting for edit layers resources to be ready."));
+				static const FText NotificationText(LOCTEXT("WaitForLayersResources", "Landscape editor waiting for edit layer resources to be ready."));
 				ShowEditLayersResourcesNotification(NotificationText, EditLayersResourcesNotification);
 			}
 			// The notification may not be visible yet (because of the initial delay) but it should not be hidden and the initial delay timer shouldn't be reset : 
