New call-site and systematic upgrades needed in Main stream for Engine+ShooterGame scalar property upgrades to Wrapped Object Pointers.

diff --git a/Engine/Plugins/Runtime/Interchange/Source/Import/Private/InterchangeImportCommon.cpp b/Engine/Plugins/Runtime/Interchange/Source/Import/Private/InterchangeImportCommon.cpp
--- a/Engine/Plugins/Runtime/Interchange/Source/Import/Private/InterchangeImportCommon.cpp
+++ b/Engine/Plugins/Runtime/Interchange/Source/Import/Private/InterchangeImportCommon.cpp
@@ -30,6 +30,15 @@ namespace UE
 			ensure(NodeContainer);
 		}
 
+		FFactoryCommon::FUpdateImportAssetDataParameters::FUpdateImportAssetDataParameters(UObject* InAssetImportDataOuter
+																							, TObjectPtr<UAssetImportData>* InAssetImportData
+																							, const UInterchangeSourceData* InSourceData
+																							, FString InNodeUniqueID
+																							, UInterchangeBaseNodeContainer* InNodeContainer)
+			: FUpdateImportAssetDataParameters(InAssetImportDataOuter, &((UAssetImportData*&)InAssetImportData), InSourceData, InNodeUniqueID, InNodeContainer)
+		{
+		}
+
 		void FFactoryCommon::UpdateImportAssetData(FUpdateImportAssetDataParameters& Parameters)
 		{
 #if WITH_EDITORONLY_DATA
diff --git a/Engine/Plugins/Runtime/Interchange/Source/Import/Private/InterchangeImportCommon.h b/Engine/Plugins/Runtime/Interchange/Source/Import/Private/InterchangeImportCommon.h
--- a/Engine/Plugins/Runtime/Interchange/Source/Import/Private/InterchangeImportCommon.h
+++ b/Engine/Plugins/Runtime/Interchange/Source/Import/Private/InterchangeImportCommon.h
@@ -35,6 +35,12 @@ namespace UE
 																	, const UInterchangeSourceData* InSourceData
 																	, FString InNodeUniqueID
 																	, UInterchangeBaseNodeContainer* InNodeContainer);
+
+				FUpdateImportAssetDataParameters(UObject* InAssetImportDataOuter
+																	, TObjectPtr<UAssetImportData>* InAssetImportData
+																	, const UInterchangeSourceData* InSourceData
+																	, FString InNodeUniqueID
+																	, UInterchangeBaseNodeContainer* InNodeContainer);
 			};
 
 			/**
diff --git a/Engine/Source/Runtime/CoreUObject/Public/UObject/WeakInterfacePtr.h b/Engine/Source/Runtime/CoreUObject/Public/UObject/WeakInterfacePtr.h
--- a/Engine/Source/Runtime/CoreUObject/Public/UObject/WeakInterfacePtr.h
+++ b/Engine/Source/Runtime/CoreUObject/Public/UObject/WeakInterfacePtr.h
@@ -8,6 +8,8 @@
 #include "Templates/Casts.h"
 #include "ScriptInterface.h"
 
+#include <type_traits>
+
 /**
  * An alternative to TWeakObjectPtr that makes it easier to work through an interface.
  */
@@ -28,11 +30,11 @@ struct TWeakInterfacePtr
 	 */
 	template<
 		typename U,
-		decltype(ImplicitConv<typename TCopyQualifiersFromTo<U, UObject>::Type*>((U*)nullptr))* = nullptr
+		decltype(ImplicitConv<typename TCopyQualifiersFromTo<U, UObject>::Type*>(std::declval<U>()))* = nullptr
 	>
-	TWeakInterfacePtr(U* Object)
+	TWeakInterfacePtr(U&& Object)
 	{
-		InterfaceInstance = Cast<T>(Object);
+		InterfaceInstance = Cast<T>(ImplicitConv<typename TCopyQualifiersFromTo<U, UObject>::Type*>(Object));
 		if (InterfaceInstance != nullptr)
 		{
 			ObjectInstance = Object;
diff --git a/Engine/Source/Runtime/Landscape/Private/LandscapeConfigHelper.cpp b/Engine/Source/Runtime/Landscape/Private/LandscapeConfigHelper.cpp
--- a/Engine/Source/Runtime/Landscape/Private/LandscapeConfigHelper.cpp
+++ b/Engine/Source/Runtime/Landscape/Private/LandscapeConfigHelper.cpp
@@ -202,11 +202,11 @@ bool FLandscapeConfigHelper::ChangeGridSize(ULandscapeInfo* InLandscapeInfo, uin
 					{
 						if (LandscapeProxy->GetLandscapeMaterial(LODIndex) == LandscapeProxy->GetLandscapeActor()->GetLandscapeMaterial(LODIndex))
 						{
-							LandscapeProxy->LandscapeMaterialsOverride.Add({ LODIndex, PreviousLODMaterial });
+							LandscapeProxy->LandscapeMaterialsOverride.Add({ LODIndex, decltype(FLandscapeProxyMaterialOverride::Material)(PreviousLODMaterial) });
 						}
 						else // If it already differs it means that the component differs from it, override on component
 						{
-							MovedComponent->OverrideMaterials.Add({ LODIndex, PreviousLODMaterial });
+							MovedComponent->OverrideMaterials.Add({ LODIndex, decltype(FLandscapeComponentMaterialOverride::Material)(PreviousLODMaterial) });
 						}
 					}
 				}
