Prototype Wrapped Object Pointers (ObjectHandle/ObjectPtr) Change ternary operators that are ambiguous when TObjectPtr can be implicitly constructed from a raw pointer.  Also allow TObjectPtr to be implicitly constructed from an appropriate raw pointer.

diff --git a/Engine/Plugins/2D/Paper2D/Source/Paper2D/Private/PaperSprite.cpp b/Engine/Plugins/2D/Paper2D/Source/Paper2D/Private/PaperSprite.cpp
--- a/Engine/Plugins/2D/Paper2D/Source/Paper2D/Private/PaperSprite.cpp
+++ b/Engine/Plugins/2D/Paper2D/Source/Paper2D/Private/PaperSprite.cpp
@@ -1946,7 +1946,7 @@ void UPaperSprite::RefreshBakedData()
 UTexture2D* UPaperSprite::GetBakedTexture() const
 {
 #if WITH_EDITOR
-	return (BakedSourceTexture != nullptr) ? BakedSourceTexture : GetSourceTexture();
+	return (BakedSourceTexture != nullptr) ? ToRawPtr(BakedSourceTexture) : GetSourceTexture();
 #else
 	return BakedSourceTexture;
 #endif
diff --git a/Engine/Plugins/Experimental/GeometryCollectionPlugin/Source/GeometryCollectionEditor/Private/GeometryCollection/GeometryCollectionConversion.cpp b/Engine/Plugins/Experimental/GeometryCollectionPlugin/Source/GeometryCollectionEditor/Private/GeometryCollection/GeometryCollectionConversion.cpp
--- a/Engine/Plugins/Experimental/GeometryCollectionPlugin/Source/GeometryCollectionEditor/Private/GeometryCollection/GeometryCollectionConversion.cpp
+++ b/Engine/Plugins/Experimental/GeometryCollectionPlugin/Source/GeometryCollectionEditor/Private/GeometryCollection/GeometryCollectionConversion.cpp
@@ -406,14 +406,14 @@ void FGeometryCollectionConversion::AppendSkeletalMesh(const USkeletalMesh* Skel
 						// for each material, add a reference in our GeometryCollectionObject
 						int CurrIdx = 0;
 
-						UMaterialInterface *CurrMaterial = SkeletalMeshComponent ? SkeletalMeshComponent->GetMaterial(CurrIdx) : SkeletalMesh->Materials[CurrIdx].MaterialInterface;
+						UMaterialInterface *CurrMaterial = SkeletalMeshComponent ? SkeletalMeshComponent->GetMaterial(CurrIdx) : ToRawPtr(SkeletalMesh->Materials[CurrIdx].MaterialInterface);
 
 
 						int MaterialStart = GeometryCollectionObject->Materials.Num();
 						while (CurrMaterial)
 						{
 							GeometryCollectionObject->Materials.Add(CurrMaterial);
-							CurrMaterial = SkeletalMeshComponent ? SkeletalMeshComponent->GetMaterial(++CurrIdx) : SkeletalMesh->Materials[++CurrIdx].MaterialInterface;
+							CurrMaterial = SkeletalMeshComponent ? SkeletalMeshComponent->GetMaterial(++CurrIdx) : ToRawPtr(SkeletalMesh->Materials[++CurrIdx].MaterialInterface);
 
 						}
 
diff --git a/Engine/Plugins/FX/Niagara/Source/Niagara/Private/NiagaraRibbonRendererProperties.cpp b/Engine/Plugins/FX/Niagara/Source/Niagara/Private/NiagaraRibbonRendererProperties.cpp
--- a/Engine/Plugins/FX/Niagara/Source/Niagara/Private/NiagaraRibbonRendererProperties.cpp
+++ b/Engine/Plugins/FX/Niagara/Source/Niagara/Private/NiagaraRibbonRendererProperties.cpp
@@ -137,7 +137,7 @@ void UNiagaraRibbonRendererProperties::GetUsedMaterials(const FNiagaraEmitterIns
 		MaterialInterface = Cast<UMaterialInterface>(InEmitter->FindBinding(MaterialUserParamBinding.Parameter));
 	}
 
-	OutMaterials.Add(MaterialInterface ? MaterialInterface : Material);
+	OutMaterials.Add(MaterialInterface ? MaterialInterface : ToRawPtr(Material));
 }
 
 void UNiagaraRibbonRendererProperties::PostInitProperties()
diff --git a/Engine/Plugins/FX/Niagara/Source/Niagara/Private/NiagaraSpriteRendererProperties.cpp b/Engine/Plugins/FX/Niagara/Source/Niagara/Private/NiagaraSpriteRendererProperties.cpp
--- a/Engine/Plugins/FX/Niagara/Source/Niagara/Private/NiagaraSpriteRendererProperties.cpp
+++ b/Engine/Plugins/FX/Niagara/Source/Niagara/Private/NiagaraSpriteRendererProperties.cpp
@@ -104,7 +104,7 @@ void UNiagaraSpriteRendererProperties::GetUsedMaterials(const FNiagaraEmitterIns
 		MaterialInterface = Cast<UMaterialInterface>(InEmitter->FindBinding(MaterialUserParamBinding.Parameter));
 	}
 
-	OutMaterials.Add(MaterialInterface ? MaterialInterface : Material);
+	OutMaterials.Add(MaterialInterface ? MaterialInterface : ToRawPtr(Material));
 }
 
 void UNiagaraSpriteRendererProperties::PostLoad()
diff --git a/Engine/Plugins/Runtime/ProceduralMeshComponent/Source/ProceduralMeshComponent/Private/ProceduralMeshComponent.cpp b/Engine/Plugins/Runtime/ProceduralMeshComponent/Source/ProceduralMeshComponent/Private/ProceduralMeshComponent.cpp
--- a/Engine/Plugins/Runtime/ProceduralMeshComponent/Source/ProceduralMeshComponent/Private/ProceduralMeshComponent.cpp
+++ b/Engine/Plugins/Runtime/ProceduralMeshComponent/Source/ProceduralMeshComponent/Private/ProceduralMeshComponent.cpp
@@ -998,7 +998,7 @@ void UProceduralMeshComponent::UpdateCollision()
 		CreateProcMeshBodySetup();
 	}
 	
-	UBodySetup* UseBodySetup = bUseAsyncCook ? AsyncBodySetupQueue.Last() : ProcMeshBodySetup;
+	UBodySetup* UseBodySetup = bUseAsyncCook ? AsyncBodySetupQueue.Last() : ToRawPtr(ProcMeshBodySetup);
 
 	// Fill in simple collision convex elements
 	UseBodySetup->AggGeom.ConvexElems = CollisionConvexElems;
diff --git a/Engine/Plugins/Runtime/ReplicationGraph/Source/Private/ReplicationGraph.cpp b/Engine/Plugins/Runtime/ReplicationGraph/Source/Private/ReplicationGraph.cpp
--- a/Engine/Plugins/Runtime/ReplicationGraph/Source/Private/ReplicationGraph.cpp
+++ b/Engine/Plugins/Runtime/ReplicationGraph/Source/Private/ReplicationGraph.cpp
@@ -2441,13 +2441,13 @@ void UNetReplicationGraphConnection::RemoveConnectionGraphNode(UReplicationGraph
 
 bool UNetReplicationGraphConnection::PrepareForReplication()
 {
-	NetConnection->ViewTarget = NetConnection->PlayerController ? NetConnection->PlayerController->GetViewTarget() : NetConnection->OwningActor;
+	NetConnection->ViewTarget = NetConnection->PlayerController ? NetConnection->PlayerController->GetViewTarget() : ToRawPtr(NetConnection->OwningActor);
 
 	// Set any children viewtargets
 	for (int32 i = 0; i < NetConnection->Children.Num(); ++i)
 	{
 		UNetConnection* CurChild = NetConnection->Children[i];
-		CurChild->ViewTarget = CurChild->PlayerController ? CurChild->PlayerController->GetViewTarget() : CurChild->OwningActor;
+		CurChild->ViewTarget = CurChild->PlayerController ? CurChild->PlayerController->GetViewTarget() : ToRawPtr(CurChild->OwningActor);
 	}
 
 	return (NetConnection->State != USOCK_Closed) && (NetConnection->ViewTarget != nullptr);
diff --git a/Engine/Source/Developer/LogVisualizer/Private/LogVisualizer.cpp b/Engine/Source/Developer/LogVisualizer/Private/LogVisualizer.cpp
--- a/Engine/Source/Developer/LogVisualizer/Private/LogVisualizer.cpp
+++ b/Engine/Source/Developer/LogVisualizer/Private/LogVisualizer.cpp
@@ -123,7 +123,7 @@ UWorld* FLogVisualizer::GetWorld(UObject* OptionalObject)
 	{
 		UEditorEngine* EEngine = Cast<UEditorEngine>(GEngine);
 		// lets use PlayWorld during PIE/Simulate and regular world from editor otherwise, to draw debug information
-		World = EEngine != nullptr && EEngine->PlayWorld != nullptr ? EEngine->PlayWorld : EEngine->GetEditorWorldContext().World();
+		World = EEngine != nullptr && EEngine->PlayWorld != nullptr ? ToRawPtr(EEngine->PlayWorld) : EEngine->GetEditorWorldContext().World();
 	}
 	else 
 #endif
diff --git a/Engine/Source/Developer/LogVisualizer/Private/SVisualLogger.cpp b/Engine/Source/Developer/LogVisualizer/Private/SVisualLogger.cpp
--- a/Engine/Source/Developer/LogVisualizer/Private/SVisualLogger.cpp
+++ b/Engine/Source/Developer/LogVisualizer/Private/SVisualLogger.cpp
@@ -72,7 +72,7 @@ namespace
 		if (GIsEditor && EEngine != nullptr && World == nullptr)
 		{
 			// lets use PlayWorld during PIE/Simulate and regular world from editor otherwise, to draw debug information
-			World = EEngine->PlayWorld != nullptr ? EEngine->PlayWorld : EEngine->GetEditorWorldContext().World();
+			World = EEngine->PlayWorld != nullptr ? ToRawPtr(EEngine->PlayWorld) : EEngine->GetEditorWorldContext().World();
 		}
 
 #endif
diff --git a/Engine/Source/Editor/BehaviorTreeEditor/Private/BehaviorTreeEditor.cpp b/Engine/Source/Editor/BehaviorTreeEditor/Private/BehaviorTreeEditor.cpp
--- a/Engine/Source/Editor/BehaviorTreeEditor/Private/BehaviorTreeEditor.cpp
+++ b/Engine/Source/Editor/BehaviorTreeEditor/Private/BehaviorTreeEditor.cpp
@@ -1646,7 +1646,7 @@ UBehaviorTree* FBehaviorTreeEditor::GetBehaviorTree() const
 
 UBlackboardData* FBehaviorTreeEditor::GetBlackboardData() const 
 {
-	return BehaviorTree == nullptr ? BlackboardData : BehaviorTree->BlackboardAsset; 
+	return BehaviorTree == nullptr ? BlackboardData : ToRawPtr(BehaviorTree->BlackboardAsset);
 }
 
 void FBehaviorTreeEditor::RefreshDebugger()
diff --git a/Engine/Source/Editor/BlueprintGraph/Private/BlueprintVariableNodeSpawner.cpp b/Engine/Source/Editor/BlueprintGraph/Private/BlueprintVariableNodeSpawner.cpp
--- a/Engine/Source/Editor/BlueprintGraph/Private/BlueprintVariableNodeSpawner.cpp
+++ b/Engine/Source/Editor/BlueprintGraph/Private/BlueprintVariableNodeSpawner.cpp
@@ -189,7 +189,7 @@ FBlueprintActionUiSpec UBlueprintVariableNodeSpawner::GetUiSpec(FBlueprintAction
 		}
 
 		UClass const* EffectiveOwnerClass = WrappedVariable->GetOwnerClass();
-		EffectiveOwnerClass = EffectiveOwnerClass ? EffectiveOwnerClass : OwnerClass;
+		EffectiveOwnerClass = EffectiveOwnerClass ? EffectiveOwnerClass : ToRawPtr(OwnerClass);
 		const bool bIsOwningClassValid = EffectiveOwnerClass && (!Cast<const UBlueprintGeneratedClass>(EffectiveOwnerClass) || EffectiveOwnerClass->ClassGeneratedBy); //todo: more general validation
 		UClass const* VariableClass = bIsOwningClassValid ? EffectiveOwnerClass->GetAuthoritativeClass() : nullptr;
 		if (VariableClass && !TargetClass->IsChildOf(VariableClass))
diff --git a/Engine/Source/Editor/Kismet/Private/SSCSEditor.cpp b/Engine/Source/Editor/Kismet/Private/SSCSEditor.cpp
--- a/Engine/Source/Editor/Kismet/Private/SSCSEditor.cpp
+++ b/Engine/Source/Editor/Kismet/Private/SSCSEditor.cpp
@@ -5548,7 +5548,7 @@ UActorComponent* SSCSEditor::AddNewComponent( UClass* NewComponentClass, UObject
 	FName TemplateVariableName;
 
 	USCS_Node* SCSNode = Cast<USCS_Node>(Asset);
-	UActorComponent* ComponentTemplate = (SCSNode ? SCSNode->ComponentTemplate : Cast<UActorComponent>(Asset));
+	UActorComponent* ComponentTemplate = (SCSNode ? ToRawPtr(SCSNode->ComponentTemplate) : Cast<UActorComponent>(Asset));
 
 	if (SCSNode)
 	{
diff --git a/Engine/Source/Editor/Persona/Private/AnimGraphNodeDetails.cpp b/Engine/Source/Editor/Persona/Private/AnimGraphNodeDetails.cpp
--- a/Engine/Source/Editor/Persona/Private/AnimGraphNodeDetails.cpp
+++ b/Engine/Source/Editor/Persona/Private/AnimGraphNodeDetails.cpp
@@ -1091,7 +1091,7 @@ FReply SParentPlayerTreeRow::OnResetButtonClicked()
 {
 	FAnimParentNodeAssetOverride* HierarchyOverride = EditorObject->GetBlueprint()->GetAssetOverrideForNode(Item->Override->ParentNodeGuid, true);
 	
-	Item->Override->NewAsset = HierarchyOverride ? HierarchyOverride->NewAsset : GraphNode->GetAnimationAsset();
+	Item->Override->NewAsset = HierarchyOverride ? ToRawPtr(HierarchyOverride->NewAsset) : GraphNode->GetAnimationAsset();
 
 	// Apply will remove the override from the object
 	EditorObject->ApplyOverrideToBlueprint(*Item->Override);
diff --git a/Engine/Source/Editor/UMGEditor/Private/Designer/SDesignerView.cpp b/Engine/Source/Editor/UMGEditor/Private/Designer/SDesignerView.cpp
--- a/Engine/Source/Editor/UMGEditor/Private/Designer/SDesignerView.cpp
+++ b/Engine/Source/Editor/UMGEditor/Private/Designer/SDesignerView.cpp
@@ -2787,7 +2787,7 @@ void SDesignerView::ProcessDropAndAddWidget(const FGeometry& MyGeometry, const F
 	if ( FindWidgetUnderCursor(MyGeometry, DragDropEvent, UPanelWidget::StaticClass(), HitResult) )
 	{
 		Target = bIsPreview ? HitResult.Widget.GetPreview() : HitResult.Widget.GetTemplate();
-		TargetTree = (bIsPreview && Target) ? Cast<UWidgetTree>(Target->GetOuter()) : BP->WidgetTree;
+		TargetTree = (bIsPreview && Target) ? Cast<UWidgetTree>(Target->GetOuter()) : ToRawPtr(BP->WidgetTree);
 	}
 	else if (BP->WidgetTree->RootWidget == nullptr || !bIsPreview)
 	{
diff --git a/Engine/Source/Editor/UnrealEd/Private/Lightmass/ImportVolumetricLightmap.cpp b/Engine/Source/Editor/UnrealEd/Private/Lightmass/ImportVolumetricLightmap.cpp
--- a/Engine/Source/Editor/UnrealEd/Private/Lightmass/ImportVolumetricLightmap.cpp
+++ b/Engine/Source/Editor/UnrealEd/Private/Lightmass/ImportVolumetricLightmap.cpp
@@ -1068,7 +1068,7 @@ void FLightmassProcessor::ImportVolumetricLightmap()
 		}
 	}
 
-	ULevel* StorageLevel = System.LightingScenario ? System.LightingScenario : System.GetWorld()->PersistentLevel;
+	ULevel* StorageLevel = System.LightingScenario ? System.LightingScenario : ToRawPtr(System.GetWorld()->PersistentLevel);
 	UMapBuildDataRegistry* Registry = StorageLevel->GetOrCreateMapBuildData();
 	FPrecomputedVolumetricLightmapData CurrentLevelData;
 
diff --git a/Engine/Source/Editor/UnrealEd/Private/PlayLevel.cpp b/Engine/Source/Editor/UnrealEd/Private/PlayLevel.cpp
--- a/Engine/Source/Editor/UnrealEd/Private/PlayLevel.cpp
+++ b/Engine/Source/Editor/UnrealEd/Private/PlayLevel.cpp
@@ -2073,11 +2073,11 @@ void UEditorEngine::OnSwitchWorldsForPIE( bool bSwitchToPieWorld, UWorld* Overri
 {
 	if( bSwitchToPieWorld )
 	{
-		SetPlayInEditorWorld( OverrideWorld ? OverrideWorld : PlayWorld );
+		SetPlayInEditorWorld( OverrideWorld ? OverrideWorld : ToRawPtr(PlayWorld) );
 	}
 	else
 	{
-		RestoreEditorWorld( OverrideWorld ? OverrideWorld : EditorWorld );
+		RestoreEditorWorld( OverrideWorld ? OverrideWorld : ToRawPtr(EditorWorld) );
 	}
 }
 
diff --git a/Engine/Source/Runtime/AIModule/Classes/BehaviorTree/BlackboardComponent.h b/Engine/Source/Runtime/AIModule/Classes/BehaviorTree/BlackboardComponent.h
--- a/Engine/Source/Runtime/AIModule/Classes/BehaviorTree/BlackboardComponent.h
+++ b/Engine/Source/Runtime/AIModule/Classes/BehaviorTree/BlackboardComponent.h
@@ -349,7 +349,7 @@ bool UBlackboardComponent::SetValue(FBlackboard::FKey KeyID, typename TDataClass
 	uint8* RawData = GetKeyRawData(KeyID) + DataOffset;
 	if (RawData)
 	{
-		UBlackboardKeyType* KeyOb = EntryInfo->KeyType->HasInstance() ? KeyInstances[KeyID] : EntryInfo->KeyType;
+		UBlackboardKeyType* KeyOb = EntryInfo->KeyType->HasInstance() ? KeyInstances[KeyID] : ToRawPtr(EntryInfo->KeyType);
 		const bool bChanged = TDataClass::SetValue((TDataClass*)KeyOb, RawData, Value);
 		if (bChanged)
 		{
@@ -366,7 +366,7 @@ bool UBlackboardComponent::SetValue(FBlackboard::FKey KeyID, typename TDataClass
 						const int32 OtherKeyID = OtherBlackboardAsset ? OtherBlackboardAsset->GetKeyID(EntryInfo->EntryName) : FBlackboard::InvalidKey;
 						if (OtherKeyID != FBlackboard::InvalidKey)
 						{
-							UBlackboardKeyType* OtherKeyOb = EntryInfo->KeyType->HasInstance() ? OtherBlackboard->KeyInstances[OtherKeyID] : EntryInfo->KeyType;
+							UBlackboardKeyType* OtherKeyOb = EntryInfo->KeyType->HasInstance() ? OtherBlackboard->KeyInstances[OtherKeyID] : ToRawPtr(EntryInfo->KeyType);
 							uint8* OtherRawData = OtherBlackboard->GetKeyRawData(OtherKeyID) + DataOffset;
 
 							TDataClass::SetValue((TDataClass*)OtherKeyOb, OtherRawData, Value);
@@ -399,7 +399,7 @@ typename TDataClass::FDataType UBlackboardComponent::GetValue(FBlackboard::FKey
 		return TDataClass::InvalidValue;
 	}
 
-	UBlackboardKeyType* KeyOb = EntryInfo->KeyType->HasInstance() ? KeyInstances[KeyID] : EntryInfo->KeyType;
+	UBlackboardKeyType* KeyOb = EntryInfo->KeyType->HasInstance() ? KeyInstances[KeyID] : ToRawPtr(EntryInfo->KeyType);
 	const uint16 DataOffset = EntryInfo->KeyType->HasInstance() ? sizeof(FBlackboardInstancedKeyMemory) : 0;
 
 	const uint8* RawData = GetKeyRawData(KeyID) + DataOffset;
diff --git a/Engine/Source/Runtime/AIModule/Private/Actions/PawnAction_Repeat.cpp b/Engine/Source/Runtime/AIModule/Private/Actions/PawnAction_Repeat.cpp
--- a/Engine/Source/Runtime/AIModule/Private/Actions/PawnAction_Repeat.cpp
+++ b/Engine/Source/Runtime/AIModule/Private/Actions/PawnAction_Repeat.cpp
@@ -94,7 +94,7 @@ bool UPawnAction_Repeat::PushSubAction()
 
 	UPawnAction* ActionCopy = SubActionTriggeringPolicy == EPawnSubActionTriggeringPolicy::CopyBeforeTriggering 
 		? Cast<UPawnAction>(StaticDuplicateObject(ActionToRepeat, this))
-		: ActionToRepeat;
+		: ToRawPtr(ActionToRepeat);
 
 	UE_VLOG(GetPawn(), LogPawnAction, Log, TEXT("%s> pushing repeted action %s %s, repeats left: %d")
 		, *GetName(), SubActionTriggeringPolicy == EPawnSubActionTriggeringPolicy::CopyBeforeTriggering ? TEXT("copy") : TEXT("instance")
diff --git a/Engine/Source/Runtime/AIModule/Private/BehaviorTree/BlackboardComponent.cpp b/Engine/Source/Runtime/AIModule/Private/BehaviorTree/BlackboardComponent.cpp
--- a/Engine/Source/Runtime/AIModule/Private/BehaviorTree/BlackboardComponent.cpp
+++ b/Engine/Source/Runtime/AIModule/Private/BehaviorTree/BlackboardComponent.cpp
@@ -246,8 +246,8 @@ void UBlackboardComponent::PopulateSynchronizedKeys()
 						uint8* RawData = GetKeyRawData(KeyID) + DataOffset;
 						uint8* RawSource = OtherBlackboard->GetKeyRawData(OtherKeyID) + DataOffset;
 
-						UBlackboardKeyType* KeyOb = bKeyHasInstance ? KeyInstances[KeyID] : Key.KeyType;
-						const UBlackboardKeyType* SourceKeyOb = bKeyHasInstance ? OtherBlackboard->KeyInstances[OtherKeyID] : Key.KeyType;
+						UBlackboardKeyType* KeyOb = bKeyHasInstance ? KeyInstances[KeyID] : ToRawPtr(Key.KeyType);
+						const UBlackboardKeyType* SourceKeyOb = bKeyHasInstance ? OtherBlackboard->KeyInstances[OtherKeyID] : ToRawPtr(Key.KeyType);
 
 						KeyOb->CopyValues(*this, RawData, SourceKeyOb, RawSource);
 					}
@@ -786,7 +786,7 @@ void UBlackboardComponent::ClearValue(FBlackboard::FKey KeyID)
 			{
 				const bool bKeyHasInstance = EntryInfo->KeyType->HasInstance();
 
-				UBlackboardKeyType* KeyOb = bKeyHasInstance ? KeyInstances[KeyID] : EntryInfo->KeyType;
+				UBlackboardKeyType* KeyOb = bKeyHasInstance ? KeyInstances[KeyID] : ToRawPtr(EntryInfo->KeyType);
 				const uint16 DataOffset = bKeyHasInstance ? sizeof(FBlackboardInstancedKeyMemory) : 0;
 				uint8* InstancedRawData = RawData + DataOffset;
 
@@ -803,7 +803,7 @@ void UBlackboardComponent::ClearValue(FBlackboard::FKey KeyID)
 						if (OtherKeyID != FBlackboard::InvalidKey)
 						{
 							const FBlackboardEntry* OtherEntryInfo = OtherBlackboard->BlackboardAsset->GetKey(OtherKeyID);
-							UBlackboardKeyType* OtherKeyOb = bKeyHasInstance ? OtherBlackboard->KeyInstances[OtherKeyID] : EntryInfo->KeyType;
+							UBlackboardKeyType* OtherKeyOb = bKeyHasInstance ? OtherBlackboard->KeyInstances[OtherKeyID] : ToRawPtr(EntryInfo->KeyType);
 							uint8* OtherRawData = OtherBlackboard->GetKeyRawData(OtherKeyID) + DataOffset;
 
 							OtherKeyOb->CopyValues(*OtherBlackboard, OtherRawData, KeyOb, InstancedRawData);
@@ -850,8 +850,8 @@ bool UBlackboardComponent::CopyKeyValue(FBlackboard::FKey SourceKeyID, FBlackboa
 	const uint8* SourceValueMem = GetKeyRawData(SourceKeyID) + MemDataOffset;
 	uint8* DestinationValueMem = GetKeyRawData(DestinationKeyID) + MemDataOffset;
 
-	const UBlackboardKeyType* SourceKeyOb = bKeyHasInstance ? KeyInstances[SourceKeyID] : SourceValueEntryInfo->KeyType;
-	UBlackboardKeyType*	DestKeyOb = bKeyHasInstance ? KeyInstances[DestinationKeyID] : DestinationValueEntryInfo->KeyType;
+	const UBlackboardKeyType* SourceKeyOb = bKeyHasInstance ? KeyInstances[SourceKeyID] : ToRawPtr(SourceValueEntryInfo->KeyType);
+	UBlackboardKeyType*	DestKeyOb = bKeyHasInstance ? KeyInstances[DestinationKeyID] : ToRawPtr(DestinationValueEntryInfo->KeyType);
 
 	DestKeyOb->CopyValues(*this, DestinationValueMem, SourceKeyOb, SourceValueMem);
 
diff --git a/Engine/Source/Runtime/AnimGraphRuntime/Private/BoneControllers/AnimNode_RigidBody.cpp b/Engine/Source/Runtime/AnimGraphRuntime/Private/BoneControllers/AnimNode_RigidBody.cpp
--- a/Engine/Source/Runtime/AnimGraphRuntime/Private/BoneControllers/AnimNode_RigidBody.cpp
+++ b/Engine/Source/Runtime/AnimGraphRuntime/Private/BoneControllers/AnimNode_RigidBody.cpp
@@ -875,7 +875,7 @@ void FAnimNode_RigidBody::InitPhysics(const UAnimInstance* InAnimInstance)
 	const USkeletalMesh* SkeletalMeshAsset = SkeletalMeshComp->SkeletalMesh;
 
 	const FReferenceSkeleton& SkelMeshRefSkel = SkeletalMeshAsset->RefSkeleton;
-	UsePhysicsAsset = OverridePhysicsAsset ? OverridePhysicsAsset : InAnimInstance->GetSkelMeshComponent()->GetPhysicsAsset();
+	UsePhysicsAsset = OverridePhysicsAsset ? ToRawPtr(OverridePhysicsAsset) : InAnimInstance->GetSkelMeshComponent()->GetPhysicsAsset();
 		
 	USkeleton* SkeletonAsset = InAnimInstance->CurrentSkeleton;
 	ensure(SkeletonAsset == SkeletalMeshAsset->Skeleton);
diff --git a/Engine/Source/Runtime/CoreUObject/Public/UObject/ObjectPtr.h b/Engine/Source/Runtime/CoreUObject/Public/UObject/ObjectPtr.h
--- a/Engine/Source/Runtime/CoreUObject/Public/UObject/ObjectPtr.h
+++ b/Engine/Source/Runtime/CoreUObject/Public/UObject/ObjectPtr.h
@@ -182,7 +182,7 @@ public:
 			decltype(ImplicitConv<T*>(std::declval<U>()))
 		>* = nullptr
 	>
-	explicit FORCEINLINE TObjectPtr(U&& Object)
+	FORCEINLINE TObjectPtr(U&& Object)
 		: FObjectPtr(const_cast<std::remove_const_t<T>*>(ImplicitConv<T*>(Object)))
 	{
 		// It would be desirable for this constructor to be implicit to allow:
diff --git a/Engine/Source/Runtime/Engine/Classes/Engine/World.h b/Engine/Source/Runtime/Engine/Classes/Engine/World.h
--- a/Engine/Source/Runtime/Engine/Classes/Engine/World.h
+++ b/Engine/Source/Runtime/Engine/Classes/Engine/World.h
@@ -2434,7 +2434,7 @@ public:
 	 * 
 	 * @return default physics volume
 	 */
-	APhysicsVolume* GetDefaultPhysicsVolume() const { return DefaultPhysicsVolume ? DefaultPhysicsVolume : InternalGetDefaultPhysicsVolume(); }
+	APhysicsVolume* GetDefaultPhysicsVolume() const { return DefaultPhysicsVolume ? ToRawPtr(DefaultPhysicsVolume) : InternalGetDefaultPhysicsVolume(); }
 
 	/** Returns true if a DefaultPhysicsVolume has been created. */
 	bool HasDefaultPhysicsVolume() const { return DefaultPhysicsVolume != nullptr; }
diff --git a/Engine/Source/Runtime/Engine/Private/Actor.cpp b/Engine/Source/Runtime/Engine/Private/Actor.cpp
--- a/Engine/Source/Runtime/Engine/Private/Actor.cpp
+++ b/Engine/Source/Runtime/Engine/Private/Actor.cpp
@@ -1868,7 +1868,7 @@ void AActor::OnRep_AttachmentReplication()
 	{
 		if (RootComponent)
 		{
-			USceneComponent* AttachParentComponent = (AttachmentReplication.AttachComponent ? AttachmentReplication.AttachComponent : AttachmentReplication.AttachParent->GetRootComponent());
+			USceneComponent* AttachParentComponent = (AttachmentReplication.AttachComponent ? ToRawPtr(AttachmentReplication.AttachComponent) : AttachmentReplication.AttachParent->GetRootComponent());
 
 			if (AttachParentComponent)
 			{
diff --git a/Engine/Source/Runtime/Engine/Private/Components/ActorComponent.cpp b/Engine/Source/Runtime/Engine/Private/Components/ActorComponent.cpp
--- a/Engine/Source/Runtime/Engine/Private/Components/ActorComponent.cpp
+++ b/Engine/Source/Runtime/Engine/Private/Components/ActorComponent.cpp
@@ -1079,7 +1079,7 @@ bool UActorComponent::SetupActorComponentTickFunction(struct FTickFunction* Tick
 		AActor* MyOwner = GetOwner();
 		if (!MyOwner || !MyOwner->IsTemplate())
 		{
-			ULevel* ComponentLevel = (MyOwner ? MyOwner->GetLevel() : GetWorld()->PersistentLevel);
+			ULevel* ComponentLevel = (MyOwner ? MyOwner->GetLevel() : ToRawPtr(GetWorld()->PersistentLevel));
 			TickFunction->SetTickFunctionEnable(TickFunction->bStartWithTickEnabled || TickFunction->IsTickFunctionEnabled());
 			TickFunction->RegisterTickFunction(ComponentLevel);
 			return true;
diff --git a/Engine/Source/Runtime/Engine/Private/Components/ChildActorComponent.cpp b/Engine/Source/Runtime/Engine/Private/Components/ChildActorComponent.cpp
--- a/Engine/Source/Runtime/Engine/Private/Components/ChildActorComponent.cpp
+++ b/Engine/Source/Runtime/Engine/Private/Components/ChildActorComponent.cpp
@@ -535,7 +535,7 @@ void UChildActorComponent::SetChildActorClass(TSubclassOf<AActor> Class, AActor*
 			DestroyChildActor();
 
 			// If an actor template was supplied, temporarily set ChildActorTemplate to create the new Actor with ActorTemplate used as the template
-			TGuardValue<AActor*> ChildActorTemplateGuard(ChildActorTemplate, (ActorTemplate ? ActorTemplate : ChildActorTemplate));
+			TGuardValue<AActor*> ChildActorTemplateGuard(ChildActorTemplate, (ActorTemplate ? ActorTemplate : ToRawPtr(ChildActorTemplate)));
 
 			CreateChildActor();
 		}
diff --git a/Engine/Source/Runtime/Engine/Private/DataChannel.cpp b/Engine/Source/Runtime/Engine/Private/DataChannel.cpp
--- a/Engine/Source/Runtime/Engine/Private/DataChannel.cpp
+++ b/Engine/Source/Runtime/Engine/Private/DataChannel.cpp
@@ -4207,7 +4207,7 @@ static void	DebugNetGUIDs( UWorld* InWorld )
 		return;
 	}
 
-	UNetConnection * Connection = (NetDriver->ServerConnection ? NetDriver->ServerConnection : (NetDriver->ClientConnections.Num() > 0 ? NetDriver->ClientConnections[0] : NULL));
+	UNetConnection * Connection = (NetDriver->ServerConnection ? ToRawPtr(NetDriver->ServerConnection) : (NetDriver->ClientConnections.Num() > 0 ? NetDriver->ClientConnections[0] : NULL));
 	if (!Connection)
 	{
 		return;
@@ -4232,7 +4232,7 @@ static void	ListOpenActorChannels( UWorld* InWorld )
 		return;
 	}
 
-	UNetConnection * Connection = (NetDriver->ServerConnection ? NetDriver->ServerConnection : (NetDriver->ClientConnections.Num() > 0 ? NetDriver->ClientConnections[0] : NULL));
+	UNetConnection * Connection = (NetDriver->ServerConnection ? ToRawPtr(NetDriver->ServerConnection) : (NetDriver->ClientConnections.Num() > 0 ? NetDriver->ClientConnections[0] : NULL));
 	if (!Connection)
 	{
 		return;
@@ -4289,7 +4289,7 @@ static void	DeleteDormantActor( UWorld* InWorld )
 		return;
 	}
 
-	UNetConnection * Connection = (NetDriver->ServerConnection ? NetDriver->ServerConnection : (NetDriver->ClientConnections.Num() > 0 ? NetDriver->ClientConnections[0] : NULL));
+	UNetConnection * Connection = (NetDriver->ServerConnection ? ToRawPtr(NetDriver->ServerConnection) : (NetDriver->ClientConnections.Num() > 0 ? NetDriver->ClientConnections[0] : NULL));
 	if (!Connection)
 	{
 		return;
diff --git a/Engine/Source/Runtime/Engine/Private/DemoNetDriver.cpp b/Engine/Source/Runtime/Engine/Private/DemoNetDriver.cpp
--- a/Engine/Source/Runtime/Engine/Private/DemoNetDriver.cpp
+++ b/Engine/Source/Runtime/Engine/Private/DemoNetDriver.cpp
@@ -1506,7 +1506,7 @@ class FReplayViewer
 public:
 	FReplayViewer(const UNetConnection* Connection) :
 		Viewer(Connection->PlayerController ? Connection->PlayerController : Connection->OwningActor), 
-		ViewTarget(Connection->PlayerController ? Connection->PlayerController->GetViewTarget() : Connection->OwningActor)
+		ViewTarget(Connection->PlayerController ? Connection->PlayerController->GetViewTarget() : ToRawPtr(Connection->OwningActor))
 	{
 		Location = ViewTarget ? ViewTarget->GetActorLocation() : FVector::ZeroVector;
 	}
@@ -2856,7 +2856,7 @@ bool UDemoNetDriver::SpawnSplitscreenViewer(ULocalPlayer* NewPlayer, UWorld* InW
 		return false;
 	}
 
-	UNetConnection* ChildConnection = CreateChild((ClientConnections.Num() > 0) ? ClientConnections[0] : ServerConnection);
+	UNetConnection* ChildConnection = CreateChild((ClientConnections.Num() > 0) ? ClientConnections[0] : ToRawPtr(ServerConnection));
 
 	APlayerController* NewSplitscreenController = ReplayHelper.CreateSpectatorController(ChildConnection);
 	if (NewSplitscreenController == nullptr)
diff --git a/Engine/Source/Runtime/Engine/Private/EngineTypes.cpp b/Engine/Source/Runtime/Engine/Private/EngineTypes.cpp
--- a/Engine/Source/Runtime/Engine/Private/EngineTypes.cpp
+++ b/Engine/Source/Runtime/Engine/Private/EngineTypes.cpp
@@ -211,7 +211,7 @@ UActorComponent* FComponentReference::GetComponent(AActor* OwningActor) const
 	else
 	{
 		// Look in Actor if specified, OwningActor if not
-		AActor* SearchActor = (OtherActor != NULL) ? OtherActor : OwningActor;
+		AActor* SearchActor = (OtherActor != NULL) ? ToRawPtr(OtherActor) : OwningActor;
 		if(SearchActor)
 		{
 			if(ComponentProperty != NAME_None)
diff --git a/Engine/Source/Runtime/Engine/Private/HLODProxy.cpp b/Engine/Source/Runtime/Engine/Private/HLODProxy.cpp
--- a/Engine/Source/Runtime/Engine/Private/HLODProxy.cpp
+++ b/Engine/Source/Runtime/Engine/Private/HLODProxy.cpp
@@ -424,7 +424,7 @@ FName UHLODProxy::GenerateKeyForActor(const ALODActor* LODActor, bool bMustUndoL
 
 	// Base us off the unique object ID
 	{
-		const UObject* Obj = LODActor->ProxyDesc ? Cast<const UObject>(LODActor->ProxyDesc) : Cast<const UObject>(LODActor);
+		const UObject* Obj = LODActor->ProxyDesc ? ToRawPtr(Cast<const UObject>(LODActor->ProxyDesc)) : Cast<const UObject>(LODActor);
 		FUniqueObjectGuid ObjectGUID = FUniqueObjectGuid::GetOrCreateIDForObject(Obj);
 		Key += TEXT("_");
 		Key += ObjectGUID.GetGuid().ToString(EGuidFormats::Digits);
diff --git a/Engine/Source/Runtime/Engine/Private/HUD.cpp b/Engine/Source/Runtime/Engine/Private/HUD.cpp
--- a/Engine/Source/Runtime/Engine/Private/HUD.cpp
+++ b/Engine/Source/Runtime/Engine/Private/HUD.cpp
@@ -501,7 +501,7 @@ AActor* AHUD::GetCurrentDebugTargetActor()
 
 		// If we hit something new, return this.
 		// Otherwise fall back to our last successful hit.
-		return DebugTargetActor ? DebugTargetActor : ShowDebugTargetActor;
+		return DebugTargetActor ? DebugTargetActor : ToRawPtr(ShowDebugTargetActor);
 	}
 	else
 	{
diff --git a/Engine/Source/Runtime/Engine/Private/LevelActor.cpp b/Engine/Source/Runtime/Engine/Private/LevelActor.cpp
--- a/Engine/Source/Runtime/Engine/Private/LevelActor.cpp
+++ b/Engine/Source/Runtime/Engine/Private/LevelActor.cpp
@@ -515,7 +515,7 @@ AActor* UWorld::SpawnActor( UClass* Class, FTransform const* UserTransformPtr, c
 	if (LevelToSpawnIn == NULL)
 	{
 		// Spawn in the same level as the owner if we have one.
-		LevelToSpawnIn = (SpawnParameters.Owner != NULL) ? SpawnParameters.Owner->GetLevel() : CurrentLevel;
+		LevelToSpawnIn = (SpawnParameters.Owner != NULL) ? SpawnParameters.Owner->GetLevel() : ToRawPtr(CurrentLevel);
 	}
 
 	// Use class's default actor as a template if none provided.
diff --git a/Engine/Source/Runtime/Engine/Private/LocalPlayer.cpp b/Engine/Source/Runtime/Engine/Private/LocalPlayer.cpp
--- a/Engine/Source/Runtime/Engine/Private/LocalPlayer.cpp
+++ b/Engine/Source/Runtime/Engine/Private/LocalPlayer.cpp
@@ -105,7 +105,7 @@ bool FLocalPlayerContext::IsValid() const
 	{
 		if (UWorld* WorldPtr = GetWorld())
 		{
-			if (APlayerController* PC = (WorldPtr ? LocalPlayerPtr->GetPlayerController(WorldPtr) : LocalPlayerPtr->PlayerController))
+			if (APlayerController* PC = (WorldPtr ? LocalPlayerPtr->GetPlayerController(WorldPtr) : ToRawPtr(LocalPlayerPtr->PlayerController)))
 			{
 				return (PC->Player != nullptr);
 			}
@@ -152,7 +152,7 @@ APlayerController* FLocalPlayerContext::GetPlayerController() const
 {
 	ULocalPlayer* LocalPlayerPtr = GetLocalPlayer();
 	UWorld* WorldPtr = World.Get();
-	return (WorldPtr ? LocalPlayerPtr->GetPlayerController(WorldPtr) : LocalPlayerPtr->PlayerController);
+	return (WorldPtr ? LocalPlayerPtr->GetPlayerController(WorldPtr) : ToRawPtr(LocalPlayerPtr->PlayerController));
 }
 
 class AGameStateBase* FLocalPlayerContext::GetGameState() const
@@ -213,7 +213,7 @@ bool FLocalPlayerContext::IsFromLocalPlayer(const AActor* ActorToTest) const
 		{
 			if (UWorld* WorldPtr = GetWorld())
 			{
-				if (APlayerController* PC = (WorldPtr ? LocalPlayerPtr->GetPlayerController(WorldPtr) : LocalPlayerPtr->PlayerController))
+				if (APlayerController* PC = (WorldPtr ? LocalPlayerPtr->GetPlayerController(WorldPtr) : ToRawPtr(LocalPlayerPtr->PlayerController)))
 				{
 					if (   ActorToTest == PC
 					    || ActorToTest == PC->GetPawn()
diff --git a/Engine/Source/Runtime/Engine/Private/NetDriver.cpp b/Engine/Source/Runtime/Engine/Private/NetDriver.cpp
--- a/Engine/Source/Runtime/Engine/Private/NetDriver.cpp
+++ b/Engine/Source/Runtime/Engine/Private/NetDriver.cpp
@@ -5086,7 +5086,7 @@ void UNetDriver::PrintDebugRelevantActors()
 void UNetDriver::DrawNetDriverDebug()
 {
 #if ENABLE_DRAW_DEBUG
-	UNetConnection *Connection = (ServerConnection ? ServerConnection : (ClientConnections.Num() >= 1 ? ClientConnections[0] : NULL));
+	UNetConnection *Connection = (ServerConnection ? ToRawPtr(ServerConnection) : (ClientConnections.Num() >= 1 ? ClientConnections[0] : NULL));
 	if (!Connection)
 	{
 		return;
diff --git a/Engine/Source/Runtime/Engine/Private/Particles/ParticleEmitterInstances.cpp b/Engine/Source/Runtime/Engine/Private/Particles/ParticleEmitterInstances.cpp
--- a/Engine/Source/Runtime/Engine/Private/Particles/ParticleEmitterInstances.cpp
+++ b/Engine/Source/Runtime/Engine/Private/Particles/ParticleEmitterInstances.cpp
@@ -3845,7 +3845,7 @@ void FParticleMeshEmitterInstance::GetMeshMaterials(
 			// Overriding the material?
 			if (Material == NULL && MeshTypeData->bOverrideMaterial == true)
 			{
-				Material = CurrentMaterial ? CurrentMaterial : LODLevel->RequiredModule->Material;
+				Material = CurrentMaterial ? CurrentMaterial : ToRawPtr(LODLevel->RequiredModule->Material);
 			}
 
 			// Use the material set on the mesh.
diff --git a/Engine/Source/Runtime/Engine/Private/SCS_Node.cpp b/Engine/Source/Runtime/Engine/Private/SCS_Node.cpp
--- a/Engine/Source/Runtime/Engine/Private/SCS_Node.cpp
+++ b/Engine/Source/Runtime/Engine/Private/SCS_Node.cpp
@@ -48,7 +48,7 @@ UActorComponent* USCS_Node::GetActualComponentTemplate(UBlueprintGeneratedClass*
 			} while (!OverridenComponentTemplate && ActualBPGC && SCS != ActualBPGC->SimpleConstructionScript);
 		}
 	}
-	return OverridenComponentTemplate ? OverridenComponentTemplate : ComponentTemplate;
+	return OverridenComponentTemplate ? OverridenComponentTemplate : ToRawPtr(ComponentTemplate);
 }
 
 const FBlueprintCookedComponentInstancingData* USCS_Node::GetActualComponentTemplateData(UBlueprintGeneratedClass* ActualBPGC) const
diff --git a/Engine/Source/Runtime/Engine/Private/SimpleConstructionScript.cpp b/Engine/Source/Runtime/Engine/Private/SimpleConstructionScript.cpp
--- a/Engine/Source/Runtime/Engine/Private/SimpleConstructionScript.cpp
+++ b/Engine/Source/Runtime/Engine/Private/SimpleConstructionScript.cpp
@@ -439,7 +439,7 @@ void USimpleConstructionScript::FixupSceneNodeHierarchy()
 					}
 				};
 
-				if (UClass* ComponentClass = (Node->ComponentClass ? Node->ComponentClass : (Node->ComponentTemplate ? Node->ComponentTemplate->GetClass() : nullptr)))
+				if (UClass* ComponentClass = (Node->ComponentClass ? ToRawPtr(Node->ComponentClass) : (Node->ComponentTemplate ? Node->ComponentTemplate->GetClass() : nullptr)))
 				{
 					if (ComponentClass->IsChildOf<USceneComponent>())
 					{
diff --git a/Engine/Source/Runtime/Engine/Private/StaticMesh.cpp b/Engine/Source/Runtime/Engine/Private/StaticMesh.cpp
--- a/Engine/Source/Runtime/Engine/Private/StaticMesh.cpp
+++ b/Engine/Source/Runtime/Engine/Private/StaticMesh.cpp
@@ -2734,7 +2734,7 @@ void FStaticMeshRenderData::Cache(const ITargetPlatform* TargetPlatform, UStatic
 			}
 
 			const FMeshBuildSettings& BuildSettings = Owner->GetSourceModel(0).BuildSettings;
-			UStaticMesh* MeshToGenerateFrom = BuildSettings.DistanceFieldReplacementMesh ? BuildSettings.DistanceFieldReplacementMesh : Owner;
+			UStaticMesh* MeshToGenerateFrom = BuildSettings.DistanceFieldReplacementMesh ? ToRawPtr(BuildSettings.DistanceFieldReplacementMesh) : Owner;
 
 			if (BuildSettings.DistanceFieldReplacementMesh)
 			{
diff --git a/Engine/Source/Runtime/Engine/Private/VisualLogger/VisualLogger.cpp b/Engine/Source/Runtime/Engine/Private/VisualLogger/VisualLogger.cpp
--- a/Engine/Source/Runtime/Engine/Private/VisualLogger/VisualLogger.cpp
+++ b/Engine/Source/Runtime/Engine/Private/VisualLogger/VisualLogger.cpp
@@ -35,7 +35,7 @@ namespace
 		if (GIsEditor && EEngine != nullptr && World == nullptr)
 		{
 			// lets use PlayWorld during PIE/Simulate and regular world from editor otherwise, to draw debug information
-			World = EEngine->PlayWorld != nullptr ? EEngine->PlayWorld : EEngine->GetEditorWorldContext().World();
+			World = EEngine->PlayWorld != nullptr ? ToRawPtr(EEngine->PlayWorld) : EEngine->GetEditorWorldContext().World();
 		}
 
 #endif
diff --git a/Engine/Source/Runtime/Engine/Public/ActiveSound.h b/Engine/Source/Runtime/Engine/Public/ActiveSound.h
--- a/Engine/Source/Runtime/Engine/Public/ActiveSound.h
+++ b/Engine/Source/Runtime/Engine/Public/ActiveSound.h
@@ -307,7 +307,7 @@ public:
 	USoundBase* GetSound() const { return Sound; }
 	void SetSound(USoundBase* InSound);
 
-	USoundEffectSourcePresetChain* GetSourceEffectChain() const { return SourceEffectChain ? SourceEffectChain : Sound->SourceEffectChain; }
+	USoundEffectSourcePresetChain* GetSourceEffectChain() const { return SourceEffectChain ? SourceEffectChain : ToRawPtr(Sound->SourceEffectChain); }
 	void SetSourceEffectChain(USoundEffectSourcePresetChain* InSourceEffectChain);
 
 	void SetSoundClass(USoundClass* SoundClass);
diff --git a/Engine/Source/Runtime/Landscape/Private/LandscapeCollision.cpp b/Engine/Source/Runtime/Landscape/Private/LandscapeCollision.cpp
--- a/Engine/Source/Runtime/Landscape/Private/LandscapeCollision.cpp
+++ b/Engine/Source/Runtime/Landscape/Private/LandscapeCollision.cpp
@@ -995,7 +995,7 @@ bool ULandscapeHeightfieldCollisionComponent::CookCollisionData(const FName& For
 					}
 					else
 					{
-						UPhysicalMaterial* DominantMaterial = Layer && Layer->PhysMaterial ? Layer->PhysMaterial : DefMaterial;
+						UPhysicalMaterial* DominantMaterial = Layer && Layer->PhysMaterial ? ToRawPtr(Layer->PhysMaterial) : DefMaterial;
 						MaterialIndex = InOutMaterials.AddUnique(DominantMaterial);
 					}
 				}
@@ -1178,7 +1178,7 @@ bool ULandscapeMeshCollisionComponent::CookCollisionData(const FName& Format, bo
 					}
 					else
 					{
-						UPhysicalMaterial* DominantMaterial = Layer && Layer->PhysMaterial ? Layer->PhysMaterial : DefMaterial;
+						UPhysicalMaterial* DominantMaterial = Layer && Layer->PhysMaterial ? ToRawPtr(Layer->PhysMaterial) : DefMaterial;
 						MaterialIndex = InOutMaterials.AddUnique(DominantMaterial);
 					}
 				}
diff --git a/Engine/Source/Runtime/Landscape/Private/LandscapeEditInterface.cpp b/Engine/Source/Runtime/Landscape/Private/LandscapeEditInterface.cpp
--- a/Engine/Source/Runtime/Landscape/Private/LandscapeEditInterface.cpp
+++ b/Engine/Source/Runtime/Landscape/Private/LandscapeEditInterface.cpp
@@ -294,7 +294,7 @@ void FLandscapeEditDataInterface::SetHeightData(int32 X1, int32 Y1, int32 X2, in
 			}
 
 			UTexture2D* Heightmap = InHeightmap != nullptr ? InHeightmap : Component->GetHeightmap(true);
-			UTexture2D* XYOffsetmapTexture = InXYOffsetmapTexture != nullptr ? InXYOffsetmapTexture : Component->XYOffsetmapTexture;
+			UTexture2D* XYOffsetmapTexture = InXYOffsetmapTexture != nullptr ? InXYOffsetmapTexture : ToRawPtr(Component->XYOffsetmapTexture);
 
 			Component->Modify(GetShouldDirtyPackage());
 
diff --git a/Engine/Source/Runtime/Landscape/Private/LandscapeRender.cpp b/Engine/Source/Runtime/Landscape/Private/LandscapeRender.cpp
--- a/Engine/Source/Runtime/Landscape/Private/LandscapeRender.cpp
+++ b/Engine/Source/Runtime/Landscape/Private/LandscapeRender.cpp
@@ -2451,7 +2451,7 @@ void FLandscapeComponentSceneProxy::GetDynamicMeshElements(const TArray<const FS
 
 				FMeshBatch& MaskMesh = Collector.AllocateMesh();
 				MaskMesh = MeshTools;
-				auto ColorMaskMaterialInstance = new FLandscapeMaskMaterialRenderProxy(GColorMaskRegionMaterial->GetRenderProxy(), EditToolRenderData.LayerContributionTexture ? EditToolRenderData.LayerContributionTexture : GLandscapeBlackTexture, true);
+				auto ColorMaskMaterialInstance = new FLandscapeMaskMaterialRenderProxy(GColorMaskRegionMaterial->GetRenderProxy(), EditToolRenderData.LayerContributionTexture ? ToRawPtr(EditToolRenderData.LayerContributionTexture) : GLandscapeBlackTexture, true);
 				MaskMesh.MaterialRenderProxy = ColorMaskMaterialInstance;
 				Collector.RegisterOneFrameMaterialProxy(ColorMaskMaterialInstance);
 				Collector.AddMesh(ViewIndex, MaskMesh);
@@ -2503,7 +2503,7 @@ void FLandscapeComponentSceneProxy::GetDynamicMeshElements(const TArray<const FS
 					FMeshBatch& MaskMesh = Collector.AllocateMesh();
 					MaskMesh = MeshTools;
 										
-					auto DirtyMaterialInstance = new FLandscapeMaskMaterialRenderProxy(GLandscapeDirtyMaterial->GetRenderProxy(), EditToolRenderData.DirtyTexture ? EditToolRenderData.DirtyTexture : GLandscapeBlackTexture, true);
+					auto DirtyMaterialInstance = new FLandscapeMaskMaterialRenderProxy(GLandscapeDirtyMaterial->GetRenderProxy(), EditToolRenderData.DirtyTexture ? ToRawPtr(EditToolRenderData.DirtyTexture) : GLandscapeBlackTexture, true);
 					MaskMesh.MaterialRenderProxy = DirtyMaterialInstance;
 					Collector.RegisterOneFrameMaterialProxy(DirtyMaterialInstance);
 					Collector.AddMesh(ViewIndex, MaskMesh);
@@ -2552,7 +2552,7 @@ void FLandscapeComponentSceneProxy::GetDynamicMeshElements(const TArray<const FS
 					{
 						FMeshBatch& SelectMesh = Collector.AllocateMesh();
 						SelectMesh = MeshTools;
-						auto SelectMaterialInstance = new FLandscapeSelectMaterialRenderProxy(GSelectionRegionMaterial->GetRenderProxy(), EditToolRenderData.DataTexture ? EditToolRenderData.DataTexture : GLandscapeBlackTexture);
+						auto SelectMaterialInstance = new FLandscapeSelectMaterialRenderProxy(GSelectionRegionMaterial->GetRenderProxy(), EditToolRenderData.DataTexture ? ToRawPtr(EditToolRenderData.DataTexture) : GLandscapeBlackTexture);
 						SelectMesh.MaterialRenderProxy = SelectMaterialInstance;
 						Collector.RegisterOneFrameMaterialProxy(SelectMaterialInstance);
 						Collector.AddMesh(ViewIndex, SelectMesh);
@@ -2580,7 +2580,7 @@ void FLandscapeComponentSceneProxy::GetDynamicMeshElements(const TArray<const FS
 					{
 						FMeshBatch& MaskMesh = Collector.AllocateMesh();
 						MaskMesh = MeshTools;
-						auto MaskMaterialInstance = new FLandscapeMaskMaterialRenderProxy(GMaskRegionMaterial->GetRenderProxy(), EditToolRenderData.DataTexture ? EditToolRenderData.DataTexture : GLandscapeBlackTexture, !!(GLandscapeEditRenderMode & ELandscapeEditRenderMode::InvertedMask));
+						auto MaskMaterialInstance = new FLandscapeMaskMaterialRenderProxy(GMaskRegionMaterial->GetRenderProxy(), EditToolRenderData.DataTexture ? ToRawPtr(EditToolRenderData.DataTexture) : GLandscapeBlackTexture, !!(GLandscapeEditRenderMode & ELandscapeEditRenderMode::InvertedMask));
 						MaskMesh.MaterialRenderProxy = MaskMaterialInstance;
 						Collector.RegisterOneFrameMaterialProxy(MaskMaterialInstance);
 						Collector.AddMesh(ViewIndex, MaskMesh);
diff --git a/Engine/Source/Runtime/Renderer/Private/ScenePrivate.h b/Engine/Source/Runtime/Renderer/Private/ScenePrivate.h
--- a/Engine/Source/Runtime/Renderer/Private/ScenePrivate.h
+++ b/Engine/Source/Runtime/Renderer/Private/ScenePrivate.h
@@ -1451,7 +1451,7 @@ public:
 		auto InputAsMID = Cast<UMaterialInstanceDynamic>(InSource);
 
 		// fixup MID parents as this is not allowed, take the next MIC or Material.
-		UMaterialInterface* ParentOfTheNewMID = InputAsMID ? InputAsMID->Parent : InSource;
+		UMaterialInterface* ParentOfTheNewMID = InputAsMID ? ToRawPtr(InputAsMID->Parent) : InSource;
 
 		// this is not allowed and would cause an error later in the code
 		check(!ParentOfTheNewMID->IsA(UMaterialInstanceDynamic::StaticClass()));
