Landscape WorldPartition support: New ALandscapeSplineActor which is indepedent of ALandscapeProxy and implements new ILandscapeSplineInterface to support existing data prior to WP worlds.
- Non WP worlds still use the ALandscapeProxy as a ULandscapeSplinesComponent container
- WP worlds create one ALandscapeSplineActor per Spline that have their own bounds (so they can be loaded properly based on that)

diff --git a/Engine/Source/Editor/LandscapeEditor/Private/LandscapeEdMode.cpp b/Engine/Source/Editor/LandscapeEditor/Private/LandscapeEdMode.cpp
--- a/Engine/Source/Editor/LandscapeEditor/Private/LandscapeEdMode.cpp
+++ b/Engine/Source/Editor/LandscapeEditor/Private/LandscapeEdMode.cpp
@@ -3487,15 +3487,15 @@ ALandscape* FEdModeLandscape::ChangeComponentSetting(int32 NumComponentsX, int32
 
 			// Clone landscape splines
 			ALandscape* OldLandscapeActor = LandscapeInfo->LandscapeActor.Get();
-			if (OldLandscapeActor != nullptr && OldLandscapeActor->SplineComponent != NULL)
+			if (OldLandscapeActor != nullptr && OldLandscapeActor->GetSplinesComponent() != nullptr)
 			{
-				ULandscapeSplinesComponent* OldSplines = OldLandscapeActor->SplineComponent;
+				ULandscapeSplinesComponent* OldSplines = OldLandscapeActor->GetSplinesComponent();
 				ULandscapeSplinesComponent* NewSplines = DuplicateObject<ULandscapeSplinesComponent>(OldSplines, NewLandscape, OldSplines->GetFName());
 				NewSplines->AttachToComponent(NewLandscape->GetRootComponent(), FAttachmentTransformRules::KeepWorldTransform);
 
 				const FVector OldSplineScale = OldSplines->GetRelativeTransform().GetScale3D();
 				NewSplines->SetRelativeScale3D(FVector(OldSplineScale.X / LandscapeScaleFactor, OldSplineScale.Y / LandscapeScaleFactor, OldSplineScale.Z));
-				NewLandscape->SplineComponent = NewSplines;
+				NewLandscape->SetSplinesComponent(NewSplines);
 				NewSplines->RegisterComponent();
 
 				// TODO: Foliage on spline meshes
diff --git a/Engine/Source/Editor/LandscapeEditor/Private/LandscapeEdMode.h b/Engine/Source/Editor/LandscapeEditor/Private/LandscapeEdMode.h
--- a/Engine/Source/Editor/LandscapeEditor/Private/LandscapeEdMode.h
+++ b/Engine/Source/Editor/LandscapeEditor/Private/LandscapeEdMode.h
@@ -308,10 +308,11 @@ public:
 	void ShowSplineProperties();
 	bool HasSelectedSplineSegments() const;
 	void FlipSelectedSplineSegments();
-	void GetSelectedSplineOwners(TSet<ALandscapeProxy*>& SelectedSplineOwners) const;
+	void GetSelectedSplineOwners(TSet<AActor*>& SelectedSplineOwners) const;
 	virtual void SelectAllConnectedSplineControlPoints();
 	virtual void SelectAllConnectedSplineSegments();
 	virtual void SplineMoveToCurrentLevel();
+	virtual bool CanMoveSplineToCurrentLevel() const;
 	virtual void UpdateSplineMeshLevels();
 	void SetbUseAutoRotateOnJoin(bool InbAutoRotateOnJoin);
 	bool GetbUseAutoRotateOnJoin();
diff --git a/Engine/Source/Editor/LandscapeEditor/Private/LandscapeEdModeSplineTools.cpp b/Engine/Source/Editor/LandscapeEditor/Private/LandscapeEdModeSplineTools.cpp
--- a/Engine/Source/Editor/LandscapeEditor/Private/LandscapeEdModeSplineTools.cpp
+++ b/Engine/Source/Editor/LandscapeEditor/Private/LandscapeEdModeSplineTools.cpp
@@ -20,6 +20,8 @@
 #include "EditorViewportClient.h"
 #include "LandscapeToolInterface.h"
 #include "LandscapeProxy.h"
+#include "LandscapeSplineActor.h"
+#include "ILandscapeSplineInterface.h"
 #include "LandscapeEdMode.h"
 #include "ScopedTransaction.h"
 #include "LandscapeRender.h"
@@ -864,6 +866,22 @@ public:
 		}
 	}
 
+	bool CanMoveSelectedToLevel() const
+	{
+		// Move to level only supportd on LandscapeProxy Splines
+		for (ULandscapeSplineControlPoint* ControlPoint : SelectedSplineControlPoints)
+		{
+			ULandscapeSplinesComponent* LandscapeSplinesComp = ControlPoint->GetOuterULandscapeSplinesComponent();
+			ALandscapeProxy* FromProxy = LandscapeSplinesComp ? Cast<ALandscapeProxy>(LandscapeSplinesComp->GetOuter()) : nullptr;
+			if (!FromProxy)
+			{
+				return false;
+			}
+		}
+
+		return true;
+	}
+
 	void MoveSelectedToLevel()
 	{
 		TSet<ALandscapeProxy*> FromProxies;
@@ -891,15 +909,16 @@ public:
 			}
 		}
 				
-		if (ToLandscape && ToLandscape->SplineComponent)
+		ULandscapeSplinesComponent* SplineComponent = ToLandscape ? ToLandscape->GetSplinesComponent() : nullptr;
+		if (SplineComponent)
 		{
-			if (!ToLandscape->SplineComponent->IsRegistered())
+			if (!SplineComponent->IsRegistered())
 			{
-				ToLandscape->SplineComponent->RegisterComponent();
+				SplineComponent->RegisterComponent();
 			}
 			else
 			{
-				ToLandscape->SplineComponent->MarkRenderStateDirty();
+				SplineComponent->MarkRenderStateDirty();
 			}
 		}
 
@@ -929,61 +948,76 @@ public:
 		if (ViewportClient->IsCtrlPressed())
 		{
 			LandscapeInfo = InTarget.LandscapeInfo.Get();
-			ALandscapeProxy* Landscape = nullptr;
+			ILandscapeSplineInterface* SplineOwner = nullptr;
 
 			// If we have a selection use the landscape of the selected spline
-			ULandscapeSplinesComponent* SplinesComponent = nullptr;
 			if (SelectedSplineControlPoints.Num() > 0)
 			{
 				ULandscapeSplineControlPoint* FirstPoint = *SelectedSplineControlPoints.CreateConstIterator();
-				SplinesComponent = FirstPoint->GetOuterULandscapeSplinesComponent();
+				ULandscapeSplinesComponent* SelectedSplinesComponent = FirstPoint->GetOuterULandscapeSplinesComponent();
 			
-				if (SplinesComponent)
+				if (SelectedSplinesComponent)
 				{
-					Landscape = SplinesComponent->GetTypedOuter<ALandscapeProxy>();
+					SplineOwner = SelectedSplinesComponent->GetSplineOwner();
 				}
 			}
 					
+			const bool bIsGridBased = EdMode->IsGridBased();
+
 			// Hit Test
-			if (!Landscape)
+			if (!SplineOwner)
 			{
 				HHitProxy* HitProxy = ViewportClient->Viewport->GetHitProxy(ViewportClient->Viewport->GetMouseX(), ViewportClient->Viewport->GetMouseY());
 				if (HitProxy->IsA(HActor::StaticGetType()))
 				{
 					HActor* ActorProxy = (HActor*)HitProxy;
-					if (ALandscapeProxy* Proxy = Cast<ALandscapeProxy>(ActorProxy->Actor))
+					// Here we want to make sure we are selecting the valid type (Grid based or not)
+					if (bIsGridBased)
+					{
+						SplineOwner = Cast<ALandscapeSplineActor>(ActorProxy->Actor);
+					}
+					else
 					{
-						Landscape = Proxy;
+						SplineOwner = Cast<ALandscapeProxy>(ActorProxy->Actor);
 					}
 				}
-
-				
 			}
-
-			// Default to Current level Landscape
-			if (!Landscape)
+				
+			// Open transaction here because we might be creating an actor or later a spline component
+			FScopedTransaction Transaction(LOCTEXT("LandscapeSpline_AddControlPoint", "Add Landscape Spline Control Point"));
+			if (!SplineOwner)
 			{
-				Landscape = LandscapeInfo->GetCurrentLevelLandscapeProxy(true);
+				if (bIsGridBased)
+				{
+					// Create Spline Actor at World Position (InHitLocation is local to Landscape)
+					FVector WorldHitLocation = LandscapeInfo->LandscapeActor.Get()->LandscapeActorToWorld().TransformPosition(InHitLocation);
+					SplineOwner = LandscapeInfo->CreateSplineActor(WorldHitLocation);
+				}
+				else
+				{
+					// Default to Current level Landscape
+					SplineOwner = LandscapeInfo->GetCurrentLevelLandscapeProxy(true);
+				}
 			}
 
-			// No Landscape found
-			if (!Landscape)
+			// No Spline Owner found
+			if (!SplineOwner)
 			{
 				return false;
 			}
 
+			ULandscapeSplinesComponent* SplinesComponent = SplineOwner->GetSplinesComponent();
 			if (!SplinesComponent)
 			{
-				if (!Landscape->SplineComponent)
-				{
-					Landscape->CreateSplineComponent();
-					check(Landscape->SplineComponent);
-				}
-				SplinesComponent = Landscape->SplineComponent;
+				SplineOwner->CreateSplineComponent();
+				SplinesComponent = SplineOwner->GetSplinesComponent();
+				check(SplinesComponent);
 			}
 
-			const FTransform LandscapeToSpline = Landscape->LandscapeActorToWorld().GetRelativeTransform(SplinesComponent->GetComponentTransform());
+			// Get Main Landscape Actor to Spline Component Transform
+			const FTransform LandscapeToSpline = SplineOwner->LandscapeActorToWorld().GetRelativeTransform(SplinesComponent->GetComponentTransform());
 
+			// Local to SplineComponent
 			AddControlPoint(SplinesComponent, LandscapeToSpline.TransformPosition(InHitLocation));
 
 			GUnrealEd->RedrawLevelEditingViewports();
@@ -1305,9 +1339,10 @@ public:
 					{
 						HLandscapeSplineProxy_Segment* SplineProxy = (HLandscapeSplineProxy_Segment*)HitProxy;
 						ClickedSplineSegment = SplineProxy->SplineSegment;
-						ALandscapeProxy* LandscapeProxy = ClickedSplineSegment->GetTypedOuter<ALandscapeProxy>();
-						check(LandscapeProxy);
-						LandscapeToSpline = LandscapeProxy->LandscapeActorToWorld().GetRelativeTransform(ClickedSplineSegment->GetOuterULandscapeSplinesComponent()->GetComponentTransform());
+						ULandscapeSplinesComponent* SplineComponent = ClickedSplineSegment->GetOuterULandscapeSplinesComponent();
+						ILandscapeSplineInterface* SplineOwner = SplineComponent->GetSplineOwner();
+						check(SplineOwner);
+						LandscapeToSpline = SplineOwner->LandscapeActorToWorld().GetRelativeTransform(SplineComponent->GetComponentTransform());
 					}
 					else if (HitProxy->IsA(HActor::StaticGetType()))
 					{
@@ -1327,8 +1362,8 @@ public:
 										ClickedSplineSegment = SplineSegment;
 										// Find actual SplineComponent owner of the Segment (not the SplineComponent owner of the mesh)
 										SplineComponent = SplineSegment->GetTypedOuter<ULandscapeSplinesComponent>();
-										ALandscapeProxy* LandscapeProxy = CastChecked<ALandscapeProxy>(SplineComponent->GetOwner());
-										LandscapeToSpline = LandscapeProxy->LandscapeActorToWorld().GetRelativeTransform(SplineComponent->GetComponentTransform());
+										ILandscapeSplineInterface* SplineOwner = SplineComponent->GetSplineOwner();
+										LandscapeToSpline = SplineOwner->LandscapeActorToWorld().GetRelativeTransform(SplineComponent->GetComponentTransform());
 									}
 								}
 							}
@@ -1716,12 +1751,12 @@ public:
 		{
 			for (const FLandscapeListInfo& Info : EdMode->GetLandscapeList())
 			{
-				Info.Info->ForAllLandscapeProxies([this](ALandscapeProxy* Proxy)
+				Info.Info->ForAllSplineActors([this](TScriptInterface<ILandscapeSplineInterface> SplineOwner)
 				{
-					if (Proxy->SplineComponent)
+					if (ULandscapeSplinesComponent* SplineComponent = SplineOwner->GetSplinesComponent())
 					{
-						Algo::CopyIf(Proxy->SplineComponent->ControlPoints, SelectedSplineControlPoints, &ULandscapeSplineControlPoint::IsSplineSelected);
-						Algo::CopyIf(Proxy->SplineComponent->Segments,      SelectedSplineSegments,      &ULandscapeSplineSegment::IsSplineSelected);
+						Algo::CopyIf(SplineComponent->ControlPoints, SelectedSplineControlPoints, &ULandscapeSplineControlPoint::IsSplineSelected);
+						Algo::CopyIf(SplineComponent->Segments,      SelectedSplineSegments,      &ULandscapeSplineSegment::IsSplineSelected);
 					}
 				});
 			}
@@ -1730,16 +1765,16 @@ public:
 		{
 			for (const FLandscapeListInfo& Info : EdMode->GetLandscapeList())
 			{
-				Info.Info->ForAllLandscapeProxies([](ALandscapeProxy* Proxy)
+				Info.Info->ForAllSplineActors([](TScriptInterface<ILandscapeSplineInterface> SplineOwner)
 				{
-					if (Proxy->SplineComponent)
+					if (ULandscapeSplinesComponent* SplineComponent = SplineOwner->GetSplinesComponent())
 					{
-						for (ULandscapeSplineControlPoint* ControlPoint : Proxy->SplineComponent->ControlPoints)
+						for (ULandscapeSplineControlPoint* ControlPoint : SplineComponent->ControlPoints)
 						{
 							ControlPoint->SetSplineSelected(false);
 						}
 
-						for (ULandscapeSplineSegment* Segment : Proxy->SplineComponent->Segments)
+						for (ULandscapeSplineSegment* Segment : SplineComponent->Segments)
 						{
 							Segment->SetSplineSelected(false);
 						}
@@ -1761,11 +1796,11 @@ public:
 
 		for (const FLandscapeListInfo& Info : EdMode->GetLandscapeList())
 		{
-			Info.Info->ForAllLandscapeProxies([](ALandscapeProxy* Proxy)
+			Info.Info->ForAllSplineActors([this](TScriptInterface<ILandscapeSplineInterface> SplineOwner)
 			{
-				if (Proxy->SplineComponent)
+				if (ULandscapeSplinesComponent* SplineComponent = SplineOwner->GetSplinesComponent())
 				{
-					Proxy->SplineComponent->ShowSplineEditorMesh(true);
+					SplineComponent->ShowSplineEditorMesh(true);
 				}
 			});
 		}
@@ -1778,11 +1813,11 @@ public:
 
 		for (const FLandscapeListInfo& Info : EdMode->GetLandscapeList())
 		{
-			Info.Info->ForAllLandscapeProxies([](ALandscapeProxy* Proxy)
+			Info.Info->ForAllSplineActors([this](TScriptInterface<ILandscapeSplineInterface> SplineOwner)
 			{
-				if (Proxy->SplineComponent)
+				if (ULandscapeSplinesComponent* SplineComponent = SplineOwner->GetSplinesComponent())
 				{
-					Proxy->SplineComponent->ShowSplineEditorMesh(false);
+					SplineComponent->ShowSplineEditorMesh(false);
 				}
 			});
 		}
@@ -1997,7 +2032,7 @@ public:
 	{
 		FString PasteString;
 		FPlatformApplicationMisc::ClipboardPaste(PasteString);
-		if (PasteString.StartsWith("BEGIN SPLINES"))
+		if (PasteString.StartsWith(FLandscapeSplineTextObjectFactory::SplineBeginTag))
 		{
 			return EEditAction::Process;
 		}
@@ -2081,6 +2116,20 @@ public:
 
 	void InternalProcessEditCopy(FString* OutData = NULL)
 	{
+		bool bFirstSplineLocation = true;
+		FVector SplineLocation;
+
+		auto GetSplineLocation = [&bFirstSplineLocation, &SplineLocation](ULandscapeSplineControlPoint* ControlPoint, const FVector& Location)
+		{
+			if (bFirstSplineLocation)
+			{
+				ILandscapeSplineInterface* SplineOwner = ControlPoint->GetOuterULandscapeSplinesComponent()->GetSplineOwner();
+				const FTransform LocalToWorld = ControlPoint->GetOuterULandscapeSplinesComponent()->GetComponentTransform();
+				SplineLocation = LocalToWorld.TransformPosition(ControlPoint->Location);
+				bFirstSplineLocation = false;
+			}
+		};
+
 		if (SelectedSplineControlPoints.Num() > 0 || SelectedSplineSegments.Num() > 0)
 		{
 			TArray<UObject*> Objects;
@@ -2089,10 +2138,12 @@ public:
 			// Control Points then segments
 			for (ULandscapeSplineControlPoint* ControlPoint : SelectedSplineControlPoints)
 			{
+				GetSplineLocation(ControlPoint, ControlPoint->Location);
 				Objects.Add(ControlPoint);
 			}
 			for (ULandscapeSplineSegment* Segment : SelectedSplineSegments)
 			{
+				GetSplineLocation(Segment->Connections[0].ControlPoint, Segment->Connections[0].ControlPoint->Location);
 				Objects.AddUnique(Segment->Connections[0].ControlPoint);
 				Objects.AddUnique(Segment->Connections[1].ControlPoint);
 			}
@@ -2105,12 +2156,13 @@ public:
 			FStringOutputDevice Ar;
 			const FExportObjectInnerContext Context;
 
-			Ar.Logf(TEXT("Begin Splines\r\n"));
+			Ar.Logf(TEXT("%s\r\n"), *FLandscapeSplineTextObjectFactory::SplineBeginTag);
+			Ar.Logf(TEXT("%s%s\r\n"), *FLandscapeSplineTextObjectFactory::SplineLocationTag, *SplineLocation.ToString());
 			for (UObject* Object : Objects)
 			{
 				UExporter::ExportToOutputDevice(&Context, Object, NULL, Ar, TEXT("copy"), 3, PPF_ExportsNotFullyQualified | PPF_Copy | PPF_Delimited, false, Object->GetOuter());
 			}
-			Ar.Logf(TEXT("End Splines\r\n"));
+			Ar.Logf(TEXT("%s\r\n"), *FLandscapeSplineTextObjectFactory::SplineEndTag);
 
 			if (OutData != NULL)
 			{
@@ -2127,17 +2179,34 @@ public:
 	{
 		FScopedTransaction Transaction(LOCTEXT("LandscapeSpline_Paste", "Paste Landscape Splines"));
 
-		ALandscapeProxy* Landscape = EdMode->CurrentToolTarget.LandscapeInfo->GetCurrentLevelLandscapeProxy(true);
-		if (!Landscape)
+		TScriptInterface<ILandscapeSplineInterface> SplineOwner;
+		const bool bGridBased = EdMode->IsGridBased();
+		if (bGridBased)
+		{
+			if (ULandscapeInfo* CurrentLandscapeInfo = EdMode->CurrentToolTarget.LandscapeInfo.Get())
+			{
+				SplineOwner = CurrentLandscapeInfo->CreateSplineActor(FVector::ZeroVector);
+			}
+		}
+		else
+		{
+			SplineOwner = EdMode->CurrentToolTarget.LandscapeInfo->GetCurrentLevelLandscapeProxy(true);
+		}
+
+		if (!SplineOwner)
 		{
 			return;
 		}
-		if (!Landscape->SplineComponent)
+		
+		ULandscapeSplinesComponent* SplineComponent = SplineOwner->GetSplinesComponent();
+		if (!SplineComponent)
 		{
-			Landscape->CreateSplineComponent();
-			check(Landscape->SplineComponent);
+			SplineOwner->CreateSplineComponent();
+			SplineComponent = SplineOwner->GetSplinesComponent();
+			check(SplineComponent);
 		}
-		Landscape->SplineComponent->Modify();
+		
+		SplineComponent->Modify();
 
 		const TCHAR* Data = NULL;
 		FString PasteString;
@@ -2152,7 +2221,12 @@ public:
 		}
 
 		FLandscapeSplineTextObjectFactory Factory;
-		TArray<UObject*> OutObjects = Factory.ImportSplines(Landscape->SplineComponent, Data);
+		TArray<UObject*> OutObjects = Factory.ImportSplines(SplineComponent, Data);
+
+		if (ALandscapeSplineActor* SplineActor = Cast<ALandscapeSplineActor>(SplineOwner.GetObject()))
+		{
+			SplineActor->SetActorLocation(Factory.SplineLocation);
+		}
 
 		if (bOffset)
 		{
@@ -2161,9 +2235,7 @@ public:
 				ULandscapeSplineControlPoint* ControlPoint = Cast<ULandscapeSplineControlPoint>(Object);
 				if (ControlPoint != NULL)
 				{
-					Landscape->SplineComponent->ControlPoints.Add(ControlPoint);
 					ControlPoint->Location += FVector(500, 500, 0);
-
 					ControlPoint->UpdateSplinePoints();
 				}
 			}
@@ -2253,16 +2325,16 @@ void FEdModeLandscape::ShowSplineProperties()
 	}
 }
 
-void FEdModeLandscape::GetSelectedSplineOwners(TSet<ALandscapeProxy*>& SelectedSplineOwners) const
+void FEdModeLandscape::GetSelectedSplineOwners(TSet<AActor*>& SelectedSplineOwners) const
 {
 	for (ULandscapeSplineSegment* Segment : SplinesTool->SelectedSplineSegments)
 	{
-		SelectedSplineOwners.Add(Segment->GetTypedOuter<ALandscapeProxy>());
+		SelectedSplineOwners.Add(Segment->GetTypedOuter<AActor>());
 	}
 
 	for (ULandscapeSplineControlPoint* ControlPoint : SplinesTool->SelectedSplineControlPoints)
 	{
-		SelectedSplineOwners.Add(ControlPoint->GetTypedOuter<ALandscapeProxy>());
+		SelectedSplineOwners.Add(ControlPoint->GetTypedOuter<AActor>());
 	}
 }
 
@@ -2311,6 +2383,16 @@ void FEdModeLandscape::SplineMoveToCurrentLevel()
 	}
 }
 
+bool FEdModeLandscape::CanMoveSplineToCurrentLevel() const
+{
+	if (SplinesTool)
+	{
+		return SplinesTool->CanMoveSelectedToLevel();
+	}
+
+	return false;
+}
+
 void FEdModeLandscape::UpdateSplineMeshLevels()
 {
 	FScopedTransaction Transaction(LOCTEXT("LandscapeSpline_UpdateSplineMeshes", "Update Spline Meshes Level"));
diff --git a/Engine/Source/Editor/LandscapeEditor/Private/LandscapeSplineDetails.cpp b/Engine/Source/Editor/LandscapeEditor/Private/LandscapeSplineDetails.cpp
--- a/Engine/Source/Editor/LandscapeEditor/Private/LandscapeSplineDetails.cpp
+++ b/Engine/Source/Editor/LandscapeEditor/Private/LandscapeSplineDetails.cpp
@@ -11,6 +11,7 @@
 #include "DetailWidgetRow.h"
 #include "DetailCategoryBuilder.h"
 #include "DetailLayoutBuilder.h"
+#include "ILandscapeSplineInterface.h"
 
 #define LOCTEXT_NAMESPACE "LandscapeSplineDetails"
 
@@ -138,7 +139,7 @@ bool FLandscapeSplineDetails::IsFlipSegmentButtonEnabled() const
 
 FText FLandscapeSplineDetails::OnGetSplineOwningLandscapeText() const
 {
-	TSet<ALandscapeProxy*> SplineOwners;
+	TSet<AActor*> SplineOwners;
 	FEdModeLandscape* LandscapeEdMode = GetEditorMode();
 	if (LandscapeEdMode && LandscapeEdMode->CurrentToolTarget.LandscapeInfo.IsValid())
 	{
@@ -146,7 +147,7 @@ FText FLandscapeSplineDetails::OnGetSplineOwningLandscapeText() const
 	}
 
 	FString SplineOwnersStr;
-	for (ALandscapeProxy* Owner : SplineOwners)
+	for (AActor* Owner : SplineOwners)
 	{
 		if (Owner)
 		{
@@ -154,7 +155,7 @@ FText FLandscapeSplineDetails::OnGetSplineOwningLandscapeText() const
 			{
 				SplineOwnersStr += ", ";
 			}
-			
+
 			SplineOwnersStr += Owner->GetActorLabel();
 		}
 	}
@@ -198,7 +199,12 @@ FReply FLandscapeSplineDetails::OnMoveToCurrentLevelButtonClicked()
 bool FLandscapeSplineDetails::IsMoveToCurrentLevelButtonEnabled() const
 {
 	FEdModeLandscape* LandscapeEdMode = GetEditorMode();
-	return (LandscapeEdMode && LandscapeEdMode->CurrentToolTarget.LandscapeInfo.IsValid() && LandscapeEdMode->CurrentToolTarget.LandscapeInfo->GetCurrentLevelLandscapeProxy(true));
+	if (LandscapeEdMode && LandscapeEdMode->CurrentToolTarget.LandscapeInfo.IsValid() && LandscapeEdMode->CurrentToolTarget.LandscapeInfo->GetCurrentLevelLandscapeProxy(true))
+	{
+		return LandscapeEdMode->CanMoveSplineToCurrentLevel();
+	}
+
+	return false;
 }
 
 FReply FLandscapeSplineDetails::OnUpdateSplineMeshLevelsButtonClicked()
diff --git a/Engine/Source/Editor/LandscapeEditor/Private/LandscapeSplineImportExport.cpp b/Engine/Source/Editor/LandscapeEditor/Private/LandscapeSplineImportExport.cpp
--- a/Engine/Source/Editor/LandscapeEditor/Private/LandscapeSplineImportExport.cpp
+++ b/Engine/Source/Editor/LandscapeEditor/Private/LandscapeSplineImportExport.cpp
@@ -6,6 +6,10 @@
 
 #define LOCTEXT_NAMESPACE "Landscape"
 
+const FString FLandscapeSplineTextObjectFactory::SplineLocationTag("SPLINE LOCATION: ");
+const FString FLandscapeSplineTextObjectFactory::SplineBeginTag("BEGIN_SPLINES");
+const FString FLandscapeSplineTextObjectFactory::SplineEndTag("END_SPLINES");
+
 FLandscapeSplineTextObjectFactory::FLandscapeSplineTextObjectFactory(FFeedbackContext* InWarningContext /*= GWarn*/)
 	: FCustomizableTextObjectFactory(InWarningContext)
 {
@@ -13,11 +17,11 @@ FLandscapeSplineTextObjectFactory::FLandscapeSplineTextObjectFactory(FFeedbackCo
 
 TArray<UObject*> FLandscapeSplineTextObjectFactory::ImportSplines(UObject* InParent, const TCHAR* TextBuffer)
 {
-	if (FParse::Command(&TextBuffer, TEXT("BEGIN SPLINES")))
+	if (FParse::Command(&TextBuffer, *SplineBeginTag))
 	{
 		ProcessBuffer(InParent, RF_Transactional, TextBuffer);
 
-		//FParse::Command(&TextBuffer, TEXT("END SPLINES"));
+		check(bProcessedEnd);
 	}
 
 	return OutObjects;
@@ -41,4 +45,16 @@ bool FLandscapeSplineTextObjectFactory::CanCreateClass(UClass* ObjectClass, bool
 	return false;
 }
 
+void FLandscapeSplineTextObjectFactory::ProcessUnidentifiedLine(const FString& StrLine)
+{
+	if (StrLine.StartsWith(SplineEndTag))
+	{
+		bProcessedEnd = true;
+	}
+	else if (StrLine.StartsWith(SplineLocationTag))
+	{
+		SplineLocation.InitFromString(StrLine.Mid(SplineLocationTag.Len()));
+	}
+}
+
 #undef LOCTEXT_NAMESPACE
diff --git a/Engine/Source/Editor/LandscapeEditor/Private/LandscapeSplineImportExport.h b/Engine/Source/Editor/LandscapeEditor/Private/LandscapeSplineImportExport.h
--- a/Engine/Source/Editor/LandscapeEditor/Private/LandscapeSplineImportExport.h
+++ b/Engine/Source/Editor/LandscapeEditor/Private/LandscapeSplineImportExport.h
@@ -12,9 +12,16 @@ public:
 
 	TArray<UObject*> ImportSplines(UObject* InParent, const TCHAR* TextBuffer);
 
+	static const FString SplineLocationTag;
+	static const FString SplineBeginTag;
+	static const FString SplineEndTag;
+
+	FVector SplineLocation;
 protected:
 	TArray<UObject*> OutObjects;
-
+	bool bProcessedEnd = false;
+	
 	virtual void ProcessConstructedObject(UObject* CreatedObject) override;
 	virtual bool CanCreateClass(UClass* ObjectClass, bool& bOmitSubObjs) const override;
+	virtual void ProcessUnidentifiedLine(const FString& StrLine) override;
 };
diff --git a/Engine/Source/Editor/UnrealEd/Private/Commandlets/WorldPartitionConvertCommandlet.cpp b/Engine/Source/Editor/UnrealEd/Private/Commandlets/WorldPartitionConvertCommandlet.cpp
--- a/Engine/Source/Editor/UnrealEd/Private/Commandlets/WorldPartitionConvertCommandlet.cpp
+++ b/Engine/Source/Editor/UnrealEd/Private/Commandlets/WorldPartitionConvertCommandlet.cpp
@@ -43,6 +43,10 @@
 #include "LandscapeStreamingProxy.h"
 #include "LandscapeInfo.h"
 #include "LandscapeConfigHelper.h"
+#include "ILandscapeSplineInterface.h"
+#include "LandscapeSplineActor.h"
+#include "LandscapeSplinesComponent.h"
+#include "LandscapeSplineControlPoint.h"
 #include "LandscapeGizmoActor.h"
 #include "WorldPartition/DataLayer/DataLayer.h"
 #include "WorldPartition/DataLayer/WorldDataLayers.h"
@@ -865,6 +869,39 @@ int32 UWorldPartitionConvertCommandlet::Main(const FString& Params)
 			LandscapeInfo->RegisterActor(NewLandscape);
 		}
 
+		TSet<AActor*> NewSplineActors;
+						
+		auto MoveControlPointToNewSplineActor = [&NewSplineActors, LandscapeInfo](ULandscapeSplineControlPoint* ControlPoint)
+		{
+			AActor* CurrentOwner = ControlPoint->GetTypedOuter<AActor>();
+			// Control point as already been moved through its connected segments
+			if (NewSplineActors.Contains(CurrentOwner))
+			{
+				return;
+			}
+			
+			const FTransform LocalToWorld = ControlPoint->GetOuterULandscapeSplinesComponent()->GetComponentTransform();
+			const FVector NewActorLocation = LocalToWorld.TransformPosition(ControlPoint->Location);
+						
+			ALandscapeSplineActor* NewSplineActor = LandscapeInfo->CreateSplineActor(NewActorLocation);
+
+			// ULandscapeSplinesComponent doesn't assign SplineEditorMesh when IsCommandlet() is true.
+			NewSplineActor->GetSplinesComponent()->SetDefaultEditorSplineMesh();
+
+			NewSplineActors.Add(NewSplineActor);
+			LandscapeInfo->MoveSpline(ControlPoint, NewSplineActor);
+		};
+				
+		// Iterate on copy since we are creating new spline actors
+		TArray<TScriptInterface<ILandscapeSplineInterface>> OldSplineActors(LandscapeInfo->GetSplineActors());
+		for (TScriptInterface<ILandscapeSplineInterface> PreviousSplineActor : OldSplineActors)
+		{
+			if (ULandscapeSplinesComponent* SplineComponent = PreviousSplineActor->GetSplinesComponent())
+			{
+				SplineComponent->ForEachControlPoint(MoveControlPointToNewSplineActor);
+			}
+		}
+
 		TSet<AActor*> ActorsToDelete;
 		FLandscapeConfigHelper::ChangeGridSize(LandscapeInfo, LandscapeGridSize, ActorsToDelete);
 		for (AActor* ActorToDelete : ActorsToDelete)
diff --git a/Engine/Source/Runtime/Landscape/Classes/ILandscapeSplineInterface.h b/Engine/Source/Runtime/Landscape/Classes/ILandscapeSplineInterface.h
new file mode 100644
--- /dev/null
+++ b/Engine/Source/Runtime/Landscape/Classes/ILandscapeSplineInterface.h
@@ -0,0 +1,33 @@
+// Copyright Epic Games, Inc. All Rights Reserved.
+
+#pragma once
+
+#include "CoreMinimal.h"
+#include "UObject/ObjectMacros.h"
+#include "UObject/Interface.h"
+#include "ILandscapeSplineInterface.generated.h"
+
+UINTERFACE()
+class LANDSCAPE_API ULandscapeSplineInterface : public UInterface
+{
+	GENERATED_UINTERFACE_BODY()
+};
+
+class ULandscapeSplinesComponent;
+class ULandscapeInfo;
+
+class LANDSCAPE_API ILandscapeSplineInterface
+{
+	GENERATED_IINTERFACE_BODY()
+public:
+	virtual FGuid GetLandscapeGuid() const PURE_VIRTUAL(ILandscapeSplineInterface::GetLandscapeGuid, return FGuid(););
+	virtual ULandscapeInfo* GetLandscapeInfo() const PURE_VIRTUAL(ILandscapeSplineInterface::GetLandscapeInfo, return nullptr;);
+	virtual FTransform LandscapeActorToWorld() const PURE_VIRTUAL(ILandscapeSplineInterface::LandscapeActorToWorld, return FTransform::Identity;);
+	virtual ULandscapeSplinesComponent* GetSplinesComponent() const PURE_VIRTUAL(ILandscapeSplineInterface::GetSplineComponent, return nullptr;);
+
+#if WITH_EDITOR
+	virtual bool SupportsForeignSplineMesh() const PURE_VIRTUAL(ILandscapeSplineInterface::SupportsForeignSplineMesh, return false;);
+	virtual void CreateSplineComponent() PURE_VIRTUAL(ILandscapeSplineInterface::CreateSplineComponent);
+	virtual void CreateSplineComponent(const FVector& Scale3D) PURE_VIRTUAL(ILandscapeSplineInterface::CreateSplineComponent);
+#endif
+};
\ No newline at end of file
diff --git a/Engine/Source/Runtime/Landscape/Classes/LandscapeInfo.h b/Engine/Source/Runtime/Landscape/Classes/LandscapeInfo.h
--- a/Engine/Source/Runtime/Landscape/Classes/LandscapeInfo.h
+++ b/Engine/Source/Runtime/Landscape/Classes/LandscapeInfo.h
@@ -7,11 +7,14 @@
 #include "UObject/Object.h"
 #include "Misc/Guid.h"
 #include "UObject/LazyObjectPtr.h"
+#include "UObject/ScriptInterface.h"
+#include "ILandscapeSplineInterface.h"
 #include "LandscapeInfo.generated.h"
 
 class ALandscape;
 class ALandscapeProxy;
 class ALandscapeStreamingProxy;
+class ALandscapeSplineActor;
 class ULandscapeComponent;
 class ULandscapeLayerInfoObject;
 class ULevel;
@@ -138,6 +141,11 @@ public:
 	TArray<ALandscapeStreamingProxy*> Proxies;
 
 private:
+#if WITH_EDITORONLY_DATA
+	UPROPERTY()
+	TArray<TScriptInterface<ILandscapeSplineInterface>> SplineActors;
+#endif
+
 	TSet<ULandscapeComponent*> SelectedComponents;
 
 	TSet<ULandscapeComponent*> SelectedRegionComponents;
@@ -248,24 +256,32 @@ public:
 	/** Moves Components to target proxy. */
 	LANDSCAPE_API ALandscapeProxy* MoveComponentsToProxy(const TArray<ULandscapeComponent*>& InComponents, ALandscapeProxy* LandscapeProxy, bool bSetPositionAndOffset = false, ULevel* TargetLevel = nullptr);
 
-	/** Moves Splines connected to this control point to target level. Creates ULandscapeSplineComponent if needed. */
+	/** Moves Splines connected to this control point to target level. Creates ULandscapeSplinesComponent if needed. */
 	LANDSCAPE_API void MoveSplineToLevel(ULandscapeSplineControlPoint* InControlPoint, ULevel* TargetLevel);
 
-	/** Moves all Splines to target level. Creates ULandscapeSplineComponent if needed. */
+	/** Moves all Splines to target level. Creates ULandscapeSplinesComponent if needed. */
 	LANDSCAPE_API void MoveSplinesToLevel(ULandscapeSplinesComponent* InSplineComponent, ULevel* TargetLevel);
 
-	/** Moves Splines connected to this control point to target Proxy. Creates ULandscapeSplineComponent if needed. */
+	/** Moves Splines connected to this control point to target Proxy. Creates ULandscapeSplinesComponent if needed. */
 	LANDSCAPE_API void MoveSplineToProxy(ULandscapeSplineControlPoint* InControlPoint, ALandscapeProxy* InLandscapeProxy);
 
 	/** Moves all Splines to target Proxy. Creates ULandscapeSplineComponent if needed */
 	LANDSCAPE_API void MoveSplinesToProxy(ULandscapeSplinesComponent* InSplineComponent, ALandscapeProxy* InLandscapeProxy);
 
+	/** Moves Splines connected to this control point to target Proxy. Creates ULandscapeSplinesComponent if needed. */
+	LANDSCAPE_API void MoveSpline(ULandscapeSplineControlPoint* InControlPoint, TScriptInterface<ILandscapeSplineInterface> InNewOwner);
+
+	/** Moves all Splines to target Spline owner. Creates ULandscapeSplinesComponent if needed. */
+	LANDSCAPE_API void MoveSplines(ULandscapeSplinesComponent* InSplineComponent, TScriptInterface<ILandscapeSplineInterface> InNewOwner);
+
 	/** Will call UpdateAllComponentMaterialInstances on all LandscapeProxies */
 	LANDSCAPE_API void UpdateAllComponentMaterialInstances();
 
 	/** Returns LandscapeStreamingProxy Cell Size in WorldPartition */
 	LANDSCAPE_API uint32 GetGridSize(uint32 InGridSizeInComponents) const;
 #endif
+	static ULandscapeInfo* Find(UWorld* InWorld, const FGuid& LandscapeGuid);
+	static ULandscapeInfo* FindOrCreate(UWorld* InWorld, const FGuid& LandscapeGuid);
 
 	/**
 	 * Runs the given function on the root landscape actor and all streaming proxies
@@ -299,11 +315,23 @@ public:
 	LANDSCAPE_API void RegisterCollisionComponent(ULandscapeHeightfieldCollisionComponent* Component);
 	LANDSCAPE_API void UnregisterCollisionComponent(ULandscapeHeightfieldCollisionComponent* Component);
 
+#if WITH_EDITOR
+	LANDSCAPE_API ALandscapeSplineActor* CreateSplineActor(const FVector& Location);
+
+	LANDSCAPE_API void ForAllSplineActors(TFunctionRef<void(TScriptInterface<ILandscapeSplineInterface>)> Fn) const;
+	LANDSCAPE_API TArray<TScriptInterface<ILandscapeSplineInterface>> GetSplineActors() const;
+
+	LANDSCAPE_API void RegisterSplineActor(TScriptInterface<ILandscapeSplineInterface> SplineActor);
+	LANDSCAPE_API void UnregisterSplineActor(TScriptInterface<ILandscapeSplineInterface> SplineActor);
+
+	LANDSCAPE_API void RequestSplineLayerUpdate();
+#endif
+
 #if WITH_EDITOR
 private:
-	bool ApplySplinesInternal(bool bOnlySelected, ALandscapeProxy* Proxy, TSet<ULandscapeComponent*>* OutModifiedComponents, bool bMarkPackageDirty, int32 LandscapeMinX, int32 LandscapeMinY, int32 LandscapeMaxX, int32 LandscapeMaxY, TFunctionRef<TSharedPtr<FModulateAlpha>(ULandscapeLayerInfoObject*)> GetOrCreateModulate);
-	void MoveSegmentToLandscape(ULandscapeSplineSegment* InSegment, ALandscapeProxy* FromProxy, ALandscapeProxy* ToProxy);
-	void MoveControlPointToLandscape(ULandscapeSplineControlPoint* InControlPoint, ALandscapeProxy* FromProxy, ALandscapeProxy* ToProxy);
+	bool ApplySplinesInternal(bool bOnlySelected, TScriptInterface<ILandscapeSplineInterface> SplineOwner, TSet<ULandscapeComponent*>* OutModifiedComponents, bool bMarkPackageDirty, int32 LandscapeMinX, int32 LandscapeMinY, int32 LandscapeMaxX, int32 LandscapeMaxY, TFunctionRef<TSharedPtr<FModulateAlpha>(ULandscapeLayerInfoObject*)> GetOrCreateModulate);
+	void MoveSegment(ULandscapeSplineSegment* InSegment, TScriptInterface<ILandscapeSplineInterface> From, TScriptInterface<ILandscapeSplineInterface> To);
+	void MoveControlPoint(ULandscapeSplineControlPoint* InControlPoint, TScriptInterface<ILandscapeSplineInterface> From, TScriptInterface<ILandscapeSplineInterface> To);
 	bool UpdateLayerInfoMapInternal(ALandscapeProxy* Proxy, bool bInvalidate);
 #endif
 };
diff --git a/Engine/Source/Runtime/Landscape/Classes/LandscapeProxy.h b/Engine/Source/Runtime/Landscape/Classes/LandscapeProxy.h
--- a/Engine/Source/Runtime/Landscape/Classes/LandscapeProxy.h
+++ b/Engine/Source/Runtime/Landscape/Classes/LandscapeProxy.h
@@ -16,7 +16,7 @@
 #include "LandscapeWeightmapUsage.h"
 #include "VT/RuntimeVirtualTextureEnum.h"
 #include "ActorPartition/PartitionActor.h"
-
+#include "ILandscapeSplineInterface.h"
 #include "LandscapeProxy.generated.h"
 
 class ALandscape;
@@ -402,7 +402,7 @@ private:
 };
 
 UCLASS(Abstract, MinimalAPI, NotBlueprintable, NotPlaceable, hidecategories=(Display, Attachment, Physics, Debug, Lighting, LOD), showcategories=(Lighting, Rendering, "Utilities|Transformation"), hidecategories=(Mobility))
-class ALandscapeProxy : public APartitionActor
+class ALandscapeProxy : public APartitionActor, public ILandscapeSplineInterface
 {
 	GENERATED_BODY()
 
@@ -411,10 +411,10 @@ public:
 
 	virtual ~ALandscapeProxy();
 
+protected:
 	UPROPERTY()
 	ULandscapeSplinesComponent* SplineComponent;
 
-protected:
 	/** Guid for LandscapeEditorInfo **/
 	UPROPERTY()
 	FGuid LandscapeGuid;
@@ -865,7 +865,7 @@ public:
 	virtual FGuid GetGridGuid() const override { return LandscapeGuid; }
 #endif	//WITH_EDITOR
 
-	FGuid GetLandscapeGuid() const { return LandscapeGuid; }
+	virtual FGuid GetLandscapeGuid() const override { return LandscapeGuid; }
 	void SetLandscapeGuid(const FGuid& Guid) { LandscapeGuid = Guid; }
 	virtual ALandscape* GetLandscapeActor() PURE_VIRTUAL(GetLandscapeActor, return nullptr;)
 	virtual const ALandscape* GetLandscapeActor() const PURE_VIRTUAL(GetLandscapeActor, return nullptr;)
@@ -953,10 +953,10 @@ public:
 	virtual void PostLoad() override;
 
 	LANDSCAPE_API ULandscapeInfo* CreateLandscapeInfo(bool bMapCheck = false);
-	LANDSCAPE_API ULandscapeInfo* GetLandscapeInfo() const;
+	virtual LANDSCAPE_API ULandscapeInfo* GetLandscapeInfo() const override;
 
 	/** Get the LandcapeActor-to-world transform with respect to landscape section offset*/
-	LANDSCAPE_API FTransform LandscapeActorToWorld() const;
+	virtual LANDSCAPE_API FTransform LandscapeActorToWorld() const override;
 
 	/**
 	* Output a landscape heightmap to a render target
@@ -970,11 +970,19 @@ public:
 
 	/** Get landscape position in section space */
 	LANDSCAPE_API FIntPoint GetSectionBaseOffset() const;
+
+	// ILandscapeSplineInterface
+	LANDSCAPE_API virtual ULandscapeSplinesComponent* GetSplinesComponent() const override { return SplineComponent; }
+
 #if WITH_EDITOR
+	LANDSCAPE_API void SetSplinesComponent(ULandscapeSplinesComponent* InSplineComponent) { check(!SplineComponent); SplineComponent = InSplineComponent; }
+
+	LANDSCAPE_API virtual bool SupportsForeignSplineMesh() const override { return true; }
+
 	LANDSCAPE_API int32 GetOutdatedGrassMapCount() const;
 	LANDSCAPE_API void BuildGrassMaps(struct FScopedSlowTask* InSlowTask = nullptr);
-	LANDSCAPE_API void CreateSplineComponent();
-	LANDSCAPE_API void CreateSplineComponent(const FVector& Scale3D);
+	LANDSCAPE_API virtual void CreateSplineComponent() override;
+	LANDSCAPE_API virtual void CreateSplineComponent(const FVector& Scale3D) override;
 
 	virtual bool CanEditChange(const FProperty* InProperty) const override;
 	virtual void PostEditChangeProperty(FPropertyChangedEvent& PropertyChangedEvent) override;
diff --git a/Engine/Source/Runtime/Landscape/Classes/LandscapeSplineActor.h b/Engine/Source/Runtime/Landscape/Classes/LandscapeSplineActor.h
new file mode 100644
--- /dev/null
+++ b/Engine/Source/Runtime/Landscape/Classes/LandscapeSplineActor.h
@@ -0,0 +1,43 @@
+// Copyright Epic Games, Inc. All Rights Reserved.
+
+#pragma once
+
+#include "CoreMinimal.h"
+#include "UObject/ObjectMacros.h"
+#include "GameFramework/Actor.h"
+#include "ILandscapeSplineInterface.h"
+#include "LandscapeSplineActor.generated.h"
+
+UCLASS(MinimalAPI, NotBlueprintable, NotPlaceable, hidecategories = (Display, Attachment, Physics, Debug, Lighting, Input))
+class ALandscapeSplineActor : public AActor, public ILandscapeSplineInterface
+{
+	GENERATED_UCLASS_BODY()
+
+protected:
+	/** Guid for LandscapeInfo **/
+	UPROPERTY()
+	FGuid LandscapeGuid;
+	
+public:
+	virtual FGuid GetLandscapeGuid() const override { return LandscapeGuid; }
+	virtual ULandscapeSplinesComponent* GetSplinesComponent() const override;
+	virtual FTransform LandscapeActorToWorld() const override;
+	virtual ULandscapeInfo* GetLandscapeInfo() const override;
+		
+#if WITH_EDITOR
+	virtual EActorGridPlacement GetDefaultGridPlacement() const override { return EActorGridPlacement::Bounds; }
+
+	void GetSharedProperties(ULandscapeInfo* InLandscapeInfo);
+
+	virtual void Destroyed() override;
+	virtual void PostRegisterAllComponents() override;
+	virtual void UnregisterAllComponents(bool bForReregister = false) override;
+	virtual void PostEditMove(bool bFinished) override;
+	
+	virtual bool SupportsForeignSplineMesh() const override { return false; }
+
+	// Interface existes for backward compatibility. Should already be created since its this actor's root component.
+	virtual void CreateSplineComponent() override { check(false); }
+	virtual void CreateSplineComponent(const FVector& Scale3D) override { check(false); }
+#endif
+};
\ No newline at end of file
diff --git a/Engine/Source/Runtime/Landscape/Classes/LandscapeSplineControlPoint.h b/Engine/Source/Runtime/Landscape/Classes/LandscapeSplineControlPoint.h
--- a/Engine/Source/Runtime/Landscape/Classes/LandscapeSplineControlPoint.h
+++ b/Engine/Source/Runtime/Landscape/Classes/LandscapeSplineControlPoint.h
@@ -238,6 +238,8 @@ public:
 	const TArray<FLandscapeSplineInterpPoint>& GetPoints() const { return Points; }
 
 #if WITH_EDITOR
+	bool SupportsForeignSplineMesh() const;
+
 	// Get the name of the best connection point (socket) to use for a particular destination
 	virtual FName GetBestConnectionTo(FVector Destination) const;
 
diff --git a/Engine/Source/Runtime/Landscape/Classes/LandscapeSplineSegment.h b/Engine/Source/Runtime/Landscape/Classes/LandscapeSplineSegment.h
--- a/Engine/Source/Runtime/Landscape/Classes/LandscapeSplineSegment.h
+++ b/Engine/Source/Runtime/Landscape/Classes/LandscapeSplineSegment.h
@@ -340,6 +340,8 @@ public:
 	const TArray<FLandscapeSplineInterpPoint>& GetPoints() const { return Points; }
 
 #if WITH_EDITOR
+	bool SupportsForeignSplineMesh() const;
+
 	bool IsSplineSelected() const { return bSelected; }
 	virtual void SetSplineSelected(bool bInSelected);
 
diff --git a/Engine/Source/Runtime/Landscape/Classes/LandscapeSplinesComponent.h b/Engine/Source/Runtime/Landscape/Classes/LandscapeSplinesComponent.h
--- a/Engine/Source/Runtime/Landscape/Classes/LandscapeSplinesComponent.h
+++ b/Engine/Source/Runtime/Landscape/Classes/LandscapeSplinesComponent.h
@@ -18,6 +18,7 @@ class UMeshComponent;
 class USplineMeshComponent;
 class UStaticMesh;
 class UTexture2D;
+class ILandscapeSplineInterface;
 
 // structs for ForeignWorldSplineDataMap
 // these are editor-only, but we don't have the concept of an editor-only USTRUCT
@@ -150,12 +151,17 @@ public:
 	UFUNCTION(BlueprintCallable, Category = LandscapeSplines)
 	TArray<USplineMeshComponent*> GetSplineMeshComponents();
 
+	LANDSCAPE_API ILandscapeSplineInterface* GetSplineOwner();
+
 	void CheckSplinesValid();
 	bool ModifySplines(bool bAlwaysMarkDirty = true);
 
 #if WITH_EDITOR
+	void RequestSplineLayerUpdate();
+
 	bool HasAnyControlPointsOrSegments() const { return ControlPoints.Num() > 0 || Segments.Num() > 0; }
 
+	LANDSCAPE_API void SetDefaultEditorSplineMesh();
 	virtual void ShowSplineEditorMesh(bool bShow);
 
 	// Rebuilds all spline points and meshes for all spline control points and segments in this splines component
@@ -197,6 +203,9 @@ public:
 	bool IsUsingEditorMesh(const USplineMeshComponent* SplineMeshComponent) const;
 
 	bool IsUsingLayerInfo(const ULandscapeLayerInfoObject* LayerInfo) const;
+
+	// Iterates throug a copy of the ControlPoints list.
+	LANDSCAPE_API void ForEachControlPoint(TFunctionRef<void(ULandscapeSplineControlPoint*)> Func);
 #endif
 
 	//~ Begin UObject Interface
diff --git a/Engine/Source/Runtime/Landscape/Private/Landscape.cpp b/Engine/Source/Runtime/Landscape/Private/Landscape.cpp
--- a/Engine/Source/Runtime/Landscape/Private/Landscape.cpp
+++ b/Engine/Source/Runtime/Landscape/Private/Landscape.cpp
@@ -44,6 +44,8 @@ Landscape.cpp: Terrain rendering
 #include "Materials/MaterialInstance.h"
 #include "Materials/MaterialInstanceDynamic.h"
 #include "ProfilingDebugging/CookStats.h"
+#include "ILandscapeSplineInterface.h"
+#include "LandscapeSplineActor.h"
 #include "LandscapeSplinesComponent.h"
 #include "EngineGlobals.h"
 #include "Engine/Engine.h"
@@ -1186,41 +1188,14 @@ ALandscape* ALandscapeStreamingProxy::GetLandscapeActor()
 
 ULandscapeInfo* ALandscapeProxy::CreateLandscapeInfo(bool bMapCheck)
 {
-	ULandscapeInfo* LandscapeInfo = nullptr;
-
-	check(LandscapeGuid.IsValid());
-	UWorld* OwningWorld = GetWorld();
-	check(OwningWorld);
-	
-	auto& LandscapeInfoMap = ULandscapeInfoMap::GetLandscapeInfoMap(OwningWorld);
-	LandscapeInfo = LandscapeInfoMap.Map.FindRef(LandscapeGuid);
-
-	if (!LandscapeInfo)
-	{
-		check(!HasAnyFlags(RF_BeginDestroyed));
-		LandscapeInfo = NewObject<ULandscapeInfo>(GetTransientPackage(), NAME_None, RF_Transactional | RF_Transient);
-		LandscapeInfoMap.Modify(false);
-		LandscapeInfoMap.Map.Add(LandscapeGuid, LandscapeInfo);
-	}
-	check(LandscapeInfo);
+	ULandscapeInfo* LandscapeInfo = ULandscapeInfo::FindOrCreate(GetWorld(), LandscapeGuid);
 	LandscapeInfo->RegisterActor(this, bMapCheck);
-
 	return LandscapeInfo;
 }
 
 ULandscapeInfo* ALandscapeProxy::GetLandscapeInfo() const
 {
-	ULandscapeInfo* LandscapeInfo = nullptr;
-
-	check(LandscapeGuid.IsValid());
-	UWorld* OwningWorld = GetWorld();
-
-	if (OwningWorld != nullptr)
-	{
-		auto& LandscapeInfoMap = ULandscapeInfoMap::GetLandscapeInfoMap(OwningWorld);
-		LandscapeInfo = LandscapeInfoMap.Map.FindRef(LandscapeGuid);
-	}
-	return LandscapeInfo;
+	return ULandscapeInfo::Find(GetWorld(), LandscapeGuid);
 }
 
 FTransform ALandscapeProxy::LandscapeActorToWorld() const
@@ -3084,6 +3059,39 @@ void ULandscapeInfo::RecreateLandscapeInfo(UWorld* InWorld, bool bMapCheck)
 
 #endif
 
+ULandscapeInfo* ULandscapeInfo::Find(UWorld* InWorld, const FGuid& LandscapeGuid)
+{
+	ULandscapeInfo* LandscapeInfo = nullptr;
+
+	check(LandscapeGuid.IsValid());
+	if (InWorld != nullptr)
+	{
+		auto& LandscapeInfoMap = ULandscapeInfoMap::GetLandscapeInfoMap(InWorld);
+		LandscapeInfo = LandscapeInfoMap.Map.FindRef(LandscapeGuid);
+	}
+	return LandscapeInfo;
+}
+
+ULandscapeInfo* ULandscapeInfo::FindOrCreate(UWorld* InWorld, const FGuid& LandscapeGuid)
+{
+	ULandscapeInfo* LandscapeInfo = nullptr;
+
+	check(LandscapeGuid.IsValid());
+	check(InWorld);
+
+	auto& LandscapeInfoMap = ULandscapeInfoMap::GetLandscapeInfoMap(InWorld);
+	LandscapeInfo = LandscapeInfoMap.Map.FindRef(LandscapeGuid);
+
+	if (!LandscapeInfo)
+	{
+		LandscapeInfo = NewObject<ULandscapeInfo>(GetTransientPackage(), NAME_None, RF_Transactional | RF_Transient);
+		LandscapeInfoMap.Modify(false);
+		LandscapeInfoMap.Map.Add(LandscapeGuid, LandscapeInfo);
+	}
+	check(LandscapeInfo);
+	return LandscapeInfo;
+}
+
 void ULandscapeInfo::ForAllLandscapeProxies(TFunctionRef<void(ALandscapeProxy*)> Fn) const
 {
 	ALandscape* Landscape = LandscapeActor.Get();
@@ -3175,6 +3183,8 @@ void ULandscapeInfo::RegisterActor(ALandscapeProxy* Proxy, bool bMapCheck)
 
 		UpdateLayerInfoMap(Proxy);
 		UpdateAllAddCollisions();
+
+		RegisterSplineActor(Proxy);
 	}
 #endif
 	
@@ -3221,6 +3231,8 @@ void ULandscapeInfo::UnregisterActor(ALandscapeProxy* Proxy)
 			Proxies.Remove(StreamingProxy);
 			StreamingProxy->LandscapeActor = nullptr;
 		}
+
+		UnregisterSplineActor(Proxy);
 	}
 #endif
 
@@ -3253,6 +3265,68 @@ void ULandscapeInfo::UnregisterActor(ALandscapeProxy* Proxy)
 #endif
 }
 
+#if WITH_EDITOR
+ALandscapeSplineActor* ULandscapeInfo::CreateSplineActor(const FVector& Location)
+{
+	check(LandscapeActor.Get());
+	UWorld* World = LandscapeActor->GetWorld();
+	check(World);
+	FActorSpawnParameters SpawnParams;
+	SpawnParams.OverrideLevel = World->PersistentLevel;
+	SpawnParams.bNoFail = true;
+	SpawnParams.ObjectFlags |= RF_Transactional;
+	ALandscapeSplineActor* SplineActor = World->SpawnActor<ALandscapeSplineActor>(Location, FRotator::ZeroRotator, SpawnParams);
+	SplineActor->GetSharedProperties(this);
+	SplineActor->GetSplinesComponent()->ShowSplineEditorMesh(true);
+	RegisterSplineActor(SplineActor);
+	return SplineActor;
+}
+
+void ULandscapeInfo::ForAllSplineActors(TFunctionRef<void(TScriptInterface<ILandscapeSplineInterface>)> Fn) const
+{
+	for (const TScriptInterface<ILandscapeSplineInterface>& SplineActor : SplineActors)
+	{
+		Fn(SplineActor);
+	}
+}
+
+TArray<TScriptInterface<ILandscapeSplineInterface>> ULandscapeInfo::GetSplineActors() const
+{
+	TArray<TScriptInterface<ILandscapeSplineInterface>> CopySplineActors(SplineActors);
+	return MoveTemp(CopySplineActors);
+}
+
+void ULandscapeInfo::RegisterSplineActor(TScriptInterface<ILandscapeSplineInterface> SplineActor)
+{
+	Modify();
+	SplineActors.AddUnique(SplineActor);
+
+	if (SplineActor->GetSplinesComponent())
+	{
+		RequestSplineLayerUpdate();
+	}
+}
+
+void ULandscapeInfo::UnregisterSplineActor(TScriptInterface<ILandscapeSplineInterface> SplineActor)
+{
+	Modify();
+	SplineActors.Remove(SplineActor);
+
+	if (SplineActor->GetSplinesComponent())
+	{
+		RequestSplineLayerUpdate();
+	}
+}
+
+void ULandscapeInfo::RequestSplineLayerUpdate()
+{
+	if (LandscapeActor)
+	{
+		LandscapeActor->RequestSplineLayerUpdate();
+	}
+}
+#endif
+
 void ULandscapeInfo::RegisterCollisionComponent(ULandscapeHeightfieldCollisionComponent* Component)
 {
 	if (Component == nullptr || !Component->IsRegistered())
diff --git a/Engine/Source/Runtime/Landscape/Private/LandscapeConfigHelper.cpp b/Engine/Source/Runtime/Landscape/Private/LandscapeConfigHelper.cpp
--- a/Engine/Source/Runtime/Landscape/Private/LandscapeConfigHelper.cpp
+++ b/Engine/Source/Runtime/Landscape/Private/LandscapeConfigHelper.cpp
@@ -215,10 +215,7 @@ bool FLandscapeConfigHelper::ChangeGridSize(ULandscapeInfo* InLandscapeInfo, uin
 
 		return true;
 	}, GridSize);
-
-	// @todo_ow: This is probably not the final solution
-	MoveSplinesToLandscape(InLandscapeInfo, InLandscapeInfo->LandscapeActor.Get(), 1.0f);
-
+		
 	// Only delete Proxies that where not reused
 	for (ALandscapeProxy* ProxyToDelete : ProxiesToDelete)
 	{
@@ -344,19 +341,19 @@ void FLandscapeConfigHelper::CopyRegionToComponent(ULandscapeInfo* InLandscapeIn
 
 void FLandscapeConfigHelper::MoveSplinesToLandscape(ULandscapeInfo* InLandscapeInfo, ALandscapeProxy* InLandscape, float InScaleFactor)
 {
-	if (!InLandscape->SplineComponent)
+	if (!InLandscape->GetSplinesComponent())
 	{
 		InLandscape->CreateSplineComponent();
 	}
 
 	InLandscapeInfo->ForAllLandscapeProxies([InLandscapeInfo, InLandscape](ALandscapeProxy* LandscapeProxy)
 	{
-		if (InLandscape == LandscapeProxy || LandscapeProxy->SplineComponent == nullptr)
+		if (InLandscape == LandscapeProxy || !LandscapeProxy->GetSplinesComponent())
 		{
 			return;
 		}
 
-		InLandscapeInfo->MoveSplinesToProxy(LandscapeProxy->SplineComponent, InLandscape);
+		InLandscapeInfo->MoveSplinesToProxy(LandscapeProxy->GetSplinesComponent(), InLandscape);
 	});
 }
 
diff --git a/Engine/Source/Runtime/Landscape/Private/LandscapeEdit.cpp b/Engine/Source/Runtime/Landscape/Private/LandscapeEdit.cpp
--- a/Engine/Source/Runtime/Landscape/Private/LandscapeEdit.cpp
+++ b/Engine/Source/Runtime/Landscape/Private/LandscapeEdit.cpp
@@ -2196,11 +2196,6 @@ bool ULandscapeInfo::AreAllComponentsRegistered() const
 
 		if (LandscapeProxy->GetLandscapeGuid() == LandscapeGuid)
 		{
-			if (LandscapeProxy->SplineComponent && !LandscapeProxy->SplineComponent->IsRegistered())
-			{
-				return false;
-			}
-
 			for (ULandscapeComponent* LandscapeComponent : LandscapeProxy->LandscapeComponents)
 			{
 				if (LandscapeComponent && !LandscapeComponent->IsRegistered())
@@ -2210,6 +2205,23 @@ bool ULandscapeInfo::AreAllComponentsRegistered() const
 			}
 		}
 	}
+		
+	for (TScriptInterface<ILandscapeSplineInterface> SplineOwner : SplineActors)
+	{
+		if (!SplineOwner.GetObject() || SplineOwner.GetObject()->IsPendingKill())
+		{
+			continue;
+		}
+
+		if (SplineOwner->GetLandscapeGuid() == LandscapeGuid)
+		{
+			if (SplineOwner->GetSplinesComponent() && !SplineOwner->GetSplinesComponent()->IsRegistered())
+			{
+				return false;
+			}
+		}
+	}
+
 
 	return true;
 }
diff --git a/Engine/Source/Runtime/Landscape/Private/LandscapeEditLayers.cpp b/Engine/Source/Runtime/Landscape/Private/LandscapeEditLayers.cpp
--- a/Engine/Source/Runtime/Landscape/Private/LandscapeEditLayers.cpp
+++ b/Engine/Source/Runtime/Landscape/Private/LandscapeEditLayers.cpp
@@ -6459,9 +6459,9 @@ void ALandscape::OnLayerInfoSplineFalloffModulationChanged(ULandscapeLayerInfoOb
 	}
 
 	bool bUsedForSplines = false;
-	LandscapeInfo->ForAllLandscapeProxies([&](ALandscapeProxy* Proxy)
+	LandscapeInfo->ForAllSplineActors([&](TScriptInterface<ILandscapeSplineInterface> SplineOwner)
 	{
-		bUsedForSplines |= (Proxy->SplineComponent && Proxy->SplineComponent->IsUsingLayerInfo(InLayerInfo));
+		bUsedForSplines |= (SplineOwner->GetSplinesComponent() && SplineOwner->GetSplinesComponent()->IsUsingLayerInfo(InLayerInfo));
 	});
 
 	if (bUsedForSplines)
diff --git a/Engine/Source/Runtime/Landscape/Private/LandscapeSplineActor.cpp b/Engine/Source/Runtime/Landscape/Private/LandscapeSplineActor.cpp
new file mode 100644
--- /dev/null
+++ b/Engine/Source/Runtime/Landscape/Private/LandscapeSplineActor.cpp
@@ -0,0 +1,89 @@
+// Copyright Epic Games, Inc. All Rights Reserved.
+
+#include "LandscapeSplineActor.h"
+#include "LandscapeInfo.h"
+#include "LandscapeInfoMap.h"
+#include "Landscape.h"
+#include "LandscapeSplinesComponent.h"
+
+ALandscapeSplineActor::ALandscapeSplineActor(const FObjectInitializer& ObjectInitializer)
+	: Super(ObjectInitializer)
+{
+	ULandscapeSplinesComponent* SplineComponent = CreateDefaultSubobject<ULandscapeSplinesComponent>(TEXT("RootComponent0"));
+	
+	RootComponent = SplineComponent;	
+	RootComponent->Mobility = EComponentMobility::Static;
+}
+
+ULandscapeSplinesComponent* ALandscapeSplineActor::GetSplinesComponent() const
+{
+	return Cast<ULandscapeSplinesComponent>(RootComponent);
+}
+
+FTransform ALandscapeSplineActor::LandscapeActorToWorld() const
+{
+	return GetLandscapeInfo()->LandscapeActor->LandscapeActorToWorld();
+}
+
+ULandscapeInfo* ALandscapeSplineActor::GetLandscapeInfo() const
+{
+	return ULandscapeInfo::Find(GetWorld(), LandscapeGuid);
+}
+
+#if WITH_EDITOR
+void ALandscapeSplineActor::GetSharedProperties(ULandscapeInfo* InLandscapeInfo)
+{
+	Modify();
+	LandscapeGuid = InLandscapeInfo->LandscapeGuid;
+}
+
+void ALandscapeSplineActor::Destroyed()
+{
+	Super::Destroyed();
+
+	UWorld* World = GetWorld();
+
+	if (GIsEditor && !World->IsGameWorld())
+	{
+		// Modify Splines Objects to support Undo/Redo
+		GetSplinesComponent()->ModifySplines();
+	}
+}
+
+void ALandscapeSplineActor::PostRegisterAllComponents()
+{
+	Super::PostRegisterAllComponents();
+	if (!IsPendingKillPending())
+	{
+		if (LandscapeGuid.IsValid())
+		{
+			ULandscapeInfo* LandscapeInfo = ULandscapeInfo::FindOrCreate(GetWorld(), LandscapeGuid);
+			LandscapeInfo->RegisterSplineActor(this);
+		}
+	}
+}
+
+void ALandscapeSplineActor::UnregisterAllComponents(bool bForReregister)
+{
+	if (GetWorld() && !GetWorld()->IsPendingKillOrUnreachable() && LandscapeGuid.IsValid())
+	{
+		if (ULandscapeInfo* LandscapeInfo = GetLandscapeInfo())
+		{
+			LandscapeInfo->UnregisterSplineActor(this);
+		}
+	}
+
+	Super::UnregisterAllComponents(bForReregister);
+}
+
+void ALandscapeSplineActor::PostEditMove(bool bFinished)
+{
+	Super::PostEditMove(bFinished);
+
+	if (bFinished)
+	{
+		GetLandscapeInfo()->RequestSplineLayerUpdate();
+	}
+}
+
+#endif
\ No newline at end of file
diff --git a/Engine/Source/Runtime/Landscape/Private/LandscapeSplineRaster.cpp b/Engine/Source/Runtime/Landscape/Private/LandscapeSplineRaster.cpp
--- a/Engine/Source/Runtime/Landscape/Private/LandscapeSplineRaster.cpp
+++ b/Engine/Source/Runtime/Landscape/Private/LandscapeSplineRaster.cpp
@@ -577,29 +577,37 @@ bool ULandscapeInfo::ApplySplines(bool bOnlySelected, TSet<ULandscapeComponent*>
 		return SharedPtr;
 	};
 
-	ForAllLandscapeProxies([&](ALandscapeProxy* Proxy)
+	ForAllSplineActors([&](TScriptInterface<ILandscapeSplineInterface> SplineOwner)
 	{
-		bResult |= ApplySplinesInternal(bOnlySelected, Proxy, OutModifiedComponents, bMarkPackageDirty, LandscapeMinX, LandscapeMinY, LandscapeMaxX, LandscapeMaxY, GetOrCreateModulate);
+		bResult |= ApplySplinesInternal(bOnlySelected, SplineOwner, OutModifiedComponents, bMarkPackageDirty, LandscapeMinX, LandscapeMinY, LandscapeMaxX, LandscapeMaxY, GetOrCreateModulate);
 	});
 
 	return bResult;
 }
 
-bool ULandscapeInfo::ApplySplinesInternal(bool bOnlySelected, ALandscapeProxy* Proxy, TSet<ULandscapeComponent*>* OutModifiedComponents, bool bMarkPackageDirty, int32 LandscapeMinX, int32 LandscapeMinY, int32 LandscapeMaxX, int32 LandscapeMaxY, TFunctionRef<TSharedPtr<FModulateAlpha>(ULandscapeLayerInfoObject*)> GetOrCreateModulate)
+bool ULandscapeInfo::ApplySplinesInternal(bool bOnlySelected, TScriptInterface<ILandscapeSplineInterface> SplineOwner, TSet<ULandscapeComponent*>* OutModifiedComponents, bool bMarkPackageDirty, int32 LandscapeMinX, int32 LandscapeMinY, int32 LandscapeMaxX, int32 LandscapeMaxY, TFunctionRef<TSharedPtr<FModulateAlpha>(ULandscapeLayerInfoObject*)> GetOrCreateModulate)
 {
-	if (!Proxy || !Proxy->SplineComponent || !Proxy->SplineComponent->IsRegistered() || Proxy->SplineComponent->ControlPoints.Num() == 0 || Proxy->SplineComponent->Segments.Num() == 0)
+	if (!SplineOwner)
 	{
 		return false;
 	}
 
-	const FTransform SplineToLandscape = Proxy->SplineComponent->GetComponentTransform().GetRelativeTransform(Proxy->LandscapeActorToWorld());
+	ULandscapeSplinesComponent* SplineComponent = SplineOwner->GetSplinesComponent();
+	
+	if (!SplineComponent || !SplineComponent->IsRegistered() || SplineComponent->ControlPoints.Num() == 0 || SplineComponent->Segments.Num() == 0)
+	{
+		return false;
+	}
+
+
+	const FTransform SplineToLandscape = SplineComponent->GetComponentTransform().GetRelativeTransform(SplineOwner->LandscapeActorToWorld());
 
 	FLandscapeEditDataInterface LandscapeEdit(this);
 	FLandscapeDoNotDirtyScope DoNotDirtyScope(LandscapeEdit, !bMarkPackageDirty);
 	TSet<ULandscapeComponent*> ModifiedComponents;
 	
 
-	for (const ULandscapeSplineControlPoint* ControlPoint : Proxy->SplineComponent->ControlPoints)
+	for (const ULandscapeSplineControlPoint* ControlPoint : SplineComponent->ControlPoints)
 	{
 		if (bOnlySelected && !ControlPoint->IsSplineSelected())
 		{
@@ -675,7 +683,7 @@ bool ULandscapeInfo::ApplySplinesInternal(bool bOnlySelected, ALandscapeProxy* P
 		}
 	}
 
-	for (const ULandscapeSplineSegment* Segment : Proxy->SplineComponent->Segments)
+	for (const ULandscapeSplineSegment* Segment : SplineComponent->Segments)
 	{
 		if (bOnlySelected && !Segment->IsSplineSelected())
 		{
@@ -750,7 +758,7 @@ bool ULandscapeInfo::ApplySplinesInternal(bool bOnlySelected, ALandscapeProxy* P
 
 	LandscapeEdit.Flush();
 		
-	if (!Proxy->HasLayersContent())
+	if (!CanHaveLayersContent())
 	{
 		ALandscapeProxy::InvalidateGeneratedComponentData(ModifiedComponents);
 	}
diff --git a/Engine/Source/Runtime/Landscape/Private/LandscapeSplines.cpp b/Engine/Source/Runtime/Landscape/Private/LandscapeSplines.cpp
--- a/Engine/Source/Runtime/Landscape/Private/LandscapeSplines.cpp
+++ b/Engine/Source/Runtime/Landscape/Private/LandscapeSplines.cpp
@@ -28,6 +28,7 @@
 #include "LandscapeSplineSegment.h"
 #include "LandscapeSplineRaster.h"
 #include "LandscapeSplineControlPoint.h"
+#include "ILandscapeSplineInterface.h"
 #include "ControlPointMeshComponent.h"
 #include "Engine/CollisionProfile.h"
 #include "Engine/Engine.h"
@@ -86,11 +87,11 @@ struct FLandscapeFixSplines
 		{
 			if (Pair.Value)
 			{
-				Pair.Value->ForAllLandscapeProxies([](ALandscapeProxy* Proxy)
+				Pair.Value->ForAllSplineActors([](TScriptInterface<ILandscapeSplineInterface> SplineOwner)
 				{
-					if (Proxy && Proxy->SplineComponent)
+					if (SplineOwner && SplineOwner->GetSplinesComponent())
 					{
-						Proxy->SplineComponent->RebuildAllSplines();
+						SplineOwner->GetSplinesComponent()->RebuildAllSplines();
 					}
 				});
 
@@ -396,6 +397,14 @@ public:
 };
 #endif
 
+//////////////////////////////////////////////////////////////////////////
+// LANDSCAPE SPLINE INTERFACE
+ULandscapeSplineInterface::ULandscapeSplineInterface(const FObjectInitializer& ObjectInitializer)
+	: Super(ObjectInitializer)
+{
+
+}
+
 //////////////////////////////////////////////////////////////////////////
 // SPLINE COMPONENT
 
@@ -454,6 +463,11 @@ TArray<USplineMeshComponent*> ULandscapeSplinesComponent::GetSplineMeshComponent
 	return SplineMeshComponents;
 }
 
+ILandscapeSplineInterface* ULandscapeSplinesComponent::GetSplineOwner()
+{
+	return Cast<ILandscapeSplineInterface>(GetOwner());
+}
+
 void ULandscapeSplinesComponent::CheckSplinesValid()
 {
 #if DO_CHECK
@@ -700,6 +714,16 @@ bool ULandscapeSplinesComponent::IsUsingEditorMesh(const USplineMeshComponent* S
 	return SplineMeshComponent->GetStaticMesh() == SplineEditorMesh && SplineMeshComponent->bHiddenInGame;
 }
 
+void ULandscapeSplinesComponent::ForEachControlPoint(TFunctionRef<void(ULandscapeSplineControlPoint*)> Func)
+{
+	// Copy in case iteration modifies the list
+	TArray<ULandscapeSplineControlPoint*> CopyControlPoints(ControlPoints);
+	for (ULandscapeSplineControlPoint* ControlPoint : CopyControlPoints)
+	{
+		Func(ControlPoint);
+	}
+}
+
 bool ULandscapeSplinesComponent::IsUsingLayerInfo(const ULandscapeLayerInfoObject* LayerInfo) const
 {
 	for (ULandscapeSplineControlPoint* ControlPoint : ControlPoints)
@@ -935,6 +959,20 @@ void ULandscapeSplinesComponent::RebuildAllSplines(bool bUpdateCollision)
 	}
 }
 
+void ULandscapeSplinesComponent::RequestSplineLayerUpdate()
+{
+	if (!IsPendingKill())
+	{
+		ILandscapeSplineInterface* SplineOwner = GetSplineOwner();
+		SplineOwner->GetLandscapeInfo()->RequestSplineLayerUpdate();
+	}
+}
+
+void ULandscapeSplinesComponent::SetDefaultEditorSplineMesh()
+{
+	SplineEditorMesh = LoadObject<UStaticMesh>(nullptr, TEXT("/Engine/EditorLandscapeResources/SplineEditorMesh"));
+}
+
 void ULandscapeSplinesComponent::ShowSplineEditorMesh(bool bShow)
 {
 	bShowSplineEditorMesh = bShow;
@@ -994,14 +1032,14 @@ ULandscapeSplinesComponent* ULandscapeSplinesComponent::GetStreamingSplinesCompo
 		if (LandscapeComponent)
 		{
 			ALandscapeProxy* ComponentLandscapeProxy = LandscapeComponent->GetLandscapeProxy();
-			if (!ComponentLandscapeProxy->SplineComponent && bCreate)
+			if (!ComponentLandscapeProxy->GetSplinesComponent() && bCreate)
 			{
 				ComponentLandscapeProxy->CreateSplineComponent(GetRelativeScale3D());
 			}
 
-			if (ComponentLandscapeProxy->SplineComponent)
+			if (ULandscapeSplinesComponent* SplineComponent = ComponentLandscapeProxy->GetSplinesComponent())
 			{
-				return ComponentLandscapeProxy->SplineComponent;
+				return SplineComponent;
 			}
 		}
 	}
@@ -1020,11 +1058,11 @@ ULandscapeSplinesComponent* ULandscapeSplinesComponent::GetStreamingSplinesCompo
 		ALandscapeProxy* Proxy = LandscapeInfo->GetLandscapeProxyForLevel(Level);
 		if (Proxy)
 		{
-			if (!Proxy->SplineComponent && bCreate)
+			if (!Proxy->GetSplinesComponent() && bCreate)
 			{
 				Proxy->CreateSplineComponent(GetRelativeScale3D());
 			}
-			return Proxy->SplineComponent;
+			return Proxy->GetSplinesComponent();
 		}
 	}
 
@@ -1044,11 +1082,11 @@ TArray<ULandscapeSplinesComponent*> ULandscapeSplinesComponent::GetAllStreamingS
 		if (LandscapeInfo)
 		{
 			TArray<ULandscapeSplinesComponent*> SplinesComponents;
-			LandscapeInfo->ForAllLandscapeProxies([&SplinesComponents](ALandscapeProxy* Proxy)
+			LandscapeInfo->ForAllSplineActors([&SplinesComponents](TScriptInterface<ILandscapeSplineInterface> SplineOwner)
 			{
-				if (Proxy->SplineComponent)
+				if (ULandscapeSplinesComponent* SplineComponent = SplineOwner->GetSplinesComponent())
 				{
-					SplinesComponents.Add(Proxy->SplineComponent);
+					SplinesComponents.Add(SplineComponent);
 				}
 			});
 			return SplinesComponents;
@@ -1649,6 +1687,11 @@ FLandscapeSplineSegmentConnection& FLandscapeSplineConnection::GetFarConnection(
 }
 
 #if WITH_EDITOR
+bool ULandscapeSplineControlPoint::SupportsForeignSplineMesh() const
+{
+	return bPlaceSplineMeshesInStreamingLevels && GetOuterULandscapeSplinesComponent()->GetSplineOwner()->SupportsForeignSplineMesh();
+}
+
 FName ULandscapeSplineControlPoint::GetBestConnectionTo(FVector Destination) const
 {
 	FName BestSocket = NAME_None;
@@ -1850,7 +1893,7 @@ void ULandscapeSplineControlPoint::UpdateSplinePoints(bool bUpdateCollision, boo
 	if (Mesh != nullptr)
 	{
 		// Attempt to place mesh components into the appropriate landscape streaming levels based on the components under the spline
-		if (bPlaceSplineMeshesInStreamingLevels)
+		if (SupportsForeignSplineMesh())
 		{
 			MeshComponentOuterSplines = OuterSplines->GetStreamingSplinesComponentByLocation(Location);
 
@@ -1881,11 +1924,10 @@ void ULandscapeSplineControlPoint::UpdateSplinePoints(bool bUpdateCollision, boo
 		}
 		else if(bUpdateMeshLevel)// Update Foreign/Local if necessary
 		{
-			ALandscapeProxy* CurrentMeshComponentOuterActor = MeshComponent->GetTypedOuter<ALandscapeProxy>();
-			ULandscapeSplinesComponent* CurrentLandscapeSplineComponent = CurrentMeshComponentOuterActor->SplineComponent;
+			AActor* CurrentMeshComponentOuterActor = MeshComponent->GetTypedOuter<AActor>();
+			AActor* MeshComponentOuterActor = MeshComponentOuterSplines->GetOwner();
+			ILandscapeSplineInterface* SplineOwner = Cast<ILandscapeSplineInterface>(MeshComponentOuterActor);
 
-			ALandscapeProxy* MeshComponentOuterActor = Cast<ALandscapeProxy>(MeshComponentOuterSplines->GetOwner());
-			
 			// Needs updating
 			if (MeshComponentOuterActor != CurrentMeshComponentOuterActor)
 			{
@@ -1897,7 +1939,7 @@ void ULandscapeSplineControlPoint::UpdateSplinePoints(bool bUpdateCollision, boo
 				MeshComponent->DetachFromComponent(FDetachmentTransformRules::KeepWorldTransform);
 				MeshComponent->InvalidateLightingCache();
 				MeshComponent->Rename(nullptr, MeshComponentOuterActor);
-				MeshComponent->AttachToComponent(MeshComponentOuterActor->SplineComponent, FAttachmentTransformRules::KeepWorldTransform);
+				MeshComponent->AttachToComponent(SplineOwner->GetSplinesComponent(), FAttachmentTransformRules::KeepWorldTransform);
 				bComponentNeedsRegistering = true;
 				bUpdateLocalForeign = true;
 			}
@@ -2218,13 +2260,7 @@ void ULandscapeSplineControlPoint::PostEditUndo()
 
 	ULandscapeSplinesComponent* SplineComponent = GetOuterULandscapeSplinesComponent();
 	SplineComponent->MarkRenderStateDirty();
-
-	ALandscapeProxy* OuterLandscape = Cast<ALandscapeProxy>(SplineComponent->GetOwner());
-	ALandscape* Landscape = OuterLandscape ? OuterLandscape->GetLandscapeActor() : nullptr;
-	if (Landscape)
-	{
-		Landscape->RequestSplineLayerUpdate();
-	}
+	SplineComponent->RequestSplineLayerUpdate();
 }
 
 void ULandscapeSplineControlPoint::PostDuplicate(bool bDuplicateForPIE)
@@ -2278,12 +2314,7 @@ void ULandscapeSplineControlPoint::PostEditChangeProperty(FPropertyChangedEvent&
 
 	if (PropertyChangedEvent.ChangeType == EPropertyChangeType::ValueSet)
 	{
-		ALandscapeProxy* OuterLandscape = Cast<ALandscapeProxy>(GetOuterULandscapeSplinesComponent()->GetOwner());
-		ALandscape* Landscape = OuterLandscape ? OuterLandscape->GetLandscapeActor() : nullptr;
-		if (Landscape)
-		{
-			Landscape->RequestSplineLayerUpdate();
-		}
+		GetOuterULandscapeSplinesComponent()->RequestSplineLayerUpdate();
 	}
 }
 #endif // WITH_EDITOR
@@ -2410,16 +2441,21 @@ void ULandscapeSplineSegment::PostLoad()
 		}
 
 		// Replace null meshes with the editor mesh
-		// Otherwise the spline will have no mesh and won't be easily selectable
+		// Also make sure that we update their visibility and collision profile if they are editor mesh
 		ULandscapeSplinesComponent* OuterSplines = GetOuterULandscapeSplinesComponent();
 		if (OuterSplines->SplineEditorMesh != nullptr)
 		{
 			for (auto* LocalMeshComponent : LocalMeshComponents)
 			{
-				if (LocalMeshComponent->GetStaticMesh() == nullptr)
+				if (LocalMeshComponent->GetStaticMesh() == nullptr || LocalMeshComponent->GetStaticMesh() == OuterSplines->SplineEditorMesh)
 				{
 					LocalMeshComponent->ConditionalPostLoad();
-					LocalMeshComponent->SetStaticMesh(OuterSplines->SplineEditorMesh);
+					
+					if (LocalMeshComponent->GetStaticMesh() == nullptr)
+					{
+						LocalMeshComponent->SetStaticMesh(OuterSplines->SplineEditorMesh);
+					}
+
 					LocalMeshComponent->SetHiddenInGame(true);
 					LocalMeshComponent->SetVisibility(OuterSplines->bShowSplineEditorMesh);
 					LocalMeshComponent->BodyInstance.SetCollisionProfileName(UCollisionProfile::NoCollision_ProfileName);
@@ -2616,6 +2652,11 @@ TArray<USplineMeshComponent*> ULandscapeSplineSegment::GetLocalMeshComponents()
 	return LocalMeshComponents;
 }
 
+bool ULandscapeSplineSegment::SupportsForeignSplineMesh() const
+{
+	return bPlaceSplineMeshesInStreamingLevels&& GetOuterULandscapeSplinesComponent()->GetSplineOwner()->SupportsForeignSplineMesh();
+}
+
 void ULandscapeSplineSegment::UpdateSplinePoints(bool bUpdateCollision, bool bUpdateMeshLevel)
 {
 	Modify();
@@ -2797,7 +2838,7 @@ void ULandscapeSplineSegment::UpdateSplinePoints(bool bUpdateCollision, bool bUp
 			ULandscapeSplinesComponent* MeshComponentOuterSplines = OuterSplines;
 
 			// Attempt to place mesh components into the appropriate landscape streaming levels based on the components under the spline
-			if (bPlaceSplineMeshesInStreamingLevels && !bUsingEditorMesh)
+			if (SupportsForeignSplineMesh() && !bUsingEditorMesh)
 			{
 				// Only "approx" because we rescale T for the 2nd pass based on how well our chosen meshes fit, but it should be good enough
 				FVector ApproxMeshLocation = SplineInfo.Eval(T + MeshT / 2, FVector::ZeroVector);
@@ -2833,10 +2874,9 @@ void ULandscapeSplineSegment::UpdateSplinePoints(bool bUpdateCollision, bool bUp
 			}
 			else if (bUpdateMeshLevel)// Update Foreign/Local if necessary
 			{
-				ALandscapeProxy* CurrentMeshComponentOuterActor = MeshComponent->GetTypedOuter<ALandscapeProxy>();
-				ULandscapeSplinesComponent* CurrentLandscapeSplineComponent = CurrentMeshComponentOuterActor->SplineComponent;
-
-				ALandscapeProxy* MeshComponentOuterActor = Cast<ALandscapeProxy>(MeshComponentOuterSplines->GetOwner());
+				AActor* CurrentMeshComponentOuterActor = MeshComponent->GetTypedOuter<AActor>();
+				AActor* MeshComponentOuterActor = MeshComponentOuterSplines->GetOwner();
+				ILandscapeSplineInterface* SplineOwner = Cast<ILandscapeSplineInterface>(MeshComponentOuterActor);
 
 				// Needs updating
 				if (MeshComponentOuterActor != CurrentMeshComponentOuterActor)
@@ -2848,7 +2888,7 @@ void ULandscapeSplineSegment::UpdateSplinePoints(bool bUpdateCollision, bool bUp
 					MeshComponent->DetachFromComponent(FDetachmentTransformRules::KeepWorldTransform);
 					MeshComponent->InvalidateLightingCache();
 					MeshComponent->Rename(nullptr, MeshComponentOuterActor);
-					MeshComponent->AttachToComponent(MeshComponentOuterActor->SplineComponent, FAttachmentTransformRules::KeepWorldTransform);
+					MeshComponent->AttachToComponent(SplineOwner->GetSplinesComponent(), FAttachmentTransformRules::KeepWorldTransform);
 				}
 			}
 
@@ -3192,13 +3232,7 @@ void ULandscapeSplineSegment::PostEditUndo()
 
 	ULandscapeSplinesComponent* SplineComponent = GetOuterULandscapeSplinesComponent();
 	SplineComponent->MarkRenderStateDirty();
-
-	ALandscapeProxy* OuterLandscape = Cast<ALandscapeProxy>(SplineComponent->GetOwner());
-	ALandscape* Landscape = OuterLandscape ? OuterLandscape->GetLandscapeActor() : nullptr;
-	if (Landscape)
-	{
-		Landscape->RequestSplineLayerUpdate();
-	}
+	SplineComponent->RequestSplineLayerUpdate();
 }
 
 void ULandscapeSplineSegment::PostDuplicate(bool bDuplicateForPIE)
@@ -3261,12 +3295,7 @@ void ULandscapeSplineSegment::PostEditChangeProperty(FPropertyChangedEvent& Prop
 
 	if (PropertyChangedEvent.ChangeType == EPropertyChangeType::ValueSet)
 	{
-		ALandscapeProxy* OuterLandscape = Cast<ALandscapeProxy>(GetOuterULandscapeSplinesComponent()->GetOwner());
-		ALandscape* Landscape = OuterLandscape ? OuterLandscape->GetLandscapeActor() : nullptr;
-		if (Landscape)
-		{
-			Landscape->RequestSplineLayerUpdate();
-		}
+		GetOuterULandscapeSplinesComponent()->RequestSplineLayerUpdate();
 	}
 }
 
@@ -3318,24 +3347,27 @@ void FindStartingControlPoints(ULandscapeSplineControlPoint* ControlPoint, TSet<
 	}
 }
 
-void ULandscapeInfo::MoveControlPointToLandscape(ULandscapeSplineControlPoint* InControlPoint, ALandscapeProxy* FromProxy, ALandscapeProxy* ToProxy)
+void ULandscapeInfo::MoveControlPoint(ULandscapeSplineControlPoint* InControlPoint,  TScriptInterface<ILandscapeSplineInterface> From, TScriptInterface<ILandscapeSplineInterface> To)
 {
-	if (ToProxy->SplineComponent == nullptr)
+	ULandscapeSplinesComponent* ToSplineComponent = To->GetSplinesComponent();
+	ULandscapeSplinesComponent* FromSplineComponent = From->GetSplinesComponent();
+	if (ToSplineComponent == nullptr)
 	{
-		ToProxy->CreateSplineComponent(FromProxy->SplineComponent->GetRelativeScale3D());
-		check(ToProxy->SplineComponent);
+		To->CreateSplineComponent(FromSplineComponent->GetRelativeScale3D());
+		ToSplineComponent = To->GetSplinesComponent();
+		check(ToSplineComponent);
 	}
 
-	if (InControlPoint->GetOuterULandscapeSplinesComponent() == ToProxy->SplineComponent)
+	if (InControlPoint->GetOuterULandscapeSplinesComponent() == ToSplineComponent)
 	{
 		return;
 	}
-	check(InControlPoint->GetOuterULandscapeSplinesComponent() == FromProxy->SplineComponent);
+	check(InControlPoint->GetOuterULandscapeSplinesComponent() == FromSplineComponent);
 
-	ToProxy->SplineComponent->Modify();
+	ToSplineComponent->Modify();
 
 	const FTransform OldToNewTransform =
-		FromProxy->SplineComponent->GetComponentTransform().GetRelativeTransform(ToProxy->SplineComponent->GetComponentTransform());
+		FromSplineComponent->GetComponentTransform().GetRelativeTransform(ToSplineComponent->GetComponentTransform());
 		
 	// Delete all Mesh Components associated with the ControlPoint. (Will get recreated in UpdateSplinePoints)
 	if (InControlPoint->LocalMeshComponent)
@@ -3343,8 +3375,8 @@ void ULandscapeInfo::MoveControlPointToLandscape(ULandscapeSplineControlPoint* I
 		InControlPoint->LocalMeshComponent->Modify();
 		InControlPoint->LocalMeshComponent->UnregisterComponent();
 		InControlPoint->LocalMeshComponent->DestroyComponent();
-		FromProxy->SplineComponent->Modify();
-		FromProxy->SplineComponent->MeshComponentLocalOwnersMap.Remove(InControlPoint->LocalMeshComponent);
+		FromSplineComponent->Modify();
+		FromSplineComponent->MeshComponentLocalOwnersMap.Remove(InControlPoint->LocalMeshComponent);
 		InControlPoint->LocalMeshComponent = nullptr;
 	}
 
@@ -3358,9 +3390,9 @@ void ULandscapeInfo::MoveControlPointToLandscape(ULandscapeSplineControlPoint* I
 	}
 
 	// Move control point to new level
-	FromProxy->SplineComponent->ControlPoints.Remove(InControlPoint);
-	InControlPoint->Rename(nullptr, ToProxy->SplineComponent);
-	ToProxy->SplineComponent->ControlPoints.Add(InControlPoint);
+	FromSplineComponent->ControlPoints.Remove(InControlPoint);
+	InControlPoint->Rename(nullptr, ToSplineComponent);
+	ToSplineComponent->ControlPoints.Add(InControlPoint);
 
 	InControlPoint->Location = OldToNewTransform.TransformPosition(InControlPoint->Location);
 
@@ -3374,27 +3406,34 @@ void ULandscapeInfo::MoveControlPointToLandscape(ULandscapeSplineControlPoint* I
 	{
 		if (!Connection.End)
 		{
-			MoveSegmentToLandscape(Connection.Segment, FromProxy, ToProxy);
+			MoveSegment(Connection.Segment, From, To);
 		}
 	}
 }
 
-void ULandscapeInfo::MoveSegmentToLandscape(ULandscapeSplineSegment* InSegment, ALandscapeProxy* FromProxy, ALandscapeProxy* ToProxy)
+void ULandscapeInfo::MoveSegment(ULandscapeSplineSegment* InSegment, TScriptInterface<ILandscapeSplineInterface> From, TScriptInterface<ILandscapeSplineInterface> To)
 {
-	ToProxy->Modify();
-	if (ToProxy->SplineComponent == nullptr)
+	AActor* ToActor = CastChecked<AActor>(To.GetObject());
+	ToActor->Modify();
+
+	AActor* FromActor = CastChecked<AActor>(From.GetObject());
+	ULandscapeSplinesComponent* ToSplineComponent = To->GetSplinesComponent();
+	ULandscapeSplinesComponent* FromSplineComponent = From->GetSplinesComponent();
+	
+	if (ToSplineComponent == nullptr)
 	{
-		ToProxy->CreateSplineComponent(FromProxy->SplineComponent->GetRelativeScale3D());
-		check(ToProxy->SplineComponent);
+		To->CreateSplineComponent(FromSplineComponent->GetRelativeScale3D());
+		ToSplineComponent = To->GetSplinesComponent();
+		check(ToSplineComponent);
 	}
 		
-	if (InSegment->GetOuterULandscapeSplinesComponent() == ToProxy->SplineComponent)
+	if (InSegment->GetOuterULandscapeSplinesComponent() == ToSplineComponent)
 	{
 		return;
 	}
-	check(InSegment->GetOuterULandscapeSplinesComponent() == FromProxy->SplineComponent);
+	check(InSegment->GetOuterULandscapeSplinesComponent() == FromSplineComponent);
 
-	ToProxy->SplineComponent->Modify();
+	ToSplineComponent->Modify();
 		
 	// Delete all Mesh Components associated with the Segment. (Will get recreated in UpdateSplinePoints)
 	for (auto* MeshComponent : InSegment->LocalMeshComponents)
@@ -3402,8 +3441,8 @@ void ULandscapeInfo::MoveSegmentToLandscape(ULandscapeSplineSegment* InSegment,
 		MeshComponent->Modify();
 		MeshComponent->UnregisterComponent();
 		MeshComponent->DestroyComponent();
-		FromProxy->Modify();
-		FromProxy->SplineComponent->MeshComponentLocalOwnersMap.Remove(MeshComponent);
+		FromActor->Modify();
+		FromSplineComponent->MeshComponentLocalOwnersMap.Remove(MeshComponent);
 	}
 	InSegment->LocalMeshComponents.Empty();
 
@@ -3420,12 +3459,12 @@ void ULandscapeInfo::MoveSegmentToLandscape(ULandscapeSplineSegment* InSegment,
 	}
 
 	// Move segment to new level
-	FromProxy->SplineComponent->Segments.Remove(InSegment);
-	InSegment->Rename(nullptr, ToProxy->SplineComponent);
-	ToProxy->SplineComponent->Segments.Add(InSegment);
+	FromSplineComponent->Segments.Remove(InSegment);
+	InSegment->Rename(nullptr, ToSplineComponent);
+	ToSplineComponent->Segments.Add(InSegment);
 		
 	// Continue ...
-	MoveControlPointToLandscape(InSegment->Connections[1].ControlPoint, FromProxy, ToProxy);
+	MoveControlPoint(InSegment->Connections[1].ControlPoint, From, To);
 
 	const bool bUpdateCollision = true; // default value
 	const bool bUpdateMeshLevel = false; // no need because mesh have been deleted 
@@ -3450,25 +3489,44 @@ void ULandscapeInfo::MoveSplineToLevel(ULandscapeSplineControlPoint* InControlPo
 }
 
 void ULandscapeInfo::MoveSplineToProxy(ULandscapeSplineControlPoint* InControlPoint, ALandscapeProxy* InLandscapeProxy)
+{
+	MoveSpline(InControlPoint, InLandscapeProxy);
+}
+
+void ULandscapeInfo::MoveSpline(ULandscapeSplineControlPoint* InControlPoint, TScriptInterface<ILandscapeSplineInterface> InNewOwner)
 {
 	TSet<ULandscapeSplineControlPoint*> ControlPoints;
 	FindStartingControlPoints(InControlPoint, ControlPoints);
-		
+	// If ControlPoints is empty it means it is looping.
+	if (ControlPoints.IsEmpty())
+	{
+		ControlPoints.Add(InControlPoint);
+	}
+	bool bUpdateBounds = false;
+
 	for (ULandscapeSplineControlPoint* ControlPoint : ControlPoints)
 	{
-		ALandscapeProxy* FromProxy = ControlPoint->GetTypedOuter<ALandscapeProxy>();
-		if (FromProxy == InLandscapeProxy)
+		ULandscapeSplinesComponent* FromSplineComponent = ControlPoint->GetOuterULandscapeSplinesComponent();
+		TScriptInterface<ILandscapeSplineInterface> From(FromSplineComponent->GetOwner());
+		if (From == InNewOwner)
 		{
 			continue;
 		}
-		FromProxy->Modify();
-		FromProxy->SplineComponent->Modify();
-		FromProxy->SplineComponent->MarkRenderStateDirty();
-		MoveControlPointToLandscape(ControlPoint, FromProxy, InLandscapeProxy);
+		bUpdateBounds = true;
+		From.GetObject()->Modify();
+		FromSplineComponent->Modify();
+		FromSplineComponent->MarkRenderStateDirty();
+		MoveControlPoint(ControlPoint, From, InNewOwner);
+		FromSplineComponent->UpdateBounds();
+	}
+
+	if (bUpdateBounds)
+	{
+		InNewOwner->GetSplinesComponent()->UpdateBounds();
 	}
 }
 
-void ULandscapeInfo::MoveSplinesToLevel(ULandscapeSplinesComponent * InSplineComponent, ULevel * TargetLevel)
+void ULandscapeInfo::MoveSplinesToLevel(ULandscapeSplinesComponent* InSplineComponent, ULevel * TargetLevel)
 {
 	ALandscapeProxy* FromProxy = InSplineComponent->GetTypedOuter<ALandscapeProxy>();
 	if (FromProxy->GetLevel() == TargetLevel)
@@ -3482,23 +3540,29 @@ void ULandscapeInfo::MoveSplinesToLevel(ULandscapeSplinesComponent * InSplineCom
 		return;
 	}
 	
-	MoveSplinesToProxy(InSplineComponent, ToProxy);
+	MoveSplines(InSplineComponent, ToProxy);
 }
 
 /** Moves all Splines to target Proxy. Creates ULandscapeSplineComponent if needed */
 void ULandscapeInfo::MoveSplinesToProxy(ULandscapeSplinesComponent* InSplineComponent, ALandscapeProxy* InLandscapeProxy)
 {
-	check(InSplineComponent->GetTypedOuter<ALandscapeProxy>() != InLandscapeProxy);
-	
-	TArray<ULandscapeSplineControlPoint*> CopyControlPoints(InSplineComponent->ControlPoints);
-	for (ULandscapeSplineControlPoint* ControlPoint : CopyControlPoints)
-	{
-		// Even if ControlPoints are part of same connected spline they will be skipped if already moved
-		MoveSplineToProxy(ControlPoint, InLandscapeProxy);
-	}
+	MoveSplines(InSplineComponent, InLandscapeProxy);
 
 	check(InSplineComponent->ControlPoints.Num() == 0 && InSplineComponent->Segments.Num() == 0);
 }
+
+
+/** Moves all Splines to target Spline owner */
+void ULandscapeInfo::MoveSplines(ULandscapeSplinesComponent* InSplineComponent, TScriptInterface<ILandscapeSplineInterface> InNewOwner)
+{
+	check(InSplineComponent != InNewOwner->GetSplinesComponent());
+
+	InSplineComponent->ForEachControlPoint([this, InNewOwner](ULandscapeSplineControlPoint* ControlPoint)
+	{
+		// Even if ControlPoints are part of same connected spline they will be skipped if already moved
+		MoveSpline(ControlPoint, InNewOwner);
+	});
+}
 #endif // WITH_EDITOR
 
 
diff --git a/Engine/Source/Runtime/Landscape/Public/LandscapeSplineProxies.h b/Engine/Source/Runtime/Landscape/Public/LandscapeSplineProxies.h
--- a/Engine/Source/Runtime/Landscape/Public/LandscapeSplineProxies.h
+++ b/Engine/Source/Runtime/Landscape/Public/LandscapeSplineProxies.h
@@ -4,19 +4,22 @@
 
 #include "CoreMinimal.h"
 #include "UObject/UObjectGlobals.h"
-#include "HitProxies.h"
+#include "EngineUtils.h"
+#include "LandscapeSplinesComponent.h"
+#include "LandscapeSplineControlPoint.h"
+#include "LandscapeSplineSegment.h"
 
 class ULandscapeSplineSegment;
 
 //////////////////////////////////////////////////////////////////////////
 // LANDSCAPE SPLINES HIT PROXY
 
-struct HLandscapeSplineProxy : public HHitProxy
+struct HLandscapeSplineProxy : public HActor
 {
 	DECLARE_HIT_PROXY( LANDSCAPE_API );
 
-	HLandscapeSplineProxy(EHitProxyPriority InPriority = HPP_Wireframe) :
-		HHitProxy(InPriority)
+	HLandscapeSplineProxy(ULandscapeSplinesComponent* SplineComponent, EHitProxyPriority InPriority = HPP_Wireframe) :
+		HActor(SplineComponent->GetOwner(), SplineComponent, InPriority)
 	{
 	}
 	virtual EMouseCursor::Type GetMouseCursor() override
@@ -32,7 +35,7 @@ struct HLandscapeSplineProxy_Segment : public HLandscapeSplineProxy
 	class ULandscapeSplineSegment* SplineSegment;
 
 	HLandscapeSplineProxy_Segment(class ULandscapeSplineSegment* InSplineSegment) :
-		HLandscapeSplineProxy(),
+		HLandscapeSplineProxy(InSplineSegment->GetOuterULandscapeSplinesComponent()),
 		SplineSegment(InSplineSegment)
 	{
 	}
@@ -49,7 +52,7 @@ struct HLandscapeSplineProxy_ControlPoint : public HLandscapeSplineProxy
 	class ULandscapeSplineControlPoint* ControlPoint;
 
 	HLandscapeSplineProxy_ControlPoint(class ULandscapeSplineControlPoint* InControlPoint) :
-		HLandscapeSplineProxy(HPP_Foreground),
+		HLandscapeSplineProxy(InControlPoint->GetOuterULandscapeSplinesComponent(), HPP_Foreground),
 		ControlPoint(InControlPoint)
 	{
 	}
@@ -67,7 +70,7 @@ struct HLandscapeSplineProxy_Tangent : public HLandscapeSplineProxy
 	uint32 End:1;
 
 	HLandscapeSplineProxy_Tangent(class ULandscapeSplineSegment* InSplineSegment, bool InEnd) :
-		HLandscapeSplineProxy(HPP_UI),
+		HLandscapeSplineProxy(InSplineSegment->GetOuterULandscapeSplinesComponent(), HPP_UI),
 		SplineSegment(InSplineSegment),
 		End(InEnd)
 	{
