WorldPartition Conversion: Fix MoveSpline with 'y' splines (instead of finding starting points walk the connection graph which already early outs when dest outers match)

diff --git a/Engine/Source/Runtime/Landscape/Private/LandscapeSplines.cpp b/Engine/Source/Runtime/Landscape/Private/LandscapeSplines.cpp
--- a/Engine/Source/Runtime/Landscape/Private/LandscapeSplines.cpp
+++ b/Engine/Source/Runtime/Landscape/Private/LandscapeSplines.cpp
@@ -3319,34 +3319,6 @@ void ALandscapeProxy::CreateSplineComponent(const FVector& Scale3D)
 	}
 }
 
-void FindStartingControlPoints(ULandscapeSplineControlPoint* ControlPoint, TSet<ULandscapeSplineControlPoint*>& StartingControlPoints)
-{
-	// Ensure we aren't already processing this point
-	bool bAlreadyInSet;
-	StartingControlPoints.Add(ControlPoint, &bAlreadyInSet);
-	if (bAlreadyInSet)
-	{
-		return;
-	}
-
-	bool bIsStartingControlPoint = true;
-	for (const FLandscapeSplineConnection& Connection : ControlPoint->ConnectedSegments)
-	{
-		if (Connection.End)
-		{
-			ULandscapeSplineControlPoint* Start = Connection.Segment->Connections[0].ControlPoint;
-			check(Start != ControlPoint);
-			FindStartingControlPoints(Start, StartingControlPoints);
-			bIsStartingControlPoint = false;
-		}
-	}
-
-	if (!bIsStartingControlPoint)
-	{
-		StartingControlPoints.Remove(ControlPoint);
-	}
-}
-
 void ULandscapeInfo::MoveControlPoint(ULandscapeSplineControlPoint* InControlPoint,  TScriptInterface<ILandscapeSplineInterface> From, TScriptInterface<ILandscapeSplineInterface> To)
 {
 	ULandscapeSplinesComponent* ToSplineComponent = To->GetSplinesComponent();
@@ -3404,10 +3376,8 @@ void ULandscapeInfo::MoveControlPoint(ULandscapeSplineControlPoint* InControlPoi
 	// Continue
 	for (const FLandscapeSplineConnection& Connection : InControlPoint->ConnectedSegments)
 	{
-		if (!Connection.End)
-		{
-			MoveSegment(Connection.Segment, From, To);
-		}
+		// Continue both directions anyways it will early out for ControlPoints that already have moved
+		MoveSegment(Connection.Segment, From, To);
 	}
 }
 
@@ -3463,7 +3433,8 @@ void ULandscapeInfo::MoveSegment(ULandscapeSplineSegment* InSegment, TScriptInte
 	InSegment->Rename(nullptr, ToSplineComponent);
 	ToSplineComponent->Segments.Add(InSegment);
 		
-	// Continue ...
+	// Continue both directions anyways it will early out for ControlPoints that already have moved
+	MoveControlPoint(InSegment->Connections[0].ControlPoint, From, To);
 	MoveControlPoint(InSegment->Connections[1].ControlPoint, From, To);
 
 	const bool bUpdateCollision = true; // default value
@@ -3495,31 +3466,21 @@ void ULandscapeInfo::MoveSplineToProxy(ULandscapeSplineControlPoint* InControlPo
 
 void ULandscapeInfo::MoveSpline(ULandscapeSplineControlPoint* InControlPoint, TScriptInterface<ILandscapeSplineInterface> InNewOwner)
 {
-	TSet<ULandscapeSplineControlPoint*> ControlPoints;
-	FindStartingControlPoints(InControlPoint, ControlPoints);
-	// If ControlPoints is empty it means it is looping.
-	if (ControlPoints.IsEmpty())
-	{
-		ControlPoints.Add(InControlPoint);
-	}
 	bool bUpdateBounds = false;
 
-	for (ULandscapeSplineControlPoint* ControlPoint : ControlPoints)
+	ULandscapeSplinesComponent* FromSplineComponent = InControlPoint->GetOuterULandscapeSplinesComponent();
+	TScriptInterface<ILandscapeSplineInterface> From(FromSplineComponent->GetOwner());
+	if (From == InNewOwner)
 	{
-		ULandscapeSplinesComponent* FromSplineComponent = ControlPoint->GetOuterULandscapeSplinesComponent();
-		TScriptInterface<ILandscapeSplineInterface> From(FromSplineComponent->GetOwner());
-		if (From == InNewOwner)
-		{
-			continue;
-		}
-		bUpdateBounds = true;
-		From.GetObject()->Modify();
-		FromSplineComponent->Modify();
-		FromSplineComponent->MarkRenderStateDirty();
-		MoveControlPoint(ControlPoint, From, InNewOwner);
-		FromSplineComponent->UpdateBounds();
+		return;
 	}
-
+	bUpdateBounds = true;
+	From.GetObject()->Modify();
+	FromSplineComponent->Modify();
+	FromSplineComponent->MarkRenderStateDirty();
+	MoveControlPoint(InControlPoint, From, InNewOwner);
+	FromSplineComponent->UpdateBounds();
+	
 	if (bUpdateBounds)
 	{
 		InNewOwner->GetSplinesComponent()->UpdateBounds();
