Refactor Landscape UpdateBakedTexture to not automatically dirty levels/actors and instead display a message so that the user can build them through the Build menu (like Grass)

diff --git a/Engine/Source/Editor/LandscapeEditor/LandscapeEditor.Build.cs b/Engine/Source/Editor/LandscapeEditor/LandscapeEditor.Build.cs
--- a/Engine/Source/Editor/LandscapeEditor/LandscapeEditor.Build.cs
+++ b/Engine/Source/Editor/LandscapeEditor/LandscapeEditor.Build.cs
@@ -26,6 +26,7 @@ public class LandscapeEditor : ModuleRules
 				"ImageWrapper",
                 "EditorWidgets",
                 "Foliage",
+				"ToolMenus"
 			}
 			);
 
diff --git a/Engine/Source/Editor/LandscapeEditor/Private/LandscapeEditorModule.cpp b/Engine/Source/Editor/LandscapeEditor/Private/LandscapeEditorModule.cpp
--- a/Engine/Source/Editor/LandscapeEditor/Private/LandscapeEditorModule.cpp
+++ b/Engine/Source/Editor/LandscapeEditor/Private/LandscapeEditorModule.cpp
@@ -28,6 +28,9 @@
 #include "LandscapeSplineDetails.h"
 
 #include "LevelEditor.h"
+#include "ToolMenus.h"
+#include "Editor/EditorEngine.h"
+#include "LandscapeSubsystem.h"
 
 #include "LandscapeRender.h"
 
@@ -72,7 +75,7 @@ public:
 			true,
 			300
 			);
-	
+
 		// register customizations
 		FPropertyEditorModule& PropertyModule = FModuleManager::GetModuleChecked<FPropertyEditorModule>("PropertyEditor");
 		PropertyModule.RegisterCustomClassLayout("LandscapeEditorObject", FOnGetDetailCustomizationInstance::CreateStatic(&FLandscapeEditorDetails::MakeInstance));
@@ -112,6 +115,20 @@ public:
 		RegisterWeightmapFileFormat(MakeShareable(new FLandscapeWeightmapFileFormat_Png()));
 		RegisterHeightmapFileFormat(MakeShareable(new FLandscapeHeightmapFileFormat_Raw()));
 		RegisterWeightmapFileFormat(MakeShareable(new FLandscapeWeightmapFileFormat_Raw()));
+
+		//Landscape extended menu
+		UToolMenu* BuildMenu = UToolMenus::Get()->ExtendMenu("LevelEditor.MainMenu.Build");
+		if (BuildMenu)
+		{
+			FToolMenuSection& Section = BuildMenu->FindOrAddSection("LevelEditorLandscape");
+
+			FUIAction ActionBakeTextures(FExecuteAction::CreateStatic(&BuildGITextures), FCanExecuteAction());
+			Section.AddMenuEntry(NAME_None, LOCTEXT("BuildGITexturesOnly", "Build GI Textures Only"),LOCTEXT("BuildGIBakedTextures ", "Build GI baked base color textures"), TAttribute<FSlateIcon>(), ActionBakeTextures,EUserInterfaceActionType::Button);
+			
+			FUIAction ActionBuildGrassMaps(FExecuteAction::CreateStatic(&BuildGrassMaps), FCanExecuteAction());
+			Section.AddMenuEntry(NAME_None, LOCTEXT("BuildGrassMapsOnly","Build Grass Maps Only"), LOCTEXT("BuildLandscapeGrassMaps","Build landscape grass maps"),TAttribute<FSlateIcon>(), ActionBuildGrassMaps, EUserInterfaceActionType::Button);
+		}
+
 	}
 
 	/**
@@ -185,6 +202,28 @@ public:
 		);
 	}
 
+	static void BuildGITextures()
+	{
+		if (UWorld* World = GEditor->GetEditorWorldContext().World())
+		{
+			if (ULandscapeSubsystem* LandscapeSubsystem = World->GetSubsystem<ULandscapeSubsystem>())
+			{
+				LandscapeSubsystem->BuildGITextures();
+			}
+		}
+	}
+
+	static void BuildGrassMaps()
+	{
+		if (UWorld* World = GEditor->GetEditorWorldContext().World())
+		{
+			if (ULandscapeSubsystem* LandscapeSubsystem = World->GetSubsystem<ULandscapeSubsystem>())
+			{
+				LandscapeSubsystem->BuildGrassMaps();
+			}
+		}
+	}
+
 	static void ChangeLandscapeViewMode(ELandscapeViewMode::Type ViewMode)
 	{
 		GLandscapeViewMode = ViewMode;
diff --git a/Engine/Source/Editor/LevelEditor/Private/LevelEditor.cpp b/Engine/Source/Editor/LevelEditor/Private/LevelEditor.cpp
--- a/Engine/Source/Editor/LevelEditor/Private/LevelEditor.cpp
+++ b/Engine/Source/Editor/LevelEditor/Private/LevelEditor.cpp
@@ -1489,8 +1489,8 @@ void FLevelEditorModule::BindGlobalLevelEditorCommands()
 	ActionList.MapAction(Commands.BuildVirtualTextureOnly,
 		FExecuteAction::CreateStatic(&FLevelEditorActionCallbacks::BuildVirtualTextureOnly_Execute));
 
-	ActionList.MapAction(Commands.BuildGrassMapsOnly,
-		FExecuteAction::CreateStatic(&FLevelEditorActionCallbacks::BuildGrassMapsOnly_Execute));
+	ActionList.MapAction(Commands.BuildAllLandscape,
+		FExecuteAction::CreateStatic(&FLevelEditorActionCallbacks::BuildAllLandscape_Execute));
 
 	ActionList.MapAction( 
 		Commands.LightingQuality_Production, 
diff --git a/Engine/Source/Editor/LevelEditor/Private/LevelEditorActions.cpp b/Engine/Source/Editor/LevelEditor/Private/LevelEditorActions.cpp
--- a/Engine/Source/Editor/LevelEditor/Private/LevelEditorActions.cpp
+++ b/Engine/Source/Editor/LevelEditor/Private/LevelEditorActions.cpp
@@ -880,10 +880,9 @@ void FLevelEditorActionCallbacks::BuildVirtualTextureOnly_Execute()
 	GEngine->DeferredCommands.AddUnique(TEXT("MAP CHECK NOTIFYRESULTS"));
 }
 
-void FLevelEditorActionCallbacks::BuildGrassMapsOnly_Execute()
+void FLevelEditorActionCallbacks::BuildAllLandscape_Execute()
 {
-	// Build Grass maps
-	FEditorBuildUtils::EditorBuild(GetWorld(), FBuildOptions::BuildGrassMaps);
+	FEditorBuildUtils::EditorBuild(GetWorld(), FBuildOptions::BuildAllLandscape);
 }
 
 bool FLevelEditorActionCallbacks::IsLightingQualityChecked( ELightingBuildQuality TestQuality )
@@ -3211,7 +3210,7 @@ void FLevelEditorCommands::RegisterCommands()
 	UI_COMMAND( BuildLODsOnly, "Build LODs", "Only builds LODs (all levels.)", EUserInterfaceActionType::Button, FInputChord() );
 	UI_COMMAND( BuildTextureStreamingOnly, "Build Texture Streaming", "Build texture streaming data", EUserInterfaceActionType::Button, FInputChord() );
 	UI_COMMAND( BuildVirtualTextureOnly, "Build Virtual Textures", "Build runtime virtual texture low mips streaming data", EUserInterfaceActionType::Button, FInputChord());
-	UI_COMMAND( BuildGrassMapsOnly, "Build Grass Maps", "Build landscape grass maps", EUserInterfaceActionType::Button, FInputChord());
+	UI_COMMAND( BuildAllLandscape, "Build All Landscape", "Build All Landscape Data(Grass, Baked Textures)", EUserInterfaceActionType::Button, FInputChord());
 
 	UI_COMMAND( LightingQuality_Production, "Production", "Sets precomputed lighting quality to highest possible quality (slowest computation time.)", EUserInterfaceActionType::RadioButton, FInputChord() );
 	UI_COMMAND( LightingQuality_High, "High", "Sets precomputed lighting quality to high quality", EUserInterfaceActionType::RadioButton, FInputChord() );
diff --git a/Engine/Source/Editor/LevelEditor/Private/LevelEditorMenu.cpp b/Engine/Source/Editor/LevelEditor/Private/LevelEditorMenu.cpp
--- a/Engine/Source/Editor/LevelEditor/Private/LevelEditorMenu.cpp
+++ b/Engine/Source/Editor/LevelEditor/Private/LevelEditorMenu.cpp
@@ -570,8 +570,8 @@ void FLevelEditorMenu::RegisterBuildMenu()
 	}
 
 	{
-		FToolMenuSection& Section = Menu->AddSection("LevelEditorGrass", LOCTEXT("GrassHeading", "Grass Maps"));
-		Section.AddMenuEntry(FLevelEditorCommands::Get().BuildGrassMapsOnly);
+		FToolMenuSection& Section = Menu->AddSection("LevelEditorLandscape", LOCTEXT("LandscapeHeading", "Landscape"));
+		Section.AddMenuEntry(FLevelEditorCommands::Get().BuildAllLandscape);
 	}
 
 	{
diff --git a/Engine/Source/Editor/LevelEditor/Public/LevelEditorActions.h b/Engine/Source/Editor/LevelEditor/Public/LevelEditorActions.h
--- a/Engine/Source/Editor/LevelEditor/Public/LevelEditorActions.h
+++ b/Engine/Source/Editor/LevelEditor/Public/LevelEditorActions.h
@@ -80,7 +80,7 @@ public:
 	TSharedPtr< FUICommandInfo > BuildLODsOnly;
 	TSharedPtr< FUICommandInfo > BuildTextureStreamingOnly;
 	TSharedPtr< FUICommandInfo > BuildVirtualTextureOnly;
-	TSharedPtr< FUICommandInfo > BuildGrassMapsOnly;
+	TSharedPtr< FUICommandInfo > BuildAllLandscape;
 	TSharedPtr< FUICommandInfo > LightingQuality_Production;
 	TSharedPtr< FUICommandInfo > LightingQuality_High;
 	TSharedPtr< FUICommandInfo > LightingQuality_Medium;
@@ -756,7 +756,7 @@ public:
 	static void BuildLODsOnly_Execute();
 	static void BuildTextureStreamingOnly_Execute();
 	static void BuildVirtualTextureOnly_Execute();
-	static void BuildGrassMapsOnly_Execute();
+	static void BuildAllLandscape_Execute();
 	static void SetLightingQuality( ELightingBuildQuality NewQuality );
 	static bool IsLightingQualityChecked( ELightingBuildQuality TestQuality );
 	static float GetLightingDensityIdeal();
diff --git a/Engine/Source/Editor/UnrealEd/Private/EditorBuildUtils.cpp b/Engine/Source/Editor/UnrealEd/Private/EditorBuildUtils.cpp
--- a/Engine/Source/Editor/UnrealEd/Private/EditorBuildUtils.cpp
+++ b/Engine/Source/Editor/UnrealEd/Private/EditorBuildUtils.cpp
@@ -65,7 +65,7 @@ const FName FBuildOptions::BuildAllOnlySelectedPaths(TEXT("BuildAllOnlySelectedP
 const FName FBuildOptions::BuildHierarchicalLOD(TEXT("BuildHierarchicalLOD"));
 const FName FBuildOptions::BuildTextureStreaming(TEXT("BuildTextureStreaming"));
 const FName FBuildOptions::BuildVirtualTexture(TEXT("BuildVirtualTexture"));
-const FName FBuildOptions::BuildGrassMaps(TEXT("BuildGrassMaps"));
+const FName FBuildOptions::BuildAllLandscape(TEXT("BuildAllLandscape"));
 
 bool FEditorBuildUtils::bBuildingNavigationFromUserRequest = false;
 TMap<FName, FEditorBuildUtils::FCustomBuildType> FEditorBuildUtils::CustomBuildTypes;
@@ -307,9 +307,9 @@ bool FEditorBuildUtils::EditorBuild( UWorld* InWorld, FName Id, const bool bAllo
 	{
 		BuildType = SBuildProgressWidget::BUILDTYPE_VirtualTexture;
 	}
-	else if (Id == FBuildOptions::BuildGrassMaps)
+	else if (Id == FBuildOptions::BuildAllLandscape)
 	{
-		BuildType = SBuildProgressWidget::BUILDTYPE_GrassMaps;
+		BuildType = SBuildProgressWidget::BUILDTYPE_AllLandscape;
 	}
 	else
 	{
@@ -432,13 +432,13 @@ bool FEditorBuildUtils::EditorBuild( UWorld* InWorld, FName Id, const bool bAllo
 			TriggerHierarchicalLODBuilder(InWorld, Id);
 		}
 	}
-	else if (Id == FBuildOptions::BuildGrassMaps)
+	else if (Id == FBuildOptions::BuildAllLandscape)
 	{
 		bDoBuild = GEditor->WarnAboutHiddenLevels(InWorld, false);
 		if (bDoBuild)
 		{
-			GEditor->ResetTransaction(NSLOCTEXT("UnrealEd", "BuildGrassMaps", "Building Grass Maps"));
-			EditorBuildGrassMaps(InWorld);
+			GEditor->ResetTransaction(NSLOCTEXT("UnrealEd", "BuildAllLandscape", "Building Landscape"));
+			EditorBuildAllLandscape(InWorld);
 		}
 	}
 	else if (Id == FBuildOptions::BuildAll || Id == FBuildOptions::BuildAllSubmit)
@@ -928,7 +928,7 @@ FBuildAllHandler::FBuildAllHandler()
 	: CurrentStep(0)
 {
 	// Add built in build steps.
-	BuildSteps.Add(FBuildOptions::BuildGrassMaps);
+	BuildSteps.Add(FBuildOptions::BuildAllLandscape);
 	BuildSteps.Add(FBuildOptions::BuildGeometry);
 	BuildSteps.Add(FBuildOptions::BuildHierarchicalLOD);
 	BuildSteps.Add(FBuildOptions::BuildAIPaths);
@@ -1030,10 +1030,10 @@ void FBuildAllHandler::ProcessBuild(const TWeakPtr<SBuildProgressWidget>& BuildP
 			BuildProgressWidget.Pin()->SetBuildType(SBuildProgressWidget::BUILDTYPE_VirtualTexture);
 			FEditorBuildUtils::EditorBuildVirtualTexture(CurrentWorld);
 		}
-		else if (StepId == FBuildOptions::BuildGrassMaps)
+		else if (StepId == FBuildOptions::BuildAllLandscape)
 		{
-			BuildProgressWidget.Pin()->SetBuildType(SBuildProgressWidget::BUILDTYPE_GrassMaps);
-			FEditorBuildUtils::EditorBuildGrassMaps(CurrentWorld);
+			BuildProgressWidget.Pin()->SetBuildType(SBuildProgressWidget::BUILDTYPE_AllLandscape);
+			FEditorBuildUtils::EditorBuildAllLandscape(CurrentWorld);
 		}
 		else if (StepId == FBuildOptions::BuildAIPaths)
 		{
@@ -1432,13 +1432,13 @@ bool FEditorBuildUtils::EditorBuildVirtualTexture(UWorld* InWorld)
 	return true;
 }
 
-void FEditorBuildUtils::EditorBuildGrassMaps(UWorld* InWorld)
+void FEditorBuildUtils::EditorBuildAllLandscape(UWorld* InWorld)
 {
 	if (InWorld)
 	{
 		if (ULandscapeSubsystem* LandscapeSubsystem = InWorld->GetSubsystem<ULandscapeSubsystem>())
 		{
-			LandscapeSubsystem->BuildGrassMaps();
+			LandscapeSubsystem->BuildAll();
 		}
 	}
 }
diff --git a/Engine/Source/Editor/UnrealEd/Public/Dialogs/SBuildProgress.h b/Engine/Source/Editor/UnrealEd/Public/Dialogs/SBuildProgress.h
--- a/Engine/Source/Editor/UnrealEd/Public/Dialogs/SBuildProgress.h
+++ b/Engine/Source/Editor/UnrealEd/Public/Dialogs/SBuildProgress.h
@@ -46,8 +46,8 @@ public:
 		BUILDTYPE_TextureStreaming, 
 		/** Virtual texture data is being built */
 		BUILDTYPE_VirtualTexture,
-		/** Grass maps are being built */
-		BUILDTYPE_GrassMaps,
+		/** All Landscape is being built*/
+		BUILDTYPE_AllLandscape,
 	};
 
 	/** The various issues that can occur. */
diff --git a/Engine/Source/Editor/UnrealEd/Public/EditorBuildUtils.h b/Engine/Source/Editor/UnrealEd/Public/EditorBuildUtils.h
--- a/Engine/Source/Editor/UnrealEd/Public/EditorBuildUtils.h
+++ b/Engine/Source/Editor/UnrealEd/Public/EditorBuildUtils.h
@@ -36,8 +36,8 @@ struct FBuildOptions
 	UNREALED_API static const FName BuildTextureStreaming;
 	/** Build virtual textures */
 	UNREALED_API static const FName BuildVirtualTexture;
-	/** Build grass maps */
-	UNREALED_API static const FName BuildGrassMaps;
+	/** Build All Landscape Data */
+	UNREALED_API static const FName BuildAllLandscape;
 };
 /**
  * Result of a custom editor build.
@@ -165,11 +165,10 @@ public:
 	static UNREALED_API bool EditorBuildVirtualTexture(UWorld* InWorld);
 
 	/**
-	 * Perform an editor build for grass maps
-	 *
-	 * @param	InWorld				WorldContext
-	 */
-	static UNREALED_API void EditorBuildGrassMaps(UWorld* InWorld);
+	* Perform an editor build for All Landscape Data
+	* @param	InWorld				WorldContext
+	*/
+	static UNREALED_API void EditorBuildAllLandscape(UWorld* InWorld);
 
 	/** 
 	* check if navigation build was was triggered from editor as user request
diff --git a/Engine/Source/Runtime/Engine/Private/UnrealEngine.cpp b/Engine/Source/Runtime/Engine/Private/UnrealEngine.cpp
--- a/Engine/Source/Runtime/Engine/Private/UnrealEngine.cpp
+++ b/Engine/Source/Runtime/Engine/Private/UnrealEngine.cpp
@@ -10210,13 +10210,7 @@ float DrawMapWarnings(UWorld* World, FViewport* Viewport, FCanvas* Canvas, UCanv
 #if WITH_EDITOR && !(UE_BUILD_SHIPPING || UE_BUILD_TEST)
 	if (ULandscapeSubsystem* LandscapeSubsystem = World->GetSubsystem<ULandscapeSubsystem>())
 	{
-		if (int32 OutdatedGrassMapCount = LandscapeSubsystem->GetOutdatedGrassMapCount())
-		{
-			SmallTextItem.SetColor(FLinearColor::Red);
-			SmallTextItem.Text = FText::Format(LOCTEXT("GRASS_MAPS_NEED_TO_BE_REBUILT_FMT", "GRASS MAPS NEED TO BE REBUILT ({0} {0}|plural(one=object,other=objects))"), OutdatedGrassMapCount);
-			Canvas->DrawItem(SmallTextItem, FVector2D(MessageX, MessageY));
-			MessageY += FontSizeY;
-		}
+		LandscapeSubsystem->DisplayBuildMessages(Canvas, MessageX, MessageY);
 	}
 #endif
 
diff --git a/Engine/Source/Runtime/Landscape/Classes/LandscapeComponent.h b/Engine/Source/Runtime/Landscape/Classes/LandscapeComponent.h
--- a/Engine/Source/Runtime/Landscape/Classes/LandscapeComponent.h
+++ b/Engine/Source/Runtime/Landscape/Classes/LandscapeComponent.h
@@ -585,6 +585,10 @@ public:
 	UPROPERTY()
 	FGuid BakedTextureMaterialGuid;
 
+	/** The Material Guid that last saved, to notify manual build operation to bake textures */
+	UPROPERTY(Transient)
+	FGuid LastBakedTextureMaterialGuid;
+
 	/** Pre-baked Base Color texture for use by distance field GI */
 	UPROPERTY(VisibleAnywhere, BlueprintReadOnly, Category = BakedTextures)
 	UTexture2D* GIBakedBaseColorTexture;
diff --git a/Engine/Source/Runtime/Landscape/Classes/LandscapeProxy.h b/Engine/Source/Runtime/Landscape/Classes/LandscapeProxy.h
--- a/Engine/Source/Runtime/Landscape/Classes/LandscapeProxy.h
+++ b/Engine/Source/Runtime/Landscape/Classes/LandscapeProxy.h
@@ -931,7 +931,7 @@ public:
 	void RenderGrassMaps(const TArray<ULandscapeComponent*>& LandscapeComponents, const TArray<ULandscapeGrassType*>& GrassTypes);
 
 	/** Update any textures baked from the landscape as necessary */
-	void UpdateBakedTextures();
+	void UpdateBakedTextures(bool bInShouldMarkDirty = false);
 
 	/** Update the landscape physical material render tasks */
 	void UpdatePhysicalMaterialTasks();
@@ -981,6 +981,8 @@ public:
 
 	LANDSCAPE_API int32 GetOutdatedGrassMapCount() const;
 	LANDSCAPE_API void BuildGrassMaps(struct FScopedSlowTask* InSlowTask = nullptr);
+	LANDSCAPE_API void BuildGITextures(struct FScopedSlowTask* InSlowTask = nullptr);
+	LANDSCAPE_API int32 GetComponentsNeedingGITextureBaking() const;
 	LANDSCAPE_API virtual void CreateSplineComponent() override;
 	LANDSCAPE_API virtual void CreateSplineComponent(const FVector& Scale3D) override;
 
@@ -1205,4 +1207,18 @@ private:
 	mutable int32 OutdatedGrassMapCount;
 	mutable double GrassMapsLastCheckTime;
 };
+
+/**
+ * Helper class used to Build or monitor Landscape GI Textures
+ */
+class LANDSCAPE_API FLandscapeBakedGITextureBuilder
+{
+public:
+	FLandscapeBakedGITextureBuilder(UWorld* InWorld);
+	void Build();
+	int32 GetComponentsNeedingTextureBaking();
+private:
+	UWorld* World;
+	int32 ComponentsNeedingTextureBaking;
+};
 #endif
\ No newline at end of file
diff --git a/Engine/Source/Runtime/Landscape/Private/Landscape.cpp b/Engine/Source/Runtime/Landscape/Private/Landscape.cpp
--- a/Engine/Source/Runtime/Landscape/Private/Landscape.cpp
+++ b/Engine/Source/Runtime/Landscape/Private/Landscape.cpp
@@ -71,6 +71,7 @@ Landscape.cpp: Terrain rendering
 #include "Algo/BinarySearch.h"
 #if WITH_EDITOR
 #include "Rendering/StaticLightingSystemInterface.h"
+#include "Misc/ScopedSlowTask.h"
 #endif
 
 /** Landscape stats */
@@ -553,6 +554,14 @@ void ULandscapeComponent::Serialize(FArchive& Ar)
 		PlatformData.Serialize(Ar, this);
 	}
 #endif
+
+#if WITH_EDITOR
+	if (Ar.IsSaving() && Ar.IsPersistent())
+	{
+		//Update the lastsaved Guid for GI texture
+		LastBakedTextureMaterialGuid = BakedTextureMaterialGuid;
+	}
+#endif
 }
 
 void ULandscapeComponent::GetResourceSizeEx(FResourceSizeEx& CumulativeResourceSize)
@@ -795,6 +804,7 @@ void ULandscapeComponent::PostLoad()
 				}
 			}
 		}
+		LastBakedTextureMaterialGuid = BakedTextureMaterialGuid;
 	}
 #endif
 
@@ -2087,6 +2097,8 @@ void ALandscapeProxy::PreSave(const class ITargetPlatform* TargetPlatform)
 #if WITH_EDITOR
 	// Work out whether we have grass or not for the next game run
 	BuildGrassMaps();
+	//Update the baked textures before saving
+	BuildGITextures();
 
 	for (ULandscapeComponent* Component : LandscapeComponents)
 	{
@@ -3749,7 +3761,66 @@ void ULandscapeComponent::SerializeStateHashes(FArchive& Ar)
 	}
 }
 
-void ALandscapeProxy::UpdateBakedTextures()
+#if WITH_EDITOR
+FLandscapeBakedGITextureBuilder::FLandscapeBakedGITextureBuilder(UWorld* InWorld)
+	:World(InWorld)
+	,ComponentsNeedingTextureBaking(0)
+{
+
+}
+
+void FLandscapeBakedGITextureBuilder::Build()
+{
+	if (World)
+	{
+		const int32 Count = GetComponentsNeedingTextureBaking();
+		FScopedSlowTask SlowTask(Count, (LOCTEXT("BakeTextures_BuildLandscapeBakedTextures", "Baking Landscape Textures")));
+		SlowTask.MakeDialog();
+
+		for (TActorIterator<ALandscapeProxy> ProxyIt(World); ProxyIt; ++ProxyIt)
+		{
+			ProxyIt->BuildGITextures(&SlowTask);
+		}
+	}
+}
+
+int32 FLandscapeBakedGITextureBuilder::GetComponentsNeedingTextureBaking()
+{
+	if (World)
+	{
+		ComponentsNeedingTextureBaking = 0;
+		for (TActorIterator<ALandscapeProxy> ProxyIt(World); ProxyIt; ++ProxyIt)
+		{
+			ComponentsNeedingTextureBaking += ProxyIt->GetComponentsNeedingGITextureBaking();
+		}
+	}
+	return ComponentsNeedingTextureBaking;
+}
+
+void ALandscapeProxy::BuildGITextures(struct FScopedSlowTask* InSlowTask /* = nullptr */)
+{
+	if (!HasAnyFlags(RF_ClassDefaultObject))
+	{
+		const bool bShouldMarkDirty = true;
+		UpdateBakedTextures(bShouldMarkDirty);
+	}
+}
+int32 ALandscapeProxy::GetComponentsNeedingGITextureBaking() const
+{
+	int32 NumOfComponentsNeedingTextureBaking = 0;
+	for (auto Component : LandscapeComponents)
+	{
+		const bool bIsDirty = Component->GetPackage()->IsDirty();
+		if (Component->LastBakedTextureMaterialGuid != Component->BakedTextureMaterialGuid && !bIsDirty)
+		{
+			NumOfComponentsNeedingTextureBaking++;
+		}
+	}
+	return NumOfComponentsNeedingTextureBaking;
+}
+#endif
+
+void ALandscapeProxy::UpdateBakedTextures(bool bInShouldMarkDirty)
 {
 	// See if we can render
 	UWorld* World = GetWorld();
@@ -3758,6 +3829,11 @@ void ALandscapeProxy::UpdateBakedTextures()
 		return;
 	}
 
+	if (bInShouldMarkDirty && GetComponentsNeedingGITextureBaking()>0)
+	{
+		MarkPackageDirty();
+	}
+
 	if (UpdateBakedTexturesCountdown-- > 0)
 	{
 		return;
@@ -3773,6 +3849,7 @@ void ALandscapeProxy::UpdateBakedTextures()
 		{
 			if (Component != nullptr && Component->GIBakedBaseColorTexture != nullptr)
 			{
+				Component->Modify(bInShouldMarkDirty);
 				Component->BakedTextureMaterialGuid.Invalidate();
 				Component->GIBakedBaseColorTexture = nullptr;
 				Component->MarkRenderStateDirty();
@@ -3784,7 +3861,7 @@ void ALandscapeProxy::UpdateBakedTextures()
 
 		return;
 	}
-	
+
 	// Stores the components and their state hash data for a single atlas
 	struct FBakedTextureSourceInfo
 	{
@@ -3816,6 +3893,8 @@ void ALandscapeProxy::UpdateBakedTextures()
 	NumComponentsNeedingTextureBaking = 0;
 	int32 NumGenerated = 0;
 
+
+
 	for (auto It = ComponentsByHeightmap.CreateConstIterator(); It; ++It)
 	{
 		const FBakedTextureSourceInfo& Info = It.Value();
@@ -3863,7 +3942,7 @@ void ALandscapeProxy::UpdateBakedTextures()
 					break;
 				}
 			}
-			
+
 			if (bNeedsBake)
 			{
 				// We throttle, baking only one atlas per frame
@@ -3909,6 +3988,7 @@ void ALandscapeProxy::UpdateBakedTextures()
 
 					for (ULandscapeComponent* Component : Info.Components)
 					{
+						Component->Modify(bInShouldMarkDirty);
 						Component->BakedTextureMaterialGuid = CombinedStateId;
 						Component->GIBakedBaseColorTexture = AtlasTexture;
 						Component->MarkRenderStateDirty();
diff --git a/Engine/Source/Runtime/Landscape/Private/LandscapeSubsystem.cpp b/Engine/Source/Runtime/Landscape/Private/LandscapeSubsystem.cpp
--- a/Engine/Source/Runtime/Landscape/Private/LandscapeSubsystem.cpp
+++ b/Engine/Source/Runtime/Landscape/Private/LandscapeSubsystem.cpp
@@ -15,6 +15,7 @@
 #include "Engine/World.h"
 #include "Math/IntRect.h"
 #include "LandscapeConfigHelper.h"
+#include "Engine/Canvas.h"
 
 static int32 GUseStreamingManagerForCameras = 1;
 static FAutoConsoleVariableRef CVarUseStreamingManagerForCameras(
@@ -24,6 +25,8 @@ static FAutoConsoleVariableRef CVarUseStreamingManagerForCameras(
 
 DECLARE_CYCLE_STAT(TEXT("LandscapeSubsystem Tick"), STAT_LandscapeSubsystemTick, STATGROUP_Landscape);
 
+#define LOCTEXT_NAMESPACE "LandscapeSubsystem"
+
 ULandscapeSubsystem::ULandscapeSubsystem()
 {
 }
@@ -56,6 +59,7 @@ void ULandscapeSubsystem::Initialize(FSubsystemCollectionBase& Collection)
 
 #if WITH_EDITOR
 	GrassMapsBuilder = new FLandscapeGrassMapsBuilder(GetWorld());
+	BakedGITextureBuilder = new FLandscapeBakedGITextureBuilder(GetWorld());
 #endif
 }
 
@@ -156,6 +160,12 @@ void ULandscapeSubsystem::Tick(float DeltaTime)
 }
 
 #if WITH_EDITOR
+void ULandscapeSubsystem::BuildAll()
+{
+	BuildGrassMaps();
+	BuildGITextures();
+}
+
 void ULandscapeSubsystem::BuildGrassMaps()
 {
 	GrassMapsBuilder->Build();
@@ -166,6 +176,16 @@ int32 ULandscapeSubsystem::GetOutdatedGrassMapCount()
 	return GrassMapsBuilder->GetOutdatedGrassMapCount(/*bInForceUpdate*/false);
 }
 
+void ULandscapeSubsystem::BuildGITextures()
+{
+	BakedGITextureBuilder->Build();
+}
+
+int32 ULandscapeSubsystem::GetComponentsNeedingGITextureBaking()
+{
+	return BakedGITextureBuilder->GetComponentsNeedingTextureBaking();
+}
+
 bool ULandscapeSubsystem::IsGridBased() const
 {
 	return UWorld::HasSubsystem<UWorldPartitionSubsystem>(GetWorld());
@@ -194,4 +214,29 @@ ALandscapeProxy* ULandscapeSubsystem::FindOrAddLandscapeProxy(ULandscapeInfo* La
 	return FLandscapeConfigHelper::FindOrAddLandscapeStreamingProxy(LandscapeInfo, SectionBase);
 }
 
-#endif
\ No newline at end of file
+void ULandscapeSubsystem::DisplayBuildMessages(FCanvas* Canvas, float& XPos, float& YPos)
+{
+	const int32 FontSizeY = 20;
+	FCanvasTextItem SmallTextItem(FVector2D(0, 0), FText::GetEmpty(), GEngine->GetSmallFont(), FLinearColor::White);
+	SmallTextItem.EnableShadow(FLinearColor::Black);
+
+	if (int32 OutdatedGrassMapCount = GetOutdatedGrassMapCount())
+	{
+		SmallTextItem.SetColor(FLinearColor::Red);
+		SmallTextItem.Text = FText::Format(LOCTEXT("GRASS_MAPS_NEED_TO_BE_REBUILT_FMT", "GRASS MAPS NEED TO BE REBUILT ({0} {0}|plural(one=object,other=objects))"), OutdatedGrassMapCount);
+		Canvas->DrawItem(SmallTextItem, FVector2D(XPos, YPos));
+		YPos += FontSizeY;
+	}
+
+	if (int32 ComponentsNeedingGITextureBaking = GetComponentsNeedingGITextureBaking())
+	{
+		SmallTextItem.SetColor(FLinearColor::Red);
+		SmallTextItem.Text = FText::Format(LOCTEXT("LANDSCAPE_TEXTURES_NEED_TO_BE_REBUILT_FMT", "LANDSCAPE BAKED TEXTURES NEED TO BE REBUILT ({0} {0}|plural(one=object,other=objects))"), ComponentsNeedingGITextureBaking);
+		Canvas->DrawItem(SmallTextItem, FVector2D(XPos, YPos));
+		YPos += FontSizeY;
+	}
+}
+
+#endif
+
+#undef LOCTEXT_NAMESPACE
\ No newline at end of file
diff --git a/Engine/Source/Runtime/Landscape/Public/LandscapeSubsystem.h b/Engine/Source/Runtime/Landscape/Public/LandscapeSubsystem.h
--- a/Engine/Source/Runtime/Landscape/Public/LandscapeSubsystem.h
+++ b/Engine/Source/Runtime/Landscape/Public/LandscapeSubsystem.h
@@ -32,12 +32,15 @@ public:
 	// End FTickableGameObject overrides
 
 #if WITH_EDITOR
+	LANDSCAPE_API void BuildAll();
 	LANDSCAPE_API void BuildGrassMaps();
 	LANDSCAPE_API int32 GetOutdatedGrassMapCount();
-
+	LANDSCAPE_API void BuildGITextures();
+	LANDSCAPE_API int32 GetComponentsNeedingGITextureBaking();
 	LANDSCAPE_API bool IsGridBased() const;
 	LANDSCAPE_API void ChangeGridSize(ULandscapeInfo* LandscapeInfo, uint32 NewGridSizeInComponents);
 	LANDSCAPE_API ALandscapeProxy* FindOrAddLandscapeProxy(ULandscapeInfo* LandscapeInfo, const FIntPoint& SectionBase);
+	LANDSCAPE_API void DisplayBuildMessages(class FCanvas* Canvas, float& XPos, float& YPos);
 #endif
 
 private:
@@ -50,6 +53,7 @@ private:
 
 #if WITH_EDITOR
 	class FLandscapeGrassMapsBuilder* GrassMapsBuilder;
+	class FLandscapeBakedGITextureBuilder* BakedGITextureBuilder;
 	FLandscapeActorDescFactory LandscapeActorDescFactory;
 #endif
 };
