- Worldpartition conversion of Landscape calls ChangeGridSize which needs to force a landscape update when using Edit Layers
- Remove Move Spline code that isn't needed anymore with Spline Actors.

diff --git a/Engine/Source/Runtime/Landscape/Classes/Landscape.h b/Engine/Source/Runtime/Landscape/Classes/Landscape.h
--- a/Engine/Source/Runtime/Landscape/Classes/Landscape.h
+++ b/Engine/Source/Runtime/Landscape/Classes/Landscape.h
@@ -343,6 +343,7 @@ public:
 	
 	LANDSCAPE_API void ToggleCanHaveLayersContent();
 	LANDSCAPE_API void ForceUpdateLayersContent(bool bIntermediateRender = false);
+	LANDSCAPE_API void ForceLayersFullUpdate();
 	LANDSCAPE_API void InitializeLandscapeLayersWeightmapUsage();
 
 private:
diff --git a/Engine/Source/Runtime/Landscape/Classes/LandscapeInfo.h b/Engine/Source/Runtime/Landscape/Classes/LandscapeInfo.h
--- a/Engine/Source/Runtime/Landscape/Classes/LandscapeInfo.h
+++ b/Engine/Source/Runtime/Landscape/Classes/LandscapeInfo.h
@@ -325,6 +325,7 @@ public:
 	LANDSCAPE_API void UnregisterSplineActor(TScriptInterface<ILandscapeSplineInterface> SplineActor);
 
 	LANDSCAPE_API void RequestSplineLayerUpdate();
+	LANDSCAPE_API void ForceLayersFullUpdate();
 #endif
 
 #if WITH_EDITOR
diff --git a/Engine/Source/Runtime/Landscape/Private/Landscape.cpp b/Engine/Source/Runtime/Landscape/Private/Landscape.cpp
--- a/Engine/Source/Runtime/Landscape/Private/Landscape.cpp
+++ b/Engine/Source/Runtime/Landscape/Private/Landscape.cpp
@@ -3338,6 +3338,14 @@ void ULandscapeInfo::RequestSplineLayerUpdate()
 		LandscapeActor->RequestSplineLayerUpdate();
 	}
 }
+
+void ULandscapeInfo::ForceLayersFullUpdate()
+{
+	if (LandscapeActor)
+	{
+		LandscapeActor->ForceLayersFullUpdate();
+	}
+}
 #endif
 
 void ULandscapeInfo::RegisterCollisionComponent(ULandscapeHeightfieldCollisionComponent* Component)
diff --git a/Engine/Source/Runtime/Landscape/Private/LandscapeConfigHelper.cpp b/Engine/Source/Runtime/Landscape/Private/LandscapeConfigHelper.cpp
--- a/Engine/Source/Runtime/Landscape/Private/LandscapeConfigHelper.cpp
+++ b/Engine/Source/Runtime/Landscape/Private/LandscapeConfigHelper.cpp
@@ -228,6 +228,11 @@ bool FLandscapeConfigHelper::ChangeGridSize(ULandscapeInfo* InLandscapeInfo, uin
 		OutActorsToDelete.Add(ProxyToDelete);
 	}
 
+	if (InLandscapeInfo->CanHaveLayersContent())
+	{
+		InLandscapeInfo->ForceLayersFullUpdate();
+	}
+
 	return true;
 }
 
@@ -549,15 +554,10 @@ ULandscapeInfo* FLandscapeConfigHelper::ChangeConfiguration(ULandscapeInfo* InLa
 			NewComponent->RegisterComponent();
 		}
 	}
-		
-	// @todo_ow: This is probably not the final solution
-	MoveSplinesToLandscape(InLandscapeInfo, NewLandscape, NewScaleFactor);
-		
+
 	if (NewLandscapeInfo->CanHaveLayersContent())
 	{
-		NewLandscape->RequestLayersContentUpdateForceAll();
-		NewLandscape->RequestSplineLayerUpdate();
-		NewLandscape->ForceUpdateLayersContent();
+		NewLandscapeInfo->ForceLayersFullUpdate();
 	}
 	else
 	{
diff --git a/Engine/Source/Runtime/Landscape/Private/LandscapeEditLayers.cpp b/Engine/Source/Runtime/Landscape/Private/LandscapeEditLayers.cpp
--- a/Engine/Source/Runtime/Landscape/Private/LandscapeEditLayers.cpp
+++ b/Engine/Source/Runtime/Landscape/Private/LandscapeEditLayers.cpp
@@ -5201,6 +5201,13 @@ void ALandscape::ForceUpdateLayersContent(bool bInIntermediateRender)
 	bInIntermediateRender = false;
 }
 
+void ALandscape::ForceLayersFullUpdate()
+{
+	RequestSplineLayerUpdate();
+	RequestLayersContentUpdateForceAll();
+	ForceUpdateLayersContent(/* bIntermediateRender */ false);
+}
+
 void ALandscape::TickLayers(float DeltaTime)
 {
 	check(GIsEditor);
