Refactor Landscape UpdateGIBakedTextures to not automatically dirty levels/actors instead display a message so that the user can build them through the Build menu

diff --git a/Engine/Source/Editor/LandscapeEditor/Private/LandscapeEditorModule.cpp b/Engine/Source/Editor/LandscapeEditor/Private/LandscapeEditorModule.cpp
--- a/Engine/Source/Editor/LandscapeEditor/Private/LandscapeEditorModule.cpp
+++ b/Engine/Source/Editor/LandscapeEditor/Private/LandscapeEditorModule.cpp
@@ -208,7 +208,7 @@ public:
 		{
 			if (ULandscapeSubsystem* LandscapeSubsystem = World->GetSubsystem<ULandscapeSubsystem>())
 			{
-				LandscapeSubsystem->BuildGITextures();
+				LandscapeSubsystem->BuildGIBakedTextures();
 			}
 		}
 	}
diff --git a/Engine/Source/Runtime/Landscape/Classes/LandscapeProxy.h b/Engine/Source/Runtime/Landscape/Classes/LandscapeProxy.h
--- a/Engine/Source/Runtime/Landscape/Classes/LandscapeProxy.h
+++ b/Engine/Source/Runtime/Landscape/Classes/LandscapeProxy.h
@@ -931,8 +931,19 @@ public:
 	/** Render grass maps for the specified components */
 	void RenderGrassMaps(const TArray<ULandscapeComponent*>& LandscapeComponents, const TArray<ULandscapeGrassType*>& GrassTypes);
 
-	/** Update any textures baked from the landscape as necessary */
-	void UpdateBakedTextures(bool bInShouldMarkDirty = false);
+	/** Stores Components and their latest BakedTextureMaterial hash*/
+	struct FGIBakedTextureState
+	{
+		FGuid CombinedStateId;
+		TArray<ULandscapeComponent*> Components;
+	};
+		
+	/** Update any GI baked textures from the landscape as necessary */
+	void UpdateGIBakedTextures(bool bBakeAllGITextures = false);
+	/** Update the landscape GI baked textures without deferring*/
+	void UpdateGIBakedTextureData(bool bInShouldMarkDirty = false);
+	/** Update the status of landscape GI baked textures*/
+	void UpdateGIBakedTextureStatus(bool* bOutGenerateLandscapeGIData, TMap<UTexture2D*, FGIBakedTextureState>* OutComponentsNeedBakingByHeightmap, int32* OutdatedComponentsCount=nullptr) const;
 
 	/** Update the landscape physical material render tasks */
 	void UpdatePhysicalMaterialTasks();
@@ -982,8 +993,8 @@ public:
 
 	LANDSCAPE_API int32 GetOutdatedGrassMapCount() const;
 	LANDSCAPE_API void BuildGrassMaps(struct FScopedSlowTask* InSlowTask = nullptr);
-	LANDSCAPE_API void BuildGITextures(struct FScopedSlowTask* InSlowTask = nullptr);
-	LANDSCAPE_API int32 GetComponentsNeedingGITextureBaking() const;
+	LANDSCAPE_API void BuildGIBakedTextures(struct FScopedSlowTask* InSlowTask = nullptr);
+	LANDSCAPE_API int32 GetOutdatedGIBakedTextureComponentsCount() const;
 	LANDSCAPE_API virtual void CreateSplineComponent() override;
 	LANDSCAPE_API virtual void CreateSplineComponent(const FVector& Scale3D) override;
 
@@ -1212,14 +1223,15 @@ private:
 /**
  * Helper class used to Build or monitor Landscape GI Textures
  */
-class LANDSCAPE_API FLandscapeBakedGITextureBuilder
+class LANDSCAPE_API FLandscapeGIBakedTextureBuilder
 {
 public:
-	FLandscapeBakedGITextureBuilder(UWorld* InWorld);
+	FLandscapeGIBakedTextureBuilder(UWorld* InWorld);
 	void Build();
-	int32 GetComponentsNeedingTextureBaking();
+	int32 GetOutdatedGIBakedTextureComponentsCount(bool bInForceUpdate = true) const;
 private:
 	UWorld* World;
-	int32 ComponentsNeedingTextureBaking;
+	mutable int32 OutdatedGIBakedTextureComponentsCount;
+	mutable double GIBakedTexturesLastCheckTime;
 };
 #endif
\ No newline at end of file
diff --git a/Engine/Source/Runtime/Landscape/Private/Landscape.cpp b/Engine/Source/Runtime/Landscape/Private/Landscape.cpp
--- a/Engine/Source/Runtime/Landscape/Private/Landscape.cpp
+++ b/Engine/Source/Runtime/Landscape/Private/Landscape.cpp
@@ -2098,7 +2098,7 @@ void ALandscapeProxy::PreSave(const class ITargetPlatform* TargetPlatform)
 	// Work out whether we have grass or not for the next game run
 	BuildGrassMaps();
 	//Update the baked textures before saving
-	BuildGITextures();
+	BuildGIBakedTextures();
 
 	for (ULandscapeComponent* Component : LandscapeComponents)
 	{
@@ -3769,66 +3769,180 @@ void ULandscapeComponent::SerializeStateHashes(FArchive& Ar)
 	}
 }
 
-#if WITH_EDITOR
-FLandscapeBakedGITextureBuilder::FLandscapeBakedGITextureBuilder(UWorld* InWorld)
+FLandscapeGIBakedTextureBuilder::FLandscapeGIBakedTextureBuilder(UWorld* InWorld)
 	:World(InWorld)
-	,ComponentsNeedingTextureBaking(0)
+	,OutdatedGIBakedTextureComponentsCount(0)
+	,GIBakedTexturesLastCheckTime(0)
 {
 
 }
 
-void FLandscapeBakedGITextureBuilder::Build()
+void FLandscapeGIBakedTextureBuilder::Build()
 {
 	if (World)
 	{
-		const int32 Count = GetComponentsNeedingTextureBaking();
-		FScopedSlowTask SlowTask(Count, (LOCTEXT("BakeTextures_BuildLandscapeBakedTextures", "Baking Landscape Textures")));
+		FScopedSlowTask SlowTask(OutdatedGIBakedTextureComponentsCount, (LOCTEXT("BakeTextures_BuildLandscapeBakedTextures", "Baking Landscape Textures")));
 		SlowTask.MakeDialog();
 
 		for (TActorIterator<ALandscapeProxy> ProxyIt(World); ProxyIt; ++ProxyIt)
 		{
-			ProxyIt->BuildGITextures(&SlowTask);
+			ProxyIt->BuildGIBakedTextures(&SlowTask);
 		}
+		//Force update the outdated count when using the build menu option.
+		OutdatedGIBakedTextureComponentsCount = 0;
+		GIBakedTexturesLastCheckTime = FPlatformTime::Seconds();
+
 	}
 }
 
-int32 FLandscapeBakedGITextureBuilder::GetComponentsNeedingTextureBaking()
+int32 FLandscapeGIBakedTextureBuilder::GetOutdatedGIBakedTextureComponentsCount(bool bInForceUpdate) const
 {
 	if (World)
 	{
-		ComponentsNeedingTextureBaking = 0;
-		for (TActorIterator<ALandscapeProxy> ProxyIt(World); ProxyIt; ++ProxyIt)
+		bool bUpdate = bInForceUpdate;
+		double GIBakedTexturesTimeNow = FPlatformTime::Seconds();
+		if (!bUpdate)
 		{
-			ComponentsNeedingTextureBaking += ProxyIt->GetComponentsNeedingGITextureBaking();
+			// Recheck every 20 secs
+			if ((GIBakedTexturesTimeNow - GIBakedTexturesLastCheckTime) > 20)
+			{
+				bUpdate = true;
+			}
+		}
+		if (bUpdate)
+		{
+			GIBakedTexturesLastCheckTime = GIBakedTexturesTimeNow;
+			OutdatedGIBakedTextureComponentsCount = 0;
+			for (TActorIterator<ALandscapeProxy> ProxyIt(World); ProxyIt; ++ProxyIt)
+			{
+				OutdatedGIBakedTextureComponentsCount += ProxyIt->GetOutdatedGIBakedTextureComponentsCount();
+			}
 		}
 	}
-	return ComponentsNeedingTextureBaking;
+	return OutdatedGIBakedTextureComponentsCount;
 }
 
-void ALandscapeProxy::BuildGITextures(struct FScopedSlowTask* InSlowTask /* = nullptr */)
+void ALandscapeProxy::BuildGIBakedTextures(struct FScopedSlowTask* InSlowTask /* = nullptr */)
 {
 	if (!HasAnyFlags(RF_ClassDefaultObject))
 	{
 		const bool bShouldMarkDirty = true;
-		UpdateBakedTextures(bShouldMarkDirty);
+		UpdateGIBakedTextureData(bShouldMarkDirty);
 	}
 }
-int32 ALandscapeProxy::GetComponentsNeedingGITextureBaking() const
+
+int32 ALandscapeProxy::GetOutdatedGIBakedTextureComponentsCount() const
 {
-	int32 NumOfComponentsNeedingTextureBaking = 0;
-	for (auto Component : LandscapeComponents)
+	int32 OutdatedGITextureComponentsCount = 0;
+	UpdateGIBakedTextureStatus(nullptr, nullptr, &OutdatedGITextureComponentsCount);
+	return OutdatedGITextureComponentsCount;
+}
+
+void ALandscapeProxy::UpdateGIBakedTextureStatus(bool* bOutGenerateLandscapeGIData, TMap<UTexture2D*, FGIBakedTextureState>* OutComponentsNeedBakingByHeightmap, int32* OutdatedComponentsCount) const
+{
+	int32 OutdatedComponents = 0;
+	int32 ComponentsNeedToBeCleared = 0;
+	int32 ComponentsNeedToBeBaked = 0;
+
+	// Check if we want to generate landscape GI data
+	static const auto DistanceFieldCVar = IConsoleManager::Get().FindTConsoleVariableDataInt(TEXT("r.GenerateMeshDistanceFields"));
+	static const auto LandscapeGICVar = IConsoleManager::Get().FindTConsoleVariableDataInt(TEXT("r.GenerateLandscapeGIData"));
+	if (DistanceFieldCVar->GetValueOnGameThread() == 0 || LandscapeGICVar->GetValueOnGameThread() == 0)
+	{
+		if (bOutGenerateLandscapeGIData)
+		{
+			*bOutGenerateLandscapeGIData = false;
+		}
+
+		for (ULandscapeComponent* Component : LandscapeComponents)
+		{
+			if (Component != nullptr && Component->GIBakedBaseColorTexture != nullptr)
+			{
+				ComponentsNeedToBeCleared++;
+			}
+		}
+
+		OutdatedComponents += ComponentsNeedToBeCleared;
+	}
+	else
 	{
-		const bool bIsDirty = Component->GetPackage()->IsDirty();
-		if (Component->LastBakedTextureMaterialGuid != Component->BakedTextureMaterialGuid && !bIsDirty)
+		// Stores the components and their state hash data for a single atlas
+		struct FGIBakeTextureStateBuilder
 		{
-			NumOfComponentsNeedingTextureBaking++;
+			// pointer as FMemoryWriter caches the address of the FBufferArchive, and this struct could be relocated on a realloc.
+			TUniquePtr<FBufferArchive> ComponentStateAr;
+			TArray<ULandscapeComponent*> Components;
+
+			FGIBakeTextureStateBuilder()
+			{
+				ComponentStateAr = MakeUnique<FBufferArchive>();
+			}
+		};
+
+		TMap<UTexture2D*, FGIBakeTextureStateBuilder> ComponentsByHeightmap;
+		for (ULandscapeComponent* Component : LandscapeComponents)
+		{
+			if (Component == nullptr)
+			{
+				continue;
+			}
+
+			FGIBakeTextureStateBuilder& Info = ComponentsByHeightmap.FindOrAdd(Component->GetHeightmap());
+			Info.Components.Add(Component);
+			Component->SerializeStateHashes(*Info.ComponentStateAr);
+		}
+
+		for (auto It = ComponentsByHeightmap.CreateIterator(); It; ++It)
+		{
+			FGIBakeTextureStateBuilder& Info =It.Value();
+
+			// Calculate a combined Guid-like ID we can use for this component
+			uint32 Hash[5];
+			FSHA1::HashBuffer(Info.ComponentStateAr->GetData(), Info.ComponentStateAr->Num(), (uint8*)Hash);
+			FGuid CombinedStateId = FGuid(Hash[0] ^ Hash[4], Hash[1], Hash[2], Hash[3]);
+
+			if (Info.Components[0]->BakedTextureMaterialGuid != CombinedStateId)
+			{
+				ComponentsNeedToBeBaked += Info.Components.Num();
+				if (OutComponentsNeedBakingByHeightmap)
+				{
+					FGIBakedTextureState& GIBakedTextureState = OutComponentsNeedBakingByHeightmap->FindOrAdd(It.Key());
+					GIBakedTextureState.Components = MoveTemp(Info.Components);
+					GIBakedTextureState.CombinedStateId = CombinedStateId;
+				}
+			}
 		}
+		OutdatedComponents += ComponentsNeedToBeBaked;
+	}
+
+	if (OutdatedComponentsCount)
+	{
+		if (OutdatedComponents == 0)
+		{
+			for (auto Component : LandscapeComponents)
+			{
+				const bool bIsDirty = Component->GetPackage()->IsDirty();
+				if (Component->LastBakedTextureMaterialGuid != Component->BakedTextureMaterialGuid && !bIsDirty)
+				{
+					OutdatedComponents++;
+				}
+			}
+		}
+		*OutdatedComponentsCount = OutdatedComponents;
 	}
-	return NumOfComponentsNeedingTextureBaking;
 }
-#endif
 
-void ALandscapeProxy::UpdateBakedTextures(bool bInShouldMarkDirty)
+void ALandscapeProxy::UpdateGIBakedTextureData(bool bInShouldMarkDirty)
+{
+	const bool bBakeAllGITextures = true;
+	UpdateGIBakedTextures(bBakeAllGITextures);
+	if (bInShouldMarkDirty && GetOutdatedGIBakedTextureComponentsCount() > 0)
+	{
+		MarkPackageDirty();
+	}
+}
+
+void ALandscapeProxy::UpdateGIBakedTextures(bool bBakeAllGITextures)
 {
 	// See if we can render
 	UWorld* World = GetWorld();
@@ -3837,27 +3951,22 @@ void ALandscapeProxy::UpdateBakedTextures(bool bInShouldMarkDirty)
 		return;
 	}
 
-	if (bInShouldMarkDirty && GetComponentsNeedingGITextureBaking()>0)
-	{
-		MarkPackageDirty();
-	}
-
-	if (UpdateBakedTexturesCountdown-- > 0)
+	if (!bBakeAllGITextures && UpdateBakedTexturesCountdown-- > 0)
 	{
 		return;
 	}
 
-	// Check if we can want to generate landscape GI data
-	static const auto DistanceFieldCVar = IConsoleManager::Get().FindTConsoleVariableDataInt(TEXT("r.GenerateMeshDistanceFields"));
-	static const auto LandscapeGICVar = IConsoleManager::Get().FindTConsoleVariableDataInt(TEXT("r.GenerateLandscapeGIData"));
-	if (DistanceFieldCVar->GetValueOnGameThread() == 0 || LandscapeGICVar->GetValueOnGameThread() == 0)
+	bool bGenerateLandscapeGIData = true;
+	TMap<UTexture2D*, FGIBakedTextureState> ComponentsToBeBakedByHeightmap;
+	UpdateGIBakedTextureStatus(&bGenerateLandscapeGIData, &ComponentsToBeBakedByHeightmap);
+	
+	if (!bGenerateLandscapeGIData)
 	{
 		// Clear out any existing GI textures
 		for (ULandscapeComponent* Component : LandscapeComponents)
 		{
 			if (Component != nullptr && Component->GIBakedBaseColorTexture != nullptr)
 			{
-				Component->Modify(bInShouldMarkDirty);
 				Component->BakedTextureMaterialGuid.Invalidate();
 				Component->GIBakedBaseColorTexture = nullptr;
 				Component->MarkRenderStateDirty();
@@ -3870,42 +3979,13 @@ void ALandscapeProxy::UpdateBakedTextures(bool bInShouldMarkDirty)
 		return;
 	}
 
-	// Stores the components and their state hash data for a single atlas
-	struct FBakedTextureSourceInfo
-	{
-		// pointer as FMemoryWriter caches the address of the FBufferArchive, and this struct could be relocated on a realloc.
-		TUniquePtr<FBufferArchive> ComponentStateAr;
-		TArray<ULandscapeComponent*> Components;
-
-		FBakedTextureSourceInfo()
-		{
-			ComponentStateAr = MakeUnique<FBufferArchive>();
-		}
-	};
-
-	// Group components by heightmap texture
-	TMap<UTexture2D*, FBakedTextureSourceInfo> ComponentsByHeightmap;
-	for (ULandscapeComponent* Component : LandscapeComponents)
-	{
-		if (Component == nullptr)
-		{
-			continue;
-		}
-
-		FBakedTextureSourceInfo& Info = ComponentsByHeightmap.FindOrAdd(Component->GetHeightmap());
-		Info.Components.Add(Component);
-		Component->SerializeStateHashes(*Info.ComponentStateAr);
-	}
-
 	TotalComponentsNeedingTextureBaking -= NumComponentsNeedingTextureBaking;
 	NumComponentsNeedingTextureBaking = 0;
 	int32 NumGenerated = 0;
 
-
-
-	for (auto It = ComponentsByHeightmap.CreateConstIterator(); It; ++It)
+	for (auto It = ComponentsToBeBakedByHeightmap.CreateConstIterator(); It; ++It)
 	{
-		const FBakedTextureSourceInfo& Info = It.Value();
+		const FGIBakedTextureState& Info = It.Value();
 
 		bool bCanBake = true;
 		for (ULandscapeComponent* Component : Info.Components)
@@ -3936,71 +4016,51 @@ void ALandscapeProxy::UpdateBakedTextures(bool bInShouldMarkDirty)
 
 		if (bCanBake)
 		{
-			// Calculate a combined Guid-like ID we can use for this component
-			uint32 Hash[5];
-			FSHA1::HashBuffer(Info.ComponentStateAr->GetData(), Info.ComponentStateAr->Num(), (uint8*)Hash);
-			FGuid CombinedStateId = FGuid(Hash[0] ^ Hash[4], Hash[1], Hash[2], Hash[3]);
-
-			bool bNeedsBake = false;
-			for (ULandscapeComponent* Component : Info.Components)
+			// We throttle, baking only one atlas per frame if bBakeAllGITextures is false.
+			if (!bBakeAllGITextures && NumGenerated > 0)
 			{
-				if (Component->BakedTextureMaterialGuid != CombinedStateId)
-				{
-					bNeedsBake = true;
-					break;
-				}
+				NumComponentsNeedingTextureBaking += Info.Components.Num();
 			}
-
-			if (bNeedsBake)
+			else
 			{
-				// We throttle, baking only one atlas per frame
-				if (NumGenerated > 0)
-				{
-					NumComponentsNeedingTextureBaking += Info.Components.Num();
-				}
-				else
-				{
-					UTexture2D* HeightmapTexture = It.Key();
-					// 1/8 the res of the heightmap
-					FIntPoint AtlasSize(HeightmapTexture->GetSizeX() >> 3, HeightmapTexture->GetSizeY() >> 3);
+				UTexture2D* HeightmapTexture = It.Key();
+				// 1/8 the res of the heightmap
+				FIntPoint AtlasSize(HeightmapTexture->GetSizeX() >> 3, HeightmapTexture->GetSizeY() >> 3);
 
-					TArray<FColor> AtlasSamples;
-					AtlasSamples.AddZeroed(AtlasSize.X * AtlasSize.Y);
+				TArray<FColor> AtlasSamples;
+				AtlasSamples.AddZeroed(AtlasSize.X * AtlasSize.Y);
 
-					for (ULandscapeComponent* Component : Info.Components)
+				for (ULandscapeComponent* Component : Info.Components)
+				{
+					// not registered; ignore this component
+					if (!Component->SceneProxy)
 					{
-						// not registered; ignore this component
-						if (!Component->SceneProxy)
-						{
-							continue;
-						}
+						continue;
+					}
 
-						int32 ComponentSamples = (SubsectionSizeQuads + 1) * NumSubsections;
-						check(FMath::IsPowerOfTwo(ComponentSamples));
+					int32 ComponentSamples = (SubsectionSizeQuads + 1) * NumSubsections;
+					check(FMath::IsPowerOfTwo(ComponentSamples));
 
-						int32 BakeSize = ComponentSamples >> 3;
-						TArray<FColor> Samples;
-						if (FMaterialUtilities::ExportBaseColor(Component, BakeSize, Samples))
+					int32 BakeSize = ComponentSamples >> 3;
+					TArray<FColor> Samples;
+					if (FMaterialUtilities::ExportBaseColor(Component, BakeSize, Samples))
+					{
+						int32 AtlasOffsetX = FMath::RoundToInt(Component->HeightmapScaleBias.Z * (float)HeightmapTexture->GetSizeX()) >> 3;
+						int32 AtlasOffsetY = FMath::RoundToInt(Component->HeightmapScaleBias.W * (float)HeightmapTexture->GetSizeY()) >> 3;
+						for (int32 y = 0; y < BakeSize; y++)
 						{
-							int32 AtlasOffsetX = FMath::RoundToInt(Component->HeightmapScaleBias.Z * (float)HeightmapTexture->GetSizeX()) >> 3;
-							int32 AtlasOffsetY = FMath::RoundToInt(Component->HeightmapScaleBias.W * (float)HeightmapTexture->GetSizeY()) >> 3;
-							for (int32 y = 0; y < BakeSize; y++)
-							{
-								FMemory::Memcpy(&AtlasSamples[(y + AtlasOffsetY)*AtlasSize.X + AtlasOffsetX], &Samples[y*BakeSize], sizeof(FColor)* BakeSize);
-							}
-							NumGenerated++;
+							FMemory::Memcpy(&AtlasSamples[(y + AtlasOffsetY)*AtlasSize.X + AtlasOffsetX], &Samples[y*BakeSize], sizeof(FColor)* BakeSize);
 						}
+						NumGenerated++;
 					}
-					UTexture2D* AtlasTexture = FMaterialUtilities::CreateTexture(GetOutermost(), HeightmapTexture->GetName() + TEXT("_BaseColor"), AtlasSize, AtlasSamples, TC_Default, TEXTUREGROUP_World, RF_NoFlags, true, CombinedStateId);
-					AtlasTexture->MarkPackageDirty();
+				}
+				UTexture2D* AtlasTexture = FMaterialUtilities::CreateTexture(GetOutermost(), HeightmapTexture->GetName() + TEXT("_BaseColor"), AtlasSize, AtlasSamples, TC_Default, TEXTUREGROUP_World, RF_NoFlags, true, Info.CombinedStateId);
 
-					for (ULandscapeComponent* Component : Info.Components)
-					{
-						Component->Modify(bInShouldMarkDirty);
-						Component->BakedTextureMaterialGuid = CombinedStateId;
-						Component->GIBakedBaseColorTexture = AtlasTexture;
-						Component->MarkRenderStateDirty();
-					}
+				for (ULandscapeComponent* Component : Info.Components)
+				{
+					Component->BakedTextureMaterialGuid = Info.CombinedStateId;
+					Component->GIBakedBaseColorTexture = AtlasTexture;
+					Component->MarkRenderStateDirty();
 				}
 			}
 		}
diff --git a/Engine/Source/Runtime/Landscape/Private/LandscapeSubsystem.cpp b/Engine/Source/Runtime/Landscape/Private/LandscapeSubsystem.cpp
--- a/Engine/Source/Runtime/Landscape/Private/LandscapeSubsystem.cpp
+++ b/Engine/Source/Runtime/Landscape/Private/LandscapeSubsystem.cpp
@@ -51,7 +51,7 @@ void ULandscapeSubsystem::Initialize(FSubsystemCollectionBase& Collection)
 
 #if WITH_EDITOR
 	GrassMapsBuilder = new FLandscapeGrassMapsBuilder(GetWorld());
-	BakedGITextureBuilder = new FLandscapeBakedGITextureBuilder(GetWorld());
+	GIBakedTextureBuilder = new FLandscapeGIBakedTextureBuilder(GetWorld());
 #endif
 }
 
@@ -62,6 +62,10 @@ void ULandscapeSubsystem::Deinitialize()
 	{
 		delete GrassMapsBuilder;
 	}
+	if (GIBakedTextureBuilder)
+	{
+		delete GIBakedTextureBuilder;
+	}
 #endif
 	Proxies.Empty();
 
@@ -132,7 +136,7 @@ void ULandscapeSubsystem::Tick(float DeltaTime)
 			// editor-only
 			if (!World->IsPlayInEditor())
 			{
-				Proxy->UpdateBakedTextures();
+				Proxy->UpdateGIBakedTextures();
 				Proxy->UpdatePhysicalMaterialTasks();
 			}
 		}
@@ -155,7 +159,7 @@ void ULandscapeSubsystem::Tick(float DeltaTime)
 void ULandscapeSubsystem::BuildAll()
 {
 	BuildGrassMaps();
-	BuildGITextures();
+	BuildGIBakedTextures();
 }
 
 void ULandscapeSubsystem::BuildGrassMaps()
@@ -168,14 +172,14 @@ int32 ULandscapeSubsystem::GetOutdatedGrassMapCount()
 	return GrassMapsBuilder->GetOutdatedGrassMapCount(/*bInForceUpdate*/false);
 }
 
-void ULandscapeSubsystem::BuildGITextures()
+void ULandscapeSubsystem::BuildGIBakedTextures()
 {
-	BakedGITextureBuilder->Build();
+	GIBakedTextureBuilder->Build();
 }
 
-int32 ULandscapeSubsystem::GetComponentsNeedingGITextureBaking()
+int32 ULandscapeSubsystem::GetOutdatedGIBakedTextureComponentsCount()
 {
-	return BakedGITextureBuilder->GetComponentsNeedingTextureBaking();
+	return GIBakedTextureBuilder->GetOutdatedGIBakedTextureComponentsCount(/*bInForceUpdate*/false);
 }
 
 bool ULandscapeSubsystem::IsGridBased() const
@@ -220,7 +224,7 @@ void ULandscapeSubsystem::DisplayBuildMessages(FCanvas* Canvas, float& XPos, flo
 		YPos += FontSizeY;
 	}
 
-	if (int32 ComponentsNeedingGITextureBaking = GetComponentsNeedingGITextureBaking())
+	if (int32 ComponentsNeedingGITextureBaking = GetOutdatedGIBakedTextureComponentsCount())
 	{
 		SmallTextItem.SetColor(FLinearColor::Red);
 		SmallTextItem.Text = FText::Format(LOCTEXT("LANDSCAPE_TEXTURES_NEED_TO_BE_REBUILT_FMT", "LANDSCAPE BAKED TEXTURES NEED TO BE REBUILT ({0} {0}|plural(one=object,other=objects))"), ComponentsNeedingGITextureBaking);
diff --git a/Engine/Source/Runtime/Landscape/Public/LandscapeSubsystem.h b/Engine/Source/Runtime/Landscape/Public/LandscapeSubsystem.h
--- a/Engine/Source/Runtime/Landscape/Public/LandscapeSubsystem.h
+++ b/Engine/Source/Runtime/Landscape/Public/LandscapeSubsystem.h
@@ -34,8 +34,8 @@ public:
 	LANDSCAPE_API void BuildAll();
 	LANDSCAPE_API void BuildGrassMaps();
 	LANDSCAPE_API int32 GetOutdatedGrassMapCount();
-	LANDSCAPE_API void BuildGITextures();
-	LANDSCAPE_API int32 GetComponentsNeedingGITextureBaking();
+	LANDSCAPE_API void BuildGIBakedTextures();
+	LANDSCAPE_API int32 GetOutdatedGIBakedTextureComponentsCount();
 	LANDSCAPE_API bool IsGridBased() const;
 	LANDSCAPE_API void ChangeGridSize(ULandscapeInfo* LandscapeInfo, uint32 NewGridSizeInComponents);
 	LANDSCAPE_API ALandscapeProxy* FindOrAddLandscapeProxy(ULandscapeInfo* LandscapeInfo, const FIntPoint& SectionBase);
@@ -52,6 +52,6 @@ private:
 
 #if WITH_EDITOR
 	class FLandscapeGrassMapsBuilder* GrassMapsBuilder;
-	class FLandscapeBakedGITextureBuilder* BakedGITextureBuilder;
+	class FLandscapeGIBakedTextureBuilder* GIBakedTextureBuilder;
 #endif
 };
