Refactor Landscape Physical Material to not automatically dirty levels/actors and instead display a message so that the user can build them through the Build menu

diff --git a/Engine/Source/Editor/LandscapeEditor/Private/LandscapeEditorModule.cpp b/Engine/Source/Editor/LandscapeEditor/Private/LandscapeEditorModule.cpp
--- a/Engine/Source/Editor/LandscapeEditor/Private/LandscapeEditorModule.cpp
+++ b/Engine/Source/Editor/LandscapeEditor/Private/LandscapeEditorModule.cpp
@@ -127,6 +127,9 @@ public:
 			
 			FUIAction ActionBuildGrassMaps(FExecuteAction::CreateStatic(&BuildGrassMaps), FCanExecuteAction());
 			Section.AddMenuEntry(NAME_None, LOCTEXT("BuildGrassMapsOnly","Build Grass Maps Only"), LOCTEXT("BuildLandscapeGrassMaps","Build landscape grass maps"),TAttribute<FSlateIcon>(), ActionBuildGrassMaps, EUserInterfaceActionType::Button);
+
+			FUIAction ActionBuildPhysicalMaterial(FExecuteAction::CreateStatic(&BuildPhysicalMaterial), FCanExecuteAction());
+			Section.AddMenuEntry(NAME_None, LOCTEXT("BuildPhysicalMaterialOnly", "Build Physical Material Only"), LOCTEXT("BuildLandscapePhysicalMaterial", "Build landscape physical material"), TAttribute<FSlateIcon>(), ActionBuildPhysicalMaterial, EUserInterfaceActionType::Button);
 		}
 
 	}
@@ -224,6 +227,17 @@ public:
 		}
 	}
 
+	static void BuildPhysicalMaterial()
+	{
+		if (UWorld* World = GEditor->GetEditorWorldContext().World())
+		{
+			if (ULandscapeSubsystem* LandscapeSubsystem = World->GetSubsystem<ULandscapeSubsystem>())
+			{
+				LandscapeSubsystem->BuildPhysicalMaterial();
+			}
+		}
+	}
+
 	static void ChangeLandscapeViewMode(ELandscapeViewMode::Type ViewMode)
 	{
 		GLandscapeViewMode = ViewMode;
diff --git a/Engine/Source/Runtime/Landscape/Classes/LandscapeComponent.h b/Engine/Source/Runtime/Landscape/Classes/LandscapeComponent.h
--- a/Engine/Source/Runtime/Landscape/Classes/LandscapeComponent.h
+++ b/Engine/Source/Runtime/Landscape/Classes/LandscapeComponent.h
@@ -623,6 +623,10 @@ public:
 	/** Represents hash for last PhysicalMaterialTask */
 	UPROPERTY()
 	uint32 PhysicalMaterialHash;
+
+	/** Represents last saved hash for PhysicalMaterialTask */
+	UPROPERTY(Transient)
+	uint32 LastSavedPhysicalMaterialHash;
 #endif
 
 	/** For mobile */
diff --git a/Engine/Source/Runtime/Landscape/Classes/LandscapeProxy.h b/Engine/Source/Runtime/Landscape/Classes/LandscapeProxy.h
--- a/Engine/Source/Runtime/Landscape/Classes/LandscapeProxy.h
+++ b/Engine/Source/Runtime/Landscape/Classes/LandscapeProxy.h
@@ -946,7 +946,9 @@ public:
 	void UpdateGIBakedTextureStatus(bool* bOutGenerateLandscapeGIData, TMap<UTexture2D*, FGIBakedTextureState>* OutComponentsNeedBakingByHeightmap, int32* OutdatedComponentsCount=nullptr) const;
 
 	/** Update the landscape physical material render tasks */
-	void UpdatePhysicalMaterialTasks();
+	void UpdatePhysicalMaterialTasks(bool bInShouldMarkDirty = false);
+
+	void UpdatePhysicalMaterialTasksStatus(TSet<ULandscapeComponent*>* OutdatedComponents, int32* OutdatedComponentsCount) const;
 
 	/** Frame counter to count down to the next time we check to update baked textures, so we don't check every frame */
 	int32 UpdateBakedTexturesCountdown;
@@ -995,6 +997,8 @@ public:
 	LANDSCAPE_API void BuildGrassMaps(struct FScopedSlowTask* InSlowTask = nullptr);
 	LANDSCAPE_API void BuildGIBakedTextures(struct FScopedSlowTask* InSlowTask = nullptr);
 	LANDSCAPE_API int32 GetOutdatedGIBakedTextureComponentsCount() const;
+	LANDSCAPE_API void BuildPhysicalMaterial(struct FScopedSlowTask* InSlowTask = nullptr);
+	LANDSCAPE_API int32 GetOudatedPhysicalMaterialComponentsCount() const;
 	LANDSCAPE_API virtual void CreateSplineComponent() override;
 	LANDSCAPE_API virtual void CreateSplineComponent(const FVector& Scale3D) override;
 
@@ -1234,4 +1238,19 @@ private:
 	mutable int32 OutdatedGIBakedTextureComponentsCount;
 	mutable double GIBakedTexturesLastCheckTime;
 };
+
+
+/**
+ * Helper class used to Build or monitor Landscape Physical Material
+ */
+class LANDSCAPE_API FLandscapePhysicalMaterialBuilder
+{
+public:
+	FLandscapePhysicalMaterialBuilder(UWorld* InWorld);
+	void Build();
+	int32 GetOudatedPhysicalMaterialComponentsCount();
+private:
+	UWorld* World;
+	int32 OudatedPhysicalMaterialComponentsCount;
+};
 #endif
\ No newline at end of file
diff --git a/Engine/Source/Runtime/Landscape/Private/Landscape.cpp b/Engine/Source/Runtime/Landscape/Private/Landscape.cpp
--- a/Engine/Source/Runtime/Landscape/Private/Landscape.cpp
+++ b/Engine/Source/Runtime/Landscape/Private/Landscape.cpp
@@ -558,8 +558,10 @@ void ULandscapeComponent::Serialize(FArchive& Ar)
 #if WITH_EDITOR
 	if (Ar.IsSaving() && Ar.IsPersistent())
 	{
-		//Update the lastsaved Guid for GI texture
+		//Update the last saved Guid for GI texture
 		LastBakedTextureMaterialGuid = BakedTextureMaterialGuid;
+		//Update the last saved Hash for physical material
+		LastSavedPhysicalMaterialHash = PhysicalMaterialHash;
 	}
 #endif
 }
@@ -805,6 +807,7 @@ void ULandscapeComponent::PostLoad()
 			}
 		}
 		LastBakedTextureMaterialGuid = BakedTextureMaterialGuid;
+		LastSavedPhysicalMaterialHash = PhysicalMaterialHash;
 	}
 #endif
 
@@ -3781,12 +3784,9 @@ void FLandscapeGIBakedTextureBuilder::Build()
 {
 	if (World)
 	{
-		FScopedSlowTask SlowTask(OutdatedGIBakedTextureComponentsCount, (LOCTEXT("BakeTextures_BuildLandscapeBakedTextures", "Baking Landscape Textures")));
-		SlowTask.MakeDialog();
-
 		for (TActorIterator<ALandscapeProxy> ProxyIt(World); ProxyIt; ++ProxyIt)
 		{
-			ProxyIt->BuildGIBakedTextures(&SlowTask);
+			ProxyIt->BuildGIBakedTextures();
 		}
 		//Force update the outdated count when using the build menu option.
 		OutdatedGIBakedTextureComponentsCount = 0;
@@ -4075,12 +4075,99 @@ void ALandscapeProxy::UpdateGIBakedTextures(bool bBakeAllGITextures)
 	}
 }
 
-void ALandscapeProxy::UpdatePhysicalMaterialTasks()
+FLandscapePhysicalMaterialBuilder::FLandscapePhysicalMaterialBuilder(UWorld* InWorld)
+	:World(InWorld)
+	,OudatedPhysicalMaterialComponentsCount(0)
 {
-	for (ULandscapeComponent* Component : LandscapeComponents)
+}
+
+void FLandscapePhysicalMaterialBuilder::Build()
+{
+	if (World)
+	{
+		for (TActorIterator<ALandscapeProxy> ProxyIt(World); ProxyIt; ++ProxyIt)
+		{
+			ProxyIt->BuildPhysicalMaterial();
+		}
+	}
+}
+
+int32 FLandscapePhysicalMaterialBuilder::GetOudatedPhysicalMaterialComponentsCount()
+{
+	if (World)
+	{
+		OudatedPhysicalMaterialComponentsCount = 0;
+		for (TActorIterator<ALandscapeProxy> ProxyIt(World); ProxyIt; ++ProxyIt)
+		{
+			OudatedPhysicalMaterialComponentsCount += ProxyIt->GetOudatedPhysicalMaterialComponentsCount();
+		}
+	}
+	return OudatedPhysicalMaterialComponentsCount;
+}
+
+int32 ALandscapeProxy::GetOudatedPhysicalMaterialComponentsCount() const
+{
+	int32 OudatedPhysicalMaterialComponentsCount = 0;
+	UpdatePhysicalMaterialTasksStatus(nullptr, &OudatedPhysicalMaterialComponentsCount);
+	return OudatedPhysicalMaterialComponentsCount;
+}
+
+void ALandscapeProxy::BuildPhysicalMaterial(struct FScopedSlowTask* InSlowTask)
+{
+	if (!HasAnyFlags(RF_ClassDefaultObject))
+	{
+		const bool bShouldMarkDirty = true;
+		UpdatePhysicalMaterialTasks(bShouldMarkDirty);
+	}
+}
+
+void ALandscapeProxy::UpdatePhysicalMaterialTasksStatus(TSet<ULandscapeComponent*>* OutdatedComponents, int32* OutdatedComponentsCount) const
+{
+	int32 OutdatedCount = 0;
+	for (auto Component : LandscapeComponents)
+	{
+		uint32 Hash = Component->CalculatePhysicalMaterialTaskHash();
+		if (Component->PhysicalMaterialHash != Hash || Component->PhysicalMaterialTask.IsValid())
+		{
+			OutdatedCount++;
+			if (OutdatedComponents)
+			{
+				OutdatedComponents->Add(Component);
+			}
+		}
+	}
+
+	if (OutdatedCount == 0)
+	{
+		for (auto Component : LandscapeComponents)
+		{
+			const bool bIsDirty = Component->GetPackage()->IsDirty();
+			if (Component->LastSavedPhysicalMaterialHash != Component->PhysicalMaterialHash && !bIsDirty)
+			{
+				OutdatedCount++;
+			}
+		}
+	}
+
+	if (OutdatedComponentsCount)
+	{
+		*OutdatedComponentsCount = OutdatedCount;
+	}
+}
+
+void ALandscapeProxy::UpdatePhysicalMaterialTasks(bool bInShouldMarkDirty)
+{
+	TSet<ULandscapeComponent*> OutdatedComponents;
+	int32 PendingComponentsToBeSaved = 0;
+	UpdatePhysicalMaterialTasksStatus(&OutdatedComponents, &PendingComponentsToBeSaved);
+	for (ULandscapeComponent* Component : OutdatedComponents)
 	{
 		Component->UpdatePhysicalMaterialTasks();
 	}
+	if (bInShouldMarkDirty && PendingComponentsToBeSaved >0)
+	{
+		MarkPackageDirty();
+	}
 }
 #endif
 
diff --git a/Engine/Source/Runtime/Landscape/Private/LandscapeSubsystem.cpp b/Engine/Source/Runtime/Landscape/Private/LandscapeSubsystem.cpp
--- a/Engine/Source/Runtime/Landscape/Private/LandscapeSubsystem.cpp
+++ b/Engine/Source/Runtime/Landscape/Private/LandscapeSubsystem.cpp
@@ -52,6 +52,7 @@ void ULandscapeSubsystem::Initialize(FSubsystemCollectionBase& Collection)
 #if WITH_EDITOR
 	GrassMapsBuilder = new FLandscapeGrassMapsBuilder(GetWorld());
 	GIBakedTextureBuilder = new FLandscapeGIBakedTextureBuilder(GetWorld());
+	PhysicalMaterialBuilder = new FLandscapePhysicalMaterialBuilder(GetWorld());
 #endif
 }
 
@@ -66,6 +67,10 @@ void ULandscapeSubsystem::Deinitialize()
 	{
 		delete GIBakedTextureBuilder;
 	}
+	if (PhysicalMaterialBuilder)
+	{
+		delete PhysicalMaterialBuilder;
+	}
 #endif
 	Proxies.Empty();
 
@@ -160,6 +165,7 @@ void ULandscapeSubsystem::BuildAll()
 {
 	BuildGrassMaps();
 	BuildGIBakedTextures();
+	BuildPhysicalMaterial();
 }
 
 void ULandscapeSubsystem::BuildGrassMaps()
@@ -182,6 +188,16 @@ int32 ULandscapeSubsystem::GetOutdatedGIBakedTextureComponentsCount()
 	return GIBakedTextureBuilder->GetOutdatedGIBakedTextureComponentsCount(/*bInForceUpdate*/false);
 }
 
+void ULandscapeSubsystem::BuildPhysicalMaterial()
+{
+	PhysicalMaterialBuilder->Build();
+}
+
+int32 ULandscapeSubsystem::GetOudatedPhysicalMaterialComponentsCount()
+{
+	return PhysicalMaterialBuilder->GetOudatedPhysicalMaterialComponentsCount();
+}
+
 bool ULandscapeSubsystem::IsGridBased() const
 {
 	return UWorld::HasSubsystem<UWorldPartitionSubsystem>(GetWorld());
@@ -219,7 +235,7 @@ void ULandscapeSubsystem::DisplayBuildMessages(FCanvas* Canvas, float& XPos, flo
 	if (int32 OutdatedGrassMapCount = GetOutdatedGrassMapCount())
 	{
 		SmallTextItem.SetColor(FLinearColor::Red);
-		SmallTextItem.Text = FText::Format(LOCTEXT("GRASS_MAPS_NEED_TO_BE_REBUILT_FMT", "GRASS MAPS NEED TO BE REBUILT ({0} {0}|plural(one=object,other=objects))"), OutdatedGrassMapCount);
+		SmallTextItem.Text = FText::Format(LOCTEXT("GRASS_MAPS_NEED_TO_BE_REBUILT_FMT", "GRASS MAPS NEEDS TO BE REBUILT ({0} {0}|plural(one=object,other=objects))"), OutdatedGrassMapCount);
 		Canvas->DrawItem(SmallTextItem, FVector2D(XPos, YPos));
 		YPos += FontSizeY;
 	}
@@ -227,7 +243,15 @@ void ULandscapeSubsystem::DisplayBuildMessages(FCanvas* Canvas, float& XPos, flo
 	if (int32 ComponentsNeedingGITextureBaking = GetOutdatedGIBakedTextureComponentsCount())
 	{
 		SmallTextItem.SetColor(FLinearColor::Red);
-		SmallTextItem.Text = FText::Format(LOCTEXT("LANDSCAPE_TEXTURES_NEED_TO_BE_REBUILT_FMT", "LANDSCAPE BAKED TEXTURES NEED TO BE REBUILT ({0} {0}|plural(one=object,other=objects))"), ComponentsNeedingGITextureBaking);
+		SmallTextItem.Text = FText::Format(LOCTEXT("LANDSCAPE_TEXTURES_NEED_TO_BE_REBUILT_FMT", "LANDSCAPE BAKED TEXTURES NEEDS TO BE REBUILT ({0} {0}|plural(one=object,other=objects))"), ComponentsNeedingGITextureBaking);
+		Canvas->DrawItem(SmallTextItem, FVector2D(XPos, YPos));
+		YPos += FontSizeY;
+	}
+
+	if (int32 ComponentsWithOudatedPhysicalMaterial = GetOudatedPhysicalMaterialComponentsCount())
+	{
+		SmallTextItem.SetColor(FLinearColor::Red);
+		SmallTextItem.Text = FText::Format(LOCTEXT("LANDSCAPE_PHYSICALMATERIAL_NEED_TO_BE_REBUILT_FMT", "LANDSCAPE PHYSICAL MATERIAL NEEDS TO BE REBUILT ({0} {0}|plural(one=object,other=objects))"), ComponentsWithOudatedPhysicalMaterial);
 		Canvas->DrawItem(SmallTextItem, FVector2D(XPos, YPos));
 		YPos += FontSizeY;
 	}
diff --git a/Engine/Source/Runtime/Landscape/Public/LandscapeSubsystem.h b/Engine/Source/Runtime/Landscape/Public/LandscapeSubsystem.h
--- a/Engine/Source/Runtime/Landscape/Public/LandscapeSubsystem.h
+++ b/Engine/Source/Runtime/Landscape/Public/LandscapeSubsystem.h
@@ -36,6 +36,8 @@ public:
 	LANDSCAPE_API int32 GetOutdatedGrassMapCount();
 	LANDSCAPE_API void BuildGIBakedTextures();
 	LANDSCAPE_API int32 GetOutdatedGIBakedTextureComponentsCount();
+	LANDSCAPE_API void BuildPhysicalMaterial();
+	LANDSCAPE_API int32 GetOudatedPhysicalMaterialComponentsCount();
 	LANDSCAPE_API bool IsGridBased() const;
 	LANDSCAPE_API void ChangeGridSize(ULandscapeInfo* LandscapeInfo, uint32 NewGridSizeInComponents);
 	LANDSCAPE_API ALandscapeProxy* FindOrAddLandscapeProxy(ULandscapeInfo* LandscapeInfo, const FIntPoint& SectionBase);
@@ -53,5 +55,6 @@ private:
 #if WITH_EDITOR
 	class FLandscapeGrassMapsBuilder* GrassMapsBuilder;
 	class FLandscapeGIBakedTextureBuilder* GIBakedTextureBuilder;
+	class FLandscapePhysicalMaterialBuilder* PhysicalMaterialBuilder;
 #endif
 };
