Ported landscape utility functions to RDG.

diff --git a/Engine/Source/Runtime/Landscape/Private/LandscapeGrass.cpp b/Engine/Source/Runtime/Landscape/Private/LandscapeGrass.cpp
--- a/Engine/Source/Runtime/Landscape/Private/LandscapeGrass.cpp
+++ b/Engine/Source/Runtime/Landscape/Private/LandscapeGrass.cpp
@@ -361,9 +361,6 @@ FLandscapeGrassWeightMeshProcessor::FLandscapeGrassWeightMeshProcessor(const FSc
 {
 	PassDrawRenderState.SetBlendState(TStaticBlendState<>::GetRHI());
 	PassDrawRenderState.SetDepthStencilState(TStaticDepthStencilState<false, CF_Always>::GetRHI());
-
-	PassDrawRenderState.SetViewUniformBuffer(InViewIfDynamicMeshCommand->ViewUniformBuffer);
-	PassDrawRenderState.SetPassUniformBuffer(nullptr);
 }
 
 void FLandscapeGrassWeightMeshProcessor::AddMeshBatch(const FMeshBatch& RESTRICT MeshBatch,
@@ -491,6 +488,11 @@ public:
 	FMatrix ViewRotationMatrix;
 	FMatrix ProjectionMatrix;
 
+	BEGIN_SHADER_PARAMETER_STRUCT(FLandscapeGrassPassParameters, )
+		SHADER_PARAMETER_STRUCT_REF(FViewUniformShaderParameters, View)
+		RENDER_TARGET_BINDING_SLOTS()
+	END_SHADER_PARAMETER_STRUCT()
+
 	void RenderLandscapeComponentToTexture_RenderThread(FRHICommandListImmediate& RHICmdList)
 	{
 		FSceneViewFamilyContext ViewFamily(FSceneViewFamily::ConstructionValues(nullptr, nullptr, FEngineShowFlags(ESFIM_Game)).SetWorldTimes(FApp::GetCurrentTime() - GStartTime, FApp::GetDeltaTime(), FApp::GetCurrentTime() - GStartTime));
@@ -508,16 +510,25 @@ public:
 		
 		const FSceneView* View = ViewFamily.Views[0];
 
-		FRHIRenderPassInfo RPInfo(RenderTargetResource->GetTextureRHI(), ERenderTargetActions::Clear_Store);
-		RHICmdList.BeginRenderPass(RPInfo, TEXT("LandscapeGrass"));
-		RHICmdList.SetViewport(View->UnscaledViewRect.Min.X, View->UnscaledViewRect.Min.Y, 0.0f, View->UnscaledViewRect.Max.X, View->UnscaledViewRect.Max.Y, 1.0f);
-		RHICmdList.SetScissorRect(false, 0, 0, 0, 0);
+		FMemMark Mark(FMemStack::Get());
+		FRDGBuilder GraphBuilder(RHICmdList);
+
+		FRDGTextureRef OutputTexture = GraphBuilder.RegisterExternalTexture(CreateRenderTarget(RenderTargetResource->GetTextureRHI(), TEXT("LandscapeGrass")));
 
+		auto* PassParameters = GraphBuilder.AllocParameters<FLandscapeGrassPassParameters>();
+		PassParameters->View = View->ViewUniformBuffer;
+		PassParameters->RenderTargets[0] = FRenderTargetBinding(OutputTexture, ERenderTargetLoadAction::EClear);
+
+		GraphBuilder.AddPass(
+			RDG_EVENT_NAME("LandscapeGrass"),
+			PassParameters,
+			ERDGPassFlags::Raster,
+			[this, &View](FRHICommandList& RHICmdList)
 		{
-			FMemMark Mark(FMemStack::Get());
+			RHICmdList.SetViewport(View->UnscaledViewRect.Min.X, View->UnscaledViewRect.Min.Y, 0.0f, View->UnscaledViewRect.Max.X, View->UnscaledViewRect.Max.Y, 1.0f);
+			RHICmdList.SetScissorRect(false, 0, 0, 0, 0);
 
-			DrawDynamicMeshPass(*View, RHICmdList,
-				[View, PassOffsetX = PassOffsetX, &ComponentInfos = ComponentInfos, NumPasses = NumPasses, FirstHeightMipsPassIndex = FirstHeightMipsPassIndex, HeightMips = HeightMips](FDynamicPassMeshDrawListContext* DynamicMeshPassContext)
+			DrawDynamicMeshPass(*View, RHICmdList, [&](FDynamicPassMeshDrawListContext* DynamicMeshPassContext)
 			{
 				FLandscapeGrassWeightMeshProcessor PassMeshProcessor(
 					nullptr,
@@ -537,9 +548,9 @@ public:
 					}
 				}
 			});
-		}
+		});
 
-		RHICmdList.EndRenderPass();
+		GraphBuilder.Execute();
 	}
 };
 
diff --git a/Engine/Source/Runtime/Landscape/Private/LandscapePhysicalMaterial.cpp b/Engine/Source/Runtime/Landscape/Private/LandscapePhysicalMaterial.cpp
--- a/Engine/Source/Runtime/Landscape/Private/LandscapePhysicalMaterial.cpp
+++ b/Engine/Source/Runtime/Landscape/Private/LandscapePhysicalMaterial.cpp
@@ -154,9 +154,6 @@ FLandscapePhysicalMaterialMeshProcessor::FLandscapePhysicalMaterialMeshProcessor
 {
 	PassDrawRenderState.SetBlendState(TStaticBlendState<>::GetRHI());
 	PassDrawRenderState.SetDepthStencilState(TStaticDepthStencilState<false, CF_Always>::GetRHI());
-
-	PassDrawRenderState.SetViewUniformBuffer(InViewIfDynamicMeshCommand->ViewUniformBuffer);
-	PassDrawRenderState.SetPassUniformBuffer(nullptr);
 }
 
 void FLandscapePhysicalMaterialMeshProcessor::AddMeshBatch(
@@ -238,6 +235,11 @@ namespace
 		OutMeshInfos[MeshInfoIndex].MeshBatch = &((FLandscapeComponentSceneProxy*)InSceneProxy)->GetGrassMeshBatch();
 		OutMeshInfos[MeshInfoIndex].MeshBatchElementMask = 1 << 0; // LOD 0 only
 	}
+	
+	BEGIN_SHADER_PARAMETER_STRUCT(FLandscapePhysicalMaterialPassParameters, )
+		SHADER_PARAMETER_STRUCT_REF(FViewUniformShaderParameters, View)
+		RENDER_TARGET_BINDING_SLOTS()
+	END_SHADER_PARAMETER_STRUCT()
 
 	/** Render the landscape physical material IDs and copy to the read back texture. */
 	void Render_RenderThread(
@@ -247,15 +249,9 @@ namespace
 		FVector ViewOrigin,
 		FMatrix ViewRotationMatrix,
 		FMatrix ProjectionMatrix,
-		FTextureRHIRef ReadbackTexture,
-		FGPUFenceRHIRef ReadbackFence)
+		FTextureRHIRef ReadbackTextureRHI,
+		FGPUFenceRHIRef ReadbackFenceRHI)
 	{
-		// Allocate temporary render target
-		TRefCountPtr<IPooledRenderTarget> PooledRenderTarget;
-		FPooledRenderTargetDesc Desc(FPooledRenderTargetDesc::Create2DDesc(TargetSize, PF_G8, FClearValueBinding::Black, TexCreate_None, TexCreate_RenderTargetable | TexCreate_ShaderResource, false));
-		GRenderTargetPool.FindFreeElement(RHICmdList, Desc, PooledRenderTarget, TEXT("LandscapePhysicalMaterialTarget"), ERenderTargetTransience::Transient);
-		FTextureRHIRef RenderTargetTexture = PooledRenderTarget->GetRenderTargetItem().TargetableTexture;
-
 		// Create the view
 		FSceneViewFamily::ConstructionValues ViewFamilyInit(nullptr, nullptr, FEngineShowFlags(ESFIM_Game));
 		ViewFamilyInit.SetWorldTimes(0.0f, 0.0f, 0.0f);
@@ -272,48 +268,54 @@ namespace
 		GetRendererModule().CreateAndInitSingleView(RHICmdList, &ViewFamily, &ViewInitOptions);
 		const FSceneView* View = ViewFamily.Views[0];
 
-		// Draw
-		{
-			SCOPED_DRAW_EVENTF(RHICmdList, LandscapePhysicalMaterial, TEXT("LandscapePhysicalMaterial"));
+		FMemMark Mark(FMemStack::Get());
+		FRDGBuilder GraphBuilder(RHICmdList);
+
+		FRDGTextureRef OutputTexture = GraphBuilder.CreateTexture(
+			FRDGTextureDesc::Create2D(TargetSize, PF_G8, FClearValueBinding::Black, TexCreate_RenderTargetable | TexCreate_ShaderResource),
+			TEXT("LandscapePhysicalMaterialTarget"));
 
-			FRHIRenderPassInfo RPInfo(RenderTargetTexture, ERenderTargetActions::Clear_Store);
-			RHICmdList.BeginRenderPass(RPInfo, TEXT("LandscapePhysicalMaterial"));
+		auto* PassParameters = GraphBuilder.AllocParameters<FLandscapePhysicalMaterialPassParameters>();
+		PassParameters->View = View->ViewUniformBuffer;
+		PassParameters->RenderTargets[0] = FRenderTargetBinding(OutputTexture, ERenderTargetLoadAction::EClear);
+
+		GraphBuilder.AddPass(
+			RDG_EVENT_NAME("LandscapePhysicalMaterial"),
+			PassParameters,
+			ERDGPassFlags::Raster,
+			[&View, &MeshInfos](FRHICommandList& RHICmdList)
+		{
 			RHICmdList.SetViewport(View->UnscaledViewRect.Min.X, View->UnscaledViewRect.Min.Y, 0.0f, View->UnscaledViewRect.Max.X, View->UnscaledViewRect.Max.Y, 1.0f);
 			RHICmdList.SetScissorRect(false, 0, 0, 0, 0);
 
+			DrawDynamicMeshPass(*View, RHICmdList, [&](FDynamicPassMeshDrawListContext* DynamicMeshPassContext)
 			{
-				SCOPED_GPU_STAT(RHICmdList, LandscapePhysicalMaterial_Draw);
+				FLandscapePhysicalMaterialMeshProcessor PassMeshProcessor(
+					nullptr,
+					View,
+					DynamicMeshPassContext);
 
-				FMemMark Mark(FMemStack::Get());
-
-				DrawDynamicMeshPass(*View, RHICmdList,
-					[View, &MeshInfos = MeshInfos](FDynamicPassMeshDrawListContext* DynamicMeshPassContext)
+				for (auto& MeshInfo : MeshInfos)
+				{
+					const FMeshBatch& Mesh = *MeshInfo.MeshBatch;
+					if (Mesh.MaterialRenderProxy != nullptr)
 					{
-						FLandscapePhysicalMaterialMeshProcessor PassMeshProcessor(
-							nullptr,
-							View,
-							DynamicMeshPassContext);
+						Mesh.MaterialRenderProxy->UpdateUniformExpressionCacheIfNeeded(View->GetFeatureLevel());
+						PassMeshProcessor.AddMeshBatch(Mesh, MeshInfo.MeshBatchElementMask, MeshInfo.Proxy);
+					}
+				}
+			});
+		});
 
-						for (auto& MeshInfo : MeshInfos)
-						{
-							const FMeshBatch& Mesh = *MeshInfo.MeshBatch;
-							if (Mesh.MaterialRenderProxy != nullptr)
-							{
-								Mesh.MaterialRenderProxy->UpdateUniformExpressionCacheIfNeeded(View->GetFeatureLevel());
-								PassMeshProcessor.AddMeshBatch(Mesh, MeshInfo.MeshBatchElementMask, MeshInfo.Proxy);
-							}
-						}
-					});
-			}
+		FRDGTextureRef ReadbackTexture = GraphBuilder.RegisterExternalTexture(CreateRenderTarget(ReadbackTextureRHI, TEXT("ReadbackTexture")));
+		AddCopyToResolveTargetPass(GraphBuilder, OutputTexture, ReadbackTexture, FResolveParams{});
 
-			RHICmdList.EndRenderPass();
+		AddPass(GraphBuilder, [&ReadbackFenceRHI](FRHICommandList& RHICmdList)
+		{
+			RHICmdList.WriteGPUFence(ReadbackFenceRHI);
+		});
 
-			// Copy to the read back texture
-			RHICmdList.Transition(FRHITransitionInfo(RenderTargetTexture, ERHIAccess::RTV, ERHIAccess::CopySrc));
-			RHICmdList.Transition(FRHITransitionInfo(ReadbackTexture, ERHIAccess::Unknown, ERHIAccess::CopyDest));
-			RHICmdList.CopyToResolveTarget(RenderTargetTexture, ReadbackTexture, FResolveParams());
-			RHICmdList.WriteGPUFence(ReadbackFence);
-		}
+		GraphBuilder.Execute();
 	}
 
 	/** Fetch the landscape physical material IDs from a read back texture. */
