Added IsLoaded function to actor descriptors handle types and updated existing code to use it.

diff --git a/Engine/Source/Runtime/Engine/Public/WorldPartition/WorldPartitionHandle.h b/Engine/Source/Runtime/Engine/Public/WorldPartition/WorldPartitionHandle.h
--- a/Engine/Source/Runtime/Engine/Public/WorldPartition/WorldPartitionHandle.h
+++ b/Engine/Source/Runtime/Engine/Public/WorldPartition/WorldPartitionHandle.h
@@ -188,6 +188,11 @@ public:
 		return ActorDesc && ActorDesc->IsValid();
 	}
 
+	FORCEINLINE bool IsLoaded() const
+	{
+		return IsValid() && (*ActorDesc)->GetActor();
+	}
+
 	FORCEINLINE FWorldPartitionActorDesc* Get() const
 	{
 		return IsValid() ? ActorDesc->Get() : nullptr;
diff --git a/Engine/Source/Runtime/Landscape/Private/LandscapeEdit.cpp b/Engine/Source/Runtime/Landscape/Private/LandscapeEdit.cpp
--- a/Engine/Source/Runtime/Landscape/Private/LandscapeEdit.cpp
+++ b/Engine/Source/Runtime/Landscape/Private/LandscapeEdit.cpp
@@ -2259,7 +2259,7 @@ bool ULandscapeInfo::HasUnloadedComponentsInRegion(int32 X1, int32 Y1, int32 X2,
 			UActorPartitionSubsystem::FCellCoord ActorCoord(LandscapeActorDesc->GridIndexX, LandscapeActorDesc->GridIndexY, LandscapeActorDesc->GridIndexZ, World->PersistentLevel);
 			if (ActorCoord.X >= MinCoord.X && ActorCoord.Y >= MinCoord.Y && ActorCoord.X <= MaxCoord.X && ActorCoord.Y <= MaxCoord.Y)
 			{
-				if (!LandscapeActorDesc->GetHardRefCount() || !LandscapeActorDesc->GetActor())
+				if (!Handle.IsLoaded())
 				{
 					return true;
 				}
