Fix crash when unregistering landscape neighbors that do not have a scene proxy.  #rb juan.canada, Yujiang.Wang

diff --git a/Engine/Source/Runtime/Landscape/Private/LandscapeRender.cpp b/Engine/Source/Runtime/Landscape/Private/LandscapeRender.cpp
--- a/Engine/Source/Runtime/Landscape/Private/LandscapeRender.cpp
+++ b/Engine/Source/Runtime/Landscape/Private/LandscapeRender.cpp
@@ -4532,17 +4532,16 @@ void FLandscapeNeighborInfo::UnregisterNeighbors(FLandscapeComponentSceneProxy*
 			{
 				FLandscapeRenderSystem& RenderSystem = *LandscapeRenderSystems.FindChecked(LandscapeKey);
 				RenderSystem.UnregisterEntity(SceneProxy);
-
-				if (RenderSystem.NumRegisteredEntities == 0)
-				{
-					FLandscapeRenderSystem* RenderSystemPtr = LandscapeRenderSystems.FindChecked(LandscapeKey);
-					delete RenderSystemPtr;
-					LandscapeRenderSystems.Remove(LandscapeKey);
-				}
 			}
 
 			if (SceneProxyMap->Num() == 0)
 			{
+				FLandscapeRenderSystem* RenderSystemPtr = LandscapeRenderSystems.FindChecked(LandscapeKey);
+				check(RenderSystemPtr->NumRegisteredEntities == 0);
+
+				delete RenderSystemPtr;
+				LandscapeRenderSystems.Remove(LandscapeKey);
+
 				// remove the entire LandscapeKey entry as this is the last scene proxy
 				SharedSceneProxyMap.Remove(LandscapeKey);
 			}
