Add progress when importing/reimporting landscape heightmap/weightmaps

diff --git a/Engine/Source/Editor/LandscapeEditor/Private/LandscapeEdMode.cpp b/Engine/Source/Editor/LandscapeEditor/Private/LandscapeEdMode.cpp
--- a/Engine/Source/Editor/LandscapeEditor/Private/LandscapeEdMode.cpp
+++ b/Engine/Source/Editor/LandscapeEditor/Private/LandscapeEdMode.cpp
@@ -3008,10 +3008,14 @@ void FEdModeLandscape::ImportData(const FLandscapeTargetListInfo& TargetInfo, co
 
 		if (TargetInfo.TargetType == ELandscapeToolTargetType::Heightmap)
 		{
+			FScopedSlowTask Progress(1, LOCTEXT("ImportingLandscapeHeightmapTask", "Importing Landscape Heightmap..."));
+			Progress.MakeDialog();
 			ImportHeightData(LandscapeInfo, GetCurrentLayerGuid(), Filename, ImportExtent, ELandscapeImportTransformType::ExpandCentered);
 		}
 		else
 		{
+			FScopedSlowTask Progress(1, LOCTEXT("ImportingLandscapeWeightmapTask", "Importing Landscape Layer Weightmap..."));
+			Progress.MakeDialog();
 			ImportWeightData(LandscapeInfo, GetCurrentLayerGuid(), TargetInfo.LayerInfoObj.Get(), Filename, ImportExtent, ELandscapeImportTransformType::ExpandCentered);
 		}
 	}
diff --git a/Engine/Source/Runtime/Landscape/Private/LandscapeEdit.cpp b/Engine/Source/Runtime/Landscape/Private/LandscapeEdit.cpp
--- a/Engine/Source/Runtime/Landscape/Private/LandscapeEdit.cpp
+++ b/Engine/Source/Runtime/Landscape/Private/LandscapeEdit.cpp
@@ -44,6 +44,7 @@ LandscapeEdit.cpp: Landscape editing
 #include "LandscapeSubsystem.h"
 #include "StaticMeshAttributes.h"
 #include "MeshUtilitiesCommon.h"
+#include "Misc/ScopedSlowTask.h"
 
 #include "EngineModule.h"
 #include "EngineUtils.h"
@@ -3647,7 +3648,8 @@ void ULandscapeInfo::ExportHeightmap(const FString& Filename)
 
 void ULandscapeInfo::ExportHeightmap(const FString& Filename, const FIntRect& ExportRegion)
 {
-	GWarn->BeginSlowTask(LOCTEXT("BeginExportingLandscapeHeightmapTask", "Exporting Landscape Heightmap"), true);
+	FScopedSlowTask Progress(1, LOCTEXT("ExportingLandscapeHeightmapTask", "Exporting Landscape Heightmap..."));
+	Progress.MakeDialog();
 
 	ILandscapeEditorModule& LandscapeEditorModule = FModuleManager::GetModuleChecked<ILandscapeEditorModule>("LandscapeEditor");
 	FLandscapeEditDataInterface LandscapeEdit(this);
@@ -3663,8 +3665,6 @@ void ULandscapeInfo::ExportHeightmap(const FString& Filename, const FIntRect& Ex
 	{
 		HeightmapFormat->Export(*Filename, NAME_None, HeightData, {(uint32)ExportWidth, (uint32)ExportHeight}, DrawScale * FVector(1, 1, LANDSCAPE_ZSCALE));
 	}
-
-	GWarn->EndSlowTask();
 }
 
 void ULandscapeInfo::ExportLayer(ULandscapeLayerInfoObject* LayerInfo, const FString& Filename)
@@ -3680,10 +3680,11 @@ void ULandscapeInfo::ExportLayer(ULandscapeLayerInfoObject* LayerInfo, const FSt
 
 void ULandscapeInfo::ExportLayer(ULandscapeLayerInfoObject* LayerInfo, const FString& Filename, const FIntRect& ExportRegion)
 {
-	check(LayerInfo);
-		
-	GWarn->BeginSlowTask(LOCTEXT("BeginExportingLandscapeWeightmapTask", "Exporting Landscape Layer Weightmap"), true);
+	FScopedSlowTask Progress(1, LOCTEXT("ExportingLandscapeWeightmapTask", "Exporting Landscape Layer Weightmap..."));
+	Progress.MakeDialog();
 
+	check(LayerInfo);
+	
 	ILandscapeEditorModule& LandscapeEditorModule = FModuleManager::GetModuleChecked<ILandscapeEditorModule>("LandscapeEditor");
 
 	TArray<uint8> WeightData;
@@ -3699,8 +3700,6 @@ void ULandscapeInfo::ExportLayer(ULandscapeLayerInfoObject* LayerInfo, const FSt
 	{
 		WeightmapFormat->Export(*Filename, LayerInfo->LayerName, WeightData, { (uint32)ExportWidth, (uint32)ExportHeight }, DrawScale * FVector(1, 1, LANDSCAPE_ZSCALE));
 	}
-
-	GWarn->EndSlowTask();
 }
 
 void ULandscapeInfo::DeleteLayer(ULandscapeLayerInfoObject* LayerInfo, const FName& LayerName)
