Fix crash: missing Modify call (which waits on async texture compiles)

diff --git a/Engine/Source/Runtime/Landscape/Private/LandscapeEditLayers.cpp b/Engine/Source/Runtime/Landscape/Private/LandscapeEditLayers.cpp
--- a/Engine/Source/Runtime/Landscape/Private/LandscapeEditLayers.cpp
+++ b/Engine/Source/Runtime/Landscape/Private/LandscapeEditLayers.cpp
@@ -1713,6 +1713,8 @@ void ALandscape::CopyOldDataToDefaultLayer(ALandscapeProxy* InProxy)
 				TArray64<uint8> ExistingMip0Data;
 				ComponentHeightmap->Source.GetMipData(ExistingMip0Data, 0);
 
+				// Calling modify here makes sure that async texture compilation finishes so we can Lock the mip
+				DefaultLayerHeightmap->Modify();
 				FColor* Mip0Data = (FColor*)DefaultLayerHeightmap->Source.LockMip(0);
 				FMemory::Memcpy(Mip0Data, ExistingMip0Data.GetData(), ExistingMip0Data.Num());
 				DefaultLayerHeightmap->Source.UnlockMip(0);
