Make landscape proxies grid placement property editable.

diff --git a/Engine/Source/Runtime/Landscape/Classes/LandscapeProxy.h b/Engine/Source/Runtime/Landscape/Classes/LandscapeProxy.h
--- a/Engine/Source/Runtime/Landscape/Classes/LandscapeProxy.h
+++ b/Engine/Source/Runtime/Landscape/Classes/LandscapeProxy.h
@@ -860,6 +860,7 @@ public:
 	virtual void PostEditMove(bool bFinished) override;
 	virtual bool ShouldImport(FString* ActorPropString, bool IsMovingLevel) override;
 	virtual bool ShouldExport() override;
+	virtual EActorGridPlacement GetDefaultGridPlacement() const override { return EActorGridPlacement::None; }
 	//~ End AActor Interface
 	virtual uint32 GetDefaultGridSize(UWorld* InWorld) const override { return 1; }
 	virtual FGuid GetGridGuid() const override { return LandscapeGuid; }
diff --git a/Engine/Source/Runtime/Landscape/Private/Landscape.cpp b/Engine/Source/Runtime/Landscape/Private/Landscape.cpp
--- a/Engine/Source/Runtime/Landscape/Private/Landscape.cpp
+++ b/Engine/Source/Runtime/Landscape/Private/Landscape.cpp
@@ -1146,6 +1146,8 @@ ALandscapeProxy::ALandscapeProxy(const FObjectInitializer& ObjectInitializer)
 #if WITH_EDITORONLY_DATA
 	LandscapeProxies.Add(this);
 #endif
+
+	GridPlacement = EActorGridPlacement::Location;
 }
 
 #if WITH_EDITORONLY_DATA
