Trivial changes from GPU instance culling work

diff --git a/Engine/Source/Runtime/Engine/Private/HierarchicalInstancedStaticMesh.cpp b/Engine/Source/Runtime/Engine/Private/HierarchicalInstancedStaticMesh.cpp
--- a/Engine/Source/Runtime/Engine/Private/HierarchicalInstancedStaticMesh.cpp
+++ b/Engine/Source/Runtime/Engine/Private/HierarchicalInstancedStaticMesh.cpp
@@ -35,6 +35,7 @@
 #if WITH_EDITOR
 #include "Rendering/StaticLightingSystemInterface.h"
 #endif
+#include "PrimitiveSceneInfo.h"
 
 #if WITH_EDITOR
 static float GDebugBuildTreeAsyncDelayInSeconds = 0.f;
diff --git a/Engine/Source/Runtime/Engine/Private/InstancedStaticMesh.cpp b/Engine/Source/Runtime/Engine/Private/InstancedStaticMesh.cpp
--- a/Engine/Source/Runtime/Engine/Private/InstancedStaticMesh.cpp
+++ b/Engine/Source/Runtime/Engine/Private/InstancedStaticMesh.cpp
@@ -3285,6 +3285,9 @@ void FInstancedStaticMeshVertexFactoryShaderParameters::GetElementShaderBindings
 	}
 	if (InstanceOffsetValue > 0 && VertexStreams.Num() > 0)
 	{
+		// GPUCULL_TODO: This here can still work together with the instance attributes for index, but note that all instance attributes then must assume they are offset wrt the on-the-fly generate buffer
+		//          so with the new scheme there is no clear way this can work in the vanilla instancing way as there is an indirection. So either other attributes must be loaded in the shader or they
+		//          would have to be copied as the instance ID is now - not good.
 		VertexFactory->OffsetInstanceStreams(InstanceOffsetValue, InputStreamType, VertexStreams);
 	}
 
diff --git a/Engine/Source/Runtime/Landscape/Private/LandscapeGrass.cpp b/Engine/Source/Runtime/Landscape/Private/LandscapeGrass.cpp
--- a/Engine/Source/Runtime/Landscape/Private/LandscapeGrass.cpp
+++ b/Engine/Source/Runtime/Landscape/Private/LandscapeGrass.cpp
@@ -174,6 +174,12 @@ static FAutoConsoleVariableRef CVarGrassMaxCreatePerFrame(
 	GGrassMaxCreatePerFrame,
 	TEXT("Maximum number of Grass components to create per frame"));
 
+static int32 GGrassUpdateAllOnRebuild = 0;
+static FAutoConsoleVariableRef CVarUpdateAllOnRebuild(
+	TEXT("grass.UpdateAllOnRebuild"),
+	GGrassUpdateAllOnRebuild,
+	TEXT(""));
+
 DECLARE_CYCLE_STAT(TEXT("Grass Async Build Time"), STAT_FoliageGrassAsyncBuildTime, STATGROUP_Foliage);
 DECLARE_CYCLE_STAT(TEXT("Grass Start Comp"), STAT_FoliageGrassStartComp, STATGROUP_Foliage);
 DECLARE_CYCLE_STAT(TEXT("Grass End Comp"), STAT_FoliageGrassEndComp, STATGROUP_Foliage);
@@ -439,6 +445,8 @@ void FLandscapeGrassWeightMeshProcessor::Process(
 // data also accessible by render thread
 class FLandscapeGrassWeightExporter_RenderThread
 {
+	FSceneInterface* SceneInterface = nullptr;
+
 	FLandscapeGrassWeightExporter_RenderThread(int32 InNumGrassMaps, bool InbNeedsHeightmap, TArray<int32> InHeightMips)
 		: RenderTargetResource(nullptr)
 		, NumPasses(0)
@@ -578,6 +586,7 @@ public:
 		, RenderTargetTexture(nullptr)
 	{
 		check(InLandscapeComponents.Num() > 0);
+		SceneInterface = InLandscapeComponents[0]->GetScene();
 
 		// todo: use a 2d target?
 		TargetSize = FIntPoint(ComponentSizeVerts * NumPasses * InLandscapeComponents.Num(), ComponentSizeVerts);
@@ -2303,7 +2312,7 @@ void ALandscapeProxy::UpdateGrassDataStatus(TSet<UTexture2D*>* OutCurrentForcedS
 
 			if (bHasGrassTypes || bBakeMaterialPositionOffsetIntoCollision)
 			{
-				if (Component->IsGrassMapOutdated() || !Component->GrassData->HasData())
+				if (Component->IsGrassMapOutdated() || !Component->GrassData->HasData() || GGrassUpdateAllOnRebuild != 0)
 				{
 					if (OutComponentsNeedingGrassMapRender)
 					{
