Expose contact shadow flag for landscape

diff --git a/Engine/Source/Runtime/Landscape/Classes/LandscapeProxy.h b/Engine/Source/Runtime/Landscape/Classes/LandscapeProxy.h
--- a/Engine/Source/Runtime/Landscape/Classes/LandscapeProxy.h
+++ b/Engine/Source/Runtime/Landscape/Classes/LandscapeProxy.h
@@ -625,6 +625,10 @@ public:
 	UPROPERTY(EditAnywhere, BlueprintReadOnly, Category = Lighting, AdvancedDisplay, meta = (EditCondition = "CastShadow", DisplayName = "Static Shadow"))
 	uint8 bCastStaticShadow : 1;
 
+	/** Whether the object should cast contact shadows. This flag is only used if CastShadow is true. */
+	UPROPERTY(EditAnywhere, BlueprintReadOnly, Category = Lighting, AdvancedDisplay, meta = (EditCondition = "CastShadow", DisplayName = "Contact Shadow"))
+	uint8 bCastContactShadow : 1;
+
 	/** When enabled, the component will be rendering into the far shadow cascades(only for directional lights).  This flag is only used if CastShadow is true. */
 	UPROPERTY(EditAnywhere, AdvancedDisplay, BlueprintReadOnly, Category = Lighting, meta = (EditCondition = "CastShadow", DisplayName = "Far Shadow"))
 	uint32 bCastFarShadow : 1;
diff --git a/Engine/Source/Runtime/Landscape/Private/Landscape.cpp b/Engine/Source/Runtime/Landscape/Private/Landscape.cpp
--- a/Engine/Source/Runtime/Landscape/Private/Landscape.cpp
+++ b/Engine/Source/Runtime/Landscape/Private/Landscape.cpp
@@ -723,6 +723,7 @@ void ULandscapeComponent::UpdatedSharedPropertiesFromActor()
 	CastShadow = LandscapeProxy->CastShadow;
 	bCastDynamicShadow = LandscapeProxy->bCastDynamicShadow;
 	bCastStaticShadow = LandscapeProxy->bCastStaticShadow;
+	bCastContactShadow = LandscapeProxy->bCastContactShadow;
 	bCastFarShadow = LandscapeProxy->bCastFarShadow;
 	bCastHiddenShadow = LandscapeProxy->bCastHiddenShadow;
 	bCastShadowAsTwoSided = LandscapeProxy->bCastShadowAsTwoSided;
@@ -1069,6 +1070,7 @@ ALandscapeProxy::ALandscapeProxy(const FObjectInitializer& ObjectInitializer)
 	CastShadow = true;
 	bCastDynamicShadow = true;
 	bCastStaticShadow = true;
+	bCastContactShadow = true;
 	bCastFarShadow = true;
 	bCastHiddenShadow = false;
 	bCastShadowAsTwoSided = false;
@@ -2571,6 +2573,7 @@ void ALandscapeProxy::GetSharedProperties(ALandscapeProxy* Landscape)
 		CastShadow = Landscape->CastShadow;
 		bCastDynamicShadow = Landscape->bCastDynamicShadow;
 		bCastStaticShadow = Landscape->bCastStaticShadow;
+		bCastContactShadow = Landscape->bCastContactShadow;
 		bCastFarShadow = Landscape->bCastFarShadow;
 		bCastHiddenShadow = Landscape->bCastHiddenShadow;
 		bCastShadowAsTwoSided = Landscape->bCastShadowAsTwoSided;
diff --git a/Engine/Source/Runtime/Landscape/Private/LandscapeEdit.cpp b/Engine/Source/Runtime/Landscape/Private/LandscapeEdit.cpp
--- a/Engine/Source/Runtime/Landscape/Private/LandscapeEdit.cpp
+++ b/Engine/Source/Runtime/Landscape/Private/LandscapeEdit.cpp
@@ -4845,6 +4845,7 @@ void ALandscapeProxy::PostEditChangeProperty(FPropertyChangedEvent& PropertyChan
 		PropertyName == FName(TEXT("CastShadow")) ||
 		PropertyName == FName(TEXT("bCastDynamicShadow")) ||
 		PropertyName == FName(TEXT("bCastStaticShadow")) ||
+		PropertyName == FName(TEXT("bCastContactShadow")) ||
 		PropertyName == FName(TEXT("bCastFarShadow")) ||
 		PropertyName == FName(TEXT("bCastHiddenShadow")) ||
 		PropertyName == FName(TEXT("bCastShadowAsTwoSided")) ||
