Avoid landscape grass update to reset bForceMiplevelsToBeResident when landscape uses edit layers (for texture streaming to finish)

diff --git a/Engine/Source/Runtime/Landscape/Private/LandscapeGrass.cpp b/Engine/Source/Runtime/Landscape/Private/LandscapeGrass.cpp
--- a/Engine/Source/Runtime/Landscape/Private/LandscapeGrass.cpp
+++ b/Engine/Source/Runtime/Landscape/Private/LandscapeGrass.cpp
@@ -2426,10 +2426,14 @@ void ALandscapeProxy::UpdateGrassData(bool bInShouldMarkDirty, FScopedSlowTask*
 	{
 		MarkPackageDirty();
 	}
-		
-	for (UTexture2D* TextureToStream : CurrentForcedStreamedTextures)
+
+	// In Edit Layers, bForceMiplevelsToBeResident needs to remain true
+	if (!HasLayersContent())
 	{
-		TextureToStream->bForceMiplevelsToBeResident = false;
+		for (UTexture2D* TextureToStream : CurrentForcedStreamedTextures)
+		{
+			TextureToStream->bForceMiplevelsToBeResident = false;
+		}
 	}
 }
 #endif
@@ -2956,11 +2960,15 @@ void ALandscapeProxy::UpdateGrass(const TArray<FVector>& Cameras, int32& InOutNu
 				{
 					Texture->bForceMiplevelsToBeResident = true;
 				}
-				for (auto Texture : CurrentForcedStreamedTextures.Difference(DesiredForceStreamedTextures))
+
+				// In Edit Layers, bForceMiplevelsToBeResident needs to remain true
+				if (!HasLayersContent())
 				{
-					Texture->bForceMiplevelsToBeResident = false;
+					for (auto Texture : CurrentForcedStreamedTextures.Difference(DesiredForceStreamedTextures))
+					{
+						Texture->bForceMiplevelsToBeResident = false;
+					}
 				}
-
 			}
 #endif
 		}
