Fixed crash when adding a second Landscape in partitioned world.

diff --git a/Engine/Source/Runtime/Engine/Public/WorldPartition/Landscape/LandscapeSplineActorDesc.h b/Engine/Source/Runtime/Engine/Public/WorldPartition/Landscape/LandscapeSplineActorDesc.h
--- a/Engine/Source/Runtime/Engine/Public/WorldPartition/Landscape/LandscapeSplineActorDesc.h
+++ b/Engine/Source/Runtime/Engine/Public/WorldPartition/Landscape/LandscapeSplineActorDesc.h
@@ -12,12 +12,13 @@
 class ENGINE_API FLandscapeSplineActorDesc : public FWorldPartitionActorDesc
 {
 #if WITH_EDITOR
+public:
+	FGuid LandscapeGuid;
+
 protected:
 	virtual void Init(const AActor* InActor) override;
 	virtual void Serialize(FArchive& Ar) override;
 	virtual void OnRegister() override;
 	virtual void OnUnregister() override;
-
-	FGuid LandscapeGuid;
 #endif
 };
diff --git a/Engine/Source/Runtime/Landscape/Private/Landscape.cpp b/Engine/Source/Runtime/Landscape/Private/Landscape.cpp
--- a/Engine/Source/Runtime/Landscape/Private/Landscape.cpp
+++ b/Engine/Source/Runtime/Landscape/Private/Landscape.cpp
@@ -72,6 +72,8 @@ Landscape.cpp: Terrain rendering
 #include "WorldPartition/WorldPartition.h"
 #include "WorldPartition/WorldPartitionHandle.h"
 #include "WorldPartition/WorldPartitionActorDescIterator.h"
+#include "WorldPartition/Landscape/LandscapeActorDesc.h"
+#include "WorldPartition/Landscape/LandscapeSplineActorDesc.h"
 #if WITH_EDITOR
 #include "Rendering/StaticLightingSystemInterface.h"
 #include "Misc/ScopedSlowTask.h"
@@ -3128,15 +3130,21 @@ void ULandscapeInfo::Initialize(UWorld* InWorld, const FGuid& InLandscapeGuid)
 	{
 		for (TWorldPartitionActorDescIterator<AActor, FWorldPartitionActorDesc> ActorDescIterator(WorldPartition); ActorDescIterator; ++ActorDescIterator)
 		{
-			if (ActorDescIterator->GetActorClass()->IsChildOf(ALandscapeStreamingProxy::StaticClass()))
+			FWorldPartitionActorDesc* ActorDesc = *ActorDescIterator;
+			if (ActorDesc->GetActorClass()->IsChildOf(ALandscapeStreamingProxy::StaticClass()))
 			{
-				ProxyHandles.Add(FWorldPartitionHandle(WorldPartition, ActorDescIterator->GetGuid()));
+				if (((FLandscapeActorDesc*)ActorDesc)->GridGuid == InLandscapeGuid)
+				{
+					ProxyHandles.Add(FWorldPartitionHandle(WorldPartition, ActorDesc->GetGuid()));
+				}
 			}
 			else if (ActorDescIterator->GetActorClass()->IsChildOf(ALandscapeSplineActor::StaticClass()))
 			{
-				SplineHandles.Emplace(FWorldPartitionHandle(WorldPartition, ActorDescIterator->GetGuid()));
+				if (((FLandscapeSplineActorDesc*)ActorDesc)->LandscapeGuid == InLandscapeGuid)
+				{
+					SplineHandles.Emplace(FWorldPartitionHandle(WorldPartition, ActorDesc->GetGuid()));
+				}
 			}
-
 		}
 	}
 #endif
