Automated wrapped object pointer upgrade for Engine + ShooterGame + ShooterGame referenced plugins

diff --git a/Engine/Source/Runtime/Experimental/GeometryCollectionEngine/Public/GeometryCollection/GeometryCollectionDebugDrawComponent.h b/Engine/Source/Runtime/Experimental/GeometryCollectionEngine/Public/GeometryCollection/GeometryCollectionDebugDrawComponent.h
--- a/Engine/Source/Runtime/Experimental/GeometryCollectionEngine/Public/GeometryCollection/GeometryCollectionDebugDrawComponent.h
+++ b/Engine/Source/Runtime/Experimental/GeometryCollectionEngine/Public/GeometryCollection/GeometryCollectionDebugDrawComponent.h
@@ -38,11 +38,11 @@ public:
 
 	/** Singleton actor, containing the debug draw properties. Automatically populated at play time unless explicitly set. */
 	UPROPERTY(EditAnywhere, Category = "Debug Draw", AdvancedDisplay)
-	AGeometryCollectionDebugDrawActor* GeometryCollectionDebugDrawActor;
+	TObjectPtr<AGeometryCollectionDebugDrawActor> GeometryCollectionDebugDrawActor;
 
 	/** Level Set singleton actor, containing the Render properties. Automatically populated at play time unless explicitly set. */
 	UPROPERTY(EditAnywhere, Category = "Debug Draw", AdvancedDisplay)
-	AGeometryCollectionRenderLevelSetActor* GeometryCollectionRenderLevelSetActor;
+	TObjectPtr<AGeometryCollectionRenderLevelSetActor> GeometryCollectionRenderLevelSetActor;
 
 	UGeometryCollectionComponent* GeometryCollectionComponent;  // the component we are debug rendering for, set by the GeometryCollectionActor after creation
 
diff --git a/Engine/Source/Runtime/Experimental/GeometryCollectionEngine/Public/GeometryCollection/GeometryCollectionObject.h b/Engine/Source/Runtime/Experimental/GeometryCollectionEngine/Public/GeometryCollection/GeometryCollectionObject.h
--- a/Engine/Source/Runtime/Experimental/GeometryCollectionEngine/Public/GeometryCollection/GeometryCollectionObject.h
+++ b/Engine/Source/Runtime/Experimental/GeometryCollectionEngine/Public/GeometryCollection/GeometryCollectionObject.h
@@ -30,7 +30,7 @@ struct GEOMETRYCOLLECTIONENGINE_API FGeometryCollectionSource
 	FTransform LocalTransform;
 
 	UPROPERTY(BlueprintReadWrite, EditAnywhere, Category = "GeometrySource")
-	TArray<UMaterialInterface*> SourceMaterial;
+	TArray<TObjectPtr<UMaterialInterface>> SourceMaterial;
 };
 
 
@@ -242,7 +242,7 @@ public:
 	TArray<FGeometryCollectionSource> GeometrySource;
 	
 	UPROPERTY(EditAnywhere, Category = "Materials")
-	TArray<UMaterialInterface*> Materials;
+	TArray<TObjectPtr<UMaterialInterface>> Materials;
 
 	/**
 	* Enable support for Nanite.
@@ -335,7 +335,7 @@ public:
 	* Materials relating to remove on fracture
 	*/
 	UPROPERTY(EditAnywhere, Category = "Fracture")
-	TArray<UMaterialInterface*> RemoveOnFractureMaterials;
+	TArray<TObjectPtr<UMaterialInterface>> RemoveOnFractureMaterials;
 
 	FORCEINLINE const int32 GetBoneSelectedMaterialIndex() const { return BoneSelectedMaterialIndex; }
 
@@ -345,7 +345,7 @@ public:
 #if WITH_EDITORONLY_DATA
 	/** Information for thumbnail rendering */
 	UPROPERTY(VisibleAnywhere, Instanced, AdvancedDisplay, Category = GeometryCollection)
-	class UThumbnailInfo* ThumbnailInfo;
+	TObjectPtr<class UThumbnailInfo> ThumbnailInfo;
 #endif // WITH_EDITORONLY_DATA
 
 private:
diff --git a/Engine/Source/Runtime/Experimental/GeometryCollectionEngine/Public/GeometryCollection/GeometryCollectionRenderLevelSetActor.h b/Engine/Source/Runtime/Experimental/GeometryCollectionEngine/Public/GeometryCollection/GeometryCollectionRenderLevelSetActor.h
--- a/Engine/Source/Runtime/Experimental/GeometryCollectionEngine/Public/GeometryCollection/GeometryCollectionRenderLevelSetActor.h
+++ b/Engine/Source/Runtime/Experimental/GeometryCollectionEngine/Public/GeometryCollection/GeometryCollectionRenderLevelSetActor.h
@@ -34,12 +34,12 @@ public:
 
 	// Volume texture to fill	
 	UPROPERTY(EditAnywhere, BlueprintReadWrite, Category = "Volume")
-	UVolumeTexture *TargetVolumeTexture;
+	TObjectPtr<UVolumeTexture> TargetVolumeTexture;
 
 	// Material that performs ray marching.  Note this must have certain parameters in order
 	// to work correctly
 	UPROPERTY(EditAnywhere, BlueprintReadWrite, Category = "Rendering")
-	UMaterial *RayMarchMaterial;
+	TObjectPtr<UMaterial> RayMarchMaterial;
 
 	// Surface tolerance used for rendering.  When surface reconstruction is noisy,
 	// try tweaking this value
diff --git a/Engine/Source/Runtime/Experimental/GeometryCollectionEngine/Public/GeometryCollection/SkeletalMeshSimulationComponent.h b/Engine/Source/Runtime/Experimental/GeometryCollectionEngine/Public/GeometryCollection/SkeletalMeshSimulationComponent.h
--- a/Engine/Source/Runtime/Experimental/GeometryCollectionEngine/Public/GeometryCollection/SkeletalMeshSimulationComponent.h
+++ b/Engine/Source/Runtime/Experimental/GeometryCollectionEngine/Public/GeometryCollection/SkeletalMeshSimulationComponent.h
@@ -36,14 +36,14 @@ public:
 
 	/** Physical Properties */
 	UPROPERTY(EditAnywhere, BlueprintReadOnly, Category = "ChaosPhysics")
-	const UChaosPhysicalMaterial* PhysicalMaterial;
+	TObjectPtr<const UChaosPhysicalMaterial> PhysicalMaterial;
 
 	/** Chaos RBD Solver */
 	UPROPERTY(EditAnywhere, Category = "ChaosPhysics", meta = (DisplayName = "Chaos Solver"))
-	AChaosSolverActor* ChaosSolverActor;
+	TObjectPtr<AChaosSolverActor> ChaosSolverActor;
 
 	UPROPERTY(EditAnywhere, Category = "ChaosPhysics")
-	UPhysicsAsset* OverridePhysicsAsset;
+	TObjectPtr<UPhysicsAsset> OverridePhysicsAsset;
 
 	//
 	// ChaosPhysics | General
diff --git a/Engine/Source/Runtime/Experimental/GeometryCollectionEngine/Public/GeometryCollection/StaticMeshSimulationComponent.h b/Engine/Source/Runtime/Experimental/GeometryCollectionEngine/Public/GeometryCollection/StaticMeshSimulationComponent.h
--- a/Engine/Source/Runtime/Experimental/GeometryCollectionEngine/Public/GeometryCollection/StaticMeshSimulationComponent.h
+++ b/Engine/Source/Runtime/Experimental/GeometryCollectionEngine/Public/GeometryCollection/StaticMeshSimulationComponent.h
@@ -87,11 +87,11 @@ public:
 	* Physical Properties
 	*/
 	UPROPERTY(EditAnywhere, BlueprintReadOnly, Category = "ChaosPhysics")
-	const UChaosPhysicalMaterial* PhysicalMaterial;
+	TObjectPtr<const UChaosPhysicalMaterial> PhysicalMaterial;
 
 	/** Chaos RBD Solver */
 	UPROPERTY(EditAnywhere, Category = "ChaosPhysics", meta = (DisplayName = "Chaos Solver"))
-	AChaosSolverActor* ChaosSolverActor;
+	TObjectPtr<AChaosSolverActor> ChaosSolverActor;
 
 	const TSharedPtr<FPhysScene_Chaos> GetPhysicsScene() const;
 
@@ -116,7 +116,7 @@ private :
 
 	/** List of component for which this simulation component created a physics object. Parallel array to PhysicsProxy, so PhysicsProxies[i] corresponds to SimulatedComponents[i] */
 	UPROPERTY()
-	TArray<UPrimitiveComponent*> SimulatedComponents;
+	TArray<TObjectPtr<UPrimitiveComponent>> SimulatedComponents;
 
 
 	//@todo(mlentine): Don't have one per static mesh
diff --git a/Engine/Source/Runtime/Experimental/InteractiveToolsFramework/Public/BaseGizmos/AxisAngleGizmo.h b/Engine/Source/Runtime/Experimental/InteractiveToolsFramework/Public/BaseGizmos/AxisAngleGizmo.h
--- a/Engine/Source/Runtime/Experimental/InteractiveToolsFramework/Public/BaseGizmos/AxisAngleGizmo.h
+++ b/Engine/Source/Runtime/Experimental/InteractiveToolsFramework/Public/BaseGizmos/AxisAngleGizmo.h
@@ -67,7 +67,7 @@ public:
 
 	/** The mouse click behavior of the gizmo is accessible so that it can be modified to use different mouse keys. */
 	UPROPERTY()
-	UClickDragInputBehavior* MouseBehavior;
+	TObjectPtr<UClickDragInputBehavior> MouseBehavior;
 
 	/** 
 	 * This gets checked to see if we should use the custom ray caster to get a destination point for the gizmo, rather
diff --git a/Engine/Source/Runtime/Experimental/InteractiveToolsFramework/Public/BaseGizmos/AxisPositionGizmo.h b/Engine/Source/Runtime/Experimental/InteractiveToolsFramework/Public/BaseGizmos/AxisPositionGizmo.h
--- a/Engine/Source/Runtime/Experimental/InteractiveToolsFramework/Public/BaseGizmos/AxisPositionGizmo.h
+++ b/Engine/Source/Runtime/Experimental/InteractiveToolsFramework/Public/BaseGizmos/AxisPositionGizmo.h
@@ -82,7 +82,7 @@ public:
 
 	/** The mouse click behavior of the gizmo is accessible so that it can be modified to use different mouse keys. */
 	UPROPERTY()
-	UClickDragInputBehavior* MouseBehavior;
+	TObjectPtr<UClickDragInputBehavior> MouseBehavior;
 
 public:
 	/** If enabled, then the sign on the parameter delta is always "increasing" when moving away from the origin point, rather than just being a projection onto the axis */
diff --git a/Engine/Source/Runtime/Experimental/InteractiveToolsFramework/Public/BaseGizmos/AxisSources.h b/Engine/Source/Runtime/Experimental/InteractiveToolsFramework/Public/BaseGizmos/AxisSources.h
--- a/Engine/Source/Runtime/Experimental/InteractiveToolsFramework/Public/BaseGizmos/AxisSources.h
+++ b/Engine/Source/Runtime/Experimental/InteractiveToolsFramework/Public/BaseGizmos/AxisSources.h
@@ -145,7 +145,7 @@ public:
 	virtual void GetTangentVectors(FVector& TangentXOut, FVector& TangentYOut) const;
 
 	UPROPERTY()
-	USceneComponent* Component;
+	TObjectPtr<USceneComponent> Component;
 
 	/** Clamped internally to 0,1,2 */
 	UPROPERTY()
diff --git a/Engine/Source/Runtime/Experimental/InteractiveToolsFramework/Public/BaseGizmos/BrushStampIndicator.h b/Engine/Source/Runtime/Experimental/InteractiveToolsFramework/Public/BaseGizmos/BrushStampIndicator.h
--- a/Engine/Source/Runtime/Experimental/InteractiveToolsFramework/Public/BaseGizmos/BrushStampIndicator.h
+++ b/Engine/Source/Runtime/Experimental/InteractiveToolsFramework/Public/BaseGizmos/BrushStampIndicator.h
@@ -96,7 +96,7 @@ public:
 	 * Optional Component that will be transformed such that it tracks the Radius/Position/Normal
 	 */
 	UPROPERTY()
-	UPrimitiveComponent* AttachedComponent;
+	TObjectPtr<UPrimitiveComponent> AttachedComponent;
 
 protected:
 	UPrimitiveComponent* ScaleInitializedComponent = nullptr;		// we are just using this as a key, never calling any functions on it
diff --git a/Engine/Source/Runtime/Experimental/InteractiveToolsFramework/Public/BaseGizmos/HitTargets.h b/Engine/Source/Runtime/Experimental/InteractiveToolsFramework/Public/BaseGizmos/HitTargets.h
--- a/Engine/Source/Runtime/Experimental/InteractiveToolsFramework/Public/BaseGizmos/HitTargets.h
+++ b/Engine/Source/Runtime/Experimental/InteractiveToolsFramework/Public/BaseGizmos/HitTargets.h
@@ -60,7 +60,7 @@ public:
 	 * Component->LineTraceComponent() is called to determine if the target is hit
 	 */
 	UPROPERTY()
-	UPrimitiveComponent* Component;
+	TObjectPtr<UPrimitiveComponent> Component;
 
 	/**
 	 * If set, this condition is checked before performing the hit test. This gives a way 
diff --git a/Engine/Source/Runtime/Experimental/InteractiveToolsFramework/Public/BaseGizmos/IntervalGizmo.h b/Engine/Source/Runtime/Experimental/InteractiveToolsFramework/Public/BaseGizmos/IntervalGizmo.h
--- a/Engine/Source/Runtime/Experimental/InteractiveToolsFramework/Public/BaseGizmos/IntervalGizmo.h
+++ b/Engine/Source/Runtime/Experimental/InteractiveToolsFramework/Public/BaseGizmos/IntervalGizmo.h
@@ -41,13 +41,13 @@ public:
 public:
 
 	UPROPERTY()
-	UGizmoLineHandleComponent* UpIntervalComponent = nullptr;
+	TObjectPtr<UGizmoLineHandleComponent> UpIntervalComponent = nullptr;
 
 	UPROPERTY()
-	UGizmoLineHandleComponent* DownIntervalComponent = nullptr;
+	TObjectPtr<UGizmoLineHandleComponent> DownIntervalComponent = nullptr;
 
 	UPROPERTY()
-	UGizmoLineHandleComponent* ForwardIntervalComponent = nullptr;
+	TObjectPtr<UGizmoLineHandleComponent> ForwardIntervalComponent = nullptr;
 
 
 public:
@@ -159,7 +159,7 @@ public:
 
 	/** State target is shared across gizmos, and created internally during SetActiveTarget() */
 	UPROPERTY()
-	UGizmoTransformChangeStateTarget* StateTarget = nullptr;
+	TObjectPtr<UGizmoTransformChangeStateTarget> StateTarget = nullptr;
 
 
 protected:
@@ -171,15 +171,15 @@ protected:
 
 	/** The active target object for the Gizmo */
 	UPROPERTY()
-	UTransformProxy* TransformProxy = nullptr;
+	TObjectPtr<UTransformProxy> TransformProxy = nullptr;
 
 	/** list of current-active child components */
 	UPROPERTY()
-	TArray<UPrimitiveComponent*> ActiveComponents;
+	TArray<TObjectPtr<UPrimitiveComponent>> ActiveComponents;
 
 	/** list of currently-active child gizmos */
 	UPROPERTY()
-	TArray<UInteractiveGizmo*> ActiveGizmos;
+	TArray<TObjectPtr<UInteractiveGizmo>> ActiveGizmos;
 
 	UGizmoLocalFloatParameterSource* UpIntervalSource = nullptr;
 	UGizmoLocalFloatParameterSource* DownIntervalSource = nullptr;
@@ -188,11 +188,11 @@ protected:
 	
 	/** Y-axis source is shared across Gizmos, and created internally during SetActiveTarget() */
 	UPROPERTY()
-	UGizmoComponentAxisSource* AxisYSource = nullptr;
+	TObjectPtr<UGizmoComponentAxisSource> AxisYSource = nullptr;
 
 	/** Z-axis source is shared across Gizmos, and created internally during SetActiveTarget() */
 	UPROPERTY()
-	UGizmoComponentAxisSource* AxisZSource = nullptr;
+	TObjectPtr<UGizmoComponentAxisSource> AxisZSource = nullptr;
 
 protected:
 	TSharedPtr<FIntervalGizmoActorFactory> GizmoActorBuilder;
diff --git a/Engine/Source/Runtime/Experimental/InteractiveToolsFramework/Public/BaseGizmos/RepositionableTransformGizmo.h b/Engine/Source/Runtime/Experimental/InteractiveToolsFramework/Public/BaseGizmos/RepositionableTransformGizmo.h
--- a/Engine/Source/Runtime/Experimental/InteractiveToolsFramework/Public/BaseGizmos/RepositionableTransformGizmo.h
+++ b/Engine/Source/Runtime/Experimental/InteractiveToolsFramework/Public/BaseGizmos/RepositionableTransformGizmo.h
@@ -66,7 +66,7 @@ protected:
 	TArray<UInteractiveGizmo*> PivotAlignmentGizmos;
 
 	UPROPERTY()
-	UGizmoTransformChangeStateTarget* RepositionStateTarget;
+	TObjectPtr<UGizmoTransformChangeStateTarget> RepositionStateTarget;
 
 	// Helper functions
 	void ModifyPivotAxisGizmo(UInteractiveGizmo* SubGizmoIn);
diff --git a/Engine/Source/Runtime/Experimental/InteractiveToolsFramework/Public/BaseGizmos/ScalableSphereGizmo.h b/Engine/Source/Runtime/Experimental/InteractiveToolsFramework/Public/BaseGizmos/ScalableSphereGizmo.h
--- a/Engine/Source/Runtime/Experimental/InteractiveToolsFramework/Public/BaseGizmos/ScalableSphereGizmo.h
+++ b/Engine/Source/Runtime/Experimental/InteractiveToolsFramework/Public/BaseGizmos/ScalableSphereGizmo.h
@@ -84,7 +84,7 @@ private:
 	bool bIsDragging{ false };
 
 	UPROPERTY()
-	UTransformProxy* ActiveTarget;
+	TObjectPtr<UTransformProxy> ActiveTarget;
 
 	// The current axis that is being dragged along
 	UPROPERTY()
diff --git a/Engine/Source/Runtime/Experimental/InteractiveToolsFramework/Public/BaseGizmos/TransformGizmo.h b/Engine/Source/Runtime/Experimental/InteractiveToolsFramework/Public/BaseGizmos/TransformGizmo.h
--- a/Engine/Source/Runtime/Experimental/InteractiveToolsFramework/Public/BaseGizmos/TransformGizmo.h
+++ b/Engine/Source/Runtime/Experimental/InteractiveToolsFramework/Public/BaseGizmos/TransformGizmo.h
@@ -50,28 +50,28 @@ public:
 
 	/** X Axis Translation Component */
 	UPROPERTY()
-	UPrimitiveComponent* TranslateX;
+	TObjectPtr<UPrimitiveComponent> TranslateX;
 
 	/** Y Axis Translation Component */
 	UPROPERTY()
-	UPrimitiveComponent* TranslateY;
+	TObjectPtr<UPrimitiveComponent> TranslateY;
 
 	/** Z Axis Translation Component */
 	UPROPERTY()
-	UPrimitiveComponent* TranslateZ;
+	TObjectPtr<UPrimitiveComponent> TranslateZ;
 
 
 	/** YZ Plane Translation Component */
 	UPROPERTY()
-	UPrimitiveComponent* TranslateYZ;
+	TObjectPtr<UPrimitiveComponent> TranslateYZ;
 
 	/** XZ Plane Translation Component */
 	UPROPERTY()
-	UPrimitiveComponent* TranslateXZ;
+	TObjectPtr<UPrimitiveComponent> TranslateXZ;
 
 	/** XY Plane Translation Component */
 	UPROPERTY()
-	UPrimitiveComponent* TranslateXY;
+	TObjectPtr<UPrimitiveComponent> TranslateXY;
 
 	//
 	// Rotation Components
@@ -79,15 +79,15 @@ public:
 
 	/** X Axis Rotation Component */
 	UPROPERTY()
-	UPrimitiveComponent* RotateX;
+	TObjectPtr<UPrimitiveComponent> RotateX;
 
 	/** Y Axis Rotation Component */
 	UPROPERTY()
-	UPrimitiveComponent* RotateY;
+	TObjectPtr<UPrimitiveComponent> RotateY;
 
 	/** Z Axis Rotation Component */
 	UPROPERTY()
-	UPrimitiveComponent* RotateZ;
+	TObjectPtr<UPrimitiveComponent> RotateZ;
 
 	//
 	// Scaling Components
@@ -95,33 +95,33 @@ public:
 
 	/** Uniform Scale Component */
 	UPROPERTY()
-	UPrimitiveComponent* UniformScale;
+	TObjectPtr<UPrimitiveComponent> UniformScale;
 
 
 	/** X Axis Scale Component */
 	UPROPERTY()
-	UPrimitiveComponent* AxisScaleX;
+	TObjectPtr<UPrimitiveComponent> AxisScaleX;
 
 	/** Y Axis Scale Component */
 	UPROPERTY()
-	UPrimitiveComponent* AxisScaleY;
+	TObjectPtr<UPrimitiveComponent> AxisScaleY;
 
 	/** Z Axis Scale Component */
 	UPROPERTY()
-	UPrimitiveComponent* AxisScaleZ;
+	TObjectPtr<UPrimitiveComponent> AxisScaleZ;
 
 
 	/** YZ Plane Scale Component */
 	UPROPERTY()
-	UPrimitiveComponent* PlaneScaleYZ;
+	TObjectPtr<UPrimitiveComponent> PlaneScaleYZ;
 
 	/** XZ Plane Scale Component */
 	UPROPERTY()
-	UPrimitiveComponent* PlaneScaleXZ;
+	TObjectPtr<UPrimitiveComponent> PlaneScaleXZ;
 
 	/** XY Plane Scale Component */
 	UPROPERTY()
-	UPrimitiveComponent* PlaneScaleXY;
+	TObjectPtr<UPrimitiveComponent> PlaneScaleXY;
 
 
 
@@ -298,7 +298,7 @@ public:
 
 	/** The active target object for the Gizmo */
 	UPROPERTY()
-	UTransformProxy* ActiveTarget;
+	TObjectPtr<UTransformProxy> ActiveTarget;
 
 	/**
 	 * @return the internal GizmoActor used by the Gizmo
@@ -386,7 +386,7 @@ protected:
 
 	/** List of current-active child components */
 	UPROPERTY()
-	TArray<UPrimitiveComponent*> ActiveComponents;
+	TArray<TObjectPtr<UPrimitiveComponent>> ActiveComponents;
 
 	/** 
 	 * List of nonuniform scale components. Subset of of ActiveComponents. These are tracked separately so they can
@@ -394,11 +394,11 @@ protected:
 	 * on Components
 	 */
 	UPROPERTY()
-	TArray<UPrimitiveComponent*> NonuniformScaleComponents;
+	TArray<TObjectPtr<UPrimitiveComponent>> NonuniformScaleComponents;
 
 	/** list of currently-active child gizmos */
 	UPROPERTY()
-	TArray<UInteractiveGizmo*> ActiveGizmos;
+	TArray<TObjectPtr<UInteractiveGizmo>> ActiveGizmos;
 
 	/** GizmoActors will be spawned in this World */
 	UWorld* World;
@@ -413,7 +413,7 @@ protected:
 
 	/** Axis that points towards camera, X/Y plane tangents aligned to right/up. Shared across Gizmos, and created internally during SetActiveTarget() */
 	UPROPERTY()
-	UGizmoConstantFrameAxisSource* CameraAxisSource;
+	TObjectPtr<UGizmoConstantFrameAxisSource> CameraAxisSource;
 
 	// internal function that updates CameraAxisSource by getting current view state from GizmoManager
 	void UpdateCameraAxisSource();
@@ -421,15 +421,15 @@ protected:
 
 	/** X-axis source is shared across Gizmos, and created internally during SetActiveTarget() */
 	UPROPERTY()
-	UGizmoComponentAxisSource* AxisXSource;
+	TObjectPtr<UGizmoComponentAxisSource> AxisXSource;
 
 	/** Y-axis source is shared across Gizmos, and created internally during SetActiveTarget() */
 	UPROPERTY()
-	UGizmoComponentAxisSource* AxisYSource;
+	TObjectPtr<UGizmoComponentAxisSource> AxisYSource;
 
 	/** Z-axis source is shared across Gizmos, and created internally during SetActiveTarget() */
 	UPROPERTY()
-	UGizmoComponentAxisSource* AxisZSource;
+	TObjectPtr<UGizmoComponentAxisSource> AxisZSource;
 
 	//
 	// Scaling support. 
@@ -438,15 +438,15 @@ protected:
 
 	/** Local X-axis source (ie 1,0,0) is shared across Scale Gizmos, and created internally during SetActiveTarget() */
 	UPROPERTY()
-	UGizmoComponentAxisSource* UnitAxisXSource;
+	TObjectPtr<UGizmoComponentAxisSource> UnitAxisXSource;
 
 	/** Y-axis source (ie 0,1,0) is shared across Scale Gizmos, and created internally during SetActiveTarget() */
 	UPROPERTY()
-	UGizmoComponentAxisSource* UnitAxisYSource;
+	TObjectPtr<UGizmoComponentAxisSource> UnitAxisYSource;
 
 	/** Z-axis source (ie 0,0,1) is shared across Scale Gizmos, and created internally during SetActiveTarget() */
 	UPROPERTY()
-	UGizmoComponentAxisSource* UnitAxisZSource;
+	TObjectPtr<UGizmoComponentAxisSource> UnitAxisZSource;
 
 
 	//
@@ -460,7 +460,7 @@ protected:
 	 * itself (IToolCommandChangeSource implementation above is called)
 	 */
 	UPROPERTY()
-	UGizmoTransformChangeStateTarget* StateTarget;
+	TObjectPtr<UGizmoTransformChangeStateTarget> StateTarget;
 
 	/**
 	 * These are used to let the translation subgizmos use raycasts into the scene to align the gizmo with scene geometry.
diff --git a/Engine/Source/Runtime/Experimental/InteractiveToolsFramework/Public/BaseGizmos/TransformSources.h b/Engine/Source/Runtime/Experimental/InteractiveToolsFramework/Public/BaseGizmos/TransformSources.h
--- a/Engine/Source/Runtime/Experimental/InteractiveToolsFramework/Public/BaseGizmos/TransformSources.h
+++ b/Engine/Source/Runtime/Experimental/InteractiveToolsFramework/Public/BaseGizmos/TransformSources.h
@@ -51,7 +51,7 @@ public:
 	virtual void SetTransform(const FTransform& NewTransform) override;
 
 	UPROPERTY()
-	USceneComponent* Component;
+	TObjectPtr<USceneComponent> Component;
 
 	/** If true, Component->Modify() is called on SetTransform */
 	UPROPERTY()
@@ -169,7 +169,7 @@ public:
 	virtual void SetTransform(const FTransform& NewTransform) override;
 
 	UPROPERTY()
-	UTransformProxy* Proxy;
+	TObjectPtr<UTransformProxy> Proxy;
 
 	/**
 	 * The underlying proxy is modified with its SetPivotMode flag set to this. Thus, the transform source
diff --git a/Engine/Source/Runtime/Experimental/InteractiveToolsFramework/Public/BaseTools/BaseBrushTool.h b/Engine/Source/Runtime/Experimental/InteractiveToolsFramework/Public/BaseTools/BaseBrushTool.h
--- a/Engine/Source/Runtime/Experimental/InteractiveToolsFramework/Public/BaseTools/BaseBrushTool.h
+++ b/Engine/Source/Runtime/Experimental/InteractiveToolsFramework/Public/BaseTools/BaseBrushTool.h
@@ -111,7 +111,7 @@ public:
 
 	/** Properties that control brush size/etc*/
 	UPROPERTY()
-	UBrushBaseProperties* BrushProperties;
+	TObjectPtr<UBrushBaseProperties> BrushProperties;
 
 	/** Set to true by Tool if user is currently in an active brush stroke*/
 	UPROPERTY()
@@ -160,7 +160,7 @@ protected:
 protected:
 
 	UPROPERTY()
-	UBrushStampIndicator* BrushStampIndicator;
+	TObjectPtr<UBrushStampIndicator> BrushStampIndicator;
 
 	virtual void SetupBrushStampIndicator();
 	virtual void UpdateBrushStampIndicator();
diff --git a/Engine/Source/Runtime/Experimental/InteractiveToolsFramework/Public/InputBehaviorSet.h b/Engine/Source/Runtime/Experimental/InteractiveToolsFramework/Public/InputBehaviorSet.h
--- a/Engine/Source/Runtime/Experimental/InteractiveToolsFramework/Public/InputBehaviorSet.h
+++ b/Engine/Source/Runtime/Experimental/InteractiveToolsFramework/Public/InputBehaviorSet.h
@@ -19,7 +19,7 @@ struct INTERACTIVETOOLSFRAMEWORK_API FBehaviorInfo
 	
 	/** Reference to a Behavior */
 	UPROPERTY()
-	UInputBehavior* Behavior = nullptr;
+	TObjectPtr<UInputBehavior> Behavior = nullptr;
 
 	/** Source object that provided this Behavior */
 	void* Source = nullptr;
diff --git a/Engine/Source/Runtime/Experimental/InteractiveToolsFramework/Public/InputRouter.h b/Engine/Source/Runtime/Experimental/InteractiveToolsFramework/Public/InputRouter.h
--- a/Engine/Source/Runtime/Experimental/InteractiveToolsFramework/Public/InputRouter.h
+++ b/Engine/Source/Runtime/Experimental/InteractiveToolsFramework/Public/InputRouter.h
@@ -83,7 +83,7 @@ protected:
 	IToolsContextTransactionsAPI* TransactionsAPI;
 
 	UPROPERTY()
-	UInputBehaviorSet* ActiveInputBehaviors;
+	TObjectPtr<UInputBehaviorSet> ActiveInputBehaviors;
 
 	UInputBehavior* ActiveKeyboardCapture = nullptr;
 	void* ActiveKeyboardCaptureOwner = nullptr;
diff --git a/Engine/Source/Runtime/Experimental/InteractiveToolsFramework/Public/InteractiveGizmo.h b/Engine/Source/Runtime/Experimental/InteractiveToolsFramework/Public/InteractiveGizmo.h
--- a/Engine/Source/Runtime/Experimental/InteractiveToolsFramework/Public/InteractiveGizmo.h
+++ b/Engine/Source/Runtime/Experimental/InteractiveToolsFramework/Public/InteractiveGizmo.h
@@ -83,7 +83,7 @@ protected:
 
 	/** The current set of InputBehaviors provided by this Gizmo */
 	UPROPERTY()
-	UInputBehaviorSet* InputBehaviors;
+	TObjectPtr<UInputBehaviorSet> InputBehaviors;
 };
 
 
diff --git a/Engine/Source/Runtime/Experimental/InteractiveToolsFramework/Public/InteractiveGizmoManager.h b/Engine/Source/Runtime/Experimental/InteractiveToolsFramework/Public/InteractiveGizmoManager.h
--- a/Engine/Source/Runtime/Experimental/InteractiveToolsFramework/Public/InteractiveGizmoManager.h
+++ b/Engine/Source/Runtime/Experimental/InteractiveToolsFramework/Public/InteractiveGizmoManager.h
@@ -240,7 +240,7 @@ protected:
 
 	/** Current set of named GizmoBuilders */
 	UPROPERTY()
-	TMap<FString, UInteractiveGizmoBuilder*> GizmoBuilders;
+	TMap<FString, TObjectPtr<UInteractiveGizmoBuilder>> GizmoBuilders;
 
 	bool bDefaultGizmosRegistered = false;
 	TSharedPtr<FTransformGizmoActorFactory> GizmoActorBuilder;
diff --git a/Engine/Source/Runtime/Experimental/InteractiveToolsFramework/Public/InteractiveTool.h b/Engine/Source/Runtime/Experimental/InteractiveToolsFramework/Public/InteractiveTool.h
--- a/Engine/Source/Runtime/Experimental/InteractiveToolsFramework/Public/InteractiveTool.h
+++ b/Engine/Source/Runtime/Experimental/InteractiveToolsFramework/Public/InteractiveTool.h
@@ -224,7 +224,7 @@ public:
 protected:
 	// CachedProperties should only ever be set to an instance of the subclass, ideally via GetPropertyCache().
 	UPROPERTY()
-	UInteractiveToolPropertySet* CachedProperties = nullptr;
+	TObjectPtr<UInteractiveToolPropertySet> CachedProperties = nullptr;
 
 	UPROPERTY()
 	bool bIsPropertySetEnabled = true;
@@ -367,11 +367,11 @@ protected:
 
 	/** The current set of InputBehaviors provided by this Tool */
 	UPROPERTY()
-	UInputBehaviorSet* InputBehaviors;
+	TObjectPtr<UInputBehaviorSet> InputBehaviors;
 
 	/** The current set of Property UObjects provided by this Tool. May contain pointer to itself. */
 	UPROPERTY()
-	TArray<UObject*> ToolPropertyObjects;
+	TArray<TObjectPtr<UObject>> ToolPropertyObjects;
 
 	/**
 	 * Add a Property object for this Tool
diff --git a/Engine/Source/Runtime/Experimental/InteractiveToolsFramework/Public/InteractiveToolManager.h b/Engine/Source/Runtime/Experimental/InteractiveToolsFramework/Public/InteractiveToolManager.h
--- a/Engine/Source/Runtime/Experimental/InteractiveToolsFramework/Public/InteractiveToolManager.h
+++ b/Engine/Source/Runtime/Experimental/InteractiveToolsFramework/Public/InteractiveToolManager.h
@@ -258,11 +258,11 @@ public:
 public:
 	/** Currently-active Left Tool, or null if no Tool is active */
 	UPROPERTY()
-	UInteractiveTool* ActiveLeftTool;
+	TObjectPtr<UInteractiveTool> ActiveLeftTool;
 
 	/** Currently-active Right Tool, or null if no Tool is active */
 	UPROPERTY()
-	UInteractiveTool* ActiveRightTool;
+	TObjectPtr<UInteractiveTool> ActiveRightTool;
 
 
 
@@ -289,7 +289,7 @@ protected:
 
 	/** Current set of named ToolBuilders */
 	UPROPERTY()
-	TMap<FString, UInteractiveToolBuilder*> ToolBuilders;
+	TMap<FString, TObjectPtr<UInteractiveToolBuilder>> ToolBuilders;
 
 	/** Currently-active Left ToolBuilder */
 	FString ActiveLeftBuilderName;
diff --git a/Engine/Source/Runtime/Experimental/InteractiveToolsFramework/Public/InteractiveToolsContext.h b/Engine/Source/Runtime/Experimental/InteractiveToolsFramework/Public/InteractiveToolsContext.h
--- a/Engine/Source/Runtime/Experimental/InteractiveToolsFramework/Public/InteractiveToolsContext.h
+++ b/Engine/Source/Runtime/Experimental/InteractiveToolsFramework/Public/InteractiveToolsContext.h
@@ -58,19 +58,19 @@ public:
 public:
 	/** current UInputRouter for this Context */
 	UPROPERTY()
-	UInputRouter* InputRouter;	
+	TObjectPtr<UInputRouter> InputRouter;	
 
 	/** current UToolTargetManager for this Context */
 	UPROPERTY()
-	UToolTargetManager* TargetManager;
+	TObjectPtr<UToolTargetManager> TargetManager;
 
 	/** current UInteractiveToolManager for this Context */
 	UPROPERTY()
-	UInteractiveToolManager* ToolManager;	
+	TObjectPtr<UInteractiveToolManager> ToolManager;	
 
 	/** current UInteractiveGizmoManager for this Context */
 	UPROPERTY()
-	UInteractiveGizmoManager* GizmoManager;
+	TObjectPtr<UInteractiveGizmoManager> GizmoManager;
 
 protected:
 	UPROPERTY()
diff --git a/Engine/Source/Runtime/Experimental/InteractiveToolsFramework/Public/InteractiveToolsSelectionStoreSubsystem.h b/Engine/Source/Runtime/Experimental/InteractiveToolsFramework/Public/InteractiveToolsSelectionStoreSubsystem.h
--- a/Engine/Source/Runtime/Experimental/InteractiveToolsFramework/Public/InteractiveToolsSelectionStoreSubsystem.h
+++ b/Engine/Source/Runtime/Experimental/InteractiveToolsFramework/Public/InteractiveToolsSelectionStoreSubsystem.h
@@ -81,5 +81,5 @@ public:
 
 private:
 	UPROPERTY()
-	const UInteractiveToolStorableSelection* StoredSelection = nullptr;
+	TObjectPtr<const UInteractiveToolStorableSelection> StoredSelection = nullptr;
 };
\ No newline at end of file
diff --git a/Engine/Source/Runtime/Experimental/InteractiveToolsFramework/Public/SingleSelectionTool.h b/Engine/Source/Runtime/Experimental/InteractiveToolsFramework/Public/SingleSelectionTool.h
--- a/Engine/Source/Runtime/Experimental/InteractiveToolsFramework/Public/SingleSelectionTool.h
+++ b/Engine/Source/Runtime/Experimental/InteractiveToolsFramework/Public/SingleSelectionTool.h
@@ -45,5 +45,5 @@ protected:
 	TUniquePtr<FPrimitiveComponentTarget> ComponentTarget{};
 
 	UPROPERTY()
-	UToolTarget* Target;
+	TObjectPtr<UToolTarget> Target;
 };
diff --git a/Engine/Source/Runtime/Experimental/InteractiveToolsFramework/Public/ToolTargetManager.h b/Engine/Source/Runtime/Experimental/InteractiveToolsFramework/Public/ToolTargetManager.h
--- a/Engine/Source/Runtime/Experimental/InteractiveToolsFramework/Public/ToolTargetManager.h
+++ b/Engine/Source/Runtime/Experimental/InteractiveToolsFramework/Public/ToolTargetManager.h
@@ -88,7 +88,7 @@ protected:
 	bool bIsActive = false;
 
 	UPROPERTY()
-	TArray<UToolTargetFactory*> Factories;
+	TArray<TObjectPtr<UToolTargetFactory>> Factories;
 
 	// More state will go here if the manager deals with target caching.
 };
\ No newline at end of file
diff --git a/Engine/Source/Runtime/Foliage/Public/FoliageType.h b/Engine/Source/Runtime/Foliage/Public/FoliageType.h
--- a/Engine/Source/Runtime/Foliage/Public/FoliageType.h
+++ b/Engine/Source/Runtime/Foliage/Public/FoliageType.h
@@ -572,7 +572,7 @@ public:
 	 * The mesh material also needs to be set up to output to a virtual texture. 
 	 */
 	UPROPERTY(EditAnywhere, BlueprintReadWrite, Category=VirtualTexture, meta = (DisplayName = "Draw in Virtual Textures"))
-	TArray<URuntimeVirtualTexture*> RuntimeVirtualTextures;
+	TArray<TObjectPtr<URuntimeVirtualTexture>> RuntimeVirtualTextures;
 
 	/**
 	 * Number of lower mips in the runtime virtual texture to skip for rendering this primitive.
diff --git a/Engine/Source/Runtime/Foliage/Public/FoliageTypeObject.h b/Engine/Source/Runtime/Foliage/Public/FoliageTypeObject.h
--- a/Engine/Source/Runtime/Foliage/Public/FoliageTypeObject.h
+++ b/Engine/Source/Runtime/Foliage/Public/FoliageTypeObject.h
@@ -41,11 +41,11 @@ struct FOLIAGE_API FFoliageTypeObject
 private:
 	/** The foliage type that will be spawned by the procedural foliage simulation */
 	UPROPERTY(Category = ProceduralFoliageSimulation, EditAnywhere, meta=(AllowedClasses="FoliageType_InstancedStaticMesh,FoliageType_Actor,Blueprint", DisplayThumbnail="true", ThumbnailSize="X=40 Y=40"))
-	UObject* FoliageTypeObject;
+	TObjectPtr<UObject> FoliageTypeObject;
 
 	/** The actual instance of the foliage type that is used for spawning */
 	UPROPERTY(transient)
-	UFoliageType* TypeInstance;
+	TObjectPtr<UFoliageType> TypeInstance;
 
 	/** Whether this contains an asset object (as opposed to a BP class) */
 	UPROPERTY()
diff --git a/Engine/Source/Runtime/Foliage/Public/FoliageType_InstancedStaticMesh.h b/Engine/Source/Runtime/Foliage/Public/FoliageType_InstancedStaticMesh.h
--- a/Engine/Source/Runtime/Foliage/Public/FoliageType_InstancedStaticMesh.h
+++ b/Engine/Source/Runtime/Foliage/Public/FoliageType_InstancedStaticMesh.h
@@ -15,10 +15,10 @@ class UFoliageType_InstancedStaticMesh : public UFoliageType
 	GENERATED_UCLASS_BODY()
 
 	UPROPERTY(EditAnywhere, BlueprintReadWrite, Category=Mesh, meta=(DisplayThumbnail="true"))
-	UStaticMesh* Mesh;
+	TObjectPtr<UStaticMesh> Mesh;
 
 	UPROPERTY(EditAnywhere, BlueprintReadWrite, AdvancedDisplay, Category = Mesh, Meta = (ToolTip = "Material overrides for foliage instances."))
-	TArray<class UMaterialInterface*> OverrideMaterials;
+	TArray<TObjectPtr<class UMaterialInterface>> OverrideMaterials;
 		
 	/** The component class to use for foliage instances. 
 	  * You can make a Blueprint subclass of FoliageInstancedStaticMeshComponent to implement custom behavior and assign that class here. */
diff --git a/Engine/Source/Runtime/Foliage/Public/InteractiveFoliageActor.h b/Engine/Source/Runtime/Foliage/Public/InteractiveFoliageActor.h
--- a/Engine/Source/Runtime/Foliage/Public/InteractiveFoliageActor.h
+++ b/Engine/Source/Runtime/Foliage/Public/InteractiveFoliageActor.h
@@ -20,7 +20,7 @@ class AInteractiveFoliageActor : public AStaticMeshActor
 private:
 	/** Collision cylinder */
 	UPROPERTY()
-	class UCapsuleComponent* CapsuleComponent;
+	TObjectPtr<class UCapsuleComponent> CapsuleComponent;
 
 	/**
 	 * Position of the last actor to enter the collision cylinder.
diff --git a/Engine/Source/Runtime/Foliage/Public/ProceduralFoliageBlockingVolume.h b/Engine/Source/Runtime/Foliage/Public/ProceduralFoliageBlockingVolume.h
--- a/Engine/Source/Runtime/Foliage/Public/ProceduralFoliageBlockingVolume.h
+++ b/Engine/Source/Runtime/Foliage/Public/ProceduralFoliageBlockingVolume.h
@@ -18,7 +18,7 @@ class AProceduralFoliageBlockingVolume : public AVolume
 	GENERATED_UCLASS_BODY()
 
 	UPROPERTY(Category = ProceduralFoliage, EditAnywhere)
-	AProceduralFoliageVolume* ProceduralFoliageVolume;
+	TObjectPtr<AProceduralFoliageVolume> ProceduralFoliageVolume;
 
 	UPROPERTY(Category = ProceduralFoliage, EditAnywhere)
 	FFoliageDensityFalloff DensityFalloff;
diff --git a/Engine/Source/Runtime/Foliage/Public/ProceduralFoliageComponent.h b/Engine/Source/Runtime/Foliage/Public/ProceduralFoliageComponent.h
--- a/Engine/Source/Runtime/Foliage/Public/ProceduralFoliageComponent.h
+++ b/Engine/Source/Runtime/Foliage/Public/ProceduralFoliageComponent.h
@@ -45,7 +45,7 @@ class FOLIAGE_API UProceduralFoliageComponent : public UActorComponent
 
 	/** The procedural foliage spawner used to generate foliage instances within this volume. */
 	UPROPERTY(Category = "ProceduralFoliage", BlueprintReadWrite, EditAnywhere)
-	UProceduralFoliageSpawner* FoliageSpawner;
+	TObjectPtr<UProceduralFoliageSpawner> FoliageSpawner;
 
 	/** The amount of overlap between simulation tiles (in cm). */
 	UPROPERTY(Category = "ProceduralFoliage", BlueprintReadWrite, EditAnywhere)
@@ -128,7 +128,7 @@ private:
 
 private:
 	UPROPERTY()
-	AVolume* SpawningVolume;
+	TObjectPtr<AVolume> SpawningVolume;
 	
 	UPROPERTY()
 	FGuid ProceduralGuid;
diff --git a/Engine/Source/Runtime/Foliage/Public/ProceduralFoliageInstance.h b/Engine/Source/Runtime/Foliage/Public/ProceduralFoliageInstance.h
--- a/Engine/Source/Runtime/Foliage/Public/ProceduralFoliageInstance.h
+++ b/Engine/Source/Runtime/Foliage/Public/ProceduralFoliageInstance.h
@@ -96,7 +96,7 @@ public:
 	float Scale;
 
 	UPROPERTY()
-	const UFoliageType* Type;
+	TObjectPtr<const UFoliageType> Type;
 
 	UActorComponent* BaseComponent;
 
diff --git a/Engine/Source/Runtime/Foliage/Public/ProceduralFoliageTile.h b/Engine/Source/Runtime/Foliage/Public/ProceduralFoliageTile.h
--- a/Engine/Source/Runtime/Foliage/Public/ProceduralFoliageTile.h
+++ b/Engine/Source/Runtime/Foliage/Public/ProceduralFoliageTile.h
@@ -123,7 +123,7 @@ private:
 private:
 
 	UPROPERTY()
-	const UProceduralFoliageSpawner* FoliageSpawner;
+	TObjectPtr<const UProceduralFoliageSpawner> FoliageSpawner;
 
 	TSet<FProceduralFoliageInstance*> PendingRemovals;
 	TSet<FProceduralFoliageInstance*> InstancesSet;
diff --git a/Engine/Source/Runtime/Foliage/Public/ProceduralFoliageVolume.h b/Engine/Source/Runtime/Foliage/Public/ProceduralFoliageVolume.h
--- a/Engine/Source/Runtime/Foliage/Public/ProceduralFoliageVolume.h
+++ b/Engine/Source/Runtime/Foliage/Public/ProceduralFoliageVolume.h
@@ -15,7 +15,7 @@ class FOLIAGE_API AProceduralFoliageVolume: public AVolume
 	GENERATED_UCLASS_BODY()
 
 	UPROPERTY(Category = ProceduralFoliage, VisibleAnywhere, BlueprintReadOnly)
-	UProceduralFoliageComponent* ProceduralComponent;
+	TObjectPtr<UProceduralFoliageComponent> ProceduralComponent;
 
 #if WITH_EDITOR
 
diff --git a/Engine/Source/Runtime/GameplayTags/Classes/GameplayTagContainer.h b/Engine/Source/Runtime/GameplayTags/Classes/GameplayTagContainer.h
--- a/Engine/Source/Runtime/GameplayTags/Classes/GameplayTagContainer.h
+++ b/Engine/Source/Runtime/GameplayTags/Classes/GameplayTagContainer.h
@@ -1020,7 +1020,7 @@ public:
 
 	/** The base expression of this query. */
 	UPROPERTY(EditDefaultsOnly, Instanced, Category = Query)
-	class UEditableGameplayTagQueryExpression* RootExpression;
+	TObjectPtr<class UEditableGameplayTagQueryExpression> RootExpression;
 
 #if WITH_EDITOR
 	/** Converts this editor query construct into the runtime-usable token stream. */
@@ -1097,7 +1097,7 @@ class UEditableGameplayTagQueryExpression_AnyExprMatch : public UEditableGamepla
 	GENERATED_BODY()
 public:
 	UPROPERTY(EditAnywhere, Instanced, Category = Expr)
-	TArray<UEditableGameplayTagQueryExpression*> Expressions;
+	TArray<TObjectPtr<UEditableGameplayTagQueryExpression>> Expressions;
 
 #if WITH_EDITOR
 	virtual void EmitTokens(TArray<uint8>& TokenStream, TArray<FGameplayTag>& TagDictionary, FString* DebugString = nullptr) const override;
@@ -1110,7 +1110,7 @@ class UEditableGameplayTagQueryExpression_AllExprMatch : public UEditableGamepla
 	GENERATED_BODY()
 public:
 	UPROPERTY(EditAnywhere, Instanced, Category = Expr)
-	TArray<UEditableGameplayTagQueryExpression*> Expressions;
+	TArray<TObjectPtr<UEditableGameplayTagQueryExpression>> Expressions;
 
 #if WITH_EDITOR
 	virtual void EmitTokens(TArray<uint8>& TokenStream, TArray<FGameplayTag>& TagDictionary, FString* DebugString = nullptr) const override;
@@ -1123,7 +1123,7 @@ class UEditableGameplayTagQueryExpression_NoExprMatch : public UEditableGameplay
 	GENERATED_BODY()
 public:
 	UPROPERTY(EditAnywhere, Instanced, Category = Expr)
-	TArray<UEditableGameplayTagQueryExpression*> Expressions;
+	TArray<TObjectPtr<UEditableGameplayTagQueryExpression>> Expressions;
 
 #if WITH_EDITOR
 	virtual void EmitTokens(TArray<uint8>& TokenStream, TArray<FGameplayTag>& TagDictionary, FString* DebugString = nullptr) const override;
diff --git a/Engine/Source/Runtime/GameplayTags/Classes/GameplayTagsManager.h b/Engine/Source/Runtime/GameplayTags/Classes/GameplayTagsManager.h
--- a/Engine/Source/Runtime/GameplayTags/Classes/GameplayTagsManager.h
+++ b/Engine/Source/Runtime/GameplayTags/Classes/GameplayTagsManager.h
@@ -98,11 +98,11 @@ struct GAMEPLAYTAGS_API FGameplayTagSource
 
 	/** If this is bound to an ini object for saving, this is the one */
 	UPROPERTY()
-	class UGameplayTagsList* SourceTagList;
+	TObjectPtr<class UGameplayTagsList> SourceTagList;
 
 	/** If this has restricted tags and is bound to an ini object for saving, this is the one */
 	UPROPERTY()
-	class URestrictedGameplayTagsList* SourceRestrictedTagList;
+	TObjectPtr<class URestrictedGameplayTagsList> SourceRestrictedTagList;
 
 	FGameplayTagSource() 
 		: SourceName(NAME_None), SourceType(EGameplayTagSourceType::Invalid), SourceTagList(nullptr), SourceRestrictedTagList(nullptr)
@@ -721,7 +721,7 @@ private:
 
 	/** Holds all of the valid gameplay-related tags that can be applied to assets */
 	UPROPERTY()
-	TArray<UDataTable*> GameplayTagTables;
+	TArray<TObjectPtr<UDataTable>> GameplayTagTables;
 
 	/** The map of ini-configured tag redirectors */
 	TMap<FName, FGameplayTag> TagRedirects;
diff --git a/Engine/Source/Runtime/GameplayTasks/Classes/GameplayTask.h b/Engine/Source/Runtime/GameplayTasks/Classes/GameplayTask.h
--- a/Engine/Source/Runtime/GameplayTasks/Classes/GameplayTask.h
+++ b/Engine/Source/Runtime/GameplayTasks/Classes/GameplayTask.h
@@ -362,7 +362,7 @@ protected:
 
 	/** child task instance */
 	UPROPERTY()
-	UGameplayTask* ChildTask;
+	TObjectPtr<UGameplayTask> ChildTask;
 
 #if !(UE_BUILD_SHIPPING || UE_BUILD_TEST)
 	mutable FString DebugDescription;
diff --git a/Engine/Source/Runtime/GameplayTasks/Classes/GameplayTasksComponent.h b/Engine/Source/Runtime/GameplayTasks/Classes/GameplayTasksComponent.h
--- a/Engine/Source/Runtime/GameplayTasks/Classes/GameplayTasksComponent.h
+++ b/Engine/Source/Runtime/GameplayTasks/Classes/GameplayTasksComponent.h
@@ -83,7 +83,7 @@ protected:
 	/** Tasks that run on simulated proxies */
 	UE_DEPRECATED(4.26, "This will be made private in future versions. Use GetSimulatedTasks, GetSimulatedTasks_Mutable, or SetSimulatedTasks instead.")
 	UPROPERTY(ReplicatedUsing = OnRep_SimulatedTasks)
-	TArray<UGameplayTask*> SimulatedTasks;
+	TArray<TObjectPtr<UGameplayTask>> SimulatedTasks;
 
 PRAGMA_DISABLE_DEPRECATION_WARNINGS
 	const TArray<UGameplayTask*>& GetSimulatedTasks()
@@ -96,7 +96,7 @@ PRAGMA_ENABLE_DEPRECATION_WARNINGS
 	void SetSimulatedTasks(const TArray<UGameplayTask*>& NewSimulatedTasks);
 
 	UPROPERTY()
-	TArray<UGameplayTask*> TaskPriorityQueue;
+	TArray<TObjectPtr<UGameplayTask>> TaskPriorityQueue;
 	
 	/** Transient array of events whose main role is to avoid
 	 *	long chain of recurrent calls if an activated/paused/removed task 
@@ -106,11 +106,11 @@ PRAGMA_ENABLE_DEPRECATION_WARNINGS
 
 	/** Array of currently active UGameplayTask that require ticking */
 	UPROPERTY()
-	TArray<UGameplayTask*> TickingTasks;
+	TArray<TObjectPtr<UGameplayTask>> TickingTasks;
 
 	/** All known tasks (processed by this component) referenced for GC */
 	UPROPERTY(transient)
-	TArray<UGameplayTask*> KnownTasks;
+	TArray<TObjectPtr<UGameplayTask>> KnownTasks;
 
 public:
 	UPROPERTY(BlueprintReadWrite, Category = "Gameplay Tasks")
diff --git a/Engine/Source/Runtime/HeadMountedDisplay/Public/MotionControllerComponent.h b/Engine/Source/Runtime/HeadMountedDisplay/Public/MotionControllerComponent.h
--- a/Engine/Source/Runtime/HeadMountedDisplay/Public/MotionControllerComponent.h
+++ b/Engine/Source/Runtime/HeadMountedDisplay/Public/MotionControllerComponent.h
@@ -68,14 +68,14 @@ class HEADMOUNTEDDISPLAY_API UMotionControllerComponent : public UPrimitiveCompo
 
 	/** A mesh override that'll be displayed attached to this MotionController. */
 	UPROPERTY(EditAnywhere, BlueprintReadWrite, BlueprintSetter=SetCustomDisplayMesh, Category="Visualization", meta=(editcondition="bDisplayDeviceModel"))
-	UStaticMesh* CustomDisplayMesh;
+	TObjectPtr<UStaticMesh> CustomDisplayMesh;
 
 	UFUNCTION(BlueprintSetter)
 	void SetCustomDisplayMesh(UStaticMesh* NewDisplayMesh);
 
 	/** Material overrides for the specified display mesh. */
 	UPROPERTY(EditAnywhere, BlueprintReadOnly, Category="Visualization", meta=(editcondition="bDisplayDeviceModel"))
-	TArray<UMaterialInterface*> DisplayMeshMaterialOverrides;
+	TArray<TObjectPtr<UMaterialInterface>> DisplayMeshMaterialOverrides;
 
 	UFUNCTION(BlueprintSetter, meta = (DeprecatedFunction, DeprecationMessage = "Please use the Motion Source property instead of Hand"))
 	void SetTrackingSource(const EControllerHand NewSource);
@@ -166,7 +166,7 @@ private:
 	TSharedPtr< FViewExtension, ESPMode::ThreadSafe > ViewExtension;	
  
 	UPROPERTY(Transient, BlueprintReadOnly, Category=Visualization, meta=(AllowPrivateAccess="true"))
-	UPrimitiveComponent* DisplayComponent;
+	TObjectPtr<UPrimitiveComponent> DisplayComponent;
 
 	/** Callback for asynchronous display model loads (to set materials, etc.) */
 	void OnDisplayModelLoaded(UPrimitiveComponent* DisplayComponent);
diff --git a/Engine/Source/Runtime/HeadMountedDisplay/Public/XRAssetFunctionLibrary.h b/Engine/Source/Runtime/HeadMountedDisplay/Public/XRAssetFunctionLibrary.h
--- a/Engine/Source/Runtime/HeadMountedDisplay/Public/XRAssetFunctionLibrary.h
+++ b/Engine/Source/Runtime/HeadMountedDisplay/Public/XRAssetFunctionLibrary.h
@@ -98,6 +98,6 @@ private:
 	ELoadStatus LoadStatus = ELoadStatus::Pending;
 
 	UPROPERTY()
-	UPrimitiveComponent* SpawnedComponent;
+	TObjectPtr<UPrimitiveComponent> SpawnedComponent;
 };
 
diff --git a/Engine/Source/Runtime/Interchange/Core/Public/Nodes/InterchangeBaseNodeContainer.h b/Engine/Source/Runtime/Interchange/Core/Public/Nodes/InterchangeBaseNodeContainer.h
--- a/Engine/Source/Runtime/Interchange/Core/Public/Nodes/InterchangeBaseNodeContainer.h
+++ b/Engine/Source/Runtime/Interchange/Core/Public/Nodes/InterchangeBaseNodeContainer.h
@@ -97,5 +97,5 @@ private:
 
 	/** Flat List of the nodes. Since the nodes are variable size, we store a pointer. */
 	UPROPERTY()
-	TMap<FString, UInterchangeBaseNode* > Nodes;
+	TMap<FString, TObjectPtr<UInterchangeBaseNode> > Nodes;
 };
diff --git a/Engine/Source/Runtime/Interchange/Engine/Public/InterchangeAssetImportData.h b/Engine/Source/Runtime/Interchange/Engine/Public/InterchangeAssetImportData.h
--- a/Engine/Source/Runtime/Interchange/Engine/Public/InterchangeAssetImportData.h
+++ b/Engine/Source/Runtime/Interchange/Engine/Public/InterchangeAssetImportData.h
@@ -63,7 +63,7 @@ public:
 	
 	/** The graph that was use to create this asset */
 	UPROPERTY(VisibleAnywhere, Category = "Interchange | AssetImportData")
-	UInterchangeBaseNodeContainer* NodeContainer;
+	TObjectPtr<UInterchangeBaseNodeContainer> NodeContainer;
 
 	/** The Node UID pass to the factory that exist in the graph that was use to create this asset */
 	UPROPERTY(VisibleAnywhere, Category = "Interchange | AssetImportData")
diff --git a/Engine/Source/Runtime/Interchange/Engine/Public/InterchangeManager.h b/Engine/Source/Runtime/Interchange/Engine/Public/InterchangeManager.h
--- a/Engine/Source/Runtime/Interchange/Engine/Public/InterchangeManager.h
+++ b/Engine/Source/Runtime/Interchange/Engine/Public/InterchangeManager.h
@@ -200,13 +200,13 @@ struct FImportAssetParameters
 	GENERATED_USTRUCT_BODY()
 
 	UPROPERTY(EditAnywhere, BlueprintReadWrite, Category = "Interchange|ImportAsset")
-	UObject* ReimportAsset = nullptr;
+	TObjectPtr<UObject> ReimportAsset = nullptr;
 
 	UPROPERTY(EditAnywhere, BlueprintReadWrite, Category = "Interchange|ImportAsset")
 	bool bIsAutomated = false;
 
 	UPROPERTY(EditAnywhere, BlueprintReadWrite, Category = "Interchange|ImportAsset")
-	TArray<UInterchangePipelineBase*> OverridePipelines;
+	TArray<TObjectPtr<UInterchangePipelineBase>> OverridePipelines;
 };
 
 UCLASS(Transient, BlueprintType)
@@ -442,24 +442,24 @@ private:
 
 	// Caching the registered translator classes to avoid double registration fast
 	UPROPERTY()
-	TSet<const UClass*> RegisteredTranslatorsClass;
+	TSet<TObjectPtr<const UClass>> RegisteredTranslatorsClass;
 
 	//The manager will create translator at every import, translator must be able to retrieve payload information when the factory ask for it.
 	//The translator stored has value is only use to know if we can use this type of translator.
 	UPROPERTY()
-	TArray<UInterchangeTranslatorBase*> RegisteredTranslators;
+	TArray<TObjectPtr<UInterchangeTranslatorBase>> RegisteredTranslators;
 	
 	//The manager will create only one pipeline per type
 	UPROPERTY()
-	TMap<const UClass*, UInterchangePipelineBase* > RegisteredPipelines;
+	TMap<TObjectPtr<const UClass>, TObjectPtr<UInterchangePipelineBase> > RegisteredPipelines;
 
 	//The manager will create only one factory per type
 	UPROPERTY()
-	TMap<const UClass*, UInterchangeFactoryBase* > RegisteredFactories;
+	TMap<TObjectPtr<const UClass>, TObjectPtr<UInterchangeFactoryBase> > RegisteredFactories;
 
 	//The manager will create only one writer per type
 	UPROPERTY()
-	TMap<const UClass*, UInterchangeWriterBase* > RegisteredWriters;
+	TMap<TObjectPtr<const UClass>, TObjectPtr<UInterchangeWriterBase> > RegisteredWriters;
 
 	//If interchange is currently importing we have a timer to watch the cancel and we block GC 
 	FThreadSafeBool bIsActive = false;
diff --git a/Engine/Source/Runtime/Landscape/Classes/ControlPointMeshActor.h b/Engine/Source/Runtime/Landscape/Classes/ControlPointMeshActor.h
--- a/Engine/Source/Runtime/Landscape/Classes/ControlPointMeshActor.h
+++ b/Engine/Source/Runtime/Landscape/Classes/ControlPointMeshActor.h
@@ -19,7 +19,7 @@ class LANDSCAPE_API AControlPointMeshActor : public AActor
 
 private:
 	UPROPERTY(Category = ControlPointMeshActor, VisibleAnywhere, BlueprintReadOnly, meta = (ExposeFunctionCategories = "Mesh,Rendering,Physics,Components|StaticMesh,Components|ControlPointMesh", AllowPrivateAccess = "true"))
-	class UControlPointMeshComponent* ControlPointMeshComponent;
+	TObjectPtr<class UControlPointMeshComponent> ControlPointMeshComponent;
 
 public:
 
diff --git a/Engine/Source/Runtime/Landscape/Classes/Landscape.h b/Engine/Source/Runtime/Landscape/Classes/Landscape.h
--- a/Engine/Source/Runtime/Landscape/Classes/Landscape.h
+++ b/Engine/Source/Runtime/Landscape/Classes/Landscape.h
@@ -130,7 +130,7 @@ private:
 
 #if WITH_EDITORONLY_DATA
 	UPROPERTY()
-	ALandscapeBlueprintBrushBase* BlueprintBrush;
+	TObjectPtr<ALandscapeBlueprintBrushBase> BlueprintBrush;
 
 	FTransform LandscapeTransform;
 	FIntPoint LandscapeSize;
@@ -188,7 +188,7 @@ struct FLandscapeLayer
 	TArray<FLandscapeLayerBrush> Brushes;
 
 	UPROPERTY()
-	TMap<ULandscapeLayerInfoObject*, bool> WeightmapLayerAllocationBlend; // True -> Substractive, False -> Additive
+	TMap<TObjectPtr<ULandscapeLayerInfoObject>, bool> WeightmapLayerAllocationBlend; // True -> Substractive, False -> Additive
 };
 
 UCLASS(MinimalAPI, showcategories=(Display, Movement, Collision, Lighting, LOD, Input), hidecategories=(Mobility))
@@ -419,10 +419,10 @@ public:
 	TArray<FLandscapeLayer> LandscapeLayers;
 
 	UPROPERTY(Transient)
-	TArray<UTextureRenderTarget2D*> HeightmapRTList;
+	TArray<TObjectPtr<UTextureRenderTarget2D>> HeightmapRTList;
 
 	UPROPERTY(Transient)
-	TArray<UTextureRenderTarget2D*> WeightmapRTList;
+	TArray<TObjectPtr<UTextureRenderTarget2D>> WeightmapRTList;
 
 private:
 	FLandscapeBlueprintBrushChangedDelegate LandscapeBlueprintBrushChangedDelegate;
@@ -430,7 +430,7 @@ private:
 
 	/** Components affected by landscape splines (used to partially clear Layer Reserved for Splines) */
 	UPROPERTY(Transient)
-	TSet<ULandscapeComponent*> LandscapeSplinesAffectedComponents;
+	TSet<TObjectPtr<ULandscapeComponent>> LandscapeSplinesAffectedComponents;
 
 	/** Provides information from LandscapeEdMode */
 	ILandscapeEdModeInterface* LandscapeEdMode;
diff --git a/Engine/Source/Runtime/Landscape/Classes/LandscapeComponent.h b/Engine/Source/Runtime/Landscape/Classes/LandscapeComponent.h
--- a/Engine/Source/Runtime/Landscape/Classes/LandscapeComponent.h
+++ b/Engine/Source/Runtime/Landscape/Classes/LandscapeComponent.h
@@ -72,11 +72,11 @@ public:
 
 	// Material used to render the tool.
 	UPROPERTY(NonTransactional)
-	UMaterialInterface* ToolMaterial;
+	TObjectPtr<UMaterialInterface> ToolMaterial;
 
 	// Material used to render the gizmo selection region...
 	UPROPERTY(NonTransactional)
-	UMaterialInterface* GizmoMaterial;
+	TObjectPtr<UMaterialInterface> GizmoMaterial;
 
 	// Component is selected
 	UPROPERTY(NonTransactional)
@@ -92,13 +92,13 @@ public:
 	int32 DebugChannelB;
 
 	UPROPERTY(NonTransactional)
-	UTexture2D* DataTexture; // Data texture other than height/weight
+	TObjectPtr<UTexture2D> DataTexture; // Data texture other than height/weight
 
 	UPROPERTY(NonTransactional)
-	UTexture2D* LayerContributionTexture; // Data texture used to represent layer contribution
+	TObjectPtr<UTexture2D> LayerContributionTexture; // Data texture used to represent layer contribution
 
 	UPROPERTY(NonTransactional)
-	UTexture2D* DirtyTexture; // Data texture used to represent layer blend dirtied area
+	TObjectPtr<UTexture2D> DirtyTexture; // Data texture used to represent layer blend dirtied area
 
 #if WITH_EDITOR
 	void UpdateDebugColorMaterial(const ULandscapeComponent* const Component);
@@ -199,7 +199,7 @@ struct FWeightmapLayerAllocationInfo
 	GENERATED_USTRUCT_BODY()
 
 	UPROPERTY()
-	ULandscapeLayerInfoObject* LayerInfo;
+	TObjectPtr<ULandscapeLayerInfoObject> LayerInfo;
 
 	UPROPERTY()
 	uint8 WeightmapTextureIndex;
@@ -294,7 +294,7 @@ struct FLandscapeComponentMaterialOverride
 	FPerPlatformInt LODIndex;
 
 	UPROPERTY(EditAnywhere, Category = LandscapeComponent)
-	UMaterialInterface* Material = nullptr;
+	TObjectPtr<UMaterialInterface> Material = nullptr;
 };
 
 USTRUCT(NotBlueprintable)
@@ -303,13 +303,13 @@ struct FWeightmapData
 	GENERATED_USTRUCT_BODY()
 
 	UPROPERTY()
-	TArray<UTexture2D*> Textures;
+	TArray<TObjectPtr<UTexture2D>> Textures;
 	
 	UPROPERTY()
 	TArray<FWeightmapLayerAllocationInfo> LayerAllocations;
 
 	UPROPERTY(Transient)
-	TArray<ULandscapeWeightmapUsage*> TextureUsages;
+	TArray<TObjectPtr<ULandscapeWeightmapUsage>> TextureUsages;
 };
 
 USTRUCT(NotBlueprintable)
@@ -318,7 +318,7 @@ struct FHeightmapData
 	GENERATED_USTRUCT_BODY()
 
 	UPROPERTY()
-	UTexture2D* Texture = nullptr;
+	TObjectPtr<UTexture2D> Texture = nullptr;
 };
 
 USTRUCT(NotBlueprintable)
@@ -446,24 +446,24 @@ class ULandscapeComponent : public UPrimitiveComponent
 	int32 NumSubsections;
 
 	UPROPERTY(EditAnywhere, BlueprintReadWrite, Category=LandscapeComponent)
-	UMaterialInterface* OverrideMaterial;
+	TObjectPtr<UMaterialInterface> OverrideMaterial;
 
 	UPROPERTY(EditAnywhere, BlueprintReadWrite, Category=LandscapeComponent, AdvancedDisplay)
-	UMaterialInterface* OverrideHoleMaterial;
+	TObjectPtr<UMaterialInterface> OverrideHoleMaterial;
 
 	UPROPERTY(EditAnywhere, Category = LandscapeComponent)
 	TArray<FLandscapeComponentMaterialOverride> OverrideMaterials;
 
 #if WITH_EDITORONLY_DATA
 	UPROPERTY()
-	UMaterialInstanceConstant* MaterialInstance_DEPRECATED;
+	TObjectPtr<UMaterialInstanceConstant> MaterialInstance_DEPRECATED;
 #endif
 
 	UPROPERTY(TextExportTransient)
-	TArray<UMaterialInstanceConstant*> MaterialInstances;
+	TArray<TObjectPtr<UMaterialInstanceConstant>> MaterialInstances;
 
 	UPROPERTY(Transient, TextExportTransient)
-	TArray<UMaterialInstanceDynamic*> MaterialInstancesDynamic;
+	TArray<TObjectPtr<UMaterialInstanceDynamic>> MaterialInstancesDynamic;
 
 	/** Mapping between LOD and Material Index*/
 	UPROPERTY(TextExportTransient)
@@ -475,7 +475,7 @@ class ULandscapeComponent : public UPrimitiveComponent
 
 	/** XYOffsetmap texture reference */
 	UPROPERTY(TextExportTransient)
-	UTexture2D* XYOffsetmapTexture;
+	TObjectPtr<UTexture2D> XYOffsetmapTexture;
 
 	/** UV offset to component's weightmap data from component local coordinates*/
 	UPROPERTY()
@@ -513,7 +513,7 @@ private:
 		
 	// Final layer data
 	UPROPERTY(Transient)
-	TArray<ULandscapeWeightmapUsage*> WeightmapTexturesUsage;
+	TArray<TObjectPtr<ULandscapeWeightmapUsage>> WeightmapTexturesUsage;
 
 	UPROPERTY(Transient)
 	uint32 LayerUpdateFlagPerMode;
@@ -524,7 +524,7 @@ private:
 
 	/** Heightmap texture reference */
 	UPROPERTY(TextExportTransient)
-	UTexture2D* HeightmapTexture;
+	TObjectPtr<UTexture2D> HeightmapTexture;
 
 	/** List of layers, and the weightmap and channel they are stored */
 	UPROPERTY()
@@ -532,11 +532,11 @@ private:
 
 	/** Weightmap texture reference */
 	UPROPERTY(TextExportTransient)
-	TArray<UTexture2D*> WeightmapTextures;
+	TArray<TObjectPtr<UTexture2D>> WeightmapTextures;
 
 	/** Used to interface the component to the LOD streamer. */
 	UPROPERTY()
-	ULandscapeLODStreamingProxy* LODStreamingProxy;
+	TObjectPtr<ULandscapeLODStreamingProxy> LODStreamingProxy;
 
 public:
 
@@ -591,7 +591,7 @@ public:
 
 	/** Pre-baked Base Color texture for use by distance field GI */
 	UPROPERTY(VisibleAnywhere, BlueprintReadOnly, Category = BakedTextures)
-	UTexture2D* GIBakedBaseColorTexture;
+	TObjectPtr<UTexture2D> GIBakedBaseColorTexture;
 
 #if WITH_EDITORONLY_DATA
 	/** LOD level Bias to use when lighting buidling via lightmass, -1 Means automatic LOD calculation based on ForcedLOD + LODBias */
@@ -600,7 +600,7 @@ public:
 
 	// List of layers allowed to be painted on this component
 	UPROPERTY(EditAnywhere, Category=LandscapeComponent)
-	TArray<ULandscapeLayerInfoObject*> LayerWhitelist;
+	TArray<TObjectPtr<ULandscapeLayerInfoObject>> LayerWhitelist;
 
 	/** Pointer to data shared with the render thread, used by the editor tools */
 	UPROPERTY(Transient, DuplicateTransient, NonTransactional)
@@ -612,7 +612,7 @@ public:
 
 	/** Represent the chosen material for each LOD */
 	UPROPERTY(DuplicateTransient)
-	TMap<UMaterialInterface*, int8> MaterialPerLOD;
+	TMap<TObjectPtr<UMaterialInterface>, int8> MaterialPerLOD;
 
 	/** Represents hash of last weightmap usage update */
 	uint32 WeightmapsHash;
@@ -634,16 +634,16 @@ public:
 	uint8 MobileBlendableLayerMask;
 
 	UPROPERTY(NonPIEDuplicateTransient)
-	UMaterialInterface* MobileMaterialInterface_DEPRECATED;
+	TObjectPtr<UMaterialInterface> MobileMaterialInterface_DEPRECATED;
 
 	/** Material interfaces used for mobile */
 	UPROPERTY(NonPIEDuplicateTransient)
-	TArray<UMaterialInterface*> MobileMaterialInterfaces;
+	TArray<TObjectPtr<UMaterialInterface>> MobileMaterialInterfaces;
 
 	/** Generated weightmap textures used for mobile. The first entry is also used for the normal map. 
 	  * Serialized only when cooking or loading cooked builds. */
 	UPROPERTY(NonPIEDuplicateTransient)
-	TArray<UTexture2D*> MobileWeightmapTextures;
+	TArray<TObjectPtr<UTexture2D>> MobileWeightmapTextures;
 
 #if WITH_EDITORONLY_DATA
 	/** Layer allocations used by mobile. Cached value here used only in the editor for usage visualization. */
@@ -652,10 +652,10 @@ public:
 	/** The editor needs to save out the combination MIC we'll use for mobile, 
 	  because we cannot generate it at runtime for standalone PIE games */
 	UPROPERTY(NonPIEDuplicateTransient)
-	TArray<UMaterialInstanceConstant*> MobileCombinationMaterialInstances;
+	TArray<TObjectPtr<UMaterialInstanceConstant>> MobileCombinationMaterialInstances;
 
 	UPROPERTY(NonPIEDuplicateTransient)
-	UMaterialInstanceConstant* MobileCombinationMaterialInstance_DEPRECATED;
+	TObjectPtr<UMaterialInstanceConstant> MobileCombinationMaterialInstance_DEPRECATED;
 #endif
 
 public:
diff --git a/Engine/Source/Runtime/Landscape/Classes/LandscapeGizmoActiveActor.h b/Engine/Source/Runtime/Landscape/Classes/LandscapeGizmoActiveActor.h
--- a/Engine/Source/Runtime/Landscape/Classes/LandscapeGizmoActiveActor.h
+++ b/Engine/Source/Runtime/Landscape/Classes/LandscapeGizmoActiveActor.h
@@ -60,7 +60,7 @@ class ALandscapeGizmoActiveActor : public ALandscapeGizmoActor
 	TEnumAsByte<enum ELandscapeGizmoType> DataType;
 
 	UPROPERTY(Transient)
-	UTexture2D* GizmoTexture;
+	TObjectPtr<UTexture2D> GizmoTexture;
 
 	UPROPERTY()
 	FVector2D TextureScale;
@@ -90,16 +90,16 @@ class ALandscapeGizmoActiveActor : public ALandscapeGizmoActor
 	FVector FrustumVerts[8];
 
 	UPROPERTY()
-	UMaterial* GizmoMaterial;
+	TObjectPtr<UMaterial> GizmoMaterial;
 
 	UPROPERTY()
-	UMaterialInstance* GizmoDataMaterial;
+	TObjectPtr<UMaterialInstance> GizmoDataMaterial;
 
 	UPROPERTY()
-	UMaterial* GizmoMeshMaterial;
+	TObjectPtr<UMaterial> GizmoMeshMaterial;
 
 	UPROPERTY(Category=LandscapeGizmoActiveActor, VisibleAnywhere)
-	TArray<ULandscapeLayerInfoObject*> LayerInfos;
+	TArray<TObjectPtr<ULandscapeLayerInfoObject>> LayerInfos;
 
 	UPROPERTY(transient)
 	bool bSnapToLandscapeGrid;
diff --git a/Engine/Source/Runtime/Landscape/Classes/LandscapeGizmoActor.h b/Engine/Source/Runtime/Landscape/Classes/LandscapeGizmoActor.h
--- a/Engine/Source/Runtime/Landscape/Classes/LandscapeGizmoActor.h
+++ b/Engine/Source/Runtime/Landscape/Classes/LandscapeGizmoActor.h
@@ -35,11 +35,11 @@ class ALandscapeGizmoActor : public AActor
 	float RelativeScaleZ;
 
 	UPROPERTY(EditAnywhere, transient, Category=Gizmo)
-	class ULandscapeInfo* TargetLandscapeInfo;
+	TObjectPtr<class ULandscapeInfo> TargetLandscapeInfo;
 
 private:
 	UPROPERTY()
-	UBillboardComponent* SpriteComponent;
+	TObjectPtr<UBillboardComponent> SpriteComponent;
 #endif // WITH_EDITORONLY_DATA
 public:
 
diff --git a/Engine/Source/Runtime/Landscape/Classes/LandscapeGrassType.h b/Engine/Source/Runtime/Landscape/Classes/LandscapeGrassType.h
--- a/Engine/Source/Runtime/Landscape/Classes/LandscapeGrassType.h
+++ b/Engine/Source/Runtime/Landscape/Classes/LandscapeGrassType.h
@@ -29,10 +29,10 @@ struct FGrassVariety
 	GENERATED_USTRUCT_BODY()
 
 	UPROPERTY(EditAnywhere, Category=Grass)
-	UStaticMesh* GrassMesh;
+	TObjectPtr<UStaticMesh> GrassMesh;
 
 	UPROPERTY(EditAnywhere, Category=Grass, meta = (ToolTip = "Material Overrides."))
-	TArray<class UMaterialInterface*> OverrideMaterials;
+	TArray<TObjectPtr<class UMaterialInterface>> OverrideMaterials;
 
 	/* Instances per 10 square meters. */
 	UPROPERTY(EditAnywhere, Category=Grass, meta = (UIMin = 0, ClampMin = 0, UIMax = 1000, ClampMax = 1000))
@@ -148,7 +148,7 @@ class ULandscapeGrassType : public UObject
 	uint32 bEnableDensityScaling : 1;
 
 	UPROPERTY()
-	UStaticMesh* GrassMesh_DEPRECATED;
+	TObjectPtr<UStaticMesh> GrassMesh_DEPRECATED;
 	UPROPERTY()
 	float GrassDensity_DEPRECATED;
 	UPROPERTY()
diff --git a/Engine/Source/Runtime/Landscape/Classes/LandscapeHeightfieldCollisionComponent.h b/Engine/Source/Runtime/Landscape/Classes/LandscapeHeightfieldCollisionComponent.h
--- a/Engine/Source/Runtime/Landscape/Classes/LandscapeHeightfieldCollisionComponent.h
+++ b/Engine/Source/Runtime/Landscape/Classes/LandscapeHeightfieldCollisionComponent.h
@@ -47,7 +47,7 @@ class ULandscapeHeightfieldCollisionComponent : public UPrimitiveComponent
 
 	/** List of layers painted on this component. Matches the WeightmapLayerAllocations array in the LandscapeComponent. */
 	UPROPERTY()
-	TArray<ULandscapeLayerInfoObject*> ComponentLayerInfos;
+	TArray<TObjectPtr<ULandscapeLayerInfoObject>> ComponentLayerInfos;
 
 	/** Offset of component in landscape quads */
 	UPROPERTY()
@@ -132,7 +132,7 @@ class ULandscapeHeightfieldCollisionComponent : public UPrimitiveComponent
 
 	/** Physical materials objects referenced by the indices in PhysicalMaterialRenderData. Stripped from cooked content */
 	UPROPERTY()
-	TArray<UPhysicalMaterial*>					PhysicalMaterialRenderObjects;
+	TArray<TObjectPtr<UPhysicalMaterial>>					PhysicalMaterialRenderObjects;
 
 	/*  Cooked editor specific heightfield data, never serialized  */
 	TArray<uint8>								CookedCollisionDataEd;
@@ -160,7 +160,7 @@ class ULandscapeHeightfieldCollisionComponent : public UPrimitiveComponent
 	
 	/** This is a list of physical materials that is actually used by a cooked HeightField */
 	UPROPERTY()
-	TArray<UPhysicalMaterial*>					CookedPhysicalMaterials;
+	TArray<TObjectPtr<UPhysicalMaterial>>					CookedPhysicalMaterials;
 	
 	/** Physics engine version of heightfield data. */
 	TRefCountPtr<FHeightfieldGeometryRef>	HeightfieldRef;
diff --git a/Engine/Source/Runtime/Landscape/Classes/LandscapeInfo.h b/Engine/Source/Runtime/Landscape/Classes/LandscapeInfo.h
--- a/Engine/Source/Runtime/Landscape/Classes/LandscapeInfo.h
+++ b/Engine/Source/Runtime/Landscape/Classes/LandscapeInfo.h
@@ -50,17 +50,17 @@ struct FLandscapeInfoLayerSettings
 	GENERATED_USTRUCT_BODY()
 
 	UPROPERTY()
-	ULandscapeLayerInfoObject* LayerInfoObj;
+	TObjectPtr<ULandscapeLayerInfoObject> LayerInfoObj;
 
 	UPROPERTY()
 	FName LayerName;
 
 #if WITH_EDITORONLY_DATA
 	UPROPERTY(transient)
-	UMaterialInstanceConstant* ThumbnailMIC;
+	TObjectPtr<UMaterialInstanceConstant> ThumbnailMIC;
 
 	UPROPERTY()
-	ALandscapeProxy* Owner;
+	TObjectPtr<ALandscapeProxy> Owner;
 
 	UPROPERTY(transient)
 	int32 DebugColorChannel;
@@ -143,7 +143,7 @@ public:
 #endif
 
 	UPROPERTY()
-	TArray<ALandscapeStreamingProxy*> Proxies;
+	TArray<TObjectPtr<ALandscapeStreamingProxy>> Proxies;
 
 private:
 #if WITH_EDITORONLY_DATA
diff --git a/Engine/Source/Runtime/Landscape/Classes/LandscapeLayerInfoObject.h b/Engine/Source/Runtime/Landscape/Classes/LandscapeLayerInfoObject.h
--- a/Engine/Source/Runtime/Landscape/Classes/LandscapeLayerInfoObject.h
+++ b/Engine/Source/Runtime/Landscape/Classes/LandscapeLayerInfoObject.h
@@ -29,7 +29,7 @@ class ULandscapeLayerInfoObject : public UObject
 	FName LayerName;
 
 	UPROPERTY(EditAnywhere, Category=LandscapeLayerInfoObject)
-	UPhysicalMaterial* PhysMaterial;
+	TObjectPtr<UPhysicalMaterial> PhysMaterial;
 
 	UPROPERTY(EditAnywhere, Category=LandscapeLayerInfoObject)
 	float Hardness;
@@ -44,7 +44,7 @@ class ULandscapeLayerInfoObject : public UObject
 
 	/** Texture to modulate the Splines Falloff Layer Alpha */
 	UPROPERTY(EditAnywhere, Category = SplineFalloffModulation, Meta = (DisplayName="Texture"))
-	UTexture2D* SplineFalloffModulationTexture;
+	TObjectPtr<UTexture2D> SplineFalloffModulationTexture;
 
 	UPROPERTY(EditAnywhere, Category = SplineFalloffModulation, Meta = (DisplayName = "Color Mask"))
 	ESplineModulationColorMask SplineFalloffModulationColorMask;
diff --git a/Engine/Source/Runtime/Landscape/Classes/LandscapeMeshProxyActor.h b/Engine/Source/Runtime/Landscape/Classes/LandscapeMeshProxyActor.h
--- a/Engine/Source/Runtime/Landscape/Classes/LandscapeMeshProxyActor.h
+++ b/Engine/Source/Runtime/Landscape/Classes/LandscapeMeshProxyActor.h
@@ -15,7 +15,7 @@ class ALandscapeMeshProxyActor : public AActor
 
 private:
 	UPROPERTY(Category = LandscapeMeshProxyActor, VisibleAnywhere, BlueprintReadOnly, meta = (ExposeFunctionCategories = "Mesh,Rendering,Physics,Components|StaticMesh", AllowPrivateAccess = "true"))
-	class ULandscapeMeshProxyComponent* LandscapeMeshProxyComponent;
+	TObjectPtr<class ULandscapeMeshProxyComponent> LandscapeMeshProxyComponent;
 
 public:
 	/** Returns StaticMeshComponent subobject **/
diff --git a/Engine/Source/Runtime/Landscape/Classes/LandscapeProxy.h b/Engine/Source/Runtime/Landscape/Classes/LandscapeProxy.h
--- a/Engine/Source/Runtime/Landscape/Classes/LandscapeProxy.h
+++ b/Engine/Source/Runtime/Landscape/Classes/LandscapeProxy.h
@@ -51,7 +51,7 @@ struct FLandscapeEditorLayerSettings
 
 #if WITH_EDITORONLY_DATA
 	UPROPERTY()
-	ULandscapeLayerInfoObject* LayerInfoObj;
+	TObjectPtr<ULandscapeLayerInfoObject> LayerInfoObj;
 
 	UPROPERTY()
 	FString ReimportLayerFilePath;
@@ -82,14 +82,14 @@ struct FLandscapeLayerStruct
 	GENERATED_USTRUCT_BODY()
 
 	UPROPERTY()
-	ULandscapeLayerInfoObject* LayerInfoObj;
+	TObjectPtr<ULandscapeLayerInfoObject> LayerInfoObj;
 
 #if WITH_EDITORONLY_DATA
 	UPROPERTY(transient)
-	ULandscapeMaterialInstanceConstant* ThumbnailMIC;
+	TObjectPtr<ULandscapeMaterialInstanceConstant> ThumbnailMIC;
 
 	UPROPERTY()
-	ALandscapeProxy* Owner;
+	TObjectPtr<ALandscapeProxy> Owner;
 
 	UPROPERTY(transient)
 	int32 DebugColorChannel;
@@ -139,7 +139,7 @@ struct FLandscapeImportLayerInfo
 	FName LayerName;
 
 	UPROPERTY(Category="Import", EditAnywhere)
-	ULandscapeLayerInfoObject* LayerInfo;
+	TObjectPtr<ULandscapeLayerInfoObject> LayerInfo;
 
 	UPROPERTY(Category="Import", EditAnywhere, meta=(DisplayName="Layer File"))
 	FString SourceFilePath; // Optional
@@ -326,7 +326,7 @@ struct FLandscapeProxyMaterialOverride
 	FPerPlatformInt LODIndex;
 
 	UPROPERTY(EditAnywhere, Category = Landscape)
-	UMaterialInterface* Material = nullptr;
+	TObjectPtr<UMaterialInterface> Material = nullptr;
 
 #if WITH_EDITORONLY_DATA
 	bool operator==(const FLandscapeProxyMaterialOverride& InOther) const
@@ -366,7 +366,7 @@ public:
 
 protected:
 	UPROPERTY()
-	ULandscapeSplinesComponent* SplineComponent;
+	TObjectPtr<ULandscapeSplinesComponent> SplineComponent;
 
 	/** Guid for LandscapeEditorInfo **/
 	UPROPERTY()
@@ -444,7 +444,7 @@ public:
 
 	/** Default physical material, used when no per-layer values physical materials */
 	UPROPERTY(EditAnywhere, Category=Landscape)
-	UPhysicalMaterial* DefaultPhysMaterial;
+	TObjectPtr<UPhysicalMaterial> DefaultPhysMaterial;
 
 	/**
 	 * Allows artists to adjust the distance where textures using UV 0 are streamed in/out.
@@ -456,7 +456,7 @@ public:
 
 	/** Combined material used to render the landscape */
 	UPROPERTY(EditAnywhere, BlueprintSetter=EditorSetLandscapeMaterial, Category=Landscape)
-	UMaterialInterface* LandscapeMaterial;
+	TObjectPtr<UMaterialInterface> LandscapeMaterial;
 
 #if !WITH_EDITORONLY_DATA
 	/** Used to cache grass types from GetGrassTypes */
@@ -471,17 +471,17 @@ public:
 
 	/** Material used to render landscape components with holes. If not set, LandscapeMaterial will be used (blend mode will be overridden to Masked if it is set to Opaque) */
 	UPROPERTY(EditAnywhere, Category=Landscape, AdvancedDisplay)
-	UMaterialInterface* LandscapeHoleMaterial;
+	TObjectPtr<UMaterialInterface> LandscapeHoleMaterial;
 
 	UPROPERTY(EditAnywhere, Category = Landscape)
 	TArray<FLandscapeProxyMaterialOverride> LandscapeMaterialsOverride;
 
 #if WITH_EDITORONLY_DATA
 	UPROPERTY(Transient)
-	UMaterialInterface* PreEditLandscapeMaterial;
+	TObjectPtr<UMaterialInterface> PreEditLandscapeMaterial;
 
 	UPROPERTY(Transient)
-	UMaterialInterface* PreEditLandscapeHoleMaterial;
+	TObjectPtr<UMaterialInterface> PreEditLandscapeHoleMaterial;
 
 	UPROPERTY(Transient)
 	TArray<FLandscapeProxyMaterialOverride> PreEditLandscapeMaterialsOverride;
@@ -503,7 +503,7 @@ public:
 	 * The material also needs to be set up to output to a virtual texture.
 	 */
 	UPROPERTY(EditAnywhere, BlueprintReadWrite, Category = VirtualTexture, meta = (DisplayName = "Draw in Virtual Textures"))
-	TArray<URuntimeVirtualTexture*> RuntimeVirtualTextures;
+	TArray<TObjectPtr<URuntimeVirtualTexture>> RuntimeVirtualTextures;
 
 	/** 
 	 * Number of mesh levels to use when rendering landscape into runtime virtual texture.
@@ -537,14 +537,14 @@ public:
 
 	/** The array of LandscapeComponent that are used by the landscape */
 	UPROPERTY()
-	TArray<ULandscapeComponent*> LandscapeComponents;
+	TArray<TObjectPtr<ULandscapeComponent>> LandscapeComponents;
 
 	/** Array of LandscapeHeightfieldCollisionComponent */
 	UPROPERTY()
-	TArray<ULandscapeHeightfieldCollisionComponent*> CollisionComponents;
+	TArray<TObjectPtr<ULandscapeHeightfieldCollisionComponent>> CollisionComponents;
 
 	UPROPERTY(transient, duplicatetransient)
-	TArray<UHierarchicalInstancedStaticMeshComponent*> FoliageComponents;
+	TArray<TObjectPtr<UHierarchicalInstancedStaticMeshComponent>> FoliageComponents;
 
 	/** A transient data structure for tracking the grass */
 	FCachedLandscapeFoliage FoliageCache;
@@ -671,7 +671,7 @@ public:
 
 #if WITH_EDITORONLY_DATA
 	UPROPERTY()
-	TArray<ULandscapeLayerInfoObject*> EditorCachedLayerInfos_DEPRECATED;
+	TArray<TObjectPtr<ULandscapeLayerInfoObject>> EditorCachedLayerInfos_DEPRECATED;
 
 	UPROPERTY()
 	FString ReimportHeightmapFilePath;
@@ -740,7 +740,7 @@ public:
 
 	/** Map of weightmap usage */
 	UPROPERTY(Transient)
-	TMap<UTexture2D*, ULandscapeWeightmapUsage*> WeightmapUsageMap;
+	TMap<TObjectPtr<UTexture2D>, TObjectPtr<ULandscapeWeightmapUsage>> WeightmapUsageMap;
 
 	// Blueprint functions
 
diff --git a/Engine/Source/Runtime/Landscape/Classes/LandscapeSplineControlPoint.h b/Engine/Source/Runtime/Landscape/Classes/LandscapeSplineControlPoint.h
--- a/Engine/Source/Runtime/Landscape/Classes/LandscapeSplineControlPoint.h
+++ b/Engine/Source/Runtime/Landscape/Classes/LandscapeSplineControlPoint.h
@@ -36,7 +36,7 @@ struct FLandscapeSplineConnection
 
 	// Segment connected to this control point
 	UPROPERTY()
-	ULandscapeSplineSegment* Segment;
+	TObjectPtr<ULandscapeSplineSegment> Segment;
 
 	// Which end of the segment is connected to this control point
 	UPROPERTY()
@@ -110,11 +110,11 @@ class ULandscapeSplineControlPoint : public UObject
 
 	/** Mesh to use on the control point */
 	UPROPERTY(EditAnywhere, Category=Mesh)
-	UStaticMesh* Mesh;
+	TObjectPtr<UStaticMesh> Mesh;
 
 	/** Overrides mesh's materials */
 	UPROPERTY(EditAnywhere, Category=Mesh)
-	TArray<UMaterialInterface*> MaterialOverrides;
+	TArray<TObjectPtr<UMaterialInterface>> MaterialOverrides;
 
 	/** Scale of the control point mesh */
 	UPROPERTY(EditAnywhere, Category=Mesh)
@@ -171,7 +171,7 @@ class ULandscapeSplineControlPoint : public UObject
 	 * The material also needs to be set up to output to a virtual texture. 
 	 */
 	UPROPERTY(EditAnywhere, Category = VirtualTexture, meta = (DisplayName = "Draw in Virtual Textures"))
-	TArray<URuntimeVirtualTexture*> RuntimeVirtualTextures;
+	TArray<TObjectPtr<URuntimeVirtualTexture>> RuntimeVirtualTextures;
 
 	/** Lod bias for rendering to runtime virtual texture. */
 	UPROPERTY(EditAnywhere, AdvancedDisplay, Category = VirtualTexture, meta = (DisplayName = "Virtual Texture LOD Bias", UIMin = "-7", UIMax = "8"))
@@ -221,7 +221,7 @@ protected:
 
 	/** Control point mesh */
 	UPROPERTY(TextExportTransient)
-	UControlPointMeshComponent* LocalMeshComponent;
+	TObjectPtr<UControlPointMeshComponent> LocalMeshComponent;
 
 #if WITH_EDITORONLY_DATA
 	/** World reference for if mesh component is stored in another streaming level */
diff --git a/Engine/Source/Runtime/Landscape/Classes/LandscapeSplineSegment.h b/Engine/Source/Runtime/Landscape/Classes/LandscapeSplineSegment.h
--- a/Engine/Source/Runtime/Landscape/Classes/LandscapeSplineSegment.h
+++ b/Engine/Source/Runtime/Landscape/Classes/LandscapeSplineSegment.h
@@ -99,7 +99,7 @@ struct FLandscapeSplineSegmentConnection
 
 	// Control point connected to this end of the segment
 	UPROPERTY()
-	ULandscapeSplineControlPoint* ControlPoint;
+	TObjectPtr<ULandscapeSplineControlPoint> ControlPoint;
 
 	// Tangent length of the connection
 	UPROPERTY(EditAnywhere, Category=LandscapeSplineSegmentConnection)
@@ -133,11 +133,11 @@ struct FLandscapeSplineMeshEntry
 
 	/** Mesh to use on the spline */
 	UPROPERTY(EditAnywhere, Category=LandscapeSplineMeshEntry)
-	UStaticMesh* Mesh;
+	TObjectPtr<UStaticMesh> Mesh;
 
 	/** Overrides mesh's materials */
 	UPROPERTY(EditAnywhere, Category=LandscapeSplineMeshEntry, AdvancedDisplay)
-	TArray<UMaterialInterface*> MaterialOverrides;
+	TArray<TObjectPtr<UMaterialInterface>> MaterialOverrides;
 
 	/** Whether to automatically center the mesh horizontally on the spline */
 	UPROPERTY(EditAnywhere, Category=LandscapeSplineMeshEntry, meta=(DisplayName="Center Horizontally"))
@@ -273,7 +273,7 @@ class ULandscapeSplineSegment : public UObject
 	 * The material also needs to be set up to output to a virtual texture. 
 	 */
 	UPROPERTY(EditAnywhere, Category = VirtualTexture, meta = (DisplayName = "Draw in Virtual Textures"))
-	TArray<URuntimeVirtualTexture*> RuntimeVirtualTextures;
+	TArray<TObjectPtr<URuntimeVirtualTexture>> RuntimeVirtualTextures;
 
 	/** Lod bias for rendering to runtime virtual texture. */
 	UPROPERTY(EditAnywhere, AdvancedDisplay, Category = VirtualTexture, meta = (DisplayName = "Virtual Texture LOD Bias", UIMin = "-7", UIMax = "8"))
@@ -323,7 +323,7 @@ protected:
 
 	/** Spline meshes */
 	UPROPERTY(TextExportTransient)
-	TArray<USplineMeshComponent*> LocalMeshComponents;
+	TArray<TObjectPtr<USplineMeshComponent>> LocalMeshComponents;
 
 #if WITH_EDITORONLY_DATA
 	/** World references for mesh components stored in other streaming levels */
diff --git a/Engine/Source/Runtime/Landscape/Classes/LandscapeSplinesComponent.h b/Engine/Source/Runtime/Landscape/Classes/LandscapeSplinesComponent.h
--- a/Engine/Source/Runtime/Landscape/Classes/LandscapeSplinesComponent.h
+++ b/Engine/Source/Runtime/Landscape/Classes/LandscapeSplinesComponent.h
@@ -32,7 +32,7 @@ struct FForeignControlPointData
 	FGuid ModificationKey;
 
 	UPROPERTY()
-	UControlPointMeshComponent* MeshComponent = nullptr;
+	TObjectPtr<UControlPointMeshComponent> MeshComponent = nullptr;
 
 	UPROPERTY()
 	TLazyObjectPtr<ULandscapeSplineControlPoint> Identifier;
@@ -54,7 +54,7 @@ struct FForeignSplineSegmentData
 	FGuid ModificationKey;
 
 	UPROPERTY()
-	TArray<USplineMeshComponent*> MeshComponents;
+	TArray<TObjectPtr<USplineMeshComponent>> MeshComponents;
 
 	UPROPERTY()
 	TLazyObjectPtr<ULandscapeSplineSegment> Identifier;
@@ -113,11 +113,11 @@ class ULandscapeSplinesComponent : public UPrimitiveComponent
 
 	/** Sprite used to draw control points */
 	UPROPERTY()
-	UTexture2D* ControlPointSprite;
+	TObjectPtr<UTexture2D> ControlPointSprite;
 
 	/** Mesh used to draw splines that have no mesh */
 	UPROPERTY()
-	UStaticMesh* SplineEditorMesh;
+	TObjectPtr<UStaticMesh> SplineEditorMesh;
 
 	/** Whether we are in-editor and showing spline editor meshes */
 	UPROPERTY(NonTransactional, Transient)
@@ -126,10 +126,10 @@ class ULandscapeSplinesComponent : public UPrimitiveComponent
 
 protected:
 	UPROPERTY(TextExportTransient)
-	TArray<ULandscapeSplineControlPoint*> ControlPoints;
+	TArray<TObjectPtr<ULandscapeSplineControlPoint>> ControlPoints;
 
 	UPROPERTY(TextExportTransient)
-	TArray<ULandscapeSplineSegment*> Segments;
+	TArray<TObjectPtr<ULandscapeSplineSegment>> Segments;
 
 #if WITH_EDITORONLY_DATA
 	// Serialized
@@ -144,7 +144,7 @@ protected:
 	// References to components owned by landscape splines in other levels
 	// for cooked build (uncooked keeps references via ForeignWorldSplineDataMap)
 	UPROPERTY(TextExportTransient)
-	TArray<UMeshComponent*> CookedForeignMeshComponents;
+	TArray<TObjectPtr<UMeshComponent>> CookedForeignMeshComponents;
 
 public:
 	/** Get a list of spline mesh components representing this landscape spline (Editor only) */
diff --git a/Engine/Source/Runtime/Landscape/Classes/LandscapeWeightmapUsage.h b/Engine/Source/Runtime/Landscape/Classes/LandscapeWeightmapUsage.h
--- a/Engine/Source/Runtime/Landscape/Classes/LandscapeWeightmapUsage.h
+++ b/Engine/Source/Runtime/Landscape/Classes/LandscapeWeightmapUsage.h
@@ -18,7 +18,7 @@ class ULandscapeWeightmapUsage : public UObject
 
 public:
 	UPROPERTY()
-	ULandscapeComponent* ChannelUsage[NumChannels];
+	TObjectPtr<ULandscapeComponent> ChannelUsage[NumChannels];
 
 	UPROPERTY()
 	FGuid LayerGuid;
diff --git a/Engine/Source/Runtime/Landscape/Classes/Materials/MaterialExpressionLandscapeGrassOutput.h b/Engine/Source/Runtime/Landscape/Classes/Materials/MaterialExpressionLandscapeGrassOutput.h
--- a/Engine/Source/Runtime/Landscape/Classes/Materials/MaterialExpressionLandscapeGrassOutput.h
+++ b/Engine/Source/Runtime/Landscape/Classes/Materials/MaterialExpressionLandscapeGrassOutput.h
@@ -22,7 +22,7 @@ struct FGrassInput
 	FName Name;
 
 	UPROPERTY(EditAnywhere, Category = Grass)
-	ULandscapeGrassType* GrassType;
+	TObjectPtr<ULandscapeGrassType> GrassType;
 
 	UPROPERTY(meta = (RequiredInput = "true"))
 	FExpressionInput Input;
diff --git a/Engine/Source/Runtime/Landscape/Classes/Materials/MaterialExpressionLandscapePhysicalMaterialOutput.h b/Engine/Source/Runtime/Landscape/Classes/Materials/MaterialExpressionLandscapePhysicalMaterialOutput.h
--- a/Engine/Source/Runtime/Landscape/Classes/Materials/MaterialExpressionLandscapePhysicalMaterialOutput.h
+++ b/Engine/Source/Runtime/Landscape/Classes/Materials/MaterialExpressionLandscapePhysicalMaterialOutput.h
@@ -14,7 +14,7 @@ struct FPhysicalMaterialInput
 	GENERATED_USTRUCT_BODY()
 
 	UPROPERTY(EditAnywhere, Category = PhysicalMaterial)
-	class UPhysicalMaterial* PhysicalMaterial;
+	TObjectPtr<class UPhysicalMaterial> PhysicalMaterial;
 
 	UPROPERTY(meta = (RequiredInput = "true"))
 	FExpressionInput Input;
diff --git a/Engine/Source/Runtime/Landscape/Public/LandscapeBlueprintBrushBase.h b/Engine/Source/Runtime/Landscape/Public/LandscapeBlueprintBrushBase.h
--- a/Engine/Source/Runtime/Landscape/Public/LandscapeBlueprintBrushBase.h
+++ b/Engine/Source/Runtime/Landscape/Public/LandscapeBlueprintBrushBase.h
@@ -18,7 +18,7 @@ class LANDSCAPE_API ALandscapeBlueprintBrushBase : public AActor
 protected:
 #if WITH_EDITORONLY_DATA
 	UPROPERTY(Transient)
-	class ALandscape* OwningLandscape;
+	TObjectPtr<class ALandscape> OwningLandscape;
 
 	UPROPERTY(Category = "Settings", EditAnywhere, BlueprintReadWrite)
 	bool AffectHeightmap;
diff --git a/Engine/Source/Runtime/LevelSequence/Public/DefaultLevelSequenceInstanceData.h b/Engine/Source/Runtime/LevelSequence/Public/DefaultLevelSequenceInstanceData.h
--- a/Engine/Source/Runtime/LevelSequence/Public/DefaultLevelSequenceInstanceData.h
+++ b/Engine/Source/Runtime/LevelSequence/Public/DefaultLevelSequenceInstanceData.h
@@ -30,7 +30,7 @@ public:
 
 	/** When set, this actor's world position will be used as the transform origin for all absolute transform sections */
 	UPROPERTY(EditAnywhere, BlueprintReadWrite, Category="General")
-	AActor* TransformOriginActor;
+	TObjectPtr<AActor> TransformOriginActor;
 
 	/** Specifies a transform that offsets all absolute transform sections in this sequence. Will compound with attach tracks. Scale is ignored. Not applied to Relative or Additive sections. */
 	UPROPERTY(EditAnywhere, BlueprintReadWrite, Category="General")
diff --git a/Engine/Source/Runtime/LevelSequence/Public/LevelSequence.h b/Engine/Source/Runtime/LevelSequence/Public/LevelSequence.h
--- a/Engine/Source/Runtime/LevelSequence/Public/LevelSequence.h
+++ b/Engine/Source/Runtime/LevelSequence/Public/LevelSequence.h
@@ -30,7 +30,7 @@ public:
 
 	/** Pointer to the movie scene that controls this animation. */
 	UPROPERTY()
-	UMovieScene* MovieScene;
+	TObjectPtr<UMovieScene> MovieScene;
 
 public:
 
@@ -123,7 +123,7 @@ protected:
 
 	/** A pointer to the director blueprint that generates this sequence's DirectorClass. */
 	UPROPERTY()
-	UBlueprint* DirectorBlueprint;
+	TObjectPtr<UBlueprint> DirectorBlueprint;
 
 #endif
 
@@ -132,7 +132,7 @@ protected:
 	 * Director instances are allocated on-demand one per sequence during evaluation and are used by event tracks for triggering events.
 	 */
 	UPROPERTY()
-	UClass* DirectorClass;
+	TObjectPtr<UClass> DirectorClass;
 
 public:
 	/**
@@ -267,12 +267,12 @@ private:
 
 	/** Array of meta-data objects associated with this level sequence. Each pointer may implement the ILevelSequenceMetaData interface in order to hook into default ULevelSequence functionality. */
 	UPROPERTY()
-	TArray<UObject*> MetaDataObjects;
+	TArray<TObjectPtr<UObject>> MetaDataObjects;
 
 #endif
 
 protected:
 	/** Array of user data stored with the asset */
 	UPROPERTY(EditAnywhere, AdvancedDisplay, Instanced, Category = Animation)
-	TArray<UAssetUserData*> AssetUserData;
+	TArray<TObjectPtr<UAssetUserData>> AssetUserData;
 };
diff --git a/Engine/Source/Runtime/LevelSequence/Public/LevelSequenceActor.h b/Engine/Source/Runtime/LevelSequence/Public/LevelSequenceActor.h
--- a/Engine/Source/Runtime/LevelSequence/Public/LevelSequenceActor.h
+++ b/Engine/Source/Runtime/LevelSequence/Public/LevelSequenceActor.h
@@ -46,7 +46,7 @@ public:
 	FSoftClassPath BurnInClass;
 
 	UPROPERTY(Instanced, EditAnywhere, BlueprintReadWrite, Category="General", meta=(EditCondition=bUseBurnIn))
-	ULevelSequenceBurnInInitSettings* Settings;
+	TObjectPtr<ULevelSequenceBurnInInitSettings> Settings;
 
 protected:
 
@@ -79,7 +79,7 @@ public:
 	FMovieSceneSequencePlaybackSettings PlaybackSettings;
 
 	UPROPERTY(Instanced, transient, replicated, BlueprintReadOnly, BlueprintGetter=GetSequencePlayer, Category="Playback", meta=(ExposeFunctionCategories="Game|Cinematic"))
-	ULevelSequencePlayer* SequencePlayer;
+	TObjectPtr<ULevelSequencePlayer> SequencePlayer;
 
 	UPROPERTY(EditAnywhere, BlueprintReadOnly, Category="General", meta=(AllowedClasses="LevelSequence"))
 	FSoftObjectPath LevelSequence;
@@ -88,11 +88,11 @@ public:
 	FLevelSequenceCameraSettings CameraSettings;
 
 	UPROPERTY(Instanced, BlueprintReadOnly, Category="General")
-	ULevelSequenceBurnInOptions* BurnInOptions;
+	TObjectPtr<ULevelSequenceBurnInOptions> BurnInOptions;
 
 	/** Mapping of actors to override the sequence bindings with */
 	UPROPERTY(Instanced, BlueprintReadOnly, Category="General")
-	UMovieSceneBindingOverrides* BindingOverrides;
+	TObjectPtr<UMovieSceneBindingOverrides> BindingOverrides;
 
 	UPROPERTY()
 	uint8 bAutoPlay_DEPRECATED : 1;
@@ -107,7 +107,7 @@ public:
 
 	/** Instance data that can be used to dynamically control sequence evaluation at runtime */
 	UPROPERTY(Instanced, BlueprintReadWrite, Category="General")
-	UObject* DefaultInstanceData;
+	TObjectPtr<UObject> DefaultInstanceData;
 
 public:
 
@@ -289,7 +289,7 @@ public:
 private:
 	/** Burn-in widget */
 	UPROPERTY()
-	ULevelSequenceBurnIn* BurnInInstance;
+	TObjectPtr<ULevelSequenceBurnIn> BurnInInstance;
 
 	UPROPERTY()
 	bool bShowBurnin;
@@ -304,7 +304,7 @@ struct FBoundActorProxy
 
 	/** Specifies the actor to override the binding with */
 	UPROPERTY(EditInstanceOnly, AdvancedDisplay, Category="General")
-	AActor* BoundActor = nullptr;
+	TObjectPtr<AActor> BoundActor = nullptr;
 
 	void Initialize(TSharedPtr<IPropertyHandle> InPropertyHandle);
 
diff --git a/Engine/Source/Runtime/LevelSequence/Public/LevelSequenceBurnIn.h b/Engine/Source/Runtime/LevelSequence/Public/LevelSequenceBurnIn.h
--- a/Engine/Source/Runtime/LevelSequence/Public/LevelSequenceBurnIn.h
+++ b/Engine/Source/Runtime/LevelSequence/Public/LevelSequenceBurnIn.h
@@ -46,5 +46,5 @@ protected:
 
 	/** The actor to get our burn in frames from */
 	UPROPERTY(BlueprintReadOnly, Category="Burn In")
-	ALevelSequenceActor* LevelSequenceActor;
+	TObjectPtr<ALevelSequenceActor> LevelSequenceActor;
 };
diff --git a/Engine/Source/Runtime/LevelSequence/Public/LevelSequenceDirector.h b/Engine/Source/Runtime/LevelSequence/Public/LevelSequenceDirector.h
--- a/Engine/Source/Runtime/LevelSequence/Public/LevelSequenceDirector.h
+++ b/Engine/Source/Runtime/LevelSequence/Public/LevelSequenceDirector.h
@@ -23,7 +23,7 @@ public:
 
 	/** Pointer to the player that's playing back this director's sequence */
 	UPROPERTY(BlueprintReadOnly, Category="Cinematics")
-	ULevelSequencePlayer* Player;
+	TObjectPtr<ULevelSequencePlayer> Player;
 };
 
 
diff --git a/Engine/Source/Runtime/LevelSequence/Public/LevelSequencePlayer.h b/Engine/Source/Runtime/LevelSequence/Public/LevelSequencePlayer.h
--- a/Engine/Source/Runtime/LevelSequence/Public/LevelSequencePlayer.h
+++ b/Engine/Source/Runtime/LevelSequence/Public/LevelSequencePlayer.h
@@ -86,7 +86,7 @@ struct FLevelSequencePlayerSnapshot
 	FLevelSequenceSnapshotSettings Settings;
 
 	UPROPERTY(VisibleAnywhere, BlueprintReadOnly, Category = "General")
-	ULevelSequence* ActiveShot = nullptr;
+	TObjectPtr<ULevelSequence> ActiveShot = nullptr;
 
 	UPROPERTY()
 	FMovieSceneSequenceID ShotID;
diff --git a/Engine/Source/Runtime/LevelSequence/Public/SequenceMediaController.h b/Engine/Source/Runtime/LevelSequence/Public/SequenceMediaController.h
--- a/Engine/Source/Runtime/LevelSequence/Public/SequenceMediaController.h
+++ b/Engine/Source/Runtime/LevelSequence/Public/SequenceMediaController.h
@@ -79,11 +79,11 @@ private:
 
 	/** Pointer to the sequence actor to use for playback */
 	UPROPERTY(EditAnywhere, BlueprintGetter=GetSequence, Category="Synchronization")
-	ALevelSequenceActor* Sequence;
+	TObjectPtr<ALevelSequenceActor> Sequence;
 
 	/** Media component that contains the media player to synchronize with */
 	UPROPERTY(VisibleAnywhere, BlueprintGetter=GetMediaComponent, Category="Synchronization")
-	UMediaComponent* MediaComponent;
+	TObjectPtr<UMediaComponent> MediaComponent;
 
 	/** Replicated time at which the server started the sequence (taken from AGameStateBase::GetServerWorldTimeSeconds) */
 	UPROPERTY(BlueprintReadOnly, Category="Synchronization", replicated, ReplicatedUsing=OnRep_ServerStartTimeSeconds, meta=(AllowPrivateAccess="true"))
diff --git a/Engine/Source/Runtime/LiveLinkInterface/Public/LiveLinkPresetTypes.h b/Engine/Source/Runtime/LiveLinkInterface/Public/LiveLinkPresetTypes.h
--- a/Engine/Source/Runtime/LiveLinkInterface/Public/LiveLinkPresetTypes.h
+++ b/Engine/Source/Runtime/LiveLinkInterface/Public/LiveLinkPresetTypes.h
@@ -23,7 +23,7 @@ struct FLiveLinkSourcePreset
 	FGuid Guid;
 
 	UPROPERTY(VisibleAnywhere, Category = "LiveLinkSourcePresets")
-	ULiveLinkSourceSettings* Settings = nullptr;
+	TObjectPtr<ULiveLinkSourceSettings> Settings = nullptr;
 
 	/** The SourceType when the source was saved to a Preset. */
 	UPROPERTY(VisibleAnywhere, Category = "LiveLinkSourcePresets")
@@ -43,10 +43,10 @@ struct FLiveLinkSubjectPreset
 	TSubclassOf<ULiveLinkRole> Role;
 
 	UPROPERTY(VisibleAnywhere, Category = "LiveLinkSubjectPresets")
-	ULiveLinkSubjectSettings* Settings = nullptr;
+	TObjectPtr<ULiveLinkSubjectSettings> Settings = nullptr;
 
 	UPROPERTY(VisibleAnywhere, Category = "LiveLinkSubjectPresets")
-	ULiveLinkVirtualSubject* VirtualSubject = nullptr;
+	TObjectPtr<ULiveLinkVirtualSubject> VirtualSubject = nullptr;
 
 	UPROPERTY(VisibleAnywhere, Category = "LiveLinkSubjectPresets")
 	bool bEnabled = false;
diff --git a/Engine/Source/Runtime/LiveLinkInterface/Public/LiveLinkSubjectSettings.h b/Engine/Source/Runtime/LiveLinkInterface/Public/LiveLinkSubjectSettings.h
--- a/Engine/Source/Runtime/LiveLinkInterface/Public/LiveLinkSubjectSettings.h
+++ b/Engine/Source/Runtime/LiveLinkInterface/Public/LiveLinkSubjectSettings.h
@@ -24,15 +24,15 @@ public:
 
 	/** List of available preprocessor the subject will use. */
 	UPROPERTY(EditAnywhere, Instanced, Category = "LiveLink", meta = (DisplayName = "Pre Processors"))
-	TArray<ULiveLinkFramePreProcessor*> PreProcessors;
+	TArray<TObjectPtr<ULiveLinkFramePreProcessor>> PreProcessors;
 
 	/** The interpolation processor the subject will use. */
 	UPROPERTY(EditAnywhere, Instanced, Category = "LiveLink", meta = (DisplayName = "Interpolation"))
-	ULiveLinkFrameInterpolationProcessor* InterpolationProcessor;
+	TObjectPtr<ULiveLinkFrameInterpolationProcessor> InterpolationProcessor;
 
 	/** List of available translator the subject can use. */
 	UPROPERTY(EditAnywhere, Instanced, Category = "LiveLink", meta = (DisplayName = "Translators"))
-	TArray<ULiveLinkFrameTranslator*> Translators;
+	TArray<TObjectPtr<ULiveLinkFrameTranslator>> Translators;
 
 	UPROPERTY()
 	TSubclassOf<ULiveLinkRole> Role;
diff --git a/Engine/Source/Runtime/LiveLinkInterface/Public/LiveLinkVirtualSubject.h b/Engine/Source/Runtime/LiveLinkInterface/Public/LiveLinkVirtualSubject.h
--- a/Engine/Source/Runtime/LiveLinkInterface/Public/LiveLinkVirtualSubject.h
+++ b/Engine/Source/Runtime/LiveLinkInterface/Public/LiveLinkVirtualSubject.h
@@ -57,7 +57,7 @@ protected:
 
 	/** List of available translator the subject can use. */
 	UPROPERTY(EditAnywhere, Instanced, Category = "LiveLink", meta=(DisplayName="Translators"))
-	TArray<ULiveLinkFrameTranslator*> FrameTranslators;
+	TArray<TObjectPtr<ULiveLinkFrameTranslator>> FrameTranslators;
 
 	/** LiveLinkClient to get access to subjects */
 	ILiveLinkClient* LiveLinkClient;
diff --git a/Engine/Source/Runtime/MRMesh/Public/MRMeshComponent.h b/Engine/Source/Runtime/MRMesh/Public/MRMeshComponent.h
--- a/Engine/Source/Runtime/MRMesh/Public/MRMeshComponent.h
+++ b/Engine/Source/Runtime/MRMesh/Public/MRMeshComponent.h
@@ -168,10 +168,10 @@ private:
 	void ClearAllBrickData_Internal();
 
 	UPROPERTY(EditAnywhere, Category = Appearance)
-	UMaterialInterface* Material;
+	TObjectPtr<UMaterialInterface> Material;
 
 	UPROPERTY(EditAnywhere, Category = Appearance)
-	UMaterialInterface* WireframeMaterial;
+	TObjectPtr<UMaterialInterface> WireframeMaterial;
 
 	/** If true, MRMesh will create a renderable mesh proxy.  If false it will not, but could still provide collision. */
 	UPROPERTY(EditAnywhere, Category = Appearance)
@@ -188,10 +188,10 @@ private:
 	bool bConnected = false;
 
 	UPROPERTY(Transient)
-	class UBodySetup* CachedBodySetup = nullptr;
+	TObjectPtr<class UBodySetup> CachedBodySetup = nullptr;
 
 	UPROPERTY(Transient)
-	TArray<UBodySetup*> BodySetups;
+	TArray<TObjectPtr<UBodySetup>> BodySetups;
 
 
 	/** Whether this mesh should write z-depth to occlude meshes or not */
diff --git a/Engine/Source/Runtime/MRMesh/Public/MockDataMeshTrackerComponent.h b/Engine/Source/Runtime/MRMesh/Public/MockDataMeshTrackerComponent.h
--- a/Engine/Source/Runtime/MRMesh/Public/MockDataMeshTrackerComponent.h
+++ b/Engine/Source/Runtime/MRMesh/Public/MockDataMeshTrackerComponent.h
@@ -109,7 +109,7 @@ public:
 
 	/** MRMeshComponent can render and provide collision based on the Mesh data. */
 	UPROPERTY(transient)
-	class UMRMeshComponent* MRMesh;
+	TObjectPtr<class UMRMeshComponent> MRMesh;
 
 	/**
 		The procedural mesh generated should bake physics and update its collision data.
diff --git a/Engine/Source/Runtime/MaterialShaderQualitySettings/Classes/MaterialShaderQualitySettings.h b/Engine/Source/Runtime/MaterialShaderQualitySettings/Classes/MaterialShaderQualitySettings.h
--- a/Engine/Source/Runtime/MaterialShaderQualitySettings/Classes/MaterialShaderQualitySettings.h
+++ b/Engine/Source/Runtime/MaterialShaderQualitySettings/Classes/MaterialShaderQualitySettings.h
@@ -38,7 +38,7 @@ private:
 	UShaderPlatformQualitySettings* GetOrCreatePlatformSettings(FName ShaderPlatformName);
 
 	UPROPERTY()
-	TMap<FName, UShaderPlatformQualitySettings*> ForwardSettingMap;
+	TMap<FName, TObjectPtr<UShaderPlatformQualitySettings>> ForwardSettingMap;
 
 #if WITH_EDITORONLY_DATA
 	UShaderPlatformQualitySettings* PreviewPlatformSettings;
diff --git a/Engine/Source/Runtime/MediaAssets/Public/MediaComponent.h b/Engine/Source/Runtime/MediaAssets/Public/MediaComponent.h
--- a/Engine/Source/Runtime/MediaAssets/Public/MediaComponent.h
+++ b/Engine/Source/Runtime/MediaAssets/Public/MediaComponent.h
@@ -33,9 +33,9 @@ private:
 
 	/** This component's media texture */
 	UPROPERTY(Category="Media", Instanced, Transient, BlueprintGetter=GetMediaTexture)
-	UMediaTexture* MediaTexture;
+	TObjectPtr<UMediaTexture> MediaTexture;
 
 	/** This component's media player */
 	UPROPERTY(Category="Media", Instanced, Transient, BlueprintGetter=GetMediaPlayer, Interp)
-	UMediaPlayer* MediaPlayer;
+	TObjectPtr<UMediaPlayer> MediaPlayer;
 };
\ No newline at end of file
diff --git a/Engine/Source/Runtime/MediaAssets/Public/MediaPlayer.h b/Engine/Source/Runtime/MediaAssets/Public/MediaPlayer.h
--- a/Engine/Source/Runtime/MediaAssets/Public/MediaPlayer.h
+++ b/Engine/Source/Runtime/MediaAssets/Public/MediaPlayer.h
@@ -1110,7 +1110,7 @@ protected:
 	 * @see OpenPlaylist, OpenPlaylistIndex
 	 */
 	UPROPERTY(BlueprintReadOnly, transient, Category=Playback, BlueprintGetter = GetPlayList)
-	mutable UMediaPlaylist* Playlist;
+	mutable TObjectPtr<UMediaPlaylist> Playlist;
 
 	/**
 	 * The current index of the source in the play list being played.
diff --git a/Engine/Source/Runtime/MediaAssets/Public/MediaPlaylist.h b/Engine/Source/Runtime/MediaAssets/Public/MediaPlaylist.h
--- a/Engine/Source/Runtime/MediaAssets/Public/MediaPlaylist.h
+++ b/Engine/Source/Runtime/MediaAssets/Public/MediaPlaylist.h
@@ -147,7 +147,7 @@ protected:
 
 	/** List of media sources to play. */
 	UPROPERTY(EditAnywhere, Category=Playlist)
-	TArray<UMediaSource*> Items;
+	TArray<TObjectPtr<UMediaSource>> Items;
 
 	virtual bool CanBeInCluster() const override { return false; }
 };
diff --git a/Engine/Source/Runtime/MediaAssets/Public/MediaSoundComponent.h b/Engine/Source/Runtime/MediaAssets/Public/MediaSoundComponent.h
--- a/Engine/Source/Runtime/MediaAssets/Public/MediaSoundComponent.h
+++ b/Engine/Source/Runtime/MediaAssets/Public/MediaSoundComponent.h
@@ -244,7 +244,7 @@ protected:
 	 * @see SetMediaPlayer
 	 */
 	UPROPERTY(EditAnywhere, Category="Media")
-	UMediaPlayer* MediaPlayer;
+	TObjectPtr<UMediaPlayer> MediaPlayer;
 
 private:
 
diff --git a/Engine/Source/Runtime/MediaAssets/Public/MediaTexture.h b/Engine/Source/Runtime/MediaAssets/Public/MediaTexture.h
--- a/Engine/Source/Runtime/MediaAssets/Public/MediaTexture.h
+++ b/Engine/Source/Runtime/MediaAssets/Public/MediaTexture.h
@@ -231,7 +231,7 @@ protected:
 	 * @see SetMediaPlayer
 	 */
 	UPROPERTY(EditAnywhere, Category="Media")
-	UMediaPlayer* MediaPlayer;
+	TObjectPtr<UMediaPlayer> MediaPlayer;
 
 private:
 
diff --git a/Engine/Source/Runtime/MediaAssets/Public/PlatformMediaSource.h b/Engine/Source/Runtime/MediaAssets/Public/PlatformMediaSource.h
--- a/Engine/Source/Runtime/MediaAssets/Public/PlatformMediaSource.h
+++ b/Engine/Source/Runtime/MediaAssets/Public/PlatformMediaSource.h
@@ -28,7 +28,7 @@ public:
 
 	/** Media sources per platform. */
 	UPROPERTY(transient, EditAnywhere, Category=Sources, Meta=(DisplayName="Media Sources"))
-	TMap<FString, UMediaSource*> PlatformMediaSources;
+	TMap<FString, TObjectPtr<UMediaSource>> PlatformMediaSources;
 
 private:
 	/** Blind data encountered at load that could not be mapped to a known platform */
@@ -74,5 +74,5 @@ private:
 	 * This media source will be used if no source was specified for a target platform.
 	 */
 	UPROPERTY()
-	UMediaSource* MediaSource;
+	TObjectPtr<UMediaSource> MediaSource;
 };
diff --git a/Engine/Source/Runtime/MovieScene/Private/Tests/MovieSceneTestObjects.h b/Engine/Source/Runtime/MovieScene/Private/Tests/MovieSceneTestObjects.h
--- a/Engine/Source/Runtime/MovieScene/Private/Tests/MovieSceneTestObjects.h
+++ b/Engine/Source/Runtime/MovieScene/Private/Tests/MovieSceneTestObjects.h
@@ -34,7 +34,7 @@ public:
 	bool bHighPassFilter;
 
 	UPROPERTY()
-	TArray<UMovieSceneSection*> SectionArray;
+	TArray<TObjectPtr<UMovieSceneSection>> SectionArray;
 };
 
 UCLASS(MinimalAPI)
@@ -62,7 +62,7 @@ public:
 	virtual UMovieScene* GetMovieScene() const override { return MovieScene; }
 
 	UPROPERTY()
-	UMovieScene* MovieScene;
+	TObjectPtr<UMovieScene> MovieScene;
 };
 
 UCLASS(MinimalAPI)
@@ -74,7 +74,7 @@ public:
 	virtual const TArray<UMovieSceneSection*>& GetAllSections() const override { return SectionArray; }
 
 	UPROPERTY()
-	TArray<UMovieSceneSection*> SectionArray;
+	TArray<TObjectPtr<UMovieSceneSection>> SectionArray;
 };
 
 UCLASS(MinimalAPI)
diff --git a/Engine/Source/Runtime/MovieScene/Public/Channels/MovieSceneObjectPathChannel.h b/Engine/Source/Runtime/MovieScene/Public/Channels/MovieSceneObjectPathChannel.h
--- a/Engine/Source/Runtime/MovieScene/Public/Channels/MovieSceneObjectPathChannel.h
+++ b/Engine/Source/Runtime/MovieScene/Public/Channels/MovieSceneObjectPathChannel.h
@@ -57,7 +57,7 @@ private:
 
 	/** Hard reference to the loaded object - relevant for any asset type which also hints the async loader to efficiently load the asset in advance */
 	UPROPERTY()
-	mutable UObject* HardPtr;
+	mutable TObjectPtr<UObject> HardPtr;
 };
 
 template<>
@@ -165,7 +165,7 @@ public:
 private:
 
 	UPROPERTY()
-	UClass* PropertyClass;
+	TObjectPtr<UClass> PropertyClass;
 
 	UPROPERTY(meta=(KeyTimes))
 	TArray<FFrameNumber> Times;
diff --git a/Engine/Source/Runtime/MovieScene/Public/EntitySystem/BuiltInComponentTypes.h b/Engine/Source/Runtime/MovieScene/Public/EntitySystem/BuiltInComponentTypes.h
--- a/Engine/Source/Runtime/MovieScene/Public/EntitySystem/BuiltInComponentTypes.h
+++ b/Engine/Source/Runtime/MovieScene/Public/EntitySystem/BuiltInComponentTypes.h
@@ -34,7 +34,7 @@ struct FEasingComponentData
 	GENERATED_BODY()
 
 	UPROPERTY()
-	UMovieSceneSection* Section = nullptr;
+	TObjectPtr<UMovieSceneSection> Section = nullptr;
 };
 
 
@@ -47,7 +47,7 @@ struct FMovieSceneTrackInstanceComponent
 	GENERATED_BODY()
 
 	UPROPERTY()
-	UMovieSceneSection* Owner = nullptr;
+	TObjectPtr<UMovieSceneSection> Owner = nullptr;
 
 	UPROPERTY()
 	TSubclassOf<UMovieSceneTrackInstance> TrackInstanceClass;
@@ -76,7 +76,7 @@ struct FTrackInstanceInputComponent
 	GENERATED_BODY()
 
 	UPROPERTY()
-	UMovieSceneSection* Section = nullptr;
+	TObjectPtr<UMovieSceneSection> Section = nullptr;
 
 	UPROPERTY()
 	int32 OutputIndex = INDEX_NONE;
diff --git a/Engine/Source/Runtime/MovieScene/Public/EntitySystem/MovieSceneEntitySystem.h b/Engine/Source/Runtime/MovieScene/Public/EntitySystem/MovieSceneEntitySystem.h
--- a/Engine/Source/Runtime/MovieScene/Public/EntitySystem/MovieSceneEntitySystem.h
+++ b/Engine/Source/Runtime/MovieScene/Public/EntitySystem/MovieSceneEntitySystem.h
@@ -155,7 +155,7 @@ private:
 protected:
 
 	UPROPERTY()
-	UMovieSceneEntitySystemLinker* Linker;
+	TObjectPtr<UMovieSceneEntitySystemLinker> Linker;
 
 protected:
 
