Temporarily mark all landscape proxies dirty when moving the landscape actor to make sure we reflect the change to the actor's metadata.

diff --git a/Engine/Source/Runtime/Landscape/Classes/LandscapeInfo.h b/Engine/Source/Runtime/Landscape/Classes/LandscapeInfo.h
--- a/Engine/Source/Runtime/Landscape/Classes/LandscapeInfo.h
+++ b/Engine/Source/Runtime/Landscape/Classes/LandscapeInfo.h
@@ -246,7 +246,7 @@ public:
 	 *  requires LandscapeActor to be loaded
 	 *  Does not work in World composition mode!
 	 */
-	LANDSCAPE_API void FixupProxiesTransform();
+	LANDSCAPE_API void FixupProxiesTransform(bool bDirty = false);
 	
 	// Update per-component layer whitelists to include the currently painted layers
 	LANDSCAPE_API void UpdateComponentLayerWhitelist();
diff --git a/Engine/Source/Runtime/Landscape/Private/Landscape.cpp b/Engine/Source/Runtime/Landscape/Private/Landscape.cpp
--- a/Engine/Source/Runtime/Landscape/Private/Landscape.cpp
+++ b/Engine/Source/Runtime/Landscape/Private/Landscape.cpp
@@ -2962,7 +2962,7 @@ void ULandscapeInfo::Reset()
 	//SelectedRegion.Empty();
 }
 
-void ULandscapeInfo::FixupProxiesTransform()
+void ULandscapeInfo::FixupProxiesTransform(bool bDirty)
 {
 	ALandscape* Landscape = LandscapeActor.Get();
 
@@ -2976,6 +2976,12 @@ void ULandscapeInfo::FixupProxiesTransform()
 	for (auto It = Proxies.CreateConstIterator(); It; ++It)
 	{
 		ALandscapeProxy* Proxy = *It;
+
+		if (bDirty)
+		{
+			Proxy->Modify();
+		}
+
 		FIntPoint LandscapeSectionOffset = Proxy->LandscapeSectionOffset - Landscape->LandscapeSectionOffset;
 		FIntPoint LandscapeSectionOffsetRem(
 			LandscapeSectionOffset.X % Proxy->ComponentSizeQuads, 
diff --git a/Engine/Source/Runtime/Landscape/Private/LandscapeEdit.cpp b/Engine/Source/Runtime/Landscape/Private/LandscapeEdit.cpp
--- a/Engine/Source/Runtime/Landscape/Private/LandscapeEdit.cpp
+++ b/Engine/Source/Runtime/Landscape/Private/LandscapeEdit.cpp
@@ -3929,7 +3929,7 @@ void ALandscape::PostEditMove(bool bFinished)
 		auto* LandscapeInfo = GetLandscapeInfo();
 		if (LandscapeInfo)
 		{
-			LandscapeInfo->FixupProxiesTransform();
+			LandscapeInfo->FixupProxiesTransform(true);
 		}
 	}
 
@@ -5141,7 +5141,7 @@ void ALandscape::PostEditChangeProperty(FPropertyChangedEvent& PropertyChangedEv
 		if (Info != nullptr)
 		{
 			// update transformations for all linked proxies 
-			Info->FixupProxiesTransform();
+			Info->FixupProxiesTransform(true);
 			bNeedsRecalcBoundingBox = true;
 		}
 
