Refactored how we register landscape proxis and splines (removed custom registration code, always register in the actor desccriptors).

diff --git a/Engine/Source/Runtime/Engine/Private/WorldPartition/Landscape/LandscapeActorDesc.cpp b/Engine/Source/Runtime/Engine/Private/WorldPartition/Landscape/LandscapeActorDesc.cpp
--- a/Engine/Source/Runtime/Engine/Private/WorldPartition/Landscape/LandscapeActorDesc.cpp
+++ b/Engine/Source/Runtime/Engine/Private/WorldPartition/Landscape/LandscapeActorDesc.cpp
@@ -46,7 +46,7 @@ void FLandscapeActorDesc::OnRegister()
 {
 	FPartitionActorDesc::OnRegister();
 
-	if (ULandscapeInfo* LandscapeInfo = ULandscapeInfo::Find(Container->GetWorld(), GridGuid))
+	if (ULandscapeInfo* LandscapeInfo = ULandscapeInfo::FindOrCreate(Container->GetWorld(), GridGuid))
 	{
 		FWorldPartitionHandle Handle(Container, GetGuid());
 		check(Handle.IsValid());
@@ -58,7 +58,7 @@ void FLandscapeActorDesc::OnUnregister()
 {
 	FPartitionActorDesc::OnUnregister();
 
-	if (ULandscapeInfo* LandscapeInfo = ULandscapeInfo::Find(Container->GetWorld(), GridGuid))
+	if (ULandscapeInfo* LandscapeInfo = ULandscapeInfo::FindOrCreate(Container->GetWorld(), GridGuid))
 	{
 		FWorldPartitionHandle Handle(Container, GetGuid());
 		check(Handle.IsValid());
diff --git a/Engine/Source/Runtime/Engine/Private/WorldPartition/Landscape/LandscapeSplineActorDesc.cpp b/Engine/Source/Runtime/Engine/Private/WorldPartition/Landscape/LandscapeSplineActorDesc.cpp
--- a/Engine/Source/Runtime/Engine/Private/WorldPartition/Landscape/LandscapeSplineActorDesc.cpp
+++ b/Engine/Source/Runtime/Engine/Private/WorldPartition/Landscape/LandscapeSplineActorDesc.cpp
@@ -35,7 +35,7 @@ void FLandscapeSplineActorDesc::OnRegister()
 {
 	FWorldPartitionActorDesc::OnRegister();
 
-	if (ULandscapeInfo* LandscapeInfo = ULandscapeInfo::Find(Container->GetWorld(), LandscapeGuid))
+	if (ULandscapeInfo* LandscapeInfo = ULandscapeInfo::FindOrCreate(Container->GetWorld(), LandscapeGuid))
 	{
 		FWorldPartitionHandle Handle(Container, GetGuid());
 		check(Handle.IsValid());
@@ -47,7 +47,7 @@ void FLandscapeSplineActorDesc::OnUnregister()
 {
 	FWorldPartitionActorDesc::OnUnregister();
 
-	if (ULandscapeInfo* LandscapeInfo = ULandscapeInfo::Find(Container->GetWorld(), LandscapeGuid))
+	if (ULandscapeInfo* LandscapeInfo = ULandscapeInfo::FindOrCreate(Container->GetWorld(), LandscapeGuid))
 	{
 		FWorldPartitionHandle Handle(Container, GetGuid());
 		check(Handle.IsValid());
diff --git a/Engine/Source/Runtime/Engine/Private/WorldPartition/WorldPartition.cpp b/Engine/Source/Runtime/Engine/Private/WorldPartition/WorldPartition.cpp
--- a/Engine/Source/Runtime/Engine/Private/WorldPartition/WorldPartition.cpp
+++ b/Engine/Source/Runtime/Engine/Private/WorldPartition/WorldPartition.cpp
@@ -215,22 +215,22 @@ void UWorldPartition::Initialize(UWorld* InWorld, const FTransform& InTransform)
 
 		for (UActorDescContainer::TIterator<> ActorDescIterator(this); ActorDescIterator; ++ActorDescIterator)
 		{
-			FWorldPartitionActorDesc* ActorDesc = *ActorDescIterator;
-
 			if (bIsInstanced)
 			{
-				const FString LongActorPackageName = ActorDesc->ActorPackage.ToString();
+				const FString LongActorPackageName = ActorDescIterator->ActorPackage.ToString();
 				const FString ActorPackageName = FPaths::GetBaseFilename(LongActorPackageName);
 				const FString InstancedName = FString::Printf(TEXT("%s_InstanceOf_%s"), *LevelPackage->GetName(), *ActorPackageName);
 
 				InstancingContext.AddMapping(*LongActorPackageName, *InstancedName);
 
-				ActorDesc->TransformInstance(ReplaceFrom, ReplaceTo, InstanceTransform);
+				ActorDescIterator->TransformInstance(ReplaceFrom, ReplaceTo, InstanceTransform);
 			}
 
+			ActorDescIterator->OnRegister();
+
 			if (bEditorOnly)
 			{
-				HashActorDesc(ActorDesc);
+				HashActorDesc(*ActorDescIterator);
 			}
 		}
 
diff --git a/Engine/Source/Runtime/Landscape/Classes/LandscapeInfo.h b/Engine/Source/Runtime/Landscape/Classes/LandscapeInfo.h
--- a/Engine/Source/Runtime/Landscape/Classes/LandscapeInfo.h
+++ b/Engine/Source/Runtime/Landscape/Classes/LandscapeInfo.h
@@ -298,7 +298,7 @@ public:
 	const TSet<FWorldPartitionHandle>& GetSplineHandles() const { return SplineHandles; }
 #endif
 	LANDSCAPE_API static ULandscapeInfo* Find(UWorld* InWorld, const FGuid& LandscapeGuid);
-	static ULandscapeInfo* FindOrCreate(UWorld* InWorld, const FGuid& LandscapeGuid);
+	LANDSCAPE_API static ULandscapeInfo* FindOrCreate(UWorld* InWorld, const FGuid& LandscapeGuid);
 
 	/** Called after creating object so that it can initialize its state */
 	void Initialize(UWorld* InWorld, const FGuid& InLandscapeGuid);
diff --git a/Engine/Source/Runtime/Landscape/Private/Landscape.cpp b/Engine/Source/Runtime/Landscape/Private/Landscape.cpp
--- a/Engine/Source/Runtime/Landscape/Private/Landscape.cpp
+++ b/Engine/Source/Runtime/Landscape/Private/Landscape.cpp
@@ -3129,30 +3129,6 @@ ULandscapeInfo* ULandscapeInfo::FindOrCreate(UWorld* InWorld, const FGuid& Lands
 void ULandscapeInfo::Initialize(UWorld* InWorld, const FGuid& InLandscapeGuid)
 {
 	LandscapeGuid = InLandscapeGuid;
-#if WITH_EDITOR
-	// If WorldPartition world get the Proxy handles
-	if (UWorldPartition* WorldPartition = InWorld->GetWorldPartition())
-	{
-		for (UActorDescContainer::TIterator<> ActorDescIterator(WorldPartition); ActorDescIterator; ++ActorDescIterator)
-		{
-			FWorldPartitionActorDesc* ActorDesc = *ActorDescIterator;
-			if (ActorDesc->GetActorClass()->IsChildOf(ALandscapeStreamingProxy::StaticClass()))
-			{
-				if (((FLandscapeActorDesc*)ActorDesc)->GridGuid == InLandscapeGuid)
-				{
-					ProxyHandles.Add(FWorldPartitionHandle(WorldPartition, ActorDesc->GetGuid()));
-				}
-			}
-			else if (ActorDescIterator->GetActorClass()->IsChildOf(ALandscapeSplineActor::StaticClass()))
-			{
-				if (((FLandscapeSplineActorDesc*)ActorDesc)->LandscapeGuid == InLandscapeGuid)
-				{
-					SplineHandles.Emplace(FWorldPartitionHandle(WorldPartition, ActorDesc->GetGuid()));
-				}
-			}
-		}
-	}
-#endif
 }
 
 void ULandscapeInfo::ForAllLandscapeProxies(TFunctionRef<void(ALandscapeProxy*)> Fn) const
