Fix lack of any progress in slow task on landscape creation

diff --git a/Engine/Source/Runtime/Landscape/Private/LandscapeEdit.cpp b/Engine/Source/Runtime/Landscape/Private/LandscapeEdit.cpp
--- a/Engine/Source/Runtime/Landscape/Private/LandscapeEdit.cpp
+++ b/Engine/Source/Runtime/Landscape/Private/LandscapeEdit.cpp
@@ -2421,7 +2421,10 @@ LANDSCAPE_API void ALandscapeProxy::Import(const FGuid& InGuid, int32 InMinX, in
 
 	check(CanHaveLayersContent() || InImportLayers == nullptr);
 
-	GWarn->BeginSlowTask(LOCTEXT("BeingImportingLandscapeTask", "Importing Landscape"), true);
+	FScopedSlowTask SlowTask(2, LOCTEXT("BeingImportingLandscapeTask", "Importing Landscape"));
+	SlowTask.MakeDialog();
+
+	SlowTask.EnterProgressFrame(1.0f);
 
 	const int32 VertsX = InMaxX - InMinX + 1;
 	const int32 VertsY = InMaxY - InMinY + 1;
@@ -2987,6 +2990,8 @@ LANDSCAPE_API void ALandscapeProxy::Import(const FGuid& InGuid, int32 InMinX, in
 	// Update all materials and recreate render state of all landscape components
 	TArray<FComponentRecreateRenderStateContext> RecreateRenderStateContexts;
 
+	SlowTask.EnterProgressFrame(1.0f);
+
 	{
 		// We disable automatic material update context, to manage it manually
 		GDisableAutomaticTextureMaterialUpdateDependencies = true;
@@ -3162,8 +3167,7 @@ LANDSCAPE_API void ALandscapeProxy::Import(const FGuid& InGuid, int32 InMinX, in
 	ReimportHeightmapFilePath = InHeightmapFileName;
 
 	LandscapeInfo->UpdateLayerInfoMap();
-
-	GWarn->EndSlowTask();
+	
 }
 
 bool ALandscapeProxy::ExportToRawMesh(int32 InExportLOD, FMeshDescription& OutRawMesh) const
