Instance culling support for Landscape

diff --git a/Engine/Shaders/Private/LandscapeGrassWeight.usf b/Engine/Shaders/Private/LandscapeGrassWeight.usf
--- a/Engine/Shaders/Private/LandscapeGrassWeight.usf
+++ b/Engine/Shaders/Private/LandscapeGrassWeight.usf
@@ -9,6 +9,9 @@ struct FLandscapeGrassWeightInterpolantsVSToPS
 	float4 Position : SV_POSITION;
 	FVertexFactoryInterpolantsVSToPS FactoryInterpolants;
 	float4 PositionPreOffset : TEXCOORD9;
+#if VF_USE_PRIMITIVE_SCENE_DATA
+	float LocalPosZ : TEXCOORD10;
+#endif // VF_USE_PRIMITIVE_SCENE_DATA
 };
 
 #if VERTEXSHADER
@@ -34,10 +37,12 @@ void VSMain(
 	// (ii) take the full offset which means that grass can be missing or corrupt at edges of each tile.
 	// We choose (i) here.
 	float3 WPO = GetMaterialWorldPositionOffset(VertexParameters);
-	WPO = mul(WPO, (float3x3)Primitive.WorldToLocal) * float3(0, 0, 1);
-	WPO = mul(WPO, (float3x3)Primitive.LocalToWorld);
+	WPO = mul(WPO, (float3x3)VertexFactoryGetWorldToLocal(VFIntermediates)) * float3(0, 0, 1);
+	WPO = mul(WPO, (float3x3)VertexFactoryGetLocalToWorld(VFIntermediates));
 	WorldPosition.xyz += WPO;
-
+#if VF_USE_PRIMITIVE_SCENE_DATA
+	Output.LocalPosZ = mul(float4(WorldPosition.xyz - ResolvedView.PreViewTranslation, 1), VertexFactoryGetWorldToLocal(VFIntermediates)).z;
+#endif // VF_USE_PRIMITIVE_SCENE_DATA
 	Output.FactoryInterpolants = VertexFactoryGetInterpolantsVSToPS(Input, VFIntermediates, VertexParameters);
 
 	float4 RasterizedWorldPosition = VertexFactoryGetRasterizedWorldPosition(Input, VFIntermediates, WorldPosition);
@@ -67,8 +72,12 @@ void PSMain(
 	CalcMaterialParameters(MaterialParameters, PixelMaterialInputs, SvPositionPreOffset, bIsFrontFace);
 
 	// Calculate and pack landscape height.
+#if VF_USE_PRIMITIVE_SCENE_DATA
+	float Height = clamp(round(Interpolants.LocalPosZ / TERRAIN_ZSCALE) + 32768.0, 0, 65535);
+#else // !VF_USE_PRIMITIVE_SCENE_DATA
 	float3 LocalPos = mul(float4(MaterialParameters.AbsoluteWorldPosition, 1), Primitive.WorldToLocal).xyz;
 	float Height = clamp(round(LocalPos.z / TERRAIN_ZSCALE) + 32768.0, 0, 65535);
+#endif // VF_USE_PRIMITIVE_SCENE_DATA
 	float2 PackedHeight = float2(floor(Height / 256), fmod(Height, 256)) / 255;
 
 #if NUM_MATERIAL_OUTPUTS_GETGRASSWEIGHT > 2
diff --git a/Engine/Shaders/Private/LandscapeVertexFactory.ush b/Engine/Shaders/Private/LandscapeVertexFactory.ush
--- a/Engine/Shaders/Private/LandscapeVertexFactory.ush
+++ b/Engine/Shaders/Private/LandscapeVertexFactory.ush
@@ -156,6 +156,12 @@ struct FVertexFactoryInput
 {
 #if FEATURE_LEVEL >= FEATURE_LEVEL_SM4
 	float4 Position: ATTRIBUTE0;
+
+#if VF_USE_PRIMITIVE_SCENE_DATA
+	uint InstanceIdOffset : ATTRIBUTE1;
+	uint DrawInstanceId : SV_InstanceID;
+#endif
+
 #else
 	float4 PackedPosition: ATTRIBUTE0; 
 	float4 LODHeights[2]: ATTRIBUTE1;  
@@ -310,6 +316,10 @@ FVertexFactoryInterpolantsVSToPS VertexFactoryAssignInterpolants(FVertexFactoryI
 
 struct FVertexFactoryIntermediates
 {
+#if VF_USE_PRIMITIVE_SCENE_DATA
+	FSceneDataIntermediates SceneData;
+#endif
+
 	float4 InputPosition;
 	float3 LocalPosition;
 	float3 WorldNormal;
@@ -326,7 +336,11 @@ float3 GetLocalPosition(FVertexFactoryIntermediates Intermediates)
 
 float4 VertexFactoryGetWorldPosition(FVertexFactoryInput Input, FVertexFactoryIntermediates Intermediates)
 {
+#if VF_USE_PRIMITIVE_SCENE_DATA
+	return INVARIANT(TransformLocalToTranslatedWorld(GetLocalPosition(Intermediates), Intermediates.SceneData.InstanceData));
+#else
 	return INVARIANT(TransformLocalToTranslatedWorld(GetLocalPosition(Intermediates)));
+#endif
 }
 
 float3 VertexFactoryGetWorldNormal(FVertexFactoryInput Input, FVertexFactoryIntermediates Intermediates)
@@ -338,8 +352,13 @@ float4 VertexFactoryGetPreviousWorldPosition(FVertexFactoryInput Input, FVertexF
 {
 	// Note we still use LocalToWorld. Landscape's transform never changes.
 	float3 LocalPosition = GetLocalPosition(Intermediates);
+#if VF_USE_PRIMITIVE_SCENE_DATA
+	float3 RotatedPosition = INVARIANT(Intermediates.SceneData.InstanceData.LocalToWorld[0].xyz * LocalPosition.xxx + Intermediates.SceneData.InstanceData.LocalToWorld[1].xyz * LocalPosition.yyy + Intermediates.SceneData.InstanceData.LocalToWorld[2].xyz * LocalPosition.zzz);
+	return INVARIANT(float4(RotatedPosition + (Intermediates.SceneData.InstanceData.LocalToWorld[3].xyz + ResolvedView.PrevPreViewTranslation.xyz),1));
+#else // !VF_USE_PRIMITIVE_SCENE_DATA
 	float3 RotatedPosition = INVARIANT(Primitive.LocalToWorld[0].xyz * LocalPosition.xxx + Primitive.LocalToWorld[1].xyz * LocalPosition.yyy + Primitive.LocalToWorld[2].xyz * LocalPosition.zzz);
 	return INVARIANT(float4(RotatedPosition + (Primitive.LocalToWorld[3].xyz + ResolvedView.PrevPreViewTranslation.xyz),1));
+#endif // VF_USE_PRIMITIVE_SCENE_DATA
 }
 
 /** Calculate the texture coordinates generated by Landscape */
@@ -559,7 +578,11 @@ FMaterialVertexParameters GetMaterialVertexParameters(FVertexFactoryInput Input,
 	Result.PreSkinnedNormal = TangentToLocal[2].xyz;
 
 	// Assumes no instancing and landscape's transform never change
+#if VF_USE_PRIMITIVE_SCENE_DATA
+	Result.PrevFrameLocalToWorld = Intermediates.SceneData.InstanceData.LocalToWorld;
+#else // !VF_USE_PRIMITIVE_SCENE_DATA
 	Result.PrevFrameLocalToWorld = Primitive.LocalToWorld;
+#endif // VF_USE_PRIMITIVE_SCENE_DATA
 
 	FLandscapeTexCoords LandscapeTexCoords = GetLandscapeTexCoords(Input, Intermediates);
 
@@ -630,6 +653,10 @@ void GetShadowMapCoordinate(FVertexFactoryInterpolantsVSToPS Interpolants, out f
 FVertexFactoryIntermediates GetVertexFactoryIntermediates(FVertexFactoryInput Input)
 {
 	FVertexFactoryIntermediates Intermediates;
+#if VF_USE_PRIMITIVE_SCENE_DATA
+	Intermediates.SceneData = GetSceneDataIntermediates(Input.InstanceIdOffset, Input.DrawInstanceId);
+#endif
+
 	Intermediates.ComponentIndex = GetComponentLinearIndex();
 	Intermediates.LodValues = GetLodValues(Intermediates.ComponentIndex);
 	Intermediates.LodBias = GetLodBias(Intermediates.ComponentIndex);
diff --git a/Engine/Source/Runtime/Landscape/Private/LandscapeGrass.cpp b/Engine/Source/Runtime/Landscape/Private/LandscapeGrass.cpp
--- a/Engine/Source/Runtime/Landscape/Private/LandscapeGrass.cpp
+++ b/Engine/Source/Runtime/Landscape/Private/LandscapeGrass.cpp
@@ -58,6 +58,7 @@
 #include "EngineUtils.h"
 #include "Misc/ScopedSlowTask.h"
 #include "TextureCompiler.h"
+#include "RenderCaptureInterface.h"
 
 #define LOCTEXT_NAMESPACE "Landscape"
 
@@ -180,6 +181,13 @@ static FAutoConsoleVariableRef CVarUpdateAllOnRebuild(
 	GGrassUpdateAllOnRebuild,
 	TEXT(""));
 
+
+static int32 GCaptureNextGrassUpdate = 0;
+static FAutoConsoleVariableRef CVarCaptureNextGrassUpdate(
+	TEXT("grass.CaptureNextGrassUpdate"),
+	GCaptureNextGrassUpdate,
+	TEXT("Trigger a renderdoc capture for the next X grass updates (calls to RenderGrassMap or RenderGrassMaps"));
+
 DECLARE_CYCLE_STAT(TEXT("Grass Async Build Time"), STAT_FoliageGrassAsyncBuildTime, STATGROUP_Foliage);
 DECLARE_CYCLE_STAT(TEXT("Grass Start Comp"), STAT_FoliageGrassStartComp, STATGROUP_Foliage);
 DECLARE_CYCLE_STAT(TEXT("Grass End Comp"), STAT_FoliageGrassEndComp, STATGROUP_Foliage);
@@ -498,13 +506,12 @@ public:
 
 	BEGIN_SHADER_PARAMETER_STRUCT(FLandscapeGrassPassParameters, )
 		SHADER_PARAMETER_STRUCT_REF(FViewUniformShaderParameters, View)
-		SHADER_PARAMETER_STRUCT_INCLUDE(FInstanceCullingDrawParams, InstanceCullingDrawParams)
 		RENDER_TARGET_BINDING_SLOTS()
 	END_SHADER_PARAMETER_STRUCT()
 
 	void RenderLandscapeComponentToTexture_RenderThread(FRHICommandListImmediate& RHICmdList)
 	{
-		FSceneViewFamilyContext ViewFamily(FSceneViewFamily::ConstructionValues(nullptr, nullptr, FEngineShowFlags(ESFIM_Game)).SetWorldTimes(FApp::GetCurrentTime() - GStartTime, FApp::GetDeltaTime(), FApp::GetCurrentTime() - GStartTime));
+		FSceneViewFamilyContext ViewFamily(FSceneViewFamily::ConstructionValues(nullptr, SceneInterface, FEngineShowFlags(ESFIM_Game)).SetWorldTimes(FApp::GetCurrentTime() - GStartTime, FApp::GetDeltaTime(), FApp::GetCurrentTime() - GStartTime));
 
 		ViewFamily.LandscapeLODOverride = 0; // Force LOD render
 
@@ -2760,6 +2767,9 @@ void ALandscapeProxy::UpdateGrass(const TArray<FVector>& Cameras, int32& InOutNu
 												continue;
 											}
 
+											RenderCaptureInterface::FScopedCapture RenderCapture((GCaptureNextGrassUpdate != 0), TEXT("RenderGrassMap"));
+											GCaptureNextGrassUpdate = FMath::Max(GCaptureNextGrassUpdate - 1, 0);
+
 											QUICK_SCOPE_CYCLE_COUNTER(STAT_GrassRenderToTexture);
 											Component->RenderGrassMap();
 											ComponentsNeedingGrassMapRender.Remove(Component);
@@ -2942,6 +2952,9 @@ void ALandscapeProxy::UpdateGrass(const TArray<FVector>& Cameras, int32& InOutNu
 					}
 					if (ComponentsToRender.Num())
 					{
+						RenderCaptureInterface::FScopedCapture RenderCapture((GCaptureNextGrassUpdate != 0), TEXT("RenderGrassMaps"));
+						GCaptureNextGrassUpdate = FMath::Max(GCaptureNextGrassUpdate - 1, 0);
+
 						RenderGrassMaps(ComponentsToRender, LandscapeGrassTypes);
 					}
 				}
diff --git a/Engine/Source/Runtime/Landscape/Private/LandscapeRender.cpp b/Engine/Source/Runtime/Landscape/Private/LandscapeRender.cpp
--- a/Engine/Source/Runtime/Landscape/Private/LandscapeRender.cpp
+++ b/Engine/Source/Runtime/Landscape/Private/LandscapeRender.cpp
@@ -1239,6 +1239,11 @@ FLandscapeComponentSceneProxy::FLandscapeComponentSceneProxy(ULandscapeComponent
 	, LightMapResolution(InComponent->GetStaticLightMapResolution())
 #endif
 {
+#if defined(GPUCULL_TODO)
+	// GPUCULL_TODO: Move to base proxy
+	bVFRequiresPrimitiveUniformBuffer = !UseGPUScene(GMaxRHIShaderPlatform, GetScene().GetFeatureLevel());
+#endif // defined(GPUCULL_TODO)
+
 	const auto FeatureLevel = GetScene().GetFeatureLevel();
 
 	if (FeatureLevel >= ERHIFeatureLevel::SM5)
@@ -1470,6 +1475,19 @@ PRAGMA_ENABLE_DEPRECATION_WARNINGS
 		}
 	}
 #endif
+
+#if defined(GPUCULL_TODO)
+	// 100% generic default-instance
+	Instances.SetNum(1);
+	FPrimitiveInstance& Instance = Instances[0];
+	Instance.PrimitiveId = ~uint32(0);
+	Instance.InstanceToLocal.SetIdentity();
+	Instance.LocalToInstance.SetIdentity();
+	Instance.LocalToWorld.SetIdentity();
+	Instance.RenderBounds = GetLocalBounds();
+	Instance.LocalBounds = Instance.RenderBounds;
+	bSupportsInstanceDataBuffer = true;
+#endif // defined(GPUCULL_TODO)
 }
 
 void FLandscapeComponentSceneProxy::CreateRenderThreadResources()
@@ -2039,6 +2057,13 @@ void FLandscapeComponentSceneProxy::OnTransformChanged()
 
 	// Recache mesh draw commands for changed uniform buffers
 	GetScene().UpdateCachedRenderStates(this);
+
+#if defined(GPUCULL_TODO)
+	// 
+	FPrimitiveInstance& Instance = Instances[0];
+	Instance.RenderBounds = GetLocalBounds();
+	Instance.LocalBounds = Instance.RenderBounds;
+#endif // defined(GPUCULL_TODO)
 }
 
 /** Creates a mesh batch for virtual texture rendering. Will render a simple fixed grid with combined subsections. */
@@ -3466,6 +3491,9 @@ void FLandscapeVertexFactory::InitRHI()
 	// position decls
 	Elements.Add(AccessStreamComponent(Data.PositionComponent, 0));
 
+#if defined(GPUCULL_TODO)
+	AddPrimitiveIdStreamElement(EVertexInputStreamType::Default, 1, Elements);
+#endif // defined(GPUCULL_TODO)
 	// create the actual device decls
 	InitDeclaration(Elements);
 }
@@ -3486,6 +3514,10 @@ bool FLandscapeVertexFactory::ShouldCompilePermutation(const FVertexFactoryShade
 void FLandscapeVertexFactory::ModifyCompilationEnvironment(const FVertexFactoryShaderPermutationParameters& Parameters, FShaderCompilerEnvironment& OutEnvironment)
 {
 	FVertexFactory::ModifyCompilationEnvironment(Parameters, OutEnvironment);
+
+#if defined(GPUCULL_TODO)
+	OutEnvironment.SetDefine(TEXT("VF_SUPPORTS_PRIMITIVE_SCENE_DATA"), Parameters.VertexFactoryType->SupportsPrimitiveIdStream() && UseGPUScene(Parameters.Platform, GetMaxSupportedFeatureLevel(Parameters.Platform)));
+#endif // defined(GPUCULL_TODO)
 }
 
 IMPLEMENT_VERTEX_FACTORY_PARAMETER_TYPE(FLandscapeVertexFactory, SF_Vertex, FLandscapeVertexFactoryVertexShaderParameters);
@@ -3495,7 +3527,11 @@ IMPLEMENT_VERTEX_FACTORY_PARAMETER_TYPE(FLandscapeVertexFactory, SF_RayHitGroup,
 #endif // RHI_RAYTRACING
 IMPLEMENT_VERTEX_FACTORY_PARAMETER_TYPE(FLandscapeVertexFactory, SF_Pixel, FLandscapeVertexFactoryPixelShaderParameters);
 
+#if defined(GPUCULL_TODO)
+IMPLEMENT_VERTEX_FACTORY_TYPE_EX(FLandscapeVertexFactory, "/Engine/Private/LandscapeVertexFactory.ush", true, true, true, false, false, true, true, false);
+#else // !defined(GPUCULL_TODO)
 IMPLEMENT_VERTEX_FACTORY_TYPE_EX(FLandscapeVertexFactory, "/Engine/Private/LandscapeVertexFactory.ush", true, true, true, false, false, true, false, false);
+#endif // defined(GPUCULL_TODO)
 
 /**
 * Copy the data from another vertex factory
@@ -3531,7 +3567,11 @@ IMPLEMENT_VERTEX_FACTORY_PARAMETER_TYPE(FLandscapeXYOffsetVertexFactory, SF_RayH
 #endif // RHI_RAYTRACING
 IMPLEMENT_VERTEX_FACTORY_PARAMETER_TYPE(FLandscapeXYOffsetVertexFactory, SF_Pixel, FLandscapeVertexFactoryPixelShaderParameters);
 
+#if defined(GPUCULL_TODO)
+IMPLEMENT_VERTEX_FACTORY_TYPE_EX(FLandscapeXYOffsetVertexFactory, "/Engine/Private/LandscapeVertexFactory.ush", true, true, true, false, false, true, true, false);
+#else // !defined(GPUCULL_TODO)
 IMPLEMENT_VERTEX_FACTORY_TYPE_EX(FLandscapeXYOffsetVertexFactory, "/Engine/Private/LandscapeVertexFactory.ush", true, true, true, false, false, true, false, false);
+#endif // defined(GPUCULL_TODO)
 
 //
 // FLandscapeFixedGridVertexFactory
@@ -3550,7 +3590,11 @@ IMPLEMENT_VERTEX_FACTORY_PARAMETER_TYPE(FLandscapeFixedGridVertexFactory, SF_Ray
 #endif
 IMPLEMENT_VERTEX_FACTORY_PARAMETER_TYPE(FLandscapeFixedGridVertexFactory, SF_Pixel, FLandscapeVertexFactoryPixelShaderParameters);
 
+#if defined(GPUCULL_TODO)
+IMPLEMENT_VERTEX_FACTORY_TYPE_EX(FLandscapeFixedGridVertexFactory, "/Engine/Private/LandscapeVertexFactory.ush", true, true, true, false, false, true, true, false);
+#else // !defined(GPUCULL_TODO)
 IMPLEMENT_VERTEX_FACTORY_TYPE_EX(FLandscapeFixedGridVertexFactory, "/Engine/Private/LandscapeVertexFactory.ush", true, true, true, false, false, true, false, false);
+#endif // defined(GPUCULL_TODO)
 
 
 /** ULandscapeMaterialInstanceConstant */
diff --git a/Engine/Source/Runtime/Landscape/Public/LandscapeRender.h b/Engine/Source/Runtime/Landscape/Public/LandscapeRender.h
--- a/Engine/Source/Runtime/Landscape/Public/LandscapeRender.h
+++ b/Engine/Source/Runtime/Landscape/Public/LandscapeRender.h
@@ -800,7 +800,11 @@ public:
 	// Key is the component size and number of subsections.
 	// Also being reused by GPULightmass currently to save mem
 	static LANDSCAPE_API TMap<uint32, FLandscapeSharedBuffers*> SharedBuffersMap;
-
+#if defined(GPUCULL_TODO)
+	// GPUCULL_TODO: put instances in base proxy and make non-virtual
+	virtual const TArray<FPrimitiveInstance>* GetPrimitiveInstances() const	{ return &Instances; }
+	virtual TArray<FPrimitiveInstance>* GetPrimitiveInstances()	{return &Instances;	}
+#endif // defined(GPUCULL_TODO)
 protected:
 	int8						MaxLOD;		// Maximum LOD level, user override possible
 	bool						UseTessellationComponentScreenSizeFalloff:1;	// Tell if we should apply a Tessellation falloff
@@ -931,6 +935,11 @@ protected:
 	// Cached versions of these
 	FMatrix					WorldToLocal;
 
+#if defined(GPUCULL_TODO)
+	// Instances (going to contain just one).
+	TArray<FPrimitiveInstance> Instances;
+#endif // defined(GPUCULL_TODO)
+
 protected:
 	virtual ~FLandscapeComponentSceneProxy();
 	
diff --git a/Engine/Source/Runtime/RenderCore/Private/VertexFactory.cpp b/Engine/Source/Runtime/RenderCore/Private/VertexFactory.cpp
--- a/Engine/Source/Runtime/RenderCore/Private/VertexFactory.cpp
+++ b/Engine/Source/Runtime/RenderCore/Private/VertexFactory.cpp
@@ -9,6 +9,7 @@
 #include "UObject/DebugSerializationFlags.h"
 #include "PipelineStateCache.h"
 #include "ShaderCompilerCore.h"
+#include "RenderUtils.h"
 
 IMPLEMENT_TYPE_LAYOUT(FVertexFactoryShaderParameters);
 
@@ -277,6 +278,19 @@ void FVertexFactory::ReleaseRHI()
 	PositionAndNormalStream.Empty();
 }
 
+bool FVertexFactory::AddPrimitiveIdStreamElement(EVertexInputStreamType InputStreamType, uint8 AttributeIndex, FVertexDeclarationElementList& VertexDeclarationElements)
+{
+	if (GetType()->SupportsPrimitiveIdStream() && UseGPUScene(GMaxRHIShaderPlatform, GMaxRHIFeatureLevel))
+	{
+		// When the VF is used for rendering in normal mesh passes, this vertex buffer and offset will be overridden
+		VertexDeclarationElements.Add(AccessStreamComponent(FVertexStreamComponent(&GPrimitiveIdDummy, 0, 0, PrimitiveIdStreamStride, VET_UInt, EVertexStreamUsage::Instancing), AttributeIndex, InputStreamType));
+		SetPrimitiveIdStreamIndex(InputStreamType, VertexDeclarationElements.Last().StreamIndex);
+		return true;
+	}
+
+	return false;
+}
+
 FVertexElement FVertexFactory::AccessStreamComponent(const FVertexStreamComponent& Component, uint8 AttributeIndex)
 {
 	FVertexStream VertexStream;
diff --git a/Engine/Source/Runtime/RenderCore/Public/VertexFactory.h b/Engine/Source/Runtime/RenderCore/Public/VertexFactory.h
--- a/Engine/Source/Runtime/RenderCore/Public/VertexFactory.h
+++ b/Engine/Source/Runtime/RenderCore/Public/VertexFactory.h
@@ -686,6 +686,8 @@ protected:
 		PrimitiveIdStreamIndex[static_cast<uint8>(InputStreamType)] = StreamIndex;
 	}
 
+	bool AddPrimitiveIdStreamElement(EVertexInputStreamType InputStreamType, uint8 AttributeIndex, FVertexDeclarationElementList& VertexDeclarationElements);
+
 	/**
 	 * Creates a vertex element for a vertex stream components.  Adds a unique stream index for the vertex buffer used by the component.
 	 * @param Component - The vertex stream component.
diff --git a/Engine/Source/Runtime/Renderer/Private/MeshPassProcessor.cpp b/Engine/Source/Runtime/Renderer/Private/MeshPassProcessor.cpp
--- a/Engine/Source/Runtime/Renderer/Private/MeshPassProcessor.cpp
+++ b/Engine/Source/Runtime/Renderer/Private/MeshPassProcessor.cpp
@@ -1308,6 +1308,19 @@ void DrawDynamicMeshPassPrivate(
 
 		check(View.bIsViewInfo);
 		const FViewInfo* ViewInfo = static_cast<const FViewInfo*>(&View);
+#if defined(GPUCULL_TODO)
+#if DO_GUARD_SLOW
+		{
+			bool bNeedsGPUSceneData = false;
+			for (const auto& VisibleMeshDrawCommand : VisibleMeshDrawCommands)
+			{
+				bNeedsGPUSceneData = bNeedsGPUSceneData || EnumHasAnyFlags(VisibleMeshDrawCommand.Flags, EFVisibleMeshDrawCommandFlags::HasPrimitiveIdStreamIndex);
+			}
+			ensure(!bNeedsGPUSceneData || ViewInfo->CachedViewUniformShaderParameters->PrimitiveSceneData != GIdentityPrimitiveBuffer.PrimitiveSceneDataBufferSRV);
+			ensure(!bNeedsGPUSceneData || ViewInfo->CachedViewUniformShaderParameters->InstanceSceneData != GIdentityPrimitiveBuffer.InstanceSceneDataBufferSRV);
+		}
+#endif // DO_GUARD_SLOW
+#endif // defined(GPUCULL_TODO)
 		SortAndMergeDynamicPassMeshDrawCommands(View.GetFeatureLevel(), VisibleMeshDrawCommands, DynamicMeshDrawCommandStorage, PrimitiveIdVertexBuffer, InstanceFactor, ViewInfo->DynamicPrimitiveCollector.GetPrimitiveIdRange());
 
 		SubmitMeshDrawCommandsRange(VisibleMeshDrawCommands, GraphicsMinimalPipelineStateSet, PrimitiveIdVertexBuffer, 0, bDynamicInstancing, 0, VisibleMeshDrawCommands.Num(), InstanceFactor, RHICmdList);
