Fixed heightfield data not being updated when updating the water brush : if, during the frame a heightmap/weightmap readback modification was detected, the component wasn't requesting the collision data to be updated, it would fail to do so when being requested to do so later on (and no modification was detected by then)

diff --git a/Engine/Source/Runtime/Landscape/Classes/Landscape.h b/Engine/Source/Runtime/Landscape/Classes/Landscape.h
--- a/Engine/Source/Runtime/Landscape/Classes/Landscape.h
+++ b/Engine/Source/Runtime/Landscape/Classes/Landscape.h
@@ -333,8 +333,8 @@ private:
 		TMap<ULandscapeComponent*, FLandscapeEditLayerReadbackResult>& InOutComponents,	FDirtyDelegate DirtyDelegate);
 
 	static bool IsUpdateFlagEnabledForModes(ELandscapeComponentUpdateFlag InFlag, uint32 InUpdateModes);
-	void UpdateForChangedHeightmaps(ULandscapeComponent* Component, int32 UpdateModes);
-	void UpdateForChangedWeightmaps(ULandscapeComponent* Component, int32 UpdateModes);
+	void UpdateForChangedHeightmaps(ULandscapeComponent* InComponent, const FLandscapeEditLayerReadbackResult& InReadbackResult);
+	void UpdateForChangedWeightmaps(ULandscapeComponent* InComponent, const FLandscapeEditLayerReadbackResult& InReadbackResult);
 	int32 UpdateCollisionAndClients(TMap<ULandscapeComponent*, FLandscapeEditLayerReadbackResult> const& Components);
 	int32 UpdateAfterReadbackResolves(TMap<ULandscapeComponent*, FLandscapeEditLayerReadbackResult> const& Components);
 
diff --git a/Engine/Source/Runtime/Landscape/Classes/LandscapeComponent.h b/Engine/Source/Runtime/Landscape/Classes/LandscapeComponent.h
--- a/Engine/Source/Runtime/Landscape/Classes/LandscapeComponent.h
+++ b/Engine/Source/Runtime/Landscape/Classes/LandscapeComponent.h
@@ -514,6 +514,12 @@ private:
 	UPROPERTY(Transient)
 	uint32 LayerUpdateFlagPerMode;
 
+	UPROPERTY(Transient)
+	bool bPendingCollisionDataUpdate;
+
+	UPROPERTY(Transient)
+	bool bPendingLayerCollisionDataUpdate;
+
 	/** Dirtied collision height region when painting (only used by Landscape Layer System) */
 	FIntRect LayerDirtyCollisionHeightData;
 #endif // WITH_EDITORONLY_DATA
@@ -743,6 +749,11 @@ public:
 	FGuid GetEditingLayerGUID() const;
 
 	void CopyFinalLayerIntoEditingLayer(FLandscapeEditDataInterface& DataInterface, TSet<UTexture2D*>& ProcessedHeightmaps);
+
+	void SetPendingCollisionDataUpdate(bool bInPendingCollisionDataUpdate) { bPendingCollisionDataUpdate = bInPendingCollisionDataUpdate; }
+	bool GetPendingCollisionDataUpdate() const { return bPendingCollisionDataUpdate; }
+	void SetPendingLayerCollisionDataUpdate(bool bInPendingLayerCollisionDataUpdate) { bPendingLayerCollisionDataUpdate = bInPendingLayerCollisionDataUpdate; }
+	bool GetPendingLayerCollisionDataUpdate() const { return bPendingLayerCollisionDataUpdate; }
 #endif 
 
 #if WITH_EDITOR
diff --git a/Engine/Source/Runtime/Landscape/Private/Landscape.cpp b/Engine/Source/Runtime/Landscape/Private/Landscape.cpp
--- a/Engine/Source/Runtime/Landscape/Private/Landscape.cpp
+++ b/Engine/Source/Runtime/Landscape/Private/Landscape.cpp
@@ -162,6 +162,8 @@ ULandscapeComponent::ULandscapeComponent(const FObjectInitializer& ObjectInitial
 #if WITH_EDITORONLY_DATA
 , CachedEditingLayerData(nullptr)
 , LayerUpdateFlagPerMode(0)
+	, bPendingCollisionDataUpdate(false)
+	, bPendingLayerCollisionDataUpdate(false)
 , WeightmapsHash(0)
 , SplineHash(0)
 , PhysicalMaterialHash(0)
@@ -2582,7 +2584,7 @@ void ALandscapeProxy::FixupSharedData(ALandscape* Landscape)
 		ComponentScreenSizeToUseSubSections = Landscape->ComponentScreenSizeToUseSubSections;
 		bUpdated = true;
 	}
-	
+
 	if (LODDistributionSetting != Landscape->LODDistributionSetting)
 	{
 		LODDistributionSetting = Landscape->LODDistributionSetting;
@@ -3069,7 +3071,7 @@ void ULandscapeInfo::RegisterActor(ALandscapeProxy* Proxy, bool bMapCheck)
 	checkSlow(Proxy);
 	check(Proxy->GetLandscapeGuid().IsValid());
 	check(LandscapeGuid.IsValid());
-
+	
 #if WITH_EDITOR
 	if (!OwningWorld->IsGameWorld())
 	{
@@ -3770,14 +3772,14 @@ void ALandscapeProxy::BuildGIBakedTextures(struct FScopedSlowTask* InSlowTask /*
 		const bool bShouldMarkDirty = true;
 		UpdateGIBakedTextureData(bShouldMarkDirty);
 	}
-}
+	}
 
 int32 ALandscapeProxy::GetOutdatedGIBakedTextureComponentsCount() const
-{
+	{
 	int32 OutdatedGITextureComponentsCount = 0;
 	UpdateGIBakedTextureStatus(nullptr, nullptr, &OutdatedGITextureComponentsCount);
 	return OutdatedGITextureComponentsCount;
-}
+	}
 
 void ALandscapeProxy::UpdateGIBakedTextureStatus(bool* bOutGenerateLandscapeGIData, TMap<UTexture2D*, FGIBakedTextureState>* OutComponentsNeedBakingByHeightmap, int32* OutdatedComponentsCount) const
 {
@@ -3805,31 +3807,31 @@ void ALandscapeProxy::UpdateGIBakedTextureStatus(bool* bOutGenerateLandscapeGIDa
 	}
 	else
 	{
-		// Stores the components and their state hash data for a single atlas
+	// Stores the components and their state hash data for a single atlas
 		struct FGIBakeTextureStateBuilder
-		{
-			// pointer as FMemoryWriter caches the address of the FBufferArchive, and this struct could be relocated on a realloc.
-			TUniquePtr<FBufferArchive> ComponentStateAr;
-			TArray<ULandscapeComponent*> Components;
+	{
+		// pointer as FMemoryWriter caches the address of the FBufferArchive, and this struct could be relocated on a realloc.
+		TUniquePtr<FBufferArchive> ComponentStateAr;
+		TArray<ULandscapeComponent*> Components;
 
 			FGIBakeTextureStateBuilder()
-			{
-				ComponentStateAr = MakeUnique<FBufferArchive>();
-			}
-		};
+		{
+			ComponentStateAr = MakeUnique<FBufferArchive>();
+		}
+	};
 
 		TMap<UTexture2D*, FGIBakeTextureStateBuilder> ComponentsByHeightmap;
-		for (ULandscapeComponent* Component : LandscapeComponents)
+	for (ULandscapeComponent* Component : LandscapeComponents)
+	{
+		if (Component == nullptr)
 		{
-			if (Component == nullptr)
-			{
-				continue;
-			}
+			continue;
+		}
 
 			FGIBakeTextureStateBuilder& Info = ComponentsByHeightmap.FindOrAdd(Component->GetHeightmap());
-			Info.Components.Add(Component);
-			Component->SerializeStateHashes(*Info.ComponentStateAr);
-		}
+		Info.Components.Add(Component);
+		Component->SerializeStateHashes(*Info.ComponentStateAr);
+	}
 
 		for (auto It = ComponentsByHeightmap.CreateIterator(); It; ++It)
 		{
@@ -3957,53 +3959,53 @@ void ALandscapeProxy::UpdateGIBakedTextures(bool bBakeAllGITextures)
 		{
 			// We throttle, baking only one atlas per frame if bBakeAllGITextures is false.
 			if (!bBakeAllGITextures && NumGenerated > 0)
-			{
-				NumComponentsNeedingTextureBaking += Info.Components.Num();
-			}
-			else
-			{
-				UTexture2D* HeightmapTexture = It.Key();
-				// 1/8 the res of the heightmap
-				FIntPoint AtlasSize(HeightmapTexture->GetSizeX() >> 3, HeightmapTexture->GetSizeY() >> 3);
+				{
+					NumComponentsNeedingTextureBaking += Info.Components.Num();
+				}
+				else
+				{
+					UTexture2D* HeightmapTexture = It.Key();
+					// 1/8 the res of the heightmap
+					FIntPoint AtlasSize(HeightmapTexture->GetSizeX() >> 3, HeightmapTexture->GetSizeY() >> 3);
 
-				TArray<FColor> AtlasSamples;
-				AtlasSamples.AddZeroed(AtlasSize.X * AtlasSize.Y);
+					TArray<FColor> AtlasSamples;
+					AtlasSamples.AddZeroed(AtlasSize.X * AtlasSize.Y);
 
-				for (ULandscapeComponent* Component : Info.Components)
-				{
-					// not registered; ignore this component
-					if (!Component->SceneProxy)
+					for (ULandscapeComponent* Component : Info.Components)
 					{
-						continue;
-					}
+						// not registered; ignore this component
+						if (!Component->SceneProxy)
+						{
+							continue;
+						}
 
-					int32 ComponentSamples = (SubsectionSizeQuads + 1) * NumSubsections;
-					check(FMath::IsPowerOfTwo(ComponentSamples));
+						int32 ComponentSamples = (SubsectionSizeQuads + 1) * NumSubsections;
+						check(FMath::IsPowerOfTwo(ComponentSamples));
 
-					int32 BakeSize = ComponentSamples >> 3;
-					TArray<FColor> Samples;
-					if (FMaterialUtilities::ExportBaseColor(Component, BakeSize, Samples))
-					{
-						int32 AtlasOffsetX = FMath::RoundToInt(Component->HeightmapScaleBias.Z * (float)HeightmapTexture->GetSizeX()) >> 3;
-						int32 AtlasOffsetY = FMath::RoundToInt(Component->HeightmapScaleBias.W * (float)HeightmapTexture->GetSizeY()) >> 3;
-						for (int32 y = 0; y < BakeSize; y++)
+						int32 BakeSize = ComponentSamples >> 3;
+						TArray<FColor> Samples;
+						if (FMaterialUtilities::ExportBaseColor(Component, BakeSize, Samples))
 						{
-							FMemory::Memcpy(&AtlasSamples[(y + AtlasOffsetY)*AtlasSize.X + AtlasOffsetX], &Samples[y*BakeSize], sizeof(FColor)* BakeSize);
+							int32 AtlasOffsetX = FMath::RoundToInt(Component->HeightmapScaleBias.Z * (float)HeightmapTexture->GetSizeX()) >> 3;
+							int32 AtlasOffsetY = FMath::RoundToInt(Component->HeightmapScaleBias.W * (float)HeightmapTexture->GetSizeY()) >> 3;
+							for (int32 y = 0; y < BakeSize; y++)
+							{
+								FMemory::Memcpy(&AtlasSamples[(y + AtlasOffsetY)*AtlasSize.X + AtlasOffsetX], &Samples[y*BakeSize], sizeof(FColor)* BakeSize);
+							}
+							NumGenerated++;
 						}
-						NumGenerated++;
 					}
-				}
 				UTexture2D* AtlasTexture = FMaterialUtilities::CreateTexture(GetOutermost(), HeightmapTexture->GetName() + TEXT("_BaseColor"), AtlasSize, AtlasSamples, TC_Default, TEXTUREGROUP_World, RF_NoFlags, true, Info.CombinedStateId);
 
-				for (ULandscapeComponent* Component : Info.Components)
-				{
+					for (ULandscapeComponent* Component : Info.Components)
+					{
 					Component->BakedTextureMaterialGuid = Info.CombinedStateId;
-					Component->GIBakedBaseColorTexture = AtlasTexture;
-					Component->MarkRenderStateDirty();
+						Component->GIBakedBaseColorTexture = AtlasTexture;
+						Component->MarkRenderStateDirty();
+					}
 				}
 			}
 		}
-	}
 
 	TotalComponentsNeedingTextureBaking += NumComponentsNeedingTextureBaking;
 
@@ -4089,7 +4091,7 @@ void ALandscapeProxy::UpdatePhysicalMaterialTasksStatus(TSet<ULandscapeComponent
 	}
 
 	if (OutdatedComponentsCount)
-	{
+{
 		*OutdatedComponentsCount = OutdatedCount;
 	}
 }
diff --git a/Engine/Source/Runtime/Landscape/Private/LandscapeEditLayers.cpp b/Engine/Source/Runtime/Landscape/Private/LandscapeEditLayers.cpp
--- a/Engine/Source/Runtime/Landscape/Private/LandscapeEditLayers.cpp
+++ b/Engine/Source/Runtime/Landscape/Private/LandscapeEditLayers.cpp
@@ -3113,25 +3113,38 @@ int32 ALandscape::RegenerateLayersHeightmaps(FTextureToComponentHelper const& Ma
 	return HeightmapUpdateModes;
 }
 
-void ALandscape::UpdateForChangedHeightmaps(ULandscapeComponent* Component, int32 UpdateMode)
+void ALandscape::UpdateForChangedHeightmaps(ULandscapeComponent* InComponent, const FLandscapeEditLayerReadbackResult& InReadbackResult)
 {
-	const int32 HeightUpdateMode = UpdateMode & (ELandscapeLayerUpdateMode::Update_Heightmap_All | ELandscapeLayerUpdateMode::Update_Heightmap_Editing | ELandscapeLayerUpdateMode::Update_Heightmap_Editing_NoCollision);
+	// If the source data has changed, mark the component as needing a collision data update:
+	//  - If ELandscapeComponentUpdateFlag::Component_Update_Heightmap_Collision is passed, it will be done immediately
+	//  - If not, at least the component's collision data will still get updated eventually, when the flag is finally passed :
+	if (InReadbackResult.bModified)
+	{
+		InComponent->SetPendingCollisionDataUpdate(true);
+	}
+
+	const int32 HeightUpdateMode = InReadbackResult.UpdateModes & (ELandscapeLayerUpdateMode::Update_Heightmap_All | ELandscapeLayerUpdateMode::Update_Heightmap_Editing | ELandscapeLayerUpdateMode::Update_Heightmap_Editing_NoCollision);
 
 	if (IsUpdateFlagEnabledForModes(ELandscapeComponentUpdateFlag::Component_Update_Heightmap_Collision, HeightUpdateMode))
 	{
-		Component->UpdateCachedBounds();
-		Component->UpdateComponentToWorld();
+		// Only update collision if there was an actual change performed on the source data : 
+		if (InComponent->GetPendingCollisionDataUpdate())
+		{
+			InComponent->UpdateCachedBounds();
+			InComponent->UpdateComponentToWorld();
 
-		// Avoid updating height field if we are going to recreate collision in this update
-		bool bUpdateHeightfieldRegion = !IsUpdateFlagEnabledForModes(ELandscapeComponentUpdateFlag::Component_Update_Recreate_Collision, HeightUpdateMode);
-		Component->UpdateCollisionData(bUpdateHeightfieldRegion);
+			// Avoid updating height field if we are going to recreate collision in this update
+			bool bUpdateHeightfieldRegion = !IsUpdateFlagEnabledForModes(ELandscapeComponentUpdateFlag::Component_Update_Recreate_Collision, HeightUpdateMode);
+			InComponent->UpdateCollisionData(bUpdateHeightfieldRegion);
+			InComponent->SetPendingCollisionDataUpdate(false);
+		}
 	}
 	else if (IsUpdateFlagEnabledForModes(ELandscapeComponentUpdateFlag::Component_Update_Approximated_Bounds, HeightUpdateMode))
 	{
 		// Update bounds with an approximated value (real computation will be done anyways when computing collision)
 		const bool bInApproximateBounds = true;
-		Component->UpdateCachedBounds(bInApproximateBounds);
-		Component->UpdateComponentToWorld();
+		InComponent->UpdateCachedBounds(bInApproximateBounds);
+		InComponent->UpdateComponentToWorld();
 	}
 }
 
@@ -4361,13 +4374,26 @@ int32 ALandscape::RegenerateLayersWeightmaps(FTextureToComponentHelper const& Ma
 	return WeightmapUpdateModes;
 }
 
-void ALandscape::UpdateForChangedWeightmaps(ULandscapeComponent* Component, int32 UpdateMode)
+void ALandscape::UpdateForChangedWeightmaps(ULandscapeComponent* InComponent, const FLandscapeEditLayerReadbackResult& InReadbackResult)
 {
-	const int32 WeightUpdateMode = UpdateMode & (ELandscapeLayerUpdateMode::Update_Weightmap_All | ELandscapeLayerUpdateMode::Update_Weightmap_Editing | ELandscapeLayerUpdateMode::Update_Weightmap_Editing_NoCollision);
+	// If the source data has changed, mark the component as needing a collision layer data update:
+	//  - If ELandscapeComponentUpdateFlag::Component_Update_Weightmap_Collision is passed, it will be done immediately
+	//  - If not, at least the component's collision layer data will still get updated eventually, when the flag is finally passed :
+	if (InReadbackResult.bModified)
+	{
+		InComponent->SetPendingLayerCollisionDataUpdate(true);
+	}
+
+	const int32 WeightUpdateMode = InReadbackResult.UpdateModes & (ELandscapeLayerUpdateMode::Update_Weightmap_All | ELandscapeLayerUpdateMode::Update_Weightmap_Editing | ELandscapeLayerUpdateMode::Update_Weightmap_Editing_NoCollision);
 
 	if (IsUpdateFlagEnabledForModes(ELandscapeComponentUpdateFlag::Component_Update_Weightmap_Collision, WeightUpdateMode))
 	{
-		Component->UpdateCollisionLayerData();
+		// Only update collision data if there was an actual change performed on the source data : 
+		if (InComponent->GetPendingLayerCollisionDataUpdate())
+		{
+			InComponent->UpdateCollisionLayerData();
+			InComponent->SetPendingLayerCollisionDataUpdate(false);
+		}
 	}
 }
 
@@ -5253,7 +5279,6 @@ int32 ALandscape::UpdateCollisionAndClients(TMap<ULandscapeComponent*, FLandscap
 	{
 		ULandscapeComponent* LandscapeComponent = Pair.Key;
 		const int32 UpdateModes = Pair.Value.UpdateModes;
-		const bool bModified = Pair.Value.bModified;
 
 		bool bDeferClientUpdateForComponent = false;
 		bool bDoUpdateClient = true;
@@ -5359,14 +5384,8 @@ int32 ALandscape::UpdateAfterReadbackResolves(TMap<ULandscapeComponent*, FLandsc
 	{
 		for (auto Pair : Components)
 		{
-			if (Pair.Value.bModified)
-			{
-				ULandscapeComponent* Component = Pair.Key;
-				int32 UpdateModes = Pair.Value.UpdateModes;
-
-				UpdateForChangedHeightmaps(Component, UpdateModes);
-				UpdateForChangedWeightmaps(Component, UpdateModes);
-			}
+			UpdateForChangedHeightmaps(Pair.Key, Pair.Value);
+			UpdateForChangedWeightmaps(Pair.Key, Pair.Value);
 		}
 
 		GetLandscapeInfo()->UpdateAllAddCollisions();
