Fixed ladnscape components bounds being invalidated when using the water brush : the component were first updated with the ELandscapeComponentUpdateFlag::Component_Update_Approximated_Bounds flag, (temporarily) invalidated the bounds, and then were updated with the IsUpdateFlagEnabledForModes(ELandscapeComponentUpdateFlag::Component_Update_Heightmap_Collision flag, but didn't end up updating the bounds if the heightfield was actually unchanged.
The fix is to not invalidate the bounds (with an approximate bounds computation) if no change is actually detected on the heightmaps

diff --git a/Engine/Source/Runtime/Landscape/Private/LandscapeEditLayers.cpp b/Engine/Source/Runtime/Landscape/Private/LandscapeEditLayers.cpp
--- a/Engine/Source/Runtime/Landscape/Private/LandscapeEditLayers.cpp
+++ b/Engine/Source/Runtime/Landscape/Private/LandscapeEditLayers.cpp
@@ -3125,10 +3125,10 @@ void ALandscape::UpdateForChangedHeightmaps(ULandscapeComponent* InComponent, co
 
 	const int32 HeightUpdateMode = InReadbackResult.UpdateModes & (ELandscapeLayerUpdateMode::Update_Heightmap_All | ELandscapeLayerUpdateMode::Update_Heightmap_Editing | ELandscapeLayerUpdateMode::Update_Heightmap_Editing_NoCollision);
 
-	if (IsUpdateFlagEnabledForModes(ELandscapeComponentUpdateFlag::Component_Update_Heightmap_Collision, HeightUpdateMode))
+	// Only update collision if there was an actual change performed on the source data : 
+	if (InComponent->GetPendingCollisionDataUpdate())
 	{
-		// Only update collision if there was an actual change performed on the source data : 
-		if (InComponent->GetPendingCollisionDataUpdate())
+		if (IsUpdateFlagEnabledForModes(ELandscapeComponentUpdateFlag::Component_Update_Heightmap_Collision, HeightUpdateMode))
 		{
 			InComponent->UpdateCachedBounds();
 			InComponent->UpdateComponentToWorld();
@@ -3138,13 +3138,13 @@ void ALandscape::UpdateForChangedHeightmaps(ULandscapeComponent* InComponent, co
 			InComponent->UpdateCollisionData(bUpdateHeightfieldRegion);
 			InComponent->SetPendingCollisionDataUpdate(false);
 		}
-	}
-	else if (IsUpdateFlagEnabledForModes(ELandscapeComponentUpdateFlag::Component_Update_Approximated_Bounds, HeightUpdateMode))
-	{
-		// Update bounds with an approximated value (real computation will be done anyways when computing collision)
-		const bool bInApproximateBounds = true;
-		InComponent->UpdateCachedBounds(bInApproximateBounds);
-		InComponent->UpdateComponentToWorld();
+		else if (IsUpdateFlagEnabledForModes(ELandscapeComponentUpdateFlag::Component_Update_Approximated_Bounds, HeightUpdateMode))
+		{
+			// Update bounds with an approximated value (real computation will be done anyways when computing collision)
+			const bool bInApproximateBounds = true;
+			InComponent->UpdateCachedBounds(bInApproximateBounds);
+			InComponent->UpdateComponentToWorld();
+		}
 	}
 }
 
