Fix static analysis warning

diff --git a/Engine/Source/Runtime/Landscape/Private/LandscapeRenderMobile.cpp b/Engine/Source/Runtime/Landscape/Private/LandscapeRenderMobile.cpp
--- a/Engine/Source/Runtime/Landscape/Private/LandscapeRenderMobile.cpp
+++ b/Engine/Source/Runtime/Landscape/Private/LandscapeRenderMobile.cpp
@@ -277,14 +277,14 @@ FLandscapeMobileRenderData::FLandscapeMobileRenderData(const TArray<uint8>& InPl
 {
 	FMemoryReader MemAr(InPlatformData);
 
-	int32 NumHoleLods;
+	int32 NumHoleLods = 0;
 	MemAr << NumHoleLods;
 	if (NumHoleLods > 0)
 	{
 		HoleData = new FLandscapeMobileHoleData;
 		HoleData->NumHoleLods = NumHoleLods;
 
-		bool b16BitIndices;
+		bool b16BitIndices = false;
 		MemAr << b16BitIndices;
 		if (b16BitIndices)
 		{
@@ -297,7 +297,7 @@ FLandscapeMobileRenderData::FLandscapeMobileRenderData(const TArray<uint8>& InPl
 	}
 
 	{
-		int32 VertexCount;
+		int32 VertexCount = 0;
 		MemAr << VertexCount;
 		TArray<uint8> VertexData;
 		VertexData.SetNumUninitialized(VertexCount * sizeof(FLandscapeMobileVertex));
@@ -306,7 +306,7 @@ FLandscapeMobileRenderData::FLandscapeMobileRenderData(const TArray<uint8>& InPl
 	}
 
 	{
-		int32 NumOccluderVertices;
+		int32 NumOccluderVertices = 0;
 		MemAr << NumOccluderVertices;
 		if (NumOccluderVertices > 0)
 		{
