Made landscape brushes AlwaysLoaded : they should have the same lifecycle as the landscape, basically

diff --git a/Engine/Source/Runtime/Landscape/Public/LandscapeBlueprintBrushBase.h b/Engine/Source/Runtime/Landscape/Public/LandscapeBlueprintBrushBase.h
--- a/Engine/Source/Runtime/Landscape/Public/LandscapeBlueprintBrushBase.h
+++ b/Engine/Source/Runtime/Landscape/Public/LandscapeBlueprintBrushBase.h
@@ -76,5 +76,7 @@ public:
 	virtual void Destroyed() override;
 
 	virtual void PushDeferredLayersContentUpdate();
+
+	virtual EActorGridPlacement GetDefaultGridPlacement() const override { return EActorGridPlacement::AlwaysLoaded; }
 #endif
 };
\ No newline at end of file
