Disable mobile landscape normalmap texture count optimization if we are already saving a weightmap channel from normalized weight blends.
Normalized weight blending only acts as a channel optimization if the number of weightmap layers is 3 or less anyway. So in that case we don't need an additional optimization to save a texture slot.
#[fyi] jonathan.bard

diff --git a/Engine/Source/Runtime/Landscape/Private/Landscape.cpp b/Engine/Source/Runtime/Landscape/Private/Landscape.cpp
--- a/Engine/Source/Runtime/Landscape/Private/Landscape.cpp
+++ b/Engine/Source/Runtime/Landscape/Private/Landscape.cpp
@@ -133,7 +133,7 @@ namespace LandscapeCookStats
 // differences, etc.) replace the version GUID below with a new one.
 // In case of merge conflicts with DDC versions, you MUST generate a new GUID
 // and set this new GUID as the version.                                       
-#define LANDSCAPE_MOBILE_COOK_VERSION TEXT("307EBBCF9AD644059520051B212C3A87")
+#define LANDSCAPE_MOBILE_COOK_VERSION TEXT("C2781BC371584D8BB9A2AB8BFDFDF7A7")
 
 #define LOCTEXT_NAMESPACE "Landscape"
 
diff --git a/Engine/Source/Runtime/Landscape/Private/LandscapeEdit.cpp b/Engine/Source/Runtime/Landscape/Private/LandscapeEdit.cpp
--- a/Engine/Source/Runtime/Landscape/Private/LandscapeEdit.cpp
+++ b/Engine/Source/Runtime/Landscape/Private/LandscapeEdit.cpp
@@ -6370,23 +6370,25 @@ void ULandscapeComponent::GeneratePlatformPixelData()
 		MobileWeightmapTextures.Add(CurrentWeightmapTexture);
 		int32 CurrentChannel = 0;
 
+		// We can potentially save a channel allocation if we have weight based blends.
+		const bool bAtLeastOneWeightBasedBlend = MobileWeightmapLayerAllocations.FindByPredicate([&](const FWeightmapLayerAllocationInfo& Allocation) -> bool { return !Allocation.LayerInfo->bNoWeightBlend; }) != nullptr;
+		const bool bUseWeightBasedChannelOptim =  bAtLeastOneWeightBasedBlend && MobileWeightmapLayerAllocations.Num() <= 3;
+		MobileBlendableLayerMask = 0;
+
 		// Give normal map a full texture if this doesn't increase the overall allocation count.
 		// This then saves a texture slot because we don't need to sample a combined normalmap/weightmap texture with two different sampler settings.
-		int32 NumTexturesCombinedNormal = FMath::DivideAndRoundUp(MobileWeightmapLayerAllocations.Num() + 2, 4);
-		int32 NumTexturesIsolatedNormal = 1 + FMath::DivideAndRoundUp(MobileWeightmapLayerAllocations.Num(), 4);
-		bool bIsolateNormalMap = NumTexturesCombinedNormal == NumTexturesIsolatedNormal;
+		// This optimization won't be useful or valid if we are already applying the weight based blending channel optimization.
+		const int32 NumTexturesCombinedNormal = FMath::DivideAndRoundUp(MobileWeightmapLayerAllocations.Num() + 2, 4);
+		const int32 NumTexturesIsolatedNormal = 1 + FMath::DivideAndRoundUp(MobileWeightmapLayerAllocations.Num(), 4);
+		const bool bIsolateNormalMap = !bUseWeightBasedChannelOptim && NumTexturesCombinedNormal == NumTexturesIsolatedNormal;
 		int32 RemainingChannels = bIsolateNormalMap ? 0 : 2;
 
-		MobileBlendableLayerMask = 0;
-
-		bool bAtLeastOneWeightBasedBlend = MobileWeightmapLayerAllocations.FindByPredicate([&](const FWeightmapLayerAllocationInfo& Allocation) -> bool { return !Allocation.LayerInfo->bNoWeightBlend; }) != nullptr;
-
 		for (auto& Allocation : MobileWeightmapLayerAllocations)
 		{
 			if (Allocation.LayerInfo)
 			{
 				// If we can pack into 2 channels with the 3rd implied, track the mask for the weight blendable layers
-				if (bAtLeastOneWeightBasedBlend && MobileWeightmapLayerAllocations.Num() <= 3)
+				if (bUseWeightBasedChannelOptim)
 				{
 					MobileBlendableLayerMask |= (!Allocation.LayerInfo->bNoWeightBlend ? (1 << CurrentChannel) : 0);
 
