Fixed landscape splines not properly emptying the MeshComponents array after moving (trying to move, actually) its content to a local array : when TObjectPtr was introduced, the move was silently replaced by a copy as the destination array was of a different non-moveable type (ptr vs. TObjectPtr), leaving the source array *not* empty after the call, while the algorithm implied it was (after a move, no assumption should be made as to the source container's state, apart that it's undefined).

diff --git a/Engine/Source/Runtime/Landscape/Private/LandscapeSplines.cpp b/Engine/Source/Runtime/Landscape/Private/LandscapeSplines.cpp
--- a/Engine/Source/Runtime/Landscape/Private/LandscapeSplines.cpp
+++ b/Engine/Source/Runtime/Landscape/Private/LandscapeSplines.cpp
@@ -2755,19 +2755,21 @@ void ULandscapeSplineSegment::UpdateSplinePoints(bool bUpdateCollision, bool bUp
 
 	TArray<USplineMeshComponent*> MeshComponents;
 
-	TArray<USplineMeshComponent*> OldLocalMeshComponents = MoveTemp(LocalMeshComponents);
+	TArray<TObjectPtr<USplineMeshComponent>> OldLocalMeshComponents;
+	Swap(LocalMeshComponents, OldLocalMeshComponents);
 	LocalMeshComponents.Reserve(20);
 
 	// Gather all ForeignMesh associated with this Segment
 	TMap<ULandscapeSplinesComponent*, TArray<USplineMeshComponent*>> ForeignMeshComponentsMap = GetForeignMeshComponents();
 	
 	// Unregister components, Remove Foreign/Local Associations
-	for (auto* LocalMeshComponent : OldLocalMeshComponents)
+	for (TObjectPtr<USplineMeshComponent> LocalMeshComponent : OldLocalMeshComponents)
 	{
-		checkSlow(OuterSplines->MeshComponentLocalOwnersMap.FindRef(LocalMeshComponent) == this);
-		verifySlow(OuterSplines->MeshComponentLocalOwnersMap.Remove(LocalMeshComponent) == 1);
-		LocalMeshComponent->Modify();
-		LocalMeshComponent->UnregisterComponent();
+		USplineMeshComponent* Comp = LocalMeshComponent.Get();
+		checkSlow(OuterSplines->MeshComponentLocalOwnersMap.FindRef(Comp) == this);
+		verifySlow(OuterSplines->MeshComponentLocalOwnersMap.Remove(Comp) == 1);
+		Comp->Modify();
+		Comp->UnregisterComponent();
 	}
 	for (auto& ForeignMeshComponentsPair : ForeignMeshComponentsMap)
 	{
@@ -2851,7 +2853,7 @@ void ULandscapeSplineSegment::UpdateSplinePoints(bool bUpdateCollision, bool bUp
 			{
 				if (OldLocalMeshComponents.Num() > 0)
 				{
-					MeshComponent = OldLocalMeshComponents.Pop(false);
+					MeshComponent = OldLocalMeshComponents.Pop(false).Get();
 				}
 			}
 			else
@@ -3130,7 +3132,7 @@ void ULandscapeSplineSegment::UpdateSplinePoints(bool bUpdateCollision, bool bUp
 	}
 	
 	// Clean up unused components
-	for (auto* LocalMeshComponent : OldLocalMeshComponents)
+	for (TObjectPtr<USplineMeshComponent> LocalMeshComponent : OldLocalMeshComponents)
 	{
 		LocalMeshComponent->DestroyComponent();
 	}
