Fixed static analysis

diff --git a/Engine/Source/Runtime/Landscape/Private/LandscapeEditInterface.cpp b/Engine/Source/Runtime/Landscape/Private/LandscapeEditInterface.cpp
--- a/Engine/Source/Runtime/Landscape/Private/LandscapeEditInterface.cpp
+++ b/Engine/Source/Runtime/Landscape/Private/LandscapeEditInterface.cpp
@@ -496,12 +496,13 @@ void FLandscapeEditDataInterface::RecalculateNormals()
 		int32 X2 = Component->GetSectionBase().X+ComponentSizeQuads+1;
 		int32 Y2 = Component->GetSectionBase().Y+ComponentSizeQuads+1;
 		int32 Stride = ComponentSizeQuads+3; 
+		int32 StrideSquared = FMath::Square(Stride);
 
-		uint16* HeightData = new uint16[FMath::Square(Stride)];
-		FVector* VertexNormals = new FVector[FMath::Square(Stride)];
-		FMemory::Memzero(VertexNormals, FMath::Square(Stride)*sizeof(FVector));
-		FVector2D* XYOffsets = new FVector2D[FMath::Square(Stride)];
-		FMemory::Memzero(XYOffsets, FMath::Square(Stride)*sizeof(FVector2D));
+		uint16* HeightData = new uint16[StrideSquared];
+		FVector* VertexNormals = new FVector[StrideSquared];
+		FMemory::Memzero(VertexNormals, StrideSquared * sizeof(FVector));
+		FVector2D* XYOffsets = new FVector2D[StrideSquared];
+		FMemory::Memzero(XYOffsets, StrideSquared * sizeof(FVector2D));
 
 		// Get XY offset
 		GetXYOffsetDataFast(X1,Y1,X2,Y2,XYOffsets,0);
@@ -513,10 +514,10 @@ void FLandscapeEditDataInterface::RecalculateNormals()
 		{
 			for( int32 X=0;X<Stride-1;X++ )
 			{
-				FVector Vert00 = FVector(XYOffsets[(X+0) + Stride*(Y+0)].X,		XYOffsets[(X+0) + Stride*(Y+0)].Y, ((float)HeightData[(X+0) + Stride*(Y+0)] - 32768.0f)*LANDSCAPE_ZSCALE) * DrawScale;
-				FVector Vert01 = FVector(XYOffsets[(X+0) + Stride*(Y+0)].X,		XYOffsets[(X+0) + Stride*(Y+0)].Y+1.0f, ((float)HeightData[(X+0) + Stride*(Y+1)] - 32768.0f)*LANDSCAPE_ZSCALE) * DrawScale;
-				FVector Vert10 = FVector(XYOffsets[(X+0) + Stride*(Y+0)].X+1.0f,	XYOffsets[(X+0) + Stride*(Y+0)].Y, ((float)HeightData[(X+1) + Stride*(Y+0)] - 32768.0f)*LANDSCAPE_ZSCALE) * DrawScale;
-				FVector Vert11 = FVector(XYOffsets[(X+0) + Stride*(Y+0)].X+1.0f,	XYOffsets[(X+0) + Stride*(Y+0)].Y+1.0f,((float)HeightData[(X+1) + Stride*(Y+1)] - 32768.0f)*LANDSCAPE_ZSCALE) * DrawScale;
+				FVector Vert00 = FVector(XYOffsets[(X+0) + Stride*(Y+0)].X,			XYOffsets[(X+0) + Stride*(Y+0)].Y,		((float)HeightData[(X+0) + Stride*(Y+0)] - 32768.0f)*LANDSCAPE_ZSCALE) * DrawScale;
+				FVector Vert01 = FVector(XYOffsets[(X+0) + Stride*(Y+0)].X,			XYOffsets[(X+0) + Stride*(Y+0)].Y+1.0f, ((float)HeightData[(X+0) + Stride*(Y+1)] - 32768.0f)*LANDSCAPE_ZSCALE) * DrawScale;
+				FVector Vert10 = FVector(XYOffsets[(X+0) + Stride*(Y+0)].X+1.0f,	XYOffsets[(X+0) + Stride*(Y+0)].Y,		((float)HeightData[(X+1) + Stride*(Y+0)] - 32768.0f)*LANDSCAPE_ZSCALE) * DrawScale;
+				FVector Vert11 = FVector(XYOffsets[(X+0) + Stride*(Y+0)].X+1.0f,	XYOffsets[(X+0) + Stride*(Y+0)].Y+1.0f,	((float)HeightData[(X+1) + Stride*(Y+1)] - 32768.0f)*LANDSCAPE_ZSCALE) * DrawScale;
 
 				FVector FaceNormal1 = ((Vert00-Vert10) ^ (Vert10-Vert11)).GetSafeNormal();
 				FVector FaceNormal2 = ((Vert11-Vert01) ^ (Vert01-Vert00)).GetSafeNormal(); 
@@ -551,6 +552,7 @@ void FLandscapeEditDataInterface::RecalculateNormals()
 						const int32 ReadX = (SubsectionSizeQuads) * SubIndexX + SubX;
 						const int32 ReadY = (SubsectionSizeQuads) * SubIndexY + SubY;
 						const int32 ReadIndex = (ReadX + 1) + (ReadY + 1) * Stride;
+						checkSlow(ReadIndex < StrideSquared);
 						const FVector Normal = VertexNormals[ReadIndex].GetSafeNormal();
 
 						// Write to (shared) Heightmap texture which has duplicated values on subsection and section borders.
