Fix missing transition when creating a landscape layer.

#rnx

diff --git a/Engine/Source/Runtime/Landscape/Private/LandscapeEditReadback.cpp b/Engine/Source/Runtime/Landscape/Private/LandscapeEditReadback.cpp
--- a/Engine/Source/Runtime/Landscape/Private/LandscapeEditReadback.cpp
+++ b/Engine/Source/Runtime/Landscape/Private/LandscapeEditReadback.cpp
@@ -78,6 +78,7 @@ void KickTask_RenderThread(FRHICommandListImmediate& RHICmdList, FLandscapeEditR
 {
 	// Transition staging textures for write.
 	TArray <FRHITransitionInfo> Transitions;
+	Transitions.Add(FRHITransitionInfo(Task.TextureResource->GetTexture2DRHI(), ERHIAccess::SRVMask, ERHIAccess::CopySrc));
 	for (uint32 MipIndex = 0; MipIndex < Task.NumMips; ++MipIndex)
 	{
 		Transitions.Add(FRHITransitionInfo(Task.StagingTextures[MipIndex], ERHIAccess::Unknown, ERHIAccess::CopyDest));
@@ -99,6 +100,7 @@ void KickTask_RenderThread(FRHICommandListImmediate& RHICmdList, FLandscapeEditR
 
 	// Transition staging textures for read.
 	Transitions.Reset();
+	Transitions.Add(FRHITransitionInfo(Task.TextureResource->GetTexture2DRHI(), ERHIAccess::CopySrc, ERHIAccess::SRVMask));
 	for (uint32 MipIndex = 0; MipIndex < Task.NumMips; ++MipIndex)
 	{
 		Transitions.Add(FRHITransitionInfo(Task.StagingTextures[MipIndex], ERHIAccess::Unknown, ERHIAccess::CPURead));
