WorldPartition: Landscape delete support

diff --git a/Engine/Source/Runtime/Engine/Private/ActorPartition/ActorPartitionSubsystem.cpp b/Engine/Source/Runtime/Engine/Private/ActorPartition/ActorPartitionSubsystem.cpp
--- a/Engine/Source/Runtime/Engine/Private/ActorPartition/ActorPartitionSubsystem.cpp
+++ b/Engine/Source/Runtime/Engine/Private/ActorPartition/ActorPartitionSubsystem.cpp
@@ -211,13 +211,12 @@ public:
 					return false;
 				}
 
-				FoundActor = CastChecked<APartitionActor>(DescActor);
-				check(FoundActor->GridSize == InGridSize && FoundActor->GetGridGuid() == InActorPartitionId.GetGridGuid());
-				// Typically, this could happen if a PartitionActor was manually deleted (it's now be prevented by APartitionActor::CanDeleteSelectedActor);
-				// If at some point it's decided to be supported, a solution for a found PendingKill PartitionActor is to
-				// call Modify(), Rename the actor to a unique/trash name and behave as if no actor was found
-				check(!FoundActor->IsPendingKill());
-				return false;
+				// Skip PendingKill actors because they will be renamed out of the way later
+				if (!DescActor->IsPendingKill())
+				{
+					FoundActor = CastChecked<APartitionActor>(DescActor);
+					return false;
+				}
 			}
 			return true;
 		};
@@ -262,11 +261,16 @@ public:
 			SpawnParams.Name = *ActorNameBuilder;
 			SpawnParams.NameMode = FActorSpawnParameters::ESpawnActorNameMode::Required_Fatal;
 
-			// Handle the case where the actor already exists, but is in the undo stack
+			// Handle the case where the actor already exists, but is in the undo stack (was deleted)
 			if (UObject* ExistingObject = StaticFindObject(nullptr, World->PersistentLevel, *SpawnParams.Name.ToString()))
 			{
-				check(CastChecked<AActor>(ExistingObject)->IsPendingKill());
-				ExistingObject->Rename(nullptr, nullptr, REN_DontCreateRedirectors | REN_DoNotDirty | REN_NonTransactional | REN_ForceNoResetLoaders);
+				AActor* ExistingActor = CastChecked<AActor>(ExistingObject);
+				check(ExistingActor->IsPendingKill());
+				ExistingActor->Modify();
+				// Don't go through AActor::Rename here because we aren't changing outers (the actor's level) and we also don't want to reset loaders
+				// if the actor is using an external package. We really just want to rename that actor out of the way so we can spawn the new one in
+				// the exact same package, keeping the package name intact.
+				ExistingActor->UObject::Rename(nullptr, nullptr, REN_DontCreateRedirectors | REN_DoNotDirty | REN_NonTransactional | REN_ForceNoResetLoaders);
 			}
 						
 			FVector CellCenter(CellBounds.GetCenter());
diff --git a/Engine/Source/Runtime/Landscape/Classes/Landscape.h b/Engine/Source/Runtime/Landscape/Classes/Landscape.h
--- a/Engine/Source/Runtime/Landscape/Classes/Landscape.h
+++ b/Engine/Source/Runtime/Landscape/Classes/Landscape.h
@@ -232,7 +232,7 @@ public:
 	virtual bool ShouldImport(FString* ActorPropString, bool IsMovingLevel) override;
 	virtual void PostEditImport() override;
 	virtual void PostDuplicate(bool bDuplicateForPIE) override;
-
+	virtual bool CanDeleteSelectedActor(FText& OutReason) const override;
 	virtual EActorGridPlacement GetDefaultGridPlacement() const override { return EActorGridPlacement::AlwaysLoaded; }
 #endif
 	virtual void PostLoad() override;
diff --git a/Engine/Source/Runtime/Landscape/Classes/LandscapeInfo.h b/Engine/Source/Runtime/Landscape/Classes/LandscapeInfo.h
--- a/Engine/Source/Runtime/Landscape/Classes/LandscapeInfo.h
+++ b/Engine/Source/Runtime/Landscape/Classes/LandscapeInfo.h
@@ -212,6 +212,9 @@ public:
 	LANDSCAPE_API int32 GetLayerInfoIndex(ULandscapeLayerInfoObject* LayerInfo, ALandscapeProxy* Owner = nullptr) const;
 	LANDSCAPE_API bool UpdateLayerInfoMap(ALandscapeProxy* Proxy = nullptr, bool bInvalidate = false);
 
+
+	LANDSCAPE_API bool CanDeleteLandscape(FText& OutReason) const;
+
 	/**
 	 *  Returns the landscape proxy of this landscape info in the given level (if it exists)
 	 *  @param  Level  Level to look in
diff --git a/Engine/Source/Runtime/Landscape/Classes/LandscapeStreamingProxy.h b/Engine/Source/Runtime/Landscape/Classes/LandscapeStreamingProxy.h
--- a/Engine/Source/Runtime/Landscape/Classes/LandscapeStreamingProxy.h
+++ b/Engine/Source/Runtime/Landscape/Classes/LandscapeStreamingProxy.h
@@ -37,7 +37,8 @@ public:
 	virtual void PostEditChangeProperty(FPropertyChangedEvent& PropertyChangedEvent) override;
 	virtual bool CanEditChange(const FProperty* InProperty) const override;
 	virtual void PostRegisterAllComponents() override;
-	virtual AActor* GetSceneOutlinerParent() const;
+	virtual AActor* GetSceneOutlinerParent() const override;
+	virtual bool CanDeleteSelectedActor(FText& OutReason) const override;
 #endif
 	//~ End UObject Interface
 
diff --git a/Engine/Source/Runtime/Landscape/Private/LandscapeEdit.cpp b/Engine/Source/Runtime/Landscape/Private/LandscapeEdit.cpp
--- a/Engine/Source/Runtime/Landscape/Private/LandscapeEdit.cpp
+++ b/Engine/Source/Runtime/Landscape/Private/LandscapeEdit.cpp
@@ -63,6 +63,7 @@ LandscapeEdit.cpp: Landscape editing
 #include "WorldPartition/WorldPartitionActorDesc.h"
 #include "WorldPartition/Landscape/LandscapeActorDesc.h"
 #include "ActorPartition/ActorPartitionSubsystem.h"
+#include "LandscapeSplineActor.h"
 #endif
 #include "Algo/Count.h"
 #include "Serialization/MemoryWriter.h"
@@ -5029,6 +5030,100 @@ AActor* ALandscapeStreamingProxy::GetSceneOutlinerParent() const
 	return Super::GetSceneOutlinerParent();
 }
 
+bool ALandscapeStreamingProxy::CanDeleteSelectedActor(FText& OutReason) const
+{
+	return true;
+}
+
+bool ALandscape::CanDeleteSelectedActor(FText& OutReason) const
+{
+	if (!IsUserManaged())
+	{
+		// Allow Delete of Actor if all other related actors have been deleted
+		ULandscapeInfo* Info = GetLandscapeInfo();
+		check(Info);
+		return Info->CanDeleteLandscape(OutReason);
+	}
+
+	return true;
+}
+
+bool ULandscapeInfo::CanDeleteLandscape(FText& OutReason) const
+{
+	check(LandscapeActor != nullptr);
+	int32 UndeletedProxyCount = 0;
+	int32 UndeletedSplineCount = 0;
+
+	// Check Registered Proxies
+	for (ALandscapeProxy* RegisteredProxy : Proxies)
+	{
+		if (RegisteredProxy == LandscapeActor)
+		{
+			continue;
+		}
+
+		check(!RegisteredProxy->IsPendingKill());
+		UndeletedProxyCount++;
+	}
+
+	// Then check for Unloaded Proxies
+	for (const FWorldPartitionHandle& Handle : ProxyHandles)
+	{
+		ALandscapeProxy* LandscapeProxy = Cast<ALandscapeProxy>(Handle->GetActor());
+		if (LandscapeProxy == LandscapeActor)
+		{
+			continue;
+		}
+		
+		// If LandscapeProxy is null then it is not loaded so not deleted.
+		if (!LandscapeProxy)
+		{
+			++UndeletedProxyCount;
+		}
+		else
+		{
+			// If Actor is loaded it should be Registered and not pending kill (already accounted for) or pending kill (deleted)
+			check(Proxies.Contains(LandscapeProxy) == !LandscapeProxy->IsPendingKill());
+		}
+	}
+
+	// Check Registered Splines
+	for (TScriptInterface<ILandscapeSplineInterface> SplineOwner : SplineActors)
+	{
+		// Only check for ALandscapeSplineActor type because ALandscapeProxy also implement the ILandscapeSplineInterface for non WP worlds
+		if(ALandscapeSplineActor* SplineActor = Cast<ALandscapeSplineActor>(SplineOwner.GetObject()))
+		{ 
+			check(!SplineActor->IsPendingKill());
+			UndeletedSplineCount++;
+		}
+	}
+
+	// Then check for Unloaded Splines
+	for (const FWorldPartitionHandle& Handle : SplineHandles)
+	{
+		ALandscapeSplineActor* SplineActor = Cast<ALandscapeSplineActor>(Handle->GetActor());
+		
+		// If SplineActor is null then it is not loaded/deleted. If it's loaded then it needs to be pending kill.
+		if (!SplineActor)
+		{
+			++UndeletedSplineCount;
+		}
+		else
+		{
+			// If Actor is loaded it should be Registered and not pending kill (already accounted for) or pending kill (deleted)
+			check(SplineActors.Contains(SplineActor) == !SplineActor->IsPendingKill());
+		}
+	}
+
+	if (UndeletedProxyCount > 0 || UndeletedSplineCount > 0)
+	{
+		OutReason = FText::Format(LOCTEXT("CanDeleteLandscapeReason", "Landscape can't be deleted because it still has {0} LandscapeStreamingProxies and {1} LandscapeSplineActors"), FText::AsNumber(UndeletedProxyCount), FText::AsNumber(UndeletedSplineCount));
+		return false;
+	}
+
+	return true;
+}
+
 void ALandscape::PreEditChange(FProperty* PropertyThatWillChange)
 {
 	PreEditLandscapeMaterial = LandscapeMaterial;
