Fix for major frame hitches when adding levels to a world.  Previously, adding a level to the world would trigger a linear loop through all primitives in the scene, testing to see if they were in the current level, and calling OnLevelAddedToWorld() if they were flagged to do so.  The more primitives in the scene, the larger the hitch.

Now, proxiess can register for needing an on added/removed from level call.  If this is requested, they are added to a map, which maintains a list of primitives that need the add/remove call, per level.

The only component that currently utilizes this is the Landscape component.  To facilitate it's needs, the generic flags used to determine if it should be visible have been removed from FPrimitiveSceneProxy, bRequiresVisibleLevelToRender bIsComponentLevelVisible, and replaced with bForceHidden.  The landscape component will now set that flag under the same circumstances as before.

diff --git a/Engine/Source/Runtime/Engine/Private/PrimitiveSceneProxy.cpp b/Engine/Source/Runtime/Engine/Private/PrimitiveSceneProxy.cpp
--- a/Engine/Source/Runtime/Engine/Private/PrimitiveSceneProxy.cpp
+++ b/Engine/Source/Runtime/Engine/Private/PrimitiveSceneProxy.cpp
@@ -129,8 +129,7 @@ FPrimitiveSceneProxy::FPrimitiveSceneProxy(const UPrimitiveComponent* InComponen
 ,	bVisibleInRayTracing(InComponent->bVisibleInRayTracing)
 ,	bRenderInDepthPass(InComponent->bRenderInDepthPass)
 ,	bRenderInMainPass(InComponent->bRenderInMainPass)
-,	bRequiresVisibleLevelToRender(false)
-,	bIsComponentLevelVisible(false)
+,	bForceHidden(false)
 ,	bCollisionEnabled(InComponent->IsCollisionEnabled())
 ,	bTreatAsBackgroundForOcclusion(InComponent->bTreatAsBackgroundForOcclusion)
 ,	bCanSkipRedundantTransformUpdates(true)
@@ -164,7 +163,7 @@ FPrimitiveSceneProxy::FPrimitiveSceneProxy(const UPrimitiveComponent* InComponen
 ,	bSupportsDistanceFieldRepresentation(false)
 ,	bSupportsMeshCardRepresentation(false)
 ,	bSupportsHeightfieldRepresentation(false)
-,	bNeedsLevelAddedToWorldNotification(false)
+,	bShouldNotifyOnWorldAddRemove(false)
 ,	bWantsSelectionOutline(true)
 ,	bVerifyUsedMaterials(true)
 ,   bHasPrevInstanceTransforms(false)
@@ -251,13 +250,6 @@ FPrimitiveSceneProxy::FPrimitiveSceneProxy(const UPrimitiveComponent* InComponen
 		bIsFoliage = FFoliageHelper::IsOwnedByFoliage(InComponent->GetOwner());
 #endif
 	}
-	
-	// 
-	// Flag components to render only after level will be fully added to the world
-	//
-	ULevel* ComponentLevel = InComponent->GetComponentLevel();
-	bRequiresVisibleLevelToRender = (ComponentLevel && ComponentLevel->bRequireFullVisibilityToRender);
-	bIsComponentLevelVisible = (!ComponentLevel || ComponentLevel->bIsVisible);
 
 	// Setup the runtime virtual texture information
 	if (UseVirtualTexturing(GetScene().GetFeatureLevel()))
@@ -805,8 +797,8 @@ bool FPrimitiveSceneProxy::IsShown(const FSceneView* View) const
 		return false;
 	}
 
-	// if primitive requires component level to be visible
-	if (bRequiresVisibleLevelToRender && !bIsComponentLevelVisible)
+	// If primitive is forcibly hidden
+	if (bForceHidden)
 	{
 		return false;
 	}
@@ -926,8 +918,8 @@ bool FPrimitiveSceneProxy::IsShadowCast(const FSceneView* View) const
 		}
 #endif	//#if WITH_EDITOR
 
-		// if primitive requires component level to be visible
-		if (bRequiresVisibleLevelToRender && !bIsComponentLevelVisible)
+		// If primitive is forcibly hidden
+		if (bForceHidden)
 		{
 			return false;
 		}
diff --git a/Engine/Source/Runtime/Engine/Public/PrimitiveSceneProxy.h b/Engine/Source/Runtime/Engine/Public/PrimitiveSceneProxy.h
--- a/Engine/Source/Runtime/Engine/Public/PrimitiveSceneProxy.h
+++ b/Engine/Source/Runtime/Engine/Public/PrimitiveSceneProxy.h
@@ -379,9 +379,18 @@ public:
 	/**
 	 * Called to notify the proxy that the level has been fully added to
 	 * the world and the primitive will now be rendered.
-	 * Only called if bNeedsLevelAddedToWorldNotification is set to true.
+	 * Only called if bShouldNotifyOnWorldAddRemove is set to true.
+	 * 
+	 * @return return true if the primitive should be added to the scene
+	 */
+	virtual bool OnLevelAddedToWorld_RenderThread() {return false;}
+
+	/**
+	 * Called to notify the proxy that the level has been fully removed from
+	 * the world and the primitive will not be rendered.
+	 * Only called if bShouldNotifyOnWorldAddRemove is set to true.
 	 */
-	virtual void OnLevelAddedToWorld() {}
+	virtual void OnLevelRemovedFromWorld_RenderThread() {}
 
 	/**
 	* @return true if the proxy can be culled when occluded by other primitives
@@ -605,8 +614,8 @@ public:
 	inline bool SupportsInstanceDataBuffer() const { return bSupportsInstanceDataBuffer; }
 	inline bool HasPrevInstanceTransforms() const { return bHasPrevInstanceTransforms; }
 	inline bool TreatAsBackgroundForOcclusion() const { return bTreatAsBackgroundForOcclusion; }
-	inline bool NeedsLevelAddedToWorldNotification() const { return bNeedsLevelAddedToWorldNotification; }
-	inline bool IsComponentLevelVisible() const { return bIsComponentLevelVisible; }
+	inline bool ShouldNotifyOnWorldAddRemove() const { return bShouldNotifyOnWorldAddRemove; }
+	inline bool IsForceHidden() const {return bForceHidden;}
 	inline bool ShouldReceiveMobileCSMShadows() const { return bReceiveMobileCSMShadows; }
 	inline bool ShouldUpdateGPUSceneTransforms() const { return bShouldUpdateGPUSceneTransforms; }
 
@@ -825,6 +834,8 @@ protected:
 	TArray<FRayTracingGeometry*> RayTracingGeometries;
 #endif
 
+	void SetForceHidden(bool bForceHiddenIn) {bForceHidden = bForceHiddenIn;}
+
 private:
 #if !(UE_BUILD_SHIPPING || UE_BUILD_TEST)
 	FLinearColor WireframeColor;
@@ -898,11 +909,8 @@ private:
 	/** If true this primitive Renders in the mainPass */
 	uint8 bRenderInMainPass : 1;
 
-	/** If true this primitive will render only after owning level becomes visible */
-	uint8 bRequiresVisibleLevelToRender : 1;
-
-	/** Whether component level is currently visible */
-	uint8 bIsComponentLevelVisible : 1;
+	/** If true this primitive will always return that it is hidden if queried*/
+	uint8 bForceHidden : 1;
 	
 	/** Whether this component has any collision enabled */
 	uint8 bCollisionEnabled : 1;
@@ -1025,7 +1033,7 @@ protected:
 	uint8 bSupportsHeightfieldRepresentation : 1;
 
 	/** Whether this primitive requires notification when its level is added to the world and made visible for the first time. */
-	uint8 bNeedsLevelAddedToWorldNotification : 1;
+	uint8 bShouldNotifyOnWorldAddRemove : 1;
 
 	/** true by default, if set to false will make given proxy never drawn with selection outline */
 	uint8 bWantsSelectionOutline : 1;
diff --git a/Engine/Source/Runtime/Landscape/Private/LandscapeRender.cpp b/Engine/Source/Runtime/Landscape/Private/LandscapeRender.cpp
--- a/Engine/Source/Runtime/Landscape/Private/LandscapeRender.cpp
+++ b/Engine/Source/Runtime/Landscape/Private/LandscapeRender.cpp
@@ -1068,6 +1068,37 @@ public:
 
 } LandscapePersistentViewUniformBufferExtension;
 
+
+void FLandscapeVisibilityHelper::Init(UPrimitiveComponent* LandscapeComponent, FPrimitiveSceneProxy* ProxyIn)
+{
+	// Flag components to render only after level will be fully added to the world
+	ULevel* ComponentLevel = LandscapeComponent->GetComponentLevel();
+	bRequiresVisibleLevelToRender = (ComponentLevel && ComponentLevel->bRequireFullVisibilityToRender);
+	bIsComponentLevelVisible = (!ComponentLevel || ComponentLevel->bIsVisible);
+}
+
+bool FLandscapeVisibilityHelper::OnAddedToWorld()
+{
+	if (bIsComponentLevelVisible)
+	{
+		return false;
+	}
+	
+	bIsComponentLevelVisible = true;
+	return true;
+}
+
+bool FLandscapeVisibilityHelper::OnRemoveFromWorld()
+{
+	if (!bIsComponentLevelVisible)
+	{
+		return false;
+	}
+
+	bIsComponentLevelVisible = false;
+	return true;
+}
+
 FLandscapeComponentSceneProxy::FLandscapeComponentSceneProxy(ULandscapeComponent* InComponent)
 	: FPrimitiveSceneProxy(InComponent, NAME_LandscapeResourceNameForDebugging)
 	, FLandscapeNeighborInfo(InComponent->GetWorld(), InComponent->GetLandscapeProxy()->GetLandscapeGuid(), InComponent->GetSectionBase() / InComponent->ComponentSizeQuads, InComponent->GetHeightmap(), InComponent->ForcedLOD, InComponent->LODBias)
@@ -1112,6 +1143,19 @@ FLandscapeComponentSceneProxy::FLandscapeComponentSceneProxy(ULandscapeComponent
 	, LightMapResolution(InComponent->GetStaticLightMapResolution())
 #endif
 {
+
+	VisibilityHelper.Init(InComponent, this);
+
+	if (!VisibilityHelper.ShouldBeVisible())
+	{
+		SetForceHidden(true);
+	}
+
+	if (VisibilityHelper.RequiresVisibleLevelToRender())
+	{
+		bShouldNotifyOnWorldAddRemove = true;
+	}
+
 #if GPUCULL_TODO
 	// GPUCULL_TODO: Move to base proxy
 	bVFRequiresPrimitiveUniformBuffer = !UseGPUScene(GMaxRHIShaderPlatform, GetScene().GetFeatureLevel());
@@ -1138,11 +1182,6 @@ FLandscapeComponentSceneProxy::FLandscapeComponentSceneProxy(ULandscapeComponent
 	LODIndexToMaterialIndex = InComponent->LODIndexToMaterialIndex;
 	check(LODIndexToMaterialIndex.Num() == MaxLOD+1);
 
-	if (!IsComponentLevelVisible())
-	{
-		bNeedsLevelAddedToWorldNotification = true;
-	}
-
 	SetLevelColor(FLinearColor(1.f, 1.f, 1.f));
 
 	if (FeatureLevel <= ERHIFeatureLevel::ES3_1)
@@ -1346,6 +1385,7 @@ FLandscapeComponentSceneProxy::FLandscapeComponentSceneProxy(ULandscapeComponent
 	Instance.LocalBounds = Instance.RenderBounds;
 	bSupportsInstanceDataBuffer = true;
 #endif // GPUCULL_TODO
+
 }
 
 void FLandscapeComponentSceneProxy::CreateRenderThreadResources()
@@ -1354,7 +1394,7 @@ void FLandscapeComponentSceneProxy::CreateRenderThreadResources()
 
 	check(HeightmapTexture != nullptr);
 
-	if (IsComponentLevelVisible())
+	if (VisibilityHelper.ShouldBeVisible())
 	{
 		RegisterNeighbors(this);
 	}
@@ -1522,9 +1562,23 @@ void FLandscapeComponentSceneProxy::DestroyRenderThreadResources()
 	UnregisterNeighbors(this);
 }
 
-void FLandscapeComponentSceneProxy::OnLevelAddedToWorld()
+bool FLandscapeComponentSceneProxy::OnLevelAddedToWorld_RenderThread()
+{
+	if (VisibilityHelper.OnAddedToWorld())
+	{
+		SetForceHidden(false);
+		RegisterNeighbors(this);
+		return true;
+	}
+	return false;
+}
+
+void FLandscapeComponentSceneProxy::OnLevelRemovedFromWorld_RenderThread()
 {
-	RegisterNeighbors(this);
+	if (VisibilityHelper.OnRemoveFromWorld())
+	{
+		SetForceHidden(true);
+	}
 }
 
 FLandscapeComponentSceneProxy::~FLandscapeComponentSceneProxy()
@@ -4131,9 +4185,11 @@ void FLandscapeNeighborInfo::UnregisterNeighbors(FLandscapeComponentSceneProxy*
 FLandscapeMeshProxySceneProxy::FLandscapeMeshProxySceneProxy(UStaticMeshComponent* InComponent, const FGuid& InGuid, const TArray<FIntPoint>& InProxyComponentBases, int8 InProxyLOD)
 	: FStaticMeshSceneProxy(InComponent, false)
 {
-	if (!IsComponentLevelVisible())
+	VisibilityHelper.Init(InComponent, this);
+
+	if (VisibilityHelper.RequiresVisibleLevelToRender())
 	{
-		bNeedsLevelAddedToWorldNotification = true;
+		bShouldNotifyOnWorldAddRemove = true;
 	}
 
 	ProxyNeighborInfos.Empty(InProxyComponentBases.Num());
@@ -4154,7 +4210,7 @@ void FLandscapeMeshProxySceneProxy::CreateRenderThreadResources()
 {
 	FStaticMeshSceneProxy::CreateRenderThreadResources();
 
-	if (IsComponentLevelVisible())
+	if (VisibilityHelper.ShouldBeVisible())
 	{
 		for (FLandscapeNeighborInfo& Info : ProxyNeighborInfos)
 		{
@@ -4163,11 +4219,27 @@ void FLandscapeMeshProxySceneProxy::CreateRenderThreadResources()
 	}
 }
 
-void FLandscapeMeshProxySceneProxy::OnLevelAddedToWorld()
+bool FLandscapeMeshProxySceneProxy::OnLevelAddedToWorld_RenderThread()
 {
-	for (FLandscapeNeighborInfo& Info : ProxyNeighborInfos)
+	if (VisibilityHelper.OnAddedToWorld())
+	{
+		SetForceHidden(false);
+
+		for (FLandscapeNeighborInfo& Info : ProxyNeighborInfos)
+		{
+			Info.RegisterNeighbors();
+		}
+		return true;
+	}
+
+	return false;
+}
+
+void FLandscapeMeshProxySceneProxy::OnLevelRemovedFromWorld_RenderThread()
+{
+	if (VisibilityHelper.OnRemoveFromWorld())
 	{
-		Info.RegisterNeighbors();
+		SetForceHidden(true);
 	}
 }
 
diff --git a/Engine/Source/Runtime/Landscape/Private/LandscapeRenderMobile.cpp b/Engine/Source/Runtime/Landscape/Private/LandscapeRenderMobile.cpp
--- a/Engine/Source/Runtime/Landscape/Private/LandscapeRenderMobile.cpp
+++ b/Engine/Source/Runtime/Landscape/Private/LandscapeRenderMobile.cpp
@@ -374,7 +374,7 @@ void FLandscapeComponentSceneProxyMobile::CreateRenderThreadResources()
 {
 	LLM_SCOPE(ELLMTag::Landscape);
 
-	if (IsComponentLevelVisible())
+	if (VisibilityHelper.ShouldBeVisible())
 	{
 		RegisterNeighbors(this);
 	}
diff --git a/Engine/Source/Runtime/Landscape/Public/LandscapeRender.h b/Engine/Source/Runtime/Landscape/Public/LandscapeRender.h
--- a/Engine/Source/Runtime/Landscape/Public/LandscapeRender.h
+++ b/Engine/Source/Runtime/Landscape/Public/LandscapeRender.h
@@ -678,6 +678,20 @@ struct FLandscapeRenderSystem
 
 extern TMap<FLandscapeNeighborInfo::FLandscapeKey, FLandscapeRenderSystem*> LandscapeRenderSystems;
 
+class FLandscapeVisibilityHelper
+{
+public:
+	void Init(UPrimitiveComponent* LandscapeComponent, FPrimitiveSceneProxy* ProxyIn);
+	bool OnAddedToWorld();
+	bool OnRemoveFromWorld();
+	bool ShouldBeVisible() const { return !bRequiresVisibleLevelToRender || bIsComponentLevelVisible; }
+	bool RequiresVisibleLevelToRender() const { return bRequiresVisibleLevelToRender; }
+private:
+	bool bRequiresVisibleLevelToRender = false;
+	bool bIsComponentLevelVisible = false;
+	FPrimitiveSceneProxy* Proxy = nullptr;
+};
+
 //
 // FLandscapeMeshProxySceneProxy
 //
@@ -690,9 +704,12 @@ public:
 	FLandscapeMeshProxySceneProxy(UStaticMeshComponent* InComponent, const FGuid& InGuid, const TArray<FIntPoint>& InProxyComponentBases, int8 InProxyLOD);
 	virtual void CreateRenderThreadResources() override;
 	virtual void DestroyRenderThreadResources() override;
-	virtual void OnLevelAddedToWorld() override;
-};
+	virtual bool OnLevelAddedToWorld_RenderThread() override;
+	virtual void OnLevelRemovedFromWorld_RenderThread() override;
 
+private:
+	FLandscapeVisibilityHelper VisibilityHelper;
+};
 
 //
 // FLandscapeComponentSceneProxy
@@ -888,6 +905,8 @@ protected:
 	TArray<FPrimitiveInstance> Instances;
 #endif // GPUCULL_TODO
 
+	FLandscapeVisibilityHelper VisibilityHelper;
+
 protected:
 	virtual ~FLandscapeComponentSceneProxy();
 	
@@ -914,7 +933,8 @@ public:
 	virtual void OnTransformChanged() override;
 	virtual void CreateRenderThreadResources() override;
 	virtual void DestroyRenderThreadResources() override;
-	virtual void OnLevelAddedToWorld() override;
+	virtual bool OnLevelAddedToWorld_RenderThread() override;
+	virtual void OnLevelRemovedFromWorld_RenderThread() override;
 	
 	friend class ULandscapeComponent;
 	friend class FLandscapeVertexFactoryVertexShaderParameters;
diff --git a/Engine/Source/Runtime/Renderer/Private/PrimitiveSceneInfo.cpp b/Engine/Source/Runtime/Renderer/Private/PrimitiveSceneInfo.cpp
--- a/Engine/Source/Runtime/Renderer/Private/PrimitiveSceneInfo.cpp
+++ b/Engine/Source/Runtime/Renderer/Private/PrimitiveSceneInfo.cpp
@@ -116,11 +116,11 @@ FPrimitiveFlagsCompact::FPrimitiveFlagsCompact(const FPrimitiveSceneProxy* Proxy
 {}
 
 FPrimitiveSceneInfoCompact::FPrimitiveSceneInfoCompact(FPrimitiveSceneInfo* InPrimitiveSceneInfo) :
+	Bounds(PrimitiveSceneInfo->Proxy->GetBounds()),
 	PrimitiveFlagsCompact(InPrimitiveSceneInfo->Proxy)
 {
 	PrimitiveSceneInfo = InPrimitiveSceneInfo;
 	Proxy = PrimitiveSceneInfo->Proxy;
-	Bounds = PrimitiveSceneInfo->Proxy->GetBounds();
 	MinDrawDistance = PrimitiveSceneInfo->Proxy->GetMinDrawDistance();
 	MaxDrawDistance = PrimitiveSceneInfo->Proxy->GetMaxDrawDistance();
 
@@ -1198,6 +1198,16 @@ void FPrimitiveSceneInfo::AddToScene(FRHICommandListImmediate& RHICmdList, FScen
 		FPrimitiveSceneProxy* Proxy = SceneInfo->Proxy;
 		INC_MEMORY_STAT_BY(STAT_PrimitiveInfoMemory, sizeof(*SceneInfo) + SceneInfo->StaticMeshes.GetAllocatedSize() + SceneInfo->StaticMeshRelevances.GetAllocatedSize() + Proxy->GetMemoryFootprint());
 	}
+
+	// Some primitive types cannot add their meshes until the level is added to the world.
+	for (FPrimitiveSceneInfo* SceneInfo : SceneInfos)
+	{
+		if (SceneInfo->Proxy->ShouldNotifyOnWorldAddRemove())
+		{
+			TArray<FPrimitiveSceneInfo*>& LevelNotifyPrimitives = Scene->PrimitivesNeedingLevelUpdateNotification.FindOrAdd(SceneInfo->Proxy->GetLevelName());
+			LevelNotifyPrimitives.Add(SceneInfo);
+		}
+	}	
 }
 
 void FPrimitiveSceneInfo::RemoveStaticMeshes()
@@ -1292,6 +1302,19 @@ void FPrimitiveSceneInfo::RemoveFromScene(bool bUpdateStaticDrawLists)
 		FVirtualTextureSystem::Get().RemoveAllProducerDestroyedCallbacks(this);
 		bRegisteredVirtualTextureProducerCallback = false;
 	}
+
+	if (Proxy->ShouldNotifyOnWorldAddRemove())
+	{
+		TArray<FPrimitiveSceneInfo*>* LevelNotifyPrimitives = Scene->PrimitivesNeedingLevelUpdateNotification.Find(Proxy->GetLevelName());
+		if (LevelNotifyPrimitives != nullptr)
+		{
+			LevelNotifyPrimitives->Remove(this);
+			if (LevelNotifyPrimitives->Num() == 0)
+			{
+				Scene->PrimitivesNeedingLevelUpdateNotification.Remove(Proxy->GetLevelName());
+			}
+		}
+	}
 }
 
 void FPrimitiveSceneInfo::UpdateRuntimeVirtualTextureFlags()
diff --git a/Engine/Source/Runtime/Renderer/Private/RendererScene.cpp b/Engine/Source/Runtime/Renderer/Private/RendererScene.cpp
--- a/Engine/Source/Runtime/Renderer/Private/RendererScene.cpp
+++ b/Engine/Source/Runtime/Renderer/Private/RendererScene.cpp
@@ -3730,18 +3730,16 @@ void FScene::OnLevelAddedToWorld_RenderThread(FName InLevelName)
 	// Mark level primitives
 	TArray<FPrimitiveSceneInfo*> PrimitivesToAdd;
 
-	for (auto It = Primitives.CreateIterator(); It; ++It)
+	if (const TArray<FPrimitiveSceneInfo*>* LevelPrimitives = PrimitivesNeedingLevelUpdateNotification.Find(InLevelName))
 	{
-		FPrimitiveSceneProxy* Proxy = (*It)->Proxy;
-		if (Proxy->LevelName == InLevelName)
+		for (FPrimitiveSceneInfo* Primitive : *LevelPrimitives)
 		{
-			Proxy->bIsComponentLevelVisible = true;
-			if (Proxy->NeedsLevelAddedToWorldNotification())
+			// If the primitive proxy returns true, it needs it's static meshes added to the scene
+			if (Primitive->Proxy->OnLevelAddedToWorld_RenderThread())
 			{
-				// The only type of SceneProxy using this is landscape
-				(*It)->RemoveStaticMeshes();
-				Proxy->OnLevelAddedToWorld();
-				PrimitivesToAdd.Add(*It);
+				// Remove static meshes and cached commands for any primitives that need to be added
+				Primitive->RemoveStaticMeshes();
+				PrimitivesToAdd.Add(Primitive);
 			}
 		}
 	}
@@ -3770,12 +3768,11 @@ void FScene::OnLevelRemovedFromWorld_RenderThread(FName InLevelName)
 {
 	TRACE_CPUPROFILER_EVENT_SCOPE(FScene::OnLevelRemovedFromWorld_RenderThread);
 
-	for (auto It = Primitives.CreateIterator(); It; ++It)
+	if (TArray<FPrimitiveSceneInfo*>* LevelPrimitives = PrimitivesNeedingLevelUpdateNotification.Find(InLevelName))
 	{
-		FPrimitiveSceneProxy* Proxy = (*It)->Proxy;
-		if (Proxy->LevelName == InLevelName)
+		for (FPrimitiveSceneInfo* Primitive : *LevelPrimitives)
 		{
-			Proxy->bIsComponentLevelVisible = false;
+			Primitive->Proxy->OnLevelRemovedFromWorld_RenderThread();
 		}
 	}
 }
diff --git a/Engine/Source/Runtime/Renderer/Private/ScenePrivate.h b/Engine/Source/Runtime/Renderer/Private/ScenePrivate.h
--- a/Engine/Source/Runtime/Renderer/Private/ScenePrivate.h
+++ b/Engine/Source/Runtime/Renderer/Private/ScenePrivate.h
@@ -2618,6 +2618,8 @@ public:
 	TArray<ERayTracingPrimitiveFlags> PrimitiveRayTracingFlags;
 #endif
 
+	TMap<FName, TArray<FPrimitiveSceneInfo*>> PrimitivesNeedingLevelUpdateNotification;
+
 #if WITH_EDITOR
 	/** Packed bit array of primitives which are selected in the editor. */
 	TBitArray<> PrimitivesSelected;
