* Fixed bad code merge leading to checking for landscape resources readiness twice
* Removed unused temporary command line option "nolandscapelayerupdate"

#rnx

diff --git a/Engine/Source/Runtime/Landscape/Private/LandscapeEditLayers.cpp b/Engine/Source/Runtime/Landscape/Private/LandscapeEditLayers.cpp
--- a/Engine/Source/Runtime/Landscape/Private/LandscapeEditLayers.cpp
+++ b/Engine/Source/Runtime/Landscape/Private/LandscapeEditLayers.cpp
@@ -4973,43 +4973,12 @@ void ALandscape::UpdateLayersContent(bool bInWaitForStreaming, bool bInSkipMonit
 		}
 	};
 
-	// Remove this command line switch after fixes for D3D12 RHI
-	if (FParse::Param(FCommandLine::Get(), TEXT("nolandscapelayerupdate")))
-	{
-		return;
-	}
-
 	ULandscapeInfo* LandscapeInfo = GetLandscapeInfo();
 	if (LandscapeInfo == nullptr || !CanHaveLayersContent() || !LandscapeInfo->AreAllComponentsRegistered())
 	{
 		return;
 	}
 
-	// We need to wait until texture resources are ready 
-	// to initialize the landscape to avoid taking the
-	// sizes and format of the default texture
-	if (!AreLayersResourcesReady(bInWaitForStreaming))
-	{
-		if (!bInWaitForStreaming)
-		{
-			static constexpr double TimeBeforeDisplayingNotification = 3.0;
-			if (WaitingForResourcesStartTime < 0.0)
-			{
-				WaitingForResourcesStartTime = FSlateApplicationBase::Get().GetCurrentTime();
-			}
-
-			if ((FSlateApplicationBase::Get().GetCurrentTime() - WaitingForResourcesStartTime) > TimeBeforeDisplayingNotification)
-			{
-				// let the user know we are waiting for resources :
-				static const FText NotificationText(LOCTEXT("WaitForLayersResources", "Landscape editor waiting for edit layer resources to be ready."));
-				ShowEditLayersResourcesNotification(NotificationText, EditLayersResourcesNotification);
-			}
-			// The notification may not be visible yet (because of the initial delay) but it should not be hidden and the initial delay timer shouldn't be reset : 
-			bHideEditLayerResourcesNotification = false;
-		}
-		return;
-	}
-
 	if (!bLandscapeLayersAreInitialized)
 	{
 		InitializeLayers();
@@ -5047,6 +5016,7 @@ void ALandscape::UpdateLayersContent(bool bInWaitForStreaming, bool bInSkipMonit
 		return;
 	}
 
+	// We need to wait until texture resources are ready to initialize the landscape to avoid taking the sizes and format of the default texture:
 	if (!AreLayersResourcesReady(bInWaitForStreaming))
 	{
 		if (!bInWaitForStreaming)
@@ -5060,7 +5030,7 @@ void ALandscape::UpdateLayersContent(bool bInWaitForStreaming, bool bInSkipMonit
 			if ((FSlateApplicationBase::Get().GetCurrentTime() - WaitingForResourcesStartTime) > TimeBeforeDisplayingNotification)
 			{
 				// let the user know we are waiting for resources :
-				static const FText NotificationText(LOCTEXT("WaitForLayersResources", "Landscape edition waiting for edit layers resources to be ready."));
+				static const FText NotificationText(LOCTEXT("WaitForLayersResources", "Landscape editor waiting for edit layers resources to be ready."));
 				ShowEditLayersResourcesNotification(NotificationText, EditLayersResourcesNotification);
 			}
 			// The notification may not be visible yet (because of the initial delay) but it should not be hidden and the initial delay timer shouldn't be reset : 
