Prepare for the deprecation of direct access to the Resource member of UTexture in favor of the GetResource() accessor.

#rnx

diff --git a/Engine/Source/Editor/Cascade/Private/CascadeEmitterCanvasClient.cpp b/Engine/Source/Editor/Cascade/Private/CascadeEmitterCanvasClient.cpp
--- a/Engine/Source/Editor/Cascade/Private/CascadeEmitterCanvasClient.cpp
+++ b/Engine/Source/Editor/Cascade/Private/CascadeEmitterCanvasClient.cpp
@@ -1060,7 +1060,7 @@ void FCascadeEmitterCanvasClient::DrawHeaderBlock(int32 Index, int32 XPos, UPart
 	}
 	else
 	{
-		Canvas->DrawTile(XPos, 0, EmitterWidth, EmitterHeadHeight, 0.f, 0.f, 1.f, 1.f, HeadColor, TexModuleDisabledBackground->Resource);
+		Canvas->DrawTile(XPos, 0, EmitterWidth, EmitterHeadHeight, 0.f, 0.f, 1.f, 1.f, HeadColor, TexModuleDisabledBackground->GetResource());
 	}
 
 	UParticleSpriteEmitter* SpriteEmitter = Cast<UParticleSpriteEmitter>(Emitter);
@@ -1281,7 +1281,7 @@ void FCascadeEmitterCanvasClient::DrawCollapsedHeaderBlock(int32 Index, int32 XP
 	}
 	else
 	{
-		Canvas->DrawTile(XPos, 0, EmitterCollapsedWidth, EmitterHeadHeight, 0.f, 0.f, 1.f, 1.f, HeadColor, TexModuleDisabledBackground->Resource);
+		Canvas->DrawTile(XPos, 0, EmitterCollapsedWidth, EmitterHeadHeight, 0.f, 0.f, 1.f, 1.f, HeadColor, TexModuleDisabledBackground->GetResource());
 	}
 
 	// Draw column background
@@ -1446,7 +1446,7 @@ void FCascadeEmitterCanvasClient::DrawModule(FCanvas* Canvas, UParticleModule* M
 	else
 	{
 		FColor BkgColor = ModuleBkgColor;
-		FTexture* BkgTexture = TexModuleDisabledBackground->Resource;
+		FTexture* BkgTexture = TexModuleDisabledBackground->GetResource();
 		if (!bIsModuleValid)
 		{
 			BkgColor.R = 255;
@@ -1769,7 +1769,7 @@ FTexture* FCascadeEmitterCanvasClient::GetIconTexture(ECascadeIcons eIcon)
 		UTexture2D* IconTexture = IconTex[eIcon];
 		if (IconTexture)
 		{
-			return IconTexture->Resource;
+			return IconTexture->GetResource();
 		}
 	}
 
diff --git a/Engine/Source/Editor/FontEditor/Private/SFontEditorViewport.cpp b/Engine/Source/Editor/FontEditor/Private/SFontEditorViewport.cpp
--- a/Engine/Source/Editor/FontEditor/Private/SFontEditorViewport.cpp
+++ b/Engine/Source/Editor/FontEditor/Private/SFontEditorViewport.cpp
@@ -198,7 +198,7 @@ void FFontEditorViewportClient::Draw(FViewport* Viewport, FCanvas* Canvas)
 							}
 
 							// Draw the font texture (with alpha blending enabled)
-							Canvas->DrawTile( XPos, YPos, Width, Height, 0.0f, 0.0f, 1.0f, 1.0f, FLinearColor::White, Texture->Resource, true);
+							Canvas->DrawTile( XPos, YPos, Width, Height, 0.0f, 0.0f, 1.0f, 1.0f, FLinearColor::White, Texture->GetResource(), true);
 						}
 					}
 					// Update our total height and current draw position
diff --git a/Engine/Source/Editor/LandscapeEditor/Private/LandscapeEdModeRampTool.cpp b/Engine/Source/Editor/LandscapeEditor/Private/LandscapeEdModeRampTool.cpp
--- a/Engine/Source/Editor/LandscapeEditor/Private/LandscapeEdModeRampTool.cpp
+++ b/Engine/Source/Editor/LandscapeEditor/Private/LandscapeEdModeRampTool.cpp
@@ -353,11 +353,11 @@ public:
 				PDI->DrawSprite(WorldPoints[i],
 					SpriteScale,
 					SpriteScale,
-					SpriteTexture->Resource,
+					SpriteTexture->GetResource(),
 					SpriteColor,
 					SDPG_Foreground,
-					0, SpriteTexture->Resource->GetSizeX(),
-					0, SpriteTexture->Resource->GetSizeY(),
+					0, SpriteTexture->GetResource()->GetSizeX(),
+					0, SpriteTexture->GetResource()->GetSizeY(),
 					SE_BLEND_Masked);
 			}
 			PDI->SetHitProxy(NULL);
diff --git a/Engine/Source/Editor/LandscapeEditor/Private/LandscapeEditorDetailCustomization_AlphaBrush.cpp b/Engine/Source/Editor/LandscapeEditor/Private/LandscapeEditorDetailCustomization_AlphaBrush.cpp
--- a/Engine/Source/Editor/LandscapeEditor/Private/LandscapeEditorDetailCustomization_AlphaBrush.cpp
+++ b/Engine/Source/Editor/LandscapeEditor/Private/LandscapeEditorDetailCustomization_AlphaBrush.cpp
@@ -165,7 +165,7 @@ void FTextureMaskThumbnailViewportClient::Draw(FViewport* Viewport, FCanvas* Can
 	Texture->WaitForStreaming();
 
 	//Draw the selected texture, uses ColourChannelBlend mode parameter to filter colour channels and apply grayscale
-	FCanvasTileItem TileItem( FVector2D( 0.0f, 0.0f ), Texture->Resource, Viewport->GetSizeXY(), FLinearColor::White );
+	FCanvasTileItem TileItem( FVector2D( 0.0f, 0.0f ), Texture->GetResource(), Viewport->GetSizeXY(), FLinearColor::White );
 	TileItem.BlendMode = (ESimpleElementBlendMode)(SE_BLEND_RGBA_MASK_START + (1<<PinnedParent->TextureChannel.Get()) + 16);
 	Canvas->DrawItem( TileItem );
 }
diff --git a/Engine/Source/Editor/MaterialEditor/Private/SMaterialEditorViewport.cpp b/Engine/Source/Editor/MaterialEditor/Private/SMaterialEditorViewport.cpp
--- a/Engine/Source/Editor/MaterialEditor/Private/SMaterialEditorViewport.cpp
+++ b/Engine/Source/Editor/MaterialEditor/Private/SMaterialEditorViewport.cpp
@@ -1090,7 +1090,7 @@ void SMaterialEditorUIPreviewZoomer::DestroyCheckerboardTexture()
 {
 	if (CheckerboardTexture)
 	{
-		if (CheckerboardTexture->Resource)
+		if (CheckerboardTexture->GetResource())
 		{
 			CheckerboardTexture->ReleaseResource();
 		}
diff --git a/Engine/Source/Editor/Matinee/Private/MatineeDraw.cpp b/Engine/Source/Editor/Matinee/Private/MatineeDraw.cpp
--- a/Engine/Source/Editor/Matinee/Private/MatineeDraw.cpp
+++ b/Engine/Source/Editor/Matinee/Private/MatineeDraw.cpp
@@ -1487,7 +1487,7 @@ void FMatineeViewportClient::CreatePushPropertiesOntoGraphButton( FCanvas* Canva
 
 		// Draw button for pushing properties onto graph view.
 		if(Canvas->IsHitTesting()) Canvas->SetHitProxy( new HMatineeTrackGraphPropBtn( InGroup, GroupIndex, Track ) );
-		Canvas->DrawTile( -14, TrackHeightToUse-11, 8, 8, 0.f, 0.f, 1.f, 1.f, FLinearColor::White, GraphTex->Resource );
+		Canvas->DrawTile( -14, TrackHeightToUse-11, 8, 8, 0.f, 0.f, 1.f, 1.f, FLinearColor::White, GraphTex->GetResource());
 		if(Canvas->IsHitTesting()) Canvas->SetHitProxy( NULL );
 	}
 }
@@ -1580,7 +1580,7 @@ void FMatineeViewportClient::DrawTrackLabel( FCanvas* Canvas, UInterpTrack* Trac
 	if( TrackIconTex )
 	{
 		TileItem.SetColor( FLinearColor::White );
-		TileItem.Texture = TrackIconTex->Resource;
+		TileItem.Texture = TrackIconTex->GetResource();
 		TileItem.Size = FVector2D( TrackIconSize, TrackIconSize );
 		TileItem.Draw( Canvas, FVector2D( -InterpEd->LabelWidth + TrackTitleIndentPixels - PaddedTrackIconSize, 0.5*(TrackHeightToUse - TrackIconSize) ) );		
 		TileItem.Texture = GWhiteTexture;
@@ -1621,11 +1621,11 @@ void FMatineeViewportClient::DrawTrackLabel( FCanvas* Canvas, UInterpTrack* Trac
 		UTexture2D* BackwardTex = (EventTrack->bFireEventsWhenBackwards) ? LabelDrawParams.BackwardEventOnTex : LabelDrawParams.BackwardEventOffTex;
 
 		if(Canvas->IsHitTesting()) Canvas->SetHitProxy( new HMatineeEventDirBtn(Group, TrackIndex, EMatineeEventDirection::IED_Backward) );
-		Canvas->DrawTile( -24, TrackHeightToUse-11, 8, 8, 0.f, 0.f, 1.f, 1.f, FLinearColor::White, BackwardTex->Resource );
+		Canvas->DrawTile( -24, TrackHeightToUse-11, 8, 8, 0.f, 0.f, 1.f, 1.f, FLinearColor::White, BackwardTex->GetResource());
 		if(Canvas->IsHitTesting()) Canvas->SetHitProxy( NULL );
 
 		if(Canvas->IsHitTesting()) Canvas->SetHitProxy( new HMatineeEventDirBtn(Group, TrackIndex, EMatineeEventDirection::IED_Forward) );
-		Canvas->DrawTile( -14, TrackHeightToUse-11, 8, 8, 0.f, 0.f, 1.f, 1.f, FLinearColor::White, ForwardTex->Resource );
+		Canvas->DrawTile( -14, TrackHeightToUse-11, 8, 8, 0.f, 0.f, 1.f, 1.f, FLinearColor::White, ForwardTex->GetResource());
 		if(Canvas->IsHitTesting()) Canvas->SetHitProxy( NULL );
 	}
 
@@ -1638,7 +1638,7 @@ void FMatineeViewportClient::DrawTrackLabel( FCanvas* Canvas, UInterpTrack* Trac
 		if(Canvas->IsHitTesting()) Canvas->SetHitProxy( new HMatineeTrackTrajectoryButton( Group, Track ) );
 		if(MovementTrack->bHide3DTrack)
 		{
-			Canvas->DrawTile( -24, TrackHeightToUse-11, 8, 8, 0.f, 0.f, 1.f, 1.f, FLinearColor::White, TrajectoryButtonTex->Resource );
+			Canvas->DrawTile( -24, TrackHeightToUse-11, 8, 8, 0.f, 0.f, 1.f, 1.f, FLinearColor::White, TrajectoryButtonTex->GetResource());
 		}
 		else
 		{
@@ -1665,7 +1665,7 @@ void FMatineeViewportClient::DrawTrackLabel( FCanvas* Canvas, UInterpTrack* Trac
 
 		if( Track->IsDisabled() == false )
 		{
-			Canvas->DrawTile( -InterpEd->LabelWidth + MatineeGlobals::DisableTrackCheckBoxHorizOffset, YPos, MatineeGlobals::DisableTrackIconSize.X, MatineeGlobals::DisableTrackIconSize.Y, 0,0,1,1, FLinearColor::White, LabelDrawParams.DisableTrackTex->Resource);
+			Canvas->DrawTile( -InterpEd->LabelWidth + MatineeGlobals::DisableTrackCheckBoxHorizOffset, YPos, MatineeGlobals::DisableTrackIconSize.X, MatineeGlobals::DisableTrackIconSize.Y, 0,0,1,1, FLinearColor::White, LabelDrawParams.DisableTrackTex->GetResource());
 		}
 
 		if( Canvas->IsHitTesting() )
@@ -2070,7 +2070,7 @@ void FMatineeViewportClient::Draw(FViewport* InViewport, FCanvas* Canvas)
 					UTexture2D* ButtonTex = (Group == InterpEd->CamViewGroup) ? LabelDrawParams.CamLockedIcon : LabelDrawParams.CamUnlockedIcon;
 					if(Canvas->IsHitTesting()) Canvas->SetHitProxy( new HMatineeGroupLockCamBtn(Group) );
 					TileItem.SetColor( FLinearColor::White );
-					TileItem.Texture = ButtonTex->Resource;
+					TileItem.Texture = ButtonTex->GetResource();
 					TileItem.Size = FVector2D( 16, 16 );
 					TileItem.Draw( Canvas, FVector2D( InterpEd->LabelWidth - 26, (0.5*GroupHeadHeight)-8 ) );
 					TileItem.Texture = GWhiteTexture;
@@ -2079,7 +2079,7 @@ void FMatineeViewportClient::Draw(FViewport* InViewport, FCanvas* Canvas)
 				if( !Group->bIsFolder )
 				{
 					TileItem.SetColor( Group->GroupColor );
-					TileItem.Texture = InterpEd->BarGradText->Resource;
+					TileItem.Texture = InterpEd->BarGradText->GetResource();
 					TileItem.Size = FVector2D( 6, GroupHeadHeight );
 					TileItem.Draw( Canvas, FVector2D( InterpEd->LabelWidth - 6, 0 ) );
 					TileItem.Texture = GWhiteTexture;
diff --git a/Engine/Source/Editor/MeshPaint/Private/TexturePaintHelpers.cpp b/Engine/Source/Editor/MeshPaint/Private/TexturePaintHelpers.cpp
--- a/Engine/Source/Editor/MeshPaint/Private/TexturePaintHelpers.cpp
+++ b/Engine/Source/Editor/MeshPaint/Private/TexturePaintHelpers.cpp
@@ -46,7 +46,7 @@ void TexturePaintHelpers::CopyTextureToRenderTargetTexture(UTexture* SourceTextu
 	UTexture2D* Texture2D = Cast<UTexture2D>(SourceTexture);
 	if (Texture2D != nullptr)
 	{
-		TextureResource = Texture2D->Resource;
+		TextureResource = Texture2D->GetResource();
 	}
 	else
 	{
diff --git a/Engine/Source/Editor/NewLevelDialog/Private/NewLevelDialogModule.cpp b/Engine/Source/Editor/NewLevelDialog/Private/NewLevelDialogModule.cpp
--- a/Engine/Source/Editor/NewLevelDialog/Private/NewLevelDialogModule.cpp
+++ b/Engine/Source/Editor/NewLevelDialog/Private/NewLevelDialogModule.cpp
@@ -49,7 +49,7 @@ public:
 		ENQUEUE_RENDER_COMMAND(UpdateSTexture2DView)(
 			[TextureView, InTexture2D](FRHICommandListImmediate& RHICmdList)
 			{
-				TextureView->ShaderResource = InTexture2D->Resource->GetTexture2DRHI();
+				TextureView->ShaderResource = InTexture2D->GetResource()->GetTexture2DRHI();
 			});
 	}
 
diff --git a/Engine/Source/Editor/Sequencer/Private/SequencerEdMode.cpp b/Engine/Source/Editor/Sequencer/Private/SequencerEdMode.cpp
--- a/Engine/Source/Editor/Sequencer/Private/SequencerEdMode.cpp
+++ b/Engine/Source/Editor/Sequencer/Private/SequencerEdMode.cpp
@@ -685,7 +685,7 @@ void FSequencerEdMode::DrawAudioTracks(FPrimitiveDrawInterface* PDI)
 								float Scale = PDI->View->WorldToScreen(Location).W * (4.0f / PDI->View->UnscaledViewRect.Width() / PDI->View->ViewMatrices.GetProjectionMatrix().M[0][0]);
 								Scale *= bIsActive ? 15.f : 10.f;
 
-								PDI->DrawSprite(Location, Scale, Scale, AudioTexture->Resource, Color, SDPG_Foreground, 0.0, 0.0, 0.0, 0.0, SE_BLEND_Masked);
+								PDI->DrawSprite(Location, Scale, Scale, AudioTexture->GetResource(), Color, SDPG_Foreground, 0.0, 0.0, 0.0, 0.0, SE_BLEND_Masked);
 								break;
 							}
 						}
diff --git a/Engine/Source/Editor/StatsViewer/Private/StatsPages/TextureStatsPage.cpp b/Engine/Source/Editor/StatsViewer/Private/StatsPages/TextureStatsPage.cpp
--- a/Engine/Source/Editor/StatsViewer/Private/StatsPages/TextureStatsPage.cpp
+++ b/Engine/Source/Editor/StatsViewer/Private/StatsPages/TextureStatsPage.cpp
@@ -217,7 +217,7 @@ struct TextureStatsGenerator : public FFindReferencedAssets
 
 			Entry->LODBias = InTexture->GetCachedLODBias();
 
-			const FTexture* Resource = InTexture->Resource; 
+			const FTexture* Resource = InTexture->GetResource();
 			if(Resource)
 			{
 				Entry->LastTimeRendered = (float)FMath::Max( FApp::GetLastTime() - Resource->LastRenderTime, 0.0 );
diff --git a/Engine/Source/Editor/TextureEditor/Private/Models/TextureEditorViewportClient.cpp b/Engine/Source/Editor/TextureEditor/Private/Models/TextureEditorViewportClient.cpp
--- a/Engine/Source/Editor/TextureEditor/Private/Models/TextureEditorViewportClient.cpp
+++ b/Engine/Source/Editor/TextureEditor/Private/Models/TextureEditorViewportClient.cpp
@@ -167,19 +167,19 @@ void FTextureEditorViewportClient::Draw(FViewport* Viewport, FCanvas* Canvas)
 		const int32 CheckerboardSizeY = FMath::Max<int32>(1, CheckerboardTexture->GetSizeY());
 		if (Settings.Background == TextureEditorBackground_CheckeredFill)
 		{
-			Canvas->DrawTile( 0.0f, 0.0f, Viewport->GetSizeXY().X, Viewport->GetSizeXY().Y, 0.0f, 0.0f, (float)Viewport->GetSizeXY().X / CheckerboardSizeX, (float)Viewport->GetSizeXY().Y / CheckerboardSizeY, FLinearColor::White, CheckerboardTexture->Resource);
+			Canvas->DrawTile( 0.0f, 0.0f, Viewport->GetSizeXY().X, Viewport->GetSizeXY().Y, 0.0f, 0.0f, (float)Viewport->GetSizeXY().X / CheckerboardSizeX, (float)Viewport->GetSizeXY().Y / CheckerboardSizeY, FLinearColor::White, CheckerboardTexture->GetResource());
 		}
 		else if (Settings.Background == TextureEditorBackground_Checkered)
 		{
-			Canvas->DrawTile( XPos, YPos, Width, Height, 0.0f, 0.0f, (float)Width / CheckerboardSizeX, (float)Height / CheckerboardSizeY, FLinearColor::White, CheckerboardTexture->Resource);
+			Canvas->DrawTile( XPos, YPos, Width, Height, 0.0f, 0.0f, (float)Width / CheckerboardSizeX, (float)Height / CheckerboardSizeY, FLinearColor::White, CheckerboardTexture->GetResource());
 		}
 	}
 
 	float Exposure = FMath::Pow(2.0f, (float)TextureEditorPinned->GetExposureBias());
 
-	if ( Texture->Resource != nullptr )
+	if ( Texture->GetResource() != nullptr )
 	{
-		FCanvasTileItem TileItem( FVector2D( XPos, YPos ), Texture->Resource, FVector2D( Width, Height ), FLinearColor(Exposure, Exposure, Exposure) );
+		FCanvasTileItem TileItem( FVector2D( XPos, YPos ), Texture->GetResource(), FVector2D( Width, Height ), FLinearColor(Exposure, Exposure, Exposure) );
 		TileItem.BlendMode = TextureEditorPinned->GetColourChannelBlendMode();
 		TileItem.BatchedElementParameters = BatchedElementParameters;
 
@@ -204,7 +204,7 @@ void FTextureEditorViewportClient::Draw(FViewport* Viewport, FCanvas* Canvas)
 		// if we are presenting a virtual texture, make the appropriate tiles resident
 		if (bIsVirtualTexture && CVarEnableVTFeedback.GetValueOnAnyThread() != 0)
 		{
-			FVirtualTexture2DResource* VTResource = static_cast<FVirtualTexture2DResource*>(Texture->Resource);
+			FVirtualTexture2DResource* VTResource = static_cast<FVirtualTexture2DResource*>(Texture->GetResource());
 			const FVector2D ScreenSpaceSize((float)Width, (float)Height);
 			const FVector2D ViewportPositon(-(float)XPos, -(float)YPos);
 			const FVector2D UV0 = TileItem.UV0;
@@ -231,7 +231,7 @@ void FTextureEditorViewportClient::Draw(FViewport* Viewport, FCanvas* Canvas)
 		const FIntPoint SizeOnMip = { Texture2D->GetSizeX() >> Mip,Texture2D->GetSizeY() >> Mip };
 		const uint64 NumPixels = static_cast<uint64>(SizeOnMip.X) * SizeOnMip.Y;
 
-		const FVirtualTexture2DResource* Resource = (FVirtualTexture2DResource*)Texture2D->Resource;
+		const FVirtualTexture2DResource* Resource = (FVirtualTexture2DResource*)Texture2D->GetResource();
 		const FIntPoint PhysicalTextureSize = Resource->GetPhysicalTextureSize(0u);
 		const uint64 NumPhysicalPixels = static_cast<uint64>(PhysicalTextureSize.X) * PhysicalTextureSize.Y;
 
@@ -509,7 +509,7 @@ void FTextureEditorViewportClient::DestroyCheckerboardTexture()
 {
 	if (CheckerboardTexture)
 	{
-		if (CheckerboardTexture->Resource)
+		if (CheckerboardTexture->GetResource())
 		{
 			CheckerboardTexture->ReleaseResource();
 		}
diff --git a/Engine/Source/Editor/TextureEditor/Private/TextureEditorToolkit.cpp b/Engine/Source/Editor/TextureEditor/Private/TextureEditorToolkit.cpp
--- a/Engine/Source/Editor/TextureEditor/Private/TextureEditorToolkit.cpp
+++ b/Engine/Source/Editor/TextureEditor/Private/TextureEditorToolkit.cpp
@@ -106,7 +106,7 @@ FTextureEditorToolkit::~FTextureEditorToolkit( )
 	UTexture2D* Texture2D = Cast<UTexture2D>(Texture);
 	if (Texture2D && Texture2D->IsCurrentlyVirtualTextured())
 	{
-		FVirtualTexture2DResource* Resource = (FVirtualTexture2DResource*)Texture2D->Resource;
+		FVirtualTexture2DResource* Resource = (FVirtualTexture2DResource*)Texture2D->GetResource();
 		Resource->ReleaseAllocatedVT();
 	}
 
@@ -423,7 +423,7 @@ UTexture* FTextureEditorToolkit::GetTexture( ) const
 
 bool FTextureEditorToolkit::HasValidTextureResource( ) const
 {
-	return Texture != nullptr && Texture->Resource != nullptr;
+	return Texture != nullptr && Texture->GetResource() != nullptr;
 }
 
 
diff --git a/Engine/Source/Editor/UnrealEd/Private/EditorModeManager.cpp b/Engine/Source/Editor/UnrealEd/Private/EditorModeManager.cpp
--- a/Engine/Source/Editor/UnrealEd/Private/EditorModeManager.cpp
+++ b/Engine/Source/Editor/UnrealEd/Private/EditorModeManager.cpp
@@ -1690,7 +1690,7 @@ void FEditorModeTools::DrawHUD( FEditorViewportClient* InViewportClient,FViewpor
 							{
 								Canvas->SetHitProxy(new HStaticMeshVert(Actor, Vertex));
 							}
-							TileItem.Texture = VertexTexture->Resource;
+							TileItem.Texture = VertexTexture->GetResource();
 
 							TileItem.Size = FVector2D(TextureSizeX, TextureSizeY);
 							Canvas->DrawItem(TileItem, FVector2D(X, Y));
diff --git a/Engine/Source/Editor/UnrealEd/Private/LevelEditorViewport.cpp b/Engine/Source/Editor/UnrealEd/Private/LevelEditorViewport.cpp
--- a/Engine/Source/Editor/UnrealEd/Private/LevelEditorViewport.cpp
+++ b/Engine/Source/Editor/UnrealEd/Private/LevelEditorViewport.cpp
@@ -4362,7 +4362,7 @@ void FLevelEditorViewportClient::DrawBrushDetails(const FSceneView* View, FPrimi
 							const FColor Color(Brush->GetWireColor());
 							PDI->SetHitProxy(new HBSPBrushVert(Brush, &poly->Vertices[VertexIndex]));
 
-							PDI->DrawSprite(WorldLocation, TextureSizeX * Scale, TextureSizeY * Scale, VertexTexture->Resource, Color, SDPG_World, 0.0f, 0.0f, 0.0f, 0.0f, SE_BLEND_Masked);
+							PDI->DrawSprite(WorldLocation, TextureSizeX * Scale, TextureSizeY * Scale, VertexTexture->GetResource(), Color, SDPG_World, 0.0f, 0.0f, 0.0f, 0.0f, SE_BLEND_Masked);
 
 							PDI->SetHitProxy(NULL);
 						}
diff --git a/Engine/Source/Editor/UnrealEd/Private/ThumbnailRendering/FontThumbnailRenderer.cpp b/Engine/Source/Editor/UnrealEd/Private/ThumbnailRendering/FontThumbnailRenderer.cpp
--- a/Engine/Source/Editor/UnrealEd/Private/ThumbnailRendering/FontThumbnailRenderer.cpp
+++ b/Engine/Source/Editor/UnrealEd/Private/ThumbnailRendering/FontThumbnailRenderer.cpp
@@ -59,7 +59,7 @@ void UFontThumbnailRenderer::Draw(UObject* Object, int32 X, int32 Y, uint32 Widt
 		case EFontCacheType::Offline:
 			if(Font->Textures.Num() > 0 && Font->Textures[0])
 			{
-				FCanvasTileItem TileItem(FVector2D(X, Y), Font->Textures[0]->Resource, FLinearColor::White);
+				FCanvasTileItem TileItem(FVector2D(X, Y), Font->Textures[0]->GetResource(), FLinearColor::White);
 				TileItem.BlendMode = (Font->ImportOptions.bUseDistanceFieldAlpha) ? SE_BLEND_TranslucentDistanceField : SE_BLEND_Translucent;
 				Canvas->DrawItem(TileItem);
 			}
diff --git a/Engine/Source/Editor/UnrealEd/Private/ThumbnailRendering/ParticleSystemThumbnailRenderer.cpp b/Engine/Source/Editor/UnrealEd/Private/ThumbnailRendering/ParticleSystemThumbnailRenderer.cpp
--- a/Engine/Source/Editor/UnrealEd/Private/ThumbnailRendering/ParticleSystemThumbnailRenderer.cpp
+++ b/Engine/Source/Editor/UnrealEd/Private/ThumbnailRendering/ParticleSystemThumbnailRenderer.cpp
@@ -90,18 +90,18 @@ void UParticleSystemThumbnailRenderer::Draw(UObject* Object, int32 X, int32 Y, u
 			else if (ParticleSystem->ThumbnailImage)
 			{
 				Canvas->DrawTile(X,Y,Width,Height,0.f,0.f,1.f,1.f,FLinearColor::White,
-					ParticleSystem->ThumbnailImage->Resource,false);
+					ParticleSystem->ThumbnailImage->GetResource(),false);
 				if (ParticleSystem->ThumbnailImageOutOfDate == true)
 				{
 					Canvas->DrawTile(X,Y,Width/2,Height/2,0.f,0.f,1.f,1.f,FLinearColor::White,
-						OutOfDate->Resource,true);
+						OutOfDate->GetResource(),true);
 				}
 			}
 			else if (NoImage)
 			{
 				// Use the texture interface to draw
 				Canvas->DrawTile(X,Y,Width,Height,0.f,0.f,1.f,1.f,FLinearColor::White,
-					NoImage->Resource,false);
+					NoImage->GetResource(),false);
 			}
 		}
 	}
diff --git a/Engine/Source/Editor/UnrealEd/Private/ThumbnailRendering/SlateBrushThumbnailRenderer.cpp b/Engine/Source/Editor/UnrealEd/Private/ThumbnailRendering/SlateBrushThumbnailRenderer.cpp
--- a/Engine/Source/Editor/UnrealEd/Private/ThumbnailRendering/SlateBrushThumbnailRenderer.cpp
+++ b/Engine/Source/Editor/UnrealEd/Private/ThumbnailRendering/SlateBrushThumbnailRenderer.cpp
@@ -28,7 +28,7 @@ void USlateBrushThumbnailRenderer::Draw(UObject* Object, int32 X, int32 Y, uint3
 		Canvas->DrawTile(
 			0.0f, 0.0f, Width, Height,							// Dimensions
 			0.0f, 0.0f, CheckerDensity, CheckerDensity,			// UVs
-			FLinearColor::White, Checker->Resource);			// Tint & Texture
+			FLinearColor::White, Checker->GetResource());			// Tint & Texture
 
 		if (Texture)
 		{
@@ -36,14 +36,14 @@ void USlateBrushThumbnailRenderer::Draw(UObject* Object, int32 X, int32 Y, uint3
 			{
 			case ESlateBrushDrawType::Image:
 				{
-					FCanvasTileItem CanvasTile( FVector2D( X, Y ), Texture->Resource, FVector2D( Width,Height ), Brush.TintColor.GetSpecifiedColor() );
+					FCanvasTileItem CanvasTile( FVector2D( X, Y ), Texture->GetResource(), FVector2D( Width,Height ), Brush.TintColor.GetSpecifiedColor() );
 					CanvasTile.BlendMode = SE_BLEND_Translucent;
 					CanvasTile.Draw( Canvas );
 				}
 				break;
 			case ESlateBrushDrawType::Border:
 				{
-					FCanvasTileItem CanvasTile( FVector2D( X, Y ), Texture->Resource, FVector2D( Width,Height ), Brush.TintColor.GetSpecifiedColor() );
+					FCanvasTileItem CanvasTile( FVector2D( X, Y ), Texture->GetResource(), FVector2D( Width,Height ), Brush.TintColor.GetSpecifiedColor() );
 					CanvasTile.BlendMode = SE_BLEND_Translucent;
 					CanvasTile.Draw( Canvas );
 				}
@@ -66,7 +66,7 @@ void USlateBrushThumbnailRenderer::Draw(UObject* Object, int32 X, int32 Y, uint3
 						FVector2D UV0( 0, 0 );
 						FVector2D UV1( Brush.Margin.Left, Brush.Margin.Top );
 
-						FCanvasTileItem CanvasTile( FVector2D( X, Y ), Texture->Resource, TopLeftSize, UV0, UV1, Brush.TintColor.GetSpecifiedColor() );
+						FCanvasTileItem CanvasTile( FVector2D( X, Y ), Texture->GetResource(), TopLeftSize, UV0, UV1, Brush.TintColor.GetSpecifiedColor() );
 						CanvasTile.BlendMode = SE_BLEND_Translucent;
 						CanvasTile.Draw( Canvas );
 					}
@@ -77,7 +77,7 @@ void USlateBrushThumbnailRenderer::Draw(UObject* Object, int32 X, int32 Y, uint3
 						FVector2D UV0( 0, 1 - Brush.Margin.Bottom );
 						FVector2D UV1( Brush.Margin.Left, 1 );
 
-						FCanvasTileItem CanvasTile( FVector2D( X, Y + Height - BottomPx ), Texture->Resource, BottomLeftSize, UV0, UV1, Brush.TintColor.GetSpecifiedColor() );
+						FCanvasTileItem CanvasTile( FVector2D( X, Y + Height - BottomPx ), Texture->GetResource(), BottomLeftSize, UV0, UV1, Brush.TintColor.GetSpecifiedColor() );
 						CanvasTile.BlendMode = SE_BLEND_Translucent;
 						CanvasTile.Draw( Canvas );
 					}
@@ -88,7 +88,7 @@ void USlateBrushThumbnailRenderer::Draw(UObject* Object, int32 X, int32 Y, uint3
 						FVector2D UV0( 1 - Brush.Margin.Right, 0 );
 						FVector2D UV1( 1, Brush.Margin.Top );
 
-						FCanvasTileItem CanvasTile( FVector2D( X + Width - RightPx, Y ), Texture->Resource, TopRightSize, UV0, UV1, Brush.TintColor.GetSpecifiedColor() );
+						FCanvasTileItem CanvasTile( FVector2D( X + Width - RightPx, Y ), Texture->GetResource(), TopRightSize, UV0, UV1, Brush.TintColor.GetSpecifiedColor() );
 						CanvasTile.BlendMode = SE_BLEND_Translucent;
 						CanvasTile.Draw( Canvas );
 					}
@@ -99,7 +99,7 @@ void USlateBrushThumbnailRenderer::Draw(UObject* Object, int32 X, int32 Y, uint3
 						FVector2D UV0( 1 - Brush.Margin.Right, 1 - Brush.Margin.Bottom );
 						FVector2D UV1( 1, 1 );
 
-						FCanvasTileItem CanvasTile( FVector2D( X + Width - RightPx, Y + Height - BottomPx ), Texture->Resource, BottomRightSize, UV0, UV1, Brush.TintColor.GetSpecifiedColor() );
+						FCanvasTileItem CanvasTile( FVector2D( X + Width - RightPx, Y + Height - BottomPx ), Texture->GetResource(), BottomRightSize, UV0, UV1, Brush.TintColor.GetSpecifiedColor() );
 						CanvasTile.BlendMode = SE_BLEND_Translucent;
 						CanvasTile.Draw( Canvas );
 					}
@@ -112,7 +112,7 @@ void USlateBrushThumbnailRenderer::Draw(UObject* Object, int32 X, int32 Y, uint3
 						FVector2D UV0( 0, Brush.Margin.Top );
 						FVector2D UV1( Brush.Margin.Left, 1 - Brush.Margin.Bottom );
 
-						FCanvasTileItem CanvasTile( FVector2D( X, Y + TopPx), Texture->Resource, CenterVerticalLeftSize, UV0, UV1, Brush.TintColor.GetSpecifiedColor() );
+						FCanvasTileItem CanvasTile( FVector2D( X, Y + TopPx), Texture->GetResource(), CenterVerticalLeftSize, UV0, UV1, Brush.TintColor.GetSpecifiedColor() );
 						CanvasTile.BlendMode = SE_BLEND_Translucent;
 						CanvasTile.Draw( Canvas );
 					}
@@ -123,7 +123,7 @@ void USlateBrushThumbnailRenderer::Draw(UObject* Object, int32 X, int32 Y, uint3
 						FVector2D UV0( 1 - Brush.Margin.Right, Brush.Margin.Top );
 						FVector2D UV1( 1, 1 - Brush.Margin.Bottom );
 
-						FCanvasTileItem CanvasTile( FVector2D( X + Width - RightPx, Y + TopPx), Texture->Resource, CenterVerticalRightSize, UV0, UV1, Brush.TintColor.GetSpecifiedColor() );
+						FCanvasTileItem CanvasTile( FVector2D( X + Width - RightPx, Y + TopPx), Texture->GetResource(), CenterVerticalRightSize, UV0, UV1, Brush.TintColor.GetSpecifiedColor() );
 						CanvasTile.BlendMode = SE_BLEND_Translucent;
 						CanvasTile.Draw( Canvas );
 					}
@@ -136,7 +136,7 @@ void USlateBrushThumbnailRenderer::Draw(UObject* Object, int32 X, int32 Y, uint3
 						FVector2D UV0( Brush.Margin.Left, 0 );
 						FVector2D UV1( 1 - Brush.Margin.Right, Brush.Margin.Top );
 
-						FCanvasTileItem CanvasTile( FVector2D( X + LeftPx, Y), Texture->Resource, CenterHorizontalTopSize, UV0, UV1, Brush.TintColor.GetSpecifiedColor() );
+						FCanvasTileItem CanvasTile( FVector2D( X + LeftPx, Y), Texture->GetResource(), CenterHorizontalTopSize, UV0, UV1, Brush.TintColor.GetSpecifiedColor() );
 						CanvasTile.BlendMode = SE_BLEND_Translucent;
 						CanvasTile.Draw( Canvas );
 					}
@@ -147,7 +147,7 @@ void USlateBrushThumbnailRenderer::Draw(UObject* Object, int32 X, int32 Y, uint3
 						FVector2D UV0( Brush.Margin.Left, 1 - Brush.Margin.Bottom );
 						FVector2D UV1( 1 - Brush.Margin.Right, 1 );
 
-						FCanvasTileItem CanvasTile( FVector2D( X + LeftPx, Y + Height - BottomPx ), Texture->Resource, CenterHorizontalBottomSize, UV0, UV1, Brush.TintColor.GetSpecifiedColor() );
+						FCanvasTileItem CanvasTile( FVector2D( X + LeftPx, Y + Height - BottomPx ), Texture->GetResource(), CenterHorizontalBottomSize, UV0, UV1, Brush.TintColor.GetSpecifiedColor() );
 						CanvasTile.BlendMode = SE_BLEND_Translucent;
 						CanvasTile.Draw( Canvas );
 					}
@@ -160,7 +160,7 @@ void USlateBrushThumbnailRenderer::Draw(UObject* Object, int32 X, int32 Y, uint3
 						FVector2D UV0( Brush.Margin.Left, Brush.Margin.Top );
 						FVector2D UV1( 1 - Brush.Margin.Right, 1 - Brush.Margin.Bottom );
 
-						FCanvasTileItem CanvasTile( FVector2D( X + LeftPx, Y + TopPx), Texture->Resource, CenterSize, UV0, UV1, Brush.TintColor.GetSpecifiedColor() );
+						FCanvasTileItem CanvasTile( FVector2D( X + LeftPx, Y + TopPx), Texture->GetResource(), CenterSize, UV0, UV1, Brush.TintColor.GetSpecifiedColor() );
 						CanvasTile.BlendMode = SE_BLEND_Translucent;
 						CanvasTile.Draw( Canvas );
 					}
@@ -168,7 +168,7 @@ void USlateBrushThumbnailRenderer::Draw(UObject* Object, int32 X, int32 Y, uint3
 				break;
 			case ESlateBrushDrawType::NoDrawType:
 				{
-					FCanvasTileItem CanvasTile( FVector2D( X, Y ), Texture->Resource, FVector2D( Width,Height ), Brush.TintColor.GetSpecifiedColor() );
+					FCanvasTileItem CanvasTile( FVector2D( X, Y ), Texture->GetResource(), FVector2D( Width,Height ), Brush.TintColor.GetSpecifiedColor() );
 					CanvasTile.BlendMode = SE_BLEND_Translucent;
 					CanvasTile.Draw( Canvas );
 				}
diff --git a/Engine/Source/Editor/UnrealEd/Private/ThumbnailRendering/TextureThumbnailRenderer.cpp b/Engine/Source/Editor/UnrealEd/Private/ThumbnailRendering/TextureThumbnailRenderer.cpp
--- a/Engine/Source/Editor/UnrealEd/Private/ThumbnailRendering/TextureThumbnailRenderer.cpp
+++ b/Engine/Source/Editor/UnrealEd/Private/ThumbnailRendering/TextureThumbnailRenderer.cpp
@@ -49,7 +49,7 @@ void UTextureThumbnailRenderer::GetThumbnailSize(UObject* Object, float Zoom, ui
 void UTextureThumbnailRenderer::Draw(UObject* Object, int32 X, int32 Y, uint32 Width, uint32 Height, FRenderTarget*, FCanvas* Canvas, bool bAdditionalViewFamily)
 {
 	UTexture* Texture = Cast<UTexture>(Object);
-	if (Texture != nullptr && Texture->Resource != nullptr) 
+	if (Texture != nullptr && Texture->GetResource() != nullptr)
 	{
 		UTexture2D* Texture2D = Cast<UTexture2D>(Texture);
 
@@ -102,11 +102,11 @@ void UTextureThumbnailRenderer::Draw(UObject* Object, int32 X, int32 Y, uint32 W
 			Canvas->DrawTile(
 				0.0f, 0.0f, Width, Height,							// Dimensions
 				0.0f, 0.0f, CheckerDensity, CheckerDensity,			// UVs
-				FLinearColor::White, Checker->Resource);			// Tint & Texture
+				FLinearColor::White, Checker->GetResource());			// Tint & Texture
 		}
 
 		// Use A canvas tile item to draw
-		FCanvasTileItem CanvasTile( FVector2D( X, Y ), Texture->Resource, FVector2D( Width,Height ), FLinearColor::White );
+		FCanvasTileItem CanvasTile( FVector2D( X, Y ), Texture->GetResource(), FVector2D( Width,Height ), FLinearColor::White );
 		CanvasTile.BlendMode = bUseTranslucentBlend ? SE_BLEND_Translucent : SE_BLEND_Opaque;
 		CanvasTile.BatchedElementParameters = BatchedElementParameters;
 		if (bIsVirtualTexture && Texture->Source.GetNumBlocks() > 1)
diff --git a/Engine/Source/Editor/UnrealEd/Private/VolumeTexturePreview.cpp b/Engine/Source/Editor/UnrealEd/Private/VolumeTexturePreview.cpp
--- a/Engine/Source/Editor/UnrealEd/Private/VolumeTexturePreview.cpp
+++ b/Engine/Source/Editor/UnrealEd/Private/VolumeTexturePreview.cpp
@@ -79,7 +79,7 @@ public:
 		SetTextureParameter(RHICmdList, ShaderRHI, InTexture, InTextureSampler, TextureValue);
 		if (GEditor && GEditor->Bad)
 		{
-			SetTextureParameter(RHICmdList, ShaderRHI, BadTexture, BadTextureSampler, GEditor->Bad->Resource);
+			SetTextureParameter(RHICmdList, ShaderRHI, BadTexture, BadTextureSampler, GEditor->Bad->GetResource());
 		}
 		else
 		{
diff --git a/Engine/Source/Editor/WorldBrowser/Private/Tiles/WorldTileThumbnails.cpp b/Engine/Source/Editor/WorldBrowser/Private/Tiles/WorldTileThumbnails.cpp
--- a/Engine/Source/Editor/WorldBrowser/Private/Tiles/WorldTileThumbnails.cpp
+++ b/Engine/Source/Editor/WorldBrowser/Private/Tiles/WorldTileThumbnails.cpp
@@ -191,7 +191,7 @@ const FSlateBrush* FTileAtlasPage::GetSlotBrush(int32 SlotIdx) const
 
 void FTileAtlasPage::UpdateSlotImageData(int32 SlotIdx, FSlateTextureDataPtr ImageData)
 {
-	if (AtlasTexture && AtlasTexture->Resource)
+	if (AtlasTexture && AtlasTexture->GetResource())
 	{
 		int32 SlotX = (SlotIdx % TileThumbnailAtlasDim)*TileThumbnailSize;
 		int32 SlotY = (SlotIdx / TileThumbnailAtlasDim)*TileThumbnailSize;
@@ -211,7 +211,7 @@ void FTileAtlasPage::UpdateSlotImageData(int32 SlotIdx, FSlateTextureDataPtr Ima
 		} 
 		Context = 
 		{
-			AtlasTexture->Resource->TextureRHI,
+			AtlasTexture->GetResource()->TextureRHI,
 			ImageData,
 			TileThumbnailSize*4,
 			UpdateRegion
diff --git a/Engine/Source/Runtime/Engine/Classes/Engine/Texture2DArray.cpp b/Engine/Source/Runtime/Engine/Classes/Engine/Texture2DArray.cpp
--- a/Engine/Source/Runtime/Engine/Classes/Engine/Texture2DArray.cpp
+++ b/Engine/Source/Runtime/Engine/Classes/Engine/Texture2DArray.cpp
@@ -370,7 +370,7 @@ ENGINE_API void UTexture2DArray::PostEditChangeProperty(FPropertyChangedEvent &
 
 bool UTexture2DArray::StreamOut(int32 NewMipCount)
 {
-	FTexture2DArrayResource* Texture2DArrayResource = Resource ? Resource->GetTexture2DArrayResource() : nullptr;
+	FTexture2DArrayResource* Texture2DArrayResource = GetResource() ? GetResource()->GetTexture2DArrayResource() : nullptr;
 	if (!HasPendingInitOrStreaming() && CachedSRRState.StreamOut(NewMipCount) && ensure(Texture2DArrayResource))
 	{
 		FTextureMipAllocator* MipAllocator = new FTexture2DArrayMipAllocator_Reallocate(this);
@@ -382,7 +382,7 @@ bool UTexture2DArray::StreamOut(int32 NewMipCount)
 
 bool UTexture2DArray::StreamIn(int32 NewMipCount, bool bHighPrio)
 {
-	FTexture2DArrayResource* Texture2DArrayResource = Resource ? Resource->GetTexture2DArrayResource() : nullptr;
+	FTexture2DArrayResource* Texture2DArrayResource = GetResource() ? GetResource()->GetTexture2DArrayResource() : nullptr;
 	if (!HasPendingInitOrStreaming() && CachedSRRState.StreamIn(NewMipCount) && ensure(Texture2DArrayResource))
 	{
 		FTextureMipDataProvider* CustomMipDataProvider = nullptr;
diff --git a/Engine/Source/Runtime/Engine/Private/CanvasRenderTarget2D.cpp b/Engine/Source/Runtime/Engine/Private/CanvasRenderTarget2D.cpp
--- a/Engine/Source/Runtime/Engine/Private/CanvasRenderTarget2D.cpp
+++ b/Engine/Source/Runtime/Engine/Private/CanvasRenderTarget2D.cpp
@@ -36,7 +36,7 @@ void UCanvasRenderTarget2D::UpdateResource()
 
 void UCanvasRenderTarget2D::FastUpdateResource()
 {
-	if (Resource == nullptr)
+	if (GetResource() == nullptr)
 	{
 		// We don't have a resource, just take the fast path
 		UpdateResource();
diff --git a/Engine/Source/Runtime/Engine/Private/Components/BillboardComponent.cpp b/Engine/Source/Runtime/Engine/Private/Components/BillboardComponent.cpp
--- a/Engine/Source/Runtime/Engine/Private/Components/BillboardComponent.cpp
+++ b/Engine/Source/Runtime/Engine/Private/Components/BillboardComponent.cpp
@@ -126,7 +126,7 @@ public:
 	{
 		QUICK_SCOPE_CYCLE_COUNTER( STAT_SpriteSceneProxy_GetDynamicMeshElements );
 
-		FTexture* TextureResource = Texture ? Texture->Resource : nullptr;
+		const FTexture* TextureResource = Texture ? Texture->GetResource() : nullptr;
 		if (TextureResource)
 		{
 			const float UL = ComponentUL == 0.0f ? TextureResource->GetSizeX() : ComponentUL;
diff --git a/Engine/Source/Runtime/Engine/Private/Components/RectLightComponent.cpp b/Engine/Source/Runtime/Engine/Private/Components/RectLightComponent.cpp
--- a/Engine/Source/Runtime/Engine/Private/Components/RectLightComponent.cpp
+++ b/Engine/Source/Runtime/Engine/Private/Components/RectLightComponent.cpp
@@ -235,7 +235,7 @@ void FRectLightSceneProxy::GetLightShaderParameters(FLightShaderParameters& Ligh
 	LightParameters.SourceRadius = SourceWidth * 0.5f;
 	LightParameters.SoftSourceRadius = 0.0f;
 	LightParameters.SourceLength = SourceHeight * 0.5f;
-	LightParameters.SourceTexture = SourceTexture ? SourceTexture->Resource->TextureRHI : GWhiteTexture->TextureRHI;
+	LightParameters.SourceTexture = SourceTexture ? SourceTexture->GetResource()->TextureRHI : GWhiteTexture->TextureRHI;
 	LightParameters.RectLightBarnCosAngle = FMath::Cos(FMath::DegreesToRadians(BarnDoorAngle));
 	LightParameters.RectLightBarnLength = BarnDoorLength;
 }
diff --git a/Engine/Source/Runtime/Engine/Private/Components/StereoLayerComponent.cpp b/Engine/Source/Runtime/Engine/Private/Components/StereoLayerComponent.cpp
--- a/Engine/Source/Runtime/Engine/Private/Components/StereoLayerComponent.cpp
+++ b/Engine/Source/Runtime/Engine/Private/Components/StereoLayerComponent.cpp
@@ -132,7 +132,7 @@ void UStereoLayerComponent::TickComponent(float DeltaTime, enum ELevelTick TickT
 	}
 
 	bool bCurrVisible = GetVisibleFlag();
-	if (!Texture || !Texture->Resource)
+	if (!Texture || !Texture->GetResource())
 	{
 		bCurrVisible = false;
 	}
@@ -148,13 +148,13 @@ void UStereoLayerComponent::TickComponent(float DeltaTime, enum ELevelTick TickT
 		if (Texture)
 		{
 			Texture->SetForceMipLevelsToBeResident(30.0f);
-			LayerDesc.Texture = Texture->Resource->TextureRHI;
+			LayerDesc.Texture = Texture->GetResource()->TextureRHI;
 			LayerDesc.Flags |= (Texture->GetMaterialType() == MCT_TextureExternal) ? IStereoLayers::LAYER_FLAG_TEX_EXTERNAL : 0;
 		}
 		if (LeftTexture)
 		{
 			Texture->SetForceMipLevelsToBeResident(30.0f);
-			LayerDesc.LeftTexture = LeftTexture->Resource->TextureRHI;
+			LayerDesc.LeftTexture = LeftTexture->GetResource()->TextureRHI;
 		}
 				
 		LayerDesc.Flags |= (bLiveTexture) ? IStereoLayers::LAYER_FLAG_TEX_CONTINUOUS_UPDATE : 0;
diff --git a/Engine/Source/Runtime/Engine/Private/CubemapUnwrapUtils.cpp b/Engine/Source/Runtime/Engine/Private/CubemapUnwrapUtils.cpp
--- a/Engine/Source/Runtime/Engine/Private/CubemapUnwrapUtils.cpp
+++ b/Engine/Source/Runtime/Engine/Private/CubemapUnwrapUtils.cpp
@@ -105,13 +105,13 @@ namespace CubemapHelpers
 	bool GenerateLongLatUnwrap(const UTextureCube* CubeTexture, TArray64<uint8>& BitsOUT, FIntPoint& SizeOUT, EPixelFormat& FormatOUT)
 	{
 		check(CubeTexture != NULL);
-		return GenerateLongLatUnwrap(CubeTexture->Resource, CubeTexture->GetSizeX(), CubeTexture->GetPixelFormat(), BitsOUT, SizeOUT, FormatOUT);
+		return GenerateLongLatUnwrap(CubeTexture->GetResource(), CubeTexture->GetSizeX(), CubeTexture->GetPixelFormat(), BitsOUT, SizeOUT, FormatOUT);
 	}
 
 	bool GenerateLongLatUnwrap(const UTextureRenderTargetCube* CubeTarget, TArray64<uint8>& BitsOUT, FIntPoint& SizeOUT, EPixelFormat& FormatOUT)
 	{
 		check(CubeTarget != NULL);
-		return GenerateLongLatUnwrap(CubeTarget->Resource, CubeTarget->SizeX, CubeTarget->GetFormat(), BitsOUT, SizeOUT, FormatOUT);
+		return GenerateLongLatUnwrap(CubeTarget->GetResource(), CubeTarget->SizeX, CubeTarget->GetFormat(), BitsOUT, SizeOUT, FormatOUT);
 	}
 }
 
diff --git a/Engine/Source/Runtime/Engine/Private/DistanceFieldAtlas.cpp b/Engine/Source/Runtime/Engine/Private/DistanceFieldAtlas.cpp
--- a/Engine/Source/Runtime/Engine/Private/DistanceFieldAtlas.cpp
+++ b/Engine/Source/Runtime/Engine/Private/DistanceFieldAtlas.cpp
@@ -1347,7 +1347,7 @@ FLandscapeTextureAtlas::FAllocation::FAllocation(UTexture2D* InTexture, uint32 I
 {}
 
 FLandscapeTextureAtlas::FPendingUpload::FPendingUpload(UTexture2D* Texture, uint32 SizeX, uint32 SizeY, uint32 MipBias, uint32 InHandle, uint32 Channel)
-	: SourceTexture(Texture->Resource->TextureRHI)
+	: SourceTexture(Texture->GetResource()->TextureRHI)
 	, SizesAndMipBias(FIntVector(SizeX, SizeY, MipBias))
 	, VisibilityChannel(Channel)
 	, Handle(InHandle)
diff --git a/Engine/Source/Runtime/Engine/Private/HUD.cpp b/Engine/Source/Runtime/Engine/Private/HUD.cpp
--- a/Engine/Source/Runtime/Engine/Private/HUD.cpp
+++ b/Engine/Source/Runtime/Engine/Private/HUD.cpp
@@ -970,7 +970,7 @@ void AHUD::DrawTexture(UTexture* Texture, float ScreenX, float ScreenY, float Sc
 {
 	if (IsCanvasValid_WarnIfNot() && Texture)
 	{
-		FCanvasTileItem TileItem(FVector2D(ScreenX, ScreenY), Texture->Resource, FVector2D(ScreenW, ScreenH) * Scale, FVector2D(TextureU, TextureV), FVector2D(TextureU + TextureUWidth, TextureV + TextureVHeight), Color);
+		FCanvasTileItem TileItem(FVector2D(ScreenX, ScreenY), Texture->GetResource(), FVector2D(ScreenW, ScreenH) * Scale, FVector2D(TextureU, TextureV), FVector2D(TextureU + TextureUWidth, TextureV + TextureVHeight), Color);
 		TileItem.Rotation = FRotator(0, Rotation, 0);
 		TileItem.PivotPoint = RotPivot;
 		if (bScalePosition)
@@ -986,7 +986,7 @@ void AHUD::DrawTextureSimple(UTexture* Texture, float ScreenX, float ScreenY, fl
 {
 	if (IsCanvasValid_WarnIfNot() && Texture)
 	{
-		FCanvasTileItem TileItem(FVector2D(ScreenX, ScreenY), Texture->Resource, FLinearColor::White);
+		FCanvasTileItem TileItem(FVector2D(ScreenX, ScreenY), Texture->GetResource(), FLinearColor::White);
 		if (bScalePosition)
 		{
 			TileItem.Position *= Scale;
diff --git a/Engine/Source/Runtime/Engine/Private/KismetRenderingLibrary.cpp b/Engine/Source/Runtime/Engine/Private/KismetRenderingLibrary.cpp
--- a/Engine/Source/Runtime/Engine/Private/KismetRenderingLibrary.cpp
+++ b/Engine/Source/Runtime/Engine/Private/KismetRenderingLibrary.cpp
@@ -46,7 +46,7 @@ void UKismetRenderingLibrary::ClearRenderTarget2D(UObject* WorldContextObject, U
 	UWorld* World = GEngine->GetWorldFromContextObject(WorldContextObject, EGetWorldErrorMode::LogAndReturnNull);
 
 	if (TextureRenderTarget
-		&& TextureRenderTarget->Resource
+		&& TextureRenderTarget->GetResource()
 		&& World)
 	{
 		FTextureRenderTargetResource* RenderTargetResource = TextureRenderTarget->GameThread_GetRenderTargetResource();
@@ -151,7 +151,7 @@ void UKismetRenderingLibrary::DrawMaterialToRenderTarget(UObject* WorldContextOb
 	{
 		FMessageLog("Blueprint").Warning(FText::Format(LOCTEXT("DrawMaterialToRenderTarget_InvalidTextureRenderTarget", "DrawMaterialToRenderTarget[{0}]: TextureRenderTarget must be non-null."), FText::FromString(GetPathNameSafe(WorldContextObject))));
 	}
-	else if (!TextureRenderTarget->Resource)
+	else if (!TextureRenderTarget->GetResource())
 	{
 		FMessageLog("Blueprint").Warning(FText::Format(LOCTEXT("DrawMaterialToRenderTarget_ReleasedTextureRenderTarget", "DrawMaterialToRenderTarget[{0}]: render target has been released."), FText::FromString(GetPathNameSafe(WorldContextObject))));
 	}
@@ -203,7 +203,7 @@ void UKismetRenderingLibrary::ExportRenderTarget(UObject* WorldContextObject, UT
 	{
 		FMessageLog("Blueprint").Warning(FText::Format(LOCTEXT("ExportRenderTarget_InvalidTextureRenderTarget", "ExportRenderTarget[{0}]: TextureRenderTarget must be non-null."), FText::FromString(GetPathNameSafe(WorldContextObject))));
 	}
-	else if (!TextureRenderTarget->Resource)
+	else if (!TextureRenderTarget->GetResource())
 	{
 		FMessageLog("Blueprint").Warning(FText::Format(LOCTEXT("ExportRenderTarget_ReleasedTextureRenderTarget", "ExportRenderTarget[{0}]: render target has been released."), FText::FromString(GetPathNameSafe(WorldContextObject))));
 	}
@@ -464,7 +464,7 @@ UTexture2D* UKismetRenderingLibrary::RenderTargetCreateStaticTexture2DEditorOnly
 		FMessageLog("Blueprint").Warning(LOCTEXT("RenderTargetCreateStaticTexture2D_InvalidRenderTarget", "RenderTargetCreateStaticTexture2DEditorOnly: RenderTarget must be non-null."));
 		return nullptr;
 	}
-	else if (!RenderTarget->Resource)
+	else if (!RenderTarget->GetResource())
 	{
 		FMessageLog("Blueprint").Warning(LOCTEXT("RenderTargetCreateStaticTexture2D_ReleasedRenderTarget", "RenderTargetCreateStaticTexture2DEditorOnly: RenderTarget has been released."));
 		return nullptr;
@@ -528,7 +528,7 @@ void UKismetRenderingLibrary::ConvertRenderTargetToTexture2DEditorOnly( UObject*
 	{
 		FMessageLog("Blueprint").Warning(FText::Format(LOCTEXT("ConvertRenderTargetToTexture2D_InvalidRenderTarget", "ConvertRenderTargetToTexture2DEditorOnly[{0}]: RenderTarget must be non-null."), FText::FromString(GetPathNameSafe(WorldContextObject))));
 	}
-	else if (!RenderTarget->Resource)
+	else if (!RenderTarget->GetResource())
 	{
 		FMessageLog("Blueprint").Warning(FText::Format(LOCTEXT("ConvertRenderTargetToTexture2D_ReleasedTextureRenderTarget", "ConvertRenderTargetToTexture2DEditorOnly[{0}]: render target has been released."), FText::FromString(GetPathNameSafe(WorldContextObject))));
 	}
@@ -627,7 +627,7 @@ void UKismetRenderingLibrary::BeginDrawCanvasToRenderTarget(UObject* WorldContex
 	{
 		FMessageLog("Blueprint").Warning(FText::Format(LOCTEXT("BeginDrawCanvasToRenderTarget_InvalidTextureRenderTarget", "BeginDrawCanvasToRenderTarget[{0}]: TextureRenderTarget must be non-null."), FText::FromString(GetPathNameSafe(WorldContextObject))));
 	}
-	else if (!TextureRenderTarget->Resource)
+	else if (!TextureRenderTarget->GetResource())
 	{
 		FMessageLog("Blueprint").Warning(FText::Format(LOCTEXT("BeginDrawCanvasToRenderTarget_ReleasedTextureRenderTarget", "BeginDrawCanvasToRenderTarget[{0}]: render target has been released."), FText::FromString(GetPathNameSafe(WorldContextObject))));
 	}
diff --git a/Engine/Source/Runtime/Engine/Private/LightMap.cpp b/Engine/Source/Runtime/Engine/Private/LightMap.cpp
--- a/Engine/Source/Runtime/Engine/Private/LightMap.cpp
+++ b/Engine/Source/Runtime/Engine/Private/LightMap.cpp
@@ -3195,7 +3195,7 @@ FLightMapInteraction FLightMap2D::GetInteraction(ERHIFeatureLevel::Type InFeatur
 	const bool bUseVirtualTextures = (CVarVirtualTexturedLightMaps.GetValueOnAnyThread() != 0) && UseVirtualTexturing(InFeatureLevel);
 	if (!bUseVirtualTextures)
 	{
-		bool bValidTextures = Textures[LightmapIndex] && Textures[LightmapIndex]->Resource;
+		bool bValidTextures = Textures[LightmapIndex] && Textures[LightmapIndex]->GetResource();
 
 		// When the FLightMap2D is first created, the textures aren't set, so that case needs to be handled.
 		if (bValidTextures)
@@ -3206,7 +3206,7 @@ FLightMapInteraction FLightMap2D::GetInteraction(ERHIFeatureLevel::Type InFeatur
 	else
 	{
 		// Preview lightmaps don't stream from disk, thus no FVirtualTexture2DResource
-		bool bValidVirtualTexture = VirtualTextures[LightmapIndex] && (VirtualTextures[LightmapIndex]->Resource != nullptr || VirtualTextures[LightmapIndex]->bPreviewLightmap);
+		bool bValidVirtualTexture = VirtualTextures[LightmapIndex] && (VirtualTextures[LightmapIndex]->GetResource() != nullptr || VirtualTextures[LightmapIndex]->bPreviewLightmap);
 		if (bValidVirtualTexture)
 		{
 			return FLightMapInteraction::InitVirtualTexture(VirtualTextures[LightmapIndex], ScaleVectors, AddVectors, CoordinateScale, CoordinateBias, bHighQuality);
@@ -3226,7 +3226,7 @@ FShadowMapInteraction FLightMap2D::GetShadowInteraction(ERHIFeatureLevel::Type I
 	if (bUseVirtualTextures)
 	{
 		// Preview lightmaps don't stream from disk, thus no FVirtualTexture2DResource
-		const bool bValidVirtualTexture = VirtualTextures[LightmapIndex] && (VirtualTextures[LightmapIndex]->Resource != nullptr || VirtualTextures[LightmapIndex]->bPreviewLightmap);
+		const bool bValidVirtualTexture = VirtualTextures[LightmapIndex] && (VirtualTextures[LightmapIndex]->GetResource() != nullptr || VirtualTextures[LightmapIndex]->bPreviewLightmap);
 		if (bValidVirtualTexture)
 		{
 			return FShadowMapInteraction::InitVirtualTexture(VirtualTextures[LightmapIndex], CoordinateScale, CoordinateBias, bShadowChannelValid, InvUniformPenumbraSize);
@@ -3476,10 +3476,10 @@ IAllocatedVirtualTexture* FLightmapResourceCluster::AcquireAllocatedVT() const
 		}
 #endif
 
-	if (!AllocatedVT && VirtualTexture && VirtualTexture->Resource)
+	if (!AllocatedVT && VirtualTexture && VirtualTexture->GetResource())
 	{
 		check(VirtualTexture->VirtualTextureStreaming);
-		const FVirtualTexture2DResource* Resource = (FVirtualTexture2DResource*)VirtualTexture->Resource;
+		const FVirtualTexture2DResource* Resource = (FVirtualTexture2DResource*)VirtualTexture->GetResource();
 		const FVirtualTextureProducerHandle ProducerHandle = Resource->GetProducerHandle();
 
 		GetRendererModule().AddVirtualTextureProducerDestroyedCallback(ProducerHandle, &OnVirtualTextureDestroyed, const_cast<FLightmapResourceCluster*>(this));
diff --git a/Engine/Source/Runtime/Engine/Private/Materials/MaterialShared.cpp b/Engine/Source/Runtime/Engine/Private/Materials/MaterialShared.cpp
--- a/Engine/Source/Runtime/Engine/Private/Materials/MaterialShared.cpp
+++ b/Engine/Source/Runtime/Engine/Private/Materials/MaterialShared.cpp
@@ -2977,7 +2977,7 @@ IAllocatedVirtualTexture* FMaterialRenderProxy::AllocateVTStack(const FMaterialR
 			continue;
 		}
 
-		const FVirtualTexture2DResource* VirtualTextureResourceForLayer = (FVirtualTexture2DResource*)Texture->Resource;
+		const FVirtualTexture2DResource* VirtualTextureResourceForLayer = (FVirtualTexture2DResource*)Texture->GetResource();
 		if (VirtualTextureResourceForLayer != nullptr)
 		{
 			// All tile sizes need to match
diff --git a/Engine/Source/Runtime/Engine/Private/Materials/MaterialUniformExpressions.cpp b/Engine/Source/Runtime/Engine/Private/Materials/MaterialUniformExpressions.cpp
--- a/Engine/Source/Runtime/Engine/Private/Materials/MaterialUniformExpressions.cpp
+++ b/Engine/Source/Runtime/Engine/Private/Materials/MaterialUniformExpressions.cpp
@@ -882,9 +882,9 @@ void FUniformExpressionSet::FillUniformBuffer(const FMaterialRenderContext& Mate
 			// It's possible for this command to trigger before a given material is cleaned up and removed from deferred update list
 			// Technically I don't think it's necessary to check 'Resource' for nullptr here, as if TextureReferenceRHI has been initialized, that should be enough
 			// Going to leave the check for now though, to hopefully avoid any unexpected problems
-			if (Value && Value->Resource && Value->TextureReference.TextureReferenceRHI && (Value->GetMaterialType() & ValidTextureTypes) != 0u)
+			if (Value && Value->GetResource() && Value->TextureReference.TextureReferenceRHI && (Value->GetMaterialType() & ValidTextureTypes) != 0u)
 			{
-				FSamplerStateRHIRef* SamplerSource = &Value->Resource->SamplerStateRHI;
+				const FSamplerStateRHIRef* SamplerSource = &Value->GetResource()->SamplerStateRHI;
 
 				const ESamplerSourceMode SourceMode = Parameter.SamplerSource;
 				if (SourceMode == SSM_Wrap_WorldGroupSettings)
@@ -908,7 +908,7 @@ void FUniformExpressionSet::FillUniformBuffer(const FMaterialRenderContext& Mate
 						TEXT("Texture %s of class %s had invalid sampler source. Material %s with texture expression in slot %i. Sampler source mode %d. Resource initialized: %d"),
 						*Value->GetName(), *Value->GetClass()->GetName(),
 						*MaterialRenderContext.Material.GetFriendlyName(), ExpressionIndex, SourceMode,
-						Value->Resource->IsInitialized());
+						Value->GetResource()->IsInitialized());
 				}
 			}
 
@@ -934,11 +934,11 @@ void FUniformExpressionSet::FillUniformBuffer(const FMaterialRenderContext& Mate
 			BufferCursor = ((uint8*)BufferCursor) + (SHADER_PARAMETER_POINTER_ALIGNMENT * 2);
 			check(BufferCursor <= TempBuffer + TempBufferSize);
 
-			if(Value && Value->Resource && (Value->GetMaterialType() & MCT_TextureCube) != 0u)
+			if(Value && Value->GetResource() && (Value->GetMaterialType() & MCT_TextureCube) != 0u)
 			{
 				check(Value->TextureReference.TextureReferenceRHI);
 				*ResourceTableTexturePtr = Value->TextureReference.TextureReferenceRHI;
-				FSamplerStateRHIRef* SamplerSource = &Value->Resource->SamplerStateRHI;
+				const FSamplerStateRHIRef* SamplerSource = &Value->GetResource()->SamplerStateRHI;
 
 				const ESamplerSourceMode SourceMode = Parameter.SamplerSource;
 				if (SourceMode == SSM_Wrap_WorldGroupSettings)
@@ -974,11 +974,11 @@ void FUniformExpressionSet::FillUniformBuffer(const FMaterialRenderContext& Mate
 			void** ResourceTableSamplerPtr = (void**)((uint8*)BufferCursor + 1 * SHADER_PARAMETER_POINTER_ALIGNMENT);
 			BufferCursor = ((uint8*)BufferCursor) + (SHADER_PARAMETER_POINTER_ALIGNMENT * 2);
 
-			if (Value && Value->Resource && (Value->GetMaterialType() & MCT_Texture2DArray) != 0u)
+			if (Value && Value->GetResource() && (Value->GetMaterialType() & MCT_Texture2DArray) != 0u)
 			{
 				check(Value->TextureReference.TextureReferenceRHI);
 				*ResourceTableTexturePtr = Value->TextureReference.TextureReferenceRHI;
-				FSamplerStateRHIRef* SamplerSource = &Value->Resource->SamplerStateRHI;
+				const FSamplerStateRHIRef* SamplerSource = &Value->GetResource()->SamplerStateRHI;
 				ESamplerSourceMode SourceMode = Parameter.SamplerSource;
 				if (SourceMode == SSM_Wrap_WorldGroupSettings)
 				{
@@ -1014,11 +1014,11 @@ void FUniformExpressionSet::FillUniformBuffer(const FMaterialRenderContext& Mate
 			BufferCursor = ((uint8*)BufferCursor) + (SHADER_PARAMETER_POINTER_ALIGNMENT * 2);
 			check(BufferCursor <= TempBuffer + TempBufferSize);
 
-			if(Value && Value->Resource && (Value->GetMaterialType() & MCT_VolumeTexture) != 0u)
+			if(Value && Value->GetResource() && (Value->GetMaterialType() & MCT_VolumeTexture) != 0u)
 			{
 				check(Value->TextureReference.TextureReferenceRHI);
 				*ResourceTableTexturePtr = Value->TextureReference.TextureReferenceRHI;
-				FSamplerStateRHIRef* SamplerSource = &Value->Resource->SamplerStateRHI;
+				const FSamplerStateRHIRef* SamplerSource = &Value->GetResource()->SamplerStateRHI;
 
 				const ESamplerSourceMode SourceMode = Parameter.SamplerSource;
 				if (SourceMode == SSM_Wrap_WorldGroupSettings)
@@ -1139,10 +1139,10 @@ void FUniformExpressionSet::FillUniformBuffer(const FMaterialRenderContext& Mate
 			{
 				const UTexture* Texture = nullptr;
 				GetTextureValue(EMaterialTextureParameterType::Virtual, ExpressionIndex, MaterialRenderContext, MaterialRenderContext.Material, Texture);
-				if (Texture && Texture->Resource)
+				if (Texture && Texture->GetResource())
 				{
 					const FVTPackedStackAndLayerIndex StackAndLayerIndex = GetVTStackAndLayerIndex(ExpressionIndex);
-					FVirtualTexture2DResource* VTResource = (FVirtualTexture2DResource*)Texture->Resource;
+					FVirtualTexture2DResource* VTResource = (FVirtualTexture2DResource*)Texture->GetResource();
 					check(VTResource);
 
 					const IAllocatedVirtualTexture* AllocatedVT = UniformExpressionCache.AllocatedVTs[StackAndLayerIndex.StackIndex];
@@ -1215,7 +1215,7 @@ uint32 FUniformExpressionSet::GetReferencedTexture2DRHIHash(const FMaterialRende
 		const uint32 ValidTextureTypes = MCT_Texture2D | MCT_TextureVirtual | MCT_TextureExternal;
 
 		FRHITexture* TexturePtr = nullptr;
-		if (Value && Value->Resource && Value->TextureReference.TextureReferenceRHI && (Value->GetMaterialType() & ValidTextureTypes) != 0u)
+		if (Value && Value->GetResource() && Value->TextureReference.TextureReferenceRHI && (Value->GetMaterialType() & ValidTextureTypes) != 0u)
 		{
 			TexturePtr = Value->TextureReference.TextureReferenceRHI->GetReferencedTexture();
 		}
diff --git a/Engine/Source/Runtime/Engine/Private/SceneManagement.cpp b/Engine/Source/Runtime/Engine/Private/SceneManagement.cpp
--- a/Engine/Source/Runtime/Engine/Private/SceneManagement.cpp
+++ b/Engine/Source/Runtime/Engine/Private/SceneManagement.cpp
@@ -883,7 +883,7 @@ bool FLightCacheInterface::GetVirtualTextureLightmapProducer(ERHIFeatureLevel::T
 		// Preview lightmaps don't stream from disk, thus no FVirtualTexture2DResource
 		if (VirtualTexture && !VirtualTexture->bPreviewLightmap)
 		{
-			FVirtualTexture2DResource* Resource = (FVirtualTexture2DResource*)VirtualTexture->Resource;
+			FVirtualTexture2DResource* Resource = (FVirtualTexture2DResource*)VirtualTexture->GetResource();
 			OutProducerHandle = Resource->GetProducerHandle();
 			return true;
 		}
@@ -896,9 +896,9 @@ IMPLEMENT_GLOBAL_SHADER_PARAMETER_STRUCT(FLightmapResourceClusterShaderParameter
 static FRHISamplerState* GetTextureSamplerState(const UTexture* Texture, FRHISamplerState* Default)
 {
 	FRHISamplerState* Result = nullptr;
-	if (Texture && Texture->Resource)
+	if (Texture && Texture->GetResource())
 	{
-		Result = Texture->Resource->SamplerStateRHI;
+		Result = Texture->GetResource()->SamplerStateRHI;
 	}
 	return Result ? Result : Default;
 }
diff --git a/Engine/Source/Runtime/Engine/Private/Shader/Preshader.cpp b/Engine/Source/Runtime/Engine/Private/Shader/Preshader.cpp
--- a/Engine/Source/Runtime/Engine/Private/Shader/Preshader.cpp
+++ b/Engine/Source/Runtime/Engine/Private/Shader/Preshader.cpp
@@ -247,11 +247,11 @@ static const UTexture* GetTextureParameter(const FMaterialRenderContext& Context
 static void EvaluateTextureSize(const FMaterialRenderContext& Context, UE::Shader::FPreshaderStack& Stack, UE::Shader::FPreshaderDataContext& RESTRICT Data)
 {
 	const UTexture* Texture = GetTextureParameter(Context, Data);
-	if (Texture && Texture->Resource)
+	if (Texture && Texture->GetResource())
 	{
-		const uint32 SizeX = Texture->Resource->GetSizeX();
-		const uint32 SizeY = Texture->Resource->GetSizeY();
-		const uint32 SizeZ = Texture->Resource->GetSizeZ();
+		const uint32 SizeX = Texture->GetResource()->GetSizeX();
+		const uint32 SizeY = Texture->GetResource()->GetSizeY();
+		const uint32 SizeZ = Texture->GetResource()->GetSizeZ();
 		Stack.Add(UE::Shader::FValue((float)SizeX, (float)SizeY, (float)SizeZ));
 	}
 	else
@@ -263,11 +263,11 @@ static void EvaluateTextureSize(const FMaterialRenderContext& Context, UE::Shade
 static void EvaluateTexelSize(const FMaterialRenderContext& Context, UE::Shader::FPreshaderStack& Stack, UE::Shader::FPreshaderDataContext& RESTRICT Data)
 {
 	const UTexture* Texture = GetTextureParameter(Context, Data);
-	if (Texture && Texture->Resource)
+	if (Texture && Texture->GetResource())
 	{
-		const uint32 SizeX = Texture->Resource->GetSizeX();
-		const uint32 SizeY = Texture->Resource->GetSizeY();
-		const uint32 SizeZ = Texture->Resource->GetSizeZ();
+		const uint32 SizeX = Texture->GetResource()->GetSizeX();
+		const uint32 SizeY = Texture->GetResource()->GetSizeY();
+		const uint32 SizeZ = Texture->GetResource()->GetSizeZ();
 		Stack.Add(UE::Shader::FValue(1.0f / (float)SizeX, 1.0f / (float)SizeY, (SizeZ > 0 ? 1.0f / (float)SizeZ : 0.0f)));
 	}
 	else
diff --git a/Engine/Source/Runtime/Engine/Private/StatsRender2.cpp b/Engine/Source/Runtime/Engine/Private/StatsRender2.cpp
--- a/Engine/Source/Runtime/Engine/Private/StatsRender2.cpp
+++ b/Engine/Source/Runtime/Engine/Private/StatsRender2.cpp
@@ -267,7 +267,7 @@ struct FStatRenderGlobals
 		else
 		{
 			UTexture2D* BackgroundTexture = UCanvas::StaticClass()->GetDefaultObject<UCanvas>()->GradientTexture0;
-			return (BackgroundTexture != nullptr) ? BackgroundTexture->Resource : nullptr;
+			return (BackgroundTexture != nullptr) ? BackgroundTexture->GetResource() : nullptr;
 		}
 	}
 };
diff --git a/Engine/Source/Runtime/Engine/Private/StereoLayerFunctionLibrary.cpp b/Engine/Source/Runtime/Engine/Private/StereoLayerFunctionLibrary.cpp
--- a/Engine/Source/Runtime/Engine/Private/StereoLayerFunctionLibrary.cpp
+++ b/Engine/Source/Runtime/Engine/Private/StereoLayerFunctionLibrary.cpp
@@ -81,11 +81,11 @@ UStereoLayerFunctionLibrary::UStereoLayerFunctionLibrary(const FObjectInitialize
 void UStereoLayerFunctionLibrary::SetSplashScreen(class UTexture* Texture, FVector2D Scale, FVector Offset, bool bShowLoadingMovie, bool bShowOnSet)
 {
 	IXRLoadingScreen* LoadingScreen = GetLoadingScreen();
-	if (LoadingScreen && Texture && Texture->Resource)
+	if (LoadingScreen && Texture && Texture->GetResource())
 	{
 		LoadingScreen->ClearSplashes();
 		IXRLoadingScreen::FSplashDesc Splash;
-		Splash.Texture = Texture->Resource->TextureRHI;
+		Splash.Texture = Texture->GetResource()->TextureRHI;
 		Splash.QuadSize = Scale;
 		Splash.Transform = FTransform(Offset);
 		LoadingScreen->AddSplash(Splash);
diff --git a/Engine/Source/Runtime/Engine/Private/Streaming/StreamingManagerTexture.cpp b/Engine/Source/Runtime/Engine/Private/Streaming/StreamingManagerTexture.cpp
--- a/Engine/Source/Runtime/Engine/Private/Streaming/StreamingManagerTexture.cpp
+++ b/Engine/Source/Runtime/Engine/Private/Streaming/StreamingManagerTexture.cpp
@@ -2400,9 +2400,9 @@ bool FRenderAssetStreamingManager::HandleLODGroupsCommand( const TCHAR* Cmd, FOu
 			{
 				return (*Texture->GetRunningPlatformData())->PixelFormat;
 			}
-			else if (Texture->Resource && Texture->Resource->TextureRHI)
+			else if (Texture->GetResource() && Texture->GetResource()->TextureRHI)
 			{
-				return Texture->Resource->TextureRHI->GetFormat();
+				return Texture->GetResource()->TextureRHI->GetFormat();
 			}
 			else
 			{
@@ -2411,7 +2411,7 @@ bool FRenderAssetStreamingManager::HandleLODGroupsCommand( const TCHAR* Cmd, FOu
 		}();
 
 		// No resource no size taken
-		if (!Texture->Resource)
+		if (!Texture->GetResource())
 		{
 			LODStats.NumNoResourceTextures++;
 		}
diff --git a/Engine/Source/Runtime/Engine/Private/Streaming/Texture2DUpdate.cpp b/Engine/Source/Runtime/Engine/Private/Streaming/Texture2DUpdate.cpp
--- a/Engine/Source/Runtime/Engine/Private/Streaming/Texture2DUpdate.cpp
+++ b/Engine/Source/Runtime/Engine/Private/Streaming/Texture2DUpdate.cpp
@@ -23,7 +23,7 @@ FTexture2DUpdateContext::FTexture2DUpdateContext(const UTexture2D* InTexture, ET
 {
 	check(InTexture);
 	checkSlow(InCurrentThread != FTexture2DUpdate::TT_Render || IsInRenderingThread());
-	Resource = Texture && Texture->Resource ? Texture->Resource->GetTexture2DResource() : nullptr;
+	Resource = Texture && Texture->GetResource() ? const_cast<UTexture2D*>(Texture)->GetResource()->GetTexture2DResource() : nullptr;
 	if (Resource)
 	{
 		MipsView = Resource->GetPlatformMipsView();
@@ -37,7 +37,7 @@ FTexture2DUpdateContext::FTexture2DUpdateContext(const UStreamableRenderAsset* I
 FTexture2DUpdate::FTexture2DUpdate(UTexture2D* InTexture) 
 	: TRenderAssetUpdate<FTexture2DUpdateContext>(InTexture)
 {
-	if (!InTexture->Resource)
+	if (!InTexture->GetResource())
 	{
 		bIsCancelled = true;
 	}
diff --git a/Engine/Source/Runtime/Engine/Private/Streaming/TextureMipDataProvider.cpp b/Engine/Source/Runtime/Engine/Private/Streaming/TextureMipDataProvider.cpp
--- a/Engine/Source/Runtime/Engine/Private/Streaming/TextureMipDataProvider.cpp
+++ b/Engine/Source/Runtime/Engine/Private/Streaming/TextureMipDataProvider.cpp
@@ -13,7 +13,7 @@ FTextureUpdateContext::FTextureUpdateContext(const UTexture* InTexture, EThreadT
 	, CurrentThread(InCurrentThread)
 {
 	check(InTexture);
-	Resource = Texture && Texture->Resource ? Texture->Resource->GetStreamableTextureResource() : nullptr;
+	Resource = Texture && Texture->GetResource() ? const_cast<UTexture*>(Texture)->GetResource()->GetStreamableTextureResource() : nullptr;
 	if (Resource)
 	{
 		MipsView = Resource->GetPlatformMipsView();
diff --git a/Engine/Source/Runtime/Engine/Private/Texture.cpp b/Engine/Source/Runtime/Engine/Private/Texture.cpp
--- a/Engine/Source/Runtime/Engine/Private/Texture.cpp
+++ b/Engine/Source/Runtime/Engine/Private/Texture.cpp
@@ -560,9 +560,9 @@ void UTexture::BeginFinalReleaseResource()
 {
 	check(!bAsyncResourceReleaseHasBeenStarted);
 	// Send the rendering thread a release message for the texture's resource.
-	if (Resource)
+	if (GetResource())
 	{
-		BeginReleaseResource(Resource);
+		BeginReleaseResource(GetResource());
 	}
 	if (TextureReference.IsInitialized_GameThread())
 	{
@@ -709,31 +709,31 @@ bool UTexture::DoesMipDataExist(const int32 MipIndex) const
 
 bool UTexture::HasPendingRenderResourceInitialization() const
 {
-	return Resource && !Resource->IsInitialized();
+	return GetResource() && !GetResource()->IsInitialized();
 }
 
 bool UTexture::HasPendingLODTransition() const
 {
-	return Resource && Resource->MipBiasFade.IsFading();
+	return GetResource() && GetResource()->MipBiasFade.IsFading();
 }
 
 float UTexture::GetLastRenderTimeForStreaming() const
 {
 	float LastRenderTime = -FLT_MAX;
-	if (Resource)
+	if (GetResource())
 	{
 		// The last render time is the last time the resource was directly bound or the last
 		// time the texture reference was cached in a resource table, whichever was later.
-		LastRenderTime = FMath::Max<double>(Resource->LastRenderTime,TextureReference.GetLastRenderTime());
+		LastRenderTime = FMath::Max<double>(GetResource()->LastRenderTime,TextureReference.GetLastRenderTime());
 	}
 	return LastRenderTime;
 }
 
 void UTexture::InvalidateLastRenderTimeForStreaming()
 {
-	if (Resource)
+	if (GetResource())
 	{
-		Resource->LastRenderTime = -FLT_MAX;
+		GetResource()->LastRenderTime = -FLT_MAX;
 	}
 	TextureReference.InvalidateLastRenderTime();
 }
diff --git a/Engine/Source/Runtime/Engine/Private/Texture2D.cpp b/Engine/Source/Runtime/Engine/Private/Texture2D.cpp
--- a/Engine/Source/Runtime/Engine/Private/Texture2D.cpp
+++ b/Engine/Source/Runtime/Engine/Private/Texture2D.cpp
@@ -477,7 +477,7 @@ int32 UTexture2D::GetNumResidentMips() const
 			in-game	resolution of the texture as it's currently loaded. An other option would be "Mips that are partially resident" as that would cover
 			somewhat the same but knowing this is additional burden on the VT system and interfaces.
 			*/
-			return static_cast<const FVirtualTexture2DResource*>(Resource)->GetNumMips();
+			return static_cast<const FVirtualTexture2DResource*>(GetResource())->GetNumMips();
 		}
 		else if (CachedSRRState.IsValid())
 		{
@@ -485,7 +485,7 @@ int32 UTexture2D::GetNumResidentMips() const
 		}
 		else
 		{
-			return Resource->GetCurrentMipCount();
+			return GetResource()->GetCurrentMipCount();
 		}
 	}
 	return 0;
diff --git a/Engine/Source/Runtime/Engine/Private/TextureRenderTarget.cpp b/Engine/Source/Runtime/Engine/Private/TextureRenderTarget.cpp
--- a/Engine/Source/Runtime/Engine/Private/TextureRenderTarget.cpp
+++ b/Engine/Source/Runtime/Engine/Private/TextureRenderTarget.cpp
@@ -28,12 +28,12 @@ UTextureRenderTarget::UTextureRenderTarget(const FObjectInitializer& ObjectIniti
 FTextureRenderTargetResource* UTextureRenderTarget::GetRenderTargetResource()
 {
 	check(IsInRenderingThread() || 
-		(IsInParallelRenderingThread() && (!Resource || Resource->IsInitialized()))); // we allow this in parallel, but only if the resource is initialized...otherwise it might be a race on intialization
+		(IsInParallelRenderingThread() && (!GetResource() || GetResource()->IsInitialized()))); // we allow this in parallel, but only if the resource is initialized...otherwise it might be a race on intialization
 	FTextureRenderTargetResource* Result = NULL;
-	if( Resource &&
-		Resource->IsInitialized() )
+	if (GetResource() &&
+		GetResource()->IsInitialized() )
 	{
-		Result = static_cast<FTextureRenderTargetResource*>( Resource );
+		Result = static_cast<FTextureRenderTargetResource*>(GetResource());
 	}
 	return Result;
 }
@@ -42,7 +42,7 @@ FTextureRenderTargetResource* UTextureRenderTarget::GetRenderTargetResource()
 FTextureRenderTargetResource* UTextureRenderTarget::GameThread_GetRenderTargetResource()
 {
 	check( IsInGameThread() );
-	return static_cast< FTextureRenderTargetResource*>( Resource );
+	return static_cast< FTextureRenderTargetResource*>(GetResource());
 }
 
 
diff --git a/Engine/Source/Runtime/Engine/Private/TextureRenderTarget2D.cpp b/Engine/Source/Runtime/Engine/Private/TextureRenderTarget2D.cpp
--- a/Engine/Source/Runtime/Engine/Private/TextureRenderTarget2D.cpp
+++ b/Engine/Source/Runtime/Engine/Private/TextureRenderTarget2D.cpp
@@ -130,9 +130,9 @@ void UTextureRenderTarget2D::ResizeTarget(uint32 InSizeX, uint32 InSizeY)
 			NumMips = FGenericPlatformMath::CeilToInt(FGenericPlatformMath::Log2(static_cast<float>(FGenericPlatformMath::Max(SizeX, SizeY))));
 		}
 
-		if (Resource)
+		if (GetResource())
 		{
-			FTextureRenderTarget2DResource* InResource = static_cast<FTextureRenderTarget2DResource*>(Resource);
+			FTextureRenderTarget2DResource* InResource = static_cast<FTextureRenderTarget2DResource*>(GetResource());
 			int32 NewSizeX = SizeX;
 			int32 NewSizeY = SizeY;
 			ENQUEUE_RENDER_COMMAND(ResizeRenderTarget)(
@@ -142,8 +142,6 @@ void UTextureRenderTarget2D::ResizeTarget(uint32 InSizeX, uint32 InSizeY)
 					InResource->UpdateDeferredResource(RHICmdList, true);
 				}
 			);
-
-
 		}
 		else
 		{
@@ -154,9 +152,9 @@ void UTextureRenderTarget2D::ResizeTarget(uint32 InSizeX, uint32 InSizeY)
 
 void UTextureRenderTarget2D::UpdateResourceImmediate(bool bClearRenderTarget/*=true*/)
 {
-	if (Resource)
+	if (GetResource())
 	{
-		FTextureRenderTarget2DResource* InResource = static_cast<FTextureRenderTarget2DResource*>(Resource);
+		FTextureRenderTarget2DResource* InResource = static_cast<FTextureRenderTarget2DResource*>(GetResource());
 		ENQUEUE_RENDER_COMMAND(UpdateResourceImmediate)(
 			[InResource, bClearRenderTarget](FRHICommandListImmediate& RHICmdList)
 			{
diff --git a/Engine/Source/Runtime/Engine/Private/TextureRenderTarget2DArray.cpp b/Engine/Source/Runtime/Engine/Private/TextureRenderTarget2DArray.cpp
--- a/Engine/Source/Runtime/Engine/Private/TextureRenderTarget2DArray.cpp
+++ b/Engine/Source/Runtime/Engine/Private/TextureRenderTarget2DArray.cpp
@@ -61,9 +61,9 @@ void UTextureRenderTarget2DArray::InitAutoFormat(uint32 InSizeX, uint32 InSizeY,
 
 void UTextureRenderTarget2DArray::UpdateResourceImmediate(bool bClearRenderTarget/*=true*/)
 {
-	if (Resource)
+	if (GetResource())
 	{
-		FTextureRenderTarget2DArrayResource* InResource = static_cast<FTextureRenderTarget2DArrayResource*>(Resource);
+		FTextureRenderTarget2DArrayResource* InResource = static_cast<FTextureRenderTarget2DArrayResource*>(GetResource());
 		ENQUEUE_RENDER_COMMAND(UpdateResourceImmediate)(
 			[InResource, bClearRenderTarget](FRHICommandListImmediate& RHICmdList)
 			{
diff --git a/Engine/Source/Runtime/Engine/Private/TextureRenderTargetCube.cpp b/Engine/Source/Runtime/Engine/Private/TextureRenderTargetCube.cpp
--- a/Engine/Source/Runtime/Engine/Private/TextureRenderTargetCube.cpp
+++ b/Engine/Source/Runtime/Engine/Private/TextureRenderTargetCube.cpp
@@ -56,9 +56,9 @@ void UTextureRenderTargetCube::InitAutoFormat(uint32 InSizeX)
 
 void UTextureRenderTargetCube::UpdateResourceImmediate(bool bClearRenderTarget/*=true*/)
 {
-	if (Resource)
+	if (GetResource())
 	{
-		FTextureRenderTargetCubeResource* InResource = static_cast<FTextureRenderTargetCubeResource*>(Resource);
+		FTextureRenderTargetCubeResource* InResource = static_cast<FTextureRenderTargetCubeResource*>(GetResource());
 		ENQUEUE_RENDER_COMMAND(UpdateResourceImmediate)(
 			[InResource, bClearRenderTarget](FRHICommandListImmediate& RHICmdList)
 			{
diff --git a/Engine/Source/Runtime/Engine/Private/TextureRenderTargetVolume.cpp b/Engine/Source/Runtime/Engine/Private/TextureRenderTargetVolume.cpp
--- a/Engine/Source/Runtime/Engine/Private/TextureRenderTargetVolume.cpp
+++ b/Engine/Source/Runtime/Engine/Private/TextureRenderTargetVolume.cpp
@@ -63,9 +63,9 @@ void UTextureRenderTargetVolume::InitAutoFormat(uint32 InSizeX, uint32 InSizeY,
 
 void UTextureRenderTargetVolume::UpdateResourceImmediate(bool bClearRenderTarget/*=true*/)
 {
-	if (Resource)
+	if (GetResource())
 	{
-		FTextureRenderTargetVolumeResource* InResource = static_cast<FTextureRenderTargetVolumeResource*>(Resource);
+		FTextureRenderTargetVolumeResource* InResource = static_cast<FTextureRenderTargetVolumeResource*>(GetResource());
 		ENQUEUE_RENDER_COMMAND(UpdateResourceImmediate)(
 			[InResource, bClearRenderTarget](FRHICommandListImmediate& RHICmdList)
 			{
diff --git a/Engine/Source/Runtime/Engine/Private/UnrealEngine.cpp b/Engine/Source/Runtime/Engine/Private/UnrealEngine.cpp
--- a/Engine/Source/Runtime/Engine/Private/UnrealEngine.cpp
+++ b/Engine/Source/Runtime/Engine/Private/UnrealEngine.cpp
@@ -6191,7 +6191,7 @@ bool UEngine::HandleRemoteTextureStatsCommand( const TCHAR* Cmd, FOutputDevice&
 		float LastTimeRendered = FLT_MAX;
 		uint32 NumUses = 0;
 		int32 LODBias = Texture->GetCachedLODBias();
-		FTexture* Resource = Texture->Resource; 
+		FTexture* Resource = Texture->GetResource();
 
 		if(Resource)
 		{
@@ -10210,7 +10210,7 @@ static void DrawProperty(UCanvas* CanvasObject, UObject* Obj, const FDebugDispla
 			}
 		}
 		Canvas->DrawTile(  X, Y, XL2 + 1, YL * FMath::Max<int>(TextLines.Num(), 1), 0, 0, CanvasObject->DefaultTexture->GetSizeX(), CanvasObject->DefaultTexture->GetSizeY(),
-			FLinearColor(0.5f, 0.5f, 0.5f, 0.5f), CanvasObject->DefaultTexture->Resource );
+			FLinearColor(0.5f, 0.5f, 0.5f, 0.5f), CanvasObject->DefaultTexture->GetResource());
 		if( bDrawPropName )
 		{
 			bDrawPropName = false;
diff --git a/Engine/Source/Runtime/Engine/Private/UserInterface/Canvas.cpp b/Engine/Source/Runtime/Engine/Private/UserInterface/Canvas.cpp
--- a/Engine/Source/Runtime/Engine/Private/UserInterface/Canvas.cpp
+++ b/Engine/Source/Runtime/Engine/Private/UserInterface/Canvas.cpp
@@ -79,9 +79,9 @@ FTexture* GetTextureForCanvasItem(UTexture* RenderTexture)
 	// Only assign a valid texture for the canvas item. Proxy textures, from the UTexture implementation, are not valid here.
 	// This is because the proxy textures get destroyed on the render thread as soon as the build completes, which could happen
 	// before the canvas item renders.
-	if (RenderTexture && RenderTexture->Resource && !RenderTexture->Resource->IsProxy())
+	if (RenderTexture && RenderTexture->GetResource() && !RenderTexture->GetResource()->IsProxy())
 	{
-		return RenderTexture->Resource;
+		return RenderTexture->GetResource();
 	}
 	else
 	{
@@ -2078,7 +2078,7 @@ void UCanvas::K2_DrawBorder(UTexture* BorderTexture, UTexture* BackgroundTexture
 {
 	if (ScreenSize.X > 0.0f && ScreenSize.Y > 0.0f && BorderTexture && BackgroundTexture && LeftBorderTexture && RightBorderTexture && TopBorderTexture && BottomBorderTexture && Canvas)
 	{
-		FCanvasBorderItem BorderItem(ScreenPosition, BorderTexture->Resource, BackgroundTexture->Resource, LeftBorderTexture->Resource, RightBorderTexture->Resource, TopBorderTexture->Resource, BottomBorderTexture->Resource, ScreenSize, RenderColor);
+		FCanvasBorderItem BorderItem(ScreenPosition, BorderTexture->GetResource(), BackgroundTexture->GetResource(), LeftBorderTexture->GetResource(), RightBorderTexture->GetResource(), TopBorderTexture->GetResource(), BottomBorderTexture->GetResource(), ScreenSize, RenderColor);
 		BorderItem.BorderScale = BorderScale;
 		BorderItem.BackgroundScale = BackgroundScale;
 		BorderItem.BorderUV0 = CoordinatePosition;
diff --git a/Engine/Source/Runtime/Engine/Private/UserInterface/CanvasItem.cpp b/Engine/Source/Runtime/Engine/Private/UserInterface/CanvasItem.cpp
--- a/Engine/Source/Runtime/Engine/Private/UserInterface/CanvasItem.cpp
+++ b/Engine/Source/Runtime/Engine/Private/UserInterface/CanvasItem.cpp
@@ -234,7 +234,7 @@ void FCanvasItemTestbed::Draw( class FViewport* Viewport, class FCanvas* Canvas
 	if( ( SelectedTexture ) && ( bShowTiles == true ))
 	{
 		// Plain tex
-		FCanvasTileItem TileItem( FVector2D( 128.0f,128.0f ), SelectedTexture->Resource, FLinearColor::White );
+		FCanvasTileItem TileItem( FVector2D( 128.0f,128.0f ), SelectedTexture->GetResource(), FLinearColor::White );
 		TileItem.Draw( Canvas );
 		TileItem.Size = FVector2D( 32.0f,32.0f );
 		TileItem.Position = FVector2D( 16.0f,16.0f );
@@ -1051,24 +1051,24 @@ void FCanvasTextItem::DrawStringInternal_OfflineCache( FCanvas* InCanvas, const
 
 		if( Font->Textures.IsValidIndex(Char.TextureIndex) && 
 			(Tex=Font->Textures[Char.TextureIndex])!=nullptr && 
-			Tex->Resource != nullptr )
+			Tex->GetResource() != nullptr )
 		{
-			if( LastTexture != Tex->Resource || BatchedElements == nullptr )
+			if( LastTexture != Tex->GetResource() || BatchedElements == nullptr )
 			{
 				FBatchedElementParameters* BatchedElementParams = nullptr;
-				BatchedElements = InCanvas->GetBatchedElements(FCanvas::ET_Triangle, BatchedElementParams, Tex->Resource, BlendMode, FontRenderInfo.GlowInfo, false);
+				BatchedElements = InCanvas->GetBatchedElements(FCanvas::ET_Triangle, BatchedElementParams, Tex->GetResource(), BlendMode, FontRenderInfo.GlowInfo, false);
 				check(BatchedElements != nullptr);
 				// Trade-off to use memory for performance by pre-allocating more reserved space 
 				// for the triangles/vertices of the batched elements used to render the text tiles
 				// Only reserve initial batch, allow growth afterwards in case there are multiple repeated calls.
 				// Reserving exactly the added amount each time would essentially force an alloc each time on subsequent calls.
-				BatchedElements->ReserveTriangles(TextLen*2,Tex->Resource,BlendMode);
+				BatchedElements->ReserveTriangles(TextLen*2,Tex->GetResource(),BlendMode);
 				BatchedElements->ReserveVertices(TextLen*4);
 
 				InvTextureSize.X = 1.0f / Tex->GetSurfaceWidth();
 				InvTextureSize.Y = 1.0f / Tex->GetSurfaceHeight();
 			}
-			LastTexture = Tex->Resource;
+			LastTexture = Tex->GetResource();
 
 			const float X		= CurrentPos.X + DrawPos.X;
 			const float Y		= CurrentPos.Y + DrawPos.Y + Char.VerticalOffset * Scale.Y;
@@ -1105,8 +1105,8 @@ void FCanvasTextItem::DrawStringInternal_OfflineCache( FCanvas* InCanvas, const
 				InColor,
 				HitProxyId);
 
-			BatchedElements->AddTriangle(V00, V10, V11, Tex->Resource, BlendMode, FontRenderInfo.GlowInfo);
-			BatchedElements->AddTriangle(V00, V11, V01, Tex->Resource, BlendMode, FontRenderInfo.GlowInfo);
+			BatchedElements->AddTriangle(V00, V10, V11, Tex->GetResource(), BlendMode, FontRenderInfo.GlowInfo);
+			BatchedElements->AddTriangle(V00, V11, V01, Tex->GetResource(), BlendMode, FontRenderInfo.GlowInfo);
 
 			// if we have another non-whitespace character to render, add the font's kerning.
 			if ( Chars[i+1] && !FChar::IsWhitespace(Chars[i+1]) )
diff --git a/Engine/Source/Runtime/Engine/Private/UserInterface/Console.cpp b/Engine/Source/Runtime/Engine/Private/UserInterface/Console.cpp
--- a/Engine/Source/Runtime/Engine/Private/UserInterface/Console.cpp
+++ b/Engine/Source/Runtime/Engine/Private/UserInterface/Console.cpp
@@ -1282,7 +1282,7 @@ void UConsole::PostRender_Console_Open(UCanvas* Canvas)
 	// Background
 	FLinearColor BackgroundColor = ConsoleDefs::AutocompleteBackgroundColor.ReinterpretAsLinear();
 	BackgroundColor.A = ConsoleSettings->BackgroundOpacityPercentage / 100.0f;
-	FCanvasTileItem ConsoleTile(FVector2D(LeftPos, 0.0f), DefaultTexture_Black->Resource, FVector2D(ClipX, Height + TopPos - yl), FVector2D(0.0f, 0.0f), FVector2D(1.0f, 1.0f), BackgroundColor);
+	FCanvasTileItem ConsoleTile(FVector2D(LeftPos, 0.0f), DefaultTexture_Black->GetResource(), FVector2D(ClipX, Height + TopPos - yl), FVector2D(0.0f, 0.0f), FVector2D(1.0f, 1.0f), BackgroundColor);
 
 	// Preserve alpha to allow single-pass composite
 	ConsoleTile.BlendMode = SE_BLEND_AlphaBlend;
@@ -1432,7 +1432,7 @@ void UConsole::PostRender_InputLine(UCanvas* Canvas, FIntPoint UserInputLinePos)
 	// Background
 	FLinearColor BackgroundColor = ConsoleDefs::AutocompleteBackgroundColor.ReinterpretAsLinear();
 	BackgroundColor.A = ConsoleSettings->BackgroundOpacityPercentage / 100.0f;
-	FCanvasTileItem ConsoleTile(FVector2D(UserInputLinePos.X, UserInputLinePos.Y - 6 - yl), DefaultTexture_Black->Resource, FVector2D(ClipX, yl + 6), FVector2D(0.0f, 0.0f), FVector2D(1.0f, 1.0f), BackgroundColor);
+	FCanvasTileItem ConsoleTile(FVector2D(UserInputLinePos.X, UserInputLinePos.Y - 6 - yl), DefaultTexture_Black->GetResource(), FVector2D(ClipX, yl + 6), FVector2D(0.0f, 0.0f), FVector2D(1.0f, 1.0f), BackgroundColor);
 
 	// Preserve alpha to allow single-pass composite
 	ConsoleTile.BlendMode = SE_BLEND_AlphaBlend;
@@ -1441,7 +1441,7 @@ void UConsole::PostRender_InputLine(UCanvas* Canvas, FIntPoint UserInputLinePos)
 
 	// Separator line
 	ConsoleTile.SetColor(ConsoleDefs::BorderColor);
-	ConsoleTile.Texture = DefaultTexture_White->Resource;
+	ConsoleTile.Texture = DefaultTexture_White->GetResource();
 	ConsoleTile.Size = FVector2D(ClipX, 2.0f);
 	Canvas->DrawItem(ConsoleTile);
 
@@ -1478,7 +1478,7 @@ void UConsole::PostRender_InputLine(UCanvas* Canvas, FIntPoint UserInputLinePos)
 		FLinearColor AutoCompleteBackgroundColor = ConsoleDefs::AutocompleteBackgroundColor;
 		AutoCompleteBackgroundColor.A = ConsoleSettings->BackgroundOpacityPercentage / 100.0f;
 		ConsoleTile.SetColor(AutoCompleteBackgroundColor);
-		ConsoleTile.Texture = DefaultTexture_White->Resource;
+		ConsoleTile.Texture = DefaultTexture_White->GetResource();
 
 		// wasteful memory allocations but when typing in a console command this is fine
 		TArray<const FAutoCompleteCommand*> AutoCompleteElements;
diff --git a/Engine/Source/Runtime/Engine/Private/VT/VirtualTextureScalability.cpp b/Engine/Source/Runtime/Engine/Private/VT/VirtualTextureScalability.cpp
--- a/Engine/Source/Runtime/Engine/Private/VT/VirtualTextureScalability.cpp
+++ b/Engine/Source/Runtime/Engine/Private/VT/VirtualTextureScalability.cpp
@@ -180,7 +180,7 @@ namespace VirtualTextureScalability
 				if (It->IsCurrentlyVirtualTextured())
 				{
 					ReleasedVirtualTextures.Add(*It);
-					BeginReleaseResource(It->Resource);
+					BeginReleaseResource(It->GetResource());
 				}
 			}
 
@@ -194,7 +194,7 @@ namespace VirtualTextureScalability
 			// Reinit streaming virtual textures
 			for (UTexture2D* Texture : ReleasedVirtualTextures)
 			{
-				BeginInitResource(Texture->Resource);
+				BeginInitResource(Texture->GetResource());
 			}
 
 			// Reinit runtime virtual textures
diff --git a/Engine/Source/Runtime/Engine/Private/VectorField.cpp b/Engine/Source/Runtime/Engine/Private/VectorField.cpp
--- a/Engine/Source/Runtime/Engine/Private/VectorField.cpp
+++ b/Engine/Source/Runtime/Engine/Private/VectorField.cpp
@@ -1128,7 +1128,7 @@ public:
 	{
 		check(IsInRenderingThread());
 
-		if (GetFeatureLevel() == ERHIFeatureLevel::SM5 && AnimatedVectorField && AnimatedVectorField->Texture && AnimatedVectorField->Texture->Resource)
+		if (GetFeatureLevel() == ERHIFeatureLevel::SM5 && AnimatedVectorField && AnimatedVectorField->Texture && AnimatedVectorField->Texture->GetResource())
 		{
 			SCOPED_DRAW_EVENT(RHICmdList, AnimateVectorField);
 
@@ -1187,7 +1187,7 @@ public:
 			CompositeCS->SetParameters(
 				RHICmdList,
 				UniformBuffer,
-				AnimatedVectorField->Texture->Resource->TextureRHI,
+				AnimatedVectorField->Texture->GetResource()->TextureRHI,
 				NoiseVolumeTextureRHI );
 			DispatchComputeShader(
 				RHICmdList,
diff --git a/Engine/Source/Runtime/Engine/Private/VolumeTexture.cpp b/Engine/Source/Runtime/Engine/Private/VolumeTexture.cpp
--- a/Engine/Source/Runtime/Engine/Private/VolumeTexture.cpp
+++ b/Engine/Source/Runtime/Engine/Private/VolumeTexture.cpp
@@ -379,7 +379,7 @@ void UVolumeTexture::UpdateMipGenSettings()
 
 bool UVolumeTexture::StreamOut(int32 NewMipCount)
 {
-	FTexture3DResource* Texture3DResource = Resource ? Resource->GetTexture3DResource() : nullptr;
+	FTexture3DResource* Texture3DResource = GetResource() ? GetResource()->GetTexture3DResource() : nullptr;
 	if (!HasPendingInitOrStreaming() && CachedSRRState.StreamOut(NewMipCount) && ensure(Texture3DResource))
 	{
 		FTextureMipAllocator* MipAllocator = nullptr;
@@ -395,7 +395,7 @@ bool UVolumeTexture::StreamOut(int32 NewMipCount)
 
 bool UVolumeTexture::StreamIn(int32 NewMipCount, bool bHighPrio)
 {
-	FTexture3DResource* Texture3DResource = Resource ? Resource->GetTexture3DResource() : nullptr;
+	FTexture3DResource* Texture3DResource = GetResource() ? GetResource()->GetTexture3DResource() : nullptr;
 	if (!HasPendingInitOrStreaming() && CachedSRRState.StreamIn(NewMipCount) && ensure(Texture3DResource))
 	{
 		FTextureMipDataProvider* CustomMipDataProvider = nullptr;
diff --git a/Engine/Source/Runtime/Engine/Public/SceneManagement.h b/Engine/Source/Runtime/Engine/Public/SceneManagement.h
--- a/Engine/Source/Runtime/Engine/Public/SceneManagement.h
+++ b/Engine/Source/Runtime/Engine/Public/SceneManagement.h
@@ -1460,7 +1460,7 @@ public:
 	inline float GetLightFunctionFadeDistance() const { return LightFunctionFadeDistance; }
 	inline float GetLightFunctionDisabledBrightness() const { return LightFunctionDisabledBrightness; }
 	inline UTextureLightProfile* GetIESTexture() const { return IESTexture; }
-	inline FTexture* GetIESTextureResource() const { return IESTexture ? IESTexture->Resource : nullptr; }
+	inline FTexture* GetIESTextureResource() const { return IESTexture ? IESTexture->GetResource() : nullptr; }
 	inline const FMaterialRenderProxy* GetLightFunctionMaterial() const { return LightFunctionMaterial; }
 	inline bool IsMovable() const { return bMovable; }
 	inline bool HasStaticLighting() const { return bStaticLighting; }
diff --git a/Engine/Source/Runtime/HeadMountedDisplay/Private/DefaultSpectatorScreenController.cpp b/Engine/Source/Runtime/HeadMountedDisplay/Private/DefaultSpectatorScreenController.cpp
--- a/Engine/Source/Runtime/HeadMountedDisplay/Private/DefaultSpectatorScreenController.cpp
+++ b/Engine/Source/Runtime/HeadMountedDisplay/Private/DefaultSpectatorScreenController.cpp
@@ -103,7 +103,7 @@ void FDefaultSpectatorScreenController::SetSpectatorScreenTextureRenderCommand(U
 	}
 
 	FTexture2DRHIRef Texture2DRHIRef;
-	FTextureResource* TextureResource = SrcTexture->Resource;
+	FTextureResource* TextureResource = SrcTexture->GetResource();
 	if (TextureResource && TextureResource->TextureRHI)
 	{
 		Texture2DRHIRef = TextureResource->TextureRHI->GetTexture2D();
diff --git a/Engine/Source/Runtime/HeadMountedDisplay/Private/XRLoadingScreenFunctionLibrary.cpp b/Engine/Source/Runtime/HeadMountedDisplay/Private/XRLoadingScreenFunctionLibrary.cpp
--- a/Engine/Source/Runtime/HeadMountedDisplay/Private/XRLoadingScreenFunctionLibrary.cpp
+++ b/Engine/Source/Runtime/HeadMountedDisplay/Private/XRLoadingScreenFunctionLibrary.cpp
@@ -25,7 +25,7 @@ UXRLoadingScreenFunctionLibrary::UXRLoadingScreenFunctionLibrary(const FObjectIn
 void UXRLoadingScreenFunctionLibrary::SetLoadingScreen(class UTexture* Texture, FVector2D Scale, FVector Offset, bool bShowLoadingMovie, bool bShowOnSet)
 {
 	IXRLoadingScreen* LoadingSrcreen = GetLoadingScreen();
-	if (LoadingSrcreen && Texture && Texture->Resource)
+	if (LoadingSrcreen && Texture && Texture->GetResource())
 	{
 		LoadingSrcreen->ClearSplashes();
 		const bool bIsExternal = Texture->GetMaterialType() == MCT_TextureExternal;
@@ -34,7 +34,7 @@ void UXRLoadingScreenFunctionLibrary::SetLoadingScreen(class UTexture* Texture,
 		Splash.QuadSize = Scale;
 		Splash.bIsDynamic = bShowLoadingMovie || bIsExternal;
 		Splash.bIsExternal = bIsExternal;
-		Splash.Texture = Texture->Resource->TextureRHI;
+		Splash.Texture = Texture->GetResource()->TextureRHI;
 		LoadingSrcreen->AddSplash(Splash);
 
 		if (bShowOnSet)
@@ -57,7 +57,7 @@ void UXRLoadingScreenFunctionLibrary::ClearLoadingScreenSplashes()
 void UXRLoadingScreenFunctionLibrary::AddLoadingScreenSplash(class UTexture* Texture, FVector Translation, FRotator Rotation, FVector2D Size, FRotator DeltaRotation, bool bClearBeforeAdd)
 {
 	IXRLoadingScreen* LoadingSrcreen = GetLoadingScreen();
-	if (LoadingSrcreen && Texture && Texture->Resource)
+	if (LoadingSrcreen && Texture && Texture->GetResource())
 	{
 		if (bClearBeforeAdd)
 		{
@@ -65,7 +65,7 @@ void UXRLoadingScreenFunctionLibrary::AddLoadingScreenSplash(class UTexture* Tex
 		}
 
 		IXRLoadingScreen::FSplashDesc Splash;
-		Splash.Texture = Texture->Resource->TextureRHI;
+		Splash.Texture = Texture->GetResource()->TextureRHI;
 		Splash.QuadSize = Size;
 		Splash.Transform = FTransform(Rotation, Translation);
 		Splash.DeltaRotation = FQuat(DeltaRotation);
diff --git a/Engine/Source/Runtime/ImageWriteQueue/Private/ImageWriteBlueprintLibrary.cpp b/Engine/Source/Runtime/ImageWriteQueue/Private/ImageWriteBlueprintLibrary.cpp
--- a/Engine/Source/Runtime/ImageWriteQueue/Private/ImageWriteBlueprintLibrary.cpp
+++ b/Engine/Source/Runtime/ImageWriteQueue/Private/ImageWriteBlueprintLibrary.cpp
@@ -55,7 +55,7 @@ bool UImageWriteBlueprintLibrary::ResolvePixelData(UTexture* InTexture, const FO
 		break;
 	}
 
-	FTextureResource* TextureResource = InTexture->Resource;
+	FTextureResource* TextureResource = InTexture->GetResource();
 	ENQUEUE_RENDER_COMMAND(ResolvePixelData)(
 		[TextureResource, OnPixelsReady](FRHICommandListImmediate& RHICmdList)
 		{
@@ -157,14 +157,14 @@ void UImageWriteBlueprintLibrary::ExportToDisk(UTexture* InTexture, const FStrin
 
 	// In the case of an error, we always call the error callbacks in a latent manner to ensure that we always trigger the calback outside of this 
 	// function - this ensures the calling context deterministic for whatever is handling the completion.
-	if (!InTexture || !InTexture->Resource || !InTexture->Resource->TextureRHI)
+	if (!InTexture || !InTexture->GetResource() || !InTexture->GetResource()->TextureRHI)
 	{
 		FFrame::KismetExecutionMessage(TEXT("Invalid texture supplied."), ELogVerbosity::Error);
 		AsyncTask(ENamedThreads::GameThread, [OnCompleteWrapper] { OnCompleteWrapper(false); });
 		return;
 	}
 
-	FTexture2DRHIRef Texture2D = InTexture->Resource->TextureRHI->GetTexture2D();
+	FTexture2DRHIRef Texture2D = InTexture->GetResource()->TextureRHI->GetTexture2D();
 	if (!Texture2D)
 	{
 		FFrame::KismetExecutionMessage(TEXT("Invalid texture supplied."), ELogVerbosity::Error);
diff --git a/Engine/Source/Runtime/Landscape/Private/LandscapeEdit.cpp b/Engine/Source/Runtime/Landscape/Private/LandscapeEdit.cpp
--- a/Engine/Source/Runtime/Landscape/Private/LandscapeEdit.cpp
+++ b/Engine/Source/Runtime/Landscape/Private/LandscapeEdit.cpp
@@ -7249,7 +7249,7 @@ TArray<FLinearColor> ALandscapeProxy::SampleRTData(UTextureRenderTarget2D* InRen
 		FMessageLog("Blueprint").Warning(LOCTEXT("SampleRTData_InvalidRenderTarget", "SampleRTData: Render Target must be non-null."));
 		return { FLinearColor(0,0,0,0) };
 	}
-	else if (!InRenderTarget->Resource)
+	else if (!InRenderTarget->GetResource())
 	{
 		FMessageLog("Blueprint").Warning(LOCTEXT("SampleRTData_ReleasedRenderTarget", "SampleRTData: Render Target has been released."));
 		return { FLinearColor(0,0,0,0) };
@@ -7327,7 +7327,7 @@ bool ALandscapeProxy::LandscapeImportHeightmapFromRenderTarget(UTextureRenderTar
 		return false;
 	}
 
-	if (InRenderTarget == nullptr || InRenderTarget->Resource == nullptr)
+	if (InRenderTarget == nullptr || InRenderTarget->GetResource() == nullptr)
 	{
 		FMessageLog("Blueprint").Error(LOCTEXT("LandscapeImportHeightmapFromRenderTarget_InvalidRT", "LandscapeImportHeightmapFromRenderTarget: Render Target must be non null and not released."));
 		return false;
diff --git a/Engine/Source/Runtime/Landscape/Private/LandscapeEditLayers.cpp b/Engine/Source/Runtime/Landscape/Private/LandscapeEditLayers.cpp
--- a/Engine/Source/Runtime/Landscape/Private/LandscapeEditLayers.cpp
+++ b/Engine/Source/Runtime/Landscape/Private/LandscapeEditLayers.cpp
@@ -512,8 +512,8 @@ public:
 
 	void SetParameters(FRHICommandList& RHICmdList, const FLandscapeLayersHeightmapShaderParameters& InParams)
 	{
-		SetTextureParameter(RHICmdList, RHICmdList.GetBoundPixelShader(), ReadTexture1Param, ReadTexture1SamplerParam, TStaticSamplerState<SF_Point, AM_Clamp, AM_Clamp, AM_Clamp>::GetRHI(), InParams.ReadHeightmap1->Resource->TextureRHI);
-		SetTextureParameter(RHICmdList, RHICmdList.GetBoundPixelShader(), ReadTexture2Param, ReadTexture2SamplerParam, TStaticSamplerState<SF_Point, AM_Clamp, AM_Clamp, AM_Clamp>::GetRHI(), InParams.ReadHeightmap2 != nullptr ? InParams.ReadHeightmap2->Resource->TextureRHI : GWhiteTexture->TextureRHI);
+		SetTextureParameter(RHICmdList, RHICmdList.GetBoundPixelShader(), ReadTexture1Param, ReadTexture1SamplerParam, TStaticSamplerState<SF_Point, AM_Clamp, AM_Clamp, AM_Clamp>::GetRHI(), InParams.ReadHeightmap1->GetResource()->TextureRHI);
+		SetTextureParameter(RHICmdList, RHICmdList.GetBoundPixelShader(), ReadTexture2Param, ReadTexture2SamplerParam, TStaticSamplerState<SF_Point, AM_Clamp, AM_Clamp, AM_Clamp>::GetRHI(), InParams.ReadHeightmap2 != nullptr ? InParams.ReadHeightmap2->GetResource()->TextureRHI : GWhiteTexture->TextureRHI);
 
 		FVector4 LayerInfo(InParams.LayerAlpha, InParams.LayerVisible ? 1.0f : 0.0f, InParams.LayerBlendMode == LSBM_AlphaBlend ? 1.0f : 0.f, 0.f);
 		FVector4 OutputConfig(InParams.ApplyLayerModifiers ? 1.0f : 0.0f, 0.0f /*unused*/, InParams.ReadHeightmap2 ? 1.0f : 0.0f, InParams.GenerateNormals ? 1.0f : 0.0f);
@@ -569,7 +569,7 @@ public:
 
 	void SetParameters(FRHICommandList& RHICmdList, const FLandscapeLayersHeightmapShaderParameters& InParams)
 	{
-		SetTextureParameter(RHICmdList, RHICmdList.GetBoundPixelShader(), ReadTexture1Param, ReadTexture1SamplerParam, TStaticSamplerState<SF_Point, AM_Clamp, AM_Clamp, AM_Clamp>::GetRHI(), InParams.ReadHeightmap1->Resource->TextureRHI);
+		SetTextureParameter(RHICmdList, RHICmdList.GetBoundPixelShader(), ReadTexture1Param, ReadTexture1SamplerParam, TStaticSamplerState<SF_Point, AM_Clamp, AM_Clamp, AM_Clamp>::GetRHI(), InParams.ReadHeightmap1->GetResource()->TextureRHI);
 
 		SetShaderValue(RHICmdList, RHICmdList.GetBoundPixelShader(), CurrentMipSizeParam, FVector2D(InParams.CurrentMipSize.X, InParams.CurrentMipSize.Y));
 		SetShaderValue(RHICmdList, RHICmdList.GetBoundPixelShader(), ParentMipSizeParam, FVector2D(InParams.ParentMipSize.X, InParams.ParentMipSize.Y));
@@ -646,8 +646,8 @@ public:
 
 	void SetParameters(FRHICommandList& RHICmdList, const FLandscapeLayersWeightmapShaderParameters& InParams)
 	{
-		SetTextureParameter(RHICmdList, RHICmdList.GetBoundPixelShader(), ReadTexture1Param, ReadTexture1SamplerParam, TStaticSamplerState<SF_Point, AM_Clamp, AM_Clamp, AM_Clamp>::GetRHI(), InParams.ReadWeightmap1->Resource->TextureRHI);
-		SetTextureParameter(RHICmdList, RHICmdList.GetBoundPixelShader(), ReadTexture2Param, ReadTexture2SamplerParam, TStaticSamplerState<SF_Point, AM_Clamp, AM_Clamp, AM_Clamp>::GetRHI(), InParams.ReadWeightmap2 != nullptr ? InParams.ReadWeightmap2->Resource->TextureRHI : GWhiteTexture->TextureRHI);
+		SetTextureParameter(RHICmdList, RHICmdList.GetBoundPixelShader(), ReadTexture1Param, ReadTexture1SamplerParam, TStaticSamplerState<SF_Point, AM_Clamp, AM_Clamp, AM_Clamp>::GetRHI(), InParams.ReadWeightmap1->GetResource()->TextureRHI);
+		SetTextureParameter(RHICmdList, RHICmdList.GetBoundPixelShader(), ReadTexture2Param, ReadTexture2SamplerParam, TStaticSamplerState<SF_Point, AM_Clamp, AM_Clamp, AM_Clamp>::GetRHI(), InParams.ReadWeightmap2 != nullptr ? InParams.ReadWeightmap2->GetResource()->TextureRHI : GWhiteTexture->TextureRHI);
 
 		FVector4 LayerInfo(InParams.LayerAlpha, InParams.LayerVisible ? 1.0f : 0.0f, InParams.LayerBlendMode == LSBM_AlphaBlend ? 1.0f : 0.f, 0.f);
 		FVector4 OutputConfig(InParams.ApplyLayerModifiers ? 1.0f : 0.0f, InParams.OutputAsSubstractive ? 1.0f : 0.0f, InParams.ReadWeightmap2 != nullptr ? 1.0f : 0.0f, InParams.OutputAsNormalized ? 1.0f : 0.0f);
@@ -698,7 +698,7 @@ public:
 
 	void SetParameters(FRHICommandList& RHICmdList, const FLandscapeLayersWeightmapShaderParameters& InParams)
 	{
-		SetTextureParameter(RHICmdList, RHICmdList.GetBoundPixelShader(), ReadTexture1Param, ReadTexture1SamplerParam, TStaticSamplerState<SF_Point, AM_Clamp, AM_Clamp, AM_Clamp>::GetRHI(), InParams.ReadWeightmap1->Resource->TextureRHI);
+		SetTextureParameter(RHICmdList, RHICmdList.GetBoundPixelShader(), ReadTexture1Param, ReadTexture1SamplerParam, TStaticSamplerState<SF_Point, AM_Clamp, AM_Clamp, AM_Clamp>::GetRHI(), InParams.ReadWeightmap1->GetResource()->TextureRHI);
 
 		SetShaderValue(RHICmdList, RHICmdList.GetBoundPixelShader(), CurrentMipSizeParam, FVector2D(InParams.CurrentMipSize.X, InParams.CurrentMipSize.Y));
 		SetShaderValue(RHICmdList, RHICmdList.GetBoundPixelShader(), ParentMipSizeParam, FVector2D(InParams.ParentMipSize.X, InParams.ParentMipSize.Y));
@@ -1121,12 +1121,12 @@ struct FLandscapeLayersCopyTextureParams
 		if (InSourceTexture != nullptr)
 		{
 			SourceResourceDebugName = InSourceTexture->GetName();
-			SourceResource = InSourceTexture->Resource;
+			SourceResource = InSourceTexture->GetResource();
 		}
 		if (InDestTexture != nullptr)
 		{
 			DestResourceDebugName = InDestTexture->GetName();
-			DestResource = InDestTexture->Resource;
+			DestResource = InDestTexture->GetResource();
 		}
 	}
 
@@ -3470,10 +3470,10 @@ void ALandscape::PrepareComponentDataToExtractMaterialLayersCS(const TArray<ULan
 				FVector2D SourcePositionOffset(FMath::RoundToInt(ComponentSectionBase.X / LocalComponentSizeQuad), FMath::RoundToInt(ComponentSectionBase.Y / LocalComponentSizeQuad));
 				FIntPoint SourceComponentVertexPosition = FIntPoint(SourcePositionOffset.X * LocalComponentSizeVerts, SourcePositionOffset.Y * LocalComponentSizeVerts);
 
-				FLandscapeLayersCopyTextureParams& CopyTextureParams = DeferredCopyTextures.Add_GetRef(FLandscapeLayersCopyTextureParams(*LayerWeightmap->GetName(), LayerWeightmap->Resource, FString::Printf(TEXT("%s WeightmapScratchTexture"), *InLayer.Name.ToString()), InOutTextureData));
+				FLandscapeLayersCopyTextureParams& CopyTextureParams = DeferredCopyTextures.Add_GetRef(FLandscapeLayersCopyTextureParams(*LayerWeightmap->GetName(), LayerWeightmap->GetResource(), FString::Printf(TEXT("%s WeightmapScratchTexture"), *InLayer.Name.ToString()), InOutTextureData));
 				// Only copy the size that's actually needed : 
-				CopyTextureParams.CopySize.X = LayerWeightmap->Resource->GetSizeX();
-				CopyTextureParams.CopySize.Y = LayerWeightmap->Resource->GetSizeY();
+				CopyTextureParams.CopySize.X = LayerWeightmap->GetResource()->GetSizeX();
+				CopyTextureParams.CopySize.Y = LayerWeightmap->GetResource()->GetSizeY();
 				// Copy from the top-left corner of the weightmap to the composited texture's position
 				CopyTextureParams.DestPosition = SourceComponentVertexPosition;
 				PrintLayersDebugTextureResource(FString::Printf(TEXT("LS Weight: %s WeightmapScratchTexture %s"), *InLayer.Name.ToString(), TEXT("WeightmapScratchTextureResource")), InOutTextureData, 0, false);
@@ -4322,8 +4322,8 @@ int32 ALandscape::RegenerateLayersWeightmaps(FTextureToComponentHelper const& Ma
 						FIntPoint TextureTopLeftPositionInAtlas(WeightmapTextureOutputOffset[NextTextureIndexToProcess - StartTextureIndex].X, WeightmapTextureOutputOffset[NextTextureIndexToProcess - StartTextureIndex].Y);
 
 						int32 CurrentMip = 0;
-						const int32 TextureSizeX = WeightmapTexture->Resource->GetSizeX();
-						const int32 TextureSizeY = WeightmapTexture->Resource->GetSizeY();
+						const int32 TextureSizeX = WeightmapTexture->GetResource()->GetSizeX();
+						const int32 TextureSizeY = WeightmapTexture->GetResource()->GetSizeY();
 						for (int32 MipRTIndex = (int32)EWeightmapRTType::WeightmapRT_Mip0; MipRTIndex < (int32)EWeightmapRTType::WeightmapRT_Count; ++MipRTIndex)
 						{
 							CurrentRT = WeightmapRTList[MipRTIndex];
@@ -4333,7 +4333,7 @@ int32 ALandscape::RegenerateLayersWeightmaps(FTextureToComponentHelper const& Ma
 								SourceDebugName = CurrentRT->GetName();
 								DestDebugName = FString::Printf(TEXT("Weightmap Mip: %d"), CurrentMip);
 
-								FLandscapeLayersCopyTextureParams& CopyTextureParams = DeferredCopyTextures.Add_GetRef(FLandscapeLayersCopyTextureParams(SourceDebugName, CurrentRT->GameThread_GetRenderTargetResource(), DestDebugName, WeightmapTexture->Resource));
+								FLandscapeLayersCopyTextureParams& CopyTextureParams = DeferredCopyTextures.Add_GetRef(FLandscapeLayersCopyTextureParams(SourceDebugName, CurrentRT->GameThread_GetRenderTargetResource(), DestDebugName, WeightmapTexture->GetResource()));
 								// Only copy the size that's actually needed : 
 								CopyTextureParams.CopySize.X = TextureSizeX >> CurrentMip;
 								CopyTextureParams.CopySize.Y = TextureSizeY >> CurrentMip;
diff --git a/Engine/Source/Runtime/Landscape/Private/LandscapeEditReadback.cpp b/Engine/Source/Runtime/Landscape/Private/LandscapeEditReadback.cpp
--- a/Engine/Source/Runtime/Landscape/Private/LandscapeEditReadback.cpp
+++ b/Engine/Source/Runtime/Landscape/Private/LandscapeEditReadback.cpp
@@ -37,7 +37,7 @@ struct FLandscapeEditReadbackTaskImpl
 /** Initialize the read back task data that is written by game thread. */
 bool InitTask_GameThread(FLandscapeEditReadbackTaskImpl& Task, UTexture2D const* InTexture, FLandscapeEditLayerReadback::FReadbackContext&& InReadbackContext, uint32 InFrameId)
 {
-	Task.TextureResource = InTexture->Resource;
+	Task.TextureResource = InTexture->GetResource();
 	Task.ReadbackContext = MoveTemp(InReadbackContext);
 	Task.InitFrameId = InFrameId;
 	Task.Size = FIntPoint(InTexture->GetSizeX(), InTexture->GetSizeY());
diff --git a/Engine/Source/Runtime/Landscape/Private/LandscapeSplines.cpp b/Engine/Source/Runtime/Landscape/Private/LandscapeSplines.cpp
--- a/Engine/Source/Runtime/Landscape/Private/LandscapeSplines.cpp
+++ b/Engine/Source/Runtime/Landscape/Private/LandscapeSplines.cpp
@@ -284,11 +284,11 @@ public:
 							ControlPointSpriteLocation,
 							ControlPointSpriteScale,
 							ControlPointSpriteScale,
-							ControlPointSprite->Resource,
+							ControlPointSprite->GetResource(),
 							ControlPointSpriteColor,
 							GetDepthPriorityGroup(View),
-							0, ControlPointSprite->Resource->GetSizeX(),
-							0, ControlPointSprite->Resource->GetSizeY(),
+							0, ControlPointSprite->GetResource()->GetSizeX(),
+							0, ControlPointSprite->GetResource()->GetSizeY(),
 							SE_BLEND_Masked);
 					}
 
diff --git a/Engine/Source/Runtime/MediaAssets/Private/Assets/MediaTexture.cpp b/Engine/Source/Runtime/MediaAssets/Private/Assets/MediaTexture.cpp
--- a/Engine/Source/Runtime/MediaAssets/Private/Assets/MediaTexture.cpp
+++ b/Engine/Source/Runtime/MediaAssets/Private/Assets/MediaTexture.cpp
@@ -311,7 +311,7 @@ void UMediaTexture::PostEditChangeProperty(FPropertyChangedEvent& PropertyChange
 
 void UMediaTexture::TickResource(FTimespan Timecode)
 {
-	if (Resource == nullptr)
+	if (GetResource() == nullptr)
 	{
 		return;
 	}
@@ -425,7 +425,7 @@ void UMediaTexture::TickResource(FTimespan Timecode)
 	RenderParams.NumMips = NumMips;
 	
 	// redraw texture resource on render thread
-	FMediaTextureResource* ResourceParam = (FMediaTextureResource*)Resource;
+	FMediaTextureResource* ResourceParam = (FMediaTextureResource*)GetResource();
 	ENQUEUE_RENDER_COMMAND(MediaTextureResourceRender)(
 		[ResourceParam, RenderParams](FRHICommandListImmediate& RHICmdList)
 		{
@@ -455,7 +455,7 @@ void UMediaTexture::UpdatePlayerAndQueue()
 		// Player changed?
 		if (CurrentGuid != PlayerGuid)
 		{
-			if (FMediaTextureResource* MediaResource = static_cast<FMediaTextureResource*>(Resource))
+			if (FMediaTextureResource* MediaResource = static_cast<FMediaTextureResource*>(GetResource()))
 			{
 				MediaResource->FlushPendingData();
 			}
@@ -474,7 +474,7 @@ void UMediaTexture::UpdatePlayerAndQueue()
 			SampleQueue.Reset();
 			CurrentGuid = DefaultGuid;
 
-			if (FMediaTextureResource* MediaResource = static_cast<FMediaTextureResource*>(Resource))
+			if (FMediaTextureResource* MediaResource = static_cast<FMediaTextureResource*>(GetResource()))
 			{
 				MediaResource->FlushPendingData();
 			}
diff --git a/Engine/Source/Runtime/Renderer/Private/BlueNoise.h b/Engine/Source/Runtime/Renderer/Private/BlueNoise.h
--- a/Engine/Source/Runtime/Renderer/Private/BlueNoise.h
+++ b/Engine/Source/Runtime/Renderer/Private/BlueNoise.h
@@ -28,5 +28,5 @@ inline void InitializeBlueNoise(FBlueNoise& BlueNoise)
 		GEngine->BlueNoiseTexture->GetSizeY() / FMath::Max<int32>(1, GEngine->BlueNoiseTexture->GetSizeX())
 	);
 
-	BlueNoise.Texture = GEngine->BlueNoiseTexture->Resource->TextureRHI;
+	BlueNoise.Texture = GEngine->BlueNoiseTexture->GetResource()->TextureRHI;
 }
diff --git a/Engine/Source/Runtime/Renderer/Private/CompositionLighting/PostProcessAmbient.cpp b/Engine/Source/Runtime/Renderer/Private/CompositionLighting/PostProcessAmbient.cpp
--- a/Engine/Source/Runtime/Renderer/Private/CompositionLighting/PostProcessAmbient.cpp
+++ b/Engine/Source/Runtime/Renderer/Private/CompositionLighting/PostProcessAmbient.cpp
@@ -26,7 +26,7 @@ void SetupAmbientCubemapParameters(const FFinalPostProcessSettings::FCubemapEntr
 
 	// cubemap texture
 	{
-		FTexture* AmbientCubemapTexture = Entry.AmbientCubemap ? Entry.AmbientCubemap->Resource : (FTexture*)GBlackTextureCube;
+		FTexture* AmbientCubemapTexture = Entry.AmbientCubemap ? Entry.AmbientCubemap->GetResource() : (FTexture*)GBlackTextureCube;
 		OutParameters->AmbientCubemap = AmbientCubemapTexture->TextureRHI;
 		OutParameters->AmbientCubemapSampler = AmbientCubemapTexture->SamplerStateRHI;
 	}
diff --git a/Engine/Source/Runtime/Renderer/Private/FogRendering.cpp b/Engine/Source/Runtime/Renderer/Private/FogRendering.cpp
--- a/Engine/Source/Runtime/Renderer/Private/FogRendering.cpp
+++ b/Engine/Source/Runtime/Renderer/Private/FogRendering.cpp
@@ -64,7 +64,7 @@ void SetupFogUniformParameters(FRDGBuilder& GraphBuilder, const FViewInfo& View,
 
 		if (View.FogInscatteringColorCubemap)
 		{
-			Cubemap = View.FogInscatteringColorCubemap->Resource;
+			Cubemap = View.FogInscatteringColorCubemap->GetResource();
 		}
 
 		OutParameters.ExponentialFogParameters = View.ExponentialFogParameters;
diff --git a/Engine/Source/Runtime/Renderer/Private/GlobalDistanceField.cpp b/Engine/Source/Runtime/Renderer/Private/GlobalDistanceField.cpp
--- a/Engine/Source/Runtime/Renderer/Private/GlobalDistanceField.cpp
+++ b/Engine/Source/Runtime/Renderer/Private/GlobalDistanceField.cpp
@@ -1765,7 +1765,7 @@ void UpdateGlobalDistanceFieldVolume(
 									FHeightfieldComponentDescription NewComponentDescription(HeightfieldPrimitive->Proxy->GetLocalToWorld());
 									HeightfieldPrimitive->Proxy->GetHeightfieldRepresentation(HeightfieldTexture, DiffuseColorTexture, VisibilityTexture, NewComponentDescription);
 
-									if (HeightfieldTexture && HeightfieldTexture->Resource && HeightfieldTexture->Resource->TextureRHI)
+									if (HeightfieldTexture && HeightfieldTexture->GetResource() && HeightfieldTexture->GetResource()->TextureRHI)
 									{
 										const FIntPoint HeightfieldSize = NewComponentDescription.HeightfieldRect.Size();
 
@@ -1929,9 +1929,9 @@ void UpdateGlobalDistanceFieldVolume(
 									PassParameters->NumHeightfields = HeightfieldDescriptions.Num();
 									PassParameters->InfluenceRadius = ClipmapInfluenceRadius;
 									PassParameters->HeightfieldThickness = ClipmapVoxelSize.X * GGlobalDistanceFieldHeightFieldThicknessScale;
-									PassParameters->HeightfieldTexture = HeightfieldTexture->Resource->TextureRHI;
+									PassParameters->HeightfieldTexture = HeightfieldTexture->GetResource()->TextureRHI;
 									PassParameters->HeightfieldSampler = TStaticSamplerState<SF_Bilinear>::GetRHI();
-									PassParameters->VisibilityTexture = VisibilityTexture ? VisibilityTexture->Resource->TextureRHI : GBlackTexture->TextureRHI;
+									PassParameters->VisibilityTexture = VisibilityTexture ? VisibilityTexture->GetResource()->TextureRHI : GBlackTexture->TextureRHI;
 									PassParameters->VisibilitySampler = TStaticSamplerState<SF_Bilinear>::GetRHI();
 									PassParameters->HeightfieldDescriptions = GraphBuilder.CreateSRV(HeightfieldDescriptionBuffer, EPixelFormat::PF_A32B32G32R32F);
 
@@ -2220,9 +2220,9 @@ void UpdateGlobalDistanceFieldVolume(
 									PassParameters->NumHeightfields = HeightfieldDescriptions.Num();
 									PassParameters->InfluenceRadius = ClipmapInfluenceRadius;
 									PassParameters->HeightfieldThickness = ClipmapVoxelSize.X * GGlobalDistanceFieldHeightFieldThicknessScale;
-									PassParameters->HeightfieldTexture = HeightfieldTexture->Resource->TextureRHI;
+									PassParameters->HeightfieldTexture = HeightfieldTexture->GetResource()->TextureRHI;
 									PassParameters->HeightfieldSampler = TStaticSamplerState<SF_Bilinear>::GetRHI();
-									PassParameters->VisibilityTexture = VisibilityTexture ? VisibilityTexture->Resource->TextureRHI : GBlackTexture->TextureRHI;
+									PassParameters->VisibilityTexture = VisibilityTexture ? VisibilityTexture->GetResource()->TextureRHI : GBlackTexture->TextureRHI;
 									PassParameters->VisibilitySampler = TStaticSamplerState<SF_Bilinear>::GetRHI();
 									PassParameters->HeightfieldDescriptions = GraphBuilder.CreateSRV(HeightfieldDescriptionBuffer, EPixelFormat::PF_A32B32G32R32F);
 
diff --git a/Engine/Source/Runtime/Renderer/Private/HeightfieldLighting.cpp b/Engine/Source/Runtime/Renderer/Private/HeightfieldLighting.cpp
--- a/Engine/Source/Runtime/Renderer/Private/HeightfieldLighting.cpp
+++ b/Engine/Source/Runtime/Renderer/Private/HeightfieldLighting.cpp
@@ -413,7 +413,7 @@ void FHeightfieldLightingViewInfo::SetupVisibleHeightfields(const FViewInfo& Vie
 				FHeightfieldComponentDescription NewComponentDescription(HeightfieldPrimitive->Proxy->GetLocalToWorld());
 				HeightfieldPrimitive->Proxy->GetHeightfieldRepresentation(HeightfieldTexture, DiffuseColorTexture, VisibilityTexture, NewComponentDescription);
 
-				if (HeightfieldTexture && HeightfieldTexture->Resource->TextureRHI)
+				if (HeightfieldTexture && HeightfieldTexture->GetResource()->TextureRHI)
 				{
 					const FIntPoint HeightfieldSize = NewComponentDescription.HeightfieldRect.Size();
 
@@ -546,7 +546,7 @@ void FHeightfieldLightingViewInfo::SetupHeightfieldsForScene(const FScene& Scene
 			FHeightfieldComponentDescription NewComponentDescription(HeightfieldPrimitive->Proxy->GetLocalToWorld());
 			HeightfieldPrimitive->Proxy->GetHeightfieldRepresentation(HeightfieldTexture, DiffuseColorTexture, VisibilityTexture, NewComponentDescription);
 
-			if (HeightfieldTexture && HeightfieldTexture->Resource->TextureRHI)
+			if (HeightfieldTexture && HeightfieldTexture->GetResource()->TextureRHI)
 			{
 				const FIntPoint HeightfieldSize = NewComponentDescription.HeightfieldRect.Size();
 
diff --git a/Engine/Source/Runtime/Renderer/Private/HeightfieldLighting.h b/Engine/Source/Runtime/Renderer/Private/HeightfieldLighting.h
--- a/Engine/Source/Runtime/Renderer/Private/HeightfieldLighting.h
+++ b/Engine/Source/Runtime/Renderer/Private/HeightfieldLighting.h
@@ -189,11 +189,11 @@ public:
 	void Set(FRHICommandList& RHICmdList, const ShaderRHIParamRef ShaderRHI, UTexture2D* HeightfieldTextureValue, UTexture2D* DiffuseColorTextureValue, UTexture2D* VisibilityTextureValue)
 	{
 		//@todo - shouldn't filter the heightfield, it's packed
-		SetTextureParameter(RHICmdList, ShaderRHI, HeightfieldTexture, HeightfieldSampler, TStaticSamplerState<SF_Bilinear>::GetRHI(), HeightfieldTextureValue->Resource->TextureRHI);
+		SetTextureParameter(RHICmdList, ShaderRHI, HeightfieldTexture, HeightfieldSampler, TStaticSamplerState<SF_Bilinear>::GetRHI(), HeightfieldTextureValue->GetResource()->TextureRHI);
 
 		if (DiffuseColorTextureValue)
 		{
-			SetTextureParameter(RHICmdList, ShaderRHI, DiffuseColorTexture, DiffuseColorSampler, TStaticSamplerState<SF_Bilinear>::GetRHI(), DiffuseColorTextureValue->Resource->TextureRHI);
+			SetTextureParameter(RHICmdList, ShaderRHI, DiffuseColorTexture, DiffuseColorSampler, TStaticSamplerState<SF_Bilinear>::GetRHI(), DiffuseColorTextureValue->GetResource()->TextureRHI);
 		}
 		else
 		{
@@ -202,7 +202,7 @@ public:
 
 		if (VisibilityTextureValue)
 		{
-			SetTextureParameter(RHICmdList, ShaderRHI, VisibilityTexture, VisibilitySampler, TStaticSamplerState<SF_Bilinear>::GetRHI(), VisibilityTextureValue->Resource->TextureRHI);
+			SetTextureParameter(RHICmdList, ShaderRHI, VisibilityTexture, VisibilitySampler, TStaticSamplerState<SF_Bilinear>::GetRHI(), VisibilityTextureValue->GetResource()->TextureRHI);
 		}
 		else
 		{
diff --git a/Engine/Source/Runtime/Renderer/Private/LightMapDensityRendering.cpp b/Engine/Source/Runtime/Renderer/Private/LightMapDensityRendering.cpp
--- a/Engine/Source/Runtime/Renderer/Private/LightMapDensityRendering.cpp
+++ b/Engine/Source/Runtime/Renderer/Private/LightMapDensityRendering.cpp
@@ -38,7 +38,7 @@ void SetupLightmapDensityPassUniformBuffer(FRDGBuilder& GraphBuilder, ERHIFeatur
 {
 	SetupSceneTextureUniformParameters(GraphBuilder, FeatureLevel, ESceneTextureSetupMode::None, LightmapDensityPassParameters.SceneTextures);
 
-	LightmapDensityPassParameters.GridTexture = GEngine->LightMapDensityTexture->Resource->TextureRHI;
+	LightmapDensityPassParameters.GridTexture = GEngine->LightMapDensityTexture->GetResource()->TextureRHI;
 	LightmapDensityPassParameters.GridTextureSampler = TStaticSamplerState<SF_Bilinear, AM_Wrap, AM_Wrap, AM_Wrap>::GetRHI();
 
 	LightmapDensityPassParameters.LightMapDensity = FVector4(
diff --git a/Engine/Source/Runtime/Renderer/Private/PostProcess/PostProcessCombineLUTs.cpp b/Engine/Source/Runtime/Renderer/Private/PostProcess/PostProcessCombineLUTs.cpp
--- a/Engine/Source/Runtime/Renderer/Private/PostProcess/PostProcessCombineLUTs.cpp
+++ b/Engine/Source/Runtime/Renderer/Private/PostProcess/PostProcessCombineLUTs.cpp
@@ -312,7 +312,7 @@ uint32 GenerateFinalTable(const FFinalPostProcessSettings& Settings, const FText
 
 			{
 				UTexture* LUTTexture = Settings.ContributingLUTs[InputIndex].LUTTexture;
-				FTexture* Internal = LUTTexture ? LUTTexture->Resource : nullptr;
+				FTexture* Internal = LUTTexture ? LUTTexture->GetResource() : nullptr;
 				for (uint32 OutputIndex = 0; OutputIndex < LocalCount; ++OutputIndex)
 				{
 					if (Internal == OutTextures[OutputIndex])
@@ -353,7 +353,7 @@ uint32 GenerateFinalTable(const FFinalPostProcessSettings& Settings, const FText
 		}
 
 		UTexture* BestLUTTexture = Settings.ContributingLUTs[BestIndex].LUTTexture;
-		FTexture* BestInternal = BestLUTTexture ? BestLUTTexture->Resource : nullptr;
+		FTexture* BestInternal = BestLUTTexture ? BestLUTTexture->GetResource() : nullptr;
 
 		OutTextures[LocalCount] = BestInternal;
 		OutWeights[LocalCount] = BestWeight;
diff --git a/Engine/Source/Runtime/Renderer/Private/PostProcess/PostProcessEyeAdaptation.cpp b/Engine/Source/Runtime/Renderer/Private/PostProcess/PostProcessEyeAdaptation.cpp
--- a/Engine/Source/Runtime/Renderer/Private/PostProcess/PostProcessEyeAdaptation.cpp
+++ b/Engine/Source/Runtime/Renderer/Private/PostProcess/PostProcessEyeAdaptation.cpp
@@ -342,7 +342,7 @@ FEyeAdaptationParameters GetEyeAdaptationParameters(const FViewInfo& View, ERHIF
 
 	//AutoExposureMeterMask
 	const FTextureRHIRef MeterMask = Settings.AutoExposureMeterMask ?
-		Settings.AutoExposureMeterMask->Resource->TextureRHI :
+		Settings.AutoExposureMeterMask->GetResource()->TextureRHI :
 		GWhiteTexture->TextureRHI;
 
 	// The distance at which we switch from linear to exponential. I.e. at StartDistance=1.5, when linear is 1.5 f-stops away from hitting the 
diff --git a/Engine/Source/Runtime/Renderer/Private/PostProcess/PostProcessFFTBloom.cpp b/Engine/Source/Runtime/Renderer/Private/PostProcess/PostProcessFFTBloom.cpp
--- a/Engine/Source/Runtime/Renderer/Private/PostProcess/PostProcessFFTBloom.cpp
+++ b/Engine/Source/Runtime/Renderer/Private/PostProcess/PostProcessFFTBloom.cpp
@@ -153,7 +153,7 @@ bool IsFFTBloomPhysicalKernelReady(const FViewInfo& View)
 		BloomConvolutionTexture = GEngine->DefaultBloomKernelTexture;
 	}
 
-	bool bValidSetup = (BloomConvolutionTexture != nullptr && BloomConvolutionTexture->Resource != nullptr);
+	bool bValidSetup = (BloomConvolutionTexture != nullptr && BloomConvolutionTexture->GetResource() != nullptr);
 
 	// The bloom convolution kernel needs to be resident to avoid visual artifacts.
 	if (bValidSetup)
@@ -642,7 +642,7 @@ FSceneRenderTargetItem* InitDomainAndGetKernel(FRHICommandList& RHICmdList, cons
 	}
 
 	// This should exist if we called IsFFTBloomPhysicalKernelReady.
-	check(BloomConvolutionTexture && BloomConvolutionTexture->Resource && BloomConvolutionTexture->Resource->TextureRHI);
+	check(BloomConvolutionTexture && BloomConvolutionTexture->GetResource() && BloomConvolutionTexture->GetResource()->TextureRHI);
 
 	const float BloomConvolutionSize = PPSettings.BloomConvolutionSize;
 	const FVector2D CenterUV = PPSettings.BloomConvolutionCenterUV;
@@ -678,7 +678,7 @@ FSceneRenderTargetItem* InitDomainAndGetKernel(FRHICommandList& RHICmdList, cons
 		const bool bSameKernelSize = FMath::IsNearlyEqual(CachedKernelScale, BloomConvolutionSize, float(1.e-6) /*tol*/);
 		const bool bSameImageSize = (ImageSize == CachedImageSize);
 		const bool bSameKernelCenterUV = CachedKernelCenterUV.Equals(CenterUV, float(1.e-6) /*tol*/);
-		const bool bSameMipLevel = bSameTexture && FFTKernel.PhysicalMipLevel == BloomConvolutionTexture->Resource->GetCurrentMipCount();
+		const bool bSameMipLevel = bSameTexture && FFTKernel.PhysicalMipLevel == BloomConvolutionTexture->GetResource()->GetCurrentMipCount();
 
 		if (bSameTexture && bSameSpectralBuffer && bSameKernelSize && bSameImageSize && bSameKernelCenterUV && bSameMipLevel)
 		{
@@ -698,7 +698,7 @@ FSceneRenderTargetItem* InitDomainAndGetKernel(FRHICommandList& RHICmdList, cons
 
 		// Sample the physical space kernel into the resized buffer
 
-		FTextureRHIRef& PhysicalSpaceKernelTextureRef = BloomConvolutionTexture->Resource->TextureRHI;
+		FTextureRHIRef& PhysicalSpaceKernelTextureRef = BloomConvolutionTexture->GetResource()->TextureRHI;
 
 		// Rescale the physical space kernel ( and omit the center if this is a 1/2 resolution fft, it will be added later)
 
@@ -736,7 +736,7 @@ FSceneRenderTargetItem* InitDomainAndGetKernel(FRHICommandList& RHICmdList, cons
 		ViewState->BloomFFTKernel.ImageSize = ImageSize;
 		ViewState->BloomFFTKernel.Physical = BloomConvolutionTexture;
 		ViewState->BloomFFTKernel.CenterUV = CenterUV;
-		ViewState->BloomFFTKernel.PhysicalMipLevel = BloomConvolutionTexture->Resource->GetCurrentMipCount();
+		ViewState->BloomFFTKernel.PhysicalMipLevel = BloomConvolutionTexture->GetResource()->GetCurrentMipCount();
 	}
 
 	// Return pointer to the transformed kernel.
diff --git a/Engine/Source/Runtime/Renderer/Private/PostProcess/PostProcessLensFlares.cpp b/Engine/Source/Runtime/Renderer/Private/PostProcess/PostProcessLensFlares.cpp
--- a/Engine/Source/Runtime/Renderer/Private/PostProcess/PostProcessLensFlares.cpp
+++ b/Engine/Source/Runtime/Renderer/Private/PostProcess/PostProcessLensFlares.cpp
@@ -362,7 +362,7 @@ FScreenPassTexture AddLensFlaresPass(
 
 		if (GEngine->DefaultBokehTexture)
 		{
-			FTextureResource* BokehTextureResource = GEngine->DefaultBokehTexture->Resource;
+			FTextureResource* BokehTextureResource = GEngine->DefaultBokehTexture->GetResource();
 
 			if (BokehTextureResource && BokehTextureResource->TextureRHI)
 			{
@@ -372,7 +372,7 @@ FScreenPassTexture AddLensFlaresPass(
 
 		if (Settings.LensFlareBokehShape)
 		{
-			FTextureResource* BokehTextureResource = Settings.LensFlareBokehShape->Resource;
+			FTextureResource* BokehTextureResource = Settings.LensFlareBokehShape->GetResource();
 
 			if (BokehTextureResource && BokehTextureResource->TextureRHI)
 			{
diff --git a/Engine/Source/Runtime/Renderer/Private/PostProcess/PostProcessMobile.cpp b/Engine/Source/Runtime/Renderer/Private/PostProcess/PostProcessMobile.cpp
--- a/Engine/Source/Runtime/Renderer/Private/PostProcess/PostProcessMobile.cpp
+++ b/Engine/Source/Runtime/Renderer/Private/PostProcess/PostProcessMobile.cpp
@@ -1121,9 +1121,9 @@ FScreenPassTexture AddMobileSunMergePass(FRDGBuilder& GraphBuilder, const FViewI
 
 	FRHITexture* BloomDirtMaskTexture = GBlackTexture->TextureRHI;
 
-	if (Settings.BloomDirtMask && Settings.BloomDirtMask->Resource)
+	if (Settings.BloomDirtMask && Settings.BloomDirtMask->GetResource())
 	{
-		BloomDirtMaskTexture = Settings.BloomDirtMask->Resource->TextureRHI;
+		BloomDirtMaskTexture = Settings.BloomDirtMask->GetResource()->TextureRHI;
 	}
 
 	FMobileSunMergePS::FParameters* PSShaderParameters = GraphBuilder.AllocParameters<FMobileSunMergePS::FParameters>();
diff --git a/Engine/Source/Runtime/Renderer/Private/PostProcess/PostProcessTonemap.cpp b/Engine/Source/Runtime/Renderer/Private/PostProcess/PostProcessTonemap.cpp
--- a/Engine/Source/Runtime/Renderer/Private/PostProcess/PostProcessTonemap.cpp
+++ b/Engine/Source/Runtime/Renderer/Private/PostProcess/PostProcessTonemap.cpp
@@ -522,9 +522,9 @@ FScreenPassTexture AddTonemapPass(FRDGBuilder& GraphBuilder, const FViewInfo& Vi
 
 	FRHITexture* BloomDirtMaskTexture = GBlackTexture->TextureRHI;
 
-	if (PostProcessSettings.BloomDirtMask && PostProcessSettings.BloomDirtMask->Resource)
+	if (PostProcessSettings.BloomDirtMask && PostProcessSettings.BloomDirtMask->GetResource())
 	{
-		BloomDirtMaskTexture = PostProcessSettings.BloomDirtMask->Resource->TextureRHI;
+		BloomDirtMaskTexture = PostProcessSettings.BloomDirtMask->GetResource()->TextureRHI;
 	}
 
 	FRHISamplerState* BilinearClampSampler = TStaticSamplerState<SF_Bilinear, AM_Clamp, AM_Clamp, AM_Clamp>::GetRHI();
diff --git a/Engine/Source/Runtime/Renderer/Private/RayTracing/RayTracingIESLightProfiles.cpp b/Engine/Source/Runtime/Renderer/Private/RayTracing/RayTracingIESLightProfiles.cpp
--- a/Engine/Source/Runtime/Renderer/Private/RayTracing/RayTracingIESLightProfiles.cpp
+++ b/Engine/Source/Runtime/Renderer/Private/RayTracing/RayTracingIESLightProfiles.cpp
@@ -75,7 +75,7 @@ void FIESLightProfileResource::BuildIESLightProfilesTexture(FRHICommandListImmed
 		FTextureRHIRef ProfileTexture; 
 		if (IsIESTextureFormatValid(LightProfileTexture))
 		{
-			ProfileTexture = LightProfileTexture->Resource->TextureRHI;
+			ProfileTexture = LightProfileTexture->GetResource()->TextureRHI;
 		}
 		else
 		{
@@ -95,8 +95,8 @@ void FIESLightProfileResource::BuildIESLightProfilesTexture(FRHICommandListImmed
 bool FIESLightProfileResource::IsIESTextureFormatValid(const UTextureLightProfile* Texture) const
 {
 	if (Texture
-		&& Texture->Resource
-		&& Texture->Resource->TextureRHI
+		&& Texture->GetResource()
+		&& Texture->GetResource()->TextureRHI
 		&& Texture->PlatformData
 		&& Texture->PlatformData->PixelFormat == AllowedIESProfileFormat
 		&& Texture->PlatformData->Mips.Num() == 1
diff --git a/Engine/Source/Runtime/Renderer/Private/ReflectionEnvironmentCapture.cpp b/Engine/Source/Runtime/Renderer/Private/ReflectionEnvironmentCapture.cpp
--- a/Engine/Source/Runtime/Renderer/Private/ReflectionEnvironmentCapture.cpp
+++ b/Engine/Source/Runtime/Renderer/Private/ReflectionEnvironmentCapture.cpp
@@ -835,7 +835,7 @@ void CopyCubemapToScratchCubemap(FRHICommandListImmediate& RHICmdList, ERHIFeatu
 	SCOPED_DRAW_EVENT(RHICmdList, CopyCubemapToScratchCubemap);
 	check(SourceCubemap);
 
-	const FTexture* SourceCubemapResource = SourceCubemap->Resource;
+	const FTexture* SourceCubemapResource = SourceCubemap->GetResource();
 	if (SourceCubemapResource == nullptr)
 	{
 		UE_LOG(LogEngine, Warning, TEXT("Unable to copy from cubemap %s, it's RHI resource is null"), *SourceCubemap->GetPathName());
diff --git a/Engine/Source/Runtime/Renderer/Private/RequiredTextureResolutionRendering.cpp b/Engine/Source/Runtime/Renderer/Private/RequiredTextureResolutionRendering.cpp
--- a/Engine/Source/Runtime/Renderer/Private/RequiredTextureResolutionRendering.cpp
+++ b/Engine/Source/Runtime/Renderer/Private/RequiredTextureResolutionRendering.cpp
@@ -44,15 +44,15 @@ void FRequiredTextureResolutionInterface::GetDebugViewModeShaderBindings(
 			{
 				const UTexture* Texture = nullptr;
 				UniformExpressions.GetTextureValue(EMaterialTextureParameterType::Standard2D, ParameterIndex, MaterialContext, Material, Texture);
-				if (Texture && Texture->Resource)
+				if (Texture && Texture->GetResource())
 				{
 					if (Texture->IsStreamable())
 					{
-						TextureResolution = 1 << (Texture->Resource->GetCurrentMipCount() - 1);
+						TextureResolution = 1 << (Texture->GetResource()->GetCurrentMipCount() - 1);
 					}
 					else
 					{
-						TextureResolution = FMath::Max(Texture->Resource->GetSizeX(), Texture->Resource->GetSizeY());
+						TextureResolution = FMath::Max(Texture->GetResource()->GetSizeX(), Texture->GetResource()->GetSizeY());
 					}
 				}
 			}
@@ -65,17 +65,17 @@ void FRequiredTextureResolutionInterface::GetDebugViewModeShaderBindings(
 		{
 			const UTexture* Texture = nullptr;
 			UniformExpressions.GetTextureValue(EMaterialTextureParameterType::Standard2D, ParameterIndex, MaterialContext, Material, Texture);
-			if (Texture && Texture->Resource && Texture->GetFName() == ViewModeParamName)
+			if (Texture && Texture->GetResource() && Texture->GetFName() == ViewModeParamName)
 			{
 				if (Texture->IsStreamable())
 				{
 					const FMaterialTextureParameterInfo& Parameter = UniformExpressions.GetTextureParameter(EMaterialTextureParameterType::Standard2D, ParameterIndex);
 					AnalysisIndex = Parameter.TextureIndex;
-					TextureResolution = 1 << (Texture->Resource->GetCurrentMipCount() - 1);
+					TextureResolution = 1 << (Texture->GetResource()->GetCurrentMipCount() - 1);
 				}
 				else
 				{
-					TextureResolution = FMath::Max(Texture->Resource->GetSizeX(), Texture->Resource->GetSizeY());
+					TextureResolution = FMath::Max(Texture->GetResource()->GetSizeX(), Texture->GetResource()->GetSizeY());
 				}
 			}
 		}
diff --git a/Engine/Source/Runtime/Renderer/Private/SceneRendering.cpp b/Engine/Source/Runtime/Renderer/Private/SceneRendering.cpp
--- a/Engine/Source/Runtime/Renderer/Private/SceneRendering.cpp
+++ b/Engine/Source/Runtime/Renderer/Private/SceneRendering.cpp
@@ -1857,7 +1857,7 @@ void FViewInfo::SetupUniformBufferParameters(
 		ViewUniformShaderParameters.FarShadowStaticMeshLODBias = 0;
 	}
 
-	ViewUniformShaderParameters.PreIntegratedBRDF = GEngine->PreIntegratedSkinBRDFTexture->Resource->TextureRHI;
+	ViewUniformShaderParameters.PreIntegratedBRDF = GEngine->PreIntegratedSkinBRDFTexture->GetResource()->TextureRHI;
 
 	ViewUniformShaderParameters.VirtualTextureFeedbackStride = GetVirtualTextureFeedbackBufferSize(SceneTexturesConfig.Extent).X;
 	// Use some low(ish) discrepancy sequence to run over every pixel in the virtual texture feedback tile.
diff --git a/Engine/Source/Runtime/Renderer/Private/ScreenPass.cpp b/Engine/Source/Runtime/Renderer/Private/ScreenPass.cpp
--- a/Engine/Source/Runtime/Renderer/Private/ScreenPass.cpp
+++ b/Engine/Source/Runtime/Renderer/Private/ScreenPass.cpp
@@ -17,7 +17,7 @@ const FTextureRHIRef& GetMiniFontTexture()
 {
 	if (GEngine->MiniFontTexture)
 	{
-		return GEngine->MiniFontTexture->Resource->TextureRHI;
+		return GEngine->MiniFontTexture->GetResource()->TextureRHI;
 	}
 	else
 	{
diff --git a/Engine/Source/Runtime/Renderer/Private/ShaderPrint.cpp b/Engine/Source/Runtime/Renderer/Private/ShaderPrint.cpp
--- a/Engine/Source/Runtime/Renderer/Private/ShaderPrint.cpp
+++ b/Engine/Source/Runtime/Renderer/Private/ShaderPrint.cpp
@@ -318,7 +318,7 @@ namespace ShaderPrint
 		// Non graph managed resources
 		FUniformBufferRef UniformBuffer = CreateUniformBuffer(View);
 		FRDGBufferSRVRef ValuesBuffer = GraphBuilder.CreateSRV(View.ShaderPrintValueBuffer);
-		FTextureRHIRef FontTexture = GEngine->MiniFontTexture != nullptr ? GEngine->MiniFontTexture->Resource->TextureRHI : GSystemTextures.BlackDummy->GetRenderTargetItem().ShaderResourceTexture;;
+		FTextureRHIRef FontTexture = GEngine->MiniFontTexture != nullptr ? GEngine->MiniFontTexture->GetResource()->TextureRHI : GSystemTextures.BlackDummy->GetRenderTargetItem().ShaderResourceTexture;;
 
 		const ERHIFeatureLevel::Type FeatureLevel = View.GetFeatureLevel();
 		FGlobalShaderMap* GlobalShaderMap = GetGlobalShaderMap(FeatureLevel);
diff --git a/Engine/Source/Runtime/Renderer/Private/SkyAtmosphereRendering.cpp b/Engine/Source/Runtime/Renderer/Private/SkyAtmosphereRendering.cpp
--- a/Engine/Source/Runtime/Renderer/Private/SkyAtmosphereRendering.cpp
+++ b/Engine/Source/Runtime/Renderer/Private/SkyAtmosphereRendering.cpp
@@ -1966,7 +1966,7 @@ void FSceneRenderer::RenderSkyAtmosphereEditorNotifications(FRDGBuilder& GraphBu
 			RenderSkyAtmosphereEditorHudPS::FParameters* PassParameters = GraphBuilder.AllocParameters<RenderSkyAtmosphereEditorHudPS::FParameters>();
 			PassParameters->RenderTargets[0] = FRenderTargetBinding(SceneColorTexture, ERenderTargetLoadAction::ELoad);
 			PassParameters->ViewUniformBuffer = View.ViewUniformBuffer;
-			PassParameters->MiniFontTexture = GEngine->MiniFontTexture ? GEngine->MiniFontTexture->Resource->TextureRHI : GSystemTextures.WhiteDummy->GetRenderTargetItem().TargetableTexture;
+			PassParameters->MiniFontTexture = GEngine->MiniFontTexture ? GEngine->MiniFontTexture->GetResource()->TextureRHI : GSystemTextures.WhiteDummy->GetRenderTargetItem().TargetableTexture;
 
 			FPixelShaderUtils::AddFullscreenPass<RenderSkyAtmosphereEditorHudPS>(GraphBuilder, View.ShaderMap, RDG_EVENT_NAME("SkyAtmosphereEditor"), PixelShader, PassParameters, View.ViewRect);
 		}
diff --git a/Engine/Source/Runtime/Renderer/Private/VisualizeTexturePresent.cpp b/Engine/Source/Runtime/Renderer/Private/VisualizeTexturePresent.cpp
--- a/Engine/Source/Runtime/Renderer/Private/VisualizeTexturePresent.cpp
+++ b/Engine/Source/Runtime/Renderer/Private/VisualizeTexturePresent.cpp
@@ -227,7 +227,7 @@ void FVisualizeTexturePresent::PresentContent(FRDGBuilder& GraphBuilder, const F
 						}
 
 						// UMax is 0.9f to avoid getting some wrap texture leaking in at the bottom
-						Canvas.DrawTile(PixelLeftTop.X, PixelLeftTop.Y, PixelRightBottom.X - PixelLeftTop.X, PixelRightBottom.Y - PixelLeftTop.Y, 0, 0, 1, 0.9f, PhaseColor, GradientTexture->Resource);
+						Canvas.DrawTile(PixelLeftTop.X, PixelLeftTop.Y, PixelRightBottom.X - PixelLeftTop.X, PixelRightBottom.Y - PixelLeftTop.Y, 0, 0, 1, 0.9f, PhaseColor, GradientTexture->GetResource());
 					}
 				}
 
diff --git a/Engine/Source/Runtime/SlateRHIRenderer/Private/SlateRHIResourceManager.cpp b/Engine/Source/Runtime/SlateRHIRenderer/Private/SlateRHIResourceManager.cpp
--- a/Engine/Source/Runtime/SlateRHIRenderer/Private/SlateRHIResourceManager.cpp
+++ b/Engine/Source/Runtime/SlateRHIRenderer/Private/SlateRHIResourceManager.cpp
@@ -833,7 +833,7 @@ FSlateShaderResourceProxy* FSlateRHIResourceManager::FindOrCreateDynamicTextureR
 					}
 				}
 
-				if ( TextureResource.IsValid() && TextureResource->GetTextureObject() && TextureResource->GetTextureObject()->Resource )
+				if ( TextureResource.IsValid() && TextureResource->GetTextureObject() && TextureResource->GetTextureObject()->GetResource())
 				{
 					TextureResource->UpdateTexture(TextureObject);
 					return TextureResource->Proxy;
diff --git a/Engine/Source/Runtime/SlateRHIRenderer/Private/SlateUTextureResource.cpp b/Engine/Source/Runtime/SlateRHIRenderer/Private/SlateUTextureResource.cpp
--- a/Engine/Source/Runtime/SlateRHIRenderer/Private/SlateUTextureResource.cpp
+++ b/Engine/Source/Runtime/SlateRHIRenderer/Private/SlateUTextureResource.cpp
@@ -102,7 +102,7 @@ void FSlateUTextureResource::UpdateTexture(UTexture* InTexture)
 		Proxy = new FSlateShaderResourceProxy;
 	}
 
-	FTexture* TextureResource = InTexture->Resource;
+	FTexture* TextureResource = InTexture->GetResource();
 
 	Proxy->Resource = this;
 	// If the RHI data has changed, it's possible the underlying size of the texture has changed,
diff --git a/Engine/Source/Runtime/SlateRHIRenderer/Private/SlateUTextureResource.h b/Engine/Source/Runtime/SlateRHIRenderer/Private/SlateUTextureResource.h
--- a/Engine/Source/Runtime/SlateRHIRenderer/Private/SlateUTextureResource.h
+++ b/Engine/Source/Runtime/SlateRHIRenderer/Private/SlateUTextureResource.h
@@ -25,9 +25,9 @@ public:
 	/** Gets the RHI resource used for rendering and updates the last render time for texture streaming */
 	FTextureRHIRef AccessRHIResource()
 	{
-		if ( TextureObject && TextureObject->Resource )
+		if ( TextureObject && TextureObject->GetResource())
 		{
-			FTexture* TextureResource = TextureObject->Resource;
+			FTexture* TextureResource = TextureObject->GetResource();
 			TextureResource->LastRenderTime = FApp::GetCurrentTime();
 
 			return TextureResource->TextureRHI;
diff --git a/Engine/Source/Runtime/UMG/Private/Blueprint/AsyncTaskDownloadImage.cpp b/Engine/Source/Runtime/UMG/Private/Blueprint/AsyncTaskDownloadImage.cpp
--- a/Engine/Source/Runtime/UMG/Private/Blueprint/AsyncTaskDownloadImage.cpp
+++ b/Engine/Source/Runtime/UMG/Private/Blueprint/AsyncTaskDownloadImage.cpp
@@ -115,7 +115,7 @@ void UAsyncTaskDownloadImage::HandleImageRequest(FHttpRequestPtr HttpRequest, FH
 						Texture->SRGB = true;
 						Texture->UpdateResource();
 
-						FTexture2DDynamicResource* TextureResource = static_cast<FTexture2DDynamicResource*>(Texture->Resource);
+						FTexture2DDynamicResource* TextureResource = static_cast<FTexture2DDynamicResource*>(Texture->GetResource());
 						if (TextureResource)
 						{
 							ENQUEUE_RENDER_COMMAND(FWriteRawDataToTexture)(
diff --git a/Engine/Source/Runtime/UMG/Private/Components/WidgetComponent.cpp b/Engine/Source/Runtime/UMG/Private/Components/WidgetComponent.cpp
--- a/Engine/Source/Runtime/UMG/Private/Components/WidgetComponent.cpp
+++ b/Engine/Source/Runtime/UMG/Private/Components/WidgetComponent.cpp
@@ -364,7 +364,7 @@ public:
 
 		if( RenderTarget )
 		{
-			FTextureResource* TextureResource = RenderTarget->Resource;
+			FTextureResource* TextureResource = RenderTarget->GetResource();
 			if ( TextureResource )
 			{
 				if (GeometryMode == EWidgetGeometryMode::Plane)
diff --git a/Engine/Source/Runtime/WebBrowserTexture/Private/WebBrowserTexture.cpp b/Engine/Source/Runtime/WebBrowserTexture/Private/WebBrowserTexture.cpp
--- a/Engine/Source/Runtime/WebBrowserTexture/Private/WebBrowserTexture.cpp
+++ b/Engine/Source/Runtime/WebBrowserTexture/Private/WebBrowserTexture.cpp
@@ -109,7 +109,7 @@ void UWebBrowserTexture::GetResourceSizeEx(FResourceSizeEx& CumulativeResourceSi
 
 void UWebBrowserTexture::TickResource(TSharedPtr<FWebBrowserTextureSample, ESPMode::ThreadSafe> Sample)
 {
-	if (Resource == nullptr)
+	if (GetResource() == nullptr)
 	{
 		return;
 	}
@@ -129,7 +129,7 @@ void UWebBrowserTexture::TickResource(TSharedPtr<FWebBrowserTextureSample, ESPMo
 		RenderParams.SampleSource = SampleQueue;
 	}
 
-	FWebBrowserTextureResource* ResourceParam = (FWebBrowserTextureResource*)Resource;
+	FWebBrowserTextureResource* ResourceParam = (FWebBrowserTextureResource*)GetResource();
 	ENQUEUE_RENDER_COMMAND(UWebBrowserTextureResourceRender)(
 		[ResourceParam, RenderParams](FRHICommandListImmediate& RHICmdList)
 		{
