Expose some Landscape classes to ease plugin development

diff --git a/Engine/Source/Runtime/Landscape/Private/LandscapeRender.cpp b/Engine/Source/Runtime/Landscape/Private/LandscapeRender.cpp
--- a/Engine/Source/Runtime/Landscape/Private/LandscapeRender.cpp
+++ b/Engine/Source/Runtime/Landscape/Private/LandscapeRender.cpp
@@ -511,68 +511,49 @@ namespace
 
 } // namespace 
 
-struct FLandscapeDebugOptions
-{
-	FLandscapeDebugOptions()
-		: bShowPatches(false)
-		, bDisableStatic(false)
-		, CombineMode(eCombineMode_Default)
-		, PatchesConsoleCommand(
-			TEXT("Landscape.Patches"),
-			TEXT("Show/hide Landscape patches"),
-			FConsoleCommandDelegate::CreateRaw(this, &FLandscapeDebugOptions::Patches))
-		, StaticConsoleCommand(
-			TEXT("Landscape.Static"),
-			TEXT("Enable/disable Landscape static drawlists"),
-			FConsoleCommandDelegate::CreateRaw(this, &FLandscapeDebugOptions::Static))
-		, CombineConsoleCommand(
-			TEXT("Landscape.Combine"),
-			TEXT("Set landscape component combining mode : 0 = Default, 1 = Combine All, 2 = Disabled"),
-			FConsoleCommandWithArgsDelegate::CreateRaw(this, &FLandscapeDebugOptions::Combine))
-	{
-	}
-
-	enum eCombineMode
-	{
-		eCombineMode_Default = 0,
-		eCombineMode_CombineAll = 1,
-		eCombineMode_Disabled = 2
-	};
-
-	bool bShowPatches;
-	bool bDisableStatic;
-	eCombineMode CombineMode;
 
-	FORCEINLINE bool IsCombinedDisabled() const { return CombineMode == eCombineMode_Disabled; }
-	FORCEINLINE bool IsCombinedAll() const { return CombineMode == eCombineMode_CombineAll; }
-	FORCEINLINE bool IsCombinedDefault() const { return CombineMode == eCombineMode_Default; }
-
-private:
-	FAutoConsoleCommand PatchesConsoleCommand;
-	FAutoConsoleCommand StaticConsoleCommand;
-	FAutoConsoleCommand CombineConsoleCommand;
+//
+// FLandscapeDebugOptions
+//
+FLandscapeDebugOptions::FLandscapeDebugOptions()
+	: bShowPatches(false)
+	, bDisableStatic(false)
+	, CombineMode(eCombineMode_Default)
+	, PatchesConsoleCommand(
+		TEXT("Landscape.Patches"),
+		TEXT("Show/hide Landscape patches"),
+		FConsoleCommandDelegate::CreateRaw(this, &FLandscapeDebugOptions::Patches))
+	, StaticConsoleCommand(
+		TEXT("Landscape.Static"),
+		TEXT("Enable/disable Landscape static drawlists"),
+		FConsoleCommandDelegate::CreateRaw(this, &FLandscapeDebugOptions::Static))
+	, CombineConsoleCommand(
+		TEXT("Landscape.Combine"),
+		TEXT("Set landscape component combining mode : 0 = Default, 1 = Combine All, 2 = Disabled"),
+		FConsoleCommandWithArgsDelegate::CreateRaw(this, &FLandscapeDebugOptions::Combine))
+{
+}
 
-	void Patches()
-	{
-		bShowPatches = !bShowPatches;
-		UE_LOG(LogLandscape, Display, TEXT("Landscape.Patches: %s"), bShowPatches ? TEXT("Show") : TEXT("Hide"));
-	}
+void FLandscapeDebugOptions::Patches()
+{
+	bShowPatches = !bShowPatches;
+	UE_LOG(LogLandscape, Display, TEXT("Landscape.Patches: %s"), bShowPatches ? TEXT("Show") : TEXT("Hide"));
+}
 
-	void Static()
-	{
-		bDisableStatic = !bDisableStatic;
-		UE_LOG(LogLandscape, Display, TEXT("Landscape.Static: %s"), bDisableStatic ? TEXT("Disabled") : TEXT("Enabled"));
-	}
+void FLandscapeDebugOptions::Static()
+{
+	bDisableStatic = !bDisableStatic;
+	UE_LOG(LogLandscape, Display, TEXT("Landscape.Static: %s"), bDisableStatic ? TEXT("Disabled") : TEXT("Enabled"));
+}
 
-	void Combine(const TArray<FString>& Args)
+void FLandscapeDebugOptions::Combine(const TArray<FString>& Args)
+{
+	if (Args.Num() >= 1)
 	{
-		if (Args.Num() >= 1)
-		{
-			CombineMode = (eCombineMode)FCString::Atoi(*Args[0]);
-			UE_LOG(LogLandscape, Display, TEXT("Landscape.Combine: %d"), (int32)CombineMode);
-		}
+		CombineMode = (eCombineMode)FCString::Atoi(*Args[0]);
+		UE_LOG(LogLandscape, Display, TEXT("Landscape.Combine: %d"), (int32)CombineMode);
 	}
-};
+}
 
 FLandscapeDebugOptions GLandscapeDebugOptions;
 
diff --git a/Engine/Source/Runtime/Landscape/Public/LandscapeRender.h b/Engine/Source/Runtime/Landscape/Public/LandscapeRender.h
--- a/Engine/Source/Runtime/Landscape/Public/LandscapeRender.h
+++ b/Engine/Source/Runtime/Landscape/Public/LandscapeRender.h
@@ -115,7 +115,7 @@ END_GLOBAL_SHADER_PARAMETER_STRUCT()
 
 typedef TUniformBufferRef<FLandscapeVertexFactoryMVFParameters> FLandscapeVertexFactoryMVFUniformBufferRef;
 
-BEGIN_GLOBAL_SHADER_PARAMETER_STRUCT(FLandscapeSectionLODUniformParameters, )
+BEGIN_GLOBAL_SHADER_PARAMETER_STRUCT(FLandscapeSectionLODUniformParameters, LANDSCAPE_API)
 	SHADER_PARAMETER(FIntPoint, Min)
 	SHADER_PARAMETER(FIntPoint, Size)
 	SHADER_PARAMETER_SRV(Buffer<float>, SectionLOD)
@@ -676,7 +676,7 @@ struct FLandscapeRenderSystem
 	void EndFrame();
 };
 
-extern TMap<FLandscapeNeighborInfo::FLandscapeKey, FLandscapeRenderSystem*> LandscapeRenderSystems;
+LANDSCAPE_API extern TMap<FLandscapeNeighborInfo::FLandscapeKey, FLandscapeRenderSystem*> LandscapeRenderSystems;
 
 class FLandscapeVisibilityHelper
 {
@@ -692,6 +692,36 @@ private:
 	FPrimitiveSceneProxy* Proxy = nullptr;
 };
 
+//
+// FLandscapeDebugOptions
+//
+struct LANDSCAPE_API FLandscapeDebugOptions
+{
+	FLandscapeDebugOptions();
+
+	enum eCombineMode
+	{
+		eCombineMode_Default = 0,
+		eCombineMode_CombineAll = 1,
+		eCombineMode_Disabled = 2
+	};
+
+	bool bShowPatches;
+	bool bDisableStatic;
+	eCombineMode CombineMode;
+
+private:
+	FAutoConsoleCommand PatchesConsoleCommand;
+	FAutoConsoleCommand StaticConsoleCommand;
+	FAutoConsoleCommand CombineConsoleCommand;
+
+	void Patches();
+	void Static();
+	void Combine(const TArray<FString>& Args);
+};
+
+LANDSCAPE_API extern FLandscapeDebugOptions GLandscapeDebugOptions;
+
 //
 // FLandscapeMeshProxySceneProxy
 //
@@ -714,7 +744,7 @@ private:
 //
 // FLandscapeComponentSceneProxy
 //
-class FLandscapeComponentSceneProxy : public FPrimitiveSceneProxy, public FLandscapeNeighborInfo
+class LANDSCAPE_API FLandscapeComponentSceneProxy : public FPrimitiveSceneProxy, public FLandscapeNeighborInfo
 {
 	friend class FLandscapeSharedBuffers;
 
@@ -774,7 +804,7 @@ public:
 	// Reference counted vertex and index buffer shared among all landscape scene proxies of the same component size
 	// Key is the component size and number of subsections.
 	// Also being reused by GPULightmass currently to save mem
-	static LANDSCAPE_API TMap<uint32, FLandscapeSharedBuffers*> SharedBuffersMap;
+	static TMap<uint32, FLandscapeSharedBuffers*> SharedBuffersMap;
 #if GPUCULL_TODO
 	// GPUCULL_TODO: put instances in base proxy and make non-virtual
 	virtual const TArray<FPrimitiveInstance>* GetPrimitiveInstances() const	{ return &Instances; }
