Fix assert occcuring when copying/duplicating map that contains landscape splines and register the component on an unitialized world.

diff --git a/Engine/Source/Runtime/Landscape/Private/LandscapeSplines.cpp b/Engine/Source/Runtime/Landscape/Private/LandscapeSplines.cpp
--- a/Engine/Source/Runtime/Landscape/Private/LandscapeSplines.cpp
+++ b/Engine/Source/Runtime/Landscape/Private/LandscapeSplines.cpp
@@ -3124,10 +3124,14 @@ void ULandscapeSplineSegment::UpdateSplinePoints(bool bUpdateCollision, bool bUp
 #endif
 		}
 
-		// Finally, register components
-		for (auto* MeshComponent : MeshComponents)
+		// Finally, register components if the world is initialized
+		UWorld* World = GetWorld();
+		if (World && World->bIsWorldInitialized)
 		{
-			MeshComponent->RegisterComponent();
+			for (USplineMeshComponent* MeshComponent : MeshComponents)
+			{
+				MeshComponent->RegisterComponent();
+			}
 		}
 	}
 	
