Landscape Splines: Fix checkSlow failing PostDuplicate (missing update of transient maps)

diff --git a/Engine/Source/Runtime/Landscape/Private/LandscapeSplines.cpp b/Engine/Source/Runtime/Landscape/Private/LandscapeSplines.cpp
--- a/Engine/Source/Runtime/Landscape/Private/LandscapeSplines.cpp
+++ b/Engine/Source/Runtime/Landscape/Private/LandscapeSplines.cpp
@@ -2275,6 +2275,11 @@ void ULandscapeSplineControlPoint::PostDuplicate(bool bDuplicateForPIE)
 			{
 				LocalMeshComponent = nullptr;
 			}
+			else
+			{
+				// If LocalMeshComponent is still valid make sure its added to the transient map that normally gets populated on PostLoad
+				OuterSplines->MeshComponentLocalOwnersMap.Add(LocalMeshComponent, this);
+			}
 		}
 
 		UpdateSplinePoints();
@@ -3255,8 +3260,14 @@ void ULandscapeSplineSegment::PostDuplicate(bool bDuplicateForPIE)
 			{
 				LocalMeshComponents.Empty();
 			}
-		}
 
+			// If LocalMeshComponents are still valid make sure they are added to the transient map that normally gets populated PostLoad
+			for (auto& LocalMeshComponent : LocalMeshComponents)
+			{
+				OuterSplines->MeshComponentLocalOwnersMap.Add(LocalMeshComponent, this);
+			}
+		}
+		
 		UpdateSplinePoints();
 	}
 
