Fix landscape material rendering to use the correct readback mechanism.

The code was trying to copy the rendered results into a CPU-accessible texture via CopyToResolveTarget, which is wrong and asserts in the RDG validation layer.

#rnx

diff --git a/Engine/Source/Runtime/Landscape/Private/LandscapePhysicalMaterial.cpp b/Engine/Source/Runtime/Landscape/Private/LandscapePhysicalMaterial.cpp
--- a/Engine/Source/Runtime/Landscape/Private/LandscapePhysicalMaterial.cpp
+++ b/Engine/Source/Runtime/Landscape/Private/LandscapePhysicalMaterial.cpp
@@ -289,8 +289,7 @@ namespace
 		FVector ViewOrigin,
 		FMatrix ViewRotationMatrix,
 		FMatrix ProjectionMatrix,
-		FTextureRHIRef ReadbackTextureRHI,
-		FGPUFenceRHIRef ReadbackFenceRHI)
+		FRHIGPUTextureReadback* Readback)
 	{
 		// Create the view
 		FSceneViewFamily::ConstructionValues ViewFamilyInit(nullptr, nullptr, FEngineShowFlags(ESFIM_Game));
@@ -348,13 +347,7 @@ namespace
 			});
 		});
 
-		FRDGTextureRef ReadbackTexture = GraphBuilder.RegisterExternalTexture(CreateRenderTarget(ReadbackTextureRHI, TEXT("ReadbackTexture")));
-		AddCopyToResolveTargetPass(GraphBuilder, OutputTexture, ReadbackTexture, FResolveParams{});
-
-		AddPass(GraphBuilder, [&ReadbackFenceRHI](FRHICommandList& RHICmdList)
-		{
-			RHICmdList.WriteGPUFence(ReadbackFenceRHI);
-		});
+		AddEnqueueCopyPass(GraphBuilder, Readback, OutputTexture);
 
 		GraphBuilder.Execute();
 	}
@@ -363,15 +356,13 @@ namespace
 	void FetchResults_RenderThread(
 		FRHICommandListImmediate& RHICmdList,
 		FIntPoint TargetSize,
-		FTextureRHIRef InReadbackTexture,
-		FGPUFenceRHIRef InReadbackFence,
+		FRHIGPUTextureReadback* Readback,
 		TArray<uint8>& OutPhysicalMaterialIds)
 	{
 		void* Data = nullptr;
-		int32 TargetWidth, TargetHeight;
-		RHICmdList.MapStagingSurface(InReadbackTexture, InReadbackFence.GetReference(), Data, TargetWidth, TargetHeight);
-		check(Data != nullptr);
-		check(TargetSize.X <= TargetWidth && TargetSize.Y <= TargetHeight);
+		int32 Pitch;
+		Readback->LockTexture(RHICmdList, Data, Pitch);
+		check(Data && TargetSize.X <= Pitch);
 
 		OutPhysicalMaterialIds.Empty(TargetSize.X * TargetSize.Y);
 		OutPhysicalMaterialIds.AddUninitialized(TargetSize.X * TargetSize.Y);
@@ -381,11 +372,11 @@ namespace
 		for (int32 Y = 0; Y < TargetSize.Y; ++Y)
 		{
 			FMemory::Memcpy(WritePtr, ReadPtr, TargetSize.X);
-			ReadPtr += TargetWidth;
+			ReadPtr += Pitch;
 			WritePtr += TargetSize.X;
 		}
 
-		RHICmdList.UnmapStagingSurface(InReadbackTexture);
+		Readback->Unlock();
 	}
 }
 
@@ -415,8 +406,7 @@ namespace
 		TArray<UPhysicalMaterial*> ResultMaterials;
 
 		// Create on render thread
-		FTextureRHIRef ReadbackTexture;
-		FGPUFenceRHIRef ReadbackFence;
+		TUniquePtr<FRHIGPUTextureReadback> Readback;
 
 		// Result written on render thread and read on game thread
 		ECompletionState CompletionState = ECompletionState::None;
@@ -444,6 +434,11 @@ namespace
 				const float ZOffset = WORLD_MAX;
 				const FMatrix ProjectionMatrix = FReversedZOrthoMatrix(TargetExtent.X, TargetExtent.Y, 0.5f / ZOffset, ZOffset);
 
+				if (Task.TargetSize != TargetSize)
+				{
+					Task.Readback = nullptr;
+				}
+
 				Task.TargetSize = TargetSize;
 				Task.ViewOrigin = TargetCenter;
 				Task.ViewRotationMatrix = ViewRotationMatrix;
@@ -460,14 +455,9 @@ namespace
 	bool InitTaskRenderResources(FLandscapePhysicalMaterialRenderTaskImpl& Task)
 	{
 		//todo: Consider pooling these and throttling to the pool size?
-		if (!Task.ReadbackTexture.IsValid() || Task.ReadbackTexture->GetSizeXYZ() != FIntVector(Task.TargetSize.X, Task.TargetSize.Y, 1))
+		if (!Task.Readback.IsValid())
 		{
-			FRHIResourceCreateInfo CreateInfo(TEXT("LandscapePhysicalMaterialReadback"));
-			Task.ReadbackTexture = RHICreateTexture2D(Task.TargetSize.X, Task.TargetSize.Y, PF_G8, 1, 1, TexCreate_CPUReadback | TexCreate_HideInVisualizeTexture, CreateInfo);
-		}
-		if (!Task.ReadbackFence.IsValid())
-		{
-			Task.ReadbackFence = RHICreateGPUFence(TEXT(""));
+			Task.Readback = MakeUnique<FRHIGPUTextureReadback>(TEXT("LandscapePhysicalMaterialReadback"));
 		}
 
 		return true;
@@ -495,8 +485,7 @@ namespace
 					Task.ViewOrigin,
 					Task.ViewRotationMatrix,
 					Task.ProjectionMatrix,
-					Task.ReadbackTexture,
-					Task.ReadbackFence);
+					Task.Readback.Get());
 
 				FPlatformMisc::MemoryBarrier();
 				Task.CompletionState = ECompletionState::Pending;
@@ -504,9 +493,14 @@ namespace
 		}
 		else if (Task.CompletionState == ECompletionState::Pending)
 		{
-			if (bFlush || Task.ReadbackFence->Poll())
+			if (bFlush || Task.Readback->IsReady())
 			{
-				FetchResults_RenderThread(RHICmdList, Task.TargetSize, Task.ReadbackTexture, Task.ReadbackFence, Task.ResultIds);
+				if (!Task.Readback->IsReady())
+				{
+					RHICmdList.ImmediateFlush(EImmediateFlushType::FlushRHIThread);
+				}
+
+				FetchResults_RenderThread(RHICmdList, Task.TargetSize, Task.Readback.Get(), Task.ResultIds);
 
 				FPlatformMisc::MemoryBarrier();
 				Task.CompletionState = ECompletionState::Complete;
@@ -594,8 +588,7 @@ public:
 					ENQUEUE_RENDER_COMMAND(FLandscapePhysicalMaterialFree)(
 						[Task](FRHICommandListImmediate& RHICmdList)
 						{
-							Task->ReadbackTexture.SafeRelease();
-							Task->ReadbackFence.SafeRelease();
+							Task->Readback = nullptr;
 						});
 				}
 			}
