Fixed missing landscape on mobile when GPUScene is enabled

diff --git a/Engine/Source/Runtime/Landscape/Private/LandscapeRender.cpp b/Engine/Source/Runtime/Landscape/Private/LandscapeRender.cpp
--- a/Engine/Source/Runtime/Landscape/Private/LandscapeRender.cpp
+++ b/Engine/Source/Runtime/Landscape/Private/LandscapeRender.cpp
@@ -1137,11 +1137,11 @@ FLandscapeComponentSceneProxy::FLandscapeComponentSceneProxy(ULandscapeComponent
 		bShouldNotifyOnWorldAddRemove = true;
 	}
 
+	const ERHIFeatureLevel::Type FeatureLevel = GetScene().GetFeatureLevel();
 	// GPUCULL_TODO: Move to base proxy
-	bVFRequiresPrimitiveUniformBuffer = !UseGPUScene(GMaxRHIShaderPlatform, GetScene().GetFeatureLevel());
-
-	const auto FeatureLevel = GetScene().GetFeatureLevel();
-
+	// Mobile LandscapeVF does not use GPUScene
+	bVFRequiresPrimitiveUniformBuffer = !UseGPUScene(GMaxRHIShaderPlatform, FeatureLevel) || (FeatureLevel == ERHIFeatureLevel::ES3_1);
+	
 	if (FeatureLevel >= ERHIFeatureLevel::SM5)
 	{
 		if (InComponent->GetLandscapeProxy()->bUseDynamicMaterialInstance)
