Flush GPU Scene before landscape grass and landscape physical material rendering

diff --git a/Engine/Source/Runtime/Landscape/Private/LandscapeGrass.cpp b/Engine/Source/Runtime/Landscape/Private/LandscapeGrass.cpp
--- a/Engine/Source/Runtime/Landscape/Private/LandscapeGrass.cpp
+++ b/Engine/Source/Runtime/Landscape/Private/LandscapeGrass.cpp
@@ -59,7 +59,7 @@
 #include "Misc/ScopedSlowTask.h"
 #include "TextureCompiler.h"
 #include "RenderCaptureInterface.h"
-#include "InstanceCulling/InstanceCullingContext.h"
+#include "SimpleMeshDrawCommandPass.h"
 
 #define LOCTEXT_NAMESPACE "Landscape"
 
@@ -554,16 +554,22 @@ public:
 
 	BEGIN_SHADER_PARAMETER_STRUCT(FLandscapeGrassPassParameters, )
 		SHADER_PARAMETER_STRUCT_REF(FViewUniformShaderParameters, View)
-		SHADER_PARAMETER_RDG_UNIFORM_BUFFER(FInstanceCullingGlobalUniforms, InstanceCulling)
+		SHADER_PARAMETER_STRUCT_INCLUDE(FInstanceCullingDrawParams, InstanceCullingDrawParams)
 		RENDER_TARGET_BINDING_SLOTS()
 	END_SHADER_PARAMETER_STRUCT()
 
 	void RenderLandscapeComponentToTexture_RenderThread(FRHICommandListImmediate& RHICmdList)
 	{
+		FMemMark Mark(FMemStack::Get());
+		FRDGBuilder GraphBuilder(RHICmdList);
+
 		FSceneViewFamilyContext ViewFamily(FSceneViewFamily::ConstructionValues(nullptr, SceneInterface, FEngineShowFlags(ESFIM_Game)).SetWorldTimes(FApp::GetCurrentTime() - GStartTime, FApp::GetDeltaTime(), FApp::GetCurrentTime() - GStartTime));
 
 		ViewFamily.LandscapeLODOverride = 0; // Force LOD render
 
+		// Ensure scene primitive rendering is valid (added primitives comitted, GPU-Scene updated, push/pop dynamic culling context).
+		FScenePrimitiveRenderingContextScopeHelper ScenePrimitiveRenderingContextScopeHelper(GetRendererModule().BeginScenePrimitiveRendering(GraphBuilder, &ViewFamily));
+		
 		FSceneViewInitOptions ViewInitOptions;
 		ViewInitOptions.SetViewRectangle(FIntRect(0, 0, TargetSize.X, TargetSize.Y));
 		ViewInitOptions.ViewOrigin = ViewOrigin;
@@ -575,26 +581,14 @@ public:
 		
 		const FSceneView* View = ViewFamily.Views[0];
 
-		FMemMark Mark(FMemStack::Get());
-		FRDGBuilder GraphBuilder(RHICmdList);
-
 		FRDGTextureRef OutputTexture = GraphBuilder.RegisterExternalTexture(CreateRenderTarget(RenderTargetResource->GetTextureRHI(), TEXT("LandscapeGrass")));
 
 		auto* PassParameters = GraphBuilder.AllocParameters<FLandscapeGrassPassParameters>();
 		PassParameters->View = View->ViewUniformBuffer;
-		PassParameters->InstanceCulling = FInstanceCullingContext::CreateDummyInstanceCullingUniformBuffer(GraphBuilder);
 		PassParameters->RenderTargets[0] = FRenderTargetBinding(OutputTexture, ERenderTargetLoadAction::EClear);
 
-		GraphBuilder.AddPass(
-			RDG_EVENT_NAME("LandscapeGrass"),
-			PassParameters,
-			ERDGPassFlags::Raster,
-			[this, &View](FRHICommandListImmediate& RHICmdList)
-		{
-			RHICmdList.SetViewport(View->UnscaledViewRect.Min.X, View->UnscaledViewRect.Min.Y, 0.0f, View->UnscaledViewRect.Max.X, View->UnscaledViewRect.Max.Y, 1.0f);
-			RHICmdList.SetScissorRect(false, 0, 0, 0, 0);
-
-			DrawDynamicMeshPass(*View, RHICmdList, [&](FDynamicPassMeshDrawListContext* DynamicMeshPassContext)
+		AddSimpleMeshPass(GraphBuilder, PassParameters, SceneInterface->GetRenderScene(), *View, nullptr, RDG_EVENT_NAME("LandscapeGrass"), View->UnscaledViewRect,
+			[&View, this](FDynamicPassMeshDrawListContext* DynamicMeshPassContext)
 			{
 				FLandscapeGrassWeightMeshProcessor PassMeshProcessor(
 					nullptr,
@@ -614,7 +608,6 @@ public:
 					}
 				}
 			});
-		});
 
 		GraphBuilder.Execute();
 	}
diff --git a/Engine/Source/Runtime/Landscape/Private/LandscapePhysicalMaterial.cpp b/Engine/Source/Runtime/Landscape/Private/LandscapePhysicalMaterial.cpp
--- a/Engine/Source/Runtime/Landscape/Private/LandscapePhysicalMaterial.cpp
+++ b/Engine/Source/Runtime/Landscape/Private/LandscapePhysicalMaterial.cpp
@@ -17,7 +17,7 @@
 #include "RenderTargetPool.h"
 #include "RHIResources.h"
 #include "SceneRenderTargetParameters.h"
-#include "InstanceCulling/InstanceCullingContext.h"
+#include "SimpleMeshDrawCommandPass.h"
 
 DECLARE_GPU_STAT_NAMED(LandscapePhysicalMaterial_Draw, TEXT("LandscapePhysicalMaterial"));
 
@@ -277,13 +277,14 @@ namespace
 	
 	BEGIN_SHADER_PARAMETER_STRUCT(FLandscapePhysicalMaterialPassParameters, )
 		SHADER_PARAMETER_STRUCT_REF(FViewUniformShaderParameters, View)
-		SHADER_PARAMETER_RDG_UNIFORM_BUFFER(FInstanceCullingGlobalUniforms, InstanceCulling)
+		SHADER_PARAMETER_STRUCT_INCLUDE(FInstanceCullingDrawParams, InstanceCullingDrawParams)
 		RENDER_TARGET_BINDING_SLOTS()
 	END_SHADER_PARAMETER_STRUCT()
 
 	/** Render the landscape physical material IDs and copy to the read back texture. */
 	void Render_RenderThread(
 		FRHICommandListImmediate& RHICmdList,
+		FSceneInterface* SceneInterface,
 		FMeshInfoArray const& MeshInfos,
 		FIntPoint TargetSize,
 		FVector ViewOrigin,
@@ -291,12 +292,17 @@ namespace
 		FMatrix ProjectionMatrix,
 		FRHIGPUTextureReadback* Readback)
 	{
+		FMemMark Mark(FMemStack::Get());
+		FRDGBuilder GraphBuilder(RHICmdList);
+
 		// Create the view
-		FSceneViewFamily::ConstructionValues ViewFamilyInit(nullptr, nullptr, FEngineShowFlags(ESFIM_Game));
+		FSceneViewFamily::ConstructionValues ViewFamilyInit(nullptr, SceneInterface, FEngineShowFlags(ESFIM_Game));
 		ViewFamilyInit.SetWorldTimes(0.0f, 0.0f, 0.0f);
 		FSceneViewFamilyContext ViewFamily(ViewFamilyInit);
 		ViewFamily.LandscapeLODOverride = 0; // Force LOD 0 render
 
+		FScenePrimitiveRenderingContextScopeHelper ScenePrimitiveRenderingContextScopeHelper(GetRendererModule().BeginScenePrimitiveRendering(GraphBuilder, &ViewFamily));
+
 		FSceneViewInitOptions ViewInitOptions;
 		ViewInitOptions.SetViewRectangle(FIntRect(0, 0, TargetSize.X, TargetSize.Y));
 		ViewInitOptions.ViewOrigin = ViewOrigin;
@@ -307,34 +313,21 @@ namespace
 		GetRendererModule().CreateAndInitSingleView(RHICmdList, &ViewFamily, &ViewInitOptions);
 		const FSceneView* View = ViewFamily.Views[0];
 
-		FMemMark Mark(FMemStack::Get());
-		FRDGBuilder GraphBuilder(RHICmdList);
-
 		FRDGTextureRef OutputTexture = GraphBuilder.CreateTexture(
 			FRDGTextureDesc::Create2D(TargetSize, PF_G8, FClearValueBinding::Black, TexCreate_RenderTargetable | TexCreate_ShaderResource),
 			TEXT("LandscapePhysicalMaterialTarget"));
 
 		auto* PassParameters = GraphBuilder.AllocParameters<FLandscapePhysicalMaterialPassParameters>();
 		PassParameters->View = View->ViewUniformBuffer;
-		PassParameters->InstanceCulling = FInstanceCullingContext::CreateDummyInstanceCullingUniformBuffer(GraphBuilder);
 		PassParameters->RenderTargets[0] = FRenderTargetBinding(OutputTexture, ERenderTargetLoadAction::EClear);
 
-		GraphBuilder.AddPass(
-			RDG_EVENT_NAME("LandscapePhysicalMaterial"),
-			PassParameters,
-			ERDGPassFlags::Raster,
-			[&View, &MeshInfos](FRHICommandListImmediate& RHICmdList)
-		{
-			RHICmdList.SetViewport(View->UnscaledViewRect.Min.X, View->UnscaledViewRect.Min.Y, 0.0f, View->UnscaledViewRect.Max.X, View->UnscaledViewRect.Max.Y, 1.0f);
-			RHICmdList.SetScissorRect(false, 0, 0, 0, 0);
-
-			DrawDynamicMeshPass(*View, RHICmdList, [&](FDynamicPassMeshDrawListContext* DynamicMeshPassContext)
+		AddSimpleMeshPass(GraphBuilder, PassParameters, SceneInterface->GetRenderScene(), *View, nullptr, RDG_EVENT_NAME("LandscapePhysicalMaterial"), View->UnscaledViewRect,
+			[View, &MeshInfos](FDynamicPassMeshDrawListContext* DynamicMeshPassContext)
 			{
 				FLandscapePhysicalMaterialMeshProcessor PassMeshProcessor(
 					nullptr,
 					View,
 					DynamicMeshPassContext);
-
 				for (auto& MeshInfo : MeshInfos)
 				{
 					const FMeshBatch& Mesh = *MeshInfo.MeshBatch;
@@ -345,7 +338,6 @@ namespace
 					}
 				}
 			});
-		});
 
 		AddEnqueueCopyPass(GraphBuilder, Readback, OutputTexture);
 
@@ -480,6 +472,7 @@ namespace
 
 				Render_RenderThread(
 					RHICmdList,
+					&SceneProxy->GetScene(),
 					MeshInfos,
 					Task.TargetSize,
 					Task.ViewOrigin,
