Fixed SSE Math double precision issues

diff --git a/Engine/Plugins/MovieScene/MovieRenderPipeline/Source/MovieRenderPipelineCore/Private/MoviePipelineRendering.cpp b/Engine/Plugins/MovieScene/MovieRenderPipeline/Source/MovieRenderPipelineCore/Private/MoviePipelineRendering.cpp
--- a/Engine/Plugins/MovieScene/MovieRenderPipeline/Source/MovieRenderPipelineCore/Private/MoviePipelineRendering.cpp
+++ b/Engine/Plugins/MovieScene/MovieRenderPipeline/Source/MovieRenderPipelineCore/Private/MoviePipelineRendering.cpp
@@ -66,8 +66,8 @@ void UMoviePipeline::SetupRenderingPipelineForShot(UMoviePipelineExecutorShot* I
 
 	// Figure out how big each sub-region (tile) is.
 	FIntPoint BackbufferResolution = FIntPoint(
-		FMath::CeilToInt(OutputResolution.X / HighResSettings->TileCount),
-		FMath::CeilToInt(OutputResolution.Y / HighResSettings->TileCount));
+		FMath::CeilToInt((float)OutputResolution.X / (float)HighResSettings->TileCount),
+		FMath::CeilToInt((float)OutputResolution.Y / (float)HighResSettings->TileCount));
 
 	// Then increase each sub-region by the overlap amount.
 	BackbufferResolution = HighResSettings->CalculatePaddedBackbufferSize(BackbufferResolution);
@@ -348,7 +348,7 @@ void UMoviePipeline::RenderFrame()
 					SpatialShiftY = r * FMath::Sin(Theta);
 				}
 
-				FIntPoint BackbufferResolution = FIntPoint(FMath::CeilToInt(OutputResolution.X / OriginalTileCount.X), FMath::CeilToInt(OutputResolution.Y / OriginalTileCount.Y));
+				FIntPoint BackbufferResolution = FIntPoint(FMath::CeilToInt((float)OutputResolution.X / (float)OriginalTileCount.X), FMath::CeilToInt((float)OutputResolution.Y / (float)OriginalTileCount.Y));
 				FIntPoint TileResolution = BackbufferResolution;
 
 				// Apply size padding.
diff --git a/Engine/Plugins/MovieScene/MovieRenderPipeline/Source/MovieRenderPipelineCore/Public/MoviePipelineHighResSetting.h b/Engine/Plugins/MovieScene/MovieRenderPipeline/Source/MovieRenderPipelineCore/Public/MoviePipelineHighResSetting.h
--- a/Engine/Plugins/MovieScene/MovieRenderPipeline/Source/MovieRenderPipelineCore/Public/MoviePipelineHighResSetting.h
+++ b/Engine/Plugins/MovieScene/MovieRenderPipeline/Source/MovieRenderPipelineCore/Public/MoviePipelineHighResSetting.h
@@ -118,8 +118,8 @@ public:
 				}
 
 				FIntPoint TileResolution = FIntPoint(
-					FMath::CeilToInt(OutputSettings->OutputResolution.X / TileCount),
-					FMath::CeilToInt(OutputSettings->OutputResolution.Y / TileCount));
+					FMath::CeilToInt((float)OutputSettings->OutputResolution.X / (float)TileCount),
+					FMath::CeilToInt((float)OutputSettings->OutputResolution.Y / (float)TileCount));
 
 				TileResolution = CalculatePaddedBackbufferSize(TileResolution);
 
diff --git a/Engine/Source/Developer/LogVisualizer/Private/VisualLoggerCanvasRenderer.cpp b/Engine/Source/Developer/LogVisualizer/Private/VisualLoggerCanvasRenderer.cpp
--- a/Engine/Source/Developer/LogVisualizer/Private/VisualLoggerCanvasRenderer.cpp
+++ b/Engine/Source/Developer/LogVisualizer/Private/VisualLoggerCanvasRenderer.cpp
@@ -217,7 +217,7 @@ void FVisualLoggerCanvasRenderer::DrawHistogramGraphs(class UCanvas* Canvas, cla
 		const FColor GraphsBackgroundColor = ULogVisualizerSettings::StaticClass()->GetDefaultObject<ULogVisualizerSettings>()->GraphsBackgroundColor;
 		const int NumberOfGraphs = CollectedGraphs.Num();
 		const int32 NumberOfColumns = FMath::CeilToInt(FMath::Sqrt(static_cast<float>(NumberOfGraphs)));
-		int32 NumberOfRows = FMath::FloorToInt(NumberOfGraphs / NumberOfColumns);
+		int32 NumberOfRows = FMath::FloorToInt((float)NumberOfGraphs / (float)NumberOfColumns);
 		if (NumberOfGraphs - NumberOfRows * NumberOfColumns > 0)
 		{
 			NumberOfRows += 1;
diff --git a/Engine/Source/Developer/MaterialBaking/Private/MaterialRenderItem.cpp b/Engine/Source/Developer/MaterialBaking/Private/MaterialRenderItem.cpp
--- a/Engine/Source/Developer/MaterialBaking/Private/MaterialRenderItem.cpp
+++ b/Engine/Source/Developer/MaterialBaking/Private/MaterialRenderItem.cpp
@@ -136,7 +136,7 @@ void FMeshMaterialRenderItem::QueueMaterial(FCanvasRenderContext& RenderContext,
 	MeshElement.bWireframe = true;
 #endif
 
-	const int32 NumTris = FMath::TruncToInt(Indices.Num() / 3);
+	const int32 NumTris = FMath::TruncToInt((float)Indices.Num() / 3);
 	if (NumTris == 0)
 	{
 		// there's nothing to do here
diff --git a/Engine/Source/Developer/MeshMergeUtilities/Private/MeshMergeUtilities.cpp b/Engine/Source/Developer/MeshMergeUtilities/Private/MeshMergeUtilities.cpp
--- a/Engine/Source/Developer/MeshMergeUtilities/Private/MeshMergeUtilities.cpp
+++ b/Engine/Source/Developer/MeshMergeUtilities/Private/MeshMergeUtilities.cpp
@@ -1069,8 +1069,8 @@ float FMeshMergeUtilities::FlattenEmissivescale(TArray<struct FFlattenMaterial>&
 	ParallelFor(NumThreads, [&InMaterialList, MaterialsPerThread, Multiplier, MaxScale]
 	(int32 Index)
 	{
-		int32 StartIndex = FMath::CeilToInt((Index)* MaterialsPerThread);
-		const int32 EndIndex = FMath::Min(FMath::CeilToInt((Index + 1) * MaterialsPerThread), InMaterialList.Num());
+		int32 StartIndex = FMath::CeilToInt((float)Index * (float)MaterialsPerThread);
+		const int32 EndIndex = FMath::Min(FMath::CeilToInt((float)(Index + 1) * (float)MaterialsPerThread), InMaterialList.Num());
 
 		for (; StartIndex < EndIndex; ++StartIndex)
 		{
diff --git a/Engine/Source/Editor/LandscapeEditor/Public/LandscapeEditorObject.h b/Engine/Source/Editor/LandscapeEditor/Public/LandscapeEditorObject.h
--- a/Engine/Source/Editor/LandscapeEditor/Public/LandscapeEditorObject.h
+++ b/Engine/Source/Editor/LandscapeEditor/Public/LandscapeEditorObject.h
@@ -708,7 +708,7 @@ public:
 	int32 ClampLandscapeSize(int32 InComponentsCount) const
 	{
 		// Max size is either whole components below 8192 verts, or 32 components
-		return FMath::Clamp(InComponentsCount, 1, FMath::Min(32, FMath::FloorToInt(8191 / (NewLandscape_SectionsPerComponent * NewLandscape_QuadsPerSection))));
+		return FMath::Clamp(InComponentsCount, 1, FMath::Min(32, FMath::FloorToInt(8191.0f / (float)(NewLandscape_SectionsPerComponent * NewLandscape_QuadsPerSection))));
 	}
 	
 	int32 CalcComponentsCount(int32 InResolution) const
diff --git a/Engine/Source/Editor/UnrealEd/Private/Fbx/FbxAnimationExport.cpp b/Engine/Source/Editor/UnrealEd/Private/Fbx/FbxAnimationExport.cpp
--- a/Engine/Source/Editor/UnrealEd/Private/Fbx/FbxAnimationExport.cpp
+++ b/Engine/Source/Editor/UnrealEd/Private/Fbx/FbxAnimationExport.cpp
@@ -739,7 +739,7 @@ void FFbxExporter::ExportAnimTrack(IAnimTrackAdapter& AnimTrackAdapter, AActor*
 		if( NextUpdateTime <= 0.0f )
 		{
 			NextUpdateTime = UpdateFrequency;
-			GWarn->StatusUpdate( FMath::RoundToInt( SampleTime ), FMath::RoundToInt(AnimationLength), NSLOCTEXT("FbxExporter", "ExportingToFbxStatus", "Exporting to FBX") );
+			GWarn->StatusUpdate( FMath::RoundToInt( SampleTime ), AnimationLength, NSLOCTEXT("FbxExporter", "ExportingToFbxStatus", "Exporting to FBX") );
 		}
 
 		TArray<FTransform> LocalBoneTransforms = InSkeletalMeshComponent->GetBoneSpaceTransforms();
diff --git a/Engine/Source/Runtime/Core/Public/Math/UnrealPlatformMathSSE.h b/Engine/Source/Runtime/Core/Public/Math/UnrealPlatformMathSSE.h
--- a/Engine/Source/Runtime/Core/Public/Math/UnrealPlatformMathSSE.h
+++ b/Engine/Source/Runtime/Core/Public/Math/UnrealPlatformMathSSE.h
@@ -30,6 +30,17 @@ namespace SSE
 		return temp;
 	}
 
+	static FORCEINLINE double InvSqrt(double InValue)
+	{
+		const __m128d One = _mm_set_sd(1.0);
+		const __m128d Y0 = _mm_set_sd(InValue);
+		const __m128d X0 = _mm_sqrt_sd(One, Y0);
+		const __m128d R0 = _mm_div_sd(One, X0);
+		double temp;
+		_mm_store_sd(&temp, R0);
+		return temp;
+	}
+
 	static FORCEINLINE float InvSqrtEst(float F)
 	{
 		// Performs one pass of Newton-Raphson iteration on the hardware estimate
@@ -50,11 +61,26 @@ namespace SSE
 		return temp;
 	}
 
+	static FORCEINLINE double InvSqrtEst(double InValue)
+	{		
+		return InvSqrt(InValue);
+	}
+
 	static FORCEINLINE int32 TruncToInt(float F)
 	{
 		return _mm_cvtt_ss2si(_mm_set_ss(F));
 	}
 
+	static FORCEINLINE int32 TruncToInt(double InValue)
+	{
+		return _mm_cvttsd_si32(_mm_set_sd(InValue));
+	}
+
+	static FORCEINLINE int64 TruncToInt64(double InValue)
+	{
+		return _mm_cvttsd_si64(_mm_set_sd(InValue));
+	}
+
 	static FORCEINLINE int32 FloorToInt(float F)
 	{
 		// Note: unlike the Generic solution and the SSE4 float solution, we implement FloorToInt using a rounding instruction, rather than implementing RoundToInt using a floor instruction.  
@@ -62,6 +88,16 @@ namespace SSE
 		return _mm_cvt_ss2si(_mm_set_ss(F + F - 0.5f)) >> 1;
 	}
 
+	static FORCEINLINE int32 FloorToInt(double InValue)
+	{
+		return _mm_cvtsd_si32(_mm_set_sd(InValue + InValue - 0.5)) >> 1;
+	}
+
+	static FORCEINLINE int64 FloorToInt64(double InValue)
+	{
+		return _mm_cvtsd_si64(_mm_set_sd(InValue + InValue - 0.5)) >> 1;
+	}
+
 	static FORCEINLINE int32 RoundToInt(float F)
 	{
 		// Note: the times-2 is to remove the rounding-to-nearest-even-number behavior that mm_cvt_ss2si uses when the fraction is .5
@@ -72,6 +108,16 @@ namespace SSE
 		return _mm_cvt_ss2si(_mm_set_ss(F + F + 0.5f)) >> 1;
 	}
 
+	static FORCEINLINE int32 RoundToInt(double InValue)
+	{
+		return _mm_cvtsd_si32(_mm_set_sd(InValue + InValue + 0.5)) >> 1;
+	}
+
+	static FORCEINLINE int64 RoundToInt64(double InValue)
+	{
+		return _mm_cvtsd_si64(_mm_set_sd(InValue + InValue + 0.5)) >> 1;
+	}
+
 	static FORCEINLINE int32 CeilToInt(float F)
 	{
 		// Note: unlike the Generic solution and the SSE4 float solution, we implement CeilToInt using a rounding instruction, rather than a dedicated ceil instruction
@@ -79,6 +125,16 @@ namespace SSE
 		return -(_mm_cvt_ss2si(_mm_set_ss(-0.5f - (F + F))) >> 1);
 	}
 
+	static FORCEINLINE int32 CeilToInt(double InValue)
+	{
+		return -(_mm_cvtsd_si32(_mm_set_sd(-0.5 - (InValue + InValue))) >> 1);
+	}
+
+	static FORCEINLINE int64 CeilToInt64(double InValue)
+	{
+		return -(_mm_cvtsd_si64(_mm_set_sd(-0.5 - (InValue + InValue))) >> 1);
+	}
+
 	// https://gist.github.com/rygorous/2156668
 	// float_to_half_rtne_SSE2
 	static __m128i FloatToHalf(__m128 f)
@@ -194,13 +250,13 @@ struct TUnrealPlatformMathSSEBase : public Base
 	}
 
 	template<typename T>
-	static FORCEINLINE float InvSqrt(T F)
+	static FORCEINLINE T InvSqrt(T F)
 	{
 		return UE4::SSE::InvSqrt(F);
 	}
 
 	template<typename T>
-	static FORCEINLINE float InvSqrtEst(T F)
+	static FORCEINLINE T InvSqrtEst(T F)
 	{
 		return UE4::SSE::InvSqrtEst(F);
 	}
@@ -214,39 +270,6 @@ struct TUnrealPlatformMathSSEBase : public Base
 	{
 		_mm_storeu_ps(Dst, UE4::SSE::HalfToFloat(_mm_loadu_si64((__m128i*)Src)));
 	}
-
-	// LWC_TODO: Implement vectorized double equivalents. 
-	// LWC_TODO: Default to double variant for non-float types, or just for int64?
-	static FORCEINLINE int32 TruncToInt(double F)
-	{
-		return Base::TruncToInt((float)F);
-	}
-
-	static FORCEINLINE int32 RoundToInt(double F)
-	{
-		return Base::RoundToInt((float)F);
-	}
-
-	static FORCEINLINE int32 FloorToInt(double F)
-	{
-		return Base::FloorToInt((float)F);
-	}
-
-	static FORCEINLINE int32 CeilToInt(double F)
-	{
-		return Base::CeilToInt((float)F);
-	}
-
-	static FORCEINLINE double InvSqrt(double F)
-	{
-		return Base::InvSqrt((float)F);
-	}
-
-	static FORCEINLINE double InvSqrtEst(double F)
-	{
-		return Base::InvSqrtEst((float)F);
-	}
-
 };
 
 #endif // PLATFORM_ENABLE_VECTORINTRINSICS
\ No newline at end of file
diff --git a/Engine/Source/Runtime/Engine/Private/GameViewportClient.cpp b/Engine/Source/Runtime/Engine/Private/GameViewportClient.cpp
--- a/Engine/Source/Runtime/Engine/Private/GameViewportClient.cpp
+++ b/Engine/Source/Runtime/Engine/Private/GameViewportClient.cpp
@@ -1589,7 +1589,7 @@ void UGameViewportClient::Draw(FViewport* InViewport, FCanvas* SceneCanvas)
 			MinY = FMath::Min<uint32>(UpscaledViewRect.Min.Y, MinY);
 			MaxX = FMath::Max<uint32>(UpscaledViewRect.Max.X, MaxX);
 			MaxY = FMath::Max<uint32>(UpscaledViewRect.Max.Y, MaxY);
-			TotalArea += FMath::TruncToInt(UpscaledViewRect.Width()) * FMath::TruncToInt(UpscaledViewRect.Height());
+			TotalArea += UpscaledViewRect.Width() * UpscaledViewRect.Height();
 		}
 
 		// To draw black borders around the rendered image (prevents artifacts from post processing passes that read outside of the image e.g. PostProcessAA)
@@ -1768,7 +1768,7 @@ void UGameViewportClient::Draw(FViewport* InViewport, FCanvas* SceneCanvas)
 					if (View != NULL)
 					{
 						// rendering to directly to viewport target
-						FVector CanvasOrigin(FMath::TruncToFloat(View->UnscaledViewRect.Min.X), FMath::TruncToInt(View->UnscaledViewRect.Min.Y), 0.f);
+						FVector CanvasOrigin(FMath::TruncToFloat(View->UnscaledViewRect.Min.X), FMath::TruncToFloat(View->UnscaledViewRect.Min.Y), 0.f);
 
 						CanvasObject->Init(View->UnscaledViewRect.Width(), View->UnscaledViewRect.Height(), View, SceneCanvas);
 
diff --git a/Engine/Source/Runtime/Engine/Private/Particles/SubUVAnimation.cpp b/Engine/Source/Runtime/Engine/Private/Particles/SubUVAnimation.cpp
--- a/Engine/Source/Runtime/Engine/Private/Particles/SubUVAnimation.cpp
+++ b/Engine/Source/Runtime/Engine/Private/Particles/SubUVAnimation.cpp
@@ -566,8 +566,8 @@ void FSubUVDerivedData::Build(UTexture2D* SubUVTexture, int32 SubImages_Horizont
 
 					for (int32 Y = 0; Y < SubImageSizeY; Y++)
 					{
-						int32 GlobalY = FMath::RoundToInt(SubImageY * SubImageSizeYFloat) + Y;
-						int32 NextGlobalY = FMath::RoundToInt(NextSubImageY * SubImageSizeYFloat) + Y;
+						int32 GlobalY = FMath::RoundToInt((float)SubImageY * SubImageSizeYFloat) + Y;
+						int32 NextGlobalY = FMath::RoundToInt((float)NextSubImageY * SubImageSizeYFloat) + Y;
 
 						for (int32 X = 0; X < SubImageSizeX; X++)
 						{
diff --git a/Engine/Source/Runtime/Landscape/Private/LandscapeEditLayers.cpp b/Engine/Source/Runtime/Landscape/Private/LandscapeEditLayers.cpp
--- a/Engine/Source/Runtime/Landscape/Private/LandscapeEditLayers.cpp
+++ b/Engine/Source/Runtime/Landscape/Private/LandscapeEditLayers.cpp
@@ -2327,8 +2327,8 @@ void ALandscape::GenerateLayersRenderQuadsAtlas(const FIntPoint& InSectionBase,
 	int32 LocalComponentSizeQuad = InSubSectionSizeQuad * NumSubsections;
 	int32 LocalComponentSizeVerts = SubsectionSizeVerts * NumSubsections;
 
-	FVector2D PositionOffset(FMath::RoundToInt(ComponentSectionBase.X / LocalComponentSizeQuad), FMath::RoundToInt(ComponentSectionBase.Y / LocalComponentSizeQuad));
-	FVector2D ComponentsPerTexture(FMath::RoundToInt(InWriteSize.X / LocalComponentSizeQuad), FMath::RoundToInt(InWriteSize.Y / LocalComponentSizeQuad));
+	FVector2D PositionOffset(FMath::RoundToInt((float)ComponentSectionBase.X / (float)LocalComponentSizeQuad), FMath::RoundToInt((float)ComponentSectionBase.Y / (float)LocalComponentSizeQuad));
+	FVector2D ComponentsPerTexture(FMath::RoundToInt((float)InWriteSize.X / (float)LocalComponentSizeQuad), FMath::RoundToInt((float)InWriteSize.Y / (float)LocalComponentSizeQuad));
 
 	if (InReadSize.X >= InWriteSize.X)
 	{
@@ -2413,8 +2413,8 @@ void ALandscape::GenerateLayersRenderQuadsMip(const FIntPoint& InSectionBase, co
 	int32 MipSubsectionSizeVerts = SubsectionSizeVerts >> InCurrentMip;
 	int32 MipLocalComponentSizeVerts = MipSubsectionSizeVerts * NumSubsections;
 
-	FVector2D PositionOffset(FMath::RoundToInt(InSectionBase.X / LocalComponentSizeQuad), FMath::RoundToInt(InSectionBase.Y / LocalComponentSizeQuad));
-	FVector2D ComponentsPerTexture(FMath::RoundToInt(InWriteSize.X / LocalComponentSizeQuad), FMath::RoundToInt(InWriteSize.Y / LocalComponentSizeQuad));
+	FVector2D PositionOffset(FMath::RoundToInt((float)InSectionBase.X / (float)LocalComponentSizeQuad), FMath::RoundToInt((float)InSectionBase.Y / (float)LocalComponentSizeQuad));
+	FVector2D ComponentsPerTexture(FMath::RoundToInt((float)InWriteSize.X / (float)LocalComponentSizeQuad), FMath::RoundToInt((float)InWriteSize.Y / (float)LocalComponentSizeQuad));
 
 	FIntPoint ComponentSectionBase(PositionOffset.X * MipLocalComponentSizeVerts, PositionOffset.Y * MipLocalComponentSizeVerts);
 	FIntPoint UVComponentSectionBase(PositionOffset.X * LocalComponentSizeVerts, PositionOffset.Y * LocalComponentSizeVerts);
@@ -2505,7 +2505,7 @@ void ALandscape::GenerateLayersRenderQuadsNonAtlasToAtlas(const FIntPoint& InSec
 	int32 LocalComponentSizeQuad = InSubSectionSizeQuad * NumSubsections;
 	int32 LocalComponentSizeVerts = SubsectionSizeVerts * NumSubsections;
 
-	FVector2D PositionOffset(FMath::RoundToInt(InSectionBase.X / LocalComponentSizeQuad), FMath::RoundToInt(InSectionBase.Y / LocalComponentSizeQuad));
+	FVector2D PositionOffset(FMath::RoundToInt((float)InSectionBase.X / (float)LocalComponentSizeQuad), FMath::RoundToInt((float)InSectionBase.Y / (float)LocalComponentSizeQuad));
 	FIntPoint ComponentSectionBase(PositionOffset.X * LocalComponentSizeVerts, PositionOffset.Y * LocalComponentSizeVerts);
 	FVector2D UVSize((float)SubsectionSizeVerts / (float)InReadSize.X, (float)SubsectionSizeVerts / (float)InReadSize.Y);
 
@@ -2929,7 +2929,7 @@ int32 ALandscape::RegenerateLayersHeightmaps(FTextureToComponentHelper const& Ma
 				int32 Index = OutHeightmaps.IndexOfByPredicate([=](const FHeightmapCopyInfo& LayerHeightmap) { return LayerHeightmap.Texture == ComponentHeightmap; });
 
 				FIntPoint ComponentSectionBase = Component->GetSectionBase() - LandscapeBaseQuads;
-				FVector2D SourcePositionOffset(FMath::RoundToInt(ComponentSectionBase.X / ComponentSizeQuad), FMath::RoundToInt(ComponentSectionBase.Y / ComponentSizeQuad));
+				FVector2D SourcePositionOffset(FMath::RoundToInt((float)ComponentSectionBase.X / ComponentSizeQuad), FMath::RoundToInt((float)ComponentSectionBase.Y / ComponentSizeQuad));
 				FIntPoint ComponentVertexPosition = FIntPoint(SourcePositionOffset.X * ComponentSizeVerts, SourcePositionOffset.Y * ComponentSizeVerts);
 				ALandscapeProxy* Proxy = Component->GetLandscapeProxy();
 
@@ -3467,7 +3467,7 @@ void ALandscape::PrepareComponentDataToExtractMaterialLayersCS(const TArray<ULan
 				check(LayerWeightmapUsage != nullptr);
 
 				FIntPoint ComponentSectionBase = Component->GetSectionBase() - InLandscapeBase;
-				FVector2D SourcePositionOffset(FMath::RoundToInt(ComponentSectionBase.X / LocalComponentSizeQuad), FMath::RoundToInt(ComponentSectionBase.Y / LocalComponentSizeQuad));
+				FVector2D SourcePositionOffset(FMath::RoundToInt((float)ComponentSectionBase.X / LocalComponentSizeQuad), FMath::RoundToInt((float)ComponentSectionBase.Y / LocalComponentSizeQuad));
 				FIntPoint SourceComponentVertexPosition = FIntPoint(SourcePositionOffset.X * LocalComponentSizeVerts, SourcePositionOffset.Y * LocalComponentSizeVerts);
 
 				FLandscapeLayersCopyTextureParams& CopyTextureParams = DeferredCopyTextures.Add_GetRef(FLandscapeLayersCopyTextureParams(*LayerWeightmap->GetName(), LayerWeightmap->GetResource(), FString::Printf(TEXT("%s WeightmapScratchTexture"), *InLayer.Name.ToString()), InOutTextureData));
@@ -3490,7 +3490,7 @@ void ALandscape::PrepareComponentDataToExtractMaterialLayersCS(const TArray<ULan
 						FIntPoint DestComponentSectionBase = DestComponent->GetSectionBase() - InLandscapeBase;
 
 						// Compute component top left vertex position from section base info
-						FVector2D DestPositionOffset(FMath::RoundToInt(DestComponentSectionBase.X / LocalComponentSizeQuad), FMath::RoundToInt(DestComponentSectionBase.Y / LocalComponentSizeQuad));
+						FVector2D DestPositionOffset(FMath::RoundToInt((float)DestComponentSectionBase.X / LocalComponentSizeQuad), FMath::RoundToInt((float)DestComponentSectionBase.Y / LocalComponentSizeQuad));
 
 						Data.ComponentVertexPosition = SourceComponentVertexPosition;
 						Data.AtlasTexturePositionOutput = FIntPoint(DestPositionOffset.X * LocalComponentSizeVerts, DestPositionOffset.Y * LocalComponentSizeVerts);
@@ -3600,7 +3600,7 @@ void ALandscape::PrepareComponentDataToPackMaterialLayersCS(int32 InCurrentWeigh
 						// Compute component top left vertex position from section base info
 						int32 LocalComponentSizeQuad = ChannelComponent->SubsectionSizeQuads * NumSubsections;
 						int32 LocalComponentSizeVerts = (ChannelComponent->SubsectionSizeQuads + 1) * NumSubsections;
-						FVector2D PositionOffset(FMath::RoundToInt(ComponentSectionBase.X / LocalComponentSizeQuad), FMath::RoundToInt(ComponentSectionBase.Y / LocalComponentSizeQuad));
+						FVector2D PositionOffset(FMath::RoundToInt((float)ComponentSectionBase.X / LocalComponentSizeQuad), FMath::RoundToInt((float)ComponentSectionBase.Y / LocalComponentSizeQuad));
 
 						Data.ComponentVertexPositionX[WeightmapChannelIndex] = PositionOffset.X * LocalComponentSizeVerts;
 						Data.ComponentVertexPositionY[WeightmapChannelIndex] = PositionOffset.Y * LocalComponentSizeVerts;
diff --git a/Engine/Source/Runtime/NavigationSystem/Private/NavMesh/NavMeshRenderingComponent.cpp b/Engine/Source/Runtime/NavigationSystem/Private/NavMesh/NavMeshRenderingComponent.cpp
--- a/Engine/Source/Runtime/NavigationSystem/Private/NavMesh/NavMeshRenderingComponent.cpp
+++ b/Engine/Source/Runtime/NavigationSystem/Private/NavMesh/NavMeshRenderingComponent.cpp
@@ -1086,7 +1086,7 @@ FNavMeshSceneProxy::FNavMeshSceneProxy(const UPrimitiveComponent* InComponent, F
 
 		FMeshBatchElement Element;
 		Element.FirstIndex = IndexBuffer.Indices.Num();
-		Element.NumPrimitives = FMath::FloorToInt(CurrentMeshBuilder.Indices.Num() / 3);
+		Element.NumPrimitives = FMath::FloorToInt((float)CurrentMeshBuilder.Indices.Num() / 3);
 		Element.MinVertexIndex = Vertices.Num();
 		Element.MaxVertexIndex = Element.MinVertexIndex + CurrentMeshBuilder.Vertices.Num() - 1;
 		Element.IndexBuffer = &IndexBuffer;
diff --git a/Engine/Source/Runtime/Slate/Private/Framework/Application/SlateApplication.cpp b/Engine/Source/Runtime/Slate/Private/Framework/Application/SlateApplication.cpp
--- a/Engine/Source/Runtime/Slate/Private/Framework/Application/SlateApplication.cpp
+++ b/Engine/Source/Runtime/Slate/Private/Framework/Application/SlateApplication.cpp
@@ -115,10 +115,10 @@ public:
 	{
 		if (Bounds)
 		{
-			CursorClipRect.Min.X = FMath::TruncToInt(Bounds->left);
-			CursorClipRect.Min.Y = FMath::TruncToInt(Bounds->top);
-			CursorClipRect.Max.X = FMath::TruncToInt(Bounds->right) - 1;
-			CursorClipRect.Max.Y = FMath::TruncToInt(Bounds->bottom) - 1;
+			CursorClipRect.Min.X = Bounds->left;
+			CursorClipRect.Min.Y = Bounds->top;
+			CursorClipRect.Max.X = Bounds->right - 1;
+			CursorClipRect.Max.Y = Bounds->bottom - 1;
 		}
 		else
 		{
diff --git a/Engine/Source/Runtime/Slate/Public/Widgets/Views/STileView.h b/Engine/Source/Runtime/Slate/Public/Widgets/Views/STileView.h
--- a/Engine/Source/Runtime/Slate/Public/Widgets/Views/STileView.h
+++ b/Engine/Source/Runtime/Slate/Public/Widgets/Views/STileView.h
@@ -451,7 +451,7 @@ protected:
 				// Only scroll the item into view if it's not already in the visible range
 				const int32 NumItemsPerLine = GetNumItemsPerLine();
 				const double ScrollLineOffset = this->GetTargetScrollOffset() / NumItemsPerLine;
-				const int32 LineOfItem = FMath::FloorToInt(IndexOfItem / NumItemsPerLine);
+				const int32 LineOfItem = FMath::FloorToInt((float)IndexOfItem / (float)NumItemsPerLine);
 				const int32 NumFullLinesInView = FMath::FloorToInt(ScrollLineOffset + NumLinesInView) - FMath::CeilToInt(ScrollLineOffset);
 				
 				const double MinDisplayedLine = this->bNavigateOnScrollIntoView ? FMath::FloorToDouble(ScrollLineOffset) : FMath::CeilToDouble(ScrollLineOffset);
