Added all of the current ray tracing shader variations to the list of excluded shader types for landscape thumbnails

diff --git a/Engine/Source/Runtime/Landscape/Private/LandscapeRender.cpp b/Engine/Source/Runtime/Landscape/Private/LandscapeRender.cpp
--- a/Engine/Source/Runtime/Landscape/Private/LandscapeRender.cpp
+++ b/Engine/Source/Runtime/Landscape/Private/LandscapeRender.cpp
@@ -3424,7 +3424,10 @@ public:
 					}
 					else
 					{
-						if (Algo::Find(GetExcludedShaderTypesInThumbnailRender(), ShaderType->GetFName()))
+						// No ray tracing on thumbnails : we don't need any variation of ray hit group shaders : 
+						const bool bIsRayHitGroupShader = (ShaderType->GetFrequency() == SF_RayHitGroup);
+						if (bIsRayHitGroupShader 
+							|| Algo::Find(GetExcludedShaderTypesInThumbnailRender(), ShaderType->GetFName()))
 						{
 							UE_LOG(LogLandscape, VeryVerbose, TEXT("Excluding shader %s from landscape thumbnail material"), ShaderType->GetName());
 							return false;
@@ -3727,8 +3730,6 @@ public:
 
 #if RHI_RAYTRACING
 			// No ray tracing on thumbnails
-			FName(TEXT("TMaterialCHSFPrecomputedVolumetricLightmapLightingPolicy")),
-			FName(TEXT("TMaterialCHSFNoLightMapPolicy")),
 			FName(TEXT("FRayTracingDynamicGeometryConverterCS")),
 			FName(TEXT("FTrivialMaterialCHS")),
 #endif // RHI_RAYTRACING
