Landscape -  LOD 0 Screen Size setting - Remove hard mention of value in the tooltip. Easy to get out of sync.

diff --git a/Engine/Source/Runtime/Landscape/Classes/LandscapeProxy.h b/Engine/Source/Runtime/Landscape/Classes/LandscapeProxy.h
--- a/Engine/Source/Runtime/Landscape/Classes/LandscapeProxy.h
+++ b/Engine/Source/Runtime/Landscape/Classes/LandscapeProxy.h
@@ -396,7 +396,7 @@ public:
 	UPROPERTY(EditAnywhere, Category = LOD, meta=(ClampMin = "0.01", ClampMax = "1.0", UIMin = "0.01", UIMax = "1.0", DisplayName= "SubSection Min Component ScreenSize"))
 	float ComponentScreenSizeToUseSubSections;
 
-	/** This is the starting screen size used to calculate the distribution, by default it's 1, but you can increase the value if you want less LOD0 component, and you use very large landscape component. */
+	/** This is the starting screen size used to calculate the distribution. You can increase the value if you want less LOD0 component, and you use very large landscape component. */
 	UPROPERTY(EditAnywhere, Category = "LOD Distribution", meta = (DisplayName = "LOD 0 Screen Size", ClampMin = "0.1", ClampMax = "10.0", UIMin = "1.0", UIMax = "10.0"))
 	float LOD0ScreenSize;
 
