Get static parameter values prior to updating static permutation for material instance in landscape components.

diff --git a/Engine/Source/Runtime/Landscape/Private/LandscapeEdit.cpp b/Engine/Source/Runtime/Landscape/Private/LandscapeEdit.cpp
--- a/Engine/Source/Runtime/Landscape/Private/LandscapeEdit.cpp
+++ b/Engine/Source/Runtime/Landscape/Private/LandscapeEdit.cpp
@@ -315,6 +315,8 @@ UMaterialInstanceConstant* ULandscapeComponent::GetCombinationMaterial(FMaterial
 			}
 
 			FStaticParameterSet StaticParameters;
+			CombinationMaterialInstance->GetStaticParameterValues(StaticParameters);
+
 			for (const FWeightmapLayerAllocationInfo& Allocation : Allocations)
 			{
 				if (Allocation.LayerInfo)
@@ -323,8 +325,8 @@ UMaterialInstanceConstant* ULandscapeComponent::GetCombinationMaterial(FMaterial
 					StaticParameters.TerrainLayerWeightParameters.Add(FStaticTerrainLayerWeightParameter(LayerParameter, Allocation.WeightmapTextureIndex, true, FGuid(), !Allocation.LayerInfo->bNoWeightBlend));
 				}
 			}
-			CombinationMaterialInstance->UpdateStaticPermutation(StaticParameters, InMaterialUpdateContext);
 
+			CombinationMaterialInstance->UpdateStaticPermutation(StaticParameters, InMaterialUpdateContext);
 			CombinationMaterialInstance->PostEditChange();
 		}
 
