Landscape edit layers fixes, preventing landscape to properly be carried over when converting a map to World Partition:
* Gather Heightmaps/Weightmaps even when doing a non-partial update, otherwise, no texture is resolved when bFlushRender is used
* Fixed draw event tag for weightmaps

#tests editor, full QA pass on landscape

diff --git a/Engine/Source/Runtime/Landscape/Private/LandscapeEditLayers.cpp b/Engine/Source/Runtime/Landscape/Private/LandscapeEditLayers.cpp
--- a/Engine/Source/Runtime/Landscape/Private/LandscapeEditLayers.cpp
+++ b/Engine/Source/Runtime/Landscape/Private/LandscapeEditLayers.cpp
@@ -2109,7 +2109,7 @@ void ALandscape::DrawWeightmapComponentsToRenderTarget(const FString& InDebugNam
 		[LayersRender, InDebugName, InDrawType, InClearRTWrite](FRHICommandListImmediate& RHICmdList) mutable
 	{
 		TRACE_CPUPROFILER_EVENT_SCOPE(LandscapeLayers_RT_RenderWeightmap);
-		SCOPED_DRAW_EVENTF(RHICmdList, LandscapeLayers, TEXT("DrawWeightmapComponentsToRenderTarget %s (%s)"), *InDebugName, *StaticEnum<EHeightmapRTType>()->GetDisplayValueAsText(InDrawType).ToString());
+		SCOPED_DRAW_EVENTF(RHICmdList, LandscapeLayers, TEXT("DrawWeightmapComponentsToRenderTarget %s (%s)"), *InDebugName, *StaticEnum<EWeightmapRTType>()->GetDisplayValueAsText(InDrawType).ToString());
 		LayersRender.Render(RHICmdList, InClearRTWrite);
 	});
 
@@ -5089,24 +5089,26 @@ void ALandscape::UpdateLayersContent(bool bInWaitForStreaming, bool bInSkipMonit
 			{
 				// Gather Neighbors (Neighbors need to be Rendered but not resolved so that the resolved Components have valid normals on edges)
 				GetLandscapeComponentNeighborsToRender(Component, NeighborsComponents);
-				// Gather Heightmaps (All Components sharing Heightmap textures need to be rendered and resolved)
-				Heightmaps.Add(Component->GetHeightmap(false));
 				Component->ForEachLayer([&](const FGuid&, FLandscapeLayerComponentData& LayerData)
 				{
 					HeightmapsToRender.Add(LayerData.HeightmapData.Texture);
 				});
-				// Gather Weightmaps
-				TArray<UTexture2D*>& WeightmapTextures = Component->GetWeightmapTextures();
-				for (FWeightmapLayerAllocationInfo const& AllocInfo : Component->GetWeightmapLayerAllocations())
-				{
-					if (AllocInfo.IsAllocated() && AllocInfo.WeightmapTextureIndex < WeightmapTextures.Num())
-					{
-						Weightmaps.Add(WeightmapTextures[AllocInfo.WeightmapTextureIndex]);
-					}
-				}
 				// Gather WeightmapUsages (Components sharing weightmap usages with the resolved Components need to be rendered so that the resolving is valid)
 				GetLandscapeComponentWeightmapsToRender(Component, WeightmapsComponents);
 			}
+
+			// Gather Heightmaps (All Components sharing Heightmap textures need to be rendered and resolved)
+			Heightmaps.Add(Component->GetHeightmap(false));
+
+			// Gather Weightmaps
+			TArray<UTexture2D*>& WeightmapTextures = Component->GetWeightmapTextures();
+			for (FWeightmapLayerAllocationInfo const& AllocInfo : Component->GetWeightmapLayerAllocations())
+			{
+				if (AllocInfo.IsAllocated() && AllocInfo.WeightmapTextureIndex < WeightmapTextures.Num())
+				{
+					Weightmaps.Add(WeightmapTextures[AllocInfo.WeightmapTextureIndex]);
+				}
+			}
 		}
 		else
 		{
