Fix Landscape GetXYOffsetData to take mip level into account.

diff --git a/Engine/Source/Runtime/Landscape/Private/LandscapeDataAccess.cpp b/Engine/Source/Runtime/Landscape/Private/LandscapeDataAccess.cpp
--- a/Engine/Source/Runtime/Landscape/Private/LandscapeDataAccess.cpp
+++ b/Engine/Source/Runtime/Landscape/Private/LandscapeDataAccess.cpp
@@ -128,7 +128,7 @@ LANDSCAPE_API FColor* FLandscapeComponentDataInterface::GetXYOffsetData(int32 Lo
 	check(LocalX >= 0 && LocalY >= 0 && LocalX < Component->ComponentSizeQuads + 1 && LocalY < Component->ComponentSizeQuads + 1);
 #endif
 
-	const int32 WeightmapSize = (Component->SubsectionSizeQuads + 1) * Component->NumSubsections;
+	const int32 WeightmapSize = ((Component->SubsectionSizeQuads + 1) * Component->NumSubsections) >> MipLevel;
 	int32 SubNumX;
 	int32 SubNumY;
 	int32 SubX;
