Landscape: Assign same Guid to all Landscape/Proxies/SplineActors sharing a Guid in the new duplicated world (regression tested against JIRA still fixed: UE-34890) (SaveAs prep work)

diff --git a/Engine/Source/Runtime/Landscape/Classes/LandscapeSplineActor.h b/Engine/Source/Runtime/Landscape/Classes/LandscapeSplineActor.h
--- a/Engine/Source/Runtime/Landscape/Classes/LandscapeSplineActor.h
+++ b/Engine/Source/Runtime/Landscape/Classes/LandscapeSplineActor.h
@@ -20,6 +20,7 @@ protected:
 	
 public:
 	virtual FGuid GetLandscapeGuid() const override { return LandscapeGuid; }
+	void SetLandscapeGuid(const FGuid& Guid) { LandscapeGuid = Guid; }
 	virtual ULandscapeSplinesComponent* GetSplinesComponent() const override;
 	virtual FTransform LandscapeActorToWorld() const override;
 	virtual ULandscapeInfo* GetLandscapeInfo() const override;
diff --git a/Engine/Source/Runtime/Landscape/Private/LandscapeEdit.cpp b/Engine/Source/Runtime/Landscape/Private/LandscapeEdit.cpp
--- a/Engine/Source/Runtime/Landscape/Private/LandscapeEdit.cpp
+++ b/Engine/Source/Runtime/Landscape/Private/LandscapeEdit.cpp
@@ -3965,18 +3965,13 @@ void ALandscape::PostEditImport()
 
 void ALandscape::PostDuplicate(bool bDuplicateForPIE)
 {
+	Super::PostDuplicate(bDuplicateForPIE);
+
 	if (!bDuplicateForPIE)
 	{
-		// Need to generate new GUID when duplicating
-		LandscapeGuid = FGuid::NewGuid();
-		// This makes sure at least we have a LandscapeInfo mapped for this GUID.
-		CreateLandscapeInfo();
-
 		// Some edit layers could be affected by BP brushes, which might need to be updated when the landscape is transformed :
 		RequestLayersContentUpdate(ELandscapeLayerUpdateMode::Update_All);
 	}
-
-	Super::PostDuplicate(bDuplicateForPIE);
 }
 #endif	//WITH_EDITOR
 
diff --git a/Engine/Source/Runtime/Landscape/Private/LandscapeModule.cpp b/Engine/Source/Runtime/Landscape/Private/LandscapeModule.cpp
--- a/Engine/Source/Runtime/Landscape/Private/LandscapeModule.cpp
+++ b/Engine/Source/Runtime/Landscape/Private/LandscapeModule.cpp
@@ -14,7 +14,10 @@
 
 #include "LandscapeProxy.h"
 #include "Landscape.h"
+#include "LandscapeSplineActor.h"
 #include "Engine/Texture2D.h"
+#include "EngineUtils.h"
+
 
 // Register the custom version with core
 FCustomVersionRegistration GRegisterLandscapeCustomVersion(FLandscapeCustomVersion::GUID, FLandscapeCustomVersion::LatestVersion, TEXT("Landscape"));
@@ -172,6 +175,25 @@ void WorldDuplicateEventFunction(UWorld* World, bool bDuplicateForPIE, TMap<UObj
 	{
 		AddPerWorldLandscapeData(World);
 	}
+
+#if WITH_EDITOR
+	// Fixup LandscapeGuid on World duplication
+	if (!bDuplicateForPIE)
+	{
+		TMap<FGuid, FGuid> NewLandscapeGuids;
+		for (ALandscapeProxy* Proxy : TActorRange<ALandscapeProxy>(World, ALandscapeProxy::StaticClass(), EActorIteratorFlags::SkipPendingKill))
+		{
+			FGuid& NewGuid = NewLandscapeGuids.FindOrAdd(Proxy->GetLandscapeGuid(), FGuid::NewGuid());
+			Proxy->SetLandscapeGuid(NewGuid);
+		}
+
+		for (ALandscapeSplineActor* SplineActor : TActorRange<ALandscapeSplineActor>(World, ALandscapeSplineActor::StaticClass(), EActorIteratorFlags::SkipPendingKill))
+		{
+			FGuid& NewGuid = NewLandscapeGuids.FindOrAdd(SplineActor->GetLandscapeGuid(), FGuid::NewGuid());
+			SplineActor->SetLandscapeGuid(NewGuid);
+		}
+	}
+#endif
 }
 
 void FLandscapeModule::StartupModule()
