Fix nullptr access for Heighfield data when filling landscape data for niagara DI
#robmerge Release-17.00

diff --git a/Engine/Source/Runtime/Landscape/Private/LandscapeCollision.cpp b/Engine/Source/Runtime/Landscape/Private/LandscapeCollision.cpp
--- a/Engine/Source/Runtime/Landscape/Private/LandscapeCollision.cpp
+++ b/Engine/Source/Runtime/Landscape/Private/LandscapeCollision.cpp
@@ -2981,6 +2981,11 @@ TOptional<float> ULandscapeHeightfieldCollisionComponent::GetHeight(float X, flo
 bool ULandscapeHeightfieldCollisionComponent::FillHeightTile(TArrayView<float> Heights, int32 Offset, int32 Stride) const
 {
 #if WITH_CHAOS
+	if (!IsValidRef(HeightfieldRef))
+	{
+		return false;
+	}
+
 	TUniquePtr<Chaos::FHeightField>& HeightFieldData = HeightfieldRef->Heightfield;
 
 	// If the heightfield data isn't valid, simply return
@@ -3023,6 +3028,11 @@ bool ULandscapeHeightfieldCollisionComponent::FillHeightTile(TArrayView<float> H
 bool ULandscapeHeightfieldCollisionComponent::FillMaterialIndexTile(TArrayView<uint8> Materials, int32 Offset, int32 Stride) const
 {
 #if WITH_CHAOS
+	if (!IsValidRef(HeightfieldRef))
+	{
+		return false;
+	}
+
 	TUniquePtr<Chaos::FHeightField>& HeightFieldData = HeightfieldRef->Heightfield;
 
 	// If the heightfield data isn't valid, simply return
