Fixed overzealous assert for landscape editor. There are some situations where reallocating weightmaps for the splines layer within a transaction is a valid case (Move component tool)

#rnx

diff --git a/Engine/Source/Runtime/Landscape/Private/LandscapeEdit.cpp b/Engine/Source/Runtime/Landscape/Private/LandscapeEdit.cpp
--- a/Engine/Source/Runtime/Landscape/Private/LandscapeEdit.cpp
+++ b/Engine/Source/Runtime/Landscape/Private/LandscapeEdit.cpp
@@ -5695,10 +5695,6 @@ void ULandscapeComponent::ReallocateWeightmaps(FLandscapeEditDataInterface* Data
 	check(!TargetProxy->HasLayersContent() || !InCanUseEditingWeightmap || EditingLayerGUID.IsValid());
 	FGuid TargetLayerGuid = InCanUseEditingWeightmap ? EditingLayerGUID : FGuid();
 
-	// We shouldn't ever be recording a transaction when reallocating weightmaps for the splines edit layer (procedural layer) :
-	FLandscapeLayer* SplinesLayer = GetLandscapeActor()->GetLandscapeSplinesReservedLayer();
-	check((SplinesLayer == nullptr) || (TargetLayerGuid != SplinesLayer->Guid) || (GUndo == nullptr));
-
 	TArray<FWeightmapLayerAllocationInfo>& ComponentWeightmapLayerAllocations = GetWeightmapLayerAllocations(InCanUseEditingWeightmap);
 	TArray<UTexture2D*>& ComponentWeightmapTextures = GetWeightmapTextures(InCanUseEditingWeightmap);
 	TArray<ULandscapeWeightmapUsage*>& ComponentWeightmapTexturesUsage = GetWeightmapTexturesUsage(InCanUseEditingWeightmap);
diff --git a/Engine/Source/Runtime/Landscape/Private/LandscapeEditLayers.cpp b/Engine/Source/Runtime/Landscape/Private/LandscapeEditLayers.cpp
--- a/Engine/Source/Runtime/Landscape/Private/LandscapeEditLayers.cpp
+++ b/Engine/Source/Runtime/Landscape/Private/LandscapeEditLayers.cpp
@@ -6483,8 +6483,6 @@ uint32 ULandscapeComponent::ComputeLayerHash() const
 
 void ALandscape::UpdateLandscapeSplines(const FGuid& InTargetLayer, bool bInUpdateOnlySelected, bool bInForceUpdateAllCompoments)
 {
-	check(GUndo == nullptr);
-
 	TRACE_CPUPROFILER_EVENT_SCOPE(LandscapeLayers_UpdateLandscapeSplines);
 	check(CanHaveLayersContent());
 	ULandscapeInfo* LandscapeInfo = GetLandscapeInfo();
