Fixed missing row of vertices when using the rendered WPO code path for landscape static lighting (contributed by [REDACTED_PERSON])

#tests editor

diff --git a/Engine/Source/Runtime/Landscape/Private/LandscapeLight.cpp b/Engine/Source/Runtime/Landscape/Private/LandscapeLight.cpp
--- a/Engine/Source/Runtime/Landscape/Private/LandscapeLight.cpp
+++ b/Engine/Source/Runtime/Landscape/Private/LandscapeLight.cpp
@@ -458,7 +458,8 @@ void FLandscapeStaticLightingMesh::GetHeightmapData(int32 InLOD, int32 GeometryL
 			RenderedWPOData = LandscapeComponent->RenderWPOHeightmap(InLOD);
 		}
 
-		for (int32 Y = 0; Y < ComponentSizeQuads + 1; Y++)
+		const int32 ComponentSizeVerts = ComponentSizeQuads + 1;
+		for (int32 Y = 0; Y < ComponentSizeVerts; Y++)
 		{
 			const FColor* const Data = DataInterface.GetHeightData(0, Y);
 
@@ -475,8 +476,8 @@ void FLandscapeStaticLightingMesh::GetHeightmapData(int32 InLOD, int32 GeometryL
 			if (bUseRenderedWPO && RenderedWPOData.Num())
 			{
 				const int32 HeightDataOffset = ExpandQuadsX + (Y + ExpandQuadsY) * NumVertices;
-				const int32 WPODataOffset    = Y * (ComponentSizeQuads + 1);
-				for (int32 X = 0; X < ComponentSizeQuads; ++X)
+				const int32 WPODataOffset    = Y * ComponentSizeVerts;
+				for (int32 X = 0; X < ComponentSizeVerts; ++X)
 				{
 					const uint16 WPOHeight = RenderedWPOData[X + WPODataOffset];
 					FColor& Height = HeightData[X + HeightDataOffset];
