Fixed crash when converting old landscape to edit layers or when deleting invalid paint layers in the landscape panel and re-abilitating them after : in the first case, there will be unassigned allocations in the component's base weightmap allocation list and in the second case, there will just be no allocation left. Yet, the list of weightmap textures in the component can still contains entries, and we'd be left with null weightmap usages in the component's list after calling InitializeProxyLayersWeightmapUsages.
We now make sure that the latter function handles those cases by not skipping unassigned allocations and by allocating empty weightmap usages for the textures that are not used within a given component

#tests editor

diff --git a/Engine/Source/Runtime/Landscape/Classes/LandscapeComponent.h b/Engine/Source/Runtime/Landscape/Classes/LandscapeComponent.h
--- a/Engine/Source/Runtime/Landscape/Classes/LandscapeComponent.h
+++ b/Engine/Source/Runtime/Landscape/Classes/LandscapeComponent.h
@@ -761,6 +761,7 @@ public:
 	LANDSCAPE_API const TArray<ULandscapeWeightmapUsage*>& GetWeightmapTexturesUsage(bool InReturnEditingWeightmap = false) const;
 	LANDSCAPE_API TArray<ULandscapeWeightmapUsage*>& GetWeightmapTexturesUsage(const FGuid& InLayerGuid);
 	LANDSCAPE_API const TArray<ULandscapeWeightmapUsage*>& GetWeightmapTexturesUsage(const FGuid& InLayerGuid) const;
+	LANDSCAPE_API void InitializeLayersWeightmapUsage(const FGuid& InLayerGuid);
 
 	LANDSCAPE_API bool HasLayersData() const;
 	LANDSCAPE_API const FLandscapeLayerComponentData* GetLayerData(const FGuid& InLayerGuid) const;
diff --git a/Engine/Source/Runtime/Landscape/Private/LandscapeEditLayers.cpp b/Engine/Source/Runtime/Landscape/Private/LandscapeEditLayers.cpp
--- a/Engine/Source/Runtime/Landscape/Private/LandscapeEditLayers.cpp
+++ b/Engine/Source/Runtime/Landscape/Private/LandscapeEditLayers.cpp
@@ -1948,77 +1948,102 @@ void ALandscapeProxy::InitializeProxyLayersWeightmapUsage()
 {
 	if (ALandscape* Landscape = GetLandscapeActor())
 	{
+		// Reset the entire proxy's usage map and then request all components to repopulate it :
 		WeightmapUsageMap.Reset();
 		for (ULandscapeComponent* Component : LandscapeComponents)
 		{
 			// Reinitialize the weightmap usages for the base (final) paint layers allocations :
+			Component->InitializeLayersWeightmapUsage(FGuid());
+
+			const FLandscapeLayer* SplinesLayer = Landscape->GetLandscapeSplinesReservedLayer();
+			for (const FLandscapeLayer& Layer : Landscape->LandscapeLayers)
 			{
-				const TArray<FWeightmapLayerAllocationInfo>& ComponentWeightmapLayerAllocations = Component->GetWeightmapLayerAllocations(/*InCanUseEditingWeightmap = */false);
-				const TArray<UTexture2D*>& ComponentWeightmapTextures = Component->GetWeightmapTextures(/*InCanUseEditingWeightmap = */false);
-				TArray<ULandscapeWeightmapUsage*>& ComponentWeightmapTexturesUsage = Component->GetWeightmapTexturesUsage(/*InCanUseEditingWeightmap = */false);
-				ComponentWeightmapTexturesUsage.Reset();
-				ComponentWeightmapTexturesUsage.AddDefaulted(ComponentWeightmapTextures.Num());
+				// Reinitialize each edit layer's weightmap usages list :
+				Component->InitializeLayersWeightmapUsage(Layer.Guid);
+			}
+		}
+	}
 
-				for (int32 LayerIdx = 0; LayerIdx < ComponentWeightmapLayerAllocations.Num(); LayerIdx++)
-				{
-					const FWeightmapLayerAllocationInfo& Allocation = ComponentWeightmapLayerAllocations[LayerIdx];
-					UTexture2D* WeightmapTexture = ComponentWeightmapTextures[Allocation.WeightmapTextureIndex];
-					UE_TRANSITIONAL_OBJECT_PTR(ULandscapeWeightmapUsage)* TempUsage = WeightmapUsageMap.Find(WeightmapTexture);
+	ValidateProxyLayersWeightmapUsage();
+	bNeedsWeightmapUsagesUpdate = false;
+}
 
-					if (TempUsage == nullptr)
-					{
-						TempUsage = &WeightmapUsageMap.Add(WeightmapTexture, CreateWeightmapUsage());
-					}
+void ULandscapeComponent::InitializeLayersWeightmapUsage(const FGuid& InLayerGuid)
+{
+	ALandscapeProxy* Proxy = GetLandscapeProxy();
+	check(Proxy);
+	ALandscape* Landscape = GetLandscapeActor();
+	check(Landscape);
+	const FLandscapeLayer* SplinesLayer = Landscape->GetLandscapeSplinesReservedLayer();
+	FGuid SplinesEditLayerGuid = SplinesLayer ? SplinesLayer->Guid : FGuid();
 
-					ULandscapeWeightmapUsage* Usage = *TempUsage;
-					ComponentWeightmapTexturesUsage[Allocation.WeightmapTextureIndex] = Usage; // Keep a ref to it for faster access
+	// Don't consider invalid edit layers : 
+	if (InLayerGuid.IsValid())
+	{
+		FLandscapeLayerComponentData* LayerData = GetLayerData(InLayerGuid);
+		if (LayerData == nullptr || !LayerData->IsInitialized())
+		{
+			return;
+		}
+	}
 
-					// Validate that there are no duplicated allocation
-					check(Usage->ChannelUsage[Allocation.WeightmapTextureChannel] == nullptr || Usage->ChannelUsage[Allocation.WeightmapTextureChannel] == Component);
+	const TArray<FWeightmapLayerAllocationInfo>& ComponentWeightmapLayerAllocations = GetWeightmapLayerAllocations(InLayerGuid);
+	const TArray<UTexture2D*>& ComponentWeightmapTextures = GetWeightmapTextures(InLayerGuid);
+	TArray<ULandscapeWeightmapUsage*>& ComponentWeightmapTexturesUsage = GetWeightmapTexturesUsage(InLayerGuid);
 
-					Usage->ChannelUsage[Allocation.WeightmapTextureChannel] = Component;
-				}
-			}
+	ComponentWeightmapTexturesUsage.Reset();
+	ComponentWeightmapTexturesUsage.AddDefaulted(ComponentWeightmapTextures.Num());
 
-			const FLandscapeLayer* SplinesLayer = Landscape->GetLandscapeSplinesReservedLayer();
-			for (const FLandscapeLayer& Layer : Landscape->LandscapeLayers)
+	for (int32 LayerIdx = 0; LayerIdx < ComponentWeightmapLayerAllocations.Num(); LayerIdx++)
+	{
+		const FWeightmapLayerAllocationInfo& Allocation = ComponentWeightmapLayerAllocations[LayerIdx];
+		if (Allocation.IsAllocated())
+		{
+			check(ComponentWeightmapTextures.IsValidIndex(Allocation.WeightmapTextureIndex));
+			UTexture2D* WeightmapTexture = ComponentWeightmapTextures[Allocation.WeightmapTextureIndex];
+			TObjectPtr<ULandscapeWeightmapUsage>* TempUsage = Proxy->WeightmapUsageMap.Find(WeightmapTexture);
+			if (TempUsage == nullptr)
 			{
-				// Compute per layer data
-				FLandscapeLayerComponentData* LayerData = Component->GetLayerData(Layer.Guid);
-				if (LayerData != nullptr && LayerData->IsInitialized())
-				{
-					LayerData->WeightmapData.TextureUsages.Reset();
-					LayerData->WeightmapData.TextureUsages.AddDefaulted(LayerData->WeightmapData.Textures.Num());
+				TempUsage = &Proxy->WeightmapUsageMap.Add(WeightmapTexture, Proxy->CreateWeightmapUsage());
+				(*TempUsage)->LayerGuid = InLayerGuid;
+			}
 
-					// regenerate the weightmap usage
-					for (int32 LayerIdx = 0; LayerIdx < LayerData->WeightmapData.LayerAllocations.Num(); LayerIdx++)
-					{
-						const FWeightmapLayerAllocationInfo& Allocation = LayerData->WeightmapData.LayerAllocations[LayerIdx];
-						UTexture2D* WeightmapTexture = LayerData->WeightmapData.Textures[Allocation.WeightmapTextureIndex];
-						TObjectPtr<ULandscapeWeightmapUsage>* TempUsage = WeightmapUsageMap.Find(WeightmapTexture);
+			ULandscapeWeightmapUsage* Usage = *TempUsage;
+			ComponentWeightmapTexturesUsage[Allocation.WeightmapTextureIndex] = Usage; // Keep a ref to it for faster access
 
-						if (TempUsage == nullptr)
-						{
-							TempUsage = &WeightmapUsageMap.Add(WeightmapTexture, CreateWeightmapUsage());
-							(*TempUsage)->LayerGuid = Layer.Guid;
-						}
+			// Validate that there are no duplicated allocation
+			check(Usage->ChannelUsage[Allocation.WeightmapTextureChannel] == nullptr || Usage->ChannelUsage[Allocation.WeightmapTextureChannel] == this);
 
-						ULandscapeWeightmapUsage* Usage = *TempUsage;
-						LayerData->WeightmapData.TextureUsages[Allocation.WeightmapTextureIndex] = Usage; // Keep a ref to it for faster access
+			// Validate that there are no duplicated allocation (except on the splines layer, since it's updated outside of a transaction and the transactor can later restore a duplicated 
+			//  allocation in 2 different components, which will assert here but will be corrected in the next UpdateLandscapeSplines, which is called right after)
+			check((SplinesEditLayerGuid.IsValid() && (InLayerGuid == SplinesEditLayerGuid))
+				|| (Usage->ChannelUsage[Allocation.WeightmapTextureChannel] == nullptr)
+				|| (Usage->ChannelUsage[Allocation.WeightmapTextureChannel] == this));
 
-						// Validate that there are no duplicated allocation (except on the splines layer, since it's updated outside of a transaction and the transactor can later restore a duplicated 
-						//  allocation in 2 different components, which will assert here but will be corrected in the next UpdateLandscapeSplines, which is called right after)
-						check((&Layer == SplinesLayer) || Usage->ChannelUsage[Allocation.WeightmapTextureChannel] == nullptr || Usage->ChannelUsage[Allocation.WeightmapTextureChannel] == Component);
+			Usage->ChannelUsage[Allocation.WeightmapTextureChannel] = this;
+		}
+	}
 
-						Usage->ChannelUsage[Allocation.WeightmapTextureChannel] = Component;
-					}
+	// If there were some invalid allocations there, we will end up with null entries in ComponentWeightmapTexturesUsage, which is not desirable since we want 
+	//  ComponentWeightmapTexturesUsage and ComponentWeightmapTextures to be in sync. Fix the situation by creating the missing usages here : 
+	for (int32 Index = 0; Index < ComponentWeightmapTexturesUsage.Num(); ++Index)
+	{
+		if (UTexture2D* WeightmapTexture = ComponentWeightmapTextures[Index])
+		{
+			if (ComponentWeightmapTexturesUsage[Index] == nullptr)
+			{
+				TObjectPtr<ULandscapeWeightmapUsage>* TempUsage = Proxy->WeightmapUsageMap.Find(WeightmapTexture);
+				if (TempUsage == nullptr)
+				{
+					TempUsage = &Proxy->WeightmapUsageMap.Add(WeightmapTexture, Proxy->CreateWeightmapUsage());
+					(*TempUsage)->LayerGuid = InLayerGuid;
 				}
+
+				ULandscapeWeightmapUsage* Usage = *TempUsage;
+				ComponentWeightmapTexturesUsage[Index] = Usage; // Keep a ref to it for faster access
 			}
 		}
 	}
-
-	ValidateProxyLayersWeightmapUsage();
-	bNeedsWeightmapUsagesUpdate = false;
 }
 
 void ALandscape::ValidateProxyLayersWeightmapUsage() const
