[Splitscreen] Landscape flickering for player 1

Landscape LOD uniform buffer for player 1 was getting overwritten with data from player 2, causing visual glitches

Also fixed an issue where the Vulkan RHI wasn't copying the data provided at the resource creation when BUF_Volatile was involved (thanks [REDACTED_HANDLE])

#localization none
#tests Tested in PC Editor & on PS5 in solo & splitscreen, with D3D11/D3D12/Vulkan

diff --git a/Engine/Shaders/Private/LandscapeVertexFactory.ush b/Engine/Shaders/Private/LandscapeVertexFactory.ush
--- a/Engine/Shaders/Private/LandscapeVertexFactory.ush
+++ b/Engine/Shaders/Private/LandscapeVertexFactory.ush
@@ -31,6 +31,17 @@ int GetComponentLinearIndex()
 #endif
 }
 
+float GetViewLODData(uint ComponentIndex)
+{
+#if FIXED_GRID
+	// Fixed grid doesn't use LandscapeContinuousLODParameters buffers
+	return 0;
+#else
+	const int AbsoluteSectionLODIndex = View.LandscapeIndirection[LandscapeContinuousLODParameters.LandscapeIndex] + ComponentIndex;
+	return View.LandscapePerComponentData[AbsoluteSectionLODIndex];
+#endif
+}
+
 /* Get current Lod for this section instance. */
 float GetSectionLod(uint ComponentIndex)
 {
@@ -38,7 +49,7 @@ float GetSectionLod(uint ComponentIndex)
 	// Fixed grid has fixed Lod
 	return LandscapeFixedGrid.LodValues.x;
 #else
-	return min(LandscapeContinuousLODParameters.SectionLOD[ComponentIndex] + View.FarShadowStaticMeshLODBias, LandscapeParameters.LastLOD);
+	return min(GetViewLODData(ComponentIndex) + View.FarShadowStaticMeshLODBias, LandscapeParameters.LastLOD);
 #endif
 }
 
@@ -106,7 +117,11 @@ float4 GetNeighborSectionLod(uint ComponentIndex)
 		int2 NeighborPos = NeighborOffset[i] + int2(LandscapeParameters.ComponentBaseX, LandscapeParameters.ComponentBaseY);
 		NeighborPos = max(LandscapeContinuousLODParameters.Min                                                     , NeighborPos);
 		NeighborPos = min(LandscapeContinuousLODParameters.Min + LandscapeContinuousLODParameters.Size - int2(1, 1), NeighborPos);
-		NeighborLod[i] = max(LandscapeContinuousLODParameters.SectionLOD[(NeighborPos.y - LandscapeContinuousLODParameters.Min.y) * LandscapeContinuousLODParameters.Size.x + (NeighborPos.x - LandscapeContinuousLODParameters.Min.x)], GetSectionLod(ComponentIndex));
+
+		uint NeighborComponentIndex = (NeighborPos.y - LandscapeContinuousLODParameters.Min.y) * LandscapeContinuousLODParameters.Size.x + (NeighborPos.x - LandscapeContinuousLODParameters.Min.x);
+		float SectionLOD = GetViewLODData(NeighborComponentIndex);
+
+		NeighborLod[i] = max(SectionLOD, GetSectionLod(ComponentIndex));
 		NeighborLod[i] = min(NeighborLod[i] + View.FarShadowStaticMeshLODBias, LandscapeParameters.LastLOD);
 	}
 	return float4(NeighborLod[0], NeighborLod[1], NeighborLod[2], NeighborLod[3]);
diff --git a/Engine/Source/Runtime/Engine/Private/SceneManagement.cpp b/Engine/Source/Runtime/Engine/Private/SceneManagement.cpp
--- a/Engine/Source/Runtime/Engine/Private/SceneManagement.cpp
+++ b/Engine/Source/Runtime/Engine/Private/SceneManagement.cpp
@@ -786,6 +786,9 @@ FViewUniformShaderParameters::FViewUniformShaderParameters()
 	WaterIndirection = GWhiteVertexBufferWithSRV->ShaderResourceViewRHI;
 	WaterData = GWhiteVertexBufferWithSRV->ShaderResourceViewRHI;
 
+	LandscapeIndirection = GWhiteVertexBufferWithSRV->ShaderResourceViewRHI;
+	LandscapePerComponentData = GWhiteVertexBufferWithSRV->ShaderResourceViewRHI;
+
 	HairScatteringLUTTexture = BlackVolume;
 	HairScatteringLUTSampler = TStaticSamplerState<SF_Bilinear>::GetRHI();
 
diff --git a/Engine/Source/Runtime/Engine/Public/SceneView.h b/Engine/Source/Runtime/Engine/Public/SceneView.h
--- a/Engine/Source/Runtime/Engine/Public/SceneView.h
+++ b/Engine/Source/Runtime/Engine/Public/SceneView.h
@@ -865,6 +865,9 @@ BEGIN_GLOBAL_SHADER_PARAMETER_STRUCT_WITH_CONSTRUCTOR(FViewUniformShaderParamete
 	SHADER_PARAMETER_SRV(Buffer<float4>, WaterIndirection)
 	SHADER_PARAMETER_SRV(Buffer<float4>, WaterData)
 
+	SHADER_PARAMETER_SRV(Buffer<uint>, LandscapeIndirection)
+	SHADER_PARAMETER_SRV(Buffer<float>, LandscapePerComponentData)
+
 	SHADER_PARAMETER_UAV(RWBuffer<uint>, VTFeedbackBuffer)
 	SHADER_PARAMETER_SRV(Buffer<uint>, EditorVisualizeLevelInstanceIds)
 	SHADER_PARAMETER_SRV(Buffer<uint>, EditorSelectedHitProxyIds)
@@ -1195,6 +1198,10 @@ public:
 	FShaderResourceViewRHIRef WaterIndirectionBuffer;
 	FShaderResourceViewRHIRef WaterDataBuffer;
 
+	/** Landscape rendering related data */
+	FShaderResourceViewRHIRef LandscapeIndirectionBuffer;
+	FShaderResourceViewRHIRef LandscapePerComponentDataBuffer;
+
 	/** Feature level for this scene */
 	const ERHIFeatureLevel::Type FeatureLevel;
 
diff --git a/Engine/Source/Runtime/Landscape/Private/LandscapeModule.cpp b/Engine/Source/Runtime/Landscape/Private/LandscapeModule.cpp
--- a/Engine/Source/Runtime/Landscape/Private/LandscapeModule.cpp
+++ b/Engine/Source/Runtime/Landscape/Private/LandscapeModule.cpp
@@ -14,6 +14,7 @@
 
 #include "LandscapeProxy.h"
 #include "Landscape.h"
+#include "LandscapeRender.h"
 #include "LandscapeSplineActor.h"
 #include "Engine/Texture2D.h"
 #include "EngineUtils.h"
@@ -28,6 +29,13 @@ public:
 	/** IModuleInterface implementation */
 	void StartupModule() override;
 	void ShutdownModule() override;
+
+private:
+	void OnPostEngineInit();
+	void OnEnginePreExit();
+
+private:
+	TSharedPtr<FLandscapeSceneViewExtension, ESPMode::ThreadSafe> SceneViewExtension;
 };
 
 /**
@@ -196,12 +204,27 @@ void FLandscapeModule::StartupModule()
 	FWorldDelegates::OnPostDuplicate.AddStatic(
 		&WorldDuplicateEventFunction
 	);
+
+	FCoreDelegates::OnPostEngineInit.AddRaw(this, &FLandscapeModule::OnPostEngineInit);
+	FCoreDelegates::OnEnginePreExit.AddRaw(this, &FLandscapeModule::OnEnginePreExit);
+}
+
+void FLandscapeModule::OnPostEngineInit()
+{
+	check(!SceneViewExtension.IsValid());
+	SceneViewExtension = FSceneViewExtensions::NewExtension<FLandscapeSceneViewExtension>();
+}
+
+void FLandscapeModule::OnEnginePreExit()
+{
+	check(SceneViewExtension.IsValid());
+	SceneViewExtension.Reset();
 }
 
 void FLandscapeModule::ShutdownModule()
 {
-	// This function may be called during shutdown to clean up your module.  For modules that support dynamic reloading,
-	// we call this function before unloading the module.
+	FCoreDelegates::OnEnginePreExit.RemoveAll(this);
+	FCoreDelegates::OnPostEngineInit.RemoveAll(this);
 }
 
 IMPLEMENT_MODULE(FLandscapeModule, Landscape);
diff --git a/Engine/Source/Runtime/Landscape/Private/LandscapeRender.cpp b/Engine/Source/Runtime/Landscape/Private/LandscapeRender.cpp
--- a/Engine/Source/Runtime/Landscape/Private/LandscapeRender.cpp
+++ b/Engine/Source/Runtime/Landscape/Private/LandscapeRender.cpp
@@ -632,7 +632,7 @@ void ULandscapeComponent::GetUsedMaterials(TArray<UMaterialInterface*>& OutMater
 static int32 GetViewLodOverride(FSceneView const& View)
 {
 	// Apply r.ForceLOD override
-	int32 LodOverride = GetCVarForceLOD();
+	int32 LodOverride = GetCVarForceLOD_AnyThread();
 #if WITH_EDITOR
 	// Apply editor landscape lod override
 	LodOverride = View.Family->LandscapeLODOverride >= 0 ? View.Family->LandscapeLODOverride : LodOverride;
@@ -664,6 +664,106 @@ static int32 GetDrawCollisionLodOverride(FSceneView const& View, FCollisionRespo
 #endif
 }
 
+
+//
+// FComputeSectionPerViewParametersTask
+//
+struct FComputeSectionPerViewParametersTask
+{
+	FLandscapeRenderSystem& RenderSystem;
+
+	FLandscapeRenderSystem::FViewParams ViewParams;
+
+	TArray<uint8> SectionCurrentFirstLODIndices;
+
+	FComputeSectionPerViewParametersTask(FLandscapeRenderSystem& InRenderSystem, const FLandscapeRenderSystem::FViewParams& InViewParams)
+		: RenderSystem(InRenderSystem)
+		, ViewParams(InViewParams)
+	{
+#if PLATFORM_SUPPORTS_LANDSCAPE_VISUAL_MESH_LOD_STREAMING
+		const int32 NumSceneProxies = InRenderSystem.SceneProxies.Num();
+		SectionCurrentFirstLODIndices.Empty(NumSceneProxies);
+		SectionCurrentFirstLODIndices.AddUninitialized(NumSceneProxies);
+
+		for (int32 Idx = 0; Idx < NumSceneProxies; ++Idx)
+		{
+			const FLandscapeComponentSceneProxy* Proxy = InRenderSystem.SceneProxies[Idx];
+			SectionCurrentFirstLODIndices[Idx] = Proxy ? Proxy->GetCurrentFirstLODIdx_RenderThread() : 0;
+		}
+#endif
+	}
+
+	FORCEINLINE TStatId GetStatId() const
+	{
+		RETURN_QUICK_DECLARE_CYCLE_STAT(FComputeSectionPerViewParametersTask, STATGROUP_TaskGraphTasks);
+	}
+
+	ENamedThreads::Type GetDesiredThread()
+	{
+		return ENamedThreads::AnyNormalThreadNormalTask;
+	}
+
+	static ESubsequentsMode::Type GetSubsequentsMode()
+	{
+		return ESubsequentsMode::TrackSubsequents;
+	}
+
+	void AnyThreadTask()
+	{
+		FOptionalTaskTagScope Scope(ETaskTag::EParallelRenderingThread);
+		RenderSystem.ComputeSectionPerViewParameters(
+			ViewParams.ViewKey, ViewParams.ViewLODOverride, ViewParams.ViewLODDistanceFactor,
+			ViewParams.ViewEngineShowFlagCollisionPawn, ViewParams.ViewEngineShowFlagCollisionVisibility,
+			ViewParams.ViewOrigin, ViewParams.ViewProjectionMatrix, SectionCurrentFirstLODIndices);
+	}
+
+	void DoTask(ENamedThreads::Type CurrentThread, const FGraphEventRef& MyCompletionGraphEvent)
+	{
+		AnyThreadTask();
+	}
+};
+
+
+//
+// FGetSectionLODBiasesTask
+//
+struct FGetSectionLODBiasesTask
+{
+	FLandscapeRenderSystem& RenderSystem;
+
+	FGetSectionLODBiasesTask(FLandscapeRenderSystem& InRenderSystem)
+		: RenderSystem(InRenderSystem)
+	{
+	}
+
+	FORCEINLINE TStatId GetStatId() const
+	{
+		RETURN_QUICK_DECLARE_CYCLE_STAT(FGetSectionLODBiasesTask, STATGROUP_TaskGraphTasks);
+	}
+
+	ENamedThreads::Type GetDesiredThread()
+	{
+		return ENamedThreads::AnyNormalThreadNormalTask;
+	}
+
+	static ESubsequentsMode::Type GetSubsequentsMode()
+	{
+		return ESubsequentsMode::TrackSubsequents;
+	}
+
+	void AnyThreadTask()
+	{
+		FOptionalTaskTagScope Scope(ETaskTag::EParallelRenderingThread);
+		RenderSystem.FetchHeightmapLODBiases();
+	}
+
+	void DoTask(ENamedThreads::Type CurrentThread, const FGraphEventRef& MyCompletionGraphEvent)
+	{
+		AnyThreadTask();
+	}
+};
+
+
 //
 // FLandscapeComponentSceneProxy
 //
@@ -674,9 +774,44 @@ const static FName NAME_LandscapeResourceNameForDebugging(TEXT("Landscape"));
 
 IMPLEMENT_GLOBAL_SHADER_PARAMETER_STRUCT(FLandscapeSectionLODUniformParameters, "LandscapeContinuousLODParameters");
 
-TGlobalResource<FNullLandscapeRenderSystemResources> GNullLandscapeRenderSystemResources;
 TMap<FLandscapeNeighborInfo::FLandscapeKey, FLandscapeRenderSystem*> LandscapeRenderSystems;
 
+TBitArray<> FLandscapeRenderSystem::LandscapeIndexAllocator;
+
+
+//
+// FLandscapeRenderSystem
+//
+FLandscapeRenderSystem::FLandscapeRenderSystem()
+	: NumRegisteredEntities(0)
+	, Min(MAX_int32, MAX_int32)
+	, Size(EForceInit::ForceInitToZero)
+{
+	SectionLODBiases.SetAllowCPUAccess(true);
+
+	LandscapeIndex = LandscapeIndexAllocator.FindAndSetFirstZeroBit();
+	if (LandscapeIndex == INDEX_NONE)
+	{
+		LandscapeIndex = LandscapeIndexAllocator.Add(true);
+	}
+
+	FLandscapeSectionLODUniformParameters Parameters;
+	Parameters.LandscapeIndex = LandscapeIndex;
+	Parameters.Size = FIntPoint(1, 1);
+	Parameters.SectionLODBias = GWhiteVertexBufferWithSRV->ShaderResourceViewRHI;
+	UniformBuffer = TUniformBufferRef<FLandscapeSectionLODUniformParameters>::CreateUniformBufferImmediate(Parameters, UniformBuffer_MultiFrame);
+}
+
+FLandscapeRenderSystem::~FLandscapeRenderSystem()
+{
+	check(LandscapeIndexAllocator[LandscapeIndex]);
+	LandscapeIndexAllocator[LandscapeIndex] = false;
+
+	// Clear slack in the array
+	int32 LastSetIndex = LandscapeIndexAllocator.FindLast(true);
+	LandscapeIndexAllocator.SetNumUninitialized(LastSetIndex + 1);
+}
+
 void FLandscapeRenderSystem::RegisterEntity(FLandscapeComponentSceneProxy* SceneProxy)
 {
 	check(IsInRenderingThread());
@@ -695,13 +830,11 @@ void FLandscapeRenderSystem::RegisterEntity(FLandscapeComponentSceneProxy* Scene
 		if (NewMin != Min || Size != SizeRequired)
 		{
 			ResizeAndMoveTo(NewMin, SizeRequired);
-			RecreateBuffers();
 		}
 	}
 	else
 	{
 		ResizeAndMoveTo(SceneProxy->ComponentBase, FIntPoint(1, 1));
-		RecreateBuffers();
 	}
 
 	NumRegisteredEntities++;
@@ -727,16 +860,13 @@ void FLandscapeRenderSystem::UnregisterEntity(FLandscapeComponentSceneProxy* Sce
 
 void FLandscapeRenderSystem::ResizeAndMoveTo(FIntPoint NewMin, FIntPoint NewSize)
 {
-	SectionLODBuffer.SafeRelease();
 	SectionLODBiasBuffer.SafeRelease();
 
-	TResourceArray<float> NewSectionLODValues;
 	TResourceArray<float> NewSectionLODBiases;
 	TArray<LODSettingsComponent> NewSectionLODSettings;
 	TArray<FVector4> NewSectionOriginAndRadius;
 	TArray<FLandscapeComponentSceneProxy*> NewSceneProxies;
 
-	NewSectionLODValues.AddZeroed(NewSize.X * NewSize.Y);
 	NewSectionLODBiases.AddZeroed(NewSize.X * NewSize.Y);
 	NewSectionLODSettings.AddZeroed(NewSize.X * NewSize.Y);
 	NewSectionOriginAndRadius.AddZeroed(NewSize.X * NewSize.Y);
@@ -751,7 +881,6 @@ void FLandscapeRenderSystem::ResizeAndMoveTo(FIntPoint NewMin, FIntPoint NewSize
 
 			if (NewLinearIndex >= 0 && NewLinearIndex < NewSize.X * NewSize.Y)
 			{
-				NewSectionLODValues[NewLinearIndex] = SectionLODValues[LinearIndex];
 				NewSectionLODBiases[NewLinearIndex] = SectionLODBiases[LinearIndex];
 				NewSectionLODSettings[NewLinearIndex] = SectionLODSettings[LinearIndex];
 				NewSectionOriginAndRadius[NewLinearIndex] = SectionOriginAndRadius[LinearIndex];
@@ -762,48 +891,36 @@ void FLandscapeRenderSystem::ResizeAndMoveTo(FIntPoint NewMin, FIntPoint NewSize
 
 	Min = NewMin;
 	Size = NewSize;
-	SectionLODValues = NewSectionLODValues;
 	SectionLODBiases = NewSectionLODBiases;
 	SectionLODSettings = NewSectionLODSettings;
 	SectionOriginAndRadius = NewSectionOriginAndRadius;
 	SceneProxies = NewSceneProxies;
 
-	SectionLODValues.SetAllowCPUAccess(true);
 	SectionLODBiases.SetAllowCPUAccess(true);
 }
 
-void FLandscapeRenderSystem::PrepareView(const FSceneView* View)
+void FLandscapeRenderSystem::PrepareView(const FViewParams& InViewParams)
 {
-#if PLATFORM_SUPPORTS_LANDSCAPE_VISUAL_MESH_LOD_STREAMING
-	const int32 NumSceneProxies = SceneProxies.Num();
-	SectionCurrentFirstLODIndices.Empty(NumSceneProxies);
-	SectionCurrentFirstLODIndices.AddUninitialized(NumSceneProxies);
-
-	for (int32 Idx = 0; Idx < NumSceneProxies; ++Idx)
-	{
-		const FLandscapeComponentSceneProxy* Proxy = SceneProxies[Idx];
-		SectionCurrentFirstLODIndices[Idx] = Proxy ? Proxy->GetCurrentFirstLODIdx_RenderThread() : 0;
-	}
-#endif
-	
 	const bool bExecuteInParallel = FApp::ShouldUseThreadingForPerformance()
 		&& GIsThreadedRendering; // Rendering thread is required to safely use rendering resources in parallel.
 
 	if (bExecuteInParallel)
 	{
-		PerViewParametersTasks.Add(View, TGraphTask<FComputeSectionPerViewParametersTask>::CreateTask(
-			nullptr, ENamedThreads::GetRenderThread()).ConstructAndDispatchWhenReady(*this, View));
+		PerViewParametersTasks.Add(InViewParams.ViewKey, TGraphTask<FComputeSectionPerViewParametersTask>::CreateTask(
+			nullptr, ENamedThreads::GetRenderThread()).ConstructAndDispatchWhenReady(*this, InViewParams));
 	}
 	else
 	{
-		FComputeSectionPerViewParametersTask Task(*this, View);
+		FComputeSectionPerViewParametersTask Task(*this, InViewParams);
 		Task.AnyThreadTask();
 	}
 }
 
-void FLandscapeRenderSystem::BeginRenderView(const FSceneView* View)
+void FLandscapeRenderSystem::BeginRender()
 {
-	TRACE_CPUPROFILER_EVENT_SCOPE(FLandscapeRenderSystem::BeginRenderView());
+	check(IsInRenderingThread());
+
+	TRACE_CPUPROFILER_EVENT_SCOPE(FLandscapeRenderSystem::BeginRender());
 
 	if (FetchHeightmapLODBiasesEventRef.IsValid())
 	{
@@ -811,29 +928,23 @@ void FLandscapeRenderSystem::BeginRenderView(const FSceneView* View)
 		FetchHeightmapLODBiasesEventRef.SafeRelease();
 	}
 
-	if (PerViewParametersTasks.Contains(View))
-	{
-		FTaskGraphInterface::Get().WaitUntilTaskCompletes(PerViewParametersTasks[View], ENamedThreads::GetRenderThread_Local());
-		PerViewParametersTasks.Remove(View);
-	}
-
+	for (auto& Pair : PerViewParametersTasks)
 	{
-		FScopeLock Lock(&CachedValuesCS);
-
-		SectionLODValues = CachedSectionLODValues[View];
+		FTaskGraphInterface::Get().WaitUntilTaskCompletes(PerViewParametersTasks[Pair.Key], ENamedThreads::GetRenderThread_Local());
 	}
 
-	RecreateBuffers(View);
+	UpdateBuffers();
 }
 
 void FLandscapeRenderSystem::ComputeSectionPerViewParameters(
-	const FSceneView* ViewPtrAsIdentifier,
+	FViewKey ViewKey,
 	int32 ViewLODOverride,
 	float ViewLODDistanceFactor,
 	bool bDrawCollisionPawn,
 	bool bDrawCollisionCollision,
-	FVector ViewOrigin,
-	FMatrix ViewProjectionMarix
+	const FVector& ViewOrigin,
+	const FMatrix& ViewProjectionMarix,
+	const TArray<uint8>& SectionCurrentFirstLODIndices
 )
 {
 	TRACE_CPUPROFILER_EVENT_SCOPE(FLandscapeRenderSystem::ComputeSectionPerViewParameters());
@@ -869,7 +980,7 @@ void FLandscapeRenderSystem::ComputeSectionPerViewParameters(
 	{
 		FScopeLock Lock(&CachedValuesCS);
 
-		CachedSectionLODValues.Add(ViewPtrAsIdentifier, NewSectionLODValues);
+		CachedSectionLODValues.Add(ViewKey, MoveTemp(NewSectionLODValues));
 	}
 }
 
@@ -894,65 +1005,48 @@ void FLandscapeRenderSystem::FetchHeightmapLODBiases()
 	}
 }
 
-void FLandscapeRenderSystem::RecreateBuffers(const FSceneView* InView /* = nullptr */)
+void FLandscapeRenderSystem::UpdateBuffers()
 {
-	if (InView == nullptr || CachedView != InView)
-	{
-		TRACE_CPUPROFILER_EVENT_SCOPE(FLandscapeRenderSystem::RecreateBuffers());
+	check(IsInRenderingThread());
 
-		if (Size != FIntPoint::ZeroValue)
-		{
-			if (!SectionLODBuffer.IsValid())
-			{
-				FRHIResourceCreateInfo CreateInfo(TEXT("SectionLODBuffer"), &SectionLODValues);
-				SectionLODBuffer = RHICreateVertexBuffer(SectionLODValues.GetResourceDataSize(), BUF_ShaderResource | BUF_Dynamic, CreateInfo);
-				SectionLODSRV = RHICreateShaderResourceView(SectionLODBuffer, sizeof(float), PF_R32_FLOAT);
-			}
-			else
-			{
-				float* Data = (float*)RHILockBuffer(SectionLODBuffer, 0, SectionLODValues.GetResourceDataSize(), RLM_WriteOnly);
-				FMemory::Memcpy(Data, SectionLODValues.GetData(), SectionLODValues.GetResourceDataSize());
-				RHIUnlockBuffer(SectionLODBuffer);
-			}
+	TRACE_CPUPROFILER_EVENT_SCOPE(FLandscapeRenderSystem::RecreateBuffers());
 
-			if (!SectionLODBiasBuffer.IsValid())
-			{
-				FRHIResourceCreateInfo CreateInfo(TEXT("SectionLODBiasBuffer"), &SectionLODBiases);
-				SectionLODBiasBuffer = RHICreateVertexBuffer(SectionLODBiases.GetResourceDataSize(), BUF_ShaderResource | BUF_Dynamic, CreateInfo);
-				SectionLODBiasSRV = RHICreateShaderResourceView(SectionLODBiasBuffer, sizeof(float), PF_R32_FLOAT);
-			}
-			else
-			{
-				float* Data = (float*)RHILockBuffer(SectionLODBiasBuffer, 0, SectionLODBiases.GetResourceDataSize(), RLM_WriteOnly);
-				FMemory::Memcpy(Data, SectionLODBiases.GetData(), SectionLODBiases.GetResourceDataSize());
-				RHIUnlockBuffer(SectionLODBiasBuffer);
-			}
+	bool bUpdateUB = false;
+
+	if (Size != FIntPoint::ZeroValue)
+	{
+		if (!SectionLODBiasBuffer.IsValid())
+		{
+			FRHIResourceCreateInfo CreateInfo(TEXT("SectionLODBiasBuffer"), &SectionLODBiases);
+			SectionLODBiasBuffer = RHICreateVertexBuffer(SectionLODBiases.GetResourceDataSize(), BUF_ShaderResource | BUF_Dynamic, CreateInfo);
+			SectionLODBiasSRV = RHICreateShaderResourceView(SectionLODBiasBuffer, sizeof(float), PF_R32_FLOAT);
+			bUpdateUB = true;
+		}
+		else
+		{
+			float* Data = (float*)RHILockBuffer(SectionLODBiasBuffer, 0, SectionLODBiases.GetResourceDataSize(), RLM_WriteOnly);
+			FMemory::Memcpy(Data, SectionLODBiases.GetData(), SectionLODBiases.GetResourceDataSize());
+			RHIUnlockBuffer(SectionLODBiasBuffer);
+		}
 
+		if (bUpdateUB)
+		{
 			FLandscapeSectionLODUniformParameters Parameters;
+			Parameters.LandscapeIndex = LandscapeIndex;
 			Parameters.Min = Min;
 			Parameters.Size = Size;
-			Parameters.SectionLOD = SectionLODSRV;
 			Parameters.SectionLODBias = SectionLODBiasSRV;
 
-			if (UniformBuffer.IsValid())
-			{
-				UniformBuffer.UpdateUniformBufferImmediate(Parameters);
-			}
-			else
-			{
-				UniformBuffer = TUniformBufferRef<FLandscapeSectionLODUniformParameters>::CreateUniformBufferImmediate(Parameters, UniformBuffer_SingleFrame);
-			}
+			UniformBuffer.UpdateUniformBufferImmediate(Parameters);
 		}
-
-		CachedView = InView;
 	}
 }
 
 void FLandscapeRenderSystem::BeginFrame()
 {
-	CachedView = nullptr;
+	check(IsInRenderingThread());
 
-	CachedSectionLODValues.Empty();
+	CachedSectionLODValues.Reset();
 
 	const bool bExecuteInParallel = FApp::ShouldUseThreadingForPerformance()
 		&& GIsThreadedRendering; // Rendering thread is required to safely use rendering resources in parallel.
@@ -971,6 +1065,8 @@ void FLandscapeRenderSystem::BeginFrame()
 
 void FLandscapeRenderSystem::EndFrame()
 {
+	check(IsInRenderingThread());
+
 	// Finalize any outstanding jobs before ~FSceneRenderer() so we don't have corrupted accesses
 	if (FetchHeightmapLODBiasesEventRef.IsValid())
 	{
@@ -980,29 +1076,31 @@ void FLandscapeRenderSystem::EndFrame()
 
 	for (auto& Pair : PerViewParametersTasks)
 	{
-		const FSceneView* View = Pair.Key;
-		FTaskGraphInterface::Get().WaitUntilTaskCompletes(PerViewParametersTasks[View], ENamedThreads::GetRenderThread_Local());
+		FTaskGraphInterface::Get().WaitUntilTaskCompletes(PerViewParametersTasks[Pair.Key], ENamedThreads::GetRenderThread_Local());
 	}
 
-	PerViewParametersTasks.Empty();
+	PerViewParametersTasks.Reset();
 }
 
-FLandscapeRenderSystem::FComputeSectionPerViewParametersTask::FComputeSectionPerViewParametersTask(FLandscapeRenderSystem& InRenderSystem, const FSceneView* InView)
-	: RenderSystem(InRenderSystem)
-	, ViewPtrAsIdentifier(InView)
-	, ViewLODOverride(GetViewLodOverride(*InView))
-	, ViewLODDistanceFactor(InView->LODDistanceFactor)
-	, ViewEngineShowFlagCollisionPawn(InView->Family->EngineShowFlags.CollisionPawn)
-	, ViewEngineShowFlagCollisionVisibility(InView->Family->EngineShowFlags.CollisionVisibility)
-	, ViewOrigin(GetLODView(*InView).ViewMatrices.GetViewOrigin())
-	, ViewProjectionMatrix(GetLODView(*InView).ViewMatrices.GetProjectionMatrix())
+
+//
+// FLandscapeSceneViewExtension
+//
+FLandscapeSceneViewExtension::FLandscapeSceneViewExtension(const FAutoRegister& AutoReg) : FSceneViewExtensionBase(AutoReg)
 {
+	FCoreDelegates::OnBeginFrame.AddRaw(this, &FLandscapeSceneViewExtension::BeginFrame);
+	FCoreDelegates::OnEndFrame.AddRaw(this, &FLandscapeSceneViewExtension::EndFrame);
 }
 
-class FLandscapePersistentViewUniformBufferExtension : public IPersistentViewUniformBufferExtension
+FLandscapeSceneViewExtension::~FLandscapeSceneViewExtension()
 {
-public:
-	virtual void BeginFrame() override
+	FCoreDelegates::OnBeginFrame.RemoveAll(this);
+	FCoreDelegates::OnEndFrame.RemoveAll(this);
+}
+
+void FLandscapeSceneViewExtension::BeginFrame()
+{
+	ENQUEUE_RENDER_COMMAND(FLandscapeSceneViewExtension_BeginFrame)([this](FRHICommandListImmediate& RHICmdList)
 	{
 		for (auto& Pair : LandscapeRenderSystems)
 		{
@@ -1010,34 +1108,84 @@ public:
 
 			RenderSystem.BeginFrame();
 		}
-	}
+	});
+}
 
-	virtual void PrepareView(const FSceneView* View) override
+void FLandscapeSceneViewExtension::SetupView(FSceneViewFamily& InViewFamily, FSceneView& InView)
+{
+	FLandscapeRenderSystem::FViewParams ViewParams;
+	ViewParams.ViewKey = InView.GetViewKey();
+	ViewParams.ViewLODOverride = GetViewLodOverride(InView);
+	ViewParams.ViewLODDistanceFactor = InView.LODDistanceFactor;
+	ViewParams.ViewEngineShowFlagCollisionPawn = InView.Family->EngineShowFlags.CollisionPawn;
+	ViewParams.ViewEngineShowFlagCollisionVisibility = InView.Family->EngineShowFlags.CollisionVisibility;
+	ViewParams.ViewOrigin = GetLODView(InView).ViewMatrices.GetViewOrigin();
+	ViewParams.ViewProjectionMatrix = GetLODView(InView).ViewMatrices.GetProjectionMatrix();
+
+	ENQUEUE_RENDER_COMMAND(FLandscapeSceneViewExtension_BeginFrame)([ViewParams](FRHICommandListImmediate& RHICmdList)
 	{
 		for (auto& Pair : LandscapeRenderSystems)
 		{
 			FLandscapeRenderSystem& RenderSystem = *Pair.Value;
-
-			RenderSystem.PrepareView(View);
+			RenderSystem.PrepareView(ViewParams);
 		}
+	});
+}
+
+void FLandscapeSceneViewExtension::PreRenderViewFamily_RenderThread(FRHICommandListImmediate& RHICmdList, FSceneViewFamily& InViewFamily)
+{
+	for (auto& Pair : LandscapeRenderSystems)
+	{
+		FLandscapeRenderSystem& RenderSystem = *Pair.Value;
+		RenderSystem.BeginRender();
 	}
+}
+
+
+void FLandscapeSceneViewExtension::PreRenderView_RenderThread(FRHICommandListImmediate& RHICmdList, FSceneView& InView)
+{
+	FShaderResourceViewRHIRef LODDataSRV;
+	FShaderResourceViewRHIRef IndirectionSRV;
 
-	virtual void BeginRenderView(const FSceneView* View, bool bShouldWaitForJobs = true) override
+	if (!LandscapeRenderSystems.IsEmpty())
 	{
-		if (!bShouldWaitForJobs)
-		{
-			return;
-		}
+		TResourceArray<uint32> LandscapeIndirection;
+		TResourceArray<float> LandscapeLODData;
+
+		LandscapeIndirection.SetNum(FLandscapeRenderSystem::LandscapeIndexAllocator.Num());
 
 		for (auto& Pair : LandscapeRenderSystems)
 		{
 			FLandscapeRenderSystem& RenderSystem = *Pair.Value;
+			uint32& IndirectionEntry = LandscapeIndirection[RenderSystem.LandscapeIndex];
 
-			RenderSystem.BeginRenderView(View);
+			// Where the data for this landscape starts
+			IndirectionEntry = LandscapeLODData.Num();
+
+			LandscapeLODData.Append(RenderSystem.CachedSectionLODValues.FindChecked(InView.GetViewKey()));
 		}
+
+		FRHIResourceCreateInfo CreateInfoLODBuffer(TEXT("LandscapeLODDataBuffer"), &LandscapeLODData);
+		FBufferRHIRef LandscapeLODDataBuffer = RHICreateStructuredBuffer(sizeof(float), LandscapeLODData.GetResourceDataSize(), BUF_ShaderResource | BUF_Volatile, CreateInfoLODBuffer);
+		LODDataSRV = RHICreateShaderResourceView(LandscapeLODDataBuffer, sizeof(float), PF_R32_FLOAT);
+
+		FRHIResourceCreateInfo CreateInfoIndirection(TEXT("LandscapeIndirectionBuffer"), &LandscapeIndirection);
+		FBufferRHIRef LandscapeIndirectionBuffer = RHICreateStructuredBuffer(sizeof(uint32), LandscapeIndirection.GetResourceDataSize(), BUF_ShaderResource | BUF_Volatile, CreateInfoIndirection);
+		IndirectionSRV = RHICreateShaderResourceView(LandscapeIndirectionBuffer, sizeof(uint32), PF_R32_UINT);
 	}
+	else
+	{
+		LODDataSRV = GWhiteVertexBufferWithSRV->ShaderResourceViewRHI;
+		IndirectionSRV = GWhiteVertexBufferWithSRV->ShaderResourceViewRHI;
+	}
+
+	InView.LandscapePerComponentDataBuffer = LODDataSRV;
+	InView.LandscapeIndirectionBuffer = IndirectionSRV;
+}
 
-	virtual void EndFrame() override
+void FLandscapeSceneViewExtension::EndFrame()
+{
+	ENQUEUE_RENDER_COMMAND(FLandscapeSceneViewExtension_EndFrame)([this](FRHICommandListImmediate& RHICmdList)
 	{
 		for (auto& Pair : LandscapeRenderSystems)
 		{
@@ -1045,11 +1193,12 @@ public:
 
 			RenderSystem.EndFrame();
 		}
-	}
-
-} LandscapePersistentViewUniformBufferExtension;
-
+	});
+}
 
+//
+// FLandscapeVisibilityHelper
+//
 void FLandscapeVisibilityHelper::Init(UPrimitiveComponent* LandscapeComponent, FPrimitiveSceneProxy* ProxyIn)
 {
 	// Flag components to render only after level will be fully added to the world
@@ -1457,11 +1606,15 @@ void FLandscapeComponentSceneProxy::CreateRenderThreadResources()
 		GrassMeshBatch.Type = PT_PointList;
 		GrassMeshBatch.DepthPriorityGroup = SDPG_World;
 
+		const FLandscapeRenderSystem& RenderSystem = *LandscapeRenderSystems.FindChecked(LandscapeKey);
+		const TUniformBufferRef<FLandscapeSectionLODUniformParameters>& LandscapeSectionLODUniformParameters = RenderSystem.UniformBuffer;
+
 		// Combined grass rendering batch element
 		FMeshBatchElement* GrassBatchElement = &GrassMeshBatch.Elements[0];
 		FLandscapeBatchElementParams* BatchElementParams = &GrassBatchParams[0];
 		BatchElementParams->LandscapeUniformShaderParametersResource = &LandscapeUniformShaderParameters;
 		BatchElementParams->FixedGridUniformShaderParameters = &LandscapeFixedGridUniformShaderParameters;
+		BatchElementParams->LandscapeSectionLODUniformParameters = &LandscapeSectionLODUniformParameters;
 		BatchElementParams->SceneProxy = this;
 		BatchElementParams->CurrentLOD = 0;
 		GrassBatchElement->UserData = BatchElementParams;
@@ -1943,10 +2096,14 @@ bool FLandscapeComponentSceneProxy::GetMeshElementForVirtualTexture(int32 InLodI
 
 	OutMeshBatch.Elements.Empty(1);
 
+	const FLandscapeRenderSystem& RenderSystem = *LandscapeRenderSystems.FindChecked(LandscapeKey);
+	const TUniformBufferRef<FLandscapeSectionLODUniformParameters>& LandscapeSectionLODUniformParameters = RenderSystem.UniformBuffer;
+
 	FLandscapeBatchElementParams* BatchElementParams = new(OutStaticBatchParamArray) FLandscapeBatchElementParams;
 	BatchElementParams->SceneProxy = this;
 	BatchElementParams->LandscapeUniformShaderParametersResource = &LandscapeUniformShaderParameters;
 	BatchElementParams->FixedGridUniformShaderParameters = &LandscapeFixedGridUniformShaderParameters;
+	BatchElementParams->LandscapeSectionLODUniformParameters = &LandscapeSectionLODUniformParameters;
 	BatchElementParams->CurrentLOD = InLodIndex;
 
 	int32 LodSubsectionSizeVerts = SubsectionSizeVerts >> InLodIndex;
@@ -2014,15 +2171,18 @@ bool FLandscapeComponentSceneProxy::GetStaticMeshElement(int32 LODIndex, bool bF
 		MeshBatch.LODIndex = LODIndex;
 		MeshBatch.bDitheredLODTransition = false;
 
-		// Combined batch element
-		FMeshBatchElement& BatchElement = MeshBatch.Elements[0];
+		const FLandscapeRenderSystem& RenderSystem = *LandscapeRenderSystems.FindChecked(LandscapeKey);
+		const TUniformBufferRef<FLandscapeSectionLODUniformParameters>& LandscapeSectionLODUniformParameters = RenderSystem.UniformBuffer;
 
 		FLandscapeBatchElementParams* BatchElementParams = new(OutStaticBatchParamArray) FLandscapeBatchElementParams;
 		BatchElementParams->LandscapeUniformShaderParametersResource = &LandscapeUniformShaderParameters;
 		BatchElementParams->FixedGridUniformShaderParameters = &LandscapeFixedGridUniformShaderParameters;
+		BatchElementParams->LandscapeSectionLODUniformParameters = &LandscapeSectionLODUniformParameters;
 		BatchElementParams->SceneProxy = this;
 		BatchElementParams->CurrentLOD = LODIndex;
 
+		// Combined batch element
+		FMeshBatchElement& BatchElement = MeshBatch.Elements[0];
 		BatchElement.UserData = BatchElementParams;
 		BatchElement.PrimitiveUniformBuffer = GetUniformBuffer();
 		BatchElement.IndexBuffer = SharedBuffers->IndexBuffers[LODIndex];
@@ -2493,12 +2653,13 @@ void FLandscapeComponentSceneProxy::GetDynamicRayTracingInstances(FRayTracingMat
 		return;
 	}
 
+	const FSceneView& SceneView = *Context.ReferenceView;
 	int32 ForcedLODLevel = ForcedLOD;
 
-	int32 ViewLodOveride = GetViewLodOverride(*Context.ReferenceView);
+	int32 ViewLodOveride = GetViewLodOverride(SceneView);
 	ForcedLODLevel = ViewLodOveride >= 0 ? ViewLodOveride : ForcedLODLevel;
 
-	float MeshScreenSizeSquared = ComputeBoundsScreenRadiusSquared(GetBounds().Origin, GetBounds().SphereRadius, *Context.ReferenceView);
+	float MeshScreenSizeSquared = ComputeBoundsScreenRadiusSquared(GetBounds().Origin, GetBounds().SphereRadius, SceneView);
 	float LODScale = Context.ReferenceView->LODDistanceFactor * CVarStaticMeshLODDistanceScale.GetValueOnRenderThread();
 	int32 LODToRender = ForcedLODLevel >= 0 ? ForcedLODLevel : GetLODFromScreenSize(MeshScreenSizeSquared, LODScale * LODScale);
 	
@@ -2531,7 +2692,8 @@ void FLandscapeComponentSceneProxy::GetDynamicRayTracingInstances(FRayTracingMat
 
 	BaseMeshBatch.Elements.Empty();
 
-	FLandscapeRenderSystem& RenderSystem = *LandscapeRenderSystems.FindChecked(LandscapeKey);
+	const FLandscapeRenderSystem& RenderSystem = *LandscapeRenderSystems.FindChecked(LandscapeKey);
+	const TUniformBufferRef<FLandscapeSectionLODUniformParameters>& LandscapeSectionLODUniformParameters = RenderSystem.UniformBuffer;
 
 	for (int32 SubY = 0; SubY < NumSubsections; SubY++)
 	{
@@ -2547,6 +2709,7 @@ void FLandscapeComponentSceneProxy::GetDynamicRayTracingInstances(FRayTracingMat
 
 			BatchElementParams.LandscapeUniformShaderParametersResource = &LandscapeUniformShaderParameters;
 			BatchElementParams.FixedGridUniformShaderParameters = &LandscapeFixedGridUniformShaderParameters;
+			BatchElementParams.LandscapeSectionLODUniformParameters = &LandscapeSectionLODUniformParameters;
 			BatchElementParams.SceneProxy = this;
 			BatchElementParams.CurrentLOD = CurrentLOD;
 			BatchElement.UserData = &BatchElementParams;
@@ -2592,10 +2755,10 @@ void FLandscapeComponentSceneProxy::GetDynamicRayTracingInstances(FRayTracingMat
 					SectionRayTracingStates[SubSectionIdx].HeightmapLODBias = RenderSystem.GetSectionLODBias(ComponentBase);
 				}
 
-				if (SectionRayTracingStates[SubSectionIdx].FractionalLOD != RenderSystem.GetSectionLODValue(ComponentBase))
+				if (SectionRayTracingStates[SubSectionIdx].FractionalLOD != RenderSystem.GetSectionLODValue(SceneView, ComponentBase))
 				{
 					bNeedsRayTracingGeometryUpdate = true;
-					SectionRayTracingStates[SubSectionIdx].FractionalLOD = RenderSystem.GetSectionLODValue(ComponentBase);
+					SectionRayTracingStates[SubSectionIdx].FractionalLOD = RenderSystem.GetSectionLODValue(SceneView, ComponentBase);
 				}
 			}
 
@@ -3063,15 +3226,7 @@ public:
 		const FLandscapeComponentSceneProxy* SceneProxy = BatchElementParams->SceneProxy;
 
 		ShaderBindings.Add(Shader->GetUniformBufferParameter<FLandscapeUniformShaderParameters>(), *BatchElementParams->LandscapeUniformShaderParametersResource);
-
-		if (SceneProxy && SceneProxy->bRegistered)
-		{
-			ShaderBindings.Add(Shader->GetUniformBufferParameter<FLandscapeSectionLODUniformParameters>(), LandscapeRenderSystems.FindChecked(SceneProxy->LandscapeKey)->UniformBuffer);
-		}
-		else
-		{
-			ShaderBindings.Add(Shader->GetUniformBufferParameter<FLandscapeSectionLODUniformParameters>(), GNullLandscapeRenderSystemResources.UniformBuffer);
-		}
+		ShaderBindings.Add(Shader->GetUniformBufferParameter<FLandscapeSectionLODUniformParameters>(), *BatchElementParams->LandscapeSectionLODUniformParameters);
 
 #if RHI_RAYTRACING
 		if (IsRayTracingEnabled())
@@ -3990,8 +4145,6 @@ void FLandscapeNeighborInfo::RegisterNeighbors(FLandscapeComponentSceneProxy* Sc
 		if (!SharedSceneProxyMap.Find(LandscapeKey))
 		{
 			LandscapeRenderSystems.Add(LandscapeKey, new FLandscapeRenderSystem {});
-
-			GetRendererModule().RegisterPersistentViewUniformBufferExtension(&LandscapePersistentViewUniformBufferExtension);
 		}
 
 		// Register ourselves in the map.
diff --git a/Engine/Source/Runtime/Landscape/Private/LandscapeRenderMobile.cpp b/Engine/Source/Runtime/Landscape/Private/LandscapeRenderMobile.cpp
--- a/Engine/Source/Runtime/Landscape/Private/LandscapeRenderMobile.cpp
+++ b/Engine/Source/Runtime/Landscape/Private/LandscapeRenderMobile.cpp
@@ -81,6 +81,7 @@ public:
 
 		const FLandscapeComponentSceneProxyMobile* SceneProxy = (const FLandscapeComponentSceneProxyMobile*)BatchElementParams->SceneProxy;
 		ShaderBindings.Add(Shader->GetUniformBufferParameter<FLandscapeUniformShaderParameters>(),*BatchElementParams->LandscapeUniformShaderParametersResource);
+		ShaderBindings.Add(Shader->GetUniformBufferParameter<FLandscapeSectionLODUniformParameters>(), *BatchElementParams->LandscapeSectionLODUniformParameters);
 
 		if (TexCoordOffsetParameter.IsBound())
 		{
@@ -92,16 +93,7 @@ public:
 			);
 			ShaderBindings.Add(TexCoordOffsetParameter, FVector2f(TexCoordOffset));
 		}
-
-		if (SceneProxy->bRegistered)
-		{
-			ShaderBindings.Add(Shader->GetUniformBufferParameter<FLandscapeSectionLODUniformParameters>(), LandscapeRenderSystems.FindChecked(SceneProxy->LandscapeKey)->UniformBuffer);
-		}
-		else
-		{
-			ShaderBindings.Add(Shader->GetUniformBufferParameter<FLandscapeSectionLODUniformParameters>(), GNullLandscapeRenderSystemResources.UniformBuffer);
-		}
-			}
+	}
 
 protected:
 	LAYOUT_FIELD(FShaderParameter, TexCoordOffsetParameter);
diff --git a/Engine/Source/Runtime/Landscape/Public/LandscapeRender.h b/Engine/Source/Runtime/Landscape/Public/LandscapeRender.h
--- a/Engine/Source/Runtime/Landscape/Public/LandscapeRender.h
+++ b/Engine/Source/Runtime/Landscape/Public/LandscapeRender.h
@@ -26,6 +26,7 @@ LandscapeRender.h: New terrain rendering
 #include "PrimitiveViewRelevance.h"
 #include "PrimitiveSceneProxy.h"
 #include "StaticMeshResources.h"
+#include "SceneViewExtension.h"
 
 // This defines the number of border blocks to surround terrain by when generating lightmaps
 #define TERRAIN_PATCH_EXPAND_SCALAR	1
@@ -116,9 +117,9 @@ END_GLOBAL_SHADER_PARAMETER_STRUCT()
 typedef TUniformBufferRef<FLandscapeVertexFactoryMVFParameters> FLandscapeVertexFactoryMVFUniformBufferRef;
 
 BEGIN_GLOBAL_SHADER_PARAMETER_STRUCT(FLandscapeSectionLODUniformParameters, LANDSCAPE_API)
+	SHADER_PARAMETER(int32, LandscapeIndex)
 	SHADER_PARAMETER(FIntPoint, Min)
 	SHADER_PARAMETER(FIntPoint, Size)
-	SHADER_PARAMETER_SRV(Buffer<float>, SectionLOD)
 	SHADER_PARAMETER_SRV(Buffer<float>, SectionLODBias)
 END_GLOBAL_SHADER_PARAMETER_STRUCT()
 
@@ -134,6 +135,7 @@ struct FLandscapeBatchElementParams
 #endif
 	const TUniformBuffer<FLandscapeUniformShaderParameters>* LandscapeUniformShaderParametersResource;
 	const TArray<TUniformBuffer<FLandscapeFixedGridUniformShaderParameters>>* FixedGridUniformShaderParameters;
+	const TUniformBufferRef<FLandscapeSectionLODUniformParameters>* LandscapeSectionLODUniformParameters;
 	const FLandscapeComponentSceneProxy* SceneProxy;
 	int32 CurrentLOD;
 };
@@ -433,44 +435,23 @@ public:
 	void UnregisterNeighbors(FLandscapeComponentSceneProxy* SceneProxy = nullptr);
 };
 
+extern RENDERER_API TAutoConsoleVariable<float> CVarStaticMeshLODDistanceScale;
 
-class FNullLandscapeRenderSystemResources : public FRenderResource
+struct FLandscapeRenderSystem
 {
-public:
+	typedef uint32 FViewKey;
 
-	FBufferRHIRef SectionLODBuffer;
-	FShaderResourceViewRHIRef SectionLODSRV;
-	TUniformBufferRef<FLandscapeSectionLODUniformParameters> UniformBuffer;
-
-	virtual void InitRHI() override
+	struct FViewParams
 	{
-		TResourceArray<float> ResourceBuffer;
-		ResourceBuffer.Add(0.0f);
-		FRHIResourceCreateInfo CreateInfo(TEXT("SectionLODBuffer"), &ResourceBuffer);
-		SectionLODBuffer = RHICreateVertexBuffer(ResourceBuffer.GetResourceDataSize(), BUF_ShaderResource | BUF_Static, CreateInfo);
-		SectionLODSRV = RHICreateShaderResourceView(SectionLODBuffer, sizeof(float), PF_R32_FLOAT);
-
-		FLandscapeSectionLODUniformParameters Parameters;
-		Parameters.Size = FIntPoint(1, 1);
-		Parameters.SectionLOD = SectionLODSRV;
-		Parameters.SectionLODBias = SectionLODSRV;
-		UniformBuffer = TUniformBufferRef<FLandscapeSectionLODUniformParameters>::CreateUniformBufferImmediate(Parameters, UniformBuffer_MultiFrame);
-	}
-
-	virtual void ReleaseRHI() override
-	{
-		SectionLODBuffer.SafeRelease();
-		SectionLODSRV.SafeRelease();
-		UniformBuffer.SafeRelease();
-	}
+		FViewKey ViewKey;
+		int32 ViewLODOverride;
+		float ViewLODDistanceFactor;
+		bool ViewEngineShowFlagCollisionPawn;
+		bool ViewEngineShowFlagCollisionVisibility;
+		FVector ViewOrigin;
+		FMatrix ViewProjectionMatrix;
 };
 
-extern TGlobalResource<FNullLandscapeRenderSystemResources> GNullLandscapeRenderSystemResources;
-
-extern RENDERER_API TAutoConsoleVariable<float> CVarStaticMeshLODDistanceScale;
-
-struct FLandscapeRenderSystem
-{
 	struct LODSettingsComponent
 	{
 		float LOD0ScreenSizeSquared;
@@ -505,126 +486,37 @@ struct FLandscapeRenderSystem
 		}
 	}
 
+	static TBitArray<> LandscapeIndexAllocator;
+
+	int32 LandscapeIndex;
 	int32 NumRegisteredEntities;
 
 	FIntPoint Min;
 	FIntPoint Size;
 
 	TArray<LODSettingsComponent> SectionLODSettings;
-	TResourceArray<float> SectionLODValues;
 	TResourceArray<float> SectionLODBiases;
 	TArray<FVector4> SectionOriginAndRadius;
 	TArray<FLandscapeComponentSceneProxy*> SceneProxies;
-	TArray<uint8> SectionCurrentFirstLODIndices;
 
-	FBufferRHIRef SectionLODBuffer;
-	FShaderResourceViewRHIRef SectionLODSRV;
 	FBufferRHIRef SectionLODBiasBuffer;
 	FShaderResourceViewRHIRef SectionLODBiasSRV;
 
 	TUniformBufferRef<FLandscapeSectionLODUniformParameters> UniformBuffer;
 
 	FCriticalSection CachedValuesCS;
-	TMap<const FSceneView*, TResourceArray<float>> CachedSectionLODValues;
-	const FSceneView* CachedView;
+	TMap<FViewKey, TResourceArray<float>> CachedSectionLODValues;
 
-	TMap<const FSceneView*, FGraphEventRef> PerViewParametersTasks;
+	TMap<FViewKey, FGraphEventRef> PerViewParametersTasks;
 	FGraphEventRef FetchHeightmapLODBiasesEventRef;
 
-	struct FComputeSectionPerViewParametersTask
-	{
-		FLandscapeRenderSystem& RenderSystem;
-		const FSceneView* ViewPtrAsIdentifier;
-		int32 ViewLODOverride;
-		float ViewLODDistanceFactor;
-		bool ViewEngineShowFlagCollisionPawn;
-		bool ViewEngineShowFlagCollisionVisibility;
-		FVector ViewOrigin;
-		FMatrix ViewProjectionMatrix;
-
-		FComputeSectionPerViewParametersTask(FLandscapeRenderSystem& InRenderSystem, const FSceneView* InView);
-
-		FORCEINLINE TStatId GetStatId() const
-		{
-			RETURN_QUICK_DECLARE_CYCLE_STAT(FComputeSectionPerViewParametersTask, STATGROUP_TaskGraphTasks);
-		}
-
-		ENamedThreads::Type GetDesiredThread()
-		{
-			return ENamedThreads::AnyNormalThreadNormalTask;
-		}
-
-		static ESubsequentsMode::Type GetSubsequentsMode()
-		{
-			return ESubsequentsMode::TrackSubsequents;
-		}
-
-		void AnyThreadTask()
-		{
-			FOptionalTaskTagScope Scope(ETaskTag::EParallelRenderingThread);
-			RenderSystem.ComputeSectionPerViewParameters(
-				ViewPtrAsIdentifier, ViewLODOverride, ViewLODDistanceFactor, 
-				ViewEngineShowFlagCollisionPawn, ViewEngineShowFlagCollisionVisibility, 
-				ViewOrigin, ViewProjectionMatrix);
-		}
-
-		void DoTask(ENamedThreads::Type CurrentThread, const FGraphEventRef& MyCompletionGraphEvent)
-		{
-			AnyThreadTask();
-		}
-	};
-
-	struct FGetSectionLODBiasesTask
-	{
-		FLandscapeRenderSystem& RenderSystem;
-
-		FGetSectionLODBiasesTask(FLandscapeRenderSystem& InRenderSystem)
-			: RenderSystem(InRenderSystem)
-		{
-		}
-
-		FORCEINLINE TStatId GetStatId() const
-		{
-			RETURN_QUICK_DECLARE_CYCLE_STAT(FGetSectionLODBiasesTask, STATGROUP_TaskGraphTasks);
-		}
-
-		ENamedThreads::Type GetDesiredThread()
-		{
-			return ENamedThreads::AnyNormalThreadNormalTask;
-		}
-
-		static ESubsequentsMode::Type GetSubsequentsMode()
-		{
-			return ESubsequentsMode::TrackSubsequents;
-		}
-
-		void AnyThreadTask()
-		{
-			FOptionalTaskTagScope Scope(ETaskTag::EParallelRenderingThread);
-			RenderSystem.FetchHeightmapLODBiases();
-		}
-
-		void DoTask(ENamedThreads::Type CurrentThread, const FGraphEventRef& MyCompletionGraphEvent)
-		{
-			AnyThreadTask();
-		}
-	};
-
-	FLandscapeRenderSystem()
-		: NumRegisteredEntities(0)
-		, Min(MAX_int32, MAX_int32)
-		, Size(EForceInit::ForceInitToZero)
-		, CachedView(nullptr)
-	{
-		SectionLODValues.SetAllowCPUAccess(true);
-		SectionLODBiases.SetAllowCPUAccess(true);
-	}
+	FLandscapeRenderSystem();
+	~FLandscapeRenderSystem();
 
 	void RegisterEntity(FLandscapeComponentSceneProxy* SceneProxy);
-
 	void UnregisterEntity(FLandscapeComponentSceneProxy* SceneProxy);
 
-	int32 GetComponentLinearIndex(FIntPoint ComponentBase)
+	int32 GetComponentLinearIndex(FIntPoint ComponentBase) const
 	{
 		return (ComponentBase.Y - Min.Y) * Size.X + ComponentBase.X - Min.X;
 	}
@@ -645,40 +537,66 @@ struct FLandscapeRenderSystem
 		SceneProxies[GetComponentLinearIndex(ComponentBase)] = SceneProxy;
 	}
 
-	float GetSectionLODValue(FIntPoint ComponentBase)
+	float GetSectionLODValue(const FSceneView& SceneView, FIntPoint ComponentBase) const
 	{
-		return SectionLODValues[GetComponentLinearIndex(ComponentBase)];
+		return CachedSectionLODValues[SceneView.GetViewKey()][GetComponentLinearIndex(ComponentBase)];
 	}
 
-	float GetSectionLODBias(FIntPoint ComponentBase)
+	float GetSectionLODBias(FIntPoint ComponentBase) const
 	{
 		return SectionLODBiases[GetComponentLinearIndex(ComponentBase)];
 	}
 
 	void ComputeSectionPerViewParameters(
-		const FSceneView* ViewPtrAsIdentifier,
+		FViewKey ViewKey,
 		int32 ViewLODOverride,
 		float ViewLODDistanceFactor,
 		bool bDrawCollisionPawn,
 		bool bDrawCollisionCollision,
-		FVector ViewOrigin,
-		FMatrix ViewProjectionMarix);
+		const FVector& ViewOrigin,
+		const FMatrix& ViewProjectionMarix,
+		const TArray<uint8>& SectionCurrentFirstLODIndices);
 
-	void PrepareView(const FSceneView* View);
+	void PrepareView(const FViewParams& InViewParams);
 
-	void BeginRenderView(const FSceneView* View);
+	void BeginRender();
 
 	void BeginFrame();
 
 	void FetchHeightmapLODBiases();
 
-	void RecreateBuffers(const FSceneView* InView = nullptr);
+	void UpdateBuffers();
 
 	void EndFrame();
 };
 
 LANDSCAPE_API extern TMap<FLandscapeNeighborInfo::FLandscapeKey, FLandscapeRenderSystem*> LandscapeRenderSystems;
 
+
+//
+// FLandscapeSceneViewExtension
+//
+class FLandscapeSceneViewExtension : public FSceneViewExtensionBase
+{
+public:
+	FLandscapeSceneViewExtension(const FAutoRegister& AutoReg);
+	virtual ~FLandscapeSceneViewExtension();
+
+	void BeginFrame();
+	void EndFrame();
+
+	virtual void SetupViewFamily(FSceneViewFamily& InViewFamily) override {}
+	virtual void SetupView(FSceneViewFamily& InViewFamily, FSceneView& InView) override;
+	virtual void BeginRenderViewFamily(FSceneViewFamily& InViewFamily) override {}
+
+	virtual void PreRenderViewFamily_RenderThread(FRHICommandListImmediate& RHICmdList, FSceneViewFamily& InViewFamily) override;
+	virtual void PreRenderView_RenderThread(FRHICommandListImmediate& RHICmdList, FSceneView& InView) override;
+};
+
+
+//
+// FLandscapeVisibilityHelper
+//
 class FLandscapeVisibilityHelper
 {
 public:
diff --git a/Engine/Source/Runtime/Renderer/Private/SceneRendering.cpp b/Engine/Source/Runtime/Renderer/Private/SceneRendering.cpp
--- a/Engine/Source/Runtime/Renderer/Private/SceneRendering.cpp
+++ b/Engine/Source/Runtime/Renderer/Private/SceneRendering.cpp
@@ -1894,7 +1894,7 @@ void FViewInfo::SetupUniformBufferParameters(
 	ViewUniformShaderParameters.ShadingEnergyClothSpecTexture	 = OrBlack2DIfNull(ViewUniformShaderParameters.ShadingEnergyClothSpecTexture);
 	ViewUniformShaderParameters.ShadingEnergyDiffuseTexture		 = OrBlack2DIfNull(ViewUniformShaderParameters.ShadingEnergyDiffuseTexture);
 
-	// Water
+	// Water global resources
 	if (WaterDataBuffer.IsValid() && WaterIndirectionBuffer.IsValid())
 	{
 		ViewUniformShaderParameters.WaterIndirection = WaterIndirectionBuffer.GetReference();
@@ -1906,6 +1906,18 @@ void FViewInfo::SetupUniformBufferParameters(
 		ViewUniformShaderParameters.WaterData = GWhiteVertexBufferWithSRV->ShaderResourceViewRHI;
 	}
 
+	// Landscape global resources
+	if (LandscapePerComponentDataBuffer.IsValid() && LandscapeIndirectionBuffer.IsValid())
+	{
+		ViewUniformShaderParameters.LandscapeIndirection = LandscapeIndirectionBuffer.GetReference();
+		ViewUniformShaderParameters.LandscapePerComponentData = LandscapePerComponentDataBuffer.GetReference();
+	}
+	else
+	{
+		ViewUniformShaderParameters.LandscapeIndirection = GWhiteVertexBufferWithSRV->ShaderResourceViewRHI;
+		ViewUniformShaderParameters.LandscapePerComponentData = GWhiteVertexBufferWithSRV->ShaderResourceViewRHI;
+	}
+
 	ViewUniformShaderParameters.VTFeedbackBuffer = GVirtualTextureFeedbackBuffer.GetUAV();
 
 	ViewUniformShaderParameters.GPUSceneViewId = GPUSceneViewId;
@@ -2990,8 +3002,8 @@ void FSceneRenderer::ComputeFamilySize()
 		MaxFamilyX = FMath::Max(MaxFamilyX, FinalViewMaxX);
 		MaxFamilyY = FMath::Max(MaxFamilyY, FinalViewMaxY);
 
-		InstancedStereoWidth = FPlatformMath::Max(InstancedStereoWidth, static_cast<uint32>(View.ViewRect.Max.X));
-	}
+			InstancedStereoWidth = FPlatformMath::Max(InstancedStereoWidth, static_cast<uint32>(View.ViewRect.Max.X));
+		}
 
 	for (FViewInfo& View : Views)
 	{
diff --git a/Engine/Source/Runtime/VulkanRHI/Private/VulkanBuffer.cpp b/Engine/Source/Runtime/VulkanRHI/Private/VulkanBuffer.cpp
--- a/Engine/Source/Runtime/VulkanRHI/Private/VulkanBuffer.cpp
+++ b/Engine/Source/Runtime/VulkanRHI/Private/VulkanBuffer.cpp
@@ -136,7 +136,17 @@ FVulkanResourceMultiBuffer::FVulkanResourceMultiBuffer(FVulkanDevice* InDevice,
 
 			// Get a dummy buffer as sometimes the high-level misbehaves and tries to use SRVs off volatile buffers before filling them in...
 			void* Data = Lock(bRenderThread, RLM_WriteOnly, InSize, 0);
-			FMemory::Memzero(Data, InSize);
+
+			if (CreateInfo.ResourceArray)
+			{
+				uint32 CopyDataSize = FMath::Min(InSize, CreateInfo.ResourceArray->GetResourceDataSize());
+				FMemory::Memcpy(Data, CreateInfo.ResourceArray->GetResourceData(), CopyDataSize);
+			}
+			else
+			{
+				FMemory::Memzero(Data, InSize);
+			}
+
 			Unlock(bRenderThread);
 		}
 		else
