Fixed black landscape material thumbnails : terrain layer weight static material parameters are not filled by GetStaticParameters anymore. So we cannot parse through the list and overriding the one that correspond to the thumbnail's terrain layer like before.

#tests editor

diff --git a/Engine/Source/Runtime/Landscape/Private/LandscapeEdit.cpp b/Engine/Source/Runtime/Landscape/Private/LandscapeEdit.cpp
--- a/Engine/Source/Runtime/Landscape/Private/LandscapeEdit.cpp
+++ b/Engine/Source/Runtime/Landscape/Private/LandscapeEdit.cpp
@@ -194,18 +194,8 @@ ULandscapeMaterialInstanceConstant* ALandscapeProxy::GetLayerThumbnailMIC(UMater
 	FStaticParameterSet StaticParameters;
 	MaterialInstance->GetStaticParameterValues(StaticParameters);
 
-	for (int32 LayerParameterIdx = 0; LayerParameterIdx < StaticParameters.TerrainLayerWeightParameters.Num(); ++LayerParameterIdx)
-	{
-		FStaticTerrainLayerWeightParameter& LayerParameter = StaticParameters.TerrainLayerWeightParameters[LayerParameterIdx];
-		if (LayerParameter.LayerName == LayerName)
-		{
-			LayerParameter.WeightmapIndex = 0;
-		}
-		else
-		{
-			LayerParameter.WeightmapIndex = INDEX_NONE;
-		}
-	}
+	// Customize that material instance to only enable our terrain layer's weightmap : 
+	StaticParameters.TerrainLayerWeightParameters.Add(FStaticTerrainLayerWeightParameter(LayerName, /*InWeightmapIndex = */0, /*bInWeightBasedBlend = */false));
 
 	// Don't recreate the render state of everything, only update the materials context
 	{
