Fix IsMobilePlatform() which should be checking for feature level ES31 and not DDPI bMobile which indicates a mobile device.

diff --git a/Engine/Source/Runtime/Landscape/Private/LandscapeGrass.cpp b/Engine/Source/Runtime/Landscape/Private/LandscapeGrass.cpp
--- a/Engine/Source/Runtime/Landscape/Private/LandscapeGrass.cpp
+++ b/Engine/Source/Runtime/Landscape/Private/LandscapeGrass.cpp
@@ -243,7 +243,7 @@ static bool ShouldCacheLandscapeGrassShaders(const FMeshMaterialShaderPermutatio
 	return (Parameters.MaterialParameters.bIsUsedWithLandscape || Parameters.MaterialParameters.bIsSpecialEngineMaterial) &&
 		IsFeatureLevelSupported(Parameters.Platform, ERHIFeatureLevel::SM5) &&
 		Parameters.VertexFactoryType == FindVertexFactoryType(FName(TEXT("FLandscapeFixedGridVertexFactory"), FNAME_Find)) && 
-		!IsConsolePlatform(Parameters.Platform);
+		!IsConsolePlatform(Parameters.Platform) && !IsVulkanMobileSM5Platform(Parameters.Platform) && !IsMetalMobileSM5Platform(Parameters.Platform);
 }
 
 class FLandscapeGrassWeightShaderElementData : public FMeshMaterialShaderElementData
diff --git a/Engine/Source/Runtime/RHI/Public/RHIDefinitions.h b/Engine/Source/Runtime/RHI/Public/RHIDefinitions.h
--- a/Engine/Source/Runtime/RHI/Public/RHIDefinitions.h
+++ b/Engine/Source/Runtime/RHI/Public/RHIDefinitions.h
@@ -1933,7 +1933,7 @@ inline bool IsPCPlatform(const FStaticShaderPlatform Platform)
 /** Whether the shader platform corresponds to the ES3.1/Metal/Vulkan feature level. */
 inline bool IsMobilePlatform(const FStaticShaderPlatform Platform)
 {
-	return FDataDrivenShaderPlatformInfo::GetIsMobile(Platform);
+	return FDataDrivenShaderPlatformInfo::GetMaxFeatureLevel(Platform) == ERHIFeatureLevel::ES3_1;
 }
 
 inline bool IsOpenGLPlatform(const FStaticShaderPlatform Platform)
