Added methods to get landscape bounds, esp. in WP situations Supports querying currently loaded & complete bounds.

diff --git a/Engine/Source/Runtime/Landscape/Classes/Landscape.h b/Engine/Source/Runtime/Landscape/Classes/Landscape.h
--- a/Engine/Source/Runtime/Landscape/Classes/Landscape.h
+++ b/Engine/Source/Runtime/Landscape/Classes/Landscape.h
@@ -240,6 +240,14 @@ public:
 	virtual void FinishDestroy() override;
 	//~ End UObject Interface
 
+	/** Computes & returns bounds containing all currently loaded landscape proxies (if any) or this landscape's bounds otherwise */
+	LANDSCAPE_API FBox GetLoadedBounds() const;
+
+#if WITH_EDITOR
+	/** Computes & returns bounds containing all landscape proxies (if any) or this landscape's bounds otherwise. Note that in non-WP worlds this will call GetLoadedBounds(). */
+	LANDSCAPE_API FBox GetCompleteBounds() const;
+#endif
+
 	LANDSCAPE_API bool IsUpToDate() const;
 	LANDSCAPE_API void TickLayers(float DeltaTime);
 	// Layers stuff
diff --git a/Engine/Source/Runtime/Landscape/Classes/LandscapeInfo.h b/Engine/Source/Runtime/Landscape/Classes/LandscapeInfo.h
--- a/Engine/Source/Runtime/Landscape/Classes/LandscapeInfo.h
+++ b/Engine/Source/Runtime/Landscape/Classes/LandscapeInfo.h
@@ -168,6 +168,12 @@ public:
 	virtual void Serialize(FArchive& Ar) override;
 	//~ End UObject Interface
 
+	FBox GetLoadedBounds() const;
+
+#if WITH_EDITOR
+	FBox GetCompleteBounds() const;
+#endif
+
 #if WITH_EDITOR
 	virtual void BeginDestroy() override;
 
diff --git a/Engine/Source/Runtime/Landscape/Private/Landscape.cpp b/Engine/Source/Runtime/Landscape/Private/Landscape.cpp
--- a/Engine/Source/Runtime/Landscape/Private/Landscape.cpp
+++ b/Engine/Source/Runtime/Landscape/Private/Landscape.cpp
@@ -2111,6 +2111,18 @@ void ALandscape::PostLoad()
 	Super::PostLoad();
 }
 
+FBox ALandscape::GetLoadedBounds() const
+{
+	return GetLandscapeInfo()->GetLoadedBounds();
+}
+
+#if WITH_EDITOR
+FBox ALandscape::GetCompleteBounds() const
+{
+	return GetLandscapeInfo()->GetCompleteBounds();
+}
+#endif
+
 #if WITH_EDITOR
 void ALandscapeProxy::OnFeatureLevelChanged(ERHIFeatureLevel::Type NewFeatureLevel)
 {
@@ -3517,6 +3529,84 @@ void ULandscapeInfo::UnregisterActorComponent(ULandscapeComponent* Component)
 	}
 }
 
+FBox ULandscapeInfo::GetLoadedBounds() const
+{
+	FBox Bounds(EForceInit::ForceInit);
+
+	auto UpdateBounds = [&Bounds](ALandscapeProxy* Proxy)
+	{
+		const bool bOnlyCollidingComponents = false;
+		const bool bIncludeChildActors = false;
+		FVector Origin;
+		FVector BoxExtents;
+
+		Proxy->GetActorBounds(bOnlyCollidingComponents, Origin, BoxExtents, bIncludeChildActors);
+
+		// Reject invalid bounds
+		if (BoxExtents != FVector::Zero())
+		{
+			Bounds += FBox::BuildAABB(Origin, BoxExtents);
+		}
+	};
+
+	if (LandscapeActor.IsValid())
+	{
+		UpdateBounds(LandscapeActor.Get());
+	}
+
+	// Since in PIE/in-game the Proxies aren't populated, we must iterate through the loaded components
+	// but this is functionally equivalent to calling ForAllLandscapeProxies
+	TSet<ALandscapeProxy*> LoadedProxies;
+	for (auto It = XYtoComponentMap.CreateConstIterator(); It; ++It)
+	{
+		if (!It.Value())
+		{
+			continue;
+		}
+
+		if (ALandscapeProxy* Proxy = Cast<ALandscapeProxy>(It.Value()->GetOwner()))
+		{
+			LoadedProxies.Add(Proxy);
+		}
+	}
+
+	for (ALandscapeProxy* Proxy : LoadedProxies)
+	{
+		UpdateBounds(Proxy);
+	}
+
+	return Bounds;
+}
+
+#if WITH_EDITOR
+FBox ULandscapeInfo::GetCompleteBounds() const
+{
+	ALandscape* Landscape = LandscapeActor.Get();
+
+	// In a non-WP situation, the current actor's bounds will do.
+	if(!Landscape || !Landscape->GetWorld() || !Landscape->GetWorld()->GetWorldPartition())
+	{
+		return GetLoadedBounds();
+	}
+
+	FBox Bounds(EForceInit::ForceInit);
+
+	for (const FWorldPartitionHandle& ProxyHandle : ProxyHandles)
+	{
+		// Skip owning landscape actor
+		ALandscapeProxy* LandscapeProxy = Cast<ALandscapeProxy>(ProxyHandle->GetActor());
+		if (LandscapeProxy == Landscape)
+		{
+			continue;
+		}
+
+		Bounds += ProxyHandle->GetBounds();
+	}
+
+	return Bounds;
+}
+#endif
+
 void ULandscapeComponent::PostInitProperties()
 {
 	Super::PostInitProperties();
