Cleaned & clarified behavior of GetCompleteBounds on the landscape when proxies are loaded.

diff --git a/Engine/Source/Runtime/Landscape/Private/Landscape.cpp b/Engine/Source/Runtime/Landscape/Private/Landscape.cpp
--- a/Engine/Source/Runtime/Landscape/Private/Landscape.cpp
+++ b/Engine/Source/Runtime/Landscape/Private/Landscape.cpp
@@ -3529,11 +3529,9 @@ void ULandscapeInfo::UnregisterActorComponent(ULandscapeComponent* Component)
 	}
 }
 
-FBox ULandscapeInfo::GetLoadedBounds() const
+namespace LandscapeInfoBoundsHelper
 {
-	FBox Bounds(EForceInit::ForceInit);
-
-	auto UpdateBounds = [&Bounds](ALandscapeProxy* Proxy)
+	void AccumulateBounds(ALandscapeProxy* Proxy, FBox& Bounds)
 	{
 		const bool bOnlyCollidingComponents = false;
 		const bool bIncludeChildActors = false;
@@ -3547,11 +3545,16 @@ FBox ULandscapeInfo::GetLoadedBounds() const
 		{
 			Bounds += FBox::BuildAABB(Origin, BoxExtents);
 		}
-	};
+	}
+}
+
+FBox ULandscapeInfo::GetLoadedBounds() const
+{
+	FBox Bounds(EForceInit::ForceInit);
 
 	if (LandscapeActor.IsValid())
 	{
-		UpdateBounds(LandscapeActor.Get());
+		LandscapeInfoBoundsHelper::AccumulateBounds(LandscapeActor.Get(), Bounds);
 	}
 
 	// Since in PIE/in-game the Proxies aren't populated, we must iterate through the loaded components
@@ -3572,7 +3575,7 @@ FBox ULandscapeInfo::GetLoadedBounds() const
 
 	for (ALandscapeProxy* Proxy : LoadedProxies)
 	{
-		UpdateBounds(Proxy);
+		LandscapeInfoBoundsHelper::AccumulateBounds(Proxy, Bounds);
 	}
 
 	return Bounds;
@@ -3599,8 +3602,15 @@ FBox ULandscapeInfo::GetCompleteBounds() const
 		{
 			continue;
 		}
-
-		Bounds += ProxyHandle->GetBounds();
+		else if (LandscapeProxy)
+		{
+			// Prioritize loaded bounds, as the bounds in the actor desc might not be up-to-date
+			LandscapeInfoBoundsHelper::AccumulateBounds(LandscapeProxy, Bounds);
+		}
+		else
+		{
+			Bounds += ProxyHandle->GetBounds();
+		}
 	}
 
 	return Bounds;
