Fix FLandscapeComponentSceneProxy::CreateRenderResources() trying to access it's associated landscape render system when no neighbors have been registered.
* This could occur if VisibilityHelper.ShouldBeVisible() returned false, in which case the neighbors weren't registered.
* Avoid creating render resources until VisibilityHelper.ShouldBeVisible() is true

diff --git a/Engine/Source/Runtime/Landscape/Private/LandscapeRender.cpp b/Engine/Source/Runtime/Landscape/Private/LandscapeRender.cpp
--- a/Engine/Source/Runtime/Landscape/Private/LandscapeRender.cpp
+++ b/Engine/Source/Runtime/Landscape/Private/LandscapeRender.cpp
@@ -1506,11 +1506,13 @@ void FLandscapeComponentSceneProxy::CreateRenderThreadResources()
 
 	check(HeightmapTexture != nullptr);
 
-	if (VisibilityHelper.ShouldBeVisible())
+	if (!VisibilityHelper.ShouldBeVisible())
 	{
-		RegisterNeighbors(this);
+		return;
 	}
 
+	RegisterNeighbors(this);
+
 	auto FeatureLevel = GetScene().GetFeatureLevel();
 
 	SharedBuffers = FLandscapeComponentSceneProxy::SharedBuffersMap.FindRef(SharedBuffersKey);
@@ -1683,7 +1685,7 @@ bool FLandscapeComponentSceneProxy::OnLevelAddedToWorld_RenderThread()
 	if (VisibilityHelper.OnAddedToWorld())
 	{
 		SetForceHidden(false);
-		RegisterNeighbors(this);
+		CreateRenderThreadResources();
 		return true;
 	}
 	return false;
@@ -2200,6 +2202,11 @@ void FLandscapeComponentSceneProxy::DrawStaticElements(FStaticPrimitiveDrawInter
 		return;
 	}
 
+	if (!VisibilityHelper.ShouldBeVisible())
+	{
+		return;
+	}
+
 	int32 TotalBatchCount = 1 + LastLOD - FirstLOD;
 	TotalBatchCount += (1 + LastVirtualTextureLOD - FirstVirtualTextureLOD) * RuntimeVirtualTextureMaterialTypes.Num();
 
