* Added detection of discrepancies between the landscape material instance's TerrainLayerWeightParameters and the landscape component's weightmap allocations. A recent change (18095319) introduced potential issues with static parameters within landscape component's "parent" material instance constants (the parent MIC of all landscape components from a landscape proxy sharing the same weightmap combination). Loading the map would register these invalid MIC in the proxy's MaterialInstanceConstantMap with a valid layer allocation key, preventing the landscape component from ever "fixing" its MICs until the layer allocation key changed (e.g. by painting with a previously unused layer on a component). When such errors are detected, we now regenerate all of the proxy's material instances in order to update the material and its TerrainLayerWeightParameters with the proper weightmap allocations
* Renamed visibility layer name to "__LANDSCAPE_VISIBILITY__" to centralize/simplify the handling of that special layer's name
* Fixed issue with FixupWeightmaps(), that could trigger invalid weightmap usage validation errors when detecting and deleting unused paint layers in components
* Removed deletion of layers with a null layer info in PostEditChangeProperty since it's done, and in a better way (for the base weightmap allocations and the edit layers ones), in PostRegisterAllComponents's FixupWeightmaps

#tests editor, full QA pass by tim.gautier

diff --git a/Engine/Source/Editor/LandscapeEditor/Private/LandscapeEditorDetailCustomization_TargetLayers.cpp b/Engine/Source/Editor/LandscapeEditor/Private/LandscapeEditorDetailCustomization_TargetLayers.cpp
--- a/Engine/Source/Editor/LandscapeEditor/Private/LandscapeEditorDetailCustomization_TargetLayers.cpp
+++ b/Engine/Source/Editor/LandscapeEditor/Private/LandscapeEditorDetailCustomization_TargetLayers.cpp
@@ -1151,7 +1151,7 @@ void FLandscapeEditorCustomNodeBuilder_TargetLayers::OnRebuildMICs(const TShared
 {
 	if (Target->LandscapeInfo.IsValid())
 	{
-		Target->LandscapeInfo.Get()->UpdateAllComponentMaterialInstances();
+		Target->LandscapeInfo.Get()->UpdateAllComponentMaterialInstances(/*bInvalidateCombinationMaterials = */true);
 	}
 }
 
diff --git a/Engine/Source/Runtime/Core/Private/Misc/MapErrors.cpp b/Engine/Source/Runtime/Core/Private/Misc/MapErrors.cpp
--- a/Engine/Source/Runtime/Core/Private/Misc/MapErrors.cpp
+++ b/Engine/Source/Runtime/Core/Private/Misc/MapErrors.cpp
@@ -31,6 +31,7 @@ FLazyName FMapErrors::DeprecatedClass(TEXT("DeprecatedClass"));
 FLazyName FMapErrors::FoliageMissingStaticMesh(TEXT("FoliageMissingStaticMesh"));
 FLazyName FMapErrors::FoliageMissingClusterComponent(TEXT("FoliageMissingStaticMesh"));
 FLazyName FMapErrors::FixedUpDeletedLayerWeightmap(TEXT("FixedUpDeletedLayerWeightmap"));
+FLazyName FMapErrors::FixedUpInvalidLandscapeMaterialInstances(TEXT("FixedUpInvalidLandscapeMaterialInstances"));
 FLazyName FMapErrors::FixedUpIncorrectLayerWeightmap(TEXT("FixedUpIncorrectLayerWeightmap"));
 FLazyName FMapErrors::FixedUpSharedLayerWeightmap(TEXT("FixedUpSharedLayerWeightmap"));
 FLazyName FMapErrors::LandscapeComponentPostLoad_Warning(TEXT("LandscapeComponentPostLoad_Warning"));
diff --git a/Engine/Source/Runtime/Core/Public/Misc/MapErrors.h b/Engine/Source/Runtime/Core/Public/Misc/MapErrors.h
--- a/Engine/Source/Runtime/Core/Public/Misc/MapErrors.h
+++ b/Engine/Source/Runtime/Core/Public/Misc/MapErrors.h
@@ -108,16 +108,19 @@ struct CORE_API FMapErrors
 
     /**  Landscape */
 
-	/** {LandscapeComponent} : Fixed up deleted layer weightmap: **TODO** */
+	/** {LandscapeComponent} : Fixed up deleted layer weightmap: */
     static FLazyName FixedUpDeletedLayerWeightmap;
     
-	/** {LandscapeComponent} : Fixed up incorrect layer weightmap texture index: **TODO** */
+	/** {LandscapeComponent} : Fixed up incorrect layer weightmap texture index: */
     static FLazyName FixedUpIncorrectLayerWeightmap;
     
-	/** Fixed up shared weightmap texture for layer {Layer} in component {Component} (shares with {Name}): **TODO** */
+	/** {LandscapeComponent} : Fixed up invalid material instances for a landscape component when its weightmap allocations don't match the material's layer configuration : */
+	static FLazyName FixedUpInvalidLandscapeMaterialInstances;
+
+	/** Fixed up shared weightmap texture for layer {Layer} in component {Component} (shares with {Name}): */
     static FLazyName FixedUpSharedLayerWeightmap;
 
-	/** Landscape ({ProxyName}) has overlapping render components at location ({X, Y}): **TODO** */
+	/** Landscape ({ProxyName}) has overlapping render components at location ({X, Y}): */
 	static FLazyName LandscapeComponentPostLoad_Warning;
 
 
diff --git a/Engine/Source/Runtime/Landscape/Classes/LandscapeComponent.h b/Engine/Source/Runtime/Landscape/Classes/LandscapeComponent.h
--- a/Engine/Source/Runtime/Landscape/Classes/LandscapeComponent.h
+++ b/Engine/Source/Runtime/Landscape/Classes/LandscapeComponent.h
@@ -1057,6 +1057,8 @@ public:
 	 */
 	static FString GetLayerAllocationKey(const TArray<FWeightmapLayerAllocationInfo>& Allocations, UMaterialInterface* LandscapeMaterial, bool bMobile = false);
 
+	bool ValidateCombinationMaterial(UMaterialInstanceConstant* InCombinationMaterial) const;
+
 	/** @todo document */
 	void GetLayerDebugColorKey(int32& R, int32& G, int32& B) const;
 
diff --git a/Engine/Source/Runtime/Landscape/Classes/LandscapeInfo.h b/Engine/Source/Runtime/Landscape/Classes/LandscapeInfo.h
--- a/Engine/Source/Runtime/Landscape/Classes/LandscapeInfo.h
+++ b/Engine/Source/Runtime/Landscape/Classes/LandscapeInfo.h
@@ -294,7 +294,7 @@ public:
 	LANDSCAPE_API void MoveSplines(ULandscapeSplinesComponent* InSplineComponent, TScriptInterface<ILandscapeSplineInterface> InNewOwner);
 
 	/** Will call UpdateAllComponentMaterialInstances on all LandscapeProxies */
-	LANDSCAPE_API void UpdateAllComponentMaterialInstances();
+	LANDSCAPE_API void UpdateAllComponentMaterialInstances(bool bInInvalidateCombinationMaterials = false);
 
 	/** Returns LandscapeStreamingProxy Cell Size in WorldPartition */
 	LANDSCAPE_API uint32 GetGridSize(uint32 InGridSizeInComponents) const;
diff --git a/Engine/Source/Runtime/Landscape/Classes/LandscapeProxy.h b/Engine/Source/Runtime/Landscape/Classes/LandscapeProxy.h
--- a/Engine/Source/Runtime/Landscape/Classes/LandscapeProxy.h
+++ b/Engine/Source/Runtime/Landscape/Classes/LandscapeProxy.h
@@ -679,7 +679,11 @@ public:
 	UPROPERTY(Transient, NonTransactional)
 	TMap<TObjectPtr<UTexture2D>, TObjectPtr<ULandscapeWeightmapUsage>> WeightmapUsageMap;
 
+	/** Set to true when on undo, when it's necessary to completely regenerate weightmap usages (since some weightmap allocations are transactional and others not, e.g. splines edit layer) */
 	bool bNeedsWeightmapUsagesUpdate = false;
+
+	/** Set to true when we know that weightmap usages are being reconstructed and might be temporarily invalid as a result (ValidateProxyLayersWeightmapUsage should be called after setting this back to false) */
+	bool bTemporarilyDisableWeightmapUsagesValidation = false;
 #endif
 
 	/** Data set at creation time */
@@ -999,8 +1003,8 @@ public:
 	LANDSCAPE_API void RemoveXYOffsets();
 
 	/** Update the material instances for all the landscape components */
-	LANDSCAPE_API void UpdateAllComponentMaterialInstances();
-	LANDSCAPE_API void UpdateAllComponentMaterialInstances(FMaterialUpdateContext& InOutMaterialContext, TArray<class FComponentRecreateRenderStateContext>& InOutRecreateRenderStateContext);
+	LANDSCAPE_API void UpdateAllComponentMaterialInstances(bool bInInvalidateCombinationMaterials = false);
+	LANDSCAPE_API void UpdateAllComponentMaterialInstances(FMaterialUpdateContext& InOutMaterialContext, TArray<class FComponentRecreateRenderStateContext>& InOutRecreateRenderStateContext, bool bInInvalidateCombinationMaterials = false);
 
 	/** Create a thumbnail material for a given layer */
 	LANDSCAPE_API static ULandscapeMaterialInstanceConstant* GetLayerThumbnailMIC(UMaterialInterface* LandscapeMaterial, FName LayerName, UTexture2D* ThumbnailWeightmap, UTexture2D* ThumbnailHeightmap, ALandscapeProxy* Proxy);
diff --git a/Engine/Source/Runtime/Landscape/Private/Landscape.cpp b/Engine/Source/Runtime/Landscape/Private/Landscape.cpp
--- a/Engine/Source/Runtime/Landscape/Private/Landscape.cpp
+++ b/Engine/Source/Runtime/Landscape/Private/Landscape.cpp
@@ -1077,6 +1077,7 @@ ALandscapeProxy::ALandscapeProxy(const FObjectInitializer& ObjectInitializer)
 		// This layer should be no weight blending
 		VisibilityLayer->bNoWeightBlend = true;
 #endif
+		VisibilityLayer->LayerName = UMaterialExpressionLandscapeVisibilityMask::ParameterName;
 		VisibilityLayer->LayerUsageDebugColor = FLinearColor(0, 0, 0, 0);
 		VisibilityLayer->AddToRoot();
 	}
@@ -2536,21 +2537,39 @@ void ALandscapeProxy::PostLoad()
 		EditorCachedLayerInfos_DEPRECATED.Empty();
 	}
 
+	bool bFixedUpInvalidMaterialInstances = false;
+
 	for (ULandscapeComponent* Comp : LandscapeComponents)
 	{
 		if (Comp)
 		{
-			// Store the layer combination in the MaterialInstanceConstantMap
+			// Validate the layer combination and store it in the MaterialInstanceConstantMap
 			if (UMaterialInstance* MaterialInstance = Comp->GetMaterialInstance(0, false))
 			{
-				if(UMaterialInstanceConstant* CombinationMaterialInstance = Cast<UMaterialInstanceConstant>(MaterialInstance->Parent))
+				UMaterialInstanceConstant* CombinationMaterialInstance = Cast<UMaterialInstanceConstant>(MaterialInstance->Parent);
+				if (Comp->ValidateCombinationMaterial(CombinationMaterialInstance))
 				{
 					MaterialInstanceConstantMap.Add(*ULandscapeComponent::GetLayerAllocationKey(Comp->GetWeightmapLayerAllocations(), CombinationMaterialInstance->Parent), CombinationMaterialInstance);
 				}
+				else
+				{
+					// There was a problem with the loaded material : it doesn't match the expected material combination, we need to regenerate the material instances : 
+					Comp->UpdateMaterialInstances();
+					bFixedUpInvalidMaterialInstances = true;
+				}
 			}
 		}
 	}
 
+	if (bFixedUpInvalidMaterialInstances)
+	{
+		FFormatNamedArguments Arguments;
+		Arguments.Add(TEXT("LandscapeName"), FText::FromString(GetPathName()));
+		FMessageLog("MapCheck").Warning()
+			->AddToken(FTextToken::Create(FText::Format(LOCTEXT("MapCheck_Message_FixedUpInvalidLandscapeMaterialInstances", "{LandscapeName} : Fixed up invalid landscape material instances. Please re-save the landscape package."), Arguments)))
+			->AddToken(FMapErrorToken::Create(FMapErrors::FixedUpInvalidLandscapeMaterialInstances));
+	}
+
 	// track feature level change to flush grass cache
 	if (GetWorld())
 	{
diff --git a/Engine/Source/Runtime/Landscape/Private/LandscapeEdit.cpp b/Engine/Source/Runtime/Landscape/Private/LandscapeEdit.cpp
--- a/Engine/Source/Runtime/Landscape/Private/LandscapeEdit.cpp
+++ b/Engine/Source/Runtime/Landscape/Private/LandscapeEdit.cpp
@@ -213,6 +213,52 @@ ULandscapeMaterialInstanceConstant* ALandscapeProxy::GetLayerThumbnailMIC(UMater
 	return MaterialInstance;
 }
 
+bool ULandscapeComponent::ValidateCombinationMaterial(UMaterialInstanceConstant* InCombinationMaterial) const
+{
+	if (InCombinationMaterial == nullptr)
+	{
+		return false;
+	}
+
+	const TArray<FStaticTerrainLayerWeightParameter>& TerrainLayerWeightParameters = InCombinationMaterial->GetStaticParameters().TerrainLayerWeightParameters;
+	const TArray<FWeightmapLayerAllocationInfo>& ComponentWeightmapAllocations = GetWeightmapLayerAllocations();
+
+	if (TerrainLayerWeightParameters.Num() != ComponentWeightmapAllocations.Num())
+	{
+		UE_LOG(LogLandscape, Display, TEXT("Material instance %s in landscape component %s doesn't match the expected shader combination: different number of allocations (expected: %i, found: %i)"),
+			*InCombinationMaterial->GetName(), *GetPathName(), ComponentWeightmapAllocations.Num(), TerrainLayerWeightParameters.Num());
+
+		return false;
+	}
+
+	for (const FWeightmapLayerAllocationInfo& Allocation : ComponentWeightmapAllocations)
+	{
+		if (Allocation.LayerInfo == nullptr)
+		{
+			UE_LOG(LogLandscape, Display, TEXT("Material instance %s in landscape component %s doesn't match the expected shader combination: invalid layer info"),
+				*InCombinationMaterial->GetName(), *GetPathName(), ComponentWeightmapAllocations.Num(), TerrainLayerWeightParameters.Num());
+
+			return false;
+		}
+
+		// Each weightmap allocation must have its equivalent in the material's TerrainLayerWeightParameters : 
+		if (!TerrainLayerWeightParameters.FindByPredicate([&](const FStaticTerrainLayerWeightParameter& TerrainLayerWeightParameter)
+			{
+				return ((TerrainLayerWeightParameter.LayerName == Allocation.LayerInfo->LayerName)
+					&& (TerrainLayerWeightParameter.WeightmapIndex == Allocation.WeightmapTextureIndex)
+					&& (TerrainLayerWeightParameter.bWeightBasedBlend == !Allocation.LayerInfo->bNoWeightBlend));
+			}))
+		{
+			UE_LOG(LogLandscape, Display, TEXT("Material instance %s in landscape component %s doesn't match the expected shader combination: missing layer %s"),
+				*InCombinationMaterial->GetName(), *GetPathName(), *Allocation.LayerInfo->LayerName.ToString());
+
+			return false;
+		}
+	}
+
+	return true;
+}
+
 /**
 * Generate a key for this component's layer allocations to use with MaterialInstanceConstantMap.
 */
@@ -309,8 +355,7 @@ UMaterialInstanceConstant* ULandscapeComponent::GetCombinationMaterial(FMaterial
 			{
 				if (Allocation.LayerInfo)
 				{
-					const FName LayerName = (Allocation.LayerInfo == ALandscapeProxy::VisibilityLayer) ? UMaterialExpressionLandscapeVisibilityMask::ParameterName : Allocation.LayerInfo->LayerName;
-					StaticParameters.TerrainLayerWeightParameters.Add(FStaticTerrainLayerWeightParameter(LayerName, Allocation.WeightmapTextureIndex, !Allocation.LayerInfo->bNoWeightBlend));
+					StaticParameters.TerrainLayerWeightParameters.Add(FStaticTerrainLayerWeightParameter(Allocation.GetLayerName(), Allocation.WeightmapTextureIndex, !Allocation.LayerInfo->bNoWeightBlend));
 				}
 			}
 
@@ -398,8 +443,7 @@ void ULandscapeComponent::UpdateMaterialInstances_Internal(FMaterialUpdateContex
 			{
 				const FWeightmapLayerAllocationInfo& Allocation = WeightmapBaseLayerAllocation[AllocIdx];
 
-				FName LayerName = Allocation.LayerInfo == ALandscapeProxy::VisibilityLayer ? UMaterialExpressionLandscapeVisibilityMask::ParameterName : Allocation.LayerInfo ? Allocation.LayerInfo->LayerName : NAME_None;
-				MaterialInstance->SetVectorParameterValueEditorOnly(FName(*FString::Printf(TEXT("LayerMask_%s"), *LayerName.ToString())), Masks[Allocation.WeightmapTextureChannel]);
+				MaterialInstance->SetVectorParameterValueEditorOnly(FName(*FString::Printf(TEXT("LayerMask_%s"), *Allocation.GetLayerName().ToString())), Masks[Allocation.WeightmapTextureChannel]);
 			}
 
 			// Set the weightmaps
@@ -481,8 +525,13 @@ void ULandscapeComponent::UpdateMaterialInstances(FMaterialUpdateContext& InOutM
 	UpdateMaterialInstances_Internal(InOutMaterialContext);
 }
 
-void ALandscapeProxy::UpdateAllComponentMaterialInstances(FMaterialUpdateContext& InOutMaterialContext, TArray<FComponentRecreateRenderStateContext>& InOutRecreateRenderStateContext)
+void ALandscapeProxy::UpdateAllComponentMaterialInstances(FMaterialUpdateContext& InOutMaterialContext, TArray<FComponentRecreateRenderStateContext>& InOutRecreateRenderStateContext, bool bInInvalidateCombinationMaterials)
 {
+	if (bInInvalidateCombinationMaterials)
+	{
+		MaterialInstanceConstantMap.Reset();
+	}
+
 	for (ULandscapeComponent* Component : LandscapeComponents)
 	{
 		Component->UpdateMaterialInstances(InOutMaterialContext, InOutRecreateRenderStateContext);
@@ -490,8 +539,13 @@ void ALandscapeProxy::UpdateAllComponentMaterialInstances(FMaterialUpdateContext
 
 }
 
-void ALandscapeProxy::UpdateAllComponentMaterialInstances()
+void ALandscapeProxy::UpdateAllComponentMaterialInstances(bool bInInvalidateCombinationMaterials)
+{
+	if (bInInvalidateCombinationMaterials)
 {
+		MaterialInstanceConstantMap.Reset();
+	}
+
 	// we're not having the material update context recreate render states because we will manually do it for only our components
 	TArray<FComponentRecreateRenderStateContext> RecreateRenderStateContexts;
 	RecreateRenderStateContexts.Reserve(LandscapeComponents.Num());
@@ -646,6 +700,17 @@ void ALandscapeProxy::FixupWeightmaps()
 {
 	WeightmapUsageMap.Empty();
 
+	// We've just reinitialized the weightmap usages map and FixupWeightmaps will reconstruct it component by component, but we might in the process delete invalid layers (e.g. those whose landscape info has been deleted),
+	//  in which case ValidateProxyLayersWeightmapUsage() on the entire proxy will be called and, because of WeightmapUsageMap's cleanup above, might report missing weightmap usages. 
+	//  To avoid triggering asserts, we simply disable validation until the fixup operation is done : 
+	bTemporarilyDisableWeightmapUsagesValidation = true;
+	ON_SCOPE_EXIT
+	{
+		bTemporarilyDisableWeightmapUsagesValidation = false;
+		// Now that the job is done, weightmap usages should be valid again
+		ValidateProxyLayersWeightmapUsage();
+	};
+
 	for (ULandscapeComponent* Component : LandscapeComponents)
 	{
 		if (Component != nullptr)
@@ -654,7 +719,6 @@ void ALandscapeProxy::FixupWeightmaps()
 		}
 	}
 
-	ValidateProxyLayersWeightmapUsage();
 }
 
 void ALandscapeProxy::RepairInvalidTextures()
@@ -4280,11 +4344,11 @@ void ULandscapeInfo::ClearDirtyData()
 	}
 }
 
-void ULandscapeInfo::UpdateAllComponentMaterialInstances()
+void ULandscapeInfo::UpdateAllComponentMaterialInstances(bool bInInvalidateCombinationMaterials)
 {
-	ForAllLandscapeProxies([](ALandscapeProxy* Proxy)
+	ForAllLandscapeProxies([=](ALandscapeProxy* Proxy)
 	{
-		Proxy->UpdateAllComponentMaterialInstances();
+		Proxy->UpdateAllComponentMaterialInstances(bInInvalidateCombinationMaterials);
 	});
 }
 
@@ -4934,38 +4998,6 @@ void ALandscapeProxy::PostEditChangeProperty(FPropertyChangedEvent& PropertyChan
 	// Remove any null landscape components
 	LandscapeComponents.RemoveAll([](const ULandscapeComponent* Component) { return Component == nullptr; });
 
-	ULandscapeInfo* Info = GetLandscapeInfo();
-	bool bRemovedAnyLayers = false;
-	for (ULandscapeComponent* Component : LandscapeComponents)
-	{
-		const TArray<FWeightmapLayerAllocationInfo>& ComponentWeightmapLayerAllocations = Component->GetWeightmapLayerAllocations(false);
-
-		int32 NumNullLayers = Algo::CountIf(ComponentWeightmapLayerAllocations, [](const FWeightmapLayerAllocationInfo& Allocation) { return Allocation.LayerInfo == nullptr; });
-		if (NumNullLayers > 0)
-		{
-			FLandscapeEditDataInterface LandscapeEdit(Info);
-			for (int32 i = 0; i < NumNullLayers; ++i)
-			{
-				// DeleteLayer doesn't expect duplicates, so we need to call it once for each null
-				Component->DeleteLayer(nullptr, LandscapeEdit);
-			}
-			bRemovedAnyLayers = true;
-		}
-	}
-	if (bRemovedAnyLayers)
-	{
-		ALandscape* LandscapeActor = GetLandscapeActor();
-
-		if(LandscapeActor != nullptr && LandscapeActor->HasLayersContent())
-		{
-			LandscapeActor->RequestLayersContentUpdate(ELandscapeLayerUpdateMode::Update_All);
-		}
-		else
-		{
-			ALandscapeProxy::InvalidateGeneratedComponentData(LandscapeComponents);
-		}
-	}
-
 	// Must do this *after* correcting the scale or reattaching the landscape components will crash!
 	// Must do this *after* clamping values / propogating values to components
 	Super::PostEditChangeProperty(PropertyChangedEvent);
@@ -6306,7 +6338,7 @@ void ULandscapeComponent::GenerateMobileWeightmapLayerAllocations()
 	GetAllMobileRelevantLayerNames(LayerNames, GetLandscapeMaterial()->GetMaterial());
 	MobileWeightmapLayerAllocations = WeightmapLayerAllocations.FilterByPredicate([&](const FWeightmapLayerAllocationInfo& Allocation) -> bool 
 		{
-			return Allocation.LayerInfo && LayerNames.Contains(Allocation.LayerInfo == ALandscapeProxy::VisibilityLayer ? UMaterialExpressionLandscapeVisibilityMask::ParameterName : Allocation.GetLayerName());
+			return Allocation.LayerInfo && LayerNames.Contains(Allocation.GetLayerName());
 		}
 	);
 	MobileWeightmapLayerAllocations.StableSort(([&](const FWeightmapLayerAllocationInfo& A, const FWeightmapLayerAllocationInfo& B) -> bool
@@ -6447,8 +6479,7 @@ void ULandscapeComponent::GeneratePlatformPixelData()
 			{
 				if (Allocation.LayerInfo)
 				{
-					FName LayerName = Allocation.LayerInfo == ALandscapeProxy::VisibilityLayer ? UMaterialExpressionLandscapeVisibilityMask::ParameterName : Allocation.LayerInfo->LayerName;
-					NewMobileMaterialInstance->SetVectorParameterValue(FName(*FString::Printf(TEXT("LayerMask_%s"), *LayerName.ToString())), Masks[Allocation.WeightmapTextureChannel]);
+					NewMobileMaterialInstance->SetVectorParameterValue(FName(*FString::Printf(TEXT("LayerMask_%s"), *Allocation.GetLayerName().ToString())), Masks[Allocation.WeightmapTextureChannel]);
 				}
 			}
 
@@ -6504,8 +6535,7 @@ void ULandscapeComponent::GeneratePlatformPixelData()
 			{
 				if (Allocation.LayerInfo)
 				{
-					FName LayerName = Allocation.LayerInfo == ALandscapeProxy::VisibilityLayer ? UMaterialExpressionLandscapeVisibilityMask::ParameterName : Allocation.LayerInfo->LayerName;
-					NewMobileMaterialInstance->SetVectorParameterValueEditorOnly(FName(*FString::Printf(TEXT("LayerMask_%s"), *LayerName.ToString())), Masks[Allocation.WeightmapTextureChannel]);
+					NewMobileMaterialInstance->SetVectorParameterValueEditorOnly(FName(*FString::Printf(TEXT("LayerMask_%s"), *Allocation.GetLayerName().ToString())), Masks[Allocation.WeightmapTextureChannel]);
 				}
 			}
 
diff --git a/Engine/Source/Runtime/Landscape/Private/LandscapeEditLayers.cpp b/Engine/Source/Runtime/Landscape/Private/LandscapeEditLayers.cpp
--- a/Engine/Source/Runtime/Landscape/Private/LandscapeEditLayers.cpp
+++ b/Engine/Source/Runtime/Landscape/Private/LandscapeEditLayers.cpp
@@ -151,8 +151,8 @@ private:
 
 	void ClearDirty()
 	{
-		if(!GWorld || GWorld->IsGameWorld())
-		{ 
+		if (!GWorld || GWorld->IsGameWorld())
+		{
 			return;
 		}
 
@@ -197,7 +197,7 @@ struct FTextureToComponentHelper
 
 				for (FWeightmapLayerAllocationInfo const& AllocInfo : AllocInfos)
 				{
-					if(AllocInfo.IsAllocated() && AllocInfo.WeightmapTextureIndex < WeightmapTextures.Num())
+					if (AllocInfo.IsAllocated() && AllocInfo.WeightmapTextureIndex < WeightmapTextures.Num())
 					{
 						UTexture2D* Weightmap = WeightmapTextures[AllocInfo.WeightmapTextureIndex];
 						check(Weightmap != nullptr);
@@ -265,7 +265,7 @@ public:
 
 		if (CreateUAV)
 		{
-			Flags |= TexCreate_RenderTargetable|TexCreate_UAV;
+			Flags |= TexCreate_RenderTargetable | TexCreate_UAV;
 		}
 
 		TextureRHI = RHICreateTexture2D(SizeX, SizeY, Format, NumMips, 1, Flags, CreateInfo);
@@ -282,7 +282,7 @@ public:
 		{
 			TextureUAV.SafeRelease();
 		}
-		
+
 		FTextureResource::ReleaseRHI();
 	}
 
@@ -746,9 +746,9 @@ public:
 		return true;
 	}
 
-	FLandscapeCopyTextureVS() 
+	FLandscapeCopyTextureVS()
 	{};
-	
+
 	FLandscapeCopyTextureVS(const ShaderMetaType::CompiledShaderInitializerType& Initializer)
 		: FGlobalShader(Initializer)
 	{}
@@ -1182,7 +1182,7 @@ class FLandscapeLayersCopyTexture_RenderThread
 public:
 	FLandscapeLayersCopyTexture_RenderThread(const FLandscapeLayersCopyTextureParams& InParams)
 		: Params(InParams)
-	{}	
+	{}
 
 	void Copy(FRHICommandListImmediate& InRHICmdList)
 	{
@@ -1204,7 +1204,7 @@ public:
 		Info.DestPosition.Y = Params.DestPosition.Y;
 		Info.SourceSliceIndex = Params.SourceArrayIndex;
 		Info.DestSliceIndex = Params.DestArrayIndex;
-		Info.SourceMipIndex = Params.SourceMip; 
+		Info.SourceMipIndex = Params.SourceMip;
 		Info.DestMipIndex = Params.DestMip;
 
 		check((Info.SourcePosition.X >= 0) && (Info.SourcePosition.Y >= 0) && (Info.DestPosition.X >= 0) && (Info.DestPosition.Y >= 0));
@@ -1263,7 +1263,7 @@ class FLandscapeLayersRender_RenderThread
 public:
 
 	FLandscapeLayersRender_RenderThread(const FString& InDebugName, UTextureRenderTarget2D* InWriteRenderTarget, const FIntPoint& InWriteRenderTargetSize, const FIntPoint& InReadRenderTargetSize, const FMatrix& InProjectionMatrix,
-										const ShaderDataType& InShaderParams, uint8 InCurrentMip, const TArray<FLandscapeLayersTriangle>& InTriangleList)
+		const ShaderDataType& InShaderParams, uint8 InCurrentMip, const TArray<FLandscapeLayersTriangle>& InTriangleList)
 		: RenderTargetResource(InWriteRenderTarget->GameThread_GetRenderTargetResource())
 		, WriteRenderTargetSize(InWriteRenderTargetSize)
 		, ReadRenderTargetSize(InReadRenderTargetSize)
@@ -1460,7 +1460,7 @@ bool ALandscape::ComputeLandscapeLayerBrushInfo(FTransform& OutLandscapeTransfor
 	{
 		return false;
 	}
-	
+
 	OutLandscapeTransform = Landscape->GetTransform();
 	FVector OffsetVector(LandscapeExtent.Min.X, LandscapeExtent.Min.Y, 0.f);
 	FVector Translation = OutLandscapeTransform.TransformFVector4(OffsetVector);
@@ -2064,7 +2064,7 @@ void ALandscape::ValidateProxyLayersWeightmapUsage() const
 
 void ALandscapeProxy::ValidateProxyLayersWeightmapUsage() const
 {
-	if (CVarLandscapeValidateProxyWeightmapUsages.GetValueOnGameThread() == 0)
+	if ((CVarLandscapeValidateProxyWeightmapUsages.GetValueOnGameThread() == 0) || bTemporarilyDisableWeightmapUsagesValidation)
 	{
 		return;
 	}
@@ -2216,7 +2216,7 @@ void ALandscape::CopyTexturePS(const FString& InSourceDebugName, FTextureResourc
 {
 	check(InSourceResource != nullptr);
 	check(InDestResource != nullptr);
-	
+
 	ENQUEUE_RENDER_COMMAND(LandscapeLayers_Cmd_CopyTexturePS)(
 		[InSourceDebugName, InSourceResource, InDestDebugName, InDestResource](FRHICommandListImmediate& RHICmdList)
 	{
@@ -2254,7 +2254,7 @@ void ALandscape::CopyTexturePS(const FString& InSourceDebugName, FTextureResourc
 }
 
 void ALandscape::DrawWeightmapComponentsToRenderTarget(const FString& InDebugName, const TArray<FIntPoint>& InSectionBaseList, const FVector2f& InScaleBias, TArray<FVector2f>* InScaleBiasPerSection, UTexture* InWeightmapRTRead, UTextureRenderTarget2D* InOptionalWeightmapRTRead2, UTextureRenderTarget2D* InWeightmapRTWrite,
-														ERTDrawingType InDrawType, bool InClearRTWrite, FLandscapeLayersWeightmapShaderParameters& InShaderParams, uint8 InMipRender) const
+	ERTDrawingType InDrawType, bool InClearRTWrite, FLandscapeLayersWeightmapShaderParameters& InShaderParams, uint8 InMipRender) const
 {
 	check(InWeightmapRTRead != nullptr);
 	check(InWeightmapRTWrite != nullptr);
@@ -2339,7 +2339,7 @@ void ALandscape::DrawWeightmapComponentsToRenderTarget(const FString& InDebugNam
 }
 
 void ALandscape::DrawWeightmapComponentsToRenderTarget(const FString& InDebugName, const TArray<ULandscapeComponent*>& InComponentsToDraw, const FIntPoint& InLandscapeBase, UTexture* InWeightmapRTRead, UTextureRenderTarget2D* InOptionalWeightmapRTRead2, UTextureRenderTarget2D* InWeightmapRTWrite,
-														ERTDrawingType InDrawType, bool InClearRTWrite, FLandscapeLayersWeightmapShaderParameters& InShaderParams, uint8 InMipRender) const
+	ERTDrawingType InDrawType, bool InClearRTWrite, FLandscapeLayersWeightmapShaderParameters& InShaderParams, uint8 InMipRender) const
 {
 
 	TArray<FIntPoint> SectionBaseList;
@@ -2388,7 +2388,7 @@ void ALandscape::DrawWeightmapComponentToRenderTargetMips(const TArray<FVector2f
 		if (WriteMipRT != nullptr)
 		{
 			DrawWeightmapComponentsToRenderTarget(FString::Printf(TEXT("LS Weight: %s = -> %s Mips %d"), *ReadMipRT->GetName(), *WriteMipRT->GetName(), CurrentMip),
-													SectionBaseToDraw, WeightmapScaleBias, nullptr, ReadMipRT, nullptr, WriteMipRT, ERTDrawingType::RTMips, InClearRTWrite, InShaderParams, CurrentMip);
+				SectionBaseToDraw, WeightmapScaleBias, nullptr, ReadMipRT, nullptr, WriteMipRT, ERTDrawingType::RTMips, InClearRTWrite, InShaderParams, CurrentMip);
 			++CurrentMip;
 		}
 
@@ -2420,7 +2420,7 @@ void ALandscape::DrawHeightmapComponentsToRenderTargetMips(const TArray<ULandsca
 		if (WriteMipRT != nullptr)
 		{
 			DrawHeightmapComponentsToRenderTarget(FString::Printf(TEXT("LS Height: %s = -> %s CombinedAtlasWithMips %d"), *ReadMipRT->GetName(), *WriteMipRT->GetName(), CurrentMip),
-												  InComponentsToDraw, InLandscapeBase, ReadMipRT, nullptr, WriteMipRT, ERTDrawingType::RTMips, InClearRTWrite, InShaderParams, CurrentMip);
+				InComponentsToDraw, InLandscapeBase, ReadMipRT, nullptr, WriteMipRT, ERTDrawingType::RTMips, InClearRTWrite, InShaderParams, CurrentMip);
 			++CurrentMip;
 		}
 
@@ -2429,7 +2429,7 @@ void ALandscape::DrawHeightmapComponentsToRenderTargetMips(const TArray<ULandsca
 }
 
 void ALandscape::DrawHeightmapComponentsToRenderTarget(const FString& InDebugName, const TArray<ULandscapeComponent*>& InComponentsToDraw, const FIntPoint& InLandscapeBase, UTexture* InHeightmapRTRead, UTextureRenderTarget2D* InOptionalHeightmapRTRead2, UTextureRenderTarget2D* InHeightmapRTWrite,
-													  ERTDrawingType InDrawType, bool InClearRTWrite, FLandscapeLayersHeightmapShaderParameters& InShaderParams, uint8 InMipRender) const
+	ERTDrawingType InDrawType, bool InClearRTWrite, FLandscapeLayersHeightmapShaderParameters& InShaderParams, uint8 InMipRender) const
 {
 	check(InHeightmapRTRead != nullptr);
 	check(InHeightmapRTWrite != nullptr);
@@ -2455,36 +2455,36 @@ void ALandscape::DrawHeightmapComponentsToRenderTarget(const FString& InDebugNam
 
 		switch (InDrawType)
 		{
-			case ERTDrawingType::RTAtlas:
-			{
-				GenerateLayersRenderQuadsAtlas(ComponentSectionBase, HeightmapScaleBias, SubsectionSizeQuads, HeightmapReadTextureSize, HeightmapWriteTextureSize, TriangleList);
-			} break;
+		case ERTDrawingType::RTAtlas:
+		{
+			GenerateLayersRenderQuadsAtlas(ComponentSectionBase, HeightmapScaleBias, SubsectionSizeQuads, HeightmapReadTextureSize, HeightmapWriteTextureSize, TriangleList);
+		} break;
 
-			case ERTDrawingType::RTAtlasToNonAtlas:
-			{
-				GenerateLayersRenderQuadsAtlasToNonAtlas(ComponentSectionBase, HeightmapScaleBias, SubsectionSizeQuads, HeightmapReadTextureSize, HeightmapWriteTextureSize, TriangleList);
-			} break;
+		case ERTDrawingType::RTAtlasToNonAtlas:
+		{
+			GenerateLayersRenderQuadsAtlasToNonAtlas(ComponentSectionBase, HeightmapScaleBias, SubsectionSizeQuads, HeightmapReadTextureSize, HeightmapWriteTextureSize, TriangleList);
+		} break;
 
-			case ERTDrawingType::RTNonAtlas:
-			{
-				GenerateLayersRenderQuadsNonAtlas(ComponentSectionBase, HeightmapScaleBias, SubsectionSizeQuads, HeightmapReadTextureSize, HeightmapWriteTextureSize, TriangleList);
-			} break;
+		case ERTDrawingType::RTNonAtlas:
+		{
+			GenerateLayersRenderQuadsNonAtlas(ComponentSectionBase, HeightmapScaleBias, SubsectionSizeQuads, HeightmapReadTextureSize, HeightmapWriteTextureSize, TriangleList);
+		} break;
 
-			case ERTDrawingType::RTNonAtlasToAtlas:
-			{
-				GenerateLayersRenderQuadsNonAtlasToAtlas(ComponentSectionBase, HeightmapScaleBias, SubsectionSizeQuads, HeightmapReadTextureSize, HeightmapWriteTextureSize, TriangleList);
-			} break;
+		case ERTDrawingType::RTNonAtlasToAtlas:
+		{
+			GenerateLayersRenderQuadsNonAtlasToAtlas(ComponentSectionBase, HeightmapScaleBias, SubsectionSizeQuads, HeightmapReadTextureSize, HeightmapWriteTextureSize, TriangleList);
+		} break;
 
-			case ERTDrawingType::RTMips:
-			{
-				GenerateLayersRenderQuadsMip(ComponentSectionBase, HeightmapScaleBias, SubsectionSizeQuads, HeightmapReadTextureSize, HeightmapWriteTextureSize, InMipRender, TriangleList);
-			} break;
+		case ERTDrawingType::RTMips:
+		{
+			GenerateLayersRenderQuadsMip(ComponentSectionBase, HeightmapScaleBias, SubsectionSizeQuads, HeightmapReadTextureSize, HeightmapWriteTextureSize, InMipRender, TriangleList);
+		} break;
 
-			default:
-			{
-				check(false);
-				return;
-			}
+		default:
+		{
+			check(false);
+			return;
+		}
 		}
 	}
 
@@ -2500,7 +2500,7 @@ void ALandscape::DrawHeightmapComponentsToRenderTarget(const FString& InDebugNam
 	}
 
 	FMatrix ProjectionMatrix = AdjustProjectionMatrixForRHI(FTranslationMatrix(FVector(0, 0, 0)) *
-															FMatrix(FPlane(1.0f / (FMath::Max<uint32>(HeightmapWriteTextureSize.X, 1.f) / 2.0f), 0.0, 0.0f, 0.0f), FPlane(0.0f, -1.0f / (FMath::Max<uint32>(HeightmapWriteTextureSize.Y, 1.f) / 2.0f), 0.0f, 0.0f), FPlane(0.0f, 0.0f, 1.0f, 0.0f), FPlane(-1.0f, 1.0f, 0.0f, 1.0f)));
+		FMatrix(FPlane(1.0f / (FMath::Max<uint32>(HeightmapWriteTextureSize.X, 1.f) / 2.0f), 0.0, 0.0f, 0.0f), FPlane(0.0f, -1.0f / (FMath::Max<uint32>(HeightmapWriteTextureSize.Y, 1.f) / 2.0f), 0.0f, 0.0f), FPlane(0.0f, 0.0f, 1.0f, 0.0f), FPlane(-1.0f, 1.0f, 0.0f, 1.0f)));
 
 	FLandscapeLayersHeightmapRender_RenderThread LayersRender(InDebugName, InHeightmapRTWrite, HeightmapWriteTextureSize, HeightmapReadTextureSize, ProjectionMatrix, InShaderParams, InMipRender, TriangleList);
 
@@ -2703,9 +2703,9 @@ void ALandscape::GenerateLayersRenderQuadsNonAtlas(const FIntPoint& InSectionBas
 	check(InReadSize.X == InWriteSize.X && InReadSize.Y == InWriteSize.Y);
 
 	int32 SubsectionSizeVerts = InSubSectionSizeQuad + 1;
-	
+
 	FVector2D UVSize((float)SubsectionSizeVerts / (float)InReadSize.X, (float)SubsectionSizeVerts / (float)InReadSize.Y);
-	
+
 	for (int8 SubY = 0; SubY < NumSubsections; ++SubY)
 	{
 		for (int8 SubX = 0; SubX < NumSubsections; ++SubX)
@@ -3181,7 +3181,7 @@ int32 ALandscape::RegenerateLayersHeightmaps(FTextureToComponentHelper const& Ma
 	// Lazily create CPU read back objects as required
 	if (HeightmapUpdateModes)
 	{
-		for(ULandscapeComponent* Component : InLandscapeComponentsToRender)
+		for (ULandscapeComponent* Component : InLandscapeComponentsToRender)
 		{
 			UTexture2D* ComponentHeightmap = Component->GetHeightmap(false);
 			ALandscapeProxy* Proxy = Component->GetLandscapeProxy();
@@ -3238,13 +3238,13 @@ int32 ALandscape::RegenerateLayersHeightmaps(FTextureToComponentHelper const& Ma
 		auto GetUniqueHeightmaps = [&](const TArray<ULandscapeComponent*>& InLandscapeComponents, TArray<FHeightmapCopyInfo>& OutHeightmaps, const FIntPoint& LandscapeBaseQuads, const FGuid& InLayerGuid = FGuid())
 		{
 			FScopedSetLandscapeEditingLayer Scope(this, InLayerGuid);
-			
+
 			const int32 ComponentSizeQuad = SubsectionSizeQuads * NumSubsections;
 			const int32 ComponentSizeVerts = (SubsectionSizeQuads + 1) * NumSubsections;
 			for (ULandscapeComponent* Component : InLandscapeComponents)
 			{
 				UTexture2D* ComponentHeightmap = Component->GetHeightmap(true);
-		
+
 				int32 Index = OutHeightmaps.IndexOfByPredicate([=](const FHeightmapCopyInfo& LayerHeightmap) { return LayerHeightmap.Texture == ComponentHeightmap; });
 
 				FIntPoint ComponentSectionBase = Component->GetSectionBase() - LandscapeBaseQuads;
@@ -3311,13 +3311,13 @@ int32 ALandscape::RegenerateLayersHeightmaps(FTextureToComponentHelper const& Ma
 
 			// NOTE: From this point on, we always work in non atlas, we'll convert back at the end to atlas only
 			DrawHeightmapComponentsToRenderTarget(FString::Printf(TEXT("LS Height: %s += -> NonAtlas %s"), *Layer.Name.ToString(), *LandscapeScratchRT1->GetName(), *LandscapeScratchRT2->GetName()),
-												  InLandscapeComponentsToRender, LandscapeExtent.Min, LandscapeScratchRT1, nullptr, LandscapeScratchRT2, ERTDrawingType::RTAtlasToNonAtlas, true, ShaderParams);
+				InLandscapeComponentsToRender, LandscapeExtent.Min, LandscapeScratchRT1, nullptr, LandscapeScratchRT2, ERTDrawingType::RTAtlasToNonAtlas, true, ShaderParams);
 
 			ShaderParams.ApplyLayerModifiers = true;
 
 			// Combine Current layer with current result
 			DrawHeightmapComponentsToRenderTarget(FString::Printf(TEXT("LS Height: %s += -> CombinedNonAtlas %s"), *Layer.Name.ToString(), *LandscapeScratchRT2->GetName(), *CombinedHeightmapNonAtlasRT->GetName()),
-												InLandscapeComponentsToRender, LandscapeExtent.Min, LandscapeScratchRT2, FirstLayer ? nullptr : LandscapeScratchRT3, CombinedHeightmapNonAtlasRT, ERTDrawingType::RTNonAtlas, FirstLayer, ShaderParams);
+				InLandscapeComponentsToRender, LandscapeExtent.Min, LandscapeScratchRT2, FirstLayer ? nullptr : LandscapeScratchRT3, CombinedHeightmapNonAtlasRT, ERTDrawingType::RTNonAtlas, FirstLayer, ShaderParams);
 
 			ShaderParams.ApplyLayerModifiers = false;
 
@@ -3365,12 +3365,12 @@ int32 ALandscape::RegenerateLayersHeightmaps(FTextureToComponentHelper const& Ma
 		}
 
 		DrawHeightmapComponentsToRenderTarget(FString::Printf(TEXT("LS Height: %s = -> CombinedNonAtlasNormals : %s"), *CombinedHeightmapNonAtlasRT->GetName(), *LandscapeScratchRT1->GetName()),
-											  InLandscapeComponentsToRender, LandscapeExtent.Min, CombinedHeightmapNonAtlasRT, nullptr, LandscapeScratchRT1, ERTDrawingType::RTNonAtlas, true, ShaderParams);
+			InLandscapeComponentsToRender, LandscapeExtent.Min, CombinedHeightmapNonAtlasRT, nullptr, LandscapeScratchRT1, ERTDrawingType::RTNonAtlas, true, ShaderParams);
 
 		ShaderParams.GenerateNormals = false;
 
 		DrawHeightmapComponentsToRenderTarget(FString::Printf(TEXT("LS Height: %s = -> CombinedAtlasFinal : %s"), *LandscapeScratchRT1->GetName(), *CombinedHeightmapAtlasRT->GetName()),
-											  InLandscapeComponentsToRender, LandscapeExtent.Min, LandscapeScratchRT1, nullptr, CombinedHeightmapAtlasRT, ERTDrawingType::RTNonAtlasToAtlas, true, ShaderParams);
+			InLandscapeComponentsToRender, LandscapeExtent.Min, LandscapeScratchRT1, nullptr, CombinedHeightmapAtlasRT, ERTDrawingType::RTNonAtlasToAtlas, true, ShaderParams);
 
 		DrawHeightmapComponentsToRenderTargetMips(InLandscapeComponentsToRender, LandscapeExtent.Min, CombinedHeightmapAtlasRT, true, ShaderParams);
 
@@ -3440,7 +3440,7 @@ void ALandscape::UpdateForChangedHeightmaps(ULandscapeComponent* InComponent, co
 	//  - If ELandscapeComponentUpdateFlag::Component_Update_Heightmap_Collision is passed, it will be done immediately
 	//  - If not, at least the component's collision data will still get updated eventually, when the flag is finally passed :
 	if (InReadbackResult.bModified)
-{
+	{
 		InComponent->SetPendingCollisionDataUpdate(true);
 	}
 
@@ -3454,8 +3454,8 @@ void ALandscape::UpdateForChangedHeightmaps(ULandscapeComponent* InComponent, co
 			InComponent->UpdateCachedBounds();
 			InComponent->UpdateComponentToWorld();
 
-		// Avoid updating height field if we are going to recreate collision in this update
-		bool bUpdateHeightfieldRegion = !IsUpdateFlagEnabledForModes(ELandscapeComponentUpdateFlag::Component_Update_Recreate_Collision, HeightUpdateMode);
+			// Avoid updating height field if we are going to recreate collision in this update
+			bool bUpdateHeightfieldRegion = !IsUpdateFlagEnabledForModes(ELandscapeComponentUpdateFlag::Component_Update_Recreate_Collision, HeightUpdateMode);
 			InComponent->UpdateCollisionData(bUpdateHeightfieldRegion);
 			InComponent->SetPendingCollisionDataUpdate(false);
 		}
@@ -3532,7 +3532,7 @@ void ALandscape::ClearDirtyData(ULandscapeComponent* InLandscapeComponent)
 	FMemory::Memzero(DirtyData.Get(), DirtyDataSize);
 	LandscapeEdit.SetDirtyData(X1, Y1, X2, Y2, DirtyData.Get(), 0);
 }
-	
+
 void ALandscape::UpdateWeightDirtyData(ULandscapeComponent* InLandscapeComponent, UTexture2D const* InWeightmap, FColor const* InOldData, FColor const* InNewData, uint8 InChannel)
 {
 	check(InOldData && InNewData);
@@ -3549,7 +3549,7 @@ void ALandscape::UpdateWeightDirtyData(ULandscapeComponent* InLandscapeComponent
 	const int32 SizeV = InWeightmap->Source.GetSizeY();
 	check(DirtyDataSize == SizeU * SizeV);
 
-	const uint8 DirtyWeight = 1<<1;
+	const uint8 DirtyWeight = 1 << 1;
 	LandscapeEdit.GetDirtyData(X1, Y1, X2, Y2, DirtyData.Get(), 0);
 
 	for (int32 Index = 0; Index < DirtyDataSize; ++Index)
@@ -3615,16 +3615,16 @@ void ALandscape::UpdateHeightDirtyData(ULandscapeComponent* InLandscapeComponent
 	const int32 X2 = X1 + ComponentSizeQuads;
 	const int32 Y1 = InLandscapeComponent->GetSectionBase().Y;
 	const int32 Y2 = Y1 + ComponentSizeQuads;
-	const int32 ComponentWidth = (SubsectionSizeQuads + 1)*NumSubsections;
+	const int32 ComponentWidth = (SubsectionSizeQuads + 1) * NumSubsections;
 	const int32 DirtyDataSize = ComponentWidth * ComponentWidth;
 	TUniquePtr<uint8[]> DirtyData = MakeUnique<uint8[]>(DirtyDataSize);
 	const int32 SizeU = InHeightmap->Source.GetSizeX();
 	const int32 SizeV = InHeightmap->Source.GetSizeY();
 	const int32 HeightmapOffsetX = InLandscapeComponent->HeightmapScaleBias.Z * SizeU;
 	const int32 HeightmapOffsetY = InLandscapeComponent->HeightmapScaleBias.W * SizeV;
-	const uint8 DirtyHeight = 1<<0;
+	const uint8 DirtyHeight = 1 << 0;
 	LandscapeEdit.GetDirtyData(X1, Y1, X2, Y2, DirtyData.Get(), 0);
-	
+
 	for (int32 X = 0; X < ComponentWidth; ++X)
 	{
 		for (int32 Y = 0; Y < ComponentWidth; ++Y)
@@ -3632,7 +3632,7 @@ void ALandscape::UpdateHeightDirtyData(ULandscapeComponent* InLandscapeComponent
 			int32 TexX = HeightmapOffsetX + X;
 			int32 TexY = HeightmapOffsetY + Y;
 			int32 TexIndex = TexX + TexY * SizeU;
-			check(TexIndex < SizeU * SizeV);
+			check(TexIndex < SizeU* SizeV);
 			if (InOldData[TexIndex] != InNewData[TexIndex])
 			{
 				DirtyData[X + Y * ComponentWidth] |= DirtyHeight;
@@ -3764,7 +3764,7 @@ bool ALandscape::ResolveLayersTexture(
 }
 
 void ALandscape::PrepareComponentDataToExtractMaterialLayersCS(const TArray<ULandscapeComponent*>& InLandscapeComponents, const FLandscapeLayer& InLayer, int32 InCurrentWeightmapToProcessIndex, const FIntPoint& InLandscapeBase, FLandscapeTexture2DResource* InOutTextureData,
-																TArray<FLandscapeLayerWeightmapExtractMaterialLayersComponentData>& OutComponentData, TMap<ULandscapeLayerInfoObject*, int32>& OutLayerInfoObjects)
+	TArray<FLandscapeLayerWeightmapExtractMaterialLayersComponentData>& OutComponentData, TMap<ULandscapeLayerInfoObject*, int32>& OutLayerInfoObjects)
 {
 	ULandscapeInfo* Info = GetLandscapeInfo();
 
@@ -3850,7 +3850,7 @@ void ALandscape::PrepareComponentDataToExtractMaterialLayersCS(const TArray<ULan
 			}
 		}
 	}
-	
+
 	ExecuteCopyLayersTexture(MoveTemp(DeferredCopyTextures));
 }
 
@@ -3992,7 +3992,7 @@ void ALandscape::ReallocateLayersWeightmaps(const TArray<ULandscapeComponent*>&
 			Usage->ClearUsage(Component);
 		}
 	}
-		
+
 	// Build a map of all the allocation per components
 	TMap<ULandscapeComponent*, TArray<ULandscapeLayerInfoObject*>> LayerAllocsPerComponent;
 
@@ -4111,7 +4111,7 @@ void ALandscape::ReallocateLayersWeightmaps(const TArray<ULandscapeComponent*>&
 				if (FLandscapeEditLayerReadback** CPUReadbackToDelete = Proxy->WeightmapsCPUReadback.Find(OldWeightmapTexture))
 				{
 					check(*CPUReadbackToDelete);
-					delete * CPUReadbackToDelete;
+					delete* CPUReadbackToDelete;
 					Proxy->WeightmapsCPUReadback.Remove(OldWeightmapTexture);
 				}
 			}
@@ -4191,7 +4191,7 @@ int32 ALandscape::RegenerateLayersWeightmaps(FTextureToComponentHelper const& Ma
 	const int32 WeightmapUpdateModes = LayerContentUpdateModes & AllWeightmapUpdateModes;
 	const bool bSkipBrush = CVarLandscapeLayerBrushOptim.GetValueOnAnyThread() == 1 && ((WeightmapUpdateModes & AllWeightmapUpdateModes) == ELandscapeLayerUpdateMode::Update_Weightmap_Editing);
 	const bool bForceRender = CVarForceLayersUpdate.GetValueOnAnyThread() != 0;
-	
+
 	ULandscapeInfo* Info = GetLandscapeInfo();
 
 	if (WeightmapUpdateModes == 0 && !bForceRender)
@@ -4203,7 +4203,7 @@ int32 ALandscape::RegenerateLayersWeightmaps(FTextureToComponentHelper const& Ma
 	{
 		return WeightmapUpdateModes;
 	}
-			
+
 	TArray<ULandscapeComponent*> ComponentThatNeedMaterialRebuild;
 	TArray<ULandscapeLayerInfoObject*> BrushRequiredAllocations;
 	int32 LayerCount = Info->Layers.Num() + 1; // due to visibility being stored at 0
@@ -4218,7 +4218,7 @@ int32 ALandscape::RegenerateLayersWeightmaps(FTextureToComponentHelper const& Ma
 		{
 			return 0;
 		}
-		
+
 		check(WeightmapRTList.Num() > 0);
 
 		UTextureRenderTarget2D* LandscapeScratchRT1 = WeightmapRTList[(int32)EWeightmapRTType::WeightmapRT_Scratch1];
@@ -4264,7 +4264,7 @@ int32 ALandscape::RegenerateLayersWeightmaps(FTextureToComponentHelper const& Ma
 				for (int32 LayerInfoSettingsIndex = 0; LayerInfoSettingsIndex < Info->Layers.Num(); ++LayerInfoSettingsIndex)
 				{
 					const FLandscapeInfoLayerSettings& InfoLayerSettings = Info->Layers[LayerInfoSettingsIndex];
-					
+
 					// It is possible that no layer info has been assigned so that InfoLayerSettings.LayerInfoObj == nullptr. In that case we don't consider it the layer here.
 					if (InfoLayerSettings.LayerInfoObj != nullptr)
 					{
@@ -4332,7 +4332,7 @@ int32 ALandscape::RegenerateLayersWeightmaps(FTextureToComponentHelper const& Ma
 				if (ComponentsData.Num() > 0)
 				{
 					PrintLayersDebugTextureResource(FString::Printf(TEXT("LS Weight: %s WeightmapScratchTexture %s"), *Layer.Name.ToString(), TEXT("WeightmapScratchTextureResource")), WeightmapScratchExtractLayerTextureResource, 0, false);
-										
+
 					FLandscapeLayerWeightmapExtractMaterialLayersComputeShaderParameters CSExtractLayersShaderParams;
 					CSExtractLayersShaderParams.AtlasWeightmapsPerLayer = CurrentLayersWeightmapAllMaterialLayersResource;
 					CSExtractLayersShaderParams.ComponentWeightmapResource = WeightmapScratchExtractLayerTextureResource;
@@ -4379,7 +4379,7 @@ int32 ALandscape::RegenerateLayersWeightmaps(FTextureToComponentHelper const& Ma
 					PSShaderParams.LayerAlpha = LayerInfoObj == ALandscapeProxy::VisibilityLayer ? 1.0f : Layer.WeightmapAlpha; // visibility can't be affected by weight
 
 					DrawWeightmapComponentsToRenderTarget(FString::Printf(TEXT("LS Weight: %s Paint: %s += -> %s"), *Layer.Name.ToString(), *LayerInfoObj->LayerName.ToString(), *LandscapeScratchRT1->GetName(), *LandscapeScratchRT2->GetName()),
-														  InLandscapeComponentsToRender, LandscapeExtent.Min, LandscapeScratchRT1, nullptr, LandscapeScratchRT2, ERTDrawingType::RTAtlas, true, PSShaderParams, 0);
+						InLandscapeComponentsToRender, LandscapeExtent.Min, LandscapeScratchRT1, nullptr, LandscapeScratchRT2, ERTDrawingType::RTAtlas, true, PSShaderParams, 0);
 
 					PSShaderParams.ApplyLayerModifiers = false;
 
@@ -4410,7 +4410,7 @@ int32 ALandscape::RegenerateLayersWeightmaps(FTextureToComponentHelper const& Ma
 					}
 
 					DrawWeightmapComponentsToRenderTarget(FString::Printf(TEXT("LS Weight: %s PaintLayer: %s, %s += -> Combined %s"), *Layer.Name.ToString(), *LayerInfoObj->LayerName.ToString(), *LandscapeScratchRT2->GetName(), *LandscapeScratchRT3->GetName()),
-														  InLandscapeComponentsToRender, LandscapeExtent.Min, LandscapeScratchRT2, bFirstLayer ? nullptr : LandscapeScratchRT1, LandscapeScratchRT3, ERTDrawingType::RTAtlasToNonAtlas, true, PSShaderParams, 0);
+						InLandscapeComponentsToRender, LandscapeExtent.Min, LandscapeScratchRT2, bFirstLayer ? nullptr : LandscapeScratchRT1, LandscapeScratchRT3, ERTDrawingType::RTAtlasToNonAtlas, true, PSShaderParams, 0);
 
 					PSShaderParams.OutputAsSubstractive = false;
 
@@ -4506,7 +4506,7 @@ int32 ALandscape::RegenerateLayersWeightmaps(FTextureToComponentHelper const& Ma
 					}
 				}
 			}
-			
+
 			int8 CurrentWeightmapToProcessIndex = 0;
 			bool HasFoundWeightmapToProcess = true; // try processing at least once
 
@@ -4548,9 +4548,9 @@ int32 ALandscape::RegenerateLayersWeightmaps(FTextureToComponentHelper const& Ma
 
 					for (int32 i = 0; i < PackLayersComponentsData.Num(); ++i)
 					{
-						check(ComponentY+ComponentSize <= WeightmapScratchPackLayerTextureResource->GetSizeY()); // This should never happen as it would be a bug in the algo
+						check(ComponentY + ComponentSize <= WeightmapScratchPackLayerTextureResource->GetSizeY()); // This should never happen as it would be a bug in the algo
 
-						if (ComponentX+ComponentSize > WeightmapScratchPackLayerTextureResource->GetSizeX())
+						if (ComponentX + ComponentSize > WeightmapScratchPackLayerTextureResource->GetSizeX())
 						{
 							ComponentY += ComponentSize;
 							ComponentX = 0;
@@ -4586,7 +4586,7 @@ int32 ALandscape::RegenerateLayersWeightmaps(FTextureToComponentHelper const& Ma
 
 					SourceDebugName = TEXT("WeightmapScratchTexture");
 					DestDebugName = CurrentRT->GetName();
-					
+
 					CopyTexturePS(SourceDebugName, WeightmapScratchPackLayerTextureResource, DestDebugName, CurrentRT->GameThread_GetRenderTargetResource());
 					DrawWeightmapComponentToRenderTargetMips(WeightmapTextureOutputOffset, CurrentRT, true, PSShaderParams);
 
@@ -4816,9 +4816,7 @@ void ALandscape::UpdateLayersMaterialInstances(const TArray<ULandscapeComponent*
 				for (int32 AllocIdx = 0; AllocIdx < WeightmapBaseLayerAllocation.Num(); AllocIdx++)
 				{
 					const FWeightmapLayerAllocationInfo& Allocation = WeightmapBaseLayerAllocation[AllocIdx];
-
-					FName LayerName = Allocation.LayerInfo == ALandscapeProxy::VisibilityLayer ? UMaterialExpressionLandscapeVisibilityMask::ParameterName : Allocation.LayerInfo ? Allocation.LayerInfo->LayerName : NAME_None;
-					MaterialInstance->SetVectorParameterValueEditorOnly(FName(*FString::Printf(TEXT("LayerMask_%s"), *LayerName.ToString())), Masks[Allocation.WeightmapTextureChannel]);
+					MaterialInstance->SetVectorParameterValueEditorOnly(FName(*FString::Printf(TEXT("LayerMask_%s"), *Allocation.GetLayerName().ToString())), Masks[Allocation.WeightmapTextureChannel]);
 				}
 
 				// Set the weightmaps
@@ -4905,7 +4903,7 @@ void ALandscape::ResolveLayersWeightmapTexture(
 			}
 		}
 	}
-		
+
 	// Weightmaps shouldn't invalidate lighting
 	const bool bInvalidateLightingCache = false;
 	InvalidateGeneratedComponentData(ChangedComponents, bInvalidateLightingCache);
@@ -5005,7 +5003,7 @@ bool ALandscape::IsUpdateFlagEnabledForModes(ELandscapeComponentUpdateFlag InFla
 			return true;
 		}
 	}
-		
+
 	if (InUpdateModes & ELandscapeLayerUpdateMode::Update_Heightmap_Editing)
 	{
 		const uint32 HeightmapEditingFlags = ELandscapeComponentUpdateFlag::Component_Update_Heightmap_Collision | ELandscapeComponentUpdateFlag::Component_Update_Client_Editing;
@@ -5190,10 +5188,10 @@ void ALandscape::GetLandscapeComponentNeighborsToRender(ULandscapeComponent* Lan
 {
 	ULandscapeInfo* LandscapeInfo = GetLandscapeInfo();
 	FIntPoint ComponentKey = LandscapeComponent->GetSectionBase() / ComponentSizeQuads;
-	
-	for (int32 IndexX = ComponentKey.X-1; IndexX <= ComponentKey.X+1; ++IndexX)
+
+	for (int32 IndexX = ComponentKey.X - 1; IndexX <= ComponentKey.X + 1; ++IndexX)
 	{
-		for (int32 IndexY = ComponentKey.Y-1; IndexY <= ComponentKey.Y+1; ++IndexY)
+		for (int32 IndexY = ComponentKey.Y - 1; IndexY <= ComponentKey.Y + 1; ++IndexY)
 		{
 			ULandscapeComponent* Result = LandscapeInfo->XYtoComponentMap.FindRef(FIntPoint(IndexX, IndexY));
 			if (Result && Result != LandscapeComponent)
@@ -5330,7 +5328,7 @@ void ALandscape::UpdateLayersContent(bool bInWaitForStreaming, bool bInSkipMonit
 		}
 		return;
 	}
-	
+
 	// Gather mappings between heightmaps/weightmaps and components
 	FTextureToComponentHelper MapHelper(*LandscapeInfo);
 
@@ -5342,7 +5340,7 @@ void ALandscape::UpdateLayersContent(bool bInWaitForStreaming, bool bInSkipMonit
 		// These flags might look like they're being mixed up but they're not!
 		const bool bDoIntermediateRender = false; // bIntermediateRender flag is for the work queued up this frame not the delayed resolves
 		const bool bDoFlushRender = bIntermediateRender; // Flush before we do an intermediate render later in this frame
-		
+
 		TMap<ULandscapeComponent*, FLandscapeEditLayerReadbackResult> ModifiedComponentInfo;
 		ResolveLayersHeightmapTexture(MapHelper, MapHelper.Heightmaps, bDoIntermediateRender, bDoFlushRender, ModifiedComponentInfo);
 		ResolveLayersWeightmapTexture(MapHelper, MapHelper.Weightmaps, bDoIntermediateRender, bDoFlushRender, ModifiedComponentInfo);
@@ -5376,7 +5374,7 @@ void ALandscape::UpdateLayersContent(bool bInWaitForStreaming, bool bInSkipMonit
 		if (!bPartialUpdate || Component->GetLayerUpdateFlagPerMode() != 0)
 		{
 			AllLandscapeComponents.Add(Component);
-		
+
 			if (bPartialUpdate)
 			{
 				// Gather Neighbors (Neighbors need to be Rendered but not resolved so that the resolved Components have valid normals on edges)
@@ -5422,7 +5420,7 @@ void ALandscape::UpdateLayersContent(bool bInWaitForStreaming, bool bInSkipMonit
 	LandscapeComponentsHeightmapsToRender.Append(AllLandscapeComponents);
 	LandscapeComponentsWeightmapsToResolve.Append(AllLandscapeComponents);
 	LandscapeComponentsWeightmapsToRender.Append(AllLandscapeComponents);
-		
+
 	if (bPartialUpdate)
 	{
 		for (ULandscapeComponent* Component : SkippedComponents)
@@ -5827,7 +5825,7 @@ bool ALandscapeProxy::RemoveObsoleteLayers(const TSet<FGuid>& InExistingLayers)
 
 			FMessageLog("MapCheck").Warning()
 				->AddToken(FUObjectToken::Create(this))
-				->AddToken(FTextToken::Create(FText::Format(LOCTEXT("MapCheck_Message_LandscapeProxyObsoleteLayer","Layer '{LayerGuid}' was removed from LandscapeProxy because it doesn't match any of the LandscapeActor Layers. Please resave '{ProxyPackage}'."), Arguments)))
+				->AddToken(FTextToken::Create(FText::Format(LOCTEXT("MapCheck_Message_LandscapeProxyObsoleteLayer", "Layer '{LayerGuid}' was removed from LandscapeProxy because it doesn't match any of the LandscapeActor Layers. Please resave '{ProxyPackage}'."), Arguments)))
 				->AddToken(FMapErrorToken::Create(FMapErrors::LandscapeComponentPostLoad_Warning));
 
 			DeleteLayer(LayerGuid);
@@ -5856,7 +5854,7 @@ bool ALandscapeProxy::AddLayer(const FGuid& InLayerGuid)
 			Component->AddLayerData(InLayerGuid, FLandscapeLayerComponentData());
 			bModified = true;
 		}
-	}	
+	}
 
 	UpdateCachedHasLayersContent();
 
@@ -5906,7 +5904,7 @@ void ALandscapeProxy::InitializeLayerWithEmptyContent(const FGuid& InLayerGuid)
 	{
 		return;
 	}
-		
+
 	// Build a mapping between each Heightmaps and Component in them
 	TMap<UTexture2D*, TArray<ULandscapeComponent*>> ComponentsPerHeightmaps;
 
@@ -5916,7 +5914,7 @@ void ALandscapeProxy::InitializeLayerWithEmptyContent(const FGuid& InLayerGuid)
 		TArray<ULandscapeComponent*>& ComponentList = ComponentsPerHeightmaps.FindOrAdd(ComponentHeightmapTexture);
 		ComponentList.Add(Component);
 	}
-		
+
 	// Init layers with valid "empty" data
 	TMap<UTexture2D*, UTexture2D*> CreatedHeightmapTextures; // < Final layer texture, New created texture for layer
 
@@ -6034,7 +6032,7 @@ float ALandscape::GetLayerAlpha(int32 InLayerIndex, bool bInHeightmap) const
 
 float ALandscape::GetClampedLayerAlpha(float InAlpha, bool bInHeightmap) const
 {
-	float AlphaClamped = FMath::Clamp<float>(InAlpha, bInHeightmap  ? -1.f : 0.f, 1.f);
+	float AlphaClamped = FMath::Clamp<float>(InAlpha, bInHeightmap ? -1.f : 0.f, 1.f);
 	return AlphaClamped;
 }
 
@@ -6212,7 +6210,7 @@ void ALandscape::CollapseLayer(int32 InLayerIndex)
 			LandscapeComponent->CopyFinalLayerIntoEditingLayer(DataInterface, ProcessedHeightmaps);
 		});
 	}
-	
+
 	TArray<ALandscapeBlueprintBrushBase*> BrushesToMove;
 	for (int32 i = 0; i < LandscapeLayers[InLayerIndex].Brushes.Num(); ++i)
 	{
@@ -6233,7 +6231,7 @@ void ALandscape::CollapseLayer(int32 InLayerIndex)
 	}
 
 	DeleteLayer(InLayerIndex);
-	
+
 	RequestLayersContentUpdateForceAll();
 }
 
@@ -6253,7 +6251,7 @@ void ALandscape::GetUsedPaintLayers(const FGuid& InLayerGuid, TArray<ULandscapeL
 	{
 		return;
 	}
-	
+
 	LandscapeInfo->GetUsedPaintLayers(InLayerGuid, OutUsedLayerInfos);
 }
 
@@ -6463,15 +6461,15 @@ uint32 ULandscapeComponent::ComputeLayerHash() const
 {
 	UTexture2D* Heightmap = GetHeightmap(true);
 	const uint8* MipData = Heightmap->Source.LockMipReadOnly(0);
-	uint32 Hash = FCrc::MemCrc32(MipData, Heightmap->GetSizeX()*Heightmap->GetSizeY() * sizeof(FColor));
+	uint32 Hash = FCrc::MemCrc32(MipData, Heightmap->GetSizeX() * Heightmap->GetSizeY() * sizeof(FColor));
 	Heightmap->Source.UnlockMip(0);
 
 	// Copy to sort
 	const TArray<UTexture2D*>& Weightmaps = GetWeightmapTextures(true);
 	TArray<FWeightmapLayerAllocationInfo> AllocationInfos = GetWeightmapLayerAllocations(true);
-	
-    // Sort allocations infos by LayerInfo Path so the Weightmaps hahses get ordered properly
-    AllocationInfos.Sort([](const FWeightmapLayerAllocationInfo& A, const FWeightmapLayerAllocationInfo& B)
+
+	// Sort allocations infos by LayerInfo Path so the Weightmaps hahses get ordered properly
+	AllocationInfos.Sort([](const FWeightmapLayerAllocationInfo& A, const FWeightmapLayerAllocationInfo& B)
 	{
 		FString PathA(A.LayerInfo ? A.LayerInfo->GetPathName() : FString());
 		FString PathB(B.LayerInfo ? B.LayerInfo->GetPathName() : FString());
@@ -6483,23 +6481,23 @@ uint32 ULandscapeComponent::ComputeLayerHash() const
 	{
 		if (AllocationInfo.IsAllocated())
 		{
-            // Compute hash of actual data of the texture that is owned by the component (per Texture Channel)
+			// Compute hash of actual data of the texture that is owned by the component (per Texture Channel)
 			UTexture2D* Weightmap = Weightmaps[AllocationInfo.WeightmapTextureIndex];
 			MipData = Weightmap->Source.LockMipReadOnly(0) + ChannelOffsets[AllocationInfo.WeightmapTextureChannel];
 			TArray<uint8> ChannelData;
-			ChannelData.AddDefaulted(Weightmap->GetSizeX()*Weightmap->GetSizeY());
-			int32 TexSize = (SubsectionSizeQuads + 1)*NumSubsections;
+			ChannelData.AddDefaulted(Weightmap->GetSizeX() * Weightmap->GetSizeY());
+			int32 TexSize = (SubsectionSizeQuads + 1) * NumSubsections;
 			for (int32 TexY = 0; TexY < TexSize; TexY++)
 			{
 				for (int32 TexX = 0; TexX < TexSize; TexX++)
 				{
 					const int32 Index = (TexX + TexY * TexSize);
-				
+
 					ChannelData.GetData()[Index] = MipData[4 * Index];
 				}
 			}
 
-			Hash = FCrc::MemCrc32(ChannelData.GetData(), Weightmap->GetSizeX()*Weightmap->GetSizeY(), Hash);
+			Hash = FCrc::MemCrc32(ChannelData.GetData(), Weightmap->GetSizeX() * Weightmap->GetSizeY(), Hash);
 			Weightmap->Source.UnlockMip(0);
 		}
 	}
@@ -6527,11 +6525,11 @@ void ALandscape::UpdateLandscapeSplines(const FGuid& InTargetLayer, bool bInUpda
 			{
 				return;
 			}
-						
+
 			TMap<ULandscapeComponent*, uint32> PreviousHashes;
 			{
 				FLandscapeEditDataInterface LandscapeEdit(LandscapeInfo);
-				
+
 				LandscapeInfo->ForAllLandscapeComponents([&](ULandscapeComponent* Component)
 				{
 					// Was never computed
@@ -6556,8 +6554,8 @@ void ALandscape::UpdateLandscapeSplines(const FGuid& InTargetLayer, bool bInUpda
 			bInUpdateOnlySelected = false;
 
 			// Apply splines without clearing up weightmap allocations
-			LandscapeInfo->ApplySplines(bInUpdateOnlySelected, ModifiedComponent, bMarkPackageDirty);		
-		
+			LandscapeInfo->ApplySplines(bInUpdateOnlySelected, ModifiedComponent, bMarkPackageDirty);
+
 			for (const TPair<ULandscapeComponent*, uint32>& Pair : PreviousHashes)
 			{
 				if (LandscapeSplinesAffectedComponents.Contains(Pair.Key))
@@ -6619,7 +6617,7 @@ bool ALandscape::IsEditingLayerReservedForSplines() const
 void ALandscape::SetEditingLayer(const FGuid& InLayerGuid)
 {
 	ensure(CanHaveLayersContent());
-	
+
 	ULandscapeInfo* LandscapeInfo = GetLandscapeInfo();
 	if (!LandscapeInfo)
 	{
@@ -6933,12 +6931,12 @@ void ALandscape::OnBlueprintBrushChanged()
 void ALandscape::OnLayerInfoSplineFalloffModulationChanged(ULandscapeLayerInfoObject* InLayerInfo)
 {
 	ULandscapeInfo* LandscapeInfo = GetLandscapeInfo();
-	
+
 	if (!LandscapeInfo)
 	{
 		return;
 	}
-	
+
 	ALandscape* Landscape = LandscapeInfo->LandscapeActor.Get();
 	if (!Landscape || !Landscape->HasLayersContent())
 	{
@@ -7065,7 +7063,7 @@ bool FLandscapeLayerBrush::Initialize(const FIntRect& InLandscapeExtent, UTextur
 				LandscapeTransform = NewLandscapeTransform;
 				LandscapeRenderTargetSize = NewLandscapeRenderTargetSize;
 				LandscapeSize = NewLandscapeSize;
-				
+
 				TGuardValue<bool> AutoRestore(GAllowActorScriptExecutionInEditor, true);
 				BlueprintBrush->Initialize(LandscapeTransform, LandscapeSize, LandscapeRenderTargetSize);
 			}
