Use BeginFrameRT/EndFrameRT instead of BeginFrame/EndFrame + ENQUEUE_RENDER_COMMAND

diff --git a/Engine/Source/Runtime/Landscape/Private/LandscapeRender.cpp b/Engine/Source/Runtime/Landscape/Private/LandscapeRender.cpp
--- a/Engine/Source/Runtime/Landscape/Private/LandscapeRender.cpp
+++ b/Engine/Source/Runtime/Landscape/Private/LandscapeRender.cpp
@@ -1088,27 +1088,24 @@ void FLandscapeRenderSystem::EndFrame()
 //
 FLandscapeSceneViewExtension::FLandscapeSceneViewExtension(const FAutoRegister& AutoReg) : FSceneViewExtensionBase(AutoReg)
 {
-	FCoreDelegates::OnBeginFrame.AddRaw(this, &FLandscapeSceneViewExtension::BeginFrame);
-	FCoreDelegates::OnEndFrame.AddRaw(this, &FLandscapeSceneViewExtension::EndFrame);
+	FCoreDelegates::OnBeginFrameRT.AddRaw(this, &FLandscapeSceneViewExtension::BeginFrame_RenderThread);
+	FCoreDelegates::OnEndFrameRT.AddRaw(this, &FLandscapeSceneViewExtension::EndFrame_RenderThread);
 }
 
 FLandscapeSceneViewExtension::~FLandscapeSceneViewExtension()
 {
-	FCoreDelegates::OnBeginFrame.RemoveAll(this);
-	FCoreDelegates::OnEndFrame.RemoveAll(this);
+	FCoreDelegates::OnBeginFrameRT.RemoveAll(this);
+	FCoreDelegates::OnEndFrameRT.RemoveAll(this);
 }
 
-void FLandscapeSceneViewExtension::BeginFrame()
+void FLandscapeSceneViewExtension::BeginFrame_RenderThread()
 {
-	ENQUEUE_RENDER_COMMAND(FLandscapeSceneViewExtension_BeginFrame)([this](FRHICommandListImmediate& RHICmdList)
+	for (auto& Pair : LandscapeRenderSystems)
 	{
-		for (auto& Pair : LandscapeRenderSystems)
-		{
-			FLandscapeRenderSystem& RenderSystem = *Pair.Value;
+		FLandscapeRenderSystem& RenderSystem = *Pair.Value;
 
-			RenderSystem.BeginFrame();
-		}
-	});
+		RenderSystem.BeginFrame();
+	}
 }
 
 void FLandscapeSceneViewExtension::PreRenderViewFamily_RenderThread(FRHICommandListImmediate& RHICmdList, FSceneViewFamily& InViewFamily)
@@ -1179,17 +1176,14 @@ void FLandscapeSceneViewExtension::PreRenderView_RenderThread(FRHICommandListImm
 	InView.LandscapeIndirectionBuffer = IndirectionSRV;
 }
 
-void FLandscapeSceneViewExtension::EndFrame()
+void FLandscapeSceneViewExtension::EndFrame_RenderThread()
 {
-	ENQUEUE_RENDER_COMMAND(FLandscapeSceneViewExtension_EndFrame)([this](FRHICommandListImmediate& RHICmdList)
+	for (auto& Pair : LandscapeRenderSystems)
 	{
-		for (auto& Pair : LandscapeRenderSystems)
-		{
-			FLandscapeRenderSystem& RenderSystem = *Pair.Value;
+		FLandscapeRenderSystem& RenderSystem = *Pair.Value;
 
-			RenderSystem.EndFrame();
-		}
-	});
+		RenderSystem.EndFrame();
+	}
 }
 
 //
diff --git a/Engine/Source/Runtime/Landscape/Public/LandscapeRender.h b/Engine/Source/Runtime/Landscape/Public/LandscapeRender.h
--- a/Engine/Source/Runtime/Landscape/Public/LandscapeRender.h
+++ b/Engine/Source/Runtime/Landscape/Public/LandscapeRender.h
@@ -582,8 +582,8 @@ public:
 	FLandscapeSceneViewExtension(const FAutoRegister& AutoReg);
 	virtual ~FLandscapeSceneViewExtension();
 
-	void BeginFrame();
-	void EndFrame();
+	void BeginFrame_RenderThread();
+	void EndFrame_RenderThread();
 
 	virtual void SetupViewFamily(FSceneViewFamily& InViewFamily) override {}
 	virtual void SetupView(FSceneViewFamily& InViewFamily, FSceneView& InView) override {}
