* Added debug name to landscape component edit layers data, so that we can display a more readable name when a component is referencing an invalid edit layer
* Changed landscape map check warnings to info messages, since the system is able to automatically cope with the problem

#tests editor

diff --git a/Engine/Source/Runtime/Landscape/Classes/LandscapeComponent.h b/Engine/Source/Runtime/Landscape/Classes/LandscapeComponent.h
--- a/Engine/Source/Runtime/Landscape/Classes/LandscapeComponent.h
+++ b/Engine/Source/Runtime/Landscape/Classes/LandscapeComponent.h
@@ -346,6 +346,21 @@ struct FLandscapeLayerComponentData
 {
 	GENERATED_USTRUCT_BODY()
 
+	FLandscapeLayerComponentData() = default;
+
+#if WITH_EDITOR
+	FLandscapeLayerComponentData(const FName& InDebugName)
+		: DebugName(InDebugName)
+	{}
+
+#endif // WITH_EDITOR
+
+#if WITH_EDITORONLY_DATA
+	// Edit layers are referenced by Guid, this name is just there to provide some insights as to what edit layer name this layer data corresponded to in case of a missing edit layer guid
+	UPROPERTY()
+	FName DebugName; 
+#endif // WITH_EDITORONLY_DATA
+
 	UPROPERTY()
 	FHeightmapData HeightmapData;
 
diff --git a/Engine/Source/Runtime/Landscape/Private/Landscape.cpp b/Engine/Source/Runtime/Landscape/Private/Landscape.cpp
--- a/Engine/Source/Runtime/Landscape/Private/Landscape.cpp
+++ b/Engine/Source/Runtime/Landscape/Private/Landscape.cpp
@@ -867,14 +867,23 @@ void ULandscapeComponent::PostLoad()
 		ReparentObject(MobileWeightmapTexture);
 	}
 
-	for (const auto& ItPair : LayersData)
+	for (auto& ItPair : LayersData)
 	{
-		const FLandscapeLayerComponentData& LayerComponentData = ItPair.Value;
+		FLandscapeLayerComponentData& LayerComponentData = ItPair.Value;
 		ReparentObject(LayerComponentData.HeightmapData.Texture);
 		for (UTexture2D* WeightmapTexture : LayerComponentData.WeightmapData.Textures)
 		{
 			ReparentObject(WeightmapTexture);
 		}
+
+		// Fixup missing/mismatching edit layer names :
+		if (const FLandscapeLayer* EditLayer = GetLandscapeActor() ? GetLandscapeActor()->GetLayer(ItPair.Key) : nullptr)
+		{
+			if (LayerComponentData.DebugName != EditLayer->Name)
+			{
+				LayerComponentData.DebugName = EditLayer->Name;
+			}
+		}
 	}
 
 	for (UMaterialInstance* MaterialInstance : MaterialInstances)
@@ -1762,7 +1771,8 @@ void ULandscapeComponent::AddDefaultLayerData(const FGuid& InLayerGuid, const TA
 
 	if (LayerData == nullptr || !LayerData->IsInitialized())
 	{
-		FLandscapeLayerComponentData NewData;
+		const FLandscapeLayer* EditLayer = GetLandscapeActor() ? GetLandscapeActor()->GetLayer(InLayerGuid) : nullptr;
+		FLandscapeLayerComponentData NewData(EditLayer ? EditLayer->Name : FName());
 
 		// Setup Heightmap data
 		UTexture2D** LayerHeightmap = InOutCreatedHeightmapTextures.Find(ComponentHeightmap);
@@ -2159,6 +2169,15 @@ void ALandscapeProxy::PreSave(FObjectPreSaveContext ObjectSaveContext)
 		for (ULandscapeComponent* LandscapeComponent : LandscapeComponents)
 		{
 			Landscape->ClearDirtyData(LandscapeComponent);
+
+			// Make sure edit layer debug names are synchronized upon save :
+			LandscapeComponent->ForEachLayer([&](const FGuid& LayerGuid, FLandscapeLayerComponentData& LayerData)
+			{
+				if (const FLandscapeLayer* EditLayer = Landscape->GetLayer(LayerGuid))
+				{
+					LayerData.DebugName = EditLayer->Name;
+				}
+			});
 		}
 	}
 #endif // WITH_EDITOR
@@ -2564,8 +2583,9 @@ void ALandscapeProxy::PostLoad()
 	{
 		FFormatNamedArguments Arguments;
 		Arguments.Add(TEXT("LandscapeName"), FText::FromString(GetPathName()));
-		FMessageLog("MapCheck").Warning()
-			->AddToken(FTextToken::Create(FText::Format(LOCTEXT("MapCheck_Message_FixedUpInvalidLandscapeMaterialInstances", "{LandscapeName} : Fixed up invalid landscape material instances. Please re-save the landscape package."), Arguments)))
+		Arguments.Add(TEXT("ProxyPackage"), FText::FromString(GetOutermost()->GetName()));
+		FMessageLog("MapCheck").Info()
+			->AddToken(FTextToken::Create(FText::Format(LOCTEXT("MapCheck_Message_FixedUpInvalidLandscapeMaterialInstances", "{LandscapeName} : Fixed up invalid landscape material instances. Please re-save {ProxyPackage}."), Arguments)))
 			->AddToken(FMapErrorToken::Create(FMapErrors::FixedUpInvalidLandscapeMaterialInstances));
 	}
 
diff --git a/Engine/Source/Runtime/Landscape/Private/LandscapeEditLayers.cpp b/Engine/Source/Runtime/Landscape/Private/LandscapeEditLayers.cpp
--- a/Engine/Source/Runtime/Landscape/Private/LandscapeEditLayers.cpp
+++ b/Engine/Source/Runtime/Landscape/Private/LandscapeEditLayers.cpp
@@ -151,8 +151,8 @@ private:
 
 	void ClearDirty()
 	{
-		if (!GWorld || GWorld->IsGameWorld())
-		{
+		if(!GWorld || GWorld->IsGameWorld())
+		{ 
 			return;
 		}
 
@@ -197,7 +197,7 @@ struct FTextureToComponentHelper
 
 				for (FWeightmapLayerAllocationInfo const& AllocInfo : AllocInfos)
 				{
-					if (AllocInfo.IsAllocated() && AllocInfo.WeightmapTextureIndex < WeightmapTextures.Num())
+					if(AllocInfo.IsAllocated() && AllocInfo.WeightmapTextureIndex < WeightmapTextures.Num())
 					{
 						UTexture2D* Weightmap = WeightmapTextures[AllocInfo.WeightmapTextureIndex];
 						check(Weightmap != nullptr);
@@ -265,7 +265,7 @@ public:
 
 		if (CreateUAV)
 		{
-			Flags |= TexCreate_RenderTargetable | TexCreate_UAV;
+			Flags |= TexCreate_RenderTargetable|TexCreate_UAV;
 		}
 
 		TextureRHI = RHICreateTexture2D(SizeX, SizeY, Format, NumMips, 1, Flags, CreateInfo);
@@ -282,7 +282,7 @@ public:
 		{
 			TextureUAV.SafeRelease();
 		}
-
+		
 		FTextureResource::ReleaseRHI();
 	}
 
@@ -746,9 +746,9 @@ public:
 		return true;
 	}
 
-	FLandscapeCopyTextureVS()
+	FLandscapeCopyTextureVS() 
 	{};
-
+	
 	FLandscapeCopyTextureVS(const ShaderMetaType::CompiledShaderInitializerType& Initializer)
 		: FGlobalShader(Initializer)
 	{}
@@ -1182,7 +1182,7 @@ class FLandscapeLayersCopyTexture_RenderThread
 public:
 	FLandscapeLayersCopyTexture_RenderThread(const FLandscapeLayersCopyTextureParams& InParams)
 		: Params(InParams)
-	{}
+	{}	
 
 	void Copy(FRHICommandListImmediate& InRHICmdList)
 	{
@@ -1204,7 +1204,7 @@ public:
 		Info.DestPosition.Y = Params.DestPosition.Y;
 		Info.SourceSliceIndex = Params.SourceArrayIndex;
 		Info.DestSliceIndex = Params.DestArrayIndex;
-		Info.SourceMipIndex = Params.SourceMip;
+		Info.SourceMipIndex = Params.SourceMip; 
 		Info.DestMipIndex = Params.DestMip;
 
 		check((Info.SourcePosition.X >= 0) && (Info.SourcePosition.Y >= 0) && (Info.DestPosition.X >= 0) && (Info.DestPosition.Y >= 0));
@@ -1263,7 +1263,7 @@ class FLandscapeLayersRender_RenderThread
 public:
 
 	FLandscapeLayersRender_RenderThread(const FString& InDebugName, UTextureRenderTarget2D* InWriteRenderTarget, const FIntPoint& InWriteRenderTargetSize, const FIntPoint& InReadRenderTargetSize, const FMatrix& InProjectionMatrix,
-		const ShaderDataType& InShaderParams, uint8 InCurrentMip, const TArray<FLandscapeLayersTriangle>& InTriangleList)
+										const ShaderDataType& InShaderParams, uint8 InCurrentMip, const TArray<FLandscapeLayersTriangle>& InTriangleList)
 		: RenderTargetResource(InWriteRenderTarget->GameThread_GetRenderTargetResource())
 		, WriteRenderTargetSize(InWriteRenderTargetSize)
 		, ReadRenderTargetSize(InReadRenderTargetSize)
@@ -1460,7 +1460,7 @@ bool ALandscape::ComputeLandscapeLayerBrushInfo(FTransform& OutLandscapeTransfor
 	{
 		return false;
 	}
-
+	
 	OutLandscapeTransform = Landscape->GetTransform();
 	FVector OffsetVector(LandscapeExtent.Min.X, LandscapeExtent.Min.Y, 0.f);
 	FVector Translation = OutLandscapeTransform.TransformFVector4(OffsetVector);
@@ -2216,7 +2216,7 @@ void ALandscape::CopyTexturePS(const FString& InSourceDebugName, FTextureResourc
 {
 	check(InSourceResource != nullptr);
 	check(InDestResource != nullptr);
-
+	
 	ENQUEUE_RENDER_COMMAND(LandscapeLayers_Cmd_CopyTexturePS)(
 		[InSourceDebugName, InSourceResource, InDestDebugName, InDestResource](FRHICommandListImmediate& RHICmdList)
 	{
@@ -2254,7 +2254,7 @@ void ALandscape::CopyTexturePS(const FString& InSourceDebugName, FTextureResourc
 }
 
 void ALandscape::DrawWeightmapComponentsToRenderTarget(const FString& InDebugName, const TArray<FIntPoint>& InSectionBaseList, const FVector2f& InScaleBias, TArray<FVector2f>* InScaleBiasPerSection, UTexture* InWeightmapRTRead, UTextureRenderTarget2D* InOptionalWeightmapRTRead2, UTextureRenderTarget2D* InWeightmapRTWrite,
-	ERTDrawingType InDrawType, bool InClearRTWrite, FLandscapeLayersWeightmapShaderParameters& InShaderParams, uint8 InMipRender) const
+														ERTDrawingType InDrawType, bool InClearRTWrite, FLandscapeLayersWeightmapShaderParameters& InShaderParams, uint8 InMipRender) const
 {
 	check(InWeightmapRTRead != nullptr);
 	check(InWeightmapRTWrite != nullptr);
@@ -2339,7 +2339,7 @@ void ALandscape::DrawWeightmapComponentsToRenderTarget(const FString& InDebugNam
 }
 
 void ALandscape::DrawWeightmapComponentsToRenderTarget(const FString& InDebugName, const TArray<ULandscapeComponent*>& InComponentsToDraw, const FIntPoint& InLandscapeBase, UTexture* InWeightmapRTRead, UTextureRenderTarget2D* InOptionalWeightmapRTRead2, UTextureRenderTarget2D* InWeightmapRTWrite,
-	ERTDrawingType InDrawType, bool InClearRTWrite, FLandscapeLayersWeightmapShaderParameters& InShaderParams, uint8 InMipRender) const
+														ERTDrawingType InDrawType, bool InClearRTWrite, FLandscapeLayersWeightmapShaderParameters& InShaderParams, uint8 InMipRender) const
 {
 
 	TArray<FIntPoint> SectionBaseList;
@@ -2388,7 +2388,7 @@ void ALandscape::DrawWeightmapComponentToRenderTargetMips(const TArray<FVector2f
 		if (WriteMipRT != nullptr)
 		{
 			DrawWeightmapComponentsToRenderTarget(FString::Printf(TEXT("LS Weight: %s = -> %s Mips %d"), *ReadMipRT->GetName(), *WriteMipRT->GetName(), CurrentMip),
-				SectionBaseToDraw, WeightmapScaleBias, nullptr, ReadMipRT, nullptr, WriteMipRT, ERTDrawingType::RTMips, InClearRTWrite, InShaderParams, CurrentMip);
+													SectionBaseToDraw, WeightmapScaleBias, nullptr, ReadMipRT, nullptr, WriteMipRT, ERTDrawingType::RTMips, InClearRTWrite, InShaderParams, CurrentMip);
 			++CurrentMip;
 		}
 
@@ -2420,7 +2420,7 @@ void ALandscape::DrawHeightmapComponentsToRenderTargetMips(const TArray<ULandsca
 		if (WriteMipRT != nullptr)
 		{
 			DrawHeightmapComponentsToRenderTarget(FString::Printf(TEXT("LS Height: %s = -> %s CombinedAtlasWithMips %d"), *ReadMipRT->GetName(), *WriteMipRT->GetName(), CurrentMip),
-				InComponentsToDraw, InLandscapeBase, ReadMipRT, nullptr, WriteMipRT, ERTDrawingType::RTMips, InClearRTWrite, InShaderParams, CurrentMip);
+												  InComponentsToDraw, InLandscapeBase, ReadMipRT, nullptr, WriteMipRT, ERTDrawingType::RTMips, InClearRTWrite, InShaderParams, CurrentMip);
 			++CurrentMip;
 		}
 
@@ -2429,7 +2429,7 @@ void ALandscape::DrawHeightmapComponentsToRenderTargetMips(const TArray<ULandsca
 }
 
 void ALandscape::DrawHeightmapComponentsToRenderTarget(const FString& InDebugName, const TArray<ULandscapeComponent*>& InComponentsToDraw, const FIntPoint& InLandscapeBase, UTexture* InHeightmapRTRead, UTextureRenderTarget2D* InOptionalHeightmapRTRead2, UTextureRenderTarget2D* InHeightmapRTWrite,
-	ERTDrawingType InDrawType, bool InClearRTWrite, FLandscapeLayersHeightmapShaderParameters& InShaderParams, uint8 InMipRender) const
+													  ERTDrawingType InDrawType, bool InClearRTWrite, FLandscapeLayersHeightmapShaderParameters& InShaderParams, uint8 InMipRender) const
 {
 	check(InHeightmapRTRead != nullptr);
 	check(InHeightmapRTWrite != nullptr);
@@ -2455,36 +2455,36 @@ void ALandscape::DrawHeightmapComponentsToRenderTarget(const FString& InDebugNam
 
 		switch (InDrawType)
 		{
-		case ERTDrawingType::RTAtlas:
-		{
-			GenerateLayersRenderQuadsAtlas(ComponentSectionBase, HeightmapScaleBias, SubsectionSizeQuads, HeightmapReadTextureSize, HeightmapWriteTextureSize, TriangleList);
-		} break;
+			case ERTDrawingType::RTAtlas:
+			{
+				GenerateLayersRenderQuadsAtlas(ComponentSectionBase, HeightmapScaleBias, SubsectionSizeQuads, HeightmapReadTextureSize, HeightmapWriteTextureSize, TriangleList);
+			} break;
 
-		case ERTDrawingType::RTAtlasToNonAtlas:
-		{
-			GenerateLayersRenderQuadsAtlasToNonAtlas(ComponentSectionBase, HeightmapScaleBias, SubsectionSizeQuads, HeightmapReadTextureSize, HeightmapWriteTextureSize, TriangleList);
-		} break;
+			case ERTDrawingType::RTAtlasToNonAtlas:
+			{
+				GenerateLayersRenderQuadsAtlasToNonAtlas(ComponentSectionBase, HeightmapScaleBias, SubsectionSizeQuads, HeightmapReadTextureSize, HeightmapWriteTextureSize, TriangleList);
+			} break;
 
-		case ERTDrawingType::RTNonAtlas:
-		{
-			GenerateLayersRenderQuadsNonAtlas(ComponentSectionBase, HeightmapScaleBias, SubsectionSizeQuads, HeightmapReadTextureSize, HeightmapWriteTextureSize, TriangleList);
-		} break;
+			case ERTDrawingType::RTNonAtlas:
+			{
+				GenerateLayersRenderQuadsNonAtlas(ComponentSectionBase, HeightmapScaleBias, SubsectionSizeQuads, HeightmapReadTextureSize, HeightmapWriteTextureSize, TriangleList);
+			} break;
 
-		case ERTDrawingType::RTNonAtlasToAtlas:
-		{
-			GenerateLayersRenderQuadsNonAtlasToAtlas(ComponentSectionBase, HeightmapScaleBias, SubsectionSizeQuads, HeightmapReadTextureSize, HeightmapWriteTextureSize, TriangleList);
-		} break;
+			case ERTDrawingType::RTNonAtlasToAtlas:
+			{
+				GenerateLayersRenderQuadsNonAtlasToAtlas(ComponentSectionBase, HeightmapScaleBias, SubsectionSizeQuads, HeightmapReadTextureSize, HeightmapWriteTextureSize, TriangleList);
+			} break;
 
-		case ERTDrawingType::RTMips:
-		{
-			GenerateLayersRenderQuadsMip(ComponentSectionBase, HeightmapScaleBias, SubsectionSizeQuads, HeightmapReadTextureSize, HeightmapWriteTextureSize, InMipRender, TriangleList);
-		} break;
+			case ERTDrawingType::RTMips:
+			{
+				GenerateLayersRenderQuadsMip(ComponentSectionBase, HeightmapScaleBias, SubsectionSizeQuads, HeightmapReadTextureSize, HeightmapWriteTextureSize, InMipRender, TriangleList);
+			} break;
 
-		default:
-		{
-			check(false);
-			return;
-		}
+			default:
+			{
+				check(false);
+				return;
+			}
 		}
 	}
 
@@ -2500,7 +2500,7 @@ void ALandscape::DrawHeightmapComponentsToRenderTarget(const FString& InDebugNam
 	}
 
 	FMatrix ProjectionMatrix = AdjustProjectionMatrixForRHI(FTranslationMatrix(FVector(0, 0, 0)) *
-		FMatrix(FPlane(1.0f / (FMath::Max<uint32>(HeightmapWriteTextureSize.X, 1.f) / 2.0f), 0.0, 0.0f, 0.0f), FPlane(0.0f, -1.0f / (FMath::Max<uint32>(HeightmapWriteTextureSize.Y, 1.f) / 2.0f), 0.0f, 0.0f), FPlane(0.0f, 0.0f, 1.0f, 0.0f), FPlane(-1.0f, 1.0f, 0.0f, 1.0f)));
+															FMatrix(FPlane(1.0f / (FMath::Max<uint32>(HeightmapWriteTextureSize.X, 1.f) / 2.0f), 0.0, 0.0f, 0.0f), FPlane(0.0f, -1.0f / (FMath::Max<uint32>(HeightmapWriteTextureSize.Y, 1.f) / 2.0f), 0.0f, 0.0f), FPlane(0.0f, 0.0f, 1.0f, 0.0f), FPlane(-1.0f, 1.0f, 0.0f, 1.0f)));
 
 	FLandscapeLayersHeightmapRender_RenderThread LayersRender(InDebugName, InHeightmapRTWrite, HeightmapWriteTextureSize, HeightmapReadTextureSize, ProjectionMatrix, InShaderParams, InMipRender, TriangleList);
 
@@ -2703,9 +2703,9 @@ void ALandscape::GenerateLayersRenderQuadsNonAtlas(const FIntPoint& InSectionBas
 	check(InReadSize.X == InWriteSize.X && InReadSize.Y == InWriteSize.Y);
 
 	int32 SubsectionSizeVerts = InSubSectionSizeQuad + 1;
-
+	
 	FVector2D UVSize((float)SubsectionSizeVerts / (float)InReadSize.X, (float)SubsectionSizeVerts / (float)InReadSize.Y);
-
+	
 	for (int8 SubY = 0; SubY < NumSubsections; ++SubY)
 	{
 		for (int8 SubX = 0; SubX < NumSubsections; ++SubX)
@@ -3181,7 +3181,7 @@ int32 ALandscape::RegenerateLayersHeightmaps(FTextureToComponentHelper const& Ma
 	// Lazily create CPU read back objects as required
 	if (HeightmapUpdateModes)
 	{
-		for (ULandscapeComponent* Component : InLandscapeComponentsToRender)
+		for(ULandscapeComponent* Component : InLandscapeComponentsToRender)
 		{
 			UTexture2D* ComponentHeightmap = Component->GetHeightmap(false);
 			ALandscapeProxy* Proxy = Component->GetLandscapeProxy();
@@ -3238,13 +3238,13 @@ int32 ALandscape::RegenerateLayersHeightmaps(FTextureToComponentHelper const& Ma
 		auto GetUniqueHeightmaps = [&](const TArray<ULandscapeComponent*>& InLandscapeComponents, TArray<FHeightmapCopyInfo>& OutHeightmaps, const FIntPoint& LandscapeBaseQuads, const FGuid& InLayerGuid = FGuid())
 		{
 			FScopedSetLandscapeEditingLayer Scope(this, InLayerGuid);
-
+			
 			const int32 ComponentSizeQuad = SubsectionSizeQuads * NumSubsections;
 			const int32 ComponentSizeVerts = (SubsectionSizeQuads + 1) * NumSubsections;
 			for (ULandscapeComponent* Component : InLandscapeComponents)
 			{
 				UTexture2D* ComponentHeightmap = Component->GetHeightmap(true);
-
+		
 				int32 Index = OutHeightmaps.IndexOfByPredicate([=](const FHeightmapCopyInfo& LayerHeightmap) { return LayerHeightmap.Texture == ComponentHeightmap; });
 
 				FIntPoint ComponentSectionBase = Component->GetSectionBase() - LandscapeBaseQuads;
@@ -3311,13 +3311,13 @@ int32 ALandscape::RegenerateLayersHeightmaps(FTextureToComponentHelper const& Ma
 
 			// NOTE: From this point on, we always work in non atlas, we'll convert back at the end to atlas only
 			DrawHeightmapComponentsToRenderTarget(FString::Printf(TEXT("LS Height: %s += -> NonAtlas %s"), *Layer.Name.ToString(), *LandscapeScratchRT1->GetName(), *LandscapeScratchRT2->GetName()),
-				InLandscapeComponentsToRender, LandscapeExtent.Min, LandscapeScratchRT1, nullptr, LandscapeScratchRT2, ERTDrawingType::RTAtlasToNonAtlas, true, ShaderParams);
+												  InLandscapeComponentsToRender, LandscapeExtent.Min, LandscapeScratchRT1, nullptr, LandscapeScratchRT2, ERTDrawingType::RTAtlasToNonAtlas, true, ShaderParams);
 
 			ShaderParams.ApplyLayerModifiers = true;
 
 			// Combine Current layer with current result
 			DrawHeightmapComponentsToRenderTarget(FString::Printf(TEXT("LS Height: %s += -> CombinedNonAtlas %s"), *Layer.Name.ToString(), *LandscapeScratchRT2->GetName(), *CombinedHeightmapNonAtlasRT->GetName()),
-				InLandscapeComponentsToRender, LandscapeExtent.Min, LandscapeScratchRT2, FirstLayer ? nullptr : LandscapeScratchRT3, CombinedHeightmapNonAtlasRT, ERTDrawingType::RTNonAtlas, FirstLayer, ShaderParams);
+												InLandscapeComponentsToRender, LandscapeExtent.Min, LandscapeScratchRT2, FirstLayer ? nullptr : LandscapeScratchRT3, CombinedHeightmapNonAtlasRT, ERTDrawingType::RTNonAtlas, FirstLayer, ShaderParams);
 
 			ShaderParams.ApplyLayerModifiers = false;
 
@@ -3365,12 +3365,12 @@ int32 ALandscape::RegenerateLayersHeightmaps(FTextureToComponentHelper const& Ma
 		}
 
 		DrawHeightmapComponentsToRenderTarget(FString::Printf(TEXT("LS Height: %s = -> CombinedNonAtlasNormals : %s"), *CombinedHeightmapNonAtlasRT->GetName(), *LandscapeScratchRT1->GetName()),
-			InLandscapeComponentsToRender, LandscapeExtent.Min, CombinedHeightmapNonAtlasRT, nullptr, LandscapeScratchRT1, ERTDrawingType::RTNonAtlas, true, ShaderParams);
+											  InLandscapeComponentsToRender, LandscapeExtent.Min, CombinedHeightmapNonAtlasRT, nullptr, LandscapeScratchRT1, ERTDrawingType::RTNonAtlas, true, ShaderParams);
 
 		ShaderParams.GenerateNormals = false;
 
 		DrawHeightmapComponentsToRenderTarget(FString::Printf(TEXT("LS Height: %s = -> CombinedAtlasFinal : %s"), *LandscapeScratchRT1->GetName(), *CombinedHeightmapAtlasRT->GetName()),
-			InLandscapeComponentsToRender, LandscapeExtent.Min, LandscapeScratchRT1, nullptr, CombinedHeightmapAtlasRT, ERTDrawingType::RTNonAtlasToAtlas, true, ShaderParams);
+											  InLandscapeComponentsToRender, LandscapeExtent.Min, LandscapeScratchRT1, nullptr, CombinedHeightmapAtlasRT, ERTDrawingType::RTNonAtlasToAtlas, true, ShaderParams);
 
 		DrawHeightmapComponentsToRenderTargetMips(InLandscapeComponentsToRender, LandscapeExtent.Min, CombinedHeightmapAtlasRT, true, ShaderParams);
 
@@ -3440,7 +3440,7 @@ void ALandscape::UpdateForChangedHeightmaps(ULandscapeComponent* InComponent, co
 	//  - If ELandscapeComponentUpdateFlag::Component_Update_Heightmap_Collision is passed, it will be done immediately
 	//  - If not, at least the component's collision data will still get updated eventually, when the flag is finally passed :
 	if (InReadbackResult.bModified)
-	{
+{
 		InComponent->SetPendingCollisionDataUpdate(true);
 	}
 
@@ -3454,8 +3454,8 @@ void ALandscape::UpdateForChangedHeightmaps(ULandscapeComponent* InComponent, co
 			InComponent->UpdateCachedBounds();
 			InComponent->UpdateComponentToWorld();
 
-			// Avoid updating height field if we are going to recreate collision in this update
-			bool bUpdateHeightfieldRegion = !IsUpdateFlagEnabledForModes(ELandscapeComponentUpdateFlag::Component_Update_Recreate_Collision, HeightUpdateMode);
+		// Avoid updating height field if we are going to recreate collision in this update
+		bool bUpdateHeightfieldRegion = !IsUpdateFlagEnabledForModes(ELandscapeComponentUpdateFlag::Component_Update_Recreate_Collision, HeightUpdateMode);
 			InComponent->UpdateCollisionData(bUpdateHeightfieldRegion);
 			InComponent->SetPendingCollisionDataUpdate(false);
 		}
@@ -3532,7 +3532,7 @@ void ALandscape::ClearDirtyData(ULandscapeComponent* InLandscapeComponent)
 	FMemory::Memzero(DirtyData.Get(), DirtyDataSize);
 	LandscapeEdit.SetDirtyData(X1, Y1, X2, Y2, DirtyData.Get(), 0);
 }
-
+	
 void ALandscape::UpdateWeightDirtyData(ULandscapeComponent* InLandscapeComponent, UTexture2D const* InWeightmap, FColor const* InOldData, FColor const* InNewData, uint8 InChannel)
 {
 	check(InOldData && InNewData);
@@ -3549,7 +3549,7 @@ void ALandscape::UpdateWeightDirtyData(ULandscapeComponent* InLandscapeComponent
 	const int32 SizeV = InWeightmap->Source.GetSizeY();
 	check(DirtyDataSize == SizeU * SizeV);
 
-	const uint8 DirtyWeight = 1 << 1;
+	const uint8 DirtyWeight = 1<<1;
 	LandscapeEdit.GetDirtyData(X1, Y1, X2, Y2, DirtyData.Get(), 0);
 
 	for (int32 Index = 0; Index < DirtyDataSize; ++Index)
@@ -3615,16 +3615,16 @@ void ALandscape::UpdateHeightDirtyData(ULandscapeComponent* InLandscapeComponent
 	const int32 X2 = X1 + ComponentSizeQuads;
 	const int32 Y1 = InLandscapeComponent->GetSectionBase().Y;
 	const int32 Y2 = Y1 + ComponentSizeQuads;
-	const int32 ComponentWidth = (SubsectionSizeQuads + 1) * NumSubsections;
+	const int32 ComponentWidth = (SubsectionSizeQuads + 1)*NumSubsections;
 	const int32 DirtyDataSize = ComponentWidth * ComponentWidth;
 	TUniquePtr<uint8[]> DirtyData = MakeUnique<uint8[]>(DirtyDataSize);
 	const int32 SizeU = InHeightmap->Source.GetSizeX();
 	const int32 SizeV = InHeightmap->Source.GetSizeY();
 	const int32 HeightmapOffsetX = InLandscapeComponent->HeightmapScaleBias.Z * SizeU;
 	const int32 HeightmapOffsetY = InLandscapeComponent->HeightmapScaleBias.W * SizeV;
-	const uint8 DirtyHeight = 1 << 0;
+	const uint8 DirtyHeight = 1<<0;
 	LandscapeEdit.GetDirtyData(X1, Y1, X2, Y2, DirtyData.Get(), 0);
-
+	
 	for (int32 X = 0; X < ComponentWidth; ++X)
 	{
 		for (int32 Y = 0; Y < ComponentWidth; ++Y)
@@ -3632,7 +3632,7 @@ void ALandscape::UpdateHeightDirtyData(ULandscapeComponent* InLandscapeComponent
 			int32 TexX = HeightmapOffsetX + X;
 			int32 TexY = HeightmapOffsetY + Y;
 			int32 TexIndex = TexX + TexY * SizeU;
-			check(TexIndex < SizeU* SizeV);
+			check(TexIndex < SizeU * SizeV);
 			if (InOldData[TexIndex] != InNewData[TexIndex])
 			{
 				DirtyData[X + Y * ComponentWidth] |= DirtyHeight;
@@ -3764,7 +3764,7 @@ bool ALandscape::ResolveLayersTexture(
 }
 
 void ALandscape::PrepareComponentDataToExtractMaterialLayersCS(const TArray<ULandscapeComponent*>& InLandscapeComponents, const FLandscapeLayer& InLayer, int32 InCurrentWeightmapToProcessIndex, const FIntPoint& InLandscapeBase, FLandscapeTexture2DResource* InOutTextureData,
-	TArray<FLandscapeLayerWeightmapExtractMaterialLayersComponentData>& OutComponentData, TMap<ULandscapeLayerInfoObject*, int32>& OutLayerInfoObjects)
+																TArray<FLandscapeLayerWeightmapExtractMaterialLayersComponentData>& OutComponentData, TMap<ULandscapeLayerInfoObject*, int32>& OutLayerInfoObjects)
 {
 	ULandscapeInfo* Info = GetLandscapeInfo();
 
@@ -3850,7 +3850,7 @@ void ALandscape::PrepareComponentDataToExtractMaterialLayersCS(const TArray<ULan
 			}
 		}
 	}
-
+	
 	ExecuteCopyLayersTexture(MoveTemp(DeferredCopyTextures));
 }
 
@@ -3992,7 +3992,7 @@ void ALandscape::ReallocateLayersWeightmaps(const TArray<ULandscapeComponent*>&
 			Usage->ClearUsage(Component);
 		}
 	}
-
+		
 	// Build a map of all the allocation per components
 	TMap<ULandscapeComponent*, TArray<ULandscapeLayerInfoObject*>> LayerAllocsPerComponent;
 
@@ -4111,7 +4111,7 @@ void ALandscape::ReallocateLayersWeightmaps(const TArray<ULandscapeComponent*>&
 				if (FLandscapeEditLayerReadback** CPUReadbackToDelete = Proxy->WeightmapsCPUReadback.Find(OldWeightmapTexture))
 				{
 					check(*CPUReadbackToDelete);
-					delete* CPUReadbackToDelete;
+					delete * CPUReadbackToDelete;
 					Proxy->WeightmapsCPUReadback.Remove(OldWeightmapTexture);
 				}
 			}
@@ -4191,7 +4191,7 @@ int32 ALandscape::RegenerateLayersWeightmaps(FTextureToComponentHelper const& Ma
 	const int32 WeightmapUpdateModes = LayerContentUpdateModes & AllWeightmapUpdateModes;
 	const bool bSkipBrush = CVarLandscapeLayerBrushOptim.GetValueOnAnyThread() == 1 && ((WeightmapUpdateModes & AllWeightmapUpdateModes) == ELandscapeLayerUpdateMode::Update_Weightmap_Editing);
 	const bool bForceRender = CVarForceLayersUpdate.GetValueOnAnyThread() != 0;
-
+	
 	ULandscapeInfo* Info = GetLandscapeInfo();
 
 	if (WeightmapUpdateModes == 0 && !bForceRender)
@@ -4203,7 +4203,7 @@ int32 ALandscape::RegenerateLayersWeightmaps(FTextureToComponentHelper const& Ma
 	{
 		return WeightmapUpdateModes;
 	}
-
+			
 	TArray<ULandscapeComponent*> ComponentThatNeedMaterialRebuild;
 	TArray<ULandscapeLayerInfoObject*> BrushRequiredAllocations;
 	int32 LayerCount = Info->Layers.Num() + 1; // due to visibility being stored at 0
@@ -4218,7 +4218,7 @@ int32 ALandscape::RegenerateLayersWeightmaps(FTextureToComponentHelper const& Ma
 		{
 			return 0;
 		}
-
+		
 		check(WeightmapRTList.Num() > 0);
 
 		UTextureRenderTarget2D* LandscapeScratchRT1 = WeightmapRTList[(int32)EWeightmapRTType::WeightmapRT_Scratch1];
@@ -4264,7 +4264,7 @@ int32 ALandscape::RegenerateLayersWeightmaps(FTextureToComponentHelper const& Ma
 				for (int32 LayerInfoSettingsIndex = 0; LayerInfoSettingsIndex < Info->Layers.Num(); ++LayerInfoSettingsIndex)
 				{
 					const FLandscapeInfoLayerSettings& InfoLayerSettings = Info->Layers[LayerInfoSettingsIndex];
-
+					
 					// It is possible that no layer info has been assigned so that InfoLayerSettings.LayerInfoObj == nullptr. In that case we don't consider it the layer here.
 					if (InfoLayerSettings.LayerInfoObj != nullptr)
 					{
@@ -4332,7 +4332,7 @@ int32 ALandscape::RegenerateLayersWeightmaps(FTextureToComponentHelper const& Ma
 				if (ComponentsData.Num() > 0)
 				{
 					PrintLayersDebugTextureResource(FString::Printf(TEXT("LS Weight: %s WeightmapScratchTexture %s"), *Layer.Name.ToString(), TEXT("WeightmapScratchTextureResource")), WeightmapScratchExtractLayerTextureResource, 0, false);
-
+										
 					FLandscapeLayerWeightmapExtractMaterialLayersComputeShaderParameters CSExtractLayersShaderParams;
 					CSExtractLayersShaderParams.AtlasWeightmapsPerLayer = CurrentLayersWeightmapAllMaterialLayersResource;
 					CSExtractLayersShaderParams.ComponentWeightmapResource = WeightmapScratchExtractLayerTextureResource;
@@ -4379,7 +4379,7 @@ int32 ALandscape::RegenerateLayersWeightmaps(FTextureToComponentHelper const& Ma
 					PSShaderParams.LayerAlpha = LayerInfoObj == ALandscapeProxy::VisibilityLayer ? 1.0f : Layer.WeightmapAlpha; // visibility can't be affected by weight
 
 					DrawWeightmapComponentsToRenderTarget(FString::Printf(TEXT("LS Weight: %s Paint: %s += -> %s"), *Layer.Name.ToString(), *LayerInfoObj->LayerName.ToString(), *LandscapeScratchRT1->GetName(), *LandscapeScratchRT2->GetName()),
-						InLandscapeComponentsToRender, LandscapeExtent.Min, LandscapeScratchRT1, nullptr, LandscapeScratchRT2, ERTDrawingType::RTAtlas, true, PSShaderParams, 0);
+														  InLandscapeComponentsToRender, LandscapeExtent.Min, LandscapeScratchRT1, nullptr, LandscapeScratchRT2, ERTDrawingType::RTAtlas, true, PSShaderParams, 0);
 
 					PSShaderParams.ApplyLayerModifiers = false;
 
@@ -4410,7 +4410,7 @@ int32 ALandscape::RegenerateLayersWeightmaps(FTextureToComponentHelper const& Ma
 					}
 
 					DrawWeightmapComponentsToRenderTarget(FString::Printf(TEXT("LS Weight: %s PaintLayer: %s, %s += -> Combined %s"), *Layer.Name.ToString(), *LayerInfoObj->LayerName.ToString(), *LandscapeScratchRT2->GetName(), *LandscapeScratchRT3->GetName()),
-						InLandscapeComponentsToRender, LandscapeExtent.Min, LandscapeScratchRT2, bFirstLayer ? nullptr : LandscapeScratchRT1, LandscapeScratchRT3, ERTDrawingType::RTAtlasToNonAtlas, true, PSShaderParams, 0);
+														  InLandscapeComponentsToRender, LandscapeExtent.Min, LandscapeScratchRT2, bFirstLayer ? nullptr : LandscapeScratchRT1, LandscapeScratchRT3, ERTDrawingType::RTAtlasToNonAtlas, true, PSShaderParams, 0);
 
 					PSShaderParams.OutputAsSubstractive = false;
 
@@ -4506,7 +4506,7 @@ int32 ALandscape::RegenerateLayersWeightmaps(FTextureToComponentHelper const& Ma
 					}
 				}
 			}
-
+			
 			int8 CurrentWeightmapToProcessIndex = 0;
 			bool HasFoundWeightmapToProcess = true; // try processing at least once
 
@@ -4548,9 +4548,9 @@ int32 ALandscape::RegenerateLayersWeightmaps(FTextureToComponentHelper const& Ma
 
 					for (int32 i = 0; i < PackLayersComponentsData.Num(); ++i)
 					{
-						check(ComponentY + ComponentSize <= WeightmapScratchPackLayerTextureResource->GetSizeY()); // This should never happen as it would be a bug in the algo
+						check(ComponentY+ComponentSize <= WeightmapScratchPackLayerTextureResource->GetSizeY()); // This should never happen as it would be a bug in the algo
 
-						if (ComponentX + ComponentSize > WeightmapScratchPackLayerTextureResource->GetSizeX())
+						if (ComponentX+ComponentSize > WeightmapScratchPackLayerTextureResource->GetSizeX())
 						{
 							ComponentY += ComponentSize;
 							ComponentX = 0;
@@ -4586,7 +4586,7 @@ int32 ALandscape::RegenerateLayersWeightmaps(FTextureToComponentHelper const& Ma
 
 					SourceDebugName = TEXT("WeightmapScratchTexture");
 					DestDebugName = CurrentRT->GetName();
-
+					
 					CopyTexturePS(SourceDebugName, WeightmapScratchPackLayerTextureResource, DestDebugName, CurrentRT->GameThread_GetRenderTargetResource());
 					DrawWeightmapComponentToRenderTargetMips(WeightmapTextureOutputOffset, CurrentRT, true, PSShaderParams);
 
@@ -4903,7 +4903,7 @@ void ALandscape::ResolveLayersWeightmapTexture(
 			}
 		}
 	}
-
+		
 	// Weightmaps shouldn't invalidate lighting
 	const bool bInvalidateLightingCache = false;
 	InvalidateGeneratedComponentData(ChangedComponents, bInvalidateLightingCache);
@@ -5003,7 +5003,7 @@ bool ALandscape::IsUpdateFlagEnabledForModes(ELandscapeComponentUpdateFlag InFla
 			return true;
 		}
 	}
-
+		
 	if (InUpdateModes & ELandscapeLayerUpdateMode::Update_Heightmap_Editing)
 	{
 		const uint32 HeightmapEditingFlags = ELandscapeComponentUpdateFlag::Component_Update_Heightmap_Collision | ELandscapeComponentUpdateFlag::Component_Update_Client_Editing;
@@ -5188,10 +5188,10 @@ void ALandscape::GetLandscapeComponentNeighborsToRender(ULandscapeComponent* Lan
 {
 	ULandscapeInfo* LandscapeInfo = GetLandscapeInfo();
 	FIntPoint ComponentKey = LandscapeComponent->GetSectionBase() / ComponentSizeQuads;
-
-	for (int32 IndexX = ComponentKey.X - 1; IndexX <= ComponentKey.X + 1; ++IndexX)
+	
+	for (int32 IndexX = ComponentKey.X-1; IndexX <= ComponentKey.X+1; ++IndexX)
 	{
-		for (int32 IndexY = ComponentKey.Y - 1; IndexY <= ComponentKey.Y + 1; ++IndexY)
+		for (int32 IndexY = ComponentKey.Y-1; IndexY <= ComponentKey.Y+1; ++IndexY)
 		{
 			ULandscapeComponent* Result = LandscapeInfo->XYtoComponentMap.FindRef(FIntPoint(IndexX, IndexY));
 			if (Result && Result != LandscapeComponent)
@@ -5328,7 +5328,7 @@ void ALandscape::UpdateLayersContent(bool bInWaitForStreaming, bool bInSkipMonit
 		}
 		return;
 	}
-
+	
 	// Gather mappings between heightmaps/weightmaps and components
 	FTextureToComponentHelper MapHelper(*LandscapeInfo);
 
@@ -5340,7 +5340,7 @@ void ALandscape::UpdateLayersContent(bool bInWaitForStreaming, bool bInSkipMonit
 		// These flags might look like they're being mixed up but they're not!
 		const bool bDoIntermediateRender = false; // bIntermediateRender flag is for the work queued up this frame not the delayed resolves
 		const bool bDoFlushRender = bIntermediateRender; // Flush before we do an intermediate render later in this frame
-
+		
 		TMap<ULandscapeComponent*, FLandscapeEditLayerReadbackResult> ModifiedComponentInfo;
 		ResolveLayersHeightmapTexture(MapHelper, MapHelper.Heightmaps, bDoIntermediateRender, bDoFlushRender, ModifiedComponentInfo);
 		ResolveLayersWeightmapTexture(MapHelper, MapHelper.Weightmaps, bDoIntermediateRender, bDoFlushRender, ModifiedComponentInfo);
@@ -5374,7 +5374,7 @@ void ALandscape::UpdateLayersContent(bool bInWaitForStreaming, bool bInSkipMonit
 		if (!bPartialUpdate || Component->GetLayerUpdateFlagPerMode() != 0)
 		{
 			AllLandscapeComponents.Add(Component);
-
+		
 			if (bPartialUpdate)
 			{
 				// Gather Neighbors (Neighbors need to be Rendered but not resolved so that the resolved Components have valid normals on edges)
@@ -5420,7 +5420,7 @@ void ALandscape::UpdateLayersContent(bool bInWaitForStreaming, bool bInSkipMonit
 	LandscapeComponentsHeightmapsToRender.Append(AllLandscapeComponents);
 	LandscapeComponentsWeightmapsToResolve.Append(AllLandscapeComponents);
 	LandscapeComponentsWeightmapsToRender.Append(AllLandscapeComponents);
-
+		
 	if (bPartialUpdate)
 	{
 		for (ULandscapeComponent* Component : SkippedComponents)
@@ -5807,28 +5807,29 @@ void ALandscapeProxy::UpdateCachedHasLayersContent(bool InCheckComponentDataInte
 
 bool ALandscapeProxy::RemoveObsoleteLayers(const TSet<FGuid>& InExistingLayers)
 {
-	TSet<FGuid> ComponentLayers;
+	TSet<TPair<FGuid, FName>> ComponentLayers;
 	for (ULandscapeComponent* Component : LandscapeComponents)
 	{
-		Component->ForEachLayer([&](const FGuid& InGuid, FLandscapeLayerComponentData&) { ComponentLayers.Add(InGuid); });
+		Component->ForEachLayer([&](const FGuid& Guid, FLandscapeLayerComponentData& ComponentData) { ComponentLayers.Add(TPair<FGuid, FName>(Guid, ComponentData.DebugName)); });
 	}
 
 	bool bModified = false;
 
-	for (const FGuid& LayerGuid : ComponentLayers)
+	for (const TPair<FGuid, FName>& LayerGuidAndName : ComponentLayers)
 	{
-		if (!InExistingLayers.Contains(LayerGuid))
+		if (!InExistingLayers.Contains(LayerGuidAndName.Key))
 		{
 			FFormatNamedArguments Arguments;
-			Arguments.Add(TEXT("LayerGuid"), FText::FromString(LayerGuid.ToString(EGuidFormats::HexValuesInBraces)));
+			Arguments.Add(TEXT("LayerName"), FText::FromString(LayerGuidAndName.Value.ToString()));
+			Arguments.Add(TEXT("LayerGuid"), FText::FromString(LayerGuidAndName.Key.ToString(EGuidFormats::HexValuesInBraces)));
 			Arguments.Add(TEXT("ProxyPackage"), FText::FromString(GetOutermost()->GetName()));
 
-			FMessageLog("MapCheck").Warning()
+			FMessageLog("MapCheck").Info()
 				->AddToken(FUObjectToken::Create(this))
-				->AddToken(FTextToken::Create(FText::Format(LOCTEXT("MapCheck_Message_LandscapeProxyObsoleteLayer", "Layer '{LayerGuid}' was removed from LandscapeProxy because it doesn't match any of the LandscapeActor Layers. Please resave '{ProxyPackage}'."), Arguments)))
+				->AddToken(FTextToken::Create(FText::Format(LOCTEXT("MapCheck_Message_LandscapeProxyObsoleteLayer","Layer '{LayerName}' ('{LayerGuid}') was removed from LandscapeProxy because it doesn't match any of the LandscapeActor Layers. Please resave '{ProxyPackage}'."), Arguments)))
 				->AddToken(FMapErrorToken::Create(FMapErrors::LandscapeComponentPostLoad_Warning));
 
-			DeleteLayer(LayerGuid);
+			DeleteLayer(LayerGuidAndName.Key);
 			bModified = true;
 		}
 	}
@@ -5851,10 +5852,11 @@ bool ALandscapeProxy::AddLayer(const FGuid& InLayerGuid)
 	{
 		if (!Component->GetLayerData(InLayerGuid))
 		{
-			Component->AddLayerData(InLayerGuid, FLandscapeLayerComponentData());
+			const FLandscapeLayer* EditLayer = GetLandscapeActor() ? GetLandscapeActor()->GetLayer(InLayerGuid) : nullptr;
+			Component->AddLayerData(InLayerGuid, FLandscapeLayerComponentData(EditLayer ? EditLayer->Name : FName()));
 			bModified = true;
 		}
-	}
+	}	
 
 	UpdateCachedHasLayersContent();
 
@@ -5904,7 +5906,7 @@ void ALandscapeProxy::InitializeLayerWithEmptyContent(const FGuid& InLayerGuid)
 	{
 		return;
 	}
-
+		
 	// Build a mapping between each Heightmaps and Component in them
 	TMap<UTexture2D*, TArray<ULandscapeComponent*>> ComponentsPerHeightmaps;
 
@@ -5914,7 +5916,7 @@ void ALandscapeProxy::InitializeLayerWithEmptyContent(const FGuid& InLayerGuid)
 		TArray<ULandscapeComponent*>& ComponentList = ComponentsPerHeightmaps.FindOrAdd(ComponentHeightmapTexture);
 		ComponentList.Add(Component);
 	}
-
+		
 	// Init layers with valid "empty" data
 	TMap<UTexture2D*, UTexture2D*> CreatedHeightmapTextures; // < Final layer texture, New created texture for layer
 
@@ -6032,7 +6034,7 @@ float ALandscape::GetLayerAlpha(int32 InLayerIndex, bool bInHeightmap) const
 
 float ALandscape::GetClampedLayerAlpha(float InAlpha, bool bInHeightmap) const
 {
-	float AlphaClamped = FMath::Clamp<float>(InAlpha, bInHeightmap ? -1.f : 0.f, 1.f);
+	float AlphaClamped = FMath::Clamp<float>(InAlpha, bInHeightmap  ? -1.f : 0.f, 1.f);
 	return AlphaClamped;
 }
 
@@ -6210,7 +6212,7 @@ void ALandscape::CollapseLayer(int32 InLayerIndex)
 			LandscapeComponent->CopyFinalLayerIntoEditingLayer(DataInterface, ProcessedHeightmaps);
 		});
 	}
-
+	
 	TArray<ALandscapeBlueprintBrushBase*> BrushesToMove;
 	for (int32 i = 0; i < LandscapeLayers[InLayerIndex].Brushes.Num(); ++i)
 	{
@@ -6231,7 +6233,7 @@ void ALandscape::CollapseLayer(int32 InLayerIndex)
 	}
 
 	DeleteLayer(InLayerIndex);
-
+	
 	RequestLayersContentUpdateForceAll();
 }
 
@@ -6251,7 +6253,7 @@ void ALandscape::GetUsedPaintLayers(const FGuid& InLayerGuid, TArray<ULandscapeL
 	{
 		return;
 	}
-
+	
 	LandscapeInfo->GetUsedPaintLayers(InLayerGuid, OutUsedLayerInfos);
 }
 
@@ -6461,15 +6463,15 @@ uint32 ULandscapeComponent::ComputeLayerHash() const
 {
 	UTexture2D* Heightmap = GetHeightmap(true);
 	const uint8* MipData = Heightmap->Source.LockMipReadOnly(0);
-	uint32 Hash = FCrc::MemCrc32(MipData, Heightmap->GetSizeX() * Heightmap->GetSizeY() * sizeof(FColor));
+	uint32 Hash = FCrc::MemCrc32(MipData, Heightmap->GetSizeX()*Heightmap->GetSizeY() * sizeof(FColor));
 	Heightmap->Source.UnlockMip(0);
 
 	// Copy to sort
 	const TArray<UTexture2D*>& Weightmaps = GetWeightmapTextures(true);
 	TArray<FWeightmapLayerAllocationInfo> AllocationInfos = GetWeightmapLayerAllocations(true);
-
-	// Sort allocations infos by LayerInfo Path so the Weightmaps hahses get ordered properly
-	AllocationInfos.Sort([](const FWeightmapLayerAllocationInfo& A, const FWeightmapLayerAllocationInfo& B)
+	
+    // Sort allocations infos by LayerInfo Path so the Weightmaps hahses get ordered properly
+    AllocationInfos.Sort([](const FWeightmapLayerAllocationInfo& A, const FWeightmapLayerAllocationInfo& B)
 	{
 		FString PathA(A.LayerInfo ? A.LayerInfo->GetPathName() : FString());
 		FString PathB(B.LayerInfo ? B.LayerInfo->GetPathName() : FString());
@@ -6481,23 +6483,23 @@ uint32 ULandscapeComponent::ComputeLayerHash() const
 	{
 		if (AllocationInfo.IsAllocated())
 		{
-			// Compute hash of actual data of the texture that is owned by the component (per Texture Channel)
+            // Compute hash of actual data of the texture that is owned by the component (per Texture Channel)
 			UTexture2D* Weightmap = Weightmaps[AllocationInfo.WeightmapTextureIndex];
 			MipData = Weightmap->Source.LockMipReadOnly(0) + ChannelOffsets[AllocationInfo.WeightmapTextureChannel];
 			TArray<uint8> ChannelData;
-			ChannelData.AddDefaulted(Weightmap->GetSizeX() * Weightmap->GetSizeY());
-			int32 TexSize = (SubsectionSizeQuads + 1) * NumSubsections;
+			ChannelData.AddDefaulted(Weightmap->GetSizeX()*Weightmap->GetSizeY());
+			int32 TexSize = (SubsectionSizeQuads + 1)*NumSubsections;
 			for (int32 TexY = 0; TexY < TexSize; TexY++)
 			{
 				for (int32 TexX = 0; TexX < TexSize; TexX++)
 				{
 					const int32 Index = (TexX + TexY * TexSize);
-
+				
 					ChannelData.GetData()[Index] = MipData[4 * Index];
 				}
 			}
 
-			Hash = FCrc::MemCrc32(ChannelData.GetData(), Weightmap->GetSizeX() * Weightmap->GetSizeY(), Hash);
+			Hash = FCrc::MemCrc32(ChannelData.GetData(), Weightmap->GetSizeX()*Weightmap->GetSizeY(), Hash);
 			Weightmap->Source.UnlockMip(0);
 		}
 	}
@@ -6525,11 +6527,11 @@ void ALandscape::UpdateLandscapeSplines(const FGuid& InTargetLayer, bool bInUpda
 			{
 				return;
 			}
-
+						
 			TMap<ULandscapeComponent*, uint32> PreviousHashes;
 			{
 				FLandscapeEditDataInterface LandscapeEdit(LandscapeInfo);
-
+				
 				LandscapeInfo->ForAllLandscapeComponents([&](ULandscapeComponent* Component)
 				{
 					// Was never computed
@@ -6554,8 +6556,8 @@ void ALandscape::UpdateLandscapeSplines(const FGuid& InTargetLayer, bool bInUpda
 			bInUpdateOnlySelected = false;
 
 			// Apply splines without clearing up weightmap allocations
-			LandscapeInfo->ApplySplines(bInUpdateOnlySelected, ModifiedComponent, bMarkPackageDirty);
-
+			LandscapeInfo->ApplySplines(bInUpdateOnlySelected, ModifiedComponent, bMarkPackageDirty);		
+		
 			for (const TPair<ULandscapeComponent*, uint32>& Pair : PreviousHashes)
 			{
 				if (LandscapeSplinesAffectedComponents.Contains(Pair.Key))
@@ -6617,7 +6619,7 @@ bool ALandscape::IsEditingLayerReservedForSplines() const
 void ALandscape::SetEditingLayer(const FGuid& InLayerGuid)
 {
 	ensure(CanHaveLayersContent());
-
+	
 	ULandscapeInfo* LandscapeInfo = GetLandscapeInfo();
 	if (!LandscapeInfo)
 	{
@@ -6931,12 +6933,12 @@ void ALandscape::OnBlueprintBrushChanged()
 void ALandscape::OnLayerInfoSplineFalloffModulationChanged(ULandscapeLayerInfoObject* InLayerInfo)
 {
 	ULandscapeInfo* LandscapeInfo = GetLandscapeInfo();
-
+	
 	if (!LandscapeInfo)
 	{
 		return;
 	}
-
+	
 	ALandscape* Landscape = LandscapeInfo->LandscapeActor.Get();
 	if (!Landscape || !Landscape->HasLayersContent())
 	{
@@ -7063,7 +7065,7 @@ bool FLandscapeLayerBrush::Initialize(const FIntRect& InLandscapeExtent, UTextur
 				LandscapeTransform = NewLandscapeTransform;
 				LandscapeRenderTargetSize = NewLandscapeRenderTargetSize;
 				LandscapeSize = NewLandscapeSize;
-
+				
 				TGuardValue<bool> AutoRestore(GAllowActorScriptExecutionInEditor, true);
 				BlueprintBrush->Initialize(LandscapeTransform, LandscapeSize, LandscapeRenderTargetSize);
 			}
