Removed old check that prevents updating landscape material instances in EditorGame mode.

diff --git a/Engine/Source/Runtime/Landscape/Private/LandscapeEdit.cpp b/Engine/Source/Runtime/Landscape/Private/LandscapeEdit.cpp
--- a/Engine/Source/Runtime/Landscape/Private/LandscapeEdit.cpp
+++ b/Engine/Source/Runtime/Landscape/Private/LandscapeEdit.cpp
@@ -370,8 +370,6 @@ UMaterialInstanceConstant* ULandscapeComponent::GetCombinationMaterial(FMaterial
 
 void ULandscapeComponent::UpdateMaterialInstances_Internal(FMaterialUpdateContext& Context)
 {
-	check(GIsEditor);
-
 	int32 MaxLOD = FMath::CeilLogTwo(SubsectionSizeQuads + 1) - 1;
 	decltype(MaterialPerLOD) NewMaterialPerLOD;
 	LODIndexToMaterialIndex.SetNumUninitialized(MaxLOD+1);
