-Fix ComputeComponentCounts to return actual number of components, could be wrong in case of trimmed landscape since min was not initalized properly
-Fix RT creation in case you have 0 landscape component (previously always returned a 1x1 count, now returns 0x0)

PR #8071: Fix ComputeComponentCounts for trimmed landscapes (Contributed by [REDACTED_PERSON])

diff --git a/Engine/Source/Runtime/Landscape/Private/LandscapeEditLayers.cpp b/Engine/Source/Runtime/Landscape/Private/LandscapeEditLayers.cpp
--- a/Engine/Source/Runtime/Landscape/Private/LandscapeEditLayers.cpp
+++ b/Engine/Source/Runtime/Landscape/Private/LandscapeEditLayers.cpp
@@ -1484,6 +1484,12 @@ void ALandscape::CreateLayersRenderingResource()
 
 	const FIntPoint ComponentCounts = ComputeComponentCounts();
 
+	// No components, can't update the render targets	
+	if (ComponentCounts.X <= 0 || ComponentCounts.Y <= 0)
+	{
+		return;
+	}
+
 	ALandscape* Landscape = GetLandscapeActor();
 	check(Landscape);
 
@@ -1716,8 +1722,8 @@ FIntPoint ALandscape::ComputeComponentCounts() const
 	}
 
 	FIntPoint NumComponents(0, 0);
-	FIntPoint MaxSectionBase(0, 0);
-	FIntPoint MinSectionBase(0, 0);
+	FIntPoint MaxSectionBase(TNumericLimits<int32>::Min(), TNumericLimits<int32>::Min());
+	FIntPoint MinSectionBase(TNumericLimits<int32>::Max(), TNumericLimits<int32>::Max());
 
 	Info->ForAllLandscapeProxies([&MaxSectionBase, &MinSectionBase](ALandscapeProxy* Proxy)
 	{
@@ -1731,8 +1737,11 @@ FIntPoint ALandscape::ComputeComponentCounts() const
 		}
 	});
 
-	NumComponents.X = ((MaxSectionBase.X - MinSectionBase.X) / ComponentSizeQuads) + 1;
-	NumComponents.Y = ((MaxSectionBase.Y - MinSectionBase.Y) / ComponentSizeQuads) + 1;
+	if ((MaxSectionBase.X >= MinSectionBase.X) && (MaxSectionBase.Y >= MinSectionBase.Y))
+	{
+		NumComponents.X = ((MaxSectionBase.X - MinSectionBase.X) / ComponentSizeQuads) + 1;
+		NumComponents.Y = ((MaxSectionBase.Y - MinSectionBase.Y) / ComponentSizeQuads) + 1;
+	}
 
 	return NumComponents;
 }
