When destroying a FLandscapeRenderSystem, ensure there are no outstanding tasks

diff --git a/Engine/Source/Runtime/Landscape/Private/LandscapeRender.cpp b/Engine/Source/Runtime/Landscape/Private/LandscapeRender.cpp
--- a/Engine/Source/Runtime/Landscape/Private/LandscapeRender.cpp
+++ b/Engine/Source/Runtime/Landscape/Private/LandscapeRender.cpp
@@ -804,6 +804,8 @@ FLandscapeRenderSystem::FLandscapeRenderSystem()
 
 FLandscapeRenderSystem::~FLandscapeRenderSystem()
 {
+	WaitForTasksCompletion();
+
 	check(LandscapeIndexAllocator[LandscapeIndex]);
 	LandscapeIndexAllocator[LandscapeIndex] = false;
 
@@ -922,16 +924,7 @@ void FLandscapeRenderSystem::BeginRender()
 
 	TRACE_CPUPROFILER_EVENT_SCOPE(FLandscapeRenderSystem::BeginRender());
 
-	if (FetchHeightmapLODBiasesEventRef.IsValid())
-	{
-		FTaskGraphInterface::Get().WaitUntilTaskCompletes(FetchHeightmapLODBiasesEventRef, ENamedThreads::GetRenderThread_Local());
-		FetchHeightmapLODBiasesEventRef.SafeRelease();
-	}
-
-	for (auto& Pair : PerViewParametersTasks)
-	{
-		FTaskGraphInterface::Get().WaitUntilTaskCompletes(PerViewParametersTasks[Pair.Key], ENamedThreads::GetRenderThread_Local());
-	}
+	WaitForTasksCompletion();
 
 	UpdateBuffers();
 }
@@ -1067,7 +1060,11 @@ void FLandscapeRenderSystem::EndFrame()
 {
 	check(IsInRenderingThread());
 
-	// Finalize any outstanding jobs before ~FSceneRenderer() so we don't have corrupted accesses
+	WaitForTasksCompletion();
+}
+
+void FLandscapeRenderSystem::WaitForTasksCompletion()
+{
 	if (FetchHeightmapLODBiasesEventRef.IsValid())
 	{
 		FTaskGraphInterface::Get().WaitUntilTaskCompletes(FetchHeightmapLODBiasesEventRef, ENamedThreads::GetRenderThread_Local());
diff --git a/Engine/Source/Runtime/Landscape/Public/LandscapeRender.h b/Engine/Source/Runtime/Landscape/Public/LandscapeRender.h
--- a/Engine/Source/Runtime/Landscape/Public/LandscapeRender.h
+++ b/Engine/Source/Runtime/Landscape/Public/LandscapeRender.h
@@ -568,6 +568,8 @@ struct FLandscapeRenderSystem
 	void UpdateBuffers();
 
 	void EndFrame();
+
+	void WaitForTasksCompletion();
 };
 
 LANDSCAPE_API extern TMap<FLandscapeNeighborInfo::FLandscapeKey, FLandscapeRenderSystem*> LandscapeRenderSystems;
