Landscape materials fixed for cooked packages opened in editor

diff --git a/Engine/Source/Runtime/Landscape/Private/Landscape.cpp b/Engine/Source/Runtime/Landscape/Private/Landscape.cpp
--- a/Engine/Source/Runtime/Landscape/Private/Landscape.cpp
+++ b/Engine/Source/Runtime/Landscape/Private/Landscape.cpp
@@ -2565,15 +2565,24 @@ void ALandscapeProxy::PostLoad()
 			if (UMaterialInstance* MaterialInstance = Comp->GetMaterialInstance(0, false))
 			{
 				UMaterialInstanceConstant* CombinationMaterialInstance = Cast<UMaterialInstanceConstant>(MaterialInstance->Parent);
-				if (Comp->ValidateCombinationMaterial(CombinationMaterialInstance))
+				// Only validate if uncooked
+				if (!Comp->GetOutermost()->HasAnyPackageFlags(PKG_FilterEditorOnly))
 				{
-					MaterialInstanceConstantMap.Add(*ULandscapeComponent::GetLayerAllocationKey(Comp->GetWeightmapLayerAllocations(), CombinationMaterialInstance->Parent), CombinationMaterialInstance);
+					if (Comp->ValidateCombinationMaterial(CombinationMaterialInstance))
+					{
+						MaterialInstanceConstantMap.Add(*ULandscapeComponent::GetLayerAllocationKey(Comp->GetWeightmapLayerAllocations(), CombinationMaterialInstance->Parent), CombinationMaterialInstance);
+					}
+					else
+					{
+						// There was a problem with the loaded material : it doesn't match the expected material combination, we need to regenerate the material instances : 
+						Comp->UpdateMaterialInstances();
+						bFixedUpInvalidMaterialInstances = true;
+					}
 				}
-				else
+				else if (CombinationMaterialInstance)
 				{
-					// There was a problem with the loaded material : it doesn't match the expected material combination, we need to regenerate the material instances : 
-					Comp->UpdateMaterialInstances();
-					bFixedUpInvalidMaterialInstances = true;
+					// Skip ValidateCombinationMaterial
+					MaterialInstanceConstantMap.Add(*ULandscapeComponent::GetLayerAllocationKey(Comp->GetWeightmapLayerAllocations(), CombinationMaterialInstance->Parent), CombinationMaterialInstance);
 				}
 			}
 		}
