Added const accessors to landscape spline data

diff --git a/Engine/Source/Runtime/Landscape/Classes/LandscapeSplinesComponent.h b/Engine/Source/Runtime/Landscape/Classes/LandscapeSplinesComponent.h
--- a/Engine/Source/Runtime/Landscape/Classes/LandscapeSplinesComponent.h
+++ b/Engine/Source/Runtime/Landscape/Classes/LandscapeSplinesComponent.h
@@ -156,6 +156,9 @@ public:
 	void CheckSplinesValid();
 	bool ModifySplines(bool bAlwaysMarkDirty = true);
 
+	const TArray<TObjectPtr<ULandscapeSplineControlPoint>>& GetControlPoints() const { return ControlPoints; }
+	const TArray<TObjectPtr<ULandscapeSplineSegment>>& GetSegments() const { return Segments; }
+
 #if WITH_EDITOR
 	void RequestSplineLayerUpdate();
 
