Prevent landscape material instances validation/recompilation if not in the editor or a commandlet (we cannot re-build material instance constants otherwise)

diff --git a/Engine/Source/Runtime/Landscape/Private/Landscape.cpp b/Engine/Source/Runtime/Landscape/Private/Landscape.cpp
--- a/Engine/Source/Runtime/Landscape/Private/Landscape.cpp
+++ b/Engine/Source/Runtime/Landscape/Private/Landscape.cpp
@@ -2565,8 +2565,8 @@ void ALandscapeProxy::PostLoad()
 			if (UMaterialInstance* MaterialInstance = Comp->GetMaterialInstance(0, false))
 			{
 				UMaterialInstanceConstant* CombinationMaterialInstance = Cast<UMaterialInstanceConstant>(MaterialInstance->Parent);
-				// Only validate if uncooked
-				if (!Comp->GetOutermost()->HasAnyPackageFlags(PKG_FilterEditorOnly))
+				// Only validate if uncooked and in the editor/commandlet mode (we cannot re-build material instance constants if this is not the case : see UMaterialInstance::CacheResourceShadersForRendering, which is only called if FApp::CanEverRender() returns true) 
+				if (!Comp->GetOutermost()->HasAnyPackageFlags(PKG_FilterEditorOnly) && (GIsEditor && FApp::CanEverRender()))
 				{
 					if (Comp->ValidateCombinationMaterial(CombinationMaterialInstance))
 					{
