Revert 18394343 to fix issue where landscape RVT isn't registered correctly in PIE leading to black RVT landscapes in WP worlds.

diff --git a/Engine/Source/Runtime/Landscape/Private/LandscapeRender.cpp b/Engine/Source/Runtime/Landscape/Private/LandscapeRender.cpp
--- a/Engine/Source/Runtime/Landscape/Private/LandscapeRender.cpp
+++ b/Engine/Source/Runtime/Landscape/Private/LandscapeRender.cpp
@@ -1497,13 +1497,11 @@ void FLandscapeComponentSceneProxy::CreateRenderThreadResources()
 
 	check(HeightmapTexture != nullptr);
 
-	if (!VisibilityHelper.ShouldBeVisible())
+	if (VisibilityHelper.ShouldBeVisible())
 	{
-		return;
+		RegisterNeighbors(this);
 	}
 
-	RegisterNeighbors(this);
-
 	auto FeatureLevel = GetScene().GetFeatureLevel();
 
 	SharedBuffers = FLandscapeComponentSceneProxy::SharedBuffersMap.FindRef(SharedBuffersKey);
@@ -1676,7 +1674,7 @@ bool FLandscapeComponentSceneProxy::OnLevelAddedToWorld_RenderThread()
 	if (VisibilityHelper.OnAddedToWorld())
 	{
 		SetForceHidden(false);
-		CreateRenderThreadResources();
+		RegisterNeighbors(this);
 		return true;
 	}
 	return false;
@@ -2193,11 +2191,6 @@ void FLandscapeComponentSceneProxy::DrawStaticElements(FStaticPrimitiveDrawInter
 		return;
 	}
 
-	if (!VisibilityHelper.ShouldBeVisible())
-	{
-		return;
-	}
-
 	int32 TotalBatchCount = 1 + LastLOD - FirstLOD;
 	TotalBatchCount += (1 + LastVirtualTextureLOD - FirstVirtualTextureLOD) * RuntimeVirtualTextureMaterialTypes.Num();
 
