Fix FLandscapeComponentSceneProxy::CreateRenderResources() trying to access it's associated landscape render system when no neighbors have been registered.

diff --git a/Engine/Source/Runtime/Landscape/Private/LandscapeRender.cpp b/Engine/Source/Runtime/Landscape/Private/LandscapeRender.cpp
--- a/Engine/Source/Runtime/Landscape/Private/LandscapeRender.cpp
+++ b/Engine/Source/Runtime/Landscape/Private/LandscapeRender.cpp
@@ -776,6 +776,18 @@ IMPLEMENT_GLOBAL_SHADER_PARAMETER_STRUCT(FLandscapeSectionLODUniformParameters,
 
 TMap<FLandscapeNeighborInfo::FLandscapeKey, FLandscapeRenderSystem*> LandscapeRenderSystems;
 
+FLandscapeRenderSystem& GetOrCreateLandscapeRenderSystem(FLandscapeNeighborInfo::FLandscapeKey LandscapeKey)
+{
+	check(IsInRenderingThread());
+
+	FLandscapeRenderSystem*& LandscapeRenderSystem = LandscapeRenderSystems.FindOrAdd(LandscapeKey);
+	if (!LandscapeRenderSystem)
+	{
+		LandscapeRenderSystem = new FLandscapeRenderSystem();
+	}
+	return *LandscapeRenderSystem;
+}
+
 TBitArray<> FLandscapeRenderSystem::LandscapeIndexAllocator;
 
 
@@ -1497,6 +1509,8 @@ void FLandscapeComponentSceneProxy::CreateRenderThreadResources()
 
 	check(HeightmapTexture != nullptr);
 
+	const FLandscapeRenderSystem& RenderSystem = GetOrCreateLandscapeRenderSystem(LandscapeKey);	
+
 	if (VisibilityHelper.ShouldBeVisible())
 	{
 		RegisterNeighbors(this);
@@ -1592,7 +1606,6 @@ void FLandscapeComponentSceneProxy::CreateRenderThreadResources()
 		GrassMeshBatch.Type = PT_PointList;
 		GrassMeshBatch.DepthPriorityGroup = SDPG_World;
 
-		const FLandscapeRenderSystem& RenderSystem = *LandscapeRenderSystems.FindChecked(LandscapeKey);
 		const TUniformBufferRef<FLandscapeSectionLODUniformParameters>& LandscapeSectionLODUniformParameters = RenderSystem.UniformBuffer;
 
 		// Combined grass rendering batch element
@@ -4128,10 +4141,7 @@ void FLandscapeNeighborInfo::RegisterNeighbors(FLandscapeComponentSceneProxy* Sc
 	check(IsInRenderingThread());
 	if (!bRegistered)
 	{
-		if (!SharedSceneProxyMap.Find(LandscapeKey))
-		{
-			LandscapeRenderSystems.Add(LandscapeKey, new FLandscapeRenderSystem {});
-		}
+		FLandscapeRenderSystem& RenderSystem = GetOrCreateLandscapeRenderSystem(LandscapeKey);
 
 		// Register ourselves in the map.
 		TMap<FIntPoint, const FLandscapeNeighborInfo*>& SceneProxyMap = SharedSceneProxyMap.FindOrAdd(LandscapeKey);
@@ -4168,7 +4178,6 @@ void FLandscapeNeighborInfo::RegisterNeighbors(FLandscapeComponentSceneProxy* Sc
 
 			if (SceneProxy != nullptr)
 			{
-				FLandscapeRenderSystem& RenderSystem = *LandscapeRenderSystems.FindChecked(LandscapeKey);
 				RenderSystem.RegisterEntity(SceneProxy);
 			}
 		}
