Changing ShouldCacheLandscapeGrassShaders to use EShaderPermutationFlags::HasEditorOnlyData.

diff --git a/Engine/Source/Runtime/Landscape/Private/LandscapeGrass.cpp b/Engine/Source/Runtime/Landscape/Private/LandscapeGrass.cpp
--- a/Engine/Source/Runtime/Landscape/Private/LandscapeGrass.cpp
+++ b/Engine/Source/Runtime/Landscape/Private/LandscapeGrass.cpp
@@ -242,8 +242,8 @@ static bool ShouldCacheLandscapeGrassShaders(const FMeshMaterialShaderPermutatio
 	// We only need grass weight shaders for Landscape vertex factories on desktop platforms
 	return (Parameters.MaterialParameters.bIsUsedWithLandscape || Parameters.MaterialParameters.bIsSpecialEngineMaterial) &&
 		IsFeatureLevelSupported(Parameters.Platform, ERHIFeatureLevel::SM5) &&
-		Parameters.VertexFactoryType == FindVertexFactoryType(FName(TEXT("FLandscapeFixedGridVertexFactory"), FNAME_Find)) && 
-		!IsConsolePlatform(Parameters.Platform) && !IsVulkanMobileSM5Platform(Parameters.Platform) && !IsMetalMobileSM5Platform(Parameters.Platform);
+		Parameters.VertexFactoryType == FindVertexFactoryType(FName(TEXT("FLandscapeFixedGridVertexFactory"), FNAME_Find)) &&
+		EnumHasAllFlags(Parameters.Flags, EShaderPermutationFlags::HasEditorOnlyData);
 }
 
 class FLandscapeGrassWeightShaderElementData : public FMeshMaterialShaderElementData
