World Partition Landscape Support
- Don't maintain a list of landscape proxy and spline handles from actor descriptors registration.
- Iterate through actor descriptors to gather landscape edges and load intersecting splines instead.

diff --git a/Engine/Source/Runtime/Engine/Private/WorldPartition/Landscape/LandscapeActorDesc.cpp b/Engine/Source/Runtime/Engine/Private/WorldPartition/Landscape/LandscapeActorDesc.cpp
--- a/Engine/Source/Runtime/Engine/Private/WorldPartition/Landscape/LandscapeActorDesc.cpp
+++ b/Engine/Source/Runtime/Engine/Private/WorldPartition/Landscape/LandscapeActorDesc.cpp
@@ -41,30 +41,4 @@ void FLandscapeActorDesc::Unload()
 
 	FPartitionActorDesc::Unload();
 }
-
-void FLandscapeActorDesc::OnRegister(UWorld* InWorld)
-{
-	FPartitionActorDesc::OnRegister(InWorld);
-
-	if (ULandscapeInfo* LandscapeInfo = ULandscapeInfo::FindOrCreate(Container->GetWorld(), GridGuid))
-	{
-		FWorldPartitionHandle Handle(Container, GetGuid());
-		check(Handle.IsValid());
-		LandscapeInfo->ProxyHandles.Add(MoveTemp(Handle));
-	}
-}
-
-void FLandscapeActorDesc::OnUnregister()
-{
-	FPartitionActorDesc::OnUnregister();
-
-	if (ULandscapeInfo* LandscapeInfo = ULandscapeInfo::Find(Container->GetWorld(), GridGuid))
-	{
-		FWorldPartitionHandle Handle(Container, GetGuid());
-		if (Handle.IsValid())
-		{
-			LandscapeInfo->ProxyHandles.Remove(Handle);
-		}
-	}
-}
 #endif
diff --git a/Engine/Source/Runtime/Engine/Private/WorldPartition/Landscape/LandscapeSplineActorDesc.cpp b/Engine/Source/Runtime/Engine/Private/WorldPartition/Landscape/LandscapeSplineActorDesc.cpp
--- a/Engine/Source/Runtime/Engine/Private/WorldPartition/Landscape/LandscapeSplineActorDesc.cpp
+++ b/Engine/Source/Runtime/Engine/Private/WorldPartition/Landscape/LandscapeSplineActorDesc.cpp
@@ -40,30 +40,4 @@ void FLandscapeSplineActorDesc::Serialize(FArchive& Ar)
 		Ar << LandscapeGuid;
 	}
 }
-
-void FLandscapeSplineActorDesc::OnRegister(UWorld* InWorld)
-{
-	FWorldPartitionActorDesc::OnRegister(InWorld);
-
-	if (ULandscapeInfo* LandscapeInfo = ULandscapeInfo::FindOrCreate(Container->GetWorld(), LandscapeGuid))
-	{
-		FWorldPartitionHandle Handle(Container, GetGuid());
-		check(Handle.IsValid());
-		LandscapeInfo->SplineHandles.Add(MoveTemp(Handle));
-	}
-}
-
-void FLandscapeSplineActorDesc::OnUnregister()
-{
-	FWorldPartitionActorDesc::OnUnregister();
-
-	if (ULandscapeInfo* LandscapeInfo = ULandscapeInfo::Find(Container->GetWorld(), LandscapeGuid))
-	{
-		FWorldPartitionHandle Handle(Container, GetGuid());
-		if (Handle.IsValid())
-		{
-			LandscapeInfo->SplineHandles.Remove(Handle);
-		}
-	}
-}
 #endif
diff --git a/Engine/Source/Runtime/Engine/Public/WorldPartition/Landscape/LandscapeActorDesc.h b/Engine/Source/Runtime/Engine/Public/WorldPartition/Landscape/LandscapeActorDesc.h
--- a/Engine/Source/Runtime/Engine/Public/WorldPartition/Landscape/LandscapeActorDesc.h
+++ b/Engine/Source/Runtime/Engine/Public/WorldPartition/Landscape/LandscapeActorDesc.h
@@ -17,8 +17,6 @@ public:
 	virtual void Unload() override;
 
 protected:
-	virtual void OnRegister(UWorld* InWorld) override;
-	virtual void OnUnregister() override;
 	virtual void Serialize(FArchive& Ar) override;
 #endif
 };
diff --git a/Engine/Source/Runtime/Engine/Public/WorldPartition/Landscape/LandscapeSplineActorDesc.h b/Engine/Source/Runtime/Engine/Public/WorldPartition/Landscape/LandscapeSplineActorDesc.h
--- a/Engine/Source/Runtime/Engine/Public/WorldPartition/Landscape/LandscapeSplineActorDesc.h
+++ b/Engine/Source/Runtime/Engine/Public/WorldPartition/Landscape/LandscapeSplineActorDesc.h
@@ -19,7 +19,5 @@ protected:
 	virtual void Init(const AActor* InActor) override;
 	virtual bool Equals(const FWorldPartitionActorDesc* Other) const override;
 	virtual void Serialize(FArchive& Ar) override;
-	virtual void OnRegister(UWorld* InWorld) override;
-	virtual void OnUnregister() override;
 #endif
 };
diff --git a/Engine/Source/Runtime/Landscape/Classes/LandscapeInfo.h b/Engine/Source/Runtime/Landscape/Classes/LandscapeInfo.h
--- a/Engine/Source/Runtime/Landscape/Classes/LandscapeInfo.h
+++ b/Engine/Source/Runtime/Landscape/Classes/LandscapeInfo.h
@@ -149,12 +149,6 @@ private:
 #if WITH_EDITORONLY_DATA
 	UPROPERTY()
 	TArray<TScriptInterface<ILandscapeSplineInterface>> SplineActors;
-
-	// World Partition handles to actors with the same LandscapeGuid
-	friend class FLandscapeActorDesc;
-	TSet<FWorldPartitionHandle> ProxyHandles;
-	friend class FLandscapeSplineActorDesc;
-	TSet<FWorldPartitionHandle> SplineHandles;
 #endif
 
 	TSet<ULandscapeComponent*> SelectedComponents;
@@ -302,9 +296,6 @@ public:
 
 	/** Returns LandscapeStreamingProxy Cell Size in WorldPartition */
 	LANDSCAPE_API uint32 GetGridSize(uint32 InGridSizeInComponents) const;
-
-	/** Returns LandscapeSplineActor WorldPartition handles */
-	const TSet<FWorldPartitionHandle>& GetSplineHandles() const { return SplineHandles; }
 #endif
 	LANDSCAPE_API static ULandscapeInfo* Find(UWorld* InWorld, const FGuid& LandscapeGuid);
 	LANDSCAPE_API static ULandscapeInfo* FindOrCreate(UWorld* InWorld, const FGuid& LandscapeGuid);
diff --git a/Engine/Source/Runtime/Landscape/Private/LandscapeEdit.cpp b/Engine/Source/Runtime/Landscape/Private/LandscapeEdit.cpp
--- a/Engine/Source/Runtime/Landscape/Private/LandscapeEdit.cpp
+++ b/Engine/Source/Runtime/Landscape/Private/LandscapeEdit.cpp
@@ -60,8 +60,11 @@ LandscapeEdit.cpp: Landscape editing
 #include "ScopedTransaction.h"
 #include "Editor.h"
 #include "PhysicalMaterials/PhysicalMaterial.h"
+#include "WorldPartition/WorldPartition.h"
+#include "WorldPartition/WorldPartitionHelpers.h"
 #include "WorldPartition/WorldPartitionActorDesc.h"
 #include "WorldPartition/Landscape/LandscapeActorDesc.h"
+#include "WorldPartition/Landscape/LandscapeSplineActorDesc.h"
 #include "ActorPartition/ActorPartitionSubsystem.h"
 #include "LandscapeSplineActor.h"
 #endif
@@ -2384,36 +2387,46 @@ bool ULandscapeInfo::AreAllComponentsRegistered() const
 
 bool ULandscapeInfo::HasUnloadedComponentsInRegion(int32 X1, int32 Y1, int32 X2, int32 Y2) const
 {
-	if (!LandscapeActor)
-	{
-		return false;
-	}
+	bool bResult = false;
 
-	UWorld* World = LandscapeActor->GetWorld();
+	if (LandscapeActor)
+	{
+		UWorld* World = LandscapeActor->GetWorld();
 
-	int32 ComponentIndexX1, ComponentIndexY1, ComponentIndexX2, ComponentIndexY2;
-	ALandscape::CalcComponentIndicesOverlap(X1, Y1, X2, Y2, ComponentSizeQuads, ComponentIndexX1, ComponentIndexY1, ComponentIndexX2, ComponentIndexY2);
+		int32 ComponentIndexX1, ComponentIndexY1, ComponentIndexX2, ComponentIndexY2;
+		ALandscape::CalcComponentIndicesOverlap(X1, Y1, X2, Y2, ComponentSizeQuads, ComponentIndexX1, ComponentIndexY1, ComponentIndexX2, ComponentIndexY2);
 
-	UActorPartitionSubsystem::FCellCoord MinCoord = UActorPartitionSubsystem::FCellCoord::GetCellCoord(FIntPoint(ComponentIndexX1 * ComponentSizeQuads, ComponentIndexY1 * ComponentSizeQuads), World->PersistentLevel, LandscapeActor->GridSize);
-	UActorPartitionSubsystem::FCellCoord MaxCoord = UActorPartitionSubsystem::FCellCoord::GetCellCoord(FIntPoint(ComponentIndexX2 * ComponentSizeQuads, ComponentIndexY2 * ComponentSizeQuads), World->PersistentLevel, LandscapeActor->GridSize);
+		const UActorPartitionSubsystem::FCellCoord MinCoord = UActorPartitionSubsystem::FCellCoord::GetCellCoord(FIntPoint(ComponentIndexX1 * ComponentSizeQuads, ComponentIndexY1 * ComponentSizeQuads), World->PersistentLevel, LandscapeActor->GridSize);
+		const UActorPartitionSubsystem::FCellCoord MaxCoord = UActorPartitionSubsystem::FCellCoord::GetCellCoord(FIntPoint(ComponentIndexX2 * ComponentSizeQuads, ComponentIndexY2 * ComponentSizeQuads), World->PersistentLevel, LandscapeActor->GridSize);
 
-	for (const FWorldPartitionHandle& Handle : ProxyHandles)
-	{
-		if (FLandscapeActorDesc* LandscapeActorDesc = (FLandscapeActorDesc*)Handle.Get())
+		if (World)
 		{
-			check(LandscapeActorDesc->GridGuid == LandscapeGuid);
-			UActorPartitionSubsystem::FCellCoord ActorCoord(LandscapeActorDesc->GridIndexX, LandscapeActorDesc->GridIndexY, LandscapeActorDesc->GridIndexZ, World->PersistentLevel);
-			if (ActorCoord.X >= MinCoord.X && ActorCoord.Y >= MinCoord.Y && ActorCoord.X <= MaxCoord.X && ActorCoord.Y <= MaxCoord.Y)
+			if (UWorldPartition* WorldPartition = World->GetWorldPartition())
 			{
-				if (!Handle.IsLoaded())
+				FWorldPartitionHelpers::ForEachActorDesc<ALandscapeProxy>(WorldPartition, [this, World, &MinCoord, &MaxCoord, &bResult](const FWorldPartitionActorDesc* ActorDesc)
 				{
+					FLandscapeActorDesc* LandscapeActorDesc = (FLandscapeActorDesc*)ActorDesc;
+
+					if (LandscapeActorDesc->GridGuid == LandscapeGuid)
+					{
+						const UActorPartitionSubsystem::FCellCoord ActorCoord(LandscapeActorDesc->GridIndexX, LandscapeActorDesc->GridIndexY, LandscapeActorDesc->GridIndexZ, World->PersistentLevel);
+						if (ActorCoord.X >= MinCoord.X && ActorCoord.Y >= MinCoord.Y && ActorCoord.X <= MaxCoord.X && ActorCoord.Y <= MaxCoord.Y)
+						{
+							if (!LandscapeActorDesc->IsLoaded())
+							{
+								bResult = true;
+								return false;
+							}
+						}
+					}
+
 					return true;
-				}
+				});
 			}
 		}
 	}
 
-	return false;
+	return bResult;
 }
 
 void ULandscapeInfo::GetComponentsInRegion(int32 X1, int32 Y1, int32 X2, int32 Y2, TSet<ULandscapeComponent*>& OutComponents, bool bOverlap) const
@@ -5103,19 +5116,24 @@ void ALandscapeStreamingProxy::PostRegisterAllComponents()
 		check(LandscapeInfo);
 		if (GEditor && !GetWorld()->IsGameWorld())
 		{
-			const TSet<FWorldPartitionHandle>& SplineHandles = LandscapeInfo->GetSplineHandles();
-			if (SplineHandles.Num())
+			if (UWorldPartition* WorldPartition = GetWorld()->GetWorldPartition())
 			{
-				FVector ActorLocation = GetActorLocation();
-				FBox Bounds(ActorLocation, ActorLocation + (GridSize * LandscapeInfo->DrawScale));
-				// Get a reference to Spline Actors that have intersecting bounds with us
-				for (const FWorldPartitionHandle& SplineHandle : LandscapeInfo->GetSplineHandles())
+				const FVector ActorLocation = GetActorLocation();
+				const FBox Bounds(ActorLocation, ActorLocation + (GridSize * LandscapeInfo->DrawScale));
+
+				FWorldPartitionHelpers::ForEachActorDesc<ALandscapeSplineActor>(WorldPartition, [this, WorldPartition, &Bounds](const FWorldPartitionActorDesc* ActorDesc) mutable
 				{
-					if (SplineHandle.IsValid() && Bounds.IntersectXY(SplineHandle->GetBounds()))
+					FLandscapeSplineActorDesc* LandscapeSplineActorDesc = (FLandscapeSplineActorDesc*)ActorDesc;
+
+					if (LandscapeSplineActorDesc->LandscapeGuid == LandscapeGuid)
 					{
-						ActorDescReferences.Add(SplineHandle);
+						if (Bounds.IntersectXY(ActorDesc->GetBounds()))
+						{
+							ActorDescReferences.Add(FWorldPartitionReference(WorldPartition, ActorDesc->GetGuid()));
+						}
 					}
-				}
+					return true;
+				});
 			}
 		}
 	}
@@ -5168,23 +5186,35 @@ bool ULandscapeInfo::CanDeleteLandscape(FText& OutReason) const
 	}
 
 	// Then check for Unloaded Proxies
-	for (const FWorldPartitionHandle& Handle : ProxyHandles)
+	if (AActor* Actor = LandscapeActor.Get())
 	{
-		ALandscapeProxy* LandscapeProxy = Cast<ALandscapeProxy>(Handle->GetActor());
-		if (LandscapeProxy == LandscapeActor)
-		{
-			continue;
-		}
-		
-		// If LandscapeProxy is null then it is not loaded so not deleted.
-		if (!LandscapeProxy)
-		{
-			++UndeletedProxyCount;
-		}
-		else
+		UWorld* World = Actor->GetWorld();
+		if (UWorldPartition* WorldPartition = World->GetWorldPartition())
 		{
-			// If Actor is loaded it should be Registered and not pending kill (already accounted for) or pending kill (deleted)
-			check(Proxies.Contains(LandscapeProxy) == IsValidChecked(LandscapeProxy));
+			FWorldPartitionHelpers::ForEachActorDesc<ALandscapeProxy>(WorldPartition, [this, &UndeletedProxyCount](const FWorldPartitionActorDesc* ActorDesc)
+			{
+				FLandscapeActorDesc* LandscapeActorDesc = (FLandscapeActorDesc*)ActorDesc;
+
+				if (LandscapeActorDesc->GridGuid == LandscapeGuid)
+				{
+					ALandscapeProxy* LandscapeProxy = Cast<ALandscapeProxy>(ActorDesc->GetActor());
+					if (LandscapeProxy != LandscapeActor)
+					{
+						// If LandscapeProxy is null then it is not loaded so not deleted.
+						if (!LandscapeProxy)
+						{
+							++UndeletedProxyCount;
+						}
+						else
+						{
+							// If Actor is loaded it should be Registered and not pending kill (already accounted for) or pending kill (deleted)
+							check(Proxies.Contains(LandscapeProxy) == IsValidChecked(LandscapeProxy));
+						}
+					}
+				}
+
+				return true;
+			});
 		}
 	}
 
@@ -5200,19 +5230,33 @@ bool ULandscapeInfo::CanDeleteLandscape(FText& OutReason) const
 	}
 
 	// Then check for Unloaded Splines
-	for (const FWorldPartitionHandle& Handle : SplineHandles)
+	if (AActor* Actor = LandscapeActor.Get())
 	{
-		ALandscapeSplineActor* SplineActor = Cast<ALandscapeSplineActor>(Handle->GetActor());
-		
-		// If SplineActor is null then it is not loaded/deleted. If it's loaded then it needs to be pending kill.
-		if (!SplineActor)
+		UWorld* World = Actor->GetWorld();
+		if (UWorldPartition* WorldPartition = World->GetWorldPartition())
 		{
-			++UndeletedSplineCount;
-		}
-		else
-		{
-			// If Actor is loaded it should be Registered and not pending kill (already accounted for) or pending kill (deleted)
-			check(SplineActors.Contains(SplineActor) == IsValidChecked(SplineActor));
+			FWorldPartitionHelpers::ForEachActorDesc<ALandscapeSplineActor>(WorldPartition, [this, &UndeletedSplineCount](const FWorldPartitionActorDesc* ActorDesc)
+			{
+				FLandscapeSplineActorDesc* LandscapeSplineActorDesc = (FLandscapeSplineActorDesc*)ActorDesc;
+
+				if (LandscapeSplineActorDesc->LandscapeGuid == LandscapeGuid)
+				{
+					ALandscapeSplineActor* SplineActor = Cast<ALandscapeSplineActor>(LandscapeSplineActorDesc->GetActor());
+		
+					// If SplineActor is null then it is not loaded/deleted. If it's loaded then it needs to be pending kill.
+					if (!SplineActor)
+					{
+						++UndeletedSplineCount;
+					}
+					else
+					{
+						// If Actor is loaded it should be Registered and not pending kill (already accounted for) or pending kill (deleted)
+						check(SplineActors.Contains(SplineActor) == IsValidChecked(SplineActor));
+					}
+				}
+
+				return true;
+			});
 		}
 	}
 
