Fix Static Analysis Win64.

diff --git a/Engine/Source/Runtime/Landscape/Private/LandscapeEdit.cpp b/Engine/Source/Runtime/Landscape/Private/LandscapeEdit.cpp
--- a/Engine/Source/Runtime/Landscape/Private/LandscapeEdit.cpp
+++ b/Engine/Source/Runtime/Landscape/Private/LandscapeEdit.cpp
@@ -2399,30 +2399,27 @@ bool ULandscapeInfo::HasUnloadedComponentsInRegion(int32 X1, int32 Y1, int32 X2,
 		const UActorPartitionSubsystem::FCellCoord MinCoord = UActorPartitionSubsystem::FCellCoord::GetCellCoord(FIntPoint(ComponentIndexX1 * ComponentSizeQuads, ComponentIndexY1 * ComponentSizeQuads), World->PersistentLevel, LandscapeActor->GridSize);
 		const UActorPartitionSubsystem::FCellCoord MaxCoord = UActorPartitionSubsystem::FCellCoord::GetCellCoord(FIntPoint(ComponentIndexX2 * ComponentSizeQuads, ComponentIndexY2 * ComponentSizeQuads), World->PersistentLevel, LandscapeActor->GridSize);
 
-		if (World)
+		if (UWorldPartition* WorldPartition = World->GetWorldPartition())
 		{
-			if (UWorldPartition* WorldPartition = World->GetWorldPartition())
+			FWorldPartitionHelpers::ForEachActorDesc<ALandscapeProxy>(WorldPartition, [this, World, &MinCoord, &MaxCoord, &bResult](const FWorldPartitionActorDesc* ActorDesc)
 			{
-				FWorldPartitionHelpers::ForEachActorDesc<ALandscapeProxy>(WorldPartition, [this, World, &MinCoord, &MaxCoord, &bResult](const FWorldPartitionActorDesc* ActorDesc)
-				{
-					FLandscapeActorDesc* LandscapeActorDesc = (FLandscapeActorDesc*)ActorDesc;
+				FLandscapeActorDesc* LandscapeActorDesc = (FLandscapeActorDesc*)ActorDesc;
 
-					if (LandscapeActorDesc->GridGuid == LandscapeGuid)
+				if (LandscapeActorDesc->GridGuid == LandscapeGuid)
+				{
+					const UActorPartitionSubsystem::FCellCoord ActorCoord(LandscapeActorDesc->GridIndexX, LandscapeActorDesc->GridIndexY, LandscapeActorDesc->GridIndexZ, World->PersistentLevel);
+					if (ActorCoord.X >= MinCoord.X && ActorCoord.Y >= MinCoord.Y && ActorCoord.X <= MaxCoord.X && ActorCoord.Y <= MaxCoord.Y)
 					{
-						const UActorPartitionSubsystem::FCellCoord ActorCoord(LandscapeActorDesc->GridIndexX, LandscapeActorDesc->GridIndexY, LandscapeActorDesc->GridIndexZ, World->PersistentLevel);
-						if (ActorCoord.X >= MinCoord.X && ActorCoord.Y >= MinCoord.Y && ActorCoord.X <= MaxCoord.X && ActorCoord.Y <= MaxCoord.Y)
+						if (!LandscapeActorDesc->IsLoaded())
 						{
-							if (!LandscapeActorDesc->IsLoaded())
-							{
-								bResult = true;
-								return false;
-							}
+							bResult = true;
+							return false;
 						}
 					}
+				}
 
-					return true;
-				});
-			}
+				return true;
+			});
 		}
 	}
 
