RHI Crash in Physics Sandbox Streaming_Parent level

Problematic UB is now a FUniformBufferRHIRef (refcounted)

diff --git a/Engine/Source/Runtime/Landscape/Private/LandscapeRender.cpp b/Engine/Source/Runtime/Landscape/Private/LandscapeRender.cpp
--- a/Engine/Source/Runtime/Landscape/Private/LandscapeRender.cpp
+++ b/Engine/Source/Runtime/Landscape/Private/LandscapeRender.cpp
@@ -811,7 +811,7 @@ FLandscapeRenderSystem::FLandscapeRenderSystem()
 	Parameters.LandscapeIndex = LandscapeIndex;
 	Parameters.Size = FIntPoint(1, 1);
 	Parameters.SectionLODBias = GWhiteVertexBufferWithSRV->ShaderResourceViewRHI;
-	UniformBuffer = TUniformBufferRef<FLandscapeSectionLODUniformParameters>::CreateUniformBufferImmediate(Parameters, UniformBuffer_MultiFrame);
+	SectionLODUniformBuffer = TUniformBufferRef<FLandscapeSectionLODUniformParameters>::CreateUniformBufferImmediate(Parameters, UniformBuffer_MultiFrame);
 }
 
 FLandscapeRenderSystem::~FLandscapeRenderSystem()
@@ -1042,7 +1042,7 @@ void FLandscapeRenderSystem::UpdateBuffers()
 			Parameters.Size = Size;
 			Parameters.SectionLODBias = SectionLODBiasSRV;
 
-			UniformBuffer.UpdateUniformBufferImmediate(Parameters);
+			RHIUpdateUniformBuffer(SectionLODUniformBuffer, &Parameters);
 		}
 	}
 }
@@ -1606,14 +1606,12 @@ void FLandscapeComponentSceneProxy::CreateRenderThreadResources()
 		GrassMeshBatch.Type = PT_PointList;
 		GrassMeshBatch.DepthPriorityGroup = SDPG_World;
 
-		const TUniformBufferRef<FLandscapeSectionLODUniformParameters>& LandscapeSectionLODUniformParameters = RenderSystem.UniformBuffer;
-
 		// Combined grass rendering batch element
 		FMeshBatchElement* GrassBatchElement = &GrassMeshBatch.Elements[0];
 		FLandscapeBatchElementParams* BatchElementParams = &GrassBatchParams[0];
 		BatchElementParams->LandscapeUniformShaderParametersResource = &LandscapeUniformShaderParameters;
 		BatchElementParams->FixedGridUniformShaderParameters = &LandscapeFixedGridUniformShaderParameters;
-		BatchElementParams->LandscapeSectionLODUniformParameters = &LandscapeSectionLODUniformParameters;
+		BatchElementParams->LandscapeSectionLODUniformParameters = RenderSystem.SectionLODUniformBuffer;
 		BatchElementParams->SceneProxy = this;
 		BatchElementParams->CurrentLOD = 0;
 		GrassBatchElement->UserData = BatchElementParams;
@@ -2096,13 +2094,12 @@ bool FLandscapeComponentSceneProxy::GetMeshElementForVirtualTexture(int32 InLodI
 	OutMeshBatch.Elements.Empty(1);
 
 	const FLandscapeRenderSystem& RenderSystem = *LandscapeRenderSystems.FindChecked(LandscapeKey);
-	const TUniformBufferRef<FLandscapeSectionLODUniformParameters>& LandscapeSectionLODUniformParameters = RenderSystem.UniformBuffer;
 
 	FLandscapeBatchElementParams* BatchElementParams = new(OutStaticBatchParamArray) FLandscapeBatchElementParams;
 	BatchElementParams->SceneProxy = this;
 	BatchElementParams->LandscapeUniformShaderParametersResource = &LandscapeUniformShaderParameters;
 	BatchElementParams->FixedGridUniformShaderParameters = &LandscapeFixedGridUniformShaderParameters;
-	BatchElementParams->LandscapeSectionLODUniformParameters = &LandscapeSectionLODUniformParameters;
+	BatchElementParams->LandscapeSectionLODUniformParameters = RenderSystem.SectionLODUniformBuffer;
 	BatchElementParams->CurrentLOD = InLodIndex;
 
 	int32 LodSubsectionSizeVerts = SubsectionSizeVerts >> InLodIndex;
@@ -2171,12 +2168,11 @@ bool FLandscapeComponentSceneProxy::GetStaticMeshElement(int32 LODIndex, bool bF
 		MeshBatch.bDitheredLODTransition = false;
 
 		const FLandscapeRenderSystem& RenderSystem = *LandscapeRenderSystems.FindChecked(LandscapeKey);
-		const TUniformBufferRef<FLandscapeSectionLODUniformParameters>& LandscapeSectionLODUniformParameters = RenderSystem.UniformBuffer;
 
 		FLandscapeBatchElementParams* BatchElementParams = new(OutStaticBatchParamArray) FLandscapeBatchElementParams;
 		BatchElementParams->LandscapeUniformShaderParametersResource = &LandscapeUniformShaderParameters;
 		BatchElementParams->FixedGridUniformShaderParameters = &LandscapeFixedGridUniformShaderParameters;
-		BatchElementParams->LandscapeSectionLODUniformParameters = &LandscapeSectionLODUniformParameters;
+		BatchElementParams->LandscapeSectionLODUniformParameters = RenderSystem.SectionLODUniformBuffer;
 		BatchElementParams->SceneProxy = this;
 		BatchElementParams->CurrentLOD = LODIndex;
 
@@ -2692,7 +2688,6 @@ void FLandscapeComponentSceneProxy::GetDynamicRayTracingInstances(FRayTracingMat
 	BaseMeshBatch.Elements.Empty();
 
 	const FLandscapeRenderSystem& RenderSystem = *LandscapeRenderSystems.FindChecked(LandscapeKey);
-	const TUniformBufferRef<FLandscapeSectionLODUniformParameters>& LandscapeSectionLODUniformParameters = RenderSystem.UniformBuffer;
 
 	for (int32 SubY = 0; SubY < NumSubsections; SubY++)
 	{
@@ -2708,7 +2703,7 @@ void FLandscapeComponentSceneProxy::GetDynamicRayTracingInstances(FRayTracingMat
 
 			BatchElementParams.LandscapeUniformShaderParametersResource = &LandscapeUniformShaderParameters;
 			BatchElementParams.FixedGridUniformShaderParameters = &LandscapeFixedGridUniformShaderParameters;
-			BatchElementParams.LandscapeSectionLODUniformParameters = &LandscapeSectionLODUniformParameters;
+			BatchElementParams.LandscapeSectionLODUniformParameters = RenderSystem.SectionLODUniformBuffer;
 			BatchElementParams.SceneProxy = this;
 			BatchElementParams.CurrentLOD = CurrentLOD;
 			BatchElement.UserData = &BatchElementParams;
@@ -3225,7 +3220,7 @@ public:
 		const FLandscapeComponentSceneProxy* SceneProxy = BatchElementParams->SceneProxy;
 
 		ShaderBindings.Add(Shader->GetUniformBufferParameter<FLandscapeUniformShaderParameters>(), *BatchElementParams->LandscapeUniformShaderParametersResource);
-		ShaderBindings.Add(Shader->GetUniformBufferParameter<FLandscapeSectionLODUniformParameters>(), *BatchElementParams->LandscapeSectionLODUniformParameters);
+		ShaderBindings.Add(Shader->GetUniformBufferParameter<FLandscapeSectionLODUniformParameters>(), BatchElementParams->LandscapeSectionLODUniformParameters);
 
 #if RHI_RAYTRACING
 		if (IsRayTracingEnabled())
diff --git a/Engine/Source/Runtime/Landscape/Private/LandscapeRenderMobile.cpp b/Engine/Source/Runtime/Landscape/Private/LandscapeRenderMobile.cpp
--- a/Engine/Source/Runtime/Landscape/Private/LandscapeRenderMobile.cpp
+++ b/Engine/Source/Runtime/Landscape/Private/LandscapeRenderMobile.cpp
@@ -81,7 +81,7 @@ public:
 
 		const FLandscapeComponentSceneProxyMobile* SceneProxy = (const FLandscapeComponentSceneProxyMobile*)BatchElementParams->SceneProxy;
 		ShaderBindings.Add(Shader->GetUniformBufferParameter<FLandscapeUniformShaderParameters>(),*BatchElementParams->LandscapeUniformShaderParametersResource);
-		ShaderBindings.Add(Shader->GetUniformBufferParameter<FLandscapeSectionLODUniformParameters>(), *BatchElementParams->LandscapeSectionLODUniformParameters);
+		ShaderBindings.Add(Shader->GetUniformBufferParameter<FLandscapeSectionLODUniformParameters>(), BatchElementParams->LandscapeSectionLODUniformParameters);
 
 		if (TexCoordOffsetParameter.IsBound())
 		{
diff --git a/Engine/Source/Runtime/Landscape/Public/LandscapeRender.h b/Engine/Source/Runtime/Landscape/Public/LandscapeRender.h
--- a/Engine/Source/Runtime/Landscape/Public/LandscapeRender.h
+++ b/Engine/Source/Runtime/Landscape/Public/LandscapeRender.h
@@ -135,7 +135,7 @@ struct FLandscapeBatchElementParams
 #endif
 	const TUniformBuffer<FLandscapeUniformShaderParameters>* LandscapeUniformShaderParametersResource;
 	const TArray<TUniformBuffer<FLandscapeFixedGridUniformShaderParameters>>* FixedGridUniformShaderParameters;
-	const TUniformBufferRef<FLandscapeSectionLODUniformParameters>* LandscapeSectionLODUniformParameters;
+	FUniformBufferRHIRef LandscapeSectionLODUniformParameters;
 	const FLandscapeComponentSceneProxy* SceneProxy;
 	int32 CurrentLOD;
 };
@@ -502,7 +502,7 @@ struct FLandscapeRenderSystem
 	FBufferRHIRef SectionLODBiasBuffer;
 	FShaderResourceViewRHIRef SectionLODBiasSRV;
 
-	TUniformBufferRef<FLandscapeSectionLODUniformParameters> UniformBuffer;
+	FUniformBufferRHIRef SectionLODUniformBuffer;
 
 	FCriticalSection CachedValuesCS;
 	TMap<FViewKey, TResourceArray<float>> CachedSectionLODValues;
