Landscape: Avoid assigning new guid on duplicate in commandlet

diff --git a/Engine/Source/Runtime/Landscape/Private/LandscapeModule.cpp b/Engine/Source/Runtime/Landscape/Private/LandscapeModule.cpp
--- a/Engine/Source/Runtime/Landscape/Private/LandscapeModule.cpp
+++ b/Engine/Source/Runtime/Landscape/Private/LandscapeModule.cpp
@@ -168,7 +168,7 @@ void WorldDuplicateEventFunction(UWorld* World, bool bDuplicateForPIE, TMap<UObj
 
 #if WITH_EDITOR
 	// Fixup LandscapeGuid on World duplication
-	if (!bDuplicateForPIE)
+	if (!bDuplicateForPIE && !IsRunningCommandlet())
 	{
 		TMap<FGuid, FGuid> NewLandscapeGuids;
 		for (ALandscapeProxy* Proxy : TActorRange<ALandscapeProxy>(World, ALandscapeProxy::StaticClass(), EActorIteratorFlags::SkipPendingKill))
