Fix bug where landscape and water meshes would not update after forcing an update.

- With on demand shader compilation enabled there are some scenarios where you do not want to rely on rendering to inform you of what shaders need to be compiled.  Specifically scenarios where materials are used to generate content (landscape brushes, thumbnails, scene captures, etc.).  In these scenarios we should explicitly compile all shaders.
- Explicitly calling SubmitCompileJobs which is a more straight forward way of compiling missing shaders from the maps.  FMaterial::CacheShaders() doesn't support this behavior and it is too risky to change.
- SubmitCompileJobs also supports being called multiple times and will not re-compile more shaders.

diff --git a/Engine/Source/Runtime/Landscape/Private/LandscapeEditLayers.cpp b/Engine/Source/Runtime/Landscape/Private/LandscapeEditLayers.cpp
--- a/Engine/Source/Runtime/Landscape/Private/LandscapeEditLayers.cpp
+++ b/Engine/Source/Runtime/Landscape/Private/LandscapeEditLayers.cpp
@@ -1433,10 +1433,18 @@ bool ALandscape::IsTextureReady(UTexture2D* InTexture, bool bInWaitForStreaming)
 bool ALandscape::IsMaterialResourceCompiled(FMaterialResource* InMaterialResource, bool bInWaitForCompilation)
 {
 	check(InMaterialResource);
-	if (bInWaitForCompilation && !InMaterialResource->IsGameThreadShaderMapComplete())
+	if (InMaterialResource->IsGameThreadShaderMapComplete())
+	{
+		return true;
+	}
+
 	{
 		TRACE_CPUPROFILER_EVENT_SCOPE(Landscape_WaitForMaterialCompilation);
-		InMaterialResource->FinishCompilation();
+		InMaterialResource->SubmitCompileJobs(EShaderCompileJobPriority::ForceLocal);
+		if (bInWaitForCompilation)
+		{
+			InMaterialResource->FinishCompilation();
+		}
 	}
 	return InMaterialResource->IsGameThreadShaderMapComplete();
 }
