Retop tool breaks navmesh generation for landscape

Implemented WITH_CHAOS code path in ULandscapeMeshCollisionComponent::DoCustomNavigableGeometryExport()

diff --git a/Engine/Source/Runtime/Landscape/Private/LandscapeCollision.cpp b/Engine/Source/Runtime/Landscape/Private/LandscapeCollision.cpp
--- a/Engine/Source/Runtime/Landscape/Private/LandscapeCollision.cpp
+++ b/Engine/Source/Runtime/Landscape/Private/LandscapeCollision.cpp
@@ -2447,24 +2447,32 @@ void ULandscapeHeightfieldCollisionComponent::PrepareGeometryExportSync()
 bool ULandscapeMeshCollisionComponent::DoCustomNavigableGeometryExport(FNavigableGeometryExport& GeomExport) const
 {
 	check(IsInGameThread());
-#if PHYSICS_INTERFACE_PHYSX
-	if (IsValidRef(MeshRef) && MeshRef->RBTriangleMesh != nullptr)
+
+	if (IsValidRef(MeshRef))
 	{
 		FTransform MeshToW = GetComponentTransform();
 		MeshToW.MultiplyScale3D(FVector(CollisionScale, CollisionScale, 1.f));
 
-		if (MeshRef->RBTriangleMesh->getTriangleMeshFlags() & PxTriangleMeshFlag::e16_BIT_INDICES)
+#if PHYSICS_INTERFACE_PHYSX
+		if (MeshRef->RBTriangleMesh != nullptr)
 		{
-			GeomExport.ExportPxTriMesh16Bit(MeshRef->RBTriangleMesh, MeshToW);
+			if (MeshRef->RBTriangleMesh->getTriangleMeshFlags() & PxTriangleMeshFlag::e16_BIT_INDICES)
+			{
+				GeomExport.ExportPxTriMesh16Bit(MeshRef->RBTriangleMesh, MeshToW);
+			}
+			else
+			{
+				GeomExport.ExportPxTriMesh32Bit(MeshRef->RBTriangleMesh, MeshToW);
+			}
 		}
-		else
+#elif WITH_CHAOS
+		if (MeshRef->Trimesh != nullptr)
 		{
-			GeomExport.ExportPxTriMesh32Bit(MeshRef->RBTriangleMesh, MeshToW);
+			GeomExport.ExportChaosTriMesh(MeshRef->Trimesh.Get(), MeshToW);
 		}
-	}
-#elif WITH_CHAOS
-	CHAOS_ENSURE(false);
 #endif
+	}
+
 	return false;
 }
 
