Default landscape actor to collapse children in the world outliner

diff --git a/Engine/Source/Runtime/Landscape/Private/Landscape.cpp b/Engine/Source/Runtime/Landscape/Private/Landscape.cpp
--- a/Engine/Source/Runtime/Landscape/Private/Landscape.cpp
+++ b/Engine/Source/Runtime/Landscape/Private/Landscape.cpp
@@ -1131,6 +1131,7 @@ ALandscape::ALandscape(const FObjectInitializer& ObjectInitializer)
 	LandscapeEdMode = nullptr;
 	bGrassUpdateEnabled = true;
 	bIsSpatiallyLoaded = false;
+	bDefaultOutlinerExpansionState = false;
 #endif // WITH_EDITORONLY_DATA
 }
 
