Landscape components flicker while sculpting

Ensure FLandscapeComponentSceneProxy::GetDynamicMeshElements() is always using the same LOD as the one stored in the view LandscapePerComponentData buffer

Added missing view extension setup in UWorldThumbnailRenderer

diff --git a/Engine/Source/Editor/UnrealEd/Private/ThumbnailRendering/WorldThumbnailRenderer.cpp b/Engine/Source/Editor/UnrealEd/Private/ThumbnailRendering/WorldThumbnailRenderer.cpp
--- a/Engine/Source/Editor/UnrealEd/Private/ThumbnailRendering/WorldThumbnailRenderer.cpp
+++ b/Engine/Source/Editor/UnrealEd/Private/ThumbnailRendering/WorldThumbnailRenderer.cpp
@@ -6,6 +6,8 @@
 #include "ShowFlags.h"
 #include "Engine/World.h"
 #include "SceneView.h"
+#include "SceneViewExtension.h"
+#include "SceneViewExtensionContext.h"
 #include "ThumbnailRendering/WorldThumbnailInfo.h"
 #include "Engine/LevelBounds.h"
 #include "ContentStreaming.h"
@@ -43,6 +45,8 @@ void UWorldThumbnailRenderer::Draw(UObject* Object, int32 X, int32 Y, uint32 Wid
 			.SetTime(UThumbnailRenderer::GetTime())
 			.SetAdditionalViewFamily(bAdditionalViewFamily));
 
+		ViewFamily.ViewExtensions = GEngine->ViewExtensions->GatherActiveExtensions(FSceneViewExtensionContext(World->Scene));
+
 		ViewFamily.EngineShowFlags.SetDiffuse(true);
 		ViewFamily.EngineShowFlags.SetSkeletalMeshes(true);
 		ViewFamily.EngineShowFlags.SetTranslucency(true);
diff --git a/Engine/Source/Runtime/Landscape/Private/LandscapeRender.cpp b/Engine/Source/Runtime/Landscape/Private/LandscapeRender.cpp
--- a/Engine/Source/Runtime/Landscape/Private/LandscapeRender.cpp
+++ b/Engine/Source/Runtime/Landscape/Private/LandscapeRender.cpp
@@ -2136,7 +2136,7 @@ void FLandscapeComponentSceneProxy::ApplyWorldOffset(FVector InOffset)
 }
 
 template<class ArrayType>
-bool FLandscapeComponentSceneProxy::GetStaticMeshElement(int32 LODIndex, bool bForToolMesh, bool bForcedLOD, FMeshBatch& MeshBatch, ArrayType& OutStaticBatchParamArray) const
+bool FLandscapeComponentSceneProxy::GetStaticMeshElement(int32 LODIndex, bool bForToolMesh, FMeshBatch& MeshBatch, ArrayType& OutStaticBatchParamArray) const
 {
 	UMaterialInterface* MaterialInterface = nullptr;
 
@@ -2225,7 +2225,7 @@ void FLandscapeComponentSceneProxy::DrawStaticElements(FStaticPrimitiveDrawInter
 	{
 		FMeshBatch MeshBatch;
 
-		if (GetStaticMeshElement(LODIndex, false, false, MeshBatch, StaticBatchParamArray))
+		if (GetStaticMeshElement(LODIndex, false, MeshBatch, StaticBatchParamArray))
 		{
 			PDI->DrawMesh(MeshBatch, LODIndex == FirstLOD ? FLT_MAX : (FMath::Sqrt(LODScreenRatioSquared[LODIndex]) * 2.0f));
 		}
@@ -2277,6 +2277,8 @@ void FLandscapeComponentSceneProxy::GetDynamicMeshElements(const TArray<const FS
 	int32 NumDrawCalls = 0;
 	const bool bIsWireframe = ViewFamily.EngineShowFlags.Wireframe;
 
+	const FLandscapeRenderSystem& RenderSystem = *LandscapeRenderSystems.FindChecked(LandscapeKey);
+
 	for (int32 ViewIndex = 0; ViewIndex < Views.Num(); ViewIndex++)
 	{
 		if (VisibilityMap & (1 << ViewIndex))
@@ -2286,32 +2288,15 @@ void FLandscapeComponentSceneProxy::GetDynamicMeshElements(const TArray<const FS
 
 			const FSceneView* View = Views[ViewIndex];
 
-			int32 ForcedLODLevel = ForcedLOD;
-
-			const int32 ViewLodOverride = GetViewLodOverride(*View);
-			ForcedLODLevel = ViewLodOverride >= 0 ? ViewLodOverride : ForcedLODLevel;
-
-#if WITH_EDITOR || !(UE_BUILD_SHIPPING || UE_BUILD_TEST)
-			const int32 DrawCollisionLodOverride = GetDrawCollisionLodOverride(*View, CollisionResponse, CollisionMipLevel, SimpleCollisionMipLevel);
-			ForcedLODLevel = DrawCollisionLodOverride >= 0 ? DrawCollisionLodOverride : ForcedLODLevel;
-#endif
-
-			ForcedLODLevel = FMath::Min(ForcedLODLevel, (int32)LODSettings.LastLODIndex);
-
-			const float LODScale = View->LODDistanceFactor * CVarStaticMeshLODDistanceScale.GetValueOnRenderThread();
-			const float MeshScreenSizeSquared = ComputeBoundsScreenRadiusSquared(GetBounds().Origin, GetBounds().SphereRadius, *View);
-			int32 LODToRender = ForcedLODLevel >= 0 ? ForcedLODLevel : GetLODFromScreenSize(MeshScreenSizeSquared, LODScale * LODScale);
+			int32 LODToRender = RenderSystem.GetSectionLODValue(*View, ComponentBase);
 
-#if PLATFORM_SUPPORTS_LANDSCAPE_VISUAL_MESH_LOD_STREAMING
-			LODToRender = FMath::Max<int32>(LODToRender, GetCurrentFirstLODIdx_RenderThread());
-#endif
 			FMeshBatch& Mesh = Collector.AllocateMesh();
-			GetStaticMeshElement(LODToRender, false, ForcedLODLevel >= 0, Mesh, ParameterArray.ElementParams);
+			GetStaticMeshElement(LODToRender, false, Mesh, ParameterArray.ElementParams);
 
 #if WITH_EDITOR
 			FMeshBatch& MeshTools = Collector.AllocateMesh();
 			// No Tessellation on tool material
-			GetStaticMeshElement(LODToRender, true, ForcedLODLevel >= 0, MeshTools, ParameterArray.ElementParams);
+			GetStaticMeshElement(LODToRender, true, MeshTools, ParameterArray.ElementParams);
 #endif
 
 			// Render the landscape component
@@ -2649,14 +2634,9 @@ void FLandscapeComponentSceneProxy::GetDynamicRayTracingInstances(FRayTracingMat
 	}
 
 	const FSceneView& SceneView = *Context.ReferenceView;
-	int32 ForcedLODLevel = ForcedLOD;
-
-	int32 ViewLodOveride = GetViewLodOverride(SceneView);
-	ForcedLODLevel = ViewLodOveride >= 0 ? ViewLodOveride : ForcedLODLevel;
-
-	float MeshScreenSizeSquared = ComputeBoundsScreenRadiusSquared(GetBounds().Origin, GetBounds().SphereRadius, SceneView);
-	float LODScale = Context.ReferenceView->LODDistanceFactor * CVarStaticMeshLODDistanceScale.GetValueOnRenderThread();
-	int32 LODToRender = ForcedLODLevel >= 0 ? ForcedLODLevel : GetLODFromScreenSize(MeshScreenSizeSquared, LODScale * LODScale);
+	const FLandscapeRenderSystem& RenderSystem = *LandscapeRenderSystems.FindChecked(LandscapeKey);
+	
+	int32 LODToRender = RenderSystem.GetSectionLODValue(SceneView, ComponentBase);
 	
 	FLandscapeElementParamArray& ParameterArray = Context.RayTracingMeshResourceCollector.AllocateOneFrameResource<FLandscapeElementParamArray>();
 	ParameterArray.ElementParams.AddDefaulted(NumSubsections * NumSubsections);
@@ -2687,8 +2667,6 @@ void FLandscapeComponentSceneProxy::GetDynamicRayTracingInstances(FRayTracingMat
 
 	BaseMeshBatch.Elements.Empty();
 
-	const FLandscapeRenderSystem& RenderSystem = *LandscapeRenderSystems.FindChecked(LandscapeKey);
-
 	for (int32 SubY = 0; SubY < NumSubsections; SubY++)
 	{
 		for (int32 SubX = 0; SubX < NumSubsections; SubX++)
diff --git a/Engine/Source/Runtime/Landscape/Public/LandscapeRender.h b/Engine/Source/Runtime/Landscape/Public/LandscapeRender.h
--- a/Engine/Source/Runtime/Landscape/Public/LandscapeRender.h
+++ b/Engine/Source/Runtime/Landscape/Public/LandscapeRender.h
@@ -856,7 +856,7 @@ protected:
 	int8 GetLODFromScreenSize(float InScreenSizeSquared, float InViewLODScale) const;
 
 	bool GetMeshElementForVirtualTexture(int32 InLodIndex, ERuntimeVirtualTextureMaterialType MaterialType, UMaterialInterface* InMaterialInterface, FMeshBatch& OutMeshBatch, TArray<FLandscapeBatchElementParams>& OutStaticBatchParamArray) const;
-	template<class ArrayType> bool GetStaticMeshElement(int32 LODIndex, bool bForToolMesh, bool bForcedLOD, FMeshBatch& MeshBatch, ArrayType& OutStaticBatchParamArray) const;
+	template<class ArrayType> bool GetStaticMeshElement(int32 LODIndex, bool bForToolMesh, FMeshBatch& MeshBatch, ArrayType& OutStaticBatchParamArray) const;
 	
 	virtual void ApplyMeshElementModifier(FMeshBatchElement& InOutMeshElement, int32 InLodIndex) const {}
 
