Added Landscape Material to Edit asset context menus

diff --git a/Engine/Source/Runtime/Landscape/Classes/LandscapeProxy.h b/Engine/Source/Runtime/Landscape/Classes/LandscapeProxy.h
--- a/Engine/Source/Runtime/Landscape/Classes/LandscapeProxy.h
+++ b/Engine/Source/Runtime/Landscape/Classes/LandscapeProxy.h
@@ -808,6 +808,7 @@ public:
 	virtual FGuid GetGridGuid() const override { return LandscapeGuid; }
 	virtual TUniquePtr<class FWorldPartitionActorDesc> CreateClassActorDesc() const override;
 	virtual bool EditorCanAttachTo(const AActor* InParent, FText& OutReason) const override { return false; }
+	virtual bool GetReferencedContentObjects(TArray<UObject*>& Objects) const override;
 #endif	//WITH_EDITOR
 
 	virtual FGuid GetLandscapeGuid() const override { return LandscapeGuid; }
diff --git a/Engine/Source/Runtime/Landscape/Classes/LandscapeStreamingProxy.h b/Engine/Source/Runtime/Landscape/Classes/LandscapeStreamingProxy.h
--- a/Engine/Source/Runtime/Landscape/Classes/LandscapeStreamingProxy.h
+++ b/Engine/Source/Runtime/Landscape/Classes/LandscapeStreamingProxy.h
@@ -39,6 +39,7 @@ public:
 	virtual void PostRegisterAllComponents() override;
 	virtual AActor* GetSceneOutlinerParent() const override;
 	virtual bool CanDeleteSelectedActor(FText& OutReason) const override;
+	virtual bool GetReferencedContentObjects(TArray<UObject*>& Objects) const override;
 #endif
 	//~ End UObject Interface
 
diff --git a/Engine/Source/Runtime/Landscape/Private/Landscape.cpp b/Engine/Source/Runtime/Landscape/Private/Landscape.cpp
--- a/Engine/Source/Runtime/Landscape/Private/Landscape.cpp
+++ b/Engine/Source/Runtime/Landscape/Private/Landscape.cpp
@@ -2877,11 +2877,7 @@ UMaterialInterface* ALandscapeProxy::GetLandscapeMaterial(int8 InLODIndex) const
 
 UMaterialInterface* ALandscapeProxy::GetLandscapeHoleMaterial() const
 {
-	if (LandscapeHoleMaterial)
-	{
-		return LandscapeHoleMaterial;
-	}
-	return nullptr;
+	return LandscapeHoleMaterial;
 }
 
 UMaterialInterface* ALandscapeStreamingProxy::GetLandscapeMaterial(int8 InLODIndex) const
diff --git a/Engine/Source/Runtime/Landscape/Private/LandscapeEdit.cpp b/Engine/Source/Runtime/Landscape/Private/LandscapeEdit.cpp
--- a/Engine/Source/Runtime/Landscape/Private/LandscapeEdit.cpp
+++ b/Engine/Source/Runtime/Landscape/Private/LandscapeEdit.cpp
@@ -699,6 +699,31 @@ TUniquePtr<FWorldPartitionActorDesc> ALandscapeProxy::CreateClassActorDesc() con
 	return TUniquePtr<FWorldPartitionActorDesc>(new FLandscapeActorDesc());
 }
 
+bool ALandscapeProxy::GetReferencedContentObjects(TArray<UObject*>& Objects) const
+{
+	Super::GetReferencedContentObjects(Objects);
+
+	if (LandscapeMaterial != nullptr)
+	{
+		Objects.AddUnique(LandscapeMaterial);
+	}
+
+	for (const FLandscapeProxyMaterialOverride& OverrideMaterial : LandscapeMaterialsOverride)
+	{
+		if (OverrideMaterial.Material != nullptr)
+		{
+			Objects.AddUnique(OverrideMaterial.Material);
+		}
+	}
+
+	if (LandscapeHoleMaterial != nullptr)
+	{
+		Objects.AddUnique(LandscapeHoleMaterial);
+	}
+
+	return true;
+}
+
 void ALandscapeProxy::FixupWeightmaps()
 {
 	WeightmapUsageMap.Empty();
@@ -5153,6 +5178,20 @@ bool ALandscapeStreamingProxy::CanDeleteSelectedActor(FText& OutReason) const
 	return true;
 }
 
+
+bool ALandscapeStreamingProxy::GetReferencedContentObjects(TArray<UObject*>& Objects) const
+{
+	Super::GetReferencedContentObjects(Objects);
+
+	// Also return the objects referenced by our parent landscape : 
+	if (LandscapeActor != nullptr)
+	{
+		LandscapeActor->GetReferencedContentObjects(Objects);
+	}
+
+	return true;
+}
+
 bool ALandscape::CanDeleteSelectedActor(FText& OutReason) const
 {
 	if (!IsUserManaged())
