Local merge of landscape edit layers : complete refactor of the merge algorithm that doesn't require allocating several render targets the size of the entire loaded landscape components. Instead, each

diff --git a/Engine/Source/Runtime/Landscape/Private/LandscapeEditLayers.cpp b/Engine/Source/Runtime/Landscape/Private/LandscapeEditLayers.cpp
--- a/Engine/Source/Runtime/Landscape/Private/LandscapeEditLayers.cpp
+++ b/Engine/Source/Runtime/Landscape/Private/LandscapeEditLayers.cpp
@@ -14,6 +14,7 @@
 #include "LandscapeDataAccess.h"
 #include "LandscapePrivate.h"
 #include "LandscapeEditReadback.h"
+#include "LandscapeEditResources.h"
 
 #include "Shader.h"
 #include "GlobalShader.h"
@@ -21,6 +22,11 @@
 #include "Engine/TextureRenderTarget2D.h"
 #include "Logging/MessageLog.h"
 #include "RenderCaptureInterface.h"
+#include "RenderGraph.h"
+#include "RenderGraphUtils.h"
+#include "PixelShaderUtils.h"
+#include "Runtime/Renderer/Private/SystemTextures.h"
+#include "Rendering/Texture2DResource.h"
 
 #if WITH_EDITOR
 #include "LandscapeEditorModule.h"
@@ -32,6 +38,8 @@
 #include "Materials/MaterialExpressionLandscapeVisibilityMask.h"
 #include "ShaderCompiler.h"
 #include "Algo/Count.h"
+#include "Algo/AnyOf.h"
+#include "Algo/Unique.h"
 #include "LandscapeSettings.h"
 #include "LandscapeRender.h"
 #include "LandscapeInfoMap.h"
@@ -77,13 +85,13 @@
 static FAutoConsoleVariableRef CVarRenderCaptureLayersNextHeightmapDraws(
 	TEXT("landscape.RenderCaptureLayersNextHeightmapDraws"),
 	RenderCaptureLayersNextHeightmapDraws,
-	TEXT("Trigger a render capture during the next heightmap draw calls."));
+	TEXT("Trigger N render capture during the next heightmap draw calls."));
 
 int32 RenderCaptureLayersNextWeightmapDraws = 0;
 static FAutoConsoleVariableRef CVarRenderCaptureLayersNextWeightmapDraws(
 	TEXT("landscape.RenderCaptureLayersNextWeightmapDraws"),
 	RenderCaptureLayersNextWeightmapDraws,
-	TEXT("Trigger a render capture during the next weightmap draw calls."));
+	TEXT("Trigger N render capture during the next weightmap draw calls."));
 
 static TAutoConsoleVariable<int32> CVarOutputLayersRTContent(
 	TEXT("landscape.OutputLayersRTContent"),
@@ -136,6 +144,47 @@
 	1,
 	TEXT("This will validate that weightmap usages in landscape proxies and their components don't get desynchronized with the landscape component layer allocations."));
 
+void OnLandscapeEditLayersLocalMergeChanged(IConsoleVariable* CVar)
+{
+	for (TObjectIterator<UWorld> It; It; ++It)
+	{
+		UWorld* CurrentWorld = *It;
+		if (!CurrentWorld->IsGameWorld())
+		{
+			auto& LandscapeInfoMap = ULandscapeInfoMap::GetLandscapeInfoMap(CurrentWorld);
+			for (TPair<FGuid, ULandscapeInfo*>& Pair : LandscapeInfoMap.Map)
+			{
+				if (ULandscapeInfo* LandscapeInfo = Pair.Value)
+				{
+					if (ALandscape* Landscape = LandscapeInfo->LandscapeActor.Get())
+					{
+						Landscape->RequestLayersInitialization(true);
+					}
+				}
+			}
+		}
+	}
+}
+
+int32 LandscapeEditLayersLocalMerge = 0;
+static FAutoConsoleVariableRef CVarLandscapeEditLayersLocalMerge(
+	TEXT("landscape.EditLayersLocalMerge.Enable"),
+	LandscapeEditLayersLocalMerge,
+	TEXT("This will allow the new merge algorithm (that merges layers at the landscape component level) to be used on landscapes that support it. This is a temporary measure while waiting for non-compatible landscapes to be deprecated. "),
+	FConsoleVariableDelegate::CreateStatic(OnLandscapeEditLayersLocalMergeChanged));
+
+TAutoConsoleVariable<int32> CVarLandscapeEditLayersMaxComponentsPerHeightmapResolveBatch(
+	TEXT("landscape.EditLayersLocalMerge.MaxComponentsPerHeightmapResolveBatch"),
+	16,
+	TEXT("Number of components being rendered in a single batch when resolving heightmaps. The higher the number, the more heightmaps can be resolved in a single batch (and the higher the GPU memory consumption since more transient textures will be needed in memory at a time)"),
+	ECVF_RenderThreadSafe);
+
+TAutoConsoleVariable<int32> CVarLandscapeEditLayersMaxComponentsPerWeightmapResolveBatch(
+	TEXT("landscape.EditLayersLocalMerge.MaxComponentsPerWeightmapResolveBatch"),
+	16,
+	TEXT("Number of components being rendered in a single batch when resolving weightmaps. The higher the number, the more weightmaps can be resolved in a single batch (and the higher the GPU memory consumption since more transient textures will be needed in memory at a time)"),
+	ECVF_RenderThreadSafe);
+
 struct FLandscapeDirty
 {
 	FLandscapeDirty()
@@ -151,21 +200,25 @@
 
 	void ClearDirty()
 	{
-		if(!GWorld || GWorld->IsGameWorld())
-		{ 
-			return;
-		}
-
-		auto& LandscapeInfoMap = ULandscapeInfoMap::GetLandscapeInfoMap(GWorld);
-		for (TPair<FGuid, ULandscapeInfo*>& Pair : LandscapeInfoMap.Map)
+		bool bCleared = false;
+		for (TObjectIterator<UWorld> It; It; ++It)
 		{
-			if (Pair.Value && Pair.Value->SupportsLandscapeEditing())
+			UWorld* CurrentWorld = *It;
+			if (!CurrentWorld->IsGameWorld())
 			{
-				Pair.Value->ClearDirtyData();
+				auto& LandscapeInfoMap = ULandscapeInfoMap::GetLandscapeInfoMap(CurrentWorld);
+				for (TPair<FGuid, ULandscapeInfo*>& Pair : LandscapeInfoMap.Map)
+				{
+					if (Pair.Value && Pair.Value->SupportsLandscapeEditing())
+					{
+						Pair.Value->ClearDirtyData();
+						bCleared = true;
+					}
+				}
 			}
 		}
 
-		UE_LOG(LogLandscape, Display, TEXT("Landscape.Dirty: Cleared"));
+		UE_LOG(LogLandscape, Display, TEXT("Landscape.Dirty: %s"), bCleared ? TEXT("Cleared") : TEXT("Landscape.Dirty: Nothing to clear"));
 	}
 };
 
@@ -177,45 +230,95 @@
  */
 struct FTextureToComponentHelper
 {
-	FTextureToComponentHelper(ULandscapeInfo const& InLandscapeInfo)
+	// Partial refresh flags : allows to recompute only a subset of the helper information :
+	enum class ERefreshFlags
 	{
-		InLandscapeInfo.ForAllLandscapeComponents([&](ULandscapeComponent* Component)
+		None = 0,
+		RefreshComponents = (1 << 0),
+		RefreshHeightmaps = (1 << 1),
+		RefreshWeightmaps = (1 << 2),
+		RefreshAll = ~0,
+	};
+	FRIEND_ENUM_CLASS_FLAGS(ERefreshFlags);
+
+	FTextureToComponentHelper(const ULandscapeInfo& InLandscapeInfo)
+		: LandscapeInfo(&InLandscapeInfo)
+	{
+		Refresh(ERefreshFlags::RefreshAll);
+	}
+
+	void Refresh(ERefreshFlags InRefreshFlags)
+	{
+		TRACE_CPUPROFILER_EVENT_SCOPE(TextureToComponentHelper_Refresh);
+		// Compute the list of components in this landscape : 
+		if (EnumHasAnyFlags(InRefreshFlags, ERefreshFlags::RefreshComponents))
 		{
-			LandscapeComponents.Add(Component);
+			// When components are refreshed, all other info has to be :
+			check(EnumHasAllFlags(InRefreshFlags, ERefreshFlags::RefreshAll));
 
+			LandscapeComponents.Reset();
+			LandscapeInfo->ForAllLandscapeComponents([&](ULandscapeComponent* Component)
 			{
-				UTexture2D* Heightmap = Component->GetHeightmap();
-				check(Heightmap != nullptr);
+				LandscapeComponents.Add(Component);
+			});
+		}
 
-				Heightmaps.Add(Heightmap);
-				HeightmapToComponents.FindOrAdd(Heightmap).Add(Component);
+		if (EnumHasAnyFlags(InRefreshFlags, ERefreshFlags::RefreshHeightmaps | ERefreshFlags::RefreshWeightmaps))
+		{
+			// Cleanup our heightmap/weightmap info:
+			if (EnumHasAnyFlags(InRefreshFlags, ERefreshFlags::RefreshHeightmaps))
+			{
+				Heightmaps.Reset();
+				HeightmapToComponents.Reset();
 			}
 
+			if (EnumHasAnyFlags(InRefreshFlags, ERefreshFlags::RefreshWeightmaps))
 			{
-				const TArray<UTexture2D*>& WeightmapTextures = Component->GetWeightmapTextures();
-				const TArray<FWeightmapLayerAllocationInfo>& AllocInfos = Component->GetWeightmapLayerAllocations();
+				Weightmaps.Reset();
+				WeightmapToComponents.Reset();
+			}
 
-				for (FWeightmapLayerAllocationInfo const& AllocInfo : AllocInfos)
+			// Iterate on all tracked landscape components and keep track of components/heightmaps/weightmaps relationship :
+			for (ULandscapeComponent* Component : LandscapeComponents)
+			{
+				if (EnumHasAnyFlags(InRefreshFlags, ERefreshFlags::RefreshHeightmaps))
 				{
-					if(AllocInfo.IsAllocated() && AllocInfo.WeightmapTextureIndex < WeightmapTextures.Num())
+					UTexture2D* Heightmap = Component->GetHeightmap();
+					check(Heightmap != nullptr);
+
+					Heightmaps.Add(Heightmap);
+					HeightmapToComponents.FindOrAdd(Heightmap).Add(Component);
+				}
+
+				if (EnumHasAnyFlags(InRefreshFlags, ERefreshFlags::RefreshWeightmaps))
+				{
+					const TArray<UTexture2D*>& WeightmapTextures = Component->GetWeightmapTextures();
+					const TArray<FWeightmapLayerAllocationInfo>& AllocInfos = Component->GetWeightmapLayerAllocations();
+
+					for (FWeightmapLayerAllocationInfo const& AllocInfo : AllocInfos)
 					{
-						UTexture2D* Weightmap = WeightmapTextures[AllocInfo.WeightmapTextureIndex];
-						check(Weightmap != nullptr);
+						if (AllocInfo.IsAllocated() && AllocInfo.WeightmapTextureIndex < WeightmapTextures.Num())
+						{
+							UTexture2D* Weightmap = WeightmapTextures[AllocInfo.WeightmapTextureIndex];
+							check(Weightmap != nullptr);
 
-						Weightmaps.Add(Weightmap);
-						WeightmapToComponents.FindOrAdd(Weightmap).AddUnique(Component);
+							Weightmaps.Add(Weightmap);
+							WeightmapToComponents.FindOrAdd(Weightmap).AddUnique(Component);
+						}
 					}
 				}
 			}
-		});
+		}
 	}
 
+	const ULandscapeInfo* LandscapeInfo = nullptr;
 	TArray< ULandscapeComponent* > LandscapeComponents;
 	TSet< UTexture2D* > Heightmaps;
 	TMap< UTexture2D*, TArray<ULandscapeComponent*> > HeightmapToComponents;
 	TSet< UTexture2D* > Weightmaps;
 	TMap< UTexture2D*, TArray<ULandscapeComponent*> > WeightmapToComponents;
 };
+ENUM_CLASS_FLAGS(FTextureToComponentHelper::ERefreshFlags);
 
 static FFileHelper::EColorChannel GetWeightmapColorChannel(const FWeightmapLayerAllocationInfo& AllocInfo)
 {
@@ -232,136 +335,6 @@
 
 #endif
 
-// Custom Resources
-
-class FLandscapeTexture2DResource : public FTextureResource
-{
-public:
-	FLandscapeTexture2DResource(uint32 InSizeX, uint32 InSizeY, EPixelFormat InFormat, uint32 InNumMips, bool InNeedUAV)
-		: SizeX(InSizeX)
-		, SizeY(InSizeY)
-		, Format(InFormat)
-		, NumMips(InNumMips)
-		, CreateUAV(InNeedUAV)
-	{}
-
-	virtual uint32 GetSizeX() const override
-	{
-		return SizeX;
-	}
-
-	virtual uint32 GetSizeY() const override
-	{
-		return SizeY;
-	}
-
-	/** Called when the resource is initialized. This is only called by the rendering thread. */
-	virtual void InitRHI() override
-	{
-		FTextureResource::InitRHI();
-
-		FRHIResourceCreateInfo CreateInfo(TEXT("FLandscapeTexture2DResource"));
-		ETextureCreateFlags Flags = TexCreate_None;
-
-		if (CreateUAV)
-		{
-			Flags |= TexCreate_RenderTargetable|TexCreate_UAV;
-		}
-
-		TextureRHI = RHICreateTexture2D(SizeX, SizeY, Format, NumMips, 1, Flags, CreateInfo);
-
-		if (CreateUAV)
-		{
-			TextureUAV = RHICreateUnorderedAccessView(TextureRHI, 0);
-		}
-	}
-
-	virtual void ReleaseRHI() override
-	{
-		if (CreateUAV)
-		{
-			TextureUAV.SafeRelease();
-		}
-		
-		FTextureResource::ReleaseRHI();
-	}
-
-	FUnorderedAccessViewRHIRef TextureUAV;
-
-private:
-	uint32 SizeX;
-	uint32 SizeY;
-	EPixelFormat Format;
-	uint32 NumMips;
-	bool CreateUAV;
-};
-
-class FLandscapeTexture2DArrayResource : public FTextureResource
-{
-public:
-	FLandscapeTexture2DArrayResource(uint32 InSizeX, uint32 InSizeY, uint32 InSizeZ, EPixelFormat InFormat, uint32 InNumMips, bool InNeedUAV)
-		: SizeX(InSizeX)
-		, SizeY(InSizeY)
-		, SizeZ(InSizeZ)
-		, Format(InFormat)
-		, NumMips(InNumMips)
-		, CreateUAV(InNeedUAV)
-	{}
-
-	virtual uint32 GetSizeX() const override
-	{
-		return SizeX;
-	}
-
-	virtual uint32 GetSizeY() const override
-	{
-		return SizeY;
-	}
-
-	virtual uint32 GetSizeZ() const
-	{
-		return SizeZ;
-	}
-
-	/** Called when the resource is initialized. This is only called by the rendering thread. */
-	virtual void InitRHI() override
-	{
-		FTextureResource::InitRHI();
-
-		FRHIResourceCreateInfo CreateInfo(TEXT("FLandscapeTexture2DArrayResource"));
-		ETextureCreateFlags Flags = TexCreate_NoTiling | TexCreate_OfflineProcessed;
-
-		if (CreateUAV)
-		{
-			Flags |= TexCreate_UAV;
-		}
-
-		TextureRHI = RHICreateTexture2DArray(SizeX, SizeY, SizeZ, Format, NumMips, 1, Flags, CreateInfo);
-
-		if (CreateUAV)
-		{
-			TextureUAV = RHICreateUnorderedAccessView(TextureRHI, 0);
-		}
-	}
-
-	virtual void ReleaseRHI() override
-	{
-		FTextureResource::ReleaseRHI();
-
-		TextureUAV.SafeRelease();
-	}
-
-	FUnorderedAccessViewRHIRef TextureUAV;
-
-private:
-	uint32 SizeX;
-	uint32 SizeY;
-	uint32 SizeZ;
-	EPixelFormat Format;
-	uint32 NumMips;
-	bool CreateUAV;
-};
-
 // Vertex format and vertex buffer
 
 struct FLandscapeLayersVertex
@@ -746,9 +719,9 @@
 		return true;
 	}
 
-	FLandscapeCopyTextureVS() 
+	FLandscapeCopyTextureVS()
 	{};
-	
+
 	FLandscapeCopyTextureVS(const ShaderMetaType::CompiledShaderInitializerType& Initializer)
 		: FGlobalShader(Initializer)
 	{}
@@ -896,7 +869,7 @@
 	void SetParameters(FRHICommandList& RHICmdList, const FLandscapeLayerWeightmapExtractMaterialLayersComputeShaderParameters& InParams)
 	{
 		SetTextureParameter(RHICmdList, RHICmdList.GetBoundComputeShader(), ComponentWeightmapParam, InParams.ComponentWeightmapResource->TextureRHI);
-		SetUAVParameter(RHICmdList, RHICmdList.GetBoundComputeShader(), AtlasPaintListsParam, InParams.AtlasWeightmapsPerLayer->TextureUAV);
+		SetUAVParameter(RHICmdList, RHICmdList.GetBoundComputeShader(), AtlasPaintListsParam, InParams.AtlasWeightmapsPerLayer->GetTextureUAV(/*InMipLevel = */0));
 		SetSRVParameter(RHICmdList, RHICmdList.GetBoundComputeShader(), ComponentsDataParam, InParams.ComputeShaderResource->ComponentsDataSRV);
 		SetShaderValue(RHICmdList, RHICmdList.GetBoundComputeShader(), ComponentSizeParam, InParams.ComponentSize);
 	}
@@ -1079,7 +1052,7 @@
 
 	void SetParameters(FRHICommandList& RHICmdList, const FLandscapeLayerWeightmapPackMaterialLayersComputeShaderParameters& InParams)
 	{
-		SetUAVParameter(RHICmdList, RHICmdList.GetBoundComputeShader(), ComponentWeightmapParam, InParams.ComponentWeightmapResource->TextureUAV);
+		SetUAVParameter(RHICmdList, RHICmdList.GetBoundComputeShader(), ComponentWeightmapParam, InParams.ComponentWeightmapResource->GetTextureUAV(/*InMipLevel = */0));
 		SetTextureParameter(RHICmdList, RHICmdList.GetBoundComputeShader(), AtlasPaintListsParam, InParams.AtlasWeightmapsPerLayer->TextureRHI);
 		SetSRVParameter(RHICmdList, RHICmdList.GetBoundComputeShader(), ComponentsDataParam, InParams.ComputeShaderResource->ComponentsDataSRV);
 		SetShaderValue(RHICmdList, RHICmdList.GetBoundComputeShader(), ComponentSizeParam, InParams.ComponentSize);
@@ -1133,6 +1106,313 @@
 	FLandscapeLayerWeightmapPackMaterialLayersComputeShaderParameters ShaderParams;
 };
 
+
+// Landscape Heightmap PS
+
+// ----------------------------------------------------------------------------------
+
+class FLandscapeLayersHeightmapsMergeEditLayersPS : public FGlobalShader
+{
+	DECLARE_GLOBAL_SHADER(FLandscapeLayersHeightmapsMergeEditLayersPS);
+	SHADER_USE_PARAMETER_STRUCT(FLandscapeLayersHeightmapsMergeEditLayersPS, FGlobalShader);
+
+public:
+	BEGIN_SHADER_PARAMETER_STRUCT(FParameters, )
+		SHADER_PARAMETER(uint32, InNumEditLayers)
+		SHADER_PARAMETER_RDG_TEXTURE_SRV(Texture2DArray<float4>, InEditLayersTextures)
+		SHADER_PARAMETER_RDG_BUFFER_SRV(StructuredBuffer<FLandscapeEditLayerHeightmapMergeInfo>, InEditLayersMergeInfos)
+		RENDER_TARGET_BINDING_SLOTS()
+	END_SHADER_PARAMETER_STRUCT()
+
+	static bool ShouldCompilePermutation(const FGlobalShaderPermutationParameters& InParameters)
+	{
+		return IsFeatureLevelSupported(InParameters.Platform, ERHIFeatureLevel::SM5) && !IsConsolePlatform(InParameters.Platform) && !IsMetalMobilePlatform(InParameters.Platform);
+	}
+
+	static void ModifyCompilationEnvironment(const FGlobalShaderPermutationParameters& InParameters, FShaderCompilerEnvironment& OutEnvironment)
+	{
+		OutEnvironment.SetDefine(TEXT("MERGE_EDIT_LAYERS"), 1);
+	}
+
+	static void MergeEditLayers(FRDGBuilder& GraphBuilder, FParameters* InParameters, const FIntPoint& InTextureSize)
+	{
+		FGlobalShaderMap* ShaderMap = GetGlobalShaderMap(GMaxRHIFeatureLevel);
+		TShaderMapRef<FLandscapeLayersHeightmapsMergeEditLayersPS> PixelShader(ShaderMap);
+
+		FPixelShaderUtils::AddFullscreenPass(
+			GraphBuilder,
+			ShaderMap,
+			RDG_EVENT_NAME("LandscapeLayers_MergeEditLayers"),
+			PixelShader,
+			InParameters,
+			FIntRect(0, 0, InTextureSize.X, InTextureSize.Y));
+	}
+};
+
+IMPLEMENT_GLOBAL_SHADER(FLandscapeLayersHeightmapsMergeEditLayersPS, "/Engine/Private/Landscape/LandscapeLayersHeightmapsPS.usf", "MergeEditLayers", SF_Pixel);
+
+
+// ----------------------------------------------------------------------------------
+
+class FLandscapeLayersHeightmapsStitchHeightmapPS : public FGlobalShader
+{
+	DECLARE_GLOBAL_SHADER(FLandscapeLayersHeightmapsStitchHeightmapPS);
+	SHADER_USE_PARAMETER_STRUCT(FLandscapeLayersHeightmapsStitchHeightmapPS, FGlobalShader);
+
+public:
+	BEGIN_SHADER_PARAMETER_STRUCT(FParameters, )
+		SHADER_PARAMETER(FUintVector2, InSourceTextureSize)
+		SHADER_PARAMETER_RDG_TEXTURE_SRV(Texture2DArray<float2>, InSourceHeightmaps)
+		SHADER_PARAMETER_SCALAR_ARRAY(uint32, InNeighborHeightmapIndices, [9])
+		RENDER_TARGET_BINDING_SLOTS()
+	END_SHADER_PARAMETER_STRUCT()
+
+	static bool ShouldCompilePermutation(const FGlobalShaderPermutationParameters& Parameters)
+	{
+		return IsFeatureLevelSupported(Parameters.Platform, ERHIFeatureLevel::SM5) && !IsConsolePlatform(Parameters.Platform) && !IsMetalMobilePlatform(Parameters.Platform);
+	}
+
+	static void ModifyCompilationEnvironment(const FGlobalShaderPermutationParameters& Parameters, FShaderCompilerEnvironment& OutEnvironment)
+	{
+		OutEnvironment.SetDefine(TEXT("STITCH_HEIGHTMAP"), 1);
+	}
+
+	static void StitchHeightmap(FRDGBuilder& GraphBuilder, FParameters* InParameters)
+	{
+		FGlobalShaderMap* ShaderMap = GetGlobalShaderMap(GMaxRHIFeatureLevel);
+		TShaderMapRef<FLandscapeLayersHeightmapsStitchHeightmapPS> PixelShader(ShaderMap);
+
+		FPixelShaderUtils::AddFullscreenPass(
+			GraphBuilder,
+			ShaderMap,
+			RDG_EVENT_NAME("LandscapeLayers_StitchHeightmap"),
+			PixelShader,
+			InParameters,
+			FIntRect(0, 0, InParameters->InSourceTextureSize.X, InParameters->InSourceTextureSize.Y));
+	}
+};
+
+IMPLEMENT_GLOBAL_SHADER(FLandscapeLayersHeightmapsStitchHeightmapPS, "/Engine/Private/Landscape/LandscapeLayersHeightmapsPS.usf", "StitchHeightmap", SF_Pixel);
+
+// ----------------------------------------------------------------------------------
+
+class FLandscapeLayersHeightmapsFinalizeHeightmapPS : public FGlobalShader
+{
+	DECLARE_GLOBAL_SHADER(FLandscapeLayersHeightmapsFinalizeHeightmapPS);
+	SHADER_USE_PARAMETER_STRUCT(FLandscapeLayersHeightmapsFinalizeHeightmapPS, FGlobalShader);
+
+public:
+	BEGIN_SHADER_PARAMETER_STRUCT(FParameters, )
+		SHADER_PARAMETER(FUintVector2, InSourceTextureSize)
+		SHADER_PARAMETER_RDG_TEXTURE_SRV(Texture2DArray<float2>, InSourceHeightmaps)
+		SHADER_PARAMETER_SCALAR_ARRAY(uint32, InNeighborHeightmapIndices, [9])
+		SHADER_PARAMETER(FUintVector4, InDestinationTextureSubregion)
+		SHADER_PARAMETER(FVector3f, InLandscapeGridScale)
+		RENDER_TARGET_BINDING_SLOTS()
+	END_SHADER_PARAMETER_STRUCT()
+
+	static bool ShouldCompilePermutation(const FGlobalShaderPermutationParameters& Parameters)
+	{
+		return IsFeatureLevelSupported(Parameters.Platform, ERHIFeatureLevel::SM5) && !IsConsolePlatform(Parameters.Platform) && !IsMetalMobilePlatform(Parameters.Platform);
+	}
+
+	static void ModifyCompilationEnvironment(const FGlobalShaderPermutationParameters& Parameters, FShaderCompilerEnvironment& OutEnvironment)
+	{
+		OutEnvironment.SetDefine(TEXT("FINALIZE_HEIGHTMAP"), 1);
+	}
+
+	static void FinalizeHeightmap(FRDGBuilder& GraphBuilder, FParameters* InParameters)
+	{
+		FGlobalShaderMap* ShaderMap = GetGlobalShaderMap(GMaxRHIFeatureLevel);
+		TShaderMapRef<FLandscapeLayersHeightmapsFinalizeHeightmapPS> PixelShader(ShaderMap);
+
+		FPixelShaderUtils::AddFullscreenPass(
+			GraphBuilder,
+			ShaderMap,
+			RDG_EVENT_NAME("LandscapeLayers_FinalizeHeightmap"),
+			PixelShader,
+			InParameters,
+			FIntRect(InParameters->InDestinationTextureSubregion.X, InParameters->InDestinationTextureSubregion.Y, InParameters->InDestinationTextureSubregion.Z, InParameters->InDestinationTextureSubregion.W));
+	}
+};
+
+IMPLEMENT_GLOBAL_SHADER(FLandscapeLayersHeightmapsFinalizeHeightmapPS, "/Engine/Private/Landscape/LandscapeLayersHeightmapsPS.usf", "FinalizeHeightmap", SF_Pixel);
+
+
+// ----------------------------------------------------------------------------------
+
+class FLandscapeLayersHeightmapsGenerateMipsPS : public FGlobalShader
+{
+	DECLARE_GLOBAL_SHADER(FLandscapeLayersHeightmapsGenerateMipsPS);
+	SHADER_USE_PARAMETER_STRUCT(FLandscapeLayersHeightmapsGenerateMipsPS, FGlobalShader);
+
+public:
+	BEGIN_SHADER_PARAMETER_STRUCT(FParameters, )
+		SHADER_PARAMETER(FUintVector2, InCurrentMipSubregionSize)
+		SHADER_PARAMETER_RDG_TEXTURE_SRV(Texture2D<float4>, InSourceHeightmap)
+		RENDER_TARGET_BINDING_SLOTS()
+	END_SHADER_PARAMETER_STRUCT()
+
+	static bool ShouldCompilePermutation(const FGlobalShaderPermutationParameters& Parameters)
+	{
+		return IsFeatureLevelSupported(Parameters.Platform, ERHIFeatureLevel::SM5) && !IsConsolePlatform(Parameters.Platform) && !IsMetalMobilePlatform(Parameters.Platform);
+	}
+
+	static void ModifyCompilationEnvironment(const FGlobalShaderPermutationParameters& Parameters, FShaderCompilerEnvironment& OutEnvironment)
+	{
+		OutEnvironment.SetDefine(TEXT("GENERATE_MIPS"), 1);
+	}
+
+	static void GenerateSingleMip(FRDGBuilder& GraphBuilder, FParameters* InParameters)
+	{
+		FGlobalShaderMap* ShaderMap = GetGlobalShaderMap(GMaxRHIFeatureLevel);
+		TShaderMapRef<FLandscapeLayersHeightmapsGenerateMipsPS> PixelShader(ShaderMap);
+
+		FIntVector MipSize = InParameters->RenderTargets[0].GetTexture()->Desc.GetSize();
+
+		FPixelShaderUtils::AddFullscreenPass(
+			GraphBuilder,
+			ShaderMap,
+			RDG_EVENT_NAME("LandscapeLayers_GenerateMip"),
+			PixelShader,
+			InParameters,
+			FIntRect(0, 0, MipSize.X, MipSize.Y));
+	}
+};
+
+IMPLEMENT_GLOBAL_SHADER(FLandscapeLayersHeightmapsGenerateMipsPS, "/Engine/Private/Landscape/LandscapeLayersHeightmapsPS.usf", "GenerateMips", SF_Pixel);
+
+// Landscape Weightmap PS
+class FLandscapeLayersWeightmapsMergeEditLayersPS : public FGlobalShader
+{
+	DECLARE_GLOBAL_SHADER(FLandscapeLayersWeightmapsMergeEditLayersPS);
+	SHADER_USE_PARAMETER_STRUCT(FLandscapeLayersWeightmapsMergeEditLayersPS, FGlobalShader);
+
+public:
+	BEGIN_SHADER_PARAMETER_STRUCT(FParameters, )
+		SHADER_PARAMETER(uint32, InNumEditLayers)
+		SHADER_PARAMETER(uint32, InStartIndexInEditLayersMergeInfos)
+		SHADER_PARAMETER_RDG_TEXTURE_SRV(Texture2DArray<float4>, InPackedWeightmaps)
+		SHADER_PARAMETER_RDG_BUFFER_SRV(StructuredBuffer<FLandscapeEditLayerWeightmapMergeInfo>, InEditLayersMergeInfos)
+		RENDER_TARGET_BINDING_SLOTS()
+		END_SHADER_PARAMETER_STRUCT()
+
+		static bool ShouldCompilePermutation(const FGlobalShaderPermutationParameters& InParameters)
+	{
+		return IsFeatureLevelSupported(InParameters.Platform, ERHIFeatureLevel::SM5) && !IsConsolePlatform(InParameters.Platform) && !IsMetalMobilePlatform(InParameters.Platform);
+	}
+
+	static void ModifyCompilationEnvironment(const FGlobalShaderPermutationParameters& InParameters, FShaderCompilerEnvironment& OutEnvironment)
+	{
+		OutEnvironment.SetDefine(TEXT("MERGE_EDIT_LAYERS"), 1);
+	}
+
+	static void MergeEditLayers(FRDGBuilder& GraphBuilder, FParameters* InParameters)
+	{
+		FGlobalShaderMap* ShaderMap = GetGlobalShaderMap(GMaxRHIFeatureLevel);
+		TShaderMapRef<FLandscapeLayersWeightmapsMergeEditLayersPS> PixelShader(ShaderMap);
+
+		FIntVector TextureSize = InParameters->RenderTargets[0].GetTexture()->Desc.GetSize();
+
+		FPixelShaderUtils::AddFullscreenPass(
+			GraphBuilder,
+			ShaderMap,
+			RDG_EVENT_NAME("LandscapeLayers_MergeEditLayers"),
+			PixelShader,
+			InParameters,
+			FIntRect(0, 0, TextureSize.X, TextureSize.Y));
+	}
+};
+
+IMPLEMENT_GLOBAL_SHADER(FLandscapeLayersWeightmapsMergeEditLayersPS, "/Engine/Private/Landscape/LandscapeLayersWeightmapsPS.usf", "MergeEditLayers", SF_Pixel);
+
+class FLandscapeLayersWeightmapsFinalizeWeightmapPS : public FGlobalShader
+{
+	DECLARE_GLOBAL_SHADER(FLandscapeLayersWeightmapsFinalizeWeightmapPS);
+	SHADER_USE_PARAMETER_STRUCT(FLandscapeLayersWeightmapsFinalizeWeightmapPS, FGlobalShader);
+
+public:
+	BEGIN_SHADER_PARAMETER_STRUCT(FParameters, )
+		SHADER_PARAMETER(uint32, InValidTextureChannelsMask)
+		SHADER_PARAMETER(FUintVector4, InPerChannelPaintLayerIndexInWeightmaps)
+		SHADER_PARAMETER(FUintVector4, InPerChannelStartPaintLayerIndex)
+		SHADER_PARAMETER(FUintVector4, InPerChannelNumPaintLayers)
+		SHADER_PARAMETER_RDG_TEXTURE_SRV_ARRAY(Texture2D<float4>, InPerChannelPaintLayerWeightmaps, [4])
+		SHADER_PARAMETER_RDG_BUFFER_SRV(StructuredBuffer<uint>, InPaintLayerInfoIndices)
+		SHADER_PARAMETER_RDG_BUFFER_SRV(StructuredBuffer<FLandscapeWeightmapPaintLayerInfo>, InPaintLayerInfos)
+		RENDER_TARGET_BINDING_SLOTS()
+		END_SHADER_PARAMETER_STRUCT()
+
+		static bool ShouldCompilePermutation(const FGlobalShaderPermutationParameters& Parameters)
+	{
+		return IsFeatureLevelSupported(Parameters.Platform, ERHIFeatureLevel::SM5) && !IsConsolePlatform(Parameters.Platform) && !IsMetalMobilePlatform(Parameters.Platform);
+	}
+
+	static void ModifyCompilationEnvironment(const FGlobalShaderPermutationParameters& Parameters, FShaderCompilerEnvironment& OutEnvironment)
+	{
+		OutEnvironment.SetDefine(TEXT("FINALIZE_WEIGHTMAP"), 1);
+	}
+
+	static void FinalizeWeightmap(FRDGBuilder& GraphBuilder, FParameters* InParameters)
+	{
+		FGlobalShaderMap* ShaderMap = GetGlobalShaderMap(GMaxRHIFeatureLevel);
+		TShaderMapRef<FLandscapeLayersWeightmapsFinalizeWeightmapPS> PixelShader(ShaderMap);
+
+		FIntVector TextureSize = InParameters->RenderTargets[0].GetTexture()->Desc.GetSize();
+
+		FPixelShaderUtils::AddFullscreenPass(
+			GraphBuilder,
+			ShaderMap,
+			RDG_EVENT_NAME("LandscapeLayers_PackWeightmap"),
+			PixelShader,
+			InParameters,
+			FIntRect(0, 0, TextureSize.X, TextureSize.Y));
+	}
+};
+
+IMPLEMENT_GLOBAL_SHADER(FLandscapeLayersWeightmapsFinalizeWeightmapPS, "/Engine/Private/Landscape/LandscapeLayersWeightmapsPS.usf", "FinalizeWeightmap", SF_Pixel);
+
+class FLandscapeLayersWeightmapsGenerateMipsPS : public FGlobalShader
+{
+	DECLARE_GLOBAL_SHADER(FLandscapeLayersWeightmapsGenerateMipsPS);
+	SHADER_USE_PARAMETER_STRUCT(FLandscapeLayersWeightmapsGenerateMipsPS, FGlobalShader);
+
+public:
+	BEGIN_SHADER_PARAMETER_STRUCT(FParameters, )
+		SHADER_PARAMETER(FUintVector2, InCurrentMipSize)
+		SHADER_PARAMETER_RDG_TEXTURE_SRV(Texture2D<float4>, InSourceWeightmap)
+		RENDER_TARGET_BINDING_SLOTS()
+		END_SHADER_PARAMETER_STRUCT()
+
+		static bool ShouldCompilePermutation(const FGlobalShaderPermutationParameters& Parameters)
+	{
+		return IsFeatureLevelSupported(Parameters.Platform, ERHIFeatureLevel::SM5) && !IsConsolePlatform(Parameters.Platform) && !IsMetalMobilePlatform(Parameters.Platform);
+	}
+
+	static void ModifyCompilationEnvironment(const FGlobalShaderPermutationParameters& Parameters, FShaderCompilerEnvironment& OutEnvironment)
+	{
+		OutEnvironment.SetDefine(TEXT("GENERATE_MIPS"), 1);
+	}
+
+	static void GenerateSingleMip(FRDGBuilder& GraphBuilder, FParameters* InParameters)
+	{
+		FGlobalShaderMap* ShaderMap = GetGlobalShaderMap(GMaxRHIFeatureLevel);
+		TShaderMapRef<FLandscapeLayersWeightmapsGenerateMipsPS> PixelShader(ShaderMap);
+
+		FIntVector MipSize = InParameters->RenderTargets[0].GetTexture()->Desc.GetSize();
+
+		FPixelShaderUtils::AddFullscreenPass(
+			GraphBuilder,
+			ShaderMap,
+			RDG_EVENT_NAME("LandscapeLayers_GenerateMip"),
+			PixelShader,
+			InParameters,
+			FIntRect(0, 0, MipSize.X, MipSize.Y));
+	}
+};
+
+IMPLEMENT_GLOBAL_SHADER(FLandscapeLayersWeightmapsGenerateMipsPS, "/Engine/Private/Landscape/LandscapeLayersWeightmapsPS.usf", "GenerateMips", SF_Pixel);
+
 // Copy texture render command
 
 struct FLandscapeLayersCopyTextureParams
@@ -1182,7 +1462,9 @@
 public:
 	FLandscapeLayersCopyTexture_RenderThread(const FLandscapeLayersCopyTextureParams& InParams)
 		: Params(InParams)
-	{}	
+	{}
 
+	const FLandscapeLayersCopyTextureParams& GetParams() const { return Params; }
+
 	void Copy(FRHICommandListImmediate& InRHICmdList)
 	{
@@ -1204,7 +1486,7 @@
 		Info.DestPosition.Y = Params.DestPosition.Y;
 		Info.SourceSliceIndex = Params.SourceArrayIndex;
 		Info.DestSliceIndex = Params.DestArrayIndex;
-		Info.SourceMipIndex = Params.SourceMip; 
+		Info.SourceMipIndex = Params.SourceMip;
 		Info.DestMipIndex = Params.DestMip;
 
 		check((Info.SourcePosition.X >= 0) && (Info.SourcePosition.Y >= 0) && (Info.DestPosition.X >= 0) && (Info.DestPosition.Y >= 0));
@@ -1263,7 +1545,7 @@
 public:
 
 	FLandscapeLayersRender_RenderThread(const FString& InDebugName, UTextureRenderTarget2D* InWriteRenderTarget, const FIntPoint& InWriteRenderTargetSize, const FIntPoint& InReadRenderTargetSize, const FMatrix& InProjectionMatrix,
-										const ShaderDataType& InShaderParams, uint8 InCurrentMip, const TArray<FLandscapeLayersTriangle>& InTriangleList)
+		const ShaderDataType& InShaderParams, uint8 InCurrentMip, const TArray<FLandscapeLayersTriangle>& InTriangleList)
 		: RenderTargetResource(InWriteRenderTarget->GameThread_GetRenderTargetResource())
 		, WriteRenderTargetSize(InWriteRenderTargetSize)
 		, ReadRenderTargetSize(InReadRenderTargetSize)
@@ -1468,7 +1750,7 @@
 	{
 		return false;
 	}
-	
+
 	OutLandscapeTransform = Landscape->GetTransform();
 	FVector OffsetVector(LandscapeExtent.Min.X, LandscapeExtent.Min.Y, 0.f);
 	FVector Translation = OutLandscapeTransform.TransformFVector4(OffsetVector);
@@ -1482,6 +1764,23 @@
 	return true;
 }
 
+bool ALandscape::SupportsEditLayersLocalMerge()
+{
+	ALandscape* Landscape = GetLandscapeActor();
+	check(Landscape);
+
+	for (const FLandscapeLayer& Layer : Landscape->LandscapeLayers)
+	{
+		// No BP brush is supported for local merge of edit layers : 
+		if (Layer.bVisible && Algo::AnyOf(Layer.Brushes, [](const FLandscapeLayerBrush& Brush) { return (Brush.GetBrush() != nullptr) && Brush.GetBrush()->IsVisible(); }))
+		{
+			return false;
+		}
+	}
+
+	return !!LandscapeEditLayersLocalMerge;
+}
+
 void ALandscape::CreateLayersRenderingResource()
 {
 	ULandscapeInfo* Info = GetLandscapeInfo();
@@ -1501,139 +1800,149 @@
 	ALandscape* Landscape = GetLandscapeActor();
 	check(Landscape);
 
-	if (Landscape->HeightmapRTList.Num() == 0)
+	if (SupportsEditLayersLocalMerge())
 	{
-		Landscape->HeightmapRTList.Init(nullptr, (int32)EHeightmapRTType::HeightmapRT_Count);
+		ReleaseLayersRenderingResource();
 
-		int32 CurrentMipSizeX = ((SubsectionSizeQuads + 1) * NumSubsections) * ComponentCounts.X;
-		int32 CurrentMipSizeY = ((SubsectionSizeQuads + 1) * NumSubsections) * ComponentCounts.Y;
+		Landscape->HeightmapRTList.Empty();
+		Landscape->WeightmapRTList.Empty();
+	}
+	else
+	{
+		if (Landscape->HeightmapRTList.Num() == 0)
+		{
+			Landscape->HeightmapRTList.Init(nullptr, (int32)EHeightmapRTType::HeightmapRT_Count);
 
-		for (int32 i = 0; i < (int32)EHeightmapRTType::HeightmapRT_Count; ++i)
-		{
-			FText DisplayName = StaticEnum<EHeightmapRTType>()->GetDisplayValueAsText((EHeightmapRTType)i);
-			FName RTName = MakeUniqueObjectName(GetTransientPackage(), UTextureRenderTarget2D::StaticClass(), FName(*DisplayName.ToString()));
-			Landscape->HeightmapRTList[i] = NewObject<UTextureRenderTarget2D>(GetTransientPackage(), RTName, RF_Transient);
-			check(Landscape->HeightmapRTList[i]);
-			Landscape->HeightmapRTList[i]->RenderTargetFormat = RTF_RGBA8;
-			Landscape->HeightmapRTList[i]->AddressX = TextureAddress::TA_Clamp;
-			Landscape->HeightmapRTList[i]->AddressY = TextureAddress::TA_Clamp;
-			Landscape->HeightmapRTList[i]->ClearColor = FLinearColor(0.0f, 0.0f, 0.0f, 0.0f);
+			int32 CurrentMipSizeX = ((SubsectionSizeQuads + 1) * NumSubsections) * ComponentCounts.X;
+			int32 CurrentMipSizeY = ((SubsectionSizeQuads + 1) * NumSubsections) * ComponentCounts.Y;
 
-			if (i < (int32)EHeightmapRTType::HeightmapRT_Mip1) // Landscape size RT
+			for (int32 i = 0; i < (int32)EHeightmapRTType::HeightmapRT_Count; ++i)
 			{
-				Landscape->HeightmapRTList[i]->InitAutoFormat(FMath::RoundUpToPowerOfTwo(CurrentMipSizeX), FMath::RoundUpToPowerOfTwo(CurrentMipSizeY));
-			}
-			else // Mips
-			{
-				CurrentMipSizeX >>= 1;
-				CurrentMipSizeY >>= 1;
-				Landscape->HeightmapRTList[i]->InitAutoFormat(FMath::RoundUpToPowerOfTwo(CurrentMipSizeX), FMath::RoundUpToPowerOfTwo(CurrentMipSizeY));
-			}
+				FText DisplayName = StaticEnum<EHeightmapRTType>()->GetDisplayValueAsText((EHeightmapRTType)i);
+				FName RTName = MakeUniqueObjectName(GetTransientPackage(), UTextureRenderTarget2D::StaticClass(), FName(*DisplayName.ToString()));
+				Landscape->HeightmapRTList[i] = NewObject<UTextureRenderTarget2D>(GetTransientPackage(), RTName, RF_Transient);
+				check(Landscape->HeightmapRTList[i]);
+				Landscape->HeightmapRTList[i]->RenderTargetFormat = RTF_RGBA8;
+				Landscape->HeightmapRTList[i]->AddressX = TextureAddress::TA_Clamp;
+				Landscape->HeightmapRTList[i]->AddressY = TextureAddress::TA_Clamp;
+				Landscape->HeightmapRTList[i]->ClearColor = FLinearColor(0.0f, 0.0f, 0.0f, 0.0f);
 
-			Landscape->HeightmapRTList[i]->UpdateResourceImmediate(true);
+				if (i < (int32)EHeightmapRTType::HeightmapRT_Mip1) // Landscape size RT
+				{
+					Landscape->HeightmapRTList[i]->InitAutoFormat(FMath::RoundUpToPowerOfTwo(CurrentMipSizeX), FMath::RoundUpToPowerOfTwo(CurrentMipSizeY));
+				}
+				else // Mips
+				{
+					CurrentMipSizeX >>= 1;
+					CurrentMipSizeY >>= 1;
+					Landscape->HeightmapRTList[i]->InitAutoFormat(FMath::RoundUpToPowerOfTwo(CurrentMipSizeX), FMath::RoundUpToPowerOfTwo(CurrentMipSizeY));
+				}
 
-			// Only generate required mips RT
-			if (CurrentMipSizeX == ComponentCounts.X && CurrentMipSizeY == ComponentCounts.Y)
-			{
-				break;
+				Landscape->HeightmapRTList[i]->UpdateResourceImmediate(true);
+
+				// Only generate required mips RT
+				if (CurrentMipSizeX == ComponentCounts.X && CurrentMipSizeY == ComponentCounts.Y)
+				{
+					break;
+				}
 			}
 		}
-	}
-	else // Simply resize the render target
-	{
-		int32 CurrentMipSizeX = ((SubsectionSizeQuads + 1) * NumSubsections) * ComponentCounts.X;
-		int32 CurrentMipSizeY = ((SubsectionSizeQuads + 1) * NumSubsections) * ComponentCounts.Y;
+		else // Simply resize the render target
+		{
+			int32 CurrentMipSizeX = ((SubsectionSizeQuads + 1) * NumSubsections) * ComponentCounts.X;
+			int32 CurrentMipSizeY = ((SubsectionSizeQuads + 1) * NumSubsections) * ComponentCounts.Y;
 
-		for (int32 i = 0; i < (int32)EHeightmapRTType::HeightmapRT_Count; ++i)
-		{
-			if (i < (int32)EHeightmapRTType::HeightmapRT_Mip1) // Landscape size RT
+			for (int32 i = 0; i < (int32)EHeightmapRTType::HeightmapRT_Count; ++i)
 			{
-				Landscape->HeightmapRTList[i]->ResizeTarget(FMath::RoundUpToPowerOfTwo(CurrentMipSizeX), FMath::RoundUpToPowerOfTwo(CurrentMipSizeY));
-			}
-			else // Mips
-			{
-				CurrentMipSizeX >>= 1;
-				CurrentMipSizeY >>= 1;
-				Landscape->HeightmapRTList[i]->ResizeTarget(FMath::RoundUpToPowerOfTwo(CurrentMipSizeX), FMath::RoundUpToPowerOfTwo(CurrentMipSizeY));
-			}
+				if (i < (int32)EHeightmapRTType::HeightmapRT_Mip1) // Landscape size RT
+				{
+					Landscape->HeightmapRTList[i]->ResizeTarget(FMath::RoundUpToPowerOfTwo(CurrentMipSizeX), FMath::RoundUpToPowerOfTwo(CurrentMipSizeY));
+				}
+				else // Mips
+				{
+					CurrentMipSizeX >>= 1;
+					CurrentMipSizeY >>= 1;
+					Landscape->HeightmapRTList[i]->ResizeTarget(FMath::RoundUpToPowerOfTwo(CurrentMipSizeX), FMath::RoundUpToPowerOfTwo(CurrentMipSizeY));
+				}
 
-			// Only generate required mips RT
-			if (CurrentMipSizeX == ComponentCounts.X && CurrentMipSizeY == ComponentCounts.Y)
-			{
-				break;
+				// Only generate required mips RT
+				if (CurrentMipSizeX == ComponentCounts.X && CurrentMipSizeY == ComponentCounts.Y)
+				{
+					break;
+				}
 			}
 		}
-	}
 
-	if (Landscape->WeightmapRTList.Num() == 0)
-	{
-		Landscape->WeightmapRTList.Init(nullptr, (int32)EWeightmapRTType::WeightmapRT_Count);
+		if (Landscape->WeightmapRTList.Num() == 0)
+		{
+			Landscape->WeightmapRTList.Init(nullptr, (int32)EWeightmapRTType::WeightmapRT_Count);
 
-		int32 CurrentMipSizeX = ((SubsectionSizeQuads + 1) * NumSubsections) * ComponentCounts.X;
-		int32 CurrentMipSizeY = ((SubsectionSizeQuads + 1) * NumSubsections) * ComponentCounts.Y;
+			int32 CurrentMipSizeX = ((SubsectionSizeQuads + 1) * NumSubsections) * ComponentCounts.X;
+			int32 CurrentMipSizeY = ((SubsectionSizeQuads + 1) * NumSubsections) * ComponentCounts.Y;
 
-		for (int32 i = 0; i < (int32)EWeightmapRTType::WeightmapRT_Count; ++i)
-		{
-			FText DisplayName = StaticEnum<EHeightmapRTType>()->GetDisplayValueAsText((EWeightmapRTType)i);
-			FName RTName = MakeUniqueObjectName(GetTransientPackage(), UTextureRenderTarget2D::StaticClass(), FName(*DisplayName.ToString()));
-			Landscape->WeightmapRTList[i] = NewObject<UTextureRenderTarget2D>(GetTransientPackage(), RTName, RF_Transient);
+			for (int32 i = 0; i < (int32)EWeightmapRTType::WeightmapRT_Count; ++i)
+			{
+				FText DisplayName = StaticEnum<EHeightmapRTType>()->GetDisplayValueAsText((EWeightmapRTType)i);
+				FName RTName = MakeUniqueObjectName(GetTransientPackage(), UTextureRenderTarget2D::StaticClass(), FName(*DisplayName.ToString()));
+				Landscape->WeightmapRTList[i] = NewObject<UTextureRenderTarget2D>(GetTransientPackage(), RTName, RF_Transient);
 
-			check(Landscape->WeightmapRTList[i]);
-			Landscape->WeightmapRTList[i]->AddressX = TextureAddress::TA_Clamp;
-			Landscape->WeightmapRTList[i]->AddressY = TextureAddress::TA_Clamp;
-			Landscape->WeightmapRTList[i]->ClearColor = FLinearColor(0.0f, 0.0f, 0.0f, 0.0f);
-			Landscape->WeightmapRTList[i]->RenderTargetFormat = RTF_RGBA8;
+				check(Landscape->WeightmapRTList[i]);
+				Landscape->WeightmapRTList[i]->AddressX = TextureAddress::TA_Clamp;
+				Landscape->WeightmapRTList[i]->AddressY = TextureAddress::TA_Clamp;
+				Landscape->WeightmapRTList[i]->ClearColor = FLinearColor(0.0f, 0.0f, 0.0f, 0.0f);
+				Landscape->WeightmapRTList[i]->RenderTargetFormat = RTF_RGBA8;
 
-			if (i < (int32)EWeightmapRTType::WeightmapRT_Mip0) // Landscape size RT, only create the number of layer we have
-			{
-				Landscape->WeightmapRTList[i]->RenderTargetFormat = (i == (int32)EWeightmapRTType::WeightmapRT_Scratch_RGBA) ? RTF_RGBA8 : RTF_R8;
-				Landscape->WeightmapRTList[i]->InitAutoFormat(FMath::RoundUpToPowerOfTwo(CurrentMipSizeX), FMath::RoundUpToPowerOfTwo(CurrentMipSizeY));
-			}
-			else // Mips
-			{
-				Landscape->WeightmapRTList[i]->InitAutoFormat(FMath::RoundUpToPowerOfTwo(CurrentMipSizeX), FMath::RoundUpToPowerOfTwo(CurrentMipSizeY));
+				if (i < (int32)EWeightmapRTType::WeightmapRT_Mip0) // Landscape size RT, only create the number of layer we have
+				{
+					Landscape->WeightmapRTList[i]->RenderTargetFormat = (i == (int32)EWeightmapRTType::WeightmapRT_Scratch_RGBA) ? RTF_RGBA8 : RTF_R8;
+					Landscape->WeightmapRTList[i]->InitAutoFormat(FMath::RoundUpToPowerOfTwo(CurrentMipSizeX), FMath::RoundUpToPowerOfTwo(CurrentMipSizeY));
+				}
+				else // Mips
+				{
+					Landscape->WeightmapRTList[i]->InitAutoFormat(FMath::RoundUpToPowerOfTwo(CurrentMipSizeX), FMath::RoundUpToPowerOfTwo(CurrentMipSizeY));
 
-				CurrentMipSizeX >>= 1;
-				CurrentMipSizeY >>= 1;
-			}
+					CurrentMipSizeX >>= 1;
+					CurrentMipSizeY >>= 1;
+				}
 
-			Landscape->WeightmapRTList[i]->UpdateResourceImmediate(true);
+				Landscape->WeightmapRTList[i]->UpdateResourceImmediate(true);
 
-			// Only generate required mips RT
-			if (CurrentMipSizeX < ComponentCounts.X && CurrentMipSizeY < ComponentCounts.Y)
-			{
-				break;
+				// Only generate required mips RT
+				if (CurrentMipSizeX < ComponentCounts.X && CurrentMipSizeY < ComponentCounts.Y)
+				{
+					break;
+				}
 			}
 		}
-	}
-	else // Simply resize the render target
-	{
-		int32 CurrentMipSizeX = ((SubsectionSizeQuads + 1) * NumSubsections) * ComponentCounts.X;
-		int32 CurrentMipSizeY = ((SubsectionSizeQuads + 1) * NumSubsections) * ComponentCounts.Y;
+		else // Simply resize the render target
+		{
+			int32 CurrentMipSizeX = ((SubsectionSizeQuads + 1) * NumSubsections) * ComponentCounts.X;
+			int32 CurrentMipSizeY = ((SubsectionSizeQuads + 1) * NumSubsections) * ComponentCounts.Y;
 
-		for (int32 i = 0; i < (int32)EWeightmapRTType::WeightmapRT_Count; ++i)
-		{
-			if (i < (int32)EWeightmapRTType::WeightmapRT_Mip0) // Landscape size RT, only create the number of layer we have
+			for (int32 i = 0; i < (int32)EWeightmapRTType::WeightmapRT_Count; ++i)
 			{
-				Landscape->WeightmapRTList[i]->ResizeTarget(FMath::RoundUpToPowerOfTwo(CurrentMipSizeX), FMath::RoundUpToPowerOfTwo(CurrentMipSizeY));
-			}
-			else // Mips
-			{
-				Landscape->WeightmapRTList[i]->ResizeTarget(FMath::RoundUpToPowerOfTwo(CurrentMipSizeX), FMath::RoundUpToPowerOfTwo(CurrentMipSizeY));
+				if (i < (int32)EWeightmapRTType::WeightmapRT_Mip0) // Landscape size RT, only create the number of layer we have
+				{
+					Landscape->WeightmapRTList[i]->ResizeTarget(FMath::RoundUpToPowerOfTwo(CurrentMipSizeX), FMath::RoundUpToPowerOfTwo(CurrentMipSizeY));
+				}
+				else // Mips
+				{
+					Landscape->WeightmapRTList[i]->ResizeTarget(FMath::RoundUpToPowerOfTwo(CurrentMipSizeX), FMath::RoundUpToPowerOfTwo(CurrentMipSizeY));
 
-				CurrentMipSizeX >>= 1;
-				CurrentMipSizeY >>= 1;
-			}
+					CurrentMipSizeX >>= 1;
+					CurrentMipSizeY >>= 1;
+				}
 
-			// Only generate required mips RT
-			if (CurrentMipSizeX < ComponentCounts.X && CurrentMipSizeY < ComponentCounts.Y)
-			{
-				break;
+				// Only generate required mips RT
+				if (CurrentMipSizeX < ComponentCounts.X && CurrentMipSizeY < ComponentCounts.Y)
+				{
+					break;
+				}
 			}
 		}
+
+		InitializeLayersWeightmapResources();
 	}
-
-	InitializeLayersWeightmapResources();
 }
 
 void ALandscape::ToggleCanHaveLayersContent()
@@ -1660,7 +1969,6 @@
 
 void ALandscape::ReleaseLayersRenderingResource()
 {
-	check(!CanHaveLayersContent());
 	ULandscapeInfo* Info = GetLandscapeInfo();
 	if (Info == nullptr)
 	{
@@ -2229,11 +2537,39 @@
 	}
 }
 
+TArray<FLandscapeLayersCopyReadbackTextureParams> PrepareLandscapeLayersCopyReadbackTextureParams(const FTextureToComponentHelper& InMapHelper, TArray<UTexture2D*> InTextures, bool bWeightmaps)
+{
+	TRACE_CPUPROFILER_EVENT_SCOPE(LandscapeLayers_PrepareLandscapeLayersCopyReadbackTextureParams);
+	TArray<FLandscapeLayersCopyReadbackTextureParams> Result;
+	Result.Reserve(InTextures.Num());
+
+	for (UTexture2D* Texture : InTextures)
+	{
+		const TMap<UTexture2D*, TArray<ULandscapeComponent*>>& TexturesToComponents = bWeightmaps ? InMapHelper.WeightmapToComponents : InMapHelper.HeightmapToComponents;
+		const TArray<ULandscapeComponent*>* Components = TexturesToComponents.Find(Texture);
+		check(Components && !Components->IsEmpty());
+		ALandscapeProxy* Proxy = (*Components)[0]->GetLandscapeProxy();
+		FLandscapeEditLayerReadback** CPUReadback = bWeightmaps ? Proxy->WeightmapsCPUReadback.Find(Texture) : Proxy->HeightmapsCPUReadback.Find(Texture);
+		check(CPUReadback && *CPUReadback);
+
+		FLandscapeLayersCopyReadbackTextureParams& CopyReadbackTextureParams = Result.Add_GetRef(FLandscapeLayersCopyReadbackTextureParams(Texture, *CPUReadback));
+		// Init the CPU read back contexts for all components dependent on this texture. This includes a context containing the current component states : 
+		for (ULandscapeComponent* ComponentToResolve : *Components)
+		{
+			const FIntPoint ComponentToResolveKey = ComponentToResolve->GetSectionBase() / ComponentToResolve->ComponentSizeQuads;
+			const int32 ComponentToResolveFlags = ComponentToResolve->GetLayerUpdateFlagPerMode();
+			CopyReadbackTextureParams.Context.Add(FLandscapeEditLayerReadback::FComponentReadbackContext(ComponentToResolveKey, ComponentToResolveFlags));
+		}
+	}
+
+	return Result;
+}
+
 void ALandscape::CopyTexturePS(const FString& InSourceDebugName, FTextureResource* InSourceResource, const FString& InDestDebugName, FTextureResource* InDestResource) const
 {
 	check(InSourceResource != nullptr);
 	check(InDestResource != nullptr);
-	
+
 	ENQUEUE_RENDER_COMMAND(LandscapeLayers_Cmd_CopyTexturePS)(
 		[InSourceDebugName, InSourceResource, InDestDebugName, InDestResource](FRHICommandListImmediate& RHICmdList)
 	{
@@ -2271,7 +2607,7 @@
 }
 
 void ALandscape::DrawWeightmapComponentsToRenderTarget(const FString& InDebugName, const TArray<FIntPoint>& InSectionBaseList, const FVector2f& InScaleBias, TArray<FVector2f>* InScaleBiasPerSection, UTexture* InWeightmapRTRead, UTextureRenderTarget2D* InOptionalWeightmapRTRead2, UTextureRenderTarget2D* InWeightmapRTWrite,
-														ERTDrawingType InDrawType, bool InClearRTWrite, FLandscapeLayersWeightmapShaderParameters& InShaderParams, uint8 InMipRender) const
+	ERTDrawingType InDrawType, bool InClearRTWrite, FLandscapeLayersWeightmapShaderParameters& InShaderParams, uint8 InMipRender) const
 {
 	check(InWeightmapRTRead != nullptr);
 	check(InWeightmapRTWrite != nullptr);
@@ -2356,7 +2692,7 @@
 }
 
 void ALandscape::DrawWeightmapComponentsToRenderTarget(const FString& InDebugName, const TArray<ULandscapeComponent*>& InComponentsToDraw, const FIntPoint& InLandscapeBase, UTexture* InWeightmapRTRead, UTextureRenderTarget2D* InOptionalWeightmapRTRead2, UTextureRenderTarget2D* InWeightmapRTWrite,
-														ERTDrawingType InDrawType, bool InClearRTWrite, FLandscapeLayersWeightmapShaderParameters& InShaderParams, uint8 InMipRender) const
+	ERTDrawingType InDrawType, bool InClearRTWrite, FLandscapeLayersWeightmapShaderParameters& InShaderParams, uint8 InMipRender) const
 {
 
 	TArray<FIntPoint> SectionBaseList;
@@ -2405,7 +2741,7 @@
 		if (WriteMipRT != nullptr)
 		{
 			DrawWeightmapComponentsToRenderTarget(FString::Printf(TEXT("LS Weight: %s = -> %s Mips %d"), *ReadMipRT->GetName(), *WriteMipRT->GetName(), CurrentMip),
-													SectionBaseToDraw, WeightmapScaleBias, nullptr, ReadMipRT, nullptr, WriteMipRT, ERTDrawingType::RTMips, InClearRTWrite, InShaderParams, CurrentMip);
+				SectionBaseToDraw, WeightmapScaleBias, nullptr, ReadMipRT, nullptr, WriteMipRT, ERTDrawingType::RTMips, InClearRTWrite, InShaderParams, CurrentMip);
 			++CurrentMip;
 		}
 
@@ -2437,7 +2773,7 @@
 		if (WriteMipRT != nullptr)
 		{
 			DrawHeightmapComponentsToRenderTarget(FString::Printf(TEXT("LS Height: %s = -> %s CombinedAtlasWithMips %d"), *ReadMipRT->GetName(), *WriteMipRT->GetName(), CurrentMip),
-												  InComponentsToDraw, InLandscapeBase, ReadMipRT, nullptr, WriteMipRT, ERTDrawingType::RTMips, InClearRTWrite, InShaderParams, CurrentMip);
+				InComponentsToDraw, InLandscapeBase, ReadMipRT, nullptr, WriteMipRT, ERTDrawingType::RTMips, InClearRTWrite, InShaderParams, CurrentMip);
 			++CurrentMip;
 		}
 
@@ -2446,7 +2782,7 @@
 }
 
 void ALandscape::DrawHeightmapComponentsToRenderTarget(const FString& InDebugName, const TArray<ULandscapeComponent*>& InComponentsToDraw, const FIntPoint& InLandscapeBase, UTexture* InHeightmapRTRead, UTextureRenderTarget2D* InOptionalHeightmapRTRead2, UTextureRenderTarget2D* InHeightmapRTWrite,
-													  ERTDrawingType InDrawType, bool InClearRTWrite, FLandscapeLayersHeightmapShaderParameters& InShaderParams, uint8 InMipRender) const
+	ERTDrawingType InDrawType, bool InClearRTWrite, FLandscapeLayersHeightmapShaderParameters& InShaderParams, uint8 InMipRender) const
 {
 	check(InHeightmapRTRead != nullptr);
 	check(InHeightmapRTWrite != nullptr);
@@ -2472,36 +2808,36 @@
 
 		switch (InDrawType)
 		{
-			case ERTDrawingType::RTAtlas:
-			{
-				GenerateLayersRenderQuadsAtlas(ComponentSectionBase, HeightmapScaleBias, SubsectionSizeQuads, HeightmapReadTextureSize, HeightmapWriteTextureSize, TriangleList);
-			} break;
+		case ERTDrawingType::RTAtlas:
+		{
+			GenerateLayersRenderQuadsAtlas(ComponentSectionBase, HeightmapScaleBias, SubsectionSizeQuads, HeightmapReadTextureSize, HeightmapWriteTextureSize, TriangleList);
+		} break;
 
-			case ERTDrawingType::RTAtlasToNonAtlas:
-			{
-				GenerateLayersRenderQuadsAtlasToNonAtlas(ComponentSectionBase, HeightmapScaleBias, SubsectionSizeQuads, HeightmapReadTextureSize, HeightmapWriteTextureSize, TriangleList);
-			} break;
+		case ERTDrawingType::RTAtlasToNonAtlas:
+		{
+			GenerateLayersRenderQuadsAtlasToNonAtlas(ComponentSectionBase, HeightmapScaleBias, SubsectionSizeQuads, HeightmapReadTextureSize, HeightmapWriteTextureSize, TriangleList);
+		} break;
 
-			case ERTDrawingType::RTNonAtlas:
-			{
-				GenerateLayersRenderQuadsNonAtlas(ComponentSectionBase, HeightmapScaleBias, SubsectionSizeQuads, HeightmapReadTextureSize, HeightmapWriteTextureSize, TriangleList);
-			} break;
+		case ERTDrawingType::RTNonAtlas:
+		{
+			GenerateLayersRenderQuadsNonAtlas(ComponentSectionBase, HeightmapScaleBias, SubsectionSizeQuads, HeightmapReadTextureSize, HeightmapWriteTextureSize, TriangleList);
+		} break;
 
-			case ERTDrawingType::RTNonAtlasToAtlas:
-			{
-				GenerateLayersRenderQuadsNonAtlasToAtlas(ComponentSectionBase, HeightmapScaleBias, SubsectionSizeQuads, HeightmapReadTextureSize, HeightmapWriteTextureSize, TriangleList);
-			} break;
+		case ERTDrawingType::RTNonAtlasToAtlas:
+		{
+			GenerateLayersRenderQuadsNonAtlasToAtlas(ComponentSectionBase, HeightmapScaleBias, SubsectionSizeQuads, HeightmapReadTextureSize, HeightmapWriteTextureSize, TriangleList);
+		} break;
 
-			case ERTDrawingType::RTMips:
-			{
-				GenerateLayersRenderQuadsMip(ComponentSectionBase, HeightmapScaleBias, SubsectionSizeQuads, HeightmapReadTextureSize, HeightmapWriteTextureSize, InMipRender, TriangleList);
-			} break;
+		case ERTDrawingType::RTMips:
+		{
+			GenerateLayersRenderQuadsMip(ComponentSectionBase, HeightmapScaleBias, SubsectionSizeQuads, HeightmapReadTextureSize, HeightmapWriteTextureSize, InMipRender, TriangleList);
+		} break;
 
-			default:
-			{
-				check(false);
-				return;
-			}
+		default:
+		{
+			check(false);
+			return;
 		}
+		}
 	}
 
@@ -2517,7 +2853,7 @@
 	}
 
 	FMatrix ProjectionMatrix = AdjustProjectionMatrixForRHI(FTranslationMatrix(FVector(0, 0, 0)) *
-															FMatrix(FPlane(1.0f / (FMath::Max<uint32>(HeightmapWriteTextureSize.X, 1.f) / 2.0f), 0.0, 0.0f, 0.0f), FPlane(0.0f, -1.0f / (FMath::Max<uint32>(HeightmapWriteTextureSize.Y, 1.f) / 2.0f), 0.0f, 0.0f), FPlane(0.0f, 0.0f, 1.0f, 0.0f), FPlane(-1.0f, 1.0f, 0.0f, 1.0f)));
+		FMatrix(FPlane(1.0f / (FMath::Max<uint32>(HeightmapWriteTextureSize.X, 1.f) / 2.0f), 0.0, 0.0f, 0.0f), FPlane(0.0f, -1.0f / (FMath::Max<uint32>(HeightmapWriteTextureSize.Y, 1.f) / 2.0f), 0.0f, 0.0f), FPlane(0.0f, 0.0f, 1.0f, 0.0f), FPlane(-1.0f, 1.0f, 0.0f, 1.0f)));
 
 	FLandscapeLayersHeightmapRender_RenderThread LayersRender(InDebugName, InHeightmapRTWrite, HeightmapWriteTextureSize, HeightmapReadTextureSize, ProjectionMatrix, InShaderParams, InMipRender, TriangleList);
 
@@ -2720,9 +3056,9 @@
 	check(InReadSize.X == InWriteSize.X && InReadSize.Y == InWriteSize.Y);
 
 	int32 SubsectionSizeVerts = InSubSectionSizeQuad + 1;
-	
+
 	FVector2D UVSize((float)SubsectionSizeVerts / (float)InReadSize.X, (float)SubsectionSizeVerts / (float)InReadSize.Y);
-	
+
 	for (int8 SubY = 0; SubY < NumSubsections; ++SubY)
 	{
 		for (int8 SubX = 0; SubX < NumSubsections; ++SubX)
@@ -3173,8 +3509,1308 @@
 	return true;
 }
 
-int32 ALandscape::RegenerateLayersHeightmaps(FTextureToComponentHelper const& MapHelper, const TArray<ULandscapeComponent*>& InLandscapeComponentsToRender, const TArray<ULandscapeComponent*>& InLandscapeComponentsToResolve)
+// Must match EEditLayerHeightmapBlendMode in LandscapeLayersHeightmapsPS.usf
+enum class ELandscapeEditLayerHeightmapBlendMode : uint32
 {
+	Additive = 0,
+	AlphaBlend,
+
+	Num,
+};
+
+// Must match EEditLayerWeightmapBlendMode in LandscapeLayersWeightmapsPS.usf
+enum class ELandscapeEditLayerWeightmapBlendMode : uint32
+{
+	Additive = 0,
+	Substractive,
+
+	Num,
+};
+
+// Must match EWeightmapPaintLayerFlags in LandscapeLayersWeightmapsPS.usf
+enum class EWeightmapPaintLayerFlags : uint32
+{
+	IsVisibilityLayer = (1 << 0), // This paint layer is the visibility layer
+	IsWeightBlended = (1 << 1), // Blend the paint layer's value with all the other paint layers weights
+
+	None = 0
+};
+
+// Must match FEditLayerHeightmapMergeInfo in LandscapeLayersHeightmapsPS.usf
+struct FLandscapeEditLayerHeightmapMergeInfo
+{
+	// COMMENT [jonathan.bard] : not used at the moment because we copy to a texture 2D array but if we didn't and had instead N statically bound textures, we could save that copy and sample the textures directly :
+	FIntRect TextureSubregion; // Subregion of the source (edit layer) texture to use
+
+	ELandscapeEditLayerHeightmapBlendMode BlendMode = ELandscapeEditLayerHeightmapBlendMode::Num; // How this layer blends with the previous ones in the layers stack
+	float Alpha = 1.0f; // Alpha value to be used in the blend
+	uint32 Padding0; // Align to next float4 
+	uint32 Padding1;
+};
+
+
+// Must match FEditLayerWeightmapMergeInfo in LandscapeLayersWeightmapsPS.usf
+struct FLandscapeEditLayerWeightmapMergeInfo
+{
+	uint32 SourceWeightmapTextureIndex = (uint32)INDEX_NONE; // The index in InPackedWeightmaps of the texture to read from for this layer
+	uint32 SourceWeightmapTextureChannel = (uint32)INDEX_NONE; // The channel of the texture to read from for this layer
+	ELandscapeEditLayerWeightmapBlendMode BlendMode = ELandscapeEditLayerWeightmapBlendMode::Num; // How this layer blends with the previous ones in the layers stack
+	float Alpha = 1.0f; // Alpha value to be used in the blend
+};
+
+// Must match FWeightmapPaintLayerInfo in LandscapeLayersWeightmapsPS.usf
+struct FLandscapeWeightmapPaintLayerInfo
+{
+	EWeightmapPaintLayerFlags Flags = EWeightmapPaintLayerFlags::None; // Additional info about this paint layer
+};
+
+// Struct that contains all the information relevant for the edit layers update operation (list of dirty components, heightmaps, weightmaps, etc.
+//  Because this information can change during the course of the update (e.g. new weightmaps are added) it can be (partially or not) refreshed if necessary :
+struct FUpdateLayersContentContext
+{
+	// Partial refresh flags : allows to recompute only a subset of the context information :
+	enum class ERefreshFlags
+	{
+		None = 0,
+		RefreshComponentInfos = (1 << 0),
+		RefreshHeightmapInfos = (1 << 1),
+		RefreshWeightmapInfos = (1 << 2),
+		RefreshMapHelper = (1 << 3),
+		RefreshAll = ~0,
+	};
+	FRIEND_ENUM_CLASS_FLAGS(ERefreshFlags);
+
+	FUpdateLayersContentContext(const FTextureToComponentHelper& InMapHelper, bool bInPartialUpdate)
+		: bPartialUpdate(bInPartialUpdate)
+		, MapHelper(InMapHelper)
+	{
+		// No need to update the map helper, it's assumed to be already ready in the constructor
+		Refresh(ERefreshFlags::RefreshComponentInfos | ERefreshFlags::RefreshHeightmapInfos | ERefreshFlags::RefreshWeightmapInfos);
+	}
+
+	FTextureToComponentHelper::ERefreshFlags RefreshFlagsToMapHelperRefreshFlags(ERefreshFlags InRefreshFlags)
+	{
+		FTextureToComponentHelper::ERefreshFlags MapHelperRefreshFlags = FTextureToComponentHelper::ERefreshFlags::None;
+		if (EnumHasAnyFlags(InRefreshFlags, ERefreshFlags::RefreshComponentInfos))
+		{
+			MapHelperRefreshFlags |= FTextureToComponentHelper::ERefreshFlags::RefreshComponents;
+		}
+		if (EnumHasAnyFlags(InRefreshFlags, ERefreshFlags::RefreshHeightmapInfos))
+		{
+			MapHelperRefreshFlags |= FTextureToComponentHelper::ERefreshFlags::RefreshHeightmaps;
+		}
+		if (EnumHasAnyFlags(InRefreshFlags, ERefreshFlags::RefreshWeightmapInfos))
+		{
+			MapHelperRefreshFlags |= FTextureToComponentHelper::ERefreshFlags::RefreshWeightmaps;
+		}
+		return MapHelperRefreshFlags;
+	}
+
+	void Refresh(ERefreshFlags InRefreshFlags)
+	{
+		TRACE_CPUPROFILER_EVENT_SCOPE(UpdateLayersContentContext_Refresh);
+		// Start by updating the map helper if necessary (keep track of components/heightmaps/weightmaps relationship) :
+		if (EnumHasAnyFlags(InRefreshFlags, ERefreshFlags::RefreshMapHelper))
+		{
+			MapHelper.Refresh(RefreshFlagsToMapHelperRefreshFlags(InRefreshFlags));
+		}
+
+		// Then triage the dirty/non-dirty components  :
+		if (EnumHasAnyFlags(InRefreshFlags, ERefreshFlags::RefreshComponentInfos))
+		{
+			// When components are refreshed, all other info has to be :
+			check(EnumHasAllFlags(InRefreshFlags, ERefreshFlags::RefreshHeightmapInfos | ERefreshFlags::RefreshWeightmapInfos));
+
+			DirtyLandscapeComponents.Reset();
+			NonDirtyLandscapeComponents.Reset();
+			for (ULandscapeComponent* Component : MapHelper.LandscapeComponents)
+			{
+				if (!bPartialUpdate || Component->GetLayerUpdateFlagPerMode() != 0)
+				{
+					DirtyLandscapeComponents.Add(Component);
+				}
+				else
+				{
+					NonDirtyLandscapeComponents.Add(Component);
+				}
+			}
+		}
+
+		if (EnumHasAnyFlags(InRefreshFlags, ERefreshFlags::RefreshHeightmapInfos | ERefreshFlags::RefreshWeightmapInfos))
+		{
+			TSet<UTexture2D*> HeightmapsToRender;
+			TSet<ULandscapeComponent*> NeighborsComponents;
+			TSet<ULandscapeComponent*> WeightmapsComponents;
+
+			// Cleanup our heightmap/weightmap info:
+			if (EnumHasAnyFlags(InRefreshFlags, ERefreshFlags::RefreshHeightmapInfos))
+			{
+				HeightmapsToResolve.Reset();
+				LandscapeComponentsHeightmapsToRender.Reset();
+				LandscapeComponentsHeightmapsToResolve.Reset();
+			}
+			if (EnumHasAnyFlags(InRefreshFlags, ERefreshFlags::RefreshWeightmapInfos))
+			{
+				WeightmapsToResolve.Reset();
+				LandscapeComponentsWeightmapsToRender.Reset();
+				LandscapeComponentsWeightmapsToResolve.Reset();
+			}
+			// Note that the AllLandscapeComponentsToResolve and AllLandscapeComponentsToUpdate are *not* reset here: they can only grow (we're assuming refresh only adds new components): 
+
+			// Iterate on all dirty components and retrieve the components that need to be resolved or rendered for their heightmap or weightmaps :
+			TArray<ULandscapeComponent*> AllLandscapeComponents;
+			for (ULandscapeComponent* Component : DirtyLandscapeComponents)
+			{
+				AllLandscapeComponents.Add(Component);
+
+				// If all components are dirty, we can take some shortcuts since all components will need to be rendered and resolved : 
+				if (bPartialUpdate)
+				{
+					if (EnumHasAnyFlags(InRefreshFlags, ERefreshFlags::RefreshHeightmapInfos))
+					{
+						// Gather Neighbors (Neighbors need to be Rendered but not resolved so that the resolved Components have valid normals on edges)
+						Component->GetLandscapeComponentNeighborsToRender(NeighborsComponents);
+						Component->ForEachLayer([&](const FGuid&, FLandscapeLayerComponentData& LayerData)
+						{
+							HeightmapsToRender.Add(LayerData.HeightmapData.Texture);
+						});
+					}
+
+					if (EnumHasAnyFlags(InRefreshFlags, ERefreshFlags::RefreshWeightmapInfos))
+					{
+						// Gather WeightmapUsages (Components sharing weightmap usages with the resolved Components need to be rendered so that the resolving is valid)
+						Component->GetLandscapeComponentWeightmapsToRender(WeightmapsComponents);
+					}
+				}
+
+				if (EnumHasAnyFlags(InRefreshFlags, ERefreshFlags::RefreshHeightmapInfos))
+				{
+					// Gather Heightmaps (All Components sharing Heightmap textures need to be rendered and resolved)
+					HeightmapsToResolve.Add(Component->GetHeightmap(/*InReturnEditingHeightmap = */false));
+				}
+
+				if (EnumHasAnyFlags(InRefreshFlags, ERefreshFlags::RefreshWeightmapInfos))
+				{
+					// Gather Weightmaps
+					const TArray<UTexture2D*>& WeightmapTextures = Component->GetWeightmapTextures(/*InReturnEditingWeightmap = */false);
+					for (FWeightmapLayerAllocationInfo const& AllocInfo : Component->GetWeightmapLayerAllocations(/*InReturnEditingWeightmap = */false))
+					{
+						if (AllocInfo.IsAllocated() && AllocInfo.WeightmapTextureIndex < WeightmapTextures.Num())
+						{
+							WeightmapsToResolve.Add(WeightmapTextures[AllocInfo.WeightmapTextureIndex]);
+						}
+					}
+				}
+			}
+
+			if (EnumHasAnyFlags(InRefreshFlags, ERefreshFlags::RefreshHeightmapInfos))
+			{
+				// Because of Heightmap Sharing anytime we render a heightmap we need to render all the components that use it
+				for (ULandscapeComponent* NeighborsComponent : NeighborsComponents)
+				{
+					NeighborsComponent->ForEachLayer([&](const FGuid&, FLandscapeLayerComponentData& LayerData)
+					{
+						HeightmapsToRender.Add(LayerData.HeightmapData.Texture);
+					});
+				}
+
+				// Copy first list into others
+				LandscapeComponentsHeightmapsToResolve.Append(AllLandscapeComponents);
+				LandscapeComponentsHeightmapsToRender.Append(AllLandscapeComponents);
+			}
+
+			if (EnumHasAnyFlags(InRefreshFlags, ERefreshFlags::RefreshWeightmapInfos))
+			{
+				LandscapeComponentsWeightmapsToResolve.Append(AllLandscapeComponents);
+				LandscapeComponentsWeightmapsToRender.Append(AllLandscapeComponents);
+			}
+
+			if (bPartialUpdate)
+			{
+				for (ULandscapeComponent* Component : NonDirtyLandscapeComponents)
+				{
+					if (EnumHasAnyFlags(InRefreshFlags, ERefreshFlags::RefreshHeightmapInfos))
+					{
+						if (HeightmapsToResolve.Contains(Component->GetHeightmap(false)))
+						{
+							AllLandscapeComponents.Add(Component);
+							LandscapeComponentsHeightmapsToRender.Add(Component);
+							LandscapeComponentsHeightmapsToResolve.Add(Component);
+						}
+						else if (NeighborsComponents.Contains(Component))
+						{
+							LandscapeComponentsHeightmapsToRender.Add(Component);
+						}
+						else
+						{
+							bool bAdd = false;
+							Component->ForEachLayer([&](const FGuid&, FLandscapeLayerComponentData& LayerData)
+							{
+								if (HeightmapsToRender.Contains(LayerData.HeightmapData.Texture))
+								{
+									bAdd = true;
+								}
+							});
+							if (bAdd)
+							{
+								LandscapeComponentsHeightmapsToRender.Add(Component);
+							}
+						}
+					}
+
+					if (EnumHasAnyFlags(InRefreshFlags, ERefreshFlags::RefreshWeightmapInfos))
+					{
+						if (WeightmapsComponents.Contains(Component))
+						{
+							LandscapeComponentsWeightmapsToRender.Add(Component);
+						}
+					}
+				}
+			}
+
+			// All selected components will have to be resolved : 
+			AllLandscapeComponentsToResolve.Append(AllLandscapeComponents);
+
+			// Add components with deferred flag to update list
+			for (ULandscapeComponent* Component : AllLandscapeComponents)
+			{
+				if (Component->GetLayerUpdateFlagPerMode() & ELandscapeLayerUpdateMode::Update_Client_Deferred)
+				{
+					AllLandscapeComponentsToUpdate.FindOrAdd(Component, FLandscapeEditLayerReadbackResult(ELandscapeLayerUpdateMode::Update_Client_Deferred, false));
+				}
+			}
+		}
+	}
+
+	// Indicates whether all components of the landscape are marked dirty :
+	const bool bPartialUpdate = false;
+	// Helper to gather mappings between heightmaps/weightmaps and components :
+	FTextureToComponentHelper MapHelper;
+	// List of landscape components that have been made dirty and need to be updated : 
+	TArray<ULandscapeComponent*> DirtyLandscapeComponents;
+	// List of landscape components that have not been made dirty : 
+	TArray<ULandscapeComponent*> NonDirtyLandscapeComponents;
+	// List of heightmap textures that might be affected by the update : 
+	TSet<UTexture2D*> HeightmapsToResolve;
+	// List of weightmap textures that might be affected by the update : 
+	TSet<UTexture2D*> WeightmapsToResolve;
+	// List of components that need to be rendered because they are either dirty or are neighbor to a component that is dirty or share a heightmap with a component that is dirty:
+	TArray<ULandscapeComponent*> LandscapeComponentsHeightmapsToRender;
+	// List of components whose heightmap needs to be resolved because they are either dirty or are neighbor to a component that is dirty:
+	TArray<ULandscapeComponent*> LandscapeComponentsHeightmapsToResolve;
+	// List of components that need to be rendered because they are either dirty or are neighbor to a component that is dirty or share a weightmap with a component that is dirty:
+	TArray<ULandscapeComponent*> LandscapeComponentsWeightmapsToRender;
+	// List of components whose weightmap needs to be resolved because they are either dirty or are neighbor to a component that is dirty:
+	TArray<ULandscapeComponent*> LandscapeComponentsWeightmapsToResolve;
+	// List of components whose heightmap or weightmaps needs to be resolved because they are either dirty or are neighbor to a component that is dirty:
+	TSet<ULandscapeComponent*> AllLandscapeComponentsToResolve;
+	// List of GPU readback results for heightmaps/weightmaps that need to be resolved, associated with their owning landscape component :
+	TMap<ULandscapeComponent*, FLandscapeEditLayerReadbackResult> AllLandscapeComponentsToUpdate;
+};
+ENUM_CLASS_FLAGS(FUpdateLayersContentContext::ERefreshFlags);
+
+// Little struct that holds information common to PerformLayersHeightmapsLocalMerge and PerformLayersHeightmapsGlobalMerge
+struct FEditLayersHeightmapMergeParams
+{
+	int32 HeightmapUpdateModes;
+	bool bForceRender;
+	bool bSkipBrush;
+};
+
+// Render-thread version of the data / functions we need for the local merge of edit layers : 
+namespace EditLayersHeightmapLocalMerge_RenderThread
+{
+	struct FComponentRenderInfo
+	{
+		// Name of the component for debug purposes :
+		FString Name;
+
+		// The component's visible layer's heightmaps (Num must be == Num of VisibleEditLayerInfos)
+		TArray<FTexture2DResourceSubregion> VisibleLayerHeightmapTextures;
+
+		// List of 9 component render info indices corresponding to the 9 neighbors of this component (index into ComponentToRenderInfos) :
+		//  Some may not be set and that's fine : the only goal is to know which are the valid neighbors when it comes to stitching adjacent rendered components
+		//  If a neighbor is missing, we can't stitch the border adjacent to it, but it also means it won't contribute to what we really want to compute : the 
+		//  component to resolve, which is guaranteed to have its neighbors present (if those do exist) :
+		TStaticArray<int32, 9> NeighborComponentToRenderInfoIndices;
+	};
+
+	struct FComponentResolveInfo
+	{
+		FComponentResolveInfo(int32 InNumComponentsToRender)
+		{
+			NeighborComponentToRenderInfoBitIndices.Init(false, InNumComponentsToRender);
+		}
+
+		void SetNeighborRenderInfo(int32 InNeighborIndex, int32 InNeighborComponentRenderInfoIndex)
+		{
+			NeighborComponentToRenderInfoIndices[InNeighborIndex] = InNeighborComponentRenderInfoIndex;
+			if (InNeighborComponentRenderInfoIndex != INDEX_NONE)
+			{
+				check(!NeighborComponentToRenderInfoBitIndices[InNeighborComponentRenderInfoIndex]);
+				NeighborComponentToRenderInfoBitIndices[InNeighborComponentRenderInfoIndex] = true;
+			}
+		}
+
+	public:
+		// Index of this components in ComponentToResolveInfos :
+		int32 ComponentToResolveInfoIndex = INDEX_NONE;
+		// Name of the component for debug purposes :
+		FString Name;
+		// Subregion of the heightmap that we want to compute normals for and resolve, that corresponds to this component :
+		FTexture2DResourceSubregion Heightmap;
+
+		// List of 9 component render info indices corresponding to the 9 neighbors of this component (index into ComponentToRenderInfos) :
+		//  Some may not be set if the component is on the border of the landscape, for example :
+		TStaticArray<int32, 9> NeighborComponentToRenderInfoIndices;
+		// Same as NeighborComponentToRenderInfoIndices but as a bit array (1 bit per component to render info) to vastly optimize the division of component resolve infos into batches, which is a O(N^2) operation : 
+		TBitArray<> NeighborComponentToRenderInfoBitIndices;
+	};
+
+	struct FTextureResolveInfo
+	{
+		// Texture that was updated and needs resolving : 
+		FTexture2DResource* Texture = nullptr;
+		// CPU readback utility to bring back the result on the CPU : 
+		FLandscapeEditLayerReadback* CPUReadback;
+	};
+
+	// Because of heightmaps being shared between one component and another, we have to group the components to render into batches
+	//  where we'll render all of the heightmaps into slices of a single scratch texture array, which we'll then be able to re-assemble into the final, packed, heightmaps (subregions) :
+	struct FComponentResolveBatchInfo
+	{
+		FComponentResolveBatchInfo(int32 InNumComponentsToRender, int32 InBatchIndex)
+			: BatchIndex(InBatchIndex)
+		{
+			ComponentToRenderInfoBitIndices.Init(false, InNumComponentsToRender);
+		}
+
+		void AddComponent(const FComponentResolveInfo& InComponentResolveInfo)
+		{
+			check(!ComponentToResolveInfoIndices.Contains(InComponentResolveInfo.ComponentToResolveInfoIndex));
+			ComponentToResolveInfoIndices.Add(InComponentResolveInfo.ComponentToResolveInfoIndex);
+
+			// Remember all the unique components that this texture needs for resolving:
+			ComponentToRenderInfoBitIndices.CombineWithBitwiseOR(InComponentResolveInfo.NeighborComponentToRenderInfoBitIndices, EBitwiseOperatorFlags::MinSize);
+		}
+
+	public:
+		int32 BatchIndex = INDEX_NONE;
+
+		// Indices (in ComponentToRenderInfos) of the components whose heightmaps we need to render within this batch in order to produce (and then resolve) the textures in TextureToResolveInfos:
+		//  It's a bit array (1 bit per component to render info) to vastly optimize the division of texture resolve infos into batches, which is a O(N^2) operation : 
+		TBitArray<> ComponentToRenderInfoBitIndices;
+		// Indices (in ComponentToResolveInfos) of components whose heightmap subregion needs to be resolved / read back on the CPU :
+		TArray<int32> ComponentToResolveInfoIndices;
+	};
+
+	struct FEditLayerInfo
+	{
+		ELandscapeEditLayerHeightmapBlendMode BlendMode = ELandscapeEditLayerHeightmapBlendMode::Num;
+		float Alpha = 1.0f;
+	};
+
+	// Description of the entire merge pass :
+	struct FMergeInfo
+	{
+		bool NeedsMerge() const
+		{
+			// If no edit layer or if no paint layer present on any edit layer, we've got nothing to do :
+			bool bNeedsMerge = (MaxNumEditLayersTexturesToMerge > 0) && (MaxNumComponentsToRenderPerResolveComponentBatch > 0);
+			check(!bNeedsMerge || !TextureToResolveInfos.IsEmpty()); // If we need merging, we must have at one texture to resolve
+			return bNeedsMerge;
+		}
+
+	public:
+		// Number of vertices per component
+		FIntPoint ComponentSizeVerts;
+
+		// Size of each heigthmap (one heightmap can contain multiple components due to heightmap sharing)
+		FIntPoint HeightmapSize;
+
+		// Heightmap pixel to world scale factor
+		FVector LandscapeGridScale;
+		
+		// Number of mips for the final heightmaps
+		int32 NumMips = 0;
+
+		// Maximum number of visible edit layers that have to be merged for a single FComponentRenderInfo : 
+		int32 MaxNumEditLayersTexturesToMerge = 0;
+
+		// Maximum number of components to render in any given FComponentResolveBatchInfo. This is the number of slices needed for the scratch texture arrays that we reuse from one batch to another :
+		int32 MaxNumComponentsToRenderPerResolveComponentBatch = 0;
+
+		// Describes how to access each visible edit layer's heightmap and how to blend it in the final heightmap for this paint layer :
+		TArray<FEditLayerInfo> VisibleEditLayerInfos;
+
+		// List of infos for each component that needs its edit layers' heightmaps to be rendered (merged) and ultimately participate to the final heightmap of the component we're trying to resolve :
+		TArray<FComponentRenderInfo> ComponentToRenderInfos;
+
+		// List of infos for each component that needs to be resolved :
+		TArray<FComponentResolveInfo> ComponentToResolveInfos;
+
+		// List of batches of FComponentResolveInfo that needs to be resolved in the same pass. This allows massive saves on transient resources on large landscapes because those can be re-cycled from one pass to another :
+		TArray<FComponentResolveBatchInfo> ComponentResolveBatchInfos;
+
+		// List of infos for each texture that needs to be resolved :
+		TArray<FTextureResolveInfo> TextureToResolveInfos;
+
+		// Not truly render-thread data because it references UTextures but it's just because FLandscapeEditLayerReadback were historically game-thread initiated so for as long as we'll use those for readback, we need to store this here : 
+		TArray<FLandscapeLayersCopyReadbackTextureParams> DeferredCopyReadbackTextures;
+	};
+
+	static ELandscapeEditLayerHeightmapBlendMode LandscapeBlendModeToEditLayerBlendMode(ELandscapeBlendMode InLandscapeBlendMode)
+	{
+		switch (InLandscapeBlendMode)
+		{
+		case LSBM_AdditiveBlend:
+			return ELandscapeEditLayerHeightmapBlendMode::Additive;
+		case LSBM_AlphaBlend:
+			return ELandscapeEditLayerHeightmapBlendMode::AlphaBlend;
+		default:
+			check(false);
+		}
+
+		return ELandscapeEditLayerHeightmapBlendMode::Num;
+	}
+
+	struct FRDGResources
+	{
+		// Texture array in which all possible edit layers heightmaps can fit : we copy the edit layers heightmaps there in order to dynamically access it in the merge shader : 
+		FRDGTextureRef EditLayersHeightmapsTextureArray;
+		FRDGTextureSRVRef EditLayersHeightmapsTextureArraySRV;
+
+		// Temporary scratch texture array that stores the output (packed height only) of all (edit layer-merged) landscape components to render within a batch (one per landscape component)  
+		//  Can be reused from one batch to another :
+		FRDGTextureRef ScratchMergedHeightmapTextureArray;
+		FRDGTextureSRVRef ScratchMergedHeightmapTextureArraySRV;
+
+		// Temporary scratch texture array that stores the output (packed height only) of all (edit layer-merged) landscape components to render within a batch (one per landscape component), after stitching step is done:
+		//  Can be reused from one batch to another :
+		FRDGTextureRef ScratchStitchedHeightmapTextureArray;
+		FRDGTextureSRVRef ScratchStitchedHeightmapTextureArraySRV;
+
+		// Single structured buffer that will contain layer merge infos (FLandscapeEditLayerHeightmapMergeInfo): doesn't changed from one component to another :
+		FRDGBufferRef EditLayersMergeInfosBuffer;
+		FRDGBufferSRVRef EditLayersMergeInfosBufferSRV;
+	};
+
+	void PrepareLayersHeightmapsLocalMergeRDGResources(const FMergeInfo& InLocalMergeInfo, FRDGBuilder& GraphBuilder, FRDGResources& OutResources)
+	{
+		{
+			int32 SizeZ = InLocalMergeInfo.MaxNumEditLayersTexturesToMerge;
+			check(SizeZ > 0);
+			// TODO [jonathan.bard] : change to PF_R8G8 once edit layers heightmaps are stored as such :
+			FRDGTextureDesc Desc = FRDGTextureDesc::Create2DArray(InLocalMergeInfo.ComponentSizeVerts, PF_B8G8R8A8, FClearValueBinding::None, TexCreate_RenderTargetable | TexCreate_ShaderResource, 
+				SizeZ, /*InNumMips = */1, /*InNumSamples = */1);
+			OutResources.EditLayersHeightmapsTextureArray = GraphBuilder.CreateTexture(Desc, TEXT("LandscapeEditLayersHeightmapsTextureArray"));
+			OutResources.EditLayersHeightmapsTextureArraySRV = GraphBuilder.CreateSRV(FRDGTextureSRVDesc::Create(OutResources.EditLayersHeightmapsTextureArray));
+		}
+
+		{
+			int32 SizeZ = InLocalMergeInfo.MaxNumComponentsToRenderPerResolveComponentBatch;
+			// We only need 2 channels since this only stores the (packed) height :
+			FRDGTextureDesc Desc = FRDGTextureDesc::Create2DArray(InLocalMergeInfo.ComponentSizeVerts, PF_R8G8, FClearValueBinding::None, TexCreate_RenderTargetable | TexCreate_TargetArraySlicesIndependently | TexCreate_ShaderResource, 
+				SizeZ, /*InNumMips = */1, /*InNumSamples = */1);
+			// Create 2 texture arrays "Merged" and "Stitched" (ScratchMergedHeightmapTextureArray will be copied/merged into ScratchStitchedHeightmapTextureArray, slice by slice)
+			OutResources.ScratchMergedHeightmapTextureArray = GraphBuilder.CreateTexture(Desc, TEXT("LandscapeEditLayersMergedHeightmapTextureArray"));
+			OutResources.ScratchStitchedHeightmapTextureArray = GraphBuilder.CreateTexture(Desc, TEXT("LandscapeEditLayersStitchedHeightmapTextureArray"));
+			OutResources.ScratchMergedHeightmapTextureArraySRV = GraphBuilder.CreateSRV(FRDGTextureSRVDesc::Create(OutResources.ScratchMergedHeightmapTextureArray));
+			OutResources.ScratchStitchedHeightmapTextureArraySRV = GraphBuilder.CreateSRV(FRDGTextureSRVDesc::Create(OutResources.ScratchStitchedHeightmapTextureArray));
+		}
+
+		{
+			// Upload layer merge info buffer once and for all since it's unchanged from one component to another : 
+			TArray<FLandscapeEditLayerHeightmapMergeInfo, TInlineAllocator<16>> EditLayerMergeInfos;
+			for (const FEditLayerInfo& EditLayerInfo : InLocalMergeInfo.VisibleEditLayerInfos)
+			{
+				FLandscapeEditLayerHeightmapMergeInfo& EditLayerMergeInfo = EditLayerMergeInfos.Emplace_GetRef();
+				EditLayerMergeInfo.BlendMode = EditLayerInfo.BlendMode;
+				EditLayerMergeInfo.Alpha = EditLayerInfo.Alpha;
+			}
+			OutResources.EditLayersMergeInfosBuffer = CreateStructuredBuffer(GraphBuilder, TEXT("LandscapeEditLayersMergeInfosBuffer"), EditLayerMergeInfos);
+			OutResources.EditLayersMergeInfosBufferSRV = GraphBuilder.CreateSRV(FRDGBufferSRVDesc(OutResources.EditLayersMergeInfosBuffer));
+		}
+	}
+
+	// Gather all textures we will want to write into or read from in the render graph and output them in OutTrackedTextures:
+	void GatherLayersHeightmapsLocalMergeRDGTextures(const FMergeInfo& InLocalMergeInfo, TMap<FTexture2DResource*, FLandscapeRDGTrackedTexture>& OutTrackedTextures)
+	{
+		// First pass, gather all textures we'll need for merging layers, i.e. the component edit layers' heightmaps : 
+		for (const FComponentRenderInfo& ComponentRenderInfo : InLocalMergeInfo.ComponentToRenderInfos)
+		{
+			for (const FTexture2DResourceSubregion& LayerHeightmap : ComponentRenderInfo.VisibleLayerHeightmapTextures)
+			{
+				check(LayerHeightmap.Texture != nullptr);
+				FLandscapeRDGTrackedTexture* LayerTrackedTexture = OutTrackedTextures.Find(LayerHeightmap.Texture);
+				if (LayerTrackedTexture == nullptr)
+				{
+					LayerTrackedTexture = &(OutTrackedTextures.Add(LayerHeightmap.Texture, FLandscapeRDGTrackedTexture(LayerHeightmap.Texture)));
+					// Layer heightmaps were not created with TexCreate_RenderTargetable and are therefore considered as ReadOnly by the RDG, which prevents their resource access from being automatically managed
+					//  (i.e. transfered to CopySrc when we copy them and back to SRVMask at the end) so we need to force the RDG to track them:
+					// TODO [jonathan.bard] : remove this if we don't use the texture array :
+					LayerTrackedTexture->bNeedsForceTracking = true;
+				}
+			}
+		}
+
+		// Second pass, gather all textures we'll need to regenerate : 
+		for (const FComponentResolveInfo& ComponentResolveInfo : InLocalMergeInfo.ComponentToResolveInfos)
+		{
+			check(ComponentResolveInfo.Heightmap.Texture != nullptr);
+			FLandscapeRDGTrackedTexture* TrackedTexture = OutTrackedTextures.Find(ComponentResolveInfo.Heightmap.Texture);
+			if (TrackedTexture == nullptr)
+			{
+				TrackedTexture = &(OutTrackedTextures.Add(ComponentResolveInfo.Heightmap.Texture, FLandscapeRDGTrackedTexture(ComponentResolveInfo.Heightmap.Texture)));
+			}
+			TrackedTexture->bNeedsScratch = true;
+		}
+	}
+
+	void MergeEditLayersHeightmapsForBatch(const FComponentResolveBatchInfo& InComponentResolveBatchInfo, const FMergeInfo& InLocalMergeInfo, const TMap<FTexture2DResource*, FLandscapeRDGTrackedTexture>& InTrackedTextures, FRDGBuilder& GraphBuilder, FRDGResources& RDGResources)
+	{
+		RDG_EVENT_SCOPE(GraphBuilder, "Merge %d edit layers on %d components", InLocalMergeInfo.VisibleEditLayerInfos.Num(), InComponentResolveBatchInfo.ComponentToRenderInfoBitIndices.CountSetBits());
+
+		// For each component to render, perform the edit layers merge and write the resulting heightmap :
+		int32 IndexInBatch = 0;
+		for (TConstSetBitIterator<> BitIt(InComponentResolveBatchInfo.ComponentToRenderInfoBitIndices); BitIt; ++BitIt, ++IndexInBatch)
+		{
+			int32 ComponentRenderInfoIndex = BitIt.GetIndex();
+			const FComponentRenderInfo& ComponentRenderInfo = InLocalMergeInfo.ComponentToRenderInfos[ComponentRenderInfoIndex];
+		
+			RDG_EVENT_SCOPE(GraphBuilder, "Component %s", *ComponentRenderInfo.Name);
+
+			// Prepare the texture array of layer heightmaps for this component : 
+			int32 NumLayers = InLocalMergeInfo.VisibleEditLayerInfos.Num();
+			for (int32 i = 0; i < NumLayers; ++i)
+			{
+				const FTexture2DResourceSubregion& LayerHeighmapSubregion = ComponentRenderInfo.VisibleLayerHeightmapTextures[i];
+
+				const FLandscapeRDGTrackedTexture* TrackedTexture = InTrackedTextures.Find(LayerHeighmapSubregion.Texture);
+				check(TrackedTexture != nullptr);
+
+				// We need to copy the (portion of the) layer's texture to the texture array : 
+				FRHICopyTextureInfo CopyTextureInfo;
+				CopyTextureInfo.Size = FIntVector(LayerHeighmapSubregion.Subregion.Size().X, LayerHeighmapSubregion.Subregion.Size().Y, 0);
+				CopyTextureInfo.DestSliceIndex = i;
+				CopyTextureInfo.SourcePosition = FIntVector(LayerHeighmapSubregion.Subregion.Min.X, LayerHeighmapSubregion.Subregion.Min.Y, 0);
+
+				AddCopyTexturePass(GraphBuilder, TrackedTexture->ExternalTextureRef, RDGResources.EditLayersHeightmapsTextureArray, CopyTextureInfo);
+			}
+
+			// Then, merge all heightmaps using the MergeEditLayers PS and write into a single slice in ScratchMergedHeightmapTextureArray: 
+			{
+				check(IndexInBatch < RDGResources.ScratchMergedHeightmapTextureArray->Desc.ArraySize);
+
+				FLandscapeLayersHeightmapsMergeEditLayersPS::FParameters* MergeEditLayersPSParams = GraphBuilder.AllocParameters<FLandscapeLayersHeightmapsMergeEditLayersPS::FParameters>();
+				MergeEditLayersPSParams->RenderTargets[0] = FRenderTargetBinding(RDGResources.ScratchMergedHeightmapTextureArray, ERenderTargetLoadAction::ENoAction, /*InMipIndex = */0, /*InArraySlice = */IndexInBatch);
+				MergeEditLayersPSParams->InNumEditLayers = NumLayers;
+				MergeEditLayersPSParams->InEditLayersTextures = RDGResources.EditLayersHeightmapsTextureArraySRV;
+				MergeEditLayersPSParams->InEditLayersMergeInfos = RDGResources.EditLayersMergeInfosBufferSRV;
+
+				FLandscapeLayersHeightmapsMergeEditLayersPS::MergeEditLayers(GraphBuilder, MergeEditLayersPSParams, /*InTextureSize =*/InLocalMergeInfo.ComponentSizeVerts);
+			}
+		}
+	}
+
+	void StitchHeightmapsForBatch(const FComponentResolveBatchInfo& InComponentResolveBatchInfo, const FMergeInfo& InLocalMergeInfo, FRDGBuilder& GraphBuilder, FRDGResources& RDGResources)
+	{
+		RDG_EVENT_SCOPE(GraphBuilder, "Stitch %d components for batch %d", InComponentResolveBatchInfo.ComponentToRenderInfoBitIndices.CountSetBits(), InComponentResolveBatchInfo.BatchIndex);
+
+		// For each component to render (i.e. including the neighbors of dirty components, that are also need to properly recompute normals), correct the heightmap 
+		//  subregion (by stitching adjacent component borders) :
+		int32 IndexInBatch = 0;
+		for (TConstSetBitIterator<> BitIt(InComponentResolveBatchInfo.ComponentToRenderInfoBitIndices); BitIt; ++BitIt, ++IndexInBatch)
+		{
+			int32 ComponentRenderInfoIndex = BitIt.GetIndex();
+			const FComponentRenderInfo& ComponentRenderInfo = InLocalMergeInfo.ComponentToRenderInfos[ComponentRenderInfoIndex];
+
+			RDG_EVENT_SCOPE(GraphBuilder, "Component %s", *ComponentRenderInfo.Name);
+
+			// Now, stitch the heightmap using the StitchHeightmapPS and output the result to a single slice in ScratchStitchedHeightmapTextureArray: 
+			{
+				FLandscapeLayersHeightmapsStitchHeightmapPS::FParameters* StitchHeightmapPSParams = GraphBuilder.AllocParameters<FLandscapeLayersHeightmapsStitchHeightmapPS::FParameters>();
+				StitchHeightmapPSParams->RenderTargets[0] = FRenderTargetBinding(RDGResources.ScratchStitchedHeightmapTextureArray, ERenderTargetLoadAction::ENoAction, /*InMipIndex = */0, /*InArraySlice = */IndexInBatch);
+				StitchHeightmapPSParams->InSourceTextureSize = FUintVector2(InLocalMergeInfo.ComponentSizeVerts.X, InLocalMergeInfo.ComponentSizeVerts.Y);
+				StitchHeightmapPSParams->InSourceHeightmaps = RDGResources.ScratchMergedHeightmapTextureArraySRV;
+
+				for (int32 NeighborIndex = 0; NeighborIndex < 9; ++NeighborIndex)
+				{
+					int32 NeighborComponentToRenderInfoIndex = ComponentRenderInfo.NeighborComponentToRenderInfoIndices[NeighborIndex];
+					// Index of the neighbor component in this batch : allows to retrieve the proper slice in the source heightmap array :
+					int32 NeighborIndexInBatch = INDEX_NONE;
+					// The neighbor could be totally absent (index == INDEX_NONE) or it can be absent from the batch. That means that it's not actually relevant for this step 
+					//  since the vertices that will be "invalid" won't be taken into account by the components we're actually trying to resolve :
+					if ((NeighborComponentToRenderInfoIndex != INDEX_NONE) && InComponentResolveBatchInfo.ComponentToRenderInfoBitIndices[NeighborComponentToRenderInfoIndex])
+					{
+						// Components are rendered in order within the batch so the index of this component to render in the batch is == to how many components are there before it:
+						NeighborIndexInBatch = InComponentResolveBatchInfo.ComponentToRenderInfoBitIndices.CountSetBits(0, NeighborComponentToRenderInfoIndex);
+						check(NeighborIndexInBatch < RDGResources.ScratchMergedHeightmapTextureArray->Desc.ArraySize);
+					}
+
+					check((NeighborIndex != 4) || (NeighborIndexInBatch != INDEX_NONE)); // The central component (the one we finalize) should always be valid
+					GET_SCALAR_ARRAY_ELEMENT(StitchHeightmapPSParams->InNeighborHeightmapIndices, NeighborIndex) = NeighborIndexInBatch;
+				}
+
+				FLandscapeLayersHeightmapsStitchHeightmapPS::StitchHeightmap(GraphBuilder, StitchHeightmapPSParams);
+			}
+		}
+	}
+
+	void FinalizeComponentsForBatch(const FComponentResolveBatchInfo& InComponentResolveBatchInfo, const FMergeInfo& InLocalMergeInfo, const TMap<FTexture2DResource*, FLandscapeRDGTrackedTexture>& InTrackedTextures, FRDGBuilder& GraphBuilder, FRDGResources& RDGResources)
+	{
+		RDG_EVENT_SCOPE(GraphBuilder, "Finalize %d components", InLocalMergeInfo.ComponentToResolveInfos.Num());
+
+		// For each component to resolve (i.e. including the neighbors of dirty components, that also need to have their normals recomputed), finalize the heightmap subregion (i.e. compute the normals) :
+		for (int32 ComponentResolveInfoIndex : InComponentResolveBatchInfo.ComponentToResolveInfoIndices)
+		{
+			const FComponentResolveInfo& ComponentResolveInfo = InLocalMergeInfo.ComponentToResolveInfos[ComponentResolveInfoIndex];
+
+			RDG_EVENT_SCOPE(GraphBuilder, "Component %s", *ComponentResolveInfo.Name);
+
+			const FLandscapeRDGTrackedTexture* TrackedTexture = InTrackedTextures.Find(ComponentResolveInfo.Heightmap.Texture);
+			check((TrackedTexture != nullptr) && (TrackedTexture->ScratchTextureRef != nullptr));
+
+			// Now, finalize the heightmap using the FinalizeHeighmap PS : 
+			{
+				FLandscapeLayersHeightmapsFinalizeHeightmapPS::FParameters* FinalizeHeightmapPSParams = GraphBuilder.AllocParameters<FLandscapeLayersHeightmapsFinalizeHeightmapPS::FParameters>();
+				FinalizeHeightmapPSParams->RenderTargets[0] = FRenderTargetBinding(TrackedTexture->ScratchTextureRef, ERenderTargetLoadAction::ENoAction);
+				FinalizeHeightmapPSParams->InSourceTextureSize = FUintVector2(InLocalMergeInfo.ComponentSizeVerts.X, InLocalMergeInfo.ComponentSizeVerts.Y);
+				FinalizeHeightmapPSParams->InSourceHeightmaps = RDGResources.ScratchStitchedHeightmapTextureArraySRV;
+				FinalizeHeightmapPSParams->InDestinationTextureSubregion = FUintVector4(ComponentResolveInfo.Heightmap.Subregion.Min.X, ComponentResolveInfo.Heightmap.Subregion.Min.Y, ComponentResolveInfo.Heightmap.Subregion.Max.X, ComponentResolveInfo.Heightmap.Subregion.Max.Y);
+				FinalizeHeightmapPSParams->InLandscapeGridScale = InLocalMergeInfo.LandscapeGridScale;
+
+				for (int32 NeighborIndex = 0; NeighborIndex < 9; ++NeighborIndex)
+				{
+					int32 NeighborComponentToRenderInfoIndex = ComponentResolveInfo.NeighborComponentToRenderInfoIndices[NeighborIndex];
+					// Index of the neighbor component in this batch : allows to retrieve the proper slice in the source heightmap array :
+					int32 NeighborIndexInBatch = INDEX_NONE;
+					// The neighbor could be absent : 
+					if (NeighborComponentToRenderInfoIndex != INDEX_NONE)
+					{
+						check(InComponentResolveBatchInfo.ComponentToRenderInfoBitIndices[NeighborComponentToRenderInfoIndex]);
+						// Components are rendered in order within the batch so the index of this component to render in the batch is == to how many components are there before it:
+						NeighborIndexInBatch = InComponentResolveBatchInfo.ComponentToRenderInfoBitIndices.CountSetBits(0, NeighborComponentToRenderInfoIndex);
+						check(NeighborIndexInBatch < RDGResources.ScratchMergedHeightmapTextureArray->Desc.ArraySize);
+					}
+
+					check((NeighborIndex != 4) || (NeighborIndexInBatch != INDEX_NONE)); // The central component (the one we finalize) should always be valid
+					GET_SCALAR_ARRAY_ELEMENT(FinalizeHeightmapPSParams->InNeighborHeightmapIndices, NeighborIndex) = NeighborIndexInBatch;
+				}
+
+				FLandscapeLayersHeightmapsFinalizeHeightmapPS::FinalizeHeightmap(GraphBuilder, FinalizeHeightmapPSParams);
+			}
+		}
+	}
+
+	void GenerateHeightmapMips(const FMergeInfo& InLocalMergeInfo, const TMap<FTexture2DResource*, FLandscapeRDGTrackedTexture>& InTrackedTextures, FRDGBuilder& GraphBuilder)
+	{
+		RDG_EVENT_SCOPE(GraphBuilder, "Generate mips for %d heightmaps", InLocalMergeInfo.TextureToResolveInfos.Num());
+
+		// For each texture to resolve, simply regenerate the mips on the entire texture :
+		for (const FTextureResolveInfo& TextureResolveInfo : InLocalMergeInfo.TextureToResolveInfos)
+		{
+			RDG_EVENT_SCOPE(GraphBuilder, "Texture %s", *TextureResolveInfo.Texture->GetTextureName().ToString());
+
+			const FLandscapeRDGTrackedTexture* TrackedTexture = InTrackedTextures.Find(TextureResolveInfo.Texture);
+			check((TrackedTexture != nullptr) && (TrackedTexture->ScratchTextureRef != nullptr));
+
+			check(InLocalMergeInfo.NumMips == TrackedTexture->ScratchTextureRef->Desc.NumMips);
+			check(TrackedTexture->ScratchTextureMipsSRVRefs.Num() == InLocalMergeInfo.NumMips);
+
+			FIntPoint CurrentMipSubregionSize = InLocalMergeInfo.ComponentSizeVerts;
+			for (int32 MipLevel = 1; MipLevel < InLocalMergeInfo.NumMips; ++MipLevel)
+			{
+				CurrentMipSubregionSize.X >>= 1;
+				CurrentMipSubregionSize.Y >>= 1;
+
+				FLandscapeLayersHeightmapsGenerateMipsPS::FParameters* GenerateMipsPSParams = GraphBuilder.AllocParameters<FLandscapeLayersHeightmapsGenerateMipsPS::FParameters>();
+				GenerateMipsPSParams->RenderTargets[0] = FRenderTargetBinding(TrackedTexture->ScratchTextureRef, ERenderTargetLoadAction::ENoAction, MipLevel);
+				GenerateMipsPSParams->InCurrentMipSubregionSize = FUintVector2(CurrentMipSubregionSize.X, CurrentMipSubregionSize.Y);
+				GenerateMipsPSParams->InSourceHeightmap = TrackedTexture->ScratchTextureMipsSRVRefs[MipLevel - 1];
+
+				FLandscapeLayersHeightmapsGenerateMipsPS::GenerateSingleMip(GraphBuilder, GenerateMipsPSParams);
+			}
+		}
+	}
+
+	void CopyScratchToSourceHeightmaps(const FMergeInfo& InLocalMergeInfo, const TMap<FTexture2DResource*, FLandscapeRDGTrackedTexture>& InTrackedTextures, FRDGBuilder& GraphBuilder)
+	{
+		RDG_EVENT_SCOPE(GraphBuilder, "Copy %d scratch to source heightmaps", InLocalMergeInfo.ComponentToResolveInfos.Num());
+
+		// For each texture to resolve, copy from scratch to final texture :
+		for (const FTextureResolveInfo& TextureResolveInfo : InLocalMergeInfo.TextureToResolveInfos)
+		{
+			const FLandscapeRDGTrackedTexture* TrackedTexture = InTrackedTextures.Find(TextureResolveInfo.Texture);
+			check((TrackedTexture != nullptr) && (TrackedTexture->ScratchTextureRef != nullptr) && (TrackedTexture->ExternalTextureRef != nullptr));
+
+			FRHICopyTextureInfo CopyTextureInfo;
+			// We want to copy all mips : 
+			CopyTextureInfo.NumMips = InLocalMergeInfo.NumMips;
+
+			AddCopyTexturePass(GraphBuilder, TrackedTexture->ScratchTextureRef, TrackedTexture->ExternalTextureRef, CopyTextureInfo);
+		}
+	}
+
+	FRDGTextureRef CreateAndClearEmptyTexture(const FMergeInfo& InLocalMergeInfo, FRDGBuilder& GraphBuilder)
+	{
+		// Convert the height value 0.0f to how it's stored in the texture : 
+		const uint16 HeightValue = LandscapeDataAccess::GetTexHeight(0.0f);
+		FLinearColor ClearHeightColor((float)((HeightValue - (HeightValue & 255)) >> 8) / 255.0f, (float)(HeightValue & 255) / 255.0f, 0.0f, 0.0f);
+
+		// TODO [jonathan.bard] : change to PF_R8G8 once edit layers heightmaps are stored as such :
+		FRDGTextureDesc Desc = FRDGTextureDesc::Create2D(InLocalMergeInfo.HeightmapSize, PF_B8G8R8A8, FClearValueBinding(ClearHeightColor), TexCreate_ShaderResource | TexCreate_RenderTargetable, InLocalMergeInfo.NumMips, /*InNumSamples = */1);
+		FRDGTextureRef EmptyTexture = GraphBuilder.CreateTexture(Desc, TEXT("LandscapeEditLayersEmptyHeightmap"));
+
+		FRDGTextureClearInfo ClearInfo;
+		ClearInfo.NumMips = InLocalMergeInfo.NumMips;
+		AddClearRenderTargetPass(GraphBuilder, EmptyTexture, ClearInfo);
+
+		return EmptyTexture;
+	}
+
+	void ClearSourceHeightmaps(const FMergeInfo& InLocalMergeInfo, FRDGBuilder& GraphBuilder)
+	{
+		RDG_EVENT_SCOPE(GraphBuilder, "Clear %d source heightmaps", InLocalMergeInfo.ComponentToResolveInfos.Num());
+
+		// We cannot clear heightmaps directly, since they are external, non-render targetable, textures so we need to copy from an empty source texture :
+		FRDGTextureRef EmptyTexture = CreateAndClearEmptyTexture(InLocalMergeInfo, GraphBuilder);
+
+		// For each texture to resolve, copy from empty to final texture :
+		for (const FTextureResolveInfo& TextureResolveInfo : InLocalMergeInfo.TextureToResolveInfos)
+		{
+			// Register the output texture to the GraphBuilder so that we can copy to it: 
+			FString* DebugName = GraphBuilder.AllocObject<FString>(TextureResolveInfo.Texture->GetTextureName().ToString());
+			TRefCountPtr<IPooledRenderTarget> RenderTarget = CreateRenderTarget(TextureResolveInfo.Texture->TextureRHI, **DebugName);
+
+			// Force tracking on the external texture, so that it can be copied to via CopyTexture within the graph : 
+			FRDGTextureRef DestinationTexture = GraphBuilder.RegisterExternalTexture(RenderTarget, ERenderTargetTexture::ShaderResource, ERDGTextureFlags::ForceTracking);
+
+			FRHICopyTextureInfo CopyTextureInfo;
+			// We want to copy all mips : 
+			CopyTextureInfo.NumMips = InLocalMergeInfo.NumMips;
+
+			AddCopyTexturePass(GraphBuilder, EmptyTexture, DestinationTexture, CopyTextureInfo);
+		}
+	}
+}; // End namespace EditLayersHeightmapLocalMerge_RenderThread
+
+void ALandscape::PrepareLayersHeightmapsLocalMergeRenderThreadData(const FUpdateLayersContentContext& InUpdateLayersContentContext, const FEditLayersHeightmapMergeParams& InMergeParams, EditLayersHeightmapLocalMerge_RenderThread::FMergeInfo& OutRenderThreadData)
+{
+	using namespace EditLayersHeightmapLocalMerge_RenderThread;
+
+	TRACE_CPUPROFILER_EVENT_SCOPE(LandscapeLayers_PrepareLayersHeightmapsLocalMergeRenderThreadData);
+
+	ULandscapeInfo* Info = GetLandscapeInfo();
+	check(Info != nullptr);
+
+	// Number of vertices for each landscape component :
+	int32 ComponentSizeVerts = (SubsectionSizeQuads + 1) * NumSubsections;
+	OutRenderThreadData.ComponentSizeVerts = FIntPoint(ComponentSizeVerts, ComponentSizeVerts);
+	OutRenderThreadData.LandscapeGridScale = GetRootComponent()->GetRelativeScale3D();
+
+	// Prepare landscape edit layers data common to all landscape components: 
+	OutRenderThreadData.VisibleEditLayerInfos.Reserve(LandscapeLayers.Num());
+	for (const FLandscapeLayer& Layer : LandscapeLayers)
+	{
+		if (Layer.bVisible && !InMergeParams.bSkipBrush)
+		{
+			FEditLayerInfo& EditLayerInfo = OutRenderThreadData.VisibleEditLayerInfos.Emplace_GetRef();
+			EditLayerInfo.BlendMode = LandscapeBlendModeToEditLayerBlendMode(Layer.BlendMode);
+			EditLayerInfo.Alpha = Layer.HeightmapAlpha;
+		}
+	}
+
+	const int32 NumComponentsToRender = InUpdateLayersContentContext.LandscapeComponentsHeightmapsToRender.Num();
+
+	// Lookup table to retrieve, for a given rendered component, the index of its FComponentRenderInfo in ComponentToRenderInfos :
+	TMap<ULandscapeComponent*, int32> ComponentToComponentRenderInfoIndex;
+
+	// Prepare all per-landscape component render data : 
+	{
+		TRACE_CPUPROFILER_EVENT_SCOPE(LandscapeLayers_PrepareHeightmapComponentRenderInfos);
+
+		OutRenderThreadData.ComponentToRenderInfos.Reserve(InUpdateLayersContentContext.LandscapeComponentsHeightmapsToRender.Num());
+		for (ULandscapeComponent* Component : InUpdateLayersContentContext.LandscapeComponentsHeightmapsToRender)
+		{
+			// Add a new component render info and set it up :
+			FComponentRenderInfo& NewComponentRenderInfo = OutRenderThreadData.ComponentToRenderInfos.Emplace_GetRef();
+			NewComponentRenderInfo.Name = Component->GetName();
+
+			// Associate the component with its index in ComponentToRenderInfos :
+			check(!ComponentToComponentRenderInfoIndex.Contains(Component));
+			ComponentToComponentRenderInfoIndex.Add(Component, OutRenderThreadData.ComponentToRenderInfos.Num() - 1);
+
+			UTexture2D* ComponentHeightmap = Component->GetHeightmap();
+			OutRenderThreadData.HeightmapSize = FIntPoint(ComponentHeightmap->Source.GetSizeX(), ComponentHeightmap->Source.GetSizeY());
+
+			const int32 HeightmapOffsetX = Component->HeightmapScaleBias.Z * (float)OutRenderThreadData.HeightmapSize.X;
+			const int32 HeightmapOffsetY = Component->HeightmapScaleBias.W * (float)OutRenderThreadData.HeightmapSize.Y;
+			// Effective area of the texture affecting this component (because of texture sharing) :
+			const FIntRect ComponentTextureSubregion(FIntPoint(HeightmapOffsetX, HeightmapOffsetY), FIntPoint(HeightmapOffsetX, HeightmapOffsetY) + OutRenderThreadData.ComponentSizeVerts);
+
+			NewComponentRenderInfo.VisibleLayerHeightmapTextures.Reserve(OutRenderThreadData.VisibleEditLayerInfos.Num());
+			for (const FLandscapeLayer& Layer : LandscapeLayers)
+			{
+				if (Layer.bVisible && !InMergeParams.bSkipBrush)
+				{
+					UTexture2D* LayerHeightmap = Component->GetLayerData(Layer.Guid)->HeightmapData.Texture.Get();
+					NewComponentRenderInfo.VisibleLayerHeightmapTextures.Add(FTexture2DResourceSubregion(LayerHeightmap->GetResource()->GetTexture2DResource(), ComponentTextureSubregion));
+				}
+			}
+		}
+
+		// Now that all landscape components have been registered, identify the (valid) neighbors for each : 
+		for (int32 ComponentToRenderIndex = 0; ComponentToRenderIndex < NumComponentsToRender; ++ComponentToRenderIndex)
+		{
+			ULandscapeComponent* Component = InUpdateLayersContentContext.LandscapeComponentsHeightmapsToRender[ComponentToRenderIndex];
+			FComponentRenderInfo& ComponentRenderInfo = OutRenderThreadData.ComponentToRenderInfos[ComponentToRenderIndex];
+
+			// Gather neighboring component infos : 
+			TStaticArray<ULandscapeComponent*, 9> NeighborComponents;
+			Component->GetLandscapeComponentNeighbors3x3(NeighborComponents);
+			for (int32 NeighborIndex = 0; NeighborIndex < 9; ++NeighborIndex)
+			{
+				int32 NeighborComponentRenderInfoIndex = INDEX_NONE;
+				if (ULandscapeComponent* NeighborComponent = NeighborComponents[NeighborIndex])
+				{
+					if (int32* NeighborComponentRenderInfoIndexPtr = ComponentToComponentRenderInfoIndex.Find(NeighborComponent))
+					{
+						NeighborComponentRenderInfoIndex = *NeighborComponentRenderInfoIndexPtr;
+					}
+				}
+				ComponentRenderInfo.NeighborComponentToRenderInfoIndices[NeighborIndex] = NeighborComponentRenderInfoIndex;
+			}
+		}
+	}
+
+	// List of UTexture2D that we need to kick off readbacks for :
+	TArray<UTexture2D*> TexturesNeedingReadback;
+
+	// Prepare per-landscape component resolve data : 
+	{
+		TRACE_CPUPROFILER_EVENT_SCOPE(LandscapeLayers_PrepareHeightmapComponentResolveInfos);
+		OutRenderThreadData.ComponentToResolveInfos.Reserve(InUpdateLayersContentContext.LandscapeComponentsHeightmapsToResolve.Num());
+		for (ULandscapeComponent* Component : InUpdateLayersContentContext.LandscapeComponentsHeightmapsToResolve)
+		{
+			UTexture2D* ComponentHeightmap = Component->GetHeightmap();
+			const int32 HeightmapOffsetX = Component->HeightmapScaleBias.Z * (float)ComponentHeightmap->Source.GetSizeX();
+			const int32 HeightmapOffsetY = Component->HeightmapScaleBias.W * (float)ComponentHeightmap->Source.GetSizeY();
+			// Effective area of the texture affecting this component (because of texture sharing) :
+			const FIntRect ComponentTextureSubregion(FIntPoint(HeightmapOffsetX, HeightmapOffsetY), FIntPoint(HeightmapOffsetX, HeightmapOffsetY) + OutRenderThreadData.ComponentSizeVerts);
+
+			FComponentResolveInfo& NewComponentResolveInfo = OutRenderThreadData.ComponentToResolveInfos.Add_GetRef(FComponentResolveInfo(NumComponentsToRender));
+			NewComponentResolveInfo.ComponentToResolveInfoIndex = OutRenderThreadData.ComponentToResolveInfos.Num() - 1;
+			NewComponentResolveInfo.Name = Component->GetName();
+			NewComponentResolveInfo.Heightmap = FTexture2DResourceSubregion(ComponentHeightmap->GetResource()->GetTexture2DResource(), ComponentTextureSubregion);
+
+			// Gather neighboring component infos : 
+			TStaticArray<ULandscapeComponent*, 9> NeighborComponents;
+			Component->GetLandscapeComponentNeighbors3x3(NeighborComponents);
+			for (int32 NeighborIndex = 0; NeighborIndex < 9; ++NeighborIndex)
+			{
+				int32 NeighborComponentRenderInfoIndex = INDEX_NONE;
+				if (ULandscapeComponent* NeighborComponent = NeighborComponents[NeighborIndex])
+				{
+					if (int32* NeighborComponentRenderInfoIndexPtr = ComponentToComponentRenderInfoIndex.Find(NeighborComponent))
+					{
+						NeighborComponentRenderInfoIndex = *NeighborComponentRenderInfoIndexPtr;
+					}
+				}
+				NewComponentResolveInfo.SetNeighborRenderInfo(NeighborIndex, NeighborComponentRenderInfoIndex);
+			}
+
+			FTextureResolveInfo* TextureResolveInfo = OutRenderThreadData.TextureToResolveInfos.FindByPredicate(
+				[TextureToResolve = NewComponentResolveInfo.Heightmap.Texture](const FTextureResolveInfo& TextureResolveInfo) { return TextureResolveInfo.Texture == TextureToResolve; });
+			if (TextureResolveInfo == nullptr)
+			{
+				ALandscapeProxy* Proxy = Component->GetLandscapeProxy();
+				FLandscapeEditLayerReadback** CPUReadback = Proxy->HeightmapsCPUReadback.Find(ComponentHeightmap);
+				check((CPUReadback != nullptr) && (*CPUReadback != nullptr));
+
+				FTextureResolveInfo NewTextureResolveInfo;
+				NewTextureResolveInfo.Texture = NewComponentResolveInfo.Heightmap.Texture;
+				NewTextureResolveInfo.CPUReadback = *CPUReadback;
+				OutRenderThreadData.TextureToResolveInfos.Add(NewTextureResolveInfo);
+
+				check(!TexturesNeedingReadback.Contains(ComponentHeightmap));
+				TexturesNeedingReadback.Add(ComponentHeightmap);
+			}
+		}
+	}
+
+	// Prepare the texture resolve batches: 
+	{
+		TRACE_CPUPROFILER_EVENT_SCOPE(LandscapeLayers_PrepareHeightmapComponentResolveInfoBatches);
+
+		int32 MaxComponentsPerResolveBatch = CVarLandscapeEditLayersMaxComponentsPerHeightmapResolveBatch.GetValueOnGameThread();
+
+		// Copy the component infos because TextureToResolveInfos indices need to remain stable at this point :
+		TArray<FComponentResolveInfo> RemainingComponentToResolveInfos = OutRenderThreadData.ComponentToResolveInfos;
+		if (!RemainingComponentToResolveInfos.IsEmpty())
+		{
+			TBitArray<TInlineAllocator<1>> TempBitArray;
+			TempBitArray.Reserve(OutRenderThreadData.ComponentToResolveInfos.Num());
+
+			while (!RemainingComponentToResolveInfos.IsEmpty())
+			{
+				const FComponentResolveInfo& ComponentResolveInfo = RemainingComponentToResolveInfos.Pop(/*bAllowShrinking = */false);
+
+				int32 BestBatchIndex = INDEX_NONE;
+				int32 MinNumComponents = MAX_int32;
+
+				// Iterate through all all batches and try to find which would be able to accept it and amongst those, which it would share the most components to render with:
+				int32 NumBatches = OutRenderThreadData.ComponentResolveBatchInfos.Num();
+				for (int32 BatchIndex = 0; BatchIndex < NumBatches; ++BatchIndex)
+				{
+					const FComponentResolveBatchInfo& Batch = OutRenderThreadData.ComponentResolveBatchInfos[BatchIndex];
+					TempBitArray = TBitArray<>::BitwiseOR(Batch.ComponentToRenderInfoBitIndices, ComponentResolveInfo.NeighborComponentToRenderInfoBitIndices, EBitwiseOperatorFlags::MinSize);
+
+					// If after adding its components, the batch still has less than MaxComponentsPerResolveBatch, it can accept it
+					int32 NumComponentsAfter = TempBitArray.CountSetBits();
+					if (NumComponentsAfter <= MaxComponentsPerResolveBatch)
+					{
+						// Is this the best candidate so far?
+						if (NumComponentsAfter < MinNumComponents)
+						{
+							BestBatchIndex = BatchIndex;
+							MinNumComponents = NumComponentsAfter;
+						}
+
+						// If the number of components after addition of this texture is unchanged, it's a perfect match, we won't ever find a better batch so just stop there for this texture:
+						if (NumComponentsAfter == Batch.ComponentToRenderInfoBitIndices.CountSetBits())
+						{
+							break;
+						}
+					}
+				}
+
+				// If we have found a batch, just add the texture to it, otherwise, add a new batch:
+				FComponentResolveBatchInfo& SelectedBatch = (BestBatchIndex != INDEX_NONE) ? OutRenderThreadData.ComponentResolveBatchInfos[BestBatchIndex]
+					: OutRenderThreadData.ComponentResolveBatchInfos.Add_GetRef(FComponentResolveBatchInfo(OutRenderThreadData.ComponentToRenderInfos.Num(), /*InBatchIndex = */OutRenderThreadData.ComponentResolveBatchInfos.Num()));
+
+				SelectedBatch.AddComponent(ComponentResolveInfo);
+				check(SelectedBatch.ComponentToRenderInfoBitIndices.CountSetBits() <= MaxComponentsPerResolveBatch);
+
+				// Keep track of the maximum number of slices in the scratch texture arrays we'll need for any given batch :
+				OutRenderThreadData.MaxNumComponentsToRenderPerResolveComponentBatch = FMath::Max(SelectedBatch.ComponentToRenderInfoBitIndices.CountSetBits(), OutRenderThreadData.MaxNumComponentsToRenderPerResolveComponentBatch);
+			}
+		}
+	}
+
+	// Finalize :
+	{
+		// Prepare the UTexture2D readbacks we'll need to perform :
+		OutRenderThreadData.DeferredCopyReadbackTextures = PrepareLandscapeLayersCopyReadbackTextureParams(InUpdateLayersContentContext.MapHelper, TexturesNeedingReadback, /*bWeightmaps = */false);
+
+		// We'll only ever need this amount of edit layers textures for any MergeEditLayers operation :
+		OutRenderThreadData.MaxNumEditLayersTexturesToMerge = OutRenderThreadData.VisibleEditLayerInfos.Num();
+
+		check(OutRenderThreadData.HeightmapSize != FIntPoint::ZeroValue);
+		OutRenderThreadData.NumMips = FMath::CeilLogTwo(OutRenderThreadData.HeightmapSize.X) + 1;
+	}
+}
+
+int32 ALandscape::PerformLayersHeightmapsLocalMerge(const FUpdateLayersContentContext& InUpdateLayersContentContext, const FEditLayersHeightmapMergeParams& InMergeParams)
+{
+	using namespace EditLayersHeightmapLocalMerge_RenderThread;
+
+	TRACE_CPUPROFILER_EVENT_SCOPE(LandscapeLayers_PerformLayersWeightmapsLocalMerge);
+
+	FMergeInfo RenderThreadData;
+	PrepareLayersHeightmapsLocalMergeRenderThreadData(InUpdateLayersContentContext, InMergeParams, RenderThreadData);
+	
+	ENQUEUE_RENDER_COMMAND(PerformLayersHeightmapsLocalMerge)([RenderThreadData](FRHICommandListImmediate& RHICmdList)
+	{
+		FMemMark Mark(FMemStack::Get());
+		FRDGBuilder GraphBuilder(RHICmdList, RDG_EVENT_NAME("PerformLayersHeightmapsLocalMerge"));
+
+		if (RenderThreadData.NeedsMerge())
+		{
+			// Prepare the GPU resources we will use during the local merge : 
+			FRDGResources RDGResources;
+			PrepareLayersHeightmapsLocalMergeRDGResources(RenderThreadData, GraphBuilder, RDGResources);
+
+			// Get a list of all external textures (heightmaps) we will manipulate during the local merge : 
+			TMap<FTexture2DResource*, FLandscapeRDGTrackedTexture> TrackedTextures;
+			GatherLayersHeightmapsLocalMergeRDGTextures(RenderThreadData, TrackedTextures);
+
+			// Start tracking those in the render graph :
+			TrackLandscapeRDGTextures(GraphBuilder, TrackedTextures);
+
+			// Process the components batch by batch in order to avoid over-allocating temporary textures :
+			for (const FComponentResolveBatchInfo& ComponentResolveBatchInfo : RenderThreadData.ComponentResolveBatchInfos)
+			{
+				RDG_EVENT_SCOPE(GraphBuilder, "Process batch %d", ComponentResolveBatchInfo.BatchIndex);
+
+				// Perform all edit layers merges, for all components to render in that batch :
+				MergeEditLayersHeightmapsForBatch(ComponentResolveBatchInfo, RenderThreadData, TrackedTextures, GraphBuilder, RDGResources);
+
+				// Correct borders of all rendered components so that they're all stitched together :
+				StitchHeightmapsForBatch(ComponentResolveBatchInfo, RenderThreadData, GraphBuilder, RDGResources);
+
+				// Finalize (compute normals) of each component to resolve :
+				FinalizeComponentsForBatch(ComponentResolveBatchInfo, RenderThreadData, TrackedTextures, GraphBuilder, RDGResources);
+			}
+
+			// Generate the mips on the entire heightmaps : 
+			GenerateHeightmapMips(RenderThreadData, TrackedTextures, GraphBuilder);
+
+			// Finally, we can put those scratch textures to good usage and update our actual heightmaps : 
+			CopyScratchToSourceHeightmaps(RenderThreadData, TrackedTextures, GraphBuilder);
+		}
+		else
+		{
+			// When there's nothing to do, we still have the obligation to output empty heightmaps :
+			ClearSourceHeightmaps(RenderThreadData, GraphBuilder);
+		}
+
+		GraphBuilder.Execute();
+
+	});
+
+	ExecuteCopyToReadbackTexture(RenderThreadData.DeferredCopyReadbackTextures);
+
+	return InMergeParams.HeightmapUpdateModes;
+}
+
+int32 ALandscape::PerformLayersHeightmapsGlobalMerge(const FUpdateLayersContentContext& InUpdateLayersContentContext, const FEditLayersHeightmapMergeParams& InMergeParams)
+{
+	ULandscapeInfo* Info = GetLandscapeInfo();
+	check(Info != nullptr);
+	check(HeightmapRTList.Num() > 0);
+
+	FIntRect LandscapeExtent;
+	if (!Info->GetLandscapeExtent(LandscapeExtent.Min.X, LandscapeExtent.Min.Y, LandscapeExtent.Max.X, LandscapeExtent.Max.Y))
+	{
+		return 0;
+	}
+
+	// Use to compute top-left vertex position per Heightmap and the actual size to copy : 
+	struct FHeightmapCopyInfo
+	{
+		FHeightmapCopyInfo(UTexture2D* InTexture, const FIntPoint& InComponentVertexPosition, int32 InComponentSizeVerts, FLandscapeEditLayerReadback* InCPUReadback = nullptr)
+			: Texture(InTexture)
+			, ComponentSizeVerts(InComponentSizeVerts)
+			, SectionRect(InComponentVertexPosition, InComponentVertexPosition + FIntPoint(InComponentSizeVerts, InComponentSizeVerts))
+			, CPUReadback(InCPUReadback)
+		{}
+
+		FHeightmapCopyInfo(FHeightmapCopyInfo&&) = default;
+		FHeightmapCopyInfo& operator=(FHeightmapCopyInfo&&) = default;
+
+		void Union(const FIntPoint& InComponentVertexPosition)
+		{
+			SectionRect.Union(FIntRect(InComponentVertexPosition, InComponentVertexPosition + FIntPoint(ComponentSizeVerts, ComponentSizeVerts)));
+		}
+
+		UTexture2D* Texture;
+		int32 ComponentSizeVerts;
+		FIntRect SectionRect;
+		FLandscapeEditLayerReadback* CPUReadback;
+	};
+
+	// Calculate Top Left Lambda
+	auto GetUniqueHeightmaps = [&](const TArray<ULandscapeComponent*>& InLandscapeComponents, TArray<FHeightmapCopyInfo>& OutHeightmaps, const FIntPoint& LandscapeBaseQuads, const FGuid& InLayerGuid = FGuid())
+	{
+		FScopedSetLandscapeEditingLayer Scope(this, InLayerGuid);
+
+		const int32 ComponentSizeQuad = SubsectionSizeQuads * NumSubsections;
+		const int32 ComponentSizeVerts = (SubsectionSizeQuads + 1) * NumSubsections;
+		for (ULandscapeComponent* Component : InLandscapeComponents)
+		{
+			UTexture2D* ComponentHeightmap = Component->GetHeightmap(true);
+
+			int32 Index = OutHeightmaps.IndexOfByPredicate([=](const FHeightmapCopyInfo& LayerHeightmap) { return LayerHeightmap.Texture == ComponentHeightmap; });
+
+			FIntPoint ComponentSectionBase = Component->GetSectionBase() - LandscapeBaseQuads;
+			FVector2D SourcePositionOffset(FMath::RoundToInt((float)ComponentSectionBase.X / ComponentSizeQuad), FMath::RoundToInt((float)ComponentSectionBase.Y / ComponentSizeQuad));
+			FIntPoint ComponentVertexPosition = FIntPoint(SourcePositionOffset.X * ComponentSizeVerts, SourcePositionOffset.Y * ComponentSizeVerts);
+			ALandscapeProxy* Proxy = Component->GetLandscapeProxy();
+
+			if (Index == INDEX_NONE)
+			{
+				FLandscapeEditLayerReadback** CPUReadback = Proxy->HeightmapsCPUReadback.Find(ComponentHeightmap);
+				OutHeightmaps.Add(FHeightmapCopyInfo(ComponentHeightmap, ComponentVertexPosition, ComponentSizeVerts, CPUReadback != nullptr ? *CPUReadback : nullptr));
+			}
+			else
+			{
+				OutHeightmaps[Index].Union(ComponentVertexPosition);
+			}
+		}
+	};
+
+	FLandscapeLayersHeightmapShaderParameters ShaderParams;
+
+	bool FirstLayer = true;
+	UTextureRenderTarget2D* CombinedHeightmapAtlasRT = HeightmapRTList[(int32)EHeightmapRTType::HeightmapRT_CombinedAtlas];
+	UTextureRenderTarget2D* CombinedHeightmapNonAtlasRT = HeightmapRTList[(int32)EHeightmapRTType::HeightmapRT_CombinedNonAtlas];
+	UTextureRenderTarget2D* LandscapeScratchRT1 = HeightmapRTList[(int32)EHeightmapRTType::HeightmapRT_Scratch1];
+	UTextureRenderTarget2D* LandscapeScratchRT2 = HeightmapRTList[(int32)EHeightmapRTType::HeightmapRT_Scratch2];
+	UTextureRenderTarget2D* LandscapeScratchRT3 = HeightmapRTList[(int32)EHeightmapRTType::HeightmapRT_Scratch3];
+
+	for (FLandscapeLayer& Layer : LandscapeLayers)
+	{
+		//Draw Layer heightmap to Combined RT Atlas
+		ShaderParams.ApplyLayerModifiers = false;
+		ShaderParams.LayerVisible = Layer.bVisible;
+		ShaderParams.GenerateNormals = false;
+		ShaderParams.LayerBlendMode = Layer.BlendMode;
+
+		if (Layer.BlendMode == LSBM_AlphaBlend)
+		{
+			// For now, only Layer reserved for Landscape Splines will use the AlphaBlendMode
+			const FLandscapeLayer* SplinesReservedLayer = GetLandscapeSplinesReservedLayer();
+			check(&Layer == SplinesReservedLayer);
+			ShaderParams.LayerAlpha = 1.0f;
+		}
+		else
+		{
+			check(Layer.BlendMode == LSBM_AdditiveBlend);
+			ShaderParams.LayerAlpha = Layer.HeightmapAlpha;
+		}
+
+		{
+			TArray<FLandscapeLayersCopyTextureParams> DeferredCopyTextures;
+			TArray<FHeightmapCopyInfo> LayerHeightmaps;
+			GetUniqueHeightmaps(InUpdateLayersContentContext.LandscapeComponentsHeightmapsToRender, LayerHeightmaps, LandscapeExtent.Min, Layer.Guid);
+			for (const FHeightmapCopyInfo& LayerHeightmap : LayerHeightmaps)
+			{
+				FLandscapeLayersCopyTextureParams& CopyTextureParams = DeferredCopyTextures.Add_GetRef(FLandscapeLayersCopyTextureParams(LayerHeightmap.Texture, LandscapeScratchRT1));
+				// Only copy the size that's actually needed : 
+				CopyTextureParams.CopySize = LayerHeightmap.SectionRect.Size();
+				// Copy from the heightmap's top-left corner to the composited texture's position :
+				CopyTextureParams.DestPosition = LayerHeightmap.SectionRect.Min;
+			}
+			ExecuteCopyLayersTexture(MoveTemp(DeferredCopyTextures));
+		}
+
+		// NOTE: From this point on, we always work in non atlas, we'll convert back at the end to atlas only
+		DrawHeightmapComponentsToRenderTarget(FString::Printf(TEXT("LS Height: %s += -> NonAtlas %s"), *Layer.Name.ToString(), *LandscapeScratchRT1->GetName(), *LandscapeScratchRT2->GetName()),
+			InUpdateLayersContentContext.LandscapeComponentsHeightmapsToRender, LandscapeExtent.Min, LandscapeScratchRT1, nullptr, LandscapeScratchRT2, ERTDrawingType::RTAtlasToNonAtlas, true, ShaderParams);
+
+		ShaderParams.ApplyLayerModifiers = true;
+
+		// Combine Current layer with current result
+		DrawHeightmapComponentsToRenderTarget(FString::Printf(TEXT("LS Height: %s += -> CombinedNonAtlas %s"), *Layer.Name.ToString(), *LandscapeScratchRT2->GetName(), *CombinedHeightmapNonAtlasRT->GetName()),
+			InUpdateLayersContentContext.LandscapeComponentsHeightmapsToRender, LandscapeExtent.Min, LandscapeScratchRT2, FirstLayer ? nullptr : LandscapeScratchRT3, CombinedHeightmapNonAtlasRT, ERTDrawingType::RTNonAtlas, FirstLayer, ShaderParams);
+
+		ShaderParams.ApplyLayerModifiers = false;
+
+		if (Layer.bVisible && !InMergeParams.bSkipBrush)
+		{
+			// Draw each brushes
+			for (int32 i = 0; i < Layer.Brushes.Num(); ++i)
+			{
+				// TODO: handle conversion from float to RG8 by using material params to write correct values
+				// TODO: handle conversion/handling of RT not same size as internal size
+
+				FLandscapeLayerBrush& Brush = Layer.Brushes[i];
+				UTextureRenderTarget2D* BrushOutputNonAtlasRT = Brush.Render(true, LandscapeExtent, CombinedHeightmapNonAtlasRT);
+				if ((BrushOutputNonAtlasRT == nullptr) 
+					|| (BrushOutputNonAtlasRT->SizeX != CombinedHeightmapNonAtlasRT->SizeX)
+					|| (BrushOutputNonAtlasRT->SizeY != CombinedHeightmapNonAtlasRT->SizeY))
+				{
+					continue;
+				}
+
+				ALandscapeBlueprintBrushBase* LandscapeBrush = Brush.GetBrush();
+				check(LandscapeBrush != nullptr); // If we managed to render, the brush should be valid
+
+				INC_DWORD_STAT(STAT_LandscapeLayersRegenerateDrawCalls); // Brush Render
+
+				PrintLayersDebugRT(FString::Printf(TEXT("LS Height: %s %s -> BrushNonAtlas %s"), *Layer.Name.ToString(), *LandscapeBrush->GetName(), *BrushOutputNonAtlasRT->GetName()), BrushOutputNonAtlasRT);
+
+				// Resolve back to Combined heightmap (it's unlikely, but possible that the brush returns the same RT as input and output, if it did various operations on it, in which case the copy is useless) :
+				if (BrushOutputNonAtlasRT != CombinedHeightmapNonAtlasRT)
+				{
+					ExecuteCopyLayersTexture({ FLandscapeLayersCopyTextureParams(BrushOutputNonAtlasRT, CombinedHeightmapNonAtlasRT) });
+					PrintLayersDebugRT(FString::Printf(TEXT("LS Height: %s Component %s += -> CombinedNonAtlas %s"), *Layer.Name.ToString(), *BrushOutputNonAtlasRT->GetName(), *CombinedHeightmapNonAtlasRT->GetName()), CombinedHeightmapNonAtlasRT);
+				}
+			}
+		}
+
+		ExecuteCopyLayersTexture({ FLandscapeLayersCopyTextureParams(CombinedHeightmapNonAtlasRT, LandscapeScratchRT3) });
+		PrintLayersDebugRT(FString::Printf(TEXT("LS Height: %s Component %s += -> CombinedNonAtlas %s"), *Layer.Name.ToString(), *CombinedHeightmapNonAtlasRT->GetName(), *LandscapeScratchRT3->GetName()), LandscapeScratchRT3);
+
+		FirstLayer = false;
+	}
+
+	ShaderParams.GenerateNormals = true;
+	ShaderParams.GridSize = GetRootComponent()->GetRelativeScale3D();
+
+	// Broadcast Event of the Full Render
+	if ((InMergeParams.HeightmapUpdateModes & Update_Heightmap_All) == Update_Heightmap_All)
+	{
+		LandscapeFullHeightmapRenderDoneDelegate.Broadcast(LandscapeScratchRT3);
+	}
+
+	DrawHeightmapComponentsToRenderTarget(FString::Printf(TEXT("LS Height: %s = -> CombinedNonAtlasNormals : %s"), *CombinedHeightmapNonAtlasRT->GetName(), *LandscapeScratchRT1->GetName()),
+		InUpdateLayersContentContext.LandscapeComponentsHeightmapsToRender, LandscapeExtent.Min, CombinedHeightmapNonAtlasRT, nullptr, LandscapeScratchRT1, ERTDrawingType::RTNonAtlas, true, ShaderParams);
+
+	ShaderParams.GenerateNormals = false;
+
+	DrawHeightmapComponentsToRenderTarget(FString::Printf(TEXT("LS Height: %s = -> CombinedAtlasFinal : %s"), *LandscapeScratchRT1->GetName(), *CombinedHeightmapAtlasRT->GetName()),
+		InUpdateLayersContentContext.LandscapeComponentsHeightmapsToRender, LandscapeExtent.Min, LandscapeScratchRT1, nullptr, CombinedHeightmapAtlasRT, ERTDrawingType::RTNonAtlasToAtlas, true, ShaderParams);
+
+	DrawHeightmapComponentsToRenderTargetMips(InUpdateLayersContentContext.LandscapeComponentsHeightmapsToRender, LandscapeExtent.Min, CombinedHeightmapAtlasRT, true, ShaderParams);
+
+	// List of UTexture2D that we need to kick off readbacks for :
+	TArray<UTexture2D*> TexturesNeedingReadback;
+
+	// Copy back all Mips to original heightmap data
+	{
+		TArray<FLandscapeLayersCopyTextureParams> DeferredCopyTextures;
+		TArray<FHeightmapCopyInfo> Heightmaps;
+		GetUniqueHeightmaps(InUpdateLayersContentContext.LandscapeComponentsHeightmapsToResolve, Heightmaps, LandscapeExtent.Min);
+		for (const FHeightmapCopyInfo& Heightmap : Heightmaps)
+		{
+			check(Heightmap.CPUReadback);
+
+			const FIntPoint Mip0CopySize = Heightmap.SectionRect.Size();
+			const FIntPoint Mip0SourcePosition = Heightmap.SectionRect.Min;
+			// Mip 0
+			{
+				FLandscapeLayersCopyTextureParams& CopyTextureParams = DeferredCopyTextures.Add_GetRef(FLandscapeLayersCopyTextureParams(CombinedHeightmapAtlasRT, Heightmap.Texture));
+				// Only copy the size that's actually needed : 
+				CopyTextureParams.CopySize = Mip0CopySize;
+				// Copy from the composited texture's position to the top-left corner of the heightmap
+				CopyTextureParams.SourcePosition = Mip0SourcePosition;
+			}
+
+			// Other Mips
+			uint8 MipIndex = 1;
+			for (int32 MipRTIndex = (int32)EHeightmapRTType::HeightmapRT_Mip1; MipRTIndex < (int32)EHeightmapRTType::HeightmapRT_Count; ++MipRTIndex)
+			{
+				UTextureRenderTarget2D* RenderTargetMip = HeightmapRTList[MipRTIndex];
+				if (RenderTargetMip != nullptr)
+				{
+					FLandscapeLayersCopyTextureParams& CopyTextureParams = DeferredCopyTextures.Add_GetRef(FLandscapeLayersCopyTextureParams(RenderTargetMip, Heightmap.Texture));
+					CopyTextureParams.CopySize.X = Mip0CopySize.X >> MipIndex;
+					CopyTextureParams.CopySize.Y = Mip0CopySize.Y >> MipIndex;
+					CopyTextureParams.SourcePosition.X = Mip0SourcePosition.X >> MipIndex;
+					CopyTextureParams.SourcePosition.Y = Mip0SourcePosition.Y >> MipIndex;
+					CopyTextureParams.DestMip = MipIndex;
+
+					++MipIndex;
+				}
+			}
+
+			check(!TexturesNeedingReadback.Contains(Heightmap.Texture));
+			TexturesNeedingReadback.Add(Heightmap.Texture);
+		}
+		ExecuteCopyLayersTexture(MoveTemp(DeferredCopyTextures));
+	}
+
+	// Prepare the UTexture2D readbacks we'll need to perform :
+	TArray<FLandscapeLayersCopyReadbackTextureParams> DeferredCopyReadbackTextures = PrepareLandscapeLayersCopyReadbackTextureParams(InUpdateLayersContentContext.MapHelper, TexturesNeedingReadback, /*bWeightmaps = */false);
+	ExecuteCopyToReadbackTexture(DeferredCopyReadbackTextures);
+
+	return InMergeParams.HeightmapUpdateModes;
+}
+
+int32 ALandscape::RegenerateLayersHeightmaps(const FUpdateLayersContentContext& InUpdateLayersContentContext)
+{
 	TRACE_CPUPROFILER_EVENT_SCOPE(LandscapeLayers_RegenerateLayersHeightmaps);
 	SCOPE_CYCLE_COUNTER(STAT_LandscapeLayersRegenerateHeightmaps);
 	ULandscapeInfo* Info = GetLandscapeInfo();
@@ -3190,7 +4826,7 @@
 	}
 
 	// Nothing to do (return that we did the processing)
-	if (InLandscapeComponentsToResolve.Num() == 0)
+	if (InUpdateLayersContentContext.LandscapeComponentsHeightmapsToRender.Num() == 0)
 	{
 		return HeightmapUpdateModes;
 	}
@@ -3198,15 +4834,15 @@
 	// Lazily create CPU read back objects as required
 	if (HeightmapUpdateModes)
 	{
-		for(ULandscapeComponent* Component : InLandscapeComponentsToRender)
+		for (ULandscapeComponent* Component : InUpdateLayersContentContext.LandscapeComponentsHeightmapsToRender)
 		{
 			UTexture2D* ComponentHeightmap = Component->GetHeightmap(false);
 			ALandscapeProxy* Proxy = Component->GetLandscapeProxy();
 			FLandscapeEditLayerReadback** CPUReadback = Proxy->HeightmapsCPUReadback.Find(ComponentHeightmap);
 			if (CPUReadback == nullptr)
-			{
+				{
 				FLandscapeEditLayerReadback* NewCPUReadback = new FLandscapeEditLayerReadback();
-				const uint8* LockedMip = ComponentHeightmap->Source.LockMipReadOnly(0);
+				const uint8* LockedMip = ComponentHeightmap->Source.LockMip(0);
 				const uint32 Hash = FLandscapeEditLayerReadback::CalculateHash(LockedMip, ComponentHeightmap->GetSizeX() * ComponentHeightmap->GetSizeY() * sizeof(FColor));
 				ComponentHeightmap->Source.UnlockMip(0);
 				NewCPUReadback->SetHash(Hash);
@@ -3218,232 +4854,22 @@
 	if (HeightmapUpdateModes || bForceRender)
 	{
 		RenderCaptureInterface::FScopedCapture RenderCapture((RenderCaptureLayersNextHeightmapDraws != 0), TEXT("LandscapeLayersHeightmapCapture"));
+		RenderCaptureLayersNextHeightmapDraws = FMath::Max(0, RenderCaptureLayersNextHeightmapDraws - 1);
 
-		check(HeightmapRTList.Num() > 0);
+		FEditLayersHeightmapMergeParams MergeParams;
+		MergeParams.HeightmapUpdateModes = HeightmapUpdateModes;
+		MergeParams.bForceRender = bForceRender;
+		MergeParams.bSkipBrush = bSkipBrush;
 
-		FIntRect LandscapeExtent;
-		if (!Info->GetLandscapeExtent(LandscapeExtent.Min.X, LandscapeExtent.Min.Y, LandscapeExtent.Max.X, LandscapeExtent.Max.Y))
+		if (SupportsEditLayersLocalMerge())
 		{
-			return 0;
+			return PerformLayersHeightmapsLocalMerge(InUpdateLayersContentContext, MergeParams);
+		}
+		else
+		{
+			return PerformLayersHeightmapsGlobalMerge(InUpdateLayersContentContext, MergeParams);
 		}
 
-		// Use to compute top-left vertex position per Heightmap and the actual size to copy : 
-		struct FHeightmapCopyInfo
-		{
-			FHeightmapCopyInfo(UTexture2D* InTexture, const FIntPoint& InComponentVertexPosition, int32 InComponentSizeVerts, FLandscapeEditLayerReadback* InCPUReadback = nullptr)
-				: Texture(InTexture)
-				, ComponentSizeVerts(InComponentSizeVerts)
-				, SectionRect(InComponentVertexPosition, InComponentVertexPosition + FIntPoint(InComponentSizeVerts, InComponentSizeVerts))
-				, CPUReadback(InCPUReadback)
-			{}
-
-			FHeightmapCopyInfo(FHeightmapCopyInfo&&) = default;
-			FHeightmapCopyInfo& operator=(FHeightmapCopyInfo&&) = default;
-
-			void Union(const FIntPoint& InComponentVertexPosition)
-			{
-				SectionRect.Union(FIntRect(InComponentVertexPosition, InComponentVertexPosition + FIntPoint(ComponentSizeVerts, ComponentSizeVerts)));
-			}
-
-			UTexture2D* Texture;
-			int32 ComponentSizeVerts;
-			FIntRect SectionRect;
-			FLandscapeEditLayerReadback* CPUReadback;
-		};
-
-		// Calculate Top Left Lambda
-		auto GetUniqueHeightmaps = [&](const TArray<ULandscapeComponent*>& InLandscapeComponents, TArray<FHeightmapCopyInfo>& OutHeightmaps, const FIntPoint& LandscapeBaseQuads, const FGuid& InLayerGuid = FGuid())
-		{
-			FScopedSetLandscapeEditingLayer Scope(this, InLayerGuid);
-			
-			const int32 ComponentSizeQuad = SubsectionSizeQuads * NumSubsections;
-			const int32 ComponentSizeVerts = (SubsectionSizeQuads + 1) * NumSubsections;
-			for (ULandscapeComponent* Component : InLandscapeComponents)
-			{
-				UTexture2D* ComponentHeightmap = Component->GetHeightmap(true);
-		
-				int32 Index = OutHeightmaps.IndexOfByPredicate([=](const FHeightmapCopyInfo& LayerHeightmap) { return LayerHeightmap.Texture == ComponentHeightmap; });
-
-				FIntPoint ComponentSectionBase = Component->GetSectionBase() - LandscapeBaseQuads;
-				FVector2D SourcePositionOffset(FMath::RoundToInt((float)ComponentSectionBase.X / ComponentSizeQuad), FMath::RoundToInt((float)ComponentSectionBase.Y / ComponentSizeQuad));
-				FIntPoint ComponentVertexPosition = FIntPoint(SourcePositionOffset.X * ComponentSizeVerts, SourcePositionOffset.Y * ComponentSizeVerts);
-				ALandscapeProxy* Proxy = Component->GetLandscapeProxy();
-
-				if (Index == INDEX_NONE)
-				{
-					FLandscapeEditLayerReadback** CPUReadback = Proxy->HeightmapsCPUReadback.Find(ComponentHeightmap);
-					OutHeightmaps.Add(FHeightmapCopyInfo(ComponentHeightmap, ComponentVertexPosition, ComponentSizeVerts, CPUReadback != nullptr ? *CPUReadback : nullptr));
-				}
-				else
-				{
-					OutHeightmaps[Index].Union(ComponentVertexPosition);
-				}
-			}
-		};
-
-		FLandscapeLayersHeightmapShaderParameters ShaderParams;
-
-		bool FirstLayer = true;
-		UTextureRenderTarget2D* CombinedHeightmapAtlasRT = HeightmapRTList[(int32)EHeightmapRTType::HeightmapRT_CombinedAtlas];
-		UTextureRenderTarget2D* CombinedHeightmapNonAtlasRT = HeightmapRTList[(int32)EHeightmapRTType::HeightmapRT_CombinedNonAtlas];
-		UTextureRenderTarget2D* LandscapeScratchRT1 = HeightmapRTList[(int32)EHeightmapRTType::HeightmapRT_Scratch1];
-		UTextureRenderTarget2D* LandscapeScratchRT2 = HeightmapRTList[(int32)EHeightmapRTType::HeightmapRT_Scratch2];
-		UTextureRenderTarget2D* LandscapeScratchRT3 = HeightmapRTList[(int32)EHeightmapRTType::HeightmapRT_Scratch3];
-
-		for (FLandscapeLayer& Layer : LandscapeLayers)
-		{
-			//Draw Layer heightmap to Combined RT Atlas
-			ShaderParams.ApplyLayerModifiers = false;
-			ShaderParams.LayerVisible = Layer.bVisible;
-			ShaderParams.GenerateNormals = false;
-			ShaderParams.LayerBlendMode = Layer.BlendMode;
-
-			if (Layer.BlendMode == LSBM_AlphaBlend)
-			{
-				// For now, only Layer reserved for Landscape Splines will use the AlphaBlendMode
-				const FLandscapeLayer* SplinesReservedLayer = GetLandscapeSplinesReservedLayer();
-				check(&Layer == SplinesReservedLayer);
-				ShaderParams.LayerAlpha = 1.0f;
-			}
-			else
-			{
-				check(Layer.BlendMode == LSBM_AdditiveBlend);
-				ShaderParams.LayerAlpha = Layer.HeightmapAlpha;
-			}
-
-			{
-				TArray<FLandscapeLayersCopyTextureParams> DeferredCopyTextures;
-				TArray<FHeightmapCopyInfo> LayerHeightmaps;
-				GetUniqueHeightmaps(InLandscapeComponentsToRender, LayerHeightmaps, LandscapeExtent.Min, Layer.Guid);
-				for (const FHeightmapCopyInfo& LayerHeightmap : LayerHeightmaps)
-				{
-					FLandscapeLayersCopyTextureParams& CopyTextureParams = DeferredCopyTextures.Add_GetRef(FLandscapeLayersCopyTextureParams(LayerHeightmap.Texture, LandscapeScratchRT1));
-					// Only copy the size that's actually needed : 
-					CopyTextureParams.CopySize = LayerHeightmap.SectionRect.Size();
-					// Copy from the heightmap's top-left corner to the composited texture's position :
-					CopyTextureParams.DestPosition = LayerHeightmap.SectionRect.Min;
-				}
-				ExecuteCopyLayersTexture(MoveTemp(DeferredCopyTextures));
-			}
-
-			// NOTE: From this point on, we always work in non atlas, we'll convert back at the end to atlas only
-			DrawHeightmapComponentsToRenderTarget(FString::Printf(TEXT("LS Height: %s += -> NonAtlas %s"), *Layer.Name.ToString(), *LandscapeScratchRT1->GetName(), *LandscapeScratchRT2->GetName()),
-												  InLandscapeComponentsToRender, LandscapeExtent.Min, LandscapeScratchRT1, nullptr, LandscapeScratchRT2, ERTDrawingType::RTAtlasToNonAtlas, true, ShaderParams);
-
-			ShaderParams.ApplyLayerModifiers = true;
-
-			// Combine Current layer with current result
-			DrawHeightmapComponentsToRenderTarget(FString::Printf(TEXT("LS Height: %s += -> CombinedNonAtlas %s"), *Layer.Name.ToString(), *LandscapeScratchRT2->GetName(), *CombinedHeightmapNonAtlasRT->GetName()),
-												InLandscapeComponentsToRender, LandscapeExtent.Min, LandscapeScratchRT2, FirstLayer ? nullptr : LandscapeScratchRT3, CombinedHeightmapNonAtlasRT, ERTDrawingType::RTNonAtlas, FirstLayer, ShaderParams);
-
-			ShaderParams.ApplyLayerModifiers = false;
-
-			if (Layer.bVisible && !bSkipBrush)
-			{
-				// Draw each brushes
-				for (int32 i = 0; i < Layer.Brushes.Num(); ++i)
-				{
-					// TODO: handle conversion from float to RG8 by using material params to write correct values
-					// TODO: handle conversion/handling of RT not same size as internal size
-
-					FLandscapeLayerBrush& Brush = Layer.Brushes[i];
-					UTextureRenderTarget2D* BrushOutputNonAtlasRT = Brush.Render(true, LandscapeExtent, CombinedHeightmapNonAtlasRT);
-					if (BrushOutputNonAtlasRT == nullptr || BrushOutputNonAtlasRT->SizeX != CombinedHeightmapNonAtlasRT->SizeX || BrushOutputNonAtlasRT->SizeY != CombinedHeightmapNonAtlasRT->SizeY)
-					{
-						continue;
-					}
-
-					ALandscapeBlueprintBrushBase* LandscapeBrush = Brush.GetBrush();
-					check(LandscapeBrush != nullptr); // If we managed to render, the brush should be valid
-
-					INC_DWORD_STAT(STAT_LandscapeLayersRegenerateDrawCalls); // Brush Render
-
-					PrintLayersDebugRT(FString::Printf(TEXT("LS Height: %s %s -> BrushNonAtlas %s"), *Layer.Name.ToString(), *LandscapeBrush->GetName(), *BrushOutputNonAtlasRT->GetName()), BrushOutputNonAtlasRT);
-
-					// Resolve back to Combined heightmap
-					ExecuteCopyLayersTexture({ FLandscapeLayersCopyTextureParams(BrushOutputNonAtlasRT, CombinedHeightmapNonAtlasRT) });
-					PrintLayersDebugRT(FString::Printf(TEXT("LS Height: %s Component %s += -> CombinedNonAtlas %s"), *Layer.Name.ToString(), *BrushOutputNonAtlasRT->GetName(), *CombinedHeightmapNonAtlasRT->GetName()), CombinedHeightmapNonAtlasRT);
-				}
-			}
-
-			ExecuteCopyLayersTexture({ FLandscapeLayersCopyTextureParams(CombinedHeightmapNonAtlasRT, LandscapeScratchRT3) });
-			PrintLayersDebugRT(FString::Printf(TEXT("LS Height: %s Component %s += -> CombinedNonAtlas %s"), *Layer.Name.ToString(), *CombinedHeightmapNonAtlasRT->GetName(), *LandscapeScratchRT3->GetName()), LandscapeScratchRT3);
-
-			FirstLayer = false;
-		}
-
-		ShaderParams.GenerateNormals = true;
-		ShaderParams.GridSize = GetRootComponent()->GetRelativeScale3D();
-
-		// Broadcast Event of the Full Render
-		if ((HeightmapUpdateModes & Update_Heightmap_All) == Update_Heightmap_All)
-		{
-			LandscapeFullHeightmapRenderDoneDelegate.Broadcast(LandscapeScratchRT3);
-		}
-
-		DrawHeightmapComponentsToRenderTarget(FString::Printf(TEXT("LS Height: %s = -> CombinedNonAtlasNormals : %s"), *CombinedHeightmapNonAtlasRT->GetName(), *LandscapeScratchRT1->GetName()),
-											  InLandscapeComponentsToRender, LandscapeExtent.Min, CombinedHeightmapNonAtlasRT, nullptr, LandscapeScratchRT1, ERTDrawingType::RTNonAtlas, true, ShaderParams);
-
-		ShaderParams.GenerateNormals = false;
-
-		DrawHeightmapComponentsToRenderTarget(FString::Printf(TEXT("LS Height: %s = -> CombinedAtlasFinal : %s"), *LandscapeScratchRT1->GetName(), *CombinedHeightmapAtlasRT->GetName()),
-											  InLandscapeComponentsToRender, LandscapeExtent.Min, LandscapeScratchRT1, nullptr, CombinedHeightmapAtlasRT, ERTDrawingType::RTNonAtlasToAtlas, true, ShaderParams);
-
-		DrawHeightmapComponentsToRenderTargetMips(InLandscapeComponentsToRender, LandscapeExtent.Min, CombinedHeightmapAtlasRT, true, ShaderParams);
-
-		// Copy back all Mips to original heightmap data
-		{
-			TArray<FLandscapeLayersCopyTextureParams> DeferredCopyTextures;
-			TArray<FLandscapeLayersCopyReadbackTextureParams> DeferredCopyReadbackTextures;
-			TArray<FHeightmapCopyInfo> Heightmaps;
-			GetUniqueHeightmaps(InLandscapeComponentsToResolve, Heightmaps, LandscapeExtent.Min);
-			for (const FHeightmapCopyInfo& Heightmap : Heightmaps)
-			{
-				check(Heightmap.CPUReadback);
-
-				const FIntPoint Mip0CopySize = Heightmap.SectionRect.Size();
-				const FIntPoint Mip0SourcePosition = Heightmap.SectionRect.Min;
-				// Mip 0
-				{
-					FLandscapeLayersCopyTextureParams& CopyTextureParams = DeferredCopyTextures.Add_GetRef(FLandscapeLayersCopyTextureParams(CombinedHeightmapAtlasRT, Heightmap.Texture));
-					// Only copy the size that's actually needed : 
-					CopyTextureParams.CopySize = Mip0CopySize;
-					// Copy from the composited texture's position to the top-left corner of the heightmap
-					CopyTextureParams.SourcePosition = Mip0SourcePosition;
-				}
-
-				// Other Mips
-				uint8 MipIndex = 1;
-				for (int32 MipRTIndex = (int32)EHeightmapRTType::HeightmapRT_Mip1; MipRTIndex < (int32)EHeightmapRTType::HeightmapRT_Count; ++MipRTIndex)
-				{
-					UTextureRenderTarget2D* RenderTargetMip = HeightmapRTList[MipRTIndex];
-					if (RenderTargetMip != nullptr)
-					{
-						FLandscapeLayersCopyTextureParams& CopyTextureParams = DeferredCopyTextures.Add_GetRef(FLandscapeLayersCopyTextureParams(RenderTargetMip, Heightmap.Texture));
-						CopyTextureParams.CopySize.X = Mip0CopySize.X >> MipIndex;
-						CopyTextureParams.CopySize.Y = Mip0CopySize.Y >> MipIndex;
-						CopyTextureParams.SourcePosition.X = Mip0SourcePosition.X >> MipIndex;
-						CopyTextureParams.SourcePosition.Y = Mip0SourcePosition.Y >> MipIndex;
-						CopyTextureParams.DestMip = MipIndex;
-
-						++MipIndex;
-					}
-				}
-
-				// Init the CPU read back parameters. This includes a context containing the current component states.
-				const int32 ReadbackIndex = DeferredCopyReadbackTextures.Add(FLandscapeLayersCopyReadbackTextureParams(Heightmap.Texture, Heightmap.CPUReadback));
-				for (ULandscapeComponent* Component : *MapHelper.HeightmapToComponents.Find(Heightmap.Texture))
-				{
-					const FIntPoint ComponentKey = Component->GetSectionBase() / Component->ComponentSizeQuads;
-					const int32 ComponentFlags = Component->GetLayerUpdateFlagPerMode();
-					DeferredCopyReadbackTextures[ReadbackIndex].Context.Add(FLandscapeEditLayerReadback::FComponentReadbackContext(ComponentKey, ComponentFlags));
-				}
-			}
-			ExecuteCopyLayersTexture(MoveTemp(DeferredCopyTextures));
-			ExecuteCopyToReadbackTexture(DeferredCopyReadbackTextures);
-		}
-
-		RenderCaptureLayersNextHeightmapDraws = 0;
 	}
 
 	return HeightmapUpdateModes;
@@ -3457,7 +4883,7 @@
 	//  - If ELandscapeComponentUpdateFlag::Component_Update_Heightmap_Collision is passed, it will be done immediately
 	//  - If not, at least the component's collision data will still get updated eventually, when the flag is finally passed :
 	if (InReadbackResult.bModified)
-{
+	{
 		InComponent->SetPendingCollisionDataUpdate(true);
 	}
 
@@ -3471,8 +4897,8 @@
 			InComponent->UpdateCachedBounds();
 			InComponent->UpdateComponentToWorld();
 
-		// Avoid updating height field if we are going to recreate collision in this update
-		bool bUpdateHeightfieldRegion = !IsUpdateFlagEnabledForModes(ELandscapeComponentUpdateFlag::Component_Update_Recreate_Collision, HeightUpdateMode);
+			// Avoid updating height field if we are going to recreate collision in this update
+			bool bUpdateHeightfieldRegion = !IsUpdateFlagEnabledForModes(ELandscapeComponentUpdateFlag::Component_Update_Recreate_Collision, HeightUpdateMode);
 			InComponent->UpdateCollisionData(bUpdateHeightfieldRegion);
 			InComponent->SetPendingCollisionDataUpdate(false);
 		}
@@ -3549,7 +4975,7 @@
 	FMemory::Memzero(DirtyData.Get(), DirtyDataSize);
 	LandscapeEdit.SetDirtyData(X1, Y1, X2, Y2, DirtyData.Get(), 0);
 }
-	
+
 void ALandscape::UpdateWeightDirtyData(ULandscapeComponent* InLandscapeComponent, UTexture2D const* InWeightmap, FColor const* InOldData, FColor const* InNewData, uint8 InChannel)
 {
 	check(InOldData && InNewData);
@@ -3566,7 +4992,7 @@
 	const int32 SizeV = InWeightmap->Source.GetSizeY();
 	check(DirtyDataSize == SizeU * SizeV);
 
-	const uint8 DirtyWeight = 1<<1;
+	const uint8 DirtyWeight = 1 << 1;
 	LandscapeEdit.GetDirtyData(X1, Y1, X2, Y2, DirtyData.Get(), 0);
 
 	for (int32 Index = 0; Index < DirtyDataSize; ++Index)
@@ -3632,16 +5058,16 @@
 	const int32 X2 = X1 + ComponentSizeQuads;
 	const int32 Y1 = InLandscapeComponent->GetSectionBase().Y;
 	const int32 Y2 = Y1 + ComponentSizeQuads;
-	const int32 ComponentWidth = (SubsectionSizeQuads + 1)*NumSubsections;
+	const int32 ComponentWidth = (SubsectionSizeQuads + 1) * NumSubsections;
 	const int32 DirtyDataSize = ComponentWidth * ComponentWidth;
 	TUniquePtr<uint8[]> DirtyData = MakeUnique<uint8[]>(DirtyDataSize);
 	const int32 SizeU = InHeightmap->Source.GetSizeX();
 	const int32 SizeV = InHeightmap->Source.GetSizeY();
 	const int32 HeightmapOffsetX = InLandscapeComponent->HeightmapScaleBias.Z * SizeU;
 	const int32 HeightmapOffsetY = InLandscapeComponent->HeightmapScaleBias.W * SizeV;
-	const uint8 DirtyHeight = 1<<0;
+	const uint8 DirtyHeight = 1 << 0;
 	LandscapeEdit.GetDirtyData(X1, Y1, X2, Y2, DirtyData.Get(), 0);
-	
+
 	for (int32 X = 0; X < ComponentWidth; ++X)
 	{
 		for (int32 Y = 0; Y < ComponentWidth; ++Y)
@@ -3649,7 +5075,7 @@
 			int32 TexX = HeightmapOffsetX + X;
 			int32 TexY = HeightmapOffsetY + Y;
 			int32 TexIndex = TexX + TexY * SizeU;
-			check(TexIndex < SizeU * SizeV);
+			check(TexIndex < SizeU* SizeV);
 			if (InOldData[TexIndex] != InNewData[TexIndex])
 			{
 				DirtyData[X + Y * ComponentWidth] |= DirtyHeight;
@@ -3781,7 +5207,376 @@
 }
 
 void ALandscape::PrepareComponentDataToExtractMaterialLayersCS(const TArray<ULandscapeComponent*>& InLandscapeComponents, const FLandscapeLayer& InLayer, int32 InCurrentWeightmapToProcessIndex, const FIntPoint& InLandscapeBase, FLandscapeTexture2DResource* InOutTextureData,
-																TArray<FLandscapeLayerWeightmapExtractMaterialLayersComponentData>& OutComponentData, TMap<ULandscapeLayerInfoObject*, int32>& OutLayerInfoObjects)
+	TArray<FLandscapeLayerWeightmapExtractMaterialLayersComponentData>& OutComponentData, TMap<ULandscapeLayerInfoObject*, int32>& OutLayerInfoObjects)
 {
 	ULandscapeInfo* Info = GetLandscapeInfo();
 
+	if (Info == nullptr)
+	{
+		return;
+	}
+
+	TArray<FLandscapeLayersCopyTextureParams> DeferredCopyTextures;
+
+	const int32 LocalComponentSizeQuad = SubsectionSizeQuads * NumSubsections;
+	const int32 LocalComponentSizeVerts = (SubsectionSizeQuads + 1) * NumSubsections;
+	for (ULandscapeComponent* Component : InLandscapeComponents)
+	{
+		FLandscapeLayerComponentData* ComponentLayerData = Component->GetLayerData(InLayer.Guid);
+
+		if (ComponentLayerData != nullptr)
+		{
+			if (ComponentLayerData->WeightmapData.Textures.IsValidIndex(InCurrentWeightmapToProcessIndex) && ComponentLayerData->WeightmapData.TextureUsages.IsValidIndex(InCurrentWeightmapToProcessIndex))
+			{
+				UTexture2D* LayerWeightmap = ComponentLayerData->WeightmapData.Textures[InCurrentWeightmapToProcessIndex];
+				check(LayerWeightmap != nullptr);
+
+				const ULandscapeWeightmapUsage* LayerWeightmapUsage = ComponentLayerData->WeightmapData.TextureUsages[InCurrentWeightmapToProcessIndex];
+				check(LayerWeightmapUsage != nullptr);
+
+				FIntPoint ComponentSectionBase = Component->GetSectionBase() - InLandscapeBase;
+				FVector2D SourcePositionOffset(FMath::RoundToInt((float)ComponentSectionBase.X / LocalComponentSizeQuad), FMath::RoundToInt((float)ComponentSectionBase.Y / LocalComponentSizeQuad));
+				FIntPoint SourceComponentVertexPosition = FIntPoint(SourcePositionOffset.X * LocalComponentSizeVerts, SourcePositionOffset.Y * LocalComponentSizeVerts);
+
+				FLandscapeLayersCopyTextureParams& CopyTextureParams = DeferredCopyTextures.Add_GetRef(FLandscapeLayersCopyTextureParams(*LayerWeightmap->GetName(), LayerWeightmap->GetResource(), FString::Printf(TEXT("%s WeightmapScratchTexture"), *InLayer.Name.ToString()), InOutTextureData));
+				// Only copy the size that's actually needed : 
+				CopyTextureParams.CopySize.X = LayerWeightmap->GetResource()->GetSizeX();
+				CopyTextureParams.CopySize.Y = LayerWeightmap->GetResource()->GetSizeY();
+				// Copy from the top-left corner of the weightmap to the composited texture's position
+				CopyTextureParams.DestPosition = SourceComponentVertexPosition;
+				PrintLayersDebugTextureResource(FString::Printf(TEXT("LS Weight: %s WeightmapScratchTexture %s"), *InLayer.Name.ToString(), TEXT("WeightmapScratchTextureResource")), InOutTextureData, 0, false);
+
+				for (const FWeightmapLayerAllocationInfo& WeightmapLayerAllocation : ComponentLayerData->WeightmapData.LayerAllocations)
+				{
+					if (WeightmapLayerAllocation.LayerInfo != nullptr && WeightmapLayerAllocation.IsAllocated() && ComponentLayerData->WeightmapData.Textures[WeightmapLayerAllocation.WeightmapTextureIndex] == LayerWeightmap)
+					{
+						FLandscapeLayerWeightmapExtractMaterialLayersComponentData Data;
+
+						const ULandscapeComponent* DestComponent = LayerWeightmapUsage->ChannelUsage[WeightmapLayerAllocation.WeightmapTextureChannel];
+						check(DestComponent);
+
+						FIntPoint DestComponentSectionBase = DestComponent->GetSectionBase() - InLandscapeBase;
+
+						// Compute component top left vertex position from section base info
+						FVector2D DestPositionOffset(FMath::RoundToInt((float)DestComponentSectionBase.X / LocalComponentSizeQuad), FMath::RoundToInt((float)DestComponentSectionBase.Y / LocalComponentSizeQuad));
+
+						Data.ComponentVertexPosition = SourceComponentVertexPosition;
+						Data.AtlasTexturePositionOutput = FIntPoint(DestPositionOffset.X * LocalComponentSizeVerts, DestPositionOffset.Y * LocalComponentSizeVerts);
+						Data.WeightmapChannelToProcess = WeightmapLayerAllocation.WeightmapTextureChannel;
+
+						if (WeightmapLayerAllocation.LayerInfo == ALandscapeProxy::VisibilityLayer)
+						{
+							Data.DestinationPaintLayerIndex = 0;
+							int32& NewLayerinfoObjectIndex = OutLayerInfoObjects.FindOrAdd(ALandscapeProxy::VisibilityLayer);
+							NewLayerinfoObjectIndex = 0;
+						}
+						else
+						{
+							for (int32 LayerInfoSettingsIndex = 0; LayerInfoSettingsIndex < Info->Layers.Num(); ++LayerInfoSettingsIndex)
+							{
+								const FLandscapeInfoLayerSettings& InfoLayerSettings = Info->Layers[LayerInfoSettingsIndex];
+
+								if (InfoLayerSettings.LayerInfoObj == WeightmapLayerAllocation.LayerInfo)
+								{
+									Data.DestinationPaintLayerIndex = LayerInfoSettingsIndex + 1; // due to visibility layer that is at 0
+									int32& NewLayerinfoObjectIndex = OutLayerInfoObjects.FindOrAdd(WeightmapLayerAllocation.LayerInfo);
+									NewLayerinfoObjectIndex = LayerInfoSettingsIndex + 1;
+
+									break;
+								}
+							}
+						}
+
+						OutComponentData.Add(Data);
+					}
+				}
+			}
+		}
+	}
+
+	ExecuteCopyLayersTexture(MoveTemp(DeferredCopyTextures));
+}
+
+void ALandscape::PrepareComponentDataToPackMaterialLayersCS(int32 InCurrentWeightmapToProcessIndex, const FIntPoint& InLandscapeBase, const TArray<ULandscapeComponent*>& InAllLandscapeComponents, TArray<UTexture2D*>& OutProcessedWeightmaps,
+	TArray<FLandscapeEditLayerReadback*>& OutProcessedCPUReadbacks, TArray<FLandscapeLayerWeightmapPackMaterialLayersComponentData>& OutComponentData)
+{
+	ULandscapeInfo* Info = GetLandscapeInfo();
+
+	if (Info == nullptr)
+	{
+		return;
+	}
+
+	// Compute a mapping of all textures for the asked index and their usage
+	TMap<UTexture2D*, ULandscapeWeightmapUsage*> WeightmapsToProcess;
+
+	for (ULandscapeComponent* Component : InAllLandscapeComponents)
+	{
+		const TArray<UTexture2D*>& ComponentWeightmapTextures = Component->GetWeightmapTextures();
+		const TArray<ULandscapeWeightmapUsage*>& ComponentWeightmapTexturesUsage = Component->GetWeightmapTexturesUsage();
+
+		if (ComponentWeightmapTextures.IsValidIndex(InCurrentWeightmapToProcessIndex))
+		{
+			UTexture2D* ComponentWeightmapTexture = ComponentWeightmapTextures[InCurrentWeightmapToProcessIndex];
+			ULandscapeWeightmapUsage* ComponentWeightmapTextureUsage = ComponentWeightmapTexturesUsage[InCurrentWeightmapToProcessIndex];
+			check(ComponentWeightmapTextureUsage != nullptr);
+
+			// It's possible the texture (and its usage, hence) has already been processed by a previous call to PrepareComponentDataToPackMaterialLayersCS with a different InCurrentWeightmapToProcessIndex
+			//  since a texture can be shared by multiple components :
+			if (OutProcessedWeightmaps.Contains(ComponentWeightmapTexture) && (WeightmapsToProcess.Find(ComponentWeightmapTexture) == nullptr))
+			{
+				WeightmapsToProcess.Add(ComponentWeightmapTexture, ComponentWeightmapTextureUsage);
+				OutProcessedWeightmaps.Add(ComponentWeightmapTexture);
+
+				FLandscapeEditLayerReadback** CPUReadback = Component->GetLandscapeProxy()->WeightmapsCPUReadback.Find(ComponentWeightmapTexture);
+				check(CPUReadback != nullptr);
+
+				OutProcessedCPUReadbacks.Add(*CPUReadback);
+			}
+		}
+	}
+
+	TArray<const FWeightmapLayerAllocationInfo*> AlreadyProcessedAllocation;
+
+	// Build for each texture what each channel should contains
+	for (auto& ItPair : WeightmapsToProcess)
+	{
+		FLandscapeLayerWeightmapPackMaterialLayersComponentData Data;
+
+		for (int32 WeightmapChannelIndex = 0; WeightmapChannelIndex < ULandscapeWeightmapUsage::NumChannels; ++WeightmapChannelIndex)
+		{
+			UTexture2D* ComponentWeightmapTexture = ItPair.Key;
+			ULandscapeWeightmapUsage* ComponentWeightmapTextureUsage = ItPair.Value;
+
+			// Clear out data to known values
+			Data.ComponentVertexPositionX[WeightmapChannelIndex] = INDEX_NONE;
+			Data.ComponentVertexPositionY[WeightmapChannelIndex] = INDEX_NONE;
+			Data.SourcePaintLayerIndex[WeightmapChannelIndex] = INDEX_NONE;
+			Data.WeightmapChannelToProcess[WeightmapChannelIndex] = INDEX_NONE;
+
+			if (ComponentWeightmapTextureUsage->ChannelUsage[WeightmapChannelIndex] != nullptr)
+			{
+				const ULandscapeComponent* ChannelComponent = ComponentWeightmapTextureUsage->ChannelUsage[WeightmapChannelIndex];
+
+				const TArray<FWeightmapLayerAllocationInfo>& ChannelLayerAllocations = ChannelComponent->GetWeightmapLayerAllocations();
+				const TArray<UTexture2D*>& ChannelComponentWeightmapTextures = ChannelComponent->GetWeightmapTextures();
+
+				for (const FWeightmapLayerAllocationInfo& ChannelLayerAllocation : ChannelLayerAllocations)
+				{
+					if (ChannelLayerAllocation.LayerInfo != nullptr && !AlreadyProcessedAllocation.Contains(&ChannelLayerAllocation) && ChannelComponentWeightmapTextures[ChannelLayerAllocation.WeightmapTextureIndex] == ComponentWeightmapTexture)
+					{
+						FIntPoint ComponentSectionBase = ChannelComponent->GetSectionBase() - InLandscapeBase;
+
+						// Compute component top left vertex position from section base info
+						int32 LocalComponentSizeQuad = ChannelComponent->SubsectionSizeQuads * NumSubsections;
+						int32 LocalComponentSizeVerts = (ChannelComponent->SubsectionSizeQuads + 1) * NumSubsections;
+						FVector2D PositionOffset(FMath::RoundToInt((float)ComponentSectionBase.X / LocalComponentSizeQuad), FMath::RoundToInt((float)ComponentSectionBase.Y / LocalComponentSizeQuad));
+
+						Data.ComponentVertexPositionX[WeightmapChannelIndex] = PositionOffset.X * LocalComponentSizeVerts;
+						Data.ComponentVertexPositionY[WeightmapChannelIndex] = PositionOffset.Y * LocalComponentSizeVerts;
+
+						Data.WeightmapChannelToProcess[WeightmapChannelIndex] = ChannelLayerAllocation.WeightmapTextureChannel;
+
+						AlreadyProcessedAllocation.Add(&ChannelLayerAllocation);
+
+						if (ChannelLayerAllocation.LayerInfo == ALandscapeProxy::VisibilityLayer)
+						{
+							Data.SourcePaintLayerIndex[WeightmapChannelIndex] = 0; // Always store after the last weightmap index
+						}
+						else
+						{
+							for (int32 LayerInfoSettingsIndex = 0; LayerInfoSettingsIndex < Info->Layers.Num(); ++LayerInfoSettingsIndex)
+							{
+								const FLandscapeInfoLayerSettings& LayerInfo = Info->Layers[LayerInfoSettingsIndex];
+
+								if (ChannelLayerAllocation.LayerInfo == LayerInfo.LayerInfoObj)
+								{
+									Data.SourcePaintLayerIndex[WeightmapChannelIndex] = LayerInfoSettingsIndex + 1; // due to visibility layer that is at 0
+									break;
+								}
+							}
+						}
+
+						break;
+					}
+				}
+			}
+		}
+
+		OutComponentData.Add(Data);
+	}
+}
+
+void ALandscape::ReallocateLayersWeightmaps(FUpdateLayersContentContext& InUpdateLayersContentContext, const TArray<ULandscapeLayerInfoObject*>& InBrushRequiredAllocations)
+{
+	TRACE_CPUPROFILER_EVENT_SCOPE(LandscapeLayers_ReallocateLayersWeightmaps);
+	SCOPE_CYCLE_COUNTER(STAT_LandscapeLayersReallocateWeightmaps);
+
+	ULandscapeInfo* Info = GetLandscapeInfo();
+
+	if (Info == nullptr)
+	{
+		return;
+	}
+
+	// Clear allocation data
+	for (ULandscapeComponent* Component : InUpdateLayersContentContext.LandscapeComponentsWeightmapsToResolve)
+	{
+		TArray<FWeightmapLayerAllocationInfo>& BaseLayerAllocations = Component->GetWeightmapLayerAllocations();
+		for (FWeightmapLayerAllocationInfo& BaseWeightmapAllocation : BaseLayerAllocations)
+		{
+			BaseWeightmapAllocation.Free();
+		}
+
+		TArray<ULandscapeWeightmapUsage*>& WeightmapTexturesUsage = Component->GetWeightmapTexturesUsage();
+		for (int32 i = 0; i < WeightmapTexturesUsage.Num(); ++i)
+		{
+			ULandscapeWeightmapUsage* Usage = WeightmapTexturesUsage[i];
+			check(Usage != nullptr);
+
+			Usage->ClearUsage(Component);
+		}
+	}
+
+	// Build a map of all the allocation per components
+	TMap<ULandscapeComponent*, TArray<ULandscapeLayerInfoObject*>> LayerAllocsPerComponent;
+	for (FLandscapeLayer& Layer : LandscapeLayers)
+	{
+		for (ULandscapeComponent* Component : InUpdateLayersContentContext.LandscapeComponentsWeightmapsToResolve)
+		{
+			TArray<ULandscapeLayerInfoObject*>* ComponentLayerAlloc = LayerAllocsPerComponent.Find(Component);
+			if (ComponentLayerAlloc == nullptr)
+			{
+				TArray<ULandscapeLayerInfoObject*> NewLayerAllocs;
+				ComponentLayerAlloc = &LayerAllocsPerComponent.Add(Component, NewLayerAllocs);
+			}
+
+			// No need for an allocation if the edit layer is invisible : 
+			if (Layer.bVisible)
+			{
+				FLandscapeLayerComponentData* LayerComponentData = Component->GetLayerData(Layer.Guid);
+				for (const FWeightmapLayerAllocationInfo& LayerWeightmapAllocation : LayerComponentData->WeightmapData.LayerAllocations)
+				{
+					if (LayerWeightmapAllocation.LayerInfo != nullptr)
+					{
+						ComponentLayerAlloc->AddUnique(LayerWeightmapAllocation.LayerInfo);
+					}
+				}
+			}
+
+			// Add the brush alloc also (only if !InMergeParams.bSkipBrush, but InBrushRequiredAllocations should be empty already if InMergeParams.bSkipBrush is true) :
+			for (ULandscapeLayerInfoObject* BrushLayerInfo : InBrushRequiredAllocations)
+			{
+				if (BrushLayerInfo != nullptr)
+				{
+					ComponentLayerAlloc->AddUnique(BrushLayerInfo);
+				}
+			}
+		}
+	}
+
+	// Trim the components that don't need weightmaps anymore (e.g. all edit layers are made invisible : there were some components LandscapeComponentsWeightmapsToResolve but there aren't anymore now):
+	InUpdateLayersContentContext.LandscapeComponentsWeightmapsToResolve.RemoveAll([&](ULandscapeComponent* Component)
+	{
+		TArray<ULandscapeLayerInfoObject*>* ComponentLayerAlloc = LayerAllocsPerComponent.Find(Component);
+		check(ComponentLayerAlloc != nullptr);
+		return ComponentLayerAlloc->IsEmpty();
+	});
+
+	// Determine if the Final layer need to add/remove some alloc
+	for (auto& ItPair : LayerAllocsPerComponent)
+	{
+		ULandscapeComponent* Component = ItPair.Key;
+		TArray<ULandscapeLayerInfoObject*>& ComponentLayerAlloc = ItPair.Value;
+		TArray<FWeightmapLayerAllocationInfo>& ComponentBaseLayerAlloc = Component->GetWeightmapLayerAllocations();
+
+		// Deal with the one that need removal
+		for (int32 i = ComponentBaseLayerAlloc.Num() - 1; i >= 0; --i)
+		{
+			const FWeightmapLayerAllocationInfo& Alloc = ComponentBaseLayerAlloc[i];
+
+			if (!ComponentLayerAlloc.Contains(Alloc.LayerInfo))
+			{
+				ComponentBaseLayerAlloc.RemoveAt(i);
+			}
+		}
+
+		// Then add the new one
+		for (ULandscapeLayerInfoObject* LayerAlloc : ComponentLayerAlloc)
+		{
+			const bool AllocExist = ComponentBaseLayerAlloc.ContainsByPredicate([&LayerAlloc](const FWeightmapLayerAllocationInfo& BaseLayerAlloc) { return (LayerAlloc == BaseLayerAlloc.LayerInfo); });
+
+			if (!AllocExist)
+			{
+				ComponentBaseLayerAlloc.Add(FWeightmapLayerAllocationInfo(LayerAlloc));
+			}
+		}
+	}
+
+	// Realloc the weightmap so it will create proper texture (if needed) and will set the allocations information
+	TArray<UTexture*> NewCreatedTextures;
+	for (ULandscapeComponent* Component : InUpdateLayersContentContext.LandscapeComponentsWeightmapsToResolve)
+	{
+		Component->ReallocateWeightmaps(nullptr, false, false, false, nullptr, &NewCreatedTextures);
+	}
+
+	// TODO: correctly only recreate what is required instead of everything..
+	//GDisableAutomaticTextureMaterialUpdateDependencies = true;
+
+	FTextureCompilingManager::Get().FinishCompilation(NewCreatedTextures);
+	for (UTexture* Texture : NewCreatedTextures)
+	{
+		Texture->bForceMiplevelsToBeResident = true;
+		Texture->WaitForStreaming();
+	}
+
+	//GDisableAutomaticTextureMaterialUpdateDependencies = false;
+
+	// Clean-up unused weightmap CPUReadback resources
+	Info->ForAllLandscapeProxies([](ALandscapeProxy* Proxy)
+	{
+		TArray<UTexture2D*, TInlineAllocator<64>> EntriesToRemoveFromMap;
+		for (auto& Pair : Proxy->WeightmapsCPUReadback)
+		{
+			UTexture2D* WeightmapTextureKey = Pair.Key;
+			bool IsTextureReferenced = false;
+			for (ULandscapeComponent* Component : Proxy->LandscapeComponents)
+			{
+				for (UTexture2D* WeightmapTexture : Component->GetWeightmapTextures(false))
+				{
+					if (WeightmapTexture == WeightmapTextureKey)
+					{
+						IsTextureReferenced = true;
+						break;
+					}
+				}
+			}
+			if (!IsTextureReferenced)
+			{
+				EntriesToRemoveFromMap.Add(WeightmapTextureKey);
+			}
+		}
+
+		if (EntriesToRemoveFromMap.Num())
+		{
+			for (UTexture2D* OldWeightmapTexture : EntriesToRemoveFromMap)
+			{
+				if (FLandscapeEditLayerReadback** CPUReadbackToDelete = Proxy->WeightmapsCPUReadback.Find(OldWeightmapTexture))
+				{
+					check(*CPUReadbackToDelete);
+					delete* CPUReadbackToDelete;
+					Proxy->WeightmapsCPUReadback.Remove(OldWeightmapTexture);
+				}
+			}
+		}
+	});
+
+	ValidateProxyLayersWeightmapUsage();
+
+	InUpdateLayersContentContext.Refresh(FUpdateLayersContentContext::ERefreshFlags::RefreshWeightmapInfos | FUpdateLayersContentContext::ERefreshFlags::RefreshMapHelper);
+}
+
+void ALandscape::InitializeLayersWeightmapResources()
+{
+	ULandscapeInfo* Info = GetLandscapeInfo();
+
@@ -3790,363 +5585,6 @@
 		return;
 	}
 
-	TArray<FLandscapeLayersCopyTextureParams> DeferredCopyTextures;
-
-	const int32 LocalComponentSizeQuad = SubsectionSizeQuads * NumSubsections;
-	const int32 LocalComponentSizeVerts = (SubsectionSizeQuads + 1) * NumSubsections;
-	for (ULandscapeComponent* Component : InLandscapeComponents)
-	{
-		FLandscapeLayerComponentData* ComponentLayerData = Component->GetLayerData(InLayer.Guid);
-
-		if (ComponentLayerData != nullptr)
-		{
-			if (ComponentLayerData->WeightmapData.Textures.IsValidIndex(InCurrentWeightmapToProcessIndex) && ComponentLayerData->WeightmapData.TextureUsages.IsValidIndex(InCurrentWeightmapToProcessIndex))
-			{
-				UTexture2D* LayerWeightmap = ComponentLayerData->WeightmapData.Textures[InCurrentWeightmapToProcessIndex];
-				check(LayerWeightmap != nullptr);
-
-				const ULandscapeWeightmapUsage* LayerWeightmapUsage = ComponentLayerData->WeightmapData.TextureUsages[InCurrentWeightmapToProcessIndex];
-				check(LayerWeightmapUsage != nullptr);
-
-				FIntPoint ComponentSectionBase = Component->GetSectionBase() - InLandscapeBase;
-				FVector2D SourcePositionOffset(FMath::RoundToInt((float)ComponentSectionBase.X / LocalComponentSizeQuad), FMath::RoundToInt((float)ComponentSectionBase.Y / LocalComponentSizeQuad));
-				FIntPoint SourceComponentVertexPosition = FIntPoint(SourcePositionOffset.X * LocalComponentSizeVerts, SourcePositionOffset.Y * LocalComponentSizeVerts);
-
-				FLandscapeLayersCopyTextureParams& CopyTextureParams = DeferredCopyTextures.Add_GetRef(FLandscapeLayersCopyTextureParams(*LayerWeightmap->GetName(), LayerWeightmap->GetResource(), FString::Printf(TEXT("%s WeightmapScratchTexture"), *InLayer.Name.ToString()), InOutTextureData));
-				// Only copy the size that's actually needed : 
-				CopyTextureParams.CopySize.X = LayerWeightmap->GetResource()->GetSizeX();
-				CopyTextureParams.CopySize.Y = LayerWeightmap->GetResource()->GetSizeY();
-				// Copy from the top-left corner of the weightmap to the composited texture's position
-				CopyTextureParams.DestPosition = SourceComponentVertexPosition;
-				PrintLayersDebugTextureResource(FString::Printf(TEXT("LS Weight: %s WeightmapScratchTexture %s"), *InLayer.Name.ToString(), TEXT("WeightmapScratchTextureResource")), InOutTextureData, 0, false);
-
-				for (const FWeightmapLayerAllocationInfo& WeightmapLayerAllocation : ComponentLayerData->WeightmapData.LayerAllocations)
-				{
-					if (WeightmapLayerAllocation.LayerInfo != nullptr && WeightmapLayerAllocation.IsAllocated() && ComponentLayerData->WeightmapData.Textures[WeightmapLayerAllocation.WeightmapTextureIndex] == LayerWeightmap)
-					{
-						FLandscapeLayerWeightmapExtractMaterialLayersComponentData Data;
-
-						const ULandscapeComponent* DestComponent = LayerWeightmapUsage->ChannelUsage[WeightmapLayerAllocation.WeightmapTextureChannel];
-						check(DestComponent);
-
-						FIntPoint DestComponentSectionBase = DestComponent->GetSectionBase() - InLandscapeBase;
-
-						// Compute component top left vertex position from section base info
-						FVector2D DestPositionOffset(FMath::RoundToInt((float)DestComponentSectionBase.X / LocalComponentSizeQuad), FMath::RoundToInt((float)DestComponentSectionBase.Y / LocalComponentSizeQuad));
-
-						Data.ComponentVertexPosition = SourceComponentVertexPosition;
-						Data.AtlasTexturePositionOutput = FIntPoint(DestPositionOffset.X * LocalComponentSizeVerts, DestPositionOffset.Y * LocalComponentSizeVerts);
-						Data.WeightmapChannelToProcess = WeightmapLayerAllocation.WeightmapTextureChannel;
-
-						if (WeightmapLayerAllocation.LayerInfo == ALandscapeProxy::VisibilityLayer)
-						{
-							Data.DestinationPaintLayerIndex = 0;
-							int32& NewLayerinfoObjectIndex = OutLayerInfoObjects.FindOrAdd(ALandscapeProxy::VisibilityLayer);
-							NewLayerinfoObjectIndex = 0;
-						}
-						else
-						{
-							for (int32 LayerInfoSettingsIndex = 0; LayerInfoSettingsIndex < Info->Layers.Num(); ++LayerInfoSettingsIndex)
-							{
-								const FLandscapeInfoLayerSettings& InfoLayerSettings = Info->Layers[LayerInfoSettingsIndex];
-
-								if (InfoLayerSettings.LayerInfoObj == WeightmapLayerAllocation.LayerInfo)
-								{
-									Data.DestinationPaintLayerIndex = LayerInfoSettingsIndex + 1; // due to visibility layer that is at 0
-									int32& NewLayerinfoObjectIndex = OutLayerInfoObjects.FindOrAdd(WeightmapLayerAllocation.LayerInfo);
-									NewLayerinfoObjectIndex = LayerInfoSettingsIndex + 1;
-
-									break;
-								}
-							}
-						}
-
-						OutComponentData.Add(Data);
-					}
-				}
-			}
-		}
-	}
-	
-	ExecuteCopyLayersTexture(MoveTemp(DeferredCopyTextures));
-}
-
-void ALandscape::PrepareComponentDataToPackMaterialLayersCS(int32 InCurrentWeightmapToProcessIndex, const FIntPoint& InLandscapeBase, const TArray<ULandscapeComponent*>& InAllLandscapeComponents, TArray<UTexture2D*>& OutProcessedWeightmaps,
-	TArray<FLandscapeEditLayerReadback*>& OutProcessedCPUReadbacks, TArray<FLandscapeLayerWeightmapPackMaterialLayersComponentData>& OutComponentData)
-{
-	ULandscapeInfo* Info = GetLandscapeInfo();
-
-	if (Info == nullptr)
-	{
-		return;
-	}
-
-	// Compute a mapping of all textures for the asked index and their usage
-	TMap<UTexture2D*, ULandscapeWeightmapUsage*> WeightmapsToProcess;
-
-	for (ULandscapeComponent* Component : InAllLandscapeComponents)
-	{
-		const TArray<UTexture2D*>& ComponentWeightmapTextures = Component->GetWeightmapTextures();
-		const TArray<ULandscapeWeightmapUsage*>& ComponentWeightmapTexturesUsage = Component->GetWeightmapTexturesUsage();
-
-		if (ComponentWeightmapTextures.IsValidIndex(InCurrentWeightmapToProcessIndex))
-		{
-			UTexture2D* ComponentWeightmapTexture = ComponentWeightmapTextures[InCurrentWeightmapToProcessIndex];
-			ULandscapeWeightmapUsage* ComponentWeightmapTextureUsage = ComponentWeightmapTexturesUsage[InCurrentWeightmapToProcessIndex];
-			check(ComponentWeightmapTextureUsage != nullptr);
-
-			if (WeightmapsToProcess.Find(ComponentWeightmapTexture) == nullptr)
-			{
-				WeightmapsToProcess.Add(ComponentWeightmapTexture, ComponentWeightmapTextureUsage);
-				OutProcessedWeightmaps.Add(ComponentWeightmapTexture);
-
-				FLandscapeEditLayerReadback** CPUReadback = Component->GetLandscapeProxy()->WeightmapsCPUReadback.Find(ComponentWeightmapTexture);
-				check(CPUReadback != nullptr);
-
-				OutProcessedCPUReadbacks.Add(*CPUReadback);
-			}
-		}
-	}
-
-	TArray<const FWeightmapLayerAllocationInfo*> AlreadyProcessedAllocation;
-
-	// Build for each texture what each channel should contains
-	for (auto& ItPair : WeightmapsToProcess)
-	{
-		FLandscapeLayerWeightmapPackMaterialLayersComponentData Data;
-
-		for (int32 WeightmapChannelIndex = 0; WeightmapChannelIndex < ULandscapeWeightmapUsage::NumChannels; ++WeightmapChannelIndex)
-		{
-			UTexture2D* ComponentWeightmapTexture = ItPair.Key;
-			ULandscapeWeightmapUsage* ComponentWeightmapTextureUsage = ItPair.Value;
-
-			// Clear out data to known values
-			Data.ComponentVertexPositionX[WeightmapChannelIndex] = INDEX_NONE;
-			Data.ComponentVertexPositionY[WeightmapChannelIndex] = INDEX_NONE;
-			Data.SourcePaintLayerIndex[WeightmapChannelIndex] = INDEX_NONE;
-			Data.WeightmapChannelToProcess[WeightmapChannelIndex] = INDEX_NONE;
-
-			if (ComponentWeightmapTextureUsage->ChannelUsage[WeightmapChannelIndex] != nullptr)
-			{
-				const ULandscapeComponent* ChannelComponent = ComponentWeightmapTextureUsage->ChannelUsage[WeightmapChannelIndex];
-
-				const TArray<FWeightmapLayerAllocationInfo>& ChannelLayerAllocations = ChannelComponent->GetWeightmapLayerAllocations();
-				const TArray<UTexture2D*>& ChannelComponentWeightmapTextures = ChannelComponent->GetWeightmapTextures();
-
-				for (const FWeightmapLayerAllocationInfo& ChannelLayerAllocation : ChannelLayerAllocations)
-				{
-					if (ChannelLayerAllocation.LayerInfo != nullptr && !AlreadyProcessedAllocation.Contains(&ChannelLayerAllocation) && ChannelComponentWeightmapTextures[ChannelLayerAllocation.WeightmapTextureIndex] == ComponentWeightmapTexture)
-					{
-						FIntPoint ComponentSectionBase = ChannelComponent->GetSectionBase() - InLandscapeBase;
-
-						// Compute component top left vertex position from section base info
-						int32 LocalComponentSizeQuad = ChannelComponent->SubsectionSizeQuads * NumSubsections;
-						int32 LocalComponentSizeVerts = (ChannelComponent->SubsectionSizeQuads + 1) * NumSubsections;
-						FVector2D PositionOffset(FMath::RoundToInt((float)ComponentSectionBase.X / LocalComponentSizeQuad), FMath::RoundToInt((float)ComponentSectionBase.Y / LocalComponentSizeQuad));
-
-						Data.ComponentVertexPositionX[WeightmapChannelIndex] = PositionOffset.X * LocalComponentSizeVerts;
-						Data.ComponentVertexPositionY[WeightmapChannelIndex] = PositionOffset.Y * LocalComponentSizeVerts;
-
-						Data.WeightmapChannelToProcess[WeightmapChannelIndex] = ChannelLayerAllocation.WeightmapTextureChannel;
-
-						AlreadyProcessedAllocation.Add(&ChannelLayerAllocation);
-
-						if (ChannelLayerAllocation.LayerInfo == ALandscapeProxy::VisibilityLayer)
-						{
-							Data.SourcePaintLayerIndex[WeightmapChannelIndex] = 0; // Always store after the last weightmap index
-						}
-						else
-						{
-							for (int32 LayerInfoSettingsIndex = 0; LayerInfoSettingsIndex < Info->Layers.Num(); ++LayerInfoSettingsIndex)
-							{
-								const FLandscapeInfoLayerSettings& LayerInfo = Info->Layers[LayerInfoSettingsIndex];
-
-								if (ChannelLayerAllocation.LayerInfo == LayerInfo.LayerInfoObj)
-								{
-									Data.SourcePaintLayerIndex[WeightmapChannelIndex] = LayerInfoSettingsIndex + 1; // due to visibility layer that is at 0
-									break;
-								}
-							}
-						}
-
-						break;
-					}
-				}
-			}
-		}
-
-		OutComponentData.Add(Data);
-	}
-}
-
-void ALandscape::ReallocateLayersWeightmaps(const TArray<ULandscapeComponent*>& InLandscapeComponents, const TArray<ULandscapeLayerInfoObject*>& InBrushRequiredAllocations)
-{
-	TRACE_CPUPROFILER_EVENT_SCOPE(LandscapeLayers_ReallocateLayersWeightmaps);
-	SCOPE_CYCLE_COUNTER(STAT_LandscapeLayersReallocateWeightmaps);
-
-	ULandscapeInfo* Info = GetLandscapeInfo();
-
-	if (Info == nullptr)
-	{
-		return;
-	}
-
-	// Clear allocation data
-	for (ULandscapeComponent* Component : InLandscapeComponents)
-	{
-		TArray<FWeightmapLayerAllocationInfo>& BaseLayerAllocations = Component->GetWeightmapLayerAllocations();
-		for (FWeightmapLayerAllocationInfo& BaseWeightmapAllocation : BaseLayerAllocations)
-		{
-			BaseWeightmapAllocation.Free();
-		}
-
-		TArray<ULandscapeWeightmapUsage*>& WeightmapTexturesUsage = Component->GetWeightmapTexturesUsage();
-		for (int32 i = 0; i < WeightmapTexturesUsage.Num(); ++i)
-		{
-			ULandscapeWeightmapUsage* Usage = WeightmapTexturesUsage[i];
-			check(Usage != nullptr);
-
-			Usage->ClearUsage(Component);
-		}
-	}
-		
-	// Build a map of all the allocation per components
-	TMap<ULandscapeComponent*, TArray<ULandscapeLayerInfoObject*>> LayerAllocsPerComponent;
-
-	for (FLandscapeLayer& Layer : LandscapeLayers)
-	{
-		for (ULandscapeComponent* Component : InLandscapeComponents)
-		{
-			TArray<ULandscapeLayerInfoObject*>* ComponentLayerAlloc = LayerAllocsPerComponent.Find(Component);
-
-			if (ComponentLayerAlloc == nullptr)
-			{
-				TArray<ULandscapeLayerInfoObject*> NewLayerAllocs;
-				ComponentLayerAlloc = &LayerAllocsPerComponent.Add(Component, NewLayerAllocs);
-			}
-
-			FLandscapeLayerComponentData* LayerComponentData = Component->GetLayerData(Layer.Guid);
-
-			for (const FWeightmapLayerAllocationInfo& LayerWeightmapAllocation : LayerComponentData->WeightmapData.LayerAllocations)
-			{
-				if (LayerWeightmapAllocation.LayerInfo != nullptr)
-				{
-					ComponentLayerAlloc->AddUnique(LayerWeightmapAllocation.LayerInfo);
-				}
-			}
-
-			// Add the brush alloc also
-			for (ULandscapeLayerInfoObject* BrushLayerInfo : InBrushRequiredAllocations)
-			{
-				if (BrushLayerInfo != nullptr)
-				{
-					ComponentLayerAlloc->AddUnique(BrushLayerInfo);
-				}
-			}
-		}
-	}
-
-	// Determine if the Final layer need to add/remove some alloc
-	for (auto& ItPair : LayerAllocsPerComponent)
-	{
-		ULandscapeComponent* Component = ItPair.Key;
-		TArray<ULandscapeLayerInfoObject*>& ComponentLayerAlloc = ItPair.Value;
-		TArray<FWeightmapLayerAllocationInfo>& ComponentBaseLayerAlloc = Component->GetWeightmapLayerAllocations();
-
-		// Deal with the one that need removal
-		for (int32 i = ComponentBaseLayerAlloc.Num() - 1; i >= 0; --i)
-		{
-			const FWeightmapLayerAllocationInfo& Alloc = ComponentBaseLayerAlloc[i];
-
-			if (!ComponentLayerAlloc.Contains(Alloc.LayerInfo))
-			{
-				ComponentBaseLayerAlloc.RemoveAt(i);
-			}
-		}
-
-		// Then add the new one
-		for (ULandscapeLayerInfoObject* LayerAlloc : ComponentLayerAlloc)
-		{
-			const bool AllocExist = ComponentBaseLayerAlloc.ContainsByPredicate([&LayerAlloc](const FWeightmapLayerAllocationInfo& BaseLayerAlloc) { return (LayerAlloc == BaseLayerAlloc.LayerInfo); });
-
-			if (!AllocExist)
-			{
-				ComponentBaseLayerAlloc.Add(FWeightmapLayerAllocationInfo(LayerAlloc));
-			}
-		}
-	}
-
-	// Realloc the weightmap so it will create proper texture (if needed) and will set the allocations information
-	TArray<UTexture*> NewCreatedTextures;
-
-	for (ULandscapeComponent* Component : InLandscapeComponents)
-	{
-		Component->ReallocateWeightmaps(nullptr, false, false, false, nullptr, &NewCreatedTextures);
-	}
-
-	// TODO: correctly only recreate what is required instead of everything..
-	//GDisableAutomaticTextureMaterialUpdateDependencies = true;
-
-	FTextureCompilingManager::Get().FinishCompilation(NewCreatedTextures);
-	for (UTexture* Texture : NewCreatedTextures)
-	{
-		Texture->bForceMiplevelsToBeResident = true;
-		Texture->WaitForStreaming();
-	}
-
-	//GDisableAutomaticTextureMaterialUpdateDependencies = false;
-
-	// Clean-up unused weightmap CPUReadback resources
-	Info->ForAllLandscapeProxies([](ALandscapeProxy* Proxy)
-	{
-		TArray<UTexture2D*, TInlineAllocator<64>> EntriesToRemoveFromMap;
-		for (auto& Pair : Proxy->WeightmapsCPUReadback)
-		{
-			UTexture2D* WeightmapTextureKey = Pair.Key;
-			bool IsTextureReferenced = false;
-			for (ULandscapeComponent* Component : Proxy->LandscapeComponents)
-			{
-				for (UTexture2D* WeightmapTexture : Component->GetWeightmapTextures(false))
-				{
-					if (WeightmapTexture == WeightmapTextureKey)
-					{
-						IsTextureReferenced = true;
-						break;
-					}
-				}
-			}
-			if (!IsTextureReferenced)
-			{
-				EntriesToRemoveFromMap.Add(WeightmapTextureKey);
-			}
-		}
-
-		if (EntriesToRemoveFromMap.Num())
-		{
-			for (UTexture2D* OldWeightmapTexture : EntriesToRemoveFromMap)
-			{
-				if (FLandscapeEditLayerReadback** CPUReadbackToDelete = Proxy->WeightmapsCPUReadback.Find(OldWeightmapTexture))
-				{
-					check(*CPUReadbackToDelete);
-					delete * CPUReadbackToDelete;
-					Proxy->WeightmapsCPUReadback.Remove(OldWeightmapTexture);
-				}
-			}
-		}
-	});
-
-	ValidateProxyLayersWeightmapUsage();
-}
-
-void ALandscape::InitializeLayersWeightmapResources()
-{
-	ULandscapeInfo* Info = GetLandscapeInfo();
-
-	if (Info == nullptr)
-	{
-		return;
-	}
-
 	// Destroy existing resource
 	TArray<FTextureResource*> ResourceToDestroy;
 	ResourceToDestroy.Add(CombinedLayersWeightmapAllMaterialLayersResource);
@@ -4175,13 +5613,13 @@
 	// Use the 1st one to compute the resource as they are all the same anyway
 	UTextureRenderTarget2D* FirstWeightmapRT = WeightmapRTList[(int32)EWeightmapRTType::WeightmapRT_Scratch1];
 
-	CombinedLayersWeightmapAllMaterialLayersResource = new FLandscapeTexture2DArrayResource(FirstWeightmapRT->SizeX, FirstWeightmapRT->SizeY, LayerCount, PF_G8, 1, true);
+	CombinedLayersWeightmapAllMaterialLayersResource = new FLandscapeTexture2DArrayResource(FirstWeightmapRT->SizeX, FirstWeightmapRT->SizeY, LayerCount, PF_G8, 1, /*bInNeedUAVs = */ true, /*bInNeedSRV = */ false);
 	BeginInitResource(CombinedLayersWeightmapAllMaterialLayersResource);
 
-	CurrentLayersWeightmapAllMaterialLayersResource = new FLandscapeTexture2DArrayResource(FirstWeightmapRT->SizeX, FirstWeightmapRT->SizeY, LayerCount, PF_G8, 1, true);
+	CurrentLayersWeightmapAllMaterialLayersResource = new FLandscapeTexture2DArrayResource(FirstWeightmapRT->SizeX, FirstWeightmapRT->SizeY, LayerCount, PF_G8, 1, /*bInNeedUAVs = */ true, /*bInNeedSRV = */ false);
 	BeginInitResource(CurrentLayersWeightmapAllMaterialLayersResource);
 
-	WeightmapScratchExtractLayerTextureResource = new FLandscapeTexture2DResource(FirstWeightmapRT->SizeX, FirstWeightmapRT->SizeY, PF_B8G8R8A8, 1, false);
+	WeightmapScratchExtractLayerTextureResource = new FLandscapeTexture2DResource(FirstWeightmapRT->SizeX, FirstWeightmapRT->SizeY, PF_B8G8R8A8, 1, /*bInNeedUAVs = */ false, /*bInNeedSRV = */ false);
 	BeginInitResource(WeightmapScratchExtractLayerTextureResource);
 
 	int32 MipCount = 0;
@@ -4196,48 +5634,885 @@
 
 	// Format for UAV can't be PF_B8G8R8A8 on Windows 7 so use PF_R8G8B8A8
 	// We make the final copy out of this to a PF_R8G8B8A8 target with CopyTexturePS() instead of CopyLayersTexture() because a pixel shader will automatically handle the channel swizzling (where a RHICopyTexture won't)
-	WeightmapScratchPackLayerTextureResource = new FLandscapeTexture2DResource(FirstWeightmapRT->SizeX, FirstWeightmapRT->SizeY, PF_R8G8B8A8, MipCount, true);
+	WeightmapScratchPackLayerTextureResource = new FLandscapeTexture2DResource(FirstWeightmapRT->SizeX, FirstWeightmapRT->SizeY, PF_R8G8B8A8, MipCount, /*bInNeedUAVs = */ true, /*bInNeedSRV = */ false);
 	BeginInitResource(WeightmapScratchPackLayerTextureResource);
 }
 
-int32 ALandscape::RegenerateLayersWeightmaps(FTextureToComponentHelper const& MapHelper, const TArray<ULandscapeComponent*>& InLandscapeComponentsToRender, const TArray<ULandscapeComponent*>& InLandscapeComponentsToResolve)
+// Little struct that holds information common to PerformLayersWeightmapsLocalMerge and PerformLayersWeightmapsGlobalMerge
+struct FEditLayersWeightmapMergeParams
 {
-	TRACE_CPUPROFILER_EVENT_SCOPE(LandscapeLayers_RegenerateLayersWeightmaps);
-	SCOPE_CYCLE_COUNTER(STAT_LandscapeLayersRegenerateWeightmaps);
-	const int32 AllWeightmapUpdateModes = (ELandscapeLayerUpdateMode::Update_Weightmap_All | ELandscapeLayerUpdateMode::Update_Weightmap_Editing | ELandscapeLayerUpdateMode::Update_Weightmap_Editing_NoCollision);
-	const int32 WeightmapUpdateModes = LayerContentUpdateModes & AllWeightmapUpdateModes;
-	const bool bSkipBrush = CVarLandscapeLayerBrushOptim.GetValueOnAnyThread() == 1 && ((WeightmapUpdateModes & AllWeightmapUpdateModes) == ELandscapeLayerUpdateMode::Update_Weightmap_Editing);
-	const bool bForceRender = CVarForceLayersUpdate.GetValueOnAnyThread() != 0;
-	
+	int32 WeightmapUpdateModes;
+	bool bForceRender;
+	bool bSkipBrush;
+};
+
+// Render-thread version of the data / functions we need for the local merge of edit layers : 
+namespace EditLayersWeightmapLocalMerge_RenderThread
+{
+	struct FEditLayerInfo
+	{
+		int32 SourceWeightmapTextureIndex = INDEX_NONE; // The index in VisibleEditLayerWeightmapTextures of the texture to read from for this layer
+		int32 SourceWeightmapTextureChannel = INDEX_NONE; // The channel of the texture to read from for this layer
+		ELandscapeEditLayerWeightmapBlendMode BlendMode = ELandscapeEditLayerWeightmapBlendMode::Num; // See EEditLayerWeightmapBlendMode
+		float Alpha = 1.0f; // Alpha value to be used in the blend
+	};
+
+	struct FComponentPaintLayerRenderInfo
+	{
+		// Name of the paint layer for debug purposes :
+		FString Name;
+		// Describes how to access each visible edit layer's weightmap and how to blend it in the final weightmap for this paint layer :
+		TArray<FEditLayerInfo> VisibleEditLayerInfos;
+		// Global index of this paint layer in the paint layer infos array : 
+		int32 PaintLayerInfoIndex = INDEX_NONE;
+	};
+
+	struct FComponentRenderInfo
+	{
+		// Name of the component for debug purposes :
+		FString Name;
+		// The information needed to render each of the component's paint layers :
+		TArray<FComponentPaintLayerRenderInfo> PaintLayerRenderInfos;
+		// The unique textures referenced by this component's visible edit layer's weightmaps for all paint layers :
+		TArray<FTexture2DResource*> VisibleEditLayerWeightmapTextures;
+	};
+
+	// For a given FTextureResolveBatchInfo, allows to identify a FComponentRenderInfo/FComponentPaintLayerRenderInfo pair (useful when recombining the weightmaps into the final -packed- weightmap)
+	struct FComponentAndPaintLayerRenderInfoIdentifier
+	{
+		FComponentAndPaintLayerRenderInfoIdentifier() = default;
+		FComponentAndPaintLayerRenderInfoIdentifier(int32 InComponentIndex, int32 InPaintLayerIndex)
+			: ComponentIndex(InComponentIndex)
+			, PaintLayerIndex(InPaintLayerIndex)
+		{}
+		bool IsValid() const { return (ComponentIndex != INDEX_NONE) && (PaintLayerIndex != INDEX_NONE); }
+
+		bool operator == (const FComponentAndPaintLayerRenderInfoIdentifier& InOther) const { return (InOther.ComponentIndex == ComponentIndex) && (InOther.PaintLayerIndex == PaintLayerIndex); }
+		bool operator != (const FComponentAndPaintLayerRenderInfoIdentifier& InOther) const { return !(*this == InOther); }
+
+		int32 ComponentIndex = INDEX_NONE; // Index of a FComponentRenderInfo in ComponentToRenderInfos
+		int32 PaintLayerIndex = INDEX_NONE; // Index of a FComponentPaintLayerRenderInfo in FComponentRenderInfo::PaintLayerRenderInfos
+	};
+
+	struct FTextureResolveInfo
+	{
+		FTextureResolveInfo(int32 InNumComponentsToRender)
+		{
+			ComponentToRenderInfoBitIndices.Init(false, InNumComponentsToRender);
+		}
+
+		void ValidatePerChannelSourceInfo(int32 InChannelIndex, const FComponentAndPaintLayerRenderInfoIdentifier& InComponentAndPaintLayerIdentifier)
+		{
+			for (int32 Index = 0; Index < 4; ++Index)
+			{
+				if (Index == InChannelIndex)
+				{
+					// Channel shouldn't be already assigned
+					check(!PerChannelSourceWeightmapsIdentifiers[Index].IsValid());
+				}
+				else
+				{
+					// There should be no duplicates in channels : that would indicate that there are 2 identical component/layer info pair for 2 different texture channels
+					check(PerChannelSourceWeightmapsIdentifiers[Index] != InComponentAndPaintLayerIdentifier);
+				}
+			}
+		}
+
+		void SetPerChannelSourceInfo(int32 InChannelIndex, const FComponentAndPaintLayerRenderInfoIdentifier& InComponentAndPaintLayerIdentifier)
+		{
+			ValidatePerChannelSourceInfo(InChannelIndex, InComponentAndPaintLayerIdentifier);
+			PerChannelSourceWeightmapsIdentifiers[InChannelIndex] = InComponentAndPaintLayerIdentifier;
+			ComponentToRenderInfoBitIndices[InComponentAndPaintLayerIdentifier.ComponentIndex] = true;
+		}
+
+	public:
+		// Index of this texture in TextureToResolveInfos :
+		int32 TextureToResolveInfoIndex = INDEX_NONE;
+		// Texture that was created or updated that needs resolving : 
+		FTexture2DResource* Texture = nullptr;
+		// List of the 4 identifiers (one per weightmap channel) of a component/paint layer association in a given FTextureResolveBatchInfo that will be used 
+		//  to recombine the individual weightmaps into the final -packed- one
+		TStaticArray<FComponentAndPaintLayerRenderInfoIdentifier, 4> PerChannelSourceWeightmapsIdentifiers;
+		// List of(up to) 4 unique component render info indices that are needed for reconstructing the 4 channels of this texture (index into ComponentToRenderInfos) :
+		//  It's a bit array (1 bit per component to render info) to vastly optimize the division of texture resolve infos into batches, which is a O(N^2) operation : 
+		TBitArray<> ComponentToRenderInfoBitIndices;
+		// CPU readback utility to bring back the result on the CPU : 
+		FLandscapeEditLayerReadback* CPUReadback = nullptr;
+	};
+
+	// Because of weightmaps being shared between one component and another (within a given landcape proxy), we have to group the components to render into batches
+	//  where we'll render all of the paint layers into individual, 1-channel, scratch textures, which we'll then be able to re-assemble into the final, packed, weightmaps :
+	struct FTextureResolveBatchInfo
+	{
+		FTextureResolveBatchInfo(int32 InNumComponentsToRender, int32 InBatchIndex)
+			: BatchIndex(InBatchIndex)
+		{
+			ComponentToRenderInfoBitIndices.Init(false, InNumComponentsToRender);
+		}
+
+		void AddTexture(const FTextureResolveInfo& InTextureResolveInfo)
+		{
+			check(!TextureToResolveInfoIndices.Contains(InTextureResolveInfo.TextureToResolveInfoIndex));
+			TextureToResolveInfoIndices.Add(InTextureResolveInfo.TextureToResolveInfoIndex);
+			// Remember all the unique components that this texture needs for resolving:
+			ComponentToRenderInfoBitIndices.CombineWithBitwiseOR(InTextureResolveInfo.ComponentToRenderInfoBitIndices, EBitwiseOperatorFlags::MinSize);
+		}
+
+	public:
+		// Index of this batch in TextureResolveBatchInfos
+		int32 BatchIndex = INDEX_NONE;
+
+		// Indices (in ComponentToRenderInfos) of the components whose weightmaps we need to render within this batch in order to produce (and then resolve) the textures in TextureToResolveInfos:
+		//  It's a bit array (1 bit per component to render info) to vastly optimize the division of texture resolve infos into batches, which is a O(N^2) operation : 
+		TBitArray<> ComponentToRenderInfoBitIndices;
+		// Indices (TextureToResolveInfos) of textures that need to be resolved / read back on the CPU :
+		TArray<int32> TextureToResolveInfoIndices;
+	};
+
+	// Description of the entire merge pass :
+	struct FMergeInfo
+	{
+		bool NeedsMerge() const
+		{
+			// If no edit layer or if no paint layer present on any edit layer, we've got nothing to do :
+			bool bNeedsMerge = (MaxNumEditLayersTexturesToMerge > 0) && (MaxNumWeightmapArraysPerResolveTextureBatch > 0);
+			check(!bNeedsMerge || !PaintLayerInfos.IsEmpty()); // If we need merging, we must have at least one paint layer
+			return bNeedsMerge;
+		}
+
+	public:
+		// Maximum number of visible edit layers that have to be merged for a single FComponentRenderInfo : 
+		int32 MaxNumEditLayersTexturesToMerge = 0;
+
+		// Maximum number of weightmap arrays that is needed for a given FTextureResolveBatchInfo (1 per FComponentRenderInfo in the batch) :
+		int32 MaxNumWeightmapArraysPerResolveTextureBatch = 0;
+
+		// Per-landscape component info :
+		FIntPoint ComponentSizeVerts = 0;
+		int32 NumMips = 0;
+
+		// List of batches of FTextureResolveInfo that needs to be resolved in the same pass. This allows massive saves on transient resources on large landscapes because those can be re-cycled from one pass to another :
+		TArray<FTextureResolveBatchInfo> TextureResolveBatchInfos;
+
+		// List of infos for each component that needs its paint layers to be rendered in order to be resolved :
+		TArray<FComponentRenderInfo> ComponentToRenderInfos;
+
+		// List of infos for each texture that needs to be resolved :
+		TArray<FTextureResolveInfo> TextureToResolveInfos;
+
+		// List of infos for each individual paint layer involved in the merge operation (including the visibility layer)
+		TArray<FLandscapeWeightmapPaintLayerInfo> PaintLayerInfos;
+
+		// Not truly render-thread data because it references UTextures but it's just because FLandscapeEditLayerReadback were historically game-thread initiated so for as long as we'll use those for readback, we need to store this here : 
+		TArray<FLandscapeLayersCopyReadbackTextureParams> DeferredCopyReadbackTextures;
+	};
+
+	// Render graph intermediate resources :
+	struct FRDGResources
+	{
+		// Contains info about each individual paint layer 
+		FRDGBufferRef PaintLayerInfosBuffer;
+		FRDGBufferSRVRef PaintLayerInfosBufferSRV;
+
+		// Texture array that can be reused from one component / one paint layer to another and that contains the list of all edit layers textures that need merging in a given pass :
+		FRDGTextureRef EditLayersWeightmapsTextureArray;
+		FRDGTextureSRVRef EditLayersWeightmapsTextureArraySRV;
+
+		// List of temporary scratch texture arrays that store the output for all (edit layer-merged) active paint layer (one per landscape component) until they can be packed onto the final weightmap textures. 
+		//  Can be reused from one batch to another :
+		TArray<FRDGTextureRef> ScratchPaintLayerWeightmapTextureArrays;
+		TArray<FRDGTextureSRVRef> ScratchPaintLayerWeightmapTextureArraysSRV;
+
+		// Single scratch texture for the weightmap finalize operation (since we cannot directly write to the final weightmaps because they were not created with TexCreate_RenderTargetable) :
+		FRDGTextureRef ScratchFinalWeightmapTexture;
+		TArray<FRDGTextureSRVRef> ScratchFinalWeightmapTextureMipsSRV; // One SRV per mip level
+
+		// Single structured buffer that will contain all possible FLandscapeEditLayerWeightmapMergeInfo we might need during the entire merge operation (this is to avoid
+		//  too many individual buffer allocations/uploads when many components need to be merged : CPU optimization) :
+		FRDGBufferRef EditLayerMergeInfosBuffer;
+		FRDGBufferSRVRef EditLayerMergeInfosBufferSRV;
+		int32 CurrentEditLayerMergeInfosBufferIndex = 0;
+
+		// Single structured buffer that will contain all possible paint layer info indices we might need during the entire merge operation (this is to avoid
+		//  too many individual buffer allocations/uploads when many components need to be merged : CPU optimization) :
+		FRDGBufferRef PaintLayerInfoIndicesBuffer;
+		FRDGBufferSRVRef PaintLayerInfoIndicesBufferSRV;
+		int32 CurrentPaintLayerInfoIndicesBufferIndex = 0;
+
+		// Dummy buffers :
+		FRDGTextureSRVRef BlackDummyArraySRV;
+	};
+
+	TArray<FLandscapeEditLayerWeightmapMergeInfo> PrepareEditLayerWeightmapMergeInfosBufferData(const FMergeInfo& InLocalMergeInfo)
+	{
+		TArray<FLandscapeEditLayerWeightmapMergeInfo> EditLayerMergeInfos;
+
+		// Batch by batch
+		for (const FTextureResolveBatchInfo& TextureResolveBatchInfo : InLocalMergeInfo.TextureResolveBatchInfos)
+		{
+			// Component by component
+			for (TConstSetBitIterator<> BitIt(TextureResolveBatchInfo.ComponentToRenderInfoBitIndices); BitIt; ++BitIt)
+			{
+				int32 ComponentRenderInfoIndex = BitIt.GetIndex();
+
+				// Paint layer by paint layer
+				for (const FComponentPaintLayerRenderInfo& ComponentPaintLayerRenderInfo : InLocalMergeInfo.ComponentToRenderInfos[ComponentRenderInfoIndex].PaintLayerRenderInfos)
+				{
+					// Update the edit layers merge info big buffer : 
+					for (const FEditLayerInfo& EditLayerInfo : ComponentPaintLayerRenderInfo.VisibleEditLayerInfos)
+					{
+						FLandscapeEditLayerWeightmapMergeInfo& EditLayerMergeInfo = EditLayerMergeInfos.Emplace_GetRef();
+						EditLayerMergeInfo.SourceWeightmapTextureIndex = EditLayerInfo.SourceWeightmapTextureIndex;
+						EditLayerMergeInfo.SourceWeightmapTextureChannel = EditLayerInfo.SourceWeightmapTextureChannel;
+						EditLayerMergeInfo.BlendMode = EditLayerInfo.BlendMode;
+						EditLayerMergeInfo.Alpha = EditLayerInfo.Alpha;
+					}
+				}
+			}
+		}
+
+		return EditLayerMergeInfos;
+	}
+
+	TArray<uint32> PreparePaintLayerInfoIndicesBufferData(const FMergeInfo& InLocalMergeInfo)
+	{
+		TArray<uint32> PaintLayerInfoIndices;
+
+		// Batch by batch
+		for (const FTextureResolveBatchInfo& TextureResolveBatchInfo : InLocalMergeInfo.TextureResolveBatchInfos)
+		{
+			// Texture by texture
+			for (int32 TextureResolveInfoIndex : TextureResolveBatchInfo.TextureToResolveInfoIndices)
+			{
+				// Output channel by output channel
+				const FTextureResolveInfo& TextureResolveInfo = InLocalMergeInfo.TextureToResolveInfos[TextureResolveInfoIndex];
+				for (int32 ChannelIndex = 0; ChannelIndex < 4; ++ChannelIndex)
+				{
+					const FComponentAndPaintLayerRenderInfoIdentifier& ComponentAndPaintLayerRenderInfoIdentifier = TextureResolveInfo.PerChannelSourceWeightmapsIdentifiers[ChannelIndex];
+					if (ComponentAndPaintLayerRenderInfoIdentifier.IsValid())
+					{
+						check(InLocalMergeInfo.ComponentToRenderInfos.IsValidIndex(ComponentAndPaintLayerRenderInfoIdentifier.ComponentIndex)); // this identifier must point to a valid render info in the FTextureResolveBatchInfo
+						const FComponentRenderInfo& ComponentRenderInfo = InLocalMergeInfo.ComponentToRenderInfos[ComponentAndPaintLayerRenderInfoIdentifier.ComponentIndex];
+
+						check(ComponentRenderInfo.PaintLayerRenderInfos.IsValidIndex(ComponentAndPaintLayerRenderInfoIdentifier.PaintLayerIndex)); // this identifier must point to a valid paint layer render info in the FComponentRenderInfo
+
+						// Update the PaintLayerInfoIndices big buffer : 
+						for (const FComponentPaintLayerRenderInfo& PaintLayerInfo : ComponentRenderInfo.PaintLayerRenderInfos)
+						{
+							check(PaintLayerInfo.PaintLayerInfoIndex < InLocalMergeInfo.PaintLayerInfos.Num());
+							PaintLayerInfoIndices.Add(PaintLayerInfo.PaintLayerInfoIndex);
+						}
+					}
+				}
+			}
+		}
+
+		return PaintLayerInfoIndices;
+	}
+
+	void PrepareLayersWeightmapsLocalMergeRDGResources(const FMergeInfo& InLocalMergeInfo, FRDGBuilder& GraphBuilder, FRDGResources& OutResources)
+	{
+		{
+			// Upload paint layer infos buffer once and for all since it's unchanged from one component to another : 
+			OutResources.PaintLayerInfosBuffer = CreateStructuredBuffer(GraphBuilder, TEXT("LandscapePaintLayerInfosBuffer"), InLocalMergeInfo.PaintLayerInfos);
+			OutResources.PaintLayerInfosBufferSRV = GraphBuilder.CreateSRV(FRDGBufferSRVDesc(OutResources.PaintLayerInfosBuffer));
+		}
+
+		{
+			// Allocate a texture array that can contain all edit layers textures to merge for any given component to render (this will be reused from one component to another) :
+			int32 SizeZ = InLocalMergeInfo.MaxNumEditLayersTexturesToMerge;
+			check(SizeZ > 0);
+			FRDGTextureDesc Desc = FRDGTextureDesc::Create2DArray(InLocalMergeInfo.ComponentSizeVerts, PF_B8G8R8A8, FClearValueBinding::None, TexCreate_ShaderResource, SizeZ, /*InNumMips = */1, /*InNumSamples = */1);
+			OutResources.EditLayersWeightmapsTextureArray = GraphBuilder.CreateTexture(Desc, TEXT("LandscapeEditLayersWeightmapsTextureArray"));
+			OutResources.EditLayersWeightmapsTextureArraySRV = GraphBuilder.CreateSRV(FRDGTextureSRVDesc::Create(OutResources.EditLayersWeightmapsTextureArray));
+		}
+
+		{
+			// Allocate as many texture arrays as needed for a given texture resolve batch (this will be reused from one batch to another) :
+			OutResources.ScratchPaintLayerWeightmapTextureArrays.Reserve(InLocalMergeInfo.MaxNumWeightmapArraysPerResolveTextureBatch);
+			OutResources.ScratchPaintLayerWeightmapTextureArraysSRV.Reserve(InLocalMergeInfo.MaxNumWeightmapArraysPerResolveTextureBatch);
+			// Each texture array (reusable from batch to batch) will contain at most a number of slices equals to the total number of active paint layers (including the visibility layer) : 
+			int32 SizeZ = InLocalMergeInfo.PaintLayerInfos.Num();
+			check(SizeZ > 0);
+			for (int32 Index = 0; Index < InLocalMergeInfo.MaxNumWeightmapArraysPerResolveTextureBatch; ++Index)
+			{
+				FRDGTextureDesc Desc = FRDGTextureDesc::Create2DArray(InLocalMergeInfo.ComponentSizeVerts, PF_G8, FClearValueBinding::None, TexCreate_ShaderResource | TexCreate_RenderTargetable | TexCreate_TargetArraySlicesIndependently, SizeZ, /*InNumMips = */1, /*InNumSamples = */1);
+				FRDGTextureRef& TextureRef = OutResources.ScratchPaintLayerWeightmapTextureArrays.Add_GetRef(GraphBuilder.CreateTexture(Desc, TEXT("LandscapeEditLayersScratchPaintLayerWeightmapTextureArray")));
+				OutResources.ScratchPaintLayerWeightmapTextureArraysSRV.Add(GraphBuilder.CreateSRV(FRDGTextureSRVDesc::Create(TextureRef)));
+			}
+		}
+
+		{
+			// Allocate a single scratch texture will all of its mip for each individual texture we want to resolve :
+			check(InLocalMergeInfo.NumMips > 0);
+			FRDGTextureDesc Desc = FRDGTextureDesc::Create2D(InLocalMergeInfo.ComponentSizeVerts, PF_B8G8R8A8, FClearValueBinding::None, TexCreate_ShaderResource | TexCreate_RenderTargetable, InLocalMergeInfo.NumMips, /*InNumSamples = */1);
+			OutResources.ScratchFinalWeightmapTexture = GraphBuilder.CreateTexture(Desc, TEXT("LandscapeEditLayersScratchFinalWeightmapTexture"));
+			for (int32 MipLevel = 0; MipLevel < InLocalMergeInfo.NumMips; ++MipLevel)
+			{
+				OutResources.ScratchFinalWeightmapTextureMipsSRV.Add(GraphBuilder.CreateSRV(FRDGTextureSRVDesc::CreateForMipLevel(OutResources.ScratchFinalWeightmapTexture, MipLevel)));
+			}
+		}
+
+		{
+			// Allocate a single structured buffer that will contain all possible FLandscapeEditLayerWeightmapMergeInfo we might need during the entire merge operation. 
+			//  Although CreateStructuredBuffer can be given a callback to provide its initial data, we need to build the source data array upfront, since it's not compatible with RDG immediate mode:
+			OutResources.EditLayerMergeInfosBuffer = CreateStructuredBuffer(GraphBuilder, TEXT("LandscapeEditLayersMergeInfosBuffer"), PrepareEditLayerWeightmapMergeInfosBufferData(InLocalMergeInfo));
+			OutResources.EditLayerMergeInfosBufferSRV = GraphBuilder.CreateSRV(FRDGBufferSRVDesc(OutResources.EditLayerMergeInfosBuffer));
+		}
+
+		{
+			// Allocate a single structured buffer that will contain all possible paint layer info indices we might need during the entire merge operation.
+			//  Although CreateStructuredBuffer can be given a callback to provide its initial data, we need to build the source data array upfront, since it's not compatible with RDG immediate mode:
+			OutResources.PaintLayerInfoIndicesBuffer = CreateStructuredBuffer(GraphBuilder, TEXT("LandscapeEditLayersPaintLayerInfoIndicesBuffer"), PreparePaintLayerInfoIndicesBufferData(InLocalMergeInfo));
+			OutResources.PaintLayerInfoIndicesBufferSRV = GraphBuilder.CreateSRV(FRDGBufferSRVDesc(OutResources.PaintLayerInfoIndicesBuffer));
+		}
+
+		{
+			// Dummy buffers for avoiding missing shader bindings
+			OutResources.BlackDummyArraySRV = GraphBuilder.CreateSRV(FRDGTextureSRVDesc::Create(GSystemTextures.GetBlackArrayDummy(GraphBuilder)));
+		}
+	}
+
+	// Gather all textures we will want to write into or read from in the render graph and output them in OutTrackedTextures:
+	void GatherLayersWeightmapsLocalMergeRDGTextures(const FMergeInfo& InLocalMergeInfo, TMap<FTexture2DResource*, FLandscapeRDGTrackedTexture>& OutTrackedTextures)
+	{
+		// Gather all source weightmaps :
+		for (const FComponentRenderInfo& ComponentRenderInfo : InLocalMergeInfo.ComponentToRenderInfos)
+		{
+			for (FTexture2DResource* VisibleEditLayersWeightmapTexture : ComponentRenderInfo.VisibleEditLayerWeightmapTextures)
+			{
+				check(VisibleEditLayersWeightmapTexture != nullptr);
+				FLandscapeRDGTrackedTexture* TrackedTexture = OutTrackedTextures.Find(VisibleEditLayersWeightmapTexture);
+				if (TrackedTexture == nullptr)
+				{
+					TrackedTexture = &(OutTrackedTextures.Add(VisibleEditLayersWeightmapTexture, FLandscapeRDGTrackedTexture(VisibleEditLayersWeightmapTexture)));
+				}
+				// Source weightmaps were not created with TexCreate_RenderTargetable and are therefore considered as ReadOnly by the RDG, which prevents their resource access from being automatically managed
+				//  (i.e. transfered to CopySrc when we copy them) so we need to force the RDG to track them: 
+				TrackedTexture->bNeedsForceTracking = true;
+				TrackedTexture->bNeedsSRV = true;
+			}
+		}
+
+		// Gather all destination weightmaps :
+		for (const FTextureResolveInfo& TextureResolveInfo : InLocalMergeInfo.TextureToResolveInfos)
+		{
+			FLandscapeRDGTrackedTexture* TrackedTexture = OutTrackedTextures.Find(TextureResolveInfo.Texture);
+			if (ensure(TrackedTexture == nullptr)) // Resolved textures should only be registered once
+			{
+				TrackedTexture = &(OutTrackedTextures.Add(TextureResolveInfo.Texture, FLandscapeRDGTrackedTexture(TextureResolveInfo.Texture)));
+			}
+
+			// Destination weightmaps were not created with TexCreate_RenderTargetable and are therefore considered as ReadOnly by the RDG, which prevents their resource access from being automatically managed
+			//  (i.e. transfered to CopyDst when we copy to them and back to SRVMask at the end) so we need to force the RDG to track them: 
+			TrackedTexture->bNeedsForceTracking = true;
+			TrackedTexture->bNeedsSRV = true;
+		}
+	}
+
+	void MergeEditLayersWeightmapsForBatch(const FTextureResolveBatchInfo& InTextureResolveBatchInfo, const FMergeInfo& InLocalMergeInfo,
+		const TMap<FTexture2DResource*, FLandscapeRDGTrackedTexture>& InTrackedTextures, FRDGBuilder& GraphBuilder, FRDGResources& RDGResources)
+	{
+		RDG_EVENT_SCOPE(GraphBuilder, "Merge paint layers on %d components for batch %d", InTextureResolveBatchInfo.ComponentToRenderInfoBitIndices.CountSetBits(), InTextureResolveBatchInfo.BatchIndex);
+
+		// For each component in the batch, perform the edit layers merge and write the resulting scratch weightmap :
+		int32 IndexInBatch = 0;
+		for (TConstSetBitIterator<> BitIt(InTextureResolveBatchInfo.ComponentToRenderInfoBitIndices); BitIt; ++BitIt, ++IndexInBatch)
+		{
+			int32 ComponentRenderInfoIndex = BitIt.GetIndex();
+
+			const FComponentRenderInfo& ComponentRenderInfo = InLocalMergeInfo.ComponentToRenderInfos[ComponentRenderInfoIndex];
+			RDG_EVENT_SCOPE(GraphBuilder, "Component %s", *ComponentRenderInfo.Name);
+
+			// Prepare a texture array of texture array that will contain all the input textures we'll need for this component, regardless of the paint layer
+			//  (done once per component since various paint layers could share the same texture) :
+			{
+				int32 NumTextures = ComponentRenderInfo.VisibleEditLayerWeightmapTextures.Num();
+				for (int32 TextureIndex = 0; TextureIndex < NumTextures; ++TextureIndex)
+				{
+					const FTexture2DResource* PackedWeightmap = ComponentRenderInfo.VisibleEditLayerWeightmapTextures[TextureIndex];
+
+					const FLandscapeRDGTrackedTexture* TrackedTexture = InTrackedTextures.Find(PackedWeightmap);
+					check(TrackedTexture != nullptr);
+
+					// We need to copy the (portion of the) layer's texture to the texture array : 
+					FRHICopyTextureInfo CopyTextureInfo;
+					CopyTextureInfo.Size = FIntVector(InLocalMergeInfo.ComponentSizeVerts.X, InLocalMergeInfo.ComponentSizeVerts.Y, 0);
+					CopyTextureInfo.DestSliceIndex = TextureIndex;
+
+					AddCopyTexturePass(GraphBuilder, TrackedTexture->ExternalTextureRef, RDGResources.EditLayersWeightmapsTextureArray, CopyTextureInfo);
+				}
+			}
+
+			// We should have a single output scratch texture array reserved for this component in RDGResources.ScratchPaintLayerWeightmapTextureArrays already :
+			check(RDGResources.ScratchPaintLayerWeightmapTextureArrays.IsValidIndex(IndexInBatch));
+			FRDGTextureRef ScratchTextureArrayRef = RDGResources.ScratchPaintLayerWeightmapTextureArrays[IndexInBatch];
+
+			// Paint layer by paint layer, merge the weightmaps from all the corresponding edit layers onto the corresponding scratch texture using the MergeEditLayers PS : 
+			int32 NumComponentPaintLayers = ComponentRenderInfo.PaintLayerRenderInfos.Num();
+			for (int32 ComponentPaintLayerIndex = 0; ComponentPaintLayerIndex < NumComponentPaintLayers; ++ComponentPaintLayerIndex)
+			{
+				const FComponentPaintLayerRenderInfo& ComponentPaintLayerRenderInfo = ComponentRenderInfo.PaintLayerRenderInfos[ComponentPaintLayerIndex];
+				RDG_EVENT_SCOPE(GraphBuilder, "Merge %d edit layers for paint layer %s", ComponentPaintLayerRenderInfo.VisibleEditLayerInfos.Num(), *ComponentPaintLayerRenderInfo.Name);
+
+				FLandscapeLayersWeightmapsMergeEditLayersPS::FParameters* MergeEditLayersPSParams = GraphBuilder.AllocParameters<FLandscapeLayersWeightmapsMergeEditLayersPS::FParameters>();
+				// We'll write to a single slice of the texture array for this component, since we're acting paint layer by paint layer here :
+				MergeEditLayersPSParams->RenderTargets[0] = FRenderTargetBinding(ScratchTextureArrayRef, ERenderTargetLoadAction::ENoAction, /*InMipIndex = */0, /*InArraySlice = */ComponentPaintLayerIndex);
+				MergeEditLayersPSParams->InNumEditLayers = ComponentPaintLayerRenderInfo.VisibleEditLayerInfos.Num();
+				MergeEditLayersPSParams->InPackedWeightmaps = RDGResources.EditLayersWeightmapsTextureArraySRV;
+				MergeEditLayersPSParams->InEditLayersMergeInfos = RDGResources.EditLayerMergeInfosBufferSRV;
+
+				// EditLayerMergeInfosBuffer is already uploaded but we need to tell the shader where we're currently at in that big buffer :
+				MergeEditLayersPSParams->InStartIndexInEditLayersMergeInfos = RDGResources.CurrentEditLayerMergeInfosBufferIndex; 
+				// Update CurrentEditLayerMergeInfosBufferIndex so that the next paint layer starts after : 
+				RDGResources.CurrentEditLayerMergeInfosBufferIndex += ComponentPaintLayerRenderInfo.VisibleEditLayerInfos.Num();
+
+				FLandscapeLayersWeightmapsMergeEditLayersPS::MergeEditLayers(GraphBuilder, MergeEditLayersPSParams);
+			}
+		}
+	}
+
+	void FinalizeSingleWeightmap(const FTextureResolveBatchInfo& InTextureResolveBatchInfo, const FMergeInfo& InLocalMergeInfo,
+		const FTextureResolveInfo& InTextureResolveInfo, FRDGBuilder& GraphBuilder, FRDGResources& RDGResources)
+	{
+		RDG_EVENT_SCOPE(GraphBuilder, "Finalize Weightmap %s", *InTextureResolveInfo.Texture->GetTextureName().ToString());
+
+		FLandscapeLayersWeightmapsFinalizeWeightmapPS::FParameters* FinalizeWeightmapPSParams = GraphBuilder.AllocParameters<FLandscapeLayersWeightmapsFinalizeWeightmapPS::FParameters>();
+		FinalizeWeightmapPSParams->RenderTargets[0] = FRenderTargetBinding(RDGResources.ScratchFinalWeightmapTexture, ERenderTargetLoadAction::ENoAction);
+		FinalizeWeightmapPSParams->InPerChannelPaintLayerIndexInWeightmaps = FUintVector4(ForceInitToZero);
+		FinalizeWeightmapPSParams->InPerChannelStartPaintLayerIndex = FUintVector4(ForceInitToZero);
+		FinalizeWeightmapPSParams->InPerChannelNumPaintLayers = FUintVector4(ForceInitToZero);
+		FinalizeWeightmapPSParams->InPaintLayerInfoIndices = RDGResources.PaintLayerInfoIndicesBufferSRV;
+		FinalizeWeightmapPSParams->InPaintLayerInfos = RDGResources.PaintLayerInfosBufferSRV;
+		FinalizeWeightmapPSParams->InValidTextureChannelsMask = 0;
+
+		check(InTextureResolveInfo.PerChannelSourceWeightmapsIdentifiers.Num() == 4);
+		for (int32 ChannelIndex = 0; ChannelIndex < 4; ++ChannelIndex)
+		{
+			const FComponentAndPaintLayerRenderInfoIdentifier& ComponentAndPaintLayerRenderInfoIdentifier = InTextureResolveInfo.PerChannelSourceWeightmapsIdentifiers[ChannelIndex];
+			if (ComponentAndPaintLayerRenderInfoIdentifier.IsValid())
+			{
+				// Indicate this channel will need to be processed :
+				FinalizeWeightmapPSParams->InValidTextureChannelsMask |= (1 << ChannelIndex);
+
+				check(InLocalMergeInfo.ComponentToRenderInfos.IsValidIndex(ComponentAndPaintLayerRenderInfoIdentifier.ComponentIndex)); // this identifier must point to a valid render info in the FTextureResolveBatchInfo
+				const FComponentRenderInfo& ComponentRenderInfo = InLocalMergeInfo.ComponentToRenderInfos[ComponentAndPaintLayerRenderInfoIdentifier.ComponentIndex];
+
+				check(ComponentRenderInfo.PaintLayerRenderInfos.IsValidIndex(ComponentAndPaintLayerRenderInfoIdentifier.PaintLayerIndex)); // this identifier must point to a valid paint layer render info in the FComponentRenderInfo
+
+				// The paint layer to process on this texture channel : 
+				FinalizeWeightmapPSParams->InPerChannelPaintLayerIndexInWeightmaps[ChannelIndex] = ComponentAndPaintLayerRenderInfoIdentifier.PaintLayerIndex;
+				// The total number of paint layers for this component/paint layer (for weight blending in-between paint layers):
+				FinalizeWeightmapPSParams->InPerChannelNumPaintLayers[ChannelIndex] = ComponentRenderInfo.PaintLayerRenderInfos.Num();
+				// The index at which we'll find the first paint layer info index in the PaintLayerInfoIndices big buffer for this channel : use the index where we're currently at in that big buffer :
+				FinalizeWeightmapPSParams->InPerChannelStartPaintLayerIndex[ChannelIndex] = RDGResources.CurrentPaintLayerInfoIndicesBufferIndex;
+				// And update the big buffer current index so that the next channel starts at the right location int he big buffer : 
+				RDGResources.CurrentPaintLayerInfoIndicesBufferIndex += ComponentRenderInfo.PaintLayerRenderInfos.Num();
+
+				// We should have a single output scratch texture array reserved for this component in RDGResources.ScratchPaintLayerWeightmapTextureArrays already :
+				check(InTextureResolveBatchInfo.ComponentToRenderInfoBitIndices[ComponentAndPaintLayerRenderInfoIdentifier.ComponentIndex]); // This component should have been rendered in that batch!
+				int32 ScratchPaintLayerWeightmapTextureIndex = InTextureResolveBatchInfo.ComponentToRenderInfoBitIndices.CountSetBits(0, ComponentAndPaintLayerRenderInfoIdentifier.ComponentIndex);
+				check(RDGResources.ScratchPaintLayerWeightmapTextureArrays.IsValidIndex(ScratchPaintLayerWeightmapTextureIndex));
+				FinalizeWeightmapPSParams->InPerChannelPaintLayerWeightmaps[ChannelIndex] = RDGResources.ScratchPaintLayerWeightmapTextureArraysSRV[ScratchPaintLayerWeightmapTextureIndex];
+			}
+			else
+			{
+				FinalizeWeightmapPSParams->InPerChannelPaintLayerWeightmaps[ChannelIndex] = RDGResources.BlackDummyArraySRV;
+			}
+		}
+
+		FLandscapeLayersWeightmapsFinalizeWeightmapPS::FinalizeWeightmap(GraphBuilder, FinalizeWeightmapPSParams);
+	}
+
+	void GenerateSingleWeightmapMips(const FTextureResolveInfo& InTextureResolveInfo, const FMergeInfo& InLocalMergeInfo,
+		FRDGBuilder& GraphBuilder, FRDGResources& RDGResources)
+	{
+		RDG_EVENT_SCOPE(GraphBuilder, "Generate mips for Weightmap %s", *InTextureResolveInfo.Texture->GetTextureName().ToString());
+		check(RDGResources.ScratchFinalWeightmapTextureMipsSRV.Num() == InLocalMergeInfo.NumMips);
+
+		FIntPoint CurrentMipSize = InLocalMergeInfo.ComponentSizeVerts;
+		for (int32 MipLevel = 1; MipLevel < InLocalMergeInfo.NumMips; ++MipLevel)
+		{
+			CurrentMipSize.X >>= 1;
+			CurrentMipSize.Y >>= 1;
+
+			// Read from scratch weightmap texture (mip N - 1) -> write to scratch weightmap texture (mip N) :
+			FLandscapeLayersWeightmapsGenerateMipsPS::FParameters* GenerateMipsPSParams = GraphBuilder.AllocParameters<FLandscapeLayersWeightmapsGenerateMipsPS::FParameters>();
+			GenerateMipsPSParams->RenderTargets[0] = FRenderTargetBinding(RDGResources.ScratchFinalWeightmapTexture, ERenderTargetLoadAction::ENoAction, MipLevel);
+			GenerateMipsPSParams->InCurrentMipSize = FUintVector2(CurrentMipSize.X, CurrentMipSize.Y);
+			GenerateMipsPSParams->InSourceWeightmap = RDGResources.ScratchFinalWeightmapTextureMipsSRV[MipLevel - 1];
+
+			FLandscapeLayersWeightmapsGenerateMipsPS::GenerateSingleMip(GraphBuilder, GenerateMipsPSParams);
+		}
+	}
+
+	void CopyScratchToSourceWeightmap(const FTextureResolveInfo& InTextureResolveInfo, const TMap<FTexture2DResource*, FLandscapeRDGTrackedTexture>& InTrackedTextures,
+		FRDGBuilder& GraphBuilder, FRDGResources& RDGResources)
+	{
+		RDG_EVENT_SCOPE(GraphBuilder, "Copy scratch to Weightmap %s", *InTextureResolveInfo.Texture->GetTextureName().ToString());
+
+		const FLandscapeRDGTrackedTexture* TrackedTexture = InTrackedTextures.Find(InTextureResolveInfo.Texture);
+		check((TrackedTexture != nullptr) && (TrackedTexture->ExternalTextureRef != nullptr));
+
+		FRHICopyTextureInfo CopyTextureInfo;
+		// We want to copy all mips : 
+		CopyTextureInfo.NumMips = TrackedTexture->ExternalTextureRef->Desc.NumMips;
+
+		AddCopyTexturePass(GraphBuilder, RDGResources.ScratchFinalWeightmapTexture, TrackedTexture->ExternalTextureRef, CopyTextureInfo);
+	}
+
+	void FinalizeAndResolveWeightmapsForBatch(const FTextureResolveBatchInfo& InTextureResolveBatchInfo, const FMergeInfo& InLocalMergeInfo,
+		const TMap<FTexture2DResource*, FLandscapeRDGTrackedTexture>& InTrackedTextures, FRDGBuilder& GraphBuilder, FRDGResources& RDGResources)
+	{
+		RDG_EVENT_SCOPE(GraphBuilder, "Finalize %d weightmaps for batch %d", InTextureResolveBatchInfo.TextureToResolveInfoIndices.Num(), InTextureResolveBatchInfo.BatchIndex);
+
+		for (int32 TextureResolveInfoIndex : InTextureResolveBatchInfo.TextureToResolveInfoIndices)
+		{
+			const FTextureResolveInfo& TextureResolveInfo = InLocalMergeInfo.TextureToResolveInfos[TextureResolveInfoIndex];
+			RDG_EVENT_SCOPE(GraphBuilder, "Finalize / resolve %s", *TextureResolveInfo.Texture->GetTextureName().ToString());
+
+			// Finalize the weightmap to the scratch texture (cannot directly write to the texture because it's not render-targetable) :
+			FinalizeSingleWeightmap(InTextureResolveBatchInfo, InLocalMergeInfo, TextureResolveInfo, GraphBuilder, RDGResources);
+
+			// Generate mips for this scratch texture :
+			GenerateSingleWeightmapMips(TextureResolveInfo, InLocalMergeInfo, GraphBuilder, RDGResources);
+
+			// And finally, copy to the output texture :
+			CopyScratchToSourceWeightmap(TextureResolveInfo, InTrackedTextures, GraphBuilder, RDGResources);
+		}
+	}
+
+} // End namespace EditLayersWeightmapLocalMerge_RenderThread
+
+void ALandscape::PrepareLayersWeightmapsLocalMergeRenderThreadData(const FUpdateLayersContentContext& InUpdateLayersContentContext, const FEditLayersWeightmapMergeParams& InMergeParams, EditLayersWeightmapLocalMerge_RenderThread::FMergeInfo& OutRenderThreadData)
+{
+	using namespace EditLayersWeightmapLocalMerge_RenderThread;
+
+	TRACE_CPUPROFILER_EVENT_SCOPE(LandscapeLayers_PrepareLayersWeightmapsLocalMergeRenderThreadData);
+
+	ULandscapeInfo* LandscapeInfo = GetLandscapeInfo();
+	check(LandscapeInfo != nullptr);
+
+	// Number of vertices for each landscape component :
+	int32 ComponentSizeVerts = (SubsectionSizeQuads + 1) * NumSubsections;
+	OutRenderThreadData.ComponentSizeVerts = FIntPoint(ComponentSizeVerts, ComponentSizeVerts);
+	OutRenderThreadData.NumMips = FMath::CeilLogTwo(ComponentSizeVerts) + 1;
+
+	// Lookup table to retrieve, for a given paint layer, its index in OutRenderThreadData.PaintLayerInfos (we don't keep UObjects in OutRenderThreadData because it's a render-thread struct) :
+	TMap<ULandscapeLayerInfoObject*, int32> PaintLayerToPaintLayerInfoIndex;
+
+	// Lookup table to retrieve, for a given rendered component/paint layer, its FComponentRenderInfo/FComponentPaintLayerRenderInfo pair identifier (indices)
+	TMap<TPair<ULandscapeComponent*, ULandscapeLayerInfoObject*>, FComponentAndPaintLayerRenderInfoIdentifier> ComponentAndPaintLayerToRenderInfoIndex;
+
+	// Prepare per-landscape component render data : 
+	{
+		TRACE_CPUPROFILER_EVENT_SCOPE(LandscapeLayers_PrepareWeightmapComponentRenderAndPaintLayerInfos);
+		for (ULandscapeComponent* Component : InUpdateLayersContentContext.LandscapeComponentsWeightmapsToRender)
+		{
+			FComponentRenderInfo& ComponentRenderInfo = OutRenderThreadData.ComponentToRenderInfos.Emplace_GetRef();
+			ComponentRenderInfo.Name = Component->GetName();
+
+			// Build, for each of the paint layers, a list of the weightmaps from each of the edit layers that affect the final result :
+			TMap<ULandscapeLayerInfoObject*, int32> PaintLayerToComponentPaintLayerRenderInfoIndex;
+			for (const FLandscapeLayer& Layer : LandscapeLayers)
+			{
+				if (Layer.bVisible && !InMergeParams.bSkipBrush)
+				{
+					// Retrieve the input textures/channels that are needed for this component and this edit layer :
+					const TArray<UTexture2D*>& ComponentEditLayerTextures = Component->GetWeightmapTextures(Layer.Guid);
+					const TArray<FWeightmapLayerAllocationInfo>& ComponentEditLayerAllocations = Component->GetWeightmapLayerAllocations(Layer.Guid);
+					int32 NumAllocations = ComponentEditLayerAllocations.Num();
+					for (int32 AllocIndex = 0; AllocIndex < NumAllocations; ++AllocIndex)
+					{
+						const FWeightmapLayerAllocationInfo& ComponentEditLayerAllocation = ComponentEditLayerAllocations[AllocIndex];
+						if (ComponentEditLayerAllocation.LayerInfo != nullptr)
+						{
+							// First, take note of the global paint layer information if it has not already been done :
+							int32* PaintLayerInfoIndex = PaintLayerToPaintLayerInfoIndex.Find(ComponentEditLayerAllocation.LayerInfo);
+							if (PaintLayerInfoIndex == nullptr)
+							{
+								FLandscapeWeightmapPaintLayerInfo NewPaintLayerInfo;
+								if (!ComponentEditLayerAllocation.LayerInfo->bNoWeightBlend)
+								{
+									NewPaintLayerInfo.Flags = EWeightmapPaintLayerFlags::IsWeightBlended;
+								}
+								if (ComponentEditLayerAllocation.LayerInfo == ALandscapeProxy::VisibilityLayer)
+								{
+									NewPaintLayerInfo.Flags = EWeightmapPaintLayerFlags::IsVisibilityLayer;
+								}
+
+								int32 NewPaintLayerInfoIndex = OutRenderThreadData.PaintLayerInfos.Add(NewPaintLayerInfo);
+								PaintLayerInfoIndex = &PaintLayerToPaintLayerInfoIndex.Add(ComponentEditLayerAllocation.LayerInfo, NewPaintLayerInfoIndex);
+								check(NewPaintLayerInfoIndex == *PaintLayerInfoIndex);
+							}
+
+							int32* ComponentPaintLayerRenderInfoIndex = PaintLayerToComponentPaintLayerRenderInfoIndex.Find(ComponentEditLayerAllocation.LayerInfo);
+							// Add one entry for each paint layer if not already there :
+							if (ComponentPaintLayerRenderInfoIndex == nullptr)
+							{
+								FComponentPaintLayerRenderInfo NewPaintLayerRenderInfo;
+								NewPaintLayerRenderInfo.Name = ComponentEditLayerAllocation.LayerInfo->LayerName.ToString();
+								// Remember which index this paint layer corresponds to in the global paint layer info array : 
+								NewPaintLayerRenderInfo.PaintLayerInfoIndex = *PaintLayerInfoIndex;
+
+								// Add it to the list and add an entry for it in the map:
+								int32 NewPaintLayerRenderInfoIndex = ComponentRenderInfo.PaintLayerRenderInfos.Add(NewPaintLayerRenderInfo);
+								ComponentPaintLayerRenderInfoIndex = &PaintLayerToComponentPaintLayerRenderInfoIndex.Add(ComponentEditLayerAllocation.LayerInfo, NewPaintLayerRenderInfoIndex);
+								check(NewPaintLayerRenderInfoIndex == *ComponentPaintLayerRenderInfoIndex);
+
+								// The next step will need an identifier for this rendered component/rendered paint layer pair:
+								TPair<ULandscapeComponent*, ULandscapeLayerInfoObject*> ComponentAndPaintLayer(Component, ComponentEditLayerAllocation.LayerInfo);
+								check(!ComponentAndPaintLayerToRenderInfoIndex.Contains(ComponentAndPaintLayer));
+								ComponentAndPaintLayerToRenderInfoIndex.Add(ComponentAndPaintLayer, FComponentAndPaintLayerRenderInfoIdentifier(OutRenderThreadData.ComponentToRenderInfos.Num() - 1, NewPaintLayerRenderInfoIndex));
+							}
+
+							FComponentPaintLayerRenderInfo& ComponentPaintLayerRenderInfo = ComponentRenderInfo.PaintLayerRenderInfos[*ComponentPaintLayerRenderInfoIndex];
+
+							// Add the texture we'll need to read from if not already there : 
+							check(ComponentEditLayerTextures.IsValidIndex(ComponentEditLayerAllocation.WeightmapTextureIndex));
+							UTexture2D* ComponentEditLayerTexture = ComponentEditLayerTextures[ComponentEditLayerAllocation.WeightmapTextureIndex];
+							int32 TextureIndexInVisibleEditLayerTextures = ComponentRenderInfo.VisibleEditLayerWeightmapTextures.AddUnique(ComponentEditLayerTexture->GetResource()->GetTexture2DResource());
+
+							// Add an entry for each edit layer that participates to this paint layer : 
+							FEditLayerInfo NewEditLayerInfo;
+							NewEditLayerInfo.SourceWeightmapTextureIndex = TextureIndexInVisibleEditLayerTextures;
+							NewEditLayerInfo.SourceWeightmapTextureChannel = ComponentEditLayerAllocation.WeightmapTextureChannel;
+							NewEditLayerInfo.Alpha = (ComponentEditLayerAllocation.LayerInfo == ALandscapeProxy::VisibilityLayer) ? 1.0f : Layer.WeightmapAlpha; // visibility can't affect or be affected by other paint layer weights
+							const bool* BlendSubstractive = Layer.WeightmapLayerAllocationBlend.Find(ComponentEditLayerAllocation.LayerInfo);
+							NewEditLayerInfo.BlendMode = ((BlendSubstractive != nullptr) && *BlendSubstractive) ? ELandscapeEditLayerWeightmapBlendMode::Substractive : ELandscapeEditLayerWeightmapBlendMode::Additive;
+
+							ComponentPaintLayerRenderInfo.VisibleEditLayerInfos.Add(NewEditLayerInfo);
+						}
+					}
+				}
+			}
+
+			// Keep track of the maximum number of weightmaps we'll need to merge for a given component :
+			OutRenderThreadData.MaxNumEditLayersTexturesToMerge = FMath::Max(ComponentRenderInfo.VisibleEditLayerWeightmapTextures.Num(), OutRenderThreadData.MaxNumEditLayersTexturesToMerge);
+		}
+	}
+
+	// Collect all UTexture2D that we need to kick off readbacks for and create a FTextureResolveInfo for each:
+	TSet<UTexture2D*> ProcessedTextures;
+	TArray<UTexture2D*> TexturesNeedingReadback;
+	{
+		TRACE_CPUPROFILER_EVENT_SCOPE(LandscapeLayers_PrepareWeightmapTextureResolveInfos);
+		for (ULandscapeComponent* LandscapeComponentToResolve : InUpdateLayersContentContext.LandscapeComponentsWeightmapsToResolve)
+		{
+			ALandscapeProxy* Proxy = LandscapeComponentToResolve->GetLandscapeProxy();
+			check(Proxy != nullptr);
+			const TArray<UTexture2D*>& ComponentBaseWeightmapTextures = LandscapeComponentToResolve->GetWeightmapTextures();
+			for (UTexture2D* ComponentWeightmap : ComponentBaseWeightmapTextures)
+			{
+				bool bIsAlreadyInSet = false;
+				ProcessedTextures.Add(ComponentWeightmap, &bIsAlreadyInSet);
+				if (!bIsAlreadyInSet)
+				{
+					FTexture2DResource* WeightmapResource = ComponentWeightmap->GetResource()->GetTexture2DResource();
+
+					FTextureResolveInfo NewTextureResolveInfo(OutRenderThreadData.ComponentToRenderInfos.Num());
+					NewTextureResolveInfo.TextureToResolveInfoIndex = OutRenderThreadData.TextureToResolveInfos.Num(); // It will be added at the end of the array 
+					NewTextureResolveInfo.Texture = WeightmapResource;
+
+					bool bDoResolve = false;
+
+					// Start preparing the information we need for resolving this texture : we'll need to know how to reconstruct each of its channels :
+					TObjectPtr<ULandscapeWeightmapUsage>* Usage = Proxy->WeightmapUsageMap.Find(ComponentWeightmap);
+					check(Usage && *Usage);
+					// Iterate over all components that participate to this texture's data :
+					for (ULandscapeComponent* SourceComponent : (*Usage)->GetUniqueValidComponents())
+					{
+						// Iterate over all of its allocations to find out which corresponds to this texture :
+						const TArray<UTexture2D*>& SourceComponentBaseWeightmapTextures = SourceComponent->GetWeightmapTextures();
+						const TArray<FWeightmapLayerAllocationInfo>& SourceComponentBaseWeightmapAllocations = SourceComponent->GetWeightmapLayerAllocations();
+						for (const FWeightmapLayerAllocationInfo& AllocationInfo : SourceComponentBaseWeightmapAllocations)
+						{
+							check(SourceComponentBaseWeightmapTextures.IsValidIndex(AllocationInfo.WeightmapTextureIndex));
+							UTexture2D* SourceComponentWeightmap = SourceComponentBaseWeightmapTextures[AllocationInfo.WeightmapTextureIndex];
+							// Same texture, we'll need this allocation
+							if (AllocationInfo.IsAllocated() && (SourceComponentWeightmap == ComponentWeightmap))
+							{
+								TPair<ULandscapeComponent*, ULandscapeLayerInfoObject*> ComponentAndPaintLayer(SourceComponent, AllocationInfo.LayerInfo);
+								FComponentAndPaintLayerRenderInfoIdentifier* ComponentAndPaintLayerIdentifier = ComponentAndPaintLayerToRenderInfoIndex.Find(ComponentAndPaintLayer);
+								// All components needed to recompose this weightmap should be in ComponentToRenderInfos and have a valid paint layer info there:
+								check((ComponentAndPaintLayerIdentifier != nullptr)
+									&& OutRenderThreadData.ComponentToRenderInfos.IsValidIndex(ComponentAndPaintLayerIdentifier->ComponentIndex)
+									&& OutRenderThreadData.ComponentToRenderInfos[ComponentAndPaintLayerIdentifier->ComponentIndex].PaintLayerRenderInfos.IsValidIndex(ComponentAndPaintLayerIdentifier->PaintLayerIndex));
+								NewTextureResolveInfo.SetPerChannelSourceInfo(AllocationInfo.WeightmapTextureChannel, *ComponentAndPaintLayerIdentifier);
+
+								// At least one channel to resolve : we need to resolve the texture : 
+								bDoResolve = true;
+							}
+						}
+					}
+
+					if (bDoResolve)
+					{
+						// Setup the CPU readback if it does not already exist:
+						FLandscapeEditLayerReadback** CPUReadback = Proxy->WeightmapsCPUReadback.Find(ComponentWeightmap);
+						if (CPUReadback == nullptr)
+						{
+							// Lazily create the readback objects as required (ReallocateLayersWeightmaps might have created new weightmaps)
+							FLandscapeEditLayerReadback* NewCPUReadback = new FLandscapeEditLayerReadback();
+							const uint8* LockedMip = ComponentWeightmap->Source.LockMipReadOnly(0);
+							const uint32 Hash = FLandscapeEditLayerReadback::CalculateHash(LockedMip, ComponentWeightmap->GetSizeX() * ComponentWeightmap->GetSizeY() * sizeof(FColor));
+							NewCPUReadback->SetHash(Hash);
+							ComponentWeightmap->Source.UnlockMip(0);
+							CPUReadback = &Proxy->WeightmapsCPUReadback.Add(ComponentWeightmap, NewCPUReadback);
+						}
+						check(*CPUReadback != nullptr);
+
+						// Register the CPU readback and add to our list of textures to resolve :
+						NewTextureResolveInfo.CPUReadback = *CPUReadback;
+						OutRenderThreadData.TextureToResolveInfos.Add(NewTextureResolveInfo);
+
+						TexturesNeedingReadback.Add(ComponentWeightmap);
+					}
+				}
+			}
+		}
+		check(OutRenderThreadData.TextureToResolveInfos.Num() == TexturesNeedingReadback.Num());
+	}
+
+	// Prepare the texture resolve batches: 
+	{
+		TRACE_CPUPROFILER_EVENT_SCOPE(LandscapeLayers_PrepareWeightmapTextureResolveInfoBatches);
+
+		int32 MaxComponentsPerResolveBatch = CVarLandscapeEditLayersMaxComponentsPerWeightmapResolveBatch.GetValueOnGameThread();
+
+		// Copy the texture infos because TextureToResolveInfos indices need to remain stable at this point :
+		TArray<FTextureResolveInfo> RemainingTextureToResolveInfos = OutRenderThreadData.TextureToResolveInfos;
+		if (!RemainingTextureToResolveInfos.IsEmpty())
+		{
+			TBitArray<TInlineAllocator<1>> TempBitArray;
+			TempBitArray.Reserve(OutRenderThreadData.ComponentToRenderInfos.Num());
+
+			while (!RemainingTextureToResolveInfos.IsEmpty())
+			{
+				const FTextureResolveInfo& TextureResolveInfo = RemainingTextureToResolveInfos.Pop(/*bAllowShrinking = */false);
+
+				int32 BestBatchIndex = INDEX_NONE;
+				int32 MinNumComponents = MAX_int32;
+
+				// Iterate through all all batches and try to find which would be able to accept it and amongst those, which it would share the most components to render with:
+				int32 NumBatches = OutRenderThreadData.TextureResolveBatchInfos.Num();
+				for (int32 BatchIndex = 0; BatchIndex < NumBatches; ++BatchIndex)
+				{
+					const FTextureResolveBatchInfo& Batch = OutRenderThreadData.TextureResolveBatchInfos[BatchIndex];
+					TempBitArray = TBitArray<>::BitwiseOR(Batch.ComponentToRenderInfoBitIndices, TextureResolveInfo.ComponentToRenderInfoBitIndices, EBitwiseOperatorFlags::MinSize);
+
+					// If after adding its components, the batch still has less than MaxComponentsPerResolveBatch components to render, it can accept it
+					int32 NumComponentsAfter = TempBitArray.CountSetBits();
+					if (NumComponentsAfter <= MaxComponentsPerResolveBatch)
+					{
+						// Is this the best candidate so far?
+						if (NumComponentsAfter < MinNumComponents)
+						{
+							BestBatchIndex = BatchIndex;
+							MinNumComponents = NumComponentsAfter;
+						}
+
+						// If the number of components after addition of this texture is unchanged, it's a perfect match, we won't ever find a better batch so just stop there for this texture:
+						if (NumComponentsAfter == Batch.ComponentToRenderInfoBitIndices.CountSetBits())
+						{
+							break;
+						}
+					}
+				}
+
+				// If we have found a batch, just add the texture to it, otherwise, add a new batch:
+				FTextureResolveBatchInfo& SelectedBatch = (BestBatchIndex != INDEX_NONE) ? OutRenderThreadData.TextureResolveBatchInfos[BestBatchIndex]
+					: OutRenderThreadData.TextureResolveBatchInfos.Add_GetRef(FTextureResolveBatchInfo(OutRenderThreadData.ComponentToRenderInfos.Num(), /*InBatchIndex = */OutRenderThreadData.TextureResolveBatchInfos.Num()));
+
+				SelectedBatch.AddTexture(TextureResolveInfo);
+				check(SelectedBatch.ComponentToRenderInfoBitIndices.CountSetBits() <= MaxComponentsPerResolveBatch);
+
+				// Keep track of the maximum number of scratch texture arrays we'll need for any given batch :
+				OutRenderThreadData.MaxNumWeightmapArraysPerResolveTextureBatch = FMath::Max(SelectedBatch.ComponentToRenderInfoBitIndices.CountSetBits(), OutRenderThreadData.MaxNumWeightmapArraysPerResolveTextureBatch);
+			}
+		}
+	}
+
+	// Finalize :
+	{
+		OutRenderThreadData.DeferredCopyReadbackTextures = PrepareLandscapeLayersCopyReadbackTextureParams(InUpdateLayersContentContext.MapHelper, TexturesNeedingReadback, /*bWeightmaps = */true);
+
+		// Finally, update the material instances to take into account potentially new material combinations : 
+		UpdateLayersMaterialInstances(InUpdateLayersContentContext.LandscapeComponentsWeightmapsToResolve);
+	}
+}
+
+int32 ALandscape::PerformLayersWeightmapsLocalMerge(FUpdateLayersContentContext& InUpdateLayersContentContext, const FEditLayersWeightmapMergeParams& InMergeParams)
+{
+	using namespace EditLayersWeightmapLocalMerge_RenderThread;
+
+	TRACE_CPUPROFILER_EVENT_SCOPE(LandscapeLayers_PerformLayersWeightmapsLocalMerge);
+
+	// We need to perform layer reallocations before doing anything, since additional weightmaps might be added in the process, which might result in new components to render/resolve : 
+	// TODO [jonathan.bard] : this is only possible to do so in local merge, since we don't yet support BP brushes, for which we need to call Render() in order to be able to know which layer allocations
+	//  they need. As such, the global merge path is broken and cannot be fixed unless the BP brush interface is changed in order to inform the system, by advance, about what paint layer it needs
+	// Make sure we have proper textures+allocations for all the final weightmaps we're about to resolve :
+	ReallocateLayersWeightmaps(InUpdateLayersContentContext, /*InBrushRequiredAllocations = */TArray<ULandscapeLayerInfoObject*>());
+
+	FMergeInfo RenderThreadData;
+	PrepareLayersWeightmapsLocalMergeRenderThreadData(InUpdateLayersContentContext, InMergeParams, RenderThreadData);
+
+	if (RenderThreadData.NeedsMerge())
+	{
+		ENQUEUE_RENDER_COMMAND(PerformLayersWeightmapsLocalMerge)([RenderThreadData](FRHICommandListImmediate& RHICmdList)
+		{
+			FMemMark Mark(FMemStack::Get());
+			FRDGBuilder GraphBuilder(RHICmdList, RDG_EVENT_NAME("PerformLayersWeightmapsLocalMerge"));
+
+			// Prepare the GPU resources we will use during the local merge : 
+			FRDGResources RDGResources;
+			PrepareLayersWeightmapsLocalMergeRDGResources(RenderThreadData, GraphBuilder, RDGResources);
+
+			// Get a list of all external textures (weightmaps) we will manipulate during the local merge : 
+			TMap<FTexture2DResource*, FLandscapeRDGTrackedTexture> TrackedTextures;
+			GatherLayersWeightmapsLocalMergeRDGTextures(RenderThreadData, TrackedTextures);
+
+			// Start tracking those in the render graph :
+			TrackLandscapeRDGTextures(GraphBuilder, TrackedTextures);
+
+			// Process the components batch by batch in order to avoid over-allocating temporary textures :
+			for (const FTextureResolveBatchInfo& TextureResolveBatchInfo : RenderThreadData.TextureResolveBatchInfos)
+			{
+				RDG_EVENT_SCOPE(GraphBuilder, "Process batch %d", TextureResolveBatchInfo.BatchIndex);
+
+				// Perform all edit layers merges, for all paint layers of all the components in that batch :
+				MergeEditLayersWeightmapsForBatch(TextureResolveBatchInfo, RenderThreadData, TrackedTextures, GraphBuilder, RDGResources);
+
+				// Pack the temporary weightmaps and generate mips on the final texture :
+				FinalizeAndResolveWeightmapsForBatch(TextureResolveBatchInfo, RenderThreadData, TrackedTextures, GraphBuilder, RDGResources);
+			}
+
+			GraphBuilder.Execute();
+		});
+
+		ExecuteCopyToReadbackTexture(RenderThreadData.DeferredCopyReadbackTextures);
+	}
+
+	return InMergeParams.WeightmapUpdateModes;
+}
+
+int32 ALandscape::PerformLayersWeightmapsGlobalMerge(FUpdateLayersContentContext& InUpdateLayersContentContext, const FEditLayersWeightmapMergeParams& InMergeParams)
+{
 	ULandscapeInfo* Info = GetLandscapeInfo();
+	check(Info != nullptr);
+	check(WeightmapRTList.Num() > 0);
 
-	if (WeightmapUpdateModes == 0 && !bForceRender)
+	FIntRect LandscapeExtent;
+	if (!Info->GetLandscapeExtent(LandscapeExtent.Min.X, LandscapeExtent.Min.Y, LandscapeExtent.Max.X, LandscapeExtent.Max.Y))
 	{
 		return 0;
 	}
 
-	if (InLandscapeComponentsToResolve.Num() == 0 || Info == nullptr)
-	{
-		return WeightmapUpdateModes;
-	}
-			
-	TArray<ULandscapeComponent*> ComponentThatNeedMaterialRebuild;
 	TArray<ULandscapeLayerInfoObject*> BrushRequiredAllocations;
 	int32 LayerCount = Info->Layers.Num() + 1; // due to visibility being stored at 0
-	bool ClearFlagsAfterUpdate = true;
 
-	if (WeightmapUpdateModes || bForceRender)
+	if (InMergeParams.WeightmapUpdateModes || InMergeParams.bForceRender)
 	{
-		RenderCaptureInterface::FScopedCapture RenderCapture((RenderCaptureLayersNextWeightmapDraws != 0), TEXT("LandscapeLayersWeightmapCapture"));
-
-		FIntRect LandscapeExtent;
-		if (!Info->GetLandscapeExtent(LandscapeExtent.Min.X, LandscapeExtent.Min.Y, LandscapeExtent.Max.X, LandscapeExtent.Max.Y))
-		{
-			return 0;
-		}
-		
-		check(WeightmapRTList.Num() > 0);
-
 		UTextureRenderTarget2D* LandscapeScratchRT1 = WeightmapRTList[(int32)EWeightmapRTType::WeightmapRT_Scratch1];
 		UTextureRenderTarget2D* LandscapeScratchRT2 = WeightmapRTList[(int32)EWeightmapRTType::WeightmapRT_Scratch2];
 		UTextureRenderTarget2D* LandscapeScratchRT3 = WeightmapRTList[(int32)EWeightmapRTType::WeightmapRT_Scratch3];
@@ -4266,7 +6541,6 @@
 
 		bool bHasWeightmapData = false;
 		bool bFirstLayer = true;
-		TMap<ULandscapeLayerInfoObject*, bool> WeightmapLayersBlendSubstractive;
 
 		for (FLandscapeLayer& Layer : LandscapeLayers)
 		{
@@ -4276,12 +6550,12 @@
 			TMap<ULandscapeLayerInfoObject*, int32> LayerInfoObjects; // <LayerInfoObj, LayerIndex>
 
 			// Determine if some brush want to write to layer that we have currently no data on
-			if (Layer.bVisible && !bSkipBrush)
+			if (Layer.bVisible && !InMergeParams.bSkipBrush)
 			{
 				for (int32 LayerInfoSettingsIndex = 0; LayerInfoSettingsIndex < Info->Layers.Num(); ++LayerInfoSettingsIndex)
 				{
 					const FLandscapeInfoLayerSettings& InfoLayerSettings = Info->Layers[LayerInfoSettingsIndex];
-					
+
 					// It is possible that no layer info has been assigned so that InfoLayerSettings.LayerInfoObj == nullptr. In that case we don't consider it the layer here.
 					if (InfoLayerSettings.LayerInfoObj != nullptr)
 					{
@@ -4311,7 +6585,8 @@
 
 				// Prepare compute shader data
 				TArray<FLandscapeLayerWeightmapExtractMaterialLayersComponentData> ComponentsData;
-				PrepareComponentDataToExtractMaterialLayersCS(InLandscapeComponentsToRender, Layer, CurrentWeightmapToProcessIndex, LandscapeExtent.Min, WeightmapScratchExtractLayerTextureResource, ComponentsData, LayerInfoObjects);
+				PrepareComponentDataToExtractMaterialLayersCS(InUpdateLayersContentContext.LandscapeComponentsWeightmapsToRender, Layer, CurrentWeightmapToProcessIndex, LandscapeExtent.Min, WeightmapScratchExtractLayerTextureResource,
+					ComponentsData, LayerInfoObjects);
 
 				HasFoundWeightmapToProcess = ComponentsData.Num() > 0;
 
@@ -4338,6 +6613,7 @@
 							DestDebugName = FString::Printf(TEXT("Weight: %s Clear CurrentProcLayerWeightmapAllLayersResource %d, "), *Layer.Name.ToString(), LayerIndex);
 
 							FLandscapeLayersCopyTextureParams& CopyTextureParams = DeferredCopyTextures.Add_GetRef(FLandscapeLayersCopyTextureParams(SourceDebugName, LandscapeScratchRT1->GameThread_GetRenderTargetResource(), DestDebugName, CurrentLayersWeightmapAllMaterialLayersResource));
+							CopyTextureParams.DestAccess = ERHIAccess::UAVMask;
 							CopyTextureParams.DestArrayIndex = LayerIndex;
 						}
 					}
@@ -4349,7 +6625,7 @@
 				if (ComponentsData.Num() > 0)
 				{
 					PrintLayersDebugTextureResource(FString::Printf(TEXT("LS Weight: %s WeightmapScratchTexture %s"), *Layer.Name.ToString(), TEXT("WeightmapScratchTextureResource")), WeightmapScratchExtractLayerTextureResource, 0, false);
-										
+
 					FLandscapeLayerWeightmapExtractMaterialLayersComputeShaderParameters CSExtractLayersShaderParams;
 					CSExtractLayersShaderParams.AtlasWeightmapsPerLayer = CurrentLayersWeightmapAllMaterialLayersResource;
 					CSExtractLayersShaderParams.ComponentWeightmapResource = WeightmapScratchExtractLayerTextureResource;
@@ -4380,12 +6656,16 @@
 					int32 LayerIndex = LayerInfoObject.Value;
 					ULandscapeLayerInfoObject* LayerInfoObj = LayerInfoObject.Key;
 
+					FString ProfilingEventName = FString::Printf(TEXT("LS Weight: %s PaintLayer: %s"), *Layer.Name.ToString(), *LayerInfoObj->LayerName.ToString());
+					SCOPED_DRAW_EVENTF_GAMETHREAD(LandscapeLayers, *ProfilingEventName);
+
 					// Copy the layer we are working on
 					SourceDebugName = FString::Printf(TEXT("Weight: %s PaintLayer: %s, CurrentProcLayerWeightmapAllLayersResource"), *Layer.Name.ToString(), *LayerInfoObj->LayerName.ToString());
 					DestDebugName = LandscapeScratchRT1->GetName();
 
 					{
 						FLandscapeLayersCopyTextureParams CopyTextureParams(SourceDebugName, CurrentLayersWeightmapAllMaterialLayersResource, DestDebugName, LandscapeScratchRT1->GameThread_GetRenderTargetResource());
+						CopyTextureParams.SourceAccess = ERHIAccess::UAVMask;
 						CopyTextureParams.SourceArrayIndex = LayerIndex;
 						ExecuteCopyLayersTexture({ MoveTemp(CopyTextureParams) });
 					}
@@ -4395,8 +6675,8 @@
 					PSShaderParams.LayerVisible = Layer.bVisible;
 					PSShaderParams.LayerAlpha = LayerInfoObj == ALandscapeProxy::VisibilityLayer ? 1.0f : Layer.WeightmapAlpha; // visibility can't be affected by weight
 
-					DrawWeightmapComponentsToRenderTarget(FString::Printf(TEXT("LS Weight: %s Paint: %s += -> %s"), *Layer.Name.ToString(), *LayerInfoObj->LayerName.ToString(), *LandscapeScratchRT1->GetName(), *LandscapeScratchRT2->GetName()),
-														  InLandscapeComponentsToRender, LandscapeExtent.Min, LandscapeScratchRT1, nullptr, LandscapeScratchRT2, ERTDrawingType::RTAtlas, true, PSShaderParams, 0);
+					DrawWeightmapComponentsToRenderTarget(FString::Printf(TEXT("LS Weight: %s PaintLayer: %s, %s += -> %s"), *Layer.Name.ToString(), *LayerInfoObj->LayerName.ToString(), *LandscapeScratchRT1->GetName(), *LandscapeScratchRT2->GetName()),
+						InUpdateLayersContentContext.LandscapeComponentsWeightmapsToRender, LandscapeExtent.Min, LandscapeScratchRT1, nullptr, LandscapeScratchRT2, ERTDrawingType::RTAtlas, true, PSShaderParams, 0);
 
 					PSShaderParams.ApplyLayerModifiers = false;
 
@@ -4406,6 +6686,7 @@
 
 					{
 						FLandscapeLayersCopyTextureParams CopyTextureParams(SourceDebugName, CombinedLayersWeightmapAllMaterialLayersResource, DestDebugName, LandscapeScratchRT1->GameThread_GetRenderTargetResource());
+						CopyTextureParams.SourceAccess = ERHIAccess::UAVMask;
 						CopyTextureParams.SourceArrayIndex = LayerIndex;
 						ExecuteCopyLayersTexture({ MoveTemp(CopyTextureParams) });
 					}
@@ -4418,16 +6699,10 @@
 					{
 						const bool* BlendSubstractive = Layer.WeightmapLayerAllocationBlend.Find(LayerInfoObj);
 						PSShaderParams.OutputAsSubstractive = BlendSubstractive != nullptr ? *BlendSubstractive : false;
-
-						if (PSShaderParams.OutputAsSubstractive)
-						{
-							bool& IsSubstractiveBlend = WeightmapLayersBlendSubstractive.FindOrAdd(LayerInfoObj);
-							IsSubstractiveBlend = true;
-						}
 					}
 
 					DrawWeightmapComponentsToRenderTarget(FString::Printf(TEXT("LS Weight: %s PaintLayer: %s, %s += -> Combined %s"), *Layer.Name.ToString(), *LayerInfoObj->LayerName.ToString(), *LandscapeScratchRT2->GetName(), *LandscapeScratchRT3->GetName()),
-														  InLandscapeComponentsToRender, LandscapeExtent.Min, LandscapeScratchRT2, bFirstLayer ? nullptr : LandscapeScratchRT1, LandscapeScratchRT3, ERTDrawingType::RTAtlasToNonAtlas, true, PSShaderParams, 0);
+						InUpdateLayersContentContext.LandscapeComponentsWeightmapsToRender, LandscapeExtent.Min, LandscapeScratchRT2, bFirstLayer ? nullptr : LandscapeScratchRT1, LandscapeScratchRT3, ERTDrawingType::RTAtlasToNonAtlas, true, PSShaderParams, 0);
 
 					PSShaderParams.OutputAsSubstractive = false;
 
@@ -4435,7 +6710,7 @@
 					DestDebugName = TEXT("CombinedProcLayerWeightmap");
 
 					// Handle brush blending
-					if (Layer.bVisible && !bSkipBrush)
+					if (Layer.bVisible && !InMergeParams.bSkipBrush)
 					{
 						// Draw each brushes				
 						for (int32 i = 0; i < Layer.Brushes.Num(); ++i)
@@ -4471,14 +6746,16 @@
 						DestDebugName = TEXT("CombinedProcLayerWeightmap");
 
 						FLandscapeLayersCopyTextureParams CopyTextureParams(SourceDebugName, LandscapeScratchRT3->GameThread_GetRenderTargetResource(), DestDebugName, CombinedLayersWeightmapAllMaterialLayersResource);
+						CopyTextureParams.DestAccess = ERHIAccess::UAVMask;
 						CopyTextureParams.DestArrayIndex = LayerIndex;
 						ExecuteCopyLayersTexture({ MoveTemp(CopyTextureParams) });
 					}
 
 					DrawWeightmapComponentsToRenderTarget(FString::Printf(TEXT("LS Weight: %s Combined Scratch No Border to %s Combined Scratch with Border"), *LandscapeScratchRT3->GetName(), *LandscapeScratchRT1->GetName()),
-						InLandscapeComponentsToRender, LandscapeExtent.Min, LandscapeScratchRT3, nullptr, LandscapeScratchRT1, ERTDrawingType::RTNonAtlasToAtlas, true, PSShaderParams, 0);
+						InUpdateLayersContentContext.LandscapeComponentsWeightmapsToRender, LandscapeExtent.Min, LandscapeScratchRT3, nullptr, LandscapeScratchRT1, ERTDrawingType::RTNonAtlasToAtlas, true, PSShaderParams, 0);
 
 					FLandscapeLayersCopyTextureParams CopyTextureParams(SourceDebugName, LandscapeScratchRT1->GameThread_GetRenderTargetResource(), DestDebugName, CombinedLayersWeightmapAllMaterialLayersResource);
+					CopyTextureParams.DestAccess = ERHIAccess::UAVMask;
 					CopyTextureParams.DestArrayIndex = LayerIndex;
 					ExecuteCopyLayersTexture({ MoveTemp(CopyTextureParams) });
 				}
@@ -4489,21 +6766,15 @@
 			bFirstLayer = false;
 		}
 
-		ReallocateLayersWeightmaps(InLandscapeComponentsToResolve, BrushRequiredAllocations);
-		TSet<UTexture2D*> ToResolve;
-		for (ULandscapeComponent* LandscapeComponentToResolve : InLandscapeComponentsToResolve)
-		{
-			const TArray<UTexture2D*>& ComponentWeightmapTextures = LandscapeComponentToResolve->GetWeightmapTextures();
-			for (UTexture2D* WeightmapTexture : ComponentWeightmapTextures)
-			{
-				ToResolve.Add(WeightmapTexture);
-			}
-		}
+		ReallocateLayersWeightmaps(InUpdateLayersContentContext, BrushRequiredAllocations);
 
+		// List of UTexture2D that we need to kick off readbacks for :
+		TArray<UTexture2D*> TexturesNeedingReadback;
+
 		if (bHasWeightmapData)
 		{
 			// Lazily create CPU read back objects as required
-			for (ULandscapeComponent* Component : InLandscapeComponentsToRender)
+			for (ULandscapeComponent* Component : InUpdateLayersContentContext.LandscapeComponentsWeightmapsToRender)
 			{
 				const TArray<UTexture2D*>& ComponentWeightmapTextures = Component->GetWeightmapTextures();
 
@@ -4523,7 +6794,7 @@
 					}
 				}
 			}
-			
+
 			int8 CurrentWeightmapToProcessIndex = 0;
 			bool HasFoundWeightmapToProcess = true; // try processing at least once
 
@@ -4536,7 +6807,7 @@
 			while (HasFoundWeightmapToProcess)
 			{
 				TArray<FLandscapeLayerWeightmapPackMaterialLayersComponentData> PackLayersComponentsData;
-				PrepareComponentDataToPackMaterialLayersCS(CurrentWeightmapToProcessIndex, LandscapeExtent.Min, InLandscapeComponentsToRender, ProcessedWeightmaps, ProcessedCPUReadbackTextures, PackLayersComponentsData);
+				PrepareComponentDataToPackMaterialLayersCS(CurrentWeightmapToProcessIndex, LandscapeExtent.Min, InUpdateLayersContentContext.LandscapeComponentsWeightmapsToRender, ProcessedWeightmaps, ProcessedCPUReadbackTextures, PackLayersComponentsData);
 				HasFoundWeightmapToProcess = PackLayersComponentsData.Num() > 0;
 
 				// Perform the compute shader
@@ -4565,9 +6836,9 @@
 
 					for (int32 i = 0; i < PackLayersComponentsData.Num(); ++i)
 					{
-						check(ComponentY+ComponentSize <= WeightmapScratchPackLayerTextureResource->GetSizeY()); // This should never happen as it would be a bug in the algo
+						check(ComponentY + ComponentSize <= WeightmapScratchPackLayerTextureResource->GetSizeY()); // This should never happen as it would be a bug in the algo
 
-						if (ComponentX+ComponentSize > WeightmapScratchPackLayerTextureResource->GetSizeX())
+						if (ComponentX + ComponentSize > WeightmapScratchPackLayerTextureResource->GetSizeX())
 						{
 							ComponentY += ComponentSize;
 							ComponentX = 0;
@@ -4583,6 +6854,13 @@
 
 					CopyTexturePS(SourceDebugName, EmptyRT->GameThread_GetRenderTargetResource(), DestDebugName, WeightmapScratchPackLayerTextureResource);
 
+					ENQUEUE_RENDER_COMMAND(LandscapeLayers_TransitionPackLayerResources)(
+						[this](FRHICommandListImmediate& RHICmdList) mutable
+					{
+						RHICmdList.Transition(FRHITransitionInfo(CombinedLayersWeightmapAllMaterialLayersResource->TextureRHI, ERHIAccess::UAVMask, ERHIAccess::SRVMask));
+						RHICmdList.Transition(FRHITransitionInfo(WeightmapScratchPackLayerTextureResource->TextureRHI, ERHIAccess::RTV, ERHIAccess::UAVMask));
+					});
+
 					FLandscapeLayerWeightmapPackMaterialLayersComputeShaderParameters CSPackLayersShaderParams;
 					CSPackLayersShaderParams.AtlasWeightmapsPerLayer = CombinedLayersWeightmapAllMaterialLayersResource;
 					CSPackLayersShaderParams.ComponentWeightmapResource = WeightmapScratchPackLayerTextureResource;
@@ -4603,18 +6881,32 @@
 
 					SourceDebugName = TEXT("WeightmapScratchTexture");
 					DestDebugName = CurrentRT->GetName();
-					
+
+					ENQUEUE_RENDER_COMMAND(LandscapeLayers_TransitionCopyResources)(
+						[this, CurrentRT](FRHICommandListImmediate& RHICmdList) mutable
+					{
+						RHICmdList.Transition(FRHITransitionInfo(CombinedLayersWeightmapAllMaterialLayersResource->TextureRHI, ERHIAccess::SRVMask, ERHIAccess::UAVMask));
+						RHICmdList.Transition(FRHITransitionInfo(WeightmapScratchPackLayerTextureResource->TextureRHI, ERHIAccess::UAVMask, ERHIAccess::SRVMask));
+						RHICmdList.Transition(FRHITransitionInfo(CurrentRT->GetRenderTargetResource()->TextureRHI, ERHIAccess::SRVMask, ERHIAccess::RTV));
+					});
+
 					CopyTexturePS(SourceDebugName, WeightmapScratchPackLayerTextureResource, DestDebugName, CurrentRT->GameThread_GetRenderTargetResource());
+
+					ENQUEUE_RENDER_COMMAND(LandscapeLayers_TransitionMip0)(
+						[this, CurrentRT](FRHICommandListImmediate& RHICmdList) mutable
+					{
+						RHICmdList.Transition(FRHITransitionInfo(CurrentRT->GetRenderTargetResource()->TextureRHI, ERHIAccess::RTV, ERHIAccess::SRVMask));
+						RHICmdList.Transition(FRHITransitionInfo(WeightmapScratchPackLayerTextureResource->TextureRHI, ERHIAccess::SRVMask, ERHIAccess::RTV));
+					});
 					DrawWeightmapComponentToRenderTargetMips(WeightmapTextureOutputOffset, CurrentRT, true, PSShaderParams);
 
 					int32 StartTextureIndex = NextTextureIndexToProcess;
 
 					TArray<FLandscapeLayersCopyTextureParams> DeferredCopyTextures;
-					TArray<FLandscapeLayersCopyReadbackTextureParams> DeferredCopyReadbackTextures;
 					for (; NextTextureIndexToProcess < ProcessedWeightmaps.Num(); ++NextTextureIndexToProcess)
 					{
 						UTexture2D* WeightmapTexture = ProcessedWeightmaps[NextTextureIndexToProcess];
-						if (!ToResolve.Contains(WeightmapTexture))
+						if (!InUpdateLayersContentContext.WeightmapsToResolve.Contains(WeightmapTexture))
 						{
 							continue;
 						}
@@ -4645,33 +6937,69 @@
 							}
 						}
 
-						// Init the CPU read back parameters. This includes a context containing the current component states.
-						const int32 ReadbackIndex = DeferredCopyReadbackTextures.Add(FLandscapeLayersCopyReadbackTextureParams(WeightmapTexture, ProcessedCPUReadbackTextures[NextTextureIndexToProcess]));
-						TArray<ULandscapeComponent*> const* Components = MapHelper.WeightmapToComponents.Find(WeightmapTexture);
-						if (Components != nullptr)
-						{
-							for (ULandscapeComponent* Component : *Components)
-							{
-								const FIntPoint ComponentKey = Component->GetSectionBase() / Component->ComponentSizeQuads;
-								const int32 ComponentFlags = Component->GetLayerUpdateFlagPerMode();
-								DeferredCopyReadbackTextures[ReadbackIndex].Context.Add(FLandscapeEditLayerReadback::FComponentReadbackContext(ComponentKey, ComponentFlags));
-							}
-						}
+						check(!TexturesNeedingReadback.Contains(WeightmapTexture));
+						TexturesNeedingReadback.Add(WeightmapTexture);
 					}
+
 					ExecuteCopyLayersTexture(MoveTemp(DeferredCopyTextures));
-					ExecuteCopyToReadbackTexture(DeferredCopyReadbackTextures);
 				}
 
 				++CurrentWeightmapToProcessIndex;
 			}
 		}
 
-		UpdateLayersMaterialInstances(InLandscapeComponentsToResolve);
+		// Prepare the UTexture2D readbacks we'll need to perform :
+		TArray<FLandscapeLayersCopyReadbackTextureParams> DeferredCopyReadbackTextures = PrepareLandscapeLayersCopyReadbackTextureParams(InUpdateLayersContentContext.MapHelper, TexturesNeedingReadback, /*bWeightmaps = */true);
+		ExecuteCopyToReadbackTexture(DeferredCopyReadbackTextures);
 
-		RenderCaptureLayersNextWeightmapDraws = 0;
+		UpdateLayersMaterialInstances(InUpdateLayersContentContext.LandscapeComponentsWeightmapsToResolve);
 	}
 
-	return WeightmapUpdateModes;
+	return InMergeParams.WeightmapUpdateModes;
 }
 
+int32 ALandscape::RegenerateLayersWeightmaps(FUpdateLayersContentContext& InUpdateLayersContentContext)
+{
+	TRACE_CPUPROFILER_EVENT_SCOPE(LandscapeLayers_RegenerateLayersWeightmaps);
+	SCOPE_CYCLE_COUNTER(STAT_LandscapeLayersRegenerateWeightmaps);
+	const int32 AllWeightmapUpdateModes = (ELandscapeLayerUpdateMode::Update_Weightmap_All | ELandscapeLayerUpdateMode::Update_Weightmap_Editing | ELandscapeLayerUpdateMode::Update_Weightmap_Editing_NoCollision);
+	const int32 WeightmapUpdateModes = LayerContentUpdateModes & AllWeightmapUpdateModes;
+	const bool bSkipBrush = CVarLandscapeLayerBrushOptim.GetValueOnAnyThread() == 1 && ((WeightmapUpdateModes & AllWeightmapUpdateModes) == ELandscapeLayerUpdateMode::Update_Weightmap_Editing);
+	const bool bForceRender = CVarForceLayersUpdate.GetValueOnAnyThread() != 0;
+
+	ULandscapeInfo* Info = GetLandscapeInfo();
+
+	if (WeightmapUpdateModes == 0 && !bForceRender)
+	{
+		return 0;
+	}
+
+	if (InUpdateLayersContentContext.LandscapeComponentsWeightmapsToResolve.Num() == 0 || Info == nullptr)
+	{
+		return WeightmapUpdateModes;
+	}
+
+	if (WeightmapUpdateModes || bForceRender)
+	{
+		RenderCaptureInterface::FScopedCapture RenderCapture((RenderCaptureLayersNextWeightmapDraws != 0), TEXT("LandscapeLayersWeightmapCapture"));
+		RenderCaptureLayersNextWeightmapDraws = FMath::Max(0, RenderCaptureLayersNextWeightmapDraws - 1);
+
+		FEditLayersWeightmapMergeParams MergeParams;
+		MergeParams.WeightmapUpdateModes = WeightmapUpdateModes;
+		MergeParams.bForceRender = bForceRender;
+		MergeParams.bSkipBrush = bSkipBrush;
+
+		if (SupportsEditLayersLocalMerge())
+		{
+			return PerformLayersWeightmapsLocalMerge(InUpdateLayersContentContext, MergeParams);
+		}
+		else
+		{
+			return PerformLayersWeightmapsGlobalMerge(InUpdateLayersContentContext, MergeParams);
+		}
+	}
+
+	return 0;
+}
+
 void ALandscape::UpdateForChangedWeightmaps(ULandscapeComponent* InComponent, const FLandscapeEditLayerReadbackResult& InReadbackResult)
@@ -4697,6 +7025,18 @@
 	}
 }
 
+void ULandscapeComponent::GetUsedPaintLayers(const FGuid& InLayerGuid, TArray<ULandscapeLayerInfoObject*>& OutUsedLayerInfos) const
+{
+	const TArray<FWeightmapLayerAllocationInfo>& AllocInfos = GetWeightmapLayerAllocations(InLayerGuid);
+	for (const FWeightmapLayerAllocationInfo& AllocInfo : AllocInfos)
+	{
+		if (AllocInfo.LayerInfo != nullptr)
+		{
+			OutUsedLayerInfos.AddUnique(AllocInfo.LayerInfo);
+		}
+	}
+}
+
 uint32 ULandscapeComponent::ComputeWeightmapsHash()
 {
 	uint32 Hash = 0;
@@ -4920,7 +7260,7 @@
 			}
 		}
 	}
-		
+
 	// Weightmaps shouldn't invalidate lighting
 	const bool bInvalidateLightingCache = false;
 	InvalidateGeneratedComponentData(ChangedComponents, bInvalidateLightingCache);
@@ -5020,7 +7360,7 @@
 			return true;
 		}
 	}
-		
+
 	if (InUpdateModes & ELandscapeLayerUpdateMode::Update_Heightmap_Editing)
 	{
 		const uint32 HeightmapEditingFlags = ELandscapeComponentUpdateFlag::Component_Update_Heightmap_Collision | ELandscapeComponentUpdateFlag::Component_Update_Client_Editing;
@@ -5201,28 +7541,44 @@
 	}
 }
 
-void ALandscape::GetLandscapeComponentNeighborsToRender(ULandscapeComponent* LandscapeComponent, TSet<ULandscapeComponent*>& NeighborComponents) const
+void ULandscapeComponent::GetLandscapeComponentNeighborsToRender(TSet<ULandscapeComponent*>& OutNeighborComponents) const
 {
 	ULandscapeInfo* LandscapeInfo = GetLandscapeInfo();
-	FIntPoint ComponentKey = LandscapeComponent->GetSectionBase() / ComponentSizeQuads;
-	
-	for (int32 IndexX = ComponentKey.X-1; IndexX <= ComponentKey.X+1; ++IndexX)
+	FIntPoint ComponentKey = GetSectionBase() / ComponentSizeQuads;
+
+	for (int32 IndexX = ComponentKey.X - 1; IndexX <= ComponentKey.X + 1; ++IndexX)
 	{
-		for (int32 IndexY = ComponentKey.Y-1; IndexY <= ComponentKey.Y+1; ++IndexY)
+		for (int32 IndexY = ComponentKey.Y - 1; IndexY <= ComponentKey.Y + 1; ++IndexY)
 		{
 			ULandscapeComponent* Result = LandscapeInfo->XYtoComponentMap.FindRef(FIntPoint(IndexX, IndexY));
-			if (Result && Result != LandscapeComponent)
+			if (Result && Result != this)
 			{
-				NeighborComponents.Add(Result);
+				OutNeighborComponents.Add(Result);
 			}
 		}
 	}
 }
 
-void ALandscape::GetLandscapeComponentWeightmapsToRender(ULandscapeComponent* LandscapeComponent, TSet<ULandscapeComponent*>& WeightmapComponents) const
+void ULandscapeComponent::GetLandscapeComponentNeighbors3x3(TStaticArray<ULandscapeComponent*, 9>& OutNeighborComponents) const
 {
+	ULandscapeInfo* LandscapeInfo = GetLandscapeInfo();
+	FIntPoint ComponentKey = GetSectionBase() / ComponentSizeQuads;
+
+	int32 LinearIndex = 0;
+	for (int32 IndexY = ComponentKey.Y - 1; IndexY <= ComponentKey.Y + 1; ++IndexY)
+	{
+		for (int32 IndexX = ComponentKey.X - 1; IndexX <= ComponentKey.X + 1; ++IndexX)
+		{
+			OutNeighborComponents[LinearIndex] = LandscapeInfo->XYtoComponentMap.FindRef(FIntPoint(IndexX, IndexY));
+			++LinearIndex;
+		}
+	}
+}
+
+void ULandscapeComponent::GetLandscapeComponentWeightmapsToRender(TSet<ULandscapeComponent*>& OutWeightmapComponents) const
+{
 	// Fill with Components that share the same weightmaps so that the Resolve of Weightmap Texture doesn't resolve null data.
-	for (ULandscapeWeightmapUsage* Usage : LandscapeComponent->GetWeightmapTexturesUsage(false))
+	for (ULandscapeWeightmapUsage* Usage : GetWeightmapTexturesUsage(/*InReturnEditingWeightmap = */false))
 	{
 		for (int32 Channel = 0; Channel < ULandscapeWeightmapUsage::NumChannels; ++Channel)
 		{
@@ -5237,7 +7593,7 @@
 						{
 							if (Usage != nullptr && Usage->ChannelUsage[Channel] != nullptr)
 							{
-								WeightmapComponents.Add(Usage->ChannelUsage[Channel]);
+								OutWeightmapComponents.Add(Usage->ChannelUsage[Channel]);
 							}
 						}
 					}
@@ -5271,6 +7627,14 @@
 		return;
 	}
 
+	// If we went from local merge to global merge or vice versa, we need to reinitialize layers : 
+	bool bUseLocalMerge = SupportsEditLayersLocalMerge();
+	if (bLandscapeLayersAreUsingLocalMerge != bUseLocalMerge)
+	{
+		RequestLayersInitialization(/*bInRequestContentUpdate = */true);
+		bLandscapeLayersAreUsingLocalMerge = bUseLocalMerge;
+	}
+
 	if (!bLandscapeLayersAreInitialized)
 	{
 		InitializeLayers();
@@ -5345,7 +7709,7 @@
 		}
 		return;
 	}
-	
+
 	// Gather mappings between heightmaps/weightmaps and components
 	FTextureToComponentHelper MapHelper(*LandscapeInfo);
 
@@ -5357,7 +7721,7 @@
 		// These flags might look like they're being mixed up but they're not!
 		const bool bDoIntermediateRender = false; // bIntermediateRender flag is for the work queued up this frame not the delayed resolves
 		const bool bDoFlushRender = bIntermediateRender; // Flush before we do an intermediate render later in this frame
-		
+
 		TMap<ULandscapeComponent*, FLandscapeEditLayerReadbackResult> ModifiedComponentInfo;
 		ResolveLayersHeightmapTexture(MapHelper, MapHelper.Heightmaps, bDoIntermediateRender, bDoFlushRender, ModifiedComponentInfo);
 		ResolveLayersWeightmapTexture(MapHelper, MapHelper.Weightmaps, bDoIntermediateRender, bDoFlushRender, ModifiedComponentInfo);
@@ -5373,121 +7737,12 @@
 	bool bUpdateAll = LayerContentUpdateModes & Update_All;
 	bool bPartialUpdate = !bForceRender && !bUpdateAll && CVarLandscapeLayerOptim.GetValueOnAnyThread() == 1;
 
-	TSet<UTexture2D*> Heightmaps;
-	TSet<UTexture2D*> HeightmapsToRender;
-	TSet<UTexture2D*> Weightmaps;
-	TSet<ULandscapeComponent*> NeighborsComponents;
-	TSet<ULandscapeComponent*> WeightmapsComponents;
-	TArray<ULandscapeComponent*> LandscapeComponentsHeightmapsToRender;
-	TArray<ULandscapeComponent*> LandscapeComponentsHeightmapsToResolve;
-	TArray<ULandscapeComponent*> LandscapeComponentsWeightmapsToRender;
-	TArray<ULandscapeComponent*> LandscapeComponentsWeightmapsToResolve;
-	TMap<ULandscapeComponent*, FLandscapeEditLayerReadbackResult> LandscapeComponentsToUpdate;
-	TArray<ULandscapeComponent*> AllLandscapeComponents;
+	FUpdateLayersContentContext UpdateLayersContentContext(MapHelper, bPartialUpdate);
 
-	TArray<ULandscapeComponent*> SkippedComponents;
-	for (ULandscapeComponent* Component : MapHelper.LandscapeComponents)
-	{
-		if (!bPartialUpdate || Component->GetLayerUpdateFlagPerMode() != 0)
-		{
-			AllLandscapeComponents.Add(Component);
-		
-			if (bPartialUpdate)
-			{
-				// Gather Neighbors (Neighbors need to be Rendered but not resolved so that the resolved Components have valid normals on edges)
-				GetLandscapeComponentNeighborsToRender(Component, NeighborsComponents);
-				Component->ForEachLayer([&](const FGuid&, FLandscapeLayerComponentData& LayerData)
-				{
-					HeightmapsToRender.Add(LayerData.HeightmapData.Texture);
-				});
-				// Gather WeightmapUsages (Components sharing weightmap usages with the resolved Components need to be rendered so that the resolving is valid)
-				GetLandscapeComponentWeightmapsToRender(Component, WeightmapsComponents);
-			}
-
-			// Gather Heightmaps (All Components sharing Heightmap textures need to be rendered and resolved)
-			Heightmaps.Add(Component->GetHeightmap(false));
-
-			// Gather Weightmaps
-			const TArray<UTexture2D*>& WeightmapTextures = Component->GetWeightmapTextures();
-			for (FWeightmapLayerAllocationInfo const& AllocInfo : Component->GetWeightmapLayerAllocations())
-			{
-				if (AllocInfo.IsAllocated() && AllocInfo.WeightmapTextureIndex < WeightmapTextures.Num())
-				{
-					Weightmaps.Add(WeightmapTextures[AllocInfo.WeightmapTextureIndex]);
-				}
-			}
-		}
-		else
-		{
-			SkippedComponents.Add(Component);
-		}
-	}
-
-	// Because of Heightmap Sharing anytime we render a heightmap we need to render all the components that use it
-	for (ULandscapeComponent* NeighborsComponent : NeighborsComponents)
-	{
-		NeighborsComponent->ForEachLayer([&](const FGuid&, FLandscapeLayerComponentData& LayerData)
-		{
-			HeightmapsToRender.Add(LayerData.HeightmapData.Texture);
-		});
-	}
-
-	// Copy first list into others
-	LandscapeComponentsHeightmapsToResolve.Append(AllLandscapeComponents);
-	LandscapeComponentsHeightmapsToRender.Append(AllLandscapeComponents);
-	LandscapeComponentsWeightmapsToResolve.Append(AllLandscapeComponents);
-	LandscapeComponentsWeightmapsToRender.Append(AllLandscapeComponents);
-		
-	if (bPartialUpdate)
-	{
-		for (ULandscapeComponent* Component : SkippedComponents)
-		{
-			if (Heightmaps.Contains(Component->GetHeightmap(false)))
-			{
-				AllLandscapeComponents.Add(Component);
-				LandscapeComponentsHeightmapsToRender.Add(Component);
-				LandscapeComponentsHeightmapsToResolve.Add(Component);
-			}
-			else if (NeighborsComponents.Contains(Component))
-			{
-				LandscapeComponentsHeightmapsToRender.Add(Component);
-			}
-			else
-			{
-				bool bAdd = false;
-				Component->ForEachLayer([&](const FGuid&, FLandscapeLayerComponentData& LayerData)
-				{
-					if (HeightmapsToRender.Contains(LayerData.HeightmapData.Texture))
-					{
-						bAdd = true;
-					}
-				});
-				if (bAdd)
-				{
-					LandscapeComponentsHeightmapsToRender.Add(Component);
-				}
-			}
-
-			if (WeightmapsComponents.Contains(Component))
-			{
-				LandscapeComponentsWeightmapsToRender.Add(Component);
-			}
-		}
-	}
-
-	// Add components with deferred flag to update list
-	for (ULandscapeComponent* Component : AllLandscapeComponents)
-	{
-		if (Component->GetLayerUpdateFlagPerMode() & ELandscapeLayerUpdateMode::Update_Client_Deferred)
-		{
-			LandscapeComponentsToUpdate.Add(Component, FLandscapeEditLayerReadbackResult(ELandscapeLayerUpdateMode::Update_Client_Deferred, false));
-		}
-	}
-
 	// Regenerate any heightmaps and weightmaps
 	int32 ProcessedModes = 0;
-	ProcessedModes |= RegenerateLayersHeightmaps(MapHelper, LandscapeComponentsHeightmapsToRender, LandscapeComponentsHeightmapsToResolve);
-	ProcessedModes |= RegenerateLayersWeightmaps(MapHelper, LandscapeComponentsWeightmapsToRender, LandscapeComponentsWeightmapsToResolve);
+	ProcessedModes |= RegenerateLayersHeightmaps(UpdateLayersContentContext);
+	ProcessedModes |= RegenerateLayersWeightmaps(UpdateLayersContentContext);
 	ProcessedModes |= (LayerContentUpdateModes & ELandscapeLayerUpdateMode::Update_Client_Deferred);
 	ProcessedModes |= (LayerContentUpdateModes & ELandscapeLayerUpdateMode::Update_Client_Editing);
 
@@ -5495,19 +7750,19 @@
 	if (bFlushRender || CVarLandscapeForceFlush.GetValueOnGameThread() != 0)
 	{
 		const bool bDoFlushRender = true;
-		ResolveLayersHeightmapTexture(MapHelper, Heightmaps, bIntermediateRender, bDoFlushRender, LandscapeComponentsToUpdate);
-		ResolveLayersWeightmapTexture(MapHelper, Weightmaps, bIntermediateRender, bDoFlushRender, LandscapeComponentsToUpdate);
+		ResolveLayersHeightmapTexture(UpdateLayersContentContext.MapHelper, UpdateLayersContentContext.HeightmapsToResolve, bIntermediateRender, bDoFlushRender, UpdateLayersContentContext.AllLandscapeComponentsToUpdate);
+		ResolveLayersWeightmapTexture(UpdateLayersContentContext.MapHelper, UpdateLayersContentContext.WeightmapsToResolve, bIntermediateRender, bDoFlushRender, UpdateLayersContentContext.AllLandscapeComponentsToUpdate);
 	}
 
 	// Clear processed mode flags
 	LayerContentUpdateModes &= ~ProcessedModes;
-	for (ULandscapeComponent* Component : AllLandscapeComponents)
+	for (ULandscapeComponent* Component : UpdateLayersContentContext.AllLandscapeComponentsToResolve)
 	{
 		Component->ClearUpdateFlagsForModes(ProcessedModes);
 	}
 
 	// Apply post resolve updates
-	const int32 ToProcessModes = UpdateAfterReadbackResolves(LandscapeComponentsToUpdate);
+	const int32 ToProcessModes = UpdateAfterReadbackResolves(UpdateLayersContentContext.AllLandscapeComponentsToUpdate);
 	LayerContentUpdateModes |= ToProcessModes;
 
 	if (LandscapeEdMode)
@@ -5683,7 +7938,6 @@
 void ALandscape::InitializeLayers()
 {
 	check(HasLayersContent());
-	check(!bLandscapeLayersAreInitialized);
 
 	CreateLayersRenderingResource();
 	InitializeLandscapeLayersWeightmapUsage();
@@ -5873,7 +8127,7 @@
 			Component->AddLayerData(InLayerGuid, FLandscapeLayerComponentData(EditLayer ? EditLayer->Name : FName()));
 			bModified = true;
 		}
-	}	
+	}
 
 	UpdateCachedHasLayersContent();
 
@@ -5923,7 +8177,7 @@
 	{
 		return;
 	}
-		
+
 	// Build a mapping between each Heightmaps and Component in them
 	TMap<UTexture2D*, TArray<ULandscapeComponent*>> ComponentsPerHeightmaps;
 
@@ -5933,7 +8187,7 @@
 		TArray<ULandscapeComponent*>& ComponentList = ComponentsPerHeightmaps.FindOrAdd(ComponentHeightmapTexture);
 		ComponentList.Add(Component);
 	}
-		
+
 	// Init layers with valid "empty" data
 	TMap<UTexture2D*, UTexture2D*> CreatedHeightmapTextures; // < Final layer texture, New created texture for layer
 
@@ -6051,7 +8305,7 @@
 
 float ALandscape::GetClampedLayerAlpha(float InAlpha, bool bInHeightmap) const
 {
-	float AlphaClamped = FMath::Clamp<float>(InAlpha, bInHeightmap  ? -1.f : 0.f, 1.f);
+	float AlphaClamped = FMath::Clamp<float>(InAlpha, bInHeightmap ? -1.f : 0.f, 1.f);
 	return AlphaClamped;
 }
 
@@ -6229,7 +8483,7 @@
 			LandscapeComponent->CopyFinalLayerIntoEditingLayer(DataInterface, ProcessedHeightmaps);
 		});
 	}
-	
+
 	TArray<ALandscapeBlueprintBrushBase*> BrushesToMove;
 	for (int32 i = 0; i < LandscapeLayers[InLayerIndex].Brushes.Num(); ++i)
 	{
@@ -6250,7 +8504,7 @@
 	}
 
 	DeleteLayer(InLayerIndex);
-	
+
 	RequestLayersContentUpdateForceAll();
 }
 
@@ -6270,7 +8524,7 @@
 	{
 		return;
 	}
-	
+
 	LandscapeInfo->GetUsedPaintLayers(InLayerGuid, OutUsedLayerInfos);
 }
 
@@ -6480,15 +8734,15 @@
 {
 	UTexture2D* Heightmap = GetHeightmap(true);
 	const uint8* MipData = Heightmap->Source.LockMipReadOnly(0);
-	uint32 Hash = FCrc::MemCrc32(MipData, Heightmap->GetSizeX()*Heightmap->GetSizeY() * sizeof(FColor));
+	uint32 Hash = FCrc::MemCrc32(MipData, Heightmap->GetSizeX() * Heightmap->GetSizeY() * sizeof(FColor));
 	Heightmap->Source.UnlockMip(0);
 
 	// Copy to sort
 	const TArray<UTexture2D*>& Weightmaps = GetWeightmapTextures(true);
 	TArray<FWeightmapLayerAllocationInfo> AllocationInfos = GetWeightmapLayerAllocations(true);
-	
-    // Sort allocations infos by LayerInfo Path so the Weightmaps hahses get ordered properly
-    AllocationInfos.Sort([](const FWeightmapLayerAllocationInfo& A, const FWeightmapLayerAllocationInfo& B)
+
+	// Sort allocations infos by LayerInfo Path so the Weightmaps hahses get ordered properly
+	AllocationInfos.Sort([](const FWeightmapLayerAllocationInfo& A, const FWeightmapLayerAllocationInfo& B)
 	{
 		FString PathA(A.LayerInfo ? A.LayerInfo->GetPathName() : FString());
 		FString PathB(B.LayerInfo ? B.LayerInfo->GetPathName() : FString());
@@ -6500,23 +8754,23 @@
 	{
 		if (AllocationInfo.IsAllocated())
 		{
-            // Compute hash of actual data of the texture that is owned by the component (per Texture Channel)
+			// Compute hash of actual data of the texture that is owned by the component (per Texture Channel)
 			UTexture2D* Weightmap = Weightmaps[AllocationInfo.WeightmapTextureIndex];
 			MipData = Weightmap->Source.LockMipReadOnly(0) + ChannelOffsets[AllocationInfo.WeightmapTextureChannel];
 			TArray<uint8> ChannelData;
-			ChannelData.AddDefaulted(Weightmap->GetSizeX()*Weightmap->GetSizeY());
-			int32 TexSize = (SubsectionSizeQuads + 1)*NumSubsections;
+			ChannelData.AddDefaulted(Weightmap->GetSizeX() * Weightmap->GetSizeY());
+			int32 TexSize = (SubsectionSizeQuads + 1) * NumSubsections;
 			for (int32 TexY = 0; TexY < TexSize; TexY++)
 			{
 				for (int32 TexX = 0; TexX < TexSize; TexX++)
 				{
 					const int32 Index = (TexX + TexY * TexSize);
-				
+
 					ChannelData.GetData()[Index] = MipData[4 * Index];
 				}
 			}
 
-			Hash = FCrc::MemCrc32(ChannelData.GetData(), Weightmap->GetSizeX()*Weightmap->GetSizeY(), Hash);
+			Hash = FCrc::MemCrc32(ChannelData.GetData(), Weightmap->GetSizeX() * Weightmap->GetSizeY(), Hash);
 			Weightmap->Source.UnlockMip(0);
 		}
 	}
@@ -6544,11 +8798,11 @@
 			{
 				return;
 			}
-						
+
 			TMap<ULandscapeComponent*, uint32> PreviousHashes;
 			{
 				FLandscapeEditDataInterface LandscapeEdit(LandscapeInfo);
-				
+
 				LandscapeInfo->ForAllLandscapeComponents([&](ULandscapeComponent* Component)
 				{
 					// Was never computed
@@ -6573,8 +8827,8 @@
 			bInUpdateOnlySelected = false;
 
 			// Apply splines without clearing up weightmap allocations
-			LandscapeInfo->ApplySplines(bInUpdateOnlySelected, ModifiedComponent, bMarkPackageDirty);		
-		
+			LandscapeInfo->ApplySplines(bInUpdateOnlySelected, ModifiedComponent, bMarkPackageDirty);
+
 			for (const TPair<ULandscapeComponent*, uint32>& Pair : PreviousHashes)
 			{
 				if (LandscapeSplinesAffectedComponents.Contains(Pair.Key))
@@ -6636,7 +8890,7 @@
 void ALandscape::SetEditingLayer(const FGuid& InLayerGuid)
 {
 	ensure(CanHaveLayersContent());
-	
+
 	ULandscapeInfo* LandscapeInfo = GetLandscapeInfo();
 	if (!LandscapeInfo)
 	{
@@ -6950,12 +9204,12 @@
 void ALandscape::OnLayerInfoSplineFalloffModulationChanged(ULandscapeLayerInfoObject* InLayerInfo)
 {
 	ULandscapeInfo* LandscapeInfo = GetLandscapeInfo();
-	
+
 	if (!LandscapeInfo)
 	{
 		return;
 	}
-	
+
 	ALandscape* Landscape = LandscapeInfo->LandscapeActor.Get();
 	if (!Landscape || !Landscape->HasLayersContent())
 	{
@@ -7082,7 +9336,7 @@
 				LandscapeTransform = NewLandscapeTransform;
 				LandscapeRenderTargetSize = NewLandscapeRenderTargetSize;
 				LandscapeSize = NewLandscapeSize;
-				
+
 				TGuardValue<bool> AutoRestore(GAllowActorScriptExecutionInEditor, true);
 				BlueprintBrush->Initialize(LandscapeTransform, LandscapeSize, LandscapeRenderTargetSize);
 			}
