Local merge of landscape edit layers : complete refactor of the merge algorithm that doesn't require allocating several render targets the size of the entire loaded landscape components. Instead, each

diff --git a/Engine/Source/Runtime/Landscape/Private/LandscapeEditResources.cpp b/Engine/Source/Runtime/Landscape/Private/LandscapeEditResources.cpp
new file mode 100644
--- /dev/null
+++ b/Engine/Source/Runtime/Landscape/Private/LandscapeEditResources.cpp
@@ -0,0 +1,206 @@
+// Copyright Epic Games, Inc. All Rights Reserved.
+
+#include "LandscapeEditResources.h"
+
+#include "Engine/Texture2D.h"
+#include "LandscapePrivate.h"
+#include "RenderingThread.h"
+#include "RenderGraph.h"
+#include "Rendering/Texture2DResource.h"
+
+
+// ----------------------------------------------------------------------------------
+
+FLandscapeTexture2DResource::FLandscapeTexture2DResource(uint32 InSizeX, uint32 InSizeY, EPixelFormat InFormat, uint32 bInNumMips, bool bInNeedUAVs, bool bInNeedSRV)
+	: SizeX(InSizeX)
+	, SizeY(InSizeY)
+	, Format(InFormat)
+	, NumMips(bInNumMips)
+	, bCreateUAVs(bInNeedUAVs)
+	, bCreateSRV(bInNeedSRV)
+{
+
+}
+
+void FLandscapeTexture2DResource::InitRHI()
+{
+	FTextureResource::InitRHI();
+
+	FRHIResourceCreateInfo CreateInfo(TEXT("FLandscapeTexture2DResource"));
+	ETextureCreateFlags Flags = TexCreate_None;
+
+	if (bCreateUAVs)
+	{
+		Flags |= TexCreate_RenderTargetable | TexCreate_UAV;
+	}
+
+	if (bCreateSRV)
+	{
+		Flags |= TexCreate_ShaderResource;
+	}
+
+	TextureRHI = RHICreateTexture2D(SizeX, SizeY, Format, NumMips, 1, Flags, CreateInfo);
+
+	if (bCreateUAVs)
+	{
+		TextureUAVs.Reserve(NumMips);
+		for (uint32 MipLevel = 0; MipLevel < NumMips; ++MipLevel)
+		{
+			TextureUAVs.Add(RHICreateUnorderedAccessView(TextureRHI, MipLevel));
+		}
+	}
+
+	if (bCreateSRV)
+	{
+		TextureSRV = RHICreateShaderResourceView(TextureRHI, /*MipLevel = */0, NumMips, Format);
+	}
+}
+
+void FLandscapeTexture2DResource::ReleaseRHI()
+{
+	FTextureResource::ReleaseRHI();
+
+	if (bCreateUAVs)
+	{
+		for (FUnorderedAccessViewRHIRef& TextureUAV : TextureUAVs)
+		{
+			TextureUAV.SafeRelease();
+		}
+		TextureUAVs.Empty();
+	}
+}
+
+FUnorderedAccessViewRHIRef FLandscapeTexture2DResource::GetTextureUAV(uint32 InMipLevel) const
+{
+	check(bCreateUAVs && (InMipLevel < NumMips));
+	check(TextureUAVs.Num() == NumMips);
+	return TextureUAVs[InMipLevel];
+}
+
+FShaderResourceViewRHIRef FLandscapeTexture2DResource::GetTextureSRV() const
+{
+	check(bCreateSRV);
+	return TextureSRV;
+}
+
+// ----------------------------------------------------------------------------------
+
+FLandscapeTexture2DArrayResource::FLandscapeTexture2DArrayResource(uint32 InSizeX, uint32 InSizeY, uint32 InSizeZ, EPixelFormat InFormat, uint32 InNumMips, bool bInNeedUAVs, bool bInNeedSRV)
+	: SizeX(InSizeX)
+	, SizeY(InSizeY)
+	, SizeZ(InSizeZ)
+	, Format(InFormat)
+	, NumMips(InNumMips)
+	, bCreateUAVs(bInNeedUAVs)
+{
+
+}
+
+void FLandscapeTexture2DArrayResource::InitRHI()
+{
+	FTextureResource::InitRHI();
+
+	FRHIResourceCreateInfo CreateInfo(TEXT("FLandscapeTexture2DArrayResource"));
+	ETextureCreateFlags Flags = TexCreate_NoTiling | TexCreate_OfflineProcessed;
+
+	if (bCreateUAVs)
+	{
+		Flags |= TexCreate_UAV;
+	}
+
+	if (bCreateSRV)
+	{
+		Flags |= TexCreate_ShaderResource;
+	}
+
+	TextureRHI = RHICreateTexture2DArray(SizeX, SizeY, SizeZ, Format, NumMips, 1, Flags, CreateInfo);
+
+	if (bCreateUAVs)
+	{
+		TextureUAVs.Reserve(NumMips);
+		for (uint32 MipLevel = 0; MipLevel < NumMips; ++MipLevel)
+		{
+			TextureUAVs.Add(RHICreateUnorderedAccessView(TextureRHI, MipLevel));
+		}
+	}
+
+	if (bCreateSRV)
+	{
+		TextureSRV = RHICreateShaderResourceView(TextureRHI, /*MipLevel = */0, NumMips, Format);
+	}
+}
+
+void FLandscapeTexture2DArrayResource::ReleaseRHI()
+{
+	FTextureResource::ReleaseRHI();
+
+	if (bCreateUAVs)
+	{
+		for (FUnorderedAccessViewRHIRef& TextureUAV : TextureUAVs)
+		{
+			TextureUAV.SafeRelease();
+		}
+		TextureUAVs.Empty();
+	}
+}
+
+FUnorderedAccessViewRHIRef FLandscapeTexture2DArrayResource::GetTextureUAV(uint32 InMipLevel) const
+{
+	check(bCreateUAVs && (InMipLevel < NumMips));
+	check(TextureUAVs.Num() == NumMips);
+	return TextureUAVs[InMipLevel];
+}
+
+FShaderResourceViewRHIRef FLandscapeTexture2DArrayResource::GetTextureSRV() const
+{
+	check(bCreateSRV);
+	return TextureSRV;
+}
+
+
+// ----------------------------------------------------------------------------------
+
+void TrackLandscapeRDGTextures(FRDGBuilder& GraphBuilder, TMap<FTexture2DResource*, FLandscapeRDGTrackedTexture>& TrackedTextures)
+{
+	// Register all these to the graph and mark those that will be written to it as targetable + create UAV/SRV : 
+	for (auto& ItTrackedTexture : TrackedTextures)
+	{
+		FLandscapeRDGTrackedTexture& TrackedTexture = ItTrackedTexture.Value;
+		check(TrackedTexture.TextureResource != nullptr);
+
+		// We need a debug name that will last until the graph is executed: 
+		TrackedTexture.DebugName = GraphBuilder.AllocObject<FString>(TrackedTexture.TextureResource->GetTextureName().ToString());
+
+		// Register the texture to the GraphBuilder : we'll read from / write to it like a render target : 
+		TRefCountPtr<IPooledRenderTarget> RenderTarget = CreateRenderTarget(TrackedTexture.TextureResource->TextureRHI, **TrackedTexture.DebugName);
+
+		// Force tracking on the external texture if necessary, so that it can be copied from/to via CopyTexture within the graph : 
+		ERDGTextureFlags TextureFlags = (TrackedTexture.bNeedsForceTracking || TrackedTexture.bNeedsScratch) ? ERDGTextureFlags::ForceTracking : ERDGTextureFlags::ReadOnly;
+		TrackedTexture.ExternalTextureRef = GraphBuilder.RegisterExternalTexture(RenderTarget, **TrackedTexture.DebugName, ERenderTargetTexture::ShaderResource, TextureFlags);
+
+		if (TrackedTexture.bNeedsScratch)
+		{
+			FString* ScratchTextureName = GraphBuilder.AllocObject<FString>(FString::Format(TEXT("{0}_Scratch"), { *TrackedTexture.DebugName }));
+			// We'll work on a targetable copy of the original texture so that we can write to it from a shader (we'll copy to the original one at the end of the process with a CopyTexture): 
+			FRDGTextureDesc ScratchTextureDesc = TrackedTexture.ExternalTextureRef->Desc;
+			ScratchTextureDesc.Flags |= ETextureCreateFlags::RenderTargetable;
+			TrackedTexture.ScratchTextureRef = GraphBuilder.CreateTexture(ScratchTextureDesc, **ScratchTextureName);
+
+			for (int32 MipLevel = 0; MipLevel < ScratchTextureDesc.NumMips; ++MipLevel)
+			{
+				TrackedTexture.ScratchTextureMipsSRVRefs.Add(GraphBuilder.CreateSRV(FRDGTextureSRVDesc::CreateForMipLevel(TrackedTexture.ScratchTextureRef, MipLevel)));
+			}
+
+			// We only need to copy the first mip : 
+			FRHICopyTextureInfo CopyTextureInfo;
+			CopyTextureInfo.NumMips = 1;
+			CopyTextureInfo.Size = FIntVector(TrackedTexture.ExternalTextureRef->Desc.GetSize().X, TrackedTexture.ExternalTextureRef->Desc.GetSize().Y, 0);
+			AddCopyTexturePass(GraphBuilder, TrackedTexture.ExternalTextureRef, TrackedTexture.ScratchTextureRef, CopyTextureInfo);
+		}
+
+		if (TrackedTexture.bNeedsSRV)
+		{
+			TrackedTexture.ExternalTextureSRVRef = GraphBuilder.CreateSRV(FRDGTextureSRVDesc::Create(TrackedTexture.ExternalTextureRef));
+		}
+	}
+}
diff --git a/Engine/Source/Runtime/Landscape/Private/LandscapeEditResources.h b/Engine/Source/Runtime/Landscape/Private/LandscapeEditResources.h
new file mode 100644
--- /dev/null
+++ b/Engine/Source/Runtime/Landscape/Private/LandscapeEditResources.h
@@ -0,0 +1,119 @@
+// Copyright Epic Games, Inc. All Rights Reserved.
+
+#pragma once
+
+#include "CoreMinimal.h"
+#include "TextureResource.h"
+
+
+// Custom Resources for edit layers landscapes :
+
+// ----------------------------------------------------------------------------------
+class FLandscapeTexture2DResource : public FTextureResource
+{
+public:
+	FLandscapeTexture2DResource(uint32 InSizeX, uint32 InSizeY, EPixelFormat InFormat, uint32 InNumMips, bool bInNeedUAVs, bool bInNeedSRV);
+
+	virtual uint32 GetSizeX() const override { return SizeX; }
+	virtual uint32 GetSizeY() const override { return SizeY; }
+
+	/** Called when the resource is initialized. This is only called by the rendering thread. */
+	virtual void InitRHI() override;
+	virtual void ReleaseRHI() override;
+
+	FUnorderedAccessViewRHIRef GetTextureUAV(uint32 InMipLevel) const;
+	FShaderResourceViewRHIRef GetTextureSRV() const;
+
+private:
+	uint32 SizeX;
+	uint32 SizeY;
+	EPixelFormat Format;
+	uint32 NumMips;
+	bool bCreateUAVs;
+	bool bCreateSRV;
+	TArray<FUnorderedAccessViewRHIRef> TextureUAVs;
+	FShaderResourceViewRHIRef TextureSRV;
+};
+
+// ----------------------------------------------------------------------------------
+class FLandscapeTexture2DArrayResource : public FTextureResource
+{
+public:
+	FLandscapeTexture2DArrayResource(uint32 InSizeX, uint32 InSizeY, uint32 InSizeZ, EPixelFormat InFormat, uint32 InNumMips, bool bInNeedUAVs, bool bInNeedSRV);
+
+	virtual uint32 GetSizeX() const override { return SizeX; }
+	virtual uint32 GetSizeY() const override { return SizeY;}
+	virtual uint32 GetSizeZ() const { return SizeZ; }
+
+	/** Called when the resource is initialized. This is only called by the rendering thread. */
+	virtual void InitRHI() override;
+	virtual void ReleaseRHI() override;
+
+	FUnorderedAccessViewRHIRef GetTextureUAV(uint32 InMipLevel) const;
+	FShaderResourceViewRHIRef GetTextureSRV() const;
+
+private:
+	uint32 SizeX;
+	uint32 SizeY;
+	uint32 SizeZ;
+	EPixelFormat Format;
+	uint32 NumMips;
+	bool bCreateUAVs;
+	bool bCreateSRV;
+	TArray<FUnorderedAccessViewRHIRef> TextureUAVs;
+	FShaderResourceViewRHIRef TextureSRV;
+};
+
+// ----------------------------------------------------------------------------------
+// Describes a texture and an associated subregion (for heightmap texture sharing):
+struct FTexture2DResourceSubregion
+{
+	FTexture2DResourceSubregion() = default;
+
+	FTexture2DResourceSubregion(FTexture2DResource* InTexture, const FIntPoint& InSubregionOffset, const FIntPoint& InSubregionSize)
+		: Texture(InTexture)
+		, Subregion(InSubregionOffset, InSubregionOffset + InSubregionSize)
+	{}
+
+	FTexture2DResourceSubregion(FTexture2DResource* InTexture, const FIntRect& InSubregion)
+		: Texture(InTexture)
+		, Subregion(InSubregion)
+	{}
+
+	FTexture2DResource* Texture = nullptr;
+	FIntRect Subregion;
+};
+
+// ----------------------------------------------------------------------------------
+// Struct that helps tracking external textures in the render graph
+struct FLandscapeRDGTrackedTexture
+{
+	FLandscapeRDGTrackedTexture(FTexture2DResource* InTextureResource)
+		: TextureResource(InTextureResource)
+	{}
+
+	/** External texture resource that needs tracking */
+	FTexture2DResource* TextureResource = nullptr;
+	/** If bNeedsScratch is set to true, a sister copy of the the texture will be allocated: FLandscapeRDGTrackedTexture is meant to track non-render targetable and UAV-less textures (UTextures...), so any change to 
+	 the texture needs to be done via a texture copy (scratch -> original) */
+	bool bNeedsScratch = false;
+	/** If bNeedsSRV is set to true, the registration of the texture in the render graph will also produce a FRDGTextureSRVRef (ExternalTextureSRVRef) in order to access this texture as a SRV */
+	bool bNeedsSRV = false;
+	/** If bNeedsForceTracking is set to true, the texture will be declared in the render graph with the ERDGTextureFlags::ForceTracking flag, so that it can be copied from/to via CopyTexture within the graph */
+	bool bNeedsForceTracking = false;
+	/** Debug name for this texture (transient memory : will last for the lifetime of the GraphBuilder) */
+	FString* DebugName = nullptr;
+	/** If bNeedsScratch is true, contains the RDG ref of the allocated intermediate scratch texture */
+	FRDGTextureRef ScratchTextureRef = nullptr;
+	/** If bNeedsScratch is true, contains the RDG refs of the intermediate scratch texture SRVs (one SRV per mip level) */
+	TArray<FRDGTextureSRVRef> ScratchTextureMipsSRVRefs; 
+	/** RDG ref for the TextureResource */
+	FRDGTextureRef ExternalTextureRef = nullptr;
+	/** If bNeedsSRV is true, RDG ref for the TextureResource's SRV */
+	FRDGTextureSRVRef ExternalTextureSRVRef = nullptr;
+};
+
+/**
+ * Register all textures we will write into or read from in the render graph :
+ */
+void TrackLandscapeRDGTextures(FRDGBuilder& GraphBuilder, TMap<FTexture2DResource*, FLandscapeRDGTrackedTexture>& TrackedTextures);
\ No newline at end of file
diff --git a/Engine/Source/Runtime/Landscape/Private/LandscapeSplines.cpp b/Engine/Source/Runtime/Landscape/Private/LandscapeSplines.cpp
--- a/Engine/Source/Runtime/Landscape/Private/LandscapeSplines.cpp
+++ b/Engine/Source/Runtime/Landscape/Private/LandscapeSplines.cpp
@@ -28,6 +28,7 @@
 #include "LandscapeSplineSegment.h"
 #include "LandscapeSplineRaster.h"
 #include "LandscapeSplineControlPoint.h"
+#include "LandscapePrivate.h"
 #include "ILandscapeSplineInterface.h"
 #include "ControlPointMeshComponent.h"
 #include "Engine/CollisionProfile.h"
@@ -77,31 +78,37 @@ struct FLandscapeFixSplines
 
 	void FixSplines()
 	{
-		if (!GWorld || GWorld->IsGameWorld())
-		{
-			return;
-		}
+		bool bFixed = false;
 
-		auto& LandscapeInfoMap = ULandscapeInfoMap::GetLandscapeInfoMap(GWorld);
-		for (TPair<FGuid, ULandscapeInfo*>& Pair : LandscapeInfoMap.Map)
+		for (TObjectIterator<UWorld> It; It; ++It)
 		{
-			if (Pair.Value && Pair.Value->SupportsLandscapeEditing())
+			UWorld* CurrentWorld = *It;
+			if (!CurrentWorld->IsGameWorld())
 			{
-				Pair.Value->ForAllSplineActors([](TScriptInterface<ILandscapeSplineInterface> SplineOwner)
+				auto& LandscapeInfoMap = ULandscapeInfoMap::GetLandscapeInfoMap(CurrentWorld);
+				for (TPair<FGuid, ULandscapeInfo*>& Pair : LandscapeInfoMap.Map)
 				{
-					if (SplineOwner && SplineOwner->GetSplinesComponent())
+					if (Pair.Value && Pair.Value->SupportsLandscapeEditing())
 					{
-						SplineOwner->GetSplinesComponent()->RebuildAllSplines();
-					}
-				});
+						Pair.Value->ForAllSplineActors([&bFixed](TScriptInterface<ILandscapeSplineInterface> SplineOwner)
+						{
+							if (SplineOwner && SplineOwner->GetSplinesComponent())
+							{
+								SplineOwner->GetSplinesComponent()->RebuildAllSplines();
+								bFixed = true;
+							}
+						});
 
-				if (Pair.Value->LandscapeActor && Pair.Value->LandscapeActor->HasLayersContent())
-				{
-					Pair.Value->LandscapeActor->RequestSplineLayerUpdate();
+						if (Pair.Value->LandscapeActor && Pair.Value->LandscapeActor->HasLayersContent())
+						{
+							Pair.Value->LandscapeActor->RequestSplineLayerUpdate();
+						}
+					}
 				}
 			}
 		}
 
+		UE_LOG(LogLandscape, Display, TEXT("Landscape.FixSplines: %s"), bFixed ? TEXT("Splines fixed") : TEXT("Nothing to fix"));
 	}
 };
 
