Fixed last minute change that prevented landscape painting to work

#rnx

diff --git a/Engine/Source/Runtime/Landscape/Private/LandscapeEditLayers.cpp b/Engine/Source/Runtime/Landscape/Private/LandscapeEditLayers.cpp
--- a/Engine/Source/Runtime/Landscape/Private/LandscapeEditLayers.cpp
+++ b/Engine/Source/Runtime/Landscape/Private/LandscapeEditLayers.cpp
@@ -5323,7 +5323,7 @@ void ALandscape::PrepareComponentDataToPackMaterialLayersCS(int32 InCurrentWeigh
 
 			// It's possible the texture (and its usage, hence) has already been processed by a previous call to PrepareComponentDataToPackMaterialLayersCS with a different InCurrentWeightmapToProcessIndex
 			//  since a texture can be shared by multiple components :
-			if (OutProcessedWeightmaps.Contains(ComponentWeightmapTexture) && (WeightmapsToProcess.Find(ComponentWeightmapTexture) == nullptr))
+			if (!OutProcessedWeightmaps.Contains(ComponentWeightmapTexture) && (WeightmapsToProcess.Find(ComponentWeightmapTexture) == nullptr))
 			{
 				WeightmapsToProcess.Add(ComponentWeightmapTexture, ComponentWeightmapTextureUsage);
 				OutProcessedWeightmaps.Add(ComponentWeightmapTexture);
