LandscapeGizmoActor should not be shown in the Scene Outliner

diff --git a/Engine/Source/Runtime/Landscape/Private/LandscapeGizmoActor.cpp b/Engine/Source/Runtime/Landscape/Private/LandscapeGizmoActor.cpp
--- a/Engine/Source/Runtime/Landscape/Private/LandscapeGizmoActor.cpp
+++ b/Engine/Source/Runtime/Landscape/Private/LandscapeGizmoActor.cpp
@@ -446,6 +446,7 @@ ALandscapeGizmoActor::ALandscapeGizmoActor(const FObjectInitializer& ObjectIniti
 	RootComponent = SceneComponent;
 
 #if WITH_EDITORONLY_DATA
+	bListedInSceneOutliner = false;
 	bEditable = false;
 	Width = 1280.0f;
 	Height = 1280.0f;
