Integrate various landscape fixes from UE5

diff --git a/Engine/Source/Runtime/Landscape/Private/LandscapeRender.cpp b/Engine/Source/Runtime/Landscape/Private/LandscapeRender.cpp
--- a/Engine/Source/Runtime/Landscape/Private/LandscapeRender.cpp
+++ b/Engine/Source/Runtime/Landscape/Private/LandscapeRender.cpp
@@ -1685,7 +1685,7 @@ bool FLandscapeComponentSceneProxy::OnLevelAddedToWorld_RenderThread()
 	if (VisibilityHelper.OnAddedToWorld())
 	{
 		SetForceHidden(false);
-		RegisterNeighbors(this);
+		CreateRenderThreadResources();
 		return true;
 	}
 	return false;
@@ -2200,6 +2200,11 @@ void FLandscapeComponentSceneProxy::DrawStaticElements(FStaticPrimitiveDrawInter
 		return;
 	}
 
+	if (!VisibilityHelper.ShouldBeVisible())
+	{
+		return;
+	}
+
 	int32 TotalBatchCount = 1 + LastLOD - FirstLOD;
 	TotalBatchCount += (1 + LastVirtualTextureLOD - FirstVirtualTextureLOD) * RuntimeVirtualTextureMaterialTypes.Num();
 
