Landscape component - Moved proxy mesh creation out of ALandscapeProxy

diff --git a/Engine/Source/Runtime/Landscape/Classes/LandscapeComponent.h b/Engine/Source/Runtime/Landscape/Classes/LandscapeComponent.h
--- a/Engine/Source/Runtime/Landscape/Classes/LandscapeComponent.h
+++ b/Engine/Source/Runtime/Landscape/Classes/LandscapeComponent.h
@@ -36,7 +36,7 @@ struct FEngineShowFlags;
 struct FLandscapeEditDataInterface;
 struct FLandscapeTextureDataInfo;
 struct FStaticLightingPrimitiveInfo;
-
+struct FMeshDescription;
 struct FLandscapeEditDataInterface;
 struct FLandscapeMobileRenderData;
 
@@ -864,6 +864,9 @@ public:
 
 	/** Generate mobile data if it's missing or outdated */
 	void CheckGenerateLandscapePlatformData(bool bIsCooking, const ITargetPlatform* TargetPlatform);
+
+	LANDSCAPE_API void ExportToMeshDescription(const int32 InExportLOD, FMeshDescription& OutMesh);
+	LANDSCAPE_API void ExportToMeshDescription(const int32 InExportLOD, const FBoxSphereBounds& InBounds, FMeshDescription& OutMesh);
 #endif
 
 	LANDSCAPE_API int32 GetMaterialInstanceCount(bool InDynamic = true) const;
diff --git a/Engine/Source/Runtime/Landscape/Private/LandscapeComponent.cpp b/Engine/Source/Runtime/Landscape/Private/LandscapeComponent.cpp
--- a/Engine/Source/Runtime/Landscape/Private/LandscapeComponent.cpp
+++ b/Engine/Source/Runtime/Landscape/Private/LandscapeComponent.cpp
@@ -6,6 +6,15 @@
 #include "LandscapeEdit.h"
 #include "LandscapeRender.h"
 
+#if WITH_EDITOR
+
+#include "MeshDescription.h"
+#include "MeshUtilitiesCommon.h"
+#include "StaticMeshAttributes.h"
+#include "LandscapeDataAccess.h"
+
+#endif
+
 FName FWeightmapLayerAllocationInfo::GetLayerName() const
 {
 	if (LayerInfo)
@@ -66,4 +75,267 @@ void ULandscapeComponent::UpdateEditToolRenderData()
 	}
 }
 
+static void ExportToMeshDescription(ULandscapeComponent* InComponent, const int32 InExportLOD, FMeshDescription& OutMesh, const FBoxSphereBounds& InBounds, bool bIgnoreBounds)
+{
+	ALandscapeProxy* LandscapeProxy = InComponent->GetLandscapeProxy();
+
+	FStaticMeshAttributes Attributes(OutMesh);
+	TVertexAttributesRef<FVector3f> VertexPositions = Attributes.GetVertexPositions();
+	TEdgeAttributesRef<bool> EdgeHardnesses = Attributes.GetEdgeHardnesses();
+	TPolygonGroupAttributesRef<FName> PolygonGroupImportedMaterialSlotNames = Attributes.GetPolygonGroupMaterialSlotNames();
+	TVertexInstanceAttributesRef<FVector3f> VertexInstanceNormals = Attributes.GetVertexInstanceNormals();
+	TVertexInstanceAttributesRef<FVector3f> VertexInstanceTangents = Attributes.GetVertexInstanceTangents();
+	TVertexInstanceAttributesRef<float> VertexInstanceBinormalSigns = Attributes.GetVertexInstanceBinormalSigns();
+	TVertexInstanceAttributesRef<FVector4f> VertexInstanceColors = Attributes.GetVertexInstanceColors();
+	TVertexInstanceAttributesRef<FVector2f> VertexInstanceUVs = Attributes.GetVertexInstanceUVs();
+
+	if (VertexInstanceUVs.GetNumChannels() < 2)
+	{
+		VertexInstanceUVs.SetNumChannels(2);
+	}
+
+	// User specified LOD to export
+	int32 LandscapeLODToExport = LandscapeProxy->ExportLOD;
+	if (InExportLOD != INDEX_NONE)
+	{
+		LandscapeLODToExport = FMath::Clamp<int32>(InExportLOD, 0, FMath::CeilLogTwo(InComponent->SubsectionSizeQuads + 1) - 1);
+	}
+
+	FLandscapeComponentDataInterface CDI(InComponent, LandscapeLODToExport);
+	const int32 ComponentSizeQuadsLOD = ((InComponent->ComponentSizeQuads + 1) >> LandscapeLODToExport) - 1;
+	const int32 SubsectionSizeQuadsLOD = ((InComponent->SubsectionSizeQuads + 1) >> LandscapeLODToExport) - 1;
+	const FIntPoint ComponentOffsetQuads = InComponent->GetSectionBase();
+	const FVector2f ComponentUVOffsetLOD = FVector2f(ComponentOffsetQuads);
+	const FVector2f ComponentUVScaleLOD = FVector2f(InComponent->ComponentSizeQuads / ComponentSizeQuadsLOD);
+
+	const int32 NumFaces = FMath::Square(ComponentSizeQuadsLOD) * 2;
+	const int32 NumVertices = NumFaces * 3;
+
+	OutMesh.ReserveNewVertices(NumVertices);
+	OutMesh.ReserveNewPolygons(NumFaces);
+	OutMesh.ReserveNewVertexInstances(NumVertices);
+	OutMesh.ReserveNewEdges(NumVertices);
+
+	FPolygonGroupID PolygonGroupID = INDEX_NONE;
+	if (OutMesh.PolygonGroups().Num() < 1)
+	{
+		PolygonGroupID = OutMesh.CreatePolygonGroup();
+		PolygonGroupImportedMaterialSlotNames[PolygonGroupID] = FName(TEXT("LandscapeMat_0"));
+	}
+	else
+	{
+		PolygonGroupID = OutMesh.PolygonGroups().GetFirstValidID();
+	}
+
+	// Check if there are any holes
+	const int32 VisThreshold = 170;
+	TArray<uint8> VisDataMap;
+	const TArray<FWeightmapLayerAllocationInfo>& ComponentWeightmapLayerAllocations = InComponent->GetWeightmapLayerAllocations();
+
+	for (int32 AllocIdx = 0; AllocIdx < ComponentWeightmapLayerAllocations.Num(); AllocIdx++)
+	{
+		const FWeightmapLayerAllocationInfo& AllocInfo = ComponentWeightmapLayerAllocations[AllocIdx];
+		if (AllocInfo.LayerInfo == ALandscapeProxy::VisibilityLayer)
+		{
+			CDI.GetWeightmapTextureData(AllocInfo.LayerInfo, VisDataMap);
+		}
+	}
+
+	const FIntPoint QuadPattern[6] =
+	{
+		//face 1
+		FIntPoint(0, 0),
+		FIntPoint(0, 1),
+		FIntPoint(1, 1),
+		//face 2
+		FIntPoint(0, 0),
+		FIntPoint(1, 1),
+		FIntPoint(1, 0),
+	};
+
+	const int32 WeightMapSize = (SubsectionSizeQuadsLOD + 1) * InComponent->NumSubsections;
+
+	const float SquaredSphereRadius = FMath::Square(InBounds.SphereRadius);
+
+	//We need to not duplicate the vertex position, so we use the FIndexAndZ to achieve fast result
+	TArray<FIndexAndZ> VertIndexAndZ;
+	VertIndexAndZ.Reserve(ComponentSizeQuadsLOD * ComponentSizeQuadsLOD * UE_ARRAY_COUNT(QuadPattern));
+	int32 CurrentIndex = 0;
+	TMap<int32, FVector> IndexToPosition;
+	IndexToPosition.Reserve(ComponentSizeQuadsLOD * ComponentSizeQuadsLOD * UE_ARRAY_COUNT(QuadPattern));
+	for (int32 y = 0; y < ComponentSizeQuadsLOD; y++)
+	{
+		for (int32 x = 0; x < ComponentSizeQuadsLOD; x++)
+		{
+			for (int32 i = 0; i < UE_ARRAY_COUNT(QuadPattern); i++)
+			{
+				int32 VertexX = x + QuadPattern[i].X;
+				int32 VertexY = y + QuadPattern[i].Y;
+				FVector Position = CDI.GetWorldVertex(VertexX, VertexY);
+
+				// If at least one vertex is within the given bounds we should process the quad  
+				new(VertIndexAndZ)FIndexAndZ(CurrentIndex, Position);
+				IndexToPosition.Add(CurrentIndex, Position);
+				CurrentIndex++;
+			}
+		}
+	}
+	// Sort the vertices by z value
+	VertIndexAndZ.Sort(FCompareIndexAndZ());
+
+	auto FindPreviousIndex = [&VertIndexAndZ, &IndexToPosition](int32 Index)->int32
+	{
+		const FVector& PositionA = IndexToPosition[Index];
+		FIndexAndZ CompressPosition(0, PositionA);
+		// Search for lowest index duplicates
+		int32 BestIndex = MAX_int32;
+		for (int32 i = 0; i < IndexToPosition.Num(); i++)
+		{
+			if (CompressPosition.Z > (VertIndexAndZ[i].Z + SMALL_NUMBER))
+			{
+				//We will not find anything there is no point searching more
+				break;
+			}
+			const FVector& PositionB = IndexToPosition[VertIndexAndZ[i].Index];
+			if (PointsEqual(PositionA, PositionB, SMALL_NUMBER))
+			{
+				if (VertIndexAndZ[i].Index < BestIndex)
+				{
+					BestIndex = VertIndexAndZ[i].Index;
+				}
+			}
+		}
+		return BestIndex < MAX_int32 ? BestIndex : Index;
+	};
+
+	// Export to MeshDescription
+	TMap<int32, FVertexID> IndexToVertexID;
+	IndexToVertexID.Reserve(CurrentIndex);
+	CurrentIndex = 0;
+	for (int32 y = 0; y < ComponentSizeQuadsLOD; y++)
+	{
+		for (int32 x = 0; x < ComponentSizeQuadsLOD; x++)
+		{
+			FVector Positions[UE_ARRAY_COUNT(QuadPattern)];
+			bool bProcess = bIgnoreBounds;
+
+			// Fill positions
+			for (int32 i = 0; i < UE_ARRAY_COUNT(QuadPattern); i++)
+			{
+				int32 VertexX = x + QuadPattern[i].X;
+				int32 VertexY = y + QuadPattern[i].Y;
+				Positions[i] = CDI.GetWorldVertex(VertexX, VertexY);
+
+				// If at least one vertex is within the given bounds we should process the quad  
+				if (!bProcess && InBounds.ComputeSquaredDistanceFromBoxToPoint(Positions[i]) < SquaredSphereRadius)
+				{
+					bProcess = true;
+				}
+			}
+
+			if (bProcess)
+			{
+				//Fill the vertexID we need
+				TArray<FVertexID> VertexIDs;
+				VertexIDs.Reserve(UE_ARRAY_COUNT(QuadPattern));
+				TArray<FVertexInstanceID> VertexInstanceIDs;
+				VertexInstanceIDs.Reserve(UE_ARRAY_COUNT(QuadPattern));
+				// Fill positions
+				for (int32 i = 0; i < UE_ARRAY_COUNT(QuadPattern); i++)
+				{
+					int32 DuplicateLowestIndex = FindPreviousIndex(CurrentIndex);
+					FVertexID VertexID;
+					if (DuplicateLowestIndex < CurrentIndex)
+					{
+						VertexID = IndexToVertexID[DuplicateLowestIndex];
+					}
+					else
+					{
+						VertexID = OutMesh.CreateVertex();
+						VertexPositions[VertexID] = Positions[i];
+					}
+					IndexToVertexID.Add(CurrentIndex, VertexID);
+					VertexIDs.Add(VertexID);
+					CurrentIndex++;
+				}
+
+				// Create triangle
+				{
+					// Whether this vertex is in hole
+					bool bInvisible = false;
+					if (VisDataMap.Num())
+					{
+						int32 TexelX, TexelY;
+						CDI.VertexXYToTexelXY(x, y, TexelX, TexelY);
+						bInvisible = (VisDataMap[CDI.TexelXYToIndex(TexelX, TexelY)] >= VisThreshold);
+					}
+					//Add vertexInstance and polygon only if we are visible
+					if (!bInvisible)
+					{
+						VertexInstanceIDs.Add(OutMesh.CreateVertexInstance(VertexIDs[0]));
+						VertexInstanceIDs.Add(OutMesh.CreateVertexInstance(VertexIDs[1]));
+						VertexInstanceIDs.Add(OutMesh.CreateVertexInstance(VertexIDs[2]));
+
+						VertexInstanceIDs.Add(OutMesh.CreateVertexInstance(VertexIDs[3]));
+						VertexInstanceIDs.Add(OutMesh.CreateVertexInstance(VertexIDs[4]));
+						VertexInstanceIDs.Add(OutMesh.CreateVertexInstance(VertexIDs[5]));
+
+						// Fill other vertex data
+						for (int32 i = 0; i < UE_ARRAY_COUNT(QuadPattern); i++)
+						{
+							int32 VertexX = x + QuadPattern[i].X;
+							int32 VertexY = y + QuadPattern[i].Y;
+
+							FVector LocalTangentX, LocalTangentY, LocalTangentZ;
+							CDI.GetLocalTangentVectors(VertexX, VertexY, LocalTangentX, LocalTangentY, LocalTangentZ);
+
+							VertexInstanceTangents[VertexInstanceIDs[i]] = LocalTangentX;
+							VertexInstanceBinormalSigns[VertexInstanceIDs[i]] = GetBasisDeterminantSign(LocalTangentX, LocalTangentY, LocalTangentZ);
+							VertexInstanceNormals[VertexInstanceIDs[i]] = LocalTangentZ;
+
+							FVector2f UV = ComponentUVOffsetLOD + FVector2f(VertexX, VertexY) * ComponentUVScaleLOD;
+							VertexInstanceUVs.Set(VertexInstanceIDs[i], 0, UV);
+							// Add lightmap UVs
+							VertexInstanceUVs.Set(VertexInstanceIDs[i], 1, UV);
+						}
+						auto AddTriangle = [&OutMesh, &EdgeHardnesses, &PolygonGroupID, &VertexIDs, &VertexInstanceIDs](int32 BaseIndex)
+						{
+							//Create a polygon from this triangle
+							TArray<FVertexInstanceID> PerimeterVertexInstances;
+							PerimeterVertexInstances.SetNum(3);
+							for (int32 Corner = 0; Corner < 3; ++Corner)
+							{
+								PerimeterVertexInstances[Corner] = VertexInstanceIDs[BaseIndex + Corner];
+							}
+							// Insert a polygon into the mesh
+							TArray<FEdgeID> NewEdgeIDs;
+							const FPolygonID NewPolygonID = OutMesh.CreatePolygon(PolygonGroupID, PerimeterVertexInstances, &NewEdgeIDs);
+							for (const FEdgeID& NewEdgeID : NewEdgeIDs)
+							{
+								EdgeHardnesses[NewEdgeID] = false;
+							}
+						};
+						AddTriangle(0);
+						AddTriangle(3);
+					}
+				}
+			}
+			else
+			{
+				CurrentIndex += UE_ARRAY_COUNT(QuadPattern);
+			}
+		}
+	}
+}
+
+void ULandscapeComponent::ExportToMeshDescription(const int32 InExportLOD, FMeshDescription& OutMesh)
+{
+	::ExportToMeshDescription(this, InExportLOD, OutMesh, FBoxSphereBounds(), true);
+}
+
+void ULandscapeComponent::ExportToMeshDescription(const int32 InExportLOD, const FBoxSphereBounds& InBounds, FMeshDescription& OutMesh)
+{
+	::ExportToMeshDescription(this, InExportLOD, OutMesh, InBounds, false);
+}
+
 #endif
diff --git a/Engine/Source/Runtime/Landscape/Private/LandscapeEdit.cpp b/Engine/Source/Runtime/Landscape/Private/LandscapeEdit.cpp
--- a/Engine/Source/Runtime/Landscape/Private/LandscapeEdit.cpp
+++ b/Engine/Source/Runtime/Landscape/Private/LandscapeEdit.cpp
@@ -3339,31 +3339,6 @@ bool ALandscapeProxy::ExportToRawMesh(int32 InExportLOD, FMeshDescription& OutRa
 	TInlineComponentArray<ULandscapeComponent*> RegisteredLandscapeComponents;
 	GetComponents<ULandscapeComponent>(RegisteredLandscapeComponents);
 
-	const FIntRect LandscapeSectionRect = GetBoundingRect();
-	const FVector2f LandscapeUVScale = FVector2f(1.0f, 1.0f) / FVector2f(LandscapeSectionRect.Size());
-
-	FStaticMeshAttributes Attributes(OutRawMesh);
-	TVertexAttributesRef<FVector3f> VertexPositions = Attributes.GetVertexPositions();
-	TEdgeAttributesRef<bool> EdgeHardnesses = Attributes.GetEdgeHardnesses();
-	TPolygonGroupAttributesRef<FName> PolygonGroupImportedMaterialSlotNames = Attributes.GetPolygonGroupMaterialSlotNames();
-	TVertexInstanceAttributesRef<FVector3f> VertexInstanceNormals = Attributes.GetVertexInstanceNormals();
-	TVertexInstanceAttributesRef<FVector3f> VertexInstanceTangents = Attributes.GetVertexInstanceTangents();
-	TVertexInstanceAttributesRef<float> VertexInstanceBinormalSigns = Attributes.GetVertexInstanceBinormalSigns();
-	TVertexInstanceAttributesRef<FVector4f> VertexInstanceColors = Attributes.GetVertexInstanceColors();
-	TVertexInstanceAttributesRef<FVector2f> VertexInstanceUVs = Attributes.GetVertexInstanceUVs();
-
-	if (VertexInstanceUVs.GetNumChannels() < 2)
-	{
-		VertexInstanceUVs.SetNumChannels(2);
-	}
-
-	// User specified LOD to export
-	int32 LandscapeLODToExport = ExportLOD;
-	if (InExportLOD != INDEX_NONE)
-	{
-		LandscapeLODToExport = FMath::Clamp<int32>(InExportLOD, 0, FMath::CeilLogTwo(SubsectionSizeQuads + 1) - 1);
-	}
-
 	// Export data for each component
 	for (auto It = RegisteredLandscapeComponents.CreateConstIterator(); It; ++It)
 	{
@@ -3374,232 +3349,8 @@ bool ALandscapeProxy::ExportToRawMesh(int32 InExportLOD, FMeshDescription& OutRa
 		{
 			continue;
 		}
-
-		FLandscapeComponentDataInterface CDI(Component, LandscapeLODToExport);
-		const int32 ComponentSizeQuadsLOD = ((Component->ComponentSizeQuads + 1) >> LandscapeLODToExport) - 1;
-		const int32 SubsectionSizeQuadsLOD = ((Component->SubsectionSizeQuads + 1) >> LandscapeLODToExport) - 1;
-		const FIntPoint ComponentOffsetQuads = Component->GetSectionBase() - LandscapeSectionOffset - LandscapeSectionRect.Min;
-		const FVector2f ComponentUVOffsetLOD = FVector2f(ComponentOffsetQuads)*((float)ComponentSizeQuadsLOD / ComponentSizeQuads);
-		const FVector2f ComponentUVScaleLOD = LandscapeUVScale*((float)ComponentSizeQuads / ComponentSizeQuadsLOD);
-
-		const int32 NumFaces = FMath::Square(ComponentSizeQuadsLOD) * 2;
-		const int32 NumVertices = NumFaces * 3;
-
-		OutRawMesh.ReserveNewVertices(NumVertices);
-		OutRawMesh.ReserveNewPolygons(NumFaces);
-		OutRawMesh.ReserveNewVertexInstances(NumVertices);
-		OutRawMesh.ReserveNewEdges(NumVertices);
-
-		FPolygonGroupID PolygonGroupID = INDEX_NONE;
-		if (OutRawMesh.PolygonGroups().Num() < 1)
-		{
-			PolygonGroupID = OutRawMesh.CreatePolygonGroup();
-			PolygonGroupImportedMaterialSlotNames[PolygonGroupID] = FName(TEXT("LandscapeMat_0"));
-		}
-		else
-		{
-			PolygonGroupID = OutRawMesh.PolygonGroups().GetFirstValidID();
-		}
-
-		// Check if there are any holes
-		const int32 VisThreshold = 170;
-		TArray<uint8> VisDataMap;
-		const TArray<FWeightmapLayerAllocationInfo>& ComponentWeightmapLayerAllocations = Component->GetWeightmapLayerAllocations();
-
-		for (int32 AllocIdx = 0; AllocIdx < ComponentWeightmapLayerAllocations.Num(); AllocIdx++)
-		{
-			const FWeightmapLayerAllocationInfo& AllocInfo = ComponentWeightmapLayerAllocations[AllocIdx];
-			if (AllocInfo.LayerInfo == ALandscapeProxy::VisibilityLayer)
-			{
-				CDI.GetWeightmapTextureData(AllocInfo.LayerInfo, VisDataMap);
-			}
-		}
-
-		const FIntPoint QuadPattern[6] =
-		{
-			//face 1
-			FIntPoint(0, 0),
-			FIntPoint(0, 1),
-			FIntPoint(1, 1),
-			//face 2
-			FIntPoint(0, 0),
-			FIntPoint(1, 1),
-			FIntPoint(1, 0),
-		};
-
-		const int32 WeightMapSize = (SubsectionSizeQuadsLOD + 1) * Component->NumSubsections;
-
-		const float SquaredSphereRadius = FMath::Square(InBounds.SphereRadius);
-
-		//We need to not duplicate the vertex position, so we use the FIndexAndZ to achieve fast result
-		TArray<FIndexAndZ> VertIndexAndZ;
-		VertIndexAndZ.Reserve(ComponentSizeQuadsLOD*ComponentSizeQuadsLOD*UE_ARRAY_COUNT(QuadPattern));
-		int32 CurrentIndex = 0;
-		TMap<int32, FVector> IndexToPosition;
-		IndexToPosition.Reserve(ComponentSizeQuadsLOD*ComponentSizeQuadsLOD*UE_ARRAY_COUNT(QuadPattern));
-		for (int32 y = 0; y < ComponentSizeQuadsLOD; y++)
-		{
-			for (int32 x = 0; x < ComponentSizeQuadsLOD; x++)
-			{
-				for (int32 i = 0; i < UE_ARRAY_COUNT(QuadPattern); i++)
-				{
-					int32 VertexX = x + QuadPattern[i].X;
-					int32 VertexY = y + QuadPattern[i].Y;
-					FVector Position = CDI.GetWorldVertex(VertexX, VertexY);
-
-					// If at least one vertex is within the given bounds we should process the quad  
-					new(VertIndexAndZ)FIndexAndZ(CurrentIndex, Position);
-					IndexToPosition.Add(CurrentIndex, Position);
-					CurrentIndex++;
-				}
-			}
-		}
-		// Sort the vertices by z value
-		VertIndexAndZ.Sort(FCompareIndexAndZ());
-
-		auto FindPreviousIndex = [&VertIndexAndZ, &IndexToPosition](int32 Index)->int32
-		{
-			const FVector& PositionA = IndexToPosition[Index];
-			FIndexAndZ CompressPosition(0, PositionA);
-			// Search for lowest index duplicates
-			int32 BestIndex = MAX_int32;
-			for (int32 i = 0; i < IndexToPosition.Num(); i++)
-			{
-				if (CompressPosition.Z > (VertIndexAndZ[i].Z + SMALL_NUMBER))
-				{
-					//We will not find anything there is no point searching more
-					break;
-				}
-				const FVector& PositionB = IndexToPosition[VertIndexAndZ[i].Index];
-				if (PointsEqual(PositionA, PositionB, SMALL_NUMBER))
-				{
-					if (VertIndexAndZ[i].Index < BestIndex)
-					{
-						BestIndex = VertIndexAndZ[i].Index;
-					}
-				}
-			}
-			return BestIndex < MAX_int32 ? BestIndex : Index;
-		};
-
-		// Export to MeshDescription
-		TMap<int32, FVertexID> IndexToVertexID;
-		IndexToVertexID.Reserve(CurrentIndex);
-		CurrentIndex = 0;
-		for (int32 y = 0; y < ComponentSizeQuadsLOD; y++)
-		{
-			for (int32 x = 0; x < ComponentSizeQuadsLOD; x++)
-			{
-				FVector Positions[UE_ARRAY_COUNT(QuadPattern)];
-				bool bProcess = bIgnoreBounds;
-
-				// Fill positions
-				for (int32 i = 0; i < UE_ARRAY_COUNT(QuadPattern); i++)
-				{
-					int32 VertexX = x + QuadPattern[i].X;
-					int32 VertexY = y + QuadPattern[i].Y;
-					Positions[i] = CDI.GetWorldVertex(VertexX, VertexY);
-
-					// If at least one vertex is within the given bounds we should process the quad  
-					if (!bProcess && InBounds.ComputeSquaredDistanceFromBoxToPoint(Positions[i]) < SquaredSphereRadius)
-					{
-						bProcess = true;
-					}
-				}
-
-				if (bProcess)
-				{
-					//Fill the vertexID we need
-					TArray<FVertexID> VertexIDs;
-					VertexIDs.Reserve(UE_ARRAY_COUNT(QuadPattern));
-					TArray<FVertexInstanceID> VertexInstanceIDs;
-					VertexInstanceIDs.Reserve(UE_ARRAY_COUNT(QuadPattern));
-					// Fill positions
-					for (int32 i = 0; i < UE_ARRAY_COUNT(QuadPattern); i++)
-					{
-						int32 DuplicateLowestIndex = FindPreviousIndex(CurrentIndex);
-						FVertexID VertexID;
-						if (DuplicateLowestIndex < CurrentIndex)
-						{
-							VertexID = IndexToVertexID[DuplicateLowestIndex];
-						}
-						else
-						{
-							VertexID = OutRawMesh.CreateVertex();
-							VertexPositions[VertexID] = Positions[i];
-						}
-						IndexToVertexID.Add(CurrentIndex, VertexID);
-						VertexIDs.Add(VertexID);
-						CurrentIndex++;
-					}
-
-					// Create triangle
-					{
-						// Whether this vertex is in hole
-						bool bInvisible = false;
-						if (VisDataMap.Num())
-						{
-							int32 TexelX, TexelY;
-							CDI.VertexXYToTexelXY(x, y, TexelX, TexelY);
-							bInvisible = (VisDataMap[CDI.TexelXYToIndex(TexelX, TexelY)] >= VisThreshold);
-						}
-						//Add vertexInstance and polygon only if we are visible
-						if (!bInvisible)
-						{
-							VertexInstanceIDs.Add(OutRawMesh.CreateVertexInstance(VertexIDs[0]));
-							VertexInstanceIDs.Add(OutRawMesh.CreateVertexInstance(VertexIDs[1]));
-							VertexInstanceIDs.Add(OutRawMesh.CreateVertexInstance(VertexIDs[2]));
-
-							VertexInstanceIDs.Add(OutRawMesh.CreateVertexInstance(VertexIDs[3]));
-							VertexInstanceIDs.Add(OutRawMesh.CreateVertexInstance(VertexIDs[4]));
-							VertexInstanceIDs.Add(OutRawMesh.CreateVertexInstance(VertexIDs[5]));
-
-							// Fill other vertex data
-							for (int32 i = 0; i < UE_ARRAY_COUNT(QuadPattern); i++)
-							{
-								int32 VertexX = x + QuadPattern[i].X;
-								int32 VertexY = y + QuadPattern[i].Y;
-
-								FVector LocalTangentX, LocalTangentY, LocalTangentZ;
-								CDI.GetLocalTangentVectors(VertexX, VertexY, LocalTangentX, LocalTangentY, LocalTangentZ);
-
-								VertexInstanceTangents[VertexInstanceIDs[i]] = LocalTangentX;
-								VertexInstanceBinormalSigns[VertexInstanceIDs[i]] = GetBasisDeterminantSign(LocalTangentX, LocalTangentY, LocalTangentZ);
-								VertexInstanceNormals[VertexInstanceIDs[i]] = LocalTangentZ;
-
-								FVector2f UV = (ComponentUVOffsetLOD + FVector2f(VertexX, VertexY))*ComponentUVScaleLOD;
-								VertexInstanceUVs.Set(VertexInstanceIDs[i], 0, UV);
-								// Add lightmap UVs
-								VertexInstanceUVs.Set(VertexInstanceIDs[i], 1, UV);
-							}
-							auto AddTriangle = [&OutRawMesh, &EdgeHardnesses, &PolygonGroupID, &VertexIDs, &VertexInstanceIDs](int32 BaseIndex)
-							{
-								//Create a polygon from this triangle
-								TArray<FVertexInstanceID> PerimeterVertexInstances;
-								PerimeterVertexInstances.SetNum(3);
-								for (int32 Corner = 0; Corner < 3; ++Corner)
-								{
-									PerimeterVertexInstances[Corner] = VertexInstanceIDs[BaseIndex + Corner];
-								}
-								// Insert a polygon into the mesh
-								TArray<FEdgeID> NewEdgeIDs;
-								const FPolygonID NewPolygonID = OutRawMesh.CreatePolygon(PolygonGroupID, PerimeterVertexInstances, &NewEdgeIDs);
-								for (const FEdgeID& NewEdgeID : NewEdgeIDs)
-								{
-									EdgeHardnesses[NewEdgeID] = false;
-								}
-							};
-							AddTriangle(0);
-							AddTriangle(3);
-						}
-					}
-				}
-				else
-				{
-					CurrentIndex += UE_ARRAY_COUNT(QuadPattern);
-				}
-			}
-		}
+		
+		Component->ExportToMeshDescription(InExportLOD, OutRawMesh);
 	}
 
 	//Compact the MeshDescription, if there was visibility mask or some bounding box clip, it need to be compacted so the sparse array are from 0 to n with no invalid data in between. 
@@ -7650,4 +7401,4 @@ bool ALandscapeProxy::LandscapeExportWeightmapToRenderTarget(UTextureRenderTarge
 
 #endif //WITH_EDITOR
 
-#undef LOCTEXT_NAMESPACE
+#undef LOCTEXT_NAMESPACE
\ No newline at end of file
