Fix incorrect geometric normals on path traced geometry with WPO materials

The post-process which applies WPO usually leaves mesh vertices in world space. However the vertex factory might still contains a local to world transform which is required to rotate the shading normals correctly. We thus end up with a mismatch when the world position is computed for ray tracing. Currently only the path tracer was manifesting an artifact due to this, but the inconsistency could have shown up in other ways (for example if a material had a dependency on world space position).

The fix is to simply keep the mesh in "local" space by instructing the RayTracingDynamicMesh compute shader to put the vertex back into local space after applying WPO.

All the geometry types were reviewed for this potential issue and the RayTracingDynamicMesh has been updated to mandate a preprocessor signal on how to handle this.

diff --git a/Engine/Plugins/FX/Niagara/Source/NiagaraVertexFactories/Private/NiagaraRibbonVertexFactory.cpp b/Engine/Plugins/FX/Niagara/Source/NiagaraVertexFactories/Private/NiagaraRibbonVertexFactory.cpp
--- a/Engine/Plugins/FX/Niagara/Source/NiagaraVertexFactories/Private/NiagaraRibbonVertexFactory.cpp
+++ b/Engine/Plugins/FX/Niagara/Source/NiagaraVertexFactories/Private/NiagaraRibbonVertexFactory.cpp
@@ -132,6 +132,9 @@ void FNiagaraRibbonVertexFactory::ModifyCompilationEnvironment(const FVertexFact
 	OutEnvironment.SetDefine(TEXT("NiagaraVFLooseParameters"), TEXT("NiagaraRibbonVFLooseParameters"));
 	
 	OutEnvironment.SetDefine(TEXT("NIAGARA_RIBBON_FACTORY"), TEXT("1"));
+
+	// Ribbons are generated in world space and never have a matrix transform in raytracing, so it is safe to leave them in world space.
+	OutEnvironment.SetDefine(TEXT("RAY_TRACING_DYNAMIC_MESH_IN_WORLD_SPACE"), TEXT("1"));
 }
 
 /**
diff --git a/Engine/Plugins/FX/Niagara/Source/NiagaraVertexFactories/Private/NiagaraSpriteVertexFactory.cpp b/Engine/Plugins/FX/Niagara/Source/NiagaraVertexFactories/Private/NiagaraSpriteVertexFactory.cpp
--- a/Engine/Plugins/FX/Niagara/Source/NiagaraVertexFactories/Private/NiagaraSpriteVertexFactory.cpp
+++ b/Engine/Plugins/FX/Niagara/Source/NiagaraVertexFactories/Private/NiagaraSpriteVertexFactory.cpp
@@ -174,6 +174,9 @@ void FNiagaraSpriteVertexFactory::ModifyCompilationEnvironment(const FVertexFact
 
 	// Set a define so we can tell in MaterialTemplate.usf when we are compiling a sprite vertex factory
 	OutEnvironment.SetDefine(TEXT("PARTICLE_SPRITE_FACTORY"),TEXT("1"));
+
+	// Sprites are generated in world space and never have a matrix transform in raytracing, so it is safe to leave them in world space.
+	OutEnvironment.SetDefine(TEXT("RAY_TRACING_DYNAMIC_MESH_IN_WORLD_SPACE"), TEXT("1"));
 }
 
 /**
diff --git a/Engine/Plugins/Runtime/HairStrands/Source/HairStrandsCore/Private/HairStrandsVertexFactory.cpp b/Engine/Plugins/Runtime/HairStrands/Source/HairStrandsCore/Private/HairStrandsVertexFactory.cpp
--- a/Engine/Plugins/Runtime/HairStrands/Source/HairStrandsCore/Private/HairStrandsVertexFactory.cpp
+++ b/Engine/Plugins/Runtime/HairStrands/Source/HairStrandsCore/Private/HairStrandsVertexFactory.cpp
@@ -261,6 +261,5 @@ IMPLEMENT_VERTEX_FACTORY_TYPE(FHairStrandsVertexFactory, "/Engine/Private/HairSt
 	| EVertexFactoryFlags::SupportsCachingMeshDrawCommands
 	| (VF_STRANDS_SUPPORT_GPU_SCENE ? EVertexFactoryFlags::SupportsPrimitiveIdStream : EVertexFactoryFlags::None)
 	| EVertexFactoryFlags::SupportsRayTracing
-	| EVertexFactoryFlags::SupportsRayTracingDynamicGeometry
 	| (VF_STRANDS_PROCEDURAL_INTERSECTOR ? EVertexFactoryFlags::SupportsRayTracingProceduralPrimitive : EVertexFactoryFlags::None)
 );
diff --git a/Engine/Shaders/Private/RayTracing/RayTracingDynamicMesh.usf b/Engine/Shaders/Private/RayTracing/RayTracingDynamicMesh.usf
--- a/Engine/Shaders/Private/RayTracing/RayTracingDynamicMesh.usf
+++ b/Engine/Shaders/Private/RayTracing/RayTracingDynamicMesh.usf
@@ -96,6 +96,15 @@ void RayTracingDynamicGeometryConverterCS(uint3 DispatchThreadId : SV_DispatchTh
 #if USE_INSTANCING
 		//Reverse transform to neutral space
 		WorldPosition = mul(WorldPosition, InverseTransform);
+#elif RAY_TRACING_DYNAMIC_MESH_IN_LOCAL_SPACE
+		// Move the point back into local space because the RT instance will be placed there
+		// This the recommended default behavior so that the transform applied by the VertexFactory and the RT transform are always the same
+		WorldPosition.xyz = LWCMultiply(LWCPromote(WorldPosition), VertexFactoryGetWorldToLocal(VFIntermediates));
+#elif RAY_TRACING_DYNAMIC_MESH_IN_WORLD_SPACE
+		// There are a few geometry types which natively generate positions in world space and therefore not need any extra processing here.
+		// This behavior must be explicitly opted-in to so as to make the code more self-documenting and avoid un-intended mismatched transforms.
+#else
+#error "Are you sure you want mesh vertices to be in world space? Please update the VertexFactory to report its requirement!"
 #endif
 
 		RWVertexPositions[OutputVertexBaseIndex + VertexIndex * 3 + 0] = WorldPosition.x;
diff --git a/Engine/Source/Runtime/Engine/Private/GeometryCollectionRendering.cpp b/Engine/Source/Runtime/Engine/Private/GeometryCollectionRendering.cpp
--- a/Engine/Source/Runtime/Engine/Private/GeometryCollectionRendering.cpp
+++ b/Engine/Source/Runtime/Engine/Private/GeometryCollectionRendering.cpp
@@ -64,6 +64,8 @@ void FGeometryCollectionVertexFactory::ModifyCompilationEnvironment(const FVerte
 
 	// Geometry collections use a custom hit proxy per bone
 	OutEnvironment.SetDefine(TEXT("USE_PER_VERTEX_HITPROXY_ID"), 1);
+
+	OutEnvironment.SetDefine(TEXT("RAY_TRACING_DYNAMIC_MESH_IN_LOCAL_SPACE"), TEXT("1"));
 }
 
 void FGeometryCollectionVertexFactory::ValidateCompiledResult(const FVertexFactoryType* Type, EShaderPlatform Platform, const FShaderParameterMap& ParameterMap, TArray<FString>& OutErrors)
diff --git a/Engine/Source/Runtime/Engine/Private/LocalVertexFactory.cpp b/Engine/Source/Runtime/Engine/Private/LocalVertexFactory.cpp
--- a/Engine/Source/Runtime/Engine/Private/LocalVertexFactory.cpp
+++ b/Engine/Source/Runtime/Engine/Private/LocalVertexFactory.cpp
@@ -214,6 +214,11 @@ void FLocalVertexFactory::ModifyCompilationEnvironment(const FVertexFactoryShade
 
 	const bool bVFSupportsPrimtiveSceneData = Parameters.VertexFactoryType->SupportsPrimitiveIdStream() && UseGPUScene(Parameters.Platform, GetMaxSupportedFeatureLevel(Parameters.Platform));
 	OutEnvironment.SetDefine(TEXT("VF_SUPPORTS_PRIMITIVE_SCENE_DATA"), bVFSupportsPrimtiveSceneData);
+
+	// When combining ray tracing and WPO, leave the mesh in local space for consistency with how shading normals are calculated.
+	// See UE-139634 for the case that lead to this.
+	OutEnvironment.SetDefine(TEXT("RAY_TRACING_DYNAMIC_MESH_IN_LOCAL_SPACE"), TEXT("1"));
+
 }
 
 void FLocalVertexFactory::ValidateCompiledResult(const FVertexFactoryType* Type, EShaderPlatform Platform, const FShaderParameterMap& ParameterMap, TArray<FString>& OutErrors)
diff --git a/Engine/Source/Runtime/Engine/Private/SkeletalMesh.cpp b/Engine/Source/Runtime/Engine/Private/SkeletalMesh.cpp
--- a/Engine/Source/Runtime/Engine/Private/SkeletalMesh.cpp
+++ b/Engine/Source/Runtime/Engine/Private/SkeletalMesh.cpp
@@ -6621,15 +6621,7 @@ void FSkeletalMeshSceneProxy::GetDynamicRayTracingInstances(FRayTracingMaterialG
 				RayTracingInstance.Materials.Add(MeshBatch);
 			}
 
-			if (bAnySegmentUsesWorldPositionOffset)
-			{
-				RayTracingInstance.InstanceTransforms.Add(FMatrix::Identity);
-			}
-			else
-			{
-				RayTracingInstance.InstanceTransforms.Add(GetLocalToWorld());
-			}
-
+			RayTracingInstance.InstanceTransforms.Add(GetLocalToWorld());
 			const uint32 VertexBufferStride = LODData.StaticVertexBuffers.PositionVertexBuffer.GetStride();
 
 			const FVertexFactory* VertexFactory = MeshObject->GetSkinVertexFactory(Context.ReferenceView, LODIndex, 0, ESkinVertexFactoryMode::RayTracing);
diff --git a/Engine/Source/Runtime/Engine/Private/StaticMeshRender.cpp b/Engine/Source/Runtime/Engine/Private/StaticMeshRender.cpp
--- a/Engine/Source/Runtime/Engine/Private/StaticMeshRender.cpp
+++ b/Engine/Source/Runtime/Engine/Private/StaticMeshRender.cpp
@@ -1765,11 +1765,9 @@ void FStaticMeshSceneProxy::GetDynamicRayTracingInstances(FRayTracingMaterialGat
 		}
 
 		RayTracingInstance.Geometry = &Geometry;
-
+		RayTracingInstance.InstanceTransforms.Add(GetLocalToWorld());
 		if (bEvaluateWPO && RenderData->LODVertexFactories[LODIndex].VertexFactory.GetType()->SupportsRayTracingDynamicGeometry())
 		{
-			RayTracingInstance.InstanceTransforms.Add(FMatrix::Identity);
-	
 			// Use the internal vertex buffer only when initialized otherwise used the shared vertex buffer - needs to be updated every frame
 			FRWBuffer* VertexBuffer = nullptr;
 			if (DynamicRayTracingGeometryVertexBuffers.Num() > (int32)LODIndex && DynamicRayTracingGeometryVertexBuffers[LODIndex].NumBytes > 0)
@@ -1789,11 +1787,7 @@ void FStaticMeshSceneProxy::GetDynamicRayTracingInstances(FRayTracingMaterialGat
 					VertexBuffer,
 					true
 				}
-		);
-		}
-		else
-		{
-			RayTracingInstance.InstanceTransforms.Add(GetLocalToWorld());
+			);
 		}
 		
 		RayTracingInstance.BuildInstanceMaskAndFlags(GetScene().GetFeatureLevel());
diff --git a/Engine/Source/Runtime/Experimental/GeometryCollectionEngine/Private/GeometryCollection/GeometryCollectionSceneProxy.cpp b/Engine/Source/Runtime/Experimental/GeometryCollectionEngine/Private/GeometryCollection/GeometryCollectionSceneProxy.cpp
--- a/Engine/Source/Runtime/Experimental/GeometryCollectionEngine/Private/GeometryCollection/GeometryCollectionSceneProxy.cpp
+++ b/Engine/Source/Runtime/Experimental/GeometryCollectionEngine/Private/GeometryCollection/GeometryCollectionSceneProxy.cpp
@@ -1031,11 +1031,7 @@ void FGeometryCollectionSceneProxy::GetDynamicRayTracingInstances(FRayTracingMat
 
 			FRayTracingInstance RayTracingInstance;
 			RayTracingInstance.Geometry = &RayTracingGeometry;
-			RayTracingInstance.InstanceTransforms.Reserve(DynamicData->Transforms.Num());
-			for (int32 TransformIndex = 0; TransformIndex < DynamicData->Transforms.Num(); ++TransformIndex)
-			{
-				RayTracingInstance.InstanceTransforms.Emplace(FMatrix44f::Identity);
-			}
+			RayTracingInstance.InstanceTransforms.Append(DynamicData->Transforms);
 
 			// Grab the material proxies we'll be using for each section
 			TArray<FMaterialRenderProxy*, TInlineAllocator<32>> MaterialProxies;
diff --git a/Engine/Source/Runtime/Landscape/Private/LandscapeRender.cpp b/Engine/Source/Runtime/Landscape/Private/LandscapeRender.cpp
--- a/Engine/Source/Runtime/Landscape/Private/LandscapeRender.cpp
+++ b/Engine/Source/Runtime/Landscape/Private/LandscapeRender.cpp
@@ -2763,7 +2763,7 @@ void FLandscapeComponentSceneProxy::GetDynamicRayTracingInstances(FRayTracingMat
 
 			FRayTracingInstance RayTracingInstance;
 			RayTracingInstance.Geometry = &SectionRayTracingStates[SubSectionIdx].Geometry;
-			RayTracingInstance.InstanceTransforms.Add(FMatrix::Identity);
+			RayTracingInstance.InstanceTransforms.Add(GetLocalToWorld());
 			RayTracingInstance.Materials.Add(MeshBatch);
 			RayTracingInstance.BuildInstanceMaskAndFlags(GetScene().GetFeatureLevel());
 			OutRayTracingInstances.Add(RayTracingInstance);
@@ -3315,6 +3315,9 @@ void FLandscapeVertexFactory::ModifyCompilationEnvironment(const FVertexFactoryS
 	FVertexFactory::ModifyCompilationEnvironment(Parameters, OutEnvironment);
 
 	OutEnvironment.SetDefine(TEXT("VF_SUPPORTS_PRIMITIVE_SCENE_DATA"), Parameters.VertexFactoryType->SupportsPrimitiveIdStream() && UseGPUScene(Parameters.Platform, GetMaxSupportedFeatureLevel(Parameters.Platform)));
+
+	// Make sure landscape vertices go back to local space so that we have consistency between the transform on normals and geometry
+	OutEnvironment.SetDefine(TEXT("RAY_TRACING_DYNAMIC_MESH_IN_LOCAL_SPACE"), TEXT("1"));
 }
 
 IMPLEMENT_VERTEX_FACTORY_PARAMETER_TYPE(FLandscapeVertexFactory, SF_Vertex, FLandscapeVertexFactoryVertexShaderParameters);
