WP Landscape HLOD - Properly stitch landscape HLODs & the actual landscape components

FLandscapeMeshProxySceneProxy now register the LOD they represent in the LandscapeRenderSystem object, causing the neighbor landscape components to morph their vertices to fill any gap.

diff --git a/Engine/Source/Runtime/Landscape/Private/LandscapeRender.cpp b/Engine/Source/Runtime/Landscape/Private/LandscapeRender.cpp
--- a/Engine/Source/Runtime/Landscape/Private/LandscapeRender.cpp
+++ b/Engine/Source/Runtime/Landscape/Private/LandscapeRender.cpp
@@ -45,14 +45,6 @@ IMPLEMENT_GLOBAL_SHADER_PARAMETER_STRUCT(FLandscapeUniformShaderParameters, "Lan
 IMPLEMENT_GLOBAL_SHADER_PARAMETER_STRUCT(FLandscapeFixedGridUniformShaderParameters, "LandscapeFixedGrid");
 IMPLEMENT_TYPE_LAYOUT(FLandscapeVertexFactoryPixelShaderParameters);
 
-int32 GLandscapeMeshLODBias = 0;
-FAutoConsoleVariableRef CVarLandscapeMeshLODBias(
-	TEXT("r.LandscapeLODBias"),
-	GLandscapeMeshLODBias,
-	TEXT("LOD bias for landscape/terrain meshes."),
-	ECVF_Scalability
-);
-
 #if !UE_BUILD_SHIPPING
 static void OnLODDistributionScaleChanged(IConsoleVariable* CVar)
 {
@@ -103,6 +95,8 @@ static FAutoConsoleVariableRef CVarAllowLandscapeShadows(
 extern TAutoConsoleVariable<int32> CVarLandscapeShowDirty;
 #endif
 
+extern RENDERER_API TAutoConsoleVariable<float> CVarStaticMeshLODDistanceScale;
+
 #if !UE_BUILD_SHIPPING
 int32 GVarDumpLandscapeLODsCurrentFrame = 0;
 bool GVarDumpLandscapeLODs = false;
@@ -665,65 +659,6 @@ static int32 GetDrawCollisionLodOverride(FSceneView const& View, FCollisionRespo
 }
 
 
-//
-// FComputeSectionPerViewParametersTask
-//
-struct FComputeSectionPerViewParametersTask
-{
-	FLandscapeRenderSystem& RenderSystem;
-
-	FLandscapeRenderSystem::FViewParams ViewParams;
-
-	TArray<uint8> SectionCurrentFirstLODIndices;
-
-	FComputeSectionPerViewParametersTask(FLandscapeRenderSystem& InRenderSystem, const FLandscapeRenderSystem::FViewParams& InViewParams)
-		: RenderSystem(InRenderSystem)
-		, ViewParams(InViewParams)
-	{
-#if PLATFORM_SUPPORTS_LANDSCAPE_VISUAL_MESH_LOD_STREAMING
-		const int32 NumSceneProxies = InRenderSystem.SceneProxies.Num();
-		SectionCurrentFirstLODIndices.Empty(NumSceneProxies);
-		SectionCurrentFirstLODIndices.AddUninitialized(NumSceneProxies);
-
-		for (int32 Idx = 0; Idx < NumSceneProxies; ++Idx)
-		{
-			const FLandscapeComponentSceneProxy* Proxy = InRenderSystem.SceneProxies[Idx];
-			SectionCurrentFirstLODIndices[Idx] = Proxy ? Proxy->GetCurrentFirstLODIdx_RenderThread() : 0;
-		}
-#endif
-	}
-
-	FORCEINLINE TStatId GetStatId() const
-	{
-		RETURN_QUICK_DECLARE_CYCLE_STAT(FComputeSectionPerViewParametersTask, STATGROUP_TaskGraphTasks);
-	}
-
-	ENamedThreads::Type GetDesiredThread()
-	{
-		return ENamedThreads::AnyNormalThreadNormalTask;
-	}
-
-	static ESubsequentsMode::Type GetSubsequentsMode()
-	{
-		return ESubsequentsMode::TrackSubsequents;
-	}
-
-	void AnyThreadTask()
-	{
-		FOptionalTaskTagScope Scope(ETaskTag::EParallelRenderingThread);
-		RenderSystem.ComputeSectionPerViewParameters(
-			ViewParams.ViewKey, ViewParams.ViewLODOverride, ViewParams.ViewLODDistanceFactor,
-			ViewParams.ViewEngineShowFlagCollisionPawn, ViewParams.ViewEngineShowFlagCollisionVisibility,
-			ViewParams.ViewOrigin, ViewParams.ViewProjectionMatrix, SectionCurrentFirstLODIndices);
-	}
-
-	void DoTask(ENamedThreads::Type CurrentThread, const FGraphEventRef& MyCompletionGraphEvent)
-	{
-		AnyThreadTask();
-	}
-};
-
-
 //
 // FGetSectionLODBiasesTask
 //
@@ -768,25 +703,12 @@ struct FGetSectionLODBiasesTask
 // FLandscapeComponentSceneProxy
 //
 TMap<uint32, FLandscapeSharedBuffers*>FLandscapeComponentSceneProxy::SharedBuffersMap;
-TMap<FLandscapeNeighborInfo::FLandscapeKey, TMap<FIntPoint, const FLandscapeNeighborInfo*> > FLandscapeNeighborInfo::SharedSceneProxyMap;
 
 const static FName NAME_LandscapeResourceNameForDebugging(TEXT("Landscape"));
 
 IMPLEMENT_GLOBAL_SHADER_PARAMETER_STRUCT(FLandscapeSectionLODUniformParameters, "LandscapeContinuousLODParameters");
 
-TMap<FLandscapeNeighborInfo::FLandscapeKey, FLandscapeRenderSystem*> LandscapeRenderSystems;
-
-FLandscapeRenderSystem& GetOrCreateLandscapeRenderSystem(FLandscapeNeighborInfo::FLandscapeKey LandscapeKey)
-{
-	check(IsInRenderingThread());
-
-	FLandscapeRenderSystem*& LandscapeRenderSystem = LandscapeRenderSystems.FindOrAdd(LandscapeKey);
-	if (!LandscapeRenderSystem)
-	{
-		LandscapeRenderSystem = new FLandscapeRenderSystem();
-	}
-	return *LandscapeRenderSystem;
-}
+TMap<uint32, FLandscapeRenderSystem*> LandscapeRenderSystems;
 
 TBitArray<> FLandscapeRenderSystem::LandscapeIndexAllocator;
 
@@ -795,9 +717,9 @@ TBitArray<> FLandscapeRenderSystem::LandscapeIndexAllocator;
 // FLandscapeRenderSystem
 //
 FLandscapeRenderSystem::FLandscapeRenderSystem()
-	: NumRegisteredEntities(0)
-	, Min(MAX_int32, MAX_int32)
+	: Min(MAX_int32, MAX_int32)
 	, Size(EForceInit::ForceInitToZero)
+	, ReferenceCount(0)
 {
 	SectionLODBiases.SetAllowCPUAccess(true);
 
@@ -826,18 +748,46 @@ FLandscapeRenderSystem::~FLandscapeRenderSystem()
 	LandscapeIndexAllocator.SetNumUninitialized(LastSetIndex + 1);
 }
 
-void FLandscapeRenderSystem::RegisterEntity(FLandscapeComponentSceneProxy* SceneProxy)
+void FLandscapeRenderSystem::CreateResources(FLandscapeSectionInfo* SectionInfo)
+{
+	check(IsInRenderingThread());
+
+	FLandscapeRenderSystem*& LandscapeRenderSystem = LandscapeRenderSystems.FindOrAdd(SectionInfo->LandscapeKey);
+	if (!LandscapeRenderSystem)
+	{
+		LandscapeRenderSystem = new FLandscapeRenderSystem();
+	}
+
+	LandscapeRenderSystem->CreateResources_Internal(SectionInfo);
+}
+
+void FLandscapeRenderSystem::DestroyResources(FLandscapeSectionInfo* SectionInfo)
 {
 	check(IsInRenderingThread());
-	check(SceneProxy != nullptr);
 
-	if (NumRegisteredEntities > 0)
+	FLandscapeRenderSystem* LandscapeRenderSystem = LandscapeRenderSystems.FindChecked(SectionInfo->LandscapeKey);
+	LandscapeRenderSystem->DestroyResources_Internal(SectionInfo);
+
+	if (LandscapeRenderSystem->ReferenceCount == 0)
+	{
+		delete LandscapeRenderSystem;
+		LandscapeRenderSystems.Remove(SectionInfo->LandscapeKey);
+	}
+}
+
+void FLandscapeRenderSystem::CreateResources_Internal(FLandscapeSectionInfo* SectionInfo)
+{
+	check(IsInRenderingThread());
+	check(SectionInfo != nullptr);
+	check(!SectionInfo->bRegistered);
+
+	if (!SectionInfos.IsEmpty())
 	{
 		// Calculate new bounding rect of landscape components
 		FIntPoint OriginalMin = Min;
 		FIntPoint OriginalMax = Min + Size - FIntPoint(1, 1);
-		FIntPoint NewMin(FMath::Min(Min.X, SceneProxy->ComponentBase.X), FMath::Min(Min.Y, SceneProxy->ComponentBase.Y));
-		FIntPoint NewMax(FMath::Max(OriginalMax.X, SceneProxy->ComponentBase.X), FMath::Max(OriginalMax.Y, SceneProxy->ComponentBase.Y));
+		FIntPoint NewMin(FMath::Min(Min.X, SectionInfo->ComponentBase.X), FMath::Min(Min.Y, SectionInfo->ComponentBase.Y));
+		FIntPoint NewMax(FMath::Max(OriginalMax.X, SectionInfo->ComponentBase.X), FMath::Max(OriginalMax.Y, SectionInfo->ComponentBase.Y));
 
 		FIntPoint SizeRequired = (NewMax - NewMin) + FIntPoint(1, 1);
 
@@ -848,28 +798,88 @@ void FLandscapeRenderSystem::RegisterEntity(FLandscapeComponentSceneProxy* Scene
 	}
 	else
 	{
-		ResizeAndMoveTo(SceneProxy->ComponentBase, FIntPoint(1, 1));
+		ResizeAndMoveTo(SectionInfo->ComponentBase, FIntPoint(1, 1));
 	}
 
-	NumRegisteredEntities++;
-	SetSectionLODSettings(SceneProxy->ComponentBase, SceneProxy->LODSettings);
-	SetSectionOriginAndRadius(SceneProxy->ComponentBase, FVector4(SceneProxy->GetBounds().Origin, SceneProxy->GetBounds().SphereRadius));
-	SetSceneProxy(SceneProxy->ComponentBase, SceneProxy);
+	ReferenceCount++;
 }
 
-void FLandscapeRenderSystem::UnregisterEntity(FLandscapeComponentSceneProxy* SceneProxy)
+void FLandscapeRenderSystem::DestroyResources_Internal(FLandscapeSectionInfo* SectionInfo)
 {
 	check(IsInRenderingThread());
-	check(SceneProxy != nullptr);
+	check(SectionInfo != nullptr);
+	check(!SectionInfo->bRegistered);
 
-	SetSceneProxy(SceneProxy->ComponentBase, nullptr);
-	SetSectionOriginAndRadius(SceneProxy->ComponentBase, FVector4(ForceInitToZero));
+	ReferenceCount--;
+}
 
-	LODSettingsComponent LODSettings;
-	FMemory::Memzero(&LODSettings, sizeof(LODSettingsComponent));
-	SetSectionLODSettings(SceneProxy->ComponentBase, LODSettings);
+void FLandscapeRenderSystem::RegisterSection(FLandscapeSectionInfo* SectionInfo)
+{
+	check(IsInRenderingThread());
+	check(SectionInfo != nullptr);
+	check(!SectionInfo->bRegistered);
 
-	NumRegisteredEntities--;
+	// With HLODs, it's possible to have multiple loaded sections representing the same
+	// landscape patch. For example, raytracing may keep the HLOD proxy around (far field),
+	// even if the actual landscape is loaded & visible.
+	// We keep a linked list of the section infos, sorted by priority, so that unregistration can
+	// properly restore a previously registered section info.
+
+	FLandscapeRenderSystem*& LandscapeRenderSystem = LandscapeRenderSystems.FindChecked(SectionInfo->LandscapeKey);
+	FLandscapeSectionInfo* ExistingSection = LandscapeRenderSystem->GetSectionInfo(SectionInfo->ComponentBase);
+	if (ExistingSection == nullptr)
+	{
+		LandscapeRenderSystem->SetSectionInfo(SectionInfo->ComponentBase, SectionInfo);
+	}
+	else
+	{
+		FLandscapeSectionInfo* CurrentSection = nullptr;
+		FLandscapeSectionInfo::TIterator SectionIt(ExistingSection);
+		for (; SectionIt; ++SectionIt)
+		{
+			CurrentSection = &*SectionIt;
+
+			// Sort on insertion
+			if (SectionInfo->GetSectionPriority() < CurrentSection->GetSectionPriority())
+			{
+				SectionInfo->LinkBefore(CurrentSection);
+				break;
+			}
+		}
+
+		if (!SectionIt)
+		{
+			// Set as tail
+			SectionInfo->LinkAfter(CurrentSection);
+		}
+		else if (CurrentSection == ExistingSection)
+		{
+			// Set as head
+			LandscapeRenderSystem->SetSectionInfo(SectionInfo->ComponentBase, SectionInfo);
+		}
+	}
+
+	SectionInfo->bRegistered = true;
+}
+
+void FLandscapeRenderSystem::UnregisterSection(FLandscapeSectionInfo* SectionInfo)
+{
+	check(IsInRenderingThread());
+	check(SectionInfo != nullptr);
+	
+	if (SectionInfo->bRegistered)
+	{
+		FLandscapeRenderSystem* LandscapeRenderSystem = LandscapeRenderSystems.FindChecked(SectionInfo->LandscapeKey);
+		FLandscapeSectionInfo* ExistingSection = LandscapeRenderSystem->GetSectionInfo(SectionInfo->ComponentBase);
+		if (ExistingSection == SectionInfo)
+		{
+			LandscapeRenderSystem->SetSectionInfo(SectionInfo->ComponentBase, SectionInfo->GetNextLink());
+		}
+
+		SectionInfo->Unlink();
+
+		SectionInfo->bRegistered = false;
+	}
 }
 
 void FLandscapeRenderSystem::ResizeAndMoveTo(FIntPoint NewMin, FIntPoint NewSize)
@@ -877,14 +887,10 @@ void FLandscapeRenderSystem::ResizeAndMoveTo(FIntPoint NewMin, FIntPoint NewSize
 	SectionLODBiasBuffer.SafeRelease();
 
 	TResourceArray<float> NewSectionLODBiases;
-	TArray<LODSettingsComponent> NewSectionLODSettings;
-	TArray<FVector4> NewSectionOriginAndRadius;
-	TArray<FLandscapeComponentSceneProxy*> NewSceneProxies;
+	TArray<FLandscapeSectionInfo*> NewSectionInfos;
 
 	NewSectionLODBiases.AddZeroed(NewSize.X * NewSize.Y);
-	NewSectionLODSettings.AddZeroed(NewSize.X * NewSize.Y);
-	NewSectionOriginAndRadius.AddZeroed(NewSize.X * NewSize.Y);
-	NewSceneProxies.AddZeroed(NewSize.X * NewSize.Y);
+	NewSectionInfos.AddZeroed(NewSize.X * NewSize.Y);
 
 	for (int32 Y = 0; Y < Size.Y; Y++)
 	{
@@ -896,40 +902,19 @@ void FLandscapeRenderSystem::ResizeAndMoveTo(FIntPoint NewMin, FIntPoint NewSize
 			if (NewLinearIndex >= 0 && NewLinearIndex < NewSize.X * NewSize.Y)
 			{
 				NewSectionLODBiases[NewLinearIndex] = SectionLODBiases[LinearIndex];
-				NewSectionLODSettings[NewLinearIndex] = SectionLODSettings[LinearIndex];
-				NewSectionOriginAndRadius[NewLinearIndex] = SectionOriginAndRadius[LinearIndex];
-				NewSceneProxies[NewLinearIndex] = SceneProxies[LinearIndex];
+				NewSectionInfos[NewLinearIndex] = SectionInfos[LinearIndex];
 			}
 		}
 	}
 
 	Min = NewMin;
 	Size = NewSize;
-	SectionLODBiases = NewSectionLODBiases;
-	SectionLODSettings = NewSectionLODSettings;
-	SectionOriginAndRadius = NewSectionOriginAndRadius;
-	SceneProxies = NewSceneProxies;
+	SectionLODBiases = MoveTemp(NewSectionLODBiases);
+	SectionInfos = MoveTemp(NewSectionInfos);
 
 	SectionLODBiases.SetAllowCPUAccess(true);
 }
 
-void FLandscapeRenderSystem::PrepareView(const FViewParams& InViewParams)
-{
-	const bool bExecuteInParallel = FApp::ShouldUseThreadingForPerformance()
-		&& GIsThreadedRendering; // Rendering thread is required to safely use rendering resources in parallel.
-
-	if (bExecuteInParallel)
-	{
-		PerViewParametersTasks.Add(InViewParams.ViewKey, TGraphTask<FComputeSectionPerViewParametersTask>::CreateTask(
-			nullptr, ENamedThreads::GetRenderThread()).ConstructAndDispatchWhenReady(*this, InViewParams));
-	}
-	else
-	{
-		FComputeSectionPerViewParametersTask Task(*this, InViewParams);
-		Task.AnyThreadTask();
-	}
-}
-
 void FLandscapeRenderSystem::BeginRender()
 {
 	check(IsInRenderingThread());
@@ -941,72 +926,34 @@ void FLandscapeRenderSystem::BeginRender()
 	UpdateBuffers();
 }
 
-void FLandscapeRenderSystem::ComputeSectionPerViewParameters(
-	FViewKey ViewKey,
-	int32 ViewLODOverride,
-	float ViewLODDistanceFactor,
-	bool bDrawCollisionPawn,
-	bool bDrawCollisionCollision,
-	const FVector& ViewOrigin,
-	const FMatrix& ViewProjectionMarix,
-	const TArray<uint8>& SectionCurrentFirstLODIndices
-)
+const TResourceArray<float>& FLandscapeRenderSystem::ComputeSectionsLODForView(const FSceneView& InView)
 {
-	TRACE_CPUPROFILER_EVENT_SCOPE(FLandscapeRenderSystem::ComputeSectionPerViewParameters());
-
-	TResourceArray<float> NewSectionLODValues;
+	TRACE_CPUPROFILER_EVENT_SCOPE(FLandscapeRenderSystem::ComputeSectionsLODForView);
 
-	NewSectionLODValues.AddZeroed(SectionLODSettings.Num());
+	TResourceArray<float>& SectionLODValues = CachedSectionLODValues.Add(InView.GetViewKey());
+	SectionLODValues.AddZeroed(SectionInfos.Num());
 
-	float LODScale = ViewLODDistanceFactor * CVarStaticMeshLODDistanceScale.GetValueOnRenderThread();
-
-	for (int32 EntityIndex = 0; EntityIndex < SectionLODSettings.Num(); EntityIndex++)
+	for (int32 SectionIndex = 0; SectionIndex < SectionInfos.Num(); SectionIndex++)
 	{
-		float MeshScreenSizeSquared = ComputeBoundsScreenRadiusSquared(FVector(SectionOriginAndRadius[EntityIndex]), SectionOriginAndRadius[EntityIndex].W, ViewOrigin, ViewProjectionMarix);
-
-		float FractionalLOD;
-		GetLODFromScreenSize(SectionLODSettings[EntityIndex], MeshScreenSizeSquared, LODScale * LODScale, FractionalLOD);
+		const float DefaultLODValue = 0.0f;
 
-		int32 ForcedLODLevel = SectionLODSettings[EntityIndex].ForcedLOD;
-		ForcedLODLevel = ViewLODOverride >= 0 ? ViewLODOverride : ForcedLODLevel;
-		const int32 DrawCollisionLODOverride = GetDrawCollisionLodOverride(bDrawCollisionPawn, bDrawCollisionCollision, SectionLODSettings[EntityIndex].DrawCollisionPawnLOD, SectionLODSettings[EntityIndex].DrawCollisionVisibilityLOD);
-		ForcedLODLevel = DrawCollisionLODOverride >= 0 ? DrawCollisionLODOverride : ForcedLODLevel;
-		ForcedLODLevel = FMath::Min<int32>(ForcedLODLevel, SectionLODSettings[EntityIndex].LastLODIndex);
-
-#if PLATFORM_SUPPORTS_LANDSCAPE_VISUAL_MESH_LOD_STREAMING
-		const float CurFirstLODIdx = (float)SectionCurrentFirstLODIndices[EntityIndex];
-#else
-		constexpr float CurFirstLODIdx = 0.f;
-#endif
-
-		NewSectionLODValues[EntityIndex] = FMath::Max(ForcedLODLevel >= 0 ? ForcedLODLevel : FractionalLOD, CurFirstLODIdx);
+		FLandscapeSectionInfo* SectionInfo = SectionInfos[SectionIndex];
+		SectionLODValues[SectionIndex] = SectionInfo ? SectionInfo->ComputeLODForView(InView) : DefaultLODValue;
 	}
 
-	{
-		FScopeLock Lock(&CachedValuesCS);
-
-		CachedSectionLODValues.Add(ViewKey, MoveTemp(NewSectionLODValues));
-	}
+	return SectionLODValues;
 }
 
 void FLandscapeRenderSystem::FetchHeightmapLODBiases()
 {
-	TRACE_CPUPROFILER_EVENT_SCOPE(FLandscapeRenderSystem::FetchHeightmapLODBiases());
+	TRACE_CPUPROFILER_EVENT_SCOPE(FLandscapeRenderSystem::FetchHeightmapLODBiases);
 
-	// TODO: this function generates A LOT OF cache misses - it should be much better if we have an event of FTexture2DResource::UpdateTexture
-	for (int32 EntityIndex = 0; EntityIndex < SceneProxies.Num(); EntityIndex++)
+	for (int32 SectionIndex = 0; SectionIndex < SectionInfos.Num(); SectionIndex++)
 	{
-		FLandscapeComponentSceneProxy* SceneProxy = SceneProxies[EntityIndex];
-		if (SceneProxy && SceneProxy->HeightmapTexture)
-		{
-			if (const FTexture2DResource* TextureResource = (const FTexture2DResource*)SceneProxy->HeightmapTexture->GetResource())
-			{
-				SectionLODBiases[EntityIndex] = SceneProxy->HeightmapTexture->GetNumMips() - SceneProxy->HeightmapTexture->GetNumResidentMips();
+		const float DefaultLODBias = 0.0f;
 
-				// TODO: support mipmap LOD bias of XY offset map
-				//XYOffsetmapTexture ? ((FTexture2DResource*)XYOffsetmapTexture->Resource)->GetCurrentFirstMip() : 0.0f);
-			}
-		}
+		FLandscapeSectionInfo* SectionInfo = SectionInfos[SectionIndex];
+		SectionLODBiases[SectionIndex] = SectionInfo ? SectionInfo->ComputeLODBias() : DefaultLODBias;
 	}
 }
 
@@ -1082,13 +1029,6 @@ void FLandscapeRenderSystem::WaitForTasksCompletion()
 		FTaskGraphInterface::Get().WaitUntilTaskCompletes(FetchHeightmapLODBiasesEventRef, ENamedThreads::GetRenderThread_Local());
 		FetchHeightmapLODBiasesEventRef.SafeRelease();
 	}
-
-	for (auto& Pair : PerViewParametersTasks)
-	{
-		FTaskGraphInterface::Get().WaitUntilTaskCompletes(PerViewParametersTasks[Pair.Key], ENamedThreads::GetRenderThread_Local());
-	}
-
-	PerViewParametersTasks.Reset();
 }
 
 
@@ -1112,31 +1052,12 @@ void FLandscapeSceneViewExtension::BeginFrame_RenderThread()
 	for (auto& Pair : LandscapeRenderSystems)
 	{
 		FLandscapeRenderSystem& RenderSystem = *Pair.Value;
-
 		RenderSystem.BeginFrame();
 	}
 }
 
 void FLandscapeSceneViewExtension::PreRenderViewFamily_RenderThread(FRHICommandListImmediate& RHICmdList, FSceneViewFamily& InViewFamily)
 {
-	for (const auto& View : InViewFamily.Views)
-	{
-		FLandscapeRenderSystem::FViewParams ViewParams;
-		ViewParams.ViewKey = View->GetViewKey();
-		ViewParams.ViewLODOverride = GetViewLodOverride(*View);
-		ViewParams.ViewLODDistanceFactor = View->LODDistanceFactor;
-		ViewParams.ViewEngineShowFlagCollisionPawn = View->Family->EngineShowFlags.CollisionPawn;
-		ViewParams.ViewEngineShowFlagCollisionVisibility = View->Family->EngineShowFlags.CollisionVisibility;
-		ViewParams.ViewOrigin = GetLODView(*View).ViewMatrices.GetViewOrigin();
-		ViewParams.ViewProjectionMatrix = GetLODView(*View).ViewMatrices.GetProjectionMatrix();
-
-		for (auto& Pair : LandscapeRenderSystems)
-		{
-			FLandscapeRenderSystem& RenderSystem = *Pair.Value;
-			RenderSystem.PrepareView(ViewParams);
-		}
-	}
-
 	for (auto& Pair : LandscapeRenderSystems)
 	{
 		FLandscapeRenderSystem& RenderSystem = *Pair.Value;
@@ -1159,12 +1080,14 @@ void FLandscapeSceneViewExtension::PreRenderView_RenderThread(FRHICommandListImm
 		for (auto& Pair : LandscapeRenderSystems)
 		{
 			FLandscapeRenderSystem& RenderSystem = *Pair.Value;
-			uint32& IndirectionEntry = LandscapeIndirection[RenderSystem.LandscapeIndex];
 
-			// Where the data for this landscape starts
+			// Store index where the LOD data for this landscape starts
+			uint32& IndirectionEntry = LandscapeIndirection[RenderSystem.LandscapeIndex];
 			IndirectionEntry = LandscapeLODData.Num();
 
-			LandscapeLODData.Append(RenderSystem.CachedSectionLODValues.FindChecked(InView.GetViewKey()));
+			// Compute sections lod values for this view & append to the global landscape LOD data
+			const TResourceArray<float>& SectionsLODValues = RenderSystem.ComputeSectionsLODForView(InView);
+			LandscapeLODData.Append(SectionsLODValues);
 		}
 
 		FRHIResourceCreateInfo CreateInfoLODBuffer(TEXT("LandscapeLODDataBuffer"), &LandscapeLODData);
@@ -1230,7 +1153,7 @@ bool FLandscapeVisibilityHelper::OnRemoveFromWorld()
 
 FLandscapeComponentSceneProxy::FLandscapeComponentSceneProxy(ULandscapeComponent* InComponent)
 	: FPrimitiveSceneProxy(InComponent, NAME_LandscapeResourceNameForDebugging)
-	, FLandscapeNeighborInfo(InComponent->GetWorld(), InComponent->GetLandscapeProxy()->GetLandscapeGuid(), InComponent->GetSectionBase() / InComponent->ComponentSizeQuads, InComponent->GetHeightmap(), InComponent->ForcedLOD, InComponent->LODBias)
+	, FLandscapeSectionInfo(InComponent->GetWorld(), InComponent->GetLandscapeProxy()->GetLandscapeGuid(), InComponent->GetSectionBase() / InComponent->ComponentSizeQuads)
 	, MaxLOD(FMath::CeilLogTwo(InComponent->SubsectionSizeQuads + 1) - 1)
 	, NumWeightmapLayerAllocations(InComponent->GetWeightmapLayerAllocations().Num())
 	, StaticLightingLOD(InComponent->GetLandscapeProxy()->StaticLightingLOD)
@@ -1250,6 +1173,7 @@ FLandscapeComponentSceneProxy::FLandscapeComponentSceneProxy(ULandscapeComponent
 	, WeightmapTextures(InComponent->GetWeightmapTextures())
 	, VisibilityWeightmapTexture(nullptr)
 	, VisibilityWeightmapChannel(-1)
+	, HeightmapTexture(InComponent->GetHeightmap())
 	, NormalmapTexture(InComponent->GetHeightmap())
 	, BaseColorForGITexture(InComponent->GIBakedBaseColorTexture)
 	, HeightmapScaleBias(InComponent->HeightmapScaleBias)
@@ -1272,7 +1196,6 @@ FLandscapeComponentSceneProxy::FLandscapeComponentSceneProxy(ULandscapeComponent
 	, LightMapResolution(InComponent->GetStaticLightMapResolution())
 #endif
 {
-
 	VisibilityHelper.Init(InComponent, this);
 
 	if (!VisibilityHelper.ShouldBeVisible())
@@ -1358,6 +1281,7 @@ FLandscapeComponentSceneProxy::FLandscapeComponentSceneProxy(ULandscapeComponent
 	}
 
 	// Clamp ForcedLOD to the valid range and then apply
+	int8 ForcedLOD = InComponent->ForcedLOD;
 	ForcedLOD = ForcedLOD >= 0 ? FMath::Clamp<int32>(ForcedLOD, FirstLOD, LastLOD) : ForcedLOD;
 	FirstLOD = ForcedLOD >= 0 ? ForcedLOD : FirstLOD;
 	LastLOD = ForcedLOD >= 0 ? ForcedLOD : LastLOD;
@@ -1366,12 +1290,6 @@ FLandscapeComponentSceneProxy::FLandscapeComponentSceneProxy(ULandscapeComponent
 	LODSettings.LastLODScreenSizeSquared = LODScreenRatioSquared[LastLOD];
 	LODSettings.ForcedLOD = ForcedLOD;
 
-	LODBias = FMath::Clamp<int8>(LODBias, -MaxLOD, MaxLOD);
-
-	int8 LocalLODBias = LODBias + (int8)GLandscapeMeshLODBias;
-	MinValidLOD = FMath::Clamp<int8>(LocalLODBias, -MaxLOD, MaxLOD);
-	MaxValidLOD = FMath::Min<int32>(MaxLOD, MaxLOD + LocalLODBias);
-
 	LastVirtualTextureLOD = MaxLOD;
 	FirstVirtualTextureLOD = FMath::Max(MaxLOD - InComponent->GetLandscapeProxy()->VirtualTextureNumLods, 0);
 	VirtualTextureLodBias = InComponent->GetLandscapeProxy()->VirtualTextureLodBias;
@@ -1509,11 +1427,11 @@ void FLandscapeComponentSceneProxy::CreateRenderThreadResources()
 
 	check(HeightmapTexture != nullptr);
 
-	const FLandscapeRenderSystem& RenderSystem = GetOrCreateLandscapeRenderSystem(LandscapeKey);	
+	FLandscapeRenderSystem::CreateResources(this);
 
 	if (VisibilityHelper.ShouldBeVisible())
 	{
-		RegisterNeighbors(this);
+		RegisterSection();
 	}
 
 	auto FeatureLevel = GetScene().GetFeatureLevel();
@@ -1611,7 +1529,7 @@ void FLandscapeComponentSceneProxy::CreateRenderThreadResources()
 		FLandscapeBatchElementParams* BatchElementParams = &GrassBatchParams[0];
 		BatchElementParams->LandscapeUniformShaderParametersResource = &LandscapeUniformShaderParameters;
 		BatchElementParams->FixedGridUniformShaderParameters = &LandscapeFixedGridUniformShaderParameters;
-		BatchElementParams->LandscapeSectionLODUniformParameters = RenderSystem.SectionLODUniformBuffer;
+		BatchElementParams->LandscapeSectionLODUniformParameters = nullptr; // Not needed for grass rendering
 		BatchElementParams->SceneProxy = this;
 		BatchElementParams->CurrentLOD = 0;
 		GrassBatchElement->UserData = BatchElementParams;
@@ -1677,7 +1595,8 @@ void FLandscapeComponentSceneProxy::CreateRenderThreadResources()
 void FLandscapeComponentSceneProxy::DestroyRenderThreadResources()
 {
 	FPrimitiveSceneProxy::DestroyRenderThreadResources();
-	UnregisterNeighbors(this);
+	FLandscapeRenderSystem::UnregisterSection(this);
+	FLandscapeRenderSystem::DestroyResources(this);
 }
 
 bool FLandscapeComponentSceneProxy::OnLevelAddedToWorld_RenderThread()
@@ -1685,7 +1604,7 @@ bool FLandscapeComponentSceneProxy::OnLevelAddedToWorld_RenderThread()
 	if (VisibilityHelper.OnAddedToWorld())
 	{
 		SetForceHidden(false);
-		CreateRenderThreadResources();
+		FLandscapeRenderSystem::RegisterSection(this);
 		return true;
 	}
 	return false;
@@ -1696,6 +1615,7 @@ void FLandscapeComponentSceneProxy::OnLevelRemovedFromWorld_RenderThread()
 	if (VisibilityHelper.OnRemoveFromWorld())
 	{
 		SetForceHidden(true);
+		FLandscapeRenderSystem::UnregisterSection(this);
 	}
 }
 
@@ -2055,14 +1975,6 @@ void FLandscapeComponentSceneProxy::OnTransformChanged()
 
 	LandscapeUniformShaderParameters.SetContents(LandscapeParams);
 
-	if (bRegistered)
-	{
-		FVector4 OriginAndSphereRadius(GetBounds().Origin, GetBounds().SphereRadius);
-
-		FLandscapeRenderSystem& RenderSystem = *LandscapeRenderSystems.FindChecked(LandscapeKey);
-		RenderSystem.SetSectionOriginAndRadius(ComponentBase, OriginAndSphereRadius);
-	}
-
 	// Recache mesh draw commands for changed uniform buffers
 	GetScene().UpdateCachedRenderStates(this);
 
@@ -2277,6 +2189,11 @@ void FLandscapeComponentSceneProxy::GetDynamicMeshElements(const TArray<const FS
 	QUICK_SCOPE_CYCLE_COUNTER(STAT_FLandscapeComponentSceneProxy_GetMeshElements);
 	SCOPE_CYCLE_COUNTER(STAT_LandscapeDynamicDrawTime);
 
+	if (!bRegistered)
+	{
+		return;
+	}
+
 	int32 NumPasses = 0;
 	int32 NumTriangles = 0;
 	int32 NumDrawCalls = 0;
@@ -2387,7 +2304,7 @@ void FLandscapeComponentSceneProxy::GetDynamicMeshElements(const TArray<const FS
 					LODMesh.Elements.Add(TemplateMesh.Elements[i]);
 					int32 CurrentLOD = ((FLandscapeBatchElementParams*)TemplateMesh.Elements[i].UserData)->CurrentLOD;
 					LODMesh.VisualizeLODIndex = CurrentLOD;
-					FLinearColor Color = GetColorForLod(CurrentLOD, ForcedLOD, true);
+					FLinearColor Color = GetColorForLod(CurrentLOD, LODSettings.ForcedLOD, true);
 					FMaterialRenderProxy* LODMaterialProxy = (FMaterialRenderProxy*)new FColoredMaterialRenderProxy(GEngine->LevelColorationUnlitMaterial->GetRenderProxy(), Color);
 					Collector.RegisterOneFrameMaterialProxy(LODMaterialProxy);
 					LODMesh.MaterialRenderProxy = LODMaterialProxy;
@@ -4114,121 +4031,117 @@ void FLandscapeComponentSceneProxy::GetLCIs(FLCIArray& LCIs)
 	}
 }
 
-//
-// FLandscapeNeighborInfo
-//
-void FLandscapeNeighborInfo::RegisterNeighbors(FLandscapeComponentSceneProxy* SceneProxy /* = nullptr */)
+float FLandscapeComponentSceneProxy::ComputeLODForView(const FSceneView& InView) const
 {
-	check(IsInRenderingThread());
-	if (!bRegistered)
+	// TODO: this function generates A LOT OF cache misses - it should be much better if we have an event of FTexture2DResource::UpdateTexture
+
+	int32 ViewLODOverride = GetViewLodOverride(InView);
+	float ViewLODDistanceFactor = InView.LODDistanceFactor;
+	bool ViewEngineShowFlagCollisionPawn = InView.Family->EngineShowFlags.CollisionPawn;
+	bool ViewEngineShowFlagCollisionVisibility = InView.Family->EngineShowFlags.CollisionVisibility;
+	const FVector& ViewOrigin = GetLODView(InView).ViewMatrices.GetViewOrigin();
+	const FMatrix& ViewProjectionMatrix = GetLODView(InView).ViewMatrices.GetProjectionMatrix();
+
+	float LODScale = ViewLODDistanceFactor * CVarStaticMeshLODDistanceScale.GetValueOnRenderThread();
+
+	int32 ForcedLODLevel = LODSettings.ForcedLOD;
+	ForcedLODLevel = ViewLODOverride >= 0 ? ViewLODOverride : ForcedLODLevel;
+	const int32 DrawCollisionLODOverride = GetDrawCollisionLodOverride(ViewEngineShowFlagCollisionPawn, ViewEngineShowFlagCollisionVisibility, LODSettings.DrawCollisionPawnLOD, LODSettings.DrawCollisionVisibilityLOD);
+	ForcedLODLevel = DrawCollisionLODOverride >= 0 ? DrawCollisionLODOverride : ForcedLODLevel;
+	ForcedLODLevel = FMath::Min<int32>(ForcedLODLevel, LODSettings.LastLODIndex);
+
+	float LODLevel = ForcedLODLevel;
+	if (ForcedLODLevel < 0)
 	{
-		FLandscapeRenderSystem& RenderSystem = GetOrCreateLandscapeRenderSystem(LandscapeKey);
+		float MeshScreenSizeSquared = ComputeBoundsScreenRadiusSquared(GetBounds().Origin, GetBounds().SphereRadius, ViewOrigin, ViewProjectionMatrix);
 
-		// Register ourselves in the map.
-		TMap<FIntPoint, const FLandscapeNeighborInfo*>& SceneProxyMap = SharedSceneProxyMap.FindOrAdd(LandscapeKey);
+		float FractionalLOD;
+		FLandscapeRenderSystem::GetLODFromScreenSize(LODSettings, MeshScreenSizeSquared, LODScale * LODScale, FractionalLOD);
+		LODLevel = FractionalLOD;
+	}
 
-		const FLandscapeNeighborInfo* Existing = SceneProxyMap.FindRef(ComponentBase);
-		if (Existing == nullptr)//(ensure(Existing == nullptr))
-		{
-			SceneProxyMap.Add(ComponentBase, this);
-			bRegistered = true;
+#if PLATFORM_SUPPORTS_LANDSCAPE_VISUAL_MESH_LOD_STREAMING
+	const float CurFirstLODIdx = GetCurrentFirstLODIdx_RenderThread();
+#else
+	constexpr float CurFirstLODIdx = 0.f;
+#endif
 
-			// Find Neighbors
-			Neighbors[0] = SceneProxyMap.FindRef(ComponentBase + FIntPoint(0, -1));
-			Neighbors[1] = SceneProxyMap.FindRef(ComponentBase + FIntPoint(-1, 0));
-			Neighbors[2] = SceneProxyMap.FindRef(ComponentBase + FIntPoint(1, 0));
-			Neighbors[3] = SceneProxyMap.FindRef(ComponentBase + FIntPoint(0, 1));
+	return FMath::Max(LODLevel, CurFirstLODIdx);
+}
 
-			// Add ourselves to our neighbors
-			if (Neighbors[0])
-			{
-				Neighbors[0]->Neighbors[3] = this;
-			}
-			if (Neighbors[1])
-			{
-				Neighbors[1]->Neighbors[2] = this;
-			}
-			if (Neighbors[2])
-			{
-				Neighbors[2]->Neighbors[1] = this;
-			}
-			if (Neighbors[3])
-			{
-				Neighbors[3]->Neighbors[0] = this;
-			}
+float FLandscapeComponentSceneProxy::ComputeLODBias() const
+{
+	float ComputedLODBias = 0;
 
-			if (SceneProxy != nullptr)
-			{
-				RenderSystem.RegisterEntity(SceneProxy);
-			}
-		}
-		else
+	if (HeightmapTexture)
+	{
+		if (const FTexture2DResource* TextureResource = (const FTexture2DResource*)HeightmapTexture->GetResource())
 		{
-			UE_LOG(LogLandscape, Warning, TEXT("Duplicate ComponentBase %d, %d"), ComponentBase.X, ComponentBase.Y);
+			ComputedLODBias = HeightmapTexture->GetNumMips() - HeightmapTexture->GetNumResidentMips();
 		}
 	}
+
+	// TODO: support mipmap LOD bias of XY offset map
+	//XYOffsetmapTexture ? ((FTexture2DResource*)XYOffsetmapTexture->Resource)->GetCurrentFirstMip() : 0.0f);
+
+	return ComputedLODBias;
 }
 
-void FLandscapeNeighborInfo::UnregisterNeighbors(FLandscapeComponentSceneProxy* SceneProxy /* = nullptr */)
+//
+// FLandscapeSectionInfo
+//
+FLandscapeSectionInfo::FLandscapeSectionInfo(const UWorld* InWorld, const FGuid& InLandscapeGuid, const FIntPoint& InSectionBase)
+	: LandscapeKey(HashCombine(GetTypeHash(InWorld), GetTypeHash(InLandscapeGuid)))
+	, ComponentBase(InSectionBase)
+	, bRegistered(false)
 {
-	check(IsInRenderingThread());
-	
-	if (bRegistered)
-	{
-		// Remove ourselves from the map
-		TMap<FIntPoint, const FLandscapeNeighborInfo*>* SceneProxyMap = SharedSceneProxyMap.Find(LandscapeKey);
-		check(SceneProxyMap);
+}
 
-		const FLandscapeNeighborInfo* MapEntry = SceneProxyMap->FindRef(ComponentBase);
-		if (MapEntry == this) //(/*ensure*/(MapEntry == this))
-		{
-			SceneProxyMap->Remove(ComponentBase);
+void FLandscapeSectionInfo::RegisterSection()
+{
+	FLandscapeRenderSystem::RegisterSection(this);
+}
 
-			if (SceneProxy != nullptr)
-			{
-				FLandscapeRenderSystem& RenderSystem = *LandscapeRenderSystems.FindChecked(LandscapeKey);
-				RenderSystem.UnregisterEntity(SceneProxy);
-			}
+void FLandscapeSectionInfo::UnregisterSection()
+{
+	FLandscapeRenderSystem::UnregisterSection(this);
+}
 
-			if (SceneProxyMap->Num() == 0)
-			{
-				FLandscapeRenderSystem* RenderSystemPtr = LandscapeRenderSystems.FindChecked(LandscapeKey);
-				check(RenderSystemPtr->NumRegisteredEntities == 0);
+//
+// FLandscapeProxySectionInfo
+//
+class FLandscapeProxySectionInfo : public FLandscapeSectionInfo
+{
+public:
+	FLandscapeProxySectionInfo(const UWorld* InWorld, const FGuid& InLandscapeGuid, const FIntPoint& InSectionBase, int8 InProxyLOD)
+		: FLandscapeSectionInfo(InWorld, InLandscapeGuid, InSectionBase)
+		, ProxyLOD(InProxyLOD)
+	{
+	}
 
-				delete RenderSystemPtr;
-				LandscapeRenderSystems.Remove(LandscapeKey);
+	virtual float ComputeLODForView(const FSceneView& InView) const override
+	{
+		return ProxyLOD;
+	}
 
-				// remove the entire LandscapeKey entry as this is the last scene proxy
-				SharedSceneProxyMap.Remove(LandscapeKey);
-			}
-			else
-			{
-				// remove reference to us from our neighbors
-				if (Neighbors[0])
-				{
-					Neighbors[0]->Neighbors[3] = nullptr;
-				}
-				if (Neighbors[1])
-				{
-					Neighbors[1]->Neighbors[2] = nullptr;
-				}
-				if (Neighbors[2])
-				{
-					Neighbors[2]->Neighbors[1] = nullptr;
-				}
-				if (Neighbors[3])
-				{
-					Neighbors[3]->Neighbors[0] = nullptr;
-				}
-			}
-		}
+	virtual float ComputeLODBias() const override
+	{
+		return 0.0f;
 	}
-}
+
+	virtual int32 GetSectionPriority() const override
+	{
+		return ProxyLOD;
+	}
+
+private:
+	int8 ProxyLOD;
+};
 
 //
 // FLandscapeMeshProxySceneProxy
 //
-FLandscapeMeshProxySceneProxy::FLandscapeMeshProxySceneProxy(UStaticMeshComponent* InComponent, const FGuid& InGuid, const TArray<FIntPoint>& InProxyComponentBases, int8 InProxyLOD)
+FLandscapeMeshProxySceneProxy::FLandscapeMeshProxySceneProxy(UStaticMeshComponent* InComponent, const FGuid& InLandscapeGuid, const TArray<FIntPoint>& InProxySectionsBases, int8 InProxyLOD)
 	: FStaticMeshSceneProxy(InComponent, false)
 {
 	VisibilityHelper.Init(InComponent, this);
@@ -4238,10 +4151,26 @@ FLandscapeMeshProxySceneProxy::FLandscapeMeshProxySceneProxy(UStaticMeshComponen
 		bShouldNotifyOnWorldAddRemove = true;
 	}
 
-	ProxyNeighborInfos.Empty(InProxyComponentBases.Num());
-	for (FIntPoint ComponentBase : InProxyComponentBases)
+	ProxySectionsInfos.Empty(InProxySectionsBases.Num());
+	for (FIntPoint SectionBase : InProxySectionsBases)
+	{
+		ProxySectionsInfos.Emplace(MakeUnique<FLandscapeProxySectionInfo>(InComponent->GetWorld(), InLandscapeGuid, SectionBase, InProxyLOD));
+	}
+}
+
+void FLandscapeMeshProxySceneProxy::RegisterSections()
+{
+	for (auto& Info : ProxySectionsInfos)
+	{
+		Info->RegisterSection();
+	}
+}
+
+void FLandscapeMeshProxySceneProxy::UnregisterSections()
+{
+	for (auto& Info : ProxySectionsInfos)
 	{
-		new(ProxyNeighborInfos) FLandscapeNeighborInfo(InComponent->GetWorld(), InGuid, ComponentBase, nullptr, InProxyLOD, 0);
+		Info->UnregisterSection();
 	}
 }
 
@@ -4255,13 +4184,15 @@ SIZE_T FLandscapeMeshProxySceneProxy::GetTypeHash() const
 void FLandscapeMeshProxySceneProxy::CreateRenderThreadResources()
 {
 	FStaticMeshSceneProxy::CreateRenderThreadResources();
+	
+	for (auto& Info : ProxySectionsInfos)
+	{
+		FLandscapeRenderSystem::CreateResources(Info.Get());
+	}
 
 	if (VisibilityHelper.ShouldBeVisible())
 	{
-		for (FLandscapeNeighborInfo& Info : ProxyNeighborInfos)
-		{
-			Info.RegisterNeighbors();
-		}
+		RegisterSections();
 	}
 }
 
@@ -4270,11 +4201,7 @@ bool FLandscapeMeshProxySceneProxy::OnLevelAddedToWorld_RenderThread()
 	if (VisibilityHelper.OnAddedToWorld())
 	{
 		SetForceHidden(false);
-
-		for (FLandscapeNeighborInfo& Info : ProxyNeighborInfos)
-		{
-			Info.RegisterNeighbors();
-		}
+		RegisterSections();
 		return true;
 	}
 
@@ -4286,16 +4213,18 @@ void FLandscapeMeshProxySceneProxy::OnLevelRemovedFromWorld_RenderThread()
 	if (VisibilityHelper.OnRemoveFromWorld())
 	{
 		SetForceHidden(true);
+		UnregisterSections();
 	}
 }
 
 void FLandscapeMeshProxySceneProxy::DestroyRenderThreadResources()
 {
 	FStaticMeshSceneProxy::DestroyRenderThreadResources();
-
-	for (FLandscapeNeighborInfo& Info : ProxyNeighborInfos)
+	UnregisterSections();
+	
+	for (auto& Info : ProxySectionsInfos)
 	{
-		Info.UnregisterNeighbors();
+		FLandscapeRenderSystem::DestroyResources(Info.Get());
 	}
 }
 
diff --git a/Engine/Source/Runtime/Landscape/Private/LandscapeRenderMobile.cpp b/Engine/Source/Runtime/Landscape/Private/LandscapeRenderMobile.cpp
--- a/Engine/Source/Runtime/Landscape/Private/LandscapeRenderMobile.cpp
+++ b/Engine/Source/Runtime/Landscape/Private/LandscapeRenderMobile.cpp
@@ -366,9 +366,11 @@ void FLandscapeComponentSceneProxyMobile::CreateRenderThreadResources()
 {
 	LLM_SCOPE(ELLMTag::Landscape);
 
+	FLandscapeRenderSystem::CreateResources(this);
+
 	if (VisibilityHelper.ShouldBeVisible())
 	{
-		RegisterNeighbors(this);
+		RegisterSection();
 	}
 	
 	auto FeatureLevel = GetScene().GetFeatureLevel();
diff --git a/Engine/Source/Runtime/Landscape/Public/LandscapeRender.h b/Engine/Source/Runtime/Landscape/Public/LandscapeRender.h
--- a/Engine/Source/Runtime/Landscape/Public/LandscapeRender.h
+++ b/Engine/Source/Runtime/Landscape/Public/LandscapeRender.h
@@ -344,114 +344,32 @@ public:
 };
 
 //
-// FLandscapeNeighborInfo
+// FLandscapeSectionInfo
 //
 
-class FLandscapeNeighborInfo
+class FLandscapeSectionInfo : public TIntrusiveLinkedList<FLandscapeSectionInfo>
 {
 public:
-	static const int8 NEIGHBOR_COUNT = 4;
+	FLandscapeSectionInfo(const UWorld* InWorld, const FGuid& InLandscapeGuid, const FIntPoint& InSectionBase);
+	virtual ~FLandscapeSectionInfo() = default;
 
-	// Key to uniquely identify the landscape to find the correct render proxy map
-	class FLandscapeKey
-	{
-		const UWorld* World;
-		const FGuid Guid;
-	public:
-		FLandscapeKey(const UWorld* InWorld, const FGuid& InGuid)
-			: World(InWorld)
-			, Guid(InGuid)
-		{}
-
-		friend inline uint32 GetTypeHash(const FLandscapeKey& InLandscapeKey)
-		{
-			return HashCombine(GetTypeHash(InLandscapeKey.World), GetTypeHash(InLandscapeKey.Guid));
-		}
-
-		friend bool operator==(const FLandscapeKey& A, const FLandscapeKey& B)
-		{
-			return A.World == B.World && A.Guid == B.Guid;
-		}
-	};
-
-	const FLandscapeNeighborInfo* GetNeighbor(int32 Index) const
-	{
-		if (Index < NEIGHBOR_COUNT)
-		{
-			return Neighbors[Index];
-		}
-
-		return nullptr;
-	}
-
-	UTexture2D*				HeightmapTexture; // PC : Heightmap, Mobile : Weightmap
-
-protected:
-
-	virtual const ULandscapeComponent* GetLandscapeComponent() const { return nullptr; }
-
-	// Map of currently registered landscape proxies, used to register with our neighbors
-	static TMap<FLandscapeKey, TMap<FIntPoint, const FLandscapeNeighborInfo*> > SharedSceneProxyMap;
-
-	// For neighbor lookup
-	FLandscapeKey			LandscapeKey;
-	FIntPoint				ComponentBase;
+	void RegisterSection();
+	void UnregisterSection();
 
-	// Pointer to our neighbor's scene proxies in NWES order (nullptr if there is currently no neighbor)
-	mutable const FLandscapeNeighborInfo* Neighbors[NEIGHBOR_COUNT];
-
-	
-	// Data we need to be able to access about our neighbor
-	int8					ForcedLOD;
-	int8					LODBias;
-	bool					bRegistered;
-
-	friend class FLandscapeComponentSceneProxy;
+	virtual float ComputeLODForView(const FSceneView& InView) const = 0;
+	virtual float ComputeLODBias() const = 0;
+	virtual int32 GetSectionPriority() const { return INDEX_NONE; }
 
 public:
-	FLandscapeNeighborInfo(const UWorld* InWorld, const FGuid& InGuid, const FIntPoint& InComponentBase, UTexture2D* InHeightmapTexture, int8 InForcedLOD, int8 InLODBias)
-	: HeightmapTexture(InHeightmapTexture)
-	, LandscapeKey(InWorld, InGuid)
-	, ComponentBase(InComponentBase)
-	, ForcedLOD(InForcedLOD)
-	, LODBias(InLODBias)
-	, bRegistered(false)
-	{
-		//       -Y       
-		//    - - 0 - -   
-		//    |       |   
-		// -X 1   P   2 +X
-		//    |       |   
-		//    - - 3 - -   
-		//       +Y       
-
-		Neighbors[0] = nullptr;
-		Neighbors[1] = nullptr;
-		Neighbors[2] = nullptr;
-		Neighbors[3] = nullptr;
-	}
-
-	void RegisterNeighbors(FLandscapeComponentSceneProxy* SceneProxy = nullptr);
-	void UnregisterNeighbors(FLandscapeComponentSceneProxy* SceneProxy = nullptr);
+	uint32		LandscapeKey;
+	FIntPoint	ComponentBase;
+	bool		bRegistered;
 };
 
-extern RENDERER_API TAutoConsoleVariable<float> CVarStaticMeshLODDistanceScale;
-
 struct FLandscapeRenderSystem
 {
 	typedef uint32 FViewKey;
 
-	struct FViewParams
-	{
-		FViewKey ViewKey;
-		int32 ViewLODOverride;
-		float ViewLODDistanceFactor;
-		bool ViewEngineShowFlagCollisionPawn;
-		bool ViewEngineShowFlagCollisionVisibility;
-		FVector ViewOrigin;
-		FMatrix ViewProjectionMatrix;
-	};
-
 	struct LODSettingsComponent
 	{
 		float LOD0ScreenSizeSquared;
@@ -489,75 +407,59 @@ struct FLandscapeRenderSystem
 	static TBitArray<> LandscapeIndexAllocator;
 
 	int32 LandscapeIndex;
-	int32 NumRegisteredEntities;
 
 	FIntPoint Min;
 	FIntPoint Size;
 
-	TArray<LODSettingsComponent> SectionLODSettings;
 	TResourceArray<float> SectionLODBiases;
-	TArray<FVector4> SectionOriginAndRadius;
-	TArray<FLandscapeComponentSceneProxy*> SceneProxies;
+	TArray<FLandscapeSectionInfo*> SectionInfos;
+	int32 ReferenceCount;
 
 	FBufferRHIRef SectionLODBiasBuffer;
 	FShaderResourceViewRHIRef SectionLODBiasSRV;
 
 	FUniformBufferRHIRef SectionLODUniformBuffer;
 
-	FCriticalSection CachedValuesCS;
 	TMap<FViewKey, TResourceArray<float>> CachedSectionLODValues;
 
-	TMap<FViewKey, FGraphEventRef> PerViewParametersTasks;
 	FGraphEventRef FetchHeightmapLODBiasesEventRef;
 	
 	FLandscapeRenderSystem();
 	~FLandscapeRenderSystem();
 
-	void RegisterEntity(FLandscapeComponentSceneProxy* SceneProxy);
-	void UnregisterEntity(FLandscapeComponentSceneProxy* SceneProxy);
+	static void CreateResources(FLandscapeSectionInfo* SectionInfo);
+	static void DestroyResources(FLandscapeSectionInfo* SectionInfo);
 
-	int32 GetComponentLinearIndex(FIntPoint ComponentBase) const
-	{
-		return (ComponentBase.Y - Min.Y) * Size.X + ComponentBase.X - Min.X;
-	}
-	void ResizeAndMoveTo(FIntPoint NewMin, FIntPoint NewSize);
+	static void RegisterSection(FLandscapeSectionInfo* SectionInfo);
+	static void UnregisterSection(FLandscapeSectionInfo* SectionInfo);
 
-	void SetSectionLODSettings(FIntPoint ComponentBase, LODSettingsComponent LODSettings)
+	int32 GetSectionLinearIndex(FIntPoint InSectionBase) const
 	{
-		SectionLODSettings[GetComponentLinearIndex(ComponentBase)] = LODSettings;
+		return (InSectionBase.Y - Min.Y) * Size.X + InSectionBase.X - Min.X;
 	}
+	void ResizeAndMoveTo(FIntPoint NewMin, FIntPoint NewSize);
 
-	void SetSectionOriginAndRadius(FIntPoint ComponentBase, FVector4 OriginAndRadius)
+	void SetSectionInfo(FIntPoint InSectionBase, FLandscapeSectionInfo* InSectionInfo)
 	{
-		SectionOriginAndRadius[GetComponentLinearIndex(ComponentBase)] = OriginAndRadius;
+		SectionInfos[GetSectionLinearIndex(InSectionBase)] = InSectionInfo;
 	}
 
-	void SetSceneProxy(FIntPoint ComponentBase, FLandscapeComponentSceneProxy* SceneProxy)
+	FLandscapeSectionInfo* GetSectionInfo(FIntPoint InSectionBase)
 	{
-		SceneProxies[GetComponentLinearIndex(ComponentBase)] = SceneProxy;
+		return SectionInfos[GetSectionLinearIndex(InSectionBase)];
 	}
 
-	float GetSectionLODValue(const FSceneView& SceneView, FIntPoint ComponentBase) const
+	float GetSectionLODValue(const FSceneView& SceneView, FIntPoint InSectionBase) const
 	{
-		return CachedSectionLODValues[SceneView.GetViewKey()][GetComponentLinearIndex(ComponentBase)];
+		return CachedSectionLODValues[SceneView.GetViewKey()][GetSectionLinearIndex(InSectionBase)];
 	}
 
-	float GetSectionLODBias(FIntPoint ComponentBase) const
+	float GetSectionLODBias(FIntPoint InSectionBase) const
 	{
-		return SectionLODBiases[GetComponentLinearIndex(ComponentBase)];
+		return SectionLODBiases[GetSectionLinearIndex(InSectionBase)];
 	}
 
-	void ComputeSectionPerViewParameters(
-		FViewKey ViewKey,
-		int32 ViewLODOverride,
-		float ViewLODDistanceFactor,
-		bool bDrawCollisionPawn,
-		bool bDrawCollisionCollision,
-		const FVector& ViewOrigin,
-		const FMatrix& ViewProjectionMarix,
-		const TArray<uint8>& SectionCurrentFirstLODIndices);
-
-	void PrepareView(const FViewParams& InViewParams);
+	const TResourceArray<float>& ComputeSectionsLODForView(const FSceneView& InView);
 
 	void BeginRender();
 
@@ -570,9 +472,11 @@ struct FLandscapeRenderSystem
 	void EndFrame();
 
 	void WaitForTasksCompletion();
-};
 
-LANDSCAPE_API extern TMap<FLandscapeNeighborInfo::FLandscapeKey, FLandscapeRenderSystem*> LandscapeRenderSystems;
+private:
+	void CreateResources_Internal(FLandscapeSectionInfo* InSectionInfo);
+	void DestroyResources_Internal(FLandscapeSectionInfo* InSectionInfo);
+};
 
 
 //
@@ -610,7 +514,6 @@ public:
 private:
 	bool bRequiresVisibleLevelToRender = false;
 	bool bIsComponentLevelVisible = false;
-	FPrimitiveSceneProxy* Proxy = nullptr;
 };
 
 //
@@ -648,24 +551,28 @@ LANDSCAPE_API extern FLandscapeDebugOptions GLandscapeDebugOptions;
 //
 class FLandscapeMeshProxySceneProxy final : public FStaticMeshSceneProxy
 {
-	TArray<FLandscapeNeighborInfo> ProxyNeighborInfos;
 public:
 	SIZE_T GetTypeHash() const override;
 
-	FLandscapeMeshProxySceneProxy(UStaticMeshComponent* InComponent, const FGuid& InGuid, const TArray<FIntPoint>& InProxyComponentBases, int8 InProxyLOD);
+	FLandscapeMeshProxySceneProxy(UStaticMeshComponent* InComponent, const FGuid& InLandscapeGuid, const TArray<FIntPoint>& InProxySectionsBases, int8 InProxyLOD);
 	virtual void CreateRenderThreadResources() override;
 	virtual void DestroyRenderThreadResources() override;
 	virtual bool OnLevelAddedToWorld_RenderThread() override;
 	virtual void OnLevelRemovedFromWorld_RenderThread() override;
 
 private:
+	void RegisterSections();
+	void UnregisterSections();
+
 	FLandscapeVisibilityHelper VisibilityHelper;
+
+	TArray<TUniquePtr<FLandscapeSectionInfo>> ProxySectionsInfos;
 };
 
 //
 // FLandscapeComponentSceneProxy
 //
-class LANDSCAPE_API FLandscapeComponentSceneProxy : public FPrimitiveSceneProxy, public FLandscapeNeighborInfo
+class LANDSCAPE_API FLandscapeComponentSceneProxy : public FPrimitiveSceneProxy, public FLandscapeSectionInfo
 {
 	friend class FLandscapeSharedBuffers;
 
@@ -739,8 +646,6 @@ protected:
 	int32						LastVirtualTextureLOD;
 	float						ComponentMaxExtend; 		// The max extend value in any axis
 	float						ComponentSquaredScreenSizeToUseSubSections; // Size at which we start to draw in sub lod if LOD are different per sub section
-	float						MinValidLOD;							// Min LOD Taking into account LODBias
-	float						MaxValidLOD;							// Max LOD Taking into account LODBias
 
 	FLandscapeRenderSystem::LODSettingsComponent LODSettings;
 
@@ -789,6 +694,7 @@ protected:
 #if WITH_EDITOR
 	TArray<FLinearColor> LayerColors;
 #endif
+	UTexture2D* HeightmapTexture; // PC : Heightmap, Mobile : Weightmap
 	UTexture2D* NormalmapTexture; // PC : Heightmap, Mobile : Weightmap
 	UTexture2D* BaseColorForGITexture;
 	FVector4f HeightmapScaleBias;
@@ -852,7 +758,6 @@ protected:
 protected:
 	virtual ~FLandscapeComponentSceneProxy();
 	
-	virtual const ULandscapeComponent* GetLandscapeComponent() const { return LandscapeComponent; }
 	int8 GetLODFromScreenSize(float InScreenSizeSquared, float InViewLODScale) const;
 
 	bool GetMeshElementForVirtualTexture(int32 InLodIndex, ERuntimeVirtualTextureMaterialType MaterialType, UMaterialInterface* InMaterialInterface, FMeshBatch& OutMeshBatch, TArray<FLandscapeBatchElementParams>& OutStaticBatchParamArray) const;
@@ -910,6 +815,10 @@ public:
 	virtual bool HasRayTracingRepresentation() const override { return true; }
 	virtual bool IsRayTracingRelevant() const override { return true; }
 #endif
+
+	// FLandscapeSceneInfo interface
+	virtual float ComputeLODForView(const FSceneView& InView) const override;
+	virtual float ComputeLODBias() const override;
 };
 
 class FLandscapeDebugMaterialRenderProxy : public FMaterialRenderProxy
