WorldPartition - Added Lanscape HLODs

Created a custom HLOD builder for ULandscapeComponent
* Landscape LOD used to generate the HLOD is autocomputed from the expected draw distance, same for the baked texture size

* Expose ALandscapeProxy::GetLODScreenSizeArray()
* Actually reverting changes done in CL 18794652 (moving proxy mesh creation back to ALandscapeProxy)
* ComputeLayerHash() now takes a bool to choose if we want to compute the hash for edition or runtime data

diff --git a/Engine/Source/Runtime/Landscape/Classes/LandscapeComponent.h b/Engine/Source/Runtime/Landscape/Classes/LandscapeComponent.h
--- a/Engine/Source/Runtime/Landscape/Classes/LandscapeComponent.h
+++ b/Engine/Source/Runtime/Landscape/Classes/LandscapeComponent.h
@@ -36,7 +36,6 @@ struct FEngineShowFlags;
 struct FLandscapeEditDataInterface;
 struct FLandscapeTextureDataInfo;
 struct FStaticLightingPrimitiveInfo;
-struct FMeshDescription;
 struct FLandscapeEditDataInterface;
 struct FLandscapeMobileRenderData;
 
@@ -759,7 +758,7 @@ public:
 	LANDSCAPE_API const TArray<FWeightmapLayerAllocationInfo>& GetWeightmapLayerAllocations(const FGuid& InLayerGuid) const;
 
 #if WITH_EDITOR
-	LANDSCAPE_API uint32 ComputeLayerHash() const;
+	LANDSCAPE_API uint32 ComputeLayerHash(bool InReturnEditingHash = true) const;
 
 	LANDSCAPE_API void SetHeightmap(UTexture2D* NewHeightmap);
 
@@ -819,6 +818,7 @@ public:
 	virtual void InvalidateLightingCacheDetailed(bool bInvalidateBuildEnqueuedLighting, bool bTranslationOnly) override;
 #endif
 	virtual void PropagateLightingScenarioChange() override;
+	virtual bool IsHLODRelevant() const { return true; }
 	//~ End UActorComponent Interface.
 
 	/** Gets the landscape info object for this landscape */
@@ -865,8 +865,7 @@ public:
 	/** Generate mobile data if it's missing or outdated */
 	void CheckGenerateLandscapePlatformData(bool bIsCooking, const ITargetPlatform* TargetPlatform);
 
-	LANDSCAPE_API void ExportToMeshDescription(const int32 InExportLOD, FMeshDescription& OutMesh);
-	LANDSCAPE_API void ExportToMeshDescription(const int32 InExportLOD, const FBoxSphereBounds& InBounds, FMeshDescription& OutMesh);
+	virtual TSubclassOf<class UHLODBuilder> GetCustomHLODBuilderClass() const override;
 #endif
 
 	LANDSCAPE_API int32 GetMaterialInstanceCount(bool InDynamic = true) const;
diff --git a/Engine/Source/Runtime/Landscape/Classes/LandscapeProxy.h b/Engine/Source/Runtime/Landscape/Classes/LandscapeProxy.h
--- a/Engine/Source/Runtime/Landscape/Classes/LandscapeProxy.h
+++ b/Engine/Source/Runtime/Landscape/Classes/LandscapeProxy.h
@@ -937,6 +937,9 @@ public:
 	// ILandscapeSplineInterface
 	LANDSCAPE_API virtual ULandscapeSplinesComponent* GetSplinesComponent() const override { return SplineComponent; }
 
+	// Retrieve the screen size at which each LOD should be rendered
+	LANDSCAPE_API TArray<float> GetLODScreenSizeArray() const;
+
 #if WITH_EDITOR
 	LANDSCAPE_API void SetSplinesComponent(ULandscapeSplinesComponent* InSplineComponent) { check(!SplineComponent); SplineComponent = InSplineComponent; }
 
diff --git a/Engine/Source/Runtime/Landscape/Landscape.Build.cs b/Engine/Source/Runtime/Landscape/Landscape.Build.cs
--- a/Engine/Source/Runtime/Landscape/Landscape.Build.cs
+++ b/Engine/Source/Runtime/Landscape/Landscape.Build.cs
@@ -64,9 +64,10 @@ public class Landscape : ModuleRules
 				new string[] {
 					"EditorFramework",
 					"UnrealEd",
-					"MaterialUtilities", 
+					"MaterialUtilities",
 					"SlateCore",
 					"Slate",
+					"GeometryCore"
 				}
 			);
 
diff --git a/Engine/Source/Runtime/Landscape/Private/Landscape.cpp b/Engine/Source/Runtime/Landscape/Private/Landscape.cpp
--- a/Engine/Source/Runtime/Landscape/Private/Landscape.cpp
+++ b/Engine/Source/Runtime/Landscape/Private/Landscape.cpp
@@ -1184,6 +1184,26 @@ FTransform ALandscapeProxy::LandscapeActorToWorld() const
 	return TM;
 }
 
+TArray<float> ALandscapeProxy::GetLODScreenSizeArray() const
+{
+	static TConsoleVariableData<float>* CVarSMLODDistanceScale = IConsoleManager::Get().FindTConsoleVariableDataFloat(TEXT("r.StaticMeshLODDistanceScale"));
+	static IConsoleVariable* CVarLSLOD0DistributionScale = IConsoleManager::Get().FindConsoleVariable(TEXT("r.LandscapeLOD0DistributionScale"));
+	float CurrentScreenSize = LOD0ScreenSize / CVarSMLODDistanceScale->GetValueOnGameThread();
+	const float ScreenSizeMult = 1.f / FMath::Max(LOD0DistributionSetting * CVarLSLOD0DistributionScale->GetFloat(), 1.01f);
+	
+	const int32 NumLODLevels = FMath::Clamp<int32>(MaxLODLevel, 0, FMath::CeilLogTwo(SubsectionSizeQuads + 1) - 1);
+	
+	TArray<float> Result;
+	Result.Empty(NumLODLevels);
+	for (int32 Idx = 0; Idx < NumLODLevels; ++Idx)
+	{
+		Result.Add(CurrentScreenSize);
+		CurrentScreenSize *= ScreenSizeMult;
+	}
+	return Result;
+}
+
+
 ALandscape* ULandscapeComponent::GetLandscapeActor() const
 {
 	ALandscapeProxy* Landscape = GetLandscapeProxy();
@@ -4458,19 +4478,7 @@ bool ULandscapeLODStreamingProxy::StreamIn(int32 NewMipCount, bool bHighPrio)
 TArray<float> ULandscapeLODStreamingProxy::GetLODScreenSizeArray() const
 {
 	check(LandscapeComponent);
-	static TConsoleVariableData<float>* CVarSMLODDistanceScale = IConsoleManager::Get().FindTConsoleVariableDataFloat(TEXT("r.StaticMeshLODDistanceScale"));
-	static IConsoleVariable* CVarLSLOD0DistributionScale = IConsoleManager::Get().FindConsoleVariable(TEXT("r.LandscapeLOD0DistributionScale"));
-	float CurrentScreenSize = LandscapeComponent->GetLandscapeProxy()->LOD0ScreenSize / CVarSMLODDistanceScale->GetValueOnGameThread();
-	const float ScreenSizeMult = 1.f / FMath::Max(LandscapeComponent->GetLandscapeProxy()->LOD0DistributionSetting * CVarLSLOD0DistributionScale->GetFloat(), 1.01f);
-	const int32 NumLODs = CachedSRRState.MaxNumLODs;
-	TArray<float> Result;
-	Result.Empty(NumLODs);
-	for (int32 Idx = 0; Idx < NumLODs; ++Idx)
-	{
-		Result.Add(CurrentScreenSize);
-		CurrentScreenSize *= ScreenSizeMult;
-	}
-	return Result;
+	return LandscapeComponent->GetLandscapeProxy()->GetLODScreenSizeArray();
 }
 
 TSharedPtr<FLandscapeMobileRenderData, ESPMode::ThreadSafe> ULandscapeLODStreamingProxy::GetRenderData() const
diff --git a/Engine/Source/Runtime/Landscape/Private/LandscapeComponent.cpp b/Engine/Source/Runtime/Landscape/Private/LandscapeComponent.cpp
--- a/Engine/Source/Runtime/Landscape/Private/LandscapeComponent.cpp
+++ b/Engine/Source/Runtime/Landscape/Private/LandscapeComponent.cpp
@@ -8,10 +8,7 @@
 
 #if WITH_EDITOR
 
-#include "MeshDescription.h"
-#include "MeshUtilitiesCommon.h"
-#include "StaticMeshAttributes.h"
-#include "LandscapeDataAccess.h"
+#include "LandscapeHLODBuilder.h"
 
 #endif
 
@@ -75,267 +72,9 @@ void ULandscapeComponent::UpdateEditToolRenderData()
 	}
 }
 
-static void ExportToMeshDescription(ULandscapeComponent* InComponent, const int32 InExportLOD, FMeshDescription& OutMesh, const FBoxSphereBounds& InBounds, bool bIgnoreBounds)
+TSubclassOf<UHLODBuilder> ULandscapeComponent::GetCustomHLODBuilderClass() const
 {
-	ALandscapeProxy* LandscapeProxy = InComponent->GetLandscapeProxy();
-
-	FStaticMeshAttributes Attributes(OutMesh);
-	TVertexAttributesRef<FVector3f> VertexPositions = Attributes.GetVertexPositions();
-	TEdgeAttributesRef<bool> EdgeHardnesses = Attributes.GetEdgeHardnesses();
-	TPolygonGroupAttributesRef<FName> PolygonGroupImportedMaterialSlotNames = Attributes.GetPolygonGroupMaterialSlotNames();
-	TVertexInstanceAttributesRef<FVector3f> VertexInstanceNormals = Attributes.GetVertexInstanceNormals();
-	TVertexInstanceAttributesRef<FVector3f> VertexInstanceTangents = Attributes.GetVertexInstanceTangents();
-	TVertexInstanceAttributesRef<float> VertexInstanceBinormalSigns = Attributes.GetVertexInstanceBinormalSigns();
-	TVertexInstanceAttributesRef<FVector4f> VertexInstanceColors = Attributes.GetVertexInstanceColors();
-	TVertexInstanceAttributesRef<FVector2f> VertexInstanceUVs = Attributes.GetVertexInstanceUVs();
-
-	if (VertexInstanceUVs.GetNumChannels() < 2)
-	{
-		VertexInstanceUVs.SetNumChannels(2);
-	}
-
-	// User specified LOD to export
-	int32 LandscapeLODToExport = LandscapeProxy->ExportLOD;
-	if (InExportLOD != INDEX_NONE)
-	{
-		LandscapeLODToExport = FMath::Clamp<int32>(InExportLOD, 0, FMath::CeilLogTwo(InComponent->SubsectionSizeQuads + 1) - 1);
-	}
-
-	FLandscapeComponentDataInterface CDI(InComponent, LandscapeLODToExport);
-	const int32 ComponentSizeQuadsLOD = ((InComponent->ComponentSizeQuads + 1) >> LandscapeLODToExport) - 1;
-	const int32 SubsectionSizeQuadsLOD = ((InComponent->SubsectionSizeQuads + 1) >> LandscapeLODToExport) - 1;
-	const FIntPoint ComponentOffsetQuads = InComponent->GetSectionBase();
-	const FVector2f ComponentUVOffsetLOD = FVector2f(ComponentOffsetQuads);
-	const FVector2f ComponentUVScaleLOD = FVector2f(InComponent->ComponentSizeQuads / ComponentSizeQuadsLOD);
-
-	const int32 NumFaces = FMath::Square(ComponentSizeQuadsLOD) * 2;
-	const int32 NumVertices = NumFaces * 3;
-
-	OutMesh.ReserveNewVertices(NumVertices);
-	OutMesh.ReserveNewPolygons(NumFaces);
-	OutMesh.ReserveNewVertexInstances(NumVertices);
-	OutMesh.ReserveNewEdges(NumVertices);
-
-	FPolygonGroupID PolygonGroupID = INDEX_NONE;
-	if (OutMesh.PolygonGroups().Num() < 1)
-	{
-		PolygonGroupID = OutMesh.CreatePolygonGroup();
-		PolygonGroupImportedMaterialSlotNames[PolygonGroupID] = FName(TEXT("LandscapeMat_0"));
-	}
-	else
-	{
-		PolygonGroupID = OutMesh.PolygonGroups().GetFirstValidID();
-	}
-
-	// Check if there are any holes
-	const int32 VisThreshold = 170;
-	TArray<uint8> VisDataMap;
-	const TArray<FWeightmapLayerAllocationInfo>& ComponentWeightmapLayerAllocations = InComponent->GetWeightmapLayerAllocations();
-
-	for (int32 AllocIdx = 0; AllocIdx < ComponentWeightmapLayerAllocations.Num(); AllocIdx++)
-	{
-		const FWeightmapLayerAllocationInfo& AllocInfo = ComponentWeightmapLayerAllocations[AllocIdx];
-		if (AllocInfo.LayerInfo == ALandscapeProxy::VisibilityLayer)
-		{
-			CDI.GetWeightmapTextureData(AllocInfo.LayerInfo, VisDataMap);
-		}
-	}
-
-	const FIntPoint QuadPattern[6] =
-	{
-		//face 1
-		FIntPoint(0, 0),
-		FIntPoint(0, 1),
-		FIntPoint(1, 1),
-		//face 2
-		FIntPoint(0, 0),
-		FIntPoint(1, 1),
-		FIntPoint(1, 0),
-	};
-
-	const int32 WeightMapSize = (SubsectionSizeQuadsLOD + 1) * InComponent->NumSubsections;
-
-	const float SquaredSphereRadius = FMath::Square(InBounds.SphereRadius);
-
-	//We need to not duplicate the vertex position, so we use the FIndexAndZ to achieve fast result
-	TArray<FIndexAndZ> VertIndexAndZ;
-	VertIndexAndZ.Reserve(ComponentSizeQuadsLOD * ComponentSizeQuadsLOD * UE_ARRAY_COUNT(QuadPattern));
-	int32 CurrentIndex = 0;
-	TMap<int32, FVector> IndexToPosition;
-	IndexToPosition.Reserve(ComponentSizeQuadsLOD * ComponentSizeQuadsLOD * UE_ARRAY_COUNT(QuadPattern));
-	for (int32 y = 0; y < ComponentSizeQuadsLOD; y++)
-	{
-		for (int32 x = 0; x < ComponentSizeQuadsLOD; x++)
-		{
-			for (int32 i = 0; i < UE_ARRAY_COUNT(QuadPattern); i++)
-			{
-				int32 VertexX = x + QuadPattern[i].X;
-				int32 VertexY = y + QuadPattern[i].Y;
-				FVector Position = CDI.GetWorldVertex(VertexX, VertexY);
-
-				// If at least one vertex is within the given bounds we should process the quad  
-				new(VertIndexAndZ)FIndexAndZ(CurrentIndex, Position);
-				IndexToPosition.Add(CurrentIndex, Position);
-				CurrentIndex++;
-			}
-		}
-	}
-	// Sort the vertices by z value
-	VertIndexAndZ.Sort(FCompareIndexAndZ());
-
-	auto FindPreviousIndex = [&VertIndexAndZ, &IndexToPosition](int32 Index)->int32
-	{
-		const FVector& PositionA = IndexToPosition[Index];
-		FIndexAndZ CompressPosition(0, PositionA);
-		// Search for lowest index duplicates
-		int32 BestIndex = MAX_int32;
-		for (int32 i = 0; i < IndexToPosition.Num(); i++)
-		{
-			if (CompressPosition.Z > (VertIndexAndZ[i].Z + SMALL_NUMBER))
-			{
-				//We will not find anything there is no point searching more
-				break;
-			}
-			const FVector& PositionB = IndexToPosition[VertIndexAndZ[i].Index];
-			if (PointsEqual(PositionA, PositionB, SMALL_NUMBER))
-			{
-				if (VertIndexAndZ[i].Index < BestIndex)
-				{
-					BestIndex = VertIndexAndZ[i].Index;
-				}
-			}
-		}
-		return BestIndex < MAX_int32 ? BestIndex : Index;
-	};
-
-	// Export to MeshDescription
-	TMap<int32, FVertexID> IndexToVertexID;
-	IndexToVertexID.Reserve(CurrentIndex);
-	CurrentIndex = 0;
-	for (int32 y = 0; y < ComponentSizeQuadsLOD; y++)
-	{
-		for (int32 x = 0; x < ComponentSizeQuadsLOD; x++)
-		{
-			FVector Positions[UE_ARRAY_COUNT(QuadPattern)];
-			bool bProcess = bIgnoreBounds;
-
-			// Fill positions
-			for (int32 i = 0; i < UE_ARRAY_COUNT(QuadPattern); i++)
-			{
-				int32 VertexX = x + QuadPattern[i].X;
-				int32 VertexY = y + QuadPattern[i].Y;
-				Positions[i] = CDI.GetWorldVertex(VertexX, VertexY);
-
-				// If at least one vertex is within the given bounds we should process the quad  
-				if (!bProcess && InBounds.ComputeSquaredDistanceFromBoxToPoint(Positions[i]) < SquaredSphereRadius)
-				{
-					bProcess = true;
-				}
-			}
-
-			if (bProcess)
-			{
-				//Fill the vertexID we need
-				TArray<FVertexID> VertexIDs;
-				VertexIDs.Reserve(UE_ARRAY_COUNT(QuadPattern));
-				TArray<FVertexInstanceID> VertexInstanceIDs;
-				VertexInstanceIDs.Reserve(UE_ARRAY_COUNT(QuadPattern));
-				// Fill positions
-				for (int32 i = 0; i < UE_ARRAY_COUNT(QuadPattern); i++)
-				{
-					int32 DuplicateLowestIndex = FindPreviousIndex(CurrentIndex);
-					FVertexID VertexID;
-					if (DuplicateLowestIndex < CurrentIndex)
-					{
-						VertexID = IndexToVertexID[DuplicateLowestIndex];
-					}
-					else
-					{
-						VertexID = OutMesh.CreateVertex();
-						VertexPositions[VertexID] = (FVector3f)Positions[i];
-					}
-					IndexToVertexID.Add(CurrentIndex, VertexID);
-					VertexIDs.Add(VertexID);
-					CurrentIndex++;
-				}
-
-				// Create triangle
-				{
-					// Whether this vertex is in hole
-					bool bInvisible = false;
-					if (VisDataMap.Num())
-					{
-						int32 TexelX, TexelY;
-						CDI.VertexXYToTexelXY(x, y, TexelX, TexelY);
-						bInvisible = (VisDataMap[CDI.TexelXYToIndex(TexelX, TexelY)] >= VisThreshold);
-					}
-					//Add vertexInstance and polygon only if we are visible
-					if (!bInvisible)
-					{
-						VertexInstanceIDs.Add(OutMesh.CreateVertexInstance(VertexIDs[0]));
-						VertexInstanceIDs.Add(OutMesh.CreateVertexInstance(VertexIDs[1]));
-						VertexInstanceIDs.Add(OutMesh.CreateVertexInstance(VertexIDs[2]));
-
-						VertexInstanceIDs.Add(OutMesh.CreateVertexInstance(VertexIDs[3]));
-						VertexInstanceIDs.Add(OutMesh.CreateVertexInstance(VertexIDs[4]));
-						VertexInstanceIDs.Add(OutMesh.CreateVertexInstance(VertexIDs[5]));
-
-						// Fill other vertex data
-						for (int32 i = 0; i < UE_ARRAY_COUNT(QuadPattern); i++)
-						{
-							int32 VertexX = x + QuadPattern[i].X;
-							int32 VertexY = y + QuadPattern[i].Y;
-
-							FVector LocalTangentX, LocalTangentY, LocalTangentZ;
-							CDI.GetLocalTangentVectors(VertexX, VertexY, LocalTangentX, LocalTangentY, LocalTangentZ);
-
-							VertexInstanceTangents[VertexInstanceIDs[i]] = (FVector3f)LocalTangentX;
-							VertexInstanceBinormalSigns[VertexInstanceIDs[i]] = GetBasisDeterminantSign(LocalTangentX, LocalTangentY, LocalTangentZ);
-							VertexInstanceNormals[VertexInstanceIDs[i]] = (FVector3f)LocalTangentZ;
-
-							FVector2f UV = ComponentUVOffsetLOD + FVector2f(VertexX, VertexY) * ComponentUVScaleLOD;
-							VertexInstanceUVs.Set(VertexInstanceIDs[i], 0, UV);
-							// Add lightmap UVs
-							VertexInstanceUVs.Set(VertexInstanceIDs[i], 1, UV);
-						}
-						auto AddTriangle = [&OutMesh, &EdgeHardnesses, &PolygonGroupID, &VertexIDs, &VertexInstanceIDs](int32 BaseIndex)
-						{
-							//Create a polygon from this triangle
-							TArray<FVertexInstanceID> PerimeterVertexInstances;
-							PerimeterVertexInstances.SetNum(3);
-							for (int32 Corner = 0; Corner < 3; ++Corner)
-							{
-								PerimeterVertexInstances[Corner] = VertexInstanceIDs[BaseIndex + Corner];
-							}
-							// Insert a polygon into the mesh
-							TArray<FEdgeID> NewEdgeIDs;
-							const FPolygonID NewPolygonID = OutMesh.CreatePolygon(PolygonGroupID, PerimeterVertexInstances, &NewEdgeIDs);
-							for (const FEdgeID& NewEdgeID : NewEdgeIDs)
-							{
-								EdgeHardnesses[NewEdgeID] = false;
-							}
-						};
-						AddTriangle(0);
-						AddTriangle(3);
-					}
-				}
-			}
-			else
-			{
-				CurrentIndex += UE_ARRAY_COUNT(QuadPattern);
-			}
-		}
-	}
-}
-
-void ULandscapeComponent::ExportToMeshDescription(const int32 InExportLOD, FMeshDescription& OutMesh)
-{
-	::ExportToMeshDescription(this, InExportLOD, OutMesh, FBoxSphereBounds(), true);
-}
-
-void ULandscapeComponent::ExportToMeshDescription(const int32 InExportLOD, const FBoxSphereBounds& InBounds, FMeshDescription& OutMesh)
-{
-	::ExportToMeshDescription(this, InExportLOD, OutMesh, InBounds, false);
+	return ULandscapeHLODBuilder::StaticClass();
 }
 
 #endif
diff --git a/Engine/Source/Runtime/Landscape/Private/LandscapeEdit.cpp b/Engine/Source/Runtime/Landscape/Private/LandscapeEdit.cpp
--- a/Engine/Source/Runtime/Landscape/Private/LandscapeEdit.cpp
+++ b/Engine/Source/Runtime/Landscape/Private/LandscapeEdit.cpp
@@ -3339,6 +3339,31 @@ bool ALandscapeProxy::ExportToRawMesh(int32 InExportLOD, FMeshDescription& OutRa
 	TInlineComponentArray<ULandscapeComponent*> RegisteredLandscapeComponents;
 	GetComponents<ULandscapeComponent>(RegisteredLandscapeComponents);
 
+	const FIntRect LandscapeSectionRect = GetBoundingRect();
+	const FVector2f LandscapeUVScale = FVector2f(1.0f, 1.0f) / FVector2f(LandscapeSectionRect.Size());
+
+	FStaticMeshAttributes Attributes(OutRawMesh);
+	TVertexAttributesRef<FVector3f> VertexPositions = Attributes.GetVertexPositions();
+	TEdgeAttributesRef<bool> EdgeHardnesses = Attributes.GetEdgeHardnesses();
+	TPolygonGroupAttributesRef<FName> PolygonGroupImportedMaterialSlotNames = Attributes.GetPolygonGroupMaterialSlotNames();
+	TVertexInstanceAttributesRef<FVector3f> VertexInstanceNormals = Attributes.GetVertexInstanceNormals();
+	TVertexInstanceAttributesRef<FVector3f> VertexInstanceTangents = Attributes.GetVertexInstanceTangents();
+	TVertexInstanceAttributesRef<float> VertexInstanceBinormalSigns = Attributes.GetVertexInstanceBinormalSigns();
+	TVertexInstanceAttributesRef<FVector4f> VertexInstanceColors = Attributes.GetVertexInstanceColors();
+	TVertexInstanceAttributesRef<FVector2f> VertexInstanceUVs = Attributes.GetVertexInstanceUVs();
+
+	if (VertexInstanceUVs.GetNumChannels() < 2)
+	{
+		VertexInstanceUVs.SetNumChannels(2);
+	}
+
+	// User specified LOD to export
+	int32 LandscapeLODToExport = ExportLOD;
+	if (InExportLOD != INDEX_NONE)
+	{
+		LandscapeLODToExport = FMath::Clamp<int32>(InExportLOD, 0, FMath::CeilLogTwo(SubsectionSizeQuads + 1) - 1);
+	}
+
 	// Export data for each component
 	for (auto It = RegisteredLandscapeComponents.CreateConstIterator(); It; ++It)
 	{
@@ -3349,8 +3374,232 @@ bool ALandscapeProxy::ExportToRawMesh(int32 InExportLOD, FMeshDescription& OutRa
 		{
 			continue;
 		}
-		
-		Component->ExportToMeshDescription(InExportLOD, OutRawMesh);
+
+		FLandscapeComponentDataInterface CDI(Component, LandscapeLODToExport);
+		const int32 ComponentSizeQuadsLOD = ((Component->ComponentSizeQuads + 1) >> LandscapeLODToExport) - 1;
+		const int32 SubsectionSizeQuadsLOD = ((Component->SubsectionSizeQuads + 1) >> LandscapeLODToExport) - 1;
+		const FIntPoint ComponentOffsetQuads = Component->GetSectionBase() - LandscapeSectionOffset - LandscapeSectionRect.Min;
+		const FVector2f ComponentUVOffsetLOD = FVector2f(ComponentOffsetQuads)*((float)ComponentSizeQuadsLOD / ComponentSizeQuads);
+		const FVector2f ComponentUVScaleLOD = LandscapeUVScale*((float)ComponentSizeQuads / ComponentSizeQuadsLOD);
+
+		const int32 NumFaces = FMath::Square(ComponentSizeQuadsLOD) * 2;
+		const int32 NumVertices = NumFaces * 3;
+
+		OutRawMesh.ReserveNewVertices(NumVertices);
+		OutRawMesh.ReserveNewPolygons(NumFaces);
+		OutRawMesh.ReserveNewVertexInstances(NumVertices);
+		OutRawMesh.ReserveNewEdges(NumVertices);
+
+		FPolygonGroupID PolygonGroupID = INDEX_NONE;
+		if (OutRawMesh.PolygonGroups().Num() < 1)
+		{
+			PolygonGroupID = OutRawMesh.CreatePolygonGroup();
+			PolygonGroupImportedMaterialSlotNames[PolygonGroupID] = FName(TEXT("LandscapeMat_0"));
+		}
+		else
+		{
+			PolygonGroupID = OutRawMesh.PolygonGroups().GetFirstValidID();
+		}
+
+		// Check if there are any holes
+		const int32 VisThreshold = 170;
+		TArray<uint8> VisDataMap;
+		const TArray<FWeightmapLayerAllocationInfo>& ComponentWeightmapLayerAllocations = Component->GetWeightmapLayerAllocations();
+
+		for (int32 AllocIdx = 0; AllocIdx < ComponentWeightmapLayerAllocations.Num(); AllocIdx++)
+		{
+			const FWeightmapLayerAllocationInfo& AllocInfo = ComponentWeightmapLayerAllocations[AllocIdx];
+			if (AllocInfo.LayerInfo == ALandscapeProxy::VisibilityLayer)
+			{
+				CDI.GetWeightmapTextureData(AllocInfo.LayerInfo, VisDataMap);
+			}
+		}
+
+		const FIntPoint QuadPattern[6] =
+		{
+			//face 1
+			FIntPoint(0, 0),
+			FIntPoint(0, 1),
+			FIntPoint(1, 1),
+			//face 2
+			FIntPoint(0, 0),
+			FIntPoint(1, 1),
+			FIntPoint(1, 0),
+		};
+
+		const int32 WeightMapSize = (SubsectionSizeQuadsLOD + 1) * Component->NumSubsections;
+
+		const float SquaredSphereRadius = FMath::Square(InBounds.SphereRadius);
+
+		//We need to not duplicate the vertex position, so we use the FIndexAndZ to achieve fast result
+		TArray<FIndexAndZ> VertIndexAndZ;
+		VertIndexAndZ.Reserve(ComponentSizeQuadsLOD*ComponentSizeQuadsLOD*UE_ARRAY_COUNT(QuadPattern));
+		int32 CurrentIndex = 0;
+		TMap<int32, FVector> IndexToPosition;
+		IndexToPosition.Reserve(ComponentSizeQuadsLOD*ComponentSizeQuadsLOD*UE_ARRAY_COUNT(QuadPattern));
+		for (int32 y = 0; y < ComponentSizeQuadsLOD; y++)
+		{
+			for (int32 x = 0; x < ComponentSizeQuadsLOD; x++)
+			{
+				for (int32 i = 0; i < UE_ARRAY_COUNT(QuadPattern); i++)
+				{
+					int32 VertexX = x + QuadPattern[i].X;
+					int32 VertexY = y + QuadPattern[i].Y;
+					FVector Position = CDI.GetWorldVertex(VertexX, VertexY);
+
+					// If at least one vertex is within the given bounds we should process the quad  
+					new(VertIndexAndZ)FIndexAndZ(CurrentIndex, Position);
+					IndexToPosition.Add(CurrentIndex, Position);
+					CurrentIndex++;
+				}
+			}
+		}
+		// Sort the vertices by z value
+		VertIndexAndZ.Sort(FCompareIndexAndZ());
+
+		auto FindPreviousIndex = [&VertIndexAndZ, &IndexToPosition](int32 Index)->int32
+		{
+			const FVector& PositionA = IndexToPosition[Index];
+			FIndexAndZ CompressPosition(0, PositionA);
+			// Search for lowest index duplicates
+			int32 BestIndex = MAX_int32;
+			for (int32 i = 0; i < IndexToPosition.Num(); i++)
+			{
+				if (CompressPosition.Z > (VertIndexAndZ[i].Z + SMALL_NUMBER))
+				{
+					//We will not find anything there is no point searching more
+					break;
+				}
+				const FVector& PositionB = IndexToPosition[VertIndexAndZ[i].Index];
+				if (PointsEqual(PositionA, PositionB, SMALL_NUMBER))
+				{
+					if (VertIndexAndZ[i].Index < BestIndex)
+					{
+						BestIndex = VertIndexAndZ[i].Index;
+					}
+				}
+			}
+			return BestIndex < MAX_int32 ? BestIndex : Index;
+		};
+
+		// Export to MeshDescription
+		TMap<int32, FVertexID> IndexToVertexID;
+		IndexToVertexID.Reserve(CurrentIndex);
+		CurrentIndex = 0;
+		for (int32 y = 0; y < ComponentSizeQuadsLOD; y++)
+		{
+			for (int32 x = 0; x < ComponentSizeQuadsLOD; x++)
+			{
+				FVector Positions[UE_ARRAY_COUNT(QuadPattern)];
+				bool bProcess = bIgnoreBounds;
+
+				// Fill positions
+				for (int32 i = 0; i < UE_ARRAY_COUNT(QuadPattern); i++)
+				{
+					int32 VertexX = x + QuadPattern[i].X;
+					int32 VertexY = y + QuadPattern[i].Y;
+					Positions[i] = CDI.GetWorldVertex(VertexX, VertexY);
+
+					// If at least one vertex is within the given bounds we should process the quad  
+					if (!bProcess && InBounds.ComputeSquaredDistanceFromBoxToPoint(Positions[i]) < SquaredSphereRadius)
+					{
+						bProcess = true;
+					}
+				}
+
+				if (bProcess)
+				{
+					//Fill the vertexID we need
+					TArray<FVertexID> VertexIDs;
+					VertexIDs.Reserve(UE_ARRAY_COUNT(QuadPattern));
+					TArray<FVertexInstanceID> VertexInstanceIDs;
+					VertexInstanceIDs.Reserve(UE_ARRAY_COUNT(QuadPattern));
+					// Fill positions
+					for (int32 i = 0; i < UE_ARRAY_COUNT(QuadPattern); i++)
+					{
+						int32 DuplicateLowestIndex = FindPreviousIndex(CurrentIndex);
+						FVertexID VertexID;
+						if (DuplicateLowestIndex < CurrentIndex)
+						{
+							VertexID = IndexToVertexID[DuplicateLowestIndex];
+						}
+						else
+						{
+							VertexID = OutRawMesh.CreateVertex();
+							VertexPositions[VertexID] = Positions[i];
+						}
+						IndexToVertexID.Add(CurrentIndex, VertexID);
+						VertexIDs.Add(VertexID);
+						CurrentIndex++;
+					}
+
+					// Create triangle
+					{
+						// Whether this vertex is in hole
+						bool bInvisible = false;
+						if (VisDataMap.Num())
+						{
+							int32 TexelX, TexelY;
+							CDI.VertexXYToTexelXY(x, y, TexelX, TexelY);
+							bInvisible = (VisDataMap[CDI.TexelXYToIndex(TexelX, TexelY)] >= VisThreshold);
+						}
+						//Add vertexInstance and polygon only if we are visible
+						if (!bInvisible)
+						{
+							VertexInstanceIDs.Add(OutRawMesh.CreateVertexInstance(VertexIDs[0]));
+							VertexInstanceIDs.Add(OutRawMesh.CreateVertexInstance(VertexIDs[1]));
+							VertexInstanceIDs.Add(OutRawMesh.CreateVertexInstance(VertexIDs[2]));
+
+							VertexInstanceIDs.Add(OutRawMesh.CreateVertexInstance(VertexIDs[3]));
+							VertexInstanceIDs.Add(OutRawMesh.CreateVertexInstance(VertexIDs[4]));
+							VertexInstanceIDs.Add(OutRawMesh.CreateVertexInstance(VertexIDs[5]));
+
+							// Fill other vertex data
+							for (int32 i = 0; i < UE_ARRAY_COUNT(QuadPattern); i++)
+							{
+								int32 VertexX = x + QuadPattern[i].X;
+								int32 VertexY = y + QuadPattern[i].Y;
+
+								FVector LocalTangentX, LocalTangentY, LocalTangentZ;
+								CDI.GetLocalTangentVectors(VertexX, VertexY, LocalTangentX, LocalTangentY, LocalTangentZ);
+
+								VertexInstanceTangents[VertexInstanceIDs[i]] = LocalTangentX;
+								VertexInstanceBinormalSigns[VertexInstanceIDs[i]] = GetBasisDeterminantSign(LocalTangentX, LocalTangentY, LocalTangentZ);
+								VertexInstanceNormals[VertexInstanceIDs[i]] = LocalTangentZ;
+
+								FVector2f UV = (ComponentUVOffsetLOD + FVector2f(VertexX, VertexY))*ComponentUVScaleLOD;
+								VertexInstanceUVs.Set(VertexInstanceIDs[i], 0, UV);
+								// Add lightmap UVs
+								VertexInstanceUVs.Set(VertexInstanceIDs[i], 1, UV);
+							}
+							auto AddTriangle = [&OutRawMesh, &EdgeHardnesses, &PolygonGroupID, &VertexIDs, &VertexInstanceIDs](int32 BaseIndex)
+							{
+								//Create a polygon from this triangle
+								TArray<FVertexInstanceID> PerimeterVertexInstances;
+								PerimeterVertexInstances.SetNum(3);
+								for (int32 Corner = 0; Corner < 3; ++Corner)
+								{
+									PerimeterVertexInstances[Corner] = VertexInstanceIDs[BaseIndex + Corner];
+								}
+								// Insert a polygon into the mesh
+								TArray<FEdgeID> NewEdgeIDs;
+								const FPolygonID NewPolygonID = OutRawMesh.CreatePolygon(PolygonGroupID, PerimeterVertexInstances, &NewEdgeIDs);
+								for (const FEdgeID& NewEdgeID : NewEdgeIDs)
+								{
+									EdgeHardnesses[NewEdgeID] = false;
+								}
+							};
+							AddTriangle(0);
+							AddTriangle(3);
+						}
+					}
+				}
+				else
+				{
+					CurrentIndex += UE_ARRAY_COUNT(QuadPattern);
+				}
+			}
+		}
 	}
 
 	//Compact the MeshDescription, if there was visibility mask or some bounding box clip, it need to be compacted so the sparse array are from 0 to n with no invalid data in between. 
diff --git a/Engine/Source/Runtime/Landscape/Private/LandscapeEditLayers.cpp b/Engine/Source/Runtime/Landscape/Private/LandscapeEditLayers.cpp
--- a/Engine/Source/Runtime/Landscape/Private/LandscapeEditLayers.cpp
+++ b/Engine/Source/Runtime/Landscape/Private/LandscapeEditLayers.cpp
@@ -8730,16 +8730,16 @@ FLandscapeLayer* ALandscape::GetLandscapeSplinesReservedLayer()
 
 LANDSCAPE_API extern bool GDisableUpdateLandscapeMaterialInstances;
 
-uint32 ULandscapeComponent::ComputeLayerHash() const
+uint32 ULandscapeComponent::ComputeLayerHash(bool InReturnEditingHash) const
 {
-	UTexture2D* Heightmap = GetHeightmap(true);
+	UTexture2D* Heightmap = GetHeightmap(InReturnEditingHash);
 	const uint8* MipData = Heightmap->Source.LockMipReadOnly(0);
 	uint32 Hash = FCrc::MemCrc32(MipData, Heightmap->GetSizeX() * Heightmap->GetSizeY() * sizeof(FColor));
 	Heightmap->Source.UnlockMip(0);
 
 	// Copy to sort
-	const TArray<UTexture2D*>& Weightmaps = GetWeightmapTextures(true);
-	TArray<FWeightmapLayerAllocationInfo> AllocationInfos = GetWeightmapLayerAllocations(true);
+	const TArray<UTexture2D*>& Weightmaps = GetWeightmapTextures(InReturnEditingHash);
+	TArray<FWeightmapLayerAllocationInfo> AllocationInfos = GetWeightmapLayerAllocations(InReturnEditingHash);
 
 	// Sort allocations infos by LayerInfo Path so the Weightmaps hahses get ordered properly
 	AllocationInfos.Sort([](const FWeightmapLayerAllocationInfo& A, const FWeightmapLayerAllocationInfo& B)
diff --git a/Engine/Source/Runtime/Landscape/Private/LandscapeHLODBuilder.cpp b/Engine/Source/Runtime/Landscape/Private/LandscapeHLODBuilder.cpp
new file mode 100644
--- /dev/null
+++ b/Engine/Source/Runtime/Landscape/Private/LandscapeHLODBuilder.cpp
@@ -0,0 +1,322 @@
+// Copyright Epic Games, Inc. All Rights Reserved.
+
+#include "LandscapeHLODBuilder.h"
+
+#if WITH_EDITOR
+
+#include "Landscape.h"
+#include "LandscapeProxy.h"
+#include "LandscapeComponent.h"
+#include "LandscapeMeshProxyComponent.h"
+
+#include "MeshDescription.h"
+#include "TriangleTypes.h"
+#include "StaticMeshAttributes.h"
+
+#include "Materials/MaterialInstanceConstant.h"
+#include "MaterialUtilities.h"
+
+#include "Engine/StaticMesh.h"
+#include "Engine/StaticMeshSourceData.h"
+
+#include "Algo/ForEach.h"
+
+#include "Serialization/ArchiveCrc32.h"
+#include "Engine/HLODProxy.h"
+
+#endif
+
+ULandscapeHLODBuilder::ULandscapeHLODBuilder(const FObjectInitializer& ObjectInitializer)
+	: Super(ObjectInitializer)
+{
+}
+
+#if WITH_EDITOR
+
+static uint32 GetTransformHash(const FRotator& ComponentRotation, const FVector& ComponentLocation, const FVector& ComponentScale)
+{
+	FArchiveCrc32 Ar;
+
+	// Include transform - round sufficiently to ensure stability
+	FIntVector Location(FMath::RoundToInt(ComponentLocation.X), FMath::RoundToInt(ComponentLocation.Y), FMath::RoundToInt(ComponentLocation.Z));
+	Ar << Location;
+
+	FRotator Rotator(ComponentRotation.GetDenormalized());
+	FIntVector Rotation(FMath::RoundToInt(Rotator.Pitch), FMath::RoundToInt(Rotator.Yaw), FMath::RoundToInt(Rotator.Roll));
+	Ar << Rotation;
+
+	float SCALE_FACTOR = 100;
+	FIntVector Scale(FMath::RoundToInt(ComponentScale.X * SCALE_FACTOR), FMath::RoundToInt(ComponentScale.Y * SCALE_FACTOR), FMath::RoundToInt(ComponentScale.Z * SCALE_FACTOR));
+	Ar << Scale;
+
+	return Ar.GetCrc();
+}
+
+uint32 ULandscapeHLODBuilder::ComputeHLODHash(const UActorComponent* InSourceComponent) const
+{
+	FArchiveCrc32 Ar;
+
+	if (const ULandscapeComponent* LSComponent = Cast<ULandscapeComponent>(InSourceComponent))
+	{
+		UE_LOG(LogHLODBuilder, VeryVerbose, TEXT("ULandscapeHLODBuilder::ComputeHLODHash %s"), *InSourceComponent->GetName());
+
+		ALandscapeProxy* LSProxy = LSComponent->GetLandscapeProxy();
+
+		// LS LOD setup
+		TArray<float> LODScreenSize = LSProxy->GetLODScreenSizeArray();
+		Ar << LODScreenSize;
+		UE_LOG(LogHLODBuilder, VeryVerbose, TEXT("     - LODScreenSize = %x"), Ar.GetCrc());
+		
+		// LS Transform
+		uint32 TransformHash = GetTransformHash(LSComponent->GetComponentRotation(), LSComponent->GetComponentLocation(), LSComponent->GetComponentScale());
+		Ar << TransformHash;
+		UE_LOG(LogHLODBuilder, VeryVerbose, TEXT("     - TransformHash = %x"), Ar.GetCrc());
+
+		// LS Content - Heightmap & Weightmaps
+		uint32 LSContentHash = LSComponent->ComputeLayerHash(/*InReturnEditingHash=*/ false);
+		Ar << LSContentHash;
+		UE_LOG(LogHLODBuilder, VeryVerbose, TEXT("     - LSContentHash = %x"), Ar.GetCrc());
+
+		TSet<UTexture*> UsedTextures;
+
+		// LS Materials
+		{
+			TArray<UMaterialInterface*> UsedMaterials;
+			if (LSProxy->bUseDynamicMaterialInstance)
+			{
+				UsedMaterials.Append(LSComponent->MaterialInstancesDynamic);
+			}
+			else
+			{
+				UsedMaterials.Append(LSComponent->MaterialInstances);
+			}
+
+			UsedMaterials.Add(LSComponent->OverrideMaterial);
+			UsedMaterials.Add(LSComponent->OverrideHoleMaterial);
+			TArray<uint32> UsedMaterialsCRC;
+			for (UMaterialInterface* MaterialInterface : UsedMaterials)
+			{
+				if (MaterialInterface)
+				{
+					uint32 MaterialCRC = UHLODProxy::GetCRC(MaterialInterface);
+					UsedMaterialsCRC.Add(MaterialCRC);
+
+					TArray<UTexture*> Textures;
+					MaterialInterface->GetUsedTextures(Textures, EMaterialQualityLevel::High, true, ERHIFeatureLevel::SM5, true);
+					UsedTextures.Append(Textures);
+				}
+			}
+			UsedMaterialsCRC.Sort();
+			Ar << UsedMaterialsCRC;
+			UE_LOG(LogHLODBuilder, VeryVerbose, TEXT("     - UsedMaterialsCRC = %x"), Ar.GetCrc());
+		}
+
+		// LS Textures
+		{
+			TArray<uint32> UsedTexturesCRC;
+			for (UTexture* Texture : UsedTextures)
+			{
+				uint32 TextureCRC = UHLODProxy::GetCRC(Texture);
+				UsedTexturesCRC.Add(TextureCRC);
+			}
+			UsedTexturesCRC.Sort();
+			Ar << UsedTexturesCRC;
+			UE_LOG(LogHLODBuilder, VeryVerbose, TEXT("     - UsedTexturesCRC = %x"), Ar.GetCrc());
+		}
+	}
+
+	return Ar.GetCrc();
+}
+
+static int32 GetMeshTextureSizeFromTargetTexelDensity(const FMeshDescription& InMesh, float InTargetTexelDensity)
+{
+	FStaticMeshConstAttributes Attributes(InMesh);
+
+	TVertexAttributesConstRef<FVector3f> VertexPositions = Attributes.GetVertexPositions();
+	
+	double Mesh3DArea = 0;
+	for (const FTriangleID TriangleID : InMesh.Triangles().GetElementIDs())
+	{
+		TArrayView<const FVertexID> TriVertices = InMesh.GetTriangleVertices(TriangleID);
+
+		// World space area
+		Mesh3DArea += UE::Geometry::VectorUtil::Area(VertexPositions[TriVertices[0]],
+													 VertexPositions[TriVertices[1]],
+													 VertexPositions[TriVertices[2]]);
+	}
+	double TexelRatio = FMath::Sqrt(1.0 / Mesh3DArea) * 100;
+
+	// Compute the perfect texture size that would get us to our texture density
+	// Also compute the nearest power of two sizes (below and above our target)
+	const int32 SizePerfect = FMath::CeilToInt(InTargetTexelDensity / TexelRatio);
+	const int32 SizeHi = FMath::RoundUpToPowerOfTwo(SizePerfect);
+	const int32 SizeLo = SizeHi >> 1;
+
+	// Compute the texel density we achieve with these two texture sizes
+	const double TexelDensityLo = SizeLo * TexelRatio;
+	const double TexelDensityHi = SizeHi * TexelRatio;
+
+	// Select best match between low & high res textures.
+	const double TexelDensityLoDiff = InTargetTexelDensity - TexelDensityLo;
+	const double TexelDensityHiDiff = TexelDensityHi - InTargetTexelDensity;
+	const int32 BestTextureSize = TexelDensityLoDiff < TexelDensityHiDiff ? SizeLo : SizeHi;
+
+	return BestTextureSize;
+}
+
+static int32 ComputeRequiredLandscapeLOD(const ALandscapeProxy* InLandscapeProxy, const float InViewDistance)
+{
+	check(InLandscapeProxy && !InLandscapeProxy->LandscapeComponents.IsEmpty());
+	
+	// These constants are showing up a lot in the screen size computation for Level HLODs. This should be configurable per project.
+	const float HalfFOV = PI * 0.25f;
+	const float ScreenWidth = 1920.0f;
+	const float ScreenHeight = 1080.0f;
+	const FPerspectiveMatrix ProjMatrix(HalfFOV, ScreenWidth, ScreenHeight, 1.0f);
+
+	TArray<float> LODScreenSizes = InLandscapeProxy->GetLODScreenSizeArray();
+
+	ULandscapeComponent* LSComponent = InLandscapeProxy->LandscapeComponents[0];
+	float ComponentRadiusScaled = LSComponent->GetLocalBounds().SphereRadius * LSComponent->GetComponentTransform().GetScale3D().GetAbsMax();
+	float ExpectedScreenSize = ComputeBoundsScreenSize(FVector::ZeroVector, ComponentRadiusScaled, FVector(0.0f, 0.0f, InViewDistance), ProjMatrix);
+
+	int32 RequiredLOD;
+	for (RequiredLOD = 0; RequiredLOD < LODScreenSizes.Num(); ++RequiredLOD)
+	{
+		if (ExpectedScreenSize > LODScreenSizes[RequiredLOD])
+		{
+			break;
+		}
+	}
+
+	return RequiredLOD;
+}
+
+static UMaterialInterface* BakeLandscapeMaterial(const FHLODBuildContext& InHLODBuildContext, const FMeshDescription& InMeshDescription, const ALandscapeProxy* InLandscapeProxy, const int32 InLandscapeLOD, int32 InTextureSize)
+{
+	// Build landscape material
+	FFlattenMaterial LandscapeFlattenMaterial;
+	LandscapeFlattenMaterial.SetPropertySize(EFlattenMaterialProperties::Diffuse, InTextureSize);
+	LandscapeFlattenMaterial.SetPropertySize(EFlattenMaterialProperties::Normal, InTextureSize);
+	LandscapeFlattenMaterial.SetPropertySize(EFlattenMaterialProperties::Metallic, InTextureSize);
+	LandscapeFlattenMaterial.SetPropertySize(EFlattenMaterialProperties::Roughness, InTextureSize);
+	LandscapeFlattenMaterial.SetPropertySize(EFlattenMaterialProperties::Specular, InTextureSize);
+	FMaterialUtilities::ExportLandscapeMaterial(InLandscapeProxy, LandscapeFlattenMaterial);
+
+	// Optimize flattened material
+	FMaterialUtilities::OptimizeFlattenMaterial(LandscapeFlattenMaterial);
+
+	UMaterial* LandscapeMaterial = GEngine->DefaultLandscapeFlattenMaterial;
+
+	// Validate that the flatten material expects world space normals
+	if (LandscapeMaterial->bTangentSpaceNormal)
+	{
+		UE_LOG(LogHLODBuilder, Error, TEXT("Landscape flatten material %s should use world space normals rather than tangent space normals."), *LandscapeMaterial->GetName());
+	}
+
+	FMaterialProxySettings MaterialProxySettings;
+	MaterialProxySettings.TextureSizingType = ETextureSizingType::TextureSizingType_UseSingleTextureSize;
+	MaterialProxySettings.TextureSize = InTextureSize;
+	MaterialProxySettings.bNormalMap = true;
+	MaterialProxySettings.bMetallicMap = true;
+	MaterialProxySettings.bRoughnessMap = true;
+	MaterialProxySettings.bSpecularMap = true;
+
+	// Create a new proxy material instance
+	TArray<UObject*> GeneratedAssets;
+	UMaterialInstanceConstant* LandscapeMaterialInstance = FMaterialUtilities::CreateFlattenMaterialInstance(InHLODBuildContext.AssetsOuter->GetPackage(), MaterialProxySettings, LandscapeMaterial, LandscapeFlattenMaterial, InHLODBuildContext.AssetsBaseName, InLandscapeProxy->GetName(), GeneratedAssets);
+
+	Algo::ForEach(GeneratedAssets, [](UObject* Asset) 
+	{
+		// We don't want any of the generate HLOD assets to be public
+		Asset->ClearFlags(RF_Public | RF_Standalone);
+
+		// Use clamp texture addressing to avoid artifacts between tiles
+		if (UTexture2D* Texture = Cast<UTexture2D>(Asset))
+		{
+			Texture->AddressX = TA_Clamp;
+			Texture->AddressY = TA_Clamp;
+		}
+	});
+
+	return LandscapeMaterialInstance;
+}
+
+// This is an initial implementation for UE 5.0
+// Multiple improvements could be done
+// * Currently, for each referenced landscape proxy, we generate individual HLOD meshes & textures. This should output a single mesh for all proxies
+// * Generated mesh could be Nanite-enabled
+TArray<UActorComponent*> ULandscapeHLODBuilder::Build(const FHLODBuildContext& InHLODBuildContext, const TArray<UActorComponent*>& InSourceComponents) const
+{
+	TArray<ULandscapeComponent*> SourceLandscapeComponents = FilterComponents<ULandscapeComponent>(InSourceComponents);
+	
+	TSet<ALandscapeProxy*> LandscapeProxies;
+	Algo::Transform(SourceLandscapeComponents, LandscapeProxies, [](ULandscapeComponent* SourceComponent) { return SourceComponent->GetLandscapeProxy(); });
+
+	// This code assume all components of a proxy are included in the same build... validate this
+	checkCode
+	(
+		for (ALandscapeProxy* LandscapeProxy : LandscapeProxies)
+		{
+			for (ULandscapeComponent* LandscapeComponent : LandscapeProxy->LandscapeComponents)
+			{
+				check(SourceLandscapeComponents.Contains(LandscapeComponent));
+			}
+		}
+	);
+
+	TArray<UActorComponent*> HLODComponents;
+	TArray<UStaticMesh*> StaticMeshes;
+
+	for (ALandscapeProxy* LandscapeProxy : LandscapeProxies)
+	{
+		UStaticMesh* StaticMesh = NewObject<UStaticMesh>(InHLODBuildContext.AssetsOuter);
+		FMeshDescription* MeshDescription = nullptr;
+
+		// Compute source landscape LOD
+		int32 LandscapeLOD = ComputeRequiredLandscapeLOD(LandscapeProxy, InHLODBuildContext.MinVisibleDistance);
+
+		// Mesh
+		{
+			FStaticMeshSourceModel& SrcModel = StaticMesh->AddSourceModel();
+			// Don't allow the engine to recalculate normals
+			SrcModel.BuildSettings.bRecomputeNormals = false;
+			SrcModel.BuildSettings.bRecomputeTangents = false;
+			SrcModel.BuildSettings.bRemoveDegenerates = false;
+			SrcModel.BuildSettings.bUseHighPrecisionTangentBasis = false;
+			SrcModel.BuildSettings.bUseFullPrecisionUVs = false;
+
+			MeshDescription = StaticMesh->CreateMeshDescription(0);
+	
+			LandscapeProxy->ExportToRawMesh(LandscapeLOD, *MeshDescription);
+
+			StaticMesh->CommitMeshDescription(0);
+
+			StaticMesh->ImportVersion = EImportStaticMeshVersion::LastVersion;
+		}
+
+		// Material
+		{
+			float TargetTexelDensityPerMeter = FMaterialUtilities::ComputeRequiredTexelDensityFromDrawDistance(InHLODBuildContext.MinVisibleDistance, MeshDescription->GetBounds().SphereRadius);
+			int32 RequiredTextureSize = GetMeshTextureSizeFromTargetTexelDensity(*MeshDescription, TargetTexelDensityPerMeter);
+
+			UMaterialInterface* LandscapeMaterial = BakeLandscapeMaterial(InHLODBuildContext, *MeshDescription, LandscapeProxy, LandscapeLOD, RequiredTextureSize);
+
+			//Assign the proxy material to the static mesh
+			StaticMesh->GetStaticMaterials().Add(FStaticMaterial(LandscapeMaterial));
+		}
+
+		StaticMeshes.Add(StaticMesh);
+
+		ULandscapeMeshProxyComponent* LandcapeMeshProxyComponent = NewObject<ULandscapeMeshProxyComponent>();
+		LandcapeMeshProxyComponent->InitializeForLandscape(LandscapeProxy, LandscapeLOD);
+		LandcapeMeshProxyComponent->SetStaticMesh(StaticMesh);
+		HLODComponents.Add(LandcapeMeshProxyComponent);
+	}
+
+	UStaticMesh::BatchBuild(StaticMeshes);
+
+	return HLODComponents;
+}
+
+#endif
diff --git a/Engine/Source/Runtime/Landscape/Public/LandscapeHLODBuilder.h b/Engine/Source/Runtime/Landscape/Public/LandscapeHLODBuilder.h
new file mode 100644
--- /dev/null
+++ b/Engine/Source/Runtime/Landscape/Public/LandscapeHLODBuilder.h
@@ -0,0 +1,24 @@
+// Copyright Epic Games, Inc. All Rights Reserved.
+
+#pragma once
+
+#include "CoreMinimal.h"
+
+#include "WorldPartition/HLOD/HLODBuilder.h"
+#include "LandscapeHLODBuilder.generated.h"
+
+UCLASS()
+class LANDSCAPE_API ULandscapeHLODBuilder : public UHLODBuilder
+{
+	GENERATED_UCLASS_BODY()
+
+public:
+#if WITH_EDITOR
+	virtual uint32 ComputeHLODHash(const UActorComponent* InSourceComponent) const;
+
+	/**
+	 * Components created with this method needs to be properly outered & assigned to your target actor.
+	 */
+	virtual TArray<UActorComponent*> Build(const FHLODBuildContext& InHLODBuildContext, const TArray<UActorComponent*>& InSourceComponents) const override;
+#endif
+};
