Fixed warning with debug shader not being excluded from landscape thumbnail renders

diff --git a/Engine/Source/Runtime/Landscape/Private/LandscapeRender.cpp b/Engine/Source/Runtime/Landscape/Private/LandscapeRender.cpp
--- a/Engine/Source/Runtime/Landscape/Private/LandscapeRender.cpp
+++ b/Engine/Source/Runtime/Landscape/Private/LandscapeRender.cpp
@@ -3641,12 +3641,7 @@ public:
 			FName(TEXT("TBasePassVSFSimpleStationaryLightVolumetricLightmapShadowsLightingPolicy")),
 
 			// Debug materials : 
-			FName(TEXT("FMaterialTexCoordScalePS")),
-			FName(TEXT("FMeshTexCoordSizeAccuracyPS")),
-			FName(TEXT("FPrimitiveDistanceAccuracyPS")),
-			FName(TEXT("FRequiredTextureResolutionPS")),
-			FName(TEXT("FComplexityAccumulatePS")),
-			FName(TEXT("FLODColorationPS")),
+			FName(TEXT("FDebugViewModePS")),
 
 			// Mobile
 			FName(TEXT("TMobileBasePassPSFMobileMovableDirectionalLightCSMWithLightmapPolicyINT32_MAXHDRLinear64Skylight")),
