Issue importing landscape with a material that has multiple layers

Prevent fallback to the layer "PreviewWeight" if a layer isn't allocated for a component.

diff --git a/Engine/Source/Runtime/Landscape/Classes/Materials/MaterialExpressionLandscapeLayerBlend.h b/Engine/Source/Runtime/Landscape/Classes/Materials/MaterialExpressionLandscapeLayerBlend.h
--- a/Engine/Source/Runtime/Landscape/Classes/Materials/MaterialExpressionLandscapeLayerBlend.h
+++ b/Engine/Source/Runtime/Landscape/Classes/Materials/MaterialExpressionLandscapeLayerBlend.h
@@ -78,6 +78,7 @@ class LANDSCAPE_API UMaterialExpressionLandscapeLayerBlend : public UMaterialExp
 	virtual uint32 GetOutputType(int32 InputIndex) override { return MCT_Unknown; }
 	virtual bool IsResultMaterialAttributes(int32 OutputIndex) override;
 	virtual int32 Compile(class FMaterialCompiler* Compiler, int32 OutputIndex) override;
+	virtual int32 CompilePreview(class FMaterialCompiler* Compiler, int32 OutputIndex) override;
 	virtual void GetCaption(TArray<FString>& OutCaptions) const override;
 	virtual const TArray<FExpressionInput*> GetInputs() override;
 	virtual FExpressionInput* GetInput(int32 InputIndex) override;
@@ -91,6 +92,9 @@ class LANDSCAPE_API UMaterialExpressionLandscapeLayerBlend : public UMaterialExp
 	virtual UObject* GetReferencedTexture() const override;
 	virtual bool CanReferenceTexture() const override { return true; }
 	//~ End UMaterialExpression Interface
+
+private:
+	int32 Compile(class FMaterialCompiler* Compiler, int32 OutputIndex, bool bForPreview);
 };
 
 
diff --git a/Engine/Source/Runtime/Landscape/Private/Materials/MaterialExpressionLandscapeLayerBlend.cpp b/Engine/Source/Runtime/Landscape/Private/Materials/MaterialExpressionLandscapeLayerBlend.cpp
--- a/Engine/Source/Runtime/Landscape/Private/Materials/MaterialExpressionLandscapeLayerBlend.cpp
+++ b/Engine/Source/Runtime/Landscape/Private/Materials/MaterialExpressionLandscapeLayerBlend.cpp
@@ -149,6 +149,16 @@ bool UMaterialExpressionLandscapeLayerBlend::IsResultMaterialAttributes(int32 Ou
 }
 
 int32 UMaterialExpressionLandscapeLayerBlend::Compile(class FMaterialCompiler* Compiler, int32 OutputIndex)
+{
+	return Compile(Compiler, OutputIndex, false);
+}
+
+int32 UMaterialExpressionLandscapeLayerBlend::CompilePreview(class FMaterialCompiler* Compiler, int32 OutputIndex)
+{
+	return Compile(Compiler, OutputIndex, true);
+}
+
+int32 UMaterialExpressionLandscapeLayerBlend::Compile(class FMaterialCompiler* Compiler, int32 OutputIndex, bool bForPreview)
 {
 	// For renormalization
 	bool bNeedsRenormalize = false;
@@ -170,7 +180,8 @@ int32 UMaterialExpressionLandscapeLayerBlend::Compile(class FMaterialCompiler* C
 			// Height input
 			const int32 HeightCode = Layer.HeightInput.Expression ? Layer.HeightInput.Compile(Compiler) : Compiler->Constant(Layer.ConstHeightInput);
 
-			const int32 WeightCode = Compiler->StaticTerrainLayerWeight(Layer.LayerName, Layer.PreviewWeight > 0.0f ? Compiler->Constant(Layer.PreviewWeight) : INDEX_NONE);
+			const int32 DefaultWeightCode = bForPreview && Layer.PreviewWeight > 0.0f ? Compiler->Constant(Layer.PreviewWeight) : INDEX_NONE;
+			const int32 WeightCode = Compiler->StaticTerrainLayerWeight(Layer.LayerName, DefaultWeightCode);
 			if (WeightCode != INDEX_NONE)
 			{
 				switch (Layer.BlendType)
@@ -232,7 +243,8 @@ int32 UMaterialExpressionLandscapeLayerBlend::Compile(class FMaterialCompiler* C
 	{
 		if (Layer.BlendType == LB_AlphaBlend)
 		{
-			const int32 WeightCode = Compiler->StaticTerrainLayerWeight(Layer.LayerName, Layer.PreviewWeight > 0.0f ? Compiler->Constant(Layer.PreviewWeight) : INDEX_NONE);
+			const int32 DefaultWeightCode = bForPreview && Layer.PreviewWeight > 0.0f ? Compiler->Constant(Layer.PreviewWeight) : INDEX_NONE;
+			const int32 WeightCode = Compiler->StaticTerrainLayerWeight(Layer.LayerName, DefaultWeightCode);
 			if (WeightCode != INDEX_NONE)
 			{
 				const int32 LayerCode = Layer.LayerInput.Expression ? Layer.LayerInput.Compile(Compiler) : Compiler->Constant3(Layer.ConstLayerInput.X, Layer.ConstLayerInput.Y, Layer.ConstLayerInput.Z);
