Fixed assert when cooking landscape components on mobile and creating materials

Cooking a ULandscapeComponent is updating a LandscapeMaterialInstanceConstant and the landscape MIC has already had BeginCacheForCookedPlatformData() called on it.

Here is what's happening:

UMaterialInstance::BeginCacheForCookedPlatformData is called for a Landscape MIC
This MIC's bHasStaticPermutationResource == false so no resources are cooked for the MIC.  The MIC's CachedMaterialResourcesForCooking is empty.
ULandscapeComponent::BeginCacheForCookedPlatformData is called on a Landscape component.
The Landscape component changes the MIC to have a static permutation and calls `CombinationMaterialInstance->UpdateStaticPermutation(StaticParameters, InMaterialUpdateContext);`
We then go to Save the MIC, which now has bHasStaticPermutationResource==true, but it doesn't have resources cached for cooking and triggers the assert.

The new code ensures we call `BeginCacheForCookedPlatformData` on any material instances created during landscape generation while cooking.

diff --git a/Engine/Source/Runtime/Landscape/Classes/LandscapeComponent.h b/Engine/Source/Runtime/Landscape/Classes/LandscapeComponent.h
--- a/Engine/Source/Runtime/Landscape/Classes/LandscapeComponent.h
+++ b/Engine/Source/Runtime/Landscape/Classes/LandscapeComponent.h
@@ -860,7 +860,7 @@ public:
 	// Generates mobile platform data for this component
 	void GenerateMobileWeightmapLayerAllocations();
 	void GeneratePlatformVertexData(const ITargetPlatform* TargetPlatform);
-	void GeneratePlatformPixelData();
+	void GeneratePlatformPixelData(bool bIsCooking, const ITargetPlatform* TargetPlatform);
 
 	/** Generate mobile data if it's missing or outdated */
 	void CheckGenerateLandscapePlatformData(bool bIsCooking, const ITargetPlatform* TargetPlatform);
diff --git a/Engine/Source/Runtime/Landscape/Private/Landscape.cpp b/Engine/Source/Runtime/Landscape/Private/Landscape.cpp
--- a/Engine/Source/Runtime/Landscape/Private/Landscape.cpp
+++ b/Engine/Source/Runtime/Landscape/Private/Landscape.cpp
@@ -338,7 +338,7 @@ void ULandscapeComponent::CheckGenerateLandscapePlatformData(bool bIsCooking, co
 
 	if (bRegeneratePixelData)
 	{
-		GeneratePlatformPixelData();
+		GeneratePlatformPixelData(bIsCooking, TargetPlatform);
 	}
 
 	MobileDataSourceHash = NewSourceHash;
diff --git a/Engine/Source/Runtime/Landscape/Private/LandscapeEdit.cpp b/Engine/Source/Runtime/Landscape/Private/LandscapeEdit.cpp
--- a/Engine/Source/Runtime/Landscape/Private/LandscapeEdit.cpp
+++ b/Engine/Source/Runtime/Landscape/Private/LandscapeEdit.cpp
@@ -6443,7 +6443,7 @@ void ULandscapeComponent::GenerateMobileWeightmapLayerAllocations()
 	}));
 }
 
-void ULandscapeComponent::GeneratePlatformPixelData()
+void ULandscapeComponent::GeneratePlatformPixelData(bool bIsCooking, const ITargetPlatform* TargetPlatform)
 {
 	check(!IsTemplate());
 
@@ -6607,6 +6607,12 @@ void ULandscapeComponent::GeneratePlatformPixelData()
 			MobileCombinationMaterialInstances[MaterialIndex] = GetCombinationMaterial(nullptr, MobileWeightmapLayerAllocations, MaterialLOD, true);
 			check(MobileCombinationMaterialInstances[MaterialIndex] != nullptr);
 
+			if (bIsCooking)
+			{
+				// If we are cooking ensure we are caching shader maps.
+				MobileCombinationMaterialInstances[MaterialIndex]->BeginCacheForCookedPlatformData(TargetPlatform);
+			}
+
 			UMaterialInstanceConstant* NewMobileMaterialInstance = NewObject<ULandscapeMaterialInstanceConstant>(this);
 
 			NewMobileMaterialInstance->SetParentEditorOnly(MobileCombinationMaterialInstances[MaterialIndex]);
