Landscape edit layers local merge fixes :
* Fixed assert when updating heightmaps of different sizes
* Fixed crash when deleting components : the deleted components were still referenced in the shared weightmap texture usages

Misc:
* Properly delete edit layers heightmaps and weightmaps when deleting landscape components

diff --git a/Engine/Source/Editor/LandscapeEditor/Private/LandscapeEdMode.cpp b/Engine/Source/Editor/LandscapeEditor/Private/LandscapeEdMode.cpp
--- a/Engine/Source/Editor/LandscapeEditor/Private/LandscapeEdMode.cpp
+++ b/Engine/Source/Editor/LandscapeEditor/Private/LandscapeEdMode.cpp
@@ -3157,24 +3157,57 @@ void FEdModeLandscape::DeleteLandscapeComponents(ULandscapeInfo* LandscapeInfo,
 			}
 		}
 
-		UTexture2D* HeightmapTexture = Component->GetHeightmap();
-
-		if (HeightmapTexture)
 		{
-			HeightmapTexture->SetFlags(RF_Transactional);
-			HeightmapTexture->Modify();
-			HeightmapTexture->MarkPackageDirty();
-			HeightmapTexture->ClearFlags(RF_Standalone); // Remove when there is no reference for this Heightmap...
+			TSet<UTexture2D*> OldHeightmapTextures;
+			OldHeightmapTextures.Add(Component->GetHeightmap(FGuid()));
+			// Also process all edit layers heightmaps :
+			Component->ForEachLayer([&](const FGuid& LayerGuid, FLandscapeLayerComponentData& LayerData)
+			{
+				OldHeightmapTextures.Add(Component->GetHeightmap(LayerGuid));
+			});
+
+			for (UTexture2D* HeightmapTexture : OldHeightmapTextures)
+			{
+				check(HeightmapTexture != nullptr);
+				HeightmapTexture->SetFlags(RF_Transactional);
+				HeightmapTexture->Modify();
+				HeightmapTexture->MarkPackageDirty();
+				HeightmapTexture->ClearFlags(RF_Standalone); // Remove when there is no reference for this Heightmap...
+			}
 		}
 
-		const TArray<UTexture2D*>& ComponentWeightmapTextures = Component->GetWeightmapTextures();
+		{
+			TArray<UTexture2D*> OldWeightmapTextures = Component->GetWeightmapTextures(FGuid());
+			// Also process all edit layers weightmaps :
+			Component->ForEachLayer([&](const FGuid& LayerGuid, FLandscapeLayerComponentData& LayerData)
+			{
+				OldWeightmapTextures.Append(Component->GetWeightmapTextures(LayerGuid));
+			});
+
+			for (UTexture2D* WeightmapTexture : OldWeightmapTextures)
+			{
+				check(WeightmapTexture != nullptr);
+				WeightmapTexture->SetFlags(RF_Transactional);
+				WeightmapTexture->Modify();
+				WeightmapTexture->MarkPackageDirty();
+				WeightmapTexture->ClearFlags(RF_Standalone);
+			}
+		}
 
-		for (UTexture2D* WeightmapTexture : ComponentWeightmapTextures)
 		{
-			WeightmapTexture->SetFlags(RF_Transactional);
-			WeightmapTexture->Modify();
-			WeightmapTexture->MarkPackageDirty();
-			WeightmapTexture->ClearFlags(RF_Standalone);
+			// The shared weightmap usages also need to be cleaned up, otherwise, one might still think these textures apply to deleted components :
+			TArray<ULandscapeWeightmapUsage*> ComponentWeightmapTexturesUsage = Component->GetWeightmapTexturesUsage(FGuid());
+			// Also process all edit layers weightmap textures usages :
+			Component->ForEachLayer([&](const FGuid& LayerGuid, FLandscapeLayerComponentData& LayerData)
+			{
+				ComponentWeightmapTexturesUsage.Append(Component->GetWeightmapTexturesUsage(LayerGuid));
+			});
+
+			for (ULandscapeWeightmapUsage* Usage : ComponentWeightmapTexturesUsage)
+			{
+				// Make sure these usages don't reference this component anymore :
+				Usage->ClearUsage(Component);
+			}
 		}
 
 		if (Component->XYOffsetmapTexture)
diff --git a/Engine/Source/Runtime/Landscape/Classes/LandscapeComponent.h b/Engine/Source/Runtime/Landscape/Classes/LandscapeComponent.h
--- a/Engine/Source/Runtime/Landscape/Classes/LandscapeComponent.h
+++ b/Engine/Source/Runtime/Landscape/Classes/LandscapeComponent.h
@@ -746,7 +746,10 @@ public:
 	virtual TArray<URuntimeVirtualTexture*> const& GetRuntimeVirtualTextures() const override;
 	virtual ERuntimeVirtualTextureMainPassType GetVirtualTextureRenderPassType() const override;
 
+	// Returns the heightmap for this component. If InReturnEditingHeightmap is passed, returns the currently active edit layer's heightmap : 
 	LANDSCAPE_API UTexture2D* GetHeightmap(bool InReturnEditingHeightmap = false) const;
+	// Returns the heightmap for this component and the edit layer specified by InLayerGuid. If InLayerGuid is invalid, returns the final (base) heightmap : 
+	LANDSCAPE_API UTexture2D* GetHeightmap(const FGuid& InLayerGuid) const;
 	LANDSCAPE_API TArray<UTexture2D*>& GetWeightmapTextures(bool InReturnEditingWeightmap = false);
 	LANDSCAPE_API const TArray<UTexture2D*>& GetWeightmapTextures(bool InReturnEditingWeightmap = false) const;
 	LANDSCAPE_API TArray<UTexture2D*>& GetWeightmapTextures(const FGuid& InLayerGuid);
diff --git a/Engine/Source/Runtime/Landscape/Private/Landscape.cpp b/Engine/Source/Runtime/Landscape/Private/Landscape.cpp
--- a/Engine/Source/Runtime/Landscape/Private/Landscape.cpp
+++ b/Engine/Source/Runtime/Landscape/Private/Landscape.cpp
@@ -1559,6 +1559,21 @@ UTexture2D* ULandscapeComponent::GetHeightmap(bool InReturnEditingHeightmap) con
 	return HeightmapTexture;
 }
 
+UTexture2D* ULandscapeComponent::GetHeightmap(const FGuid& InLayerGuid) const
+{
+#if WITH_EDITORONLY_DATA
+	if (InLayerGuid.IsValid())
+	{
+		if (const FLandscapeLayerComponentData* LayerData = GetLayerData(InLayerGuid))
+		{
+			return LayerData->HeightmapData.Texture;
+		}
+	}
+#endif
+
+	return HeightmapTexture;
+}
+
 const TArray<UTexture2D*>& ULandscapeComponent::GetWeightmapTextures(bool InReturnEditingWeightmap) const
 {
 #if WITH_EDITORONLY_DATA
diff --git a/Engine/Source/Runtime/Landscape/Private/LandscapeEdit.cpp b/Engine/Source/Runtime/Landscape/Private/LandscapeEdit.cpp
--- a/Engine/Source/Runtime/Landscape/Private/LandscapeEdit.cpp
+++ b/Engine/Source/Runtime/Landscape/Private/LandscapeEdit.cpp
@@ -4473,12 +4473,17 @@ ALandscapeProxy* ULandscapeInfo::MoveComponentsToProxy(const TArray<ULandscapeCo
 		}
 	}
 
-	// Check which ones are need for height map change
+	// Check which heightmap will need to be renewed :
 	TSet<UTexture2D*> OldHeightmapTextures;
 	for (ULandscapeComponent* Component : TargetSelectedComponents)
 	{
 		Component->Modify();
 		OldHeightmapTextures.Add(Component->GetHeightmap());
+		// Also process all edit layers heightmaps :
+		Component->ForEachLayer([&](const FGuid& LayerGuid, FLandscapeLayerComponentData& LayerData)
+		{
+			OldHeightmapTextures.Add(Component->GetHeightmap(LayerGuid));
+		});
 	}
 
 	// Need to split all the component which share Heightmap with selected components
@@ -4528,9 +4533,10 @@ ALandscapeProxy* ULandscapeInfo::MoveComponentsToProxy(const TArray<ULandscapeCo
 		ALandscape::SplitHeightmap(HeightmapUpdateComponentPair.Key, HeightmapUpdateComponentPair.Value ? LandscapeProxy : nullptr);
 	}
 
-	// Delete if it is no referenced textures...
+	// Delete if textures are not referenced anymore...
 	for (UTexture2D* Texture : OldHeightmapTextures)
 	{
+		check(Texture != nullptr);
 		Texture->SetFlags(RF_Transactional);
 		Texture->Modify();
 		Texture->MarkPackageDirty();
@@ -4569,7 +4575,7 @@ ALandscapeProxy* ULandscapeInfo::MoveComponentsToProxy(const TArray<ULandscapeCo
 			Landscape->RequestLayersContentUpdateForceAll();
 		}
 
-		// Need to Repacking all the Weight map (to make it packed well...)
+		// Need to re-pack all the Weight map (to have it optimally re-packed...)
 		for (ALandscapeProxy* Proxy : SelectProxies)
 		{
 			Proxy->RemoveInvalidWeightmaps();
diff --git a/Engine/Source/Runtime/Landscape/Private/LandscapeEditLayers.cpp b/Engine/Source/Runtime/Landscape/Private/LandscapeEditLayers.cpp
--- a/Engine/Source/Runtime/Landscape/Private/LandscapeEditLayers.cpp
+++ b/Engine/Source/Runtime/Landscape/Private/LandscapeEditLayers.cpp
@@ -3869,6 +3869,10 @@ namespace EditLayersHeightmapLocalMerge_RenderThread
 
 	struct FTextureResolveInfo
 	{
+		// Size of the entire texture that needs resolving :
+		FIntPoint TextureSize = FIntPoint(ForceInitToZero);
+		// Number of mips corresponding to that size : 
+		int32 NumMips = 0;
 		// Texture that was updated and needs resolving : 
 		FTexture2DResource* Texture = nullptr;
 		// CPU readback utility to bring back the result on the CPU : 
@@ -3923,17 +3927,17 @@ namespace EditLayersHeightmapLocalMerge_RenderThread
 
 	public:
 		// Number of vertices per component
-		FIntPoint ComponentSizeVerts;
+		FIntPoint ComponentSizeVerts = FIntPoint(ForceInitToZero);
+
+		// Maximum size of all heigthmaps (one heightmap can contain multiple components due to heightmap sharing)
+		FIntPoint MaxHeightmapSize = FIntPoint(ForceInitToZero);
 
-		// Size of each heigthmap (one heightmap can contain multiple components due to heightmap sharing)
-		FIntPoint HeightmapSize;
+		// Maximum number of mips of all heightmaps (which can be of different sizes) :
+		int32 MaxHeightmapNumMips = 0;
 
 		// Heightmap pixel to world scale factor
-		FVector LandscapeGridScale;
+		FVector LandscapeGridScale = FVector::ZeroVector;
 		
-		// Number of mips for the final heightmaps
-		int32 NumMips = 0;
-
 		// Maximum number of visible edit layers that have to be merged for a single FComponentRenderInfo : 
 		int32 MaxNumEditLayersTexturesToMerge = 0;
 
@@ -4215,11 +4219,11 @@ namespace EditLayersHeightmapLocalMerge_RenderThread
 			const FLandscapeRDGTrackedTexture* TrackedTexture = InTrackedTextures.Find(TextureResolveInfo.Texture);
 			check((TrackedTexture != nullptr) && (TrackedTexture->ScratchTextureRef != nullptr));
 
-			check(InLocalMergeInfo.NumMips == TrackedTexture->ScratchTextureRef->Desc.NumMips);
-			check(TrackedTexture->ScratchTextureMipsSRVRefs.Num() == InLocalMergeInfo.NumMips);
+			check(TextureResolveInfo.NumMips == TrackedTexture->ScratchTextureRef->Desc.NumMips);
+			check(TrackedTexture->ScratchTextureMipsSRVRefs.Num() == TextureResolveInfo.NumMips);
 
 			FIntPoint CurrentMipSubregionSize = InLocalMergeInfo.ComponentSizeVerts;
-			for (int32 MipLevel = 1; MipLevel < InLocalMergeInfo.NumMips; ++MipLevel)
+			for (int32 MipLevel = 1; MipLevel < TextureResolveInfo.NumMips; ++MipLevel)
 			{
 				CurrentMipSubregionSize.X >>= 1;
 				CurrentMipSubregionSize.Y >>= 1;
@@ -4246,7 +4250,7 @@ namespace EditLayersHeightmapLocalMerge_RenderThread
 
 			FRHICopyTextureInfo CopyTextureInfo;
 			// We want to copy all mips : 
-			CopyTextureInfo.NumMips = InLocalMergeInfo.NumMips;
+			CopyTextureInfo.NumMips = TextureResolveInfo.NumMips;
 
 			AddCopyTexturePass(GraphBuilder, TrackedTexture->ScratchTextureRef, TrackedTexture->ExternalTextureRef, CopyTextureInfo);
 		}
@@ -4258,12 +4262,13 @@ namespace EditLayersHeightmapLocalMerge_RenderThread
 		const uint16 HeightValue = LandscapeDataAccess::GetTexHeight(0.0f);
 		FLinearColor ClearHeightColor((float)((HeightValue - (HeightValue & 255)) >> 8) / 255.0f, (float)(HeightValue & 255) / 255.0f, 0.0f, 0.0f);
 
+		// Even if we have heightmaps of different sizes to handle, we only need one empty heightmap to copy from (whose size is MaxHeightmapSize) :
 		// TODO [jonathan.bard] : change to PF_R8G8 once edit layers heightmaps are stored as such :
-		FRDGTextureDesc Desc = FRDGTextureDesc::Create2D(InLocalMergeInfo.HeightmapSize, PF_B8G8R8A8, FClearValueBinding(ClearHeightColor), TexCreate_ShaderResource | TexCreate_RenderTargetable, InLocalMergeInfo.NumMips, /*InNumSamples = */1);
+		FRDGTextureDesc Desc = FRDGTextureDesc::Create2D(InLocalMergeInfo.MaxHeightmapSize, PF_B8G8R8A8, FClearValueBinding(ClearHeightColor), TexCreate_ShaderResource | TexCreate_RenderTargetable, InLocalMergeInfo.MaxHeightmapNumMips, /*InNumSamples = */1);
 		FRDGTextureRef EmptyTexture = GraphBuilder.CreateTexture(Desc, TEXT("LandscapeEditLayersEmptyHeightmap"));
 
 		FRDGTextureClearInfo ClearInfo;
-		ClearInfo.NumMips = InLocalMergeInfo.NumMips;
+		ClearInfo.NumMips = InLocalMergeInfo.MaxHeightmapNumMips;
 		AddClearRenderTargetPass(GraphBuilder, EmptyTexture, ClearInfo);
 
 		return EmptyTexture;
@@ -4288,7 +4293,10 @@ namespace EditLayersHeightmapLocalMerge_RenderThread
 
 			FRHICopyTextureInfo CopyTextureInfo;
 			// We want to copy all mips : 
-			CopyTextureInfo.NumMips = InLocalMergeInfo.NumMips;
+			CopyTextureInfo.NumMips = TextureResolveInfo.NumMips;
+			// We need specify the size since the empty texture might be of higher size :
+			check(EmptyTexture->Desc.GetSize().X >= TextureResolveInfo.TextureSize.X);
+			CopyTextureInfo.Size = FIntVector(TextureResolveInfo.TextureSize.X, TextureResolveInfo.TextureSize.X, 0);
 
 			AddCopyTexturePass(GraphBuilder, EmptyTexture, DestinationTexture, CopyTextureInfo);
 		}
@@ -4342,10 +4350,15 @@ void ALandscape::PrepareLayersHeightmapsLocalMergeRenderThreadData(const FUpdate
 			ComponentToComponentRenderInfoIndex.Add(Component, OutRenderThreadData.ComponentToRenderInfos.Num() - 1);
 
 			UTexture2D* ComponentHeightmap = Component->GetHeightmap();
-			OutRenderThreadData.HeightmapSize = FIntPoint(ComponentHeightmap->Source.GetSizeX(), ComponentHeightmap->Source.GetSizeY());
+			int32 TextureSize = ComponentHeightmap->Source.GetSizeX();
+			if (ComponentHeightmap->Source.GetSizeX() > OutRenderThreadData.MaxHeightmapSize.X)
+			{
+				OutRenderThreadData.MaxHeightmapSize = FIntPoint(TextureSize, TextureSize);
+				OutRenderThreadData.MaxHeightmapNumMips = FMath::CeilLogTwo(TextureSize) + 1;
+			}
 
-			const int32 HeightmapOffsetX = Component->HeightmapScaleBias.Z * (float)OutRenderThreadData.HeightmapSize.X;
-			const int32 HeightmapOffsetY = Component->HeightmapScaleBias.W * (float)OutRenderThreadData.HeightmapSize.Y;
+			const int32 HeightmapOffsetX = Component->HeightmapScaleBias.Z * (float)TextureSize;
+			const int32 HeightmapOffsetY = Component->HeightmapScaleBias.W * (float)TextureSize;
 			// Effective area of the texture affecting this component (because of texture sharing) :
 			const FIntRect ComponentTextureSubregion(FIntPoint(HeightmapOffsetX, HeightmapOffsetY), FIntPoint(HeightmapOffsetX, HeightmapOffsetY) + OutRenderThreadData.ComponentSizeVerts);
 
@@ -4429,6 +4442,8 @@ void ALandscape::PrepareLayersHeightmapsLocalMergeRenderThreadData(const FUpdate
 				check((CPUReadback != nullptr) && (*CPUReadback != nullptr));
 
 				FTextureResolveInfo NewTextureResolveInfo;
+				NewTextureResolveInfo.TextureSize = FIntPoint(ComponentHeightmap->Source.GetSizeX(), ComponentHeightmap->Source.GetSizeX());
+				NewTextureResolveInfo.NumMips = FMath::CeilLogTwo(NewTextureResolveInfo.TextureSize.X) + 1;
 				NewTextureResolveInfo.Texture = NewComponentResolveInfo.Heightmap.Texture;
 				NewTextureResolveInfo.CPUReadback = *CPUReadback;
 				OutRenderThreadData.TextureToResolveInfos.Add(NewTextureResolveInfo);
@@ -4505,9 +4520,6 @@ void ALandscape::PrepareLayersHeightmapsLocalMergeRenderThreadData(const FUpdate
 
 		// We'll only ever need this amount of edit layers textures for any MergeEditLayers operation :
 		OutRenderThreadData.MaxNumEditLayersTexturesToMerge = OutRenderThreadData.VisibleEditLayerInfos.Num();
-
-		check(OutRenderThreadData.HeightmapSize != FIntPoint::ZeroValue);
-		OutRenderThreadData.NumMips = FMath::CeilLogTwo(OutRenderThreadData.HeightmapSize.X) + 1;
 	}
 }
 
