* Removed incomplete per-platform per-LOD level material overrides for landscape. The code isn't meant to support that since material instances are built at edit-time. We would need to somewhat hack it at cook time in order to regenerate (and recompile) the list of material instances per-platform, which is not a recommended thing to to in the engine. Instead, we choose to simply make the per-LOD level material overrides constant across platforms since the recommended way is to customize the material for platforms using material expressions such as Quality Level switc, Shading Path Switc or Feature Level Switc.
* Deprecated/merged FLandscapeComponentMaterialOverride and FLandscapeProxyMaterialOverride into FLandscapePerLODMaterialOverride since they were identical
* Fixed UpdateLayersMaterialInstances which, unlike UpdateMaterialInstances, didn't take into account per-component material overrides

From shelve 17931928 on behalf on Jonathan.Bard

diff --git a/Engine/Source/Editor/LandscapeEditor/Private/LandscapeEdMode.cpp b/Engine/Source/Editor/LandscapeEditor/Private/LandscapeEdMode.cpp
--- a/Engine/Source/Editor/LandscapeEditor/Private/LandscapeEdMode.cpp
+++ b/Engine/Source/Editor/LandscapeEditor/Private/LandscapeEdMode.cpp
@@ -3434,7 +3434,7 @@ ALandscape* FEdModeLandscape::ChangeComponentSetting(int32 NumComponentsX, int32
 			NewLandscape->SetActorRelativeScale3D(FVector(OldScale.X * LandscapeScaleFactor, OldScale.Y * LandscapeScaleFactor, OldScale.Z));
 
 			NewLandscape->LandscapeMaterial = OldLandscape->LandscapeMaterial;
-			NewLandscape->LandscapeMaterialsOverride = OldLandscape->LandscapeMaterialsOverride;
+			NewLandscape->SetPerLODOverrideMaterials(OldLandscape->GetPerLODOverrideMaterials());
 			NewLandscape->CollisionMipLevel = OldLandscape->CollisionMipLevel;
 			NewLandscape->MaxLODLevel = OldLandscape->MaxLODLevel;
 			NewLandscape->LODDistanceFactor_DEPRECATED = OldLandscape->LODDistanceFactor_DEPRECATED;
diff --git a/Engine/Source/Runtime/Landscape/Classes/LandscapeComponent.h b/Engine/Source/Runtime/Landscape/Classes/LandscapeComponent.h
--- a/Engine/Source/Runtime/Landscape/Classes/LandscapeComponent.h
+++ b/Engine/Source/Runtime/Landscape/Classes/LandscapeComponent.h
@@ -300,8 +300,8 @@ struct FLandscapeComponentGrassData
 	friend FArchive& operator<<(FArchive& Ar, FLandscapeComponentGrassData& Data);
 };
 
-USTRUCT(NotBlueprintable)
-struct FLandscapeComponentMaterialOverride
+USTRUCT(NotBlueprintable, meta = (Deprecated = "5.1"))
+struct UE_DEPRECATED(5.1, "FLandscapeComponentMaterialOverride is deprecated; please use FLandscapePerLODMaterialOverride instead") FLandscapeComponentMaterialOverride
 {
 	GENERATED_USTRUCT_BODY()
 
@@ -312,6 +312,24 @@ struct FLandscapeComponentMaterialOverride
 	TObjectPtr<UMaterialInterface> Material = nullptr;
 };
 
+USTRUCT(NotBlueprintable)
+struct FLandscapePerLODMaterialOverride
+{
+	GENERATED_USTRUCT_BODY()
+
+	UPROPERTY(EditAnywhere, Category = Material, meta = (UIMin = 0, UIMax = 8, ClampMin = 0, ClampMax = 8))
+	int32 LODIndex;
+
+	UPROPERTY(EditAnywhere, Category = Material)
+	TObjectPtr<UMaterialInterface> Material = nullptr;
+
+	bool operator == (const FLandscapePerLODMaterialOverride & InOther) const
+	{
+		return (LODIndex == InOther.LODIndex)
+			&& (Material == InOther.Material);
+	}
+};
+
 USTRUCT(NotBlueprintable)
 struct FWeightmapData
 {
@@ -480,10 +498,13 @@ class ULandscapeComponent : public UPrimitiveComponent
 	UPROPERTY(EditAnywhere, BlueprintReadWrite, Category=LandscapeComponent, AdvancedDisplay)
 	TObjectPtr<UMaterialInterface> OverrideHoleMaterial;
 
-	UPROPERTY(EditAnywhere, Category = LandscapeComponent)
-	TArray<FLandscapeComponentMaterialOverride> OverrideMaterials;
-
 #if WITH_EDITORONLY_DATA
+	PRAGMA_DISABLE_DEPRECATION_WARNINGS
+	UE_DEPRECATED(5.1, "OverrideMaterials has been deprecated, use PerLODOverrideMaterials instead.")
+	UPROPERTY()
+	TArray<FLandscapeComponentMaterialOverride> OverrideMaterials_DEPRECATED;
+	PRAGMA_ENABLE_DEPRECATION_WARNINGS
+
 	UPROPERTY()
 	TObjectPtr<UMaterialInstanceConstant> MaterialInstance_DEPRECATED;
 #endif
@@ -569,6 +590,9 @@ private:
 	UPROPERTY()
 	TObjectPtr<ULandscapeLODStreamingProxy> LODStreamingProxy;
 
+	UPROPERTY(EditAnywhere, Category = LandscapeComponent)
+	TArray<FLandscapePerLODMaterialOverride> PerLODOverrideMaterials;
+
 public:
 
 	/** Uniquely identifies this component's built map data. */
@@ -760,6 +784,9 @@ public:
 	LANDSCAPE_API TArray<FWeightmapLayerAllocationInfo>& GetWeightmapLayerAllocations(const FGuid& InLayerGuid);
 	LANDSCAPE_API const TArray<FWeightmapLayerAllocationInfo>& GetWeightmapLayerAllocations(const FGuid& InLayerGuid) const;
 
+	const TArray<FLandscapePerLODMaterialOverride>& GetPerLODOverrideMaterials() const { return PerLODOverrideMaterials; }
+	void SetPerLODOverrideMaterials(const TArray<FLandscapePerLODMaterialOverride>& InValue) { PerLODOverrideMaterials = InValue; }
+
 #if WITH_EDITOR
 	LANDSCAPE_API uint32 ComputeLayerHash(bool InReturnEditingHash = true) const;
 
diff --git a/Engine/Source/Runtime/Landscape/Classes/LandscapeProxy.h b/Engine/Source/Runtime/Landscape/Classes/LandscapeProxy.h
--- a/Engine/Source/Runtime/Landscape/Classes/LandscapeProxy.h
+++ b/Engine/Source/Runtime/Landscape/Classes/LandscapeProxy.h
@@ -320,8 +320,8 @@ public:
 	~FAsyncGrassTask();
 };
 
-USTRUCT()
-struct FLandscapeProxyMaterialOverride
+USTRUCT(meta = (Deprecated = "5.1"))
+struct UE_DEPRECATED(5.1, "FLandscapeProxyMaterialOverride is deprecated; please use FLandscapePerLODMaterialOverride instead") FLandscapeProxyMaterialOverride
 {
 	GENERATED_USTRUCT_BODY()
 
@@ -330,31 +330,6 @@ struct FLandscapeProxyMaterialOverride
 
 	UPROPERTY(EditAnywhere, Category = Landscape)
 	TObjectPtr<UMaterialInterface> Material = nullptr;
-
-#if WITH_EDITORONLY_DATA
-	bool operator==(const FLandscapeProxyMaterialOverride& InOther) const
-	{
-		if (Material != InOther.Material)
-		{
-			return false;
-		}
-
-		if (LODIndex.Default != InOther.LODIndex.Default || LODIndex.PerPlatform.Num() != InOther.LODIndex.PerPlatform.Num())
-		{
-			return false;
-		}
-
-		for (auto& ItPair : LODIndex.PerPlatform)
-		{
-			if (!InOther.LODIndex.PerPlatform.Contains(ItPair.Key))
-			{
-				return false;
-			}
-		}
-
-		return true;
-	}
-#endif
 };
 
 UCLASS(Abstract, MinimalAPI, NotBlueprintable, NotPlaceable, hidecategories=(Display, Attachment, Physics, Debug, Lighting, HLOD), showcategories=(Lighting, Rendering, Transformation), hidecategories=(Mobility))
@@ -375,6 +350,14 @@ protected:
 	UPROPERTY()
 	FGuid LandscapeGuid;
 
+	UPROPERTY(EditAnywhere, Category = Landscape)
+	TArray<FLandscapePerLODMaterialOverride> PerLODOverrideMaterials;
+
+#if WITH_EDITORONLY_DATA
+	UPROPERTY(Transient)
+	TArray<FLandscapePerLODMaterialOverride> PreEditPerLODOverrideMaterials;
+#endif // WITH_EDITORONLY_DATA
+
 public:
 	LANDSCAPE_API TOptional<float> GetHeightAtLocation(FVector Location) const;
 
@@ -389,11 +372,13 @@ public:
 	UPROPERTY(EditAnywhere, Category=LOD)
 	int32 MaxLODLevel;
 
+#if WITH_EDITORONLY_DATA
 	UPROPERTY()
 	float LODDistanceFactor_DEPRECATED;
 
 	UPROPERTY()
 	TEnumAsByte<ELandscapeLODFalloff::Type> LODFalloff_DEPRECATED;
+#endif // WITH_EDITORONLY_DATA
 
 	/** Component screen size (0.0 - 1.0) at which we should keep sub sections. This is mostly pertinent if you have large component of > 64 and component are close to the camera. The goal is to reduce draw call, so if a component is smaller than the value, we merge all subsections into 1 drawcall. */
 	UPROPERTY(EditAnywhere, Category = LOD, meta=(ClampMin = "0.01", ClampMax = "1.0", UIMin = "0.01", UIMax = "1.0", DisplayName= "SubSection Min Component ScreenSize"))
@@ -460,19 +445,21 @@ public:
 	UPROPERTY(EditAnywhere, Category=Landscape, AdvancedDisplay)
 	TObjectPtr<UMaterialInterface> LandscapeHoleMaterial;
 
-	UPROPERTY(EditAnywhere, Category = Landscape)
-	TArray<FLandscapeProxyMaterialOverride> LandscapeMaterialsOverride;
 
 #if WITH_EDITORONLY_DATA
+
+	PRAGMA_DISABLE_DEPRECATION_WARNINGS
+	UE_DEPRECATED(5.1, "LandscapeComponentMaterialOverride has been deprecated, use PerLODOverrideMaterials instead.")
+	UPROPERTY()
+	TArray<FLandscapeProxyMaterialOverride> LandscapeMaterialsOverride_DEPRECATED;
+	PRAGMA_ENABLE_DEPRECATION_WARNINGS
+
 	UPROPERTY(Transient)
 	TObjectPtr<UMaterialInterface> PreEditLandscapeMaterial;
 
 	UPROPERTY(Transient)
 	TObjectPtr<UMaterialInterface> PreEditLandscapeHoleMaterial;
 
-	UPROPERTY(Transient)
-	TArray<FLandscapeProxyMaterialOverride> PreEditLandscapeMaterialsOverride;
-
 	UPROPERTY(Transient)
 	bool bIsPerformingInteractiveActionOnLandscapeMaterialOverride;
 #endif 
@@ -728,8 +715,6 @@ public:
 	UPROPERTY()
 	bool bHasLayersContent;
 
-public:
-
 #if WITH_EDITOR
 	LANDSCAPE_API static ULandscapeLayerInfoObject* VisibilityLayer;
 #endif
@@ -816,6 +801,9 @@ public:
 	virtual ALandscape* GetLandscapeActor() PURE_VIRTUAL(GetLandscapeActor, return nullptr;)
 	virtual const ALandscape* GetLandscapeActor() const PURE_VIRTUAL(GetLandscapeActor, return nullptr;)
 
+	const TArray<FLandscapePerLODMaterialOverride>& GetPerLODOverrideMaterials() const { return PerLODOverrideMaterials; }
+	void SetPerLODOverrideMaterials(const TArray<FLandscapePerLODMaterialOverride>& InValue) { PerLODOverrideMaterials = InValue; }
+
 	static void SetGrassUpdateInterval(int32 Interval) { GrassUpdateInterval = Interval; }
 
 	/* Per-frame call to update dynamic grass placement and render grassmaps */
@@ -851,9 +839,9 @@ public:
 		* @param bForceSync if true, block and finish all work
 	*/
 	LANDSCAPE_API void UpdateGrass(const TArray<FVector>& Cameras, int32& InOutNumComponentsCreated, bool bForceSync = false);
-
 	LANDSCAPE_API void UpdateGrass(const TArray<FVector>& Cameras, bool bForceSync = false);
 
+	// TODO [jonathan.bard] : Rename to "AddGrassExlusionBox" + no reason for any of this to be static
 	LANDSCAPE_API static void AddExclusionBox(FWeakObjectPtr Owner, const FBox& BoxToRemove);
 	LANDSCAPE_API static void RemoveExclusionBox(FWeakObjectPtr Owner);
 	LANDSCAPE_API static void RemoveAllExclusionBoxes();
@@ -941,6 +929,9 @@ public:
 	LANDSCAPE_API TArray<float> GetLODScreenSizeArray() const;
 
 #if WITH_EDITOR
+	/* Serialize all hashes/guids that record the current state of this proxy */
+	void SerializeStateHashes(FArchive& Ar);
+
 	LANDSCAPE_API void SetSplinesComponent(ULandscapeSplinesComponent* InSplineComponent) { check(!SplineComponent); SplineComponent = InSplineComponent; }
 
 	LANDSCAPE_API virtual bool SupportsForeignSplineMesh() const override { return true; }
diff --git a/Engine/Source/Runtime/Landscape/Private/Landscape.cpp b/Engine/Source/Runtime/Landscape/Private/Landscape.cpp
--- a/Engine/Source/Runtime/Landscape/Private/Landscape.cpp
+++ b/Engine/Source/Runtime/Landscape/Private/Landscape.cpp
@@ -594,17 +594,10 @@ UMaterialInterface* ULandscapeComponent::GetLandscapeMaterial(int8 InLODIndex) c
 
 		if (World != nullptr)
 		{
-			for (const FLandscapeComponentMaterialOverride& Material : OverrideMaterials)
+			if (const FLandscapePerLODMaterialOverride* LocalMaterialOverride = PerLODOverrideMaterials.FindByPredicate(
+				[InLODIndex](const FLandscapePerLODMaterialOverride& InOverride) { return (InOverride.LODIndex == InLODIndex) && (InOverride.Material != nullptr); }))
 			{
-				if (Material.LODIndex.GetValue() == InLODIndex)
-				{
-					if (Material.Material != nullptr)
-					{
-						return Material.Material;
-					}
-
-					break;
-				}
+				return LocalMaterialOverride->Material;
 			}
 		}
 	}
@@ -823,8 +816,20 @@ void ULandscapeComponent::PostLoad()
 		}
 		LastBakedTextureMaterialGuid = BakedTextureMaterialGuid;
 		LastSavedPhysicalMaterialHash = PhysicalMaterialHash;
+
+		PRAGMA_DISABLE_DEPRECATION_WARNINGS;
+		if (!OverrideMaterials_DEPRECATED.IsEmpty())
+		{
+			PerLODOverrideMaterials.Reserve(OverrideMaterials_DEPRECATED.Num());
+			for (const FLandscapeComponentMaterialOverride& LocalMaterialOverride : OverrideMaterials_DEPRECATED)
+			{
+				PerLODOverrideMaterials.Add({ LocalMaterialOverride.LODIndex.Default, LocalMaterialOverride.Material });
+			}
+			OverrideMaterials_DEPRECATED.Reset();
+		}
+		PRAGMA_ENABLE_DEPRECATION_WARNINGS;
 	}
-#endif
+#endif // WITH_EDITOR
 
 #if WITH_EDITORONLY_DATA
 	// Handle old MaterialInstance
@@ -2240,6 +2245,7 @@ void ALandscapeProxy::Serialize(FArchive& Ar)
 	Ar.UsingCustomVersion(FLandscapeCustomVersion::GUID);
 	Ar.UsingCustomVersion(FEditorObjectVersion::GUID);
 
+#if WITH_EDITORONLY_DATA
 	if (Ar.IsLoading() && Ar.CustomVer(FLandscapeCustomVersion::GUID) < FLandscapeCustomVersion::MigrateOldPropertiesToNewRenderingProperties)
 	{
 		if (LODDistanceFactor_DEPRECATED > 0)
@@ -2261,6 +2267,7 @@ void ALandscapeProxy::Serialize(FArchive& Ar)
 			}
 		}
 	}
+#endif // WITH_EDITORONLY_DATA
 }
 
 void ALandscapeProxy::AddReferencedObjects(UObject* InThis, FReferenceCollector& Collector)
@@ -2583,6 +2590,18 @@ void ALandscapeProxy::PostLoad()
 		CreateLandscapeInfo();
 	}
 #if WITH_EDITOR
+	PRAGMA_DISABLE_DEPRECATION_WARNINGS;
+	if (!LandscapeMaterialsOverride_DEPRECATED.IsEmpty())
+	{
+		PerLODOverrideMaterials.Reserve(LandscapeMaterialsOverride_DEPRECATED.Num());
+		for (const FLandscapeProxyMaterialOverride& LocalMaterialOverride : LandscapeMaterialsOverride_DEPRECATED)
+		{
+			PerLODOverrideMaterials.Add({ LocalMaterialOverride.LODIndex.Default, LocalMaterialOverride.Material });
+		}
+		LandscapeMaterialsOverride_DEPRECATED.Reset();
+	}
+	PRAGMA_ENABLE_DEPRECATION_WARNINGS;
+
 	if (GIsEditor && GetWorld() && !GetWorld()->IsGameWorld())
 	{
 		if ((GetLinker() && (GetLinker()->UEVer() < VER_UE4_LANDSCAPE_COMPONENT_LAZY_REFERENCES)) ||
@@ -2741,7 +2760,7 @@ void ALandscapeProxy::GetSharedProperties(ALandscapeProxy* Landscape)
 		if (!LandscapeMaterial)
 		{
 			LandscapeMaterial = Landscape->LandscapeMaterial;
-			LandscapeMaterialsOverride = Landscape->LandscapeMaterialsOverride;
+			PerLODOverrideMaterials = Landscape->PerLODOverrideMaterials;
 		}
 		if (!LandscapeHoleMaterial)
 		{
@@ -2893,17 +2912,10 @@ UMaterialInterface* ALandscapeProxy::GetLandscapeMaterial(int8 InLODIndex) const
 
 		if (World != nullptr)
 		{
-			for (const FLandscapeProxyMaterialOverride& OverrideMaterial : LandscapeMaterialsOverride)
+			if (const FLandscapePerLODMaterialOverride* LocalMaterialOverride = PerLODOverrideMaterials.FindByPredicate(
+				[InLODIndex](const FLandscapePerLODMaterialOverride& InOverride) { return (InOverride.LODIndex == InLODIndex) && (InOverride.Material != nullptr); }))
 			{
-				if (OverrideMaterial.LODIndex.GetValue() == InLODIndex)
-				{
-					if (OverrideMaterial.Material != nullptr)
-					{
-						return OverrideMaterial.Material;
-					}
-
-					break;
-				}
+				return LocalMaterialOverride->Material;
 			}
 		}
 	}
@@ -2924,17 +2936,10 @@ UMaterialInterface* ALandscapeStreamingProxy::GetLandscapeMaterial(int8 InLODInd
 
 		if (World != nullptr)
 		{
-			for (const FLandscapeProxyMaterialOverride& OverrideMaterial : LandscapeMaterialsOverride)
+			if (const FLandscapePerLODMaterialOverride* LocalMaterialOverride = PerLODOverrideMaterials.FindByPredicate(
+				[InLODIndex](const FLandscapePerLODMaterialOverride& InOverride) { return (InOverride.LODIndex == InLODIndex) && (InOverride.Material != nullptr); }))
 			{
-				if (OverrideMaterial.LODIndex.GetValue() == InLODIndex)
-				{
-					if (OverrideMaterial.Material != nullptr)
-					{
-						return OverrideMaterial.Material;
-					}
-
-					break;
-				}
+				return LocalMaterialOverride->Material;
 			}
 		}
 	}
@@ -3841,6 +3846,7 @@ void FLandscapeComponentDerivedData::SaveToDDC(const FGuid& StateId, UObject* Co
 	GetDerivedDataCacheRef().Put(*GetDDCKeyString(StateId), Bytes, Component->GetPathName());
 }
 
+PRAGMA_DISABLE_DEPRECATION_WARNINGS
 ALandscapeProxy::~ALandscapeProxy()
 {
 	for (int32 Index = 0; Index < AsyncFoliageTasks.Num(); Index++)
@@ -3863,6 +3869,7 @@ ALandscapeProxy::~ALandscapeProxy()
 	LandscapeProxies.Remove(this);
 #endif
 }
+PRAGMA_ENABLE_DEPRECATION_WARNINGS
 
 //
 // ALandscapeMeshProxyActor
@@ -3907,6 +3914,19 @@ void ULandscapeMeshProxyComponent::InitializeForLandscape(ALandscapeProxy* Lands
 }
 
 #if WITH_EDITOR
+void ALandscapeProxy::SerializeStateHashes(FArchive& Ar)
+{
+	for (FLandscapePerLODMaterialOverride& MaterialOverride : PerLODOverrideMaterials)
+	{
+		if (MaterialOverride.Material != nullptr)
+		{
+			FGuid LocalStateId = MaterialOverride.Material->GetMaterial_Concurrent()->StateId;
+			Ar << LocalStateId;
+			Ar << MaterialOverride.LODIndex;
+		}
+	}
+}
+
 void ULandscapeComponent::SerializeStateHashes(FArchive& Ar)
 {
 	FGuid HeightmapGuid = HeightmapTexture->Source.GetId();
@@ -3931,7 +3951,7 @@ void ULandscapeComponent::SerializeStateHashes(FArchive& Ar)
 		Ar << LocalStateId;
 	}
 
-	for (FLandscapeComponentMaterialOverride& MaterialOverride : OverrideMaterials)
+	for (FLandscapePerLODMaterialOverride& MaterialOverride : PerLODOverrideMaterials)
 	{
 		if (MaterialOverride.Material != nullptr)
 		{
@@ -3949,15 +3969,7 @@ void ULandscapeComponent::SerializeStateHashes(FArchive& Ar)
 		Ar << LocalStateId;
 	}
 
-	for (FLandscapeProxyMaterialOverride& MaterialOverride : Proxy->LandscapeMaterialsOverride)
-	{
-		if (MaterialOverride.Material != nullptr)
-		{
-			FGuid LocalStateId = MaterialOverride.Material->GetMaterial_Concurrent()->StateId;
-			Ar << LocalStateId;
-			Ar << MaterialOverride.LODIndex;
-		}
-	}
+	Proxy->SerializeStateHashes(Ar);
 }
 
 FLandscapeGIBakedTextureBuilder::FLandscapeGIBakedTextureBuilder(UWorld* InWorld)
diff --git a/Engine/Source/Runtime/Landscape/Private/LandscapeConfigHelper.cpp b/Engine/Source/Runtime/Landscape/Private/LandscapeConfigHelper.cpp
--- a/Engine/Source/Runtime/Landscape/Private/LandscapeConfigHelper.cpp
+++ b/Engine/Source/Runtime/Landscape/Private/LandscapeConfigHelper.cpp
@@ -193,7 +193,8 @@ bool FLandscapeConfigHelper::ChangeGridSize(ULandscapeInfo* InLandscapeInfo, uin
 					}
 				}
 
-				MovedComponent->OverrideMaterials.Reset();
+				TArray<FLandscapePerLODMaterialOverride> PerLODOverrideMaterialsForComponent;
+				TArray<FLandscapePerLODMaterialOverride> PerLODOverrideMaterialsForProxy = LandscapeProxy->GetPerLODOverrideMaterials();
 				for (int32 LODIndex = 0; LODIndex <= 8; ++LODIndex)
 				{
 					UMaterialInterface* PreviousLODMaterial = PreviousLandscapeLODMaterials.FindChecked(LODIndex);
@@ -202,14 +203,16 @@ bool FLandscapeConfigHelper::ChangeGridSize(ULandscapeInfo* InLandscapeInfo, uin
 					{
 						if (LandscapeProxy->GetLandscapeMaterial(LODIndex) == LandscapeProxy->GetLandscapeActor()->GetLandscapeMaterial(LODIndex))
 						{
-							LandscapeProxy->LandscapeMaterialsOverride.Add({ LODIndex, decltype(FLandscapeProxyMaterialOverride::Material)(PreviousLODMaterial) });
+							PerLODOverrideMaterialsForProxy.Add({ LODIndex, TObjectPtr<UMaterialInterface>(PreviousLODMaterial) });
 						}
 						else // If it already differs it means that the component differs from it, override on component
 						{
-							MovedComponent->OverrideMaterials.Add({ LODIndex, decltype(FLandscapeComponentMaterialOverride::Material)(PreviousLODMaterial) });
+							PerLODOverrideMaterialsForComponent.Add({ LODIndex, TObjectPtr<UMaterialInterface>(PreviousLODMaterial) });
 						}
 					}
 				}
+				MovedComponent->SetPerLODOverrideMaterials(PerLODOverrideMaterialsForComponent);
+				LandscapeProxy->SetPerLODOverrideMaterials(PerLODOverrideMaterialsForProxy);
 			}
 		}
 
diff --git a/Engine/Source/Runtime/Landscape/Private/LandscapeEdit.cpp b/Engine/Source/Runtime/Landscape/Private/LandscapeEdit.cpp
--- a/Engine/Source/Runtime/Landscape/Private/LandscapeEdit.cpp
+++ b/Engine/Source/Runtime/Landscape/Private/LandscapeEdit.cpp
@@ -710,11 +710,11 @@ bool ALandscapeProxy::GetReferencedContentObjects(TArray<UObject*>& Objects) con
 		Objects.AddUnique(LandscapeMaterial);
 	}
 
-	for (const FLandscapeProxyMaterialOverride& OverrideMaterial : LandscapeMaterialsOverride)
+	for (const FLandscapePerLODMaterialOverride& LODOverrideMaterial : PerLODOverrideMaterials)
 	{
-		if (OverrideMaterial.Material != nullptr)
+		if (LODOverrideMaterial.Material != nullptr)
 		{
-			Objects.AddUnique(OverrideMaterial.Material);
+			Objects.AddUnique(LODOverrideMaterial.Material);
 		}
 	}
 
@@ -5082,11 +5082,11 @@ void ALandscapeStreamingProxy::PostEditChangeProperty(FPropertyChangedEvent& Pro
 			LandscapeActor = nullptr;
 		}
 	}
-	else if (PropertyName == FName(TEXT("LandscapeMaterial")) || PropertyName == FName(TEXT("LandscapeHoleMaterial")) || PropertyName == FName(TEXT("LandscapeMaterialsOverride")))
+	else if (PropertyName == FName(TEXT("LandscapeMaterial")) || PropertyName == FName(TEXT("LandscapeHoleMaterial")) || PropertyName == FName(TEXT("PerLODOverrideMaterials")))
 	{
 		bool RecreateMaterialInstances = true;
 
-		if (PropertyName == FName(TEXT("LandscapeMaterialsOverride")) && PropertyChangedEvent.ChangeType == EPropertyChangeType::ArrayAdd)
+		if (PropertyName == FName(TEXT("PerLODOverrideMaterials")) && PropertyChangedEvent.ChangeType == EPropertyChangeType::ArrayAdd)
 		{
 			RecreateMaterialInstances = false;
 		}
@@ -5312,7 +5312,7 @@ void ALandscape::PreEditChange(FProperty* PropertyThatWillChange)
 {
 	PreEditLandscapeMaterial = LandscapeMaterial;
 	PreEditLandscapeHoleMaterial = LandscapeHoleMaterial;
-	PreEditLandscapeMaterialsOverride = LandscapeMaterialsOverride;
+	PreEditPerLODOverrideMaterials = PerLODOverrideMaterials;
 
 	Super::PreEditChange(PropertyThatWillChange);
 }
@@ -5329,24 +5329,24 @@ void ALandscape::PostEditChangeProperty(FPropertyChangedEvent& PropertyChangedEv
 
 	ULandscapeInfo* Info = GetLandscapeInfo();
 
-	if ((PropertyName == GET_MEMBER_NAME_CHECKED(ALandscapeProxy, LandscapeMaterial) || PropertyName == GET_MEMBER_NAME_CHECKED(ALandscapeProxy, LandscapeHoleMaterial) || MemberPropertyName == GET_MEMBER_NAME_CHECKED(ALandscapeProxy, LandscapeMaterialsOverride))
+	if ((PropertyName == GET_MEMBER_NAME_CHECKED(ALandscapeProxy, LandscapeMaterial) || PropertyName == GET_MEMBER_NAME_CHECKED(ALandscapeProxy, LandscapeHoleMaterial) || MemberPropertyName == GET_MEMBER_NAME_CHECKED(ALandscapeProxy, PerLODOverrideMaterials))
 		&& PropertyChangedEvent.ChangeType != EPropertyChangeType::ArrayAdd)
 	{
 		bool HasMaterialChanged = false;
 
 		if (PropertyChangedEvent.ChangeType != EPropertyChangeType::Interactive)
 		{
-			if (PreEditLandscapeMaterial != LandscapeMaterial || PreEditLandscapeHoleMaterial != LandscapeHoleMaterial || PreEditLandscapeMaterialsOverride.Num() != LandscapeMaterialsOverride.Num() || bIsPerformingInteractiveActionOnLandscapeMaterialOverride)
+			if (PreEditLandscapeMaterial != LandscapeMaterial || PreEditLandscapeHoleMaterial != LandscapeHoleMaterial || PreEditPerLODOverrideMaterials.Num() != PerLODOverrideMaterials.Num() || bIsPerformingInteractiveActionOnLandscapeMaterialOverride)
 			{
 				HasMaterialChanged = true;
 			}
 
 			if (!HasMaterialChanged)
 			{
-				for (int32 i = 0; i < LandscapeMaterialsOverride.Num(); ++i)
+				for (int32 i = 0; i < PerLODOverrideMaterials.Num(); ++i)
 				{
-					const FLandscapeProxyMaterialOverride& NewMaterialOverride = LandscapeMaterialsOverride[i];
-					const FLandscapeProxyMaterialOverride& PreEditMaterialOverride = PreEditLandscapeMaterialsOverride[i];
+					const FLandscapePerLODMaterialOverride& NewMaterialOverride = PerLODOverrideMaterials[i];
+					const FLandscapePerLODMaterialOverride& PreEditMaterialOverride = PreEditPerLODOverrideMaterials[i];
 
 					if (!(PreEditMaterialOverride == NewMaterialOverride))
 					{
@@ -5360,8 +5360,8 @@ void ALandscape::PostEditChangeProperty(FPropertyChangedEvent& PropertyChangedEv
 		}
 		else
 		{
-			// We are probably using a slider or something similar in LandscapeMaterialsOverride
-			bIsPerformingInteractiveActionOnLandscapeMaterialOverride = MemberPropertyName == GET_MEMBER_NAME_CHECKED(ALandscapeProxy, LandscapeMaterialsOverride);
+			// We are probably using a slider or something similar in PerLODOverrideMaterials
+			bIsPerformingInteractiveActionOnLandscapeMaterialOverride = MemberPropertyName == GET_MEMBER_NAME_CHECKED(ALandscapeProxy, PerLODOverrideMaterials);
 		}
 
 		if (Info != nullptr && HasMaterialChanged)
@@ -5398,32 +5398,32 @@ void ALandscape::PostEditChangeProperty(FPropertyChangedEvent& PropertyChangedEv
 		// Some edit layers could be affected by BP brushes, which might need to be updated when the landscape is transformed :
 		RequestLayersContentUpdate(ELandscapeLayerUpdateMode::Update_All);
 	}
-	else if (GIsEditor && PropertyName == FName(TEXT("MaxLODLevel")))
+	else if (GIsEditor && PropertyName == GET_MEMBER_NAME_CHECKED(ALandscapeProxy, MaxLODLevel))
 	{
 		MaxLODLevel = FMath::Clamp<int32>(MaxLODLevel, -1, FMath::CeilLogTwo(SubsectionSizeQuads + 1) - 1);
 		bPropagateToProxies = true;
 	}
-	else if (PropertyName == FName(TEXT("ComponentScreenSizeToUseSubSections")))
+	else if (PropertyName == GET_MEMBER_NAME_CHECKED(ALandscapeProxy, ComponentScreenSizeToUseSubSections))
 	{
 		ComponentScreenSizeToUseSubSections = FMath::Clamp<float>(ComponentScreenSizeToUseSubSections, 0.01f, 1.0f);
 		bPropagateToProxies = true;
 	}
-	else if (PropertyName == FName(TEXT("LODDistributionSetting")))
+	else if (PropertyName == GET_MEMBER_NAME_CHECKED(ALandscapeProxy, LODDistributionSetting))
 	{
 		LODDistributionSetting = FMath::Clamp<float>(LODDistributionSetting, 1.0f, 10.0f);
 		bPropagateToProxies = true;
 	}
-	else if (PropertyName == FName(TEXT("LOD0DistributionSetting")))
+	else if (PropertyName == GET_MEMBER_NAME_CHECKED(ALandscapeProxy, LOD0DistributionSetting))
 	{
 		LOD0DistributionSetting = FMath::Clamp<float>(LOD0DistributionSetting, 1.0f, 10.0f);
 		bPropagateToProxies = true;
 	}
-	else if (PropertyName == FName(TEXT("LOD0ScreenSize")))
+	else if (PropertyName == GET_MEMBER_NAME_CHECKED(ALandscapeProxy, LOD0ScreenSize))
 	{
 		LOD0ScreenSize = FMath::Clamp<float>(LOD0ScreenSize, 0.1f, 10.0f);
 		bPropagateToProxies = true;
 	}
-	else if (PropertyName == FName(TEXT("CollisionMipLevel")))
+	else if (PropertyName == GET_MEMBER_NAME_CHECKED(ALandscapeProxy, CollisionMipLevel))
 	{
 		CollisionMipLevel = FMath::Clamp<int32>(CollisionMipLevel, 0, FMath::CeilLogTwo(SubsectionSizeQuads + 1) - 1);
 		bPropagateToProxies = true;
@@ -5449,7 +5449,7 @@ void ALandscape::PostEditChangeProperty(FPropertyChangedEvent& PropertyChangedEv
 	{
 		bPropagateToProxies = true;
 	}
-	else if (GIsEditor && PropertyName == FName(TEXT("StaticLightingResolution")))
+	else if (GIsEditor && PropertyName == GET_MEMBER_NAME_CHECKED(ALandscapeProxy, StaticLightingResolution))
 	{
 		StaticLightingResolution = ::AdjustStaticLightingResolution(StaticLightingResolution, NumSubsections, SubsectionSizeQuads, ComponentSizeQuads);
 		bChangedLighting = true;
@@ -5549,7 +5549,7 @@ void ALandscape::PostEditChangeProperty(FPropertyChangedEvent& PropertyChangedEv
 
 	PreEditLandscapeMaterial = nullptr;
 	PreEditLandscapeHoleMaterial = nullptr;
-	PreEditLandscapeMaterialsOverride.Empty();
+	PreEditPerLODOverrideMaterials.Empty();
 }
 
 void ALandscapeProxy::ChangedPhysMaterial()
@@ -5647,11 +5647,11 @@ void ULandscapeComponent::PostEditChangeProperty(FPropertyChangedEvent& Property
 	const FName PropertyName = PropertyChangedEvent.Property ? PropertyChangedEvent.Property->GetFName() : NAME_None;
 	const FName MemberPropertyName = PropertyChangedEvent.MemberProperty ? PropertyChangedEvent.MemberProperty->GetFName() : NAME_None;
 
-	if (PropertyName == FName(TEXT("OverrideMaterial")) || MemberPropertyName == FName(TEXT("OverrideMaterials")) || MemberPropertyName == FName(TEXT("MaterialPerLOD_Key")))
+	if (PropertyName == FName(TEXT("OverrideMaterial")) || MemberPropertyName == FName(TEXT("PerLODOverrideMaterials")) || MemberPropertyName == FName(TEXT("MaterialPerLOD_Key")))
 	{
 		bool RecreateMaterialInstances = true;
 
-		if (PropertyName == FName(TEXT("OverrideMaterials")) && PropertyChangedEvent.ChangeType == EPropertyChangeType::ArrayAdd)
+		if (PropertyName == FName(TEXT("PerLODOverrideMaterials")) && PropertyChangedEvent.ChangeType == EPropertyChangeType::ArrayAdd)
 		{
 			RecreateMaterialInstances = false;
 		}
diff --git a/Engine/Source/Runtime/Landscape/Private/LandscapeEditLayers.cpp b/Engine/Source/Runtime/Landscape/Private/LandscapeEditLayers.cpp
--- a/Engine/Source/Runtime/Landscape/Private/LandscapeEditLayers.cpp
+++ b/Engine/Source/Runtime/Landscape/Private/LandscapeEditLayers.cpp
@@ -7118,12 +7118,12 @@ void ALandscape::UpdateLayersMaterialInstances(const TArray<ULandscapeComponent*
 		Component->LODIndexToMaterialIndex.SetNumUninitialized(MaxLOD + 1);
 		int8 LastLODIndex = INDEX_NONE;
 
-		UMaterialInterface* BaseMaterial = GetLandscapeMaterial();
-		UMaterialInterface* LOD0Material = GetLandscapeMaterial(0);
+		UMaterialInterface* BaseMaterial = Component->GetLandscapeMaterial();
+		UMaterialInterface* LOD0Material = Component->GetLandscapeMaterial(0);
 
 		for (int32 LODIndex = 0; LODIndex <= MaxLOD; ++LODIndex)
 		{
-			UMaterialInterface* CurrentMaterial = GetLandscapeMaterial(LODIndex);
+			UMaterialInterface* CurrentMaterial = Component->GetLandscapeMaterial(LODIndex);
 
 			// if we have a LOD0 override, do not let the base material override it, it should override everything!
 			if (CurrentMaterial == BaseMaterial && BaseMaterial != LOD0Material)
