Fix for crash with landscapes when using -rhivalidation.  Ensures that the RHIBuffer is held for the life of the frame because not all RHI backends will have the SRV maintain a hard reference to the source buffer

diff --git a/Engine/Source/Runtime/Landscape/Private/LandscapeRender.cpp b/Engine/Source/Runtime/Landscape/Private/LandscapeRender.cpp
--- a/Engine/Source/Runtime/Landscape/Private/LandscapeRender.cpp
+++ b/Engine/Source/Runtime/Landscape/Private/LandscapeRender.cpp
@@ -1091,11 +1091,11 @@ void FLandscapeSceneViewExtension::PreRenderView_RenderThread(FRHICommandListImm
 		}
 
 		FRHIResourceCreateInfo CreateInfoLODBuffer(TEXT("LandscapeLODDataBuffer"), &LandscapeLODData);
-		FBufferRHIRef LandscapeLODDataBuffer = RHICreateVertexBuffer(LandscapeLODData.GetResourceDataSize(), BUF_ShaderResource | BUF_Volatile, CreateInfoLODBuffer);
+		LandscapeLODDataBuffer = RHICreateVertexBuffer(LandscapeLODData.GetResourceDataSize(), BUF_ShaderResource | BUF_Volatile, CreateInfoLODBuffer);
 		LODDataSRV = RHICreateShaderResourceView(LandscapeLODDataBuffer, sizeof(float), PF_R32_FLOAT);
 
 		FRHIResourceCreateInfo CreateInfoIndirection(TEXT("LandscapeIndirectionBuffer"), &LandscapeIndirection);
-		FBufferRHIRef LandscapeIndirectionBuffer = RHICreateVertexBuffer(LandscapeIndirection.GetResourceDataSize(), BUF_ShaderResource | BUF_Volatile, CreateInfoIndirection);
+		LandscapeIndirectionBuffer = RHICreateVertexBuffer(LandscapeIndirection.GetResourceDataSize(), BUF_ShaderResource | BUF_Volatile, CreateInfoIndirection);
 		IndirectionSRV = RHICreateShaderResourceView(LandscapeIndirectionBuffer, sizeof(uint32), PF_R32_UINT);
 	}
 	else
@@ -1116,6 +1116,9 @@ void FLandscapeSceneViewExtension::EndFrame_RenderThread()
 
 		RenderSystem.EndFrame();
 	}
+
+	LandscapeLODDataBuffer = FBufferRHIRef();
+	LandscapeIndirectionBuffer = FBufferRHIRef();
 }
 
 //
diff --git a/Engine/Source/Runtime/Landscape/Public/LandscapeRender.h b/Engine/Source/Runtime/Landscape/Public/LandscapeRender.h
--- a/Engine/Source/Runtime/Landscape/Public/LandscapeRender.h
+++ b/Engine/Source/Runtime/Landscape/Public/LandscapeRender.h
@@ -497,6 +497,10 @@ public:
 
 	virtual void PreRenderViewFamily_RenderThread(FRHICommandListImmediate& RHICmdList, FSceneViewFamily& InViewFamily) override;
 	virtual void PreRenderView_RenderThread(FRHICommandListImmediate& RHICmdList, FSceneView& InView) override;
+
+private:
+	FBufferRHIRef LandscapeLODDataBuffer;
+	FBufferRHIRef LandscapeIndirectionBuffer;
 };
 
 
