Fixed crash when using invalid scale factor on landscape copy/paste gizmo

diff --git a/Engine/Source/Runtime/Landscape/Private/LandscapeGizmoActor.cpp b/Engine/Source/Runtime/Landscape/Private/LandscapeGizmoActor.cpp
--- a/Engine/Source/Runtime/Landscape/Private/LandscapeGizmoActor.cpp
+++ b/Engine/Source/Runtime/Landscape/Private/LandscapeGizmoActor.cpp
@@ -104,8 +104,8 @@ public:
 		case EMaterialParameterType::Texture:
 			if (ParameterInfo.Name == FName(TEXT("AlphaTexture")))
 			{
-				// FIXME: This needs to return a black texture if AlphaTexture is NULL.
-				// Returning NULL will cause the material to use GWhiteTexture.
+				// FIXME: This needs to return a black texture if AlphaTexture is nullptr.
+				// Returning nullptr will cause the material to use GWhiteTexture.
 				OutValue = AlphaTexture;
 				return true;
 			}
@@ -128,26 +128,27 @@ public:
 	}
 
 	FVector XAxis, YAxis, Origin;
-	float SampleSizeX, SampleSizeY;
-	bool bHeightmapRendering;
-	HHitProxy* HitProxy;
+	float SampleSizeX = 0.0f;
+	float SampleSizeY = 0.0f;
+	bool bHeightmapRendering = false;
+	bool bIsValid = false;
+	HHitProxy* HitProxy = nullptr;
 	FMatrix MeshRT;
 	FVector FrustumVerts[8];
 	TArray<FVector> SampledPositions;
 	TArray<FVector> SampledNormals;
-	FLandscapeGizmoMeshRenderProxy* HeightmapRenderProxy;
-	FMaterialRenderProxy* GizmoRenderProxy;
-
-	FLandscapeGizmoRenderSceneProxy(const ULandscapeGizmoRenderComponent* InComponent):
-		FPrimitiveSceneProxy(InComponent),
-		bHeightmapRendering(false),
-		HitProxy(nullptr),
-		HeightmapRenderProxy(nullptr),
-		GizmoRenderProxy(nullptr)
+	FLandscapeGizmoMeshRenderProxy* HeightmapRenderProxy = nullptr;
+	FMaterialRenderProxy* GizmoRenderProxy = nullptr;
+
+	FLandscapeGizmoRenderSceneProxy(const ULandscapeGizmoRenderComponent* InComponent)
+		: FPrimitiveSceneProxy(InComponent)
 	{
 #if WITH_EDITOR	
 		ALandscapeGizmoActiveActor* Gizmo = Cast<ALandscapeGizmoActiveActor>(InComponent->GetOwner());
-		if (Gizmo && Gizmo->GizmoMeshMaterial && Gizmo->GizmoDataMaterial && Gizmo->GetRootComponent())
+		if (Gizmo && Gizmo->GizmoMeshMaterial && Gizmo->GizmoDataMaterial && Gizmo->GetRootComponent() 
+			&& !FMath::IsNearlyZero(Gizmo->CachedWidth) 
+			&& !FMath::IsNearlyZero(Gizmo->CachedHeight) 
+			&& !FMath::IsNearlyZero(ScaleXY))
 		{
 			ULandscapeInfo* LandscapeInfo = Gizmo->TargetLandscapeInfo;
 			if (LandscapeInfo && LandscapeInfo->GetLandscapeProxy())
@@ -182,37 +183,39 @@ public:
 				const FMatrix WToL = LToW.ToMatrixWithScale().InverseFast();
 				const FVector BaseLocation = WToL.TransformPosition(Gizmo->GetActorLocation());
 				const float ScaleXY = LandscapeInfo->DrawScale.X;
-
-				MeshRT = FTranslationMatrix(FVector(-W / ScaleXY + 0.5, -H / ScaleXY + 0.5, 0) * GizmoScale3D) * FRotationTranslationMatrix(FRotator(0, Gizmo->GetActorRotation().Yaw, 0), FVector(BaseLocation.X, BaseLocation.Y, 0)) * LToW.ToMatrixWithScale();
-				HeightmapRenderProxy = new FLandscapeGizmoMeshRenderProxy( Gizmo->GizmoMeshMaterial->GetRenderProxy(), BaseLocation.Z + L, BaseLocation.Z, Gizmo->GizmoTexture, FLinearColor(Gizmo->TextureScale.X, Gizmo->TextureScale.Y, 0, 0), WToL );
-
-				GizmoRenderProxy = (Gizmo->DataType != LGT_None) ? Gizmo->GizmoDataMaterial->GetRenderProxy() : Gizmo->GizmoMaterial->GetRenderProxy();
-
 				// Cache sampled height
 				float ScaleX = Gizmo->GetWidth() / Gizmo->CachedWidth / ScaleXY * Gizmo->CachedScaleXY;
 				float ScaleY = Gizmo->GetHeight() / Gizmo->CachedHeight / ScaleXY * Gizmo->CachedScaleXY;
-				FScaleMatrix Mat(FVector(ScaleX, ScaleY, L));
-				FMatrix NormalM = Mat.InverseFast().GetTransposed();
-
-				int32 SamplingSize = Gizmo->SampleSizeX * Gizmo->SampleSizeY;
-				SampledPositions.Empty(SamplingSize);
-				SampledNormals.Empty(SamplingSize);
 
-				for (int32 Y = 0; Y < Gizmo->SampleSizeY; ++Y)
+				if (!FMath::IsNearlyZero(ScaleX) && !FMath::IsNearlyZero(ScaleY))
 				{
-					for (int32 X = 0; X < Gizmo->SampleSizeX; ++X)
+					MeshRT = FTranslationMatrix(FVector(-W / ScaleXY + 0.5, -H / ScaleXY + 0.5, 0) * GizmoScale3D) * FRotationTranslationMatrix(FRotator(0, Gizmo->GetActorRotation().Yaw, 0), FVector(BaseLocation.X, BaseLocation.Y, 0)) * LToW.ToMatrixWithScale();
+
+					HeightmapRenderProxy = new FLandscapeGizmoMeshRenderProxy( Gizmo->GizmoMeshMaterial->GetRenderProxy(), BaseLocation.Z + L, BaseLocation.Z, Gizmo->GizmoTexture, FLinearColor(Gizmo->TextureScale.X, Gizmo->TextureScale.Y, 0, 0), WToL );
+					GizmoRenderProxy = (Gizmo->DataType != LGT_None) ? Gizmo->GizmoDataMaterial->GetRenderProxy() : Gizmo->GizmoMaterial->GetRenderProxy();
+
+					FScaleMatrix Mat(FVector(ScaleX, ScaleY, L));
+					FMatrix NormalM = Mat.InverseFast().GetTransposed();
+
+					int32 SamplingSize = Gizmo->SampleSizeX * Gizmo->SampleSizeY;
+					SampledPositions.Empty(SamplingSize);
+					SampledNormals.Empty(SamplingSize);
+
+					for (int32 Y = 0; Y < Gizmo->SampleSizeY; ++Y)
 					{
-						FVector SampledPos = Gizmo->SampledHeight[X + Y * ALandscapeGizmoActiveActor::DataTexSize];
-						SampledPos.X *= ScaleX;
-						SampledPos.Y *= ScaleY;
-						SampledPos.Z = Gizmo->GetLandscapeHeight(SampledPos.Z);
+						for (int32 X = 0; X < Gizmo->SampleSizeX; ++X)
+						{
+							FVector SampledPos = Gizmo->SampledHeight[X + Y * ALandscapeGizmoActiveActor::DataTexSize];
+							SampledPos.X *= ScaleX;
+							SampledPos.Y *= ScaleY;
+							SampledPos.Z = Gizmo->GetLandscapeHeight(SampledPos.Z);
 
-						FVector SampledNormal = NormalM.TransformVector(Gizmo->SampledNormal[X + Y * ALandscapeGizmoActiveActor::DataTexSize]);
-						SampledNormal = SampledNormal.GetSafeNormal();
+							FVector SampledNormal = NormalM.TransformVector(Gizmo->SampledNormal[X + Y * ALandscapeGizmoActiveActor::DataTexSize]);
+							SampledNormal = SampledNormal.GetSafeNormal();
 
-						SampledPositions.Add(SampledPos);
-						SampledNormals.Add(SampledNormal);
-						//MeshBuilder.AddVertex(SampledPos, FVector2D((float)X / (Gizmo->SampleSizeX), (float)Y / (Gizmo->SampleSizeY)), TangentX, SampledNormal^TangentX, SampledNormal, FColor::White );
+							SampledPositions.Add(SampledPos);
+							SampledNormals.Add(SampledNormal);
+						}
 					}
 				}
 			}
@@ -223,7 +226,7 @@ public:
 	~FLandscapeGizmoRenderSceneProxy()
 	{
 		delete HeightmapRenderProxy;
-		HeightmapRenderProxy = NULL;
+		HeightmapRenderProxy = nullptr;
 	}
 
 #if WITH_EDITOR
@@ -242,7 +245,7 @@ public:
 	{
 		//FMemMark Mark(FMemStack::Get());
 #if WITH_EDITOR
-		if( GizmoRenderProxy &&  HeightmapRenderProxy )
+		if(GizmoRenderProxy && HeightmapRenderProxy)
 		{
 			for (int32 ViewIndex = 0; ViewIndex < Views.Num(); ViewIndex++)
 			{
@@ -470,7 +473,7 @@ void ALandscapeGizmoActor::Duplicate(ALandscapeGizmoActor* Gizmo)
 	Gizmo->SetActorLocation( GetActorLocation(), false );
 	Gizmo->SetActorRotation( GetActorRotation() );
 
-	if (Gizmo->GetRootComponent() != NULL && GetRootComponent() != NULL)
+	if (Gizmo->GetRootComponent() != nullptr && GetRootComponent() != nullptr)
 	{
 		Gizmo->GetRootComponent()->SetRelativeScale3D(GetRootComponent()->GetRelativeScale3D());
 	}
@@ -914,7 +917,7 @@ void ALandscapeGizmoActiveActor::SampleData(int32 SizeX, int32 SizeY)
 
 LANDSCAPE_API void ALandscapeGizmoActiveActor::Import( int32 VertsX, int32 VertsY, uint16* HeightData, TArray<ULandscapeLayerInfoObject*> ImportLayerInfos, uint8* LayerDataPointers[] )
 {
-	if (VertsX <= 0 || VertsY <= 0 || HeightData == NULL || TargetLandscapeInfo == NULL || GizmoTexture == NULL || (ImportLayerInfos.Num() && !LayerDataPointers) )
+	if (VertsX <= 0 || VertsY <= 0 || HeightData == nullptr || TargetLandscapeInfo == nullptr || GizmoTexture == nullptr || (ImportLayerInfos.Num() && !LayerDataPointers) )
 	{
 		return;
 	}
@@ -1208,7 +1211,7 @@ void ALandscapeGizmoActiveActor::ImportFromClipboard()
 						Str++;
 					}
 					StrBuf[i] = 0;
-					LayerInfos.Add( LoadObject<ULandscapeLayerInfoObject>(NULL, StrBuf) );
+					LayerInfos.Add( LoadObject<ULandscapeLayerInfoObject>(nullptr, StrBuf) );
 				}
 			}
 
