Build fix

diff --git a/Engine/Source/Runtime/Landscape/Private/LandscapeGizmoActor.cpp b/Engine/Source/Runtime/Landscape/Private/LandscapeGizmoActor.cpp
--- a/Engine/Source/Runtime/Landscape/Private/LandscapeGizmoActor.cpp
+++ b/Engine/Source/Runtime/Landscape/Private/LandscapeGizmoActor.cpp
@@ -147,12 +147,12 @@ public:
 		ALandscapeGizmoActiveActor* Gizmo = Cast<ALandscapeGizmoActiveActor>(InComponent->GetOwner());
 		if (Gizmo && Gizmo->GizmoMeshMaterial && Gizmo->GizmoDataMaterial && Gizmo->GetRootComponent() 
 			&& !FMath::IsNearlyZero(Gizmo->CachedWidth) 
-			&& !FMath::IsNearlyZero(Gizmo->CachedHeight) 
-			&& !FMath::IsNearlyZero(ScaleXY))
+			&& !FMath::IsNearlyZero(Gizmo->CachedHeight))
 		{
 			ULandscapeInfo* LandscapeInfo = Gizmo->TargetLandscapeInfo;
-			if (LandscapeInfo && LandscapeInfo->GetLandscapeProxy())
+			if (LandscapeInfo && LandscapeInfo->GetLandscapeProxy() && !FMath::IsNearlyZero(LandscapeInfo->DrawScale.X))
 			{
+				const float ScaleXY = LandscapeInfo->DrawScale.X;
 				SampleSizeX = Gizmo->SampleSizeX;
 				SampleSizeY = Gizmo->SampleSizeY;
 				bHeightmapRendering = (Gizmo->DataType & LGT_Height);
@@ -182,7 +182,7 @@ public:
 
 				const FMatrix WToL = LToW.ToMatrixWithScale().InverseFast();
 				const FVector BaseLocation = WToL.TransformPosition(Gizmo->GetActorLocation());
-				const float ScaleXY = LandscapeInfo->DrawScale.X;
+				
 				// Cache sampled height
 				float ScaleX = Gizmo->GetWidth() / Gizmo->CachedWidth / ScaleXY * Gizmo->CachedScaleXY;
 				float ScaleY = Gizmo->GetHeight() / Gizmo->CachedHeight / ScaleXY * Gizmo->CachedScaleXY;
