PR #7998: Added setting for contact shadows on grass (Contributed by [REDACTED_PERSON])

diff --git a/Engine/Source/Runtime/Landscape/Classes/LandscapeGrassType.h b/Engine/Source/Runtime/Landscape/Classes/LandscapeGrassType.h
--- a/Engine/Source/Runtime/Landscape/Classes/LandscapeGrassType.h
+++ b/Engine/Source/Runtime/Landscape/Classes/LandscapeGrassType.h
@@ -106,6 +106,10 @@ struct FGrassVariety
 	UPROPERTY(EditAnywhere, Category = Grass)
 	bool bCastDynamicShadow;
 
+	/** Whether the grass should cast contact shadows. **/
+	UPROPERTY(EditAnywhere, Category = Grass)
+	bool bCastContactShadow;
+
 	/** Whether we should keep a cpu copy of the instance buffer. This should be set to true if you plan on using GetOverlappingXXXXCount functions of the component otherwise it won't return any data.**/
 	UPROPERTY(EditAnywhere, Category = Grass)
 	bool bKeepInstanceBufferCPUCopy;
@@ -127,6 +131,7 @@ struct FGrassVariety
 		, bUseLandscapeLightmap(false)
 		, bReceivesDecals(true)
 		, bCastDynamicShadow(true)
+		, bCastContactShadow(true)
 		, bKeepInstanceBufferCPUCopy(false)
 	{
 	}
diff --git a/Engine/Source/Runtime/Landscape/Private/LandscapeGrass.cpp b/Engine/Source/Runtime/Landscape/Private/LandscapeGrass.cpp
--- a/Engine/Source/Runtime/Landscape/Private/LandscapeGrass.cpp
+++ b/Engine/Source/Runtime/Landscape/Private/LandscapeGrass.cpp
@@ -3023,8 +3023,9 @@ void ALandscapeProxy::UpdateGrass(const TArray<FVector>& Cameras, int32& InOutNu
 										HierarchicalInstancedStaticMeshComponent->InstancingRandomSeed = FolSeed;
 										HierarchicalInstancedStaticMeshComponent->LightingChannels = GrassVariety.LightingChannels;
 										HierarchicalInstancedStaticMeshComponent->bCastStaticShadow = false;
-										HierarchicalInstancedStaticMeshComponent->CastShadow = GrassVariety.bCastDynamicShadow && !bDisableDynamicShadows;
+										HierarchicalInstancedStaticMeshComponent->CastShadow = (GrassVariety.bCastDynamicShadow || GrassVariety.bCastContactShadow) && !bDisableDynamicShadows;
 										HierarchicalInstancedStaticMeshComponent->bCastDynamicShadow = GrassVariety.bCastDynamicShadow && !bDisableDynamicShadows;
+										HierarchicalInstancedStaticMeshComponent->bCastContactShadow = GrassVariety.bCastContactShadow && !bDisableDynamicShadows;
 										HierarchicalInstancedStaticMeshComponent->OverrideMaterials = GrassVariety.OverrideMaterials;
 
 										const FMeshMapBuildData* MeshMapBuildData = Component->GetMeshMapBuildData();
