Added support for subsections in landscape edit layers local merge

diff --git a/Engine/Shaders/Private/Landscape/LandscapeLayersHeightmapsPS.usf b/Engine/Shaders/Private/Landscape/LandscapeLayersHeightmapsPS.usf
--- a/Engine/Shaders/Private/Landscape/LandscapeLayersHeightmapsPS.usf
+++ b/Engine/Shaders/Private/Landscape/LandscapeLayersHeightmapsPS.usf
@@ -7,8 +7,8 @@
 #if defined (__INTELLISENSE__)
 // Uncomment the appropriate define for enabling syntax highlighting with HLSL Tools for Visual Studio : 
 //#define MERGE_EDIT_LAYERS 1
-#define STITCH_HEIGHTMAP 1
-#define FINALIZE_HEIGHTMAP 1
+//#define STITCH_HEIGHTMAP 1
+//#define FINALIZE_HEIGHTMAP 1
 //#define GENERATE_MIPS 1
 #endif // defined (__INTELLISENSE__)
 
@@ -16,6 +16,9 @@
 
 // ----------------------------------------------------------------------------------
 
+// All shaders inputs/outputs
+uint InNumSubsections; // Number of sub-sections for this landscape (1 or 2)
+
 #if MERGE_EDIT_LAYERS
 // MergeEditLayers inputs/outputs :
 
@@ -160,6 +163,32 @@ float2 LoadNeighborHeightmapUnpackedHeight(bool2 bInIsMinComponentBorder, bool2
 			SampleCoordinates.y = bSkipDuplicateRowColumn ? +1 : 0;
 		}
 
+		// Handle the sub-section case, where there's a duplicate row/column of vertices in the middle : 
+		if (InNumSubsections == 2)
+		{
+			uint2 SubSectionSize = InSourceTextureSize / 2;
+			// If we need the neighbor on the right and we're on the left (duplicate) column of the subsection, skip a column to the right :
+			if ((InRelativeTextureCoordinates.x == SubSectionSize.x - 1) && (InNeighborOffset.x > 0))
+			{
+				SampleCoordinates.x += 1;
+			}
+			// If we need the neighbor on the left and we're on the right (duplicate) column of the subsection, skip a column to the left :
+			else if ((InRelativeTextureCoordinates.x == SubSectionSize.x) && (InNeighborOffset.x < 0))
+			{
+				SampleCoordinates.x -= 1;
+			}
+			// If we need the neighbor on the bottom and we're on the top (duplicate) row of the subsection, skip a row to the bottom :
+			if ((InRelativeTextureCoordinates.y == SubSectionSize.y - 1) && (InNeighborOffset.y > 0))
+			{
+				SampleCoordinates.y += 1;
+			}
+			// If we need the neighbor on the top and we're on the bottom (duplicate) row of the subsection, skip a row to the top :
+			else if ((InRelativeTextureCoordinates.y == SubSectionSize.y) && (InNeighborOffset.y < 0))
+			{
+				SampleCoordinates.y -= 1;
+			}
+		}
+
 		int4 LoadCoordinates = int4(SampleCoordinates, NeighborHeightmapIndex, 0); // xy = relative coordinates, z = index in texture array, w = mip level
 		float2 PackedHeight = InSourceHeightmaps.Load(LoadCoordinates);
 		Height = UnpackHeight(PackedHeight);
@@ -461,10 +490,19 @@ void GenerateMips(in float4 SVPos : SV_POSITION, out float4 OutColor : SV_Target
 	// | * :   |   | * : * |   |   : * | 
 	// +-------+   +-------+   +-------+
 
+	bool bIsLastMip = all(InCurrentMipSubregionSize == 1);
+	
 	uint2 RelativeTextureCoordinates = TextureCoordinates % InCurrentMipSubregionSize;
 	bool2 bIsMinBorder = (RelativeTextureCoordinates == 0);
 	bool2 bIsMaxBorder = (RelativeTextureCoordinates == (InCurrentMipSubregionSize - 1));
-	bool bIsLastMip = all(InCurrentMipSubregionSize == 1);
+
+	// If we use subsections, we'll have an additional column/row of min border and another of max border next to them
+	if (InNumSubsections == 2)
+	{
+		uint2 SubsectionSize = InCurrentMipSubregionSize / 2;
+		bIsMinBorder = bIsMinBorder || (RelativeTextureCoordinates == SubsectionSize);
+		bIsMaxBorder = bIsMaxBorder || (RelativeTextureCoordinates == SubsectionSize - 1);
+	}
 		
 	float SampleWeights[4] =
 	{ 
diff --git a/Engine/Shaders/Private/Landscape/LandscapeLayersWeightmapsPS.usf b/Engine/Shaders/Private/Landscape/LandscapeLayersWeightmapsPS.usf
--- a/Engine/Shaders/Private/Landscape/LandscapeLayersWeightmapsPS.usf
+++ b/Engine/Shaders/Private/Landscape/LandscapeLayersWeightmapsPS.usf
@@ -7,7 +7,7 @@
 #if defined (__INTELLISENSE__)
 // Uncomment the appropriate define for enabling syntax highlighting with HLSL Tools for Visual Studio : 
 //#define MERGE_EDIT_LAYERS 1
-#define FINALIZE_WEIGHTMAP 1
+//#define FINALIZE_WEIGHTMAP 1
 //#define GENERATE_MIPS 1
 #endif // defined (__INTELLISENSE__)
 
@@ -77,6 +77,7 @@ StructuredBuffer<FWeightmapPaintLayerInfo> InPaintLayerInfos;
 #if GENERATE_MIPS
 // GenerateMips inputs/outputs :
 uint2 InCurrentMipSize; // Size of the texture mip currently being generated
+uint InNumSubsections; // Number of sub-sections for this landscape (1 or 2)
 
 Texture2D<float4> InSourceWeightmap;
 #endif // GENERATE_MIPS
@@ -260,9 +261,17 @@ void GenerateMips(in float4 SVPos : SV_POSITION, out float4 OutColor : SV_Target
 	// | * :   |   | * : * |   |   : * | <-- bIsMaxBorder.y = true
 	// +-------+   +-------+   +-------+
 
+	bool bIsLastMip = all(InCurrentMipSize == 1);
+
 	bool2 bIsMinBorder = (TextureCoordinates == 0);
 	bool2 bIsMaxBorder = (TextureCoordinates == (InCurrentMipSize - 1));
-	bool bIsLastMip = all(InCurrentMipSize == 1);
+	// If we use subsections, we'll have an additional column/row of min border and another of max border next to them
+	if (InNumSubsections == 2)
+	{
+		uint2 SubsectionSize = InCurrentMipSize / 2;
+		bIsMinBorder = bIsMinBorder || (TextureCoordinates == SubsectionSize);
+		bIsMaxBorder = bIsMaxBorder || (TextureCoordinates == SubsectionSize - 1);
+	}
 		
 	float SampleWeights[4] =
 	{ 
diff --git a/Engine/Source/Runtime/Landscape/Private/LandscapeEditLayers.cpp b/Engine/Source/Runtime/Landscape/Private/LandscapeEditLayers.cpp
--- a/Engine/Source/Runtime/Landscape/Private/LandscapeEditLayers.cpp
+++ b/Engine/Source/Runtime/Landscape/Private/LandscapeEditLayers.cpp
@@ -1163,6 +1163,7 @@ class FLandscapeLayersHeightmapsStitchHeightmapPS : public FGlobalShader
 public:
 	BEGIN_SHADER_PARAMETER_STRUCT(FParameters, )
 		SHADER_PARAMETER(FUintVector2, InSourceTextureSize)
+		SHADER_PARAMETER(uint32, InNumSubsections)
 		SHADER_PARAMETER_RDG_TEXTURE_SRV(Texture2DArray<float2>, InSourceHeightmaps)
 		SHADER_PARAMETER_SCALAR_ARRAY(uint32, InNeighborHeightmapIndices, [9])
 		RENDER_TARGET_BINDING_SLOTS()
@@ -1205,6 +1206,7 @@ class FLandscapeLayersHeightmapsFinalizeHeightmapPS : public FGlobalShader
 public:
 	BEGIN_SHADER_PARAMETER_STRUCT(FParameters, )
 		SHADER_PARAMETER(FUintVector2, InSourceTextureSize)
+		SHADER_PARAMETER(uint32, InNumSubsections)
 		SHADER_PARAMETER_RDG_TEXTURE_SRV(Texture2DArray<float2>, InSourceHeightmaps)
 		SHADER_PARAMETER_SCALAR_ARRAY(uint32, InNeighborHeightmapIndices, [9])
 		SHADER_PARAMETER(FUintVector4, InDestinationTextureSubregion)
@@ -1250,6 +1252,7 @@ class FLandscapeLayersHeightmapsGenerateMipsPS : public FGlobalShader
 public:
 	BEGIN_SHADER_PARAMETER_STRUCT(FParameters, )
 		SHADER_PARAMETER(FUintVector2, InCurrentMipSubregionSize)
+		SHADER_PARAMETER(uint32, InNumSubsections)
 		SHADER_PARAMETER_RDG_TEXTURE_SRV(Texture2D<float4>, InSourceHeightmap)
 		RENDER_TARGET_BINDING_SLOTS()
 	END_SHADER_PARAMETER_STRUCT()
@@ -1381,6 +1384,7 @@ class FLandscapeLayersWeightmapsGenerateMipsPS : public FGlobalShader
 public:
 	BEGIN_SHADER_PARAMETER_STRUCT(FParameters, )
 		SHADER_PARAMETER(FUintVector2, InCurrentMipSize)
+		SHADER_PARAMETER(uint32, InNumSubsections)
 		SHADER_PARAMETER_RDG_TEXTURE_SRV(Texture2D<float4>, InSourceWeightmap)
 		RENDER_TARGET_BINDING_SLOTS()
 		END_SHADER_PARAMETER_STRUCT()
@@ -3930,6 +3934,9 @@ namespace EditLayersHeightmapLocalMerge_RenderThread
 		// Number of vertices per component
 		FIntPoint ComponentSizeVerts = FIntPoint(ForceInitToZero);
 
+		// Number of sub sections for this landscape 
+		uint32 NumSubsections = 1;
+
 		// Maximum size of all heigthmaps (one heightmap can contain multiple components due to heightmap sharing)
 		FIntPoint MaxHeightmapSize = FIntPoint(ForceInitToZero);
 
@@ -4137,6 +4144,7 @@ namespace EditLayersHeightmapLocalMerge_RenderThread
 				FLandscapeLayersHeightmapsStitchHeightmapPS::FParameters* StitchHeightmapPSParams = GraphBuilder.AllocParameters<FLandscapeLayersHeightmapsStitchHeightmapPS::FParameters>();
 				StitchHeightmapPSParams->RenderTargets[0] = FRenderTargetBinding(RDGResources.ScratchStitchedHeightmapTextureArray, ERenderTargetLoadAction::ENoAction, /*InMipIndex = */0, /*InArraySlice = */IndexInBatch);
 				StitchHeightmapPSParams->InSourceTextureSize = FUintVector2(InLocalMergeInfo.ComponentSizeVerts.X, InLocalMergeInfo.ComponentSizeVerts.Y);
+				StitchHeightmapPSParams->InNumSubsections = InLocalMergeInfo.NumSubsections;
 				StitchHeightmapPSParams->InSourceHeightmaps = RDGResources.ScratchMergedHeightmapTextureArraySRV;
 
 				for (int32 NeighborIndex = 0; NeighborIndex < 9; ++NeighborIndex)
@@ -4181,6 +4189,7 @@ namespace EditLayersHeightmapLocalMerge_RenderThread
 				FLandscapeLayersHeightmapsFinalizeHeightmapPS::FParameters* FinalizeHeightmapPSParams = GraphBuilder.AllocParameters<FLandscapeLayersHeightmapsFinalizeHeightmapPS::FParameters>();
 				FinalizeHeightmapPSParams->RenderTargets[0] = FRenderTargetBinding(TrackedTexture->ScratchTextureRef, ERenderTargetLoadAction::ENoAction);
 				FinalizeHeightmapPSParams->InSourceTextureSize = FUintVector2(InLocalMergeInfo.ComponentSizeVerts.X, InLocalMergeInfo.ComponentSizeVerts.Y);
+				FinalizeHeightmapPSParams->InNumSubsections = InLocalMergeInfo.NumSubsections;
 				FinalizeHeightmapPSParams->InSourceHeightmaps = RDGResources.ScratchStitchedHeightmapTextureArraySRV;
 				FinalizeHeightmapPSParams->InDestinationTextureSubregion = FUintVector4(ComponentResolveInfo.Heightmap.Subregion.Min.X, ComponentResolveInfo.Heightmap.Subregion.Min.Y, ComponentResolveInfo.Heightmap.Subregion.Max.X, ComponentResolveInfo.Heightmap.Subregion.Max.Y);
 				FinalizeHeightmapPSParams->InLandscapeGridScale = (FVector3f)InLocalMergeInfo.LandscapeGridScale;
@@ -4232,6 +4241,7 @@ namespace EditLayersHeightmapLocalMerge_RenderThread
 				FLandscapeLayersHeightmapsGenerateMipsPS::FParameters* GenerateMipsPSParams = GraphBuilder.AllocParameters<FLandscapeLayersHeightmapsGenerateMipsPS::FParameters>();
 				GenerateMipsPSParams->RenderTargets[0] = FRenderTargetBinding(TrackedTexture->ScratchTextureRef, ERenderTargetLoadAction::ENoAction, MipLevel);
 				GenerateMipsPSParams->InCurrentMipSubregionSize = FUintVector2(CurrentMipSubregionSize.X, CurrentMipSubregionSize.Y);
+				GenerateMipsPSParams->InNumSubsections = InLocalMergeInfo.NumSubsections;
 				GenerateMipsPSParams->InSourceHeightmap = TrackedTexture->ScratchTextureMipsSRVRefs[MipLevel - 1];
 
 				FLandscapeLayersHeightmapsGenerateMipsPS::GenerateSingleMip(GraphBuilder, GenerateMipsPSParams);
@@ -4317,6 +4327,7 @@ void ALandscape::PrepareLayersHeightmapsLocalMergeRenderThreadData(const FUpdate
 	int32 ComponentSizeVerts = (SubsectionSizeQuads + 1) * NumSubsections;
 	OutRenderThreadData.ComponentSizeVerts = FIntPoint(ComponentSizeVerts, ComponentSizeVerts);
 	OutRenderThreadData.LandscapeGridScale = GetRootComponent()->GetRelativeScale3D();
+	OutRenderThreadData.NumSubsections = NumSubsections;
 
 	// Prepare landscape edit layers data common to all landscape components: 
 	OutRenderThreadData.VisibleEditLayerInfos.Reserve(LandscapeLayers.Num());
@@ -5800,8 +5811,13 @@ namespace EditLayersWeightmapLocalMerge_RenderThread
 		// Maximum number of weightmap arrays that is needed for a given FTextureResolveBatchInfo (1 per FComponentRenderInfo in the batch) :
 		int32 MaxNumWeightmapArraysPerResolveTextureBatch = 0;
 
-		// Per-landscape component info :
+		// Number of vertices per component : 
 		FIntPoint ComponentSizeVerts = 0;
+
+		// Number of sub sections for this landscape :
+		uint32 NumSubsections = 1;
+
+		// Number of mips for the weightmaps of this landscape
 		int32 NumMips = 0;
 
 		// List of batches of FTextureResolveInfo that needs to be resolved in the same pass. This allows massive saves on transient resources on large landscapes because those can be re-cycled from one pass to another :
@@ -6153,6 +6169,7 @@ namespace EditLayersWeightmapLocalMerge_RenderThread
 			FLandscapeLayersWeightmapsGenerateMipsPS::FParameters* GenerateMipsPSParams = GraphBuilder.AllocParameters<FLandscapeLayersWeightmapsGenerateMipsPS::FParameters>();
 			GenerateMipsPSParams->RenderTargets[0] = FRenderTargetBinding(RDGResources.ScratchFinalWeightmapTexture, ERenderTargetLoadAction::ENoAction, MipLevel);
 			GenerateMipsPSParams->InCurrentMipSize = FUintVector2(CurrentMipSize.X, CurrentMipSize.Y);
+			GenerateMipsPSParams->InNumSubsections = InLocalMergeInfo.NumSubsections;
 			GenerateMipsPSParams->InSourceWeightmap = RDGResources.ScratchFinalWeightmapTextureMipsSRV[MipLevel - 1];
 
 			FLandscapeLayersWeightmapsGenerateMipsPS::GenerateSingleMip(GraphBuilder, GenerateMipsPSParams);
@@ -6210,6 +6227,7 @@ void ALandscape::PrepareLayersWeightmapsLocalMergeRenderThreadData(const FUpdate
 	int32 ComponentSizeVerts = (SubsectionSizeQuads + 1) * NumSubsections;
 	OutRenderThreadData.ComponentSizeVerts = FIntPoint(ComponentSizeVerts, ComponentSizeVerts);
 	OutRenderThreadData.NumMips = FMath::CeilLogTwo(ComponentSizeVerts) + 1;
+	OutRenderThreadData.NumSubsections = NumSubsections;
 
 	// Lookup table to retrieve, for a given paint layer, its index in OutRenderThreadData.PaintLayerInfos (we don't keep UObjects in OutRenderThreadData because it's a render-thread struct) :
 	TMap<ULandscapeLayerInfoObject*, int32> PaintLayerToPaintLayerInfoIndex;
