Move SystemGUIDs in to a cpp to avoid massive recompiles when changing
Convert landscape mobile cook version to system guid system
#rnx

diff --git a/Engine/Source/Runtime/Core/Private/UObject/DevObjectVersion.cpp b/Engine/Source/Runtime/Core/Private/UObject/DevObjectVersion.cpp
--- a/Engine/Source/Runtime/Core/Private/UObject/DevObjectVersion.cpp
+++ b/Engine/Source/Runtime/Core/Private/UObject/DevObjectVersion.cpp
@@ -116,6 +116,7 @@ FGuid FDevSystemGuidRegistry::GetSystemGuid(FGuid System)
 
 FDevSystemGuids::FDevSystemGuids()
 	: GLOBALSHADERMAP_DERIVEDDATA_VER(0x7912A706, 0x52B8450A, 0x9CD66FA7, 0xEFBBFA0A)
+	, LANDSCAPE_MOBILE_COOK_VERSION(0x0E9ADF72, 0xD6B64E0D, 0x81C4A92B, 0x081A37AB)
 	, MATERIALSHADERMAP_DERIVEDDATA_VER(0x2579AAFE, 0x8F1D4E4F, 0xB04278C6, 0x35917C17)
 	, NANITE_DERIVEDDATA_VER(0xBEB0226A, 0x070E4ECA, 0x972C1E7D, 0xD8599E68)
 	, NIAGARASHADERMAP_DERIVEDDATA_VER(0x7BBD9913, 0xC1554D20, 0xADAE9F17, 0xB006299E)
diff --git a/Engine/Source/Runtime/Core/Private/UObject/FortniteMainBranchObjectVersion.cpp b/Engine/Source/Runtime/Core/Private/UObject/FortniteMainBranchObjectVersion.cpp
new file mode 100644
--- /dev/null
+++ b/Engine/Source/Runtime/Core/Private/UObject/FortniteMainBranchObjectVersion.cpp
@@ -0,0 +1,19 @@
+// Copyright Epic Games, Inc. All Rights Reserved.
+#include "UObject/FortniteMainBranchObjectVersion.h"
+
+TMap<FGuid, FGuid> FFortniteMainBranchObjectVersion::GetSystemGuids()
+{
+	TMap<FGuid, FGuid> SystemGuids;
+	const FDevSystemGuids& DevGuids = FDevSystemGuids::Get();
+
+	SystemGuids.Add(DevGuids.GLOBALSHADERMAP_DERIVEDDATA_VER, FGuid("8D3A12924CDF4D9F84C0A900E9445CAD"));
+	SystemGuids.Add(DevGuids.LANDSCAPE_MOBILE_COOK_VERSION, FGuid("32D02EF867C74B71A0D4E0FA41392732"));
+	SystemGuids.Add(DevGuids.MATERIALSHADERMAP_DERIVEDDATA_VER, FGuid("1DA85DA5733C4A489A1BE117CDC5ACCC"));
+	SystemGuids.Add(DevGuids.NANITE_DERIVEDDATA_VER, FGuid("49FBF137DBFB4204B6C9E5B91DA15EBD"));
+	SystemGuids.Add(DevGuids.NIAGARASHADERMAP_DERIVEDDATA_VER, FGuid("2289D5116CF94BC0AFEAEC541468E645"));
+	SystemGuids.Add(DevGuids.Niagara_LatestScriptCompileVersion, FGuid("68A430A98C770B4BB1331E206470CC26"));
+	SystemGuids.Add(DevGuids.SkeletalMeshDerivedDataVersion, FGuid("ACF593EAAE354FCBB34CF44F5AA1BFD2"));
+	SystemGuids.Add(DevGuids.STATICMESH_DERIVEDDATA_VER, FGuid("3006d21f867648a68cf4edee185446b7"));
+
+	return SystemGuids;
+}
diff --git a/Engine/Source/Runtime/Core/Private/UObject/UE5MainStreamObjectVersion.cpp b/Engine/Source/Runtime/Core/Private/UObject/UE5MainStreamObjectVersion.cpp
new file mode 100644
--- /dev/null
+++ b/Engine/Source/Runtime/Core/Private/UObject/UE5MainStreamObjectVersion.cpp
@@ -0,0 +1,22 @@
+// Copyright Epic Games, Inc. All Rights Reserved.
+#include "UObject/UE5MainStreamObjectVersion.h"
+#include "Containers/Map.h"
+#include "Misc/Guid.h"
+
+
+TMap<FGuid, FGuid> FUE5MainStreamObjectVersion::GetSystemGuids()
+{
+	TMap<FGuid, FGuid> SystemGuids;
+	const FDevSystemGuids& DevGuids = FDevSystemGuids::Get();
+
+	SystemGuids.Add(DevGuids.GLOBALSHADERMAP_DERIVEDDATA_VER, FGuid("62F5564D1FED4A2D8864DF300EC5AA2F"));
+	SystemGuids.Add(DevGuids.LANDSCAPE_MOBILE_COOK_VERSION, FGuid("71000000000000000000000000000035"));
+	SystemGuids.Add(DevGuids.MATERIALSHADERMAP_DERIVEDDATA_VER, FGuid("65C01D817C9A4EEFAE9E988D41A1F3DD"));
+	SystemGuids.Add(DevGuids.NANITE_DERIVEDDATA_VER, FGuid("7E26CA5B0F5F03C40A3FD7173D88B8FA"));
+	SystemGuids.Add(DevGuids.NIAGARASHADERMAP_DERIVEDDATA_VER, FGuid("8A37C45D24F2423CBE5F8F371DE33575"));
+	SystemGuids.Add(DevGuids.Niagara_LatestScriptCompileVersion, FGuid("BFA6C45A89CE2B84BB2C444C78C46E76"));
+	SystemGuids.Add(DevGuids.SkeletalMeshDerivedDataVersion, FGuid("FB7B2FB546DE49B2A2F88B7F69763CEE"));
+	SystemGuids.Add(DevGuids.STATICMESH_DERIVEDDATA_VER, FGuid("98F7D79A5811013825E75BBCC41ED3E9"));
+
+	return SystemGuids;
+}
diff --git a/Engine/Source/Runtime/Core/Private/UObject/UE5ReleaseStreamObjectVersion.cpp b/Engine/Source/Runtime/Core/Private/UObject/UE5ReleaseStreamObjectVersion.cpp
new file mode 100644
--- /dev/null
+++ b/Engine/Source/Runtime/Core/Private/UObject/UE5ReleaseStreamObjectVersion.cpp
@@ -0,0 +1,18 @@
+// Copyright Epic Games, Inc. All Rights Reserved.
+#include "UObject/UE5ReleaseStreamObjectVersion.h"
+#include "Containers/Map.h"
+#include "Misc/Guid.h"
+
+TMap<FGuid, FGuid> FUE5ReleaseStreamObjectVersion::GetSystemGuids()
+{
+	TMap<FGuid, FGuid> SystemGuids;
+	const FDevSystemGuids& DevGuids = FDevSystemGuids::Get();
+
+	SystemGuids.Add(DevGuids.GLOBALSHADERMAP_DERIVEDDATA_VER, FGuid("D0BF3452816D46908073DFDD4B855AE5"));
+	SystemGuids.Add(DevGuids.MATERIALSHADERMAP_DERIVEDDATA_VER, FGuid("79090A66F9D94B5BB285D49E5D39468E"));
+	SystemGuids.Add(DevGuids.NIAGARASHADERMAP_DERIVEDDATA_VER, FGuid("D204CCC166ED482685DAEE8AB72895AB"));
+	SystemGuids.Add(DevGuids.Niagara_LatestScriptCompileVersion, FGuid("3B3F2E14B150418481F033E5FAD9D13A"));
+	SystemGuids.Add(DevGuids.SkeletalMeshDerivedDataVersion, FGuid("25C49E579B3142DDA2A8C14037267679"));
+
+	return SystemGuids;
+}
diff --git a/Engine/Source/Runtime/Core/Public/UObject/DevObjectVersion.h b/Engine/Source/Runtime/Core/Public/UObject/DevObjectVersion.h
--- a/Engine/Source/Runtime/Core/Public/UObject/DevObjectVersion.h
+++ b/Engine/Source/Runtime/Core/Public/UObject/DevObjectVersion.h
@@ -19,6 +19,7 @@ struct CORE_API FDevSystemGuids
 	
 	FDevSystemGuids();
 	const FGuid GLOBALSHADERMAP_DERIVEDDATA_VER;
+	const FGuid LANDSCAPE_MOBILE_COOK_VERSION;
 	const FGuid MATERIALSHADERMAP_DERIVEDDATA_VER;
 	const FGuid NANITE_DERIVEDDATA_VER;
 	const FGuid NIAGARASHADERMAP_DERIVEDDATA_VER;
diff --git a/Engine/Source/Runtime/Core/Public/UObject/FortniteMainBranchObjectVersion.h b/Engine/Source/Runtime/Core/Public/UObject/FortniteMainBranchObjectVersion.h
--- a/Engine/Source/Runtime/Core/Public/UObject/FortniteMainBranchObjectVersion.h
+++ b/Engine/Source/Runtime/Core/Public/UObject/FortniteMainBranchObjectVersion.h
@@ -210,21 +210,7 @@ struct CORE_API FFortniteMainBranchObjectVersion
 	// The GUID for this custom version number
 	const static FGuid GUID;
 
-	static TMap<FGuid, FGuid> GetSystemGuids()
-	{
-		TMap<FGuid, FGuid> SystemGuids;
-		const FDevSystemGuids& DevGuids = FDevSystemGuids::Get();
-
-		SystemGuids.Add(DevGuids.GLOBALSHADERMAP_DERIVEDDATA_VER,    FGuid("8D3A12924CDF4D9F84C0A900E9445CAD"));
-		SystemGuids.Add(DevGuids.MATERIALSHADERMAP_DERIVEDDATA_VER,  FGuid("1DA85DA5733C4A489A1BE117CDC5ACCC"));
-		SystemGuids.Add(DevGuids.NANITE_DERIVEDDATA_VER,             FGuid("49FBF137DBFB4204B6C9E5B91DA15EBD"));
-		SystemGuids.Add(DevGuids.NIAGARASHADERMAP_DERIVEDDATA_VER,   FGuid("2289D5116CF94BC0AFEAEC541468E645"));
-		SystemGuids.Add(DevGuids.Niagara_LatestScriptCompileVersion, FGuid("68A430A98C770B4BB1331E206470CC26"));
-		SystemGuids.Add(DevGuids.SkeletalMeshDerivedDataVersion,     FGuid("ACF593EAAE354FCBB34CF44F5AA1BFD2"));
-		SystemGuids.Add(DevGuids.STATICMESH_DERIVEDDATA_VER,         FGuid("3006d21f867648a68cf4edee185446b7"));
-
-		return SystemGuids;
-	}
+	static TMap<FGuid, FGuid> GetSystemGuids();
 
 private:
 	FFortniteMainBranchObjectVersion() {}
diff --git a/Engine/Source/Runtime/Core/Public/UObject/UE5MainStreamObjectVersion.h b/Engine/Source/Runtime/Core/Public/UObject/UE5MainStreamObjectVersion.h
--- a/Engine/Source/Runtime/Core/Public/UObject/UE5MainStreamObjectVersion.h
+++ b/Engine/Source/Runtime/Core/Public/UObject/UE5MainStreamObjectVersion.h
@@ -202,21 +202,7 @@ struct CORE_API FUE5MainStreamObjectVersion
 	// The GUID for this custom version number
 	const static FGuid GUID;
 
-	static TMap<FGuid, FGuid> GetSystemGuids()
-	{
-		TMap<FGuid, FGuid> SystemGuids;
-		const FDevSystemGuids& DevGuids = FDevSystemGuids::Get();
-
-		SystemGuids.Add(DevGuids.GLOBALSHADERMAP_DERIVEDDATA_VER, FGuid("62F5564D1FED4A2D8864DF300EC5AA2F"));
-		SystemGuids.Add(DevGuids.MATERIALSHADERMAP_DERIVEDDATA_VER, FGuid("65C01D817C9A4EEFAE9E988D41A1F3DD"));
-		SystemGuids.Add(DevGuids.NANITE_DERIVEDDATA_VER, FGuid("7E26CA5B0F5F03C40A3FD7173D88B8FA"));
-		SystemGuids.Add(DevGuids.NIAGARASHADERMAP_DERIVEDDATA_VER, FGuid("8A37C45D24F2423CBE5F8F371DE33575"));
-		SystemGuids.Add(DevGuids.Niagara_LatestScriptCompileVersion, FGuid("BFA6C45A89CE2B84BB2C444C78C46E76"));
-		SystemGuids.Add(DevGuids.SkeletalMeshDerivedDataVersion, FGuid("FB7B2FB546DE49B2A2F88B7F69763CEE"));
-		SystemGuids.Add(DevGuids.STATICMESH_DERIVEDDATA_VER, FGuid("98F7D79A5811013825E75BBCC41ED3E9"));
-
-		return SystemGuids;
-	}
+	static TMap<FGuid, FGuid> GetSystemGuids();
 
 	FUE5MainStreamObjectVersion() = delete;
 };
diff --git a/Engine/Source/Runtime/Core/Public/UObject/UE5ReleaseStreamObjectVersion.h b/Engine/Source/Runtime/Core/Public/UObject/UE5ReleaseStreamObjectVersion.h
--- a/Engine/Source/Runtime/Core/Public/UObject/UE5ReleaseStreamObjectVersion.h
+++ b/Engine/Source/Runtime/Core/Public/UObject/UE5ReleaseStreamObjectVersion.h
@@ -129,19 +129,7 @@ struct CORE_API FUE5ReleaseStreamObjectVersion
 	// The GUID for this custom version number
 	const static FGuid GUID;
 
-	static TMap<FGuid, FGuid> GetSystemGuids()
-	{
-		TMap<FGuid, FGuid> SystemGuids;
-		const FDevSystemGuids& DevGuids = FDevSystemGuids::Get();
-
-		SystemGuids.Add(DevGuids.GLOBALSHADERMAP_DERIVEDDATA_VER, FGuid("D0BF3452816D46908073DFDD4B855AE5"));
-		SystemGuids.Add(DevGuids.MATERIALSHADERMAP_DERIVEDDATA_VER, FGuid("79090A66F9D94B5BB285D49E5D39468E"));
-		SystemGuids.Add(DevGuids.NIAGARASHADERMAP_DERIVEDDATA_VER, FGuid("D204CCC166ED482685DAEE8AB72895AB"));
-		SystemGuids.Add(DevGuids.Niagara_LatestScriptCompileVersion, FGuid("3B3F2E14B150418481F033E5FAD9D13A"));
-		SystemGuids.Add(DevGuids.SkeletalMeshDerivedDataVersion, FGuid("25C49E579B3142DDA2A8C14037267679"));
-
-		return SystemGuids;
-	}
+	static TMap<FGuid, FGuid> GetSystemGuids();
 
 	FUE5ReleaseStreamObjectVersion() = delete;
 };
diff --git a/Engine/Source/Runtime/Landscape/Private/Landscape.cpp b/Engine/Source/Runtime/Landscape/Private/Landscape.cpp
--- a/Engine/Source/Runtime/Landscape/Private/Landscape.cpp
+++ b/Engine/Source/Runtime/Landscape/Private/Landscape.cpp
@@ -13,6 +13,7 @@ Landscape.cpp: Terrain rendering
 #include "UObject/UObjectIterator.h"
 #include "UObject/PropertyPortFlags.h"
 #include "UObject/ConstructorHelpers.h"
+#include "UObject/DevObjectVersion.h"
 #include "UObject/LinkerLoad.h"
 #include "LandscapeStreamingProxy.h"
 #include "LandscapeInfo.h"
@@ -130,12 +131,6 @@ namespace LandscapeCookStats
 // Set this to 0 to disable landscape cooking and thus disable it on device.
 #define ENABLE_LANDSCAPE_COOKING 1
 
-// If mobile landscape data data needs to be rebuilt (new format, serialization
-// differences, etc.) replace the version GUID below with a new one.
-// In case of merge conflicts with DDC versions, you MUST generate a new GUID
-// and set this new GUID as the version.                                       
-#define LANDSCAPE_MOBILE_COOK_VERSION TEXT("43D02EF867C74B71A0D4E0FA41392707")
-
 #define LOCTEXT_NAMESPACE "Landscape"
 
 static void PrintNumLandscapeShadows()
@@ -293,7 +288,7 @@ void ULandscapeComponent::CheckGenerateLandscapePlatformData(bool bIsCooking, co
 	}
 
 	// Serialize the version guid as part of the hash so we can invalidate DDC data if needed
-	FString Version(LANDSCAPE_MOBILE_COOK_VERSION);
+	FString Version = FDevSystemGuids::GetSystemGuid(FDevSystemGuids::Get().LANDSCAPE_MOBILE_COOK_VERSION).ToString();
 	ComponentStateAr << Version;
 
 	uint32 Hash[5];
