Landscape using Brushes now forces its LS Streaming Proxies to also be loaded in editor (until partial loading of landscape proxies is properly supported when using BP brushes).

diff --git a/Engine/Source/Runtime/Engine/Private/WorldPartition/Landscape/LandscapeActorDesc.cpp b/Engine/Source/Runtime/Engine/Private/WorldPartition/Landscape/LandscapeActorDesc.cpp
--- a/Engine/Source/Runtime/Engine/Private/WorldPartition/Landscape/LandscapeActorDesc.cpp
+++ b/Engine/Source/Runtime/Engine/Private/WorldPartition/Landscape/LandscapeActorDesc.cpp
@@ -50,9 +50,9 @@ void FLandscapeActorDesc::Serialize(FArchive& Ar)
 
 void FLandscapeActorDesc::Unload()
 {
-	if (ALandscapeStreamingProxy* LandscapeStreamingProxy = Cast<ALandscapeStreamingProxy>(GetActor()))
+	if (ALandscapeProxy* LandscapeProxy = Cast<ALandscapeProxy>(GetActor()))
 	{
-		LandscapeStreamingProxy->ActorDescReferences.Empty();
+		LandscapeProxy->ActorDescReferences.Empty();
 	}
 
 	FPartitionActorDesc::Unload();
diff --git a/Engine/Source/Runtime/Landscape/Classes/Landscape.h b/Engine/Source/Runtime/Landscape/Classes/Landscape.h
--- a/Engine/Source/Runtime/Landscape/Classes/Landscape.h
+++ b/Engine/Source/Runtime/Landscape/Classes/Landscape.h
@@ -242,6 +242,7 @@ public:
 	virtual void PostEditChangeProperty(FPropertyChangedEvent& PropertyChangedEvent) override;
 	virtual void PostEditMove(bool bFinished) override;
 	virtual void PostEditUndo() override;
+	virtual void PostRegisterAllComponents() override;
 	virtual bool ShouldImport(FString* ActorPropString, bool IsMovingLevel) override;
 	virtual void PostEditImport() override;
 	virtual void PostDuplicate(bool bDuplicateForPIE) override;
diff --git a/Engine/Source/Runtime/Landscape/Classes/LandscapeProxy.h b/Engine/Source/Runtime/Landscape/Classes/LandscapeProxy.h
--- a/Engine/Source/Runtime/Landscape/Classes/LandscapeProxy.h
+++ b/Engine/Source/Runtime/Landscape/Classes/LandscapeProxy.h
@@ -17,6 +17,11 @@
 #include "VT/RuntimeVirtualTextureEnum.h"
 #include "ActorPartition/PartitionActor.h"
 #include "ILandscapeSplineInterface.h"
+
+#if WITH_EDITOR
+#include "WorldPartition/WorldPartitionHandle.h"
+#endif
+
 #include "LandscapeProxy.generated.h"
 
 class ALandscape;
@@ -343,6 +348,20 @@ public:
 	virtual ~ALandscapeProxy();
 
 protected:
+
+#if WITH_EDITORONLY_DATA
+	/** 
+	* Hard refs to actors that need to be loaded when this proxy is loaded.
+	* It is currently used for 2 cases : 
+	* 1- ALandscapeStreamingProxy forces the loading of its intersecting ALandscapeSplineActor.
+	* 2- ALandscape forces the loading of its child landscape proxies when the Landscape has Layer brushes.
+	*    This is a temporary solution until landscape layer brushes support partial landscape loading.
+	*/
+	TSet<FWorldPartitionReference> ActorDescReferences;
+
+	friend class FLandscapeActorDesc;
+#endif
+
 	UPROPERTY()
 	TObjectPtr<ULandscapeSplinesComponent> SplineComponent;
 
diff --git a/Engine/Source/Runtime/Landscape/Classes/LandscapeStreamingProxy.h b/Engine/Source/Runtime/Landscape/Classes/LandscapeStreamingProxy.h
--- a/Engine/Source/Runtime/Landscape/Classes/LandscapeStreamingProxy.h
+++ b/Engine/Source/Runtime/Landscape/Classes/LandscapeStreamingProxy.h
@@ -6,10 +6,6 @@
 #include "UObject/ObjectMacros.h"
 #include "LandscapeProxy.h"
 
-#if WITH_EDITOR
-#include "WorldPartition/WorldPartitionHandle.h"
-#endif
-
 #include "LandscapeStreamingProxy.generated.h"
 
 class ALandscape;
@@ -26,11 +22,6 @@ public:
 	UPROPERTY(EditAnywhere, Category=LandscapeProxy)
 	TLazyObjectPtr<ALandscape> LandscapeActor;
 
-#if WITH_EDITORONLY_DATA
-	/** hard refs to actors that need to be loaded when this proxy is loaded */
-	TSet<FWorldPartitionReference> ActorDescReferences;
-#endif
-
 	//~ Begin UObject Interface
 #if WITH_EDITOR
 	virtual bool ShouldExport() override { return false;  }
diff --git a/Engine/Source/Runtime/Landscape/Private/LandscapeEdit.cpp b/Engine/Source/Runtime/Landscape/Private/LandscapeEdit.cpp
--- a/Engine/Source/Runtime/Landscape/Private/LandscapeEdit.cpp
+++ b/Engine/Source/Runtime/Landscape/Private/LandscapeEdit.cpp
@@ -63,6 +63,7 @@ LandscapeEdit.cpp: Landscape editing
 #include "Editor.h"
 #include "PhysicalMaterials/PhysicalMaterial.h"
 #include "WorldPartition/WorldPartition.h"
+#include "WorldPartition/WorldPartitionHandle.h"
 #include "WorldPartition/WorldPartitionHelpers.h"
 #include "WorldPartition/WorldPartitionActorDesc.h"
 #include "WorldPartition/Landscape/LandscapeActorDesc.h"
@@ -4106,6 +4107,44 @@ void ALandscape::PostEditUndo()
 	RequestLayersContentUpdate(ELandscapeLayerUpdateMode::Update_All);
 }
 
+void ALandscape::PostRegisterAllComponents()
+{
+	Super::PostRegisterAllComponents();
+
+	auto HasValidBrush = [this]()
+	{
+		for (const FLandscapeLayer& Layer : LandscapeLayers)
+		{
+			for (const FLandscapeLayerBrush& Brush : Layer.Brushes)
+			{
+				if (IsValid(Brush.GetBrush()))
+				{
+					return true;
+				}
+			}
+		}
+		return false;
+	};
+
+	// Until it is properly supported, ALandscape with layer brushes will force all of its proxies to be loaded in editor
+	if (GEditor && !GetWorld()->IsGameWorld() && LandscapeGuid.IsValid() && HasValidBrush())
+	{
+		if (UWorldPartition* WorldPartition = GetWorld()->GetWorldPartition())
+		{
+			FWorldPartitionHelpers::ForEachActorDesc<ALandscapeProxy>(WorldPartition, [this, WorldPartition](const FWorldPartitionActorDesc* ActorDesc)
+			{
+				FLandscapeActorDesc* LandscapeActorDesc = (FLandscapeActorDesc*)ActorDesc;
+
+				if (LandscapeActorDesc->GridGuid == LandscapeGuid)
+				{
+					ActorDescReferences.Add(FWorldPartitionReference(WorldPartition, ActorDesc->GetGuid()));
+				}
+				return true;
+			});
+		}
+	}
+}
+
 bool ALandscape::ShouldImport(FString* ActorPropString, bool IsMovingLevel)
 {
 	return GetWorld() != nullptr && !GetWorld()->IsGameWorld();
