Removed async task for computing landscape component LOD biases to fix race condition: the task was started at the beginning of the render thread frame and synced when starting to render view family/views. But the task was operating on UTextures, which could have their internal FTextureResource pointer updated within that lapse of time, leading to a crash if GetNumMips() / GetNumResidentMips() was executed at the same time
Furthermore, there's no per-view async task anymore so there was no point in kicking off that async task in BeginRenderViewFamily_RenderThread (which was the right thing to do, since at this point, FTextureResource pointers couldn't be modified by render commands happening between BeginFrame_RenderThread and BeginRenderViewFamily_RenderThread) and waiting after in that same function since no other parallel work was to be executed anyway. The job was less than 0.1ms on console anyway so the necessity to have this running asynchronously was questionable in the first place.

diff --git a/Engine/Source/Runtime/Landscape/Private/LandscapeRender.cpp b/Engine/Source/Runtime/Landscape/Private/LandscapeRender.cpp
--- a/Engine/Source/Runtime/Landscape/Private/LandscapeRender.cpp
+++ b/Engine/Source/Runtime/Landscape/Private/LandscapeRender.cpp
@@ -77,13 +77,6 @@ FAutoConsoleVariableRef CVarLandscapeLODDistributionScale(
 	ECVF_Scalability
 );
 
-float GShadowMapWorldUnitsToTexelFactor = -1.0f;
-static FAutoConsoleVariableRef CVarShadowMapWorldUnitsToTexelFactor(
-	TEXT("Landscape.ShadowMapWorldUnitsToTexelFactor"),
-	GShadowMapWorldUnitsToTexelFactor,
-	TEXT("Used to specify tolerance factor for mesh size related to cascade shadow resolution")
-);
-
 int32 GAllowLandscapeShadows = 1;
 static FAutoConsoleVariableRef CVarAllowLandscapeShadows(
 	TEXT("r.AllowLandscapeShadows"),
@@ -618,8 +611,8 @@ void ULandscapeComponent::GetUsedMaterials(TArray<UMaterialInterface*>& OutMater
 #endif
 }
 
-/** 
- * Return any global Lod override for landscape. 
+/**
+ * Return any global Lod override for landscape.
  * A return value less than 0 means no override.
  * Any positive value must still be clamped into the valid Lod range for the landscape.
  */
@@ -659,46 +652,6 @@ static int32 GetDrawCollisionLodOverride(FSceneView const& View, FCollisionRespo
 }
 
 
-//
-// FGetSectionLODBiasesTask
-//
-struct FGetSectionLODBiasesTask
-{
-	FLandscapeRenderSystem& RenderSystem;
-
-	FGetSectionLODBiasesTask(FLandscapeRenderSystem& InRenderSystem)
-		: RenderSystem(InRenderSystem)
-	{
-	}
-
-	FORCEINLINE TStatId GetStatId() const
-	{
-		RETURN_QUICK_DECLARE_CYCLE_STAT(FGetSectionLODBiasesTask, STATGROUP_TaskGraphTasks);
-	}
-
-	ENamedThreads::Type GetDesiredThread()
-	{
-		return ENamedThreads::AnyNormalThreadNormalTask;
-	}
-
-	static ESubsequentsMode::Type GetSubsequentsMode()
-	{
-		return ESubsequentsMode::TrackSubsequents;
-	}
-
-	void AnyThreadTask()
-	{
-		FOptionalTaskTagScope Scope(ETaskTag::EParallelRenderingThread);
-		RenderSystem.FetchHeightmapLODBiases();
-	}
-
-	void DoTask(ENamedThreads::Type CurrentThread, const FGraphEventRef& MyCompletionGraphEvent)
-	{
-		AnyThreadTask();
-	}
-};
-
-
 //
 // FLandscapeComponentSceneProxy
 //
@@ -738,8 +691,6 @@ FLandscapeRenderSystem::FLandscapeRenderSystem()
 
 FLandscapeRenderSystem::~FLandscapeRenderSystem()
 {
-	WaitForTasksCompletion();
-
 	check(LandscapeIndexAllocator[LandscapeIndex]);
 	LandscapeIndexAllocator[LandscapeIndex] = false;
 
@@ -866,7 +817,7 @@ void FLandscapeRenderSystem::UnregisterSection(FLandscapeSectionInfo* SectionInf
 {
 	check(IsInRenderingThread());
 	check(SectionInfo != nullptr);
-	
+
 	if (SectionInfo->bRegistered)
 	{
 		FLandscapeRenderSystem* LandscapeRenderSystem = LandscapeRenderSystems.FindChecked(SectionInfo->LandscapeKey);
@@ -915,17 +866,6 @@ void FLandscapeRenderSystem::ResizeAndMoveTo(FIntPoint NewMin, FIntPoint NewSize
 	SectionLODBiases.SetAllowCPUAccess(true);
 }
 
-void FLandscapeRenderSystem::BeginRender()
-{
-	check(IsInRenderingThread());
-
-	TRACE_CPUPROFILER_EVENT_SCOPE(FLandscapeRenderSystem::BeginRender());
-
-	WaitForTasksCompletion();
-
-	UpdateBuffers();
-}
-
 const TResourceArray<float>& FLandscapeRenderSystem::ComputeSectionsLODForView(const FSceneView& InView)
 {
 	TRACE_CPUPROFILER_EVENT_SCOPE(FLandscapeRenderSystem::ComputeSectionsLODForView);
@@ -944,6 +884,16 @@ const TResourceArray<float>& FLandscapeRenderSystem::ComputeSectionsLODForView(c
 	return SectionLODValues;
 }
 
+void FLandscapeRenderSystem::PreRenderViewFamily_RenderThread()
+{
+	TRACE_CPUPROFILER_EVENT_SCOPE(FLandscapeRenderSystem::PreRenderViewFamily_RenderThread);
+	check(IsInRenderingThread());
+
+	CachedSectionLODValues.Reset();
+	FetchHeightmapLODBiases();
+	UpdateBuffers();
+}
+
 void FLandscapeRenderSystem::FetchHeightmapLODBiases()
 {
 	TRACE_CPUPROFILER_EVENT_SCOPE(FLandscapeRenderSystem::FetchHeightmapLODBiases);
@@ -961,7 +911,7 @@ void FLandscapeRenderSystem::UpdateBuffers()
 {
 	check(IsInRenderingThread());
 
-	TRACE_CPUPROFILER_EVENT_SCOPE(FLandscapeRenderSystem::RecreateBuffers());
+	TRACE_CPUPROFILER_EVENT_SCOPE(FLandscapeRenderSystem::UpdateBuffers);
 
 	bool bUpdateUB = false;
 
@@ -994,74 +944,26 @@ void FLandscapeRenderSystem::UpdateBuffers()
 	}
 }
 
-void FLandscapeRenderSystem::BeginFrame()
-{
-	check(IsInRenderingThread());
-
-	CachedSectionLODValues.Reset();
-
-	const bool bExecuteInParallel = FApp::ShouldUseThreadingForPerformance()
-		&& GIsThreadedRendering; // Rendering thread is required to safely use rendering resources in parallel.
-
-	if (bExecuteInParallel)
-	{
-		FetchHeightmapLODBiasesEventRef = TGraphTask<FGetSectionLODBiasesTask>::CreateTask(
-			nullptr, ENamedThreads::GetRenderThread()).ConstructAndDispatchWhenReady(*this);
-	}
-	else
-	{
-		FGetSectionLODBiasesTask Task(*this);
-		Task.AnyThreadTask();
-	}
-}
-
-void FLandscapeRenderSystem::EndFrame()
-{
-	check(IsInRenderingThread());
-
-	WaitForTasksCompletion();
-}
-
-void FLandscapeRenderSystem::WaitForTasksCompletion()
-{
-	if (FetchHeightmapLODBiasesEventRef.IsValid())
-	{
-		FTaskGraphInterface::Get().WaitUntilTaskCompletes(FetchHeightmapLODBiasesEventRef, ENamedThreads::GetRenderThread_Local());
-		FetchHeightmapLODBiasesEventRef.SafeRelease();
-	}
-}
-
 
 //
 // FLandscapeSceneViewExtension
 //
 FLandscapeSceneViewExtension::FLandscapeSceneViewExtension(const FAutoRegister& AutoReg) : FSceneViewExtensionBase(AutoReg)
 {
-	FCoreDelegates::OnBeginFrameRT.AddRaw(this, &FLandscapeSceneViewExtension::BeginFrame_RenderThread);
 	FCoreDelegates::OnEndFrameRT.AddRaw(this, &FLandscapeSceneViewExtension::EndFrame_RenderThread);
 }
 
 FLandscapeSceneViewExtension::~FLandscapeSceneViewExtension()
 {
-	FCoreDelegates::OnBeginFrameRT.RemoveAll(this);
 	FCoreDelegates::OnEndFrameRT.RemoveAll(this);
 }
 
-void FLandscapeSceneViewExtension::BeginFrame_RenderThread()
-{
-	for (auto& Pair : LandscapeRenderSystems)
-	{
-		FLandscapeRenderSystem& RenderSystem = *Pair.Value;
-		RenderSystem.BeginFrame();
-	}
-}
-
 void FLandscapeSceneViewExtension::PreRenderViewFamily_RenderThread(FRHICommandListImmediate& RHICmdList, FSceneViewFamily& InViewFamily)
 {
 	for (auto& Pair : LandscapeRenderSystems)
 	{
 		FLandscapeRenderSystem& RenderSystem = *Pair.Value;
-		RenderSystem.BeginRender();
+		RenderSystem.PreRenderViewFamily_RenderThread();
 	}
 }
 
@@ -1108,15 +1010,10 @@ void FLandscapeSceneViewExtension::PreRenderView_RenderThread(FRHICommandListImm
 	InView.LandscapeIndirectionBuffer = IndirectionSRV;
 }
 
+// TODO [jonathan.bard] Ideally this should be symmetrical with FLandscapeSceneViewExtension::PreRenderView_RenderThread and should be called in FLandscapeSceneViewExtension::PostRenderView_RenderThread
+//  but PostRenderView_RenderThread is not currently called on all code paths (scene/reflection captures, etc.)
 void FLandscapeSceneViewExtension::EndFrame_RenderThread()
 {
-	for (auto& Pair : LandscapeRenderSystems)
-	{
-		FLandscapeRenderSystem& RenderSystem = *Pair.Value;
-
-		RenderSystem.EndFrame();
-	}
-
 	LandscapeLODDataBuffer = FBufferRHIRef();
 	LandscapeIndirectionBuffer = FBufferRHIRef();
 }
@@ -1138,7 +1035,7 @@ bool FLandscapeVisibilityHelper::OnAddedToWorld()
 	{
 		return false;
 	}
-	
+
 	bIsComponentLevelVisible = true;
 	return true;
 }
@@ -1216,7 +1113,7 @@ FLandscapeComponentSceneProxy::FLandscapeComponentSceneProxy(ULandscapeComponent
 	}
 
 	EnableGPUSceneSupportFlags();
-	
+
 	const ERHIFeatureLevel::Type FeatureLevel = GetScene().GetFeatureLevel();
 	if (FeatureLevel >= ERHIFeatureLevel::SM5)
 	{
@@ -1235,7 +1132,7 @@ FLandscapeComponentSceneProxy::FLandscapeComponentSceneProxy(ULandscapeComponent
 	}
 
 	LODIndexToMaterialIndex = InComponent->LODIndexToMaterialIndex;
-	check(LODIndexToMaterialIndex.Num() == MaxLOD+1);
+	check(LODIndexToMaterialIndex.Num() == MaxLOD + 1);
 
 	SetLevelColor(FLinearColor(1.f, 1.f, 1.f));
 
@@ -1281,7 +1178,7 @@ FLandscapeComponentSceneProxy::FLandscapeComponentSceneProxy(ULandscapeComponent
 
 	// Clamp to MaxLODLevel
 	const int32 MaxLODLevel = InComponent->GetLandscapeProxy()->MaxLODLevel;
-	if ( MaxLODLevel >= 0)
+	if (MaxLODLevel >= 0)
 	{
 		MaxLOD = FMath::Min<int8>(MaxLODLevel, MaxLOD);
 		LastLOD = FMath::Min<int32>(MaxLODLevel, LastLOD);
@@ -1494,7 +1391,7 @@ void FLandscapeComponentSceneProxy::CreateRenderThreadResources()
 	LandscapeUniformShaderParameters.InitResource();
 
 	// Create per Lod uniform buffers
-	const int32 NumMips = FMath::CeilLogTwo(SubsectionSizeVerts); 
+	const int32 NumMips = FMath::CeilLogTwo(SubsectionSizeVerts);
 	// create as many as there are potential mips (even if MaxLOD can be inferior than that), because the grass could need that much :
 	LandscapeFixedGridUniformShaderParameters.AddDefaulted(NumMips);
 	for (int32 LodIndex = 0; LodIndex < NumMips; ++LodIndex)
@@ -1502,7 +1399,7 @@ void FLandscapeComponentSceneProxy::CreateRenderThreadResources()
 		LandscapeFixedGridUniformShaderParameters[LodIndex].InitResource();
 		FLandscapeFixedGridUniformShaderParameters Parameters;
 		Parameters.LodValues = FVector4f(
-			LodIndex, 
+			LodIndex,
 			0.f,
 			(float)((SubsectionSizeVerts >> LodIndex) - 1),
 			1.f / (float)((SubsectionSizeVerts >> LodIndex) - 1));
@@ -1519,7 +1416,7 @@ void FLandscapeComponentSceneProxy::CreateRenderThreadResources()
 		GrassMeshBatch.Elements.AddDefaulted(NumMips);
 		GrassBatchParams.Empty(NumMips);
 		GrassBatchParams.AddDefaulted(NumMips);
-		
+
 		// Grass is being generated using LOD0 material only
 		// It uses the fixed grid vertex factory so it doesn't support XY offsets
 		FMaterialRenderProxy* RenderProxy = AvailableMaterials[LODIndexToMaterialIndex[0]]->GetRenderProxy();
@@ -2126,8 +2023,8 @@ void FLandscapeComponentSceneProxy::DrawStaticElements(FStaticPrimitiveDrawInter
 	PDI->ReserveMemoryForMeshes(TotalBatchCount);
 
 	// Add fixed grid mesh batches for runtime virtual texture usage
- 	for (ERuntimeVirtualTextureMaterialType MaterialType : RuntimeVirtualTextureMaterialTypes)
- 	{
+	for (ERuntimeVirtualTextureMaterialType MaterialType : RuntimeVirtualTextureMaterialTypes)
+	{
 		const int32 MaterialIndex = LODIndexToMaterialIndex[FirstLOD];
 
 		for (int32 LODIndex = FirstVirtualTextureLOD; LODIndex <= LastVirtualTextureLOD; ++LODIndex)
@@ -2138,7 +2035,7 @@ void FLandscapeComponentSceneProxy::DrawStaticElements(FStaticPrimitiveDrawInter
 				PDI->DrawMesh(RuntimeVirtualTextureMeshBatch, FLT_MAX);
 			}
 		}
- 	}
+	}
 
 	for (int32 LODIndex = FirstLOD; LODIndex <= LastLOD; LODIndex++)
 	{
@@ -2407,7 +2304,7 @@ void FLandscapeComponentSceneProxy::GetDynamicMeshElements(const TArray<const FS
 
 					FMeshBatch& MaskMesh = Collector.AllocateMesh();
 					MaskMesh = MeshTools;
-										
+
 					auto DirtyMaterialInstance = new FLandscapeMaskMaterialRenderProxy(GLandscapeDirtyMaterial->GetRenderProxy(), EditToolRenderData.DirtyTexture ? ToRawPtr(EditToolRenderData.DirtyTexture) : GLandscapeBlackTexture, true);
 					MaskMesh.MaterialRenderProxy = DirtyMaterialInstance;
 					Collector.RegisterOneFrameMaterialProxy(DirtyMaterialInstance);
@@ -2559,9 +2456,9 @@ void FLandscapeComponentSceneProxy::GetDynamicRayTracingInstances(FRayTracingMat
 
 	const FSceneView& SceneView = *Context.ReferenceView;
 	const FLandscapeRenderSystem& RenderSystem = *LandscapeRenderSystems.FindChecked(LandscapeKey);
-	
+
 	int32 LODToRender = RenderSystem.GetSectionLODValue(SceneView, ComponentBase);
-	
+
 	FLandscapeElementParamArray& ParameterArray = Context.RayTracingMeshResourceCollector.AllocateOneFrameResource<FLandscapeElementParamArray>();
 	ParameterArray.ElementParams.AddDefaulted(NumSubsections * NumSubsections);
 
@@ -3015,8 +2912,8 @@ FLandscapeSharedBuffers::FLandscapeSharedBuffers(const int32 InSharedBuffersKey,
 		// Vertex Buffer cannot be shared
 		VertexBuffer = new FLandscapeVertexBuffer(InFeatureLevel, NumVertices, SubsectionSizeVerts, NumSubsections);
 	}
-	IndexBuffers = new FIndexBuffer*[NumIndexBuffers];
-	FMemory::Memzero(IndexBuffers, sizeof(FIndexBuffer*)* NumIndexBuffers);
+	IndexBuffers = new FIndexBuffer * [NumIndexBuffers];
+	FMemory::Memzero(IndexBuffers, sizeof(FIndexBuffer*) * NumIndexBuffers);
 	IndexRanges = new FLandscapeIndexRanges[NumIndexBuffers]();
 
 #if RHI_RAYTRACING
@@ -3135,7 +3032,7 @@ public:
 
 IMPLEMENT_TYPE_LAYOUT(FLandscapeVertexFactoryVertexShaderParameters);
 
-/** 
+/**
   * Shader parameters for use with FLandscapeFixedGridVertexFactory
   * Simple grid rendering (without dynamic lod blend) needs a simpler fixed setup.
   */
@@ -3246,8 +3143,8 @@ IMPLEMENT_VERTEX_FACTORY_PARAMETER_TYPE(FLandscapeVertexFactory, SF_RayHitGroup,
 #endif // RHI_RAYTRACING
 IMPLEMENT_VERTEX_FACTORY_PARAMETER_TYPE(FLandscapeVertexFactory, SF_Pixel, FLandscapeVertexFactoryPixelShaderParameters);
 
-IMPLEMENT_VERTEX_FACTORY_TYPE(FLandscapeVertexFactory, "/Engine/Private/LandscapeVertexFactory.ush", 
-	  EVertexFactoryFlags::UsedWithMaterials
+IMPLEMENT_VERTEX_FACTORY_TYPE(FLandscapeVertexFactory, "/Engine/Private/LandscapeVertexFactory.ush",
+	EVertexFactoryFlags::UsedWithMaterials
 	| EVertexFactoryFlags::SupportsStaticLighting
 	| EVertexFactoryFlags::SupportsDynamicLighting
 	| EVertexFactoryFlags::SupportsCachingMeshDrawCommands
@@ -3268,9 +3165,9 @@ void FLandscapeVertexFactory::Copy(const FLandscapeVertexFactory& Other)
 	const FDataType* DataCopy = &Other.Data;
 	ENQUEUE_RENDER_COMMAND(FLandscapeVertexFactoryCopyData)(
 		[VertexFactory, DataCopy](FRHICommandListImmediate& RHICmdList)
-		{
-			VertexFactory->Data = *DataCopy;
-		});
+	{
+		VertexFactory->Data = *DataCopy;
+	});
 	BeginUpdateResourceRHI(this);
 }
 
@@ -3291,8 +3188,8 @@ IMPLEMENT_VERTEX_FACTORY_PARAMETER_TYPE(FLandscapeXYOffsetVertexFactory, SF_RayH
 #endif // RHI_RAYTRACING
 IMPLEMENT_VERTEX_FACTORY_PARAMETER_TYPE(FLandscapeXYOffsetVertexFactory, SF_Pixel, FLandscapeVertexFactoryPixelShaderParameters);
 
-IMPLEMENT_VERTEX_FACTORY_TYPE(FLandscapeXYOffsetVertexFactory, "/Engine/Private/LandscapeVertexFactory.ush", 
-	  EVertexFactoryFlags::UsedWithMaterials
+IMPLEMENT_VERTEX_FACTORY_TYPE(FLandscapeXYOffsetVertexFactory, "/Engine/Private/LandscapeVertexFactory.ush",
+	EVertexFactoryFlags::UsedWithMaterials
 	| EVertexFactoryFlags::SupportsStaticLighting
 	| EVertexFactoryFlags::SupportsDynamicLighting
 	| EVertexFactoryFlags::SupportsCachingMeshDrawCommands
@@ -3319,8 +3216,8 @@ IMPLEMENT_VERTEX_FACTORY_PARAMETER_TYPE(FLandscapeFixedGridVertexFactory, SF_Ray
 #endif
 IMPLEMENT_VERTEX_FACTORY_PARAMETER_TYPE(FLandscapeFixedGridVertexFactory, SF_Pixel, FLandscapeVertexFactoryPixelShaderParameters);
 
-IMPLEMENT_VERTEX_FACTORY_TYPE(FLandscapeFixedGridVertexFactory, "/Engine/Private/LandscapeVertexFactory.ush", 
-	  EVertexFactoryFlags::UsedWithMaterials
+IMPLEMENT_VERTEX_FACTORY_TYPE(FLandscapeFixedGridVertexFactory, "/Engine/Private/LandscapeVertexFactory.ush",
+	EVertexFactoryFlags::UsedWithMaterials
 	| EVertexFactoryFlags::SupportsStaticLighting
 	| EVertexFactoryFlags::SupportsDynamicLighting
 	| EVertexFactoryFlags::SupportsCachingMeshDrawCommands
@@ -3482,7 +3379,7 @@ public:
 					{
 						// No ray tracing on thumbnails : we don't need any variation of ray hit group shaders : 
 						const bool bIsRayHitGroupShader = (ShaderType->GetFrequency() == SF_RayHitGroup);
-						if (bIsRayHitGroupShader 
+						if (bIsRayHitGroupShader
 							|| Algo::Find(GetExcludedShaderTypesInThumbnailRender(), ShaderType->GetFName()))
 						{
 							UE_LOG(LogLandscape, VeryVerbose, TEXT("Excluding shader %s from landscape thumbnail material"), ShaderType->GetName());
@@ -3753,7 +3650,7 @@ public:
 			FName(TEXT("FLightmapGBufferVS")),
 			FName(TEXT("FLightmapGBufferPS")),
 		};
-		return ShaderTypes;		
+		return ShaderTypes;
 	}
 
 	static const TArray<FName>& GetGrassShaderTypes()
@@ -3960,7 +3857,7 @@ void ALandscapeProxy::ChangeComponentScreenSizeToUseSubSections(float InComponen
 	if (LandscapeComponents.Num() > 0)
 	{
 		int32 ComponentCount = LandscapeComponents.Num();
-		FLandscapeComponentSceneProxy** RenderProxies = new FLandscapeComponentSceneProxy*[ComponentCount];
+		FLandscapeComponentSceneProxy** RenderProxies = new FLandscapeComponentSceneProxy * [ComponentCount];
 		for (int32 Idx = 0; Idx < ComponentCount; ++Idx)
 		{
 			RenderProxies[Idx] = (FLandscapeComponentSceneProxy*)(LandscapeComponents[Idx]->SceneProxy);
@@ -3969,17 +3866,17 @@ void ALandscapeProxy::ChangeComponentScreenSizeToUseSubSections(float InComponen
 		float ComponentScreenSizeToUseSubSectionsLocal = ComponentScreenSizeToUseSubSections;
 		ENQUEUE_RENDER_COMMAND(LandscapeChangeComponentScreenSizeToUseSubSectionsCommand)(
 			[RenderProxies, ComponentCount, ComponentScreenSizeToUseSubSectionsLocal](FRHICommandListImmediate& RHICmdList)
+		{
+			for (int32 Idx = 0; Idx < ComponentCount; ++Idx)
 			{
-				for (int32 Idx = 0; Idx < ComponentCount; ++Idx)
+				if (RenderProxies[Idx] != nullptr)
 				{
-					if (RenderProxies[Idx] != nullptr)
-					{
-						RenderProxies[Idx]->ChangeComponentScreenSizeToUseSubSections_RenderThread(ComponentScreenSizeToUseSubSectionsLocal);
-					}
+					RenderProxies[Idx]->ChangeComponentScreenSizeToUseSubSections_RenderThread(ComponentScreenSizeToUseSubSectionsLocal);
 				}
-
-				delete[] RenderProxies;
 			}
+
+			delete[] RenderProxies;
+		}
 		);
 	}
 }
@@ -4024,7 +3921,7 @@ void FLandscapeComponentSceneProxy::GetHeightfieldRepresentation(UTexture2D*& Ou
 	OutHeightmapTexture = HeightmapTexture;
 	OutDiffuseColorTexture = BaseColorForGITexture;
 	OutVisibilityTexture = VisibilityWeightmapTexture;
-	
+
 	OutDescription.HeightfieldScaleBias = HeightmapScaleBias;
 
 	OutDescription.MinMaxUV = FVector4f(
@@ -4204,7 +4101,7 @@ SIZE_T FLandscapeMeshProxySceneProxy::GetTypeHash() const
 void FLandscapeMeshProxySceneProxy::CreateRenderThreadResources()
 {
 	FStaticMeshSceneProxy::CreateRenderThreadResources();
-	
+
 	for (auto& Info : ProxySectionsInfos)
 	{
 		FLandscapeRenderSystem::CreateResources(Info.Get());
@@ -4241,7 +4138,7 @@ void FLandscapeMeshProxySceneProxy::DestroyRenderThreadResources()
 {
 	FStaticMeshSceneProxy::DestroyRenderThreadResources();
 	UnregisterSections();
-	
+
 	for (auto& Info : ProxySectionsInfos)
 	{
 		FLandscapeRenderSystem::DestroyResources(Info.Get());
@@ -4250,7 +4147,7 @@ void FLandscapeMeshProxySceneProxy::DestroyRenderThreadResources()
 
 FPrimitiveSceneProxy* ULandscapeMeshProxyComponent::CreateSceneProxy()
 {
-	if (GetStaticMesh() == nullptr 
+	if (GetStaticMesh() == nullptr
 		|| GetStaticMesh()->IsCompiling()
 		|| GetStaticMesh()->GetRenderData() == nullptr
 		|| GetStaticMesh()->GetRenderData()->LODResources.Num() == 0
diff --git a/Engine/Source/Runtime/Landscape/Public/LandscapeRender.h b/Engine/Source/Runtime/Landscape/Public/LandscapeRender.h
--- a/Engine/Source/Runtime/Landscape/Public/LandscapeRender.h
+++ b/Engine/Source/Runtime/Landscape/Public/LandscapeRender.h
@@ -362,8 +362,8 @@ public:
 
 public:
 	uint32		LandscapeKey;
-	FIntPoint	ComponentBase;
-	bool		bRegistered;
+	FIntPoint				ComponentBase;
+	bool					bRegistered;
 };
 
 struct FLandscapeRenderSystem
@@ -422,8 +422,6 @@ struct FLandscapeRenderSystem
 
 	TMap<FViewKey, TResourceArray<float>> CachedSectionLODValues;
 
-	FGraphEventRef FetchHeightmapLODBiasesEventRef;
-	
 	FLandscapeRenderSystem();
 	~FLandscapeRenderSystem();
 
@@ -460,18 +458,9 @@ struct FLandscapeRenderSystem
 	}
 
 	const TResourceArray<float>& ComputeSectionsLODForView(const FSceneView& InView);
-
-	void BeginRender();
-
-	void BeginFrame();
-
 	void FetchHeightmapLODBiases();
-
 	void UpdateBuffers();
-
-	void EndFrame();
-
-	void WaitForTasksCompletion();
+	void PreRenderViewFamily_RenderThread();
 
 private:
 	void CreateResources_Internal(FLandscapeSectionInfo* InSectionInfo);
@@ -488,7 +477,6 @@ public:
 	FLandscapeSceneViewExtension(const FAutoRegister& AutoReg);
 	virtual ~FLandscapeSceneViewExtension();
 
-	void BeginFrame_RenderThread();
 	void EndFrame_RenderThread();
 
 	virtual void SetupViewFamily(FSceneViewFamily& InViewFamily) override {}
