Fix landscape splines to maintain spline mesh vertex color upon minor updates.

diff --git a/Engine/Source/Runtime/Landscape/Private/LandscapeSplines.cpp b/Engine/Source/Runtime/Landscape/Private/LandscapeSplines.cpp
--- a/Engine/Source/Runtime/Landscape/Private/LandscapeSplines.cpp
+++ b/Engine/Source/Runtime/Landscape/Private/LandscapeSplines.cpp
@@ -2814,6 +2814,13 @@ void ULandscapeSplineSegment::UpdateSplinePoints(bool bUpdateCollision, bool bUp
 		TArray<FMeshSettings> MeshSettings;
 		MeshSettings.Reserve(21);
 
+		int32 NumOldSplineMeshComponents = OldLocalMeshComponents.Num();
+		TArray<USplineMeshComponent*>* TmpForeignMeshComponents = ForeignMeshComponentsMap.Find(OuterSplines);
+		if (TmpForeignMeshComponents)
+		{
+			NumOldSplineMeshComponents += TmpForeignMeshComponents->Num();
+		}
+
 		FRandomStream Random(RandomSeed);
 
 		// First pass:
@@ -2823,6 +2830,9 @@ void ULandscapeSplineSegment::UpdateSplinePoints(bool bUpdateCollision, bool bUp
 			const float CosInterp = 0.5f - 0.5f * FMath::Cos(T * PI);
 			const float Width = FMath::Lerp(StartWidth, EndWidth, CosInterp);
 
+			int32 NextLocalMeshComponentIndex = 0;
+			int32 NextForeignMeshComponentIndex = 0;
+
 			const FLandscapeSplineMeshEntry* MeshEntry = UsableMeshes[Random.RandHelper(UsableMeshes.Num())];
 			UStaticMesh* Mesh = MeshEntry->Mesh;
 			const FBoxSphereBounds MeshBounds = Mesh->GetBounds();
@@ -2863,18 +2873,25 @@ void ULandscapeSplineSegment::UpdateSplinePoints(bool bUpdateCollision, bool bUp
 			USplineMeshComponent* MeshComponent = nullptr;
 			if (MeshComponentOuterSplines == OuterSplines)
 			{
-				if (OldLocalMeshComponents.Num() > 0)
+				if (NextLocalMeshComponentIndex < OldLocalMeshComponents.Num())
 				{
-					MeshComponent = OldLocalMeshComponents.Pop(false).Get();
+					// Reuse in the same order so pre-existing static mesh instance vertex color is maintained.
+					MeshComponent = OldLocalMeshComponents[NextLocalMeshComponentIndex].Get();
+					NextLocalMeshComponentIndex++;
 				}
 			}
 			else
 			{
 				TArray<USplineMeshComponent*>* ForeignMeshComponents = ForeignMeshComponentsMap.Find(MeshComponentOuterSplines);
-				if (ForeignMeshComponents && ForeignMeshComponents->Num() > 0)
+				if (ForeignMeshComponents && NextForeignMeshComponentIndex < ForeignMeshComponents->Num())
 				{
+					// Reuse in the same order so pre-existing static mesh instance vertex color may be maintained.
+					// Note: this will only retain the same mesh instance vertex color over the segment if the foreign meshes postceded the 
+					// local meshes along the segment. The relative order would need to be stored to properly handle the case when 
+					// the foreign meshes precede the local meshes (or both precede and postcede).
 					MeshComponentOuterSplines->UpdateModificationKey(this);
-					MeshComponent = ForeignMeshComponents->Pop(false);
+					MeshComponent = (*ForeignMeshComponents)[NextForeignMeshComponentIndex];
+					NextForeignMeshComponentIndex++;
 				}
 			}
 
@@ -2939,6 +2956,9 @@ void ULandscapeSplineSegment::UpdateSplinePoints(bool bUpdateCollision, bool bUp
 		// Add terminating key
 		MeshSettings.Add(FMeshSettings(T, nullptr));
 
+		// Clear vertex color if there is more than one static mesh and the number of spline meshes has changed
+		bool bClearInstanceVertexColors = (UsableMeshes.Num() > 1 && iMesh != NumOldSplineMeshComponents);
+
 		// Second pass:
 		// Rescale components to fit a whole number to the spline, set up final parameters
 		const float Rescale = 1.0f / T;
@@ -3111,6 +3131,11 @@ void ULandscapeSplineSegment::UpdateSplinePoints(bool bUpdateCollision, bool bUp
 			MeshComponent->SetCastShadow(bCastShadow);
 			MeshComponent->InvalidateLightingCache();
 
+			if (bClearInstanceVertexColors)
+			{
+				MeshComponent->RemoveInstanceVertexColors();
+			}
+
 #if WITH_EDITOR
 			if (!bUsingEditorMesh)
 			{
