Landscape: Resolve static analysis warning in LandscapeEditLayers.cpp

diff --git a/Engine/Source/Runtime/Landscape/Private/LandscapeEditLayers.cpp b/Engine/Source/Runtime/Landscape/Private/LandscapeEditLayers.cpp
--- a/Engine/Source/Runtime/Landscape/Private/LandscapeEditLayers.cpp
+++ b/Engine/Source/Runtime/Landscape/Private/LandscapeEditLayers.cpp
@@ -5274,8 +5274,6 @@ void ALandscape::PrepareComponentDataToExtractMaterialLayersCS(const TArray<ULan
 				{
 					if (WeightmapLayerAllocation.LayerInfo != nullptr && WeightmapLayerAllocation.IsAllocated() && ComponentLayerData->WeightmapData.Textures[WeightmapLayerAllocation.WeightmapTextureIndex] == LayerWeightmap)
 					{
-						FLandscapeLayerWeightmapExtractMaterialLayersComponentData Data;
-
 						const ULandscapeComponent* DestComponent = LayerWeightmapUsage->ChannelUsage[WeightmapLayerAllocation.WeightmapTextureChannel];
 						check(DestComponent);
 
@@ -5284,15 +5282,17 @@ void ALandscape::PrepareComponentDataToExtractMaterialLayersCS(const TArray<ULan
 						// Compute component top left vertex position from section base info
 						FVector2D DestPositionOffset(FMath::RoundToInt((float)DestComponentSectionBase.X / LocalComponentSizeQuad), FMath::RoundToInt((float)DestComponentSectionBase.Y / LocalComponentSizeQuad));
 
-						Data.ComponentVertexPosition = SourceComponentVertexPosition;
-						Data.AtlasTexturePositionOutput = FIntPoint(DestPositionOffset.X * LocalComponentSizeVerts, DestPositionOffset.Y * LocalComponentSizeVerts);
-						Data.WeightmapChannelToProcess = WeightmapLayerAllocation.WeightmapTextureChannel;
+						FLandscapeLayerWeightmapExtractMaterialLayersComponentData Data{
+							SourceComponentVertexPosition,
+							0,
+							WeightmapLayerAllocation.WeightmapTextureChannel,
+							FIntPoint(DestPositionOffset.X * LocalComponentSizeVerts, DestPositionOffset.Y * LocalComponentSizeVerts)
+						};
 
 						if (WeightmapLayerAllocation.LayerInfo == ALandscapeProxy::VisibilityLayer)
 						{
-							Data.DestinationPaintLayerIndex = 0;
-							int32& NewLayerinfoObjectIndex = OutLayerInfoObjects.FindOrAdd(ALandscapeProxy::VisibilityLayer);
-							NewLayerinfoObjectIndex = 0;
+							int32& NewLayerInfoObjectIndex = OutLayerInfoObjects.FindOrAdd(ALandscapeProxy::VisibilityLayer);
+							NewLayerInfoObjectIndex = 0;
 						}
 						else
 						{
