Properly account for landscape collision components in memory reports

diff --git a/Engine/Source/Runtime/Engine/Classes/AI/Navigation/NavigationTypes.h b/Engine/Source/Runtime/Engine/Classes/AI/Navigation/NavigationTypes.h
--- a/Engine/Source/Runtime/Engine/Classes/AI/Navigation/NavigationTypes.h
+++ b/Engine/Source/Runtime/Engine/Classes/AI/Navigation/NavigationTypes.h
@@ -780,6 +780,7 @@ struct ENGINE_API FNavHeightfieldSamples
 	TBitArray<> Holes;
 
 	FNavHeightfieldSamples();
+	void GetResourceSizeEx(FResourceSizeEx& CumulativeResourceSize);
 
 	FORCEINLINE bool IsEmpty() const { return Heights.Num() == 0; }
 };
diff --git a/Engine/Source/Runtime/Engine/Private/AI/Navigation/NavigationTypes.cpp b/Engine/Source/Runtime/Engine/Private/AI/Navigation/NavigationTypes.cpp
--- a/Engine/Source/Runtime/Engine/Private/AI/Navigation/NavigationTypes.cpp
+++ b/Engine/Source/Runtime/Engine/Private/AI/Navigation/NavigationTypes.cpp
@@ -243,6 +243,11 @@ FNavHeightfieldSamples::FNavHeightfieldSamples()
 #endif // WITH_PHYSX
 }
 
+void FNavHeightfieldSamples::GetResourceSizeEx(FResourceSizeEx& CumulativeResourceSize)
+{
+	CumulativeResourceSize.AddDedicatedSystemMemoryBytes(sizeof(*this) + Heights.GetAllocatedSize() + Holes.GetAllocatedSize());
+}
+
 //----------------------------------------------------------------------//
 // FNavAgentProperties
 //----------------------------------------------------------------------//
diff --git a/Engine/Source/Runtime/Landscape/Classes/LandscapeHeightfieldCollisionComponent.h b/Engine/Source/Runtime/Landscape/Classes/LandscapeHeightfieldCollisionComponent.h
--- a/Engine/Source/Runtime/Landscape/Classes/LandscapeHeightfieldCollisionComponent.h
+++ b/Engine/Source/Runtime/Landscape/Classes/LandscapeHeightfieldCollisionComponent.h
@@ -97,23 +97,24 @@ class ULandscapeHeightfieldCollisionComponent : public UPrimitiveComponent
 		TArray<physx::PxMaterial*> UsedPhysicalMaterialArray;
 		physx::PxHeightField* RBHeightfield = nullptr;
 		physx::PxHeightField* RBHeightfieldSimple = nullptr;
-#if WITH_EDITOR
+#if WITH_EDITORONLY_DATA
 		physx::PxHeightField* RBHeightfieldEd = nullptr; // Used only by landscape editor, does not have holes in it
-#endif	//WITH_EDITOR
+#endif	//WITH_EDITORONLY_DATA
 #endif	//WITH_PHYSX
 
 #if WITH_CHAOS
 		TArray<Chaos::FMaterialHandle> UsedChaosMaterials;
 		TUniquePtr<Chaos::FHeightField> Heightfield;
 	    TUniquePtr<Chaos::FHeightField> HeightfieldSimple;
-#if WITH_EDITOR
+#if WITH_EDITORONLY_DATA
 		TUniquePtr<Chaos::FHeightField> EditorHeightfield;
-#endif
-#endif
+#endif // WITH_EDITORONLY_DATA
+#endif // WITH_CHAOS
 
 		FHeightfieldGeometryRef(FGuid& InGuid);
 
 		virtual ~FHeightfieldGeometryRef();
+		void GetResourceSizeEx(FResourceSizeEx& CumulativeResourceSize);
 	};
 	
 #if WITH_EDITORONLY_DATA
@@ -221,6 +222,7 @@ public:
 	//~ Begin UObject Interface.
 	virtual void Serialize(FArchive& Ar) override;
 	virtual void BeginDestroy() override;
+	virtual void GetResourceSizeEx(FResourceSizeEx& CumulativeResourceSize) override;
 	virtual void PostLoad() override;
 	PRAGMA_DISABLE_DEPRECATION_WARNINGS // Suppress compiler warning on override of deprecated function
 	UE_DEPRECATED(5.0, "Use version that takes FObjectPreSaveContext instead.")
diff --git a/Engine/Source/Runtime/Landscape/Classes/LandscapeMeshCollisionComponent.h b/Engine/Source/Runtime/Landscape/Classes/LandscapeMeshCollisionComponent.h
--- a/Engine/Source/Runtime/Landscape/Classes/LandscapeMeshCollisionComponent.h
+++ b/Engine/Source/Runtime/Landscape/Classes/LandscapeMeshCollisionComponent.h
@@ -57,31 +57,33 @@ public:
 		/** List of PxMaterials used on this landscape */
 		TArray<physx::PxMaterial*>	UsedPhysicalMaterialArray;
 		physx::PxTriangleMesh*		RBTriangleMesh;
-#if WITH_EDITOR
+#if WITH_EDITORONLY_DATA
 		physx::PxTriangleMesh*		RBTriangleMeshEd; // Used only by landscape editor, does not have holes in it
-#endif	//WITH_EDITOR
+#endif	//WITH_EDITORONLY_DATA
 #endif	//WITH_PHYSX
 
 #if WITH_CHAOS
 		TArray<Chaos::FMaterialHandle> UsedChaosMaterials;
 		TUniquePtr<Chaos::FTriangleMeshImplicitObject> Trimesh;
-#if WITH_EDITOR
+#if WITH_EDITORONLY_DATA
 		TUniquePtr<Chaos::FTriangleMeshImplicitObject> EditorTrimesh;
-#endif // WITH_EDITOR
+#endif // WITH_EDITORONLY_DATA
 #endif // WITH_CHAOS
 
 		FTriMeshGeometryRef();
 		FTriMeshGeometryRef(FGuid& InGuid);
 		virtual ~FTriMeshGeometryRef();
+
+		void GetResourceSizeEx(FResourceSizeEx& CumulativeResourceSize);
 	};
 
 #if WITH_EDITORONLY_DATA
 	/** The collision mesh values. */
-	FWordBulkData								CollisionXYOffsetData; //  X, Y Offset in raw format...
+	FWordBulkData CollisionXYOffsetData; //  X, Y Offset in raw format...
 #endif //WITH_EDITORONLY_DATA
 
 	/** Physics engine version of heightfield data. */
-	TRefCountPtr<FTriMeshGeometryRef>			MeshRef;
+	TRefCountPtr<FTriMeshGeometryRef> MeshRef;
 
 	//~ Begin UActorComponent Interface.
 protected:
@@ -101,6 +103,7 @@ public:
 	//~ Begin UObject Interface.
 	virtual void Serialize(FArchive& Ar) override;
 	virtual void BeginDestroy() override;
+	virtual void GetResourceSizeEx(FResourceSizeEx& CumulativeResourceSize) override;
 #if WITH_EDITOR
 	virtual void ExportCustomProperties(FOutputDevice& Out, uint32 Indent) override;
 	virtual void ImportCustomProperties(const TCHAR* SourceText, FFeedbackContext* Warn) override;
diff --git a/Engine/Source/Runtime/Landscape/Private/LandscapeCollision.cpp b/Engine/Source/Runtime/Landscape/Private/LandscapeCollision.cpp
--- a/Engine/Source/Runtime/Landscape/Private/LandscapeCollision.cpp
+++ b/Engine/Source/Runtime/Landscape/Private/LandscapeCollision.cpp
@@ -48,7 +48,7 @@
 
 #if WITH_EDITOR && PHYSICS_INTERFACE_PHYSX
 	#include "IPhysXCooking.h"
-#endif
+#endif // WITH_EDITOR && PHYSICS_INTERFACE_PHYSX
 
 #if WITH_CHAOS
 #include "PhysXToChaosUtil.h"
@@ -62,7 +62,7 @@
 #include "PhysicsEngine/Experimental/ChaosCooking.h"
 #include "Chaos/ChaosArchive.h"
 #include "PhysicsProxy/SingleParticlePhysicsProxy.h"
-#endif
+#endif // WITH_CHAOS
 
 using namespace PhysicsInterfaceTypes;
 
@@ -98,27 +98,59 @@ ULandscapeHeightfieldCollisionComponent::FHeightfieldGeometryRef::~FHeightfieldG
 		GPhysXPendingKillHeightfield.Add(RBHeightfield);
 		RBHeightfield = NULL;
 	}
-#if WITH_EDITOR
+#if WITH_EDITORONLY_DATA
 	if (RBHeightfieldEd)
 	{
 		GPhysXPendingKillHeightfield.Add(RBHeightfieldEd);
 		RBHeightfieldEd = NULL;
 	}
-#endif// WITH_EDITOR
+#endif// WITH_EDITORONLY_DATA
 #endif// WITH_PHYSX
 
 	// Remove ourselves from the shared map.
 	GSharedHeightfieldRefs.Remove(Guid);
 }
 
+void ULandscapeHeightfieldCollisionComponent::FHeightfieldGeometryRef::GetResourceSizeEx(FResourceSizeEx& CumulativeResourceSize)
+{
+	CumulativeResourceSize.AddDedicatedSystemMemoryBytes(sizeof(*this));
+
+#if WITH_PHYSX
+	CumulativeResourceSize.AddDedicatedSystemMemoryBytes(UsedPhysicalMaterialArray.GetAllocatedSize());
+	// COMMENT [jonathan.bard] : unaccounted for : RBHeightfield, RBHeightfieldSimple 
+#endif //WITH_PHYSX
+
+#if WITH_CHAOS
+	CumulativeResourceSize.AddDedicatedSystemMemoryBytes(UsedChaosMaterials.GetAllocatedSize());
+
+	if (Heightfield.IsValid())
+	{
+		TArray<uint8> Data;
+		FMemoryWriter MemAr(Data);
+		Chaos::FChaosArchive ChaosAr(MemAr);
+		Heightfield->Serialize(ChaosAr);
+		CumulativeResourceSize.AddDedicatedSystemMemoryBytes(Data.Num());
+	}
+
+	if (HeightfieldSimple.IsValid())
+	{
+		TArray<uint8> Data;
+		FMemoryWriter MemAr(Data);
+		Chaos::FChaosArchive ChaosAr(MemAr);
+		HeightfieldSimple->Serialize(ChaosAr);
+		CumulativeResourceSize.AddDedicatedSystemMemoryBytes(Data.Num());
+	}
+#endif // WITH_CHAOS
+}
+
 TMap<FGuid, ULandscapeMeshCollisionComponent::FTriMeshGeometryRef* > GSharedMeshRefs;
 
 ULandscapeMeshCollisionComponent::FTriMeshGeometryRef::FTriMeshGeometryRef()
 #if WITH_PHYSX
 	:	RBTriangleMesh(nullptr)
-#if WITH_EDITOR
+#if WITH_EDITORONLY_DATA
 	, RBTriangleMeshEd(nullptr)
-#endif	//WITH_EDITOR
+#endif	//WITH_EDITORONLY_DATA
 #endif	//WITH_PHYSX
 {}
 
@@ -126,9 +158,9 @@ ULandscapeMeshCollisionComponent::FTriMeshGeometryRef::FTriMeshGeometryRef(FGuid
 	: Guid(InGuid)
 #if WITH_PHYSX
 	, RBTriangleMesh(nullptr)
-#if WITH_EDITOR
+#if WITH_EDITORONLY_DATA
 	, RBTriangleMeshEd(nullptr)
-#endif	//WITH_EDITOR
+#endif	//WITH_EDITORONLY_DATA
 #endif	//WITH_PHYSX
 {}
 
@@ -142,19 +174,42 @@ ULandscapeMeshCollisionComponent::FTriMeshGeometryRef::~FTriMeshGeometryRef()
 		RBTriangleMesh = nullptr;
 	}
 
-#if WITH_EDITOR
+#if WITH_EDITORONLY_DATA
 	if (RBTriangleMeshEd)
 	{
 		GPhysXPendingKillTriMesh.Add(RBTriangleMeshEd);
 		RBTriangleMeshEd = nullptr;
 	}
-#endif// WITH_EDITOR
+#endif// WITH_EDITORONLY_DATA
 #endif// WITH_PHYSX
 
 	// Remove ourselves from the shared map.
 	GSharedMeshRefs.Remove(Guid);
 }
 
+void ULandscapeMeshCollisionComponent::FTriMeshGeometryRef::GetResourceSizeEx(FResourceSizeEx& CumulativeResourceSize)
+{
+	CumulativeResourceSize.AddDedicatedSystemMemoryBytes(sizeof(*this));
+
+#if WITH_PHYSX
+	CumulativeResourceSize.AddDedicatedSystemMemoryBytes(UsedPhysicalMaterialArray.GetAllocatedSize());
+	// COMMENT [jonathan.bard] : unaccounted for : RBTriangleMesh 
+#endif //WITH_PHYSX
+
+#if WITH_CHAOS
+	CumulativeResourceSize.AddDedicatedSystemMemoryBytes(UsedChaosMaterials.GetAllocatedSize());
+
+	if (Trimesh.IsValid())
+	{
+		TArray<uint8> Data;
+		FMemoryWriter MemAr(Data);
+		Chaos::FChaosArchive ChaosAr(MemAr);
+		Trimesh->Serialize(ChaosAr);
+		CumulativeResourceSize.AddDedicatedSystemMemoryBytes(Data.Num());
+	}
+#endif // WITH_CHAOS
+}
+
 // Generate a new guid to force a recache of landscape collison derived data
 #define LANDSCAPE_COLLISION_DERIVEDDATA_VER	TEXT("AB01CAE4167F48E3A3ED396A900F77AC")
 
@@ -500,7 +555,7 @@ void ULandscapeHeightfieldCollisionComponent::OnCreatePhysicsState()
 					Geoms.Emplace(MoveTemp(ChaosEditorHeightFieldFromCooked));
 					ShapeArray.Emplace(MoveTemp(NewEditorShape));
 				}
-#endif
+#endif // WITH_EDITOR
 				// Push the shapes to the actor
 				if(Geoms.Num() == 1)
 				{
@@ -598,7 +653,7 @@ void ULandscapeHeightfieldCollisionComponent::CreateCollisionObject()
 
 #if WITH_EDITOR
 		const bool bNeedsEditorHeightField = !World->IsGameWorld() && !GetOutermost()->bIsCookedForEditor;
-#endif
+#endif // WITH_EDITOR
 		FHeightfieldGeometryRef* ExistingHeightfieldRef = nullptr;
 		bool bCheckDDC = true;
 
@@ -618,11 +673,11 @@ void ULandscapeHeightfieldCollisionComponent::CreateCollisionObject()
 		if (ExistingHeightfieldRef && (!bNeedsEditorHeightField
 #if WITH_CHAOS
 			|| ExistingHeightfieldRef->EditorHeightfield != nullptr
-#endif
+#endif // WITH_CHAOS
 			))
-#else
+#else // WITH_EDITOR
 		if (ExistingHeightfieldRef)
-#endif
+#endif // !WITH_EDITOR
 		{
 			HeightfieldRef = ExistingHeightfieldRef;
 		}
@@ -736,7 +791,7 @@ void ULandscapeHeightfieldCollisionComponent::CreateCollisionObject()
 					}
 				}
 #endif //WITH_EDITOR
-#endif
+#endif // WITH_CHAOS
 			}
 		}
 	}
@@ -760,7 +815,7 @@ void ULandscapeHeightfieldCollisionComponent::SpeculativelyLoadAsyncDDCCollsionD
 			World->AsyncPreRegisterDDCRequests.Add(SpeculativeDDCRequest);
 		}
 	}
-#endif
+#endif // WITH_PHYSX
 }
 
 #if PHYSICS_INTERFACE_PHYSX
@@ -1050,7 +1105,7 @@ bool ULandscapeHeightfieldCollisionComponent::CookCollisionData(const FName& For
 	}
 
 	Succeeded = true;
-#endif
+#endif // WITH_CHAOS
 
 	if (CollisionHeightData.IsLocked())
 	{
@@ -1290,7 +1345,7 @@ bool ULandscapeMeshCollisionComponent::CookCollisionData(const FName& Format, bo
 		ChaosAr << Trimesh;
 		Result = OutData.Num() > 0;
 	}
-#endif
+#endif // WITH_CHAOS
 
 	if (Result)
 	{
@@ -1352,7 +1407,7 @@ void ULandscapeMeshCollisionComponent::CreateCollisionObject()
 			// Create cooked physics data
 			static FName PhysicsFormatName(FPlatformProperties::GetPhysicsFormat());
 			CookCollisionData(PhysicsFormatName, false, bCheckDDC, CookedCollisionData, CookedPhysicalMaterials);
-#endif //WITH_EDITOR
+#endif // WITH_EDITOR
 
 			if (CookedCollisionData.Num())
 			{
@@ -1366,7 +1421,7 @@ void ULandscapeMeshCollisionComponent::CreateCollisionObject()
 				FMemoryReader Reader(CookedCollisionData);
 				Chaos::FChaosArchive Ar(Reader);
 				Ar << MeshRef->Trimesh;
-#endif
+#endif // WITH_CHAOS
 
 				for (UPhysicalMaterial* PhysicalMaterial : CookedPhysicalMaterials)
 				{
@@ -1398,7 +1453,7 @@ void ULandscapeMeshCollisionComponent::CreateCollisionObject()
 						FMemoryReader EdReader(CookedCollisionData);
 						Chaos::FChaosArchive EdAr(EdReader);
 						EdAr << MeshRef->EditorTrimesh;
-#endif
+#endif // WITH_CHAOS
 					}
 				}
 #endif //WITH_EDITOR
@@ -1439,7 +1494,7 @@ struct FMeshCollisionInitHelper
 		PxGeomEd.scale.scale.x = ComponentScale.X * CollisionScale;
 		PxGeomEd.scale.scale.y = ComponentScale.Y * CollisionScale;
 		PxGeomEd.scale.scale.z = ComponentScale.Z;
-#endif
+#endif // WITH_EDITOR
 
 		check(World);
 		PhysScene = World->GetPhysicsScene();
@@ -1529,7 +1584,7 @@ struct FMeshCollisionInitHelper
 			PActor->attachShape(*NewShape);
 			NewShape->release();
 		}
-#endif
+#endif // WITH_EDITOR
 
 		// Set body instance data
 		TargetInstance->PhysicsUserData = FPhysicsUserData(TargetInstance);
@@ -1661,7 +1716,7 @@ struct FMeshCollisionInitHelper
 			Geometries.Emplace(MoveTemp(ScaledTrimeshEd));
 			ShapeArray.Emplace(MoveTemp(NewEdShape));
 		}
-#endif
+#endif // WITH_EDITOR
 
 		if(Geometries.Num() == 1)
 		{
@@ -1721,7 +1776,7 @@ private:
 
 	FPhysicsActorHandle ActorHandle;
 };
-#endif
+#endif // WITH_CHAOS
 
 void ULandscapeMeshCollisionComponent::OnCreatePhysicsState()
 {
@@ -1772,7 +1827,7 @@ void ULandscapeMeshCollisionComponent::OnCreatePhysicsState()
 						QueryFilterDataEd.Word3 |= (EPDF_SimpleCollision | EPDF_ComplexCollision);
 
 				Initializer.SetEditorFilter(QueryFilterDataEd);
-#endif
+#endif // WITH_EDITOR
 
 				Initializer.CreateActors();
 				Initializer.AddToScene();
@@ -1955,7 +2010,6 @@ void ULandscapeHeightfieldCollisionComponent::UpdateHeightfieldRegion(int32 Comp
 
 			HeightfieldRef->EditorHeightfield->EditHeights(Samples, HeightfieldY1, HeightfieldX1, DstVertsY, DstVertsX);
 
-#if WITH_CHAOS
 			// Rebuild geometry to update local bounds, and update in acceleration structure.
 			const Chaos::FImplicitObjectUnion& Union = PhysActorHandle->GetGameThreadAPI().Geometry()->GetObjectChecked<Chaos::FImplicitObjectUnion>();
 			TArray<TUniquePtr<Chaos::FImplicitObject>> NewGeometry;
@@ -1968,8 +2022,7 @@ void ULandscapeHeightfieldCollisionComponent::UpdateHeightfieldRegion(int32 Comp
 
 			FPhysScene* PhysScene = GetWorld()->GetPhysicsScene();
 			PhysScene->UpdateActorInAccelerationStructure(PhysActorHandle);
-#endif
-#endif
+#endif // WITH_CHAOS
 		});
 	}
 
@@ -2022,7 +2075,7 @@ bool ULandscapeHeightfieldCollisionComponent::RecreateCollision()
 			return false;
 		}
 		CollisionHash = NewHash;
-#endif
+#endif // WITH_EDITOR
 		HeightfieldRef = NULL;
 		HeightfieldGuid = FGuid();
 
@@ -2309,7 +2362,7 @@ bool ULandscapeHeightfieldCollisionComponent::ComponentIsTouchingSelectionFrustu
 
 	return false;
 }
-#endif
+#endif // WITH_EDITOR
 
 bool ULandscapeHeightfieldCollisionComponent::DoCustomNavigableGeometryExport(FNavigableGeometryExport& GeomExport) const
 {
@@ -2346,7 +2399,7 @@ bool ULandscapeHeightfieldCollisionComponent::DoCustomNavigableGeometryExport(FN
 			GeomExport.ExportChaosHeightField(HeightfieldRef->Heightfield.Get(), HFToW);
 		}
 	}
-#endif
+#endif // WITH_CHAOS
 
 	return false;
 }
@@ -2363,7 +2416,7 @@ void ULandscapeHeightfieldCollisionComponent::GatherGeometrySlice(FNavigableGeom
 		GeomExport.ExportPxHeightFieldSlice(CachedHeightFieldSamples, HeightfieldRowsCount, HeightfieldColumnsCount, HFToW, SliceBox);
 #elif WITH_CHAOS
 		GeomExport.ExportChaosHeightFieldSlice(CachedHeightFieldSamples, HeightfieldRowsCount, HeightfieldColumnsCount, HFToW, SliceBox);
-#endif
+#endif // WITH_CHAOS
 	}
 }
 
@@ -2436,8 +2489,7 @@ void ULandscapeHeightfieldCollisionComponent::PrepareGeometryExportSync()
 			}
 		}
 	}
-
-#endif// WITH_PHYSX
+#endif // WITH_CHAOS
 }
 
 bool ULandscapeMeshCollisionComponent::DoCustomNavigableGeometryExport(FNavigableGeometryExport& GeomExport) const
@@ -2466,7 +2518,7 @@ bool ULandscapeMeshCollisionComponent::DoCustomNavigableGeometryExport(FNavigabl
 		{
 			GeomExport.ExportChaosTriMesh(MeshRef->Trimesh.Get(), MeshToW);
 		}
-#endif
+#endif // WITH_CHAOS
 	}
 
 	return false;
@@ -2507,7 +2559,7 @@ void ULandscapeHeightfieldCollisionComponent::PostLoad()
 			SpeculativelyLoadAsyncDDCCollsionData();
 		}
 	}
-#endif//WITH_EDITOR
+#endif // WITH_EDITOR
 }
 
 void ULandscapeHeightfieldCollisionComponent::PreSave(const class ITargetPlatform* TargetPlatform)
@@ -3014,7 +3066,7 @@ TOptional<float> ULandscapeHeightfieldCollisionComponent::GetHeight(float X, flo
 	{
 		Height = HeightfieldRef->Heightfield->GetHeightAt({ X, Y });
 	}
-#endif
+#endif // WITH_CHAOS
 	return Height;
 }
 
@@ -3028,7 +3080,7 @@ struct FHeightFieldAccessor
 #elif WITH_CHAOS
 	, NumX(InGeometryRef.Heightfield.IsValid() ? InGeometryRef.Heightfield->GetNumCols() : 0)
 	, NumY(InGeometryRef.Heightfield.IsValid() ? InGeometryRef.Heightfield->GetNumRows() : 0)
-#endif
+#endif // WITH_CHAOS
 	{
 #if WITH_PHYSX && PHYSICS_INTERFACE_PHYSX
 		const int32 CellCount = NumX * NumY;
@@ -3037,7 +3089,7 @@ struct FHeightFieldAccessor
 			HFSamples.SetNumUninitialized(CellCount);
 			GeometryRef.RBHeightfield->saveCells(HFSamples.GetData(), CellCount * HFSamples.GetTypeSize());
 		}
-#endif
+#endif // WITH_PHYSX && PHYSICS_INTERFACE_PHYSX
 	}
 
 	float GetUnscaledHeight(int32 X, int32 Y) const
@@ -3046,7 +3098,7 @@ struct FHeightFieldAccessor
 		return float(HFSamples[NumX * (NumY - 1 - X) + Y].height);
 #elif WITH_CHAOS
 		return GeometryRef.Heightfield->GetHeight(X, Y);
-#else
+#else // WITH_CHAOS
 		return 0.0f;
 #endif
 	}
@@ -3058,7 +3110,7 @@ struct FHeightFieldAccessor
 		return HFSamples[NumX * (NumY - 1 - X) + Y].materialIndex0;
 #elif WITH_CHAOS
 		return GeometryRef.Heightfield->GetMaterialIndex(X, Y);
-#else
+#else // WITH_CHAOS
 		return 0;
 #endif
 	}
@@ -3070,7 +3122,7 @@ struct FHeightFieldAccessor
 #if WITH_PHYSX && PHYSICS_INTERFACE_PHYSX
 private:
 	TArray<PxHeightFieldSample> HFSamples;
-#endif
+#endif // WITH_PHYSX && PHYSICS_INTERFACE_PHYSX
 };
 
 bool ULandscapeHeightfieldCollisionComponent::FillHeightTile(TArrayView<float> Heights, int32 Offset, int32 Stride) const
@@ -3137,6 +3189,31 @@ bool ULandscapeHeightfieldCollisionComponent::FillMaterialIndexTile(TArrayView<u
 	return true;
 }
 
+void ULandscapeHeightfieldCollisionComponent::GetResourceSizeEx(FResourceSizeEx& CumulativeResourceSize)
+{
+	Super::GetResourceSizeEx(CumulativeResourceSize);
+
+	CumulativeResourceSize.AddDedicatedSystemMemoryBytes(sizeof(CookedCollisionData) + CookedCollisionData.GetAllocatedSize() + sizeof(HeightfieldRowsCount) + sizeof(HeightfieldColumnsCount));
+	
+	if (IsValidRef(HeightfieldRef))
+	{
+		HeightfieldRef->GetResourceSizeEx(CumulativeResourceSize);
+	}
+
+	CachedHeightFieldSamples.GetResourceSizeEx(CumulativeResourceSize);
+}
+
+void ULandscapeMeshCollisionComponent::GetResourceSizeEx(FResourceSizeEx& CumulativeResourceSize)
+{
+	Super::GetResourceSizeEx(CumulativeResourceSize);
+
+	if (IsValidRef(MeshRef))
+	{
+		MeshRef->GetResourceSizeEx(CumulativeResourceSize);
+	}
+}
+
+
 TOptional<float> ALandscapeProxy::GetHeightAtLocation(FVector Location) const
 {
 	TOptional<float> Height;
@@ -3248,5 +3325,5 @@ void ALandscapeProxy::GetHeightValues(int32& SizeX, int32& SizeY, TArray<float>
 			}
 		}
 	}
-#endif
-}
\ No newline at end of file
+#endif // WITH_CHAOS
+}
