Landscape: Change LogLandscape default verbosity from Warning to Log

diff --git a/Engine/Source/Runtime/Landscape/Private/LandscapeCollision.cpp b/Engine/Source/Runtime/Landscape/Private/LandscapeCollision.cpp
--- a/Engine/Source/Runtime/Landscape/Private/LandscapeCollision.cpp
+++ b/Engine/Source/Runtime/Landscape/Private/LandscapeCollision.cpp
@@ -1834,7 +1834,7 @@ void ULandscapeMeshCollisionComponent::OnCreatePhysicsState()
 			}
 			else
 			{
-				UE_LOG(LogLandscape, Log, TEXT("ULandscapeMeshCollisionComponent::OnCreatePhysicsState(): TriMesh invalid"));
+				UE_LOG(LogLandscape, Warning, TEXT("ULandscapeMeshCollisionComponent::OnCreatePhysicsState(): TriMesh invalid"));
 			}
 		}
 	}
diff --git a/Engine/Source/Runtime/Landscape/Private/LandscapeEdit.cpp b/Engine/Source/Runtime/Landscape/Private/LandscapeEdit.cpp
--- a/Engine/Source/Runtime/Landscape/Private/LandscapeEdit.cpp
+++ b/Engine/Source/Runtime/Landscape/Private/LandscapeEdit.cpp
@@ -342,7 +342,7 @@ UMaterialInstanceConstant* ULandscapeComponent::GetCombinationMaterial(FMaterial
 			ULandscapeMaterialInstanceConstant* LandscapeCombinationMaterialInstance = NewObject<ULandscapeMaterialInstanceConstant>(GetOuter());
 			LandscapeCombinationMaterialInstance->bMobile = bMobile;
 			CombinationMaterialInstance = LandscapeCombinationMaterialInstance;
-			UE_LOG(LogLandscape, Log, TEXT("Looking for key %s, making new combination %s"), *LayerKey, *CombinationMaterialInstance->GetName());
+			UE_LOG(LogLandscape, Verbose, TEXT("Looking for key %s, making new combination %s"), *LayerKey, *CombinationMaterialInstance->GetName());
 			Proxy->MaterialInstanceConstantMap.Add(*LayerKey, CombinationMaterialInstance);
 			CombinationMaterialInstance->SetParentEditorOnly(MaterialToUse, false);
 
@@ -2743,7 +2743,7 @@ LANDSCAPE_API void ALandscapeProxy::Import(const FGuid& InGuid, int32 InMinX, in
 			}
 
 
-			UE_LOG(LogLandscape, Log, TEXT("%s needs %d alphamaps"), *LandscapeComponent->GetName(), EditingAlphaLayerData.Num());
+			UE_LOG(LogLandscape, VeryVerbose, TEXT("%s needs %d alphamaps"), *LandscapeComponent->GetName(), EditingAlphaLayerData.Num());
 
 			TArray<FWeightmapLayerAllocationInfo>& ComponentWeightmapLayerAllocations = LandscapeComponent->GetWeightmapLayerAllocations();
 
@@ -2887,7 +2887,7 @@ LANDSCAPE_API void ALandscapeProxy::Import(const FGuid& InGuid, int32 InMinX, in
 			// Pointers to the texture data where we'll store each layer. Stride is 4 (FColor)
 			TArray<uint8*> WeightmapTextureDataPointers;
 
-			UE_LOG(LogLandscape, Log, TEXT("%s needs %d weightmap channels"), *LandscapeComponent->GetName(), WeightValues.Num());
+			UE_LOG(LogLandscape, VeryVerbose, TEXT("%s needs %d weightmap channels"), *LandscapeComponent->GetName(), WeightValues.Num());
 
 			// Find texture channels to store each layer.
 			int32 LayerIndex = 0;
@@ -2926,7 +2926,7 @@ LANDSCAPE_API void ALandscapeProxy::Import(const FGuid& InGuid, int32 InMinX, in
 					ULandscapeWeightmapUsage* WeightmapUsage = WeightmapUsageMap.FindChecked(Allocation.Texture);
 					ComponentWeightmapTexturesUsage.Add(WeightmapUsage);
 
-					UE_LOG(LogLandscape, Log, TEXT("  ==> Storing %d channels starting at %s[%d]"), RemainingLayers, *Allocation.Texture->GetName(), Allocation.ChannelsInUse);
+					UE_LOG(LogLandscape, VeryVerbose, TEXT("  ==> Storing %d channels starting at %s[%d]"), RemainingLayers, *Allocation.Texture->GetName(), Allocation.ChannelsInUse);
 
 					for (int32 i = 0; i < RemainingLayers; i++)
 					{
@@ -2966,7 +2966,7 @@ LANDSCAPE_API void ALandscapeProxy::Import(const FGuid& InGuid, int32 InMinX, in
 					ULandscapeWeightmapUsage* WeightmapUsage = WeightmapUsageMap.Add(WeightmapTexture, CreateWeightmapUsage());
 					ComponentWeightmapTexturesUsage.Add(WeightmapUsage);
 
-					UE_LOG(LogLandscape, Log, TEXT("  ==> Storing %d channels in new texture %s"), ThisAllocationLayers, *WeightmapTexture->GetName());
+					UE_LOG(LogLandscape, VeryVerbose, TEXT("  ==> Storing %d channels in new texture %s"), ThisAllocationLayers, *WeightmapTexture->GetName());
 
 					WeightmapTextureDataPointers.Add((uint8*)&MipData->R);
 					ComponentWeightmapLayerAllocations[LayerIndex + 0].WeightmapTextureIndex = ComponentWeightmapTextures.Num();
@@ -7677,7 +7677,7 @@ bool ALandscapeProxy::LandscapeImportWeightmapFromRenderTarget(UTextureRenderTar
 				AlphamapAccessor.SetData(MinX, MinY, MaxX, MaxY, LayerData.GetData(), ELandscapeLayerPaintingRestriction::None);
 
 				uint64 CycleEnd = FPlatformTime::Cycles64();
-				UE_LOG(LogLandscape, Log, TEXT("Took %f seconds to import heightmap from render target"), FPlatformTime::ToSeconds64(CycleEnd));
+				UE_LOG(LogLandscape, Verbose, TEXT("Took %f seconds to import heightmap from render target"), FPlatformTime::ToSeconds64(CycleEnd));
 
 				return true;
 			}
diff --git a/Engine/Source/Runtime/Landscape/Private/LandscapeEditLayers.cpp b/Engine/Source/Runtime/Landscape/Private/LandscapeEditLayers.cpp
--- a/Engine/Source/Runtime/Landscape/Private/LandscapeEditLayers.cpp
+++ b/Engine/Source/Runtime/Landscape/Private/LandscapeEditLayers.cpp
@@ -219,7 +219,7 @@ private:
 			}
 		}
 
-		UE_LOG(LogLandscape, Display, TEXT("Landscape.Dirty: %s"), bCleared ? TEXT("Cleared") : TEXT("Landscape.Dirty: Nothing to clear"));
+		UE_LOG(LogLandscape, Verbose, TEXT("Landscape.Dirty: %s"), bCleared ? TEXT("Cleared") : TEXT("Landscape.Dirty: Nothing to clear"));
 	}
 };
 
diff --git a/Engine/Source/Runtime/Landscape/Private/LandscapePrivate.h b/Engine/Source/Runtime/Landscape/Private/LandscapePrivate.h
--- a/Engine/Source/Runtime/Landscape/Private/LandscapePrivate.h
+++ b/Engine/Source/Runtime/Landscape/Private/LandscapePrivate.h
@@ -6,7 +6,7 @@
 #include "Stats/Stats.h"
 
 
-LANDSCAPE_API DECLARE_LOG_CATEGORY_EXTERN(LogLandscape, Warning, All);
+LANDSCAPE_API DECLARE_LOG_CATEGORY_EXTERN(LogLandscape, Log, All);
 LANDSCAPE_API DECLARE_LOG_CATEGORY_EXTERN(LogLandscapeBP, Display, All);
 
 /**
