Landscape: Resurrect Landscape.DumpLODs console command functionality

diff --git a/Engine/Source/Runtime/Landscape/Private/LandscapeRender.cpp b/Engine/Source/Runtime/Landscape/Private/LandscapeRender.cpp
--- a/Engine/Source/Runtime/Landscape/Private/LandscapeRender.cpp
+++ b/Engine/Source/Runtime/Landscape/Private/LandscapeRender.cpp
@@ -92,17 +92,14 @@ extern RENDERER_API TAutoConsoleVariable<float> CVarStaticMeshLODDistanceScale;
 
 #if !UE_BUILD_SHIPPING
 int32 GVarDumpLandscapeLODsCurrentFrame = 0;
-bool GVarDumpLandscapeLODs = false;
 
 static void OnDumpLandscapeLODs(const TArray< FString >& Args)
 {
 	if (Args.Num() >= 1)
 	{
-		GVarDumpLandscapeLODs = FCString::Atoi(*Args[0]) == 0 ? false : true;
+		// Add some buffer to be able to correctly catch the frame during the rendering
+		GVarDumpLandscapeLODsCurrentFrame = FCString::Atoi(*Args[0]) != 0 ? GFrameNumberRenderThread + 3 : INDEX_NONE;
 	}
-
-	// Add some buffer to be able to correctly catch the frame during the rendering
-	GVarDumpLandscapeLODsCurrentFrame = GVarDumpLandscapeLODs ? GFrameNumberRenderThread + 3 : INDEX_NONE;
 }
 
 static FAutoConsoleCommand CVarDumpLandscapeLODs(
@@ -1699,6 +1696,14 @@ FPrimitiveViewRelevance FLandscapeComponentSceneProxy::GetViewRelevance(const FS
 	}
 
 	Result.bShadowRelevance = (GAllowLandscapeShadows > 0) && IsShadowCast(View) && View->Family->EngineShowFlags.Landscape;
+
+#if !UE_BUILD_SHIPPING
+	if (GVarDumpLandscapeLODsCurrentFrame == GFrameNumberRenderThread)
+	{
+		Result.bDynamicRelevance = true;
+	}
+#endif // !UE_BUILD_SHIPPING
+
 	return Result;
 }
 
@@ -2447,6 +2452,44 @@ void FLandscapeComponentSceneProxy::GetDynamicMeshElements(const TArray<const FS
 	INC_DWORD_STAT_BY(STAT_LandscapeComponentRenderPasses, NumPasses);
 	INC_DWORD_STAT_BY(STAT_LandscapeDrawCalls, NumDrawCalls);
 	INC_DWORD_STAT_BY(STAT_LandscapeTriangles, NumTriangles * NumPasses);
+
+#if !UE_BUILD_SHIPPING
+	if (GVarDumpLandscapeLODsCurrentFrame == GFrameNumberRenderThread)
+	{
+		for (const FSceneView* View : Views)
+		{
+			const FString& LandscapeName = LandscapeComponent->GetLandscapeInfo()->LandscapeActor
+											 ? LandscapeComponent->GetLandscapeInfo()->LandscapeActor->GetName()
+											 : LexToString(LandscapeComponent->GetLandscapeInfo()->LandscapeGuid);
+			const FString& ComponentName = LandscapeComponent->GetName();
+
+			const float LODValue = ComputeLODForView(*View);
+			const int32 LOD = FMath::FloorToInt(LODValue);
+			const int32 Resolution = (ComponentSizeQuads + 1) >> LOD;
+
+			const int32 LoadedHeightmapResolution = LandscapeComponent->HeightmapTexture
+				                                        ? 1 << (LandscapeComponent->HeightmapTexture->GetResource()->GetCurrentMipCount() - 1) : 0;
+			const int32 LoadedWeightmapResolution = [this]
+			{
+				if (LandscapeComponent->WeightmapTextures.IsEmpty())
+				{
+					return 0;
+				}
+				int32 MaxMipCount = 0;
+				for (const UTexture2D* WeightmapTexture : LandscapeComponent->WeightmapTextures)
+				{
+					MaxMipCount = FMath::Max(MaxMipCount, WeightmapTexture ? WeightmapTexture->GetResource()->GetCurrentMipCount() : 0);
+				}
+				return MaxMipCount > 0 ? 1 << (MaxMipCount - 1) : 0;
+			}();
+
+			UE_LOG(LogLandscape, Display, TEXT("\nView: %d, Landscape: %s, Component: %s [%s], "
+				       "LODValue: %f, LOD: %d, Resolution: %d, LoadedHeightmapMIP: %d, LoadedWeightmapMIP: %d"),
+			       View->GetViewKey(), *LandscapeName, *ComponentName, *SectionBase.ToString(),
+			       LODValue, LOD, Resolution, LoadedHeightmapResolution, LoadedWeightmapResolution);
+		}
+	}
+#endif // !UE_BUILD_SHIPPING
 }
 
 #if RHI_RAYTRACING
