Landscape: Fix Static Analysis warning for re-added DumpLODs functionality

diff --git a/Engine/Source/Runtime/Landscape/Private/LandscapeRender.cpp b/Engine/Source/Runtime/Landscape/Private/LandscapeRender.cpp
--- a/Engine/Source/Runtime/Landscape/Private/LandscapeRender.cpp
+++ b/Engine/Source/Runtime/Landscape/Private/LandscapeRender.cpp
@@ -2467,8 +2467,10 @@ void FLandscapeComponentSceneProxy::GetDynamicMeshElements(const TArray<const FS
 			const int32 LOD = FMath::FloorToInt(LODValue);
 			const int32 Resolution = (ComponentSizeQuads + 1) >> LOD;
 
-			const int32 LoadedHeightmapResolution = LandscapeComponent->HeightmapTexture
-				                                        ? 1 << (LandscapeComponent->HeightmapTexture->GetResource()->GetCurrentMipCount() - 1) : 0;
+			const int32 LoadedHeightmapResolution = LandscapeComponent->HeightmapTexture && LandscapeComponent->HeightmapTexture->GetResource() &&
+			                                        LandscapeComponent->HeightmapTexture->GetResource()->GetCurrentMipCount() > 0
+				                                        ? 1 << (LandscapeComponent->HeightmapTexture->GetResource()->GetCurrentMipCount() - 1)
+				                                        : 0;
 			const int32 LoadedWeightmapResolution = [this]
 			{
 				if (LandscapeComponent->WeightmapTextures.IsEmpty())
