Fixed mem report texture usage computation for landscape components :
* HeightmapTexture / XYOffsetmapTexture, which are only used on the non-mobile path, were artificially increasing the ref counts on platforms where mobile and non-mobile paths are enabled (in mobile mode)
* Weightmap[0], which contains the normal map in mobile mode, was not accounted for by the material instance's GetUsedTextures call, since it uses inlined hlsl code to sample it, rather than a texture sample material expression

diff --git a/Engine/Source/Runtime/Landscape/Private/LandscapeRender.cpp b/Engine/Source/Runtime/Landscape/Private/LandscapeRender.cpp
--- a/Engine/Source/Runtime/Landscape/Private/LandscapeRender.cpp
+++ b/Engine/Source/Runtime/Landscape/Private/LandscapeRender.cpp
@@ -3831,33 +3831,59 @@ void ULandscapeComponent::GetStreamingRenderAssetInfo(FStreamingTextureLevelCont
 		}
 	}
 
-	// Weightmap
-	for (int32 TextureIndex = 0; TextureIndex < WeightmapTextures.Num(); TextureIndex++)
+	if (bUseMobileLandscapeMesh && !MobileWeightmapTextures.IsEmpty())
 	{
-		FStreamingRenderAssetPrimitiveInfo& StreamingWeightmap = *new(OutStreamingRenderAssets)FStreamingRenderAssetPrimitiveInfo;
-		StreamingWeightmap.Bounds = BoundingSphere;
-		StreamingWeightmap.TexelFactor = TexelFactor;
-		StreamingWeightmap.RenderAsset = WeightmapTextures[TextureIndex];
+		// Special case for mobile : the normal map is in the first weightmap texture (+ eventually 2 weightmap layers in case of a number of weight-blend layers <= 3, in the 2 remaining channels) 
+		//  is sampled through inlined hlsl code rather than texture material expression, which prevents them from being gathered by GetUsedTextures : 
+		if (UTexture2D* NormalmapTexture = MobileWeightmapTextures[0])
+		{
+			FStreamingRenderAssetPrimitiveInfo& StreamingWeightmap = *new(OutStreamingRenderAssets)FStreamingRenderAssetPrimitiveInfo;
+			StreamingWeightmap.Bounds = BoundingSphere;
+			StreamingWeightmap.TexelFactor = TexelFactor;
+			StreamingWeightmap.RenderAsset = NormalmapTexture;
+		}
+		// Make sure that all other mobile weightmap textures in the list all have been accounted for by GetUsedTextures on the material :
+		for (int32 TextureIndex = 1; TextureIndex < MobileWeightmapTextures.Num(); ++TextureIndex)
+		{
+			if (UTexture2D* WeightmapTexture = MobileWeightmapTextures[TextureIndex])
+			{
+				check(OutStreamingRenderAssets.FindByPredicate([=](const FStreamingRenderAssetPrimitiveInfo& StreamingWeightmap) { return StreamingWeightmap.RenderAsset == WeightmapTexture;}) != nullptr);
+			}
+		}
 	}
-
-	// Heightmap
-	if (HeightmapTexture)
+	else
 	{
-		FStreamingRenderAssetPrimitiveInfo& StreamingHeightmap = *new(OutStreamingRenderAssets)FStreamingRenderAssetPrimitiveInfo;
-		StreamingHeightmap.Bounds = BoundingSphere;
+		// Make sure that all weightmap textures in the list all have been accounted for by GetUsedTextures on the material :
+		for (int32 TextureIndex = 0; TextureIndex < WeightmapTextures.Num(); ++TextureIndex)
+		{
+			if (UTexture2D* WeightmapTexture = WeightmapTextures[TextureIndex])
+			{
+				check(OutStreamingRenderAssets.FindByPredicate([=](const FStreamingRenderAssetPrimitiveInfo& StreamingWeightmap) { return StreamingWeightmap.RenderAsset == WeightmapTexture; }) != nullptr);
+			}
+		}
 
-		float HeightmapTexelFactor = TexelFactor * (static_cast<float>(HeightmapTexture->GetSizeY()) / (ComponentSizeQuads + 1));
-		StreamingHeightmap.TexelFactor = ForcedLOD >= 0 ? -(1 << (13 - ForcedLOD)) : HeightmapTexelFactor; // Minus Value indicate forced resolution (Mip 13 for 8k texture)
-		StreamingHeightmap.RenderAsset = HeightmapTexture;
-	}
+		// Heightmap has not been accounted for by GetUsedTextures on the material :
+		if (HeightmapTexture)
+		{
+			// Heightmap should not have been accounted for already:
+			check(OutStreamingRenderAssets.FindByPredicate([=](const FStreamingRenderAssetPrimitiveInfo& StreamingWeightmap) { return StreamingWeightmap.RenderAsset == HeightmapTexture; }) == nullptr);
 
-	// XYOffset
-	if (XYOffsetmapTexture)
-	{
-		FStreamingRenderAssetPrimitiveInfo& StreamingXYOffset = *new(OutStreamingRenderAssets)FStreamingRenderAssetPrimitiveInfo;
-		StreamingXYOffset.Bounds = BoundingSphere;
-		StreamingXYOffset.TexelFactor = TexelFactor;
-		StreamingXYOffset.RenderAsset = XYOffsetmapTexture;
+			FStreamingRenderAssetPrimitiveInfo& StreamingHeightmap = *new(OutStreamingRenderAssets)FStreamingRenderAssetPrimitiveInfo;
+			StreamingHeightmap.Bounds = BoundingSphere;
+
+			float HeightmapTexelFactor = TexelFactor * (static_cast<float>(HeightmapTexture->GetSizeY()) / (ComponentSizeQuads + 1));
+			StreamingHeightmap.TexelFactor = ForcedLOD >= 0 ? -(1 << (13 - ForcedLOD)) : HeightmapTexelFactor; // Minus Value indicate forced resolution (Mip 13 for 8k texture)
+			StreamingHeightmap.RenderAsset = HeightmapTexture;
+		}
+
+		// XYOffset has not been accounted for by GetUsedTextures on the material :
+		if (XYOffsetmapTexture)
+		{
+			FStreamingRenderAssetPrimitiveInfo& StreamingXYOffset = *new(OutStreamingRenderAssets)FStreamingRenderAssetPrimitiveInfo;
+			StreamingXYOffset.Bounds = BoundingSphere;
+			StreamingXYOffset.TexelFactor = TexelFactor;
+			StreamingXYOffset.RenderAsset = XYOffsetmapTexture;
+		}
 	}
 
 #if WITH_EDITOR
