Temporarily removed asserts about unreported weightmap textures in ULandscapeComponent::GetStreamingRenderAssetInfo

diff --git a/Engine/Source/Runtime/Landscape/Private/LandscapeRender.cpp b/Engine/Source/Runtime/Landscape/Private/LandscapeRender.cpp
--- a/Engine/Source/Runtime/Landscape/Private/LandscapeRender.cpp
+++ b/Engine/Source/Runtime/Landscape/Private/LandscapeRender.cpp
@@ -3843,24 +3843,24 @@ void ULandscapeComponent::GetStreamingRenderAssetInfo(FStreamingTextureLevelCont
 			StreamingWeightmap.RenderAsset = NormalmapTexture;
 		}
 		// Make sure that all other mobile weightmap textures in the list all have been accounted for by GetUsedTextures on the material :
-		for (int32 TextureIndex = 1; TextureIndex < MobileWeightmapTextures.Num(); ++TextureIndex)
-		{
-			if (UTexture2D* WeightmapTexture = MobileWeightmapTextures[TextureIndex])
-			{
-				check(OutStreamingRenderAssets.FindByPredicate([=](const FStreamingRenderAssetPrimitiveInfo& StreamingWeightmap) { return StreamingWeightmap.RenderAsset == WeightmapTexture;}) != nullptr);
-			}
-		}
+		//for (int32 TextureIndex = 1; TextureIndex < MobileWeightmapTextures.Num(); ++TextureIndex)
+		//{
+		//	if (UTexture2D* WeightmapTexture = MobileWeightmapTextures[TextureIndex])
+		//	{
+		//		check(OutStreamingRenderAssets.FindByPredicate([=](const FStreamingRenderAssetPrimitiveInfo& StreamingWeightmap) { return StreamingWeightmap.RenderAsset == WeightmapTexture;}) != nullptr);
+		//	}
+		//}
 	}
 	else
 	{
 		// Make sure that all weightmap textures in the list all have been accounted for by GetUsedTextures on the material :
-		for (int32 TextureIndex = 0; TextureIndex < WeightmapTextures.Num(); ++TextureIndex)
-		{
-			if (UTexture2D* WeightmapTexture = WeightmapTextures[TextureIndex])
-			{
-				check(OutStreamingRenderAssets.FindByPredicate([=](const FStreamingRenderAssetPrimitiveInfo& StreamingWeightmap) { return StreamingWeightmap.RenderAsset == WeightmapTexture; }) != nullptr);
-			}
-		}
+		//for (int32 TextureIndex = 0; TextureIndex < WeightmapTextures.Num(); ++TextureIndex)
+		//{
+		//	if (UTexture2D* WeightmapTexture = WeightmapTextures[TextureIndex])
+		//	{
+		//		check(OutStreamingRenderAssets.FindByPredicate([=](const FStreamingRenderAssetPrimitiveInfo& StreamingWeightmap) { return StreamingWeightmap.RenderAsset == WeightmapTexture; }) != nullptr);
+		//	}
+		//}
 
 		// Heightmap has not been accounted for by GetUsedTextures on the material :
 		if (HeightmapTexture)
