Landscape: Fix static analysis warning in DumpLODs code

diff --git a/Engine/Source/Runtime/Landscape/Private/LandscapeRender.cpp b/Engine/Source/Runtime/Landscape/Private/LandscapeRender.cpp
--- a/Engine/Source/Runtime/Landscape/Private/LandscapeRender.cpp
+++ b/Engine/Source/Runtime/Landscape/Private/LandscapeRender.cpp
@@ -2467,20 +2467,26 @@ void FLandscapeComponentSceneProxy::GetDynamicMeshElements(const TArray<const FS
 			const int32 LOD = FMath::FloorToInt(LODValue);
 			const int32 Resolution = (ComponentSizeQuads + 1) >> LOD;
 
-			const int32 LoadedHeightmapResolution = LandscapeComponent->HeightmapTexture && LandscapeComponent->HeightmapTexture->GetResource() &&
-			                                        LandscapeComponent->HeightmapTexture->GetResource()->GetCurrentMipCount() > 0
-				                                        ? 1 << (LandscapeComponent->HeightmapTexture->GetResource()->GetCurrentMipCount() - 1)
-				                                        : 0;
-			const int32 LoadedWeightmapResolution = [this]
+			const int32 LoadedHeightmapResolution = [this]
 			{
-				if (LandscapeComponent->WeightmapTextures.IsEmpty())
+				if (!(LandscapeComponent->HeightmapTexture && LandscapeComponent->HeightmapTexture->GetResource()))
 				{
 					return 0;
 				}
+				const int32 MipCount = LandscapeComponent->HeightmapTexture->GetResource()->GetCurrentMipCount();
+				return MipCount > 0 ? 1 << (MipCount - 1) : 0;
+			}();
+
+			const int32 LoadedWeightmapResolution = [this]
+			{
 				int32 MaxMipCount = 0;
 				for (const UTexture2D* WeightmapTexture : LandscapeComponent->WeightmapTextures)
 				{
-					MaxMipCount = FMath::Max(MaxMipCount, WeightmapTexture ? WeightmapTexture->GetResource()->GetCurrentMipCount() : 0);
+					if (!(WeightmapTexture && WeightmapTexture->GetResource()))
+					{
+						continue;
+					}
+					MaxMipCount = FMath::Max(MaxMipCount, WeightmapTexture->GetResource()->GetCurrentMipCount());
 				}
 				return MaxMipCount > 0 ? 1 << (MaxMipCount - 1) : 0;
 			}();
