export SetLOD

diff --git a/Engine/Source/Runtime/Landscape/Classes/LandscapeComponent.h b/Engine/Source/Runtime/Landscape/Classes/LandscapeComponent.h
--- a/Engine/Source/Runtime/Landscape/Classes/LandscapeComponent.h
+++ b/Engine/Source/Runtime/Landscape/Classes/LandscapeComponent.h
@@ -1151,7 +1151,7 @@ public:
 	friend struct FLandscapeComponentDataInterface;
 	friend class ULandscapeLODStreamingProxy;
 
-	void SetLOD(bool bForced, int32 InLODValue);
+	LANDSCAPE_API void SetLOD(bool bForced, int32 InLODValue);
 
 protected:
 
