Added BP setters for ForcedLOD and LODBias on landscape components. SetLOD can now be called both at runtime and edit-time.

diff --git a/Engine/Source/Runtime/Landscape/Classes/LandscapeComponent.h b/Engine/Source/Runtime/Landscape/Classes/LandscapeComponent.h
--- a/Engine/Source/Runtime/Landscape/Classes/LandscapeComponent.h
+++ b/Engine/Source/Runtime/Landscape/Classes/LandscapeComponent.h
@@ -1153,6 +1153,12 @@ public:
 
 	LANDSCAPE_API void SetLOD(bool bForced, int32 InLODValue);
 
+	UFUNCTION(BlueprintCallable, Category = "LandscapeComponent")
+	LANDSCAPE_API void SetForcedLOD(int32 InForcedLOD);
+
+	UFUNCTION(BlueprintCallable, Category = "LandscapeComponent")
+	LANDSCAPE_API void SetLODBias(int32 InLODBias);
+
 protected:
 
 #if WITH_EDITOR
diff --git a/Engine/Source/Runtime/Landscape/Private/Landscape.cpp b/Engine/Source/Runtime/Landscape/Private/Landscape.cpp
--- a/Engine/Source/Runtime/Landscape/Private/Landscape.cpp
+++ b/Engine/Source/Runtime/Landscape/Private/Landscape.cpp
@@ -347,6 +347,71 @@ void ULandscapeComponent::CheckGenerateLandscapePlatformData(bool bIsCooking, co
 }
 #endif
 
+void ULandscapeComponent::SetForcedLOD(int32 InForcedLOD)
+{
+	SetLOD(/*bForced = */true, InForcedLOD);
+}
+
+void ULandscapeComponent::SetLODBias(int32 InLODBias)
+{
+	SetLOD(/*bForced = */false, InLODBias);
+}
+
+void ULandscapeComponent::SetLOD(bool bForcedLODChanged, int32 InLODValue)
+{
+	if (bForcedLODChanged)
+	{
+		ForcedLOD = InLODValue;
+		if (ForcedLOD >= 0)
+		{
+			ForcedLOD = FMath::Clamp<int32>(ForcedLOD, 0, FMath::CeilLogTwo(SubsectionSizeQuads + 1) - 1);
+		}
+		else
+		{
+			ForcedLOD = -1;
+		}
+	}
+	else
+	{
+		int32 MaxLOD = FMath::CeilLogTwo(SubsectionSizeQuads + 1) - 1;
+		LODBias = FMath::Clamp<int32>(InLODValue, -MaxLOD, MaxLOD);
+	}
+
+	InvalidateLightingCache();
+	MarkRenderStateDirty();
+
+#if WITH_EDITOR
+	// Update neighbor components for lighting cache (only relevant in the editor ATM) : 
+	ULandscapeInfo* Info = GetLandscapeInfo();
+	if (Info)
+	{
+		FIntPoint ComponentBase = GetSectionBase() / ComponentSizeQuads;
+		FIntPoint LandscapeKey[8] =
+		{
+			ComponentBase + FIntPoint(-1, -1),
+			ComponentBase + FIntPoint(+0, -1),
+			ComponentBase + FIntPoint(+1, -1),
+			ComponentBase + FIntPoint(-1, +0),
+			ComponentBase + FIntPoint(+1, +0),
+			ComponentBase + FIntPoint(-1, +1),
+			ComponentBase + FIntPoint(+0, +1),
+			ComponentBase + FIntPoint(+1, +1)
+		};
+
+		for (int32 Idx = 0; Idx < 8; ++Idx)
+		{
+			ULandscapeComponent* Comp = Info->XYtoComponentMap.FindRef(LandscapeKey[Idx]);
+			if (Comp)
+			{
+				Comp->Modify();
+				Comp->InvalidateLightingCache();
+				Comp->MarkRenderStateDirty();
+			}
+		}
+	}
+#endif // WITH_EDITOR
+}
+
 void ULandscapeComponent::Serialize(FArchive& Ar)
 {
 	LLM_SCOPE(ELLMTag::Landscape);
diff --git a/Engine/Source/Runtime/Landscape/Private/LandscapeEdit.cpp b/Engine/Source/Runtime/Landscape/Private/LandscapeEdit.cpp
--- a/Engine/Source/Runtime/Landscape/Private/LandscapeEdit.cpp
+++ b/Engine/Source/Runtime/Landscape/Private/LandscapeEdit.cpp
@@ -5608,59 +5608,6 @@ void ALandscapeProxy::ChangedPhysMaterial()
 	}
 }
 
-void ULandscapeComponent::SetLOD(bool bForcedLODChanged, int32 InLODValue)
-{
-	if (bForcedLODChanged)
-	{
-		ForcedLOD = InLODValue;
-		if (ForcedLOD >= 0)
-		{
-			ForcedLOD = FMath::Clamp<int32>(ForcedLOD, 0, FMath::CeilLogTwo(SubsectionSizeQuads + 1) - 1);
-		}
-		else
-		{
-			ForcedLOD = -1;
-		}
-	}
-	else
-	{
-		int32 MaxLOD = FMath::CeilLogTwo(SubsectionSizeQuads + 1) - 1;
-		LODBias = FMath::Clamp<int32>(InLODValue, -MaxLOD, MaxLOD);
-	}
-
-	InvalidateLightingCache();
-	MarkRenderStateDirty();
-
-	// Update neighbor components
-	ULandscapeInfo* Info = GetLandscapeInfo();
-	if (Info)
-	{
-		FIntPoint ComponentBase = GetSectionBase() / ComponentSizeQuads;
-		FIntPoint LandscapeKey[8] =
-		{
-			ComponentBase + FIntPoint(-1, -1),
-			ComponentBase + FIntPoint(+0, -1),
-			ComponentBase + FIntPoint(+1, -1),
-			ComponentBase + FIntPoint(-1, +0),
-			ComponentBase + FIntPoint(+1, +0),
-			ComponentBase + FIntPoint(-1, +1),
-			ComponentBase + FIntPoint(+0, +1),
-			ComponentBase + FIntPoint(+1, +1)
-		};
-
-		for (int32 Idx = 0; Idx < 8; ++Idx)
-		{
-			ULandscapeComponent* Comp = Info->XYtoComponentMap.FindRef(LandscapeKey[Idx]);
-			if (Comp)
-			{
-				Comp->Modify();
-				Comp->InvalidateLightingCache();
-				Comp->MarkRenderStateDirty();
-			}
-		}
-	}
-}
-
 void ULandscapeComponent::PreEditChange(FProperty* PropertyThatWillChange)
 {
 	Super::PreEditChange(PropertyThatWillChange);
