Make Landscape async grass builder and HISM's AcceptPrebuiltTree work for Nanite ISMs

diff --git a/Engine/Source/Runtime/Engine/Classes/Components/HierarchicalInstancedStaticMeshComponent.h b/Engine/Source/Runtime/Engine/Classes/Components/HierarchicalInstancedStaticMeshComponent.h
--- a/Engine/Source/Runtime/Engine/Classes/Components/HierarchicalInstancedStaticMeshComponent.h
+++ b/Engine/Source/Runtime/Engine/Classes/Components/HierarchicalInstancedStaticMeshComponent.h
@@ -256,7 +256,7 @@ public:
 
 	bool BuildTreeIfOutdated(bool Async, bool ForceUpdate);
 	static void BuildTreeAnyThread(TArray<FMatrix>& InstanceTransforms, TArray<float>& InstanceCustomDataFloats, int32 NumCustomDataFloats, const FBox& MeshBox, TArray<FClusterNode>& OutClusterTree, TArray<int32>& OutSortedInstances, TArray<int32>& OutInstanceReorderTable, int32& OutOcclusionLayerNum, int32 MaxInstancesPerLeaf, bool InGenerateInstanceScalingRange);
-	void AcceptPrebuiltTree(TArray<FClusterNode>& InClusterTree, int32 InOcclusionLayerNumNodes, int32 InNumBuiltRenderInstances);
+	void AcceptPrebuiltTree(TArray<FInstancedStaticMeshInstanceData>& InInstanceData, TArray<FClusterNode>& InClusterTree, int32 InOcclusionLayerNumNodes, int32 InNumBuiltRenderInstances);
 	bool IsAsyncBuilding() const { return bIsAsyncBuilding; }
 	bool IsTreeFullyBuilt() const { return !bIsOutOfDate; }
 
diff --git a/Engine/Source/Runtime/Engine/Private/HierarchicalInstancedStaticMesh.cpp b/Engine/Source/Runtime/Engine/Private/HierarchicalInstancedStaticMesh.cpp
--- a/Engine/Source/Runtime/Engine/Private/HierarchicalInstancedStaticMesh.cpp
+++ b/Engine/Source/Runtime/Engine/Private/HierarchicalInstancedStaticMesh.cpp
@@ -2721,13 +2721,15 @@ void UHierarchicalInstancedStaticMeshComponent::BuildTreeAnyThread(
 	OutSortedInstances = MoveTemp(Builder.Result->SortedInstances);
 }
 
-void UHierarchicalInstancedStaticMeshComponent::AcceptPrebuiltTree(TArray<FClusterNode>& InClusterTree, int32 InOcclusionLayerNumNodes, int32 InNumBuiltRenderInstances)
+void UHierarchicalInstancedStaticMeshComponent::AcceptPrebuiltTree(TArray<FInstancedStaticMeshInstanceData>& InInstanceData, TArray<FClusterNode>& InClusterTree, int32 InOcclusionLayerNumNodes, int32 InNumBuiltRenderInstances)
 {
 	checkSlow(IsInGameThread());
 
 	QUICK_SCOPE_CYCLE_COUNTER(STAT_UHierarchicalInstancedStaticMeshComponent_AcceptPrebuiltTree);
+	
 	// this is only for prebuild data, already in the correct order
 	check(!PerInstanceSMData.Num());
+
 	NumBuiltInstances = 0;
 	TranslatedInstanceSpaceOrigin = FVector::Zero();
 	check(PerInstanceRenderData.IsValid());
@@ -2753,6 +2755,13 @@ void UHierarchicalInstancedStaticMeshComponent::AcceptPrebuiltTree(TArray<FClust
 	if(bMeshIsValid)
 	{
 		*ClusterTreePtr = MoveTemp(InClusterTree);
+
+		// We only need to copy off the instances if it is a Nanite mesh, since the Nanite scene proxy uses them instead
+		if (ShouldCreateNaniteProxy())
+		{
+			PerInstanceSMData = MoveTemp(InInstanceData);
+		}
+		
 		PostBuildStats();
 
 	}
diff --git a/Engine/Source/Runtime/Landscape/Private/LandscapeGrass.cpp b/Engine/Source/Runtime/Landscape/Private/LandscapeGrass.cpp
--- a/Engine/Source/Runtime/Landscape/Private/LandscapeGrass.cpp
+++ b/Engine/Source/Runtime/Landscape/Private/LandscapeGrass.cpp
@@ -1770,6 +1770,7 @@ struct FAsyncGrassBuilder : public FGrassBuilderBase
 	TArray<FBox> ExcludedBoxes;
 
 	// output
+	TArray<FInstancedStaticMeshInstanceData> InstanceData;
 	FStaticMeshInstanceData InstanceBuffer;
 	TArray<FClusterNode> ClusterTree;
 	int32 OutOcclusionLayerNum;
@@ -2144,11 +2145,14 @@ struct FAsyncGrassBuilder : public FGrassBuilderBase
 			TArray<float> InstanceCustomDataDummy;
 			UHierarchicalInstancedStaticMeshComponent::BuildTreeAnyThread(InstanceTransforms, InstanceCustomDataDummy, 0, MeshBox, ClusterTree, SortedInstances, InstanceReorderTable, OutOcclusionLayerNum, DesiredInstancesPerLeaf, false);
 
-			// in-place sort the instances
+			InstanceData.Reset(NumInstances);
 			
+			// in-place sort the instances and generate the sorted instance data
 			for (int32 FirstUnfixedIndex = 0; FirstUnfixedIndex < NumInstances; FirstUnfixedIndex++)
 			{
-				int32 LoadFrom = SortedInstances[FirstUnfixedIndex];
+				int32 LoadFrom = SortedInstances[FirstUnfixedIndex];				
+				InstanceData.Emplace(InstanceTransforms[LoadFrom]);
+
 				if (LoadFrom != FirstUnfixedIndex)
 				{
 					check(LoadFrom > FirstUnfixedIndex);
@@ -3298,7 +3302,7 @@ void ALandscapeProxy::UpdateGrass(const TArray<FVector>& Cameras, int32& InOutNu
 							HierarchicalInstancedStaticMeshComponent->PerInstanceRenderData->UpdateFromPreallocatedData(Inner.Builder->InstanceBuffer);
 						}
 
-						HierarchicalInstancedStaticMeshComponent->AcceptPrebuiltTree(Inner.Builder->ClusterTree, Inner.Builder->OutOcclusionLayerNum, NumBuiltRenderInstances);
+						HierarchicalInstancedStaticMeshComponent->AcceptPrebuiltTree(Inner.Builder->InstanceData, Inner.Builder->ClusterTree, Inner.Builder->OutOcclusionLayerNum, NumBuiltRenderInstances);
 						if (bForceSync && GetWorld())
 						{
 							QUICK_SCOPE_CYCLE_COUNTER(STAT_FoliageGrassEndComp_SyncUpdate);
