* Added common function for whether landscape edit layers are supported on a given platform
* Added landscape runtime utils for this (the function is needed for global shaders, which have to live in a runtime module)

diff --git a/Engine/Source/Runtime/Landscape/Private/LandscapeEditLayers.cpp b/Engine/Source/Runtime/Landscape/Private/LandscapeEditLayers.cpp
--- a/Engine/Source/Runtime/Landscape/Private/LandscapeEditLayers.cpp
+++ b/Engine/Source/Runtime/Landscape/Private/LandscapeEditLayers.cpp
@@ -15,6 +15,7 @@ LandscapeEditLayers.cpp: Landscape editing layers mode
 #include "LandscapePrivate.h"
 #include "LandscapeEditReadback.h"
 #include "LandscapeEditResources.h"
+#include "LandscapeUtils.h"
 
 #include "Shader.h"
 #include "GlobalShader.h"
@@ -415,7 +416,7 @@ public:
 
 	static bool ShouldCompilePermutation(const FGlobalShaderPermutationParameters& Parameters)
 	{
-		return IsFeatureLevelSupported(Parameters.Platform, ERHIFeatureLevel::SM5) && !IsConsolePlatform(Parameters.Platform) && !IsMetalMobilePlatform(Parameters.Platform);
+		return UE::Landscape::DoesPlatformSupportEditLayers(Parameters.Platform);
 	}
 
 	static void ModifyCompilationEnvironment(const FGlobalShaderPermutationParameters& Parameters, FShaderCompilerEnvironment& OutEnvironment)
@@ -480,7 +481,7 @@ public:
 
 	static bool ShouldCompilePermutation(const FGlobalShaderPermutationParameters& Parameters)
 	{
-		return IsFeatureLevelSupported(Parameters.Platform, ERHIFeatureLevel::SM5) && !IsConsolePlatform(Parameters.Platform) && !IsMetalMobilePlatform(Parameters.Platform);
+		return UE::Landscape::DoesPlatformSupportEditLayers(Parameters.Platform);
 	}
 
 	static void ModifyCompilationEnvironment(const FGlobalShaderPermutationParameters& Parameters, FShaderCompilerEnvironment& OutEnvironment)
@@ -542,7 +543,7 @@ public:
 
 	static bool ShouldCompilePermutation(const FGlobalShaderPermutationParameters& Parameters)
 	{
-		return IsFeatureLevelSupported(Parameters.Platform, ERHIFeatureLevel::SM5) && !IsConsolePlatform(Parameters.Platform) && !IsMetalMobilePlatform(Parameters.Platform);
+		return UE::Landscape::DoesPlatformSupportEditLayers(Parameters.Platform);
 	}
 
 	static void ModifyCompilationEnvironment(const FGlobalShaderPermutationParameters& Parameters, FShaderCompilerEnvironment& OutEnvironment)
@@ -617,7 +618,7 @@ public:
 
 	static bool ShouldCompilePermutation(const FGlobalShaderPermutationParameters& Parameters)
 	{
-		return IsFeatureLevelSupported(Parameters.Platform, ERHIFeatureLevel::SM5) && !IsConsolePlatform(Parameters.Platform) && !IsMetalMobilePlatform(Parameters.Platform);
+		return UE::Landscape::DoesPlatformSupportEditLayers(Parameters.Platform);
 	}
 
 	static void ModifyCompilationEnvironment(const FGlobalShaderPermutationParameters& Parameters, FShaderCompilerEnvironment& OutEnvironment)
@@ -671,7 +672,7 @@ public:
 
 	static bool ShouldCompilePermutation(const FGlobalShaderPermutationParameters& Parameters)
 	{
-		return IsFeatureLevelSupported(Parameters.Platform, ERHIFeatureLevel::SM5) && !IsConsolePlatform(Parameters.Platform) && !IsMetalMobilePlatform(Parameters.Platform);
+		return UE::Landscape::DoesPlatformSupportEditLayers(Parameters.Platform);
 	}
 
 	static void ModifyCompilationEnvironment(const FGlobalShaderPermutationParameters& Parameters, FShaderCompilerEnvironment& OutEnvironment)
@@ -715,9 +716,9 @@ class FLandscapeCopyTextureVS : public FGlobalShader
 public:
 	DECLARE_GLOBAL_SHADER(FLandscapeCopyTextureVS);
 
-	static bool ShouldCompilePermutation(const FGlobalShaderPermutationParameters&)
+	static bool ShouldCompilePermutation(const FGlobalShaderPermutationParameters& InParameters)
 	{
-		return true;
+		return UE::Landscape::DoesPlatformSupportEditLayers(InParameters.Platform);
 	}
 
 	FLandscapeCopyTextureVS()
@@ -735,7 +736,7 @@ public:
 
 	static bool ShouldCompilePermutation(const FGlobalShaderPermutationParameters& Parameters)
 	{
-		return IsFeatureLevelSupported(Parameters.Platform, ERHIFeatureLevel::SM5) && !IsConsolePlatform(Parameters.Platform) && !IsMetalMobilePlatform(Parameters.Platform);
+		return UE::Landscape::DoesPlatformSupportEditLayers(Parameters.Platform);
 	}
 
 	static void ModifyCompilationEnvironment(const FGlobalShaderPermutationParameters& Parameters, FShaderCompilerEnvironment& OutEnvironment)
@@ -846,7 +847,7 @@ public:
 
 	static bool ShouldCompilePermutation(const FGlobalShaderPermutationParameters& Parameters)
 	{
-		return IsFeatureLevelSupported(Parameters.Platform, ERHIFeatureLevel::SM5) && !IsConsolePlatform(Parameters.Platform) && !IsMetalMobilePlatform(Parameters.Platform);
+		return UE::Landscape::DoesPlatformSupportEditLayers(Parameters.Platform);
 	}
 
 	static void ModifyCompilationEnvironment(const FGlobalShaderPermutationParameters& Parameters, FShaderCompilerEnvironment& OutEnvironment)
@@ -1028,7 +1029,7 @@ public:
 
 	static bool ShouldCompilePermutation(const FGlobalShaderPermutationParameters& Parameters)
 	{
-		return IsFeatureLevelSupported(Parameters.Platform, ERHIFeatureLevel::SM5) && !IsConsolePlatform(Parameters.Platform) && !IsMetalMobilePlatform(Parameters.Platform);
+		return UE::Landscape::DoesPlatformSupportEditLayers(Parameters.Platform);
 	}
 
 	static void ModifyCompilationEnvironment(const FGlobalShaderPermutationParameters& Parameters, FShaderCompilerEnvironment& OutEnvironment)
@@ -1127,7 +1128,7 @@ public:
 
 	static bool ShouldCompilePermutation(const FGlobalShaderPermutationParameters& InParameters)
 	{
-		return IsFeatureLevelSupported(InParameters.Platform, ERHIFeatureLevel::SM5) && !IsConsolePlatform(InParameters.Platform) && !IsMetalMobilePlatform(InParameters.Platform);
+		return UE::Landscape::DoesPlatformSupportEditLayers(InParameters.Platform);
 	}
 
 	static void ModifyCompilationEnvironment(const FGlobalShaderPermutationParameters& InParameters, FShaderCompilerEnvironment& OutEnvironment)
@@ -1171,7 +1172,7 @@ public:
 
 	static bool ShouldCompilePermutation(const FGlobalShaderPermutationParameters& Parameters)
 	{
-		return IsFeatureLevelSupported(Parameters.Platform, ERHIFeatureLevel::SM5) && !IsConsolePlatform(Parameters.Platform) && !IsMetalMobilePlatform(Parameters.Platform);
+		return UE::Landscape::DoesPlatformSupportEditLayers(Parameters.Platform);
 	}
 
 	static void ModifyCompilationEnvironment(const FGlobalShaderPermutationParameters& Parameters, FShaderCompilerEnvironment& OutEnvironment)
@@ -1216,7 +1217,7 @@ public:
 
 	static bool ShouldCompilePermutation(const FGlobalShaderPermutationParameters& Parameters)
 	{
-		return IsFeatureLevelSupported(Parameters.Platform, ERHIFeatureLevel::SM5) && !IsConsolePlatform(Parameters.Platform) && !IsMetalMobilePlatform(Parameters.Platform);
+		return UE::Landscape::DoesPlatformSupportEditLayers(Parameters.Platform);
 	}
 
 	static void ModifyCompilationEnvironment(const FGlobalShaderPermutationParameters& Parameters, FShaderCompilerEnvironment& OutEnvironment)
@@ -1259,7 +1260,7 @@ public:
 
 	static bool ShouldCompilePermutation(const FGlobalShaderPermutationParameters& Parameters)
 	{
-		return IsFeatureLevelSupported(Parameters.Platform, ERHIFeatureLevel::SM5) && !IsConsolePlatform(Parameters.Platform) && !IsMetalMobilePlatform(Parameters.Platform);
+		return UE::Landscape::DoesPlatformSupportEditLayers(Parameters.Platform);
 	}
 
 	static void ModifyCompilationEnvironment(const FGlobalShaderPermutationParameters& Parameters, FShaderCompilerEnvironment& OutEnvironment)
@@ -1303,7 +1304,7 @@ public:
 
 		static bool ShouldCompilePermutation(const FGlobalShaderPermutationParameters& InParameters)
 	{
-		return IsFeatureLevelSupported(InParameters.Platform, ERHIFeatureLevel::SM5) && !IsConsolePlatform(InParameters.Platform) && !IsMetalMobilePlatform(InParameters.Platform);
+		return UE::Landscape::DoesPlatformSupportEditLayers(InParameters.Platform);
 	}
 
 	static void ModifyCompilationEnvironment(const FGlobalShaderPermutationParameters& InParameters, FShaderCompilerEnvironment& OutEnvironment)
@@ -1349,7 +1350,7 @@ public:
 
 		static bool ShouldCompilePermutation(const FGlobalShaderPermutationParameters& Parameters)
 	{
-		return IsFeatureLevelSupported(Parameters.Platform, ERHIFeatureLevel::SM5) && !IsConsolePlatform(Parameters.Platform) && !IsMetalMobilePlatform(Parameters.Platform);
+		return UE::Landscape::DoesPlatformSupportEditLayers(Parameters.Platform);
 	}
 
 	static void ModifyCompilationEnvironment(const FGlobalShaderPermutationParameters& Parameters, FShaderCompilerEnvironment& OutEnvironment)
@@ -1391,7 +1392,7 @@ public:
 
 		static bool ShouldCompilePermutation(const FGlobalShaderPermutationParameters& Parameters)
 	{
-		return IsFeatureLevelSupported(Parameters.Platform, ERHIFeatureLevel::SM5) && !IsConsolePlatform(Parameters.Platform) && !IsMetalMobilePlatform(Parameters.Platform);
+		return UE::Landscape::DoesPlatformSupportEditLayers(Parameters.Platform);
 	}
 
 	static void ModifyCompilationEnvironment(const FGlobalShaderPermutationParameters& Parameters, FShaderCompilerEnvironment& OutEnvironment)
diff --git a/Engine/Source/Runtime/Landscape/Private/LandscapeUtils.cpp b/Engine/Source/Runtime/Landscape/Private/LandscapeUtils.cpp
new file mode 100644
--- /dev/null
+++ b/Engine/Source/Runtime/Landscape/Private/LandscapeUtils.cpp
@@ -0,0 +1,16 @@
+// Copyright Epic Games, Inc. All Rights Reserved.
+
+#include "LandscapeUtils.h"
+
+namespace UE::Landscape
+{
+
+bool DoesPlatformSupportEditLayers(EShaderPlatform InShaderPlatform)
+{
+	// Edit layers work on the GPU and are only available on SM5+ and in the editor : 
+	return IsFeatureLevelSupported(InShaderPlatform, ERHIFeatureLevel::SM5)
+		&& !IsConsolePlatform(InShaderPlatform)
+		&& !IsMobilePlatform(InShaderPlatform);
+}
+
+} // end namespace UE::Landscape
diff --git a/Engine/Source/Runtime/Landscape/Public/LandscapeUtils.h b/Engine/Source/Runtime/Landscape/Public/LandscapeUtils.h
new file mode 100644
--- /dev/null
+++ b/Engine/Source/Runtime/Landscape/Public/LandscapeUtils.h
@@ -0,0 +1,17 @@
+// Copyright Epic Games, Inc. All Rights Reserved.
+
+#pragma once
+
+#include "CoreMinimal.h"
+#include "RHIDefinitions.h"
+
+namespace UE::Landscape
+{
+
+/**
+* Returns true if edit layers (GPU landscape tools) are enabled on this platform :
+* Note: this is intended for the editor but is in runtime code since global shaders need to exist in runtime modules 
+*/
+bool DoesPlatformSupportEditLayers(EShaderPlatform InShaderPlatform);
+
+} // end namespace UE::Landscape
