Fix a bug with Nanite Landscape grass where they would be placed incorrectly.
Also fix an issue where WPO would be disabled for Nanite grass if `r.OptimizedWPO=1`.

diff --git a/Engine/Source/Runtime/Landscape/Private/LandscapeGrass.cpp b/Engine/Source/Runtime/Landscape/Private/LandscapeGrass.cpp
--- a/Engine/Source/Runtime/Landscape/Private/LandscapeGrass.cpp
+++ b/Engine/Source/Runtime/Landscape/Private/LandscapeGrass.cpp
@@ -2146,17 +2146,21 @@ struct FAsyncGrassBuilder : public FGrassBuilderBase
 			UHierarchicalInstancedStaticMeshComponent::BuildTreeAnyThread(InstanceTransforms, InstanceCustomDataDummy, 0, MeshBox, ClusterTree, SortedInstances, InstanceReorderTable, OutOcclusionLayerNum, DesiredInstancesPerLeaf, false);
 
 			InstanceData.Reset(NumInstances);
+			for (const FMatrix& Transform : InstanceTransforms)
+			{
+				InstanceData.Emplace(Transform);
+			}
 			
 			// in-place sort the instances and generate the sorted instance data
 			for (int32 FirstUnfixedIndex = 0; FirstUnfixedIndex < NumInstances; FirstUnfixedIndex++)
 			{
 				int32 LoadFrom = SortedInstances[FirstUnfixedIndex];				
-				InstanceData.Emplace(InstanceTransforms[LoadFrom]);
 
 				if (LoadFrom != FirstUnfixedIndex)
 				{
 					check(LoadFrom > FirstUnfixedIndex);
 					InstanceBuffer.SwapInstance(FirstUnfixedIndex, LoadFrom);
+					InstanceData.Swap(FirstUnfixedIndex, LoadFrom);
 
 					int32 SwapGoesTo = InstanceReorderTable[FirstUnfixedIndex];
 					check(SwapGoesTo > FirstUnfixedIndex);
@@ -3032,6 +3036,7 @@ void ALandscapeProxy::UpdateGrass(const TArray<FVector>& Cameras, int32& InOutNu
 										HierarchicalInstancedStaticMeshComponent->bCastDynamicShadow = GrassVariety.bCastDynamicShadow && !bDisableDynamicShadows;
 										HierarchicalInstancedStaticMeshComponent->bCastContactShadow = GrassVariety.bCastContactShadow && !bDisableDynamicShadows;
 										HierarchicalInstancedStaticMeshComponent->OverrideMaterials = GrassVariety.OverrideMaterials;
+										HierarchicalInstancedStaticMeshComponent->bEvaluateWorldPositionOffset = true;
 
 										const FMeshMapBuildData* MeshMapBuildData = Component->GetMeshMapBuildData();
 
