UHierarchicalInstancedStaticMeshComponent: Stop assuming that an empty PerinstanceSMData means that the component is grass (incorrect for Nanite grass since recent changes to AcceptPrebuiltTree).
Created UGrassHierarchicalInstancedStaticMeshComponent for the purpose of moving all Landscape grass-specific logic to a subclass (similar to Foliage), and out of the generic HISM base as much as possible.
Deprecated a couple of HISM methods that are currently only used by grass in preparation to move them to the new subclass.

diff --git a/Engine/Source/Runtime/Engine/Classes/Components/HierarchicalInstancedStaticMeshComponent.h b/Engine/Source/Runtime/Engine/Classes/Components/HierarchicalInstancedStaticMeshComponent.h
--- a/Engine/Source/Runtime/Engine/Classes/Components/HierarchicalInstancedStaticMeshComponent.h
+++ b/Engine/Source/Runtime/Engine/Classes/Components/HierarchicalInstancedStaticMeshComponent.h
@@ -255,7 +255,9 @@ public:
 	virtual bool ShouldCreatePhysicsState() const override;
 
 	bool BuildTreeIfOutdated(bool Async, bool ForceUpdate);
+	UE_DEPRECATED(5.1, "The BuildTreeAnyThread method is moving to UGrassInstancedStaticMeshComponent. Please update your project to use the new component and method or your project may not compile in the next udpate.")
 	static void BuildTreeAnyThread(TArray<FMatrix>& InstanceTransforms, TArray<float>& InstanceCustomDataFloats, int32 NumCustomDataFloats, const FBox& MeshBox, TArray<FClusterNode>& OutClusterTree, TArray<int32>& OutSortedInstances, TArray<int32>& OutInstanceReorderTable, int32& OutOcclusionLayerNum, int32 MaxInstancesPerLeaf, bool InGenerateInstanceScalingRange);
+	UE_DEPRECATED(5.1, "The AcceptPrebuiltTree method is moving to UGrassInstancedStaticMeshComponent. Please update your project to use the new component and method or your project may not compile in the next udpate.")
 	void AcceptPrebuiltTree(TArray<FInstancedStaticMeshInstanceData>& InInstanceData, TArray<FClusterNode>& InClusterTree, int32 InOcclusionLayerNumNodes, int32 InNumBuiltRenderInstances);
 	bool IsAsyncBuilding() const { return bIsAsyncBuilding; }
 	bool IsTreeFullyBuilt() const { return !bIsOutOfDate; }
@@ -273,7 +275,7 @@ public:
 
 	virtual void PropagateLightingScenarioChange() override;
 
-	virtual EHISMViewRelevanceType GetViewRelevanceType() const { return PerInstanceSMData.Num() ? EHISMViewRelevanceType::HISM : EHISMViewRelevanceType::Grass; }
+	EHISMViewRelevanceType GetViewRelevanceType() const { return ViewRelevanceType; }
 
 protected:
 	void BuildTree();
@@ -308,6 +310,7 @@ protected:
 	mutable FBox AccumulatedNavigationDirtyArea;
 
 	FGraphEventArray BuildTreeAsyncTasks;
+	EHISMViewRelevanceType ViewRelevanceType = EHISMViewRelevanceType::HISM;
 
 protected:
 	friend FStaticLightingTextureMapping_InstancedStaticMesh;
diff --git a/Engine/Source/Runtime/Engine/Private/HierarchicalInstancedStaticMesh.cpp b/Engine/Source/Runtime/Engine/Private/HierarchicalInstancedStaticMesh.cpp
--- a/Engine/Source/Runtime/Engine/Private/HierarchicalInstancedStaticMesh.cpp
+++ b/Engine/Source/Runtime/Engine/Private/HierarchicalInstancedStaticMesh.cpp
@@ -2645,8 +2645,10 @@ void UHierarchicalInstancedStaticMeshComponent::PostBuildStats()
 	const TArray<FClusterNode>& ClusterTree = *ClusterTreePtr;
 	FString MeshName = GetStaticMesh() ? GetStaticMesh()->GetPathName() : FString(TEXT("null"));
 	check(PerInstanceRenderData.IsValid());
-	bool bIsGrass = !PerInstanceSMData.Num();
-	NumInst = bIsGrass ? PerInstanceRenderData->InstanceBuffer.GetNumInstances() : PerInstanceSMData.Num();
+
+	// Non-nanite HISMs that were built with AcceptPrebuiltTree won't have PerInstanceSMData, so get the instance count from the buffer
+	int32 NumInst = PerInstanceSMData.Num() == 0 ? PerInstanceRenderData->InstanceBuffer.GetNumInstances() : PerInstanceSMData.Num();
+	bool bIsGrass = GetViewRelevanceType() == EHISMViewRelevanceType::Grass;
 
 	UE_LOG(LogStaticMesh, Display, TEXT("Built a foliage hierarchy with %d instances, %d nodes, %f instances / leaf (desired %d) and %d verts in LOD0. Grass? %d    %s"), NumInst, ClusterTree.Num(), ActualInstancesPerLeaf(), DesiredInstancesPerLeaf(), GetVertsForLOD(0), bIsGrass), *MeshName);
 #endif
@@ -2693,6 +2695,7 @@ void UHierarchicalInstancedStaticMeshComponent::BuildTree()
 	}
 }
 
+// TODO: Move this implementation to UGrassInstancedStaticMeshComponent.cpp for UE 5.2
 void UHierarchicalInstancedStaticMeshComponent::BuildTreeAnyThread(
 	TArray<FMatrix>& InstanceTransforms, 
 	TArray<float>& InstanceCustomDataFloats,
@@ -2721,6 +2724,7 @@ void UHierarchicalInstancedStaticMeshComponent::BuildTreeAnyThread(
 	OutSortedInstances = MoveTemp(Builder.Result->SortedInstances);
 }
 
+// TODO: Move this implementation to UGrassInstancedStaticMeshComponent.cpp for UE 5.2
 void UHierarchicalInstancedStaticMeshComponent::AcceptPrebuiltTree(TArray<FInstancedStaticMeshInstanceData>& InInstanceData, TArray<FClusterNode>& InClusterTree, int32 InOcclusionLayerNumNodes, int32 InNumBuiltRenderInstances)
 {
 	checkSlow(IsInGameThread());
diff --git a/Engine/Source/Runtime/Foliage/Private/InstancedFoliage.cpp b/Engine/Source/Runtime/Foliage/Private/InstancedFoliage.cpp
--- a/Engine/Source/Runtime/Foliage/Private/InstancedFoliage.cpp
+++ b/Engine/Source/Runtime/Foliage/Private/InstancedFoliage.cpp
@@ -5641,6 +5641,8 @@ UFoliageInstancedStaticMeshComponent::UFoliageInstancedStaticMeshComponent(const
 #if WITH_EDITORONLY_DATA
 	bEnableAutoLODGeneration = false;
 #endif
+
+	ViewRelevanceType = EHISMViewRelevanceType::Foliage;
 }
 
 void UFoliageInstancedStaticMeshComponent::ReceiveComponentDamage(float DamageAmount, FDamageEvent const& DamageEvent, AController* EventInstigator, AActor* DamageCauser)
diff --git a/Engine/Source/Runtime/Foliage/Private/InstancedGrass.cpp b/Engine/Source/Runtime/Foliage/Private/InstancedGrass.cpp
new file mode 100644
--- /dev/null
+++ b/Engine/Source/Runtime/Foliage/Private/InstancedGrass.cpp
@@ -0,0 +1,51 @@
+// Copyright Epic Games, Inc. All Rights Reserved.
+
+#include "GrassInstancedStaticMeshComponent.h"
+
+UGrassInstancedStaticMeshComponent::UGrassInstancedStaticMeshComponent(const FObjectInitializer& ObjectInitializer) :
+	Super(ObjectInitializer)
+{
+	ViewRelevanceType = EHISMViewRelevanceType::Grass;
+}
+
+void UGrassInstancedStaticMeshComponent::BuildTreeAnyThread(
+	TArray<FMatrix>& InstanceTransforms, 
+	TArray<float>& InstanceCustomDataFloats,
+	int32 NumCustomDataFloats,
+	const FBox& MeshBox,
+	TArray<FClusterNode>& OutClusterTree,
+	TArray<int32>& OutSortedInstances,
+	TArray<int32>& OutInstanceReorderTable,
+	int32& OutOcclusionLayerNum,
+	int32 MaxInstancesPerLeaf,
+	bool InGenerateInstanceScalingRange
+	)
+{
+	// TODO: Move the implementation of this function in the base here in 5.2
+	PRAGMA_DISABLE_DEPRECATION_WARNINGS
+	Super::BuildTreeAnyThread(
+		InstanceTransforms,
+		InstanceCustomDataFloats,
+		NumCustomDataFloats,
+		MeshBox,
+		OutClusterTree,
+		OutSortedInstances,
+		OutInstanceReorderTable,
+		OutOcclusionLayerNum,
+		MaxInstancesPerLeaf,
+		InGenerateInstanceScalingRange);
+	PRAGMA_ENABLE_DEPRECATION_WARNINGS
+}
+
+void UGrassInstancedStaticMeshComponent::AcceptPrebuiltTree(TArray<FInstancedStaticMeshInstanceData>& InInstanceData, TArray<FClusterNode>& InClusterTree, int32 InOcclusionLayerNumNodes, int32 InNumBuiltRenderInstances)
+{
+	// TODO: Move the implementation of this function in the base to here in 5.2
+	PRAGMA_DISABLE_DEPRECATION_WARNINGS
+	Super::AcceptPrebuiltTree(
+		InInstanceData,
+		InClusterTree,
+		InOcclusionLayerNumNodes,
+		InNumBuiltRenderInstances);
+	PRAGMA_ENABLE_DEPRECATION_WARNINGS
+}
+
diff --git a/Engine/Source/Runtime/Foliage/Public/FoliageInstancedStaticMeshComponent.h b/Engine/Source/Runtime/Foliage/Public/FoliageInstancedStaticMeshComponent.h
--- a/Engine/Source/Runtime/Foliage/Public/FoliageInstancedStaticMeshComponent.h
+++ b/Engine/Source/Runtime/Foliage/Public/FoliageInstancedStaticMeshComponent.h
@@ -47,8 +47,6 @@ public:
 	void SetGenerationGuid(const FGuid& InGuid) { GenerationGuid = InGuid; }
 	const FGuid& GetGenerationGuid() const { return GenerationGuid; }
 
-	virtual EHISMViewRelevanceType GetViewRelevanceType() const override { return EHISMViewRelevanceType::Foliage; }
-
 private:
 
 	UPROPERTY()
diff --git a/Engine/Source/Runtime/Foliage/Public/GrassInstancedStaticMeshComponent.h b/Engine/Source/Runtime/Foliage/Public/GrassInstancedStaticMeshComponent.h
new file mode 100644
--- /dev/null
+++ b/Engine/Source/Runtime/Foliage/Public/GrassInstancedStaticMeshComponent.h
@@ -0,0 +1,19 @@
+// Copyright Epic Games, Inc. All Rights Reserved.
+
+#pragma once
+
+#include "CoreMinimal.h"
+#include "Components/HierarchicalInstancedStaticMeshComponent.h"
+#include "GrassInstancedStaticMeshComponent.generated.h"
+
+UCLASS(ClassGroup = Foliage)
+class FOLIAGE_API UGrassInstancedStaticMeshComponent : public UHierarchicalInstancedStaticMeshComponent
+{
+	GENERATED_UCLASS_BODY()
+
+public:
+	static void BuildTreeAnyThread(TArray<FMatrix>& InstanceTransforms, TArray<float>& InstanceCustomDataFloats, int32 NumCustomDataFloats, const FBox& MeshBox, TArray<FClusterNode>& OutClusterTree, TArray<int32>& OutSortedInstances, TArray<int32>& OutInstanceReorderTable, int32& OutOcclusionLayerNum, int32 MaxInstancesPerLeaf, bool InGenerateInstanceScalingRange);
+	
+	void AcceptPrebuiltTree(TArray<FInstancedStaticMeshInstanceData>& InInstanceData, TArray<FClusterNode>& InClusterTree, int32 InOcclusionLayerNumNodes, int32 InNumBuiltRenderInstances);
+};
+
diff --git a/Engine/Source/Runtime/Landscape/Private/LandscapeGrass.cpp b/Engine/Source/Runtime/Landscape/Private/LandscapeGrass.cpp
--- a/Engine/Source/Runtime/Landscape/Private/LandscapeGrass.cpp
+++ b/Engine/Source/Runtime/Landscape/Private/LandscapeGrass.cpp
@@ -41,7 +41,7 @@
 #include "LandscapeDataAccess.h"
 #include "StaticMeshResources.h"
 #include "LandscapeLight.h"
-#include "Components/HierarchicalInstancedStaticMeshComponent.h"
+#include "GrassInstancedStaticMeshComponent.h"
 #include "Materials/MaterialInstanceConstant.h"
 #include "ShaderParameterUtils.h"
 #include "EngineModule.h"
@@ -102,13 +102,13 @@ static int32 GMaxInstancesPerComponent = 65536;
 static FAutoConsoleVariableRef CVarMaxInstancesPerComponent(
 	TEXT("grass.MaxInstancesPerComponent"),
 	GMaxInstancesPerComponent,
-	TEXT("Used to control the number of hierarchical components created. More can be more efficient, but can be hitchy as new components come into range"));
+	TEXT("Used to control the number of grass components created. More can be more efficient, but can be hitchy as new components come into range"));
 
 static int32 GMaxAsyncTasks = 4;
 static FAutoConsoleVariableRef CVarMaxAsyncTasks(
 	TEXT("grass.MaxAsyncTasks"),
 	GMaxAsyncTasks,
-	TEXT("Used to control the number of hierarchical components created at a time."));
+	TEXT("Used to control the number of grass components created at a time."));
 
 static int32 GUseHaltonDistribution = 0;
 static FAutoConsoleVariableRef CVarUseHaltonDistribution(
@@ -1775,7 +1775,7 @@ struct FAsyncGrassBuilder : public FGrassBuilderBase
 	TArray<FClusterNode> ClusterTree;
 	int32 OutOcclusionLayerNum;
 
-	FAsyncGrassBuilder(ALandscapeProxy* Landscape, ULandscapeComponent* Component, const ULandscapeGrassType* GrassType, const FGrassVariety& GrassVariety, ERHIFeatureLevel::Type FeatureLevel, UHierarchicalInstancedStaticMeshComponent* HierarchicalInstancedStaticMeshComponent, int32 SqrtSubsections, int32 SubX, int32 SubY, uint32 InHaltonBaseIndex, TArray<FBox>& InExcludedBoxes)
+	FAsyncGrassBuilder(ALandscapeProxy* Landscape, ULandscapeComponent* Component, const ULandscapeGrassType* GrassType, const FGrassVariety& GrassVariety, ERHIFeatureLevel::Type FeatureLevel, UGrassInstancedStaticMeshComponent* GrassInstancedStaticMeshComponent, int32 SqrtSubsections, int32 SubX, int32 SubY, uint32 InHaltonBaseIndex, TArray<FBox>& InExcludedBoxes)
 		: FGrassBuilderBase(Landscape, Component, GrassVariety, FeatureLevel, SqrtSubsections, SubX, SubY, GrassType->bEnableDensityScaling)
 		, GrassData(Component, GrassType)
 		, Scaling(GrassVariety.Scaling)
@@ -1786,10 +1786,10 @@ struct FAsyncGrassBuilder : public FGrassBuilderBase
 		, RandomScale(false)
 		, AlignToSurface(GrassVariety.AlignToSurface)
 		, PlacementJitter(GrassVariety.PlacementJitter)
-		, RandomStream(HierarchicalInstancedStaticMeshComponent->InstancingRandomSeed)
-		, XForm(LandscapeToWorld * HierarchicalInstancedStaticMeshComponent->GetComponentTransform().ToMatrixWithScale().Inverse())
+		, RandomStream(GrassInstancedStaticMeshComponent->InstancingRandomSeed)
+		, XForm(LandscapeToWorld * GrassInstancedStaticMeshComponent->GetComponentTransform().ToMatrixWithScale().Inverse())
 		, MeshBox(GrassVariety.GrassMesh->GetBounds().GetBox())
-		, DesiredInstancesPerLeaf(HierarchicalInstancedStaticMeshComponent->DesiredInstancesPerLeaf())
+		, DesiredInstancesPerLeaf(GrassInstancedStaticMeshComponent->DesiredInstancesPerLeaf())
 
 		, BuildTime(0)
 		, TotalInstances(0)
@@ -1827,7 +1827,7 @@ struct FAsyncGrassBuilder : public FGrassBuilderBase
 
 		if (InExcludedBoxes.Num())
 		{
-			FMatrix BoxXForm = HierarchicalInstancedStaticMeshComponent->GetComponentToWorld().ToMatrixWithScale().Inverse() * XForm.Inverse();
+			FMatrix BoxXForm = GrassInstancedStaticMeshComponent->GetComponentToWorld().ToMatrixWithScale().Inverse() * XForm.Inverse();
 			for (const FBox& Box : InExcludedBoxes)
 			{
 				ExcludedBoxes.Add(Box.TransformBy(BoxXForm));
@@ -2143,7 +2143,7 @@ struct FAsyncGrassBuilder : public FGrassBuilderBase
 			TArray<int32> SortedInstances;
 			TArray<int32> InstanceReorderTable;
 			TArray<float> InstanceCustomDataDummy;
-			UHierarchicalInstancedStaticMeshComponent::BuildTreeAnyThread(InstanceTransforms, InstanceCustomDataDummy, 0, MeshBox, ClusterTree, SortedInstances, InstanceReorderTable, OutOcclusionLayerNum, DesiredInstancesPerLeaf, false);
+			UGrassInstancedStaticMeshComponent::BuildTreeAnyThread(InstanceTransforms, InstanceCustomDataDummy, 0, MeshBox, ClusterTree, SortedInstances, InstanceReorderTable, OutOcclusionLayerNum, DesiredInstancesPerLeaf, false);
 
 			InstanceData.Reset(NumInstances);
 			for (const FMatrix& Transform : InstanceTransforms)
@@ -3011,32 +3011,32 @@ void ALandscapeProxy::UpdateGrass(const TArray<FVector>& Cameras, int32& InOutNu
 										ClearFlags(RF_Transactional);
 										bool PreviousPackageDirtyFlag = GetOutermost()->IsDirty();
 
-										UHierarchicalInstancedStaticMeshComponent* HierarchicalInstancedStaticMeshComponent;
+										UGrassInstancedStaticMeshComponent* GrassInstancedStaticMeshComponent;
 										{
 											QUICK_SCOPE_CYCLE_COUNTER(STAT_GrassCreateComp);
-											HierarchicalInstancedStaticMeshComponent = NewObject<UHierarchicalInstancedStaticMeshComponent>(this, NAME_None, RF_Transient);
+											GrassInstancedStaticMeshComponent = NewObject<UGrassInstancedStaticMeshComponent>(this, NAME_None, RF_Transient);
 										}
-										NewComp.Foliage = HierarchicalInstancedStaticMeshComponent;
+										NewComp.Foliage = GrassInstancedStaticMeshComponent;
 										FoliageCache.CachedGrassComps.Add(NewComp);
 
-										HierarchicalInstancedStaticMeshComponent->Mobility = EComponentMobility::Static;
-										HierarchicalInstancedStaticMeshComponent->SetStaticMesh(GrassVariety.GrassMesh);
-										HierarchicalInstancedStaticMeshComponent->MinLOD = GrassVariety.MinLOD;
-										HierarchicalInstancedStaticMeshComponent->bSelectable = false;
-										HierarchicalInstancedStaticMeshComponent->bHasPerInstanceHitProxies = false;
-										HierarchicalInstancedStaticMeshComponent->bReceivesDecals = GrassVariety.bReceivesDecals;
+										GrassInstancedStaticMeshComponent->Mobility = EComponentMobility::Static;
+										GrassInstancedStaticMeshComponent->SetStaticMesh(GrassVariety.GrassMesh);
+										GrassInstancedStaticMeshComponent->MinLOD = GrassVariety.MinLOD;
+										GrassInstancedStaticMeshComponent->bSelectable = false;
+										GrassInstancedStaticMeshComponent->bHasPerInstanceHitProxies = false;
+										GrassInstancedStaticMeshComponent->bReceivesDecals = GrassVariety.bReceivesDecals;
 										static FName NoCollision(TEXT("NoCollision"));
-										HierarchicalInstancedStaticMeshComponent->SetCollisionProfileName(NoCollision);
-										HierarchicalInstancedStaticMeshComponent->bDisableCollision = true;
-										HierarchicalInstancedStaticMeshComponent->SetCanEverAffectNavigation(false);
-										HierarchicalInstancedStaticMeshComponent->InstancingRandomSeed = FolSeed;
-										HierarchicalInstancedStaticMeshComponent->LightingChannels = GrassVariety.LightingChannels;
-										HierarchicalInstancedStaticMeshComponent->bCastStaticShadow = false;
-										HierarchicalInstancedStaticMeshComponent->CastShadow = (GrassVariety.bCastDynamicShadow || GrassVariety.bCastContactShadow) && !bDisableDynamicShadows;
-										HierarchicalInstancedStaticMeshComponent->bCastDynamicShadow = GrassVariety.bCastDynamicShadow && !bDisableDynamicShadows;
-										HierarchicalInstancedStaticMeshComponent->bCastContactShadow = GrassVariety.bCastContactShadow && !bDisableDynamicShadows;
-										HierarchicalInstancedStaticMeshComponent->OverrideMaterials = GrassVariety.OverrideMaterials;
-										HierarchicalInstancedStaticMeshComponent->bEvaluateWorldPositionOffset = true;
+										GrassInstancedStaticMeshComponent->SetCollisionProfileName(NoCollision);
+										GrassInstancedStaticMeshComponent->bDisableCollision = true;
+										GrassInstancedStaticMeshComponent->SetCanEverAffectNavigation(false);
+										GrassInstancedStaticMeshComponent->InstancingRandomSeed = FolSeed;
+										GrassInstancedStaticMeshComponent->LightingChannels = GrassVariety.LightingChannels;
+										GrassInstancedStaticMeshComponent->bCastStaticShadow = false;
+										GrassInstancedStaticMeshComponent->CastShadow = (GrassVariety.bCastDynamicShadow || GrassVariety.bCastContactShadow) && !bDisableDynamicShadows;
+										GrassInstancedStaticMeshComponent->bCastDynamicShadow = GrassVariety.bCastDynamicShadow && !bDisableDynamicShadows;
+										GrassInstancedStaticMeshComponent->bCastContactShadow = GrassVariety.bCastContactShadow && !bDisableDynamicShadows;
+										GrassInstancedStaticMeshComponent->OverrideMaterials = GrassVariety.OverrideMaterials;
+										GrassInstancedStaticMeshComponent->bEvaluateWorldPositionOffset = true;
 
 										const FMeshMapBuildData* MeshMapBuildData = Component->GetMeshMapBuildData();
 
@@ -3045,12 +3045,12 @@ void ALandscapeProxy::UpdateGrass(const TArray<FVector>& Cameras, int32& InOutNu
 											&& MeshMapBuildData
 											&& MeshMapBuildData->LightMap)
 										{
-											HierarchicalInstancedStaticMeshComponent->SetLODDataCount(GrassVariety.GrassMesh->GetNumLODs(), GrassVariety.GrassMesh->GetNumLODs());
+											GrassInstancedStaticMeshComponent->SetLODDataCount(GrassVariety.GrassMesh->GetNumLODs(), GrassVariety.GrassMesh->GetNumLODs());
 
 											FLightMapRef GrassLightMap = new FLandscapeGrassLightMap(*MeshMapBuildData->LightMap->GetLightMap2D());
 											FShadowMapRef GrassShadowMap = MeshMapBuildData->ShadowMap ? new FLandscapeGrassShadowMap(*MeshMapBuildData->ShadowMap->GetShadowMap2D()) : nullptr;
 
-											for (auto& LOD : HierarchicalInstancedStaticMeshComponent->LODData)
+											for (auto& LOD : GrassInstancedStaticMeshComponent->LODData)
 											{
 												// This trasient OverrideMapBuildData will be cleaned up by UMapBuildDataRegistry::CleanupTransientOverrideMapBuildData() if the underlying MeshMapBuildData is gone
 												LOD.OverrideMapBuildData = MakeUnique<FMeshMapBuildData>();
@@ -3063,27 +3063,27 @@ void ALandscapeProxy::UpdateGrass(const TArray<FVector>& Cameras, int32& InOutNu
 										if (!Cameras.Num() || bDisableGPUCull)
 										{
 											// if we don't have any cameras, then we are rendering landscape LOD materials or somesuch and we want to disable culling
-											HierarchicalInstancedStaticMeshComponent->InstanceStartCullDistance = 0;
-											HierarchicalInstancedStaticMeshComponent->InstanceEndCullDistance = 0;
+											GrassInstancedStaticMeshComponent->InstanceStartCullDistance = 0;
+											GrassInstancedStaticMeshComponent->InstanceEndCullDistance = 0;
 										}
 										else
 										{
-											HierarchicalInstancedStaticMeshComponent->InstanceStartCullDistance = GrassVariety.StartCullDistance.GetValue() * CullDistanceScale;
-											HierarchicalInstancedStaticMeshComponent->InstanceEndCullDistance = GrassVariety.EndCullDistance.GetValue() * CullDistanceScale;
+											GrassInstancedStaticMeshComponent->InstanceStartCullDistance = GrassVariety.StartCullDistance.GetValue() * CullDistanceScale;
+											GrassInstancedStaticMeshComponent->InstanceEndCullDistance = GrassVariety.EndCullDistance.GetValue() * CullDistanceScale;
 										}
 
 										//@todo - take the settings from a UFoliageType object.  For now, disable distance field lighting on grass so we don't hitch.
-										HierarchicalInstancedStaticMeshComponent->bAffectDistanceFieldLighting = false;
+										GrassInstancedStaticMeshComponent->bAffectDistanceFieldLighting = false;
 
 										{
 											QUICK_SCOPE_CYCLE_COUNTER(STAT_GrassAttachComp);
 
-											HierarchicalInstancedStaticMeshComponent->AttachToComponent(GetRootComponent(), FAttachmentTransformRules::KeepRelativeTransform);
+											GrassInstancedStaticMeshComponent->AttachToComponent(GetRootComponent(), FAttachmentTransformRules::KeepRelativeTransform);
 											FTransform DesiredTransform = GetRootComponent()->GetComponentTransform();
 											DesiredTransform.RemoveScaling();
-											HierarchicalInstancedStaticMeshComponent->SetWorldTransform(DesiredTransform);
+											GrassInstancedStaticMeshComponent->SetWorldTransform(DesiredTransform);
 
-											FoliageComponents.Add(HierarchicalInstancedStaticMeshComponent);
+											FoliageComponents.Add(GrassInstancedStaticMeshComponent);
 										}
 
 										FAsyncGrassBuilder* Builder;
@@ -3097,12 +3097,12 @@ void ALandscapeProxy::UpdateGrass(const TArray<FVector>& Cameras, int32& InOutNu
 												check(HaltonBaseIndex > (uint32)MaxInstancesSub);
 												HaltonIndexForSub = HaltonBaseIndex - (uint32)MaxInstancesSub;
 											}
-											Builder = new FAsyncGrassBuilder(this, Component, GrassType, GrassVariety, FeatureLevel, HierarchicalInstancedStaticMeshComponent, SqrtSubsections, SubX, SubY, HaltonIndexForSub, NewComp.ExcludedBoxes);
+											Builder = new FAsyncGrassBuilder(this, Component, GrassType, GrassVariety, FeatureLevel, GrassInstancedStaticMeshComponent, SqrtSubsections, SubX, SubY, HaltonIndexForSub, NewComp.ExcludedBoxes);
 										}
 
 										if (Builder->bHaveValidData)
 										{
-											FAsyncTask<FAsyncGrassTask>* Task = new FAsyncTask<FAsyncGrassTask>(Builder, NewComp.Key, HierarchicalInstancedStaticMeshComponent);
+											FAsyncTask<FAsyncGrassTask>* Task = new FAsyncTask<FAsyncGrassTask>(Builder, NewComp.Key, GrassInstancedStaticMeshComponent);
 
 											Task->StartBackgroundTask();
 
@@ -3115,7 +3115,7 @@ void ALandscapeProxy::UpdateGrass(const TArray<FVector>& Cameras, int32& InOutNu
 										{
 											QUICK_SCOPE_CYCLE_COUNTER(STAT_GrassRegisterComp);
 
-											HierarchicalInstancedStaticMeshComponent->RegisterComponent();
+											GrassInstancedStaticMeshComponent->RegisterComponent();
 										}
 
 										SetFlags(RF_Transactional);
@@ -3288,30 +3288,30 @@ void ALandscapeProxy::UpdateGrass(const TArray<FVector>& Cameras, int32& InOutNu
 				SCOPE_CYCLE_COUNTER(STAT_FoliageGrassEndComp);
 				FAsyncGrassTask& Inner = Task->GetTask();
 				AsyncFoliageTasks.RemoveAtSwap(Index--);
-				UHierarchicalInstancedStaticMeshComponent* HierarchicalInstancedStaticMeshComponent = Inner.Foliage.Get();
+				UGrassInstancedStaticMeshComponent* GrassISMComponent = Cast<UGrassInstancedStaticMeshComponent>(Inner.Foliage.Get());
 				int32 NumBuiltRenderInstances = Inner.Builder->InstanceBuffer.GetNumInstances();
 				//UE_LOG(LogCore, Display, TEXT("%d instances in %4.0fms     %6.0f instances / sec"), NumBuiltRenderInstances, 1000.0f * float(Inner.Builder->BuildTime), float(NumBuiltRenderInstances) / float(Inner.Builder->BuildTime));
 
-				if (HierarchicalInstancedStaticMeshComponent && StillUsed.Contains(HierarchicalInstancedStaticMeshComponent))
+				if (GrassISMComponent && StillUsed.Contains(GrassISMComponent))
 				{
 					if (NumBuiltRenderInstances > 0)
 					{
 						QUICK_SCOPE_CYCLE_COUNTER(STAT_FoliageGrassEndComp_AcceptPrebuiltTree);
 
-						if (!HierarchicalInstancedStaticMeshComponent->PerInstanceRenderData.IsValid())
+						if (!GrassISMComponent->PerInstanceRenderData.IsValid())
 						{
-							HierarchicalInstancedStaticMeshComponent->InitPerInstanceRenderData(true, &Inner.Builder->InstanceBuffer, Inner.Builder->RequireCPUAccess);
+							GrassISMComponent->InitPerInstanceRenderData(true, &Inner.Builder->InstanceBuffer, Inner.Builder->RequireCPUAccess);
 						}
 						else
 						{
-							HierarchicalInstancedStaticMeshComponent->PerInstanceRenderData->UpdateFromPreallocatedData(Inner.Builder->InstanceBuffer);
+							GrassISMComponent->PerInstanceRenderData->UpdateFromPreallocatedData(Inner.Builder->InstanceBuffer);
 						}
 
-						HierarchicalInstancedStaticMeshComponent->AcceptPrebuiltTree(Inner.Builder->InstanceData, Inner.Builder->ClusterTree, Inner.Builder->OutOcclusionLayerNum, NumBuiltRenderInstances);
+						GrassISMComponent->AcceptPrebuiltTree(Inner.Builder->InstanceData, Inner.Builder->ClusterTree, Inner.Builder->OutOcclusionLayerNum, NumBuiltRenderInstances);
 						if (bForceSync && GetWorld())
 						{
 							QUICK_SCOPE_CYCLE_COUNTER(STAT_FoliageGrassEndComp_SyncUpdate);
-							HierarchicalInstancedStaticMeshComponent->RecreateRenderState_Concurrent();
+							GrassISMComponent->RecreateRenderState_Concurrent();
 						}
 					}
 				}
