Fixed GetBlueprintRenderDependencies on landscape BP brushes not getting called : the function needs to be marked CallInEditor in order to be callable outside of BeginPlay.
This fixes BP brushes potentially being not ready and flattening the terrain when running the landscape merge (when implemented on the BP side, of course).

#rnx

diff --git a/Engine/Source/Runtime/Landscape/Public/LandscapeBlueprintBrushBase.h b/Engine/Source/Runtime/Landscape/Public/LandscapeBlueprintBrushBase.h
--- a/Engine/Source/Runtime/Landscape/Public/LandscapeBlueprintBrushBase.h
+++ b/Engine/Source/Runtime/Landscape/Public/LandscapeBlueprintBrushBase.h
@@ -48,7 +48,7 @@ public:
 	UFUNCTION(BlueprintCallable, Category = "Landscape")
 	void RequestLandscapeUpdate();
 
-	UFUNCTION(BlueprintImplementableEvent)
+	UFUNCTION(BlueprintImplementableEvent, CallInEditor)
 	void GetBlueprintRenderDependencies(TArray<UObject*>& OutStreamableAssets);
 
 #if WITH_EDITOR
