Deprecated unused class member

#rnx

diff --git a/Engine/Source/Runtime/Landscape/Classes/LandscapeMaterialInstanceConstant.h b/Engine/Source/Runtime/Landscape/Classes/LandscapeMaterialInstanceConstant.h
--- a/Engine/Source/Runtime/Landscape/Classes/LandscapeMaterialInstanceConstant.h
+++ b/Engine/Source/Runtime/Landscape/Classes/LandscapeMaterialInstanceConstant.h
@@ -35,7 +35,7 @@ class ULandscapeMaterialInstanceConstant : public UMaterialInstanceConstant
 	uint32 bIsLayerThumbnail:1;
 
 	UPROPERTY()
-	uint32 bDisableTessellation:1;
+	uint32 bDisableTessellation_DEPRECATED:1;
 
 	UPROPERTY()
 	uint32 bMobile:1;
