Fixed bHasLayersContent not being set properly for the main ALandscape actor in distributed landscape setups (One ALandscape, multiple ALandscapeStreamingProxy actors).
This meant that the test in ALandscapeProxy::PostRegisterAllComponents() to know whether data migration (from non-edit layers to edit layers) should occur, which relies on the cached value bHasLayersContent (which was always false in that case, since there's no landscape component in the ALandscape) and the result of virtual function HasLayersContent() (overridden in ALandscape::HasLayersContent to return true if there's at least one registered landscape edit layer) was always forcing data migration to occur, which would do nothing (no landscape component, hence, no data to migrate) but was still marking the level as dirty

diff --git a/Engine/Source/Runtime/Landscape/Classes/Landscape.h b/Engine/Source/Runtime/Landscape/Classes/Landscape.h
--- a/Engine/Source/Runtime/Landscape/Classes/Landscape.h
+++ b/Engine/Source/Runtime/Landscape/Classes/Landscape.h
@@ -269,6 +269,7 @@ public:
 	LANDSCAPE_API void RegisterLandscapeEdMode(ILandscapeEdModeInterface* InLandscapeEdMode) { LandscapeEdMode = InLandscapeEdMode; }
 	LANDSCAPE_API void UnregisterLandscapeEdMode() { LandscapeEdMode = nullptr; }
 	LANDSCAPE_API virtual bool HasLayersContent() const override;
+	LANDSCAPE_API virtual void UpdateCachedHasLayersContent(bool bInCheckComponentDataIntegrity) override;
 	LANDSCAPE_API void RequestSplineLayerUpdate();
 	LANDSCAPE_API void RequestLayersInitialization(bool bInRequestContentUpdate = true);
 	LANDSCAPE_API void RequestLayersContentUpdateForceAll(ELandscapeLayerUpdateMode InModeMask = ELandscapeLayerUpdateMode::Update_All);
diff --git a/Engine/Source/Runtime/Landscape/Classes/LandscapeProxy.h b/Engine/Source/Runtime/Landscape/Classes/LandscapeProxy.h
--- a/Engine/Source/Runtime/Landscape/Classes/LandscapeProxy.h
+++ b/Engine/Source/Runtime/Landscape/Classes/LandscapeProxy.h
@@ -1121,7 +1121,7 @@ public:
 	/** Will tell if the landscape proxy can have some content related to the layer system */
 	LANDSCAPE_API bool CanHaveLayersContent() const;
 
-	LANDSCAPE_API void UpdateCachedHasLayersContent(bool InCheckComponentDataIntegrity = false);
+	LANDSCAPE_API virtual void UpdateCachedHasLayersContent(bool InCheckComponentDataIntegrity = false);
 
 protected:
 	friend class ALandscape;
diff --git a/Engine/Source/Runtime/Landscape/Private/Landscape.cpp b/Engine/Source/Runtime/Landscape/Private/Landscape.cpp
--- a/Engine/Source/Runtime/Landscape/Private/Landscape.cpp
+++ b/Engine/Source/Runtime/Landscape/Private/Landscape.cpp
@@ -2118,7 +2118,7 @@ void ALandscapeProxy::PostRegisterAllComponents()
 				}
 			}
 			else
-#endif
+#endif // WITH_EDITOR
 			{
 				LandscapeInfo = CreateLandscapeInfo(true);
 			}
diff --git a/Engine/Source/Runtime/Landscape/Private/LandscapeEditLayers.cpp b/Engine/Source/Runtime/Landscape/Private/LandscapeEditLayers.cpp
--- a/Engine/Source/Runtime/Landscape/Private/LandscapeEditLayers.cpp
+++ b/Engine/Source/Runtime/Landscape/Private/LandscapeEditLayers.cpp
@@ -2092,6 +2092,12 @@ void ALandscape::CopyOldDataToDefaultLayer(ALandscapeProxy* InProxy)
 		return;
 	}
 
+	if (InProxy->LandscapeComponents.IsEmpty())
+	{
+		// No data to migrate, we can early-out to avoid modifying the proxy:
+		return;
+	}
+
 	InProxy->Modify();
 
 	FLandscapeLayer* DefaultLayer = GetLayer(0);
@@ -7293,6 +7299,16 @@ bool ALandscape::HasLayersContent() const
 	return LandscapeLayers.Num() > 0;
 }
 
+void ALandscape::UpdateCachedHasLayersContent(bool bInCheckComponentDataIntegrity)
+{
+	Super::UpdateCachedHasLayersContent(bInCheckComponentDataIntegrity);
+
+	// For consistency with the ALandscape::HasLayersContent() override above, make sure the cached bHasLayersContent boolean is also valid when we have at least one edit layer : 
+	//  Otherwise, as ALandscapeProxy::UpdateCachedHasLayersContent relies on the presence of landscape components and in distributed landscape setups (one ALandscape + multiple ALandscapeStreamingProxy),  
+	//  the master ALandscape actor doesn't have any landscape, hence it would have bHasLayersContent erroneously set to false (while ALandscape::HasLayersContent() would actually return true!)
+	bHasLayersContent |= ALandscape::HasLayersContent();
+}
+
 void ALandscape::RequestLayersInitialization(bool bInRequestContentUpdate)
 {
 	if (!CanHaveLayersContent())
