Corrected ALandscapeProxy::GetLODScreenSizeArray() which didn't handle MaxLODLevel == INDEX_NONE properly
* MaxLODLevel == INDEX_NONE means the max available
* Also fixed off-by-one error (Num LOD levels vs Max LOD level)

diff --git a/Engine/Source/Runtime/Landscape/Private/Landscape.cpp b/Engine/Source/Runtime/Landscape/Private/Landscape.cpp
--- a/Engine/Source/Runtime/Landscape/Private/Landscape.cpp
+++ b/Engine/Source/Runtime/Landscape/Private/Landscape.cpp
@@ -1284,7 +1284,10 @@ static TArray<float> GetLODScreenSizeArray(const ALandscapeProxy* InLandscapePro
 
 TArray<float> ALandscapeProxy::GetLODScreenSizeArray() const
 {
-	const int32 NumLODLevels = FMath::Clamp<int32>(MaxLODLevel, 0, FMath::CeilLogTwo(SubsectionSizeQuads + 1) - 1);
+	const int32 MaxPossibleLOD = FMath::CeilLogTwo(SubsectionSizeQuads + 1) - 1;
+	const int32 MaxLOD = MaxLODLevel != INDEX_NONE ? FMath::Min<int32>(MaxLODLevel, MaxPossibleLOD) : MaxPossibleLOD;
+
+	const int32 NumLODLevels = MaxLOD + 1;
 	return ::GetLODScreenSizeArray(this, NumLODLevels);
 }
 
