Landscape: Default Preview to false on LandscapeLayerSwitch material node. When true the behavior is to always show the Used branch, which is almost never what we want.
It would be better to have Preview path only taken on non-landscape mesh, but current static parameter set doesn't distinguish between non-landscape and landscape with no layers.

diff --git a/Engine/Source/Runtime/Landscape/Private/Materials/MaterialExpressionLandscapeLayerSwitch.cpp b/Engine/Source/Runtime/Landscape/Private/Materials/MaterialExpressionLandscapeLayerSwitch.cpp
--- a/Engine/Source/Runtime/Landscape/Private/Materials/MaterialExpressionLandscapeLayerSwitch.cpp
+++ b/Engine/Source/Runtime/Landscape/Private/Materials/MaterialExpressionLandscapeLayerSwitch.cpp
@@ -34,7 +34,7 @@ UMaterialExpressionLandscapeLayerSwitch::UMaterialExpressionLandscapeLayerSwitch
 	bCollapsed = false;
 #endif
 
-	PreviewUsed = true;
+	PreviewUsed = false;
 }
 
 #if WITH_EDITOR
