Fix heighmap visibility not available when building Global Distance Field in cooked builds
- determine which layer contains visibility mask using UMaterialExpressionLandscapeVisibilityMask::ParameterName.

diff --git a/Engine/Source/Runtime/Landscape/Private/LandscapeRender.cpp b/Engine/Source/Runtime/Landscape/Private/LandscapeRender.cpp
--- a/Engine/Source/Runtime/Landscape/Private/LandscapeRender.cpp
+++ b/Engine/Source/Runtime/Landscape/Private/LandscapeRender.cpp
@@ -13,6 +13,7 @@ LandscapeRender.cpp: New terrain rendering
 #include "Materials/Material.h"
 #include "Materials/MaterialExpressionTextureCoordinate.h"
 #include "Materials/MaterialExpressionLandscapeLayerCoords.h"
+#include "Materials/MaterialExpressionLandscapeVisibilityMask.h"
 #include "Materials/MaterialInstanceDynamic.h"
 #include "Materials/MaterialInstanceConstant.h"
 #include "ShaderParameterUtils.h"
@@ -1308,18 +1309,18 @@ FLandscapeComponentSceneProxy::FLandscapeComponentSceneProxy(ULandscapeComponent
 			LayerColors.Add(Allocation.LayerInfo->LayerUsageDebugColor);
 		}
 	}
+#endif
 
 	for (int32 Idx = 0; Idx < InComponent->WeightmapLayerAllocations.Num(); Idx++)
 	{
 		const FWeightmapLayerAllocationInfo& Allocation = InComponent->WeightmapLayerAllocations[Idx];
-		if (Allocation.LayerInfo == ALandscapeProxy::VisibilityLayer && Allocation.IsAllocated())
+		if (Allocation.GetLayerName() == UMaterialExpressionLandscapeVisibilityMask::ParameterName && Allocation.IsAllocated())
 		{
 			VisibilityWeightmapTexture = WeightmapTextures[Allocation.WeightmapTextureIndex];
 			VisibilityWeightmapChannel = Allocation.WeightmapTextureChannel;
 			break;
 		}
 	}
-#endif
 
 	bSupportsInstanceDataBuffer = true;
 	UpdateDefaultInstanceSceneData();
