Fixed crash when trying to access the visibility landscape mask on mobile : when the mobile mesh path is enabled, there's no such texture.
This  means that GlobalDistanceField will not support landscape visibility.

diff --git a/Engine/Source/Runtime/Landscape/Private/LandscapeRender.cpp b/Engine/Source/Runtime/Landscape/Private/LandscapeRender.cpp
--- a/Engine/Source/Runtime/Landscape/Private/LandscapeRender.cpp
+++ b/Engine/Source/Runtime/Landscape/Private/LandscapeRender.cpp
@@ -1300,6 +1300,24 @@ FLandscapeComponentSceneProxy::FLandscapeComponentSceneProxy(ULandscapeComponent
 	bSupportsHeightfieldRepresentation = !bMobileLandscapeMesh;
 	bSupportsMeshCardRepresentation = true;
 
+	// Find where the visibility weightmap lies, if available (no such thing with mobile landscape mesh, though, since visibility is baked in the vertex buffers) :
+	if (!bMobileLandscapeMesh)
+	{
+		for (int32 Idx = 0; Idx < InComponent->WeightmapLayerAllocations.Num(); Idx++)
+		{
+			const FWeightmapLayerAllocationInfo& Allocation = InComponent->WeightmapLayerAllocations[Idx];
+			if (Allocation.GetLayerName() == UMaterialExpressionLandscapeVisibilityMask::ParameterName && Allocation.IsAllocated())
+			{
+				VisibilityWeightmapTexture = WeightmapTextures[Allocation.WeightmapTextureIndex];
+				VisibilityWeightmapChannel = Allocation.WeightmapTextureChannel;
+				break;
+			}
+		}
+	}
+
+	bSupportsInstanceDataBuffer = true;
+	UpdateDefaultInstanceSceneData();
+
 #if WITH_EDITOR
 	const TArray<FWeightmapLayerAllocationInfo>& ComponentWeightmapLayerAllocations = InComponent->GetWeightmapLayerAllocations();
 	for (const FWeightmapLayerAllocationInfo& Allocation : ComponentWeightmapLayerAllocations)
@@ -1310,20 +1328,6 @@ FLandscapeComponentSceneProxy::FLandscapeComponentSceneProxy(ULandscapeComponent
 		}
 	}
 #endif
-
-	for (int32 Idx = 0; Idx < InComponent->WeightmapLayerAllocations.Num(); Idx++)
-	{
-		const FWeightmapLayerAllocationInfo& Allocation = InComponent->WeightmapLayerAllocations[Idx];
-		if (Allocation.GetLayerName() == UMaterialExpressionLandscapeVisibilityMask::ParameterName && Allocation.IsAllocated())
-		{
-			VisibilityWeightmapTexture = WeightmapTextures[Allocation.WeightmapTextureIndex];
-			VisibilityWeightmapChannel = Allocation.WeightmapTextureChannel;
-			break;
-		}
-	}
-
-	bSupportsInstanceDataBuffer = true;
-	UpdateDefaultInstanceSceneData();
 }
 
 void FLandscapeComponentSceneProxy::CreateRenderThreadResources()
