Fixed random crash ocurring in HLOD builder commandlet, where we're hashing landscape texture source data while potentially asynchronously building textures. UTexture2D::GetSizeX() / UTexture2D::GetSizeY() should be prevented in those cases, since unlike in the landscape editor, nothing ensures that FinishCachePlatformData() is called on those textures and the returned sizes end up being wrong / thread-unsafe.
Also fixed a couple of places where we were not using the FTextureSource's data (even though those calls should be fine since they happen in a place where texture compilation is done)

diff --git a/Engine/Source/Runtime/Landscape/Private/Landscape.cpp b/Engine/Source/Runtime/Landscape/Private/Landscape.cpp
--- a/Engine/Source/Runtime/Landscape/Private/Landscape.cpp
+++ b/Engine/Source/Runtime/Landscape/Private/Landscape.cpp
@@ -4128,14 +4128,14 @@ void ALandscapeProxy::BuildGIBakedTextures(struct FScopedSlowTask* InSlowTask /*
 		const bool bShouldMarkDirty = true;
 		UpdateGIBakedTextureData(bShouldMarkDirty);
 	}
-	}
+}
 
 int32 ALandscapeProxy::GetOutdatedGIBakedTextureComponentsCount() const
-	{
+{
 	int32 OutdatedGITextureComponentsCount = 0;
 	UpdateGIBakedTextureStatus(nullptr, nullptr, &OutdatedGITextureComponentsCount);
 	return OutdatedGITextureComponentsCount;
-	}
+}
 
 void ALandscapeProxy::UpdateGIBakedTextureStatus(bool* bOutGenerateLandscapeGIData, TMap<UTexture2D*, FGIBakedTextureState>* OutComponentsNeedBakingByHeightmap, int32* OutdatedComponentsCount) const
 {
@@ -4315,53 +4315,53 @@ void ALandscapeProxy::UpdateGIBakedTextures(bool bBakeAllGITextures)
 		{
 			// We throttle, baking only one atlas per frame if bBakeAllGITextures is false.
 			if (!bBakeAllGITextures && NumGenerated > 0)
-				{
-					NumComponentsNeedingTextureBaking += Info.Components.Num();
-				}
-				else
-				{
-					UTexture2D* HeightmapTexture = It.Key();
-					// 1/8 the res of the heightmap
-					FIntPoint AtlasSize(HeightmapTexture->GetSizeX() >> 3, HeightmapTexture->GetSizeY() >> 3);
+			{
+				NumComponentsNeedingTextureBaking += Info.Components.Num();
+			}
+			else
+			{
+				UTexture2D* HeightmapTexture = It.Key();
+				// 1/8 the res of the heightmap
+				FIntPoint AtlasSize(HeightmapTexture->Source.GetSizeX() >> 3, HeightmapTexture->Source.GetSizeY() >> 3);
 
-					TArray<FColor> AtlasSamples;
-					AtlasSamples.AddZeroed(AtlasSize.X * AtlasSize.Y);
+				TArray<FColor> AtlasSamples;
+				AtlasSamples.AddZeroed(AtlasSize.X * AtlasSize.Y);
 
-					for (ULandscapeComponent* Component : Info.Components)
+				for (ULandscapeComponent* Component : Info.Components)
+				{
+					// not registered; ignore this component
+					if (!Component->SceneProxy)
 					{
-						// not registered; ignore this component
-						if (!Component->SceneProxy)
-						{
-							continue;
-						}
+						continue;
+					}
 
-						int32 ComponentSamples = (SubsectionSizeQuads + 1) * NumSubsections;
-						check(FMath::IsPowerOfTwo(ComponentSamples));
+					int32 ComponentSamples = (SubsectionSizeQuads + 1) * NumSubsections;
+					check(FMath::IsPowerOfTwo(ComponentSamples));
 
-						int32 BakeSize = ComponentSamples >> 3;
-						TArray<FColor> Samples;
-						if (FMaterialUtilities::ExportBaseColor(Component, BakeSize, Samples))
+					int32 BakeSize = ComponentSamples >> 3;
+					TArray<FColor> Samples;
+					if (FMaterialUtilities::ExportBaseColor(Component, BakeSize, Samples))
+					{
+						int32 AtlasOffsetX = FMath::RoundToInt(Component->HeightmapScaleBias.Z * (float)HeightmapTexture->Source.GetSizeX()) >> 3;
+						int32 AtlasOffsetY = FMath::RoundToInt(Component->HeightmapScaleBias.W * (float)HeightmapTexture->Source.GetSizeY()) >> 3;
+						for (int32 y = 0; y < BakeSize; y++)
 						{
-							int32 AtlasOffsetX = FMath::RoundToInt(Component->HeightmapScaleBias.Z * (float)HeightmapTexture->GetSizeX()) >> 3;
-							int32 AtlasOffsetY = FMath::RoundToInt(Component->HeightmapScaleBias.W * (float)HeightmapTexture->GetSizeY()) >> 3;
-							for (int32 y = 0; y < BakeSize; y++)
-							{
-								FMemory::Memcpy(&AtlasSamples[(y + AtlasOffsetY) * AtlasSize.X + AtlasOffsetX], &Samples[y * BakeSize], sizeof(FColor) * BakeSize);
-							}
-							NumGenerated++;
+							FMemory::Memcpy(&AtlasSamples[(y + AtlasOffsetY) * AtlasSize.X + AtlasOffsetX], &Samples[y * BakeSize], sizeof(FColor) * BakeSize);
 						}
+						NumGenerated++;
 					}
+				}
 				UTexture2D* AtlasTexture = FMaterialUtilities::CreateTexture(GetOutermost(), HeightmapTexture->GetName() + TEXT("_BaseColor"), AtlasSize, AtlasSamples, TC_Default, TEXTUREGROUP_World, RF_NoFlags, true, Info.CombinedStateId);
 
-					for (ULandscapeComponent* Component : Info.Components)
-					{
+				for (ULandscapeComponent* Component : Info.Components)
+				{
 					Component->BakedTextureMaterialGuid = Info.CombinedStateId;
-						Component->GIBakedBaseColorTexture = AtlasTexture;
-						Component->MarkRenderStateDirty();
-					}
+					Component->GIBakedBaseColorTexture = AtlasTexture;
+					Component->MarkRenderStateDirty();
 				}
 			}
 		}
+	}
 
 	TotalComponentsNeedingTextureBaking += NumComponentsNeedingTextureBaking;
 
diff --git a/Engine/Source/Runtime/Landscape/Private/LandscapeEditLayers.cpp b/Engine/Source/Runtime/Landscape/Private/LandscapeEditLayers.cpp
--- a/Engine/Source/Runtime/Landscape/Private/LandscapeEditLayers.cpp
+++ b/Engine/Source/Runtime/Landscape/Private/LandscapeEditLayers.cpp
@@ -4866,10 +4866,10 @@ int32 ALandscape::RegenerateLayersHeightmaps(const FUpdateLayersContentContext&
 			ALandscapeProxy* Proxy = Component->GetLandscapeProxy();
 			FLandscapeEditLayerReadback** CPUReadback = Proxy->HeightmapsCPUReadback.Find(ComponentHeightmap);
 			if (CPUReadback == nullptr)
-				{
+			{
 				FLandscapeEditLayerReadback* NewCPUReadback = new FLandscapeEditLayerReadback();
 				const uint8* LockedMip = ComponentHeightmap->Source.LockMip(0);
-				const uint32 Hash = FLandscapeEditLayerReadback::CalculateHash(LockedMip, ComponentHeightmap->GetSizeX() * ComponentHeightmap->GetSizeY() * sizeof(FColor));
+				const uint32 Hash = FLandscapeEditLayerReadback::CalculateHash(LockedMip, ComponentHeightmap->Source.GetSizeX() * ComponentHeightmap->Source.GetSizeY() * sizeof(FColor));
 				ComponentHeightmap->Source.UnlockMip(0);
 				NewCPUReadback->SetHash(Hash);
 				Proxy->HeightmapsCPUReadback.Add(ComponentHeightmap, NewCPUReadback);
@@ -6382,7 +6382,7 @@ void ALandscape::PrepareLayersWeightmapsLocalMergeRenderThreadData(const FUpdate
 							// Lazily create the readback objects as required (ReallocateLayersWeightmaps might have created new weightmaps)
 							FLandscapeEditLayerReadback* NewCPUReadback = new FLandscapeEditLayerReadback();
 							const uint8* LockedMip = ComponentWeightmap->Source.LockMipReadOnly(0);
-							const uint32 Hash = FLandscapeEditLayerReadback::CalculateHash(LockedMip, ComponentWeightmap->GetSizeX() * ComponentWeightmap->GetSizeY() * sizeof(FColor));
+							const uint32 Hash = FLandscapeEditLayerReadback::CalculateHash(LockedMip, ComponentWeightmap->Source.GetSizeX() * ComponentWeightmap->Source.GetSizeY() * sizeof(FColor));
 							NewCPUReadback->SetHash(Hash);
 							ComponentWeightmap->Source.UnlockMip(0);
 							CPUReadback = &Proxy->WeightmapsCPUReadback.Add(ComponentWeightmap, NewCPUReadback);
@@ -6814,7 +6814,7 @@ int32 ALandscape::PerformLayersWeightmapsGlobalMerge(FUpdateLayersContentContext
 					{
 						FLandscapeEditLayerReadback* NewCPUReadback = new FLandscapeEditLayerReadback();
 						const uint8* LockedMip = WeightmapTexture->Source.LockMipReadOnly(0);
-						const uint32 Hash = FLandscapeEditLayerReadback::CalculateHash(LockedMip, WeightmapTexture->GetSizeX() * WeightmapTexture->GetSizeY() * sizeof(FColor));
+						const uint32 Hash = FLandscapeEditLayerReadback::CalculateHash(LockedMip, WeightmapTexture->Source.GetSizeX() * WeightmapTexture->Source.GetSizeY() * sizeof(FColor));
 						NewCPUReadback->SetHash(Hash);
 						WeightmapTexture->Source.UnlockMip(0);
 						Proxy->WeightmapsCPUReadback.Add(WeightmapTexture, NewCPUReadback);
@@ -8771,7 +8771,7 @@ uint32 ULandscapeComponent::ComputeLayerHash(bool InReturnEditingHash) const
 {
 	UTexture2D* Heightmap = GetHeightmap(InReturnEditingHash);
 	const uint8* MipData = Heightmap->Source.LockMipReadOnly(0);
-	uint32 Hash = FCrc::MemCrc32(MipData, Heightmap->GetSizeX() * Heightmap->GetSizeY() * sizeof(FColor));
+	uint32 Hash = FCrc::MemCrc32(MipData, Heightmap->Source.GetSizeX() * Heightmap->Source.GetSizeY() * sizeof(FColor));
 	Heightmap->Source.UnlockMip(0);
 
 	// Copy to sort
@@ -8795,7 +8795,7 @@ uint32 ULandscapeComponent::ComputeLayerHash(bool InReturnEditingHash) const
 			UTexture2D* Weightmap = Weightmaps[AllocationInfo.WeightmapTextureIndex];
 			MipData = Weightmap->Source.LockMipReadOnly(0) + ChannelOffsets[AllocationInfo.WeightmapTextureChannel];
 			TArray<uint8> ChannelData;
-			ChannelData.AddDefaulted(Weightmap->GetSizeX() * Weightmap->GetSizeY());
+			ChannelData.AddDefaulted(Weightmap->Source.GetSizeX() * Weightmap->Source.GetSizeY());
 			int32 TexSize = (SubsectionSizeQuads + 1) * NumSubsections;
 			for (int32 TexY = 0; TexY < TexSize; TexY++)
 			{
