Changed landscape edit layers streaming textures tracking to weak object pointers in order to avoid preventing worlds being dynamically unloaded (e.g. world composition) to hold references to internal landscape textures

diff --git a/Engine/Source/Runtime/Landscape/Classes/Landscape.h b/Engine/Source/Runtime/Landscape/Classes/Landscape.h
--- a/Engine/Source/Runtime/Landscape/Classes/Landscape.h
+++ b/Engine/Source/Runtime/Landscape/Classes/Landscape.h
@@ -380,7 +380,7 @@ private:
 	bool PrepareLayersTextureResources(bool bInWaitForStreaming);
 	bool PrepareLayersTextureResources(const TArray<FLandscapeLayer>& InLayers, bool bInWaitForStreaming);
 	bool PrepareLayersBrushResources(bool bInWaitForStreaming);
-	void InvalidateRVTForTextures(const TSet<TObjectPtr<UTexture2D>>& InTextures);
+	void InvalidateRVTForTextures(const TSet<TWeakObjectPtr<UTexture2D>>& InTextures);
 	void PrepareLayersHeightmapsLocalMergeRenderThreadData(const FUpdateLayersContentContext& InUpdateLayersContentContext, const FEditLayersHeightmapMergeParams& InMergeParams, EditLayersHeightmapLocalMerge_RenderThread::FMergeInfo& OutRenderThreadData);
 	void PrepareLayersWeightmapsLocalMergeRenderThreadData(const FUpdateLayersContentContext& InUpdateLayersContentContext, const FEditLayersWeightmapMergeParams& InMergeParams, EditLayersWeightmapLocalMerge_RenderThread::FMergeInfo& OutRenderThreadData);
 
@@ -467,7 +467,7 @@ public:
 
 	/** List of textures that are not fully streamed in yet (updated every frame to track textures that have finished streaming in) */
 	UPROPERTY(Transient, DuplicateTransient, TextExportTransient)	
-	TSet<TObjectPtr<UTexture2D>> TrackedStreamingInTextures;
+	TSet<TWeakObjectPtr<UTexture2D>> TrackedStreamingInTextures;
 
 private:
 	FLandscapeBlueprintBrushChangedDelegate LandscapeBlueprintBrushChangedDelegate;
diff --git a/Engine/Source/Runtime/Landscape/Private/LandscapeEditLayers.cpp b/Engine/Source/Runtime/Landscape/Private/LandscapeEditLayers.cpp
--- a/Engine/Source/Runtime/Landscape/Private/LandscapeEditLayers.cpp
+++ b/Engine/Source/Runtime/Landscape/Private/LandscapeEditLayers.cpp
@@ -3351,7 +3351,7 @@ bool ALandscape::PrepareTextureResources(bool bInWaitForStreaming)
 		return false;
 	}
 
-	TSet<TObjectPtr<UTexture2D>> StreamingInTexturesBefore(MoveTemp(TrackedStreamingInTextures));
+	TSet<TWeakObjectPtr<UTexture2D>> StreamingInTexturesBefore(MoveTemp(TrackedStreamingInTextures));
 
 	bool bIsReady = true;
 	Info->ForAllLandscapeProxies([&](ALandscapeProxy* Proxy)
@@ -3383,7 +3383,7 @@ bool ALandscape::PrepareTextureResources(bool bInWaitForStreaming)
 	});
 
 	// The assets that were streaming in before and are not anymore can be considered streamed in: 
-	TSet<TObjectPtr<UTexture2D>> StreamedInTextures = StreamingInTexturesBefore.Difference(TrackedStreamingInTextures);
+	TSet<TWeakObjectPtr<UTexture2D>> StreamedInTextures = StreamingInTexturesBefore.Difference(TrackedStreamingInTextures);
 	InvalidateRVTForTextures(StreamedInTextures);
 
 	return bIsReady;
@@ -3391,7 +3391,7 @@ bool ALandscape::PrepareTextureResources(bool bInWaitForStreaming)
 
 // Note: this approach is generic, because FObjectCacheContextScope is a fast texture->material interface->primitive component lookup. 
 // If FObjectCacheContextScope was available at runtime, it could become an efficient way to automatically invalidate RVT areas corresponding to primitive components that use textures that are being streamed in:
-void ALandscape::InvalidateRVTForTextures(const TSet<TObjectPtr<UTexture2D>>& InTextures)
+void ALandscape::InvalidateRVTForTextures(const TSet<TWeakObjectPtr<UTexture2D>>& InTextures)
 {
 #if WITH_EDITOR
 	TRACE_CPUPROFILER_EVENT_SCOPE(ALandscape_InvalidateRVTForTextures);
@@ -3402,17 +3402,20 @@ void ALandscape::InvalidateRVTForTextures(const TSet<TObjectPtr<UTexture2D>>& In
 		FObjectCacheContextScope ObjectCacheScope;
 		TSet<UPrimitiveComponent*> PrimitiveComponentsToInvalidate;
 
-		for (UTexture2D* Texture : InTextures)
+		for (TWeakObjectPtr<UTexture2D> TexturePtr : InTextures)
 		{
-			// First, find all the materials referencing this texture that are writing to the RVT in order to invalidate the primitive components referencing them when the texture 
-			//  gets fully streamed in so that we're not left with low-res mips being rendered in the RVT tiles : 
-			for (UMaterialInterface* MaterialInterface : ObjectCacheScope.GetContext().GetMaterialsAffectedByTexture(Texture))
+			if (UTexture2D* Texture = TexturePtr.Get())
 			{
-				if (MaterialInterface->GetCachedExpressionData().bHasRuntimeVirtualTextureOutput)
+				// First, find all the materials referencing this texture that are writing to the RVT in order to invalidate the primitive components referencing them when the texture 
+				//  gets fully streamed in so that we're not left with low-res mips being rendered in the RVT tiles : 
+				for (UMaterialInterface* MaterialInterface : ObjectCacheScope.GetContext().GetMaterialsAffectedByTexture(Texture))
 				{
-					for (UPrimitiveComponent* PrimitiveComponent : ObjectCacheScope.GetContext().GetPrimitivesAffectedByMaterial(MaterialInterface))
+					if (MaterialInterface->GetCachedExpressionData().bHasRuntimeVirtualTextureOutput)
 					{
-						PrimitiveComponentsToInvalidate.Add(PrimitiveComponent);
+						for (UPrimitiveComponent* PrimitiveComponent : ObjectCacheScope.GetContext().GetPrimitivesAffectedByMaterial(MaterialInterface))
+						{
+							PrimitiveComponentsToInvalidate.Add(PrimitiveComponent);
+						}
 					}
 				}
 			}
