Landscape fills in velocity in view relevance.
This allows it to write velocity when editor moving or when r.Velocity.ForceOutput is set.

diff --git a/Engine/Source/Runtime/Landscape/Private/LandscapeRender.cpp b/Engine/Source/Runtime/Landscape/Private/LandscapeRender.cpp
--- a/Engine/Source/Runtime/Landscape/Private/LandscapeRender.cpp
+++ b/Engine/Source/Runtime/Landscape/Private/LandscapeRender.cpp
@@ -1664,6 +1664,8 @@ FPrimitiveViewRelevance FLandscapeComponentSceneProxy::GetViewRelevance(const FS
 		}
 	}
 
+	Result.bVelocityRelevance = DrawsVelocity() && Result.bOpaque && Result.bRenderInMainPass;
+
 	// Various visualizations need to render using dynamic relevance
 	if ((View->Family->EngineShowFlags.Bounds && IsSelected()) ||
 		GLandscapeDebugOptions.bShowPatches)
