Added support for non-square heightmaps in landscape edit layers local merge

diff --git a/Engine/Source/Runtime/Landscape/Private/LandscapeEditLayers.cpp b/Engine/Source/Runtime/Landscape/Private/LandscapeEditLayers.cpp
--- a/Engine/Source/Runtime/Landscape/Private/LandscapeEditLayers.cpp
+++ b/Engine/Source/Runtime/Landscape/Private/LandscapeEditLayers.cpp
@@ -4367,17 +4367,13 @@ void ALandscape::PrepareLayersHeightmapsLocalMergeRenderThreadData(const FUpdate
 			ComponentToComponentRenderInfoIndex.Add(Component, OutRenderThreadData.ComponentToRenderInfos.Num() - 1);
 
 			UTexture2D* ComponentHeightmap = Component->GetHeightmap();
-			int32 TextureSize = ComponentHeightmap->Source.GetSizeX();
-			if (ComponentHeightmap->Source.GetSizeX() > OutRenderThreadData.MaxHeightmapSize.X)
-			{
-				OutRenderThreadData.MaxHeightmapSize = FIntPoint(TextureSize, TextureSize);
-				OutRenderThreadData.MaxHeightmapNumMips = FMath::CeilLogTwo(TextureSize) + 1;
-			}
+			FIntPoint TextureSize(ComponentHeightmap->Source.GetSizeX(), ComponentHeightmap->Source.GetSizeY());
+			OutRenderThreadData.MaxHeightmapSize = TextureSize.ComponentMax(OutRenderThreadData.MaxHeightmapSize);
+			OutRenderThreadData.MaxHeightmapNumMips = FMath::Max((int32)FMath::CeilLogTwo(TextureSize.GetMin()) + 1, OutRenderThreadData.MaxHeightmapNumMips);
 
-			const int32 HeightmapOffsetX = Component->HeightmapScaleBias.Z * (float)TextureSize;
-			const int32 HeightmapOffsetY = Component->HeightmapScaleBias.W * (float)TextureSize;
+			FIntPoint HeightmapOffset(Component->HeightmapScaleBias.Z * TextureSize.X, Component->HeightmapScaleBias.W * TextureSize.Y);
 			// Effective area of the texture affecting this component (because of texture sharing) :
-			const FIntRect ComponentTextureSubregion(FIntPoint(HeightmapOffsetX, HeightmapOffsetY), FIntPoint(HeightmapOffsetX, HeightmapOffsetY) + OutRenderThreadData.ComponentSizeVerts);
+			const FIntRect ComponentTextureSubregion(HeightmapOffset, HeightmapOffset + OutRenderThreadData.ComponentSizeVerts);
 
 			NewComponentRenderInfo.VisibleLayerHeightmapTextures.Reserve(OutRenderThreadData.VisibleEditLayerInfos.Num());
 			for (const FLandscapeLayer& Layer : LandscapeLayers)
@@ -4459,8 +4455,8 @@ void ALandscape::PrepareLayersHeightmapsLocalMergeRenderThreadData(const FUpdate
 				check((CPUReadback != nullptr) && (*CPUReadback != nullptr));
 
 				FTextureResolveInfo NewTextureResolveInfo;
-				NewTextureResolveInfo.TextureSize = FIntPoint(ComponentHeightmap->Source.GetSizeX(), ComponentHeightmap->Source.GetSizeX());
-				NewTextureResolveInfo.NumMips = FMath::CeilLogTwo(NewTextureResolveInfo.TextureSize.X) + 1;
+				NewTextureResolveInfo.TextureSize = FIntPoint(ComponentHeightmap->Source.GetSizeX(), ComponentHeightmap->Source.GetSizeY());
+				NewTextureResolveInfo.NumMips = (int32)FMath::CeilLogTwo(NewTextureResolveInfo.TextureSize.GetMin()) + 1;
 				NewTextureResolveInfo.Texture = NewComponentResolveInfo.Heightmap.Texture;
 				NewTextureResolveInfo.CPUReadback = *CPUReadback;
 				OutRenderThreadData.TextureToResolveInfos.Add(NewTextureResolveInfo);
