UE-120740: Crash in async trace related to landscape data lifetime
Fix from UDN report

diff --git a/Engine/Source/Runtime/Landscape/Private/LandscapeCollision.cpp b/Engine/Source/Runtime/Landscape/Private/LandscapeCollision.cpp
--- a/Engine/Source/Runtime/Landscape/Private/LandscapeCollision.cpp
+++ b/Engine/Source/Runtime/Landscape/Private/LandscapeCollision.cpp
@@ -2079,7 +2079,8 @@ bool ULandscapeHeightfieldCollisionComponent::RecreateCollision()
 		}
 		CollisionHash = NewHash;
 #endif // WITH_EDITOR
-		HeightfieldRef = NULL;
+		TRefCountPtr<FHeightfieldGeometryRef> HeightfieldRefLifetimeExtender = HeightfieldRef; // Ensure heightfield data is alive until removed from physics world
+		HeightfieldRef = nullptr; // Ensure data will be recreated
 		HeightfieldGuid = FGuid();
 
 		RecreatePhysicsState();
@@ -2208,9 +2209,12 @@ void ULandscapeHeightfieldCollisionComponent::SnapFoliageInstances(const FBox& I
 
 bool ULandscapeMeshCollisionComponent::RecreateCollision()
 {
+	TRefCountPtr<FTriMeshGeometryRef> TriMeshLifetimeExtender = nullptr; // Ensure heightfield data is alive until removed from physics world
+
 	if (!HasAnyFlags(RF_ClassDefaultObject))
 	{
-		MeshRef = NULL;
+		TriMeshLifetimeExtender = MeshRef;
+		MeshRef = nullptr; // Ensure data will be recreated
 		MeshGuid = FGuid();
 		CachedHeightFieldSamples.Heights.Empty();
 		CachedHeightFieldSamples.Holes.Empty();
