Fixes compilation. Added forward declaration.

diff --git a/Engine/Source/Runtime/Landscape/Public/LandscapeSettings.h b/Engine/Source/Runtime/Landscape/Public/LandscapeSettings.h
--- a/Engine/Source/Runtime/Landscape/Public/LandscapeSettings.h
+++ b/Engine/Source/Runtime/Landscape/Public/LandscapeSettings.h
@@ -6,6 +6,8 @@
 #include "Engine/DeveloperSettings.h"
 #include "LandscapeSettings.generated.h"
 
+class UMaterialInterface;
+
 UCLASS(config = Engine, defaultconfig, meta = (DisplayName = "Landscape"))
 class LANDSCAPE_API ULandscapeSettings : public UDeveloperSettings
 {
