Landscape: fix relative location mismatch detection.

#rnx

diff --git a/Engine/Source/Runtime/Landscape/Private/Landscape.cpp b/Engine/Source/Runtime/Landscape/Private/Landscape.cpp
--- a/Engine/Source/Runtime/Landscape/Private/Landscape.cpp
+++ b/Engine/Source/Runtime/Landscape/Private/Landscape.cpp
@@ -845,8 +845,8 @@ void ULandscapeComponent::PostLoad()
 		FVector LocalRelativeLocation = GetRelativeLocation();
 		float CheckRelativeLocationX = float(SectionBaseX - LandscapeProxy->LandscapeSectionOffset.X);
 		float CheckRelativeLocationY = float(SectionBaseY - LandscapeProxy->LandscapeSectionOffset.Y);
-		if (CheckRelativeLocationX != LocalRelativeLocation.X ||
-			CheckRelativeLocationY != LocalRelativeLocation.Y)
+		if (!FMath::IsNearlyEqual(CheckRelativeLocationX, LocalRelativeLocation.X, UE_DOUBLE_KINDA_SMALL_NUMBER) ||
+			!FMath::IsNearlyEqual(CheckRelativeLocationY, LocalRelativeLocation.Y, UE_DOUBLE_KINDA_SMALL_NUMBER))
 		{
 			UE_LOG(LogLandscape, Warning, TEXT("LandscapeComponent RelativeLocation disagrees with its section base, attempted automated fix: '%s', %f,%f vs %f,%f."),
 				*GetFullName(), LocalRelativeLocation.X, LocalRelativeLocation.Y, CheckRelativeLocationX, CheckRelativeLocationY);
