Landscape:
- Allow changing bIsSpatiallyLoaded flag on ALandscapeStreamingProxy and ALandscapeSplineActor
- New flag in ALandscape used to set the bIsSpatiallyloaded flags on newly created AlandscapeStreamingProxy and ALandscapeSplineActor (bAreNewLandscapeActorsSpatiallyLoaded)

diff --git a/Engine/Source/Runtime/Landscape/Classes/Landscape.h b/Engine/Source/Runtime/Landscape/Classes/Landscape.h
--- a/Engine/Source/Runtime/Landscape/Classes/Landscape.h
+++ b/Engine/Source/Runtime/Landscape/Classes/Landscape.h
@@ -341,12 +341,10 @@ public:
 	LANDSCAPE_API void ForceLayersFullUpdate();
 	LANDSCAPE_API void InitializeLandscapeLayersWeightmapUsage();
 
-#if WITH_EDITOR
 	LANDSCAPE_API bool ComputeLandscapeLayerBrushInfo(FTransform& OutLandscapeTransform, FIntPoint& OutLandscapeSize, FIntPoint& OutLandscapeRenderTargetSize);
 	void RequestProxyLayersWeightmapUsageUpdate();
 	void UpdateProxyLayersWeightmapUsage();
 	void ValidateProxyLayersWeightmapUsage() const;
-#endif // WITH_EDITOR
 
 private:
 	bool SupportsEditLayersLocalMerge();
@@ -437,6 +435,10 @@ private:
 public:
 
 #if WITH_EDITORONLY_DATA
+	/** Landscape actor has authority on default streaming behavior for new actors : LandscapeStreamingProxies & LandscapeSplineActors */
+	UPROPERTY(EditAnywhere, Category = WorldPartition)
+	bool bAreNewLandscapeActorsSpatiallyLoaded = true;
+
 	UPROPERTY(EditAnywhere, Category=Landscape)
 	bool bCanHaveLayersContent = false;
 
diff --git a/Engine/Source/Runtime/Landscape/Classes/LandscapeInfo.h b/Engine/Source/Runtime/Landscape/Classes/LandscapeInfo.h
--- a/Engine/Source/Runtime/Landscape/Classes/LandscapeInfo.h
+++ b/Engine/Source/Runtime/Landscape/Classes/LandscapeInfo.h
@@ -299,6 +299,9 @@ public:
 
 	/** Returns LandscapeStreamingProxy Cell Size in WorldPartition */
 	LANDSCAPE_API uint32 GetGridSize(uint32 InGridSizeInComponents) const;
+
+	/** Returns true if new Landscape actors should be spatially loaded in WorldPartition (LandscapeStreamingProxy & LandscapeSplineActor) */
+	LANDSCAPE_API bool AreNewLandscapeActorsSpatiallyLoaded() const;
 #endif
 	LANDSCAPE_API static ULandscapeInfo* Find(UWorld* InWorld, const FGuid& LandscapeGuid);
 	LANDSCAPE_API static ULandscapeInfo* FindOrCreate(UWorld* InWorld, const FGuid& LandscapeGuid);
diff --git a/Engine/Source/Runtime/Landscape/Classes/LandscapeSplineActor.h b/Engine/Source/Runtime/Landscape/Classes/LandscapeSplineActor.h
--- a/Engine/Source/Runtime/Landscape/Classes/LandscapeSplineActor.h
+++ b/Engine/Source/Runtime/Landscape/Classes/LandscapeSplineActor.h
@@ -26,7 +26,7 @@ public:
 	virtual ULandscapeInfo* GetLandscapeInfo() const override;
 		
 #if WITH_EDITOR
-	virtual bool CanChangeIsSpatiallyLoadedFlag() const override { return false; }
+	virtual bool CanChangeIsSpatiallyLoadedFlag() const override { return true; }
 	virtual TUniquePtr<class FWorldPartitionActorDesc> CreateClassActorDesc() const override;
 
 	void GetSharedProperties(ULandscapeInfo* InLandscapeInfo);
diff --git a/Engine/Source/Runtime/Landscape/Classes/LandscapeStreamingProxy.h b/Engine/Source/Runtime/Landscape/Classes/LandscapeStreamingProxy.h
--- a/Engine/Source/Runtime/Landscape/Classes/LandscapeStreamingProxy.h
+++ b/Engine/Source/Runtime/Landscape/Classes/LandscapeStreamingProxy.h
@@ -31,6 +31,7 @@ public:
 	virtual AActor* GetSceneOutlinerParent() const override;
 	virtual bool CanDeleteSelectedActor(FText& OutReason) const override;
 	virtual bool GetReferencedContentObjects(TArray<UObject*>& Objects) const override;
+	virtual bool CanChangeIsSpatiallyLoadedFlag() const override { return true; }
 #endif
 	//~ End UObject Interface
 
diff --git a/Engine/Source/Runtime/Landscape/Private/Landscape.cpp b/Engine/Source/Runtime/Landscape/Private/Landscape.cpp
--- a/Engine/Source/Runtime/Landscape/Private/Landscape.cpp
+++ b/Engine/Source/Runtime/Landscape/Private/Landscape.cpp
@@ -3523,6 +3523,7 @@ ALandscapeSplineActor* ULandscapeInfo::CreateSplineActor(const FVector& Location
 	ALandscapeSplineActor* SplineActor = World->SpawnActor<ALandscapeSplineActor>(Location, FRotator::ZeroRotator, SpawnParams);
 	SplineActor->GetSharedProperties(this);
 	SplineActor->GetSplinesComponent()->ShowSplineEditorMesh(true);
+	SplineActor->SetIsSpatiallyLoaded(AreNewLandscapeActorsSpatiallyLoaded());
 	RegisterSplineActor(SplineActor);
 	return SplineActor;
 }
diff --git a/Engine/Source/Runtime/Landscape/Private/LandscapeConfigHelper.cpp b/Engine/Source/Runtime/Landscape/Private/LandscapeConfigHelper.cpp
--- a/Engine/Source/Runtime/Landscape/Private/LandscapeConfigHelper.cpp
+++ b/Engine/Source/Runtime/Landscape/Private/LandscapeConfigHelper.cpp
@@ -74,6 +74,7 @@ ALandscapeProxy* FLandscapeConfigHelper::FindOrAddLandscapeStreamingProxy(UActor
 		LandscapeProxy->CreateLandscapeInfo();
 		LandscapeProxy->SetActorLocationAndRotation(ProxyLocation, Landscape->GetActorRotation());
 		LandscapeProxy->LandscapeSectionOffset = FIntPoint(CellLocation.X, CellLocation.Y);
+		LandscapeProxy->SetIsSpatiallyLoaded(LandscapeProxy->GetLandscapeInfo()->AreNewLandscapeActorsSpatiallyLoaded());
 	};
 
 	const bool bCreate = true;
diff --git a/Engine/Source/Runtime/Landscape/Private/LandscapeEdit.cpp b/Engine/Source/Runtime/Landscape/Private/LandscapeEdit.cpp
--- a/Engine/Source/Runtime/Landscape/Private/LandscapeEdit.cpp
+++ b/Engine/Source/Runtime/Landscape/Private/LandscapeEdit.cpp
@@ -4418,6 +4418,16 @@ uint32 ULandscapeInfo::GetGridSize(uint32 InGridSizeInComponents) const
 	return InGridSizeInComponents * ComponentSizeQuads;
 }
 
+bool ULandscapeInfo::AreNewLandscapeActorsSpatiallyLoaded() const
+{
+	if (ALandscape* Landscape = LandscapeActor.Get())
+	{
+		return Landscape->bAreNewLandscapeActorsSpatiallyLoaded;
+	}
+	
+	return false;
+}
+
 ALandscapeProxy* ULandscapeInfo::MoveComponentsToLevel(const TArray<ULandscapeComponent*>& InComponents, ULevel* TargetLevel, FName NewProxyName)
 {
 	ALandscape* Landscape = LandscapeActor.Get();
