New Landscapes will always have GridSize in their name

diff --git a/Engine/Source/Runtime/Engine/Private/ActorPartition/ActorPartitionSubsystem.cpp b/Engine/Source/Runtime/Engine/Private/ActorPartition/ActorPartitionSubsystem.cpp
--- a/Engine/Source/Runtime/Engine/Private/ActorPartition/ActorPartitionSubsystem.cpp
+++ b/Engine/Source/Runtime/Engine/Private/ActorPartition/ActorPartitionSubsystem.cpp
@@ -254,7 +254,7 @@
 				ActorNameBuilder += TEXT("_");
 			}
 
-			if (InActorPartitionId.GetClass()->GetDefaultObject<APartitionActor>()->ShouldIncludeGridSizeInName(World))
+			if (InActorPartitionId.GetClass()->GetDefaultObject<APartitionActor>()->ShouldIncludeGridSizeInName(World, InActorPartitionId))
 			{
 				ActorNameBuilder += FString::Printf(TEXT("%d_"), InGridSize);
 			}
diff --git a/Engine/Source/Runtime/Engine/Private/ActorPartition/PartitionActor.cpp b/Engine/Source/Runtime/Engine/Private/ActorPartition/PartitionActor.cpp
--- a/Engine/Source/Runtime/Engine/Private/ActorPartition/PartitionActor.cpp
+++ b/Engine/Source/Runtime/Engine/Private/ActorPartition/PartitionActor.cpp
@@ -45,7 +45,7 @@
 	return true;
 }
 
-bool APartitionActor::ShouldIncludeGridSizeInName(UWorld* InWorld) const
+bool APartitionActor::ShouldIncludeGridSizeInName(UWorld* InWorld, const FActorPartitionIdentifier& InIdentifier) const
 {
 	return InWorld->GetWorldSettings()->bIncludeGridSizeInNameForPartitionedActors;
 }
diff --git a/Engine/Source/Runtime/Engine/Public/ActorPartition/PartitionActor.h b/Engine/Source/Runtime/Engine/Public/ActorPartition/PartitionActor.h
--- a/Engine/Source/Runtime/Engine/Public/ActorPartition/PartitionActor.h
+++ b/Engine/Source/Runtime/Engine/Public/ActorPartition/PartitionActor.h
@@ -6,7 +6,9 @@
 #include "Misc/Guid.h"
 #include "PartitionActor.generated.h"
 
-class UBoxComponent;
+#if WITH_EDITOR
+struct FActorPartitionIdentifier;
+#endif
 
 // Actor base class for instance containers placed on a grid.
 // See UActorPartitionSubsystem.
@@ -23,7 +25,7 @@
 	virtual bool CanChangeIsSpatiallyLoadedFlag() const override { return false; }
 	virtual TUniquePtr<class FWorldPartitionActorDesc> CreateClassActorDesc() const override;
 	virtual uint32 GetDefaultGridSize(UWorld* InWorld) const PURE_VIRTUAL(APartitionActor, return 0;)
-	virtual bool ShouldIncludeGridSizeInName(UWorld * InWorld) const;
+	virtual bool ShouldIncludeGridSizeInName(UWorld * InWorld, const FActorPartitionIdentifier& InIdentifier) const;
 	virtual FGuid GetGridGuid() const { return FGuid(); }
 	virtual bool IsUserManaged() const override;
 #endif
diff --git a/Engine/Source/Runtime/Foliage/Private/InstancedFoliage.cpp b/Engine/Source/Runtime/Foliage/Private/InstancedFoliage.cpp
--- a/Engine/Source/Runtime/Foliage/Private/InstancedFoliage.cpp
+++ b/Engine/Source/Runtime/Foliage/Private/InstancedFoliage.cpp
@@ -4969,7 +4969,7 @@
 	return InWorld->GetWorldSettings()->InstancedFoliageGridSize;
 }
 
-bool AInstancedFoliageActor::ShouldIncludeGridSizeInName(UWorld* InWorld) const
+bool AInstancedFoliageActor::ShouldIncludeGridSizeInName(UWorld* InWorld, const FActorPartitionIdentifier& InIdentifier) const
 {
 	return InWorld->GetWorldSettings()->bIncludeGridSizeInNameForFoliageActors;
 }
diff --git a/Engine/Source/Runtime/Foliage/Public/InstancedFoliageActor.h b/Engine/Source/Runtime/Foliage/Public/InstancedFoliageActor.h
--- a/Engine/Source/Runtime/Foliage/Public/InstancedFoliageActor.h
+++ b/Engine/Source/Runtime/Foliage/Public/InstancedFoliageActor.h
@@ -15,7 +15,10 @@
 
 #include "InstancedFoliageActor.generated.h"
 
+#if WITH_EDITOR
 class UProceduralFoliageComponent;
+struct FActorPartitionIdentifier;
+#endif
 
 // Function for filtering out hit components during FoliageTrace
 typedef TFunction<bool(const UPrimitiveComponent*)> FFoliageTraceFilterFunc;
@@ -99,7 +102,7 @@
 	FOLIAGE_API void DetectFoliageTypeChangeAndUpdate();
 
 	virtual uint32 GetDefaultGridSize(UWorld* InWorld) const override;
-	virtual bool ShouldIncludeGridSizeInName(UWorld* InWorld) const override;
+	virtual bool ShouldIncludeGridSizeInName(UWorld* InWorld, const FActorPartitionIdentifier& InIdentifier) const override;
 
 	// Delegate type for selection change events
 	DECLARE_MULTICAST_DELEGATE_TwoParams(FOnSelectionChanged, bool, const TArray<AActor*>&);
diff --git a/Engine/Source/Runtime/Landscape/Classes/Landscape.h b/Engine/Source/Runtime/Landscape/Classes/Landscape.h
--- a/Engine/Source/Runtime/Landscape/Classes/Landscape.h
+++ b/Engine/Source/Runtime/Landscape/Classes/Landscape.h
@@ -243,6 +243,7 @@
 	virtual void PostEditMove(bool bFinished) override;
 	virtual void PostEditUndo() override;
 	virtual void PostRegisterAllComponents() override;
+	virtual void PostActorCreated() override;
 	virtual bool ShouldImport(FString* ActorPropString, bool IsMovingLevel) override;
 	virtual void PostEditImport() override;
 	virtual void PostDuplicate(bool bDuplicateForPIE) override;
@@ -439,6 +440,10 @@
 	UPROPERTY(EditAnywhere, Category = WorldPartition)
 	bool bAreNewLandscapeActorsSpatiallyLoaded = true;
 
+	/** If true, LandscapeStreamingProxy actors have the grid size included in their name, for backward compatibility we also check the AWorldSettings::bIncludeGridSizeInNameForPartitionedActors */
+	UPROPERTY()
+	bool bIncludeGridSizeInNameForLandscapeActors = false;
+
 	UPROPERTY(EditAnywhere, Category=Landscape)
 	bool bCanHaveLayersContent = false;
 
diff --git a/Engine/Source/Runtime/Landscape/Classes/LandscapeStreamingProxy.h b/Engine/Source/Runtime/Landscape/Classes/LandscapeStreamingProxy.h
--- a/Engine/Source/Runtime/Landscape/Classes/LandscapeStreamingProxy.h
+++ b/Engine/Source/Runtime/Landscape/Classes/LandscapeStreamingProxy.h
@@ -9,7 +9,11 @@
 #include "LandscapeStreamingProxy.generated.h"
 
 class ALandscape;
+
+#if WITH_EDITOR
 class UMaterialInterface;
+struct FActorPartitionIdentifier;
+#endif
 
 UCLASS(MinimalAPI, notplaceable)
 class ALandscapeStreamingProxy : public ALandscapeProxy
@@ -32,6 +36,7 @@
 	virtual bool CanDeleteSelectedActor(FText& OutReason) const override;
 	virtual bool GetReferencedContentObjects(TArray<UObject*>& Objects) const override;
 	virtual bool CanChangeIsSpatiallyLoadedFlag() const override { return true; }
+	virtual bool ShouldIncludeGridSizeInName(UWorld* InWorld, const FActorPartitionIdentifier& InIdentifier) const override;
 #endif
 	//~ End UObject Interface
 
diff --git a/Engine/Source/Runtime/Landscape/Private/LandscapeConfigHelper.cpp b/Engine/Source/Runtime/Landscape/Private/LandscapeConfigHelper.cpp
--- a/Engine/Source/Runtime/Landscape/Private/LandscapeConfigHelper.cpp
+++ b/Engine/Source/Runtime/Landscape/Private/LandscapeConfigHelper.cpp
@@ -54,10 +54,10 @@
 	const uint32 GridSize = InLandscapeInfo->LandscapeActor->GridSize;
 
 	UActorPartitionSubsystem::FCellCoord CellCoord = UActorPartitionSubsystem::FCellCoord::GetCellCoord(InSectionBase, World->PersistentLevel, GridSize);
-	return FLandscapeConfigHelper::FindOrAddLandscapeStreamingProxy(ActorPartitionSubsystem, InLandscapeInfo, CellCoord, GridSize);
+	return FLandscapeConfigHelper::FindOrAddLandscapeStreamingProxy(ActorPartitionSubsystem, InLandscapeInfo, CellCoord);
 }
 
-ALandscapeProxy* FLandscapeConfigHelper::FindOrAddLandscapeStreamingProxy(UActorPartitionSubsystem* InActorPartitionSubsystem, ULandscapeInfo* InLandscapeInfo, const UActorPartitionSubsystem::FCellCoord& InCellCoord, uint32 InSearchGridSize)
+ALandscapeProxy* FLandscapeConfigHelper::FindOrAddLandscapeStreamingProxy(UActorPartitionSubsystem* InActorPartitionSubsystem, ULandscapeInfo* InLandscapeInfo, const UActorPartitionSubsystem::FCellCoord& InCellCoord)
 {
 	ALandscape* Landscape = InLandscapeInfo->LandscapeActor.Get();
 	check(Landscape);
@@ -80,13 +80,8 @@
 	const bool bCreate = true;
 	const bool bBoundsSearch = false;
 
-	// SearchGridSize might be different than Landscape->GridSize when we are trying to recycle Proxies if the GridSize is being changed.
-	ALandscapeProxy* LandscapeProxy = Cast<ALandscapeProxy>(InActorPartitionSubsystem->GetActor(ALandscapeStreamingProxy::StaticClass(), InCellCoord, bCreate, InLandscapeInfo->LandscapeGuid, InSearchGridSize, bBoundsSearch, LandscapeProxyCreated));
-	if (LandscapeProxy && LandscapeProxy->GridSize != Landscape->GridSize)
-	{
-		LandscapeProxy->Modify();
-		LandscapeProxy->GridSize = Landscape->GridSize;
-	}
+	ALandscapeProxy* LandscapeProxy = Cast<ALandscapeProxy>(InActorPartitionSubsystem->GetActor(ALandscapeStreamingProxy::StaticClass(), InCellCoord, bCreate, InLandscapeInfo->LandscapeGuid, Landscape->GridSize, bBoundsSearch, LandscapeProxyCreated));
+	check(LandscapeProxy->GridSize == Landscape->GridSize);
 	return LandscapeProxy;
 }
 
@@ -95,7 +90,6 @@
 	check(InLandscapeInfo);
 
 	const uint32 GridSize = InLandscapeInfo->GetGridSize(InNewGridSizeInComponents);
-	const uint32 OldGridSize = InLandscapeInfo->LandscapeActor->GridSize;
 	
 	InLandscapeInfo->LandscapeActor->Modify();
 	InLandscapeInfo->LandscapeActor->GridSize = GridSize;
@@ -103,6 +97,9 @@
 	// This needs to be done before moving components
 	InLandscapeInfo->LandscapeActor->InitializeLandscapeLayersWeightmapUsage();
 
+	// Make sure if actor didn't include grid size in name it now does. This will avoid recycling 
+	// LandscapeStreamingProxy actors and create new ones with the proper name.
+	InLandscapeInfo->LandscapeActor->bIncludeGridSizeInNameForLandscapeActors = true;
 
 	FIntRect Extent;
 	InLandscapeInfo->GetLandscapeExtent(Extent.Min.X, Extent.Min.Y, Extent.Max.X, Extent.Max.Y);
@@ -120,7 +117,7 @@
 
 	TSet<ALandscapeProxy*> ProxiesToDelete;
 	
-	FActorPartitionGridHelper::ForEachIntersectingCell(ALandscapeStreamingProxy::StaticClass(), Extent, World->PersistentLevel, [ActorPartitionSubsystem, InLandscapeInfo, InNewGridSizeInComponents, &LandscapeComponents, OldGridSize, &ProxiesToDelete](const UActorPartitionSubsystem::FCellCoord& CellCoord, const FIntRect& CellBounds)
+	FActorPartitionGridHelper::ForEachIntersectingCell(ALandscapeStreamingProxy::StaticClass(), Extent, World->PersistentLevel, [ActorPartitionSubsystem, InLandscapeInfo, InNewGridSizeInComponents, &LandscapeComponents, &ProxiesToDelete](const UActorPartitionSubsystem::FCellCoord& CellCoord, const FIntRect& CellBounds)
 	{
 		TMap<ULandscapeComponent*, UMaterialInterface*> ComponentMaterials;
 		TMap<ULandscapeComponent*, UMaterialInterface*> ComponentHoleMaterials;
@@ -155,7 +152,7 @@
 		check(ComponentsToMove.Num() <= MaxComponents);
 		if (ComponentsToMove.Num())
 		{
-			ALandscapeProxy* LandscapeProxy = FLandscapeConfigHelper::FindOrAddLandscapeStreamingProxy(ActorPartitionSubsystem, InLandscapeInfo, CellCoord, OldGridSize);
+			ALandscapeProxy* LandscapeProxy = FLandscapeConfigHelper::FindOrAddLandscapeStreamingProxy(ActorPartitionSubsystem, InLandscapeInfo, CellCoord);
 			check(LandscapeProxy);
 			InLandscapeInfo->MoveComponentsToProxy(ComponentsToMove, LandscapeProxy);
 
@@ -521,7 +518,7 @@
 			}
 
 			UActorPartitionSubsystem::FCellCoord CellCoord = UActorPartitionSubsystem::FCellCoord::GetCellCoord(NewSectionBase, World->PersistentLevel, NewLandscape->GridSize);
-			ALandscapeProxy* LandscapeProxy = FLandscapeConfigHelper::FindOrAddLandscapeStreamingProxy(ActorPartitionSubsystem, NewLandscapeInfo, CellCoord, NewLandscape->GridSize);
+			ALandscapeProxy* LandscapeProxy = FLandscapeConfigHelper::FindOrAddLandscapeStreamingProxy(ActorPartitionSubsystem, NewLandscapeInfo, CellCoord);
 			check(LandscapeProxy);
 
 			ULandscapeComponent* NewComponent = NewObject<ULandscapeComponent>(LandscapeProxy, NAME_None, RF_Transactional);
diff --git a/Engine/Source/Runtime/Landscape/Private/LandscapeEdit.cpp b/Engine/Source/Runtime/Landscape/Private/LandscapeEdit.cpp
--- a/Engine/Source/Runtime/Landscape/Private/LandscapeEdit.cpp
+++ b/Engine/Source/Runtime/Landscape/Private/LandscapeEdit.cpp
@@ -4137,6 +4137,14 @@
 	}
 }
 
+void ALandscape::PostActorCreated()
+{
+	Super::PostActorCreated();
+
+	// Newly spawned Landscapes always set this value to true
+	bIncludeGridSizeInNameForLandscapeActors = true;
+}
+
 bool ALandscape::ShouldImport(FString* ActorPropString, bool IsMovingLevel)
 {
 	return GetWorld() != nullptr && !GetWorld()->IsGameWorld();
@@ -5234,6 +5242,27 @@
 	return true;
 }
 
+bool ALandscapeStreamingProxy::ShouldIncludeGridSizeInName(UWorld* InWorld, const FActorPartitionIdentifier& InIdentifier) const
+{
+	// Always return true if this world setting flag is true
+	if (Super::ShouldIncludeGridSizeInName(InWorld, InIdentifier))
+	{
+		return true;
+	}
+
+	if (ULandscapeInfo* LandscapeInfo = ULandscapeInfo::Find(InWorld, InIdentifier.GetGridGuid()))
+	{
+		if (ALandscape* Landscape = LandscapeInfo->LandscapeActor.Get())
+		{
+			// This new flag is to support Landscapes that were created with bIncludeGridSizeInNameForPartitionedActors == false or 
+			// that were reconfigured with FLandscapeConfigHelper::ChangeGridSize
+			return Landscape->bIncludeGridSizeInNameForLandscapeActors;
+		}
+	}
+
+	return false;
+}
+
 bool ALandscape::CanDeleteSelectedActor(FText& OutReason) const
 {
 	if (!IsUserManaged())
diff --git a/Engine/Source/Runtime/Landscape/Public/LandscapeConfigHelper.h b/Engine/Source/Runtime/Landscape/Public/LandscapeConfigHelper.h
--- a/Engine/Source/Runtime/Landscape/Public/LandscapeConfigHelper.h
+++ b/Engine/Source/Runtime/Landscape/Public/LandscapeConfigHelper.h
@@ -72,7 +72,7 @@
 	LANDSCAPE_API static bool ChangeGridSize(ULandscapeInfo* InLandscapeInfo, uint32 InNewGridSizeInComponents, TSet<AActor*>& OutActorsToDelete);
 
 private:
-	static ALandscapeProxy* FindOrAddLandscapeStreamingProxy(UActorPartitionSubsystem* ActorPartitionSubsystem, ULandscapeInfo* LandscapeInfo, const UActorPartitionSubsystem::FCellCoord& CellCoord, uint32 InSearchGridSize);
+	static ALandscapeProxy* FindOrAddLandscapeStreamingProxy(UActorPartitionSubsystem* ActorPartitionSubsystem, ULandscapeInfo* LandscapeInfo, const UActorPartitionSubsystem::FCellCoord& CellCoord);
 	static void CopyRegionToComponent(ULandscapeInfo* LandscapeInfo, const FIntRect& Region, bool bResample, ULandscapeComponent* Component);
 	static void ExtractLandscapeData(ULandscapeInfo* LandscapeInfo, const FIntRect& Region, const FGuid& LayerGuid, TArray<uint16>& OutHeightData, TArray<FLandscapeImportLayerInfo>& OutImportMaterialLayerInfos);
 	static void MoveSplinesToLandscape(ULandscapeInfo* InLandscapeInfo, ALandscapeProxy* InLandscape, float InScaleFactor);
