Submitting on behalf of Jonathan.Bard, see review 19835889

Landscape: added utility function to render the merged heightmap from anywhere in world space. This is done in 2 steps : render the heightmaps  of all landscape components overlapping with the requested region into an atlas texture (uncompressed) subsection by subsection, skipping the duplicated rows/columns. Then resample the atlas at the requested location.
If the render target is 8bits/channel, the resulting heightmap will be compressed on the R and G channels. If the target is 16 bits uint, the output will contain the uncompressed height value.
The function supports heightmaps with different minimum mip level and will use the best available mip level. However, currently the intermediate atlas render target is always allocated using the landscape's full resolution.
This is meant to be compatible with runtime but is not at the moment, because it relies on ULandscapeInfo members which are not editor-only but are not available at runtime (yet). This will be fixed later on.

diff --git a/Engine/Shaders/Private/Landscape/LandscapeMergeHeightmapsPS.usf b/Engine/Shaders/Private/Landscape/LandscapeMergeHeightmapsPS.usf
new file mode 100644
--- /dev/null
+++ b/Engine/Shaders/Private/Landscape/LandscapeMergeHeightmapsPS.usf
@@ -0,0 +1,83 @@
+// Copyright Epic Games, Inc. All Rights Reserved.
+
+#include "../Common.ush"
+#include "LandscapeCommon.ush"
+
+// ----------------------------------------------------------------------------------
+
+#if defined (__INTELLISENSE__)
+// Uncomment the appropriate define for enabling syntax highlighting with HLSL Tools for Visual Studio : 
+//#define MERGE_HEIGHTMAP 1
+//#define RESAMPLE_HEIGHTMAP 1
+#endif // defined (__INTELLISENSE__)
+
+// ----------------------------------------------------------------------------------
+
+#if MERGE_HEIGHTMAP
+// MergeHeightmap inputs/outputs :
+
+uint4 InHeightmapAtlasSubregion; // Subregion of the atlas texture that is being rendered to
+uint4 InHeightmapSubregion; // Subregion of the input heightmap that we copy from
+Texture2D<float4> InHeightmap; // Source texture (packed height in the .rg channels)
+#endif // MERGE_HEIGHTMAP
+
+#if RESAMPLE_HEIGHTMAP
+// ResampleHeightmap inputs/outputs :
+
+float4x4 InOutputUVToMergedHeightmapUV; // Transform to apply to the texture coordinates in order to resample the heightmap in the requested render area
+Texture2D<float> InMergedHeightmap; // Source texture (unpacked floating point height)
+SamplerState InMergedHeightmapSampler;
+uint2 InRenderAreaSize; // Size of the output render target
+#endif // RESAMPLE_HEIGHTMAP
+
+static const float LANDSCAPE_MAX_VALUE = 65535.0f;
+
+// ----------------------------------------------------------------------------------
+// Util functions : 
+
+#if MERGE_HEIGHTMAP
+
+void MergeHeightmap(in float4 SVPos : SV_POSITION, out float OutHeight : SV_Target0)
+{
+	// Input heightmap and output heightmap atlas may not be at the same resolution (different mips), so we need to renormalize the sample position before sampling the input heightmap :
+	uint2 AtlasSubregionSize = InHeightmapAtlasSubregion.zw - InHeightmapAtlasSubregion.xy;
+	float2 SampleCoordinatesNormalized = (SVPos.xy - InHeightmapAtlasSubregion.xy) / (float2) AtlasSubregionSize;
+	uint2 HeightmapSubregionSize = InHeightmapSubregion.zw - InHeightmapSubregion.xy;
+	uint2 SampleCoordinates = floor(InHeightmapSubregion.xy + SampleCoordinatesNormalized * HeightmapSubregionSize);
+		
+	float4 HeightmapSample = InHeightmap.Load(int3(SampleCoordinates, 0));
+	OutHeight = UnpackHeight(HeightmapSample.xy)/LANDSCAPE_MAX_VALUE;
+}
+
+#endif // MERGE_HEIGHTMAP
+
+
+// ----------------------------------------------------------------------------------
+
+#if RESAMPLE_HEIGHTMAP
+
+void ResampleHeightmap(in float4 SVPos : SV_POSITION, out float4 OutColor : SV_Target0)
+{
+	OutColor = 0;
+
+	// Pixel coordinates to UV (SvPos is already offset by 0.5 to be at pixel center) :
+	float2 UV = SVPos.xy / (float2)InRenderAreaSize;
+	float4 SampleUV = float4(UV, 0.0f, 1.0f);
+	SampleUV = mul(SampleUV, InOutputUVToMergedHeightmapUV);
+
+	float Height = 0.0f;
+	if (all(SampleUV.xy >= 0.0f) && all(SampleUV.xy <= 1.0f))
+	{
+		Height = InMergedHeightmap.Sample(InMergedHeightmapSampler, SampleUV.xy);
+	}
+
+	Height *= LANDSCAPE_MAX_VALUE;
+
+#if COMPRESS_HEIGHT
+	OutColor.xy = PackHeight(Height);
+#else // COMPRESS_HEIGHT
+	OutColor.x = Height;
+#endif // !COMPRESS_HEIGHT
+}
+
+#endif // RESAMPLE_HEIGHTMAP
diff --git a/Engine/Source/Runtime/Landscape/Classes/Landscape.h b/Engine/Source/Runtime/Landscape/Classes/Landscape.h
--- a/Engine/Source/Runtime/Landscape/Classes/Landscape.h
+++ b/Engine/Source/Runtime/Landscape/Classes/Landscape.h
@@ -215,9 +215,9 @@
 	//~ Begin ALandscapeProxy Interface
 	LANDSCAPE_API virtual ALandscape* GetLandscapeActor() override;
 	LANDSCAPE_API virtual const ALandscape* GetLandscapeActor() const override;
-#if WITH_EDITOR
 	//~ End ALandscapeProxy Interface
 
+#if WITH_EDITOR
 	LANDSCAPE_API bool HasAllComponent(); // determine all component is in this actor
 	
 	// Include Components with overlapped vertices
@@ -249,7 +249,8 @@
 	virtual void PostDuplicate(bool bDuplicateForPIE) override;
 	virtual bool CanDeleteSelectedActor(FText& OutReason) const override;
 	virtual bool CanChangeIsSpatiallyLoadedFlag() const override { return false; }
-#endif
+#endif // WITH_EDITOR
+
 	virtual void PostLoad() override;
 	virtual void BeginDestroy() override;
 	virtual void FinishDestroy() override;
@@ -258,15 +259,17 @@
 	/** Computes & returns bounds containing all currently loaded landscape proxies (if any) or this landscape's bounds otherwise */
 	LANDSCAPE_API FBox GetLoadedBounds() const;
 
-#if WITH_EDITOR
-	/** Computes & returns bounds containing all landscape proxies (if any) or this landscape's bounds otherwise. Note that in non-WP worlds this will call GetLoadedBounds(). */
-	LANDSCAPE_API FBox GetCompleteBounds() const;
-#endif
-
 	LANDSCAPE_API bool IsUpToDate() const;
 	LANDSCAPE_API void TickLayers(float DeltaTime);
-	// Layers stuff
+
+	UFUNCTION(BlueprintCallable, Category = "Landscape|Runtime")
+	LANDSCAPE_API void RenderHeightmap(const FTransform& InWorldTransform, const FBox2D& InExtents, UTextureRenderTarget2D* OutRenderTarget);
+
+	bool IsValidRenderTargetFormatHeightmap(EPixelFormat InRenderTargetFormat, bool& bOutCompressHeight);
+
 #if WITH_EDITOR
+	/** Computes & returns bounds containing all landscape proxies (if any) or this landscape's bounds otherwise. Note that in non-WP worlds this will call GetLoadedBounds(). */
+	LANDSCAPE_API FBox GetCompleteBounds() const;
 	LANDSCAPE_API void RegisterLandscapeEdMode(ILandscapeEdModeInterface* InLandscapeEdMode) { LandscapeEdMode = InLandscapeEdMode; }
 	LANDSCAPE_API void UnregisterLandscapeEdMode() { LandscapeEdMode = nullptr; }
 	LANDSCAPE_API virtual bool HasLayersContent() const override;
@@ -347,6 +350,9 @@
 	void UpdateProxyLayersWeightmapUsage();
 	void ValidateProxyLayersWeightmapUsage() const;
 
+protected:
+	FName GenerateUniqueLayerName(FName InName = NAME_None) const;
+
 private:
 	bool SupportsEditLayersLocalMerge();
 	void CreateLayersRenderingResource();
@@ -431,7 +437,7 @@
 	static bool IsMaterialResourceCompiled(FMaterialResource* InMaterialResource, bool bInWaitForCompilation);
 	static void ShowEditLayersResourcesNotification(const FText& InText, TWeakPtr<SNotificationItem>& NotificationItem);
 	static void HideEditLayersResourcesNotification(TWeakPtr<SNotificationItem>& InNotificationItem);
-#endif
+#endif // WITH_EDITOR
 
 public:
 
@@ -533,11 +539,6 @@
 	// Used in packing the material layer data contained into CombinedLayersWeightmapAllMaterialLayersResource to be set again for each component weightmap (size of the landscape)
 	class FLandscapeTexture2DResource* WeightmapScratchPackLayerTextureResource;
 #endif
-
-protected:
-#if WITH_EDITOR
-	FName GenerateUniqueLayerName(FName InName = NAME_None) const;
-#endif
 };
 
 #if WITH_EDITOR
diff --git a/Engine/Source/Runtime/Landscape/Classes/LandscapeInfo.h b/Engine/Source/Runtime/Landscape/Classes/LandscapeInfo.h
--- a/Engine/Source/Runtime/Landscape/Classes/LandscapeInfo.h
+++ b/Engine/Source/Runtime/Landscape/Classes/LandscapeInfo.h
@@ -341,6 +341,18 @@ public:
 	LANDSCAPE_API void RegisterCollisionComponent(ULandscapeHeightfieldCollisionComponent* Component);
 	LANDSCAPE_API void UnregisterCollisionComponent(ULandscapeHeightfieldCollisionComponent* Component);
 
+	/**
+	 * Retrieve the components currently loaded that overlap with a given "window" area
+	 * @param InAreaWorldTransform : transform of the requested area (at the center)
+	 * @param InAreaExtents : extents of requested area (i.e. around the center)
+	 * @param OutOverlappedComponents : loaded components that overlap this area (key = xy index of the component, value = component)
+	 * @param OutComponentIndicesBoundingRect : bounding rectangle of the overlapped components, in component index space, with
+	 *  the max being exclusive. For instance a rectangle with min=(0,0) and max=(2,1) includes components (0,0) and (1,0).
+	 * 
+	 * @return true if at least one overlapped component
+	 */
+	LANDSCAPE_API bool GetOverlappedComponents(const FTransform& InAreaWorldTransform, const FBox2D& InAreaExtents, TMap<FIntPoint, ULandscapeComponent*>& OutOverlappedComponents, FIntRect& OutComponentIndicesBoundingRect);
+
 #if WITH_EDITOR
 	LANDSCAPE_API ALandscapeSplineActor* CreateSplineActor(const FVector& Location);
 
diff --git a/Engine/Source/Runtime/Landscape/Private/Landscape.cpp b/Engine/Source/Runtime/Landscape/Private/Landscape.cpp
--- a/Engine/Source/Runtime/Landscape/Private/Landscape.cpp
+++ b/Engine/Source/Runtime/Landscape/Private/Landscape.cpp
@@ -52,6 +52,7 @@ Landscape.cpp: Terrain rendering
 #include "LandscapeSplinesComponent.h"
 #include "EngineGlobals.h"
 #include "Engine/Engine.h"
+#include "Engine/TextureRenderTarget2D.h"
 #include "EngineUtils.h"
 #include "ComponentRecreateRenderStateContext.h"
 #include "LandscapeWeightmapUsage.h"
@@ -59,6 +60,13 @@ Landscape.cpp: Terrain rendering
 #include "Streaming/LandscapeMeshMobileUpdate.h"
 #include "ContentStreaming.h"
 #include "UObject/ObjectSaveContext.h"
+#include "GlobalShader.h"
+#include "ShaderParameterStruct.h"
+#include "PixelShaderUtils.h"
+#include "LandscapeEditResources.h"
+#include "Rendering/Texture2DResource.h"
+#include "RenderCaptureInterface.h"
+#include "VisualLogger/VisualLogger.h"
 
 #if WITH_EDITOR
 #include "LandscapeEdit.h"
@@ -160,6 +168,13 @@ FAutoConsoleCommand CmdPrintNumLandscapeShadows(
 	FConsoleCommandDelegate::CreateStatic(PrintNumLandscapeShadows)
 	);
 
+int32 RenderCaptureNextHeightmapRenders = 0;
+static FAutoConsoleVariableRef CVarRenderCaptureNextHeightmapRenders(
+	TEXT("landscape.RenderCaptureNextHeightmapRenders"),
+	RenderCaptureNextHeightmapRenders,
+	TEXT("Trigger a render capture during the next N RenderHeightmap draws"));
+
+
 ULandscapeComponent::ULandscapeComponent(const FObjectInitializer& ObjectInitializer)
 	: Super(ObjectInitializer)
 #if WITH_EDITORONLY_DATA
@@ -2258,14 +2273,389 @@ FBox ALandscape::GetLoadedBounds() const
 	return GetLandscapeInfo()->GetLoadedBounds();
 }
 
+
+// ----------------------------------------------------------------------------------
+
+// This shader allows to render parts of the heightmaps (all pixels except the redundant ones on the right/bottom edges) in an atlas render target (uncompressed height)
+class FLandscapeMergeHeightmapsPS : public FGlobalShader
+{
+	DECLARE_GLOBAL_SHADER(FLandscapeMergeHeightmapsPS);
+	SHADER_USE_PARAMETER_STRUCT(FLandscapeMergeHeightmapsPS, FGlobalShader);
+
+public:
+	BEGIN_SHADER_PARAMETER_STRUCT(FParameters, )
+		SHADER_PARAMETER(FUintVector4, InHeightmapAtlasSubregion)
+		SHADER_PARAMETER(FUintVector4, InHeightmapSubregion)
+		SHADER_PARAMETER_RDG_TEXTURE_SRV(Texture2D<float4>, InHeightmap)
+		RENDER_TARGET_BINDING_SLOTS()
+		END_SHADER_PARAMETER_STRUCT()
+
+		static bool ShouldCompilePermutation(const FGlobalShaderPermutationParameters& InParameters)
+	{
+		return true;
+	}
+
+	static void ModifyCompilationEnvironment(const FGlobalShaderPermutationParameters& InParameters, FShaderCompilerEnvironment& OutEnvironment)
+	{
+		OutEnvironment.SetDefine(TEXT("MERGE_HEIGHTMAP"), 1);
+	}
+
+	static void MergeHeightmap(FRDGBuilder& GraphBuilder, FParameters* InParameters, const FIntRect& InRenderTargetArea)
+	{
+		FGlobalShaderMap* ShaderMap = GetGlobalShaderMap(GMaxRHIFeatureLevel);
+		TShaderMapRef<FLandscapeMergeHeightmapsPS> PixelShader(ShaderMap);
+
+		FPixelShaderUtils::AddFullscreenPass(
+			GraphBuilder,
+			ShaderMap,
+			RDG_EVENT_NAME("LandscapeLayers_MergeHeightmap"),
+			PixelShader,
+			InParameters,
+			InRenderTargetArea);
+	}
+};
+
+IMPLEMENT_GLOBAL_SHADER(FLandscapeMergeHeightmapsPS, "/Engine/Private/Landscape/LandscapeMergeHeightmapsPS.usf", "MergeHeightmap", SF_Pixel);
+
+
+// ----------------------------------------------------------------------------------
+
+// This shader allows to resample the heightmap (bilinear interpolation) from a given atlas usually heightmap produced by FLandscapeMergeHeightmapsPS :
+//  The output heightmap's heights can be either compressed or uncompressed depending on the render target format (8 bits/channel for the former, 16/32 bits/channel for the latter)
+class FLandscapeResampleHeightmapsPS : public FGlobalShader
+{
+	DECLARE_GLOBAL_SHADER(FLandscapeResampleHeightmapsPS);
+	SHADER_USE_PARAMETER_STRUCT(FLandscapeResampleHeightmapsPS, FGlobalShader);
+
+public:
+	BEGIN_SHADER_PARAMETER_STRUCT(FParameters, )
+		SHADER_PARAMETER(FMatrix44f, InOutputUVToMergedHeightmapUV)
+		SHADER_PARAMETER_RDG_TEXTURE_SRV(Texture2D<float4>, InMergedHeightmap)
+		SHADER_PARAMETER_SAMPLER(SamplerState, InMergedHeightmapSampler)
+		SHADER_PARAMETER(FUintVector2, InRenderAreaSize)
+		RENDER_TARGET_BINDING_SLOTS()
+		END_SHADER_PARAMETER_STRUCT()
+
+		class FCompressHeight : SHADER_PERMUTATION_BOOL("COMPRESS_HEIGHT");
+	using FPermutationDomain = TShaderPermutationDomain<FCompressHeight>;
+
+	static FPermutationDomain GetPermutationVector(bool bCompressHeight)
+	{
+		FPermutationDomain PermutationVector;
+		PermutationVector.Set<FCompressHeight>(bCompressHeight);
+		return PermutationVector;
+	}
+
+	static bool ShouldCompilePermutation(const FGlobalShaderPermutationParameters& InParameters)
+	{
+		return true;
+	}
+
+	static void ModifyCompilationEnvironment(const FGlobalShaderPermutationParameters& InParameters, FShaderCompilerEnvironment& OutEnvironment)
+	{
+		OutEnvironment.SetDefine(TEXT("RESAMPLE_HEIGHTMAP"), 1);
+	}
+
+	static void ResampleHeightmap(FRDGBuilder& GraphBuilder, FParameters* InParameters, bool bCompressHeight)
+	{
+		FGlobalShaderMap* ShaderMap = GetGlobalShaderMap(GMaxRHIFeatureLevel);
+
+		const FLandscapeResampleHeightmapsPS::FPermutationDomain PixelPermutationVector = FLandscapeResampleHeightmapsPS::GetPermutationVector(bCompressHeight);
+
+		TShaderMapRef<FLandscapeResampleHeightmapsPS> PixelShader(ShaderMap, PixelPermutationVector);
+
+		FPixelShaderUtils::AddFullscreenPass(
+			GraphBuilder,
+			ShaderMap,
+			RDG_EVENT_NAME("LandscapeLayers_ResampleHeightmap"),
+			PixelShader,
+			InParameters,
+			FIntRect(0, 0, InParameters->InRenderAreaSize.X, InParameters->InRenderAreaSize.Y));
+	}
+};
+
+IMPLEMENT_GLOBAL_SHADER(FLandscapeResampleHeightmapsPS, "/Engine/Private/Landscape/LandscapeMergeHeightmapsPS.usf", "ResampleHeightmap", SF_Pixel);
+
+
+// Render-thread version of the data / functions we need for the local merge of edit layers : 
+namespace RenderMergedLandscape_RenderThread
+{
+	struct FHeightmapRenderInfo
+	{
+		// Transform to go from the output render area space ((0,0) in the lower left corner, (1,1) in the upper-right) to the temporary render target space
+		FMatrix OutputUVToMergedHeightmapUV;
+		FBox RenderAreaExtents;
+		// TODO [jonathan.bard] : remove ? This is only used to compute the temporary render target size
+		FIntPoint SubsectionSizeQuads;
+		int32 NumSubsections = 1;
+		bool bCompressHeight = false;
+
+		TMap<FIntPoint, FTexture2DResourceSubregion> ComponentHeightmapsToRender;
+	};
+
+	void RenderMergedHeightmap(const FHeightmapRenderInfo& InRenderInfo, FRDGBuilder& GraphBuilder, FRDGTextureRef OutputTexture)
+	{
+		// Find the total area that those components need to be rendered to :
+		FIntRect ComponentKeyRect;
+		for (auto Iter = InRenderInfo.ComponentHeightmapsToRender.CreateConstIterator(); Iter; ++Iter)
+		{
+			ComponentKeyRect.Include(Iter.Key());
+		}
+
+		ComponentKeyRect.Max += FIntPoint(1, 1);
+		FIntPoint NumComponentsToRender(ComponentKeyRect.Width(), ComponentKeyRect.Height());
+		FIntPoint NumSubsectionsToRender = NumComponentsToRender * InRenderInfo.NumSubsections;
+		FIntPoint RenderTargetSize = NumSubsectionsToRender * InRenderInfo.SubsectionSizeQuads + 1; // add one for the end vertex
+		FIntPoint ComponentSizeQuads = InRenderInfo.SubsectionSizeQuads * InRenderInfo.NumSubsections;
+
+		// We need a temporary render target that can contain all heightmaps. Use PF_G16 (decoded height) as this will be resampled using bilinear sampling :
+		FRDGTextureDesc Desc = FRDGTextureDesc::Create2D(RenderTargetSize, PF_G16, FClearValueBinding::Black, TexCreate_RenderTargetable | TexCreate_ShaderResource);
+		FRDGTextureRef AtlasTexture = GraphBuilder.CreateTexture(Desc, TEXT("LandscapeHeightmapAtlas"));
+		FRDGTextureSRVRef AtlasTextureSRV = GraphBuilder.CreateSRV(FRDGTextureSRVDesc::Create(AtlasTexture));
+		FRenderTargetBinding AtlasTextureRT(AtlasTexture, ERenderTargetLoadAction::ELoad);
+
+		TMap<FTexture2DResource*, FRDGTextureSRVRef> HeightmapTextureSRVs;
+
+		// Fill that render target subsection by subsection, in order to bypass the redundant columns/lines on the subsection edges:
+		for (int32 ComponentY = ComponentKeyRect.Min.Y; ComponentY < ComponentKeyRect.Max.Y; ++ComponentY)
+		{
+			for (int32 ComponentX = ComponentKeyRect.Min.X; ComponentX < ComponentKeyRect.Max.X; ++ComponentX)
+			{
+				FIntPoint LandscapeComponentKey(ComponentX, ComponentY);
+				if (const FTexture2DResourceSubregion* HeightmapResourceSubregion = InRenderInfo.ComponentHeightmapsToRender.Find(LandscapeComponentKey))
+				{
+					FIntPoint SubsectionSubregionSize = HeightmapResourceSubregion->Subregion.Size() / InRenderInfo.NumSubsections;
+					FRDGTextureSRVRef* Heightmap = HeightmapTextureSRVs.Find(HeightmapResourceSubregion->Texture);
+					if (Heightmap == nullptr)
+					{
+						FString* DebugString = GraphBuilder.AllocObject<FString>(HeightmapResourceSubregion->Texture->GetTextureName().ToString());
+						FRDGTextureRef TextureRef = GraphBuilder.RegisterExternalTexture(CreateRenderTarget(HeightmapResourceSubregion->Texture->TextureRHI, **DebugString));
+						Heightmap = &HeightmapTextureSRVs.Add(HeightmapResourceSubregion->Texture, GraphBuilder.CreateSRV(FRDGTextureSRVDesc::Create(TextureRef)));
+					}
+
+					for (int32 SubsectionY = 0; SubsectionY < InRenderInfo.NumSubsections; ++SubsectionY)
+					{
+						for (int32 SubsectionX = 0; SubsectionX < InRenderInfo.NumSubsections; ++SubsectionX)
+						{
+							FIntPoint SubsectionLocalKey(SubsectionX, SubsectionY);
+							FIntPoint SubsectionKey = LandscapeComponentKey * InRenderInfo.NumSubsections + SubsectionLocalKey;
+
+							FIntRect HeightmapAtlasSubregion;
+							HeightmapAtlasSubregion.Min = SubsectionKey * InRenderInfo.SubsectionSizeQuads;
+							// We only really need the +1 on the very last subsection to get the last row/column, since we end up overwriting the other end
+							// rows/columns when we proceed to the next tile. However it's much easier to add the +1 here and do a small amount of duplicate
+							// writes, because otherwise we would have to adjust HeightmapSubregion to align with the region we're writing, which would get
+							// messy in cases of different mip levels.
+							HeightmapAtlasSubregion.Max = HeightmapAtlasSubregion.Min + InRenderInfo.SubsectionSizeQuads + 1;
+
+							FIntRect HeightmapSubregion;
+							HeightmapSubregion.Min = HeightmapResourceSubregion->Subregion.Min + SubsectionLocalKey * SubsectionSubregionSize;
+							HeightmapSubregion.Max = HeightmapSubregion.Min + SubsectionSubregionSize;
+
+							FLandscapeMergeHeightmapsPS::FParameters* MergeHeightmapsPSParams = GraphBuilder.AllocParameters<FLandscapeMergeHeightmapsPS::FParameters>();
+							MergeHeightmapsPSParams->InHeightmapAtlasSubregion = FUintVector4(HeightmapAtlasSubregion.Min.X, HeightmapAtlasSubregion.Min.Y, HeightmapAtlasSubregion.Max.X, HeightmapAtlasSubregion.Max.Y);
+							MergeHeightmapsPSParams->InHeightmapSubregion = FUintVector4(HeightmapSubregion.Min.X, HeightmapSubregion.Min.Y, HeightmapSubregion.Max.X, HeightmapSubregion.Max.Y);
+							MergeHeightmapsPSParams->InHeightmap = *Heightmap;
+							MergeHeightmapsPSParams->RenderTargets[0] = AtlasTextureRT;
+
+							FLandscapeMergeHeightmapsPS::MergeHeightmap(GraphBuilder, MergeHeightmapsPSParams, HeightmapAtlasSubregion);
+						}
+					}
+				}
+				else
+				{
+					// We can clear all subsections of a given component at once : 
+					FIntRect HeightmapAtlasSubregion;
+					HeightmapAtlasSubregion.Min = LandscapeComponentKey * ComponentSizeQuads;
+					HeightmapAtlasSubregion.Max = HeightmapAtlasSubregion.Min + ComponentSizeQuads;
+					FRDGTextureClearInfo ClearInfo;
+					ClearInfo.Viewport = HeightmapAtlasSubregion;
+					AddClearRenderTargetPass(GraphBuilder, AtlasTexture, ClearInfo);
+				}
+			}
+		}
+
+		{
+			FIntVector RenderAreaSize = OutputTexture->Desc.GetSize();
+			FLandscapeResampleHeightmapsPS::FParameters* ResampleHeightmapsPSParams = GraphBuilder.AllocParameters<FLandscapeResampleHeightmapsPS::FParameters>();
+			ResampleHeightmapsPSParams->InOutputUVToMergedHeightmapUV = FMatrix44f(InRenderInfo.OutputUVToMergedHeightmapUV);
+			ResampleHeightmapsPSParams->InMergedHeightmap = AtlasTextureSRV;
+			ResampleHeightmapsPSParams->InMergedHeightmapSampler = TStaticSamplerState<SF_Bilinear, AM_Clamp, AM_Clamp, AM_Clamp>::GetRHI();
+			ResampleHeightmapsPSParams->InRenderAreaSize = FUintVector2((uint32)RenderAreaSize.X, (uint32)RenderAreaSize.Y);
+			ResampleHeightmapsPSParams->RenderTargets[0] = FRenderTargetBinding(OutputTexture, ERenderTargetLoadAction::ENoAction);;
+
+			// We now need to resample the atlas texture where the render area is : 
+			FLandscapeResampleHeightmapsPS::ResampleHeightmap(GraphBuilder, ResampleHeightmapsPSParams, InRenderInfo.bCompressHeight);
+		}
+	}
+}
+
+bool ALandscape::IsValidRenderTargetFormatHeightmap(EPixelFormat InRenderTargetFormat, bool& bOutCompressHeight)
+{
+	bOutCompressHeight = false;
+	switch (InRenderTargetFormat)
+	{
+		// 8 bits formats : need compression
+	case PF_A8R8G8B8:
+	case PF_R8G8B8A8:
+	case PF_R8G8:
+	case PF_B8G8R8A8:
+	{
+		bOutCompressHeight = true;
+		return true;
+	}
+	// 16 bits formats :
+	case PF_G16:
+	// We don't use 16 bit float formats because they will have precision issues
+	// (we need 16 bits of mantissa)
+
+	// TODO: We can support 32 bit floating point formats, but for these, we probably
+	// want to output the height as an unpacked, signed values. We'll add support for
+	// that in a later CL.
+	//case PF_R32_FLOAT:
+	//case PF_G32R32F:
+	//case PF_R32G32B32F:
+	//case PF_A32B32G32R32F:
+	{
+		return true;
+	}
+	default:
+		break;
+	}
+
+	return false;
+}
+
+// TODO [jonathan.bard] This does not support being called at runtime yet because ULandscapeInfo::LandscapeActor is not valid at runtime or in non-editor worlds (e.g. PIE), for no good reason :
+void ALandscape::RenderHeightmap(const FTransform& InRenderAreaWorldTransform, const FBox2D& InRenderAreaExtents, UTextureRenderTarget2D* OutRenderTarget)
+{
+	// TODO: We may want a version of this function that returns a lambda that can be passed to the render thread and run
+	// there to add the pass to an existing FRDGBuilder, in case the user wants this to be a part of a render graph with
+	// other passes. In that case RenderHeightmap would just use that function.
+
+	TRACE_CPUPROFILER_EVENT_SCOPE(Landscape_RenderMergedHeightmap);
+
+	RenderCaptureInterface::FScopedCapture RenderCapture((RenderCaptureNextHeightmapRenders != 0), TEXT("RenderHeightmapCapture"));
+	RenderCaptureNextHeightmapRenders = FMath::Max(RenderCaptureNextHeightmapRenders - 1, 0);
+
+	ULandscapeInfo* Info = GetLandscapeInfo();
+	if (Info == nullptr)
+	{
+		UE_LOG(LogLandscape, Error, TEXT("RenderHeightmap : Cannot render anything if there's no associated landscape info with this landscape (%s)"), *GetFullName());
+		return;
+	}
+
+	// Check render target validity :
+	if (OutRenderTarget == nullptr)
+	{
+		UE_LOG(LogLandscape, Error, TEXT("RenderHeightmap : Missing render target"));
+		return;
+	}
+
+	bool bCompressHeight = false;
+	if (!IsValidRenderTargetFormatHeightmap(OutRenderTarget->GetFormat(), bCompressHeight))
+	{
+		UE_LOG(LogLandscape, Warning, TEXT("RenderHeightmap : invalid render target format for rendering heightmap (%s)"), GetPixelFormatString(OutRenderTarget->GetFormat()));
+		return;
+	}
+
+	// It can be helpful to visualize where the render happened so leave a visual log for that: 
+	UE_VLOG_OBOX(this, LogLandscape, Log, FBox(FVector(InRenderAreaExtents.Min, 0.0), FVector(InRenderAreaExtents.Max, 0.0)), InRenderAreaWorldTransform.ToMatrixWithScale(), FColor::Blue, TEXT(""));
+
+	// Don't do anything if this render area overlaps with no landscape component :
+	TMap<FIntPoint, ULandscapeComponent*> OverlappedComponents;
+	FIntRect ComponentIndicesBoundingRect;
+	if (!Info->GetOverlappedComponents(InRenderAreaWorldTransform, InRenderAreaExtents, OverlappedComponents, ComponentIndicesBoundingRect))
+	{
+		UE_LOG(LogLandscape, Log, TEXT("RenderHeightmap : no heightmap to render"));
+		return;
+	}
+
+	const FTransform& LandscapeTransform = GetTransform();
+
+	RenderMergedLandscape_RenderThread::FHeightmapRenderInfo HeightmapMergeRenderInfo;
+	// For now, merge the heightmap at max resolution :
+	HeightmapMergeRenderInfo.SubsectionSizeQuads = SubsectionSizeQuads;
+	HeightmapMergeRenderInfo.NumSubsections = NumSubsections;
+	HeightmapMergeRenderInfo.bCompressHeight = bCompressHeight;
+
+	for (auto It : OverlappedComponents)
+	{
+		ULandscapeComponent* Component = It.Value;
+		FIntPoint ComponentKey = It.Key;
+		UTexture2D* ComponentHeightmap = Component->GetHeightmap();
+		
+		// Get the subregion of the heightmap that this component uses (differs due to texture sharing).
+		// HeightmapScaleBias ZW values give us the offset of the component in a shared texture. You might think that XY would
+		// give the portion of the texture it occupies, but no, XY are 1/size, for some reason. Just calculate the subregion
+		// size ourselves.
+		int32 ComponentSize = Component->NumSubsections * (Component->SubsectionSizeQuads + 1);
+
+		FIntPoint HeightmapOffset(0, 0);
+		FTextureResource* HeightmapResource = ComponentHeightmap->GetResource();
+		if (ensure(HeightmapResource))
+		{
+			// We get the overall heightmap size via the resource instead of direct GetSizeX/Y calls because apparently
+			// the latter are unreliable while the texture is being built.
+			HeightmapOffset = FIntPoint(
+				FMath::RoundToDouble(Component->HeightmapScaleBias.Z * HeightmapResource->GetSizeX()),
+				FMath::RoundToDouble(Component->HeightmapScaleBias.W * HeightmapResource->GetSizeY()));
+		}
+			
+		// When mips are partially loaded, we need to take that into consideration when merging the heightmap :
+		uint32 MipBias = ComponentHeightmap->GetNumMips() - ComponentHeightmap->GetNumResidentMips();
+
+		// Theoretically speaking, all of our component heightmaps should be powers of two when we include the duplicated
+		// rows/columns across subsections, so we shouldn't get weird truncation results here...
+		HeightmapOffset.X >>= MipBias;
+		HeightmapOffset.Y >>= MipBias;
+		ComponentSize >>= MipBias;
+
+		// Effective area of the texture affecting this component (because of texture sharing) :
+		FIntRect HeightmapSubregion(HeightmapOffset, HeightmapOffset + ComponentSize);
+		HeightmapMergeRenderInfo.ComponentHeightmapsToRender.Add(ComponentKey, FTexture2DResourceSubregion(ComponentHeightmap->GetResource()->GetTexture2DResource(), HeightmapSubregion));
+	}
+
+	// Create the transform that will go from output target UVs to world space: 
+	FVector OutputUVOrigin = InRenderAreaWorldTransform.TransformPosition(FVector(InRenderAreaExtents.Min.X, InRenderAreaExtents.Min.Y, 0.0));
+	FVector OutputUVScale = InRenderAreaWorldTransform.GetScale3D() * FVector(InRenderAreaExtents.GetSize(), 1.0);
+	FTransform OutputUVToWorld(InRenderAreaWorldTransform.GetRotation(), OutputUVOrigin, OutputUVScale);
+
+	// Create the transform that will go from merged heightmap UVs to world space. Note that this is slightly trickier because
+	// vertices in the landscape correspond to pixel centers. So UV (0,0) is not at the minimal landscape vertex, but is instead 
+	// half a quad further (one pixel is one quad in size, so the center of the first pixel ends up at the minimal vertex).
+	// For related reasons, the size of the merged heightmap in world coordinates is actually one quad bigger in each direction.
+	check((ComponentIndicesBoundingRect.Min.X < ComponentIndicesBoundingRect.Max.X) && (ComponentIndicesBoundingRect.Min.Y < ComponentIndicesBoundingRect.Max.Y));
+	FVector MergedHeightmapScale = (FVector(ComponentIndicesBoundingRect.Max - ComponentIndicesBoundingRect.Min) * static_cast<double>(ComponentSizeQuads) + 1) 
+		* LandscapeTransform.GetScale3D();
+	MergedHeightmapScale.Z = 1.0f;
+	FVector MergedHeightmapUVOrigin = LandscapeTransform.TransformPosition(FVector(ComponentIndicesBoundingRect.Min) * (double)ComponentSizeQuads - FVector(0.5,0.5,0));
+	FTransform MergedHeightmapUVToWorld(LandscapeTransform.GetRotation(), MergedHeightmapUVOrigin, MergedHeightmapScale);
+
+	HeightmapMergeRenderInfo.OutputUVToMergedHeightmapUV = OutputUVToWorld.ToMatrixWithScale() * MergedHeightmapUVToWorld.ToInverseMatrixWithScale();
+
+	// Extract the render thread version of the render target :
+	FTextureRenderTarget2DResource* OutputRenderTargetResource = OutRenderTarget->GameThread_GetRenderTargetResource()->GetTextureRenderTarget2DResource();
+	check(OutputRenderTargetResource != nullptr);
+
+	ENQUEUE_RENDER_COMMAND(RenderHeightmap)([HeightmapMergeRenderInfo, OutputRenderTargetResource](FRHICommandListImmediate& RHICmdList)
+	{
+		FMemMark Mark(FMemStack::Get());
+		FRDGBuilder GraphBuilder(RHICmdList, RDG_EVENT_NAME("RenderHeightmap"));
+
+		FRDGTextureRef OutputHeightmap = GraphBuilder.RegisterExternalTexture(CreateRenderTarget(OutputRenderTargetResource->GetTextureRHI(), TEXT("MergedHeightmap")));
+		RenderMergedLandscape_RenderThread::RenderMergedHeightmap(HeightmapMergeRenderInfo, GraphBuilder, OutputHeightmap);
+
+		GraphBuilder.Execute();
+	});
+}
+
 #if WITH_EDITOR
 FBox ALandscape::GetCompleteBounds() const
 {
 	return GetLandscapeInfo()->GetCompleteBounds();
 }
-#endif
 
-#if WITH_EDITOR
 void ALandscapeProxy::OnFeatureLevelChanged(ERHIFeatureLevel::Type NewFeatureLevel)
 {
 	FlushGrassComponents();
@@ -2763,7 +3153,7 @@ void ALandscapeProxy::PostLoad()
 		FeatureLevelChangedDelegateHandle = GetWorld()->AddOnFeatureLevelChangedHandler(FeatureLevelChangedDelegate);
 	}
 	RepairInvalidTextures();
-#endif
+#endif // WITH_EDITOR
 }
 
 FIntPoint ALandscapeProxy::GetSectionBaseOffset() const
@@ -3627,6 +4017,57 @@ void ULandscapeInfo::UnregisterCollisionComponent(ULandscapeHeightfieldCollision
 	}
 }
 
+// TODO [jonathan.bard] This does not support being called at runtime yet because ULandscapeInfo::LandscapeActor 
+// is not valid at runtime or in non-editor worlds (e.g. PIE), for no good reason, and also XYtoComponentMap is
+// described as being valid only in editor.
+bool ULandscapeInfo::GetOverlappedComponents(const FTransform& InAreaWorldTransform, const FBox2D& InAreaExtents, 
+	TMap<FIntPoint, ULandscapeComponent*>& OutOverlappedComponents, FIntRect& OutComponentIndicesBoundingRect)
+{
+	if (!LandscapeActor.IsValid())
+	{ 
+		return false;
+	}
+
+	// Compute the AABB for this area in landscape space to find which of the landscape components are overlapping :
+	FVector Extremas[4];
+	const FTransform& LandscapeTransform = LandscapeActor->GetTransform();
+	Extremas[0] = LandscapeTransform.InverseTransformPosition(InAreaWorldTransform.TransformPosition(FVector(InAreaExtents.Min.X, InAreaExtents.Min.Y, 0.0)));
+	Extremas[1] = LandscapeTransform.InverseTransformPosition(InAreaWorldTransform.TransformPosition(FVector(InAreaExtents.Min.X, InAreaExtents.Max.Y, 0.0)));
+	Extremas[2] = LandscapeTransform.InverseTransformPosition(InAreaWorldTransform.TransformPosition(FVector(InAreaExtents.Max.X, InAreaExtents.Min.Y, 0.0)));
+	Extremas[3] = LandscapeTransform.InverseTransformPosition(InAreaWorldTransform.TransformPosition(FVector(InAreaExtents.Max.X, InAreaExtents.Max.Y, 0.0)));
+	FBox LocalExtents(Extremas, 4);
+
+	// Indices of the landscape components needed for rendering this area : 
+	FIntRect BoundingIndices;
+	BoundingIndices.Min = FIntPoint(FMath::FloorToInt(LocalExtents.Min.X / ComponentSizeQuads), FMath::FloorToInt(LocalExtents.Min.Y / ComponentSizeQuads));
+	// The max here is meant to be an exclusive bound, hence the +1
+	BoundingIndices.Max = FIntPoint(FMath::FloorToInt(LocalExtents.Max.X / ComponentSizeQuads), FMath::FloorToInt(LocalExtents.Max.Y / ComponentSizeQuads)) + FIntPoint(1);
+
+	// There could be missing components, so the effective area is actually a subset of this area :
+	FIntRect EffectiveBoundingIndices;
+	// Go through each loaded component and find out the actual bounds of the area we need to render :
+	for (int32 KeyY = BoundingIndices.Min.Y; KeyY < BoundingIndices.Max.Y; ++KeyY)
+	{
+		for (int32 KeyX = BoundingIndices.Min.X; KeyX < BoundingIndices.Max.X; ++KeyX)
+		{
+			FIntPoint Key(KeyX, KeyY);
+			if (ULandscapeComponent* Component = XYtoComponentMap.FindRef(Key))
+			{
+				EffectiveBoundingIndices.Union(FIntRect(Key, Key + FIntPoint(1)));
+				OutOverlappedComponents.Add(Key, Component);
+			}
+		}
+	}
+
+	if (OutOverlappedComponents.IsEmpty())
+	{
+		return false;
+	}
+
+	OutComponentIndicesBoundingRect = EffectiveBoundingIndices;
+	return true;
+}
+
 void ULandscapeInfo::RegisterActorComponent(ULandscapeComponent* Component, bool bMapCheck)
 {
 	// Do not register components which are not part of the world
diff --git a/Engine/Source/Runtime/Landscape/Private/LandscapeEditResources.h b/Engine/Source/Runtime/Landscape/Private/LandscapeEditResources.h
--- a/Engine/Source/Runtime/Landscape/Private/LandscapeEditResources.h
+++ b/Engine/Source/Runtime/Landscape/Private/LandscapeEditResources.h
@@ -70,11 +70,6 @@ struct FTexture2DResourceSubregion
 {
 	FTexture2DResourceSubregion() = default;
 
-	FTexture2DResourceSubregion(FTexture2DResource* InTexture, const FIntPoint& InSubregionOffset, const FIntPoint& InSubregionSize)
-		: Texture(InTexture)
-		, Subregion(InSubregionOffset, InSubregionOffset + InSubregionSize)
-	{}
-
 	FTexture2DResourceSubregion(FTexture2DResource* InTexture, const FIntRect& InSubregion)
 		: Texture(InTexture)
 		, Subregion(InSubregion)
