Fixes include warning

diff --git a/Engine/Source/Runtime/Landscape/Private/LandscapeEdit.cpp b/Engine/Source/Runtime/Landscape/Private/LandscapeEdit.cpp
--- a/Engine/Source/Runtime/Landscape/Private/LandscapeEdit.cpp
+++ b/Engine/Source/Runtime/Landscape/Private/LandscapeEdit.cpp
@@ -5,7 +5,7 @@ LandscapeEdit.cpp: Landscape editing
 =============================================================================*/
 
 #include "LandscapeEdit.h"
-#include "AssetRegistryModule.h"
+#include "AssetRegistry/AssetRegistryModule.h"
 #include "Misc/MessageDialog.h"
 #include "Misc/Paths.h"
 #include "Misc/FeedbackContext.h"
