Expose landscape virtual texture pass type to blueprint setter

diff --git a/Engine/Source/Runtime/Landscape/Classes/LandscapeProxy.h b/Engine/Source/Runtime/Landscape/Classes/LandscapeProxy.h
--- a/Engine/Source/Runtime/Landscape/Classes/LandscapeProxy.h
+++ b/Engine/Source/Runtime/Landscape/Classes/LandscapeProxy.h
@@ -517,9 +517,12 @@ public:
 	int32 VirtualTextureLodBias = 0;
 
 	/** Controls if this component draws in the main pass as well as in the virtual texture. */
-	UPROPERTY(EditAnywhere, BlueprintReadOnly, Category = VirtualTexture, meta = (DisplayName = "Draw in Main Pass"))
+	UPROPERTY(EditAnywhere, BlueprintSetter = SetVirtualTextureRenderPassType, Category = VirtualTexture, meta = (DisplayName = "Draw in Main Pass"))
 	ERuntimeVirtualTextureMainPassType VirtualTextureRenderPassType = ERuntimeVirtualTextureMainPassType::Always;
 
+	UFUNCTION(BlueprintSetter)
+	void SetVirtualTextureRenderPassType(ERuntimeVirtualTextureMainPassType InType);
+
 	/** Allows overriding the landscape bounds. This is useful if you distort the landscape with world-position-offset, for example
 	 *  Extension value in the negative Z axis, positive value increases bound size
 	 *  Note that this can also be overridden per-component when the component is selected with the component select tool */
diff --git a/Engine/Source/Runtime/Landscape/Private/LandscapeBlueprintSupport.cpp b/Engine/Source/Runtime/Landscape/Private/LandscapeBlueprintSupport.cpp
--- a/Engine/Source/Runtime/Landscape/Private/LandscapeBlueprintSupport.cpp
+++ b/Engine/Source/Runtime/Landscape/Private/LandscapeBlueprintSupport.cpp
@@ -126,4 +126,20 @@ float ULandscapeComponent::EditorGetPaintLayerWeightAtLocation(const FVector& In
 	}
 #endif
 	return 0.f;
-}
\ No newline at end of file
+}
+
+void ALandscapeProxy::SetVirtualTextureRenderPassType(ERuntimeVirtualTextureMainPassType InType)
+{
+	if (InType != VirtualTextureRenderPassType)
+	{
+		VirtualTextureRenderPassType = InType;
+
+		for (ULandscapeComponent* Component : LandscapeComponents)
+		{
+			if (Component != nullptr)
+			{
+				Component->MarkRenderStateDirty();
+			}
+		}
+	}
+}
