Added delegate for changes on landscape proxy.
Added some robustness/validation to the GetCompleteBounds call.

diff --git a/Engine/Source/Runtime/Landscape/Classes/LandscapeProxy.h b/Engine/Source/Runtime/Landscape/Classes/LandscapeProxy.h
--- a/Engine/Source/Runtime/Landscape/Classes/LandscapeProxy.h
+++ b/Engine/Source/Runtime/Landscape/Classes/LandscapeProxy.h
@@ -42,11 +42,14 @@ class UTexture2D;
 class FLandscapeEditLayerReadback;
 struct FAsyncGrassBuilder;
 struct FLandscapeInfoLayerSettings;
+class FLandscapeProxyComponentDataChangedParams;
 struct FMeshDescription;
 enum class ENavDataGatheringMode : uint8;
 
 #if WITH_EDITOR
 LANDSCAPE_API extern bool GLandscapeEditModeActive;
+
+DECLARE_MULTICAST_DELEGATE_TwoParams(FOnLandscapeProxyComponentDataChanged, ALandscapeProxy*, const FLandscapeProxyComponentDataChangedParams&);
 #endif // WITH_EDITOR
 
 
@@ -746,6 +749,10 @@ public:
 	TMap<FString, UMaterialInstanceConstant*> MaterialInstanceConstantMap;
 #endif
 
+#if WITH_EDITOR
+	FOnLandscapeProxyComponentDataChanged OnComponentDataChanged;
+#endif
+
 	// Blueprint functions
 
 	/** Change the Level of Detail distance factor */
@@ -1237,5 +1244,19 @@ private:
 	int32 OudatedPhysicalMaterialComponentsCount;
 };
 
+/**
+* Helper class to store proxy changes information 
+*/
+class LANDSCAPE_API FLandscapeProxyComponentDataChangedParams
+{
+public:
+	FLandscapeProxyComponentDataChangedParams(const TSet<ULandscapeComponent*>& InComponents);
+	void ForEachComponent(TFunctionRef<void(const ULandscapeComponent*)> Func) const;
+	const TArray<ULandscapeComponent*>& GetComponents() const { return Components; }
+
+private:
+	TArray<ULandscapeComponent*> Components;
+};
+
 DEFINE_ACTORDESC_TYPE(ALandscapeProxy, FLandscapeActorDesc);
 #endif
\ No newline at end of file
diff --git a/Engine/Source/Runtime/Landscape/Private/Landscape.cpp b/Engine/Source/Runtime/Landscape/Private/Landscape.cpp
--- a/Engine/Source/Runtime/Landscape/Private/Landscape.cpp
+++ b/Engine/Source/Runtime/Landscape/Private/Landscape.cpp
@@ -2653,7 +2653,14 @@ void ALandscape::RenderHeightmap(const FTransform& InRenderAreaWorldTransform, c
 #if WITH_EDITOR
 FBox ALandscape::GetCompleteBounds() const
 {
-	return GetLandscapeInfo()->GetCompleteBounds();
+	if (ensure(GetLandscapeInfo()))
+	{
+		return GetLandscapeInfo()->GetCompleteBounds();
+	}
+	else
+	{
+		return FBox(EForceInit::ForceInit);
+	}
 }
 
 void ALandscapeProxy::OnFeatureLevelChanged(ERHIFeatureLevel::Type NewFeatureLevel)
@@ -4910,8 +4917,12 @@ void ALandscapeProxy::UpdatePhysicalMaterialTasks(bool bInShouldMarkDirty)
 }
 #endif
 
-template<class ContainerType>
-void InvalidateGeneratedComponentDataImpl(const ContainerType& Components, bool bInvalidateLightingCache)
+void ALandscapeProxy::InvalidateGeneratedComponentData(bool bInvalidateLightingCache)
+{
+	InvalidateGeneratedComponentData(LandscapeComponents, bInvalidateLightingCache);
+}
+
+void ALandscapeProxy::InvalidateGeneratedComponentData(const TArray<ULandscapeComponent*>& Components, bool bInvalidateLightingCache)
 {
 	TMap<ALandscapeProxy*, TSet<ULandscapeComponent*>> ByProxy;
 	for (auto Iter = Components.CreateConstIterator(); Iter; ++Iter)
@@ -4928,22 +4939,17 @@ void InvalidateGeneratedComponentDataImpl(const ContainerType& Components, bool
 	for (auto Iter = ByProxy.CreateConstIterator(); Iter; ++Iter)
 	{
 		Iter.Key()->FlushGrassComponents(&Iter.Value());
-	}
-}
-
-void ALandscapeProxy::InvalidateGeneratedComponentData(bool bInvalidateLightingCache)
-{
-	InvalidateGeneratedComponentDataImpl(LandscapeComponents, bInvalidateLightingCache);
-}
 
-void ALandscapeProxy::InvalidateGeneratedComponentData(const TArray<ULandscapeComponent*>& Components, bool bInvalidateLightingCache)
-{
-	InvalidateGeneratedComponentDataImpl(Components, bInvalidateLightingCache);
+#if WITH_EDITOR
+		FLandscapeProxyComponentDataChangedParams ChangeParams(Iter.Value());
+		Iter.Key()->OnComponentDataChanged.Broadcast(Iter.Key(), ChangeParams);
+#endif
+	}
 }
 
 void ALandscapeProxy::InvalidateGeneratedComponentData(const TSet<ULandscapeComponent*>& Components, bool bInvalidateLightingCache)
 {
-	InvalidateGeneratedComponentDataImpl(Components, bInvalidateLightingCache);
+	InvalidateGeneratedComponentData(Components.Array(), bInvalidateLightingCache);
 }
 
 ULandscapeLODStreamingProxy::ULandscapeLODStreamingProxy(const FObjectInitializer& ObjectInitializer)
@@ -5103,4 +5109,19 @@ void ULandscapeLODStreamingProxy::InitResourceStateForMobileStreaming()
 	CachedSRRState.bHasPendingInitHint = true;
 }
 
+#if WITH_EDITOR
+FLandscapeProxyComponentDataChangedParams::FLandscapeProxyComponentDataChangedParams(const TSet<ULandscapeComponent*>& InComponents)
+	: Components(InComponents.Array())
+{
+}
+
+void FLandscapeProxyComponentDataChangedParams::ForEachComponent(TFunctionRef<void(const ULandscapeComponent*)> Func) const
+{
+	for (ULandscapeComponent* Component : Components)
+	{
+		Func(Component);
+	}
+}
+#endif
+
 #undef LOCTEXT_NAMESPACE
