Sanitized suggested landscape layer info package name when creating a landscape layer info file associated with a layer name (which doesn't have the same restrictions that package names have)

diff --git a/Engine/Source/Runtime/Landscape/Private/LandscapeUtils.cpp b/Engine/Source/Runtime/Landscape/Private/LandscapeUtils.cpp
--- a/Engine/Source/Runtime/Landscape/Private/LandscapeUtils.cpp
+++ b/Engine/Source/Runtime/Landscape/Private/LandscapeUtils.cpp
@@ -5,7 +5,8 @@
 #include "Engine/Level.h"
 
 #if WITH_EDITOR
-# include "EditorDirectories.h"
+#include "EditorDirectories.h"
+#include "ObjectTools.h"
 #endif
 
 namespace UE::Landscape
@@ -45,13 +46,12 @@ FString GetLayerInfoObjectPackageName(const ULevel* InLevel, const FName& InLaye
 	FString SharedAssetsPath = GetSharedAssetsPath(InLevel);
 	int32 Suffix = 1;
 
-	OutLayerObjectName = FName(*FString::Printf(TEXT("%s_LayerInfo"), *InLayerName.ToString()));
+	OutLayerObjectName = FName(*FString::Printf(TEXT("%s_LayerInfo"), *ObjectTools::SanitizeInvalidChars(*InLayerName.ToString(), INVALID_LONGPACKAGE_CHARACTERS)));
 	FPackageName::TryConvertFilenameToLongPackageName(SharedAssetsPath / OutLayerObjectName.ToString(), PackageName);
 
 	while (FPackageName::DoesPackageExist(PackageName, &PackageFilename))
 	{
-		OutLayerObjectName = FName(*FString::Printf(TEXT("%s_LayerInfo_%d"), *InLayerName.ToString(), Suffix));
-
+		OutLayerObjectName = FName(*FString::Printf(TEXT("%s_LayerInfo_%d"), *ObjectTools::SanitizeInvalidChars(*InLayerName.ToString(), INVALID_LONGPACKAGE_CHARACTERS), Suffix));
 		if (!FPackageName::TryConvertFilenameToLongPackageName(SharedAssetsPath / OutLayerObjectName.ToString(), PackageName))
 		{
 			break;
@@ -65,4 +65,4 @@ FString GetLayerInfoObjectPackageName(const ULevel* InLevel, const FName& InLaye
 
 #endif //!WITH_EDITOR
 
-} // end namespace UE::Landscape
+} // end namespace UE::Landscape
\ No newline at end of file
