Physics interface cleanup.

diff --git a/Engine/Source/Runtime/Landscape/Private/LandscapeCollision.cpp b/Engine/Source/Runtime/Landscape/Private/LandscapeCollision.cpp
--- a/Engine/Source/Runtime/Landscape/Private/LandscapeCollision.cpp
+++ b/Engine/Source/Runtime/Landscape/Private/LandscapeCollision.cpp
@@ -46,11 +46,6 @@
 #include "Physics/PhysicsInterfaceCore.h"
 #include "Physics/PhysicsInterfaceUtils.h"
 
-#if WITH_EDITOR && PHYSICS_INTERFACE_PHYSX
-	#include "IPhysXCooking.h"
-#endif // WITH_EDITOR && PHYSICS_INTERFACE_PHYSX
-
-#if WITH_CHAOS
 #include "PhysXToChaosUtil.h"
 #include "Chaos/ParticleHandle.h"
 #include "Chaos/Vector.h"
@@ -62,7 +57,6 @@
 #include "PhysicsEngine/Experimental/ChaosCooking.h"
 #include "Chaos/ChaosArchive.h"
 #include "PhysicsProxy/SingleParticlePhysicsProxy.h"
-#endif // WITH_CHAOS
 
 using namespace PhysicsInterfaceTypes;
 
@@ -92,22 +86,6 @@ ULandscapeHeightfieldCollisionComponent::FHeightfieldGeometryRef::FHeightfieldGe
 
 ULandscapeHeightfieldCollisionComponent::FHeightfieldGeometryRef::~FHeightfieldGeometryRef()
 {
-#if WITH_PHYSX && PHYSICS_INTERFACE_PHYSX
-	// Free the existing heightfield data.
-	if (RBHeightfield)
-	{
-		GPhysXPendingKillHeightfield.Add(RBHeightfield);
-		RBHeightfield = NULL;
-	}
-#if WITH_EDITORONLY_DATA
-	if (RBHeightfieldEd)
-	{
-		GPhysXPendingKillHeightfield.Add(RBHeightfieldEd);
-		RBHeightfieldEd = NULL;
-	}
-#endif// WITH_EDITORONLY_DATA
-#endif// WITH_PHYSX
-
 	// Remove ourselves from the shared map.
 	GSharedHeightfieldRefs.Remove(Guid);
 }
@@ -116,12 +94,9 @@ void ULandscapeHeightfieldCollisionComponent::FHeightfieldGeometryRef::GetResour
 {
 	CumulativeResourceSize.AddDedicatedSystemMemoryBytes(sizeof(*this));
 
-#if WITH_PHYSX
-	CumulativeResourceSize.AddDedicatedSystemMemoryBytes(UsedPhysicalMaterialArray.GetAllocatedSize());
+	//CumulativeResourceSize.AddDedicatedSystemMemoryBytes(UsedPhysicalMaterialArray.GetAllocatedSize());
 	// COMMENT [jonathan.bard] : unaccounted for : RBHeightfield, RBHeightfieldSimple 
-#endif //WITH_PHYSX
 
-#if WITH_CHAOS
 	CumulativeResourceSize.AddDedicatedSystemMemoryBytes(UsedChaosMaterials.GetAllocatedSize());
 
 	if (Heightfield.IsValid())
@@ -141,49 +116,19 @@ void ULandscapeHeightfieldCollisionComponent::FHeightfieldGeometryRef::GetResour
 		HeightfieldSimple->Serialize(ChaosAr);
 		CumulativeResourceSize.AddDedicatedSystemMemoryBytes(Data.Num());
 	}
-#endif // WITH_CHAOS
 }
 
 TMap<FGuid, ULandscapeMeshCollisionComponent::FTriMeshGeometryRef* > GSharedMeshRefs;
 
 ULandscapeMeshCollisionComponent::FTriMeshGeometryRef::FTriMeshGeometryRef()
-#if WITH_PHYSX
-	:	RBTriangleMesh(nullptr)
-#if WITH_EDITORONLY_DATA
-	, RBTriangleMeshEd(nullptr)
-#endif	//WITH_EDITORONLY_DATA
-#endif	//WITH_PHYSX
 {}
 
 ULandscapeMeshCollisionComponent::FTriMeshGeometryRef::FTriMeshGeometryRef(FGuid& InGuid)
 	: Guid(InGuid)
-#if WITH_PHYSX
-	, RBTriangleMesh(nullptr)
-#if WITH_EDITORONLY_DATA
-	, RBTriangleMeshEd(nullptr)
-#endif	//WITH_EDITORONLY_DATA
-#endif	//WITH_PHYSX
 {}
 
 ULandscapeMeshCollisionComponent::FTriMeshGeometryRef::~FTriMeshGeometryRef()
 {
-#if WITH_PHYSX && PHYSICS_INTERFACE_PHYSX
-	// Free the existing heightfield data.
-	if (RBTriangleMesh)
-	{
-		GPhysXPendingKillTriMesh.Add(RBTriangleMesh);
-		RBTriangleMesh = nullptr;
-	}
-
-#if WITH_EDITORONLY_DATA
-	if (RBTriangleMeshEd)
-	{
-		GPhysXPendingKillTriMesh.Add(RBTriangleMeshEd);
-		RBTriangleMeshEd = nullptr;
-	}
-#endif// WITH_EDITORONLY_DATA
-#endif// WITH_PHYSX
-
 	// Remove ourselves from the shared map.
 	GSharedMeshRefs.Remove(Guid);
 }
@@ -192,12 +137,9 @@ void ULandscapeMeshCollisionComponent::FTriMeshGeometryRef::GetResourceSizeEx(FR
 {
 	CumulativeResourceSize.AddDedicatedSystemMemoryBytes(sizeof(*this));
 
-#if WITH_PHYSX
-	CumulativeResourceSize.AddDedicatedSystemMemoryBytes(UsedPhysicalMaterialArray.GetAllocatedSize());
+	//CumulativeResourceSize.AddDedicatedSystemMemoryBytes(UsedPhysicalMaterialArray.GetAllocatedSize());
 	// COMMENT [jonathan.bard] : unaccounted for : RBTriangleMesh 
-#endif //WITH_PHYSX
 
-#if WITH_CHAOS
 	CumulativeResourceSize.AddDedicatedSystemMemoryBytes(UsedChaosMaterials.GetAllocatedSize());
 
 	if (Trimesh.IsValid())
@@ -208,7 +150,6 @@ void ULandscapeMeshCollisionComponent::FTriMeshGeometryRef::GetResourceSizeEx(FR
 		Trimesh->Serialize(ChaosAr);
 		CumulativeResourceSize.AddDedicatedSystemMemoryBytes(Data.Num());
 	}
-#endif // WITH_CHAOS
 }
 
 // Generate a new guid to force a recache of landscape collison derived data
@@ -245,13 +186,7 @@ static FString GetHFDDCKeyString(const FName& Format, bool bDefMaterial, const F
 		CombinedStateId = FGuid(Hash[0] ^ Hash[4], Hash[1], Hash[2], Hash[3]);
 	}
 
-#if PHYSICS_INTERFACE_PHYSX
-	const FString InterfacePrefix = TEXT("PHYSX");
-#elif WITH_CHAOS
 	const FString InterfacePrefix = FString::Printf(TEXT("%s_%s"), TEXT("CHAOS"), Chaos::ChaosVersionGUID);
-#else
-	const FString InterfacePrefix = TEXT("UNDEFINED");
-#endif
 
 	const FString KeyPrefix = FString::Printf(TEXT("%s_%s_%s"), *InterfacePrefix, *Format.ToString(), (bDefMaterial ? TEXT("VIS") : TEXT("FULL")));
 	return FDerivedDataCacheInterface::BuildCacheKey(*KeyPrefix, LANDSCAPE_COLLISION_DERIVEDDATA_VER, *CombinedStateId.ToString());
@@ -369,114 +304,6 @@ void ULandscapeHeightfieldCollisionComponent::OnCreatePhysicsState()
 			// Create the geometry
 			FVector FinalScale(LandscapeScale.X * CollisionScale, LandscapeScale.Y * CollisionScale, LandscapeScale.Z * LANDSCAPE_ZSCALE);
 
-#if PHYSICS_INTERFACE_PHYSX
-			PxTransform PhysXLandscapeComponentTransform = U2PTransform(FTransform(LandscapeComponentMatrix));
-			PxHeightFieldGeometry LandscapeComponentGeom(HeightfieldRef->RBHeightfield, PxMeshGeometryFlag::eDOUBLE_SIDED, LandscapeScale.Z * LANDSCAPE_ZSCALE, LandscapeScale.Y * CollisionScale, LandscapeScale.X * CollisionScale);
-
-			if (LandscapeComponentGeom.isValid())
-			{
-				// Creating both a sync and async actor, since this object is static
-
-				// Create the sync scene actor
-				PxRigidStatic* HeightFieldActorSync = GPhysXSDK->createRigidStatic(PhysXLandscapeComponentTransform);
-				PxShape* HeightFieldShapeSync = GPhysXSDK->createShape(LandscapeComponentGeom, HeightfieldRef->UsedPhysicalMaterialArray.GetData(), HeightfieldRef->UsedPhysicalMaterialArray.Num(), true);
-				HeightFieldShapeSync->setLocalPose(U2PTransform(LandscapeShapeTM));
-				check(HeightFieldShapeSync);
-
-				// Setup filtering
-				FCollisionFilterData QueryFilterData, SimFilterData;
-				CreateShapeFilterData(GetCollisionObjectType(), FMaskFilter(0), GetOwner()->GetUniqueID(), GetCollisionResponseToChannels(), GetUniqueID(), 0, QueryFilterData, SimFilterData, true, false, true);
-
-				// Heightfield is used for simple and complex collision
-				QueryFilterData.Word3 |= bCreateSimpleCollision ? EPDF_ComplexCollision : (EPDF_SimpleCollision | EPDF_ComplexCollision);
-				SimFilterData.Word3 |= bCreateSimpleCollision ? EPDF_ComplexCollision : (EPDF_SimpleCollision | EPDF_ComplexCollision);
-				HeightFieldShapeSync->setQueryFilterData(U2PFilterData(QueryFilterData));
-				HeightFieldShapeSync->setSimulationFilterData(U2PFilterData(SimFilterData));
-				HeightFieldShapeSync->setFlag(PxShapeFlag::eSCENE_QUERY_SHAPE, true);
-				HeightFieldShapeSync->setFlag(PxShapeFlag::eSIMULATION_SHAPE, true);
-				HeightFieldShapeSync->setFlag(PxShapeFlag::eVISUALIZATION, true);
-
-				HeightFieldActorSync->attachShape(*HeightFieldShapeSync);
-
-				// attachShape holds its own ref(), so release this here.
-				HeightFieldShapeSync->release();
-
-				if (bCreateSimpleCollision)
-				{
-					PxHeightFieldGeometry LandscapeComponentGeomSimple(HeightfieldRef->RBHeightfieldSimple, PxMeshGeometryFlags(), LandscapeScale.Z * LANDSCAPE_ZSCALE, LandscapeScale.Y * SimpleCollisionScale, LandscapeScale.X * SimpleCollisionScale);
-					check(LandscapeComponentGeomSimple.isValid());
-					PxShape* HeightFieldShapeSimpleSync = GPhysXSDK->createShape(LandscapeComponentGeomSimple, HeightfieldRef->UsedPhysicalMaterialArray.GetData(), HeightfieldRef->UsedPhysicalMaterialArray.Num(), true);
-					HeightFieldShapeSimpleSync->setLocalPose(U2PTransform(LandscapeShapeTM));
-					check(HeightFieldShapeSimpleSync);
-
-					// Setup filtering
-					FCollisionFilterData QueryFilterDataSimple = QueryFilterData;
-					FCollisionFilterData SimFilterDataSimple = SimFilterData;
-					QueryFilterDataSimple.Word3 = (QueryFilterDataSimple.Word3 & ~EPDF_ComplexCollision) | EPDF_SimpleCollision;
-					SimFilterDataSimple.Word3 = (SimFilterDataSimple.Word3 & ~EPDF_ComplexCollision) | EPDF_SimpleCollision;
-					HeightFieldShapeSimpleSync->setQueryFilterData(U2PFilterData(QueryFilterDataSimple));
-					HeightFieldShapeSimpleSync->setSimulationFilterData(U2PFilterData(SimFilterDataSimple));
-					HeightFieldShapeSimpleSync->setFlag(PxShapeFlag::eSCENE_QUERY_SHAPE, true);
-					HeightFieldShapeSimpleSync->setFlag(PxShapeFlag::eSIMULATION_SHAPE, true);
-					HeightFieldShapeSimpleSync->setFlag(PxShapeFlag::eVISUALIZATION, true);
-
-					HeightFieldActorSync->attachShape(*HeightFieldShapeSimpleSync);
-
-					// attachShape holds its own ref(), so release this here.
-					HeightFieldShapeSimpleSync->release();
-				}
-
-#if WITH_EDITOR
-				// Create a shape for a heightfield which is used only by the landscape editor
-				if (!GetWorld()->IsGameWorld())
-				{
-					PxHeightFieldGeometry LandscapeComponentGeomEd(HeightfieldRef->RBHeightfieldEd, PxMeshGeometryFlags(), LandscapeScale.Z * LANDSCAPE_ZSCALE, LandscapeScale.Y * CollisionScale, LandscapeScale.X * CollisionScale);
-					if (LandscapeComponentGeomEd.isValid())
-					{
-#if WITH_CHAOS || WITH_IMMEDIATE_PHYSX
-						UE_LOG(LogLandscape, Warning, TEXT("Failed to create editor shapes, currently unimplemented for Chaos"));
-#else
-						FPhysicsMaterialHandle_PhysX MaterialHandle = GEngine->DefaultPhysMaterial->GetPhysicsMaterial();
-						PxMaterial* PDefaultMat = MaterialHandle.Material;
-						PxShape* HeightFieldEdShapeSync = GPhysXSDK->createShape(LandscapeComponentGeomEd, &PDefaultMat, 1, true);
-						HeightFieldEdShapeSync->setLocalPose(U2PTransform(LandscapeShapeTM));
-						check(HeightFieldEdShapeSync);
-
-						FCollisionResponseContainer CollisionResponse;
-						CollisionResponse.SetAllChannels(ECollisionResponse::ECR_Ignore);
-						CollisionResponse.SetResponse(ECollisionChannel::ECC_Visibility, ECR_Block);
-						FCollisionFilterData QueryFilterDataEd, SimFilterDataEd;
-						CreateShapeFilterData(ECollisionChannel::ECC_Visibility, FMaskFilter(0), GetOwner()->GetUniqueID(), CollisionResponse, GetUniqueID(), 0, QueryFilterDataEd, SimFilterDataEd, true, false, true);
-
-						QueryFilterDataEd.Word3 |= (EPDF_SimpleCollision | EPDF_ComplexCollision);
-						HeightFieldEdShapeSync->setQueryFilterData(U2PFilterData(QueryFilterDataEd));
-						HeightFieldEdShapeSync->setFlag(PxShapeFlag::eSCENE_QUERY_SHAPE, true);
-
-						HeightFieldActorSync->attachShape(*HeightFieldEdShapeSync);
-
-						// attachShape holds its own ref(), so release this here.
-						HeightFieldEdShapeSync->release();
-#endif
-					}
-				}
-#endif// WITH_EDITOR
-
-				FPhysScene* PhysScene = GetWorld()->GetPhysicsScene();
-
-				// Set body instance data
-				BodyInstance.PhysicsUserData = FPhysicsUserData(&BodyInstance);
-				BodyInstance.OwnerComponent = this;
-
-				BodyInstance.ActorHandle.SyncActor = HeightFieldActorSync;
-				HeightFieldActorSync->userData = &BodyInstance.PhysicsUserData;
-
-				// Add to scenes
-				PxScene* SyncScene = PhysScene->GetPxScene();
-				SCOPED_SCENE_WRITE_LOCK(SyncScene);
-				SyncScene->addActor(*HeightFieldActorSync);
-			}
-
-#elif WITH_CHAOS
 			{
 				FActorCreationParams Params;
 				Params.InitialTM = LandscapeComponentTransform;
@@ -604,7 +431,6 @@ void ULandscapeHeightfieldCollisionComponent::OnCreatePhysicsState()
 				}
 
 			}
-#endif // WITH_CHAOS
 		}
 	}
 }
@@ -613,7 +439,6 @@ void ULandscapeHeightfieldCollisionComponent::OnDestroyPhysicsState()
 {
 	Super::OnDestroyPhysicsState();
 	
-#if WITH_CHAOS
 	if (FPhysScene_Chaos* PhysScene = GetWorld()->GetPhysicsScene())
 	{
 		FPhysicsActorHandle& ActorHandle = BodyInstance.GetPhysicsActorHandle();
@@ -626,7 +451,6 @@ void ULandscapeHeightfieldCollisionComponent::OnDestroyPhysicsState()
 			PhysScene->UnRegisterForCollisionEvents(this);
 		}
 	}
-#endif // WITH_CHAOS
 }
 
 void ULandscapeHeightfieldCollisionComponent::ApplyWorldOffset(const FVector& InOffset, bool bWorldShift)
@@ -641,12 +465,8 @@ void ULandscapeHeightfieldCollisionComponent::ApplyWorldOffset(const FVector& In
 
 void ULandscapeHeightfieldCollisionComponent::CreateCollisionObject()
 {
-#if WITH_CHAOS
 	LLM_SCOPE(ELLMTag::ChaosLandscape);
-#else
-	//NOTE: this currently gets ignored because of low level allocator
-	LLM_SCOPE(ELLMTag::PhysXLandscape);
-#endif
+
 	// If we have not created a heightfield yet - do it now.
 	if (!IsValidRef(HeightfieldRef))
 	{
@@ -671,11 +491,7 @@ void ULandscapeHeightfieldCollisionComponent::CreateCollisionObject()
 
 #if WITH_EDITOR
 		// Use existing heightfield except if it is missing its editor heightfield and the component needs it.
-		if (ExistingHeightfieldRef && (!bNeedsEditorHeightField
-#if WITH_CHAOS
-			|| ExistingHeightfieldRef->EditorHeightfield != nullptr
-#endif // WITH_CHAOS
-			))
+		if (ExistingHeightfieldRef && (!bNeedsEditorHeightField || ExistingHeightfieldRef->EditorHeightfield != nullptr))
 #else // WITH_EDITOR
 		if (ExistingHeightfieldRef)
 #endif // !WITH_EDITOR
@@ -704,43 +520,6 @@ void ULandscapeHeightfieldCollisionComponent::CreateCollisionObject()
 			{
 				HeightfieldRef = GSharedHeightfieldRefs.Add(HeightfieldGuid, new FHeightfieldGeometryRef(HeightfieldGuid));
 
-#if WITH_PHYSX && PHYSICS_INTERFACE_PHYSX
-				// Create heightfield shape
-				{
-					FPhysXInputStream HeightFieldStream(CookedCollisionData.GetData(), CookedCollisionData.Num());
-					HeightfieldRef->RBHeightfield = GPhysXSDK->createHeightField(HeightFieldStream);
-					if (SimpleCollisionSizeQuads > 0)
-					{
-						HeightfieldRef->RBHeightfieldSimple = GPhysXSDK->createHeightField(HeightFieldStream);
-					}
-				}
-
-				for (UPhysicalMaterial* PhysicalMaterial : CookedPhysicalMaterials)
-				{
-					const FPhysicsMaterialHandle_PhysX& MaterialHandle = PhysicalMaterial->GetPhysicsMaterial();
-					HeightfieldRef->UsedPhysicalMaterialArray.Add(MaterialHandle.Material);
-				}
-
-				// Release cooked collison data
-				// In cooked builds created collision object will never be deleted while component is alive, so we don't need this data anymore
-				if (FPlatformProperties::RequiresCookedData() || World->IsGameWorld())
-				{
-					CookedCollisionData.Empty();
-				}
-
-#if WITH_EDITOR
-				// Create heightfield for the landscape editor (no holes in it)
-				if (!World->IsGameWorld())
-				{
-					TArray<UPhysicalMaterial*> CookedMaterialsEd;
-					if (CookCollisionData(PhysicsFormatName, true, bCheckDDC, CookedCollisionDataEd, CookedMaterialsEd))
-					{
-						FPhysXInputStream HeightFieldStream(CookedCollisionDataEd.GetData(), CookedCollisionDataEd.Num());
-						HeightfieldRef->RBHeightfieldEd = GPhysXSDK->createHeightField(HeightFieldStream);
-					}
-				}
-#endif //WITH_EDITOR
-#elif WITH_CHAOS
 				// Create heightfields
 				{
 					FMemoryReader Reader(CookedCollisionData);
@@ -792,7 +571,6 @@ void ULandscapeHeightfieldCollisionComponent::CreateCollisionObject()
 					}
 				}
 #endif //WITH_EDITOR
-#endif // WITH_CHAOS
 			}
 		}
 	}
@@ -801,7 +579,6 @@ void ULandscapeHeightfieldCollisionComponent::CreateCollisionObject()
 #if WITH_EDITOR
 void ULandscapeHeightfieldCollisionComponent::SpeculativelyLoadAsyncDDCCollsionData()
 {
-#if WITH_PHYSX
 	if (GetLinkerUEVersion() >= VER_UE4_LANDSCAPE_SERIALIZE_PHYSICS_MATERIALS && !GLandscapeCollisionSkipDDC)
 	{
 		UWorld* World = GetWorld();
@@ -816,94 +593,7 @@ void ULandscapeHeightfieldCollisionComponent::SpeculativelyLoadAsyncDDCCollsionD
 			World->AsyncPreRegisterDDCRequests.Add(SpeculativeDDCRequest);
 		}
 	}
-#endif // WITH_PHYSX
-}
-
-#if PHYSICS_INTERFACE_PHYSX
-TArray<PxHeightFieldSample> ConvertHeightfieldDataForPhysx(
-	const ULandscapeHeightfieldCollisionComponent* const Component,
-	const int32 CollisionSizeVerts,
-	const bool bIsMirrored,
-	const uint16* Heights,
-	const bool bUseDefMaterial,
-	UPhysicalMaterial* const DefMaterial,
-	const uint8* DominantLayers,
-	const uint8* RenderPhysicalMaterialIds,
-	TArray<UPhysicalMaterial*> const& PhysicalMaterialRenderObjects,
-	TArray<UPhysicalMaterial*>& InOutMaterials) 
-{
-	const int32 NumSamples = FMath::Square(CollisionSizeVerts);
-	check(DefMaterial);
-	// Might return INDEX_NONE if DefMaterial wasn't added yet
-	int32 DefaultMaterialIndex = InOutMaterials.IndexOfByKey(DefMaterial);
-	
-	TArray<PxHeightFieldSample> Samples;
-	Samples.Reserve(NumSamples);
-	Samples.AddZeroed(NumSamples);
-
-	for (int32 RowIndex = 0; RowIndex < CollisionSizeVerts; RowIndex++)
-	{
-		for (int32 ColIndex = 0; ColIndex < CollisionSizeVerts; ColIndex++)
-		{
-			int32 SrcSampleIndex = (ColIndex * CollisionSizeVerts) + (bIsMirrored ? RowIndex : (CollisionSizeVerts - RowIndex - 1));
-			int32 DstSampleIndex = (RowIndex * CollisionSizeVerts) + ColIndex;
-
-			PxHeightFieldSample& Sample = Samples[DstSampleIndex];
-			Sample.height = ((int32)Heights[SrcSampleIndex] - 32768);
-
-			// Materials are not relevant on the last row/column because they are per-triangle and the last row/column don't own any
-			if (RowIndex < CollisionSizeVerts - 1 &&
-				ColIndex < CollisionSizeVerts - 1)
-			{
-				int32 MaterialIndex = DefaultMaterialIndex; // Default physical material.
-				if (!bUseDefMaterial && DominantLayers)
-				{
-					uint8 DominantLayerIdx = DominantLayers ? DominantLayers[SrcSampleIndex] : -1;
-					ULandscapeLayerInfoObject* Layer = Component->ComponentLayerInfos.IsValidIndex(DominantLayerIdx) ? Component->ComponentLayerInfos[DominantLayerIdx] : nullptr;
-					if (Layer == ALandscapeProxy::VisibilityLayer)
-					{
-						// If it's a hole, override with the hole flag.
-						MaterialIndex = PxHeightFieldMaterial::eHOLE;
-					}
-					else if (RenderPhysicalMaterialIds)
-					{
-						uint8 RenderIdx = RenderPhysicalMaterialIds[SrcSampleIndex];
-						if (RenderIdx > 0)
-						{
-							MaterialIndex = InOutMaterials.AddUnique(PhysicalMaterialRenderObjects[RenderIdx - 1]);
-						}
-					}
-					else if (Layer && Layer->PhysMaterial)
-					{
-						MaterialIndex = InOutMaterials.AddUnique(Layer->PhysMaterial);
-					}
-				}
-
-				// Default Material but Def Material wasn't added yet...
-				if (MaterialIndex == INDEX_NONE)
-				{
-					MaterialIndex = DefaultMaterialIndex = InOutMaterials.Add(DefMaterial);
-				}
-
-				Sample.materialIndex0 = MaterialIndex;
-				Sample.materialIndex1 = MaterialIndex;
-			}
-
-			// TODO: edge turning
-		}
-	}
-
-	// Handle case where Component uses 100% Visibility layer
-	// Add the default material because PhysX running PX_CHECKED (Debug) will return a nullptr on CreateShape if the material list is empty
-	if (InOutMaterials.Num() == 0)
-	{
-		check(DefaultMaterialIndex == INDEX_NONE);
-		InOutMaterials.Add(DefMaterial);
-	}
-
-	return Samples;
 }
-#endif // WITH_PHYSX
 
 bool ULandscapeHeightfieldCollisionComponent::CookCollisionData(const FName& Format, bool bUseDefMaterial, bool bCheckDDC, TArray<uint8>& OutCookedData, TArray<UPhysicalMaterial*>& InOutMaterials) const
 {
@@ -1015,37 +705,6 @@ bool ULandscapeHeightfieldCollisionComponent::CookCollisionData(const FName& For
 	// List of materials which is actually used by heightfield
 	InOutMaterials.Empty();
 
-#if WITH_PHYSX && PHYSICS_INTERFACE_PHYSX
-		
-	TArray<PxHeightFieldSample> Samples;
-	TArray<PxHeightFieldSample> SimpleSamples;
-	Samples = ConvertHeightfieldDataForPhysx(this, CollisionSizeVerts, bIsMirrored, Heights, bUseDefMaterial, DefMaterial, DominantLayers, RenderPhysicalMaterialIds, PhysicalMaterialRenderObjects, InOutMaterials);
-
-	if (bGenerateSimpleCollision)
-	{
-		SimpleSamples = ConvertHeightfieldDataForPhysx(this, SimpleCollisionSizeVerts, bIsMirrored, SimpleHeights, bUseDefMaterial, DefMaterial, SimpleDominantLayers, SimpleRenderPhysicalMaterialIds, PhysicalMaterialRenderObjects, InOutMaterials);
-	}
-
-	CollisionHeightData.Unlock();
-	if (DominantLayers)
-	{
-		DominantLayerData.Unlock();
-	}
-	
-	//
-	FIntPoint HFSize = FIntPoint(CollisionSizeVerts, CollisionSizeVerts);
-
-	ITargetPlatformManagerModule* TPM = GetTargetPlatformManager();
-	const IPhysXCooking* Cooker = TPM->FindPhysXCooking(Format);
-	Succeeded = Cooker->CookHeightField(Format, HFSize, Samples.GetData(), Samples.GetTypeSize(), OutData);
-
-	if (Succeeded && bGenerateSimpleCollision)
-	{
-		FIntPoint HFSizeSimple = FIntPoint(SimpleCollisionSizeVerts, SimpleCollisionSizeVerts);
-		Succeeded = Cooker->CookHeightField(Format, HFSizeSimple, SimpleSamples.GetData(), SimpleSamples.GetTypeSize(), OutData);
-	}
-#elif WITH_CHAOS
-
 	// Generate material indices
 	TArray<uint8> MaterialIndices;
 	MaterialIndices.Reserve(NumSamples + NumSimpleSamples);
@@ -1108,7 +767,6 @@ bool ULandscapeHeightfieldCollisionComponent::CookCollisionData(const FName& For
 	}
 
 	Succeeded = true;
-#endif // WITH_CHAOS
 
 	if (CollisionHeightData.IsLocked())
 	{
@@ -1324,12 +982,7 @@ bool ULandscapeMeshCollisionComponent::CookCollisionData(const FName& Format, bo
 
 	TArray<uint8> OutData;
 	bool Result = false;
-#if PHYSICS_INTERFACE_PHYSX
-	bool bFlipNormals = true;
-	ITargetPlatformManagerModule* TPM = GetTargetPlatformManager();
-	const IPhysXCooking* Cooker = TPM->FindPhysXCooking(Format);
-	Result = Cooker->CookTriMesh(Format, EPhysXMeshCookFlags::Default, Vertices, Indices, MaterialIndices, bFlipNormals, OutData);
-#elif WITH_CHAOS
+
 	FCookBodySetupInfo CookInfo;
 	FTriMeshCollisionData& MeshDesc = CookInfo.TriangleMeshDesc;
 	MeshDesc.bFlipNormals = true;
@@ -1348,7 +1001,6 @@ bool ULandscapeMeshCollisionComponent::CookCollisionData(const FName& Format, bo
 		ChaosAr << Trimesh;
 		Result = OutData.Num() > 0;
 	}
-#endif // WITH_CHAOS
 
 	if (Result)
 	{
@@ -1417,22 +1069,13 @@ void ULandscapeMeshCollisionComponent::CreateCollisionObject()
 				MeshRef = GSharedMeshRefs.Add(MeshGuid, new FTriMeshGeometryRef(MeshGuid));
 
 				// Create physics objects
-#if PHYSICS_INTERFACE_PHYSX
-				FPhysXInputStream Buffer(CookedCollisionData.GetData(), CookedCollisionData.Num());
-				MeshRef->RBTriangleMesh = GPhysXSDK->createTriangleMesh(Buffer);
-#elif WITH_CHAOS
 				FMemoryReader Reader(CookedCollisionData);
 				Chaos::FChaosArchive Ar(Reader);
 				Ar << MeshRef->Trimesh;
-#endif // WITH_CHAOS
 
 				for (UPhysicalMaterial* PhysicalMaterial : CookedPhysicalMaterials)
 				{
-#if WITH_CHAOS || WITH_IMMEDIATE_PHYSX
 					MeshRef->UsedChaosMaterials.Add(PhysicalMaterial->GetPhysicsMaterial());
-#else
-					MeshRef->UsedPhysicalMaterialArray.Add(PhysicalMaterial->GetPhysicsMaterial().Material);
-#endif
 				}
 
 				// Release cooked collison data
@@ -1449,14 +1092,9 @@ void ULandscapeMeshCollisionComponent::CreateCollisionObject()
 					TArray<UPhysicalMaterial*> CookedMaterialsEd;
 					if (CookCollisionData(PhysicsFormatName, true, bCheckDDC, CookedCollisionDataEd, CookedMaterialsEd))
 					{
-#if PHYSICS_INTERFACE_PHYSX
-						FPhysXInputStream MeshStream(CookedCollisionDataEd.GetData(), CookedCollisionDataEd.Num());
-						MeshRef->RBTriangleMeshEd = GPhysXSDK->createTriangleMesh(MeshStream);
-#elif WITH_CHAOS
 						FMemoryReader EdReader(CookedCollisionData);
 						Chaos::FChaosArchive EdAr(EdReader);
 						EdAr << MeshRef->EditorTrimesh;
-#endif // WITH_CHAOS
 					}
 				}
 #endif //WITH_EDITOR
@@ -1474,157 +1112,6 @@ ULandscapeMeshCollisionComponent::ULandscapeMeshCollisionComponent()
 
 ULandscapeMeshCollisionComponent::~ULandscapeMeshCollisionComponent() = default;
 
-#if PHYSICS_INTERFACE_PHYSX
-struct FMeshCollisionInitHelper
-{
-	FMeshCollisionInitHelper() = delete;
-	FMeshCollisionInitHelper(TRefCountPtr<ULandscapeMeshCollisionComponent::FTriMeshGeometryRef> InMeshRef, UWorld* InWorld, UPrimitiveComponent* InComponent, FBodyInstance* InBodyInstance)
-		: ComponentToWorld(FTransform::Identity)
-		, ComponentScale(FVector::OneVector)
-		, CollisionScale(1.0f)
-		, MeshRef(InMeshRef)
-		, World(InWorld)
-		, Component(InComponent)
-		, TargetInstance(InBodyInstance)
-	{
-		PxGeom.triangleMesh = MeshRef->RBTriangleMesh;
-		PxGeom.scale.scale.x = ComponentScale.X * CollisionScale;
-		PxGeom.scale.scale.y = ComponentScale.Y * CollisionScale;
-		PxGeom.scale.scale.z = ComponentScale.Z;
-
-#if WITH_EDITOR
-		PxGeomEd.triangleMesh = MeshRef->RBTriangleMeshEd;
-		PxGeomEd.scale.scale.x = ComponentScale.X * CollisionScale;
-		PxGeomEd.scale.scale.y = ComponentScale.Y * CollisionScale;
-		PxGeomEd.scale.scale.z = ComponentScale.Z;
-#endif // WITH_EDITOR
-
-		check(World);
-		PhysScene = World->GetPhysicsScene();
-		check(PhysScene);
-		check(Component);
-		check(TargetInstance);
-	}
-
-	void SetComponentScale3D(const FVector& InScale)
-	{
-		ComponentScale = InScale;
-
-		PxGeom.scale.scale.x = ComponentScale.X * CollisionScale;
-		PxGeom.scale.scale.y = ComponentScale.Y * CollisionScale;
-		PxGeom.scale.scale.z = ComponentScale.Z;
-
-		PxGeomEd.scale.scale.x = ComponentScale.X * CollisionScale;
-		PxGeomEd.scale.scale.y = ComponentScale.Y * CollisionScale;
-		PxGeomEd.scale.scale.z = ComponentScale.Z;
-	}
-
-	void SetCollisionScale(float InScale)
-	{
-		CollisionScale = InScale;
-
-		PxGeom.scale.scale.x = ComponentScale.X * CollisionScale;
-		PxGeom.scale.scale.y = ComponentScale.Y * CollisionScale;
-		PxGeom.scale.scale.z = ComponentScale.Z;
-
-		PxGeomEd.scale.scale.x = ComponentScale.X * CollisionScale;
-		PxGeomEd.scale.scale.y = ComponentScale.Y * CollisionScale;
-		PxGeomEd.scale.scale.z = ComponentScale.Z;
-	}
-
-	void SetComponentToWorld(const FTransform& InTransform)
-	{
-		ComponentToWorld = InTransform;
-
-		PxComponentTransform = U2PTransform(ComponentToWorld);
-	}
-
-	void SetFilters(const FCollisionFilterData& InQueryFilter, const FCollisionFilterData& InSimulationFilter)
-	{
-		QueryFilter = InQueryFilter;
-		SimulationFilter = InSimulationFilter;
-	}
-
-	void SetEditorFilter(const FCollisionFilterData& InFilter)
-	{
-		QueryFilterEd = InFilter;
-	}
-
-	bool IsGeometryValid() const
-	{
-		return PxGeom.isValid();
-	}
-
-	void CreateActors()
-	{
-		{
-			// Create the sync scene actor
-			PActor = GPhysXSDK->createRigidStatic(PxComponentTransform);
-			PxShape* NewShape = GPhysXSDK->createShape(PxGeom, MeshRef->UsedPhysicalMaterialArray.GetData(), MeshRef->UsedPhysicalMaterialArray.Num(), true);
-			check(NewShape);
-
-			// Heightfield is used for simple and complex collision
-			NewShape->setQueryFilterData(U2PFilterData(QueryFilter));
-			NewShape->setSimulationFilterData(U2PFilterData(SimulationFilter));
-			NewShape->setFlag(PxShapeFlag::eSCENE_QUERY_SHAPE, true);
-			NewShape->setFlag(PxShapeFlag::eSIMULATION_SHAPE, true);
-			NewShape->setFlag(PxShapeFlag::eVISUALIZATION, true);
-
-			PActor->attachShape(*NewShape);
-			NewShape->release();
-		}
-
-#if WITH_EDITOR
-		if(!World->IsGameWorld()) // Need to create editor shape
-		{
-			PxMaterial* PDefaultMat = GEngine->DefaultPhysMaterial->GetPhysicsMaterial().Material;
-			PxShape* NewShape = GPhysXSDK->createShape(PxGeom, &PDefaultMat, 1, true);
-			check(NewShape);
-
-			NewShape->setQueryFilterData(U2PFilterData(QueryFilterEd));
-			NewShape->setFlag(PxShapeFlag::eSCENE_QUERY_SHAPE, true);
-
-			PActor->attachShape(*NewShape);
-			NewShape->release();
-		}
-#endif // WITH_EDITOR
-
-		// Set body instance data
-		TargetInstance->PhysicsUserData = FPhysicsUserData(TargetInstance);
-		TargetInstance->OwnerComponent = Component;
-		TargetInstance->ActorHandle.SyncActor = PActor;
-		PActor->userData = &TargetInstance->PhysicsUserData;
-	}
-
-	void AddToScene()
-	{
-		check(PhysScene);
-
-		// Add to scenes
-		PxScene* SyncScene = PhysScene->GetPxScene();
-		SCOPED_SCENE_WRITE_LOCK(SyncScene);
-		SyncScene->addActor(*PActor);
-	}
-
-private:
-	FTransform ComponentToWorld;
-	FVector ComponentScale;
-	float CollisionScale;
-	TRefCountPtr<ULandscapeMeshCollisionComponent::FTriMeshGeometryRef> MeshRef;
-	FPhysScene* PhysScene;
-	FCollisionFilterData QueryFilter;
-	FCollisionFilterData SimulationFilter;
-	FCollisionFilterData QueryFilterEd;
-	UWorld* World;
-	UPrimitiveComponent* Component;
-	FBodyInstance* TargetInstance;
-
-	PxTriangleMeshGeometry PxGeom;
-	PxTriangleMeshGeometry PxGeomEd;
-	PxTransform PxComponentTransform;
-	PxRigidStatic* PActor;
-};
-#elif WITH_CHAOS
 struct FMeshCollisionInitHelper
 {
 	FMeshCollisionInitHelper() = delete;
@@ -1779,7 +1266,7 @@ private:
 
 	FPhysicsActorHandle ActorHandle;
 };
-#endif // WITH_CHAOS
+
 
 void ULandscapeMeshCollisionComponent::OnCreatePhysicsState()
 {
@@ -1893,7 +1380,6 @@ uint32 ULandscapeHeightfieldCollisionComponent::ComputeCollisionHash() const
 
 void ULandscapeHeightfieldCollisionComponent::UpdateHeightfieldRegion(int32 ComponentX1, int32 ComponentY1, int32 ComponentX2, int32 ComponentY2)
 {
-#if WITH_PHYSX
 	if (IsValidRef(HeightfieldRef))
 	{
 		// If we're currently sharing this data with a PIE session, we need to make a new heightfield.
@@ -1903,17 +1389,10 @@ void ULandscapeHeightfieldCollisionComponent::UpdateHeightfieldRegion(int32 Comp
 			return;
 		}
 
-#if WITH_CHAOS || WITH_IMMEDIATE_PHYSX
 		if(!BodyInstance.ActorHandle)
 		{
 			return;
 		}
-#else
-		if (BodyInstance.ActorHandle.SyncActor == NULL)
-		{
-			return;
-		}
-#endif
 
 		// We don't lock the async scene as we only set the geometry in the sync scene's RigidActor.
 		// This function is used only during painting for line traces by the painting tools.
@@ -1929,65 +1408,6 @@ void ULandscapeHeightfieldCollisionComponent::UpdateHeightfieldRegion(int32 Comp
 			uint16* Heights = (uint16*)CollisionHeightData.Lock(LOCK_READ_ONLY);
 			check(CollisionHeightData.GetElementCount() == (FMath::Square(CollisionSizeVerts) + FMath::Square(SimpleCollisionSizeVerts)));
 	
-#if PHYSICS_INTERFACE_PHYSX
-			// PhysX heightfield has the X and Y axis swapped, and the X component is also inverted
-			int32 HeightfieldX1 = ComponentY1;
-			int32 HeightfieldY1 = (bIsMirrored ? ComponentX1 : (CollisionSizeVerts - ComponentX2 - 1));
-			int32 DstVertsX = ComponentY2 - ComponentY1 + 1;
-			int32 DstVertsY = ComponentX2 - ComponentX1 + 1;
-	
-			TArray<PxHeightFieldSample> Samples;
-			Samples.AddZeroed(DstVertsX*DstVertsY);
-	
-			// Traverse the area in destination heigthfield coordinates
-			for (int32 RowIndex = 0; RowIndex < DstVertsY; RowIndex++)
-			{
-				for (int32 ColIndex = 0; ColIndex < DstVertsX; ColIndex++)
-				{
-					int32 SrcX = bIsMirrored ? (RowIndex + ComponentX1) : (ComponentX2 - RowIndex);
-					int32 SrcY = ColIndex + ComponentY1;
-					int32 SrcSampleIndex = (SrcY * CollisionSizeVerts) + SrcX;
-					check(SrcSampleIndex < FMath::Square(CollisionSizeVerts));
-					int32 DstSampleIndex = (RowIndex * DstVertsX) + ColIndex;
-	
-					PxHeightFieldSample& Sample = Samples[DstSampleIndex];
-					Sample.height = FMath::Clamp<int32>(((int32)Heights[SrcSampleIndex] - 32768), -32768, 32767);
-	
-					Sample.materialIndex0 = 0;
-					Sample.materialIndex1 = 0;
-				}
-			}
-	
-			CollisionHeightData.Unlock();
-	
-			PxHeightFieldDesc SubDesc;
-			SubDesc.format = PxHeightFieldFormat::eS16_TM;
-			SubDesc.nbColumns = DstVertsX;
-			SubDesc.nbRows = DstVertsY;
-			SubDesc.samples.data = Samples.GetData();
-			SubDesc.samples.stride = sizeof(PxU32);
-			SubDesc.flags = PxHeightFieldFlag::eNO_BOUNDARY_EDGES;
-	
-			HeightfieldRef->RBHeightfieldEd->modifySamples(HeightfieldX1, HeightfieldY1, SubDesc, true);
-	
-			//
-			// Reset geometry of heightfield shape. Required by the modifySamples
-			//
-			FVector LandscapeScale = GetComponentToWorld().GetScale3D().GetAbs();
-			// Create the geometry
-			PxHeightFieldGeometry LandscapeComponentGeom(HeightfieldRef->RBHeightfieldEd, PxMeshGeometryFlags(), LandscapeScale.Z * LANDSCAPE_ZSCALE, LandscapeScale.Y * CollisionScale, LandscapeScale.X * CollisionScale);
-
-			{
-				FInlineShapeArray PShapes;
-
-				const int32 NumShapes = FillInlineShapeArray_AssumesLocked(PShapes, Actor);
-				if(NumShapes > 1)
-				{
-					FPhysicsInterface::SetGeometry(PShapes[1], LandscapeComponentGeom);
-				}
-			}
-
-#elif WITH_CHAOS
 			int32 HeightfieldY1 = ComponentY1;
 			int32 HeightfieldX1 = (bIsMirrored ? ComponentX1 : (CollisionSizeVerts - ComponentX2 - 1));
 			int32 DstVertsX = ComponentX2 - ComponentX1 + 1;
@@ -2025,11 +1445,8 @@ void ULandscapeHeightfieldCollisionComponent::UpdateHeightfieldRegion(int32 Comp
 
 			FPhysScene* PhysScene = GetWorld()->GetPhysicsScene();
 			PhysScene->UpdateActorInAccelerationStructure(PhysActorHandle);
-#endif // WITH_CHAOS
 		});
 	}
-
-#endif// WITH_PHYSX
 }
 #endif// WITH_EDITOR
 
@@ -2359,23 +1776,6 @@ bool ULandscapeHeightfieldCollisionComponent::IsShown(const FEngineShowFlags& Sh
 bool ULandscapeHeightfieldCollisionComponent::DoCustomNavigableGeometryExport(FNavigableGeometryExport& GeomExport) const
 {
 	check(IsInGameThread());
-#if WITH_PHYSX && PHYSICS_INTERFACE_PHYSX
-	if (IsValidRef(HeightfieldRef) && HeightfieldRef->RBHeightfield)
-	{
-		FTransform HFToW = GetComponentTransform();
-		if (HeightfieldRef->RBHeightfieldSimple)
-		{
-			const float SimpleCollisionScale = CollisionScale * CollisionSizeQuads / SimpleCollisionSizeQuads;
-			HFToW.MultiplyScale3D(FVector(SimpleCollisionScale, SimpleCollisionScale, LANDSCAPE_ZSCALE));
-			GeomExport.ExportPxHeightField(HeightfieldRef->RBHeightfieldSimple, HFToW);
-		}
-		else
-		{
-			HFToW.MultiplyScale3D(FVector(CollisionScale, CollisionScale, LANDSCAPE_ZSCALE));
-			GeomExport.ExportPxHeightField(HeightfieldRef->RBHeightfield, HFToW);
-		}
-	}
-#elif WITH_CHAOS
 	if(IsValidRef(HeightfieldRef) && HeightfieldRef->Heightfield)
 	{
 		FTransform HFToW = GetComponentTransform();
@@ -2391,7 +1791,6 @@ bool ULandscapeHeightfieldCollisionComponent::DoCustomNavigableGeometryExport(FN
 			GeomExport.ExportChaosHeightField(HeightfieldRef->Heightfield.Get(), HFToW);
 		}
 	}
-#endif // WITH_CHAOS
 
 	return false;
 }
@@ -2404,11 +1803,7 @@ void ULandscapeHeightfieldCollisionComponent::GatherGeometrySlice(FNavigableGeom
 		FTransform HFToW = GetComponentTransform();
 		HFToW.MultiplyScale3D(FVector(CollisionScale, CollisionScale, LANDSCAPE_ZSCALE));
 
-#if WITH_PHYSX && PHYSICS_INTERFACE_PHYSX
-		GeomExport.ExportPxHeightFieldSlice(CachedHeightFieldSamples, HeightfieldRowsCount, HeightfieldColumnsCount, HFToW, SliceBox);
-#elif WITH_CHAOS
 		GeomExport.ExportChaosHeightFieldSlice(CachedHeightFieldSamples, HeightfieldRowsCount, HeightfieldColumnsCount, HFToW, SliceBox);
-#endif // WITH_CHAOS
 	}
 }
 
@@ -2420,38 +1815,6 @@ ENavDataGatheringMode ULandscapeHeightfieldCollisionComponent::GetGeometryGather
 
 void ULandscapeHeightfieldCollisionComponent::PrepareGeometryExportSync()
 {
-	//check(IsInGameThread());
-#if WITH_PHYSX && PHYSICS_INTERFACE_PHYSX
-	if (IsValidRef(HeightfieldRef) && HeightfieldRef->RBHeightfield != nullptr && CachedHeightFieldSamples.IsEmpty())
-	{
-		const UWorld* World = GetWorld();
-
-		if (World != nullptr)
-		{
-			HeightfieldRowsCount = HeightfieldRef->RBHeightfield->getNbRows();
-			HeightfieldColumnsCount = HeightfieldRef->RBHeightfield->getNbColumns();
-			const int32 SamplesCount = HeightfieldRowsCount * HeightfieldColumnsCount;
-				
-			if (CachedHeightFieldSamples.Heights.Num() != SamplesCount)
-			{
-				QUICK_SCOPE_CYCLE_COUNTER(STAT_NavMesh_ExportPxHeightField_saveCells);
-
-				CachedHeightFieldSamples.Heights.SetNumUninitialized(SamplesCount);
-
-				TArray<PxHeightFieldSample> HFSamples;
-				HFSamples.SetNumUninitialized(SamplesCount);
-				HeightfieldRef->RBHeightfield->saveCells(HFSamples.GetData(), HFSamples.Num()*HFSamples.GetTypeSize());
-
-				for (int32 SampleIndex = 0; SampleIndex < HFSamples.Num(); ++SampleIndex)
-				{
-					const PxHeightFieldSample& Sample = HFSamples[SampleIndex];
-					CachedHeightFieldSamples.Heights[SampleIndex] = Sample.height;
-					CachedHeightFieldSamples.Holes.Add((Sample.materialIndex0 == PxHeightFieldMaterial::eHOLE));
-				}
-			}
-		}
-	}
-#elif WITH_CHAOS
 	if(IsValidRef(HeightfieldRef) && HeightfieldRef->Heightfield.Get() && CachedHeightFieldSamples.IsEmpty())
 	{
 		const UWorld* World = GetWorld();
@@ -2481,7 +1844,6 @@ void ULandscapeHeightfieldCollisionComponent::PrepareGeometryExportSync()
 			}
 		}
 	}
-#endif // WITH_CHAOS
 }
 
 bool ULandscapeMeshCollisionComponent::DoCustomNavigableGeometryExport(FNavigableGeometryExport& GeomExport) const
@@ -2493,24 +1855,10 @@ bool ULandscapeMeshCollisionComponent::DoCustomNavigableGeometryExport(FNavigabl
 		FTransform MeshToW = GetComponentTransform();
 		MeshToW.MultiplyScale3D(FVector(CollisionScale, CollisionScale, 1.f));
 
-#if PHYSICS_INTERFACE_PHYSX
-		if (MeshRef->RBTriangleMesh != nullptr)
-		{
-			if (MeshRef->RBTriangleMesh->getTriangleMeshFlags() & PxTriangleMeshFlag::e16_BIT_INDICES)
-			{
-				GeomExport.ExportPxTriMesh16Bit(MeshRef->RBTriangleMesh, MeshToW);
-			}
-			else
-			{
-				GeomExport.ExportPxTriMesh32Bit(MeshRef->RBTriangleMesh, MeshToW);
-			}
-		}
-#elif WITH_CHAOS
 		if (MeshRef->Trimesh != nullptr)
 		{
 			GeomExport.ExportChaosTriMesh(MeshRef->Trimesh.Get(), MeshToW);
 		}
-#endif // WITH_CHAOS
 	}
 
 	return false;
@@ -3048,17 +2396,12 @@ TOptional<float> ULandscapeHeightfieldCollisionComponent::GetHeight(float X, flo
 {
 	TOptional<float> Height;
 	const float ZScale = GetComponentTransform().GetScale3D().Z * LANDSCAPE_ZSCALE;
-#if WITH_PHYSX && PHYSICS_INTERFACE_PHYSX
-	if (IsValidRef(HeightfieldRef) && HeightfieldRef->RBHeightfield != nullptr)
-	{
-		Height = HeightfieldRef->RBHeightfield->getHeight(HeightfieldRef->RBHeightfield->getNbRows() - 1 - X, Y) * ZScale;
-	}
-#elif WITH_CHAOS
+
 	if (IsValidRef(HeightfieldRef) && HeightfieldRef->Heightfield.Get())
 	{
 		Height = HeightfieldRef->Heightfield->GetHeightAt({ X, Y });
 	}
-#endif // WITH_CHAOS
+
 	return Height;
 }
 
@@ -3066,55 +2409,24 @@ struct FHeightFieldAccessor
 {
 	FHeightFieldAccessor(const ULandscapeHeightfieldCollisionComponent::FHeightfieldGeometryRef& InGeometryRef)
 	: GeometryRef(InGeometryRef)
-#if WITH_PHYSX && PHYSICS_INTERFACE_PHYSX
-	, NumX(InGeometryRef.RBHeightfield ? InGeometryRef.RBHeightfield->getNbColumns() : 0)
-	, NumY(InGeometryRef.RBHeightfield ? InGeometryRef.RBHeightfield->getNbRows() : 0)
-#elif WITH_CHAOS
 	, NumX(InGeometryRef.Heightfield.IsValid() ? InGeometryRef.Heightfield->GetNumCols() : 0)
 	, NumY(InGeometryRef.Heightfield.IsValid() ? InGeometryRef.Heightfield->GetNumRows() : 0)
-#endif // WITH_CHAOS
 	{
-#if WITH_PHYSX && PHYSICS_INTERFACE_PHYSX
-		const int32 CellCount = NumX * NumY;
-		if (CellCount > 0)
-		{
-			HFSamples.SetNumUninitialized(CellCount);
-			GeometryRef.RBHeightfield->saveCells(HFSamples.GetData(), CellCount * HFSamples.GetTypeSize());
-		}
-#endif // WITH_PHYSX && PHYSICS_INTERFACE_PHYSX
 	}
 
 	float GetUnscaledHeight(int32 X, int32 Y) const
 	{
-#if WITH_PHYSX && PHYSICS_INTERFACE_PHYSX
-		return float(HFSamples[NumX * (NumY - 1 - X) + Y].height);
-#elif WITH_CHAOS
 		return GeometryRef.Heightfield->GetHeight(X, Y);
-#else // WITH_CHAOS
-		return 0.0f;
-#endif
 	}
 
 	uint8 GetMaterialIndex(int32 X, int32 Y) const
 	{
-#if WITH_PHYSX && PHYSICS_INTERFACE_PHYSX
-		// we'll just use the sample from the first triangle
-		return HFSamples[NumX * (NumY - 1 - X) + Y].materialIndex0;
-#elif WITH_CHAOS
 		return GeometryRef.Heightfield->GetMaterialIndex(X, Y);
-#else // WITH_CHAOS
-		return 0;
-#endif
 	}
 
 	const ULandscapeHeightfieldCollisionComponent::FHeightfieldGeometryRef& GeometryRef;
 	const int32 NumX = 0;
 	const int32 NumY = 0;
-
-#if WITH_PHYSX && PHYSICS_INTERFACE_PHYSX
-private:
-	TArray<PxHeightFieldSample> HFSamples;
-#endif // WITH_PHYSX && PHYSICS_INTERFACE_PHYSX
 };
 
 bool ULandscapeHeightfieldCollisionComponent::FillHeightTile(TArrayView<float> Heights, int32 Offset, int32 Stride) const
@@ -3233,7 +2545,6 @@ void ALandscapeProxy::GetHeightValues(int32& SizeX, int32& SizeY, TArray<float>
 	SizeY = 0;
 	ArrayValues.SetNum(0);
 	
-#if WITH_CHAOS
 	// Exit if we have no landscape data
 	if (LandscapeComponents.Num() == 0 || CollisionComponents.Num() == 0)
 	{
@@ -3317,5 +2628,4 @@ void ALandscapeProxy::GetHeightValues(int32& SizeX, int32& SizeY, TArray<float>
 			}
 		}
 	}
-#endif // WITH_CHAOS
 }
diff --git a/Engine/Source/Runtime/MRMesh/Private/MRMeshComponent.cpp b/Engine/Source/Runtime/MRMesh/Private/MRMeshComponent.cpp
--- a/Engine/Source/Runtime/MRMesh/Private/MRMeshComponent.cpp
+++ b/Engine/Source/Runtime/MRMesh/Private/MRMeshComponent.cpp
@@ -17,10 +17,7 @@
 #include "PhysicsEngine/BodyInstance.h"
 #include "PhysicsEngine/PhysicsSettings.h"
 #include "PhysicsPublic.h"
-#if PHYSICS_INTERFACE_PHYSX
-#include "IPhysXCooking.h"
-#include "PhysXCookHelper.h"
-#endif
+
 #include "Misc/RuntimeErrors.h"
 #include "Engine/Engine.h"
 #include "UObject/UObjectThreadContext.h"
@@ -28,9 +25,7 @@
 #include "UObject/ConstructorHelpers.h"
 #include "Materials/MaterialInstanceDynamic.h"
 
-#if WITH_CHAOS
-	#include "Experimental/ChaosDerivedData.h"
-#endif
+#include "Experimental/ChaosDerivedData.h"
 
 // See ISpatialAcceleration.h, the header is not included to avoid having a dependency on Chaos
 #ifndef CHAOS_SERIALIZE_OUT
@@ -38,7 +33,7 @@
 #endif
 
 // When using the Chaos code path, we're using a private class FChaosDerivedDataCooker which is not exported
-#define SUPPORTS_PHYSICS_COOKING WITH_PHYSX && (PHYSICS_INTERFACE_PHYSX || (WITH_CHAOS && CHAOS_SERIALIZE_OUT && IS_MONOLITHIC))
+#define SUPPORTS_PHYSICS_COOKING (CHAOS_SERIALIZE_OUT && IS_MONOLITHIC)
 
 DECLARE_CYCLE_STAT(TEXT("MrMesh SetCollisionProfileName"), STAT_MrMesh_SetCollisionProfileName, STATGROUP_Physics);
 DECLARE_CYCLE_STAT(TEXT("Update Collision"), STAT_UpdateCollision, STATGROUP_MRMESH);
diff --git a/Engine/Source/Runtime/NavigationSystem/NavigationSystem.Build.cs b/Engine/Source/Runtime/NavigationSystem/NavigationSystem.Build.cs
--- a/Engine/Source/Runtime/NavigationSystem/NavigationSystem.Build.cs
+++ b/Engine/Source/Runtime/NavigationSystem/NavigationSystem.Build.cs
@@ -44,15 +44,12 @@ namespace UnrealBuildTool.Rules
 
             SetupModulePhysicsSupport(Target);
 
-			if (Target.bCompileChaos || Target.bUseChaos)
-            {
-                PublicDependencyModuleNames.AddRange(
-                    new string[] {
-						"Chaos",
-                        "GeometryCollectionEngine",
-                    }
-					);
-            }
+            PublicDependencyModuleNames.AddRange(
+                new string[] {
+					"Chaos",
+                    "GeometryCollectionEngine",
+                }
+			);
 
             if (Target.bCompileRecast)
             {
diff --git a/Engine/Source/Runtime/NavigationSystem/Private/NavCollision.cpp b/Engine/Source/Runtime/NavigationSystem/Private/NavCollision.cpp
--- a/Engine/Source/Runtime/NavigationSystem/Private/NavCollision.cpp
+++ b/Engine/Source/Runtime/NavigationSystem/Private/NavCollision.cpp
@@ -36,11 +36,7 @@ namespace NavCollisionCookStats
 }
 #endif
 
-#if WITH_PHYSX && PHYSICS_INTERFACE_PHYSX
-static const FName NAVCOLLISION_FORMAT = TEXT("NavCollision_X");
-#elif WITH_CHAOS
 static const FName NAVCOLLISION_FORMAT = TEXT("NavCollision_Chaos");
-#endif
 
 class FNavCollisionDataReader
 {
diff --git a/Engine/Source/Runtime/NavigationSystem/Private/NavMesh/RecastNavMeshGenerator.cpp b/Engine/Source/Runtime/NavigationSystem/Private/NavMesh/RecastNavMeshGenerator.cpp
--- a/Engine/Source/Runtime/NavigationSystem/Private/NavMesh/RecastNavMeshGenerator.cpp
+++ b/Engine/Source/Runtime/NavigationSystem/Private/NavMesh/RecastNavMeshGenerator.cpp
@@ -17,13 +17,9 @@
 #include "UObject/GarbageCollection.h"
 
 #if WITH_RECAST
-#if WITH_PHYSX
-	#include "PhysXPublic.h"
-#endif
-#if WITH_CHAOS
-	#include "Chaos/HeightField.h"
-	#include "Chaos/TriangleMeshImplicitObject.h"
-#endif
+
+#include "Chaos/HeightField.h"
+#include "Chaos/TriangleMeshImplicitObject.h"
 #include "NavMesh/PImplRecastNavMesh.h"
 
 // recast includes
@@ -233,19 +229,10 @@ struct FRecastGeometryExport : public FNavigableGeometryExport
 	TNavStatArray<int32> IndexBuffer;
 	FWalkableSlopeOverride SlopeOverride;
 
-#if PHYSICS_INTERFACE_PHYSX
-	virtual void ExportPxTriMesh16Bit(physx::PxTriangleMesh const * const TriMesh, const FTransform& LocalToWorld) override;
-	virtual void ExportPxTriMesh32Bit(physx::PxTriangleMesh const * const TriMesh, const FTransform& LocalToWorld) override;
-	virtual void ExportPxConvexMesh(physx::PxConvexMesh const * const ConvexMesh, const FTransform& LocalToWorld) override;
-	virtual void ExportPxHeightField(physx::PxHeightField const * const HeightField, const FTransform& LocalToWorld) override;
-	virtual void ExportPxHeightFieldSlice(const FNavHeightfieldSamples& PrefetchedHeightfieldSamples, const int32 NumRows, const int32 NumCols, const FTransform& LocalToWorld, const FBox& SliceBox) override;
-#endif // PHYSICS_INTERFACE_PHYSX
-#if WITH_CHAOS
 	virtual void ExportChaosTriMesh(const Chaos::FTriangleMeshImplicitObject* const TriMesh, const FTransform& LocalToWorld) override;
 	virtual void ExportChaosConvexMesh(const FKConvexElem* const Convex, const FTransform& LocalToWorld) override;
 	virtual void ExportChaosHeightField(const Chaos::FHeightField* const Heightfield, const FTransform& LocalToWorld) override;
 	virtual void ExportChaosHeightFieldSlice(const FNavHeightfieldSamples& PrefetchedHeightfieldSamples, const int32 NumRows, const int32 NumCols, const FTransform& LocalToWorld, const FBox& SliceBox) override;
-#endif
 	virtual void ExportCustomMesh(const FVector* InVertices, int32 NumVerts, const int32* InIndices, int32 NumIndices, const FTransform& LocalToWorld) override;
 	virtual void ExportRigidBodySetup(UBodySetup& BodySetup, const FTransform& LocalToWorld) override;
 	virtual void AddNavModifiers(const FCompositeNavModifier& Modifiers) override;
@@ -356,308 +343,6 @@ static void StoreCollisionCache(FRecastGeometryExport& GeomExport)
 	FMemory::Memcpy(RawMemory + HeaderSize + CoordsSize, GeomExport.IndexBuffer.GetData(), IndicesSize);
 }
 
-#if PHYSICS_INTERFACE_PHYSX
-/** exports PxConvexMesh as trimesh */
-void ExportPxConvexMesh(PxConvexMesh const * const ConvexMesh, const FTransform& LocalToWorld,
-						TNavStatArray<FVector::FReal>& VertexBuffer, TNavStatArray<int32>& IndexBuffer,
-						FBox& UnrealBounds)
-{
-	// after FKConvexElem::AddCachedSolidConvexGeom
-	if(ConvexMesh == NULL)
-	{
-		return;
-	}
-
-	int32 StartVertOffset = VertexBuffer.Num() / 3;
-	const bool bNegX = LocalToWorld.GetDeterminant() < 0;
-
-	// get PhysX data
-	const PxVec3* PVertices = ConvexMesh->getVertices();
-	const PxU8* PIndexBuffer = ConvexMesh->getIndexBuffer();
-	const PxU32 NbPolygons = ConvexMesh->getNbPolygons();
-
-#if SHOW_NAV_EXPORT_PREVIEW
-	UWorld* DebugWorld = FindEditorWorld();
-#endif // SHOW_NAV_EXPORT_PREVIEW
-
-	for(PxU32 i = 0; i < NbPolygons; ++i)
-	{
-		PxHullPolygon Data;
-		bool bStatus = ConvexMesh->getPolygonData(i, Data);
-		check(bStatus);
-
-		const PxU8* indices = PIndexBuffer + Data.mIndexBase;
-		
-		// add vertices 
-		for(PxU32 j = 0; j < Data.mNbVerts; ++j)
-		{
-			const int32 VertIndex = indices[j];
-			const FVector UnrealCoords = LocalToWorld.TransformPosition( P2UVector(PVertices[VertIndex]) );
-			UnrealBounds += UnrealCoords;
-
-			VertexBuffer.Add(UnrealCoords.X);
-			VertexBuffer.Add(UnrealCoords.Y);
-			VertexBuffer.Add(UnrealCoords.Z);
-		}
-
-		// Add indices
-		const PxU32 nbTris = Data.mNbVerts - 2;
-		for(PxU32 j = 0; j < nbTris; ++j)
-		{
-			IndexBuffer.Add(StartVertOffset + 0 );
-			IndexBuffer.Add(StartVertOffset + j + 2);
-			IndexBuffer.Add(StartVertOffset + j + 1);
-
-#if SHOW_NAV_EXPORT_PREVIEW
-			if (DebugWorld)
-			{
-				FVector V0(VertexBuffer[(StartVertOffset + 0) * 3+0], VertexBuffer[(StartVertOffset + 0) * 3+1], VertexBuffer[(StartVertOffset + 0) * 3+2]);
-				FVector V1(VertexBuffer[(StartVertOffset + j + 2) * 3+0], VertexBuffer[(StartVertOffset + j + 2) * 3+1], VertexBuffer[(StartVertOffset + j + 2) * 3+2]);
-				FVector V2(VertexBuffer[(StartVertOffset + j + 1) * 3+0], VertexBuffer[(StartVertOffset + j + 1) * 3+1], VertexBuffer[(StartVertOffset + j + 1) * 3+2]);
-
-				DrawDebugLine(DebugWorld, V0, V1, bNegX ? FColor::Red : FColor::Blue, true);
-				DrawDebugLine(DebugWorld, V1, V2, bNegX ? FColor::Red : FColor::Blue, true);
-				DrawDebugLine(DebugWorld, V2, V0, bNegX ? FColor::Red : FColor::Blue, true);
-			}
-#endif // SHOW_NAV_EXPORT_PREVIEW
-		}
-
-		StartVertOffset += Data.mNbVerts;
-	}
-}
-
-template<typename TIndicesType> 
-FORCEINLINE_DEBUGGABLE void ExportPxTriMesh(PxTriangleMesh const * const TriMesh, const FTransform& LocalToWorld,
-											TNavStatArray<FVector::FReal>& VertexBuffer, TNavStatArray<int32>& IndexBuffer,
-											FBox& UnrealBounds)
-{
-	if (TriMesh == NULL)
-	{
-		return;
-	}
-
-	int32 VertOffset = VertexBuffer.Num() / 3;
-	const PxVec3* PVerts = TriMesh->getVertices();
-	const PxU32 NumTris = TriMesh->getNbTriangles();
-
-	const TIndicesType* Indices = (TIndicesType*)TriMesh->getTriangles();;
-		
-	VertexBuffer.Reserve(VertexBuffer.Num() + NumTris*3);
-	IndexBuffer.Reserve(IndexBuffer.Num() + NumTris*3);
-	const bool bFlipCullMode = (LocalToWorld.GetDeterminant() < 0.f);
-	const int32 IndexOrder[3] = { bFlipCullMode ? 0 : 2, 1, bFlipCullMode ? 2 : 0 };
-
-#if SHOW_NAV_EXPORT_PREVIEW
-	UWorld* DebugWorld = FindEditorWorld();
-#endif // SHOW_NAV_EXPORT_PREVIEW
-
-	for(PxU32 TriIdx = 0; TriIdx < NumTris; ++TriIdx)
-	{
-		for (int32 i = 0; i < 3; i++)
-		{
-			const FVector UnrealCoords = LocalToWorld.TransformPosition(P2UVector(PVerts[Indices[i]]));
-			UnrealBounds += UnrealCoords;
-
-			VertexBuffer.Add(UnrealCoords.X);
-			VertexBuffer.Add(UnrealCoords.Y);
-			VertexBuffer.Add(UnrealCoords.Z);
-		}
-		Indices += 3;
-
-		IndexBuffer.Add(VertOffset + IndexOrder[0]);
-		IndexBuffer.Add(VertOffset + IndexOrder[1]);
-		IndexBuffer.Add(VertOffset + IndexOrder[2]);
-
-#if SHOW_NAV_EXPORT_PREVIEW
-		if (DebugWorld)
-		{
-			FVector V0(VertexBuffer[(VertOffset + IndexOrder[0]) * 3+0], VertexBuffer[(VertOffset + IndexOrder[0]) * 3+1], VertexBuffer[(VertOffset + IndexOrder[0]) * 3+2]);
-			FVector V1(VertexBuffer[(VertOffset + IndexOrder[1]) * 3+0], VertexBuffer[(VertOffset + IndexOrder[1]) * 3+1], VertexBuffer[(VertOffset + IndexOrder[1]) * 3+2]);
-			FVector V2(VertexBuffer[(VertOffset + IndexOrder[2]) * 3+0], VertexBuffer[(VertOffset + IndexOrder[2]) * 3+1], VertexBuffer[(VertOffset + IndexOrder[2]) * 3+2]);
-
-			DrawDebugLine(DebugWorld, V0, V1, bFlipCullMode ? FColor::Red : FColor::Blue, true);
-			DrawDebugLine(DebugWorld, V1, V2, bFlipCullMode ? FColor::Red : FColor::Blue, true);
-			DrawDebugLine(DebugWorld, V2, V0, bFlipCullMode ? FColor::Red : FColor::Blue, true);
-		}
-#endif // SHOW_NAV_EXPORT_PREVIEW
-
-		VertOffset += 3;
-	}
-}
-
-void ExportPxHeightField(PxHeightField const * const HeightField, const FTransform& LocalToWorld
-	, TNavStatArray<FVector::FReal>& VertexBuffer, TNavStatArray<int32>& IndexBuffer
-	, FBox& UnrealBounds)
-{
-	if (HeightField == NULL)
-	{
-		return;
-	}
-
-	QUICK_SCOPE_CYCLE_COUNTER(STAT_NavMesh_ExportPxHeightField);
-
-	const int32 NumRows = HeightField->getNbRows();
-	const int32 NumCols = HeightField->getNbColumns();
-	const int32 VertexCount = NumRows * NumCols;
-
-	// Unfortunately we have to use PxHeightField::saveCells instead PxHeightField::getHeight here 
-	// because current PxHeightField interface does not provide an access to a triangle material index by HF 2D coordinates
-	// PxHeightField::getTriangleMaterialIndex uses some internal adressing which does not match HF 2D coordinates
-	TArray<PxHeightFieldSample> HFSamples;
-	HFSamples.SetNumUninitialized(VertexCount);
-	{
-		QUICK_SCOPE_CYCLE_COUNTER(STAT_NavMesh_ExportPxHeightField_saveCells);
-		HeightField->saveCells(HFSamples.GetData(), VertexCount*HFSamples.GetTypeSize());
-	}
-
-	//
-	const int32 VertOffset = VertexBuffer.Num() / 3;
-	const int32 NumQuads = (NumRows - 1)*(NumCols - 1);
-
-	VertexBuffer.Reserve(VertexBuffer.Num() + VertexCount * 3);
-	IndexBuffer.Reserve(IndexBuffer.Num() + NumQuads * 6);
-
-	const bool bMirrored = (LocalToWorld.GetDeterminant() < 0.f);
-
-	for (int32 Y = 0; Y < NumRows; Y++)
-	{
-		for (int32 X = 0; X < NumCols; X++)
-		{
-			const int32 SampleIdx = (bMirrored ? X : (NumCols - X - 1))*NumCols + Y;
-
-			const PxHeightFieldSample& Sample = HFSamples[SampleIdx];
-			const FVector UnrealCoords = LocalToWorld.TransformPosition(FVector(X, Y, Sample.height));
-			UnrealBounds += UnrealCoords;
-
-			VertexBuffer.Add(UnrealCoords.X);
-			VertexBuffer.Add(UnrealCoords.Y);
-			VertexBuffer.Add(UnrealCoords.Z);
-		}
-	}
-
-	for (int32 Y = 0; Y < NumRows - 1; Y++)
-	{
-		for (int32 X = 0; X < NumCols - 1; X++)
-		{
-			const int32 SampleIdx = (bMirrored ? X : (NumCols - X - 1 - 1))*NumCols + Y;
-			const PxHeightFieldSample& Sample = HFSamples[SampleIdx];
-			const bool bIsHole = (Sample.materialIndex0 == PxHeightFieldMaterial::eHOLE);
-			if (bIsHole)
-			{
-				continue;
-			}
-
-			const int32 I00 = X + 0 + (Y + 0)*NumCols;
-			int32 I01 = X + 0 + (Y + 1)*NumCols;
-			int32 I10 = X + 1 + (Y + 0)*NumCols;
-			const int32 I11 = X + 1 + (Y + 1)*NumCols;
-
-			if (bMirrored)
-			{
-				Swap(I01, I10);
-			}
-
-			IndexBuffer.Add(VertOffset + I00);
-			IndexBuffer.Add(VertOffset + I11);
-			IndexBuffer.Add(VertOffset + I10);
-
-			IndexBuffer.Add(VertOffset + I00);
-			IndexBuffer.Add(VertOffset + I01);
-			IndexBuffer.Add(VertOffset + I11);
-		}
-	}
-}
-
-void ExportPxHeightFieldSlice(const FNavHeightfieldSamples& PrefetchedHeightfieldSamples, const int32 NumRows, const int32 NumCols, const FTransform& LocalToWorld
-	, TNavStatArray<FVector::FReal>& VertexBuffer, TNavStatArray<int32>& IndexBuffer, const FBox& SliceBox
-	, FBox& UnrealBounds)
-{
-	QUICK_SCOPE_CYCLE_COUNTER(STAT_NavMesh_ExportPxHeightFieldSlice);
-
-#if PHYSICS_INTERFACE_PHYSX
-	static const uint32 SizeOfPx = sizeof(physx::PxI16);
-	static const uint32 SizeOfHeight = PrefetchedHeightfieldSamples.Heights.GetTypeSize();
-	ensure(SizeOfPx == SizeOfHeight);
-#endif // PHYSICS_INTERFACE_PHYSX
-
-	// calculate the actual start and number of columns we want
-	const FBox LocalBox = SliceBox.TransformBy(LocalToWorld.Inverse());
-	const bool bMirrored = (LocalToWorld.GetDeterminant() < 0.f);
-
-	const int32 MinX = FMath::Clamp(FMath::FloorToInt(LocalBox.Min.X) - 1, 0, NumCols);
-	const int32 MinY = FMath::Clamp(FMath::FloorToInt(LocalBox.Min.Y) - 1, 0, NumRows);
-	const int32 MaxX = FMath::Clamp(FMath::CeilToInt(LocalBox.Max.X) + 1, 0, NumCols);
-	const int32 MaxY = FMath::Clamp(FMath::CeilToInt(LocalBox.Max.Y) + 1, 0, NumRows);
-	const int32 SizeX = MaxX - MinX;
-	const int32 SizeY = MaxY - MinY;
-
-	if (SizeX <= 0 || SizeY <= 0)
-	{
-		// slice is outside bounds, skip
-		return;
-	}
-
-	const int32 VertOffset = VertexBuffer.Num() / 3;
-	const int32 NumVerts = SizeX * SizeY;
-	const int32 NumQuads = (SizeX - 1) * (SizeY - 1);
-	{
-		QUICK_SCOPE_CYCLE_COUNTER(STAT_RecastGeometryExport_AllocatingMemory);
-		VertexBuffer.Reserve(VertexBuffer.Num() + NumVerts * 3);
-		IndexBuffer.Reserve(IndexBuffer.Num() + NumQuads * 3 * 2);
-	}
-
-	for (int32 IdxY = 0; IdxY < SizeY; IdxY++)
-	{
-		for (int32 IdxX = 0; IdxX < SizeX; IdxX++)
-		{
-			const int32 CoordX = IdxX + MinX;
-			const int32 CoordY = IdxY + MinY;
-			const int32 SampleIdx = ((bMirrored ? CoordX : (NumCols - CoordX - 1)) * NumCols) + CoordY;
-
-			const FVector UnrealCoords = LocalToWorld.TransformPosition(FVector(CoordX, CoordY, PrefetchedHeightfieldSamples.Heights[SampleIdx]));
-			VertexBuffer.Add(UnrealCoords.X);
-			VertexBuffer.Add(UnrealCoords.Y);
-			VertexBuffer.Add(UnrealCoords.Z);
-		}
-	}
-
-	for (int32 IdxY = 0; IdxY < SizeY - 1; IdxY++)
-	{
-		for (int32 IdxX = 0; IdxX < SizeX - 1; IdxX++)
-		{
-			const int32 CoordX = IdxX + MinX;
-			const int32 CoordY = IdxY + MinY;
-			const int32 SampleIdx = ((bMirrored ? CoordX : (NumCols - CoordX - 1)) * NumCols) + CoordY;
-
-			const bool bIsHole = PrefetchedHeightfieldSamples.Holes[SampleIdx];
-			if (bIsHole)
-			{
-				continue;
-			}
-
-			const int32 I00 = (IdxX + 0) + (IdxY + 0) * SizeX;
-			int32 I01 = (IdxX + 0) + (IdxY + 1) * SizeX;
-			int32 I10 = (IdxX + 1) + (IdxY + 0) * SizeX;
-			const int32 I11 = (IdxX + 1) + (IdxY + 1) * SizeX;
-			if (bMirrored)
-			{
-				Swap(I01, I10);
-			}
-
-			IndexBuffer.Add(VertOffset + I00);
-			IndexBuffer.Add(VertOffset + I11);
-			IndexBuffer.Add(VertOffset + I10);
-
-			IndexBuffer.Add(VertOffset + I00);
-			IndexBuffer.Add(VertOffset + I01);
-			IndexBuffer.Add(VertOffset + I11);
-		}
-	}
-}
-
-#endif // WITH_PHYSX
-
-#if WITH_CHAOS
 void ExportChaosTriMesh(const Chaos::FTriangleMeshImplicitObject* const TriMesh, const FTransform& LocalToWorld
 	, TNavStatArray<FVector::FReal>& VertexBuffer, TNavStatArray<int32>& IndexBuffer
 	, FBox& UnrealBounds)
@@ -960,7 +645,7 @@ void ExportChaosHeightFieldSlice(const FNavHeightfieldSamples& PrefetchedHeightf
 		}
 	}
 }
-#endif
+
 
 void ExportCustomMesh(const FVector* InVertices, int32 NumVerts, const int32* InIndices, int32 NumIndices, const FTransform& LocalToWorld,
 					  TNavStatArray<FVector::FReal>& VertexBuffer, TNavStatArray<int32>& IndexBuffer, FBox& UnrealBounds)
@@ -1057,25 +742,11 @@ FORCEINLINE_DEBUGGABLE void ExportRigidBodyConvexElements(UBodySetup& BodySetup,
 		// Store index of first vertex in shape buffer
 		ShapeBuffer.Add(VertexBuffer.Num() / 3);
 
-#if WITH_PHYSX && PHYSICS_INTERFACE_PHYSX
-		// Get verts/triangles from this hull.
-		if (!ConvexElem->GetConvexMesh() && ConvexElem->GetMirroredConvexMesh())
-		{
-			// If there is only a NegX mesh (e.g. a mirrored volume), use it
-			ExportPxConvexMesh(ConvexElem->GetMirroredConvexMesh(), NegXScale * LocalToWorld, VertexBuffer, IndexBuffer, UnrealBounds);
-		}
-		else
-		{
-			// Otherwise use the regular mesh in the case that both exist
-			ExportPxConvexMesh(ConvexElem->GetConvexMesh(), LocalToWorld, VertexBuffer, IndexBuffer, UnrealBounds);
-		}
-#elif WITH_CHAOS
 		if (ConvexElem->GetChaosConvexMesh())
 		{
 			// TODO use ConvexElem->GetTransform?() transform?
 			ExportChaosConvexMesh(ConvexElem, LocalToWorld, VertexBuffer, IndexBuffer, UnrealBounds);
 		}
-#endif
 	}
 }
 
@@ -1084,25 +755,11 @@ FORCEINLINE_DEBUGGABLE void ExportRigidBodyTriMesh(UBodySetup& BodySetup, TNavSt
 {
 	if (BodySetup.GetCollisionTraceFlag() == CTF_UseComplexAsSimple)
 	{
-#if WITH_PHYSX && PHYSICS_INTERFACE_PHYSX
-		for (PxTriangleMesh* TriMesh : BodySetup.TriMeshes)
-		{
-			if (TriMesh->getTriangleMeshFlags() & PxTriangleMeshFlag::e16_BIT_INDICES)
-			{
-				ExportPxTriMesh<PxU16>(TriMesh, LocalToWorld, VertexBuffer, IndexBuffer, UnrealBounds);
-			}
-			else
-			{
-				ExportPxTriMesh<PxU32>(TriMesh, LocalToWorld, VertexBuffer, IndexBuffer, UnrealBounds);
-			}
-		}
-#elif WITH_CHAOS
 		for(const auto& TriMesh : BodySetup.ChaosTriMeshes)
 		{
 			ExportChaosTriMesh(TriMesh.Get(), LocalToWorld, VertexBuffer, IndexBuffer, UnrealBounds);
 		}
-#endif
-}
+	}
 }
 
 void ExportRigidBodyBoxElements(const FKAggregateGeom& AggGeom, TNavStatArray<FVector::FReal>& VertexBuffer, TNavStatArray<int32>& IndexBuffer,
@@ -1434,34 +1091,6 @@ void ExportVertexSoup(const TArray<FVector>& VertexSoup, TNavStatArray<FVector::
 
 } // namespace RecastGeometryExport
 
-#if PHYSICS_INTERFACE_PHYSX
-void FRecastGeometryExport::ExportPxTriMesh16Bit(physx::PxTriangleMesh const * const TriMesh, const FTransform& LocalToWorld)
-{
-	RecastGeometryExport::ExportPxTriMesh<PxU16>(TriMesh, LocalToWorld, VertexBuffer, IndexBuffer, Data->Bounds);
-}
-
-void FRecastGeometryExport::ExportPxTriMesh32Bit(physx::PxTriangleMesh const * const TriMesh, const FTransform& LocalToWorld)
-{
-	RecastGeometryExport::ExportPxTriMesh<PxU32>(TriMesh, LocalToWorld, VertexBuffer, IndexBuffer, Data->Bounds);
-}
-
-void FRecastGeometryExport::ExportPxConvexMesh(physx::PxConvexMesh const * const ConvexMesh, const FTransform& LocalToWorld)
-{
-	RecastGeometryExport::ExportPxConvexMesh(ConvexMesh, LocalToWorld, VertexBuffer, IndexBuffer, Data->Bounds);
-}
-
-void FRecastGeometryExport::ExportPxHeightField(physx::PxHeightField const * const HeightField, const FTransform& LocalToWorld)
-{
-	RecastGeometryExport::ExportPxHeightField(HeightField, LocalToWorld, VertexBuffer, IndexBuffer, Data->Bounds);
-}
-
-void FRecastGeometryExport::ExportPxHeightFieldSlice(const FNavHeightfieldSamples& PrefetchedHeightfieldSamples, const int32 NumRows, const int32 NumCols, const FTransform& LocalToWorld, const FBox& SliceBox)
-{
-	RecastGeometryExport::ExportPxHeightFieldSlice(PrefetchedHeightfieldSamples, NumRows, NumCols, LocalToWorld, VertexBuffer, IndexBuffer, SliceBox, Data->Bounds);
-}
-#endif // WITH_PHYSX
-
-#if WITH_CHAOS
 void FRecastGeometryExport::ExportChaosTriMesh(const Chaos::FTriangleMeshImplicitObject* const TriMesh, const FTransform& LocalToWorld)
 {
 	RecastGeometryExport::ExportChaosTriMesh(TriMesh, LocalToWorld, VertexBuffer, IndexBuffer, Data->Bounds);
@@ -1481,7 +1110,6 @@ void FRecastGeometryExport::ExportChaosHeightFieldSlice(const FNavHeightfieldSam
 {
 	RecastGeometryExport::ExportChaosHeightFieldSlice(PrefetchedHeightfieldSamples, NumRows, NumCols, LocalToWorld, VertexBuffer, IndexBuffer, SliceBox, Data->Bounds);
 }
-#endif
 
 void FRecastGeometryExport::ExportCustomMesh(const FVector* InVertices, int32 NumVerts, const int32* InIndices, int32 NumIndices, const FTransform& LocalToWorld)
 {
diff --git a/Engine/Source/Runtime/NavigationSystem/Private/NavigationSystem.cpp b/Engine/Source/Runtime/NavigationSystem/Private/NavigationSystem.cpp
--- a/Engine/Source/Runtime/NavigationSystem/Private/NavigationSystem.cpp
+++ b/Engine/Source/Runtime/NavigationSystem/Private/NavigationSystem.cpp
@@ -3427,15 +3427,7 @@ void UNavigationSystemV1::Build()
 	TRACE_CPUPROFILER_EVENT_SCOPE(UNavigationSystemV1::Build);
 	
 	UE_LOG(LogNavigationDataBuild, Display, TEXT("UNavigationSystemV1::Build started..."));
-#if PHYSICS_INTERFACE_PHYSX
-	UE_LOG(LogNavigationDataBuild, Display, TEXT("   Building navigation data using PHYSICS_INTERFACE_PHYSX."));
-#endif
-#if WITH_PHYSX
-	UE_LOG(LogNavigationDataBuild, Display, TEXT("   Building navigation data using WITH_PHYSX."));
-#endif
-#if WITH_CHAOS
-	UE_LOG(LogNavigationDataBuild, Display, TEXT("   Building navigation data using WITH_CHAOS."));
-#endif
+	UE_LOG(LogNavigationDataBuild, Display, TEXT("   Building navigation data using Chaos physics."));
 
 	UWorld* World = GetWorld();
 	if (!World)
diff --git a/Engine/Source/Runtime/NavigationSystem/Private/NavigationSystemTypes.cpp b/Engine/Source/Runtime/NavigationSystem/Private/NavigationSystemTypes.cpp
--- a/Engine/Source/Runtime/NavigationSystem/Private/NavigationSystemTypes.cpp
+++ b/Engine/Source/Runtime/NavigationSystem/Private/NavigationSystemTypes.cpp
@@ -231,13 +231,7 @@ namespace NavigationHelper
 
 	bool IsBodyNavigationRelevant(const UBodySetup& BodySetup)
 	{
-#if PHYSICS_INTERFACE_PHYSX
-		const bool bBodyHasGeometry = (BodySetup.AggGeom.GetElementCount() > 0 || BodySetup.TriMeshes.Num() > 0);
-#elif WITH_CHAOS
 		const bool bBodyHasGeometry = (BodySetup.AggGeom.GetElementCount() > 0 || BodySetup.ChaosTriMeshes.Num() > 0);
-#else
-		const bool bBodyHasGeometry = (BodySetup.AggGeom.GetElementCount() > 0);
-#endif
 
 		// has any colliding geometry
 		return bBodyHasGeometry
diff --git a/Engine/Source/Runtime/PhysicsCore/PhysicsCore.Build.cs b/Engine/Source/Runtime/PhysicsCore/PhysicsCore.Build.cs
--- a/Engine/Source/Runtime/PhysicsCore/PhysicsCore.Build.cs
+++ b/Engine/Source/Runtime/PhysicsCore/PhysicsCore.Build.cs
@@ -28,22 +28,6 @@ public class PhysicsCore: ModuleRules
 		PublicIncludePathModuleNames.Remove("PhysicsCore");
 		PublicDependencyModuleNames.Remove("PhysicsCore");
 
-		if (Target.bCompileChaos == false && Target.bUseChaos == false)
-        {
-            if (Target.bCompilePhysX)
-            {
-                // Not ideal but as this module publicly exposes PhysX types
-                // to other modules when PhysX is enabled it requires that its
-                // public files have access to PhysX includes
-                PublicDependencyModuleNames.Add("PhysX");
-            }
-
-            if (Target.bCompileAPEX)
-            {
-                PublicDependencyModuleNames.Add("APEX");
-            }
-        }
-
 		PrivateDefinitions.Add("CHAOS_INCLUDE_LEVEL_1=1");
 	}
 }
diff --git a/Engine/Source/Runtime/PhysicsCore/Private/ChaosEngineInterface.cpp b/Engine/Source/Runtime/PhysicsCore/Private/ChaosEngineInterface.cpp
--- a/Engine/Source/Runtime/PhysicsCore/Private/ChaosEngineInterface.cpp
+++ b/Engine/Source/Runtime/PhysicsCore/Private/ChaosEngineInterface.cpp
@@ -11,7 +11,6 @@
 
 FPhysicsDelegatesCore::FOnUpdatePhysXMaterial FPhysicsDelegatesCore::OnUpdatePhysXMaterial;
 
-#if WITH_CHAOS
 #include "ChaosInterfaceWrapperCore.h"
 #include "Chaos/TriangleMeshImplicitObject.h"
 #include "Chaos/Sphere.h"
@@ -671,14 +670,12 @@ FVector FChaosEngineInterface::GetWorldVelocityAtPoint_AssumesLocked(const FPhys
 	return FVector(0);
 }
 
-#if WITH_CHAOS
 FVector FChaosEngineInterface::GetWorldVelocityAtPoint_AssumesLocked(const Chaos::FRigidBodyHandle_Internal* Body_Internal, const FVector& InPoint)
 {
 	const Chaos::FVec3 COM = Body_Internal->CanTreatAsRigid() ? Chaos::FParticleUtilitiesGT::GetCoMWorldPosition(Body_Internal) : (Chaos::FVec3)Chaos::FParticleUtilitiesGT::GetActorWorldTransform(Body_Internal).GetTranslation();
 	const Chaos::FVec3 Diff = InPoint - COM;
 	return Body_Internal->V() - Chaos::FVec3::CrossProduct(Diff, Body_Internal->W());
 }
-#endif
 
 FTransform FChaosEngineInterface::GetComTransform_AssumesLocked(const FPhysicsActorHandle& InActorReference)
 {
@@ -1806,61 +1803,3 @@ void FChaosEngineInterface::SetKinematicTarget_AssumesLocked(const FPhysicsActor
 	FChaosScene* Scene = GetCurrentScene(InActorReference);
 	Scene->UpdateActorInAccelerationStructure(InActorReference);
 }
-
-#elif WITH_ENGINE //temp physx code to make moving out of engine easier
-
-#include "PhysXSupportCore.h"
-
-FPhysicsMaterialHandle FChaosEngineInterface::CreateMaterial(const UPhysicalMaterial* InMaterial)
-{
-	check(GPhysXSDK);
-
-	FPhysicsMaterialHandle_PhysX NewRef;
-
-	const float Friction = InMaterial->Friction;
-	const float Restitution = InMaterial->Restitution;
-
-	NewRef.Material = GPhysXSDK->createMaterial(Friction,Friction,Restitution);
-
-	return NewRef;
-}
-
-void FChaosEngineInterface::ReleaseMaterial(FPhysicsMaterialHandle_PhysX& InHandle)
-{
-	if(InHandle.IsValid())
-	{
-		InHandle.Material->userData = nullptr;
-		GPhysXPendingKillMaterial.Add(InHandle.Material);
-		InHandle.Material = nullptr;
-	}
-}
-
-void FChaosEngineInterface::UpdateMaterial(FPhysicsMaterialHandle_PhysX& InHandle,UPhysicalMaterial* InMaterial)
-{
-	if(InHandle.IsValid())
-	{
-		PxMaterial* PMaterial = InHandle.Material;
-
-		PMaterial->setStaticFriction(InMaterial->Friction);
-		PMaterial->setDynamicFriction(InMaterial->Friction);
-		PMaterial->setRestitution(InMaterial->Restitution);
-
-		const uint32 UseFrictionCombineMode = (InMaterial->bOverrideFrictionCombineMode ? InMaterial->FrictionCombineMode.GetValue() : UPhysicsSettingsCore::Get()->FrictionCombineMode.GetValue());
-		PMaterial->setFrictionCombineMode(static_cast<physx::PxCombineMode::Enum>(UseFrictionCombineMode));
-
-		const uint32 UseRestitutionCombineMode = (InMaterial->bOverrideRestitutionCombineMode ? InMaterial->RestitutionCombineMode.GetValue() : UPhysicsSettingsCore::Get()->RestitutionCombineMode.GetValue());
-		PMaterial->setRestitutionCombineMode(static_cast<physx::PxCombineMode::Enum>(UseRestitutionCombineMode));
-
-		FPhysicsDelegatesCore::OnUpdatePhysXMaterial.Broadcast(InMaterial);
-	}
-}
-
-void FChaosEngineInterface::SetUserData(FPhysicsMaterialHandle_PhysX& InHandle,void* InUserData)
-{
-	if(InHandle.IsValid())
-	{
-		InHandle.Material->userData = InUserData;
-	}
-}
-
-#endif
\ No newline at end of file
diff --git a/Engine/Source/Runtime/PhysicsCore/Private/ChaosScene.cpp b/Engine/Source/Runtime/PhysicsCore/Private/ChaosScene.cpp
--- a/Engine/Source/Runtime/PhysicsCore/Private/ChaosScene.cpp
+++ b/Engine/Source/Runtime/PhysicsCore/Private/ChaosScene.cpp
@@ -1,7 +1,5 @@
 // Copyright Epic Games, Inc. All Rights Reserved.
 
-#if WITH_CHAOS
-
 #include "Chaos/ChaosScene.h"
 
 #include "Async/AsyncWork.h"
@@ -183,14 +181,10 @@ void FChaosScene::Flush()
 
 void FChaosScene::RemoveActorFromAccelerationStructure(FPhysicsActorHandle Actor)
 {
-#if WITH_CHAOS
 	using namespace Chaos;
 	RemoveActorFromAccelerationStructureImp(Actor->GetParticle_LowLevel());
-#endif
 }
 
-#if WITH_CHAOS
-
 void FChaosScene::RemoveActorFromAccelerationStructureImp(Chaos::FGeometryParticle* Particle)
 {
 	using namespace Chaos;
@@ -201,12 +195,9 @@ void FChaosScene::RemoveActorFromAccelerationStructureImp(Chaos::FGeometryPartic
 		GetSpacialAcceleration()->RemoveElementFrom(AccelerationHandle, Particle->SpatialIdx());
 	}
 }
-#endif
-
 
 void FChaosScene::UpdateActorInAccelerationStructure(const FPhysicsActorHandle& Actor)
 {
-#if WITH_CHAOS
 	using namespace Chaos;
 
 	if(GetSpacialAcceleration())
@@ -233,12 +224,10 @@ void FChaosScene::UpdateActorInAccelerationStructure(const FPhysicsActorHandle&
 
 		GetSolver()->UpdateParticleInAccelerationStructure_External(Actor->GetParticle_LowLevel(),/*bDelete=*/false);
 	}
-#endif
 }
 
 void FChaosScene::UpdateActorsInAccelerationStructure(const TArrayView<FPhysicsActorHandle>& Actors)
 {
-#if WITH_CHAOS
 	using namespace Chaos;
 
 	if(GetSpacialAcceleration())
@@ -279,12 +268,10 @@ void FChaosScene::UpdateActorsInAccelerationStructure(const TArrayView<FPhysicsA
 			}
 		}
 	}
-#endif
 }
 
 void FChaosScene::AddActorsToScene_AssumesLocked(TArray<FPhysicsActorHandle>& InHandles,const bool bImmediate)
 {
-#if WITH_CHAOS
 	TRACE_CPUPROFILER_EVENT_SCOPE(FChaosScene::AddActorsToScene_AssumesLocked)
 
 	Chaos::FPhysicsSolver* Solver = GetSolver();
@@ -311,7 +298,6 @@ void FChaosScene::AddActorsToScene_AssumesLocked(TArray<FPhysicsActorHandle>& In
 			SpatialAcceleration->UpdateElementIn(AccelerationHandle,WorldBounds,bHasBounds, Body_External.SpatialIdx());
 		}
 	}
-#endif
 }
 
 void FChaosSceneSimCallback::OnPreSimulate_Internal()
@@ -329,7 +315,6 @@ void FChaosScene::SetGravity(const Chaos::FVec3& Acceleration)
 
 void FChaosScene::SetUpForFrame(const FVector* NewGrav,float InDeltaSeconds /*= 0.0f*/,float InMinPhysicsDeltaTime /*= 0.0f*/,float InMaxPhysicsDeltaTime /*= 0.0f*/,float InMaxSubstepDeltaTime /*= 0.0f*/,int32 InMaxSubsteps,bool bSubstepping)
 {
-#if WITH_CHAOS
 	using namespace Chaos;
 	SetGravity(*NewGrav);
 
@@ -358,12 +343,10 @@ void FChaosScene::SetUpForFrame(const FVector* NewGrav,float InDeltaSeconds /*=
 		}
 		Solver->SetMinDeltaTime_External(InMinPhysicsDeltaTime);
 	}
-#endif
 }
 
 void FChaosScene::StartFrame()
 {
-#if WITH_CHAOS
 	using namespace Chaos;
 
 	SCOPE_CYCLE_COUNTER(STAT_Scene_StartFrame);
@@ -390,7 +373,6 @@ void FChaosScene::StartFrame()
 		CompletionEvents.Add(Solver->AdvanceAndDispatch_External(UseDeltaTime));
 	}
 
-#endif
 }
 
 void FChaosScene::OnSyncBodies(Chaos::FPhysicsSolverBase* Solver)
@@ -434,11 +416,9 @@ bool FChaosScene::IsCompletionEventComplete() const
 template <typename TSolver>
 void FChaosScene::SyncBodies(TSolver* Solver)
 {
-#if WITH_CHAOS
 	DECLARE_SCOPE_CYCLE_COUNTER(TEXT("SyncBodies"),STAT_SyncBodies,STATGROUP_Physics);
 	CSV_SCOPED_TIMING_STAT_EXCLUSIVE(SyncBodies);
 	OnSyncBodies(Solver);
-#endif
 }
 
 // Accumulate all the AABBTree stats
@@ -470,7 +450,6 @@ void GetAABBTreeStats(Chaos::ISpatialAccelerationCollection<Chaos::FAcceleration
 
 void FChaosScene::EndFrame()
 {
-#if WITH_CHAOS
 	using namespace Chaos;
 	using SpatialAccelerationCollection = ISpatialAccelerationCollection<FAccelerationStructureHandle,FReal,3>;
 
@@ -550,7 +529,6 @@ void FChaosScene::EndFrame()
 	}
 
 	OnPhysScenePostTick.Broadcast(this);
-#endif
 }
 
 void FChaosScene::WaitPhysScenes()
@@ -566,5 +544,3 @@ FGraphEventArray FChaosScene::GetCompletionEvents()
 {
 	return CompletionEvents;
 }
-
-#endif
\ No newline at end of file
diff --git a/Engine/Source/Runtime/PhysicsCore/Private/PhysTestSerializer.cpp b/Engine/Source/Runtime/PhysicsCore/Private/PhysTestSerializer.cpp
--- a/Engine/Source/Runtime/PhysicsCore/Private/PhysTestSerializer.cpp
+++ b/Engine/Source/Runtime/PhysicsCore/Private/PhysTestSerializer.cpp
@@ -4,13 +4,6 @@
 
 #include "PhysTestSerializer.h"
 
-
-#if PHYSICS_INTERFACE_PHYSX
-#include "PhysXIncludes.h"
-#include "PhysXSupportCore.h"
-#include "PhysXToChaosUtil.h"
-#endif
-
 #include "PhysicsCore.h"
 #include "Chaos/PBDRigidsEvolution.h"
 #include "Chaos/PBDRigidParticles.h"
@@ -93,10 +86,6 @@ void FPhysTestSerializer::Serialize(Chaos::FChaosArchive& Ar)
 
 	if (Ar.IsLoading())
 	{
-#if PHYSICS_INTERFACE_PHYSX
-		CreatePhysXData();
-#endif
-
 #if 0
 		CreateChaosData();
 #endif
@@ -131,98 +120,6 @@ void FPhysTestSerializer::SetPhysicsData(Chaos::FPBDRigidsEvolution& Evolution)
 	ArchiveVersion = Ar.GetCustomVersions();
 }
 
-
-
-#if PHYSICS_INTERFACE_PHYSX
-void FPhysTestSerializer::SetPhysicsData(physx::PxScene& Scene)
-{
-	check(AlignedDataHelper == nullptr || &Scene != AlignedDataHelper->PhysXScene);
-
-	PxSerializationRegistry* Registry = PxSerialization::createSerializationRegistry(*GPhysXSDK);
-	PxCollection* Collection = PxCollectionExt::createCollection(Scene);
-
-	PxSerialization::complete(*Collection, *Registry);
-
-	//give an ID for every object so we can find it later. This only holds for direct objects like actors and shapes
-	const uint32 NumObjects = Collection->getNbObjects();
-	TArray<PxBase*> Objects;
-	Objects.AddUninitialized(NumObjects);
-	Collection->getObjects(Objects.GetData(), NumObjects);
-	for (PxBase* Obj : Objects)
-	{
-		Collection->add(*Obj, (PxSerialObjectId)Obj);
-	}
-
-	Data.Empty();
-	FPhysXOutputStream Stream(&Data);
-	PxSerialization::serializeCollectionToBinary(Stream, *Collection, *Registry);
-	Collection->release();
-	Registry->release();
-
-	bDiskDataIsChaos = false;
-}
-
-void FPhysTestSerializer::CreatePhysXData()
-{
-	if (bDiskDataIsChaos == false)	//For the moment we don't support chaos to physx direction
-	{
-		{
-			check(Data.Num());	//no data, was the physx scene set?
-			AlignedDataHelper = MakeUnique<FPhysXSerializerData>(Data.Num());
-			FMemory::Memcpy(AlignedDataHelper->Data, Data.GetData(), Data.Num());
-		}
-
-		PxSceneDesc Desc = CreateDummyPhysXSceneDescriptor();	//question: does it matter that this is default and not the one set by user settings?
-		AlignedDataHelper->PhysXScene = GPhysXSDK->createScene(Desc);
-
-		AlignedDataHelper->Registry = PxSerialization::createSerializationRegistry(*GPhysXSDK);
-		AlignedDataHelper->Collection = PxSerialization::createCollectionFromBinary(AlignedDataHelper->Data, *AlignedDataHelper->Registry);
-		AlignedDataHelper->PhysXScene->addCollection(*AlignedDataHelper->Collection);
-	}
-}
-
-physx::PxBase* FPhysTestSerializer::FindObject(uint64 Id)
-{
-	if (!AlignedDataHelper)
-	{
-		CreatePhysXData();
-	}
-
-	physx::PxBase* Ret = AlignedDataHelper->Collection->find(Id);
-	ensure(Ret);
-#if 0
-		CreateChaosData();
-#endif
-	return Ret;
-}
-
-FPhysTestSerializer::FPhysXSerializerData::~FPhysXSerializerData()
-{
-	if (PhysXScene)
-	{
-		//release all resources the collection created (calling release on the collection is not enough)
-		const uint32 NumObjects = Collection->getNbObjects();
-		TArray<PxBase*> Objects;
-		Objects.AddUninitialized(NumObjects);
-		Collection->getObjects(Objects.GetData(), NumObjects);
-		for (PxBase* Obj : Objects)
-		{
-			if (Obj->isReleasable())
-			{
-				Obj->release();
-			}
-		}
-
-		Collection->release();
-		Registry->release();
-		PhysXScene->release();
-	}
-	FMemory::Free(Data);
-}
-
-
-#endif
-
 #if 0
 
 void FPhysTestSerializer::CreateChaosData()
diff --git a/Engine/Source/Runtime/PhysicsCore/Private/PhysXLibs.cpp b/Engine/Source/Runtime/PhysicsCore/Private/PhysXLibs.cpp
--- a/Engine/Source/Runtime/PhysicsCore/Private/PhysXLibs.cpp
+++ b/Engine/Source/Runtime/PhysicsCore/Private/PhysXLibs.cpp
@@ -10,8 +10,6 @@
 #include "HAL/PlatformFileManager.h"
 #include "GenericPlatform/GenericPlatformFile.h"
 
-#if WITH_PHYSX 
-
 // PhysX library imports
 namespace PhysDLLHelper
 {
@@ -25,13 +23,6 @@ namespace PhysDLLHelper
 	void* PxPvdSDKHandle = nullptr;
 	void* PhysX3CookingHandle = nullptr;
 	void* nvToolsExtHandle = nullptr;
-	#if WITH_APEX
-		void* APEXFrameworkHandle = nullptr;
-		void* APEX_LegacyHandle = nullptr;
-		#if WITH_APEX_CLOTHING
-			void* APEX_ClothingHandle = nullptr;
-		#endif  //WITH_APEX_CLOTHING
-	#endif	//WITH_APEX
 #endif
 
 #if PLATFORM_WINDOWS || PLATFORM_HOLOLENS
@@ -166,19 +157,6 @@ void* LoadPhysicsLibrary(const FString& PathEnd)
 	return Handle;
 }
 
-#if WITH_APEX
-PHYSICSCORE_API void* LoadAPEXModule(const FString& Path)
-{
-#if PLATFORM_WINDOWS || PLATFORM_HOLOLENS
-	return LoadPhysicsLibrary(RootAPEXPath + Path + APEXSuffix);
-#elif PLATFORM_MAC
-	const FString APEX_HandleLibName = FString::Printf(TEXT("%slib%s%s"), *PhysXBinariesRoot, *Path, *APEXSuffix);
-	return LoadPhysicsLibrary(APEX_HandleLibName);
-#endif
-	return nullptr;
-}
-#endif
-
 /**
  *	Load the required modules for PhysX
  */
@@ -202,15 +180,6 @@ PHYSICSCORE_API bool LoadPhysXModules(bool bLoadCookingModule)
 		PhysX3CookingHandle = LoadPhysicsLibrary(RootPhysXPath + "PhysX3Cooking" + PhysXSuffix);
 	}
 
-	#if WITH_APEX
-		APEXFrameworkHandle = LoadPhysicsLibrary(RootAPEXPath + "APEXFramework" + APEXSuffix);
-		#if WITH_APEX_LEGACY
-			APEX_LegacyHandle = LoadPhysicsLibrary(RootAPEXPath + "APEX_Legacy" + APEXSuffix);
-		#endif //WITH_APEX_LEGACY
-		#if WITH_APEX_CLOTHING
-			APEX_ClothingHandle = LoadPhysicsLibrary(RootAPEXPath + "APEX_Clothing" + APEXSuffix);
-		#endif //WITH_APEX_CLOTHING
-	#endif	//WITH_APEX
 #elif PLATFORM_MAC
 	const FString PxFoundationLibName = FString::Printf(TEXT("%slibPxFoundation%s"), *PhysXBinariesRoot, *PhysXSuffix);
 	PxFoundationHandle = LoadPhysicsLibrary(PxFoundationLibName);
@@ -229,19 +198,6 @@ PHYSICSCORE_API bool LoadPhysXModules(bool bLoadCookingModule)
 		const FString PhysX3CookinLibName = FString::Printf(TEXT("%slibPhysX3Cooking%s"), *PhysXBinariesRoot, *PhysXSuffix);
 		PhysX3CookingHandle = LoadPhysicsLibrary(PhysX3CookinLibName);
 	}
-
-	#if WITH_APEX
-		const FString APEXFrameworkLibName = FString::Printf(TEXT("%slibAPEXFramework%s"), *PhysXBinariesRoot, *APEXSuffix);
-		APEXFrameworkHandle = LoadPhysicsLibrary(APEXFrameworkLibName);
-		#if WITH_APEX_LEGACY
-			const FString APEX_LegacyHandleLibName = FString::Printf(TEXT("%slibAPEX_Legacy%s"), *PhysXBinariesRoot, *APEXSuffix);
-			APEX_LegacyHandle = LoadPhysicsLibrary(APEX_LegacyHandleLibName);
-		#endif //WITH_APEX_LEGACY
-		#if WITH_APEX_CLOTHING
-			const FString APEX_ClothingHandleLibName = FString::Printf(TEXT("%slibAPEX_Clothing%s"), *PhysXBinariesRoot, *APEXSuffix);
-			APEX_ClothingHandle = LoadPhysicsLibrary(APEX_ClothingHandleLibName);
-		#endif //WITH_APEX_CLOTHING
-	#endif	//WITH_APEX
 #endif	//PLATFORM_WINDOWS
 
 	bool bSucceeded = true;
@@ -257,15 +213,6 @@ PHYSICSCORE_API bool LoadPhysXModules(bool bLoadCookingModule)
 	// Cooking module if present
 	bSucceeded = bSucceeded && (!bLoadCookingModule || PhysX3CookingHandle);
 	// Apex if present
-#if WITH_APEX
-	bSucceeded = bSucceeded && APEXFrameworkHandle;
-#if WITH_APEX_LEGACY
-	bSucceeded = bSucceeded && APEX_LegacyHandle;
-#endif //WITH_APEX_LEGACY
-#if WITH_APEX_CLOTHING
-	bSucceeded = bSucceeded && APEX_ClothingHandle;
-#endif // WITH_APEX_CLOTHING
-#endif // WITH_APEX
 #endif // PLATFORM_WINDOWS || PLATFORM_MAC
 
 	return bSucceeded;
@@ -285,26 +232,6 @@ PHYSICSCORE_API void UnloadPhysXModules()
 	}
 	FPlatformProcess::FreeDllHandle(PhysX3CommonHandle);
 	FPlatformProcess::FreeDllHandle(PxFoundationHandle);
-	#if WITH_APEX
-		FPlatformProcess::FreeDllHandle(APEXFrameworkHandle);
-		FPlatformProcess::FreeDllHandle(APEX_LegacyHandle);
-		#if WITH_APEX_CLOTHING
-			FPlatformProcess::FreeDllHandle(APEX_ClothingHandle);
-		#endif //WITH_APEX_CLOTHING
-	#endif	//WITH_APEX
 #endif
 }
-
-#if WITH_APEX
-PHYSICSCORE_API void UnloadAPEXModule(void* Handle)
-{
-#if PLATFORM_WINDOWS || PLATFORM_HOLOLENS || PLATFORM_MAC
-	if(Handle)
-	{
-		FPlatformProcess::FreeDllHandle(Handle);
-	}
-#endif
-}
-#endif
 }
-#endif // WITH_PHYSX
\ No newline at end of file
diff --git a/Engine/Source/Runtime/PhysicsCore/Private/PhysXPublicCore.cpp b/Engine/Source/Runtime/PhysicsCore/Private/PhysXPublicCore.cpp
--- a/Engine/Source/Runtime/PhysicsCore/Private/PhysXPublicCore.cpp
+++ b/Engine/Source/Runtime/PhysicsCore/Private/PhysXPublicCore.cpp
@@ -1,112 +1 @@
 // Copyright Epic Games, Inc. All Rights Reserved.
-
-#if PHYSICS_INTERFACE_PHYSX
-
-#include "PhysXPublicCore.h"
-#include "PhysicsPublicCore.h"
-#include "IPhysXCookingModule.h"
-#include "Modules/ModuleManager.h"
-
-#if WITH_APEX
-PHYSICSCORE_API nvidia::apex::PhysX3Interface* GPhysX3Interface = nullptr;
-#endif	// #if WITH_APEX
-
-///////////////////// Unreal to PhysX conversion /////////////////////
-
-PxTransform UMatrix2PTransform(const FMatrix& UTM)
-{
-	PxQuat PQuat = U2PQuat(UTM.ToQuat());
-	PxVec3 PPos = U2PVector(UTM.GetOrigin());
-
-	PxTransform Result(PPos, PQuat);
-
-	return Result;
-}
-
-PxMat44 U2PMatrix(const FMatrix& UTM)
-{
-	PxMat44 Result;
-
-	const physx::PxMat44 *MatPtr = (const physx::PxMat44 *)(&UTM);
-	Result = *MatPtr;
-
-	return Result;
-}
-
-UCollision2PGeom::UCollision2PGeom(const FCollisionShape& CollisionShape)
-{
-	switch (CollisionShape.ShapeType)
-	{
-	case ECollisionShape::Box:
-	{
-		new (Storage)PxBoxGeometry(U2PVector(CollisionShape.GetBox()));
-		break;
-	}
-	case ECollisionShape::Sphere:
-	{
-		new (Storage)PxSphereGeometry(CollisionShape.GetSphereRadius());
-		break;
-	}
-	case ECollisionShape::Capsule:
-	{
-		new (Storage)PxCapsuleGeometry(CollisionShape.GetCapsuleRadius(), CollisionShape.GetCapsuleAxisHalfLength());
-		break;
-	}
-	default:
-		// invalid point
-		ensure(false);
-	}
-}
-
-///////////////////// PhysX to Unreal conversion /////////////////////
-
-FMatrix P2UMatrix(const PxMat44& PMat)
-{
-	FMatrix Result;
-	// we have to use Memcpy instead of typecasting, because PxMat44's are not aligned like FMatrix is
-	FMemory::Memcpy(&Result, &PMat, sizeof(PMat));
-	return Result;
-}
-
-FMatrix PTransform2UMatrix(const PxTransform& PTM)
-{
-	FQuat UQuat = P2UQuat(PTM.q);
-	FVector UPos = P2UVector(PTM.p);
-
-	FMatrix Result = FQuatRotationTranslationMatrix(UQuat, UPos);
-
-	return Result;
-}
-
-FTransform P2UTransform(const PxTransform& PTM)
-{
-	FQuat UQuat = P2UQuat(PTM.q);
-	FVector UPos = P2UVector(PTM.p);
-
-	FTransform Result = FTransform(UQuat, UPos);
-
-	return Result;
-}
-
-IPhysXCookingModule* GetPhysXCookingModule(bool bForceLoad)
-{
-	check(IsInGameThread());
-
-	if (bForceLoad)
-	{
-#if WITH_PHYSX_COOKING
-		return FModuleManager::LoadModulePtr<IPhysXCookingModule>("PhysXCooking");	//in some configurations (for example the editor) we must have physx cooking
-#else
-		return FModuleManager::LoadModulePtr<IPhysXCookingModule>("RuntimePhysXCooking");	//in some configurations (mobile) we can choose to opt in for physx cooking via plugin
-#endif
-	}
-	else
-	{
-#if WITH_PHYSX_COOKING
-		return FModuleManager::GetModulePtr<IPhysXCookingModule>("PhysXCooking");	//in some configurations (for example the editor) we must have physx cooking
-#else
-		return FModuleManager::GetModulePtr<IPhysXCookingModule>("RuntimePhysXCooking");	//in some configurations (mobile) we can choose to opt in for physx cooking via plugin
-#endif
-	}
-}
-#endif // WITH_PHYSX
diff --git a/Engine/Source/Runtime/PhysicsCore/Private/PhysXSupportCore.cpp b/Engine/Source/Runtime/PhysicsCore/Private/PhysXSupportCore.cpp
--- a/Engine/Source/Runtime/PhysicsCore/Private/PhysXSupportCore.cpp
+++ b/Engine/Source/Runtime/PhysicsCore/Private/PhysXSupportCore.cpp
@@ -1,353 +1,3 @@
 // Copyright Epic Games, Inc. All Rights Reserved.
 
 #include "PhysXSupportCore.h"
-
-#if PHYSICS_INTERFACE_PHYSX
-#include "PhysicsPublicCore.h"
-#include "PhysXSupportCore.h"
-#include "HAL/IConsoleManager.h"
-
-PxFoundation* GPhysXFoundation = nullptr;
-FPhysXAllocator* GPhysXAllocator = nullptr;
-PxPvd* GPhysXVisualDebugger = nullptr;
-physx::PxPhysics* GPhysXSDK = nullptr;
-
-TArray<PxMaterial*> GPhysXPendingKillMaterial;
-
-#if WITH_APEX
-PHYSICSCORE_API apex::ApexSDK* GApexSDK = nullptr;
-
-#if WITH_APEX_LEGACY
-PHYSICSCORE_API apex::Module* GApexModuleLegacy = nullptr;
-#endif
-
-#if WITH_APEX_CLOTHING
-PHYSICSCORE_API apex::ModuleClothing* GApexModuleClothing = nullptr;
-#endif
-#endif
-
-#if WITH_APEX
-FApexNullRenderResourceManager		GApexNullRenderResourceManager;
-FApexResourceCallback				GApexResourceCallback;
-#endif	// #if WITH_APEX
-
-
-int32 GPhysXHackLoopCounter = -1;
-FAutoConsoleVariableRef CVarHackLoopCounter(TEXT("p.TriMeshBufferOverflowCounter"), GPhysXHackLoopCounter, TEXT("Loop logging counter - set to -1 to disable logging"), ECVF_Default);
-
-
-/** Util to convert PhysX error code to string */
-FString ErrorCodeToString(PxErrorCode::Enum e)
-{
-	FString CodeString;
-
-	switch (e)
-	{
-	case PxErrorCode::eNO_ERROR:
-		CodeString = TEXT("eNO_ERROR");
-		break;
-	case PxErrorCode::eDEBUG_INFO:
-		CodeString = TEXT("eDEBUG_INFO");
-		break;
-	case PxErrorCode::eDEBUG_WARNING:
-		CodeString = TEXT("eDEBUG_WARNING");
-		break;
-	case PxErrorCode::eINVALID_PARAMETER:
-		CodeString = TEXT("eINVALID_PARAMETER");
-		break;
-	case PxErrorCode::eINVALID_OPERATION:
-		CodeString = TEXT("eINVALID_OPERATION");
-		break;
-	case PxErrorCode::eOUT_OF_MEMORY:
-		CodeString = TEXT("eOUT_OF_MEMORY");
-		break;
-	case PxErrorCode::eINTERNAL_ERROR:
-		CodeString = TEXT("eINTERNAL_ERROR");
-		break;
-	case PxErrorCode::eABORT:
-		CodeString = TEXT("eABORT");
-		break;
-	case PxErrorCode::ePERF_WARNING:
-		CodeString = TEXT("ePERF_WARNING");
-		break;
-	case PxErrorCode::eLOGGING_INFO:
-		CodeString = TEXT("eLOGGING_INFO");
-		break;
-	default:
-		CodeString = TEXT("UNKONWN");
-	}
-
-	return CodeString;
-}
-
-
-int32 GPhysXHackCurrentLoopCounter = 0;
-
-void FPhysXErrorCallback::reportError(PxErrorCode::Enum e, const char* message, const char* file, int line)
-{
-	// if not in game, ignore Perf warnings - i.e. Moving Static actor in editor will produce this warning
-	if (GIsEditor && e == PxErrorCode::ePERF_WARNING)
-	{
-		return;
-	}
-
-	if (e == PxErrorCode::eLOGGING_INFO)
-	{
-		if (GPhysXHackLoopCounter == -1)
-		{
-			return;
-		}
-		GPhysXHackCurrentLoopCounter++;
-		if (GPhysXHackCurrentLoopCounter <= GPhysXHackLoopCounter)
-		{
-			return;
-		}
-	}
-
-	if (e == PxErrorCode::eINTERNAL_ERROR)
-	{
-		const char* HillClimbError = "HillClimbing";
-		const char* TestSATCapsulePoly = "testSATCapsulePoly";
-		const char* MeshCleanFailed = "cleaning the mesh failed";
-
-		// HACK: Internal errors which we want to suppress in release builds should be changed to debug warning error codes.
-		// This way we see them in debug but not in production.
-		if (FPlatformString::Strstr(message, MeshCleanFailed))
-		{
-			e = PxErrorCode::eDEBUG_WARNING;
-		}
-	}
-	// Make string to print out, include physx file/line
-	FString ErrorString = FString::Printf(TEXT("PHYSX: (%s %d) %s : %s"), ANSI_TO_TCHAR(file), line, *ErrorCodeToString(e), ANSI_TO_TCHAR(message));
-
-	if (e == PxErrorCode::eOUT_OF_MEMORY || e == PxErrorCode::eABORT)
-	{
-		UE_LOG(LogPhysicsCore, Error, TEXT("%s"), *ErrorString);
-		//ensureMsgf(false, TEXT("%s"), *ErrorString);
-	}
-	else if (e == PxErrorCode::eINVALID_PARAMETER || e == PxErrorCode::eINVALID_OPERATION)
-	{
-		UE_LOG(LogPhysicsCore, Error, TEXT("%s"), *ErrorString);
-		//ensureMsgf(false, TEXT("%s"), *ErrorString);
-	}
-	else if (e == PxErrorCode::ePERF_WARNING || e == PxErrorCode::eINTERNAL_ERROR || e == PxErrorCode::eLOGGING_INFO)
-	{
-		UE_LOG(LogPhysicsCore, Warning, TEXT("%s"), *ErrorString);
-	}
-#if UE_BUILD_DEBUG
-	else if (e == PxErrorCode::eDEBUG_WARNING)
-	{
-		UE_LOG(LogPhysicsCore, Warning, TEXT("%s"), *ErrorString);
-	}
-#endif
-	else
-	{
-		UE_LOG(LogPhysicsCore, Log, TEXT("%s"), *ErrorString);
-	}
-}
-
-FPhysxSharedData* FPhysxSharedData::Singleton = nullptr;
-
-void FPhysxSharedData::Initialize()
-{
-	check(Singleton == nullptr);
-	Singleton = new FPhysxSharedData();
-}
-
-void FPhysxSharedData::Terminate()
-{
-	if (Singleton)
-	{
-		delete Singleton;
-		Singleton = nullptr;
-	}
-}
-
-void FPhysxSharedData::LockAccess()
-{
-	if (Singleton)
-	{
-		Singleton->CriticalSection.Lock();
-	}
-}
-
-void FPhysxSharedData::UnlockAccess()
-{
-	if (Singleton)
-	{
-		Singleton->CriticalSection.Unlock();
-	}
-}
-
-void FPhysxSharedData::Add( PxBase* Obj, const FString& OwnerName )
-{
-	if(Obj) 
-	{ 
-		SharedObjects->add(*Obj, (PxSerialObjectId)Obj);
-#if !(UE_BUILD_SHIPPING || UE_BUILD_TEST)
-		OwnerNames.Add(Obj, OwnerName);
-#endif
-	}
-}
-
-void FPhysxSharedData::Remove(PxBase* Obj) 
-{ 
-	// Check for containment first due to multiple UBodySetups sharing the same ref counted object causing harmless double-frees
-	if (Obj && SharedObjects->contains(*Obj)) 
-	{ 
-		SharedObjects->remove(*Obj); 
-		OwnerNames.Remove(Obj); 
-	} 
-}	
-
-SIZE_T GetPhysxObjectSize(PxBase* Obj, const PxCollection* SharedCollection)
-{
-	PxSerializationRegistry* Sr = PxSerialization::createSerializationRegistry(*GPhysXSDK);
-	PxCollection* Collection = PxCreateCollection();
-
-	Collection->add(*Obj);
-	PxSerialization::complete(*Collection, *Sr, SharedCollection);	// chase all other stuff (shared shaps, materials, etc) needed to serialize this collection
-
-	FPhysXCountMemoryStream Out;
-	PxSerialization::serializeCollectionToBinary(Out, *Collection, *Sr, SharedCollection);
-
-	Collection->release();
-	Sr->release();
-
-	return Out.UsedMemory;
-}
-
-struct FSharedResourceEntry
-{
-	uint64 MemorySize;
-	uint64 Count;
-};
-
-void HelperCollectUsage(const TMap<FString, TArray<PxBase*> >& ObjectsByType, TMap<FString, FSharedResourceEntry>& AllocationsByType, uint64& OverallSize, int32& OverallCount)
-{
-	TArray<FString> TypeNames;
-	ObjectsByType.GetKeys(TypeNames);
-
-	for (int32 TypeIdx = 0; TypeIdx < TypeNames.Num(); ++TypeIdx)
-	{
-		const FString& TypeName = TypeNames[TypeIdx];
-
-		const TArray<PxBase*>* ObjectsArray = ObjectsByType.Find(TypeName);
-		check(ObjectsArray);
-
-		PxSerializationRegistry* Sr = PxSerialization::createSerializationRegistry(*GPhysXSDK);
-		PxCollection* Collection = PxCreateCollection();
-
-		for (int32 i = 0; i < ObjectsArray->Num(); ++i)
-		{
-			Collection->add(*((*ObjectsArray)[i]));;
-		}
-
-		PxSerialization::complete(*Collection, *Sr);	// chase all other stuff (shared shaps, materials, etc) needed to serialize this collection
-
-		FPhysXCountMemoryStream Out;
-		PxSerialization::serializeCollectionToBinary(Out, *Collection, *Sr);
-
-		Collection->release();
-		Sr->release();
-
-		OverallSize += Out.UsedMemory;
-		OverallCount += ObjectsArray->Num();
-
-		FSharedResourceEntry NewEntry;
-		NewEntry.Count = ObjectsArray->Num();
-		NewEntry.MemorySize = Out.UsedMemory;
-
-		AllocationsByType.Add(TypeName, NewEntry);
-	}
-}
-
-void FPhysxSharedData::DumpSharedMemoryUsage(FOutputDevice* Ar)
-{
-	
-
-	struct FSortBySize
-	{
-		FORCEINLINE bool operator()( const FSharedResourceEntry& A, const FSharedResourceEntry& B ) const 
-		{ 
-			// Sort descending
-			return B.MemorySize < A.MemorySize;
-		}
-	};
-
-	TMap<FString, FSharedResourceEntry> AllocationsByType;
-	TMap<FString, FSharedResourceEntry> AllocationsByObject;
-
-	uint64 OverallSize = 0;
-	int32 OverallCount = 0;
-
-	TMap<FString, TArray<PxBase*> > ObjectsByType;
-	TMap<FString, TArray<PxBase*> > ObjectsByObjectName;	//array is just there for code reuse, should be a single object
-
-	for (int32 i=0; i < (int32)SharedObjects->getNbObjects(); ++i)
-	{
-		PxBase& Obj = SharedObjects->getObject(i);
-		FString TypeName = ANSI_TO_TCHAR(Obj.getConcreteTypeName());
-
-		TArray<PxBase*>* ObjectsArray = ObjectsByType.Find(TypeName);
-		if (ObjectsArray == NULL)
-		{
-			ObjectsByType.Add(TypeName, TArray<PxBase*>());
-			ObjectsArray = ObjectsByType.Find(TypeName);
-		}
-
-		check(ObjectsArray);
-		ObjectsArray->Add(&Obj);
-
-		if (const FString* OwnerName = OwnerNames.Find(&Obj))
-		{
-			TArray<PxBase*> Objs;
-			Objs.Add(&Obj);
-			ObjectsByObjectName.Add(*OwnerName, Objs);
-		}
-	}
-
-	HelperCollectUsage(ObjectsByType, AllocationsByType, OverallSize, OverallCount);
-
-	uint64 IgnoreSize;
-	int32 IgnoreCount;
-	HelperCollectUsage(ObjectsByObjectName, AllocationsByObject, IgnoreSize, IgnoreCount);
-
-
-	Ar->Logf(TEXT(""));
-	Ar->Logf(TEXT("Shared Resources:"));
-	Ar->Logf(TEXT(""));
-
-	AllocationsByType.ValueSort(FSortBySize());
-	AllocationsByObject.ValueSort(FSortBySize());
-	
-	Ar->Logf(TEXT("%-10d %s (%d)"), OverallSize, TEXT("Overall"), OverallCount );
-	
-	for( auto It=AllocationsByType.CreateConstIterator(); It; ++It )
-	{
-		Ar->Logf(TEXT("%-10d %s (%d)"), It.Value().MemorySize, *It.Key(), It.Value().Count );
-	}
-
-	Ar->Logf(TEXT("Detailed:"));
-
-	for (auto It = AllocationsByObject.CreateConstIterator(); It; ++It)
-	{
-		Ar->Logf(TEXT("%-10d %s (%d)"), It.Value().MemorySize, *It.Key(), It.Value().Count);
-	}
-}
-
-void PvdConnect(FString Host, bool bVisualization)
-{
-	int32	Port = 5425;         // TCP port to connect to, where PVD is listening
-	uint32	Timeout = 100;          // timeout in milliseconds to wait for PVD to respond, consoles and remote PCs need a higher timeout.
-
-	PxPvdInstrumentationFlags ConnectionFlags = bVisualization ? PxPvdInstrumentationFlag::eALL : (PxPvdInstrumentationFlag::ePROFILE | PxPvdInstrumentationFlag::eMEMORY);
-
-	PxPvdTransport* transport = PxDefaultPvdSocketTransportCreate(TCHAR_TO_ANSI(*Host), Port, Timeout);
-	GPhysXVisualDebugger->disconnect();	//make sure we're disconnected first
-	GPhysXVisualDebugger->connect(*transport, ConnectionFlags);
-
-	// per scene properties (e.g. PxPvdSceneFlag::eTRANSMIT_CONSTRAINTS) are 
-	// set on the PxPvdSceneClient in PhysScene.cpp, FPhysScene::InitPhysScene
-}
-#endif // WITH_PHYSX
diff --git a/Engine/Source/Runtime/PhysicsCore/Private/PhysicalMaterial.cpp b/Engine/Source/Runtime/PhysicsCore/Private/PhysicalMaterial.cpp
--- a/Engine/Source/Runtime/PhysicsCore/Private/PhysicalMaterial.cpp
+++ b/Engine/Source/Runtime/PhysicsCore/Private/PhysicalMaterial.cpp
@@ -7,10 +7,7 @@
 #include "PhysicalMaterials/PhysicalMaterial.h"
 #include "PhysicalMaterials/PhysicalMaterialPropertyBase.h"
 #include "UObject/UObjectIterator.h"
-
-#if WITH_CHAOS
-	#include "Chaos/PhysicalMaterials.h"
-#endif
+#include "Chaos/PhysicalMaterials.h"
 
 UDEPRECATED_PhysicalMaterialPropertyBase::UDEPRECATED_PhysicalMaterialPropertyBase(const FObjectInitializer& ObjectInitializer)
 	: Super(ObjectInitializer)
