Landscape fix : replaced usage of TSet<TWeakPtr<...>> by TArray<TWeakPtr<...>> because the former will have invalid results until re-hashing after pointers expire.

diff --git a/Engine/Source/Runtime/Landscape/Classes/Landscape.h b/Engine/Source/Runtime/Landscape/Classes/Landscape.h
--- a/Engine/Source/Runtime/Landscape/Classes/Landscape.h
+++ b/Engine/Source/Runtime/Landscape/Classes/Landscape.h
@@ -385,7 +385,7 @@ private:
 	bool PrepareLayersTextureResources(bool bInWaitForStreaming);
 	bool PrepareLayersTextureResources(const TArray<FLandscapeLayer>& InLayers, bool bInWaitForStreaming);
 	bool PrepareLayersBrushResources(bool bInWaitForStreaming);
-	void InvalidateRVTForTextures(const TSet<TWeakObjectPtr<UTexture2D>>& InTextures);
+	void InvalidateRVTForTextures(const TSet<UTexture2D*>& InTextures);
 	void PrepareLayersHeightmapsLocalMergeRenderThreadData(const FUpdateLayersContentContext& InUpdateLayersContentContext, const FEditLayersHeightmapMergeParams& InMergeParams, EditLayersHeightmapLocalMerge_RenderThread::FMergeInfo& OutRenderThreadData);
 	void PrepareLayersWeightmapsLocalMergeRenderThreadData(const FUpdateLayersContentContext& InUpdateLayersContentContext, const FEditLayersWeightmapMergeParams& InMergeParams, EditLayersWeightmapLocalMerge_RenderThread::FMergeInfo& OutRenderThreadData);
 
@@ -480,7 +480,7 @@ public:
 
 	/** List of textures that are not fully streamed in yet (updated every frame to track textures that have finished streaming in) */
 	UPROPERTY(Transient, DuplicateTransient, TextExportTransient)	
-	TSet<TWeakObjectPtr<UTexture2D>> TrackedStreamingInTextures;
+	TArray<TWeakObjectPtr<UTexture2D>> TrackedStreamingInTextures;
 
 private:
 	FLandscapeBlueprintBrushChangedDelegate LandscapeBlueprintBrushChangedDelegate;
diff --git a/Engine/Source/Runtime/Landscape/Private/LandscapeEditLayers.cpp b/Engine/Source/Runtime/Landscape/Private/LandscapeEditLayers.cpp
--- a/Engine/Source/Runtime/Landscape/Private/LandscapeEditLayers.cpp
+++ b/Engine/Source/Runtime/Landscape/Private/LandscapeEditLayers.cpp
@@ -42,6 +42,7 @@ LandscapeEditLayers.cpp: Landscape editing layers mode
 #include "Algo/Count.h"
 #include "Algo/AnyOf.h"
 #include "Algo/Unique.h"
+#include "Algo/Transform.h"
 #include "LandscapeSettings.h"
 #include "LandscapeRender.h"
 #include "LandscapeInfoMap.h"
@@ -3350,8 +3351,18 @@ bool ALandscape::PrepareTextureResources(bool bInWaitForStreaming)
 	{
 		return false;
 	}
-
-	TSet<TWeakObjectPtr<UTexture2D>> StreamingInTexturesBefore(MoveTemp(TrackedStreamingInTextures));
+	
+	// Only keep the textures that are still valid:
+	TSet<UTexture2D*> StreamingInTexturesBefore;
+	StreamingInTexturesBefore.Reserve(TrackedStreamingInTextures.Num());	
+	Algo::TransformIf(TrackedStreamingInTextures, StreamingInTexturesBefore,
+		[](const TWeakObjectPtr<UTexture2D>& Texture) { return Texture.IsValid(); },
+		[](const TWeakObjectPtr<UTexture2D>& Texture) { return Texture.Get(); });
+	TrackedStreamingInTextures.Empty();
+
+	// Start tracking the new textures :
+	TSet<UTexture2D*> StreamingInTexturesAfter;
+	StreamingInTexturesAfter.Reserve(TrackedStreamingInTextures.Num());
 
 	bool bIsReady = true;
 	Info->ForAllLandscapeProxies([&](ALandscapeProxy* Proxy)
@@ -3364,7 +3375,7 @@ bool ALandscape::PrepareTextureResources(bool bInWaitForStreaming)
 				bool bIsTextureReady = IsTextureReady(ComponentHeightmap, bInWaitForStreaming);
 				if (!bIsTextureReady)
 				{
-					TrackedStreamingInTextures.Add(ComponentHeightmap);
+					StreamingInTexturesAfter.Add(ComponentHeightmap);
 				}
 				bIsReady &= bIsTextureReady;
 			}
@@ -3375,7 +3386,7 @@ bool ALandscape::PrepareTextureResources(bool bInWaitForStreaming)
 				// If the texture is not ready, start tracking its state changes to be notified when it's fully streamed in : 
 				if (!bIsTextureReady)
 				{
-					TrackedStreamingInTextures.Add(ComponentWeightmap);
+					StreamingInTexturesAfter.Add(ComponentWeightmap);
 				}
 				bIsReady &= bIsTextureReady;
 			}
@@ -3383,15 +3394,19 @@ bool ALandscape::PrepareTextureResources(bool bInWaitForStreaming)
 	});
 
 	// The assets that were streaming in before and are not anymore can be considered streamed in: 
-	TSet<TWeakObjectPtr<UTexture2D>> StreamedInTextures = StreamingInTexturesBefore.Difference(TrackedStreamingInTextures);
+	TRACE_CPUPROFILER_EVENT_SCOPE(LandscapeLayers_Difference);
+	TSet<UTexture2D*> StreamedInTextures = StreamingInTexturesBefore.Difference(StreamingInTexturesAfter);
 	InvalidateRVTForTextures(StreamedInTextures);
 
+	// Store as a list of TWeakObjectPtr<UTexture2D> so as not to keep references on the tracked textures :
+	Algo::Transform(StreamingInTexturesAfter, TrackedStreamingInTextures, [](UTexture2D* Texture) { return TWeakObjectPtr<UTexture2D>(Texture); });
+
 	return bIsReady;
 }
 
 // Note: this approach is generic, because FObjectCacheContextScope is a fast texture->material interface->primitive component lookup. 
 // If FObjectCacheContextScope was available at runtime, it could become an efficient way to automatically invalidate RVT areas corresponding to primitive components that use textures that are being streamed in:
-void ALandscape::InvalidateRVTForTextures(const TSet<TWeakObjectPtr<UTexture2D>>& InTextures)
+void ALandscape::InvalidateRVTForTextures(const TSet<UTexture2D*>& InTextures)
 {
 #if WITH_EDITOR
 	TRACE_CPUPROFILER_EVENT_SCOPE(ALandscape_InvalidateRVTForTextures);
@@ -3402,9 +3417,9 @@ void ALandscape::InvalidateRVTForTextures(const TSet<TWeakObjectPtr<UTexture2D>>
 		FObjectCacheContextScope ObjectCacheScope;
 		TSet<UPrimitiveComponent*> PrimitiveComponentsToInvalidate;
 
-		for (TWeakObjectPtr<UTexture2D> TexturePtr : InTextures)
+		for (UTexture2D* Texture : InTextures)
 		{
-			if (UTexture2D* Texture = TexturePtr.Get())
+			if (Texture != nullptr)
 			{
 				// First, find all the materials referencing this texture that are writing to the RVT in order to invalidate the primitive components referencing them when the texture 
 				//  gets fully streamed in so that we're not left with low-res mips being rendered in the RVT tiles : 
