Work in progress Nanite support for landscape

diff --git a/Engine/Source/Runtime/Landscape/Classes/Landscape.h b/Engine/Source/Runtime/Landscape/Classes/Landscape.h
--- a/Engine/Source/Runtime/Landscape/Classes/Landscape.h
+++ b/Engine/Source/Runtime/Landscape/Classes/Landscape.h
@@ -268,6 +268,7 @@ public:
 	bool IsValidRenderTargetFormatHeightmap(EPixelFormat InRenderTargetFormat, bool& bOutCompressHeight);
 
 #if WITH_EDITOR
+	LANDSCAPE_API virtual bool IsNaniteEnabled() const override;
 	/** Computes & returns bounds containing all landscape proxies (if any) or this landscape's bounds otherwise. Note that in non-WP worlds this will call GetLoadedBounds(). */
 	LANDSCAPE_API FBox GetCompleteBounds() const;
 	LANDSCAPE_API void RegisterLandscapeEdMode(ILandscapeEdModeInterface* InLandscapeEdMode) { LandscapeEdMode = InLandscapeEdMode; }
@@ -442,6 +443,10 @@ private:
 public:
 
 #if WITH_EDITORONLY_DATA
+	/** Use Nanite to render landscape as a mesh on supported platforms. */
+	UPROPERTY(EditAnywhere, BlueprintReadOnly, Category = Nanite, meta = (DisplayName = "Enable Nanite"))
+	bool bEnableNanite = false;
+
 	/** Landscape actor has authority on default streaming behavior for new actors : LandscapeStreamingProxies & LandscapeSplineActors */
 	UPROPERTY(EditAnywhere, Category = WorldPartition)
 	bool bAreNewLandscapeActorsSpatiallyLoaded = true;
diff --git a/Engine/Source/Runtime/Landscape/Classes/LandscapeComponent.h b/Engine/Source/Runtime/Landscape/Classes/LandscapeComponent.h
--- a/Engine/Source/Runtime/Landscape/Classes/LandscapeComponent.h
+++ b/Engine/Source/Runtime/Landscape/Classes/LandscapeComponent.h
@@ -890,11 +890,11 @@ public:
 
 	// Generates mobile platform data for this component
 	void GenerateMobileWeightmapLayerAllocations();
-	void GeneratePlatformVertexData(const ITargetPlatform* TargetPlatform);
-	void GeneratePlatformPixelData(bool bIsCooking, const ITargetPlatform* TargetPlatform);
+	void GenerateMobilePlatformVertexData(const ITargetPlatform* TargetPlatform);
+	void GenerateMobilePlatformPixelData(bool bIsCooking, const ITargetPlatform* TargetPlatform);
 
 	/** Generate mobile data if it's missing or outdated */
-	void CheckGenerateLandscapePlatformData(bool bIsCooking, const ITargetPlatform* TargetPlatform);
+	void CheckGenerateMobilePlatformData(bool bIsCooking, const ITargetPlatform* TargetPlatform);
 
 	virtual TSubclassOf<class UHLODBuilder> GetCustomHLODBuilderClass() const override;
 #endif
diff --git a/Engine/Source/Runtime/Landscape/Classes/LandscapeProxy.h b/Engine/Source/Runtime/Landscape/Classes/LandscapeProxy.h
--- a/Engine/Source/Runtime/Landscape/Classes/LandscapeProxy.h
+++ b/Engine/Source/Runtime/Landscape/Classes/LandscapeProxy.h
@@ -13,6 +13,7 @@
 #include "Engine/Texture.h"
 #include "PerPlatformProperties.h"
 #include "LandscapeComponent.h"
+#include "LandscapeNaniteComponent.h"
 #include "LandscapeWeightmapUsage.h"
 #include "VT/RuntimeVirtualTextureEnum.h"
 #include "ActorPartition/PartitionActor.h"
@@ -27,6 +28,7 @@
 class ALandscape;
 class ALandscapeProxy;
 class UHierarchicalInstancedStaticMeshComponent;
+class ULandscapeNaniteComponent;
 class ULandscapeComponent;
 class ULandscapeGrassType;
 class ULandscapeHeightfieldCollisionComponent;
@@ -549,6 +551,9 @@ public:
 	UPROPERTY(transient, duplicatetransient)
 	TArray<TObjectPtr<UHierarchicalInstancedStaticMeshComponent>> FoliageComponents;
 
+	UPROPERTY()
+	TObjectPtr<ULandscapeNaniteComponent> NaniteComponent;
+
 	/** A transient data structure for tracking the grass */
 	FCachedLandscapeFoliage FoliageCache;
 	/** A transient data structure for tracking the grass tasks*/
@@ -823,6 +828,7 @@ public:
 	virtual TUniquePtr<class FWorldPartitionActorDesc> CreateClassActorDesc() const override;
 	virtual bool EditorCanAttachTo(const AActor* InParent, FText& OutReason) const override { return false; }
 	virtual bool GetReferencedContentObjects(TArray<UObject*>& Objects) const override;
+	virtual bool IsNaniteEnabled() const PURE_VIRTUAL(IsNaniteEnabled, return false;)
 #endif	//WITH_EDITOR
 
 	virtual FGuid GetLandscapeGuid() const override { return LandscapeGuid; }
@@ -920,6 +926,8 @@ public:
 
 	/** Handle so we can unregister the delegate */
 	FDelegateHandle FeatureLevelChangedDelegateHandle;
+
+	FGuid GetNaniteContentId() const;
 #endif
 
 	//~ Begin UObject Interface.
@@ -1058,19 +1066,37 @@ public:
 	 */
 	LANDSCAPE_API bool ExportToRawMesh(int32 InExportLOD, FMeshDescription& OutRawMesh) const;
 
-
 	/**
 	* Exports landscape geometry contained within InBounds into a raw mesh
 	*
-	* @param InExportLOD Landscape LOD level to use while exporting, INDEX_NONE will use ALanscapeProxy::ExportLOD settings
+	* @param InExportLOD - Landscape LOD level to use while exporting, INDEX_NONE will use ALanscapeProxy::ExportLOD settings
 	* @param OutRawMesh - Resulting raw mesh
 	* @param InBounds - Box/Sphere bounds which limit the geometry exported out into OutRawMesh
+	* @param bIgnoreBounds - If false, InBounds will be ignored during export
 	* @return true if successful
 	*/
 	LANDSCAPE_API bool ExportToRawMesh(int32 InExportLOD, FMeshDescription& OutRawMesh, const FBoxSphereBounds& InBounds, bool bIgnoreBounds = false) const;
 
-	/** Generate platform data if it's missing or outdated */
+	/**
+	* Exports landscape geometry contained within InBounds into a raw mesh
+	*
+	* @param InComponents - Specific landscape component(s) to export
+	* @param InExportLOD - Landscape LOD level to use while exporting, INDEX_NONE will use ALanscapeProxy::ExportLOD settings
+	* @param OutRawMesh - Resulting raw mesh
+	* @param InBounds - Box/Sphere bounds which limit the geometry exported out into OutRawMesh
+	* @param bIgnoreBounds - If false, InBounds will be ignored during export
+	* @return true if successful
+	*/
+	LANDSCAPE_API bool ExportToRawMesh(const TArrayView<ULandscapeComponent*>& InComponents, int32 InExportLOD, FMeshDescription& OutRawMesh, const FBoxSphereBounds& InBounds, bool bIgnoreBounds = false) const;
+
+	UE_DEPRECATED(5.1, "CheckGenerateLandscapePlatformData has been deprecated, please use CheckGenerateMobilePlatformData instead.")
 	LANDSCAPE_API void CheckGenerateLandscapePlatformData(bool bIsCooking, const ITargetPlatform* TargetPlatform);
+
+	/** Generate mobile platform data if it's missing or outdated */
+	LANDSCAPE_API void CheckGenerateMobilePlatformData(bool bIsCooking, const ITargetPlatform* TargetPlatform);
+
+	/** Generate Nanite platform data if it's missing or outdated */
+	LANDSCAPE_API void CheckGenerateNanitePlatformData(bool bIsCooking, const ITargetPlatform* TargetPlatform);
 	
 	/** @return Current size of bounding rectangle in quads space */
 	LANDSCAPE_API FIntRect GetBoundingRect() const;
diff --git a/Engine/Source/Runtime/Landscape/Classes/LandscapeStreamingProxy.h b/Engine/Source/Runtime/Landscape/Classes/LandscapeStreamingProxy.h
--- a/Engine/Source/Runtime/Landscape/Classes/LandscapeStreamingProxy.h
+++ b/Engine/Source/Runtime/Landscape/Classes/LandscapeStreamingProxy.h
@@ -46,6 +46,7 @@
 #if WITH_EDITOR
 	virtual UMaterialInterface* GetLandscapeMaterial(int8 InLODIndex = INDEX_NONE) const override;
 	virtual UMaterialInterface* GetLandscapeHoleMaterial() const override;
+	virtual bool IsNaniteEnabled() const override;
 #endif
 	//~ End ALandscapeBase Interface
 
diff --git a/Engine/Source/Runtime/Landscape/Landscape.Build.cs b/Engine/Source/Runtime/Landscape/Landscape.Build.cs
--- a/Engine/Source/Runtime/Landscape/Landscape.Build.cs
+++ b/Engine/Source/Runtime/Landscape/Landscape.Build.cs
@@ -63,10 +63,17 @@ public class Landscape : ModuleRules
 					"MaterialUtilities",
 					"SlateCore",
 					"Slate",
-					"GeometryCore"
+					"GeometryCore",
+					"MeshUtilities",
+					"MeshUtilitiesCommon",
+					"MeshBuilderCommon",
+					"MeshBuilder",
 				}
 			);
 
+			DynamicallyLoadedModuleNames.Add("NaniteBuilder");
+			PrivateIncludePathModuleNames.Add("NaniteBuilder");
+
 			CircularlyReferencedDependentModules.AddRange(
 				new string[] {
 					"UnrealEd",
diff --git a/Engine/Source/Runtime/Landscape/Private/Landscape.cpp b/Engine/Source/Runtime/Landscape/Private/Landscape.cpp
--- a/Engine/Source/Runtime/Landscape/Private/Landscape.cpp
+++ b/Engine/Source/Runtime/Landscape/Private/Landscape.cpp
@@ -25,6 +25,7 @@ Landscape.cpp: Terrain rendering
 #include "LandscapeInfoMap.h"
 #include "EditorSupportDelegates.h"
 #include "LandscapeMeshProxyComponent.h"
+#include "LandscapeNaniteComponent.h"
 #include "LandscapeRender.h"
 #include "LandscapeRenderMobile.h"
 #include "Logging/TokenizedMessage.h"
@@ -136,9 +137,6 @@ namespace LandscapeCookStats
 }
 #endif
 
-// Set this to 0 to disable landscape cooking and thus disable it on device.
-#define ENABLE_LANDSCAPE_COOKING 1
-
 static bool UseMobileLandscapeMesh(const ITargetPlatform* TargetPlatform)
 {
 	return TargetPlatform->SupportsFeature(ETargetPlatformFeatures::MobileLandscapeMesh);
@@ -266,30 +264,134 @@ UMaterialInstanceDynamic* ULandscapeComponent::GetMaterialInstanceDynamic(int32
 	return nullptr;
 }
 
-
 #if WITH_EDITOR
+
 void ULandscapeComponent::BeginCacheForCookedPlatformData(const ITargetPlatform* TargetPlatform)
 {
 	Super::BeginCacheForCookedPlatformData(TargetPlatform);
 
-	if (UseMobileLandscapeMesh(TargetPlatform) && !HasAnyFlags(RF_ClassDefaultObject))
+	if (!HasAnyFlags(RF_ClassDefaultObject))
 	{
-		CheckGenerateLandscapePlatformData(true, TargetPlatform);
+		if (UseMobileLandscapeMesh(TargetPlatform))
+		{
+			CheckGenerateMobilePlatformData(/*bIsCooking = */ true, TargetPlatform);
+		}
 	}
 }
 
+// Deprecated, use CheckGenerateMobilePlatformData
 void ALandscapeProxy::CheckGenerateLandscapePlatformData(bool bIsCooking, const ITargetPlatform* TargetPlatform)
+{
+	return CheckGenerateMobilePlatformData(bIsCooking, TargetPlatform);
+}
+
+void ALandscapeProxy::CheckGenerateMobilePlatformData(bool bIsCooking, const ITargetPlatform* TargetPlatform)
 {
 	for (ULandscapeComponent* Component : LandscapeComponents)
 	{
-		Component->CheckGenerateLandscapePlatformData(bIsCooking, TargetPlatform);
+		Component->CheckGenerateMobilePlatformData(bIsCooking, TargetPlatform);
+	}
+}
+
+void ALandscapeProxy::CheckGenerateNanitePlatformData(bool bIsCooking, const ITargetPlatform* TargetPlatform)
+{
+	if (IsNaniteEnabled())
+	{
+		FGuid NaniteContentId = GetNaniteContentId();
+		if (NaniteComponent == nullptr)
+		{
+			NaniteComponent = NewObject<ULandscapeNaniteComponent>(this);
+
+			NaniteComponent->SetCollisionProfileName(UCollisionProfile::NoCollision_ProfileName);
+			NaniteComponent->SetCollisionEnabled(ECollisionEnabled::NoCollision);
+			NaniteComponent->SetMobility(EComponentMobility::Static);
+			NaniteComponent->SetGenerateOverlapEvents(false);
+			NaniteComponent->SetCanEverAffectNavigation(false);
+			NaniteComponent->CanCharacterStepUpOn = ECanBeCharacterBase::ECB_No;
+			NaniteComponent->bSelectable = false;
+			NaniteComponent->DepthPriorityGroup = SDPG_World;
+			NaniteComponent->RegisterComponent();
+		}
+
+		if (NaniteComponent->GetProxyContentId() != NaniteContentId)
+		{
+			NaniteComponent->InitializeForLandscape(this, NaniteContentId);
+		}
+	}
+	else if (NaniteComponent != nullptr)
+	{
+		NaniteComponent->DetachFromComponent(FDetachmentTransformRules::KeepRelativeTransform);
+		NaniteComponent->DestroyComponent();
+		NaniteComponent = nullptr;
 	}
 }
 
-void ULandscapeComponent::CheckGenerateLandscapePlatformData(bool bIsCooking, const ITargetPlatform* TargetPlatform)
+FGuid ALandscapeProxy::GetNaniteContentId() const
 {
-#if ENABLE_LANDSCAPE_COOKING
+	if (!IsNaniteEnabled())
+	{
+		return FGuid();
+	}
 
+	FBufferArchive ContentStateAr;
+
+	struct FCompareULandscapeComponentBySectionBase
+	{
+		FORCEINLINE bool operator()(const ULandscapeComponent& A, const ULandscapeComponent& B) const
+		{
+			// Sort components based on their SectionBase (i.e. 2D index relative to the entire landscape) to ensure stable ID generation
+			return (A.GetSectionBase().X == B.GetSectionBase().X) ? (A.GetSectionBase().Y < B.GetSectionBase().Y) : (A.GetSectionBase().X < B.GetSectionBase().X);
+		}
+	};
+	TArray<ULandscapeComponent*> StableOrderComponents(LandscapeComponents);
+	StableOrderComponents.Sort(FCompareULandscapeComponentBySectionBase());
+
+	for (ULandscapeComponent* Component : StableOrderComponents)
+	{
+		if (Component == nullptr)
+		{
+			continue;
+		}
+
+		FGuid HeightmapGuid = Component->GetHeightmap()->Source.GetId();
+		ContentStateAr << HeightmapGuid;
+
+		// Take into account the Heightmap offset per component
+		ContentStateAr << Component->HeightmapScaleBias.Z;
+		ContentStateAr << Component->HeightmapScaleBias.W;
+
+		if (Component->ComponentHasVisibilityPainted())
+		{
+			const TArray<UTexture2D*>& WeightmapTextures = Component->GetWeightmapTextures();
+			const TArray<FWeightmapLayerAllocationInfo>& AllocInfos = Component->GetWeightmapLayerAllocations();
+			for (const FWeightmapLayerAllocationInfo& AllocInfo : AllocInfos)
+			{
+				if (AllocInfo.IsAllocated() && AllocInfo.LayerInfo == ALandscapeProxy::VisibilityLayer)
+				{
+					UTexture2D* VisibilityWeightmap = WeightmapTextures[AllocInfo.WeightmapTextureIndex];
+					check(VisibilityWeightmap != nullptr);
+
+					FGuid VisibilityWeightmapGuid = VisibilityWeightmap->Source.GetId();
+					ContentStateAr << VisibilityWeightmapGuid;
+				}
+			}
+		}
+
+		// Nanite only cares about LOD0 data
+		if (UMaterialInterface* LandscapeLOD0Material = Component->GetLandscapeMaterial(/*InLODIndex = */ 0))
+		{
+			FGuid LocalStateId = LandscapeLOD0Material->GetMaterial_Concurrent()->StateId;
+			ContentStateAr << LocalStateId;
+		}
+	}
+
+	uint32 Hash[5];
+	FSHA1::HashBuffer(ContentStateAr.GetData(), ContentStateAr.Num(), (uint8*)Hash);
+	return FGuid(Hash[0] ^ Hash[4], Hash[1], Hash[2], Hash[3]);
+}
+
+void ULandscapeComponent::CheckGenerateMobilePlatformData(bool bIsCooking, const ITargetPlatform* TargetPlatform)
+{
 	// Regenerate platform data only when it's missing or there is a valid hash-mismatch.
 
 	FBufferArchive ComponentStateAr;
@@ -308,8 +410,8 @@ void ULandscapeComponent::CheckGenerateLandscapePlatformData(bool bIsCooking, co
 	}
 
 	// Serialize the version guid as part of the hash so we can invalidate DDC data if needed
-	FString Version = FDevSystemGuids::GetSystemGuid(FDevSystemGuids::Get().LANDSCAPE_MOBILE_COOK_VERSION).ToString();
-	ComponentStateAr << Version;
+	FString MobileVersion = FDevSystemGuids::GetSystemGuid(FDevSystemGuids::Get().LANDSCAPE_MOBILE_COOK_VERSION).ToString();
+	ComponentStateAr << MobileVersion;
 
 	uint32 Hash[5];
 	FSHA1::HashBuffer(ComponentStateAr.GetData(), ComponentStateAr.Num(), (uint8*)Hash);
@@ -335,7 +437,7 @@ void ULandscapeComponent::CheckGenerateLandscapePlatformData(bool bIsCooking, co
 // 			}
 // 			else
 			{
-				GeneratePlatformVertexData(TargetPlatform);
+				GenerateMobilePlatformVertexData(TargetPlatform);
 // 				PlatformData.SaveToDDC(NewSourceHash, this);
 				COOK_STAT(Timer.AddMiss(PlatformData.GetPlatformDataSize()));
 			}
@@ -345,7 +447,7 @@ void ULandscapeComponent::CheckGenerateLandscapePlatformData(bool bIsCooking, co
 			// When not cooking (e.g. mobile preview) DDC data isn't sufficient to 
 			// display correctly, so the platform vertex data must be regenerated.
 
-			GeneratePlatformVertexData(TargetPlatform);
+			GenerateMobilePlatformVertexData(TargetPlatform);
 		}
 	}
 
@@ -353,14 +455,13 @@ void ULandscapeComponent::CheckGenerateLandscapePlatformData(bool bIsCooking, co
 
 	if (bRegeneratePixelData)
 	{
-		GeneratePlatformPixelData(bIsCooking, TargetPlatform);
+		GenerateMobilePlatformPixelData(bIsCooking, TargetPlatform);
 	}
 
 	MobileDataSourceHash = NewSourceHash;
-
-#endif
 }
-#endif
+
+#endif // WITH_EDITOR
 
 void ULandscapeComponent::SetForcedLOD(int32 InForcedLOD)
 {
@@ -435,12 +536,15 @@ void ULandscapeComponent::Serialize(FArchive& Ar)
 	Ar.UsingCustomVersion(FEditorObjectVersion::GUID);
 
 #if WITH_EDITOR
-	if (Ar.IsCooking() && !HasAnyFlags(RF_ClassDefaultObject) && UseMobileLandscapeMesh(Ar.CookingTarget()))
+	if (Ar.IsCooking() && !HasAnyFlags(RF_ClassDefaultObject))
 	{
-		// for -oldcook:
-		// the old cooker calls BeginCacheForCookedPlatformData after the package export set is tagged, so the mobile material doesn't get saved, so we have to do CheckGenerateLandscapePlatformData in serialize
-		// the new cooker clears the texture source data before calling serialize, causing GeneratePlatformVertexData to crash, so we have to do CheckGenerateLandscapePlatformData in BeginCacheForCookedPlatformData
-		CheckGenerateLandscapePlatformData(true, Ar.CookingTarget());
+		if (UseMobileLandscapeMesh(Ar.CookingTarget()))
+		{
+			// for -oldcook:
+			// the old cooker calls BeginCacheForCookedPlatformData after the package export set is tagged, so the mobile material doesn't get saved, so we have to do CheckGenerateMobilePlatformData in serialize
+			// the new cooker clears the texture source data before calling serialize, causing GeneratePlatformVertexData to crash, so we have to do CheckGenerateMobilePlatformData in BeginCacheForCookedPlatformData
+			CheckGenerateMobilePlatformData(/*bIsCooking = */ true, Ar.CookingTarget());
+		}
 	}
 
 	// Avoid the archiver in the PIE duplicate writer case because we want to share landscape textures & materials
@@ -631,7 +735,6 @@ void ULandscapeComponent::Serialize(FArchive& Ar)
 		UE_LOG(LogLandscape, Fatal, TEXT("This platform requires cooked packages, and this landscape does not contain cooked data %s."), *GetName());
 	}
 
-#if ENABLE_LANDSCAPE_COOKING
 	if (bCooked)
 	{
 		bool bCookedMobileData = Ar.IsCooking() && UseMobileLandscapeMesh(Ar.CookingTarget());
@@ -647,7 +750,6 @@ void ULandscapeComponent::Serialize(FArchive& Ar)
 			PlatformData.Serialize(Ar, this);
 		}
 	}
-#endif
 
 #if WITH_EDITOR
 	if (Ar.GetPortFlags() & PPF_DuplicateForPIE)
@@ -1089,11 +1191,11 @@ void ULandscapeComponent::PostLoad()
 		UWorld* World = GetWorld();
 		ERHIFeatureLevel::Type FeatureLevel = ((GEngine->GetDefaultWorldFeatureLevel() == ERHIFeatureLevel::ES3_1) || (World && (World->FeatureLevel <= ERHIFeatureLevel::ES3_1))) 
 			? ERHIFeatureLevel::ES3_1 : GMaxRHIFeatureLevel;
-				
+
 		// If we're loading on a platform that doesn't require cooked data, but defaults to a mobile feature level, generate or preload data from the DDC
 		if (!FPlatformProperties::RequiresCookedData() && UseMobileLandscapeMesh(GShaderPlatformForFeatureLevel[FeatureLevel]))
 		{
-			CheckGenerateLandscapePlatformData(false, nullptr);
+			CheckGenerateMobilePlatformData(/*bIsCooking = */ false, /*TargetPlatform = */ nullptr);
 		}
 	}
 
@@ -2651,6 +2753,7 @@ void ALandscape::RenderHeightmap(const FTransform& InRenderAreaWorldTransform, c
 }
 
 #if WITH_EDITOR
+
 FBox ALandscape::GetCompleteBounds() const
 {
 	if (ensure(GetLandscapeInfo()))
@@ -2663,6 +2766,11 @@ FBox ALandscape::GetCompleteBounds() const
 	}
 }
 
+bool ALandscape::IsNaniteEnabled() const
+{
+	return bEnableNanite;
+}
+
 void ALandscapeProxy::OnFeatureLevelChanged(ERHIFeatureLevel::Type NewFeatureLevel)
 {
 	FlushGrassComponents();
@@ -2675,7 +2783,7 @@ void ALandscapeProxy::OnFeatureLevelChanged(ERHIFeatureLevel::Type NewFeatureLev
 		{
 			if (Component != nullptr)
 			{
-				Component->CheckGenerateLandscapePlatformData(false, nullptr);
+				Component->CheckGenerateMobilePlatformData(/*bIsCooking = */ false, /*TargetPlatform = */ nullptr);
 			}
 		}
 	}
@@ -2753,7 +2861,17 @@ void ALandscapeProxy::Serialize(FArchive& Ar)
 			}
 		}
 	}
-#endif // WITH_EDITORONLY_DATA
+#endif
+
+#if WITH_EDITOR
+	if (Ar.IsCooking() && !HasAnyFlags(RF_ClassDefaultObject))
+	{
+		//if (UseNaniteLandscapeMesh(Ar.CookingTarget())) // TODO: WIP
+		{
+			CheckGenerateNanitePlatformData(/*bIsCooking = */ true, Ar.CookingTarget());
+		}
+	}
+#endif
 }
 
 void ALandscapeProxy::AddReferencedObjects(UObject* InThis, FReferenceCollector& Collector)
@@ -3153,13 +3271,24 @@ void ALandscapeProxy::PostLoad()
 			->AddToken(FMapErrorToken::Create(FMapErrors::FixedUpInvalidLandscapeMaterialInstances));
 	}
 
+	UWorld* World = GetWorld();
+
 	// track feature level change to flush grass cache
-	if (GetWorld())
+	if (World)
 	{
 		FOnFeatureLevelChanged::FDelegate FeatureLevelChangedDelegate = FOnFeatureLevelChanged::FDelegate::CreateUObject(this, &ALandscapeProxy::OnFeatureLevelChanged);
-		FeatureLevelChangedDelegateHandle = GetWorld()->AddOnFeatureLevelChangedHandler(FeatureLevelChangedDelegate);
+		FeatureLevelChangedDelegateHandle = World->AddOnFeatureLevelChangedHandler(FeatureLevelChangedDelegate);
 	}
 	RepairInvalidTextures();
+
+	if (World && !HasAnyFlags(RF_ClassDefaultObject) && !FPlatformProperties::RequiresCookedData())
+	{
+		// If we're loading on a platform that doesn't require cooked data, but defaults to a Nanite feature level, generate or preload data from the DDC
+		if (UseNaniteLandscapeMesh(GShaderPlatformForFeatureLevel[World->FeatureLevel.GetValue()]))
+		{
+			CheckGenerateNanitePlatformData(/*bIsCooking = */ false, /*TargetPlatform = */ nullptr);
+		}
+	}
 #endif // WITH_EDITOR
 }
 
@@ -3456,6 +3585,16 @@ UMaterialInterface* ALandscapeStreamingProxy::GetLandscapeHoleMaterial() const
 	return nullptr;
 }
 
+bool ALandscapeStreamingProxy::IsNaniteEnabled() const
+{
+	if (const ALandscape* Landscape = GetLandscapeActor())
+	{
+		return Landscape->IsNaniteEnabled();
+	}
+	
+	return false;
+}
+
 void ALandscape::PreSave(const class ITargetPlatform* TargetPlatform)
 {
 	PRAGMA_DISABLE_DEPRECATION_WARNINGS;
@@ -4485,6 +4624,9 @@ void ULandscapeComponent::SerializeStateHashes(FArchive& Ar)
 		Ar << WeightmapGuid;
 	}
 
+	bool bEnableNanite = GetLandscapeProxy()->IsNaniteEnabled();
+	Ar << bEnableNanite;
+
 	bool bMeshHoles = GetLandscapeProxy()->bMeshHoles;
 	uint8 MeshHolesMaxLod = GetLandscapeProxy()->MeshHolesMaxLod;
 	Ar << bMeshHoles << MeshHolesMaxLod;
@@ -4939,6 +5081,7 @@ void ALandscapeProxy::InvalidateGeneratedComponentData(const TArray<ULandscapeCo
 	for (auto Iter = ByProxy.CreateConstIterator(); Iter; ++Iter)
 	{
 		Iter.Key()->FlushGrassComponents(&Iter.Value());
+		Iter.Key()->CheckGenerateNanitePlatformData(/*bIsCooking = */ false, /*TargetPlatform = */ nullptr);
 
 #if WITH_EDITOR
 		FLandscapeProxyComponentDataChangedParams ChangeParams(Iter.Value());
diff --git a/Engine/Source/Runtime/Landscape/Private/LandscapeEdit.cpp b/Engine/Source/Runtime/Landscape/Private/LandscapeEdit.cpp
--- a/Engine/Source/Runtime/Landscape/Private/LandscapeEdit.cpp
+++ b/Engine/Source/Runtime/Landscape/Private/LandscapeEdit.cpp
@@ -47,7 +47,14 @@ LandscapeEdit.cpp: Landscape editing
 #include "Engine/World.h"
 #include "LandscapeSubsystem.h"
 #include "StaticMeshAttributes.h"
+#include "StaticMeshDescription.h"
+#include "StaticMeshOperations.h"
 #include "MeshUtilitiesCommon.h"
+#include "OverlappingCorners.h"
+#include "MeshBuild.h"
+#include "StaticMeshBuilder.h"
+#include "NaniteBuilder.h"
+#include "Rendering/NaniteResources.h"
 #include "Misc/ScopedSlowTask.h"
 
 #include "EngineModule.h"
@@ -632,7 +639,7 @@ void ULandscapeComponent::PreFeatureLevelChange(ERHIFeatureLevel::Type PendingFe
 	if (UseMobileLandscapeMesh(GShaderPlatformForFeatureLevel[PendingFeatureLevel]))
 	{
 		// See if we need to cook platform data for mobile preview in editor
-		CheckGenerateLandscapePlatformData(false, nullptr);
+		CheckGenerateMobilePlatformData(/*bIsCooking = */ false, /*TargetPlatform = */ nullptr);
 	}
 }
 
@@ -3329,7 +3336,22 @@ bool ALandscapeProxy::ExportToRawMesh(int32 InExportLOD, FMeshDescription& OutRa
 {
 	TInlineComponentArray<ULandscapeComponent*> RegisteredLandscapeComponents;
 	GetComponents<ULandscapeComponent>(RegisteredLandscapeComponents);
+	return ExportToRawMesh(
+		MakeArrayView(RegisteredLandscapeComponents.GetData(), RegisteredLandscapeComponents.Num()),
+		InExportLOD,
+		OutRawMesh,
+		InBounds,
+		bIgnoreBounds
+	);
+}
 
+bool ALandscapeProxy::ExportToRawMesh(
+	const TArrayView<ULandscapeComponent*>& RegisteredLandscapeComponents,
+	int32 InExportLOD,
+	FMeshDescription& OutRawMesh,
+	const FBoxSphereBounds& InBounds,
+	bool bIgnoreBounds /*= false*/) const
+{
 	const FIntRect LandscapeSectionRect = GetBoundingRect();
 	const FVector2f LandscapeUVScale = FVector2f(1.0f, 1.0f) / FVector2f(LandscapeSectionRect.Size());
 
@@ -3356,10 +3378,8 @@ bool ALandscapeProxy::ExportToRawMesh(int32 InExportLOD, FMeshDescription& OutRa
 	}
 
 	// Export data for each component
-	for (auto It = RegisteredLandscapeComponents.CreateConstIterator(); It; ++It)
+	for (ULandscapeComponent* Component : RegisteredLandscapeComponents)
 	{
-		ULandscapeComponent* Component = (*It);
-
 		// Early out if the Landscape bounds and given bounds do not overlap at all
 		if (!bIgnoreBounds && !FBoxSphereBounds::SpheresIntersect(Component->Bounds, InBounds))
 		{
@@ -5072,7 +5092,12 @@ void ALandscapeProxy::PostEditChangeProperty(FPropertyChangedEvent& PropertyChan
 	else if (GIsEditor && 
 		(PropertyName == GET_MEMBER_NAME_CHECKED(ALandscapeProxy, bMeshHoles) || PropertyName == GET_MEMBER_NAME_CHECKED(ALandscapeProxy, MeshHolesMaxLod)))
 	{
-		CheckGenerateLandscapePlatformData(false, nullptr);
+		CheckGenerateMobilePlatformData(false, nullptr);
+		MarkComponentsRenderStateDirty();
+	}
+	if (GIsEditor && PropertyName == FName(TEXT("bEnableNanite")))
+	{
+		CheckGenerateNanitePlatformData(/*bIsCooking = */ false, /*TargetPlatform = */ nullptr);
 		MarkComponentsRenderStateDirty();
 	}
 	else if (PropertyName == FName(TEXT("bUseDynamicMaterialInstance")))
@@ -5173,15 +5198,23 @@ void ALandscapeStreamingProxy::PostEditChangeProperty(FPropertyChangedEvent& Pro
 
 			UWorld* World = GetWorld();
 
-			if (World != nullptr && UseMobileLandscapeMesh(GShaderPlatformForFeatureLevel[World->FeatureLevel]))
+			if (World != nullptr)
 			{
-				for (ULandscapeComponent * Component : LandscapeComponents)
+				if (UseMobileLandscapeMesh(GShaderPlatformForFeatureLevel[World->FeatureLevel]))
 				{
-					if (Component != nullptr)
+					for (ULandscapeComponent* Component : LandscapeComponents)
 					{
-						Component->CheckGenerateLandscapePlatformData(false, nullptr);
+						if (Component != nullptr)
+						{
+							Component->CheckGenerateMobilePlatformData(/*bIsCooking = */ false, /*TargetPlatform = */ nullptr);
+						}
 					}
 				}
+
+				if (UseNaniteLandscapeMesh(GShaderPlatformForFeatureLevel[World->FeatureLevel]))
+				{
+					CheckGenerateNanitePlatformData(/*bIsCooking = */ false, /*TargetPlatform = */ nullptr);
+				}
 			}
 		}
 	}
@@ -5406,6 +5439,7 @@ void ALandscape::PostEditChangeProperty(FPropertyChangedEvent& PropertyChangedEv
 	bool bNeedsRecalcBoundingBox = false;
 	bool bChangedLighting = false;
 	bool bPropagateToProxies = false;
+	bool bNaniteToggled = false;
 
 	ULandscapeInfo* Info = GetLandscapeInfo();
 
@@ -5543,11 +5577,20 @@ void ALandscape::PostEditChangeProperty(FPropertyChangedEvent& PropertyChangedEv
 	{
 		ExportLOD = FMath::Clamp<int32>(ExportLOD, 0, FMath::CeilLogTwo(SubsectionSizeQuads + 1) - 1);
 	}
+	else if (GIsEditor && (PropertyName == GET_MEMBER_NAME_CHECKED(ALandscape, bEnableNanite)))
+	{
+		bNaniteToggled = true;
+
+		// Generate Nanite data for a landscape with components on it, and recreate render state
+		// Streaming proxies won't be built here, but the bPropagateToProxies path will.
+		CheckGenerateNanitePlatformData(/*bIsCooking = */ false, /*TargetPlatform = */ nullptr);
+		MarkComponentsRenderStateDirty();
+	}
 
 	// Must do this *after* clamping values
 	Super::PostEditChangeProperty(PropertyChangedEvent);
 
-	bPropagateToProxies = bPropagateToProxies || bNeedsRecalcBoundingBox || bChangedLighting;
+	bPropagateToProxies = bPropagateToProxies || bNeedsRecalcBoundingBox || bChangedLighting || bNaniteToggled;
 
 	if (Info != nullptr)
 	{
@@ -5598,15 +5641,23 @@ void ALandscape::PostEditChangeProperty(FPropertyChangedEvent& PropertyChangedEv
 
 				UWorld* World = GetWorld();
 
-				if (World != nullptr && UseMobileLandscapeMesh(GShaderPlatformForFeatureLevel[World->FeatureLevel]))
+				if (World != nullptr)
 				{
-					for (ULandscapeComponent * Component : LandscapeComponents)
+					if (UseMobileLandscapeMesh(GShaderPlatformForFeatureLevel[World->FeatureLevel]))
 					{
-						if (Component != nullptr)
+						for (ULandscapeComponent* Component : LandscapeComponents)
 						{
-							Component->CheckGenerateLandscapePlatformData(false, nullptr);
+							if (Component != nullptr)
+							{
+								Component->CheckGenerateMobilePlatformData(/*bIsCooking = */ false, /*TargetPlatform = */ nullptr);
+							}
 						}
 					}
+
+					if (UseNaniteLandscapeMesh(GShaderPlatformForFeatureLevel[World->FeatureLevel]))
+					{
+						CheckGenerateNanitePlatformData(/*bIsCooking = */ false, /*TargetPlatform = */ nullptr);
+					}
 				}
 			}
 		}
@@ -5689,9 +5740,12 @@ void ULandscapeComponent::PostEditChangeProperty(FPropertyChangedEvent& Property
 
 			UWorld* World = GetWorld();
 
-			if (World != nullptr && UseMobileLandscapeMesh(GShaderPlatformForFeatureLevel[World->FeatureLevel]))
+			if (World != nullptr)
 			{
-				CheckGenerateLandscapePlatformData(false, nullptr);
+				if (UseMobileLandscapeMesh(GShaderPlatformForFeatureLevel[World->FeatureLevel]))
+				{
+					CheckGenerateMobilePlatformData(/*bIsCooking = */ false, /*TargetPlatform = */ nullptr);
+				}
 			}
 		}
 	}
@@ -6478,7 +6532,7 @@ void ULandscapeComponent::GenerateMobileWeightmapLayerAllocations()
 	}));
 }
 
-void ULandscapeComponent::GeneratePlatformPixelData(bool bIsCooking, const ITargetPlatform* TargetPlatform)
+void ULandscapeComponent::GenerateMobilePlatformPixelData(bool bIsCooking, const ITargetPlatform* TargetPlatform)
 {
 	check(!IsTemplate());
 
@@ -6993,7 +7047,7 @@ void BuildHoleRenderData(int32 InNumSubsections, int32 InSubsectionSizeVerts, TA
 
 // Generates vertex and index buffer data from the component's height map and visibility textures.
 // For use on mobile platforms that don't use vertex texture fetch for height or alpha testing for visibility.
-void ULandscapeComponent::GeneratePlatformVertexData(const ITargetPlatform* TargetPlatform)
+void ULandscapeComponent::GenerateMobilePlatformVertexData(const ITargetPlatform* TargetPlatform)
 {
 	if (IsTemplate())
 	{
@@ -7688,4 +7742,4 @@ bool ALandscapeProxy::LandscapeExportWeightmapToRenderTarget(UTextureRenderTarge
 
 #endif //WITH_EDITOR
 
-#undef LOCTEXT_NAMESPACE
\ No newline at end of file
+#undef LOCTEXT_NAMESPACE
diff --git a/Engine/Source/Runtime/Landscape/Private/LandscapeRender.cpp b/Engine/Source/Runtime/Landscape/Private/LandscapeRender.cpp
--- a/Engine/Source/Runtime/Landscape/Private/LandscapeRender.cpp
+++ b/Engine/Source/Runtime/Landscape/Private/LandscapeRender.cpp
@@ -10,6 +10,7 @@ LandscapeRender.cpp: New terrain rendering
 #include "LandscapeLayerInfoObject.h"
 #include "LandscapePrivate.h"
 #include "LandscapeMeshProxyComponent.h"
+#include "LandscapeNaniteComponent.h"
 #include "Materials/Material.h"
 #include "Materials/MaterialExpressionTextureCoordinate.h"
 #include "Materials/MaterialExpressionLandscapeLayerCoords.h"
@@ -41,6 +42,8 @@ LandscapeRender.cpp: New terrain rendering
 #include "VT/RuntimeVirtualTexture.h"
 #include "RayTracingInstance.h"
 #include "ProfilingDebugging/LoadTimeTracker.h"
+#include "StaticMeshResources.h"
+#include "NaniteSceneProxy.h"
 
 IMPLEMENT_GLOBAL_SHADER_PARAMETER_STRUCT(FLandscapeUniformShaderParameters, "LandscapeParameters");
 IMPLEMENT_GLOBAL_SHADER_PARAMETER_STRUCT(FLandscapeFixedGridUniformShaderParameters, "LandscapeFixedGrid");
@@ -4240,3 +4243,20 @@ FPrimitiveSceneProxy* ULandscapeMeshProxyComponent::CreateSceneProxy()
 
 	return new FLandscapeMeshProxySceneProxy(this, LandscapeGuid, ProxyComponentBases, ProxyLOD);
 }
+
+FPrimitiveSceneProxy* ULandscapeNaniteComponent::CreateSceneProxy()
+{
+	// TODO: WIP
+	/*if (GetStaticMesh() == nullptr
+		|| GetStaticMesh()->IsCompiling()
+		|| GetStaticMesh()->GetRenderData() == nullptr
+		|| GetStaticMesh()->GetRenderData()->LODResources.Num() == 0
+		|| GetStaticMesh()->GetRenderData()->LODResources[0].VertexBuffers.StaticMeshVertexBuffer.GetNumVertices() == 0)
+	{
+		return nullptr;
+	}
+
+	return new FLandscapeNaniteSceneProxy(this, LandscapeGuid, ProxyComponentBases, ProxyLOD);*/
+	return Super::CreateSceneProxy();
+}
+
diff --git a/Engine/Source/Runtime/RenderCore/Private/RenderUtils.cpp b/Engine/Source/Runtime/RenderCore/Private/RenderUtils.cpp
--- a/Engine/Source/Runtime/RenderCore/Private/RenderUtils.cpp
+++ b/Engine/Source/Runtime/RenderCore/Private/RenderUtils.cpp
@@ -1762,6 +1762,11 @@ RENDERCORE_API bool UseMobileLandscapeMesh(EShaderPlatform ShaderPlatform)
 	return false;
 }
 
+RENDERCORE_API bool UseNaniteLandscapeMesh(EShaderPlatform ShaderPlatform)
+{
+	return DoesPlatformSupportNanite(ShaderPlatform);
+}
+
 RENDERCORE_API bool ExcludeNonPipelinedShaderTypes(EShaderPlatform ShaderPlatform)
 {
 	if (RHISupportsShaderPipelines(ShaderPlatform))
diff --git a/Engine/Source/Runtime/RenderCore/Public/RenderUtils.h b/Engine/Source/Runtime/RenderCore/Public/RenderUtils.h
--- a/Engine/Source/Runtime/RenderCore/Public/RenderUtils.h
+++ b/Engine/Source/Runtime/RenderCore/Public/RenderUtils.h
@@ -804,6 +804,11 @@ RENDERCORE_API bool UseVirtualTextureLightmap(const FStaticFeatureLevel InFeatur
 */
 RENDERCORE_API bool UseMobileLandscapeMesh(EShaderPlatform ShaderPlatform);
 
+/**
+*	Checks if platform uses a Nanite landscape mesh
+*/
+RENDERCORE_API bool UseNaniteLandscapeMesh(EShaderPlatform ShaderPlatform);
+
 /**
  *  Checks if the non-pipeline shaders will not be compild and ones from FShaderPipeline used instead.
  */
