Expose non-const interfaces of control points and segments

diff --git a/Engine/Source/Runtime/Landscape/Classes/LandscapeSplinesComponent.h b/Engine/Source/Runtime/Landscape/Classes/LandscapeSplinesComponent.h
--- a/Engine/Source/Runtime/Landscape/Classes/LandscapeSplinesComponent.h
+++ b/Engine/Source/Runtime/Landscape/Classes/LandscapeSplinesComponent.h
@@ -157,7 +157,10 @@ public:
 	bool ModifySplines(bool bAlwaysMarkDirty = true);
 
 	const TArray<TObjectPtr<ULandscapeSplineControlPoint>>& GetControlPoints() const { return ControlPoints; }
+	TArray<TObjectPtr<ULandscapeSplineControlPoint>>& GetControlPoints() { return ControlPoints; }
+
 	const TArray<TObjectPtr<ULandscapeSplineSegment>>& GetSegments() const { return Segments; }
+	TArray<TObjectPtr<ULandscapeSplineSegment>>& GetSegments() { return Segments; }
 
 #if WITH_EDITOR
 	void RequestSplineLayerUpdate();
