More improvements to Nanite Landscape - each ALandscapeProxy contains multiple ULandscapeComponent, each rendering a grid cell of the heightmap terrain. My implementation for Nanite triangulates and merges all the components under a proxy (you can also have multiple proxies) into a single Nanite mesh with 4 material slots on it, and stores it on a new ULandscapeNaniteComponent also on each proxy.

The rendering of the 4 landscape components is needed for passes (like RVT) where Nanite is not supported or implemented, and we only want the Nanite landscape component to render in the base pass and virtual shadow maps. (mutually exclusive).

I've now implemented a new mutually exclusive render state that each proxy propagates down to the landscape components to override the calculated view relevance without destructively changing properties like bRenderInMainPass (since we don't want to break non-Nanite rendering when Nanite is not supported by the runtime or platform).

diff --git a/Engine/Source/Runtime/Landscape/Classes/LandscapeComponent.h b/Engine/Source/Runtime/Landscape/Classes/LandscapeComponent.h
--- a/Engine/Source/Runtime/Landscape/Classes/LandscapeComponent.h
+++ b/Engine/Source/Runtime/Landscape/Classes/LandscapeComponent.h
@@ -545,6 +545,9 @@ class ULandscapeComponent : public UPrimitiveComponent
 	TLazyObjectPtr<ULandscapeHeightfieldCollisionComponent> CollisionComponent;
 
 private:
+	UPROPERTY(Transient)
+	bool bNaniteActive;
+
 #if WITH_EDITORONLY_DATA
 	/** Unique ID for this component, used for caching during distributed lighting */
 	UPROPERTY()
@@ -1159,6 +1162,13 @@ public:
 	UFUNCTION(BlueprintCallable, Category = "LandscapeComponent")
 	LANDSCAPE_API void SetLODBias(int32 InLODBias);
 
+	void SetNaniteActive(bool bValue);
+
+	inline bool IsNaniteActive() const
+	{
+		return bNaniteActive;
+	}
+
 protected:
 
 #if WITH_EDITOR
diff --git a/Engine/Source/Runtime/Landscape/Private/Landscape.cpp b/Engine/Source/Runtime/Landscape/Private/Landscape.cpp
--- a/Engine/Source/Runtime/Landscape/Private/Landscape.cpp
+++ b/Engine/Source/Runtime/Landscape/Private/Landscape.cpp
@@ -175,6 +175,7 @@ static FAutoConsoleVariableRef CVarRenderCaptureNextHeightmapRenders(
 
 ULandscapeComponent::ULandscapeComponent(const FObjectInitializer& ObjectInitializer)
 	: Super(ObjectInitializer)
+	, bNaniteActive(false)
 #if WITH_EDITORONLY_DATA
 	, CachedEditingLayerData(nullptr)
 	, LayerUpdateFlagPerMode(0)
@@ -324,6 +325,17 @@ void ALandscapeProxy::CheckGenerateNanitePlatformData(bool bIsCooking, const ITa
 		NaniteComponent->DestroyComponent();
 		NaniteComponent = nullptr;
 	}
+
+	if (!bIsCooking)
+	{
+		for (ULandscapeComponent* Component : LandscapeComponents)
+		{
+			if (Component)
+			{
+				Component->SetNaniteActive(IsNaniteEnabled());
+			}
+		}
+	}
 }
 
 FGuid ALandscapeProxy::GetNaniteContentId() const
@@ -528,6 +540,15 @@ void ULandscapeComponent::SetLOD(bool bForcedLODChanged, int32 InLODValue)
 #endif // WITH_EDITOR
 }
 
+void ULandscapeComponent::SetNaniteActive(bool bValue)
+{
+	if (bNaniteActive != bValue)
+	{
+		bNaniteActive = bValue;
+		MarkRenderStateDirty();
+	}
+}
+
 void ULandscapeComponent::Serialize(FArchive& Ar)
 {
 	LLM_SCOPE(ELLMTag::Landscape);
diff --git a/Engine/Source/Runtime/Landscape/Private/LandscapeRender.cpp b/Engine/Source/Runtime/Landscape/Private/LandscapeRender.cpp
--- a/Engine/Source/Runtime/Landscape/Private/LandscapeRender.cpp
+++ b/Engine/Source/Runtime/Landscape/Private/LandscapeRender.cpp
@@ -1118,6 +1118,8 @@ FLandscapeComponentSceneProxy::FLandscapeComponentSceneProxy(ULandscapeComponent
 		bShouldNotifyOnWorldAddRemove = true;
 	}
 
+	bNaniteActive = InComponent->IsNaniteActive();
+
 	EnableGPUSceneSupportFlags();
 
 	const ERHIFeatureLevel::Type FeatureLevel = GetScene().GetFeatureLevel();
@@ -1719,6 +1721,15 @@ FPrimitiveViewRelevance FLandscapeComponentSceneProxy::GetViewRelevance(const FS
 	}
 #endif // !UE_BUILD_SHIPPING
 
+	if (bNaniteActive)
+	{
+		Result.bShadowRelevance = false;
+		Result.bVelocityRelevance = false;
+		Result.bRenderInMainPass = false;
+		Result.bRenderCustomDepth = false;
+		Result.bTranslucentSelfShadow = false;
+	}
+
 	return Result;
 }
 
diff --git a/Engine/Source/Runtime/Landscape/Public/LandscapeRender.h b/Engine/Source/Runtime/Landscape/Public/LandscapeRender.h
--- a/Engine/Source/Runtime/Landscape/Public/LandscapeRender.h
+++ b/Engine/Source/Runtime/Landscape/Public/LandscapeRender.h
@@ -678,6 +678,8 @@ protected:
 	// Storage for static draw list batch params
 	TArray<FLandscapeBatchElementParams> StaticBatchParamArray;
 
+	bool bNaniteActive;
+
 
 #if WITH_EDITOR
 	// Precomputed grass rendering MeshBatch and per-LOD params
