Final cleanup of PhysX stuff in landscape code

diff --git a/Engine/Source/Runtime/Landscape/Classes/LandscapeHeightfieldCollisionComponent.h b/Engine/Source/Runtime/Landscape/Classes/LandscapeHeightfieldCollisionComponent.h
--- a/Engine/Source/Runtime/Landscape/Classes/LandscapeHeightfieldCollisionComponent.h
+++ b/Engine/Source/Runtime/Landscape/Classes/LandscapeHeightfieldCollisionComponent.h
@@ -63,7 +63,7 @@ class ULandscapeHeightfieldCollisionComponent : public UPrimitiveComponent
 	UPROPERTY()
 	TArray<uint8> CollisionQuadFlags;
 
-	/** Guid used to share PhysX heightfield objects in the editor */
+	/** Guid used to share Physics heightfield objects in the editor */
 	UPROPERTY()
 	FGuid HeightfieldGuid;
 
@@ -221,7 +221,7 @@ public:
 	 */
 	virtual bool CookCollisionData(const FName& Format, bool bUseOnlyDefMaterial, bool bCheckDDC, TArray<uint8>& OutCookedData, TArray<UPhysicalMaterial*>& InOutMaterials) const;
 
-	/** Modify a sub-region of the PhysX heightfield. Note that this does not update the physical material */
+	/** Modify a sub-region of the physics heightfield. Note that this does not update the physical material */
 	void UpdateHeightfieldRegion(int32 ComponentX1, int32 ComponentY1, int32 ComponentX2, int32 ComponentY2);
 
 	/** Computes a hash of all the data that will impact final collision */
diff --git a/Engine/Source/Runtime/Landscape/Private/LandscapeCollision.cpp b/Engine/Source/Runtime/Landscape/Private/LandscapeCollision.cpp
--- a/Engine/Source/Runtime/Landscape/Private/LandscapeCollision.cpp
+++ b/Engine/Source/Runtime/Landscape/Private/LandscapeCollision.cpp
@@ -15,7 +15,6 @@
 #include "Misc/SecureHash.h"
 #include "CollisionQueryParams.h"
 #include "Engine/World.h"
-#include "PhysxUserData.h"
 #include "LandscapeProxy.h"
 #include "LandscapeInfo.h"
 #include "Interfaces/Interface_CollisionDataProvider.h"
@@ -25,8 +24,6 @@
 #include "LandscapePrivate.h"
 #include "PhysicsPublic.h"
 #include "LandscapeDataAccess.h"
-#include "PhysXPublic.h"
-#include "PhysicsEngine/PhysXSupport.h"
 #include "DerivedDataCacheInterface.h"
 #include "PhysicalMaterials/PhysicalMaterial.h"
 #include "LandscapeHeightfieldCollisionComponent.h"
@@ -43,10 +40,10 @@
 #include "EngineUtils.h"
 #include "Engine/Engine.h"
 #include "Materials/MaterialInstanceConstant.h"
+#include "Physics/PhysicsFiltering.h"
 #include "Physics/PhysicsInterfaceCore.h"
 #include "Physics/PhysicsInterfaceUtils.h"
 
-#include "PhysXToChaosUtil.h"
 #include "Chaos/ParticleHandle.h"
 #include "Chaos/Vector.h"
 #include "Chaos/Core.h"
@@ -94,9 +91,6 @@ void ULandscapeHeightfieldCollisionComponent::FHeightfieldGeometryRef::GetResour
 {
 	CumulativeResourceSize.AddDedicatedSystemMemoryBytes(sizeof(*this));
 
-	//CumulativeResourceSize.AddDedicatedSystemMemoryBytes(UsedPhysicalMaterialArray.GetAllocatedSize());
-	// COMMENT [jonathan.bard] : unaccounted for : RBHeightfield, RBHeightfieldSimple 
-
 	CumulativeResourceSize.AddDedicatedSystemMemoryBytes(UsedChaosMaterials.GetAllocatedSize());
 
 	if (Heightfield.IsValid())
@@ -137,9 +131,6 @@ void ULandscapeMeshCollisionComponent::FTriMeshGeometryRef::GetResourceSizeEx(FR
 {
 	CumulativeResourceSize.AddDedicatedSystemMemoryBytes(sizeof(*this));
 
-	//CumulativeResourceSize.AddDedicatedSystemMemoryBytes(UsedPhysicalMaterialArray.GetAllocatedSize());
-	// COMMENT [jonathan.bard] : unaccounted for : RBTriangleMesh 
-
 	CumulativeResourceSize.AddDedicatedSystemMemoryBytes(UsedChaosMaterials.GetAllocatedSize());
 
 	if (Trimesh.IsValid())
@@ -277,19 +268,8 @@ void ULandscapeHeightfieldCollisionComponent::OnCreatePhysicsState()
 			// Make transform for this landscape component PxActor
 			FTransform LandscapeComponentTransform = GetComponentToWorld();
 			FMatrix LandscapeComponentMatrix = LandscapeComponentTransform.ToMatrixWithScale();
-			FTransform LandscapeShapeTM = FTransform::Identity;
-
-			// Get the scale to give to PhysX
 			FVector LandscapeScale = LandscapeComponentMatrix.ExtractScaling();
 
-			bool bIsMirrored = LandscapeComponentMatrix.Determinant() < 0.f;
-			if (!bIsMirrored)
-			{
-				// Unreal and PhysX have opposite handedness, so we need to translate the origin and rearrange the data
-				LandscapeShapeTM.SetTranslation(FVector(-CollisionSizeQuads*CollisionScale*LandscapeScale.X, 0, 0));
-			}
-
-
 			// Reorder the axes
 			FVector TerrainX = LandscapeComponentMatrix.GetScaledAxis(EAxis::X);
 			FVector TerrainY = LandscapeComponentMatrix.GetScaledAxis(EAxis::Y);
@@ -1284,14 +1264,7 @@ void ULandscapeMeshCollisionComponent::OnCreatePhysicsState()
 			// Make transform for this landscape component PxActor
 			FTransform LandscapeComponentTransform = GetComponentToWorld();
 			FMatrix LandscapeComponentMatrix = LandscapeComponentTransform.ToMatrixWithScale();
-			bool bIsMirrored = LandscapeComponentMatrix.Determinant() < 0.f;
-			if (bIsMirrored)
-			{
-				// Unreal and PhysX have opposite handedness, so we need to translate the origin and rearrange the data
-				LandscapeComponentMatrix = FTranslationMatrix(FVector(CollisionSizeQuads, 0, 0)) * LandscapeComponentMatrix;
-			}
 
-			// Get the scale to give to PhysX
 			FVector LandscapeScale = LandscapeComponentMatrix.ExtractScaling();
 
 			Initializer.SetComponentToWorld(LandscapeComponentTransform);
@@ -1671,7 +1644,7 @@ void ULandscapeHeightfieldCollisionComponent::Serialize(FArchive& Ar)
 
 		if (FPlatformProperties::RequiresCookedData() && !bCooked && Ar.IsLoading())
 		{
-			UE_LOG(LogPhysics, Fatal, TEXT("This platform requires cooked packages, and physX data was not cooked into %s."), *GetFullName());
+			UE_LOG(LogPhysics, Fatal, TEXT("This platform requires cooked packages, and physics data was not cooked into %s."), *GetFullName());
 		}
 
 		if (bCooked)
@@ -1709,7 +1682,7 @@ void ULandscapeMeshCollisionComponent::Serialize(FArchive& Ar)
 #endif// WITH_EDITORONLY_DATA
 	}
 
-	// PhysX cooking mesh data
+	// Physics cooking mesh data
 	bool bCooked = false;
 	if (Ar.UEVer() >= VER_UE4_ADD_COOKED_TO_LANDSCAPE)
 	{
@@ -1719,7 +1692,7 @@ void ULandscapeMeshCollisionComponent::Serialize(FArchive& Ar)
 
 	if (FPlatformProperties::RequiresCookedData() && !bCooked && Ar.IsLoading())
 	{
-		UE_LOG(LogPhysics, Fatal, TEXT("This platform requires cooked packages, and physX data was not cooked into %s."), *GetFullName());
+		UE_LOG(LogPhysics, Fatal, TEXT("This platform requires cooked packages, and physics data was not cooked into %s."), *GetFullName());
 	}
 
 	if (bCooked)
