Landscape fixes :
* Fixed ASan issue with dirty data (weightmap/heightmap texture resolving debugging code) by removing usage of lambda in favor of direct call to the dirty function
* Fixed redundant ifdefs and invalid comments
* Removed overallocation of material instances array (was previously done to account for hardware tessellation which doesn't exist anymore)
* Renamed FLandscapeEditLayerReadbackResult to FLandscapeEditLayerComponentReadbackResult to better match what it actually is and changed the output type of readback results from ResolveLayersTexture from a map to an array, because a map was overkill and it allows to batch-handle readbacks in UpdateForChangedWeightmaps / UpdateForChangedHeightmaps

diff --git a/Engine/Shaders/Private/Landscape/LandscapeLayersHeightmapsPS.usf b/Engine/Shaders/Private/Landscape/LandscapeLayersHeightmapsPS.usf
--- a/Engine/Shaders/Private/Landscape/LandscapeLayersHeightmapsPS.usf
+++ b/Engine/Shaders/Private/Landscape/LandscapeLayersHeightmapsPS.usf
@@ -231,7 +231,7 @@ float3 ComputeNullableTriangleNormal(float4 InPoint0, float4 InPoint1, float4 In
 
 
 // ----------------------------------------------------------------------------------
-// Compute shaders : 
+// Pixel shaders : 
 
 #if MERGE_EDIT_LAYERS
 // Take all source edit layers heightmap regions and merge them all into a (almost) final heightmap
diff --git a/Engine/Shaders/Private/Landscape/LandscapeLayersWeightmapsPS.usf b/Engine/Shaders/Private/Landscape/LandscapeLayersWeightmapsPS.usf
--- a/Engine/Shaders/Private/Landscape/LandscapeLayersWeightmapsPS.usf
+++ b/Engine/Shaders/Private/Landscape/LandscapeLayersWeightmapsPS.usf
@@ -174,7 +174,7 @@ float FinalizeWeightmapChannel(uint2 InSampleCoordinates, uint InChannelIndex)
 
 
 // ----------------------------------------------------------------------------------
-// Compute shaders : 
+// Pixel shaders : 
 
 #if MERGE_EDIT_LAYERS
 void MergeEditLayers(in float4 SVPos : SV_POSITION, out float OutColor : SV_Target0)
diff --git a/Engine/Source/Runtime/Landscape/Classes/Landscape.h b/Engine/Source/Runtime/Landscape/Classes/Landscape.h
--- a/Engine/Source/Runtime/Landscape/Classes/Landscape.h
+++ b/Engine/Source/Runtime/Landscape/Classes/Landscape.h
@@ -15,7 +15,7 @@ class ILandscapeEdModeInterface;
 class SNotificationItem;
 class UStreamableRenderAsset;
 class FMaterialResource;
-struct FLandscapeEditLayerReadbackResult;
+struct FLandscapeEditLayerComponentReadbackResult;
 struct FNotificationInfo;
 struct FTextureToComponentHelper;
 struct FUpdateLayersContentContext;
@@ -365,22 +365,21 @@ private:
 	int32 RegenerateLayersHeightmaps(const FUpdateLayersContentContext& InUpdateLayersContentContext);
 	int32 PerformLayersHeightmapsLocalMerge(const FUpdateLayersContentContext& InUpdateLayersContentContext, const FEditLayersHeightmapMergeParams& InMergeParams);
 	int32 PerformLayersHeightmapsGlobalMerge(const FUpdateLayersContentContext& InUpdateLayersContentContext, const FEditLayersHeightmapMergeParams& InMergeParams);
-	void ResolveLayersHeightmapTexture(FTextureToComponentHelper const& MapHelper, TSet<UTexture2D*> const& HeightmapsToResolve, bool bIntermediateRender, bool bFlushRender, TMap<ULandscapeComponent*, FLandscapeEditLayerReadbackResult>& InOutComponents);
+	void ResolveLayersHeightmapTexture(const FTextureToComponentHelper& MapHelper, const TSet<UTexture2D*>& HeightmapsToResolve, bool bIntermediateRender, bool bFlushRender, TArray<FLandscapeEditLayerComponentReadbackResult>& InOutComponentReadbackResults);
 
 	int32 RegenerateLayersWeightmaps(FUpdateLayersContentContext& InUpdateLayersContentContext);
 	int32 PerformLayersWeightmapsLocalMerge(FUpdateLayersContentContext& InUpdateLayersContentContext, const FEditLayersWeightmapMergeParams& InMergeParams);
 	int32 PerformLayersWeightmapsGlobalMerge(FUpdateLayersContentContext& InUpdateLayersContentContext, const FEditLayersWeightmapMergeParams& InMergeParams);
-	void ResolveLayersWeightmapTexture(FTextureToComponentHelper const& MapHelper, TSet<UTexture2D*> const& WeightmapsToResolve, bool bIntermediateRender, bool bFlushRender, TMap<ULandscapeComponent*, FLandscapeEditLayerReadbackResult>& InOutComponents);
+	void ResolveLayersWeightmapTexture(const FTextureToComponentHelper& MapHelper, const TSet<UTexture2D*>& WeightmapsToResolve, bool bIntermediateRender, bool bFlushRender, TArray<FLandscapeEditLayerComponentReadbackResult>& InOutComponentReadbackResults);
 
-	using FDirtyDelegate = TFunctionRef<void(UTexture2D const*, FColor const*, FColor const*)>;
 	bool ResolveLayersTexture(FTextureToComponentHelper const& MapHelper, FLandscapeEditLayerReadback* InCPUReadBack, UTexture2D* InOutputTexture, bool bIntermediateRender,	bool bFlushRender,
-		TMap<ULandscapeComponent*, FLandscapeEditLayerReadbackResult>& InOutComponents,	FDirtyDelegate DirtyDelegate);
+		TArray<FLandscapeEditLayerComponentReadbackResult>& InOutComponentReadbackResults, bool bIsWeightmap);
 
 	static bool IsUpdateFlagEnabledForModes(ELandscapeComponentUpdateFlag InFlag, uint32 InUpdateModes);
-	void UpdateForChangedHeightmaps(ULandscapeComponent* InComponent, const FLandscapeEditLayerReadbackResult& InReadbackResult);
-	void UpdateForChangedWeightmaps(ULandscapeComponent* InComponent, const FLandscapeEditLayerReadbackResult& InReadbackResult);
-	int32 UpdateCollisionAndClients(TMap<ULandscapeComponent*, FLandscapeEditLayerReadbackResult> const& Components);
-	int32 UpdateAfterReadbackResolves(TMap<ULandscapeComponent*, FLandscapeEditLayerReadbackResult> const& Components);
+	void UpdateForChangedHeightmaps(const TArrayView<FLandscapeEditLayerComponentReadbackResult>& InComponentReadbackResults);
+	void UpdateForChangedWeightmaps(const TArrayView<FLandscapeEditLayerComponentReadbackResult>& InComponentReadbackResults);
+	uint32 UpdateCollisionAndClients(const TArrayView<FLandscapeEditLayerComponentReadbackResult>& Components);
+	uint32 UpdateAfterReadbackResolves(const TArrayView<FLandscapeEditLayerComponentReadbackResult>& Components);
 
 	bool PrepareTextureResources(bool bInWaitForStreaming);
 	bool PrepareLayersTextureResources(bool bInWaitForStreaming);
diff --git a/Engine/Source/Runtime/Landscape/Private/Landscape.cpp b/Engine/Source/Runtime/Landscape/Private/Landscape.cpp
--- a/Engine/Source/Runtime/Landscape/Private/Landscape.cpp
+++ b/Engine/Source/Runtime/Landscape/Private/Landscape.cpp
@@ -2125,7 +2125,6 @@ void ULandscapeComponent::SetHeightmap(UTexture2D* NewHeightmap)
 
 void ULandscapeComponent::SetWeightmapTextures(const TArray<UTexture2D*>& InNewWeightmapTextures, bool InApplyToEditingWeightmap)
 {
-#if WITH_EDITORONLY_DATA
 	FLandscapeLayerComponentData* EditingLayer = GetEditingLayer();
 
 	if (InApplyToEditingWeightmap && EditingLayer != nullptr)
@@ -2134,7 +2133,6 @@ void ULandscapeComponent::SetWeightmapTextures(const TArray<UTexture2D*>& InNewW
 		EditingLayer->WeightmapData.Textures.Append(InNewWeightmapTextures);
 	}
 	else
-#endif
 	{
 		WeightmapTextures = InNewWeightmapTextures;
 	}
@@ -2147,7 +2145,6 @@ void ULandscapeComponent::SetWeightmapLayerAllocations(const TArray<FWeightmapLa
 
 TArray<ULandscapeWeightmapUsage*>& ULandscapeComponent::GetWeightmapTexturesUsage(bool InReturnEditingWeightmap)
 {
-#if WITH_EDITORONLY_DATA
 	if (InReturnEditingWeightmap)
 	{
 		if (FLandscapeLayerComponentData* EditingLayer = GetEditingLayer())
@@ -2155,14 +2152,12 @@ TArray<ULandscapeWeightmapUsage*>& ULandscapeComponent::GetWeightmapTexturesUsag
 			return EditingLayer->WeightmapData.TextureUsages;
 		}
 	}
-#endif
 
 	return WeightmapTexturesUsage;
 }
 
 const TArray<ULandscapeWeightmapUsage*>& ULandscapeComponent::GetWeightmapTexturesUsage(bool InReturnEditingWeightmap) const
 {
-#if WITH_EDITORONLY_DATA
 	if (InReturnEditingWeightmap)
 	{
 		if (const FLandscapeLayerComponentData* EditingLayer = GetEditingLayer())
@@ -2170,14 +2165,12 @@ const TArray<ULandscapeWeightmapUsage*>& ULandscapeComponent::GetWeightmapTextur
 			return EditingLayer->WeightmapData.TextureUsages;
 		}
 	}
-#endif
 
 	return WeightmapTexturesUsage;
 }
 
 TArray<ULandscapeWeightmapUsage*>& ULandscapeComponent::GetWeightmapTexturesUsage(const FGuid& InLayerGuid)
 {
-#if WITH_EDITORONLY_DATA
 	if (InLayerGuid.IsValid())
 	{
 		if (FLandscapeLayerComponentData* LayerData = GetLayerData(InLayerGuid))
@@ -2185,14 +2178,12 @@ TArray<ULandscapeWeightmapUsage*>& ULandscapeComponent::GetWeightmapTexturesUsag
 			return LayerData->WeightmapData.TextureUsages;
 		}
 	}
-#endif
 
 	return WeightmapTexturesUsage;
 }
 
 const TArray<ULandscapeWeightmapUsage*>& ULandscapeComponent::GetWeightmapTexturesUsage(const FGuid& InLayerGuid) const
 {
-#if WITH_EDITORONLY_DATA
 	if (InLayerGuid.IsValid())
 	{
 		if (const FLandscapeLayerComponentData* LayerData = GetLayerData(InLayerGuid))
@@ -2200,14 +2191,12 @@ const TArray<ULandscapeWeightmapUsage*>& ULandscapeComponent::GetWeightmapTextur
 			return LayerData->WeightmapData.TextureUsages;
 		}
 	}
-#endif
 
 	return WeightmapTexturesUsage;
 }
 
 void ULandscapeComponent::SetWeightmapTexturesUsage(const TArray<ULandscapeWeightmapUsage*>& InNewWeightmapTexturesUsage, bool InApplyToEditingWeightmap)
 {
-#if WITH_EDITORONLY_DATA
 	FLandscapeLayerComponentData* EditingLayer = GetEditingLayer();
 
 	if (InApplyToEditingWeightmap && EditingLayer != nullptr)
@@ -2216,13 +2205,12 @@ void ULandscapeComponent::SetWeightmapTexturesUsage(const TArray<ULandscapeWeigh
 		EditingLayer->WeightmapData.TextureUsages.Append(InNewWeightmapTexturesUsage);
 	}
 	else
-#endif
 	{
 		WeightmapTexturesUsage = InNewWeightmapTexturesUsage;
 	}
 }
 
-#endif
+#endif // WITH_EDITOR
 
 void ALandscapeProxy::PostRegisterAllComponents()
 {
@@ -2283,7 +2271,7 @@ void ALandscapeProxy::PostRegisterAllComponents()
 			LandscapeInfo->FixupProxiesTransform();
 		}
 	}
-#endif
+#endif // WITH_EDITOR
 }
 
 void ALandscapeProxy::UnregisterAllComponents(const bool bForReregister)
@@ -4317,9 +4305,9 @@ void ULandscapeInfo::UnregisterActorComponent(ULandscapeComponent* Component)
 		XYComponentBounds = FIntRect(MAX_int32, MAX_int32, MIN_int32, MIN_int32);
 
 		for (const auto& XYComponentPair : XYtoComponentMap)
-		{
+	{
 			XYComponentBounds.Include(XYComponentPair.Key);
-		}
+	}
 	}
 }
 
@@ -4740,14 +4728,14 @@ void ALandscapeProxy::BuildGIBakedTextures(struct FScopedSlowTask* InSlowTask /*
 		const bool bShouldMarkDirty = true;
 		UpdateGIBakedTextureData(bShouldMarkDirty);
 	}
-}
+	}
 
 int32 ALandscapeProxy::GetOutdatedGIBakedTextureComponentsCount() const
-{
+	{
 	int32 OutdatedGITextureComponentsCount = 0;
 	UpdateGIBakedTextureStatus(nullptr, nullptr, &OutdatedGITextureComponentsCount);
 	return OutdatedGITextureComponentsCount;
-}
+	}
 
 void ALandscapeProxy::UpdateGIBakedTextureStatus(bool* bOutGenerateLandscapeGIData, TMap<UTexture2D*, FGIBakedTextureState>* OutComponentsNeedBakingByHeightmap, int32* OutdatedComponentsCount) const
 {
@@ -4927,53 +4915,53 @@ void ALandscapeProxy::UpdateGIBakedTextures(bool bBakeAllGITextures)
 		{
 			// We throttle, baking only one atlas per frame if bBakeAllGITextures is false.
 			if (!bBakeAllGITextures && NumGenerated > 0)
-			{
-				NumComponentsNeedingTextureBaking += Info.Components.Num();
-			}
-			else
-			{
-				UTexture2D* HeightmapTexture = It.Key();
-				// 1/8 the res of the heightmap
+				{
+					NumComponentsNeedingTextureBaking += Info.Components.Num();
+				}
+				else
+				{
+					UTexture2D* HeightmapTexture = It.Key();
+					// 1/8 the res of the heightmap
 				FIntPoint AtlasSize(HeightmapTexture->Source.GetSizeX() >> 3, HeightmapTexture->Source.GetSizeY() >> 3);
 
-				TArray<FColor> AtlasSamples;
-				AtlasSamples.AddZeroed(AtlasSize.X * AtlasSize.Y);
+					TArray<FColor> AtlasSamples;
+					AtlasSamples.AddZeroed(AtlasSize.X * AtlasSize.Y);
 
-				for (ULandscapeComponent* Component : Info.Components)
-				{
-					// not registered; ignore this component
-					if (!Component->SceneProxy)
+					for (ULandscapeComponent* Component : Info.Components)
 					{
-						continue;
-					}
+						// not registered; ignore this component
+						if (!Component->SceneProxy)
+						{
+							continue;
+						}
 
-					int32 ComponentSamples = (SubsectionSizeQuads + 1) * NumSubsections;
-					check(FMath::IsPowerOfTwo(ComponentSamples));
+						int32 ComponentSamples = (SubsectionSizeQuads + 1) * NumSubsections;
+						check(FMath::IsPowerOfTwo(ComponentSamples));
 
-					int32 BakeSize = ComponentSamples >> 3;
-					TArray<FColor> Samples;
-					if (FMaterialUtilities::ExportBaseColor(Component, BakeSize, Samples))
-					{
+						int32 BakeSize = ComponentSamples >> 3;
+						TArray<FColor> Samples;
+						if (FMaterialUtilities::ExportBaseColor(Component, BakeSize, Samples))
+						{
 						int32 AtlasOffsetX = FMath::RoundToInt(Component->HeightmapScaleBias.Z * (float)HeightmapTexture->Source.GetSizeX()) >> 3;
 						int32 AtlasOffsetY = FMath::RoundToInt(Component->HeightmapScaleBias.W * (float)HeightmapTexture->Source.GetSizeY()) >> 3;
-						for (int32 y = 0; y < BakeSize; y++)
-						{
-							FMemory::Memcpy(&AtlasSamples[(y + AtlasOffsetY) * AtlasSize.X + AtlasOffsetX], &Samples[y * BakeSize], sizeof(FColor) * BakeSize);
+							for (int32 y = 0; y < BakeSize; y++)
+							{
+								FMemory::Memcpy(&AtlasSamples[(y + AtlasOffsetY) * AtlasSize.X + AtlasOffsetX], &Samples[y * BakeSize], sizeof(FColor) * BakeSize);
+							}
+							NumGenerated++;
 						}
-						NumGenerated++;
 					}
-				}
 				UTexture2D* AtlasTexture = FMaterialUtilities::CreateTexture(GetOutermost(), HeightmapTexture->GetName() + TEXT("_BaseColor"), AtlasSize, AtlasSamples, TC_Default, TEXTUREGROUP_World, RF_NoFlags, true, Info.CombinedStateId);
 
-				for (ULandscapeComponent* Component : Info.Components)
-				{
+					for (ULandscapeComponent* Component : Info.Components)
+					{
 					Component->BakedTextureMaterialGuid = Info.CombinedStateId;
-					Component->GIBakedBaseColorTexture = AtlasTexture;
-					Component->MarkRenderStateDirty();
+						Component->GIBakedBaseColorTexture = AtlasTexture;
+						Component->MarkRenderStateDirty();
+					}
 				}
 			}
 		}
-	}
 
 	TotalComponentsNeedingTextureBaking += NumComponentsNeedingTextureBaking;
 
diff --git a/Engine/Source/Runtime/Landscape/Private/LandscapeEdit.cpp b/Engine/Source/Runtime/Landscape/Private/LandscapeEdit.cpp
--- a/Engine/Source/Runtime/Landscape/Private/LandscapeEdit.cpp
+++ b/Engine/Source/Runtime/Landscape/Private/LandscapeEdit.cpp
@@ -420,7 +420,7 @@ void ULandscapeComponent::UpdateMaterialInstances_Internal(FMaterialUpdateContex
 
 	MaterialPerLOD = NewMaterialPerLOD;
 
-	MaterialInstances.SetNumZeroed(MaterialPerLOD.Num() * 2); // over allocate in case we are using tessellation
+	MaterialInstances.SetNumZeroed(MaterialPerLOD.Num());
 	int8 MaterialIndex = 0;
 
 	const TArray<FWeightmapLayerAllocationInfo>& WeightmapBaseLayerAllocation = GetWeightmapLayerAllocations();
diff --git a/Engine/Source/Runtime/Landscape/Private/LandscapeEditLayers.cpp b/Engine/Source/Runtime/Landscape/Private/LandscapeEditLayers.cpp
--- a/Engine/Source/Runtime/Landscape/Private/LandscapeEditLayers.cpp
+++ b/Engine/Source/Runtime/Landscape/Private/LandscapeEditLayers.cpp
@@ -2511,21 +2511,17 @@ void ExecuteCopyLayersTexture(TArray<FLandscapeLayersCopyTextureParams>&& InCopy
 }
 
 /** Per component information from read back results. */
-struct FLandscapeEditLayerReadbackResult
+struct FLandscapeEditLayerComponentReadbackResult
 {
+	ULandscapeComponent* LandscapeComponent = nullptr;
 	/** ELandscapeLayerUpdateMode flags set on ULandscapeComponent at time when read back task was submitted. */
-	int32 UpdateModes;
+	uint32 UpdateModes = 0;
 	/** Were the associated heightmap/weightmaps modified. */
-	bool bModified;
-
-	FLandscapeEditLayerReadbackResult()
-		: UpdateModes(0)
-		, bModified(false)
-	{}
+	bool bModified = false;
 
-	FLandscapeEditLayerReadbackResult(int32 InUpdateModes, bool bInModified)
-		: UpdateModes(InUpdateModes)
-		, bModified(bInModified)
+	FLandscapeEditLayerComponentReadbackResult(ULandscapeComponent* InLandscapeComponent, uint32 InUpdateModes)
+		: LandscapeComponent(InLandscapeComponent)
+		, UpdateModes(InUpdateModes)
 	{}
 };
 
@@ -3684,7 +3680,7 @@ struct FUpdateLayersContentContext
 				LandscapeComponentsWeightmapsToRender.Reset();
 				LandscapeComponentsWeightmapsToResolve.Reset();
 			}
-			// Note that the AllLandscapeComponentsToResolve and AllLandscapeComponentsToUpdate are *not* reset here: they can only grow (we're assuming refresh only adds new components): 
+			// Note that the AllLandscapeComponentsToResolve and AllLandscapeComponentReadbackResults are *not* reset here: they can only grow (we're assuming refresh only adds new components): 
 
 			// Iterate on all dirty components and retrieve the components that need to be resolved or rendered for their heightmap or weightmaps :
 			TArray<ULandscapeComponent*> AllLandscapeComponents;
@@ -3805,7 +3801,11 @@ struct FUpdateLayersContentContext
 			{
 				if (Component->GetLayerUpdateFlagPerMode() & ELandscapeLayerUpdateMode::Update_Client_Deferred)
 				{
-					AllLandscapeComponentsToUpdate.FindOrAdd(Component, FLandscapeEditLayerReadbackResult(ELandscapeLayerUpdateMode::Update_Client_Deferred, false));
+					FLandscapeEditLayerComponentReadbackResult* ComponentReadbackResult = AllLandscapeComponentReadbackResults.FindByPredicate([Component](const FLandscapeEditLayerComponentReadbackResult& Element) { return Element.LandscapeComponent == Component; });
+					if (ComponentReadbackResult == nullptr)
+					{
+						AllLandscapeComponentReadbackResults.Add(FLandscapeEditLayerComponentReadbackResult(Component, ELandscapeLayerUpdateMode::Update_Client_Deferred));
+					}
 				}
 			}
 		}
@@ -3834,7 +3834,7 @@ struct FUpdateLayersContentContext
 	// List of components whose heightmap or weightmaps needs to be resolved because they are either dirty or are neighbor to a component that is dirty:
 	TSet<ULandscapeComponent*> AllLandscapeComponentsToResolve;
 	// List of GPU readback results for heightmaps/weightmaps that need to be resolved, associated with their owning landscape component :
-	TMap<ULandscapeComponent*, FLandscapeEditLayerReadbackResult> AllLandscapeComponentsToUpdate;
+	TArray<FLandscapeEditLayerComponentReadbackResult> AllLandscapeComponentReadbackResults;
 };
 ENUM_CLASS_FLAGS(FUpdateLayersContentContext::ERefreshFlags);
 
@@ -4915,39 +4915,42 @@ int32 ALandscape::RegenerateLayersHeightmaps(const FUpdateLayersContentContext&
 	return HeightmapUpdateModes;
 }
 
-void ALandscape::UpdateForChangedHeightmaps(ULandscapeComponent* InComponent, const FLandscapeEditLayerReadbackResult& InReadbackResult)
+void ALandscape::UpdateForChangedHeightmaps(const TArrayView<FLandscapeEditLayerComponentReadbackResult>& InComponentReadbackResults)
 {
 	TRACE_CPUPROFILER_EVENT_SCOPE(LandscapeLayers_UpdateForChangedHeightmaps);
 
-	// If the source data has changed, mark the component as needing a collision data update:
-	//  - If ELandscapeComponentUpdateFlag::Component_Update_Heightmap_Collision is passed, it will be done immediately
-	//  - If not, at least the component's collision data will still get updated eventually, when the flag is finally passed :
-	if (InReadbackResult.bModified)
+	for (const FLandscapeEditLayerComponentReadbackResult& ComponentReadbackResult : InComponentReadbackResults)
 	{
-		InComponent->SetPendingCollisionDataUpdate(true);
-	}
+		// If the source data has changed, mark the component as needing a collision data update:
+		//  - If ELandscapeComponentUpdateFlag::Component_Update_Heightmap_Collision is passed, it will be done immediately
+		//  - If not, at least the component's collision data will still get updated eventually, when the flag is finally passed :
+		if (ComponentReadbackResult.bModified)
+		{
+			ComponentReadbackResult.LandscapeComponent->SetPendingCollisionDataUpdate(true);
+		}
 
-	const int32 HeightUpdateMode = InReadbackResult.UpdateModes & (ELandscapeLayerUpdateMode::Update_Heightmap_All | ELandscapeLayerUpdateMode::Update_Heightmap_Editing | ELandscapeLayerUpdateMode::Update_Heightmap_Editing_NoCollision);
+		const uint32 HeightUpdateMode = ComponentReadbackResult.UpdateModes & (ELandscapeLayerUpdateMode::Update_Heightmap_All | ELandscapeLayerUpdateMode::Update_Heightmap_Editing | ELandscapeLayerUpdateMode::Update_Heightmap_Editing_NoCollision);
 
-	// Only update collision if there was an actual change performed on the source data : 
-	if (InComponent->GetPendingCollisionDataUpdate())
-	{
-		if (IsUpdateFlagEnabledForModes(ELandscapeComponentUpdateFlag::Component_Update_Heightmap_Collision, HeightUpdateMode))
+		// Only update collision if there was an actual change performed on the source data : 
+		if (ComponentReadbackResult.LandscapeComponent->GetPendingCollisionDataUpdate())
 		{
-			InComponent->UpdateCachedBounds();
-			InComponent->UpdateComponentToWorld();
+			if (IsUpdateFlagEnabledForModes(ELandscapeComponentUpdateFlag::Component_Update_Heightmap_Collision, HeightUpdateMode))
+			{
+				ComponentReadbackResult.LandscapeComponent->UpdateCachedBounds();
+				ComponentReadbackResult.LandscapeComponent->UpdateComponentToWorld();
 
-			// Avoid updating height field if we are going to recreate collision in this update
-			bool bUpdateHeightfieldRegion = !IsUpdateFlagEnabledForModes(ELandscapeComponentUpdateFlag::Component_Update_Recreate_Collision, HeightUpdateMode);
-			InComponent->UpdateCollisionData(bUpdateHeightfieldRegion);
-			InComponent->SetPendingCollisionDataUpdate(false);
-		}
-		else if (IsUpdateFlagEnabledForModes(ELandscapeComponentUpdateFlag::Component_Update_Approximated_Bounds, HeightUpdateMode))
-		{
-			// Update bounds with an approximated value (real computation will be done anyways when computing collision)
-			const bool bInApproximateBounds = true;
-			InComponent->UpdateCachedBounds(bInApproximateBounds);
-			InComponent->UpdateComponentToWorld();
+				// Avoid updating height field if we are going to recreate collision in this update
+				bool bUpdateHeightfieldRegion = !IsUpdateFlagEnabledForModes(ELandscapeComponentUpdateFlag::Component_Update_Recreate_Collision, HeightUpdateMode);
+				ComponentReadbackResult.LandscapeComponent->UpdateCollisionData(bUpdateHeightfieldRegion);
+				ComponentReadbackResult.LandscapeComponent->SetPendingCollisionDataUpdate(false);
+			}
+			else if (IsUpdateFlagEnabledForModes(ELandscapeComponentUpdateFlag::Component_Update_Approximated_Bounds, HeightUpdateMode))
+			{
+				// Update bounds with an approximated value (real computation will be done anyways when computing collision)
+				const bool bInApproximateBounds = true;
+				ComponentReadbackResult.LandscapeComponent->UpdateCachedBounds(bInApproximateBounds);
+				ComponentReadbackResult.LandscapeComponent->UpdateComponentToWorld();
+			}
 		}
 	}
 }
@@ -4957,7 +4960,7 @@ void ALandscape::ResolveLayersHeightmapTexture(
 	TSet<UTexture2D*> const& HeightmapsToResolve,
 	bool bIntermediateRender,
 	bool bFlushRender,
-	TMap<ULandscapeComponent*, FLandscapeEditLayerReadbackResult>& InOutComponents)
+	TArray<FLandscapeEditLayerComponentReadbackResult>& InOutComponentReadbackResults)
 {
 	TRACE_CPUPROFILER_EVENT_SCOPE(LandscapeLayers_ResolveLayersHeightmapTexture);
 	SCOPE_CYCLE_COUNTER(STAT_LandscapeLayersResolveHeightmaps);
@@ -4968,11 +4971,6 @@ void ALandscape::ResolveLayersHeightmapTexture(
 		return;
 	}
 
-	FDirtyDelegate OnDirtyDelegate = [this, &MapHelper](UTexture2D const* Heightmap, FColor const* OldData, FColor const* NewData)
-	{
-		OnDirtyHeightmap(MapHelper, Heightmap, OldData, NewData);
-	};
-
 	TArray<ULandscapeComponent*> ChangedComponents;
 	for (UTexture2D* Heightmap : HeightmapsToResolve)
 	{
@@ -4980,7 +4978,7 @@ void ALandscape::ResolveLayersHeightmapTexture(
 		check(LandscapeProxy);
 		if (FLandscapeEditLayerReadback** CPUReadback = LandscapeProxy->HeightmapsCPUReadback.Find(Heightmap))
 		{
-			const bool bChanged = ResolveLayersTexture(MapHelper, *CPUReadback, Heightmap, bIntermediateRender, bFlushRender, InOutComponents, OnDirtyDelegate);
+			const bool bChanged = ResolveLayersTexture(MapHelper, *CPUReadback, Heightmap, bIntermediateRender, bFlushRender, InOutComponentReadbackResults, /*bIsWeightmap = */false);
 			if (bChanged)
 			{
 				ChangedComponents.Append(MapHelper.HeightmapToComponents[Heightmap]);
@@ -5171,8 +5169,8 @@ bool ALandscape::ResolveLayersTexture(
 	UTexture2D* InOutputTexture,
 	bool bIntermediateRender,
 	bool bFlushRender,
-	TMap<ULandscapeComponent*, FLandscapeEditLayerReadbackResult>& InOutComponents,
-	FDirtyDelegate DirtyDelegate)
+	TArray<FLandscapeEditLayerComponentReadbackResult>& InOutComponentReadbackResults,
+	bool bIsWeightmap)
 {
 	TRACE_CPUPROFILER_EVENT_SCOPE(LandscapeLayers_ResolveLayersTexture);
 	SCOPE_CYCLE_COUNTER(STAT_LandscapeLayersResolveTexture);
@@ -5206,10 +5204,21 @@ bool ALandscape::ResolveLayersTexture(
 				// Don't do this for intermediate renders.
 				if (MipIndex == 0 && !bIntermediateRender)
 				{
-					const uint32 Hash = FLandscapeEditLayerReadback::CalculateHash((uint8*)OutMipsData[MipIndex].GetData(), OutMipsData[MipIndex].Num() * sizeof(FColor));
+					uint32 Hash = 0;
+					{
+						TRACE_CPUPROFILER_EVENT_SCOPE(LandscapeLayers_CalculateHash);
+						Hash = FLandscapeEditLayerReadback::CalculateHash((uint8*)OutMipsData[MipIndex].GetData(), OutMipsData[MipIndex].Num() * sizeof(FColor));
+					}
 					if (InCPUReadback->SetHash(Hash))
 					{
-						DirtyDelegate(InOutputTexture, (FColor*)TextureData, OutMipsData[MipIndex].GetData());
+						if (bIsWeightmap)
+						{
+							OnDirtyWeightmap(MapHelper, InOutputTexture, (FColor*)TextureData, OutMipsData[MipIndex].GetData());
+						}
+						else
+						{
+							OnDirtyHeightmap(MapHelper, InOutputTexture, (FColor*)TextureData, OutMipsData[MipIndex].GetData());
+						}
 						bChanged = true;
 
 						// We're about to modify the texture's source data, the texture needs to know so that it can handle properly update cached platform data (additionally, the package needs to be dirtied) :
@@ -5232,9 +5241,13 @@ bool ALandscape::ResolveLayersTexture(
 				ULandscapeComponent** Component = GetLandscapeInfo()->XYtoComponentMap.Find(ComponentContext.ComponentKey);
 				if (Component != nullptr && *Component != nullptr)
 				{
-					FLandscapeEditLayerReadbackResult& ComponentInfo = InOutComponents.FindOrAdd(*Component);
-					ComponentInfo.UpdateModes |= ComponentContext.UpdateModes;
-					ComponentInfo.bModified |= bChanged ? 1 : 0;
+					FLandscapeEditLayerComponentReadbackResult* ComponentReadbackResult = InOutComponentReadbackResults.FindByPredicate([LandscapeComponent = *Component](const FLandscapeEditLayerComponentReadbackResult& Element) { return Element.LandscapeComponent == LandscapeComponent; });
+					if (ComponentReadbackResult == nullptr)
+					{
+						ComponentReadbackResult = &InOutComponentReadbackResults.Add_GetRef(FLandscapeEditLayerComponentReadbackResult(*Component, ELandscapeLayerUpdateMode::Update_None));
+					}
+					ComponentReadbackResult->UpdateModes |= ComponentContext.UpdateModes;
+					ComponentReadbackResult->bModified |= bChanged ? 1 : 0;
 				}
 			}
 		}
@@ -7043,25 +7056,29 @@ int32 ALandscape::RegenerateLayersWeightmaps(FUpdateLayersContentContext& InUpda
 	return 0;
 }
 
-void ALandscape::UpdateForChangedWeightmaps(ULandscapeComponent* InComponent, const FLandscapeEditLayerReadbackResult& InReadbackResult)
+void ALandscape::UpdateForChangedWeightmaps(const TArrayView<FLandscapeEditLayerComponentReadbackResult>& InComponentReadbackResults)
 {
-	// If the source data has changed, mark the component as needing a collision layer data update:
-	//  - If ELandscapeComponentUpdateFlag::Component_Update_Weightmap_Collision is passed, it will be done immediately
-	//  - If not, at least the component's collision layer data will still get updated eventually, when the flag is finally passed :
-	if (InReadbackResult.bModified)
-	{
-		InComponent->SetPendingLayerCollisionDataUpdate(true);
-	}
-
-	const int32 WeightUpdateMode = InReadbackResult.UpdateModes & (ELandscapeLayerUpdateMode::Update_Weightmap_All | ELandscapeLayerUpdateMode::Update_Weightmap_Editing | ELandscapeLayerUpdateMode::Update_Weightmap_Editing_NoCollision);
+	TArray<ULandscapeComponent*> ComponentsNeedingMaterialInstanceUpdates;
 
-	if (IsUpdateFlagEnabledForModes(ELandscapeComponentUpdateFlag::Component_Update_Weightmap_Collision, WeightUpdateMode))
+	for (const FLandscapeEditLayerComponentReadbackResult& ComponentReadbackResult : InComponentReadbackResults)
 	{
-		// Only update collision data if there was an actual change performed on the source data : 
-		if (InComponent->GetPendingLayerCollisionDataUpdate())
+		// If the source data has changed, mark the component as needing a collision layer data update:
+		//  - If ELandscapeComponentUpdateFlag::Component_Update_Weightmap_Collision is passed, it will be done immediately
+		//  - If not, at least the component's collision layer data will still get updated eventually, when the flag is finally passed :
+		if (ComponentReadbackResult.bModified)
 		{
-			InComponent->UpdateCollisionLayerData();
-			InComponent->SetPendingLayerCollisionDataUpdate(false);
+			ComponentReadbackResult.LandscapeComponent->SetPendingLayerCollisionDataUpdate(true);
+		}
+
+		const int32 WeightUpdateMode = ComponentReadbackResult.UpdateModes & (ELandscapeLayerUpdateMode::Update_Weightmap_All | ELandscapeLayerUpdateMode::Update_Weightmap_Editing | ELandscapeLayerUpdateMode::Update_Weightmap_Editing_NoCollision);
+		if (IsUpdateFlagEnabledForModes(ELandscapeComponentUpdateFlag::Component_Update_Weightmap_Collision, WeightUpdateMode))
+		{
+			// Only update collision data if there was an actual change performed on the source data : 
+			if (ComponentReadbackResult.LandscapeComponent->GetPendingLayerCollisionDataUpdate())
+			{
+				ComponentReadbackResult.LandscapeComponent->UpdateCollisionLayerData();
+				ComponentReadbackResult.LandscapeComponent->SetPendingLayerCollisionDataUpdate(false);
+			}
 		}
 	}
 }
@@ -7271,7 +7288,7 @@ void ALandscape::ResolveLayersWeightmapTexture(
 	TSet<UTexture2D*> const& WeightmapsToResolve,
 	bool bIntermediateRender,
 	bool bFlushRender,
-	TMap<ULandscapeComponent*, FLandscapeEditLayerReadbackResult>& InOutComponents)
+	TArray<FLandscapeEditLayerComponentReadbackResult>& InOutComponentReadbackResults)
 {
 	TRACE_CPUPROFILER_EVENT_SCOPE(LandscapeLayers_ResolveLayersWeightmapTexture);
 	SCOPE_CYCLE_COUNTER(STAT_LandscapeLayersResolveWeightmaps);
@@ -7282,11 +7299,6 @@ void ALandscape::ResolveLayersWeightmapTexture(
 		return;
 	}
 
-	FDirtyDelegate OnDirtyDelegate = [this, &MapHelper](UTexture2D const* Weightmap, FColor const* OldData, FColor const* NewData)
-	{
-		OnDirtyWeightmap(MapHelper, Weightmap, OldData, NewData);
-	};
-
 	TArray<ULandscapeComponent*> ChangedComponents;
 	for (UTexture2D* Weightmap : WeightmapsToResolve)
 	{
@@ -7294,7 +7306,7 @@ void ALandscape::ResolveLayersWeightmapTexture(
 		check(LandscapeProxy);
 		if (FLandscapeEditLayerReadback** CPUReadback = LandscapeProxy->WeightmapsCPUReadback.Find(Weightmap))
 		{
-			const bool bChanged = ResolveLayersTexture(MapHelper, *CPUReadback, Weightmap, bIntermediateRender, bFlushRender, InOutComponents, OnDirtyDelegate);
+			const bool bChanged = ResolveLayersTexture(MapHelper, *CPUReadback, Weightmap, bIntermediateRender, bFlushRender, InOutComponentReadbackResults, /*bIsWeightmap = */true);
 			if (bChanged)
 			{
 				ChangedComponents.Append(MapHelper.WeightmapToComponents[Weightmap]);
@@ -7773,11 +7785,10 @@ void ALandscape::UpdateLayersContent(bool bInWaitForStreaming, bool bInSkipMonit
 		const bool bDoIntermediateRender = false; // bIntermediateRender flag is for the work queued up this frame not the delayed resolves
 		const bool bDoFlushRender = bIntermediateRender; // Flush before we do an intermediate render later in this frame
 
-		TMap<ULandscapeComponent*, FLandscapeEditLayerReadbackResult> ModifiedComponentInfo;
-		ResolveLayersHeightmapTexture(MapHelper, MapHelper.Heightmaps, bDoIntermediateRender, bDoFlushRender, ModifiedComponentInfo);
-		ResolveLayersWeightmapTexture(MapHelper, MapHelper.Weightmaps, bDoIntermediateRender, bDoFlushRender, ModifiedComponentInfo);
-
-		LayerContentUpdateModes |= UpdateAfterReadbackResolves(ModifiedComponentInfo);
+		TArray<FLandscapeEditLayerComponentReadbackResult> ComponentReadbackResults;
+		ResolveLayersHeightmapTexture(MapHelper, MapHelper.Heightmaps, bDoIntermediateRender, bDoFlushRender, ComponentReadbackResults);
+		ResolveLayersWeightmapTexture(MapHelper, MapHelper.Weightmaps, bDoIntermediateRender, bDoFlushRender, ComponentReadbackResults);
+		LayerContentUpdateModes |= UpdateAfterReadbackResolves(ComponentReadbackResults);
 	}
 
 	if (LayerContentUpdateModes == 0 && !bForceRender)
@@ -7801,8 +7812,8 @@ void ALandscape::UpdateLayersContent(bool bInWaitForStreaming, bool bInSkipMonit
 	if (bFlushRender || CVarLandscapeForceFlush.GetValueOnGameThread() != 0)
 	{
 		const bool bDoFlushRender = true;
-		ResolveLayersHeightmapTexture(UpdateLayersContentContext.MapHelper, UpdateLayersContentContext.HeightmapsToResolve, bIntermediateRender, bDoFlushRender, UpdateLayersContentContext.AllLandscapeComponentsToUpdate);
-		ResolveLayersWeightmapTexture(UpdateLayersContentContext.MapHelper, UpdateLayersContentContext.WeightmapsToResolve, bIntermediateRender, bDoFlushRender, UpdateLayersContentContext.AllLandscapeComponentsToUpdate);
+		ResolveLayersHeightmapTexture(UpdateLayersContentContext.MapHelper, UpdateLayersContentContext.HeightmapsToResolve, bIntermediateRender, bDoFlushRender, UpdateLayersContentContext.AllLandscapeComponentReadbackResults);
+		ResolveLayersWeightmapTexture(UpdateLayersContentContext.MapHelper, UpdateLayersContentContext.WeightmapsToResolve, bIntermediateRender, bDoFlushRender, UpdateLayersContentContext.AllLandscapeComponentReadbackResults);
 	}
 
 	// Clear processed mode flags
@@ -7813,7 +7824,7 @@ void ALandscape::UpdateLayersContent(bool bInWaitForStreaming, bool bInSkipMonit
 	}
 
 	// Apply post resolve updates
-	const int32 ToProcessModes = UpdateAfterReadbackResolves(UpdateLayersContentContext.AllLandscapeComponentsToUpdate);
+	const uint32 ToProcessModes = UpdateAfterReadbackResolves(UpdateLayersContentContext.AllLandscapeComponentReadbackResults);
 	LayerContentUpdateModes |= ToProcessModes;
 
 	if (LandscapeEdMode)
@@ -7853,7 +7864,7 @@ private:
 	ULandscapeHeightfieldCollisionComponent* CollisionComponent;
 };
 
-int32 ALandscape::UpdateCollisionAndClients(TMap<ULandscapeComponent*, FLandscapeEditLayerReadbackResult> const& Components)
+uint32 ALandscape::UpdateCollisionAndClients(const TArrayView<FLandscapeEditLayerComponentReadbackResult>& InComponentReadbackResults)
 {
 	TRACE_CPUPROFILER_EVENT_SCOPE(LandscapeLayers_PostResolve_CollisionAndClients);
 
@@ -7865,27 +7876,24 @@ int32 ALandscape::UpdateCollisionAndClients(TMap<ULandscapeComponent*, FLandscap
 	TArray<uint16> HeightData;
 	TArray<uint8> LayerContributionMaskData;
 
-	for (auto Pair : Components)
+	for (const FLandscapeEditLayerComponentReadbackResult& ComponentReadbackResult : InComponentReadbackResults)
 	{
-		ULandscapeComponent* LandscapeComponent = Pair.Key;
-		const int32 UpdateModes = Pair.Value.UpdateModes;
-
 		bool bDeferClientUpdateForComponent = false;
 		bool bDoUpdateClient = true;
-		if (IsUpdateFlagEnabledForModes(ELandscapeComponentUpdateFlag::Component_Update_Recreate_Collision, UpdateModes))
+		if (IsUpdateFlagEnabledForModes(ELandscapeComponentUpdateFlag::Component_Update_Recreate_Collision, ComponentReadbackResult.UpdateModes))
 		{
-			if (ULandscapeHeightfieldCollisionComponent* CollisionComp = LandscapeComponent->CollisionComponent.Get())
+			if (ULandscapeHeightfieldCollisionComponent* CollisionComp = ComponentReadbackResult.LandscapeComponent->CollisionComponent.Get())
 			{
 				FEnableCollisionHashOptimScope Scope(CollisionComp);
 				bDoUpdateClient = CollisionComp->RecreateCollision();
 			}
 		}
 
-		if (bDoUpdateClient && IsUpdateFlagEnabledForModes(ELandscapeComponentUpdateFlag::Component_Update_Client, UpdateModes))
+		if (bDoUpdateClient && IsUpdateFlagEnabledForModes(ELandscapeComponentUpdateFlag::Component_Update_Client, ComponentReadbackResult.UpdateModes))
 		{
 			if (!GUndo)
 			{
-				if (ULandscapeHeightfieldCollisionComponent* CollisionComp = LandscapeComponent->CollisionComponent.Get())
+				if (ULandscapeHeightfieldCollisionComponent* CollisionComp = ComponentReadbackResult.LandscapeComponent->CollisionComponent.Get())
 				{
 					FNavigationSystem::UpdateComponentData(*CollisionComp);
 					CollisionComp->SnapFoliageInstances();
@@ -7898,11 +7906,11 @@ int32 ALandscape::UpdateCollisionAndClients(TMap<ULandscapeComponent*, FLandscap
 			}
 		}
 
-		if (IsUpdateFlagEnabledForModes(ELandscapeComponentUpdateFlag::Component_Update_Client_Editing, UpdateModes))
+		if (IsUpdateFlagEnabledForModes(ELandscapeComponentUpdateFlag::Component_Update_Client_Editing, ComponentReadbackResult.UpdateModes))
 		{
 			if (LandscapeEdModeInfo.ViewMode == ELandscapeViewMode::LayerContribution)
 			{
-				check(ComponentSizeQuads == LandscapeComponent->ComponentSizeQuads);
+				check(ComponentSizeQuads == ComponentReadbackResult.LandscapeComponent->ComponentSizeQuads);
 				const int32 Stride = (1 + ComponentSizeQuads);
 				const int32 ArraySize = Stride * Stride;
 				if (LayerContributionMaskData.Num() != ArraySize)
@@ -7910,13 +7918,13 @@ int32 ALandscape::UpdateCollisionAndClients(TMap<ULandscapeComponent*, FLandscap
 					LayerContributionMaskData.AddZeroed(ArraySize);
 				}
 				uint8* LayerContributionMaskDataPtr = LayerContributionMaskData.GetData();
-				const int32 X1 = LandscapeComponent->GetSectionBase().X;
+				const int32 X1 = ComponentReadbackResult.LandscapeComponent->GetSectionBase().X;
 				const int32 X2 = X1 + ComponentSizeQuads;
-				const int32 Y1 = LandscapeComponent->GetSectionBase().Y;
+				const int32 Y1 = ComponentReadbackResult.LandscapeComponent->GetSectionBase().Y;
 				const int32 Y2 = Y1 + ComponentSizeQuads;
 				bool bLayerContributionWrittenData = false;
 
-				ULandscapeInfo* Info = LandscapeComponent->GetLandscapeInfo();
+				ULandscapeInfo* Info = ComponentReadbackResult.LandscapeComponent->GetLandscapeInfo();
 				check(Info);
 				FLandscapeEditDataInterface LandscapeEdit(Info);
 
@@ -7956,7 +7964,7 @@ int32 ALandscape::UpdateCollisionAndClients(TMap<ULandscapeComponent*, FLandscap
 
 		if (bDeferClientUpdateForComponent)
 		{
-			LandscapeComponent->RequestDeferredClientUpdate();
+			ComponentReadbackResult.LandscapeComponent->RequestDeferredClientUpdate();
 		}
 	}
 
@@ -7964,23 +7972,20 @@ int32 ALandscape::UpdateCollisionAndClients(TMap<ULandscapeComponent*, FLandscap
 	return bAllClientsUpdated ? 0 : ELandscapeLayerUpdateMode::Update_Client_Deferred;
 }
 
-int32 ALandscape::UpdateAfterReadbackResolves(TMap<ULandscapeComponent*, FLandscapeEditLayerReadbackResult> const& Components)
+uint32 ALandscape::UpdateAfterReadbackResolves(const TArrayView<FLandscapeEditLayerComponentReadbackResult>& InComponentReadbackResults)
 {
 	TRACE_CPUPROFILER_EVENT_SCOPE(LandscapeLayers_PostResolve_Updates);
 
-	int32 NewUpdateFlags = 0;
+	uint32 NewUpdateFlags = 0;
 
-	if (Components.Num())
+	if (InComponentReadbackResults.Num())
 	{
-		for (auto Pair : Components)
-		{
-			UpdateForChangedHeightmaps(Pair.Key, Pair.Value);
-			UpdateForChangedWeightmaps(Pair.Key, Pair.Value);
-		}
+		UpdateForChangedHeightmaps(InComponentReadbackResults);
+		UpdateForChangedWeightmaps(InComponentReadbackResults);
 
 		GetLandscapeInfo()->UpdateAllAddCollisions();
 
-		NewUpdateFlags |= UpdateCollisionAndClients(Components);
+		NewUpdateFlags |= UpdateCollisionAndClients(InComponentReadbackResults);
 	}
 
 	return NewUpdateFlags;
