Check if FSlateApplicationBase is initialized before accessing it

diff --git a/Engine/Source/Runtime/Landscape/Private/LandscapeEditLayers.cpp b/Engine/Source/Runtime/Landscape/Private/LandscapeEditLayers.cpp
--- a/Engine/Source/Runtime/Landscape/Private/LandscapeEditLayers.cpp
+++ b/Engine/Source/Runtime/Landscape/Private/LandscapeEditLayers.cpp
@@ -7753,7 +7753,7 @@ void ALandscape::UpdateLayersContent(bool bInWaitForStreaming, bool bInSkipMonit
 	bResourcesReady &= PrepareLayersBrushResources(bInWaitForStreaming);
 	if (!bResourcesReady)
 	{
-		if (!bInWaitForStreaming)
+		if (FSlateApplicationBase::IsInitialized() && !bInWaitForStreaming)
 		{
 			static constexpr double TimeBeforeDisplayingNotification = 3.0;
 			if (WaitingForResourcesStartTime < 0.0)
