Refactored Landscape Nanite rebuilds to (by default) only occur during saving, cooking, or explicitly clicking the "Rebuild Data" button on the top level Landscape actor. This makes the sculpting/editing tools usable without massive rebuild hitches. By default now, any edits will invalidate Nanite and fall back to regular landscape rendering until the representation is updated again. This is a hopefully short term approach until we have proper asynchronous builds happening without user involvement. Live rebuilds (with hitching) can be re-enabled now with the landscape.LiveRebuildNaniteOnModification cvar.

diff --git a/Engine/Source/Editor/DetailCustomizations/Private/LandscapeUIDetails.cpp b/Engine/Source/Editor/DetailCustomizations/Private/LandscapeUIDetails.cpp
--- a/Engine/Source/Editor/DetailCustomizations/Private/LandscapeUIDetails.cpp
+++ b/Engine/Source/Editor/DetailCustomizations/Private/LandscapeUIDetails.cpp
@@ -14,6 +14,8 @@
 #include "Widgets/Input/SCheckBox.h"
 #include "Misc/MessageDialog.h"
 #include "Editor.h"
+#include "LandscapeInfo.h"
+#include "LandscapeStreamingProxy.h"
 
 #define LOCTEXT_NAMESPACE "FLandscapeUIDetails"
 
@@ -48,16 +50,16 @@ void FLandscapeUIDetails::CustomizeDetails( IDetailLayoutBuilder& DetailBuilder
 			TSharedRef<IPropertyHandle> ComponentScreenSizeToUseSubSectionsProp = DetailBuilder.GetProperty(GET_MEMBER_NAME_CHECKED(ALandscapeProxy, ComponentScreenSizeToUseSubSections));
 			DetailBuilder.HideProperty(ComponentScreenSizeToUseSubSectionsProp);
 		}
-							   
-		TSharedRef<IPropertyHandle> PropertyHandle = DetailBuilder.GetProperty(GET_MEMBER_NAME_CHECKED(ALandscape, bCanHaveLayersContent));
-		DetailBuilder.HideProperty(PropertyHandle);
+
+		TSharedRef<IPropertyHandle> CanHaveLayersPropertyHandle = DetailBuilder.GetProperty(GET_MEMBER_NAME_CHECKED(ALandscape, bCanHaveLayersContent));
+		DetailBuilder.HideProperty(CanHaveLayersPropertyHandle);
 		const FText DisplayAndFilterText(LOCTEXT("LandscapeToggleLayerName", "Enable Edit Layers"));
 		const FText ToolTipText(LOCTEXT("LandscapeToggleLayerToolTip", "Toggle whether or not to support edit layers on this Landscape. Toggling this will clear the undo stack."));
-		DetailBuilder.AddCustomRowToCategory(PropertyHandle, DisplayAndFilterText)
+		DetailBuilder.AddCustomRowToCategory(CanHaveLayersPropertyHandle, DisplayAndFilterText)
 		.RowTag(TEXT("EnableEditLayers"))
 		.NameContent()
 		[
-			PropertyHandle->CreatePropertyNameWidget(DisplayAndFilterText, ToolTipText)
+			CanHaveLayersPropertyHandle->CreatePropertyNameWidget(DisplayAndFilterText, ToolTipText)
 		]
 		.ValueContent()
 		[
@@ -77,6 +79,35 @@ void FLandscapeUIDetails::CustomizeDetails( IDetailLayoutBuilder& DetailBuilder
 				}
 			})
 		];
+
+		TSharedRef<IPropertyHandle> EnableNanitePropertyHandle = DetailBuilder.GetProperty(GET_MEMBER_NAME_CHECKED(ALandscape, bEnableNanite));
+		DetailBuilder.AddCustomRowToCategory(EnableNanitePropertyHandle, LOCTEXT("RebuildNaniteData", "Rebuild Data"))
+			.ValueContent()
+			[
+				SNew(SButton)
+				.Text(LOCTEXT("RebuildNaniteData", "Rebuild Data"))
+				.HAlign(HAlign_Center)
+				.ToolTipText(LOCTEXT("RebuildNaniteDataTooltip", "Rebuilds the Nanite mesh representation from the Landscape data"))
+				.OnClicked_Lambda([Landscape]()
+				{
+					if (Landscape.IsValid())
+					{
+						Landscape->UpdateNaniteRepresentation();
+						Landscape->UpdateRenderingMethod();
+
+						ULandscapeInfo* LandscapeInfo = Landscape->GetLandscapeInfo();
+						if (LandscapeInfo != nullptr)
+						{
+							for (ALandscapeProxy* Proxy : LandscapeInfo->Proxies)
+							{
+								Proxy->UpdateNaniteRepresentation();
+								Proxy->UpdateRenderingMethod();
+							}
+						}
+					}
+					return FReply::Handled();
+				})
+			];
 	}
 }
 
diff --git a/Engine/Source/Runtime/Landscape/Classes/LandscapeProxy.h b/Engine/Source/Runtime/Landscape/Classes/LandscapeProxy.h
--- a/Engine/Source/Runtime/Landscape/Classes/LandscapeProxy.h
+++ b/Engine/Source/Runtime/Landscape/Classes/LandscapeProxy.h
@@ -892,6 +892,8 @@ public:
 	/* Invalidate the precomputed grass and baked texture data on all components */
 	LANDSCAPE_API void InvalidateGeneratedComponentData(bool bInvalidateLightingCache = false);
 
+	LANDSCAPE_API void UpdateRenderingMethod();
+
 #if WITH_EDITOR
 	/** Update Grass maps */
 	void UpdateGrassData(bool bInShouldMarkDirty = false, struct FScopedSlowTask* InSlowTask = nullptr);
@@ -1095,8 +1097,15 @@ public:
 	/** Generate mobile platform data if it's missing or outdated */
 	LANDSCAPE_API void CheckGenerateMobilePlatformData(bool bIsCooking, const ITargetPlatform* TargetPlatform);
 
-	/** Generate Nanite platform data if it's missing or outdated */
-	LANDSCAPE_API void CheckGenerateNanitePlatformData(bool bIsCooking, const ITargetPlatform* TargetPlatform);
+	/** Update Nanite representation if it's missing or outdated */
+	LANDSCAPE_API void UpdateNaniteRepresentation();
+
+	/** 
+	* Invalidate and disable Nanite representation until a subsequent rebuild occurs
+	*
+	* @param bCheckContentId - If true, only invalidate when the content Id of the proxy mismatches with the Nanite representation
+	*/
+	LANDSCAPE_API void InvalidateNaniteRepresentation(bool bCheckContentId = true);
 	
 	/** @return Current size of bounding rectangle in quads space */
 	LANDSCAPE_API FIntRect GetBoundingRect() const;
@@ -1213,9 +1222,6 @@ private:
 	void UpdateGrassDataStatus(TSet<UTexture2D*>* OutCurrentForcedStreamedTextures, TSet<UTexture2D*>* OutDesiredForcedStreamedTextures, TSet<ULandscapeComponent*>* OutComponentsNeedingGrassMapRender, TSet<ULandscapeComponent*>* OutOutdatedComponents, bool bInEnableForceResidentFlag, int32* OutOutdatedGrassMaps = nullptr) const;
 #endif
 
-protected:
-	void UpdateRenderingMethod();
-
 #if WITH_EDITORONLY_DATA
 public:
 	static const TArray<ALandscapeProxy*>& GetLandscapeProxies() { return LandscapeProxies; }
diff --git a/Engine/Source/Runtime/Landscape/Private/Landscape.cpp b/Engine/Source/Runtime/Landscape/Private/Landscape.cpp
--- a/Engine/Source/Runtime/Landscape/Private/Landscape.cpp
+++ b/Engine/Source/Runtime/Landscape/Private/Landscape.cpp
@@ -172,6 +172,16 @@ static FAutoConsoleVariableRef CVarRenderCaptureNextHeightmapRenders(
 	RenderCaptureNextHeightmapRenders,
 	TEXT("Trigger a render capture during the next N RenderHeightmap draws"));
 
+#if WITH_EDITOR
+
+int32 LiveRebuildNaniteOnModification = 0;
+static FAutoConsoleVariableRef CVarLiveRebuildNaniteOnModification(
+	TEXT("landscape.LiveRebuildNaniteOnModification"),
+	LiveRebuildNaniteOnModification,
+	TEXT("Trigger a rebuild of Nanite representation immediately when a modification is performed"));
+
+#endif
+
 
 ULandscapeComponent::ULandscapeComponent(const FObjectInitializer& ObjectInitializer)
 	: Super(ObjectInitializer)
@@ -294,11 +304,11 @@ void ALandscapeProxy::CheckGenerateMobilePlatformData(bool bIsCooking, const ITa
 	}
 }
 
-void ALandscapeProxy::CheckGenerateNanitePlatformData(bool bIsCooking, const ITargetPlatform* TargetPlatform)
+void ALandscapeProxy::UpdateNaniteRepresentation()
 {
-	if (IsNaniteEnabled())
+	if (IsNaniteEnabled() && LandscapeComponents.Num() > 0)
 	{
-		FGuid NaniteContentId = GetNaniteContentId();
+		const FGuid NaniteContentId = GetNaniteContentId();
 		if (NaniteComponent == nullptr)
 		{
 			NaniteComponent = NewObject<ULandscapeNaniteComponent>(this);
@@ -321,11 +331,22 @@ void ALandscapeProxy::CheckGenerateNanitePlatformData(bool bIsCooking, const ITa
 
 		NaniteComponent->UpdatedSharedPropertiesFromActor();
 	}
-	else if (NaniteComponent != nullptr)
+	else
 	{
-		NaniteComponent->DetachFromComponent(FDetachmentTransformRules::KeepRelativeTransform);
-		NaniteComponent->DestroyComponent();
-		NaniteComponent = nullptr;
+		InvalidateNaniteRepresentation(/* bCheckContentId = */ false);
+	}
+}
+
+void ALandscapeProxy::InvalidateNaniteRepresentation(bool bCheckContentId)
+{
+	if (NaniteComponent != nullptr)
+	{
+		if (!bCheckContentId || NaniteComponent->GetProxyContentId() != GetNaniteContentId())
+		{
+			NaniteComponent->DetachFromComponent(FDetachmentTransformRules::KeepRelativeTransform);
+			NaniteComponent->DestroyComponent();
+			NaniteComponent = nullptr;
+		}
 	}
 }
 
@@ -2866,11 +2887,11 @@ void ALandscapeProxy::Serialize(FArchive& Ar)
 #endif
 
 #if WITH_EDITOR
-	if (Ar.IsCooking() && !HasAnyFlags(RF_ClassDefaultObject))
+	if (!HasAnyFlags(RF_ClassDefaultObject))
 	{
-		//if (UseNaniteLandscapeMesh(Ar.CookingTarget())) // TODO: WIP
+		if (Ar.IsCooking() || (Ar.IsSaving() && !Ar.IsTransacting()))
 		{
-			CheckGenerateNanitePlatformData(Ar.IsCooking(), Ar.CookingTarget());
+			UpdateNaniteRepresentation();
 		}
 	}
 #endif
@@ -3285,11 +3306,7 @@ void ALandscapeProxy::PostLoad()
 
 	if (World && !HasAnyFlags(RF_ClassDefaultObject) && !FPlatformProperties::RequiresCookedData())
 	{
-		// If we're loading on a platform that doesn't require cooked data, but defaults to a Nanite feature level, generate or preload data from the DDC
-		if (UseNaniteLandscapeMesh(GShaderPlatformForFeatureLevel[World->FeatureLevel.GetValue()]))
-		{
-			CheckGenerateNanitePlatformData(/*bIsCooking = */ false, /*TargetPlatform = */ nullptr);
-		}
+		UpdateNaniteRepresentation();
 	}
 #endif // WITH_EDITOR
 
@@ -5088,7 +5105,14 @@ void ALandscapeProxy::InvalidateGeneratedComponentData(const TArray<ULandscapeCo
 		Iter.Key()->FlushGrassComponents(&Iter.Value());
 
 	#if WITH_EDITOR
-		Iter.Key()->CheckGenerateNanitePlatformData(/*bIsCooking = */ false, /*TargetPlatform = */ nullptr);
+		if (LiveRebuildNaniteOnModification != 0)
+		{
+			Iter.Key()->UpdateNaniteRepresentation();
+		}
+		else
+		{
+			Iter.Key()->InvalidateNaniteRepresentation();
+		}
 
 		FLandscapeProxyComponentDataChangedParams ChangeParams(Iter.Value());
 		Iter.Key()->OnComponentDataChanged.Broadcast(Iter.Key(), ChangeParams);
@@ -5105,6 +5129,11 @@ void ALandscapeProxy::InvalidateGeneratedComponentData(const TSet<ULandscapeComp
 
 void ALandscapeProxy::UpdateRenderingMethod()
 {
+	if (LandscapeComponents.Num() == 0)
+	{
+		return;
+	}
+
 	bool bAllowNanite = false;
 	if (ALandscape* LandscapeActor = GetLandscapeActor())
 	{
diff --git a/Engine/Source/Runtime/Landscape/Private/LandscapeEdit.cpp b/Engine/Source/Runtime/Landscape/Private/LandscapeEdit.cpp
--- a/Engine/Source/Runtime/Landscape/Private/LandscapeEdit.cpp
+++ b/Engine/Source/Runtime/Landscape/Private/LandscapeEdit.cpp
@@ -104,6 +104,8 @@ LANDSCAPE_API bool GDisableUpdateLandscapeMaterialInstances = false;
 // Channel remapping
 extern const size_t ChannelOffsets[4];
 
+extern int32 LiveRebuildNaniteOnModification;
+
 ULandscapeLayerInfoObject* ALandscapeProxy::VisibilityLayer = nullptr;
 
 void ULandscapeComponent::Init(int32 InBaseX, int32 InBaseY, int32 InComponentSizeQuads, int32 InNumSubsections, int32 InSubsectionSizeQuads)
@@ -5100,7 +5102,15 @@ void ALandscapeProxy::PostEditChangeProperty(FPropertyChangedEvent& PropertyChan
 	}
 	if (GIsEditor && PropertyName == FName(TEXT("bEnableNanite")))
 	{
-		CheckGenerateNanitePlatformData(/*bIsCooking = */ false, /*TargetPlatform = */ nullptr);
+		if (LiveRebuildNaniteOnModification != 0)
+		{
+			UpdateNaniteRepresentation();
+		}
+		else
+		{
+			InvalidateNaniteRepresentation();
+		}
+
 		UpdateRenderingMethod();
 		MarkComponentsRenderStateDirty();
 	}
@@ -5215,11 +5225,16 @@ void ALandscapeStreamingProxy::PostEditChangeProperty(FPropertyChangedEvent& Pro
 					}
 				}
 
-				if (UseNaniteLandscapeMesh(GShaderPlatformForFeatureLevel[World->FeatureLevel]))
+				if (LiveRebuildNaniteOnModification != 0)
 				{
-					CheckGenerateNanitePlatformData(/*bIsCooking = */ false, /*TargetPlatform = */ nullptr);
-					UpdateRenderingMethod();
+					UpdateNaniteRepresentation();
+				}
+				else
+				{
+					InvalidateNaniteRepresentation();
 				}
+
+				UpdateRenderingMethod();
 			}
 		}
 	}
@@ -5588,7 +5603,15 @@ void ALandscape::PostEditChangeProperty(FPropertyChangedEvent& PropertyChangedEv
 
 		// Generate Nanite data for a landscape with components on it, and recreate render state
 		// Streaming proxies won't be built here, but the bPropagateToProxies path will.
-		CheckGenerateNanitePlatformData(/*bIsCooking = */ false, /*TargetPlatform = */ nullptr);
+		if (LiveRebuildNaniteOnModification != 0)
+		{
+			UpdateNaniteRepresentation();
+		}
+		else
+		{
+			InvalidateNaniteRepresentation();
+		}
+
 		UpdateRenderingMethod();
 		MarkComponentsRenderStateDirty();
 	}
@@ -5660,11 +5683,16 @@ void ALandscape::PostEditChangeProperty(FPropertyChangedEvent& PropertyChangedEv
 						}
 					}
 
-					if (UseNaniteLandscapeMesh(GShaderPlatformForFeatureLevel[World->FeatureLevel]))
+					if (LiveRebuildNaniteOnModification != 0)
 					{
-						CheckGenerateNanitePlatformData(/*bIsCooking = */ false, /*TargetPlatform = */ nullptr);
-						UpdateRenderingMethod();
+						UpdateNaniteRepresentation();
 					}
+					else
+					{
+						InvalidateNaniteRepresentation();
+					}
+
+					UpdateRenderingMethod();
 				}
 			}
 		}
