Support rendering the Landscape editor tool/gizmo on top of Nanite

diff --git a/Engine/Source/Runtime/Landscape/Private/LandscapeRender.cpp b/Engine/Source/Runtime/Landscape/Private/LandscapeRender.cpp
--- a/Engine/Source/Runtime/Landscape/Private/LandscapeRender.cpp
+++ b/Engine/Source/Runtime/Landscape/Private/LandscapeRender.cpp
@@ -1725,9 +1725,21 @@ FPrimitiveViewRelevance FLandscapeComponentSceneProxy::GetViewRelevance(const FS
 	{
 		Result.bShadowRelevance = false;
 		Result.bVelocityRelevance = false;
-		Result.bRenderInMainPass = false;
 		Result.bRenderCustomDepth = false;
 		Result.bTranslucentSelfShadow = false;
+	#if WITH_EDITOR
+		if (GLandscapeEditModeActive && Result.bDynamicRelevance)
+		{
+			if (!View->bIsVirtualTexture || View->bIsSceneCapture)
+			{
+				Result.bStaticRelevance = false;
+			}
+		}
+		else
+	#endif
+		{
+			Result.bRenderInMainPass = false;
+		}
 	}
 
 	return Result;
