New "AffectsLandscape"  ActorDesc property: Implemented for Water, PatchComponent and LandscapeSplineActor
- Deprecated LandscapeSplineActorDesc as it is no longer needed with this property

diff --git a/Engine/Plugins/Experimental/LandscapePatch/Source/LandscapePatch/Private/LandscapePatchComponent.cpp b/Engine/Plugins/Experimental/LandscapePatch/Source/LandscapePatch/Private/LandscapePatchComponent.cpp
--- a/Engine/Plugins/Experimental/LandscapePatch/Source/LandscapePatch/Private/LandscapePatchComponent.cpp
+++ b/Engine/Plugins/Experimental/LandscapePatch/Source/LandscapePatch/Private/LandscapePatchComponent.cpp
@@ -158,6 +158,16 @@ void ULandscapePatchComponent::OnRegister()
 	RequestLandscapeUpdate();
 }
 
+void ULandscapePatchComponent::GetActorDescProperties(FPropertyPairsMap& PropertyPairsMap) const
+{
+	Super::GetActorDescProperties(PropertyPairsMap);
+
+	if (Landscape.IsValid())
+	{
+		PropertyPairsMap.AddProperty(ALandscape::AffectsLandscapeActorDescProperty, *Landscape->GetLandscapeGuid().ToString());
+	}
+}
+
 void ULandscapePatchComponent::OnUpdateTransform(EUpdateTransformFlags UpdateTransformFlags, ETeleportType Teleport)
 {
 	Super::OnUpdateTransform(UpdateTransformFlags, Teleport);
diff --git a/Engine/Plugins/Experimental/LandscapePatch/Source/LandscapePatch/Public/LandscapePatchComponent.h b/Engine/Plugins/Experimental/LandscapePatch/Source/LandscapePatch/Public/LandscapePatchComponent.h
--- a/Engine/Plugins/Experimental/LandscapePatch/Source/LandscapePatch/Public/LandscapePatchComponent.h
+++ b/Engine/Plugins/Experimental/LandscapePatch/Source/LandscapePatch/Public/LandscapePatchComponent.h
@@ -49,6 +49,7 @@ public:
 	virtual void OnComponentCreated() override;
 	virtual void OnComponentDestroyed(bool bDestroyingHierarchy) override;
 	virtual void OnRegister() override;
+	virtual void GetActorDescProperties(FPropertyPairsMap& PropertyPairsMap) const override;
 
 	// UObject
 	virtual void PostEditChangeProperty(FPropertyChangedEvent& PropertyChangedEvent) override;
diff --git a/Engine/Plugins/Experimental/Water/Source/Runtime/Private/WaterBodyActor.cpp b/Engine/Plugins/Experimental/Water/Source/Runtime/Private/WaterBodyActor.cpp
--- a/Engine/Plugins/Experimental/Water/Source/Runtime/Private/WaterBodyActor.cpp
+++ b/Engine/Plugins/Experimental/Water/Source/Runtime/Private/WaterBodyActor.cpp
@@ -28,6 +28,7 @@
 #include "WaterIconHelper.h"
 #include "Components/BillboardComponent.h"
 #include "Editor/EditorEngine.h"
+#include "Landscape.h"
 #endif
 
 #define LOCTEXT_NAMESPACE "Water"
@@ -611,6 +612,16 @@ bool AWaterBody::SetIsHiddenEdLayer(bool bIsHiddenEdLayer)
 	}
 	return false;
 }
+
+void AWaterBody::GetActorDescProperties(FPropertyPairsMap& PropertyPairsMap) const
+{
+	Super::GetActorDescProperties(PropertyPairsMap);
+	if (AffectsLandscape())
+	{
+		PropertyPairsMap.AddProperty(ALandscape::AffectsLandscapeActorDescProperty);
+	}
+}
+
 #endif // WITH_EDITOR
 
 void AWaterBody::UnregisterAllComponents(bool bForReregister)
diff --git a/Engine/Plugins/Experimental/Water/Source/Runtime/Private/WaterBodyIslandActor.cpp b/Engine/Plugins/Experimental/Water/Source/Runtime/Private/WaterBodyIslandActor.cpp
--- a/Engine/Plugins/Experimental/Water/Source/Runtime/Private/WaterBodyIslandActor.cpp
+++ b/Engine/Plugins/Experimental/Water/Source/Runtime/Private/WaterBodyIslandActor.cpp
@@ -20,6 +20,7 @@
 #include "WaterIconHelper.h"
 #include "WaterSubsystem.h"
 #include "WaterModule.h"
+#include "Landscape.h"
 #endif // WITH_EDITOR
 
 // ----------------------------------------------------------------------------------
@@ -165,6 +166,15 @@ void AWaterBodyIsland::PostLoad()
 }
 
 #if WITH_EDITOR
+void AWaterBodyIsland::GetActorDescProperties(FPropertyPairsMap& PropertyPairsMap) const
+{
+	Super::GetActorDescProperties(PropertyPairsMap);
+	if (AffectsLandscape())
+	{
+		PropertyPairsMap.AddProperty(ALandscape::AffectsLandscapeActorDescProperty);
+	}
+}
+
 void AWaterBodyIsland::UpdateOverlappingWaterBodyComponents()
 {
 	TArray<FOverlapResult> Overlaps;
diff --git a/Engine/Plugins/Experimental/Water/Source/Runtime/Public/WaterBodyActor.h b/Engine/Plugins/Experimental/Water/Source/Runtime/Public/WaterBodyActor.h
--- a/Engine/Plugins/Experimental/Water/Source/Runtime/Public/WaterBodyActor.h
+++ b/Engine/Plugins/Experimental/Water/Source/Runtime/Public/WaterBodyActor.h
@@ -73,6 +73,7 @@ public:
 	virtual void SetActorHiddenInGame(bool bNewHidden) override;
 	virtual void SetIsTemporarilyHiddenInEditor(bool bIsHidden) override;
 	virtual bool SetIsHiddenEdLayer(bool bIsHiddenEdLayer) override;
+	virtual void GetActorDescProperties(FPropertyPairsMap& PropertyPairsMap) const override;
 #endif // WITH_EDITOR
 	
 	/** Returns the type of body */
diff --git a/Engine/Plugins/Experimental/Water/Source/Runtime/Public/WaterBodyIslandActor.h b/Engine/Plugins/Experimental/Water/Source/Runtime/Public/WaterBodyIslandActor.h
--- a/Engine/Plugins/Experimental/Water/Source/Runtime/Public/WaterBodyIslandActor.h
+++ b/Engine/Plugins/Experimental/Water/Source/Runtime/Public/WaterBodyIslandActor.h
@@ -45,6 +45,7 @@ public:
 	virtual void PostLoad() override;
 	
 #if WITH_EDITOR
+	virtual void GetActorDescProperties(FPropertyPairsMap& PropertyPairsMap) const override;
 	void UpdateOverlappingWaterBodyComponents();
 	void UpdateActorIcon();
 #endif // WITH_EDITOR
diff --git a/Engine/Source/Runtime/Core/Public/UObject/UE5MainStreamObjectVersion.h b/Engine/Source/Runtime/Core/Public/UObject/UE5MainStreamObjectVersion.h
--- a/Engine/Source/Runtime/Core/Public/UObject/UE5MainStreamObjectVersion.h
+++ b/Engine/Source/Runtime/Core/Public/UObject/UE5MainStreamObjectVersion.h
@@ -241,6 +241,10 @@ struct CORE_API FUE5MainStreamObjectVersion
 
 		// Add IAnimationDataModel interface and replace UObject based representation for Animation Assets
 		AnimationDataModelInterface,
+
+		// Deprecate LandscapeSplineActorDesc
+		LandscapeSplineActorDescDeprecation,
+
 		// -----<new versions can be added above this line>-------------------------------------------------
 		VersionPlusOne,
 		LatestVersion = VersionPlusOne - 1
diff --git a/Engine/Source/Runtime/Engine/Private/WorldPartition/Landscape/LandscapeSplineActorDesc.cpp b/Engine/Source/Runtime/Engine/Private/WorldPartition/Landscape/LandscapeSplineActorDesc.cpp
deleted file mode 100644
diff --git a/Engine/Source/Runtime/Engine/Public/WorldPartition/Landscape/LandscapeSplineActorDesc.h b/Engine/Source/Runtime/Engine/Public/WorldPartition/Landscape/LandscapeSplineActorDesc.h
deleted file mode 100644
diff --git a/Engine/Source/Runtime/Landscape/Classes/Landscape.h b/Engine/Source/Runtime/Landscape/Classes/Landscape.h
--- a/Engine/Source/Runtime/Landscape/Classes/Landscape.h
+++ b/Engine/Source/Runtime/Landscape/Classes/Landscape.h
@@ -218,6 +218,8 @@ public:
 	//~ End ALandscapeProxy Interface
 
 #if WITH_EDITOR
+	static LANDSCAPE_API FName AffectsLandscapeActorDescProperty;
+
 	LANDSCAPE_API bool HasAllComponent(); // determine all component is in this actor
 	
 	// Include Components with overlapped vertices
diff --git a/Engine/Source/Runtime/Landscape/Classes/LandscapeSplineActor.h b/Engine/Source/Runtime/Landscape/Classes/LandscapeSplineActor.h
--- a/Engine/Source/Runtime/Landscape/Classes/LandscapeSplineActor.h
+++ b/Engine/Source/Runtime/Landscape/Classes/LandscapeSplineActor.h
@@ -27,7 +27,7 @@ public:
 		
 #if WITH_EDITOR
 	virtual bool CanChangeIsSpatiallyLoadedFlag() const override { return true; }
-	virtual TUniquePtr<class FWorldPartitionActorDesc> CreateClassActorDesc() const override;
+	virtual void GetActorDescProperties(FPropertyPairsMap& PropertyPairsMap) const override;
 
 	void GetSharedProperties(ULandscapeInfo* InLandscapeInfo);
 
@@ -44,5 +44,3 @@ public:
 	virtual void CreateSplineComponent(const FVector& Scale3D) override { check(false); }
 #endif
 };
-
-DEFINE_ACTORDESC_TYPE(ALandscapeSplineActor, FLandscapeSplineActorDesc);
\ No newline at end of file
diff --git a/Engine/Source/Runtime/Landscape/Private/Landscape.cpp b/Engine/Source/Runtime/Landscape/Private/Landscape.cpp
--- a/Engine/Source/Runtime/Landscape/Private/Landscape.cpp
+++ b/Engine/Source/Runtime/Landscape/Private/Landscape.cpp
@@ -86,7 +86,6 @@ Landscape.cpp: Terrain rendering
 #include "WorldPartition/WorldPartitionHelpers.h"
 #include "WorldPartition/WorldPartitionHandle.h"
 #include "WorldPartition/Landscape/LandscapeActorDesc.h"
-#include "WorldPartition/Landscape/LandscapeSplineActorDesc.h"
 #if WITH_EDITOR
 #include "Rendering/StaticLightingSystemInterface.h"
 #include "Misc/ScopedSlowTask.h"
diff --git a/Engine/Source/Runtime/Landscape/Private/LandscapeEdit.cpp b/Engine/Source/Runtime/Landscape/Private/LandscapeEdit.cpp
--- a/Engine/Source/Runtime/Landscape/Private/LandscapeEdit.cpp
+++ b/Engine/Source/Runtime/Landscape/Private/LandscapeEdit.cpp
@@ -75,7 +75,6 @@ LandscapeEdit.cpp: Landscape editing
 #include "WorldPartition/WorldPartitionHelpers.h"
 #include "WorldPartition/WorldPartitionActorDesc.h"
 #include "WorldPartition/Landscape/LandscapeActorDesc.h"
-#include "WorldPartition/Landscape/LandscapeSplineActorDesc.h"
 #include "ActorPartition/ActorPartitionSubsystem.h"
 #include "LandscapeUtils.h"
 #include "LandscapeSplineActor.h"
@@ -101,6 +100,8 @@ static TAutoConsoleVariable<int32> CVarMobileCompressLanscapeWeightMaps(
 // Instead, one call per component is done at the end
 LANDSCAPE_API bool GDisableUpdateLandscapeMaterialInstances = false;
 
+LANDSCAPE_API FName ALandscape::AffectsLandscapeActorDescProperty(TEXT("AffectsLandscape"));
+
 // Channel remapping
 extern const size_t ChannelOffsets[4];
 
@@ -5251,22 +5252,31 @@ void ALandscapeStreamingProxy::PostRegisterAllComponents()
 		check(LandscapeInfo);
 		if (GEditor && !GetWorld()->IsGameWorld())
 		{
-			if (UWorldPartition* WorldPartition = GetWorld()->GetWorldPartition())
+			if (UWorldPartition* WorldPartition = GetWorld()->GetWorldPartition(); WorldPartition->IsInitialized())
 			{
 				const FVector ActorLocation = GetActorLocation();
-				const FBox Bounds(ActorLocation, ActorLocation + (GridSize * LandscapeInfo->DrawScale));
+				FBox Bounds(ActorLocation, ActorLocation + (GridSize * LandscapeInfo->DrawScale));
 
-				FWorldPartitionHelpers::ForEachActorDesc<ALandscapeSplineActor>(WorldPartition, [this, WorldPartition, &Bounds](const FWorldPartitionActorDesc* ActorDesc) mutable
-				{
-					FLandscapeSplineActorDesc* LandscapeSplineActorDesc = (FLandscapeSplineActorDesc*)ActorDesc;
+				// all actors that intersect Landscape in 2D need to be considered
+				Bounds.Min.Z = -HALF_WORLD_MAX;
+				Bounds.Max.Z = HALF_WORLD_MAX;
 
-					if (LandscapeSplineActorDesc->LandscapeGuid == LandscapeGuid)
+				FWorldPartitionHelpers::ForEachIntersectingActorDesc(WorldPartition, Bounds, [this, WorldPartition](const FWorldPartitionActorDesc* ActorDesc) mutable
+				{
+					FName PropertyValue;
+					if(ActorDesc->GetProperty(ALandscape::AffectsLandscapeActorDescProperty, &PropertyValue))
 					{
-						if (Bounds.IntersectXY(ActorDesc->GetBounds()))
+						// If no Guid specified then consider actor as affecting all landscapes
+						if (PropertyValue.IsNone())
+						{
+							ActorDescReferences.Add(FWorldPartitionReference(WorldPartition, ActorDesc->GetGuid()));
+						}
+						else if(FGuid ParsedGuid; FGuid::Parse(PropertyValue.ToString(), ParsedGuid) && ParsedGuid == LandscapeGuid)
 						{
 							ActorDescReferences.Add(FWorldPartitionReference(WorldPartition, ActorDesc->GetGuid()));
 						}
 					}
+										
 					return true;
 				});
 			}
@@ -5407,21 +5417,24 @@ bool ULandscapeInfo::CanDeleteLandscape(FText& OutReason) const
 		{
 			FWorldPartitionHelpers::ForEachActorDesc<ALandscapeSplineActor>(WorldPartition, [this, &UndeletedSplineCount](const FWorldPartitionActorDesc* ActorDesc)
 			{
-				FLandscapeSplineActorDesc* LandscapeSplineActorDesc = (FLandscapeSplineActorDesc*)ActorDesc;
-
-				if (LandscapeSplineActorDesc->LandscapeGuid == LandscapeGuid)
+				FName AffectsLandscapeProperty;
+				if (ActorDesc->GetProperty(ALandscape::AffectsLandscapeActorDescProperty, &AffectsLandscapeProperty))
 				{
-					ALandscapeSplineActor* SplineActor = Cast<ALandscapeSplineActor>(LandscapeSplineActorDesc->GetActor());
-		
-					// If SplineActor is null then it is not loaded/deleted. If it's loaded then it needs to be pending kill.
-					if (!SplineActor)
-					{
-						++UndeletedSplineCount;
-					}
-					else
+					FGuid ParsedLandscapeGuid;
+					if (FGuid::Parse(AffectsLandscapeProperty.ToString(), ParsedLandscapeGuid) && ParsedLandscapeGuid == LandscapeGuid)
 					{
-						// If Actor is loaded it should be Registered and not pending kill (already accounted for) or pending kill (deleted)
-						check(SplineActors.Contains(SplineActor) == IsValidChecked(SplineActor));
+						ALandscapeSplineActor* SplineActor = Cast<ALandscapeSplineActor>(ActorDesc->GetActor());
+
+						// If SplineActor is null then it is not loaded/deleted. If it's loaded then it needs to be pending kill.
+						if (!SplineActor)
+						{
+							++UndeletedSplineCount;
+						}
+						else
+						{
+							// If Actor is loaded it should be Registered and not pending kill (already accounted for) or pending kill (deleted)
+							check(SplineActors.Contains(SplineActor) == IsValidChecked(SplineActor));
+						}
 					}
 				}
 
diff --git a/Engine/Source/Runtime/Landscape/Private/LandscapeModule.cpp b/Engine/Source/Runtime/Landscape/Private/LandscapeModule.cpp
--- a/Engine/Source/Runtime/Landscape/Private/LandscapeModule.cpp
+++ b/Engine/Source/Runtime/Landscape/Private/LandscapeModule.cpp
@@ -19,6 +19,11 @@
 #include "Engine/Texture2D.h"
 #include "EngineUtils.h"
 
+#if WITH_EDITOR
+#include "WorldPartition/WorldPartitionActorDesc.h"
+#include "UObject/UE5MainStreamObjectVersion.h"
+#endif
+
 
 // Register the custom version with core
 FCustomVersionRegistration GRegisterLandscapeCustomVersion(FLandscapeCustomVersion::GUID, FLandscapeCustomVersion::LatestVersion, TEXT("Landscape"));
@@ -208,6 +213,24 @@ void FLandscapeModule::StartupModule()
 
 	FCoreDelegates::OnPostEngineInit.AddRaw(this, &FLandscapeModule::OnPostEngineInit);
 	FCoreDelegates::OnEnginePreExit.AddRaw(this, &FLandscapeModule::OnEnginePreExit);
+
+#if WITH_EDITOR
+	// Register LandscapeSplineActorDesc Deprecation
+	FWorldPartitionActorDesc::RegisterActorDescDeprecator(ALandscapeSplineActor::StaticClass(), [](FArchive& Ar, FWorldPartitionActorDesc* ActorDesc)
+	{
+		check(Ar.IsLoading());
+		if (Ar.CustomVer(FUE5MainStreamObjectVersion::GUID) < FUE5MainStreamObjectVersion::AddedLandscapeSplineActorDesc)
+		{
+			ActorDesc->AddProperty(ALandscape::AffectsLandscapeActorDescProperty);
+		}
+		else if (Ar.CustomVer(FUE5MainStreamObjectVersion::GUID) < FUE5MainStreamObjectVersion::LandscapeSplineActorDescDeprecation)
+		{
+			FGuid LandscapeGuid;
+			Ar << LandscapeGuid;
+			ActorDesc->AddProperty(ALandscape::AffectsLandscapeActorDescProperty, *LandscapeGuid.ToString());
+		}
+	});
+#endif
 }
 
 void FLandscapeModule::OnPostEngineInit()
diff --git a/Engine/Source/Runtime/Landscape/Private/LandscapeSplineActor.cpp b/Engine/Source/Runtime/Landscape/Private/LandscapeSplineActor.cpp
--- a/Engine/Source/Runtime/Landscape/Private/LandscapeSplineActor.cpp
+++ b/Engine/Source/Runtime/Landscape/Private/LandscapeSplineActor.cpp
@@ -8,7 +8,6 @@
 
 #if WITH_EDITOR
 #include "WorldPartition/WorldPartitionActorDesc.h"
-#include "WorldPartition/Landscape/LandscapeSplineActorDesc.h"
 #endif
 
 ALandscapeSplineActor::ALandscapeSplineActor(const FObjectInitializer& ObjectInitializer)
@@ -36,9 +35,14 @@ ULandscapeInfo* ALandscapeSplineActor::GetLandscapeInfo() const
 }
 
 #if WITH_EDITOR
-TUniquePtr<FWorldPartitionActorDesc> ALandscapeSplineActor::CreateClassActorDesc() const
+void ALandscapeSplineActor::GetActorDescProperties(FPropertyPairsMap& PropertyPairsMap) const
 {
-	return TUniquePtr<FWorldPartitionActorDesc>(new FLandscapeSplineActorDesc());
+	Super::GetActorDescProperties(PropertyPairsMap);
+
+	if (LandscapeGuid.IsValid())
+	{
+		PropertyPairsMap.AddProperty(ALandscape::AffectsLandscapeActorDescProperty, *LandscapeGuid.ToString());
+	}
 }
 
 void ALandscapeSplineActor::GetSharedProperties(ULandscapeInfo* InLandscapeInfo)
