Fix crash when cooking maps with Landscape.

#rnx

diff --git a/Engine/Source/Runtime/Landscape/Private/LandscapeEdit.cpp b/Engine/Source/Runtime/Landscape/Private/LandscapeEdit.cpp
--- a/Engine/Source/Runtime/Landscape/Private/LandscapeEdit.cpp
+++ b/Engine/Source/Runtime/Landscape/Private/LandscapeEdit.cpp
@@ -5255,28 +5255,19 @@ void ALandscapeStreamingProxy::PostRegisterAllComponents()
 			if (UWorldPartition* WorldPartition = GetWorld()->GetWorldPartition(); WorldPartition->IsInitialized())
 			{
 				const FVector ActorLocation = GetActorLocation();
-				FBox Bounds(ActorLocation, ActorLocation + (GridSize * LandscapeInfo->DrawScale));
+				const FBox Bounds(ActorLocation, ActorLocation + (GridSize * LandscapeInfo->DrawScale));
 
-				// all actors that intersect Landscape in 2D need to be considered
-				Bounds.Min.Z = -HALF_WORLD_MAX;
-				Bounds.Max.Z = HALF_WORLD_MAX;
-
-				FWorldPartitionHelpers::ForEachIntersectingActorDesc(WorldPartition, Bounds, [this, WorldPartition](const FWorldPartitionActorDesc* ActorDesc) mutable
+				FWorldPartitionHelpers::ForEachActorDesc<ALandscapeSplineActor>(WorldPartition, [this, WorldPartition, &Bounds](const FWorldPartitionActorDesc* ActorDesc) mutable
 				{
 					FName PropertyValue;
-					if(ActorDesc->GetProperty(ALandscape::AffectsLandscapeActorDescProperty, &PropertyValue))
+					if (Bounds.IntersectXY(ActorDesc->GetBounds()) && ActorDesc->GetProperty(ALandscape::AffectsLandscapeActorDescProperty, &PropertyValue))
 					{
 						// If no Guid specified then consider actor as affecting all landscapes
-						if (PropertyValue.IsNone())
-						{
-							ActorDescReferences.Add(FWorldPartitionReference(WorldPartition, ActorDesc->GetGuid()));
-						}
-						else if(FGuid ParsedGuid; FGuid::Parse(PropertyValue.ToString(), ParsedGuid) && ParsedGuid == LandscapeGuid)
+						if(FGuid ParsedGuid; PropertyValue.IsNone() || (FGuid::Parse(PropertyValue.ToString(), ParsedGuid) && ParsedGuid == LandscapeGuid))
 						{
 							ActorDescReferences.Add(FWorldPartitionReference(WorldPartition, ActorDesc->GetGuid()));
 						}
 					}
-										
 					return true;
 				});
 			}
