Landscape: Fix missing nullcheck for WorldPartition.
#rnx

diff --git a/Engine/Source/Runtime/Landscape/Private/LandscapeEdit.cpp b/Engine/Source/Runtime/Landscape/Private/LandscapeEdit.cpp
--- a/Engine/Source/Runtime/Landscape/Private/LandscapeEdit.cpp
+++ b/Engine/Source/Runtime/Landscape/Private/LandscapeEdit.cpp
@@ -5252,7 +5252,7 @@ void ALandscapeStreamingProxy::PostRegisterAllComponents()
 		check(LandscapeInfo);
 		if (GEditor && !GetWorld()->IsGameWorld())
 		{
-			if (UWorldPartition* WorldPartition = GetWorld()->GetWorldPartition(); WorldPartition->IsInitialized())
+			if (UWorldPartition* WorldPartition = GetWorld()->GetWorldPartition(); WorldPartition && WorldPartition->IsInitialized())
 			{
 				const FVector ActorLocation = GetActorLocation();
 				const FBox Bounds(ActorLocation, ActorLocation + (GridSize * LandscapeInfo->DrawScale));
