Check for !IsRunningCookCommandlet() and use ForEachIntersectingActorDesc

#rnx

diff --git a/Engine/Source/Runtime/Landscape/Private/LandscapeEdit.cpp b/Engine/Source/Runtime/Landscape/Private/LandscapeEdit.cpp
--- a/Engine/Source/Runtime/Landscape/Private/LandscapeEdit.cpp
+++ b/Engine/Source/Runtime/Landscape/Private/LandscapeEdit.cpp
@@ -5250,17 +5250,21 @@ void ALandscapeStreamingProxy::PostRegisterAllComponents()
 	{
 		ULandscapeInfo* LandscapeInfo = GetLandscapeInfo();
 		check(LandscapeInfo);
-		if (GEditor && !GetWorld()->IsGameWorld())
+		if (GEditor && !GetWorld()->IsGameWorld() && !IsRunningCommandlet())
 		{
 			if (UWorldPartition* WorldPartition = GetWorld()->GetWorldPartition(); WorldPartition && WorldPartition->IsInitialized())
 			{
 				const FVector ActorLocation = GetActorLocation();
-				const FBox Bounds(ActorLocation, ActorLocation + (GridSize * LandscapeInfo->DrawScale));
+				FBox Bounds(ActorLocation, ActorLocation + (GridSize * LandscapeInfo->DrawScale));
 
-				FWorldPartitionHelpers::ForEachActorDesc<ALandscapeSplineActor>(WorldPartition, [this, WorldPartition, &Bounds](const FWorldPartitionActorDesc* ActorDesc) mutable
+				// all actors that intersect Landscape in 2D need to be considered
+				Bounds.Min.Z = -HALF_WORLD_MAX;
+				Bounds.Max.Z = HALF_WORLD_MAX;
+
+				FWorldPartitionHelpers::ForEachIntersectingActorDesc(WorldPartition, Bounds, [this, WorldPartition](const FWorldPartitionActorDesc* ActorDesc) mutable
 				{
 					FName PropertyValue;
-					if (Bounds.IntersectXY(ActorDesc->GetBounds()) && ActorDesc->GetProperty(ALandscape::AffectsLandscapeActorDescProperty, &PropertyValue))
+					if (ActorDesc->GetProperty(ALandscape::AffectsLandscapeActorDescProperty, &PropertyValue))
 					{
 						// If no Guid specified then consider actor as affecting all landscapes
 						if(FGuid ParsedGuid; PropertyValue.IsNone() || (FGuid::Parse(PropertyValue.ToString(), ParsedGuid) && ParsedGuid == LandscapeGuid))
