Made ULandscapeInfo members properly set outside of editor and for PIE worlds

#tests UEENGQA-71560

diff --git a/Engine/Source/Runtime/Landscape/Classes/LandscapeInfo.h b/Engine/Source/Runtime/Landscape/Classes/LandscapeInfo.h
--- a/Engine/Source/Runtime/Landscape/Classes/LandscapeInfo.h
+++ b/Engine/Source/Runtime/Landscape/Classes/LandscapeInfo.h
@@ -140,7 +140,7 @@ public:
 	/** Lookup map used by the "add component" tool. Only available near valid LandscapeComponents.
 	    only for use by the "add component" tool. Todo - move into the tool? */
 	TMap<FIntPoint, FLandscapeAddCollision> XYtoAddCollisionMap;
-#endif
+#endif // WITH_EDITORONLY_DATA
 
 	UPROPERTY()
 	TArray<TObjectPtr<ALandscapeStreamingProxy>> Proxies;
@@ -149,7 +149,7 @@ private:
 #if WITH_EDITORONLY_DATA
 	UPROPERTY()
 	TArray<TScriptInterface<ILandscapeSplineInterface>> SplineActors;
-#endif
+#endif // WITH_EDITORONLY_DATA
 
 	TSet<ULandscapeComponent*> SelectedComponents;
 
@@ -225,6 +225,8 @@ public:
 	 */
 	LANDSCAPE_API ALandscapeProxy* GetLandscapeProxyForLevel(ULevel* Level) const;
 
+#endif //WITH_EDITOR
+
 	/**
 	 *  Returns landscape which is spawned in the current level that was previously added to this landscape info object
 	 *  @param	bRegistered		Whether to consider only registered(visible) landscapes
@@ -238,6 +240,8 @@ public:
 	 */
 	LANDSCAPE_API ALandscapeProxy* GetLandscapeProxy() const;
 
+#if WITH_EDITOR
+
 	/** Resets all actors, proxies, components registrations */
 	LANDSCAPE_API void Reset();
 
diff --git a/Engine/Source/Runtime/Landscape/Private/Landscape.cpp b/Engine/Source/Runtime/Landscape/Private/Landscape.cpp
--- a/Engine/Source/Runtime/Landscape/Private/Landscape.cpp
+++ b/Engine/Source/Runtime/Landscape/Private/Landscape.cpp
@@ -1794,8 +1794,7 @@ void ULandscapeComponent::OnUnregister()
 		// AActor::GetWorld checks for Unreachable and BeginDestroyed
 		UWorld* World = GetLandscapeProxy()->GetWorld();
 
-		// Game worlds don't have landscape infos
-		if (World && !World->IsGameWorld())
+		if (World)
 		{
 			ULandscapeInfo* Info = GetLandscapeInfo();
 			if (Info)
@@ -2252,7 +2251,7 @@ void ALandscapeProxy::PostRegisterAllComponents()
 		if (LandscapeGuid.IsValid())
 		{
 #if WITH_EDITOR
-			if (GIsEditor && !GetWorld()->IsGameWorld())
+			if (GIsEditor)
 			{
 				// Note: This can happen when loading certain cooked assets in an editor
 				// Todo: Determine the root cause of this and fix it at a higher level!
@@ -2293,8 +2292,7 @@ void ALandscapeProxy::PostRegisterAllComponents()
 		}
 	}
 #if WITH_EDITOR
-	// Game worlds don't have landscape infos
-	if (!GetWorld()->IsGameWorld() && !IsPendingKillPending())
+	if (!IsPendingKillPending())
 	{
 		if (LandscapeGuid.IsValid())
 		{
@@ -2306,7 +2304,6 @@ void ALandscapeProxy::PostRegisterAllComponents()
 
 void ALandscapeProxy::UnregisterAllComponents(const bool bForReregister)
 {
-	// Game worlds don't have landscape infos
 	// On shutdown the world will be unreachable
 	if (GetWorld() && IsValidChecked(GetWorld()) && !GetWorld()->IsUnreachable() &&
 		// When redoing the creation of a landscape we may get UnregisterAllComponents called when
@@ -3338,7 +3335,7 @@ void ALandscapeProxy::Destroyed()
 
 	UWorld* World = GetWorld();
 
-	if (GIsEditor && !World->IsGameWorld())
+	if (GIsEditor)
 	{
 		ULandscapeInfo::RecreateLandscapeInfo(World, false);
 
@@ -3667,6 +3664,8 @@ ALandscapeProxy* ULandscapeInfo::GetLandscapeProxyForLevel(ULevel* Level) const
 	return LandscapeProxy;
 }
 
+#endif // WITH_EDITOR
+
 ALandscapeProxy* ULandscapeInfo::GetCurrentLevelLandscapeProxy(bool bRegistered) const
 {
 	ALandscapeProxy* LandscapeProxy = nullptr;
@@ -3723,6 +3722,8 @@ ALandscapeProxy* ULandscapeInfo::GetLandscapeProxy() const
 	return nullptr;
 }
 
+#if WITH_EDITOR
+
 void ULandscapeInfo::Reset()
 {
 	LandscapeActor.Reset();
@@ -3865,7 +3866,7 @@ void ULandscapeInfo::RecreateLandscapeInfo(UWorld* InWorld, bool bMapCheck)
 }
 
 
-#endif
+#endif // WITH_EDITOR
 
 ULandscapeInfo* ULandscapeInfo::Find(UWorld* InWorld, const FGuid& LandscapeGuid)
 {
@@ -3928,80 +3929,84 @@ void ULandscapeInfo::RegisterActor(ALandscapeProxy* Proxy, bool bMapCheck)
 	check(Proxy->GetLandscapeGuid().IsValid());
 	check(LandscapeGuid.IsValid());
 	
-#if WITH_EDITOR
-	if (!OwningWorld->IsGameWorld())
+	// in case this Info object is not initialized yet
+	// initialized it with properties from passed actor
+	if (GetLandscapeProxy() == nullptr)
 	{
-		// in case this Info object is not initialized yet
-		// initialized it with properties from passed actor
-		if (GetLandscapeProxy() == nullptr)
-		{
-			ComponentSizeQuads = Proxy->ComponentSizeQuads;
-			ComponentNumSubsections = Proxy->NumSubsections;
-			SubsectionSizeQuads = Proxy->SubsectionSizeQuads;
-			DrawScale = Proxy->GetRootComponent() != nullptr ? Proxy->GetRootComponent()->GetRelativeScale3D() : FVector(100.0f);
-		}
+		ComponentSizeQuads = Proxy->ComponentSizeQuads;
+		ComponentNumSubsections = Proxy->NumSubsections;
+		SubsectionSizeQuads = Proxy->SubsectionSizeQuads;
+		DrawScale = Proxy->GetRootComponent() != nullptr ? Proxy->GetRootComponent()->GetRelativeScale3D() : FVector(100.0f);
+	}
+
+	// check that passed actor matches all shared parameters
+	check(LandscapeGuid == Proxy->GetLandscapeGuid());
+	check(ComponentSizeQuads == Proxy->ComponentSizeQuads);
+	check(ComponentNumSubsections == Proxy->NumSubsections);
+	check(SubsectionSizeQuads == Proxy->SubsectionSizeQuads);
 
-		// check that passed actor matches all shared parameters
-		check(LandscapeGuid == Proxy->GetLandscapeGuid());
-		check(ComponentSizeQuads == Proxy->ComponentSizeQuads);
-		check(ComponentNumSubsections == Proxy->NumSubsections);
-		check(SubsectionSizeQuads == Proxy->SubsectionSizeQuads);
+	if (Proxy->GetRootComponent() != nullptr && !DrawScale.Equals(Proxy->GetRootComponent()->GetRelativeScale3D()))
+	{
+		UE_LOG(LogLandscape, Warning, TEXT("Landscape proxy (%s) scale (%s) does not match to main actor scale (%s)."),
+			*Proxy->GetName(), *Proxy->GetRootComponent()->GetRelativeScale3D().ToCompactString(), *DrawScale.ToCompactString());
+	}
+
+	// register
+	if (ALandscape* Landscape = Cast<ALandscape>(Proxy))
+	{
+		checkf(!LandscapeActor || LandscapeActor == Landscape, TEXT("Multiple landscapes with the same GUID detected: %s vs %s"), *LandscapeActor->GetPathName(), *Landscape->GetPathName());
+		LandscapeActor = Landscape;
 
-		if (Proxy->GetRootComponent() != nullptr && !DrawScale.Equals(Proxy->GetRootComponent()->GetRelativeScale3D()))
+#if WITH_EDITOR
+		// In world composition user is not allowed to move landscape in editor, only through WorldBrowser 
+		bool bIsLockLocation = LandscapeActor->IsLockLocation();
+		bIsLockLocation |= OwningWorld != nullptr ? OwningWorld->WorldComposition != nullptr : false;
+		LandscapeActor->SetLockLocation(bIsLockLocation);
+#endif // WITH_EDITOR
+
+		// update proxies reference actor
+		for (ALandscapeStreamingProxy* StreamingProxy : Proxies)
 		{
-			UE_LOG(LogLandscape, Warning, TEXT("Landscape proxy (%s) scale (%s) does not match to main actor scale (%s)."),
-				*Proxy->GetName(), *Proxy->GetRootComponent()->GetRelativeScale3D().ToCompactString(), *DrawScale.ToCompactString());
+			StreamingProxy->LandscapeActor = LandscapeActor;
+#if WITH_EDITOR
+			StreamingProxy->FixupSharedData(Landscape);
+#endif // WITH_EDITOR
 		}
-
-		// register
-		if (ALandscape* Landscape = Cast<ALandscape>(Proxy))
+	}
+	else
+	{
+		auto LamdbdaLowerBound = [](ALandscapeProxy* A, ALandscapeProxy* B)
 		{
-			checkf(!LandscapeActor || LandscapeActor == Landscape, TEXT("Multiple landscapes with the same GUID detected: %s vs %s"), *LandscapeActor->GetPathName(), *Landscape->GetPathName());
-			LandscapeActor = Landscape;
-			// In world composition user is not allowed to move landscape in editor, only through WorldBrowser 
-			bool bIsLockLocation = LandscapeActor->IsLockLocation();
-			bIsLockLocation |= OwningWorld != nullptr ? OwningWorld->WorldComposition != nullptr : false;
-			LandscapeActor->SetLockLocation(bIsLockLocation);
+			FIntPoint SectionBaseA = A->GetSectionBaseOffset();
+			FIntPoint SectionBaseB = B->GetSectionBaseOffset();
 
-			// update proxies reference actor
-			for (ALandscapeStreamingProxy* StreamingProxy : Proxies)
+			if (SectionBaseA.X != SectionBaseB.X)
 			{
-				StreamingProxy->LandscapeActor = LandscapeActor;
-				StreamingProxy->FixupSharedData(Landscape);
+				return SectionBaseA.X < SectionBaseB.X;
 			}
-		}
-		else
-		{
-			auto LamdbdaLowerBound = [](ALandscapeProxy* A, ALandscapeProxy* B)
-			{
-				FIntPoint SectionBaseA = A->GetSectionBaseOffset();
-				FIntPoint SectionBaseB = B->GetSectionBaseOffset();
-
-				if (SectionBaseA.X != SectionBaseB.X)
-				{
-					return SectionBaseA.X < SectionBaseB.X;
-				}
 
-				return SectionBaseA.Y < SectionBaseB.Y;
-			};
+			return SectionBaseA.Y < SectionBaseB.Y;
+		};
 
-			// Insert Proxies in a sorted fashion for generating deterministic results in the Layer system
-			ALandscapeStreamingProxy* StreamingProxy = CastChecked<ALandscapeStreamingProxy>(Proxy);
-			if (!Proxies.Contains(Proxy))
-			{
-				uint32 InsertIndex = Algo::LowerBound(Proxies, Proxy, LamdbdaLowerBound);
-				Proxies.Insert(StreamingProxy, InsertIndex);
-			}
-			StreamingProxy->LandscapeActor = LandscapeActor;
-			StreamingProxy->FixupSharedData(LandscapeActor.Get());
+		// Insert Proxies in a sorted fashion for generating deterministic results in the Layer system
+		ALandscapeStreamingProxy* StreamingProxy = CastChecked<ALandscapeStreamingProxy>(Proxy);
+		if (!Proxies.Contains(Proxy))
+		{
+			uint32 InsertIndex = Algo::LowerBound(Proxies, Proxy, LamdbdaLowerBound);
+			Proxies.Insert(StreamingProxy, InsertIndex);
 		}
 
-		UpdateLayerInfoMap(Proxy);
-		UpdateAllAddCollisions();
-
-		RegisterSplineActor(Proxy);
+		StreamingProxy->LandscapeActor = LandscapeActor;
+#if WITH_EDITOR
+		StreamingProxy->FixupSharedData(LandscapeActor.Get());
+#endif // WITH_EDITOR
 	}
-#endif
+
+#if WITH_EDITOR
+	UpdateLayerInfoMap(Proxy);
+	UpdateAllAddCollisions();
+	RegisterSplineActor(Proxy);
+#endif // WITH_EDITOR
 
 	//
 	// add proxy components to the XY map
@@ -4020,36 +4025,33 @@ void ULandscapeInfo::RegisterActor(ALandscapeProxy* Proxy, bool bMapCheck)
 void ULandscapeInfo::UnregisterActor(ALandscapeProxy* Proxy)
 {
 	UWorld* OwningWorld = Proxy->GetWorld();
-#if WITH_EDITOR
-	if (!OwningWorld->IsGameWorld())
+	if (ALandscape* Landscape = Cast<ALandscape>(Proxy))
 	{
-		if (ALandscape* Landscape = Cast<ALandscape>(Proxy))
-		{
-			// Note: UnregisterActor sometimes gets triggered twice, e.g. it has been observed to happen during redo
-			// Note: In some cases LandscapeActor could be updated to a new landscape actor before the old landscape is unregistered/destroyed
-			// e.g. this has been observed when merging levels in the editor
-
-			if (LandscapeActor.Get() == Landscape)
-			{
-				LandscapeActor = nullptr;
-			}
+		// Note: UnregisterActor sometimes gets triggered twice, e.g. it has been observed to happen during redo
+		// Note: In some cases LandscapeActor could be updated to a new landscape actor before the old landscape is unregistered/destroyed
+		// e.g. this has been observed when merging levels in the editor
 
-			// update proxies reference to landscape actor
-			for (ALandscapeStreamingProxy* StreamingProxy : Proxies)
-			{
-				StreamingProxy->LandscapeActor = LandscapeActor;
-			}
-		}
-		else
+		if (LandscapeActor.Get() == Landscape)
 		{
-			ALandscapeStreamingProxy* StreamingProxy = CastChecked<ALandscapeStreamingProxy>(Proxy);
-			Proxies.Remove(StreamingProxy);
-			StreamingProxy->LandscapeActor = nullptr;
+			LandscapeActor = nullptr;
 		}
 
-		UnregisterSplineActor(Proxy);
+		// update proxies reference to landscape actor
+		for (ALandscapeStreamingProxy* StreamingProxy : Proxies)
+		{
+			StreamingProxy->LandscapeActor = LandscapeActor;
+		}
 	}
-#endif
+	else
+	{
+		ALandscapeStreamingProxy* StreamingProxy = CastChecked<ALandscapeStreamingProxy>(Proxy);
+		Proxies.Remove(StreamingProxy);
+		StreamingProxy->LandscapeActor = nullptr;
+	}
+
+#if WITH_EDITOR
+	UnregisterSplineActor(Proxy);
+#endif // WITH_EDITOR
 
 	// remove proxy components from the XY map
 	for (int32 CompIdx = 0; CompIdx < Proxy->LandscapeComponents.Num(); ++CompIdx)
@@ -4072,11 +4074,8 @@ void ULandscapeInfo::UnregisterActor(ALandscapeProxy* Proxy)
 	XYtoCollisionComponentMap.Compact();
 
 #if WITH_EDITOR
-	if (!OwningWorld->IsGameWorld())
-	{
-		UpdateLayerInfoMap();
-		UpdateAllAddCollisions();
-	}
+	UpdateLayerInfoMap();
+	UpdateAllAddCollisions();
 #endif
 }
 
@@ -5137,7 +5136,6 @@ void ALandscapeProxy::InvalidateGeneratedComponentData(const TArray<ULandscapeCo
 void ALandscapeProxy::InvalidateGeneratedComponentData(const TSet<ULandscapeComponent*>& Components, bool bInvalidateLightingCache)
 {
 	InvalidateGeneratedComponentData(Components.Array(), bInvalidateLightingCache);
-}
 
 void ALandscapeProxy::UpdateRenderingMethod(bool bDisableNanite)
 {
diff --git a/Engine/Source/Runtime/Landscape/Private/LandscapeEdit.cpp b/Engine/Source/Runtime/Landscape/Private/LandscapeEdit.cpp
--- a/Engine/Source/Runtime/Landscape/Private/LandscapeEdit.cpp
+++ b/Engine/Source/Runtime/Landscape/Private/LandscapeEdit.cpp
@@ -89,9 +89,9 @@ DEFINE_LOG_CATEGORY(LogLandscapeBP);
 #define LOCTEXT_NAMESPACE "Landscape"
 
 static TAutoConsoleVariable<int32> CVarMobileCompressLanscapeWeightMaps(
-	TEXT("r.Mobile.CompressLandscapeWeightMaps"),
+    TEXT("r.Mobile.CompressLandscapeWeightMaps"),
 	0,
-	TEXT("Whether to compress the terrain weight maps for mobile."),
+    TEXT("Whether to compress the terrain weight maps for mobile."),
 	ECVF_ReadOnly);
 
 #if WITH_EDITOR
@@ -4072,7 +4072,7 @@ void ALandscapeProxy::PostEditMove(bool bFinished)
 	// This point is only reached when Copy and Pasted
 	Super::PostEditMove(bFinished);
 
-	if (bFinished && !GetWorld()->IsGameWorld())
+	if (bFinished)
 	{
 		ULandscapeInfo::RecreateLandscapeInfo(GetWorld(), true);
 		RecreateComponentsState();
@@ -4099,7 +4099,7 @@ void ALandscapeProxy::PostEditImport()
 
 void ALandscape::PostEditMove(bool bFinished)
 {
-	if (bFinished && !GetWorld()->IsGameWorld())
+	if (bFinished)
 	{
 		// align all proxies to landscape actor
 		auto* LandscapeInfo = GetLandscapeInfo();
@@ -4231,9 +4231,9 @@ ULandscapeLayerInfoObject::ULandscapeLayerInfoObject(const FObjectInitializer& O
 
 FLinearColor ULandscapeLayerInfoObject::GenerateLayerUsageDebugColor() const
 {
-	uint8 Hash[20];
-	FString PathNameString = GetPathName();
-	FSHA1::HashBuffer(*PathNameString, PathNameString.Len() * sizeof(PathNameString[0]), Hash);
+		uint8 Hash[20];
+		FString PathNameString = GetPathName();
+		FSHA1::HashBuffer(*PathNameString, PathNameString.Len() * sizeof(PathNameString[0]), Hash);
 
 	return FLinearColor(float(Hash[0]) / 255.f, float(Hash[1]) / 255.f, float(Hash[2]) / 255.f, 1.f);
 }
@@ -5214,11 +5214,11 @@ void ALandscapeStreamingProxy::PostEditChangeProperty(FPropertyChangedEvent& Pro
 			if (World != nullptr)
 			{
 				if (UseMobileLandscapeMesh(GShaderPlatformForFeatureLevel[World->FeatureLevel]))
+			{
+				for (ULandscapeComponent * Component : LandscapeComponents)
 				{
-					for (ULandscapeComponent* Component : LandscapeComponents)
+					if (Component != nullptr)
 					{
-						if (Component != nullptr)
-						{
 							Component->CheckGenerateMobilePlatformData(/*bIsCooking = */ false, /*TargetPlatform = */ nullptr);
 						}
 					}
@@ -5227,7 +5227,7 @@ void ALandscapeStreamingProxy::PostEditChangeProperty(FPropertyChangedEvent& Pro
 				if (LiveRebuildNaniteOnModification != 0)
 				{
 					UpdateNaniteRepresentation();
-				}
+					}
 				else
 				{
 					InvalidateNaniteRepresentation();
@@ -5250,7 +5250,7 @@ void ALandscapeStreamingProxy::PostRegisterAllComponents()
 	{
 		ULandscapeInfo* LandscapeInfo = GetLandscapeInfo();
 		check(LandscapeInfo);
-		if (GEditor && !GetWorld()->IsGameWorld() && !IsRunningCommandlet())
+		if (GEditor && !IsRunningCommandlet())
 		{
 			if (UWorldPartition* WorldPartition = GetWorld()->GetWorldPartition(); WorldPartition && WorldPartition->IsInitialized())
 			{
@@ -5417,21 +5417,21 @@ bool ULandscapeInfo::CanDeleteLandscape(FText& OutReason) const
 				{
 					FGuid ParsedLandscapeGuid;
 					if (FGuid::Parse(AffectsLandscapeProperty.ToString(), ParsedLandscapeGuid) && ParsedLandscapeGuid == LandscapeGuid)
-					{
+				{
 						ALandscapeSplineActor* SplineActor = Cast<ALandscapeSplineActor>(ActorDesc->GetActor());
-
-						// If SplineActor is null then it is not loaded/deleted. If it's loaded then it needs to be pending kill.
-						if (!SplineActor)
-						{
-							++UndeletedSplineCount;
-						}
-						else
-						{
-							// If Actor is loaded it should be Registered and not pending kill (already accounted for) or pending kill (deleted)
-							check(SplineActors.Contains(SplineActor) == IsValidChecked(SplineActor));
-						}
+		
+					// If SplineActor is null then it is not loaded/deleted. If it's loaded then it needs to be pending kill.
+					if (!SplineActor)
+					{
+						++UndeletedSplineCount;
+					}
+					else
+					{
+						// If Actor is loaded it should be Registered and not pending kill (already accounted for) or pending kill (deleted)
+						check(SplineActors.Contains(SplineActor) == IsValidChecked(SplineActor));
 					}
 				}
+				}
 
 				return true;
 			});
@@ -5679,15 +5679,15 @@ void ALandscape::PostEditChangeProperty(FPropertyChangedEvent& PropertyChangedEv
 				if (World != nullptr)
 				{
 					if (UseMobileLandscapeMesh(GShaderPlatformForFeatureLevel[World->FeatureLevel]))
+				{
+					for (ULandscapeComponent * Component : LandscapeComponents)
 					{
-						for (ULandscapeComponent* Component : LandscapeComponents)
+						if (Component != nullptr)
 						{
-							if (Component != nullptr)
-							{
 								Component->CheckGenerateMobilePlatformData(/*bIsCooking = */ false, /*TargetPlatform = */ nullptr);
-							}
 						}
 					}
+				}
 
 					if (LiveRebuildNaniteOnModification != 0)
 					{
@@ -7784,4 +7784,4 @@ bool ALandscapeProxy::LandscapeExportWeightmapToRenderTarget(UTextureRenderTarge
 
 #endif //WITH_EDITOR
 
-#undef LOCTEXT_NAMESPACE
+#undef LOCTEXT_NAMESPACE
\ No newline at end of file
diff --git a/Engine/Source/Runtime/Landscape/Private/LandscapeEditLayers.cpp b/Engine/Source/Runtime/Landscape/Private/LandscapeEditLayers.cpp
--- a/Engine/Source/Runtime/Landscape/Private/LandscapeEditLayers.cpp
+++ b/Engine/Source/Runtime/Landscape/Private/LandscapeEditLayers.cpp
@@ -154,7 +154,7 @@ void OnLandscapeEditLayersLocalMergeChanged(IConsoleVariable* CVar)
 		UWorld* CurrentWorld = *It;
 		if (!CurrentWorld->IsGameWorld())
 		{
-			auto& LandscapeInfoMap = ULandscapeInfoMap::GetLandscapeInfoMap(CurrentWorld);
+			ULandscapeInfoMap& LandscapeInfoMap = ULandscapeInfoMap::GetLandscapeInfoMap(CurrentWorld);
 			for (TPair<FGuid, ULandscapeInfo*>& Pair : LandscapeInfoMap.Map)
 			{
 				if (ULandscapeInfo* LandscapeInfo = Pair.Value)
@@ -209,7 +209,7 @@ private:
 			UWorld* CurrentWorld = *It;
 			if (!CurrentWorld->IsGameWorld())
 			{
-				auto& LandscapeInfoMap = ULandscapeInfoMap::GetLandscapeInfoMap(CurrentWorld);
+				ULandscapeInfoMap& LandscapeInfoMap = ULandscapeInfoMap::GetLandscapeInfoMap(CurrentWorld);
 				for (TPair<FGuid, ULandscapeInfo*>& Pair : LandscapeInfoMap.Map)
 				{
 					if (Pair.Value && Pair.Value->SupportsLandscapeEditing())
