Restore lost closing brace
#rnx

diff --git a/Engine/Source/Runtime/Landscape/Private/Landscape.cpp b/Engine/Source/Runtime/Landscape/Private/Landscape.cpp
--- a/Engine/Source/Runtime/Landscape/Private/Landscape.cpp
+++ b/Engine/Source/Runtime/Landscape/Private/Landscape.cpp
@@ -5136,6 +5136,7 @@ void ALandscapeProxy::InvalidateGeneratedComponentData(const TArray<ULandscapeCo
 void ALandscapeProxy::InvalidateGeneratedComponentData(const TSet<ULandscapeComponent*>& Components, bool bInvalidateLightingCache)
 {
 	InvalidateGeneratedComponentData(Components.Array(), bInvalidateLightingCache);
+}
 
 void ALandscapeProxy::UpdateRenderingMethod(bool bDisableNanite)
 {
