* Fixed incorrect removal of IsGameWorld() check w.r.t. world partitioned-worlds : ActorDescReferences is only meant for making sure spline actors intersecting a landscape proxy are loaded alongside the proxy but this shouldn't be done for PIE worlds
* Fixed various places where components within the LandscapeComponents member of ALandscapeProxy could potentially be null (they are null in the dedicated server), like it is done elsewhere in the code

diff --git a/Engine/Source/Runtime/Landscape/Private/Landscape.cpp b/Engine/Source/Runtime/Landscape/Private/Landscape.cpp
--- a/Engine/Source/Runtime/Landscape/Private/Landscape.cpp
+++ b/Engine/Source/Runtime/Landscape/Private/Landscape.cpp
@@ -3056,7 +3056,7 @@ bool ULandscapeInfo::UpdateLayerInfoMapInternal(ALandscapeProxy* Proxy, bool bIn
 					ULandscapeComponent* Component = Proxy->LandscapeComponents[ComponentIndex];
 
 					// Add layers from per-component override materials
-					if (Component->OverrideMaterial != nullptr)
+					if ((Component != nullptr) && (Component->OverrideMaterial != nullptr))
 					{
 						TArray<FName> ComponentLayerNames = Proxy->GetLayersFromMaterial(Component->OverrideMaterial);
 						for (int32 i = 0; i < ComponentLayerNames.Num(); i++)
@@ -4787,35 +4787,35 @@ void ALandscapeProxy::UpdateGIBakedTextureStatus(bool* bOutGenerateLandscapeGIDa
 	}
 	else
 	{
-	// Stores the components and their state hash data for a single atlas
+		// Stores the components and their state hash data for a single atlas
 		struct FGIBakeTextureStateBuilder
-	{
-		// pointer as FMemoryWriter caches the address of the FBufferArchive, and this struct could be relocated on a realloc.
-		TUniquePtr<FBufferArchive> ComponentStateAr;
-		TArray<ULandscapeComponent*> Components;
+		{
+			// pointer as FMemoryWriter caches the address of the FBufferArchive, and this struct could be relocated on a realloc.
+			TUniquePtr<FBufferArchive> ComponentStateAr;
+			TArray<ULandscapeComponent*> Components;
 
 			FGIBakeTextureStateBuilder()
-		{
-			ComponentStateAr = MakeUnique<FBufferArchive>();
-		}
-	};
+			{
+				ComponentStateAr = MakeUnique<FBufferArchive>();
+			}
+		};
 
 		TMap<UTexture2D*, FGIBakeTextureStateBuilder> ComponentsByHeightmap;
-	for (ULandscapeComponent* Component : LandscapeComponents)
-	{
-		if (Component == nullptr)
+		for (ULandscapeComponent* Component : LandscapeComponents)
 		{
-			continue;
-		}
+			if (Component == nullptr)
+			{
+				continue;
+			}
 
 			FGIBakeTextureStateBuilder& Info = ComponentsByHeightmap.FindOrAdd(Component->GetHeightmap());
-		Info.Components.Add(Component);
-		Component->SerializeStateHashes(*Info.ComponentStateAr);
-	}
+			Info.Components.Add(Component);
+			Component->SerializeStateHashes(*Info.ComponentStateAr);
+		}
 
 		for (auto It = ComponentsByHeightmap.CreateIterator(); It; ++It)
 		{
-			FGIBakeTextureStateBuilder& Info =It.Value();
+			FGIBakeTextureStateBuilder& Info = It.Value();
 
 			// Calculate a combined Guid-like ID we can use for this component
 			uint32 Hash[5];
diff --git a/Engine/Source/Runtime/Landscape/Private/LandscapeEdit.cpp b/Engine/Source/Runtime/Landscape/Private/LandscapeEdit.cpp
--- a/Engine/Source/Runtime/Landscape/Private/LandscapeEdit.cpp
+++ b/Engine/Source/Runtime/Landscape/Private/LandscapeEdit.cpp
@@ -3624,9 +3624,12 @@ FIntRect ALandscapeProxy::GetBoundingRect() const
 	if (LandscapeComponents.Num() > 0)
 	{
 		FIntRect Rect(MAX_int32, MAX_int32, MIN_int32, MIN_int32);
-		for (int32 CompIdx = 0; CompIdx < LandscapeComponents.Num(); CompIdx++)
+		for (ULandscapeComponent* Component : LandscapeComponents)
 		{
-			Rect.Include(LandscapeComponents[CompIdx]->GetSectionBase());
+			if (Component != nullptr)
+			{
+				Rect.Include(Component->GetSectionBase());
+			}
 		}
 		Rect.Max += FIntPoint(ComponentSizeQuads, ComponentSizeQuads);
 		Rect -= LandscapeSectionOffset;
@@ -3656,7 +3659,10 @@ bool ULandscapeInfo::GetLandscapeExtent(ALandscapeProxy* LandscapeProxy, FIntRec
 
 	for (ULandscapeComponent* LandscapeComponent : LandscapeProxy->LandscapeComponents)
 	{
-		LandscapeComponent->GetComponentExtent(ProxyExtent.Min.X, ProxyExtent.Min.Y, ProxyExtent.Max.X, ProxyExtent.Max.Y);
+		if (LandscapeComponent != nullptr)
+		{
+			LandscapeComponent->GetComponentExtent(ProxyExtent.Min.X, ProxyExtent.Min.Y, ProxyExtent.Max.X, ProxyExtent.Max.Y);
+		}
 	}
 	
 	return ProxyExtent.Min.X != INT32_MAX;
@@ -3697,7 +3703,10 @@ LANDSCAPE_API void ULandscapeInfo::ForAllLandscapeComponents(TFunctionRef<void(U
 	{
 		for (ULandscapeComponent* Component : Proxy->LandscapeComponents)
 		{
-			Fn(Component);
+			if (Component != nullptr)
+			{
+				Fn(Component);
+			}
 		}
 	});
 }
@@ -3976,7 +3985,10 @@ void ULandscapeInfo::GetUsedPaintLayers(const FGuid& InLayerGuid, TArray<ULandsc
 	{
 		for (ULandscapeComponent* Component : Proxy->LandscapeComponents)
 		{
-			Component->GetUsedPaintLayers(InLayerGuid, OutUsedLayerInfos);
+			if (Component != nullptr)
+			{
+				Component->GetUsedPaintLayers(InLayerGuid, OutUsedLayerInfos);
+			}
 		}
 	});
 }
@@ -5214,11 +5226,11 @@ void ALandscapeStreamingProxy::PostEditChangeProperty(FPropertyChangedEvent& Pro
 			if (World != nullptr)
 			{
 				if (UseMobileLandscapeMesh(GShaderPlatformForFeatureLevel[World->FeatureLevel]))
-			{
-				for (ULandscapeComponent * Component : LandscapeComponents)
 				{
-					if (Component != nullptr)
+					for (ULandscapeComponent* Component : LandscapeComponents)
 					{
+						if (Component != nullptr)
+						{
 							Component->CheckGenerateMobilePlatformData(/*bIsCooking = */ false, /*TargetPlatform = */ nullptr);
 						}
 					}
@@ -5250,7 +5262,7 @@ void ALandscapeStreamingProxy::PostRegisterAllComponents()
 	{
 		ULandscapeInfo* LandscapeInfo = GetLandscapeInfo();
 		check(LandscapeInfo);
-		if (GEditor && !IsRunningCommandlet())
+		if (GEditor && !GetWorld()->IsGameWorld() && !IsRunningCommandlet())
 		{
 			if (UWorldPartition* WorldPartition = GetWorld()->GetWorldPartition(); WorldPartition && WorldPartition->IsInitialized())
 			{
@@ -5679,15 +5691,15 @@ void ALandscape::PostEditChangeProperty(FPropertyChangedEvent& PropertyChangedEv
 				if (World != nullptr)
 				{
 					if (UseMobileLandscapeMesh(GShaderPlatformForFeatureLevel[World->FeatureLevel]))
-				{
-					for (ULandscapeComponent * Component : LandscapeComponents)
 					{
-						if (Component != nullptr)
+						for (ULandscapeComponent* Component : LandscapeComponents)
 						{
+							if (Component != nullptr)
+							{
 								Component->CheckGenerateMobilePlatformData(/*bIsCooking = */ false, /*TargetPlatform = */ nullptr);
+							}
 						}
 					}
-				}
 
 					if (LiveRebuildNaniteOnModification != 0)
 					{
@@ -6209,10 +6221,12 @@ void ALandscapeProxy::RemoveInvalidWeightmaps()
 		}
 
 		// Remove Unused Weightmaps...
-		for (int32 Idx = 0; Idx < LandscapeComponents.Num(); ++Idx)
+		for (ULandscapeComponent* Component : LandscapeComponents)
 		{
-			ULandscapeComponent* Component = LandscapeComponents[Idx];
-			Component->RemoveInvalidWeightmaps();
+			if (Component != nullptr)
+			{
+				Component->RemoveInvalidWeightmaps();
+			}
 		}
 	}
 }
diff --git a/Engine/Source/Runtime/Landscape/Private/LandscapeEditLayers.cpp b/Engine/Source/Runtime/Landscape/Private/LandscapeEditLayers.cpp
--- a/Engine/Source/Runtime/Landscape/Private/LandscapeEditLayers.cpp
+++ b/Engine/Source/Runtime/Landscape/Private/LandscapeEditLayers.cpp
@@ -8121,13 +8121,13 @@ bool ALandscapeProxy::HasLayersContent() const
 void ALandscapeProxy::UpdateCachedHasLayersContent(bool InCheckComponentDataIntegrity)
 {
 	// In the case of InCheckComponentDataIntegrity we will loop through all components to make sure they all have the same state and in the other case we will assume that the 1st component represent the state of all the others.
-	bHasLayersContent = LandscapeComponents.Num() > 0 ? LandscapeComponents[0]->HasLayersData() : false;
+	bHasLayersContent = (!LandscapeComponents.IsEmpty() && (LandscapeComponents[0] != nullptr)) ? LandscapeComponents[0]->HasLayersData() : false;
 
 	if (InCheckComponentDataIntegrity)
 	{
 		for (const ULandscapeComponent* Component : LandscapeComponents)
 		{
-			check(bHasLayersContent == Component->HasLayersData());
+			check((Component == nullptr) || (bHasLayersContent == Component->HasLayersData()));
 		}
 	}
 }
@@ -8137,7 +8137,10 @@ bool ALandscapeProxy::RemoveObsoleteLayers(const TSet<FGuid>& InExistingLayers)
 	TSet<TPair<FGuid, FName>> ComponentLayers;
 	for (ULandscapeComponent* Component : LandscapeComponents)
 	{
-		Component->ForEachLayer([&](const FGuid& Guid, FLandscapeLayerComponentData& ComponentData) { ComponentLayers.Add(TPair<FGuid, FName>(Guid, ComponentData.DebugName)); });
+		if (Component != nullptr)
+		{
+			Component->ForEachLayer([&](const FGuid& Guid, FLandscapeLayerComponentData& ComponentData) { ComponentLayers.Add(TPair<FGuid, FName>(Guid, ComponentData.DebugName)); });
+		}
 	}
 
 	bool bModified = false;
@@ -8177,7 +8180,7 @@ bool ALandscapeProxy::AddLayer(const FGuid& InLayerGuid)
 	bool bModified = false;
 	for (ULandscapeComponent* Component : LandscapeComponents)
 	{
-		if (!Component->GetLayerData(InLayerGuid))
+		if ((Component != nullptr) && !Component->GetLayerData(InLayerGuid))
 		{
 			const FLandscapeLayer* EditLayer = GetLandscapeActor() ? GetLandscapeActor()->GetLayer(InLayerGuid) : nullptr;
 			Component->AddLayerData(InLayerGuid, FLandscapeLayerComponentData(EditLayer ? EditLayer->Name : FName()));
@@ -8199,26 +8202,29 @@ void ALandscapeProxy::DeleteLayer(const FGuid& InLayerGuid)
 {
 	for (ULandscapeComponent* Component : LandscapeComponents)
 	{
-		const FLandscapeLayerComponentData* LayerComponentData = Component->GetLayerData(InLayerGuid);
-
-		if (LayerComponentData != nullptr)
+		if (Component != nullptr)
 		{
-			for (const FWeightmapLayerAllocationInfo& Allocation : LayerComponentData->WeightmapData.LayerAllocations)
-			{
-				UTexture2D* WeightmapTexture = LayerComponentData->WeightmapData.Textures[Allocation.WeightmapTextureIndex];
-				TObjectPtr<ULandscapeWeightmapUsage>* Usage = WeightmapUsageMap.Find(WeightmapTexture);
+			const FLandscapeLayerComponentData* LayerComponentData = Component->GetLayerData(InLayerGuid);
 
-				if (Usage != nullptr && (*Usage) != nullptr)
+			if (LayerComponentData != nullptr)
+			{
+				for (const FWeightmapLayerAllocationInfo& Allocation : LayerComponentData->WeightmapData.LayerAllocations)
 				{
-					(*Usage)->ChannelUsage[Allocation.WeightmapTextureChannel] = nullptr;
+					UTexture2D* WeightmapTexture = LayerComponentData->WeightmapData.Textures[Allocation.WeightmapTextureIndex];
+					TObjectPtr<ULandscapeWeightmapUsage>* Usage = WeightmapUsageMap.Find(WeightmapTexture);
 
-					if ((*Usage)->IsEmpty())
+					if (Usage != nullptr && (*Usage) != nullptr)
 					{
-						WeightmapUsageMap.Remove(WeightmapTexture);
+						(*Usage)->ChannelUsage[Allocation.WeightmapTextureChannel] = nullptr;
+
+						if ((*Usage)->IsEmpty())
+						{
+							WeightmapUsageMap.Remove(WeightmapTexture);
+						}
 					}
 				}
+				Component->RemoveLayerData(InLayerGuid);
 			}
-			Component->RemoveLayerData(InLayerGuid);
 		}
 	}
 
@@ -8239,9 +8245,12 @@ void ALandscapeProxy::InitializeLayerWithEmptyContent(const FGuid& InLayerGuid)
 
 	for (ULandscapeComponent* Component : LandscapeComponents)
 	{
-		UTexture2D* ComponentHeightmapTexture = Component->GetHeightmap();
-		TArray<ULandscapeComponent*>& ComponentList = ComponentsPerHeightmaps.FindOrAdd(ComponentHeightmapTexture);
-		ComponentList.Add(Component);
+		if (Component != nullptr)
+		{
+			UTexture2D* ComponentHeightmapTexture = Component->GetHeightmap();
+			TArray<ULandscapeComponent*>& ComponentList = ComponentsPerHeightmaps.FindOrAdd(ComponentHeightmapTexture);
+			ComponentList.Add(Component);
+		}
 	}
 
 	// Init layers with valid "empty" data
@@ -8249,11 +8258,14 @@ void ALandscapeProxy::InitializeLayerWithEmptyContent(const FGuid& InLayerGuid)
 
 	for (ULandscapeComponent* Component : LandscapeComponents)
 	{
-		UTexture2D* ComponentHeightmap = Component->GetHeightmap();
-		const TArray<ULandscapeComponent*>* ComponentsUsingHeightmap = ComponentsPerHeightmaps.Find(ComponentHeightmap);
-		check(ComponentsUsingHeightmap != nullptr);
+		if (Component != nullptr)
+		{
+			UTexture2D* ComponentHeightmap = Component->GetHeightmap();
+			const TArray<ULandscapeComponent*>* ComponentsUsingHeightmap = ComponentsPerHeightmaps.Find(ComponentHeightmap);
+			check(ComponentsUsingHeightmap != nullptr);
 
-		Component->AddDefaultLayerData(InLayerGuid, *ComponentsUsingHeightmap, CreatedHeightmapTextures);
+			Component->AddDefaultLayerData(InLayerGuid, *ComponentsUsingHeightmap, CreatedHeightmapTextures);
+		}
 	}
 }
 #endif
