Landscape: Moved assert to catch component weightmap atlas offset issue earlier

diff --git a/Engine/Source/Runtime/Landscape/Private/LandscapeEditLayers.cpp b/Engine/Source/Runtime/Landscape/Private/LandscapeEditLayers.cpp
--- a/Engine/Source/Runtime/Landscape/Private/LandscapeEditLayers.cpp
+++ b/Engine/Source/Runtime/Landscape/Private/LandscapeEditLayers.cpp
@@ -6884,20 +6884,25 @@ int32 ALandscape::PerformLayersWeightmapsGlobalMerge(FUpdateLayersContentContext
 					TArray<FVector2f> WeightmapTextureOutputOffset;
 
 					// Compute each weightmap location so compute shader will be able to output at expected location
-					int32 ComponentSize = (SubsectionSizeQuads + 1) * NumSubsections;
+					const int32 WeightmapSizeX = WeightmapScratchPackLayerTextureResource->GetSizeX();
+					const int32 WeightmapSizeY = WeightmapScratchPackLayerTextureResource->GetSizeY();
+					const int32 ComponentSize = (SubsectionSizeQuads + 1) * NumSubsections;
+
 					float ComponentY = 0;
 					float ComponentX = 0;
 
 					for (int32 i = 0; i < PackLayersComponentsData.Num(); ++i)
 					{
-						check(ComponentY + ComponentSize <= WeightmapScratchPackLayerTextureResource->GetSizeY()); // This should never happen as it would be a bug in the algo
-
-						if (ComponentX + ComponentSize > WeightmapScratchPackLayerTextureResource->GetSizeX())
+						if (ComponentX + ComponentSize > WeightmapSizeX)
 						{
 							ComponentY += ComponentSize;
 							ComponentX = 0;
 						}
 
+						// This should never happen as it would be a bug in the algo
+						check(ComponentX + ComponentSize <= WeightmapSizeX);
+						check(ComponentY + ComponentSize <= WeightmapSizeY);
+
 						WeightmapTextureOutputOffset.Add(FVector2f(ComponentX, ComponentY));
 						ComponentX += ComponentSize;
 					}
