Only update the Landscape Nanite representation if we have a valid landscape actor. There is no landscape actor during cooking, and this was causing IsNaniteEnabled to return false, destroying the Nanite data (and falling back on non-Nanite for cooked targets).

diff --git a/Engine/Source/Runtime/Landscape/Private/Landscape.cpp b/Engine/Source/Runtime/Landscape/Private/Landscape.cpp
--- a/Engine/Source/Runtime/Landscape/Private/Landscape.cpp
+++ b/Engine/Source/Runtime/Landscape/Private/Landscape.cpp
@@ -2866,10 +2866,10 @@ void ALandscapeProxy::PreSave(FObjectPreSaveContext ObjectSaveContext)
 				}
 			});
 		}
-	}
 
-	UpdateNaniteRepresentation();
-	UpdateRenderingMethod();
+		UpdateNaniteRepresentation();
+		UpdateRenderingMethod();
+	}
 #endif // WITH_EDITOR
 }
 
