Made UpdateRenderingMethod Nanite-path selection work in non-editor builds (where not ALandscape exists).
- disables / skips extra mobile platform logic.

diff --git a/Engine/Source/Runtime/Landscape/Private/Landscape.cpp b/Engine/Source/Runtime/Landscape/Private/Landscape.cpp
--- a/Engine/Source/Runtime/Landscape/Private/Landscape.cpp
+++ b/Engine/Source/Runtime/Landscape/Private/Landscape.cpp
@@ -5149,15 +5149,19 @@ void ALandscapeProxy::UpdateRenderingMethod(bool bDisableNanite)
 
 	if (!bDisableNanite && NaniteComponent != nullptr)
 	{
+		bNaniteActive = UseNanite(GShaderPlatformForFeatureLevel[GMaxRHIFeatureLevel]);
+#if WITH_EDITOR
 		if (ALandscape* LandscapeActor = GetLandscapeActor())
 		{
 			if (UWorld* World = LandscapeActor->GetWorld())
 			{
-				const bool bIsMobile = ((GEngine->GetDefaultWorldFeatureLevel() == ERHIFeatureLevel::ES3_1) || (World && (World->FeatureLevel <= ERHIFeatureLevel::ES3_1)));
-				const ERHIFeatureLevel::Type FeatureLevel = bIsMobile ? ERHIFeatureLevel::ES3_1 : GMaxRHIFeatureLevel;
-				bNaniteActive = UseNanite(GShaderPlatformForFeatureLevel[FeatureLevel]);
+				if ((GEngine->GetDefaultWorldFeatureLevel() == ERHIFeatureLevel::ES3_1) || (World->FeatureLevel <= ERHIFeatureLevel::ES3_1))
+				{
+					bNaniteActive = UseNanite(GShaderPlatformForFeatureLevel[ERHIFeatureLevel::ES3_1]);
+				}
 			}
 		}
+#endif
 	}
 
 	for (ULandscapeComponent* Component : LandscapeComponents)
