World Partition - HLOD - Expose HLOD options on landscape

diff --git a/Engine/Source/Runtime/Landscape/Classes/LandscapeProxy.h b/Engine/Source/Runtime/Landscape/Classes/LandscapeProxy.h
--- a/Engine/Source/Runtime/Landscape/Classes/LandscapeProxy.h
+++ b/Engine/Source/Runtime/Landscape/Classes/LandscapeProxy.h
@@ -1,3 +1,4 @@
+
 // Copyright Epic Games, Inc. All Rights Reserved.
 
 #pragma once
@@ -342,7 +343,7 @@ struct UE_DEPRECATED(5.1, "FLandscapeProxyMaterialOverride is deprecated; please
 	TObjectPtr<UMaterialInterface> Material = nullptr;
 };
 
-UCLASS(Abstract, MinimalAPI, NotBlueprintable, NotPlaceable, hidecategories=(Display, Attachment, Physics, Debug, Lighting, HLOD), showcategories=(Lighting, Rendering, Transformation), hidecategories=(Mobility))
+UCLASS(Abstract, MinimalAPI, NotBlueprintable, NotPlaceable, hidecategories=(Display, Attachment, Physics, Debug, Lighting), showcategories=(Lighting, Rendering, Transformation), hidecategories=(Mobility))
 class ALandscapeProxy : public APartitionActor, public ILandscapeSplineInterface
 {
 	GENERATED_BODY()
@@ -738,7 +739,7 @@ public:
 #endif
 
 	/** Flag whether or not this Landscape's surface can be used for culling hidden triangles **/
-	UPROPERTY(EditAnywhere, Category = HierarchicalLOD)
+	UPROPERTY(EditAnywhere, Category = HLOD)
 	bool bUseLandscapeForCullingInvisibleHLODVertices;
 
 	/** Flag that tell if we have some layers content **/
