World Partition - HLOD - Enable Nanite on landscape HLODs if the source landscape proxies are Nanite-enabled

diff --git a/Engine/Source/Runtime/Landscape/Private/LandscapeHLODBuilder.cpp b/Engine/Source/Runtime/Landscape/Private/LandscapeHLODBuilder.cpp
--- a/Engine/Source/Runtime/Landscape/Private/LandscapeHLODBuilder.cpp
+++ b/Engine/Source/Runtime/Landscape/Private/LandscapeHLODBuilder.cpp
@@ -292,6 +292,9 @@ TArray<UActorComponent*> ULandscapeHLODBuilder::Build(const FHLODBuildContext& I
 
 			StaticMesh->CommitMeshDescription(0);
 
+			// Nanite settings
+			StaticMesh->NaniteSettings.bEnabled = LandscapeProxy->IsNaniteEnabled();
+
 			StaticMesh->ImportVersion = EImportStaticMeshVersion::LastVersion;
 		}
 
