World Partition - HLOD - Take into account the LS proxy nanite enabled flag in the component CRC used to check for HLOD dirtiness
#rnx

diff --git a/Engine/Source/Runtime/Landscape/Private/LandscapeHLODBuilder.cpp b/Engine/Source/Runtime/Landscape/Private/LandscapeHLODBuilder.cpp
--- a/Engine/Source/Runtime/Landscape/Private/LandscapeHLODBuilder.cpp
+++ b/Engine/Source/Runtime/Landscape/Private/LandscapeHLODBuilder.cpp
@@ -123,6 +123,10 @@ uint32 ULandscapeHLODBuilder::ComputeHLODHash(const UActorComponent* InSourceCom
 			Ar << UsedTexturesCRC;
 			UE_LOG(LogHLODBuilder, VeryVerbose, TEXT("     - UsedTexturesCRC = %x"), Ar.GetCrc());
 		}
+
+		// Nanite enabled?
+		bool bNaniteEnabled = LSProxy->IsNaniteEnabled();
+		Ar << bNaniteEnabled;
 	}
 
 	return Ar.GetCrc();
