Removed unacceptable words from landscape code

diff --git a/Engine/Source/Runtime/Landscape/Classes/LandscapeComponent.h b/Engine/Source/Runtime/Landscape/Classes/LandscapeComponent.h
--- a/Engine/Source/Runtime/Landscape/Classes/LandscapeComponent.h
+++ b/Engine/Source/Runtime/Landscape/Classes/LandscapeComponent.h
@@ -1126,7 +1126,7 @@ public:
 
 	LANDSCAPE_API FIntRect GetComponentExtent() const;
 
-	/** Updates navigation properties to match landscape's master switch */
+	/** Updates navigation properties to match landscape actor's */
 	void UpdateNavigationRelevance();
 
 	/** Updates the reject navmesh underneath flag in the collision component */
diff --git a/Engine/Source/Runtime/Landscape/Private/LandscapeEdit.cpp b/Engine/Source/Runtime/Landscape/Private/LandscapeEdit.cpp
--- a/Engine/Source/Runtime/Landscape/Private/LandscapeEdit.cpp
+++ b/Engine/Source/Runtime/Landscape/Private/LandscapeEdit.cpp
@@ -165,7 +165,7 @@ void ULandscapeComponent::UpdateCachedBounds(bool bInApproximateBounds)
 	ULandscapeHeightfieldCollisionComponent* HFCollisionComponent = CollisionComponent.Get();
 	if (HFCollisionComponent)
 	{
-        // In Landscape Layers the Collision Component is slave and doesn't need to be transacted
+        // In Landscape Layers the Collision Component gets regenerated after the heightmap changes are undone and doesn't need to be transacted
 		if (!GetLandscapeProxy()->HasLayersContent())
 		{
 			HFCollisionComponent->Modify();
@@ -1055,7 +1055,7 @@ void ULandscapeComponent::UpdateCollisionHeightData(const FColor* const Heightma
 	bool CreatedNew = false;
 	bool ChangeType = false;
 
-    // In Landscape Layers the Collision Component is slave and doesn't need to be transacted
+	// In Landscape Layers the Collision Component gets regenerated after the heightmap changes are undone and doesn't need to be transacted
 	if (!Proxy->HasLayersContent())
 	{
 		if (CollisionComp)
diff --git a/Engine/Source/Runtime/Landscape/Private/LandscapeEditLayers.cpp b/Engine/Source/Runtime/Landscape/Private/LandscapeEditLayers.cpp
--- a/Engine/Source/Runtime/Landscape/Private/LandscapeEditLayers.cpp
+++ b/Engine/Source/Runtime/Landscape/Private/LandscapeEditLayers.cpp
@@ -7333,7 +7333,7 @@ void ALandscape::UpdateCachedHasLayersContent(bool bInCheckComponentDataIntegrit
 
 	// For consistency with the ALandscape::HasLayersContent() override above, make sure the cached bHasLayersContent boolean is also valid when we have at least one edit layer : 
 	//  Otherwise, as ALandscapeProxy::UpdateCachedHasLayersContent relies on the presence of landscape components and in distributed landscape setups (one ALandscape + multiple ALandscapeStreamingProxy),  
-	//  the master ALandscape actor doesn't have any landscape, hence it would have bHasLayersContent erroneously set to false (while ALandscape::HasLayersContent() would actually return true!)
+	//  the "parent" ALandscape actor doesn't have any landscape component, hence it would have bHasLayersContent erroneously set to false (while ALandscape::HasLayersContent() would actually return true!)
 	bHasLayersContent |= ALandscape::HasLayersContent();
 }
 
