Landscape Custom Brushes: Added a property called 'Update on Property Change' that when false skips updating the brush on details editing. This was requested for extremely large maps where editing landscape triggers grass and PCG to re-update and is too slow.

diff --git a/Engine/Source/Runtime/Landscape/Private/LandscapeBlueprintBrushBase.cpp b/Engine/Source/Runtime/Landscape/Private/LandscapeBlueprintBrushBase.cpp
--- a/Engine/Source/Runtime/Landscape/Private/LandscapeBlueprintBrushBase.cpp
+++ b/Engine/Source/Runtime/Landscape/Private/LandscapeBlueprintBrushBase.cpp
@@ -23,6 +23,7 @@ static TAutoConsoleVariable<int32> CVarLandscapeBrushPadding(
 ALandscapeBlueprintBrushBase::ALandscapeBlueprintBrushBase(const FObjectInitializer& ObjectInitializer)
 #if WITH_EDITORONLY_DATA
 	: OwningLandscape(nullptr)
+	, UpdateOnPropertyChange(true)
 	, AffectHeightmap(false)
 	, AffectWeightmap(false)
 	, bIsVisible(true)
@@ -186,7 +187,7 @@ void ALandscapeBlueprintBrushBase::PostEditChangeProperty(FPropertyChangedEvent&
 {
 	Super::PostEditChangeProperty(PropertyChangedEvent);
 #if WITH_EDITORONLY_DATA
-	if (OwningLandscape)
+	if (OwningLandscape && UpdateOnPropertyChange)
 	{
 		OwningLandscape->OnBlueprintBrushChanged();
 	}
diff --git a/Engine/Source/Runtime/Landscape/Public/LandscapeBlueprintBrushBase.h b/Engine/Source/Runtime/Landscape/Public/LandscapeBlueprintBrushBase.h
--- a/Engine/Source/Runtime/Landscape/Public/LandscapeBlueprintBrushBase.h
+++ b/Engine/Source/Runtime/Landscape/Public/LandscapeBlueprintBrushBase.h
@@ -20,6 +20,9 @@ protected:
 	UPROPERTY(Transient)
 	TObjectPtr<class ALandscape> OwningLandscape;
 
+	UPROPERTY(Category = "Settings", EditAnywhere, BlueprintReadWrite)
+	bool UpdateOnPropertyChange;
+
 	UPROPERTY(Category = "Settings", EditAnywhere, BlueprintReadWrite)
 	bool AffectHeightmap;
 
