Landscape Collision View: Support collision view mode for landscape heightfield collision components

diff --git a/Engine/Source/Runtime/Landscape/Classes/LandscapeHeightfieldCollisionComponent.h b/Engine/Source/Runtime/Landscape/Classes/LandscapeHeightfieldCollisionComponent.h
--- a/Engine/Source/Runtime/Landscape/Classes/LandscapeHeightfieldCollisionComponent.h
+++ b/Engine/Source/Runtime/Landscape/Classes/LandscapeHeightfieldCollisionComponent.h
@@ -31,6 +31,8 @@ namespace Chaos
 UCLASS(MinimalAPI, Within=LandscapeProxy)
 class ULandscapeHeightfieldCollisionComponent : public UPrimitiveComponent
 {
+	friend class FLandscapeHeightfieldCollisionComponentSceneProxy;
+
 	GENERATED_UCLASS_BODY()
 
 	ULandscapeHeightfieldCollisionComponent(FVTableHelper& Helper);
@@ -230,6 +232,11 @@ public:
 	/** Gets the landscape info object for this landscape */
 	ULandscapeInfo* GetLandscapeInfo() const;
 
+#if !(UE_BUILD_SHIPPING || UE_BUILD_TEST) && WITH_EDITORONLY_DATA
+	// The scene proxy is only for debug purposes :
+	virtual FPrimitiveSceneProxy* CreateSceneProxy() override;
+#endif // !(UE_BUILD_SHIPPING || UE_BUILD_TEST) && WITH_EDITORONLY_DATA
+
 	/** Creates collision object from a cooked collision data */
 	virtual void CreateCollisionObject();
 
diff --git a/Engine/Source/Runtime/Landscape/Private/LandscapeCollision.cpp b/Engine/Source/Runtime/Landscape/Private/LandscapeCollision.cpp
--- a/Engine/Source/Runtime/Landscape/Private/LandscapeCollision.cpp
+++ b/Engine/Source/Runtime/Landscape/Private/LandscapeCollision.cpp
@@ -43,6 +43,7 @@
 #include "Physics/PhysicsFiltering.h"
 #include "Physics/PhysicsInterfaceCore.h"
 #include "Physics/PhysicsInterfaceUtils.h"
+#include "DynamicMeshBuilder.h"
 
 #include "Chaos/ParticleHandle.h"
 #include "Chaos/Vector.h"
@@ -61,6 +62,13 @@ using namespace PhysicsInterfaceTypes;
 // It's a lot faster to compute than to request from DDC, so always skip.
 bool GLandscapeCollisionSkipDDC = true;
 
+static TAutoConsoleVariable<int32> CVarLandscapeShowCollisionMesh(
+	TEXT("landscape.ShowCollisionMesh"),
+	0,
+	TEXT("Selects which heightfield to visualize when ShowFlags.Collision is used. 0 for simple, 1 for complex, 2 for editor only."),
+	ECVF_RenderThreadSafe
+);
+
 #if ENABLE_COOK_STATS
 namespace LandscapeCollisionCookStats
 {
@@ -443,6 +451,233 @@ void ULandscapeHeightfieldCollisionComponent::ApplyWorldOffset(const FVector& In
 	}
 }
 
+// Callback to flag scene proxy as dirty when cvar changes
+static void OnLandscapeShowCollisionMeshChanged()
+{
+	for (ULandscapeHeightfieldCollisionComponent* LandscapeHeightfieldCollisionComponent : TObjectRange<ULandscapeHeightfieldCollisionComponent>(RF_ClassDefaultObject | RF_ArchetypeObject, true, EInternalObjectFlags::Garbage))
+	{
+		LandscapeHeightfieldCollisionComponent->MarkRenderStateDirty();
+	}
+}
+
+// Registers callback with cvar
+static FAutoConsoleVariableSink OnLandscapeCollisionHeightfieldChangedSink(FConsoleCommandDelegate::CreateStatic(&OnLandscapeShowCollisionMeshChanged));
+
+#if !(UE_BUILD_SHIPPING || UE_BUILD_TEST) && WITH_EDITORONLY_DATA
+FPrimitiveSceneProxy* ULandscapeHeightfieldCollisionComponent::CreateSceneProxy()
+{
+	class FLandscapeHeightfieldCollisionComponentSceneProxy final : public FPrimitiveSceneProxy
+	{
+	public:
+		SIZE_T GetTypeHash() const override
+		{
+			static size_t UniquePointer;
+			return reinterpret_cast<size_t>(&UniquePointer);
+		}
+
+		// Constructor exists to populate Vertices and Indices arrays which are
+		// used to construct the collision mesh inside GetDynamicMeshElements
+		FLandscapeHeightfieldCollisionComponentSceneProxy(const ULandscapeHeightfieldCollisionComponent* InComponent, const Chaos::FHeightField& InHeightfield, const FLinearColor& InWireframeColor, bool bInShowHoles)
+			: FPrimitiveSceneProxy(InComponent)
+		{
+			const Chaos::FHeightField::FData<uint16>& GeomData = InHeightfield.GeomData;
+			const int32 CollisionSizeVerts = InComponent->CollisionSizeQuads + 1;
+			const int32 SimpleCollisionSizeVerts = InComponent->SimpleCollisionSizeQuads > 0 ? InComponent->SimpleCollisionSizeQuads + 1 : 0;
+			const int32 NumVerts = FMath::Square(CollisionSizeVerts);
+			const int32 NumSimpleVerts = FMath::Square(SimpleCollisionSizeVerts);
+
+			check(InComponent->CollisionHeightData.GetElementCount() == NumVerts + NumSimpleVerts);
+
+			Vertices.SetNumUninitialized(NumVerts);
+			for (int32 i = 0; i < NumVerts; i++)
+			{
+				int32 X = i % CollisionSizeVerts;
+				int32 Y = i / CollisionSizeVerts;
+				Vertices[i] = FVector3f(X, Y, (GeomData.MinValue + (float)GeomData.Heights[i] * GeomData.HeightPerUnit) * LANDSCAPE_ZSCALE);
+			}
+
+			const int32 NumTris = FMath::Square(InComponent->CollisionSizeQuads) * 2;
+			Indices.SetNumUninitialized(NumTris * 3);
+
+			const uint8* DominantLayers = nullptr;
+			if (InComponent->DominantLayerData.GetElementCount() > 0)
+			{
+				DominantLayers = (const uint8*)InComponent->DominantLayerData.LockReadOnly();
+			}
+
+			int32 TriangleIdx = 0;
+			for (int32 Y = 0; Y < InComponent->CollisionSizeQuads; Y++)
+			{
+				for (int32 X = 0; X < InComponent->CollisionSizeQuads; X++)
+				{
+					int32 DataIdx = X + Y * CollisionSizeVerts;
+					bool bHole = false;
+
+					if (bInShowHoles && DominantLayers)
+					{
+						uint8 DominantLayerIdx = DominantLayers[DataIdx];
+						if (InComponent->ComponentLayerInfos.IsValidIndex(DominantLayerIdx) && 
+							InComponent->ComponentLayerInfos[DominantLayerIdx] == ALandscapeProxy::VisibilityLayer)
+						{
+							// If it's a hole, override with the hole flag.
+							bHole = true;
+						}
+					}
+
+					if (bHole)
+					{
+						Indices[TriangleIdx + 0] = (X + 0) + (Y + 0) * CollisionSizeVerts;
+						Indices[TriangleIdx + 1] = Indices[TriangleIdx + 0];
+						Indices[TriangleIdx + 2] = Indices[TriangleIdx + 0];
+					}
+					else
+					{
+						Indices[TriangleIdx + 0] = (X + 0) + (Y + 0) * CollisionSizeVerts;
+						Indices[TriangleIdx + 1] = (X + 1) + (Y + 1) * CollisionSizeVerts;
+						Indices[TriangleIdx + 2] = (X + 1) + (Y + 0) * CollisionSizeVerts;
+					}
+
+					TriangleIdx += 3;
+
+					if (bHole)
+					{
+						Indices[TriangleIdx + 0] = (X + 0) + (Y + 0) * CollisionSizeVerts;
+						Indices[TriangleIdx + 1] = Indices[TriangleIdx + 0];
+						Indices[TriangleIdx + 2] = Indices[TriangleIdx + 0];
+					}
+					else
+					{
+						Indices[TriangleIdx + 0] = (X + 0) + (Y + 0) * CollisionSizeVerts;
+						Indices[TriangleIdx + 1] = (X + 0) + (Y + 1) * CollisionSizeVerts;
+						Indices[TriangleIdx + 2] = (X + 1) + (Y + 1) * CollisionSizeVerts;
+					}
+
+					TriangleIdx += 3;
+				}
+			}
+
+			if (DominantLayers)
+			{
+				InComponent->DominantLayerData.Unlock();
+			}
+
+			WireframeMaterialInstance.Reset(new FColoredMaterialRenderProxy(
+				GEngine->WireframeMaterial ? GEngine->WireframeMaterial->GetRenderProxy() : NULL,
+				InWireframeColor));
+		}
+
+		virtual void GetDynamicMeshElements(const TArray<const FSceneView*>& Views, const FSceneViewFamily& ViewFamily, uint32 VisibilityMap, FMeshElementCollector& Collector) const override
+		{
+			const FMatrix& LocalToWorld = GetLocalToWorld();
+
+			const bool bDrawCollision = ViewFamily.EngineShowFlags.Collision && IsCollisionEnabled();
+
+			if (bDrawCollision && AllowDebugViewmodes() && WireframeMaterialInstance.IsValid())
+			{
+				// Set up mesh builder
+				FDynamicMeshBuilder MeshBuilder(ERHIFeatureLevel::Type::SM6);
+				MeshBuilder.AddVertices(Vertices);
+				MeshBuilder.AddTriangles(Indices);
+
+				// For each view, pass the mesh to the collector
+				for (int32 ViewIndex = 0; ViewIndex != Views.Num(); ++ViewIndex)
+				{
+					MeshBuilder.GetMesh(LocalToWorld, WireframeMaterialInstance.Get(), SDPG_World, false, false, ViewIndex, Collector);
+				}
+			}
+		}
+
+		virtual FPrimitiveViewRelevance GetViewRelevance(const FSceneView* View) const override
+		{
+			// Should we draw this because collision drawing is enabled, and we have collision
+			const bool bShowForCollision = View->Family->EngineShowFlags.Collision && IsCollisionEnabled();
+
+			FPrimitiveViewRelevance Result;
+			Result.bDrawRelevance = IsShown(View) || bShowForCollision;
+			Result.bDynamicRelevance = true;
+			Result.bShadowRelevance = false;
+			Result.bEditorPrimitiveRelevance = UseEditorCompositing(View);
+			return Result;
+		}
+
+		virtual uint32 GetMemoryFootprint(void) const override { return(sizeof(*this) + GetAllocatedSize()); }
+		uint32 GetAllocatedSize(void) const { return FPrimitiveSceneProxy::GetAllocatedSize(); }
+
+	private:
+		TArray<FDynamicMeshVertex> Vertices;
+		TArray<uint32> Indices;
+		TUniquePtr<FColoredMaterialRenderProxy> WireframeMaterialInstance = nullptr;
+	};
+
+	enum class EHeightfieldSource
+	{
+		Simple = 0,
+		Complex = 1,
+		Editor = 2
+	};
+
+	FLandscapeHeightfieldCollisionComponentSceneProxy* Proxy = nullptr;
+
+	if (HeightfieldRef.IsValid() && IsValidRef(HeightfieldRef))
+	{
+		const Chaos::FHeightField* LocalHeightfield = nullptr;
+		FLinearColor WireframeColor;
+		bool bShowHoles = true;
+
+		switch (static_cast<EHeightfieldSource>(CVarLandscapeShowCollisionMesh.GetValueOnAnyThread()))
+		{
+		case EHeightfieldSource::Simple:
+			if (RenderComponent->SimpleCollisionMipLevel > RenderComponent->CollisionMipLevel)
+			{
+				if (HeightfieldRef->HeightfieldSimple.IsValid())
+				{
+					LocalHeightfield = HeightfieldRef->HeightfieldSimple.Get();
+				}
+			}
+			else
+			{
+				if (HeightfieldRef->Heightfield.IsValid())
+				{
+					LocalHeightfield = HeightfieldRef->Heightfield.Get();
+				}
+			}
+
+			WireframeColor = FColor(157, 149, 223, 255);
+			break;
+
+		case EHeightfieldSource::Complex:
+			if (HeightfieldRef->Heightfield.IsValid())
+			{
+				LocalHeightfield = HeightfieldRef->Heightfield.Get();
+			}
+
+			WireframeColor = FColor(0, 255, 255, 255);
+			break;
+
+		case EHeightfieldSource::Editor:
+			if (HeightfieldRef->EditorHeightfield.IsValid())
+			{
+				LocalHeightfield = HeightfieldRef->EditorHeightfield.Get();
+			}
+
+			WireframeColor = FColor(157, 223, 149, 255);
+			bShowHoles = false;
+			break;
+
+		default:
+			UE_LOG(LogLandscape, Warning, TEXT("Invalid Value for CVar landscape.ShowCollisionMesh"));
+		}
+
+		if (LocalHeightfield != nullptr)
+		{
+			Proxy = new FLandscapeHeightfieldCollisionComponentSceneProxy(this, *LocalHeightfield, WireframeColor, bShowHoles);
+		}
+	}
+
+	return Proxy;
+}
+#endif // !(UE_BUILD_SHIPPING || UE_BUILD_TEST) && WITH_EDITORONLY_DATA
+
 void ULandscapeHeightfieldCollisionComponent::CreateCollisionObject()
 {
 	LLM_SCOPE(ELLMTag::ChaosLandscape);
@@ -1474,6 +1709,8 @@ bool ULandscapeHeightfieldCollisionComponent::RecreateCollision()
 		HeightfieldGuid = FGuid();
 
 		RecreatePhysicsState();
+
+		MarkRenderStateDirty();
 	}
 	return true;
 }
