Re-enable PrecomputedLightingUniformBuffer for landscape on mobile as GPUScene support hasn't been implemented there

diff --git a/Engine/Source/Runtime/Landscape/Private/LandscapeRender.cpp b/Engine/Source/Runtime/Landscape/Private/LandscapeRender.cpp
--- a/Engine/Source/Runtime/Landscape/Private/LandscapeRender.cpp
+++ b/Engine/Source/Runtime/Landscape/Private/LandscapeRender.cpp
@@ -1252,7 +1252,7 @@ FLandscapeComponentSceneProxy::FLandscapeComponentSceneProxy(ULandscapeComponent
 		StaticLightingResolution = InComponent->GetLandscapeProxy()->StaticLightingResolution;
 	}
 
-	ComponentLightInfo = MakeUnique<FLandscapeLCI>(InComponent);
+	ComponentLightInfo = MakeUnique<FLandscapeLCI>(InComponent, FeatureLevel);
 	check(ComponentLightInfo);
 
 	const bool bHasStaticLighting = ComponentLightInfo->GetLightMap() || ComponentLightInfo->GetShadowMap();
diff --git a/Engine/Source/Runtime/Landscape/Public/LandscapeRender.h b/Engine/Source/Runtime/Landscape/Public/LandscapeRender.h
--- a/Engine/Source/Runtime/Landscape/Public/LandscapeRender.h
+++ b/Engine/Source/Runtime/Landscape/Public/LandscapeRender.h
@@ -581,7 +581,7 @@ class LANDSCAPE_API FLandscapeComponentSceneProxy : public FPrimitiveSceneProxy,
 	{
 	public:
 		/** Initialization constructor. */
-		FLandscapeLCI(const ULandscapeComponent* InComponent)
+		FLandscapeLCI(const ULandscapeComponent* InComponent, ERHIFeatureLevel::Type FeatureLevel)
 			: FLightCacheInterface()
 		{
 			const FMeshMapBuildData* MapBuildData = InComponent->GetMeshMapBuildData();
@@ -591,7 +591,12 @@ class LANDSCAPE_API FLandscapeComponentSceneProxy : public FPrimitiveSceneProxy,
 				SetLightMap(MapBuildData->LightMap);
 				SetShadowMap(MapBuildData->ShadowMap);
 				SetResourceCluster(MapBuildData->ResourceCluster);
-				bCanUsePrecomputedLightingParametersFromGPUScene = true;
+				if (FeatureLevel >= ERHIFeatureLevel::SM5)
+				{
+					// Landscape does not support GPUScene on mobile
+					// TODO: enable this when GPUScene support is implemented
+					bCanUsePrecomputedLightingParametersFromGPUScene = true;
+				}
 				IrrelevantLights = MapBuildData->IrrelevantLights;
 			}
 		}
