Fixed crash when lazy object ptr for render component is invalid when landscape collision scene proxy is created

#rnx

diff --git a/Engine/Source/Runtime/Landscape/Private/LandscapeCollision.cpp b/Engine/Source/Runtime/Landscape/Private/LandscapeCollision.cpp
--- a/Engine/Source/Runtime/Landscape/Private/LandscapeCollision.cpp
+++ b/Engine/Source/Runtime/Landscape/Private/LandscapeCollision.cpp
@@ -627,7 +627,7 @@ FPrimitiveSceneProxy* ULandscapeHeightfieldCollisionComponent::CreateSceneProxy(
 		switch (static_cast<EHeightfieldSource>(CVarLandscapeShowCollisionMesh.GetValueOnAnyThread()))
 		{
 		case EHeightfieldSource::Simple:
-			if (RenderComponent->SimpleCollisionMipLevel > RenderComponent->CollisionMipLevel)
+			if (RenderComponent.IsValid() && (RenderComponent->SimpleCollisionMipLevel > RenderComponent->CollisionMipLevel))
 			{
 				if (HeightfieldRef->HeightfieldSimple.IsValid())
 				{
