Fixed: Errors in Incremental Cook Lyra

diff --git a/Engine/Source/Runtime/Landscape/Private/Landscape.cpp b/Engine/Source/Runtime/Landscape/Private/Landscape.cpp
--- a/Engine/Source/Runtime/Landscape/Private/Landscape.cpp
+++ b/Engine/Source/Runtime/Landscape/Private/Landscape.cpp
@@ -290,7 +290,10 @@ void ULandscapeComponent::BeginCacheForCookedPlatformData(const ITargetPlatform*
 
 	if (!HasAnyFlags(RF_ClassDefaultObject))
 	{
-		CheckGenerateMobilePlatformData(/*bIsCooking = */ true, TargetPlatform);
+		if (TargetPlatform->SupportsFeature(ETargetPlatformFeatures::MobileRendering))
+		{
+			CheckGenerateMobilePlatformData(/*bIsCooking = */ true, TargetPlatform);
+		}
 	}
 }
 
@@ -536,7 +539,10 @@ void ULandscapeComponent::Serialize(FArchive& Ar)
 		// for -oldcook:
 		// the old cooker calls BeginCacheForCookedPlatformData after the package export set is tagged, so the mobile material doesn't get saved, so we have to do CheckGenerateMobilePlatformData in serialize
 		// the new cooker clears the texture source data before calling serialize, causing GeneratePlatformVertexData to crash, so we have to do CheckGenerateMobilePlatformData in BeginCacheForCookedPlatformData
-		CheckGenerateMobilePlatformData(/*bIsCooking = */ true, Ar.CookingTarget());
+		if (Ar.CookingTarget()->SupportsFeature(ETargetPlatformFeatures::MobileRendering))
+		{
+			CheckGenerateMobilePlatformData(/*bIsCooking = */ true, Ar.CookingTarget());
+		}
 	}
 
 	// Avoid the archiver in the PIE duplicate writer case because we want to share landscape textures & materials
