Fixed: Errors in Incremental Cook ShooterGame Android

diff --git a/Engine/Source/Runtime/Landscape/Private/Landscape.cpp b/Engine/Source/Runtime/Landscape/Private/Landscape.cpp
--- a/Engine/Source/Runtime/Landscape/Private/Landscape.cpp
+++ b/Engine/Source/Runtime/Landscape/Private/Landscape.cpp
@@ -439,7 +439,7 @@ void ULandscapeComponent::CheckGenerateMobilePlatformData(bool bIsCooking, const
 	FGuid NewSourceHash = FGuid(Hash[0] ^ Hash[4], Hash[1], Hash[2], Hash[3]);
 
 	bool bHashMismatch = MobileDataSourceHash != NewSourceHash;
-	bool bMissingPixelData = MobileMaterialInterfaces.Num() == 0 || MobileWeightmapTextures.Num() == 0 || MaterialPerLOD.Num() == 0;
+	bool bMissingPixelData = MobileMaterialInterfaces.Num() == 0 || MaterialPerLOD.Num() == 0;
 	bool bRegeneratePixelData = bMissingPixelData || bHashMismatch;
 
 	if (bRegeneratePixelData)
