Landscape Collision View: Remove dependency on RenderComponent when creating heightfield scene proxy.

diff --git a/Engine/Source/Runtime/Landscape/Private/LandscapeCollision.cpp b/Engine/Source/Runtime/Landscape/Private/LandscapeCollision.cpp
--- a/Engine/Source/Runtime/Landscape/Private/LandscapeCollision.cpp
+++ b/Engine/Source/Runtime/Landscape/Private/LandscapeCollision.cpp
@@ -627,19 +627,13 @@ FPrimitiveSceneProxy* ULandscapeHeightfieldCollisionComponent::CreateSceneProxy(
 		switch (static_cast<EHeightfieldSource>(CVarLandscapeShowCollisionMesh.GetValueOnAnyThread()))
 		{
 		case EHeightfieldSource::Simple:
-			if (RenderComponent.IsValid() && (RenderComponent->SimpleCollisionMipLevel > RenderComponent->CollisionMipLevel))
+			if (HeightfieldRef->HeightfieldSimple.IsValid())
 			{
-				if (HeightfieldRef->HeightfieldSimple.IsValid())
-				{
-					LocalHeightfield = HeightfieldRef->HeightfieldSimple.Get();
-				}
+				LocalHeightfield = HeightfieldRef->HeightfieldSimple.Get();
 			}
-			else
+			else if (HeightfieldRef->Heightfield.IsValid())
 			{
-				if (HeightfieldRef->Heightfield.IsValid())
-				{
-					LocalHeightfield = HeightfieldRef->Heightfield.Get();
-				}
+				LocalHeightfield = HeightfieldRef->Heightfield.Get();
 			}
 
 			WireframeColor = FColor(157, 149, 223, 255);
