* Fixed landscape edit layers not properly removing weightmap allocations on components with unused paint layers in the procedural brush case
* Added list of layers affected to each channel to the weightmap readback struct, in order to properly analyze the content of the texture when the weightmap layer allocations might have changed on the component since the moment the readback was issued

#tests UEENGQA-74353

diff --git a/Engine/Source/Runtime/Landscape/Classes/LandscapeComponent.h b/Engine/Source/Runtime/Landscape/Classes/LandscapeComponent.h
--- a/Engine/Source/Runtime/Landscape/Classes/LandscapeComponent.h
+++ b/Engine/Source/Runtime/Landscape/Classes/LandscapeComponent.h
@@ -759,9 +759,12 @@ public:
 
 #if WITH_EDITOR
 
-	/** Deletes a layer from this component, removing all its data */
+	/** Deletes a layer from this component, removing all its data, adjusting other layer's weightmaps if necessary, etc. */
 	LANDSCAPE_API void DeleteLayer(ULandscapeLayerInfoObject* LayerInfo, FLandscapeEditDataInterface& LandscapeEdit);
 
+	/** Deletes a layer from this component, but doesn't do anything else (assumes the user knows what he's doing, use DeleteLayer otherwise) */
+	void DeleteLayerAllocation(const FGuid& InEditLayerGuid, int32 InLayerIdx, bool bInShouldDirtyPackage);
+
 	/** Fills a layer to 100% on this component, adding it if needed and removing other layers that get painted away */
 	LANDSCAPE_API void FillLayer(ULandscapeLayerInfoObject* LayerInfo, FLandscapeEditDataInterface& LandscapeEdit);
 
diff --git a/Engine/Source/Runtime/Landscape/Private/Landscape.cpp b/Engine/Source/Runtime/Landscape/Private/Landscape.cpp
--- a/Engine/Source/Runtime/Landscape/Private/Landscape.cpp
+++ b/Engine/Source/Runtime/Landscape/Private/Landscape.cpp
@@ -2139,6 +2139,54 @@ void ULandscapeComponent::SetWeightmapTexturesUsage(const TArray<ULandscapeWeigh
 	}
 }
 
+void ULandscapeComponent::DeleteLayerAllocation(const FGuid& InEditLayerGuid, int32 InLayerIdx, bool bInShouldDirtyPackage)
+{
+	TArray<FWeightmapLayerAllocationInfo>& ComponentWeightmapLayerAllocations = GetWeightmapLayerAllocations(InEditLayerGuid);
+	TArray<UTexture2D*>& ComponentWeightmapTextures = GetWeightmapTextures(InEditLayerGuid);
+	TArray<ULandscapeWeightmapUsage*>& ComponentWeightmapTexturesUsage = GetWeightmapTexturesUsage(InEditLayerGuid);
+	const FWeightmapLayerAllocationInfo& LayerAllocation = ComponentWeightmapLayerAllocations[InLayerIdx];
+	const int32 DeleteLayerWeightmapTextureIndex = LayerAllocation.WeightmapTextureIndex;
+
+	ALandscapeProxy* Proxy = GetLandscapeProxy();
+	Modify(bInShouldDirtyPackage);
+	Proxy->Modify(bInShouldDirtyPackage);
+
+	// Mark the channel as unallocated, so we can reuse it later
+	ULandscapeWeightmapUsage* Usage = ComponentWeightmapTexturesUsage.IsValidIndex(DeleteLayerWeightmapTextureIndex) ? ComponentWeightmapTexturesUsage[DeleteLayerWeightmapTextureIndex] : nullptr;
+	if (Usage) // can be null if WeightmapUsageMap hasn't been built yet
+	{
+		Usage->ChannelUsage[LayerAllocation.WeightmapTextureChannel] = nullptr;
+	}
+
+	// Remove the layer:
+	ComponentWeightmapLayerAllocations.RemoveAt(InLayerIdx);
+
+	// Check if the weightmap texture used by the layer we just removed is used by any other layer, and if so, remove the texture too
+	bool bCanRemoveLayerTexture = !ComponentWeightmapLayerAllocations.ContainsByPredicate([DeleteLayerWeightmapTextureIndex](const FWeightmapLayerAllocationInfo& Allocation) { return Allocation.WeightmapTextureIndex == DeleteLayerWeightmapTextureIndex; });
+	if (bCanRemoveLayerTexture)
+	{
+		ComponentWeightmapTextures[DeleteLayerWeightmapTextureIndex]->ClearFlags(RF_Standalone);
+		ComponentWeightmapTextures.RemoveAt(DeleteLayerWeightmapTextureIndex);
+		if (Usage)
+		{
+			ComponentWeightmapTexturesUsage.RemoveAt(DeleteLayerWeightmapTextureIndex);
+		}
+
+		// Adjust WeightmapTextureChannel index for other layers
+		for (FWeightmapLayerAllocationInfo& Allocation : ComponentWeightmapLayerAllocations)
+		{
+			if (Allocation.WeightmapTextureIndex > DeleteLayerWeightmapTextureIndex)
+			{
+				Allocation.WeightmapTextureIndex--;
+			}
+
+			check(Allocation.WeightmapTextureIndex < ComponentWeightmapTextures.Num());
+		}
+	}
+
+	Proxy->ValidateProxyLayersWeightmapUsage();
+}
+
 #endif // WITH_EDITOR
 
 void ALandscapeProxy::PostRegisterAllComponents()
diff --git a/Engine/Source/Runtime/Landscape/Private/LandscapeEditInterface.cpp b/Engine/Source/Runtime/Landscape/Private/LandscapeEditInterface.cpp
--- a/Engine/Source/Runtime/Landscape/Private/LandscapeEditInterface.cpp
+++ b/Engine/Source/Runtime/Landscape/Private/LandscapeEditInterface.cpp
@@ -1586,9 +1586,9 @@ void ULandscapeComponent::DeleteLayer(ULandscapeLayerInfoObject* LayerInfo, FLan
 	TRACE_CPUPROFILER_EVENT_SCOPE(LandscapeComponent_DeleteLayer);
 
 	ULandscapeComponent* Component = this;
-	TArray<FWeightmapLayerAllocationInfo>& ComponentWeightmapLayerAllocations = Component->GetWeightmapLayerAllocations(true);
-	TArray<UTexture2D*>& ComponentWeightmapTextures = Component->GetWeightmapTextures(true);
-	TArray<ULandscapeWeightmapUsage*>& ComponentWeightmapTexturesUsage = Component->GetWeightmapTexturesUsage(true);
+	FGuid EditLayerGuid = Component->GetEditingLayerGUID();
+	const TArray<FWeightmapLayerAllocationInfo>& ComponentWeightmapLayerAllocations = Component->GetWeightmapLayerAllocations(EditLayerGuid);
+	const TArray<UTexture2D*>& ComponentWeightmapTextures = Component->GetWeightmapTextures(EditLayerGuid);
 
 	// Find the index for this layer in this component.
 	const int32 DeleteLayerIdx = ComponentWeightmapLayerAllocations.IndexOfByPredicate(
@@ -1599,22 +1599,7 @@ void ULandscapeComponent::DeleteLayer(ULandscapeLayerInfoObject* LayerInfo, FLan
 		return;
 	}
 
-	FWeightmapLayerAllocationInfo& DeleteLayerAllocation = ComponentWeightmapLayerAllocations[DeleteLayerIdx];
-	int32 DeleteLayerWeightmapTextureIndex = DeleteLayerAllocation.WeightmapTextureIndex;
-
-	// See if we'll be able to remove the texture completely.
-	bool bCanRemoveLayerTexture = true;
-	for (int32 LayerIdx = 0; LayerIdx < ComponentWeightmapLayerAllocations.Num(); LayerIdx++)
-	{
-		FWeightmapLayerAllocationInfo& Allocation = ComponentWeightmapLayerAllocations[LayerIdx];
-
-		// check if we will be able to remove the texture also
-		if (LayerIdx != DeleteLayerIdx && Allocation.WeightmapTextureIndex == DeleteLayerWeightmapTextureIndex)
-		{
-			bCanRemoveLayerTexture = false;
-			break;
-		}
-	}
+	Component->DeleteLayerAllocation(EditLayerGuid, DeleteLayerIdx, LandscapeEdit.GetShouldDirtyPackage());
 
 	// See if the deleted layer is a NoWeightBlend layer - if not, we don't have to worry about normalization
 	bool bDeleteLayerIsNoWeightBlend = (LayerInfo && LayerInfo->bNoWeightBlend);
@@ -1634,7 +1619,7 @@ void ULandscapeComponent::DeleteLayer(ULandscapeLayerInfoObject* LayerInfo, FLan
 		// Get the data for each layer
 		for (int32 LayerIdx = 0; LayerIdx < ComponentWeightmapLayerAllocations.Num(); LayerIdx++)
 		{
-			FWeightmapLayerAllocationInfo& Allocation = ComponentWeightmapLayerAllocations[LayerIdx];
+			const FWeightmapLayerAllocationInfo& Allocation = ComponentWeightmapLayerAllocations[LayerIdx];
 			LayerDataPtrs.Add((uint8*)TexDataInfos[Allocation.WeightmapTextureIndex]->GetMipData(0) + ChannelOffsets[Allocation.WeightmapTextureChannel]);
 
 			// Find the layer info and record if it is a bNoWeightBlend layer.
@@ -1704,12 +1689,6 @@ void ULandscapeComponent::DeleteLayer(ULandscapeLayerInfoObject* LayerInfo, FLan
 		// Update all the textures and mips
 		for (int32 Idx = 0; Idx < ComponentWeightmapTextures.Num(); Idx++)
 		{
-			if (bCanRemoveLayerTexture && Idx == DeleteLayerWeightmapTextureIndex)
-			{
-				// We're going to remove this texture anyway, so don't bother updating
-				continue;
-			}
-
 			UTexture2D* WeightmapTexture = ComponentWeightmapTextures[Idx];
 			FLandscapeTextureDataInfo* WeightmapDataInfo = TexDataInfos[Idx];
 
@@ -1727,49 +1706,6 @@ void ULandscapeComponent::DeleteLayer(ULandscapeLayerInfoObject* LayerInfo, FLan
 		}
 	}
 
-	// Mark the channel as unallocated, so we can reuse it later
-	ALandscapeProxy* Proxy = Component->GetLandscapeProxy();
-	Component->Modify(LandscapeEdit.GetShouldDirtyPackage());
-	Proxy->Modify(LandscapeEdit.GetShouldDirtyPackage());
-
-	ULandscapeWeightmapUsage* Usage = ComponentWeightmapTexturesUsage.IsValidIndex(DeleteLayerAllocation.WeightmapTextureIndex) ? ComponentWeightmapTexturesUsage[DeleteLayerAllocation.WeightmapTextureIndex] : nullptr;
-	if (Usage) // can be null if WeightmapUsageMap hasn't been built yet
-	{
-		Usage->ChannelUsage[DeleteLayerAllocation.WeightmapTextureChannel] = nullptr;
-	}
-
-	// Remove the layer
-	ComponentWeightmapLayerAllocations.RemoveAt(DeleteLayerIdx);
-
-	// If this layer was the last usage for this channel in this layer, we can remove it.
-	if (bCanRemoveLayerTexture)
-	{
-		ComponentWeightmapTextures[DeleteLayerWeightmapTextureIndex]->SetFlags(RF_Transactional);
-		ComponentWeightmapTextures[DeleteLayerWeightmapTextureIndex]->Modify(LandscapeEdit.GetShouldDirtyPackage());
-		ComponentWeightmapTextures[DeleteLayerWeightmapTextureIndex]->ClearFlags(RF_Standalone);
-
-		// possible that usages have not been built yet
-		if (ComponentWeightmapTexturesUsage.IsValidIndex(DeleteLayerWeightmapTextureIndex))
-		{
-			ComponentWeightmapTexturesUsage.RemoveAt(DeleteLayerWeightmapTextureIndex);
-		}
-
-		ComponentWeightmapTextures.RemoveAt(DeleteLayerWeightmapTextureIndex);
-
-		// Adjust WeightmapTextureIndex index for other layers
-		for (int32 LayerIdx = 0; LayerIdx < ComponentWeightmapLayerAllocations.Num(); LayerIdx++)
-		{
-			FWeightmapLayerAllocationInfo& Allocation = ComponentWeightmapLayerAllocations[LayerIdx];
-
-			if (Allocation.WeightmapTextureIndex > DeleteLayerWeightmapTextureIndex)
-			{
-				Allocation.WeightmapTextureIndex--;
-			}
-
-			check(Allocation.WeightmapTextureIndex < ComponentWeightmapTextures.Num());
-		}
-	}
-
 	// Update the shaders for this component
 	if (!Component->GetLandscapeProxy()->HasLayersContent())
 	{
@@ -1805,8 +1741,6 @@ void ULandscapeComponent::DeleteLayer(ULandscapeLayerInfoObject* LayerInfo, FLan
 			LocalCollisionComponent->RecreateCollision();
 		}
 	}
-
-	Proxy->ValidateProxyLayersWeightmapUsage();
 }
 
 void FLandscapeEditDataInterface::DeleteLayer(ULandscapeLayerInfoObject* LayerInfo)
@@ -2500,42 +2434,7 @@ bool DeleteLayerIfAllZero(ULandscapeComponent* const Component, const uint8* con
 		}
 	}
 
-	ALandscapeProxy* Proxy = Component->GetLandscapeProxy();
-	Component->Modify(bShouldDirtyPackage);
-	Proxy->Modify(bShouldDirtyPackage);
-
-	TArray<FWeightmapLayerAllocationInfo>& ComponentWeightmapLayerAllocations = Component->GetWeightmapLayerAllocations(true);
-	TArray<UTexture2D*>& ComponentWeightmapTextures = Component->GetWeightmapTextures(true);
-	TArray<ULandscapeWeightmapUsage*>& ComponentWeightmapTexturesUsage = Component->GetWeightmapTexturesUsage(true);
-
-	// Mark the channel as unallocated, so we can reuse it later
-	const int32 DeleteLayerWeightmapTextureIndex = ComponentWeightmapLayerAllocations[LayerIdx].WeightmapTextureIndex;
-	ULandscapeWeightmapUsage& Usage = *ComponentWeightmapTexturesUsage[DeleteLayerWeightmapTextureIndex];
-	Usage.ChannelUsage[ComponentWeightmapLayerAllocations[LayerIdx].WeightmapTextureChannel] = nullptr;
-
-	// Remove the layer as it's totally painted away.
-	ComponentWeightmapLayerAllocations.RemoveAt(LayerIdx);
-
-	// Check if the weightmap texture used by the layer we just removed is used by any other layer, and if so, remove the texture too
-	bool bCanRemoveLayerTexture = !ComponentWeightmapLayerAllocations.ContainsByPredicate([DeleteLayerWeightmapTextureIndex](const FWeightmapLayerAllocationInfo& Allocation){ return Allocation.WeightmapTextureIndex == DeleteLayerWeightmapTextureIndex; });
-	if (bCanRemoveLayerTexture)
-	{
-		ComponentWeightmapTextures[DeleteLayerWeightmapTextureIndex]->ClearFlags(RF_Standalone);
-		ComponentWeightmapTextures.RemoveAt(DeleteLayerWeightmapTextureIndex);
-		ComponentWeightmapTexturesUsage.RemoveAt(DeleteLayerWeightmapTextureIndex);
-
-		// Adjust WeightmapTextureChannel index for other layers
-		for (auto It = ComponentWeightmapLayerAllocations.CreateIterator(); It; ++It)
-		{
-			FWeightmapLayerAllocationInfo& Allocation = *It;
-			if (Allocation.WeightmapTextureIndex > DeleteLayerWeightmapTextureIndex)
-			{
-				Allocation.WeightmapTextureIndex--;
-			}
-		}
-	}
-
-	Component->GetLandscapeProxy()->ValidateProxyLayersWeightmapUsage();
+	Component->DeleteLayerAllocation(Component->GetEditingLayerGUID(), LayerIdx, bShouldDirtyPackage);
 
 	return true;
 }
diff --git a/Engine/Source/Runtime/Landscape/Private/LandscapeEditLayers.cpp b/Engine/Source/Runtime/Landscape/Private/LandscapeEditLayers.cpp
--- a/Engine/Source/Runtime/Landscape/Private/LandscapeEditLayers.cpp
+++ b/Engine/Source/Runtime/Landscape/Private/LandscapeEditLayers.cpp
@@ -147,6 +147,11 @@ TAutoConsoleVariable<int32> CVarLandscapeValidateProxyWeightmapUsages(
 	1,
 	TEXT("This will validate that weightmap usages in landscape proxies and their components don't get desynchronized with the landscape component layer allocations."));
 
+TAutoConsoleVariable<int32> CVarLandscapeRemoveEmptyPaintLayersOnEdit(
+	TEXT("landscape.RemoveEmptyPaintLayersOnEdit"),
+	1,
+	TEXT("This will analyze weightmaps on readback and remove unneeded allocations (for unpainted layers)."));
+
 void OnLandscapeEditLayersLocalMergeChanged(IConsoleVariable* CVar)
 {
 	for (TObjectIterator<UWorld> It; It; ++It)
@@ -2518,8 +2523,10 @@ struct FLandscapeEditLayerComponentReadbackResult
 	uint32 UpdateModes = 0;
 	/** Were the associated heightmap/weightmaps modified. */
 	bool bModified = false;
+	/** Indicates which of the component's weightmaps is not needed anymore. */
+	TArray<ULandscapeLayerInfoObject*> AllZeroLayers;
 
-	FLandscapeEditLayerComponentReadbackResult(ULandscapeComponent* InLandscapeComponent, uint32 InUpdateModes)
+	FLandscapeEditLayerComponentReadbackResult(ULandscapeComponent* InLandscapeComponent, uint32 InUpdateModes = 0)
 		: LandscapeComponent(InLandscapeComponent)
 		, UpdateModes(InUpdateModes)
 	{}
@@ -2567,7 +2574,25 @@ TArray<FLandscapeLayersCopyReadbackTextureParams> PrepareLandscapeLayersCopyRead
 		{
 			const FIntPoint ComponentToResolveKey = ComponentToResolve->GetSectionBase() / ComponentToResolve->ComponentSizeQuads;
 			const int32 ComponentToResolveFlags = ComponentToResolve->GetLayerUpdateFlagPerMode();
-			CopyReadbackTextureParams.Context.Add(FLandscapeEditLayerReadback::FComponentReadbackContext(ComponentToResolveKey, ComponentToResolveFlags));
+			FLandscapeEditLayerReadback::FPerChannelLayerNames PerChannelLayerNames;
+
+			// Weightmaps could be reallocated randomly before we actually perform the readback, so we need to keep a picture of which channel was affected to which paint layer at readback time:
+			if (bWeightmaps)
+			{
+				const TArray<UTexture2D*>& WeightmapTextures = ComponentToResolve->GetWeightmapTextures(/*InReturnEditingWeightmap = */false);
+				for (FWeightmapLayerAllocationInfo const& AllocInfo : ComponentToResolve->GetWeightmapLayerAllocations(/*InReturnEditingWeightmap = */false))
+				{
+					if (AllocInfo.IsAllocated())
+					{
+						UTexture2D* PaintLayerTexture = WeightmapTextures[AllocInfo.WeightmapTextureIndex];
+						if (PaintLayerTexture == Texture)
+						{
+							PerChannelLayerNames[AllocInfo.WeightmapTextureChannel] = AllocInfo.LayerInfo->LayerName;
+						}
+					}
+				}
+			}
+			CopyReadbackTextureParams.Context.Add(FLandscapeEditLayerReadback::FComponentReadbackContext { ComponentToResolveKey, ComponentToResolveFlags, PerChannelLayerNames });
 		}
 	}
 
@@ -3347,7 +3372,7 @@ bool ALandscape::PrepareTextureResources(bool bInWaitForStreaming)
 	{
 		return false;
 	}
-	
+
 	// Only keep the textures that are still valid:
 	TSet<UTexture2D*> StreamingInTexturesBefore;
 	StreamingInTexturesBefore.Reserve(TrackedStreamingInTextures.Num());	
@@ -5231,11 +5256,31 @@ bool ALandscape::ResolveLayersTexture(
 			}
 		}
 
+		// Find out whether some channels from this weightmap are now all zeros : 
+		static constexpr uint32 AllChannelsAllZerosMask = 15;
+		uint32 AllZerosTextureChannelMask = AllChannelsAllZerosMask;
+		const bool bCheckForEmptyChannels = CVarLandscapeRemoveEmptyPaintLayersOnEdit.GetValueOnGameThread() != 0;
+		if (bIsWeightmap && bCheckForEmptyChannels)
+		{
+			TRACE_CPUPROFILER_EVENT_SCOPE(LandscapeLayers_AnalyzeWeightmap);
+			const FColor* TextureData = reinterpret_cast<const FColor*>(InOutputTexture->Source.LockMipReadOnly(0));
+			const int32 TexSize = OutMipsData[0].Num();
+			// We can stop iterating as soon as all of the channels are non-zero :
+			for (int32 Index = 0; (Index < TexSize) && (AllZerosTextureChannelMask != 0); ++Index)
+			{
+				AllZerosTextureChannelMask &= (((TextureData[Index].R == 0) ? 1 : 0) << 0)
+					| (((TextureData[Index].G == 0) ? 1 : 0) << 1)
+					| (((TextureData[Index].B == 0) ? 1 : 0) << 2)
+					| (((TextureData[Index].A == 0) ? 1 : 0) << 3);
+			}
+			InOutputTexture->Source.UnlockMip(0);
+		}
+
 		// Process component flags from all result contexts.
 		for (int32 ResultIndex = 0; ResultIndex < CompletedReadbackNum; ++ResultIndex)
 		{
-			FLandscapeEditLayerReadback::FReadbackContext const& ResultContext = InCPUReadback->GetResultContext(ResultIndex);
-			for (auto ComponentContext : ResultContext)
+			const FLandscapeEditLayerReadback::FReadbackContext& ResultContext = InCPUReadback->GetResultContext(ResultIndex);
+			for (const FLandscapeEditLayerReadback::FComponentReadbackContext& ComponentContext : ResultContext)
 			{
 				ULandscapeComponent** Component = GetLandscapeInfo()->XYtoComponentMap.Find(ComponentContext.ComponentKey);
 				if (Component != nullptr && *Component != nullptr)
@@ -5251,6 +5296,43 @@ bool ALandscape::ResolveLayersTexture(
 			}
 		}
 
+		// We need to find the weightmap layers that are effectively empty in order to let the component clean them up eventually :
+		if (bIsWeightmap && bCheckForEmptyChannels && (AllZerosTextureChannelMask != 0))
+		{
+			// Only use the latest readback context, since it's the only one we've actually read back : 
+			const FLandscapeEditLayerReadback::FReadbackContext& EffectiveResultContext = InCPUReadback->GetResultContext(CompletedReadbackNum - 1);
+			while (AllZerosTextureChannelMask != 0)
+			{
+				int32 AllZerosTextureChannelIndex = NumBitsPerDWORD - 1 - FMath::CountLeadingZeros(AllZerosTextureChannelMask);
+				for (const FLandscapeEditLayerReadback::FComponentReadbackContext& ComponentContext : EffectiveResultContext)
+				{
+					FName AllZerosLayerInfoName = ComponentContext.PerChannelLayerNames[AllZerosTextureChannelIndex];
+					ULandscapeComponent** Component = GetLandscapeInfo()->XYtoComponentMap.Find(ComponentContext.ComponentKey);
+					if (Component != nullptr && *Component != nullptr)
+					{
+						const TArray<FWeightmapLayerAllocationInfo>& WeightmapLayerAllocations = (*Component)->GetWeightmapLayerAllocations();
+						const TArray<UTexture2D*>& WeightmapTextures = (*Component)->GetWeightmapTextures();
+						for (const FWeightmapLayerAllocationInfo& WeightmapLayerAllocation : WeightmapLayerAllocations)
+						{
+							if (WeightmapLayerAllocation.IsAllocated())
+							{
+								UTexture2D* Texture = WeightmapTextures[WeightmapLayerAllocation.WeightmapTextureIndex];
+								if ((Texture == InOutputTexture) && (AllZerosLayerInfoName == WeightmapLayerAllocation.LayerInfo->LayerName))
+								{
+									FLandscapeEditLayerComponentReadbackResult* ComponentReadbackResult = InOutComponentReadbackResults.FindByPredicate([LandscapeComponent = *Component](const FLandscapeEditLayerComponentReadbackResult& Element) { return Element.LandscapeComponent == LandscapeComponent; });
+									check(ComponentReadbackResult != nullptr);
+
+									// Mark this layer info within this component as being all-zero :
+									ComponentReadbackResult->AllZeroLayers.AddUnique(WeightmapLayerAllocation.LayerInfo);
+								}
+							}
+						}
+					}
+				}
+				AllZerosTextureChannelMask &= ~((uint32)1 << AllZerosTextureChannelIndex);
+			}
+		}
+
 		// Release the processed read backs
 		InCPUReadback->ReleaseCompletedResults(CompletedReadbackNum);
 	}
@@ -6786,11 +6868,12 @@ int32 ALandscape::PerformLayersWeightmapsGlobalMerge(FUpdateLayersContentContext
 							PrintLayersDebugRT(FString::Printf(TEXT("LS Weight: %s %s -> Brush %s"), *Layer.Name.ToString(), *LandscapeBrush->GetName(), *BrushOutputRT->GetName()), BrushOutputRT);
 
 							// Copy result back if brush did not edit things in place.
+
+							// Resolve back to Combined heightmap (it's unlikely, but possible that the brush returns the same RT as input and output, if it did various operations on it, in which case the copy is useless) :
 							if (BrushOutputRT != LandscapeScratchRT3)
 							{
 								SourceDebugName = FString::Printf(TEXT("Weight: %s PaintLayer: %s Brush: %s"), *Layer.Name.ToString(), *LayerInfoObj->LayerName.ToString(), *BrushOutputRT->GetName());
 								DestDebugName = LandscapeScratchRT3->GetName();
-
 								ExecuteCopyLayersTexture({ FLandscapeLayersCopyTextureParams(SourceDebugName, BrushOutputRT->GameThread_GetRenderTargetResource(), DestDebugName, LandscapeScratchRT3->GameThread_GetRenderTargetResource()) });
 								PrintLayersDebugRT(FString::Printf(TEXT("LS Weight: %s Component %s += -> Combined %s"), *Layer.Name.ToString(), *BrushOutputRT->GetName(), *LandscapeScratchRT3->GetName()), LandscapeScratchRT3);
 							}
@@ -7077,6 +7160,25 @@ void ALandscape::UpdateForChangedWeightmaps(const TArrayView<FLandscapeEditLayer
 			ComponentReadbackResult.LandscapeComponent->SetPendingLayerCollisionDataUpdate(true);
 		}
 
+		// If this component has a layer with only zeros, remove it so that we don't end up with weightmaps we don't end up using :
+		if (!ComponentReadbackResult.AllZeroLayers.IsEmpty())
+		{
+			const TArray<FWeightmapLayerAllocationInfo>& ComponentWeightmapLayerAllocations = ComponentReadbackResult.LandscapeComponent->GetWeightmapLayerAllocations(FGuid());
+			for (ULandscapeLayerInfoObject* AllZeroLayerInfo : ComponentReadbackResult.AllZeroLayers)
+			{
+				check(AllZeroLayerInfo != nullptr);
+				// Find the index for this layer in this component.
+				int32 AllZeroLayerIndex = ComponentWeightmapLayerAllocations.IndexOfByPredicate(
+					[AllZeroLayerInfo](const FWeightmapLayerAllocationInfo& Allocation) { return Allocation.LayerInfo == AllZeroLayerInfo; });
+				check(AllZeroLayerIndex != INDEX_NONE);
+
+				ComponentReadbackResult.LandscapeComponent->DeleteLayerAllocation(FGuid(), AllZeroLayerIndex, /*bShouldDirtyPackage = */true);
+
+				// We removed a weightmap allocation so the material instance for this landscape component needs updating :
+				ComponentsNeedingMaterialInstanceUpdates.Add(ComponentReadbackResult.LandscapeComponent);
+			}
+		}
+
 		const int32 WeightUpdateMode = ComponentReadbackResult.UpdateModes & (ELandscapeLayerUpdateMode::Update_Weightmap_All | ELandscapeLayerUpdateMode::Update_Weightmap_Editing | ELandscapeLayerUpdateMode::Update_Weightmap_Editing_NoCollision);
 		if (IsUpdateFlagEnabledForModes(ELandscapeComponentUpdateFlag::Component_Update_Weightmap_Collision, WeightUpdateMode))
 		{
@@ -7088,6 +7190,8 @@ void ALandscape::UpdateForChangedWeightmaps(const TArrayView<FLandscapeEditLayer
 			}
 		}
 	}
+
+	UpdateLayersMaterialInstances(ComponentsNeedingMaterialInstanceUpdates);
 }
 
 void ULandscapeComponent::GetUsedPaintLayers(const FGuid& InLayerGuid, TArray<ULandscapeLayerInfoObject*>& OutUsedLayerInfos) const
diff --git a/Engine/Source/Runtime/Landscape/Private/LandscapeEditReadback.h b/Engine/Source/Runtime/Landscape/Private/LandscapeEditReadback.h
--- a/Engine/Source/Runtime/Landscape/Private/LandscapeEditReadback.h
+++ b/Engine/Source/Runtime/Landscape/Private/LandscapeEditReadback.h
@@ -27,6 +27,8 @@ public:
 	/** Get the stored hash value. */
 	uint32 GetHash() const { return Hash; }
 
+	using FPerChannelLayerNames = TStaticArray<FName, 4>;
+
 	/** Per component context required for processing read back results. */
 	struct FComponentReadbackContext
 	{
@@ -34,11 +36,8 @@ public:
 		FIntPoint ComponentKey;
 		/** Component ELandscapeLayerUpdateMode flags. */
 		int32 UpdateModes = 0;
-
-		FComponentReadbackContext(FIntPoint InComponentKey, int32 InUpdateModes)
-			: ComponentKey(InComponentKey)
-			, UpdateModes(InUpdateModes)
-		{}
+		/** For weightmaps only : configuration of the channels for this component and texture when the readback was performed (useful in case the channel configuration changed before we could perform the readback) */
+		FPerChannelLayerNames PerChannelLayerNames;
 	};
 	/** Full context for processing read back results. */
 	using FReadbackContext = TArray<FComponentReadbackContext>;
