Landscape: Give unique label to ALandscapeSplineActor

diff --git a/Engine/Source/Runtime/Landscape/Private/Landscape.cpp b/Engine/Source/Runtime/Landscape/Private/Landscape.cpp
--- a/Engine/Source/Runtime/Landscape/Private/Landscape.cpp
+++ b/Engine/Source/Runtime/Landscape/Private/Landscape.cpp
@@ -72,6 +72,7 @@ Landscape.cpp: Terrain rendering
 #include "LandscapeEdit.h"
 #include "MaterialUtilities.h"
 #include "Editor.h"
+#include "Editor/EditorEngine.h"
 #include "Algo/Transform.h"
 #include "Algo/BinarySearch.h"
 #include "Engine/Texture2D.h"
@@ -4040,6 +4041,9 @@ ALandscapeSplineActor* ULandscapeInfo::CreateSplineActor(const FVector& Location
 	SplineActor->GetSharedProperties(this);
 	SplineActor->GetSplinesComponent()->ShowSplineEditorMesh(true);
 	SplineActor->SetIsSpatiallyLoaded(AreNewLandscapeActorsSpatiallyLoaded());
+	
+	FActorLabelUtilities::SetActorLabelUnique(SplineActor, ALandscapeSplineActor::StaticClass()->GetName());
+
 	RegisterSplineActor(SplineActor);
 	return SplineActor;
 }
