Landscape: remove `ensure` on GetCompleteBounds

diff --git a/Engine/Source/Runtime/Landscape/Private/Landscape.cpp b/Engine/Source/Runtime/Landscape/Private/Landscape.cpp
--- a/Engine/Source/Runtime/Landscape/Private/Landscape.cpp
+++ b/Engine/Source/Runtime/Landscape/Private/Landscape.cpp
@@ -2740,7 +2740,7 @@ void ALandscape::RenderHeightmap(const FTransform& InRenderAreaWorldTransform, c
 
 FBox ALandscape::GetCompleteBounds() const
 {
-	if (ensure(GetLandscapeInfo()))
+	if (GetLandscapeInfo())
 	{
 		return GetLandscapeInfo()->GetCompleteBounds();
 	}
