Fixed: static analysys warning in a LandscapeRender.cpp

diff --git a/Engine/Source/Runtime/Landscape/Private/LandscapeRender.cpp b/Engine/Source/Runtime/Landscape/Private/LandscapeRender.cpp
--- a/Engine/Source/Runtime/Landscape/Private/LandscapeRender.cpp
+++ b/Engine/Source/Runtime/Landscape/Private/LandscapeRender.cpp
@@ -799,8 +799,13 @@ FLandscapeComponentSceneProxy::FLandscapeComponentSceneProxy(ULandscapeComponent
 
 	SetLevelColor(FLinearColor(1.f, 1.f, 1.f));
 			
-	HeightmapSubsectionOffsetU = ((float)(InComponent->SubsectionSizeQuads + 1) / (float)FMath::Max<int32>(1, HeightmapTexture->GetSizeX()));
-	HeightmapSubsectionOffsetV = ((float)(InComponent->SubsectionSizeQuads + 1) / (float)FMath::Max<int32>(1, HeightmapTexture->GetSizeY()));
+	HeightmapSubsectionOffsetU = 0;
+	HeightmapSubsectionOffsetV = 0;
+	if (HeightmapTexture)
+	{
+		HeightmapSubsectionOffsetU = ((float)(InComponent->SubsectionSizeQuads + 1) / (float)FMath::Max<int32>(1, HeightmapTexture->GetSizeX()));
+		HeightmapSubsectionOffsetV = ((float)(InComponent->SubsectionSizeQuads + 1) / (float)FMath::Max<int32>(1, HeightmapTexture->GetSizeY()));
+	}
 
 	float ScreenSizeRatioDivider = FMath::Max(InComponent->GetLandscapeProxy()->LOD0DistributionSetting * GLandscapeLOD0DistributionScale, 1.01f);
 	// Cancel out so that landscape is not affected by r.StaticMeshLODDistanceScale
