Prevents ULandscapeInfo::RegisterActor to try loading twice the same landscape in creative JIP.

diff --git a/Engine/Source/Runtime/Landscape/Private/Landscape.cpp b/Engine/Source/Runtime/Landscape/Private/Landscape.cpp
--- a/Engine/Source/Runtime/Landscape/Private/Landscape.cpp
+++ b/Engine/Source/Runtime/Landscape/Private/Landscape.cpp
@@ -3902,23 +3902,29 @@ void ULandscapeInfo::RegisterActor(ALandscapeProxy* Proxy, bool bMapCheck)
 	// register
 	if (ALandscape* Landscape = Cast<ALandscape>(Proxy))
 	{
-		checkf(!LandscapeActor || LandscapeActor == Landscape, TEXT("Multiple landscapes with the same GUID detected: %s vs %s"), *LandscapeActor->GetPathName(), *Landscape->GetPathName());
-		LandscapeActor = Landscape;
+		if (LandscapeActor == nullptr)
+		{
+			LandscapeActor = Landscape;
 
 #if WITH_EDITOR
-		// In world composition user is not allowed to move landscape in editor, only through WorldBrowser 
-		bool bIsLockLocation = LandscapeActor->IsLockLocation();
-		bIsLockLocation |= OwningWorld != nullptr ? OwningWorld->WorldComposition != nullptr : false;
-		LandscapeActor->SetLockLocation(bIsLockLocation);
+			// In world composition user is not allowed to move landscape in editor, only through WorldBrowser 
+			bool bIsLockLocation = LandscapeActor->IsLockLocation();
+			bIsLockLocation |= OwningWorld != nullptr ? OwningWorld->WorldComposition != nullptr : false;
+			LandscapeActor->SetLockLocation(bIsLockLocation);
 #endif // WITH_EDITOR
 
-		// update proxies reference actor
-		for (ALandscapeStreamingProxy* StreamingProxy : Proxies)
-		{
-			StreamingProxy->LandscapeActor = LandscapeActor;
+			// update proxies reference actor
+			for (ALandscapeStreamingProxy* StreamingProxy : Proxies)
+			{
+				StreamingProxy->LandscapeActor = LandscapeActor;
 #if WITH_EDITOR
-			StreamingProxy->FixupSharedData(Landscape);
+				StreamingProxy->FixupSharedData(Landscape);
 #endif // WITH_EDITOR
+			}
+		}
+		else if (LandscapeActor == Landscape)
+		{
+			UE_LOG(LogLandscape, Warning, TEXT("Multiple landscapes actors with the same GUID detected: %s vs %s"), * LandscapeActor->GetPathName(), * Landscape->GetPathName());
 		}
 	}
 	else
