Fixed: WaterScreenshots IslandLake and IslandRiver failing due to changes in the landscape

diff --git a/Engine/Shaders/Private/LandscapeVertexFactory.ush b/Engine/Shaders/Private/LandscapeVertexFactory.ush
--- a/Engine/Shaders/Private/LandscapeVertexFactory.ush
+++ b/Engine/Shaders/Private/LandscapeVertexFactory.ush
@@ -361,7 +361,7 @@ float3x3 VertexFactoryGetPerPixelTangentBasis(FVertexFactoryInterpolantsVSToPS I
 {
 	float3x3 Result;
 #if PIXELSHADER || RAYHITGROUPSHADER || COMPUTESHADER
-	float4 SampleValue = Texture2DSample(LandscapeParameters.HeightmapTexture, LandscapeParameters.HeightmapTextureSampler, Interpolants.WeightHeightMapTexCoord.zw);
+	float4 SampleValue = Texture2DSample(LandscapeParameters.NormalmapTexture, LandscapeParameters.NormalmapTextureSampler, Interpolants.WeightHeightMapTexCoord.zw);
 	float2 SampleNormal = float2(SampleValue.b, SampleValue.a) * float2(2.0,2.0) - float2(1.0,1.0);
 	float3 WorldNormal = float3( SampleNormal, sqrt(max(1.0-dot(SampleNormal,SampleNormal),0.0)) );
 
diff --git a/Engine/Source/Runtime/Landscape/Private/LandscapeRender.cpp b/Engine/Source/Runtime/Landscape/Private/LandscapeRender.cpp
--- a/Engine/Source/Runtime/Landscape/Private/LandscapeRender.cpp
+++ b/Engine/Source/Runtime/Landscape/Private/LandscapeRender.cpp
@@ -1535,11 +1535,15 @@ void FLandscapeComponentSceneProxy::OnTransformChanged()
 	{
 		LandscapeParams.HeightmapTexture = HeightmapTexture->TextureReference.TextureReferenceRHI;
 		LandscapeParams.HeightmapTextureSampler = TStaticSamplerState<SF_Point>::GetRHI();
+		LandscapeParams.NormalmapTexture = HeightmapTexture->TextureReference.TextureReferenceRHI;
+		LandscapeParams.NormalmapTextureSampler = TStaticSamplerState<SF_Bilinear>::GetRHI();
 	}
 	else
 	{
 		LandscapeParams.HeightmapTexture = GBlackTexture->TextureRHI;
 		LandscapeParams.HeightmapTextureSampler = GBlackTexture->SamplerStateRHI;
+		LandscapeParams.NormalmapTexture = GBlackTexture->TextureRHI;
+		LandscapeParams.NormalmapTextureSampler = GBlackTexture->SamplerStateRHI;
 	}
 
 	if (XYOffsetmapTexture)
diff --git a/Engine/Source/Runtime/Landscape/Public/LandscapeRender.h b/Engine/Source/Runtime/Landscape/Public/LandscapeRender.h
--- a/Engine/Source/Runtime/Landscape/Public/LandscapeRender.h
+++ b/Engine/Source/Runtime/Landscape/Public/LandscapeRender.h
@@ -103,6 +103,8 @@ BEGIN_GLOBAL_SHADER_PARAMETER_STRUCT(FLandscapeUniformShaderParameters, LANDSCAP
 	SHADER_PARAMETER(FMatrix44f, LocalToWorldNoScaling)
 	SHADER_PARAMETER_TEXTURE(Texture2D, HeightmapTexture)
 	SHADER_PARAMETER_SAMPLER(SamplerState, HeightmapTextureSampler)
+	SHADER_PARAMETER_TEXTURE(Texture2D, NormalmapTexture)
+	SHADER_PARAMETER_SAMPLER(SamplerState, NormalmapTextureSampler)
 	SHADER_PARAMETER_TEXTURE(Texture2D, XYOffsetmapTexture)
 	SHADER_PARAMETER_SAMPLER(SamplerState, XYOffsetmapTextureSampler)
 END_GLOBAL_SHADER_PARAMETER_STRUCT()
