Moved landscape LOD calculations off the render thread.

diff --git a/Engine/Source/Runtime/Engine/Public/SceneViewExtension.h b/Engine/Source/Runtime/Engine/Public/SceneViewExtension.h
--- a/Engine/Source/Runtime/Engine/Public/SceneViewExtension.h
+++ b/Engine/Source/Runtime/Engine/Public/SceneViewExtension.h
@@ -150,6 +150,11 @@ public:
 	 */
 	virtual void PreRenderView_RenderThread(FRDGBuilder& GraphBuilder, FSceneView& InView) {}
 
+	/**
+	 * Called on render thread prior to initializing views.
+	 */
+	virtual void PreInitViews_RenderThread(FRDGBuilder& GraphBuilder) {}
+
 	/**
 	 * Called right after Base Pass rendering finished when using the deferred renderer.
 	 */
diff --git a/Engine/Source/Runtime/Landscape/Private/LandscapeRender.cpp b/Engine/Source/Runtime/Landscape/Private/LandscapeRender.cpp
--- a/Engine/Source/Runtime/Landscape/Private/LandscapeRender.cpp
+++ b/Engine/Source/Runtime/Landscape/Private/LandscapeRender.cpp
@@ -88,6 +88,12 @@ static FAutoConsoleVariableRef CVarAllowLandscapeShadows(
 	TEXT("Allow Landscape Shadows")
 );
 
+int32 GLandscapeUseAsyncTasksForLODComputation = 1;
+FAutoConsoleVariableRef CVarLandscapeUseAsyncTasksForLODComputation(
+	TEXT("r.LandscapeUseAsyncTasksForLODComputation"),
+	GLandscapeUseAsyncTasksForLODComputation,
+	TEXT("Use async tasks for computing per-landscape component LOD biases."));
+
 int32 GDisableLandscapeNaniteGI = 1;
 static FAutoConsoleVariableRef CVarDisableLandscapeNaniteGI(
 	TEXT("r.DisableLandscapeNaniteGI"),
@@ -538,16 +544,6 @@ const TResourceArray<float>& FLandscapeRenderSystem::ComputeSectionsLODForView(c
 	return SectionLODValues;
 }
 
-void FLandscapeRenderSystem::PreRenderViewFamily_RenderThread()
-{
-	TRACE_CPUPROFILER_EVENT_SCOPE(FLandscapeRenderSystem::PreRenderViewFamily_RenderThread);
-	check(IsInRenderingThread());
-
-	CachedSectionLODValues.Reset();
-	FetchHeightmapLODBiases();
-	UpdateBuffers();
-}
-
 void FLandscapeRenderSystem::FetchHeightmapLODBiases()
 {
 	TRACE_CPUPROFILER_EVENT_SCOPE(FLandscapeRenderSystem::FetchHeightmapLODBiases);
@@ -612,56 +608,85 @@ FLandscapeSceneViewExtension::~FLandscapeSceneViewExtension()
 	FCoreDelegates::OnEndFrameRT.RemoveAll(this);
 }
 
-void FLandscapeSceneViewExtension::PreRenderViewFamily_RenderThread(FRDGBuilder& GraphBuilder, FSceneViewFamily& InViewFamily)
-{
-	for (auto& Pair : LandscapeRenderSystems)
-	{
-		FLandscapeRenderSystem& RenderSystem = *Pair.Value;
-		RenderSystem.PreRenderViewFamily_RenderThread();
-	}
-}
-
 void FLandscapeSceneViewExtension::PreRenderView_RenderThread(FRDGBuilder& GraphBuilder, FSceneView& InView)
 {
-	FShaderResourceViewRHIRef LODDataSRV;
-	FShaderResourceViewRHIRef IndirectionSRV;
+	LandscapeViews.Emplace(InView);
 
-	if (!LandscapeRenderSystems.IsEmpty())
+	// Kick the job once all views have been collected.
+	if (!LandscapeRenderSystems.IsEmpty() && LandscapeViews.Num() == InView.Family->Views.Num())
 	{
-		TResourceArray<uint32> LandscapeIndirection;
-		TResourceArray<float> LandscapeLODData;
+		const auto ComputeLODs = [this]
+		{
+			TRACE_CPUPROFILER_EVENT_SCOPE(FLandscapeRenderSystem::ComputeLODs);
+			FOptionalTaskTagScope Scope(ETaskTag::EParallelRenderingThread);
 
-		LandscapeIndirection.SetNum(FLandscapeRenderSystem::LandscapeIndexAllocator.Num());
+			for (auto& Pair : LandscapeRenderSystems)
+			{
+				FLandscapeRenderSystem& RenderSystem = *Pair.Value;
+				RenderSystem.CachedSectionLODValues.Reset();
+				RenderSystem.FetchHeightmapLODBiases();
+			}
 
-		for (auto& Pair : LandscapeRenderSystems)
-		{
-			FLandscapeRenderSystem& RenderSystem = *Pair.Value;
+			for (FLandscapeViewData& LandscapeView : LandscapeViews)
+			{
+				LandscapeView.LandscapeIndirection.SetNum(FLandscapeRenderSystem::LandscapeIndexAllocator.Num());
+
+				for (auto& Pair : LandscapeRenderSystems)
+				{
+					FLandscapeRenderSystem& RenderSystem = *Pair.Value;
 
-			// Store index where the LOD data for this landscape starts
-			uint32& IndirectionEntry = LandscapeIndirection[RenderSystem.LandscapeIndex];
-			IndirectionEntry = LandscapeLODData.Num();
+					// Store index where the LOD data for this landscape starts
+					LandscapeView.LandscapeIndirection[RenderSystem.LandscapeIndex] = LandscapeView.LandscapeLODData.Num();
 
-			// Compute sections lod values for this view & append to the global landscape LOD data
-			const TResourceArray<float>& SectionsLODValues = RenderSystem.ComputeSectionsLODForView(InView);
-			LandscapeLODData.Append(SectionsLODValues);
+					// Compute sections lod values for this view & append to the global landscape LOD data
+					LandscapeView.LandscapeLODData.Append(RenderSystem.ComputeSectionsLODForView(*LandscapeView.View));
+				}
+			}
+		};
+
+		if (GIsThreadedRendering && GLandscapeUseAsyncTasksForLODComputation)
+		{
+			LandscapeSetupTask = UE::Tasks::Launch(UE_SOURCE_LOCATION, ComputeLODs, LowLevelTasks::ETaskPriority::Normal);
+		}
+		else
+		{
+			ComputeLODs();
 		}
+	}
+}
 
-		FRHIResourceCreateInfo CreateInfoLODBuffer(TEXT("LandscapeLODDataBuffer"), &LandscapeLODData);
-		LandscapeLODDataBuffer = RHICreateVertexBuffer(LandscapeLODData.GetResourceDataSize(), BUF_ShaderResource | BUF_Volatile, CreateInfoLODBuffer);
-		LODDataSRV = RHICreateShaderResourceView(LandscapeLODDataBuffer, sizeof(float), PF_R32_FLOAT);
+void FLandscapeSceneViewExtension::PreInitViews_RenderThread(FRDGBuilder& GraphBuilder)
+{
+	LandscapeSetupTask.Wait();
 
-		FRHIResourceCreateInfo CreateInfoIndirection(TEXT("LandscapeIndirectionBuffer"), &LandscapeIndirection);
-		LandscapeIndirectionBuffer = RHICreateVertexBuffer(LandscapeIndirection.GetResourceDataSize(), BUF_ShaderResource | BUF_Volatile, CreateInfoIndirection);
-		IndirectionSRV = RHICreateShaderResourceView(LandscapeIndirectionBuffer, sizeof(uint32), PF_R32_UINT);
+	for (auto& Pair : LandscapeRenderSystems)
+	{
+		FLandscapeRenderSystem& RenderSystem = *Pair.Value;
+		RenderSystem.UpdateBuffers();
 	}
-	else
+
+	for (FLandscapeViewData& LandscapeView : LandscapeViews)
 	{
-		LODDataSRV = GWhiteVertexBufferWithSRV->ShaderResourceViewRHI;
-		IndirectionSRV = GWhiteVertexBufferWithSRV->ShaderResourceViewRHI;
+		if (!LandscapeRenderSystems.IsEmpty())
+		{
+			FBufferRHIRef Buffer;
+
+			FRHIResourceCreateInfo CreateInfoLODBuffer(TEXT("LandscapeLODDataBuffer"), &LandscapeView.LandscapeLODData);
+			Buffer = RHICreateVertexBuffer(LandscapeView.LandscapeLODData.GetResourceDataSize(), BUF_ShaderResource | BUF_Volatile, CreateInfoLODBuffer);
+			LandscapeView.View->LandscapePerComponentDataBuffer = RHICreateShaderResourceView(Buffer, sizeof(float), PF_R32_FLOAT);
+
+			FRHIResourceCreateInfo CreateInfoIndirection(TEXT("LandscapeIndirectionBuffer"), &LandscapeView.LandscapeIndirection);
+			Buffer = RHICreateVertexBuffer(LandscapeView.LandscapeIndirection.GetResourceDataSize(), BUF_ShaderResource | BUF_Volatile, CreateInfoIndirection);
+			LandscapeView.View->LandscapeIndirectionBuffer = RHICreateShaderResourceView(Buffer, sizeof(uint32), PF_R32_UINT);
+		}
+		else
+		{
+			LandscapeView.View->LandscapePerComponentDataBuffer = GWhiteVertexBufferWithSRV->ShaderResourceViewRHI;
+			LandscapeView.View->LandscapeIndirectionBuffer = GWhiteVertexBufferWithSRV->ShaderResourceViewRHI;
+		}
 	}
 
-	InView.LandscapePerComponentDataBuffer = LODDataSRV;
-	InView.LandscapeIndirectionBuffer = IndirectionSRV;
+	LandscapeViews.Reset();
 }
 
 // TODO [jonathan.bard] Ideally this should be symmetrical with FLandscapeSceneViewExtension::PreRenderView_RenderThread and should be called in FLandscapeSceneViewExtension::PostRenderView_RenderThread
diff --git a/Engine/Source/Runtime/Landscape/Public/LandscapeRender.h b/Engine/Source/Runtime/Landscape/Public/LandscapeRender.h
--- a/Engine/Source/Runtime/Landscape/Public/LandscapeRender.h
+++ b/Engine/Source/Runtime/Landscape/Public/LandscapeRender.h
@@ -27,6 +27,7 @@ LandscapeRender.h: New terrain rendering
 #include "PrimitiveSceneProxy.h"
 #include "StaticMeshResources.h"
 #include "SceneViewExtension.h"
+#include "Tasks/Task.h"
 
 // This defines the number of border blocks to surround terrain by when generating lightmaps
 #define TERRAIN_PATCH_EXPAND_SCALAR	1
@@ -465,7 +466,6 @@ struct FLandscapeRenderSystem
 	const TResourceArray<float>& ComputeSectionsLODForView(const FSceneView& InView);
 	void FetchHeightmapLODBiases();
 	void UpdateBuffers();
-	void PreRenderViewFamily_RenderThread();
 
 private:
 	void CreateResources_Internal(FLandscapeSectionInfo* InSectionInfo);
@@ -488,13 +488,29 @@ public:
 	virtual void SetupView(FSceneViewFamily& InViewFamily, FSceneView& InView) override {}
 	virtual void BeginRenderViewFamily(FSceneViewFamily& InViewFamily) override {}
 
-	virtual void PreRenderViewFamily_RenderThread(FRDGBuilder& GraphBuilder, FSceneViewFamily& InViewFamily) override;
 	virtual void PreRenderView_RenderThread(FRDGBuilder& GraphBuilder, FSceneView& InView) override;
+	virtual void PreInitViews_RenderThread(FRDGBuilder& GraphBuilder) override;
 
 	LANDSCAPE_API const TMap<uint32, FLandscapeRenderSystem*>& GetLandscapeRenderSystems() const;
 private:
 	FBufferRHIRef LandscapeLODDataBuffer;
 	FBufferRHIRef LandscapeIndirectionBuffer;
+
+	struct FLandscapeViewData
+	{
+		FLandscapeViewData() = default;
+
+		FLandscapeViewData(FSceneView& InView)
+			: View(&InView)
+		{}
+
+		FSceneView* View = nullptr;
+		TResourceArray<uint32> LandscapeIndirection;
+		TResourceArray<float> LandscapeLODData;
+	};
+
+	TArray<FLandscapeViewData> LandscapeViews;
+	UE::Tasks::FTask LandscapeSetupTask;
 };
 
 
diff --git a/Engine/Source/Runtime/Renderer/Private/SceneVisibility.cpp b/Engine/Source/Runtime/Renderer/Private/SceneVisibility.cpp
--- a/Engine/Source/Runtime/Renderer/Private/SceneVisibility.cpp
+++ b/Engine/Source/Runtime/Renderer/Private/SceneVisibility.cpp
@@ -46,7 +46,7 @@
 #include "TemporalAA.h"
 #include "RayTracing/RayTracingInstanceCulling.h"
 #include "RendererModule.h"
-
+#include "SceneViewExtension.h"
 
 /*------------------------------------------------------------------------------
 	Globals
@@ -4016,6 +4016,11 @@ void FSceneRenderer::PreVisibilityFrameSetup(FRDGBuilder& GraphBuilder, const FS
 		DitherUniformShaderParameters.LODFactor = View.GetTemporalLODTransition() - 1.0f;
 		View.DitherFadeInUniformBuffer = FDitherUniformBufferRef::CreateUniformBufferImmediate(DitherUniformShaderParameters, UniformBuffer_SingleFrame);
 	}
+
+	for (const auto& ViewExtension : ViewFamily.ViewExtensions)
+	{
+		ViewExtension->PreInitViews_RenderThread(GraphBuilder);
+	}
 }
 
 void FSceneViewState::UpdateMotionBlurTimeScale(const FViewInfo& View)
