Hold a reference to landscape buffers to avoid dangling pointer.

diff --git a/Engine/Source/Runtime/Landscape/Private/LandscapeRender.cpp b/Engine/Source/Runtime/Landscape/Private/LandscapeRender.cpp
--- a/Engine/Source/Runtime/Landscape/Private/LandscapeRender.cpp
+++ b/Engine/Source/Runtime/Landscape/Private/LandscapeRender.cpp
@@ -669,15 +669,13 @@ void FLandscapeSceneViewExtension::PreInitViews_RenderThread(FRDGBuilder& GraphB
 	{
 		if (!LandscapeRenderSystems.IsEmpty())
 		{
-			FBufferRHIRef Buffer;
-
 			FRHIResourceCreateInfo CreateInfoLODBuffer(TEXT("LandscapeLODDataBuffer"), &LandscapeView.LandscapeLODData);
-			Buffer = RHICreateVertexBuffer(LandscapeView.LandscapeLODData.GetResourceDataSize(), BUF_ShaderResource | BUF_Volatile, CreateInfoLODBuffer);
-			LandscapeView.View->LandscapePerComponentDataBuffer = RHICreateShaderResourceView(Buffer, sizeof(float), PF_R32_FLOAT);
+			LandscapeLODDataBuffer = RHICreateVertexBuffer(LandscapeView.LandscapeLODData.GetResourceDataSize(), BUF_ShaderResource | BUF_Volatile, CreateInfoLODBuffer);
+			LandscapeView.View->LandscapePerComponentDataBuffer = RHICreateShaderResourceView(LandscapeLODDataBuffer, sizeof(float), PF_R32_FLOAT);
 
 			FRHIResourceCreateInfo CreateInfoIndirection(TEXT("LandscapeIndirectionBuffer"), &LandscapeView.LandscapeIndirection);
-			Buffer = RHICreateVertexBuffer(LandscapeView.LandscapeIndirection.GetResourceDataSize(), BUF_ShaderResource | BUF_Volatile, CreateInfoIndirection);
-			LandscapeView.View->LandscapeIndirectionBuffer = RHICreateShaderResourceView(Buffer, sizeof(uint32), PF_R32_UINT);
+			LandscapeIndirectionBuffer = RHICreateVertexBuffer(LandscapeView.LandscapeIndirection.GetResourceDataSize(), BUF_ShaderResource | BUF_Volatile, CreateInfoIndirection);
+			LandscapeView.View->LandscapeIndirectionBuffer = RHICreateShaderResourceView(LandscapeIndirectionBuffer, sizeof(uint32), PF_R32_UINT);
 		}
 		else
 		{
