Landscape: Prevent Landscape preview settings affecting actual Landscape rendering.

diff --git a/Engine/Source/Editor/UnrealEd/Private/PreviewMaterial.cpp b/Engine/Source/Editor/UnrealEd/Private/PreviewMaterial.cpp
--- a/Engine/Source/Editor/UnrealEd/Private/PreviewMaterial.cpp
+++ b/Engine/Source/Editor/UnrealEd/Private/PreviewMaterial.cpp
@@ -227,6 +227,8 @@ public:
 	 */
 	virtual bool IsPersistent() const override { return false; }
 	virtual FString GetAssetName() const override { return FString::Printf(TEXT("Preview:%s"), *FMaterialResource::GetAssetName()); }
+
+	virtual bool IsPreview() const override { return true; }
 };
 
 /** Implementation of Preview Material functions*/
diff --git a/Engine/Source/Runtime/Engine/Private/Materials/HLSLMaterialTranslator.cpp b/Engine/Source/Runtime/Engine/Private/Materials/HLSLMaterialTranslator.cpp
--- a/Engine/Source/Runtime/Engine/Private/Materials/HLSLMaterialTranslator.cpp
+++ b/Engine/Source/Runtime/Engine/Private/Materials/HLSLMaterialTranslator.cpp
@@ -6963,7 +6963,7 @@ int32 FHLSLMaterialTranslator::StaticTerrainLayerWeight(FName LayerName,int32 De
 		}
 	}
 
-	if(!bFoundParameter)
+	if(!bFoundParameter && Material->IsPreview())
 	{
 		return Default;
 	}
diff --git a/Engine/Source/Runtime/Engine/Public/MaterialShared.h b/Engine/Source/Runtime/Engine/Public/MaterialShared.h
--- a/Engine/Source/Runtime/Engine/Public/MaterialShared.h
+++ b/Engine/Source/Runtime/Engine/Public/MaterialShared.h
@@ -1939,6 +1939,8 @@ public:
 	virtual bool IsPersistent() const = 0;
 	virtual UMaterialInterface* GetMaterialInterface() const { return NULL; }
 
+	virtual bool IsPreview() const { return false; }
+
 #if WITH_EDITOR
 	UE_DEPRECATED(5.1, "Should instead check individual properties")
 	virtual bool HasMaterialAttributesConnected() const { return false; }
diff --git a/Engine/Source/Runtime/Landscape/Classes/Materials/MaterialExpressionLandscapeLayerBlend.h b/Engine/Source/Runtime/Landscape/Classes/Materials/MaterialExpressionLandscapeLayerBlend.h
--- a/Engine/Source/Runtime/Landscape/Classes/Materials/MaterialExpressionLandscapeLayerBlend.h
+++ b/Engine/Source/Runtime/Landscape/Classes/Materials/MaterialExpressionLandscapeLayerBlend.h
@@ -78,7 +78,6 @@ class LANDSCAPE_API UMaterialExpressionLandscapeLayerBlend : public UMaterialExp
 	virtual uint32 GetOutputType(int32 InputIndex) override { return MCT_Unknown; }
 	virtual bool IsResultMaterialAttributes(int32 OutputIndex) override;
 	virtual int32 Compile(class FMaterialCompiler* Compiler, int32 OutputIndex) override;
-	virtual int32 CompilePreview(class FMaterialCompiler* Compiler, int32 OutputIndex) override;
 	virtual void GetCaption(TArray<FString>& OutCaptions) const override;
 	virtual const TArray<FExpressionInput*> GetInputs() override;
 	virtual FExpressionInput* GetInput(int32 InputIndex) override;
diff --git a/Engine/Source/Runtime/Landscape/Private/Materials/MaterialExpressionLandscapeLayerBlend.cpp b/Engine/Source/Runtime/Landscape/Private/Materials/MaterialExpressionLandscapeLayerBlend.cpp
--- a/Engine/Source/Runtime/Landscape/Private/Materials/MaterialExpressionLandscapeLayerBlend.cpp
+++ b/Engine/Source/Runtime/Landscape/Private/Materials/MaterialExpressionLandscapeLayerBlend.cpp
@@ -149,21 +149,10 @@ bool UMaterialExpressionLandscapeLayerBlend::IsResultMaterialAttributes(int32 Ou
 }
 
 int32 UMaterialExpressionLandscapeLayerBlend::Compile(class FMaterialCompiler* Compiler, int32 OutputIndex)
-{
-	return Compile(Compiler, OutputIndex, false);
-}
-
-int32 UMaterialExpressionLandscapeLayerBlend::CompilePreview(class FMaterialCompiler* Compiler, int32 OutputIndex)
-{
-	return Compile(Compiler, OutputIndex, true);
-}
-
-int32 UMaterialExpressionLandscapeLayerBlend::Compile(class FMaterialCompiler* Compiler, int32 OutputIndex, bool bForPreview)
 {
 	// For renormalization
 	bool bNeedsRenormalize = false;
 	int32 WeightSumCode = Compiler->Constant(0);
-
 	// Temporary store for each layer's weight
 	TArray<int32> WeightCodes;
 	WeightCodes.Empty(Layers.Num());
@@ -180,7 +169,7 @@ int32 UMaterialExpressionLandscapeLayerBlend::Compile(class FMaterialCompiler* C
 			// Height input
 			const int32 HeightCode = Layer.HeightInput.Expression ? Layer.HeightInput.Compile(Compiler) : Compiler->Constant(Layer.ConstHeightInput);
 
-			const int32 DefaultWeightCode = bForPreview && Layer.PreviewWeight > 0.0f ? Compiler->Constant(Layer.PreviewWeight) : INDEX_NONE;
+			const int32 DefaultWeightCode = Layer.PreviewWeight > 0.0f ? Compiler->Constant(Layer.PreviewWeight) : INDEX_NONE;
 			const int32 WeightCode = Compiler->StaticTerrainLayerWeight(Layer.LayerName, DefaultWeightCode);
 			if (WeightCode != INDEX_NONE)
 			{
@@ -243,7 +232,7 @@ int32 UMaterialExpressionLandscapeLayerBlend::Compile(class FMaterialCompiler* C
 	{
 		if (Layer.BlendType == LB_AlphaBlend)
 		{
-			const int32 DefaultWeightCode = bForPreview && Layer.PreviewWeight > 0.0f ? Compiler->Constant(Layer.PreviewWeight) : INDEX_NONE;
+			const int32 DefaultWeightCode = Layer.PreviewWeight > 0.0f ? Compiler->Constant(Layer.PreviewWeight) : INDEX_NONE;
 			const int32 WeightCode = Compiler->StaticTerrainLayerWeight(Layer.LayerName, DefaultWeightCode);
 			if (WeightCode != INDEX_NONE)
 			{
diff --git a/Engine/Source/Runtime/Landscape/Private/Materials/MaterialExpressionLandscapeLayerSwitch.cpp b/Engine/Source/Runtime/Landscape/Private/Materials/MaterialExpressionLandscapeLayerSwitch.cpp
--- a/Engine/Source/Runtime/Landscape/Private/Materials/MaterialExpressionLandscapeLayerSwitch.cpp
+++ b/Engine/Source/Runtime/Landscape/Private/Materials/MaterialExpressionLandscapeLayerSwitch.cpp
@@ -34,7 +34,7 @@ UMaterialExpressionLandscapeLayerSwitch::UMaterialExpressionLandscapeLayerSwitch
 	bCollapsed = false;
 #endif
 
-	PreviewUsed = false;
+	PreviewUsed = true;
 }
 
 #if WITH_EDITOR
