LandscapeNotifications: Unify landscape notification's into a single notification/notification manager class.

diff --git a/Engine/Source/Editor/UnrealEd/Private/GrassRenderingNotification.cpp b/Engine/Source/Editor/UnrealEd/Private/GrassRenderingNotification.cpp
deleted file mode 100644
diff --git a/Engine/Source/Editor/UnrealEd/Private/LandscapeTextureBakingNotification.cpp b/Engine/Source/Editor/UnrealEd/Private/LandscapeTextureBakingNotification.cpp
deleted file mode 100644
diff --git a/Engine/Source/Runtime/Landscape/Classes/Landscape.h b/Engine/Source/Runtime/Landscape/Classes/Landscape.h
--- a/Engine/Source/Runtime/Landscape/Classes/Landscape.h
+++ b/Engine/Source/Runtime/Landscape/Classes/Landscape.h
@@ -16,7 +16,7 @@ class SNotificationItem;
 class UStreamableRenderAsset;
 class FMaterialResource;
 struct FLandscapeEditLayerComponentReadbackResult;
-struct FNotificationInfo;
+struct FLandscapeNotification;
 struct FTextureToComponentHelper;
 struct FUpdateLayersContentContext;
 struct FEditLayersHeightmapMergeParams;
@@ -437,8 +437,6 @@ private:
 
 	static bool IsTextureReady(UTexture2D* InTexture, bool bInWaitForStreaming);
 	static bool IsMaterialResourceCompiled(FMaterialResource* InMaterialResource, bool bInWaitForCompilation);
-	static void ShowEditLayersResourcesNotification(const FText& InText, TWeakPtr<SNotificationItem>& NotificationItem);
-	static void HideEditLayersResourcesNotification(TWeakPtr<SNotificationItem>& InNotificationItem);
 #endif // WITH_EDITOR
 
 public:
@@ -527,11 +525,18 @@ private:
 	UPROPERTY(Transient)
 	bool bSplineLayerUpdateRequested;
 
-	/** Time since waiting for edit layers resources to be ready (for displaying a notification to the user) */
-	double WaitingForResourcesStartTime = -1.0;
+	/** Time since waiting for landscape resources to be ready (for displaying a notification to the user) */
+	double WaitingForLandscapeResourcesStartTime = -1.0;
 
-	/** User notification for when edit layer resources are not ready */
-	TWeakPtr<SNotificationItem> EditLayersResourcesNotification;
+	/** Time since waiting for brush resources to be ready (for displaying a notification to the user) */
+	double WaitingForLandscapeBrushResourcesStartTime = -1.0;
+
+	/** Non-stackable user notifications for landscape editor */
+	TSharedPtr<FLandscapeNotification> WaitingForTexturesNotification;
+	TSharedPtr<FLandscapeNotification> WaitingForBrushesNotification;
+	TSharedPtr<FLandscapeNotification> InvalidShadingModelNotification;
+	TSharedPtr<FLandscapeNotification> TextureBakingNotification;
+	TSharedPtr<FLandscapeNotification> GrassRenderingNotification;
 
 	// Represent all the resolved paint layer, from all layers blended together (size of the landscape x material layer count)
 	class FLandscapeTexture2DArrayResource* CombinedLayersWeightmapAllMaterialLayersResource;
diff --git a/Engine/Source/Runtime/Landscape/Classes/LandscapeProxy.h b/Engine/Source/Runtime/Landscape/Classes/LandscapeProxy.h
--- a/Engine/Source/Runtime/Landscape/Classes/LandscapeProxy.h
+++ b/Engine/Source/Runtime/Landscape/Classes/LandscapeProxy.h
@@ -1113,6 +1113,7 @@ public:
 
 	/* For the grassmap rendering notification */
 	int32 NumComponentsNeedingGrassMapRender;
+	UE_DEPRECATED(5.1, "No longer used for the grassmap rendering notification.")
 	LANDSCAPE_API static int32 TotalComponentsNeedingGrassMapRender;
 
 	/* To throttle texture streaming when we're trying to render a grassmap */
@@ -1121,6 +1122,7 @@ public:
 
 	/* For the texture baking notification */
 	int32 NumComponentsNeedingTextureBaking;
+	UE_DEPRECATED(5.1, "No longer used for the texture baking notification.")
 	LANDSCAPE_API static int32 TotalComponentsNeedingTextureBaking;
 
 	/** remove an overlapping component. Called from MapCheck. */
diff --git a/Engine/Source/Runtime/Landscape/Private/Landscape.cpp b/Engine/Source/Runtime/Landscape/Private/Landscape.cpp
--- a/Engine/Source/Runtime/Landscape/Private/Landscape.cpp
+++ b/Engine/Source/Runtime/Landscape/Private/Landscape.cpp
@@ -80,6 +80,8 @@ Landscape.cpp: Terrain rendering
 #include "LandscapeVersion.h"
 #include "UObject/FortniteMainBranchObjectVersion.h"
 #include "LandscapeDataAccess.h"
+#include "LandscapeNotification.h"
+#include "LandscapeHeightfieldCollisionComponent.h"
 #include "UObject/EditorObjectVersion.h"
 #include "Algo/BinarySearch.h"
 #include "WorldPartition/WorldPartition.h"
@@ -89,6 +91,8 @@ Landscape.cpp: Terrain rendering
 #if WITH_EDITOR
 #include "Rendering/StaticLightingSystemInterface.h"
 #include "Misc/ScopedSlowTask.h"
+#include "UnrealEdGlobals.h"
+#include "Editor/UnrealEdEngine.h"
 #endif
 
 /** Landscape stats */
@@ -1293,7 +1297,6 @@ ALandscape::ALandscape(const FObjectInitializer& ObjectInitializer)
 	CombinedLayersWeightmapAllMaterialLayersResource = nullptr;
 	CurrentLayersWeightmapAllMaterialLayersResource = nullptr;
 	WeightmapScratchExtractLayerTextureResource = nullptr;
-	WeightmapScratchPackLayerTextureResource = nullptr;
 	bLandscapeLayersAreInitialized = false;
 	bLandscapeLayersAreUsingLocalMerge = false;
 	LandscapeEdMode = nullptr;
@@ -1301,6 +1304,50 @@ ALandscape::ALandscape(const FObjectInitializer& ObjectInitializer)
 	bIsSpatiallyLoaded = false;
 	bDefaultOutlinerExpansionState = false;
 #endif // WITH_EDITORONLY_DATA
+
+#if WITH_EDITOR
+	if (!IsTemplate())
+	{
+		FLandscapeNotificationManager* LandscapeNotificationManager = GetWorld()->GetSubsystem<ULandscapeSubsystem>()->GetNotificationManager();
+	
+		// Create conditional TextureBakingNotification
+		TextureBakingNotification = MakeShared<FLandscapeNotification>(this, FLandscapeNotification::EType::TextureBaking,
+			[this]() { return NumComponentsNeedingTextureBaking > 0; },
+			[this](FText& InText)
+			{
+				// Accumulating all outstanding textures needing baking for ALL landscapes because only one such notification can be displayed at a time
+				int ComponentsNeedingTextureBakingAllLandscapes = 0;
+				for (const ALandscape* Landscape : TObjectRange<ALandscape>(RF_ClassDefaultObject | RF_ArchetypeObject, true, EInternalObjectFlags::Garbage))
+				{
+					ComponentsNeedingTextureBakingAllLandscapes += Landscape->NumComponentsNeedingTextureBaking;
+				}
+				
+				FFormatNamedArguments Args;
+				Args.Add(TEXT("OutstandingTextures"), FText::AsNumber(ComponentsNeedingTextureBakingAllLandscapes));
+				InText = FText::Format(NSLOCTEXT("TextureBaking", "TextureBakingFormat", "Baking Landscape Textures ({OutstandingTextures})"), Args);
+			});
+
+		// Create conditional GrassRenderingNotification
+		GrassRenderingNotification = MakeShared<FLandscapeNotification>(this, FLandscapeNotification::EType::GrassRendering,
+			[this]() { return NumComponentsNeedingGrassMapRender > 0 && !(FSlateApplication::Get().HasAnyMouseCaptor() || GUnrealEd->IsUserInteracting()); },
+			[this](FText& InText)
+			{
+				// Accumulating all outstanding grass maps that need to be rendered for ALL landscapes because only one such notification can be displayed at a time
+				int ComponentsNeedingGrassMapRenderAllLandscapes = 0;
+				for (const ALandscape* Landscape : TObjectRange<ALandscape>(RF_ClassDefaultObject | RF_ArchetypeObject, true, EInternalObjectFlags::Garbage))
+				{
+					ComponentsNeedingGrassMapRenderAllLandscapes += Landscape->NumComponentsNeedingGrassMapRender;
+				}
+				
+				FFormatNamedArguments Args;
+				Args.Add(TEXT("OutstandingGrassMaps"), FText::AsNumber(ComponentsNeedingGrassMapRenderAllLandscapes));
+				InText = FText::Format(NSLOCTEXT("GrassMapRender", "GrassMapRenderFormat", "Building Grass Maps ({OutstandingGrassMaps})"), Args);
+			});
+
+		LandscapeNotificationManager->RegisterNotification(TextureBakingNotification);
+		LandscapeNotificationManager->RegisterNotification(GrassRenderingNotification);
+	}
+#endif
 }
 
 ALandscapeStreamingProxy::ALandscapeStreamingProxy(const FObjectInitializer& ObjectInitializer)
@@ -3292,7 +3339,6 @@ void ALandscapeProxy::Destroyed()
 			SplineComponent->ModifySplines();
 		}
 
-		TotalComponentsNeedingGrassMapRender -= NumComponentsNeedingGrassMapRender;
 		NumComponentsNeedingGrassMapRender = 0;
 		TotalTexturesToStreamForVisibleGrassMapRender -= NumTexturesToStreamForVisibleGrassMapRender;
 		NumTexturesToStreamForVisibleGrassMapRender = 0;
@@ -4431,7 +4477,6 @@ ALandscapeProxy::~ALandscapeProxy()
 	AsyncFoliageTasks.Empty();
 
 #if WITH_EDITOR
-	TotalComponentsNeedingGrassMapRender -= NumComponentsNeedingGrassMapRender;
 	NumComponentsNeedingGrassMapRender = 0;
 	TotalTexturesToStreamForVisibleGrassMapRender -= NumTexturesToStreamForVisibleGrassMapRender;
 	NumTexturesToStreamForVisibleGrassMapRender = 0;
@@ -4754,7 +4799,6 @@ void ALandscapeProxy::UpdateGIBakedTextures(bool bBakeAllGITextures)
 		return;
 	}
 
-	TotalComponentsNeedingTextureBaking -= NumComponentsNeedingTextureBaking;
 	NumComponentsNeedingTextureBaking = 0;
 	int32 NumGenerated = 0;
 
@@ -4840,9 +4884,7 @@ void ALandscapeProxy::UpdateGIBakedTextures(bool bBakeAllGITextures)
 				}
 			}
 		}
-
-	TotalComponentsNeedingTextureBaking += NumComponentsNeedingTextureBaking;
-
+	
 	if (NumGenerated == 0)
 	{
 		// Don't check if we need to update anything for another 60 frames
diff --git a/Engine/Source/Runtime/Landscape/Private/LandscapeEditLayers.cpp b/Engine/Source/Runtime/Landscape/Private/LandscapeEditLayers.cpp
--- a/Engine/Source/Runtime/Landscape/Private/LandscapeEditLayers.cpp
+++ b/Engine/Source/Runtime/Landscape/Private/LandscapeEditLayers.cpp
@@ -57,10 +57,12 @@ LandscapeEditLayers.cpp: Landscape editing layers mode
 #include "Misc/ScopedSlowTask.h"
 #include "TextureCompiler.h"
 #include "Editor.h"
-#include "Widgets/Notifications/SNotificationList.h"
-#include "Framework/Notifications/NotificationManager.h"
+#include "LandscapeNotification.h"
+#include "LandscapeSubsystem.h"
 #include "ObjectCacheContext.h"
 #include "Components/RuntimeVirtualTextureComponent.h"
+#include "UnrealEdGlobals.h" // GUnrealEd
+#include "Editor/UnrealEdEngine.h"
 #endif
 
 #define LOCTEXT_NAMESPACE "Landscape"
@@ -1682,36 +1684,6 @@ typedef FLandscapeLayersRender_RenderThread<FLandscapeLayersWeightmapShaderParam
 
 #if WITH_EDITOR
 
-void ALandscape::ShowEditLayersResourcesNotification(const FText& InText, TWeakPtr<SNotificationItem>& NotificationItem)
-{
-	TSharedPtr<SNotificationItem> PinnedItem = NotificationItem.Pin();
-	if (!PinnedItem.IsValid())
-	{
-		FNotificationInfo Info(InText);
-		Info.bUseThrobber = true;
-		PinnedItem = FSlateNotificationManager::Get().AddNotification(Info);
-		NotificationItem = PinnedItem;
-	}
-	else
-	{
-		PinnedItem->SetText(InText);
-	}
-	PinnedItem->SetCompletionState(SNotificationItem::ECompletionState::CS_Pending);
-	PinnedItem->SetExpireDuration(1.0f);
-	PinnedItem->ExpireAndFadeout();
-}
-
-void ALandscape::HideEditLayersResourcesNotification(TWeakPtr<SNotificationItem>& InNotificationItem)
-{
-	TSharedPtr<SNotificationItem> PinnedItem = InNotificationItem.Pin();
-	if (PinnedItem.IsValid() && (PinnedItem->GetCompletionState() != SNotificationItem::ECompletionState::CS_Success))
-	{
-		PinnedItem->SetCompletionState(SNotificationItem::ECompletionState::CS_Success);
-		PinnedItem->SetExpireDuration(1.0f);
-		PinnedItem->ExpireAndFadeout();
-	}
-}
-
 bool ALandscape::IsTextureReady(UTexture2D* InTexture, bool bInWaitForStreaming)
 {
 	check(InTexture);
@@ -2526,7 +2498,7 @@ struct FLandscapeEditLayerComponentReadbackResult
 	/** Indicates which of the component's weightmaps is not needed anymore. */
 	TArray<ULandscapeLayerInfoObject*> AllZeroLayers;
 
-	FLandscapeEditLayerComponentReadbackResult(ULandscapeComponent* InLandscapeComponent, uint32 InUpdateModes = 0)
+	FLandscapeEditLayerComponentReadbackResult(ULandscapeComponent* InLandscapeComponent, uint32 InUpdateModes)
 		: LandscapeComponent(InLandscapeComponent)
 		, UpdateModes(InUpdateModes)
 	{}
@@ -3516,6 +3488,7 @@ bool ALandscape::PrepareLayersTextureResources(const TArray<FLandscapeLayer>& In
 bool ALandscape::PrepareLayersBrushResources(bool bInWaitForStreaming)
 {
 	TRACE_CPUPROFILER_EVENT_SCOPE(LandscapeLayers_PrepareLayersBrushTextureResources);
+
 	TSet<UObject*> Dependencies;
 	for (const FLandscapeLayer& Layer : LandscapeLayers)
 	{
@@ -7781,14 +7754,17 @@ void ALandscape::UpdateLayersContent(bool bInWaitForStreaming, bool bInSkipMonit
 {
 	TRACE_CPUPROFILER_EVENT_SCOPE(LandscapeLayers_UpdateLayersContent);
 
-	bool bHideEditLayerResourcesNotification = true;
+	bool bHideNotifications = true;
 	ON_SCOPE_EXIT
 	{
-		// Make sure to hide the notification if necessary when we early-out :
-		if (bHideEditLayerResourcesNotification)
+		// Make sure that we don't leave any notification behind when we leave this function without explicitly displaying one :
+		if (bHideNotifications)
 		{
-			HideEditLayersResourcesNotification(EditLayersResourcesNotification);
-			WaitingForResourcesStartTime = -1.0;
+			WaitingForTexturesNotification.Reset();
+			WaitingForBrushesNotification.Reset();
+			InvalidShadingModelNotification.Reset();
+			WaitingForLandscapeResourcesStartTime = -1.0;
+			WaitingForLandscapeBrushResourcesStartTime = -1.0;
 		}
 	};
 
@@ -7857,30 +7833,77 @@ void ALandscape::UpdateLayersContent(bool bInWaitForStreaming, bool bInSkipMonit
 	{
 		return;
 	}
+	
+	FLandscapeNotificationManager* LandscapeNotificationManager = GetWorld()->GetSubsystem<ULandscapeSubsystem>()->GetNotificationManager();
+	
+	// The Edit layers shaders only work on SM5 : cancel any update that might happen when SM5+ shading model is not active :
+	ERHIFeatureLevel::Type FeatureLevel = GetWorld() ? (ERHIFeatureLevel::Type)GetWorld()->FeatureLevel : GMaxRHIFeatureLevel;
+	if (FeatureLevel < ERHIFeatureLevel::SM5)
+	{
+		if (!InvalidShadingModelNotification.IsValid())
+		{
+			InvalidShadingModelNotification = MakeShared<FLandscapeNotification>(this, FLandscapeNotification::EType::ShadingModelInvalid);
+			static const FText NotificationText(LOCTEXT("InvalidShadingModel", "Landscape editor cannot update landscape with a feature level less than SM5!"));
+			InvalidShadingModelNotification->NotificationText = NotificationText;
+		}
+		LandscapeNotificationManager->RegisterNotification(InvalidShadingModelNotification);
+		bHideNotifications = false;
+		return;
+	}
+	else
+	{
+		InvalidShadingModelNotification.Reset();
+	}
 
 	// We need to wait until layers texture resources are ready to initialize the landscape to avoid taking the sizes and format of the default texture:
+	static constexpr double TimeBeforeDisplayingWaitingForResourcesNotification = 3.0;
+
 	bResourcesReady &= PrepareLayersTextureResources(bInWaitForStreaming);
-	// We need to wait until brushes texture/shader resources are ready to process the edit layers:
-	bResourcesReady &= PrepareLayersBrushResources(bInWaitForStreaming);
 	if (!bResourcesReady)
 	{
-		if (FSlateApplicationBase::IsInitialized() && !bInWaitForStreaming)
+		if (WaitingForLandscapeResourcesStartTime < 0.0)
 		{
-			static constexpr double TimeBeforeDisplayingNotification = 3.0;
-			if (WaitingForResourcesStartTime < 0.0)
+			WaitingForLandscapeResourcesStartTime = FSlateApplicationBase::Get().GetCurrentTime();
+			if (!WaitingForTexturesNotification.IsValid())
 			{
-				WaitingForResourcesStartTime = FSlateApplicationBase::Get().GetCurrentTime();
+				WaitingForTexturesNotification = MakeShared<FLandscapeNotification>(this, FLandscapeNotification::EType::LandscapeResourcesNotReady);
+				static const FText NotificationText(LOCTEXT("WaitForLanscapeResources", "Landscape editor waiting for landscape resources to be ready."));
+				WaitingForTexturesNotification->NotificationText = NotificationText;
+				WaitingForTexturesNotification->NotificationStartTime = WaitingForLandscapeResourcesStartTime + TimeBeforeDisplayingWaitingForResourcesNotification;
 			}
+			LandscapeNotificationManager->RegisterNotification(WaitingForTexturesNotification);
+			bHideNotifications = false;
+		}
+	}
+	else
+	{
+		WaitingForTexturesNotification.Reset();
+	}
 
-			if ((FSlateApplicationBase::Get().GetCurrentTime() - WaitingForResourcesStartTime) > TimeBeforeDisplayingNotification)
+	bResourcesReady &= PrepareLayersBrushResources(bInWaitForStreaming);
+	if (!bResourcesReady)
+	{
+		if (WaitingForLandscapeBrushResourcesStartTime < 0.0)
+		{
+			WaitingForLandscapeBrushResourcesStartTime = FSlateApplicationBase::Get().GetCurrentTime();
+			if (!WaitingForBrushesNotification.IsValid())
 			{
-				// let the user know we are waiting for resources :
-				static const FText NotificationText(LOCTEXT("WaitForLayersResources", "Landscape editor waiting for edit layers resources to be ready."));
-				ShowEditLayersResourcesNotification(NotificationText, EditLayersResourcesNotification);
+				WaitingForBrushesNotification = MakeShared<FLandscapeNotification>(this, FLandscapeNotification::EType::LandscapeResourcesNotReady);
+				static const FText NotificationText(LOCTEXT("WaitForLanscapeResources", "Landscape editor waiting for landscape resources to be ready."));
+				WaitingForBrushesNotification->NotificationText = NotificationText;
+				WaitingForBrushesNotification->NotificationStartTime = WaitingForLandscapeResourcesStartTime + TimeBeforeDisplayingWaitingForResourcesNotification;
 			}
-			// The notification may not be visible yet (because of the initial delay) but it should not be hidden and the initial delay timer shouldn't be reset : 
-			bHideEditLayerResourcesNotification = false;
+			LandscapeNotificationManager->RegisterNotification(WaitingForBrushesNotification);
+			bHideNotifications = false;
 		}
+	}
+	else
+	{
+		WaitingForBrushesNotification.Reset();
+	}
+
+	if (!bResourcesReady)
+	{
 		return;
 	}
 
diff --git a/Engine/Source/Runtime/Landscape/Private/LandscapeGrass.cpp b/Engine/Source/Runtime/Landscape/Private/LandscapeGrass.cpp
--- a/Engine/Source/Runtime/Landscape/Private/LandscapeGrass.cpp
+++ b/Engine/Source/Runtime/Landscape/Private/LandscapeGrass.cpp
@@ -2455,9 +2455,7 @@ int32 ALandscapeProxy::GetOutdatedGrassMapCount() const
 	return OutdatedGrassMaps;
 }
 
-int32 ALandscapeProxy::TotalComponentsNeedingGrassMapRender = 0;
 int32 ALandscapeProxy::TotalTexturesToStreamForVisibleGrassMapRender = 0;
-int32 ALandscapeProxy::TotalComponentsNeedingTextureBaking = 0;
 
 void ALandscapeProxy::UpdateGrassDataStatus(TSet<UTexture2D*>* OutCurrentForcedStreamedTextures, TSet<UTexture2D*>* OutDesiredForcedStreamedTextures, TSet<ULandscapeComponent*>* OutComponentsNeedingGrassMapRender, TSet<ULandscapeComponent*>* OutOutdatedComponents, bool bInEnableForceResidentFlag, int32* OutOutdatedGrassMaps) const
 {
@@ -2610,7 +2608,6 @@ void ALandscapeProxy::UpdateGrassData(bool bInShouldMarkDirty, FScopedSlowTask*
 	NumTexturesToStreamForVisibleGrassMapRender = 0;
 
 	// Remove local count from global count
-	TotalComponentsNeedingGrassMapRender -= NumComponentsNeedingGrassMapRender;
 	NumComponentsNeedingGrassMapRender = 0;
 
 	// Wait for Texture Streaming
@@ -3181,9 +3178,7 @@ void ALandscapeProxy::UpdateGrass(const TArray<FVector>& Cameras, int32& InOutNu
 					}
 				}
 
-				TotalComponentsNeedingGrassMapRender -= NumComponentsNeedingGrassMapRender;
 				NumComponentsNeedingGrassMapRender = ComponentsNeedingGrassMapRender.Num() - NumComponentsRendered - NumComponentsUnableToRender;
-				TotalComponentsNeedingGrassMapRender += NumComponentsNeedingGrassMapRender;
 
 				// Update resident flags
 				for (auto Texture : DesiredForceStreamedTextures.Difference(CurrentForcedStreamedTextures))
diff --git a/Engine/Source/Runtime/Landscape/Private/LandscapeNotification.cpp b/Engine/Source/Runtime/Landscape/Private/LandscapeNotification.cpp
new file mode 100644
--- /dev/null
+++ b/Engine/Source/Runtime/Landscape/Private/LandscapeNotification.cpp
@@ -0,0 +1,129 @@
+// Copyright Epic Games, Inc. All Rights Reserved.
+
+#include "LandscapeNotification.h"
+#include "LandscapeSubsystem.h"
+
+#if WITH_EDITOR
+
+#include "Widgets/Notifications/SNotificationList.h"
+#include "Framework/Notifications/NotificationManager.h"
+#include "Algo/Count.h"
+
+#define LOCTEXT_NAMESPACE "Landscape"
+
+// ----------------------------------------------------------------------------------
+
+FLandscapeNotification::FLandscapeNotification(const TWeakObjectPtr<ALandscape>& InLandscape, EType InNotificationType, FConditionCallback InConditionCallback, FUpdateTextCallback InUpdateTextCallback)
+	: Landscape(InLandscape)
+	, NotificationType(InNotificationType)
+	, ConditionCallback(InConditionCallback)
+	, UpdateTextCallback(InUpdateTextCallback)
+{}
+
+bool FLandscapeNotification::operator == (const FLandscapeNotification& Other) const
+{
+	return (Landscape == Other.Landscape) && (NotificationType == Other.NotificationType);
+}
+
+bool FLandscapeNotification::operator < (const FLandscapeNotification& Other) const
+{
+	if (NotificationType == Other.NotificationType)
+	{
+		return (Landscape.Get() < Other.Landscape.Get());
+	}
+	return (NotificationType < Other.NotificationType);
+}
+
+bool FLandscapeNotification::ShouldShowNotification() const 
+{
+	return ConditionCallback() && (FSlateApplicationBase::Get().GetCurrentTime() >= NotificationStartTime || NotificationStartTime < 0.0);
+}
+
+void FLandscapeNotification::SetNotificationText()
+{
+	UpdateTextCallback(NotificationText);
+}
+
+// ----------------------------------------------------------------------------------
+
+void FLandscapeNotificationManager::Tick()
+{
+	// Cleanup notifications for stale landscapes:
+	LandscapeNotifications.RemoveAll([](const TWeakPtr<FLandscapeNotification>& Notification) 
+		{ return !Notification.IsValid() || !Notification.Pin()->GetLandscape().IsValid(); });
+
+	// Don't keep notifications that should not yet be displayed: 
+	TArray<TWeakPtr<FLandscapeNotification>> ValidNotifications = LandscapeNotifications.FilterByPredicate([=](const TWeakPtr<FLandscapeNotification>& Notification) 
+		{ return Notification.Pin()->ShouldShowNotification(); });
+	
+	if (ValidNotifications.IsEmpty())
+	{
+		HideNotificationItem();
+	}
+	else
+	{
+		// Sort by priority: 
+		ValidNotifications.Sort([](const TWeakPtr<FLandscapeNotification>& LHS, const TWeakPtr<FLandscapeNotification>& RHS) { return (*LHS.Pin().Get() < *RHS.Pin().Get()); });
+
+		// Only display the topmost priority:
+		ActiveNotification = ValidNotifications[0];
+		FLandscapeNotification* LocalActiveNotification = ActiveNotification.Pin().Get();
+		check(LocalActiveNotification != nullptr);
+
+		LocalActiveNotification->SetNotificationText();
+
+		// There might be multiple notifications for this type:
+		int32 NumIdenticalNotifications = Algo::CountIf(ValidNotifications, [NotificationType = LocalActiveNotification->GetNotificationType()](const TWeakPtr<FLandscapeNotification>& Notification)
+		{ return (Notification.Pin().Get()->GetNotificationType() == NotificationType); });
+
+		FText Text = FText::Format(LOCTEXT("NotificationFooter", "{0}\n[{1}]"),
+			LocalActiveNotification->NotificationText,
+			(NumIdenticalNotifications > 1) ? LOCTEXT("MultipleLandscapes", "Multiple landscapes") : FText::FromString(LocalActiveNotification->GetLandscape()->GetActorLabel()));
+
+		ShowNotificationItem(Text);
+	}
+}
+
+void FLandscapeNotificationManager::RegisterNotification(const TWeakPtr<FLandscapeNotification>& InNotification)
+{
+	LandscapeNotifications.AddUnique(InNotification);
+}
+
+void FLandscapeNotificationManager::UnregisterNotification(const TWeakPtr<FLandscapeNotification>& InNotification)
+{
+	LandscapeNotifications.Remove(InNotification);
+}
+
+void FLandscapeNotificationManager::ShowNotificationItem(const FText& InText)
+{
+	TSharedPtr<SNotificationItem> PinnedItem = NotificationItem.Pin();
+	if (!PinnedItem.IsValid())
+	{
+		FNotificationInfo Info(InText);
+		Info.bUseThrobber = true;
+		Info.bFireAndForget = false;
+		PinnedItem = FSlateNotificationManager::Get().AddNotification(Info);
+		PinnedItem->SetFadeInDuration(0.0f);
+		NotificationItem = PinnedItem;
+	}
+	else
+	{
+		PinnedItem->SetText(InText);
+	}
+	PinnedItem->SetCompletionState(SNotificationItem::ECompletionState::CS_Pending);
+}
+
+void FLandscapeNotificationManager::HideNotificationItem()
+{
+	TSharedPtr<SNotificationItem> PinnedItem = NotificationItem.Pin();
+	if (PinnedItem.IsValid() && (PinnedItem->GetCompletionState() != SNotificationItem::ECompletionState::CS_Success))
+	{
+		PinnedItem->SetCompletionState(SNotificationItem::ECompletionState::CS_Success);
+		PinnedItem->SetExpireDuration(1.0f);
+		PinnedItem->ExpireAndFadeout();
+	}
+}
+
+#undef LOCTEXT_NAMESPACE
+
+#endif // WITH_EDITOR
\ No newline at end of file
diff --git a/Engine/Source/Runtime/Landscape/Private/LandscapeSubsystem.cpp b/Engine/Source/Runtime/Landscape/Private/LandscapeSubsystem.cpp
--- a/Engine/Source/Runtime/Landscape/Private/LandscapeSubsystem.cpp
+++ b/Engine/Source/Runtime/Landscape/Private/LandscapeSubsystem.cpp
@@ -14,6 +14,7 @@
 #include "ActorPartition/ActorPartitionSubsystem.h"
 #include "Engine/World.h"
 #include "Math/IntRect.h"
+#include "LandscapeNotification.h"
 #include "LandscapeConfigHelper.h"
 #include "Engine/Canvas.h"
 #include "EngineUtils.h"
@@ -54,24 +55,17 @@ void ULandscapeSubsystem::Initialize(FSubsystemCollectionBase& Collection)
 	GrassMapsBuilder = new FLandscapeGrassMapsBuilder(GetWorld());
 	GIBakedTextureBuilder = new FLandscapeGIBakedTextureBuilder(GetWorld());
 	PhysicalMaterialBuilder = new FLandscapePhysicalMaterialBuilder(GetWorld());
+	NotificationManager = new FLandscapeNotificationManager();
 #endif
 }
 
 void ULandscapeSubsystem::Deinitialize()
 {
 #if WITH_EDITOR
-	if (GrassMapsBuilder)
-	{
-		delete GrassMapsBuilder;
-	}
-	if (GIBakedTextureBuilder)
-	{
-		delete GIBakedTextureBuilder;
-	}
-	if (PhysicalMaterialBuilder)
-	{
-		delete PhysicalMaterialBuilder;
-	}
+	delete GrassMapsBuilder;
+	delete GIBakedTextureBuilder;
+	delete PhysicalMaterialBuilder;
+	delete NotificationManager;
 #endif
 	Proxies.Empty();
 
@@ -159,6 +153,7 @@ void ULandscapeSubsystem::Tick(float DeltaTime)
 	if (GIsEditor && !World->IsPlayInEditor())
 	{
 		LandscapePhysicalMaterial::GarbageCollectTasks();
+		NotificationManager->Tick();
 	}
 #endif
 }
diff --git a/Engine/Source/Runtime/Landscape/Public/LandscapeNotification.h b/Engine/Source/Runtime/Landscape/Public/LandscapeNotification.h
new file mode 100644
--- /dev/null
+++ b/Engine/Source/Runtime/Landscape/Public/LandscapeNotification.h
@@ -0,0 +1,88 @@
+// Copyright Epic Games, Inc. All Rights Reserved.
+
+#pragma once
+
+#if WITH_EDITOR
+
+#include "CoreMinimal.h"
+#include "Landscape.h"
+
+struct FLandscapeNotification
+{
+	/** The various types of landscape notifications, ordered by priority (a single notification is displayed at a time). */
+	enum class EType : uint8
+	{
+		LandscapeResourcesNotReady = 0,
+		LandscapeBrushResourcesNotReady,
+		ShadingModelInvalid,
+		TextureBaking,
+		GrassRendering,
+	};
+
+	using FConditionCallback = TFunction<bool(void)>;
+	using FUpdateTextCallback = TFunction<void(FText&)>;
+
+	/** 
+	* Constructs a notification which is expected to have a longer lifetime and behave in a dynamic way.
+	* 
+	* @param InLandscape - The landscape actor associated with this notification.
+	* @param InNotificationType - The type of notification which determines its display priority.
+	* @param InConditionCallback - Lambda used to conditionally determine whether this notification should be displayed per tick.
+	* @param InUpdateTextCallback - Lambda used to construct notification text when the text can change (e.g. displaying the value of a variable).
+	*/ 
+	FLandscapeNotification(const TWeakObjectPtr<ALandscape>& InLandscape, EType InNotificationType, FConditionCallback InConditionCallback = FConditionCallback([]() { return true; }), FUpdateTextCallback InUpdateTextCallback = [](FText& InText) {});
+	~FLandscapeNotification() = default;
+
+	bool operator == (const FLandscapeNotification& Other) const;
+	bool operator < (const FLandscapeNotification& Other) const;
+
+	TWeakObjectPtr<ALandscape> GetLandscape() const { return Landscape; }
+	EType GetNotificationType() const { return NotificationType; }
+
+	// Public wrappers for the callback functions.
+	bool ShouldShowNotification() const;
+	void SetNotificationText();
+
+private:
+	FLandscapeNotification() = delete;
+
+public:
+	/** The text to be displayed.Can either be set by user directly or by setting SetTextCallback at notification construction. */ 
+	FText NotificationText;
+
+	/** Determines when a notification should be displayed, ignored if less than 0. */ 
+	double NotificationStartTime = -1.0;
+
+private:
+	TWeakObjectPtr<ALandscape> Landscape;
+	EType NotificationType;
+
+	/** Determines whether or not a notification should be shown. Defaults to return true. */
+	FConditionCallback ConditionCallback;
+
+	/** Sets the notification text per frame by manipulating InText&. Defaults to do nothing, allowing direct manipulation of NotificationText. */
+	FUpdateTextCallback UpdateTextCallback;
+};
+
+/**
+ * FLandscapeNotificationManager : centralizes landscape-related user notifications so that the user doesn't get flooded by toasts when multiple messages from possibly multiple landscapes are tossed around
+ */
+class FLandscapeNotificationManager
+{
+public:
+	void Tick();
+
+	void RegisterNotification(const TWeakPtr<FLandscapeNotification>& InNotification);
+	void UnregisterNotification(const TWeakPtr<FLandscapeNotification>& InNotification);
+
+private:
+	void ShowNotificationItem(const FText& InText);
+	void HideNotificationItem();
+
+private:
+	TArray<TWeakPtr<FLandscapeNotification>> LandscapeNotifications;
+	TWeakPtr<FLandscapeNotification> ActiveNotification;
+	TWeakPtr<SNotificationItem> NotificationItem;
+};
+
+#endif // WITH_EDITOR
\ No newline at end of file
diff --git a/Engine/Source/Runtime/Landscape/Public/LandscapeSubsystem.h b/Engine/Source/Runtime/Landscape/Public/LandscapeSubsystem.h
--- a/Engine/Source/Runtime/Landscape/Public/LandscapeSubsystem.h
+++ b/Engine/Source/Runtime/Landscape/Public/LandscapeSubsystem.h
@@ -9,6 +9,7 @@
 
 class ALandscapeProxy;
 class ULandscapeInfo;
+class FLandscapeNotificationManager;
 
 UCLASS(MinimalAPI)
 class ULandscapeSubsystem : public UTickableWorldSubsystem
@@ -42,6 +43,7 @@ public:
 	LANDSCAPE_API void ChangeGridSize(ULandscapeInfo* LandscapeInfo, uint32 NewGridSizeInComponents);
 	LANDSCAPE_API ALandscapeProxy* FindOrAddLandscapeProxy(ULandscapeInfo* LandscapeInfo, const FIntPoint& SectionBase);
 	LANDSCAPE_API void DisplayBuildMessages(class FCanvas* Canvas, float& XPos, float& YPos);
+	FLandscapeNotificationManager* GetNotificationManager() { return NotificationManager; }
 #endif
 
 private:
@@ -56,5 +58,7 @@ private:
 	class FLandscapeGrassMapsBuilder* GrassMapsBuilder;
 	class FLandscapeGIBakedTextureBuilder* GIBakedTextureBuilder;
 	class FLandscapePhysicalMaterialBuilder* PhysicalMaterialBuilder;
+	
+	FLandscapeNotificationManager* NotificationManager;
 #endif
 };
