[Backout] - CL21458279
Original CL Desc
-----------------------------------------------------------------
LandscapeNotifications: Unify landscape notification's into a single notification/notification manager class.

diff --git a/Engine/Source/Editor/UnrealEd/Private/GrassRenderingNotification.cpp b/Engine/Source/Editor/UnrealEd/Private/GrassRenderingNotification.cpp
new file mode 100644
--- /dev/null
+++ b/Engine/Source/Editor/UnrealEd/Private/GrassRenderingNotification.cpp
@@ -0,0 +1,38 @@
+// Copyright Epic Games, Inc. All Rights Reserved.
+
+#include "CoreMinimal.h"
+#include "GlobalEditorNotification.h"
+#include "LandscapeProxy.h"
+#include "Widgets/Notifications/SNotificationList.h"
+#include "Framework/Application/SlateApplication.h"
+#include "UnrealEdGlobals.h"
+#include "Editor/UnrealEdEngine.h"
+
+/** Notification class for grassmmap rendering. */
+class FGrassRenderingNotificationImpl : public FGlobalEditorNotification
+{
+protected:
+	/** FGlobalEditorNotification interface */
+	virtual bool ShouldShowNotification(const bool bIsNotificationAlreadyActive) const override;
+	virtual void SetNotificationText(const TSharedPtr<SNotificationItem>& InNotificationItem) const override;
+};
+
+/** Global notification object. */
+FGrassRenderingNotificationImpl GGrassRenderingNotification;
+
+bool FGrassRenderingNotificationImpl::ShouldShowNotification(const bool bIsNotificationAlreadyActive) const
+{
+	// Avoid showing the notification when a user is interacting as it causes a focus lost event which prevents them from dragging
+	return ALandscapeProxy::TotalComponentsNeedingGrassMapRender > 0 && !(FSlateApplication::Get().HasAnyMouseCaptor() || GUnrealEd->IsUserInteracting());
+}
+
+void FGrassRenderingNotificationImpl::SetNotificationText(const TSharedPtr<SNotificationItem>& InNotificationItem) const
+{
+	if (ALandscapeProxy::TotalComponentsNeedingGrassMapRender > 0)
+	{
+		FFormatNamedArguments Args;
+		Args.Add(TEXT("OutstandingGrassMaps"), FText::AsNumber(ALandscapeProxy::TotalComponentsNeedingGrassMapRender));
+		const FText ProgressMessage = FText::Format(NSLOCTEXT("GrassMapRender", "GrassMapRenderFormat", "Building Grass Maps ({OutstandingGrassMaps})"), Args);
+		InNotificationItem->SetText(ProgressMessage);
+	}
+}
diff --git a/Engine/Source/Editor/UnrealEd/Private/LandscapeTextureBakingNotification.cpp b/Engine/Source/Editor/UnrealEd/Private/LandscapeTextureBakingNotification.cpp
new file mode 100644
--- /dev/null
+++ b/Engine/Source/Editor/UnrealEd/Private/LandscapeTextureBakingNotification.cpp
@@ -0,0 +1,34 @@
+// Copyright Epic Games, Inc. All Rights Reserved.
+
+#include "CoreMinimal.h"
+#include "GlobalEditorNotification.h"
+#include "LandscapeProxy.h"
+#include "Widgets/Notifications/SNotificationList.h"
+
+/** Notification class for grassmmap rendering. */
+class FLandscapeTextureBakingNotificationImpl : public FGlobalEditorNotification
+{
+protected:
+	/** FGlobalEditorNotification interface */
+	virtual bool ShouldShowNotification(const bool bIsNotificationAlreadyActive) const override;
+	virtual void SetNotificationText(const TSharedPtr<SNotificationItem>& InNotificationItem) const override;
+};
+
+/** Global notification object. */
+FLandscapeTextureBakingNotificationImpl GLandscapeTextureBakingNotification;
+
+bool FLandscapeTextureBakingNotificationImpl::ShouldShowNotification(const bool bIsNotificationAlreadyActive) const
+{
+	return ALandscapeProxy::TotalComponentsNeedingTextureBaking > 0;
+}
+
+void FLandscapeTextureBakingNotificationImpl::SetNotificationText(const TSharedPtr<SNotificationItem>& InNotificationItem) const
+{
+	if (ALandscapeProxy::TotalComponentsNeedingTextureBaking > 0)
+	{
+		FFormatNamedArguments Args;
+		Args.Add(TEXT("OutstandingTextures"), FText::AsNumber(ALandscapeProxy::TotalComponentsNeedingTextureBaking));
+		const FText ProgressMessage = FText::Format(NSLOCTEXT("TextureBaking", "TextureBakingFormat", "Baking Landscape Textures ({OutstandingTextures})"), Args);
+		InNotificationItem->SetText(ProgressMessage);
+	}
+}
diff --git a/Engine/Source/Runtime/Landscape/Classes/Landscape.h b/Engine/Source/Runtime/Landscape/Classes/Landscape.h
--- a/Engine/Source/Runtime/Landscape/Classes/Landscape.h
+++ b/Engine/Source/Runtime/Landscape/Classes/Landscape.h
@@ -16,7 +16,7 @@ class SNotificationItem;
 class UStreamableRenderAsset;
 class FMaterialResource;
 struct FLandscapeEditLayerComponentReadbackResult;
-struct FLandscapeNotification;
+struct FNotificationInfo;
 struct FTextureToComponentHelper;
 struct FUpdateLayersContentContext;
 struct FEditLayersHeightmapMergeParams;
@@ -437,6 +437,8 @@ private:
 
 	static bool IsTextureReady(UTexture2D* InTexture, bool bInWaitForStreaming);
 	static bool IsMaterialResourceCompiled(FMaterialResource* InMaterialResource, bool bInWaitForCompilation);
+	static void ShowEditLayersResourcesNotification(const FText& InText, TWeakPtr<SNotificationItem>& NotificationItem);
+	static void HideEditLayersResourcesNotification(TWeakPtr<SNotificationItem>& InNotificationItem);
 #endif // WITH_EDITOR
 
 public:
@@ -525,18 +527,11 @@ private:
 	UPROPERTY(Transient)
 	bool bSplineLayerUpdateRequested;
 
-	/** Time since waiting for landscape resources to be ready (for displaying a notification to the user) */
-	double WaitingForLandscapeResourcesStartTime = -1.0;
+	/** Time since waiting for edit layers resources to be ready (for displaying a notification to the user) */
+	double WaitingForResourcesStartTime = -1.0;
 
-	/** Time since waiting for brush resources to be ready (for displaying a notification to the user) */
-	double WaitingForLandscapeBrushResourcesStartTime = -1.0;
-
-	/** Non-stackable user notifications for landscape editor */
-	TSharedPtr<FLandscapeNotification> WaitingForTexturesNotification;
-	TSharedPtr<FLandscapeNotification> WaitingForBrushesNotification;
-	TSharedPtr<FLandscapeNotification> InvalidShadingModelNotification;
-	TSharedPtr<FLandscapeNotification> TextureBakingNotification;
-	TSharedPtr<FLandscapeNotification> GrassRenderingNotification;
+	/** User notification for when edit layer resources are not ready */
+	TWeakPtr<SNotificationItem> EditLayersResourcesNotification;
 
 	// Represent all the resolved paint layer, from all layers blended together (size of the landscape x material layer count)
 	class FLandscapeTexture2DArrayResource* CombinedLayersWeightmapAllMaterialLayersResource;
diff --git a/Engine/Source/Runtime/Landscape/Classes/LandscapeProxy.h b/Engine/Source/Runtime/Landscape/Classes/LandscapeProxy.h
--- a/Engine/Source/Runtime/Landscape/Classes/LandscapeProxy.h
+++ b/Engine/Source/Runtime/Landscape/Classes/LandscapeProxy.h
@@ -1113,7 +1113,6 @@ public:
 
 	/* For the grassmap rendering notification */
 	int32 NumComponentsNeedingGrassMapRender;
-	UE_DEPRECATED(5.1, "No longer used for the grassmap rendering notification.")
 	LANDSCAPE_API static int32 TotalComponentsNeedingGrassMapRender;
 
 	/* To throttle texture streaming when we're trying to render a grassmap */
@@ -1122,7 +1121,6 @@ public:
 
 	/* For the texture baking notification */
 	int32 NumComponentsNeedingTextureBaking;
-	UE_DEPRECATED(5.1, "No longer used for the texture baking notification.")
 	LANDSCAPE_API static int32 TotalComponentsNeedingTextureBaking;
 
 	/** remove an overlapping component. Called from MapCheck. */
diff --git a/Engine/Source/Runtime/Landscape/Private/Landscape.cpp b/Engine/Source/Runtime/Landscape/Private/Landscape.cpp
--- a/Engine/Source/Runtime/Landscape/Private/Landscape.cpp
+++ b/Engine/Source/Runtime/Landscape/Private/Landscape.cpp
@@ -80,8 +80,6 @@ Landscape.cpp: Terrain rendering
 #include "LandscapeVersion.h"
 #include "UObject/FortniteMainBranchObjectVersion.h"
 #include "LandscapeDataAccess.h"
-#include "LandscapeNotification.h"
-#include "LandscapeHeightfieldCollisionComponent.h"
 #include "UObject/EditorObjectVersion.h"
 #include "Algo/BinarySearch.h"
 #include "WorldPartition/WorldPartition.h"
@@ -91,8 +89,6 @@ Landscape.cpp: Terrain rendering
 #if WITH_EDITOR
 #include "Rendering/StaticLightingSystemInterface.h"
 #include "Misc/ScopedSlowTask.h"
-#include "UnrealEdGlobals.h"
-#include "Editor/UnrealEdEngine.h"
 #endif
 
 /** Landscape stats */
@@ -1297,6 +1293,7 @@ ALandscape::ALandscape(const FObjectInitializer& ObjectInitializer)
 	CombinedLayersWeightmapAllMaterialLayersResource = nullptr;
 	CurrentLayersWeightmapAllMaterialLayersResource = nullptr;
 	WeightmapScratchExtractLayerTextureResource = nullptr;
+	WeightmapScratchPackLayerTextureResource = nullptr;
 	bLandscapeLayersAreInitialized = false;
 	bLandscapeLayersAreUsingLocalMerge = false;
 	LandscapeEdMode = nullptr;
@@ -1304,50 +1301,6 @@ ALandscape::ALandscape(const FObjectInitializer& ObjectInitializer)
 	bIsSpatiallyLoaded = false;
 	bDefaultOutlinerExpansionState = false;
 #endif // WITH_EDITORONLY_DATA
-
-#if WITH_EDITOR
-	if (!IsTemplate())
-	{
-		FLandscapeNotificationManager* LandscapeNotificationManager = GetWorld()->GetSubsystem<ULandscapeSubsystem>()->GetNotificationManager();
-	
-		// Create conditional TextureBakingNotification
-		TextureBakingNotification = MakeShared<FLandscapeNotification>(this, FLandscapeNotification::EType::TextureBaking,
-			[this]() { return NumComponentsNeedingTextureBaking > 0; },
-			[this](FText& InText)
-			{
-				// Accumulating all outstanding textures needing baking for ALL landscapes because only one such notification can be displayed at a time
-				int ComponentsNeedingTextureBakingAllLandscapes = 0;
-				for (const ALandscape* Landscape : TObjectRange<ALandscape>(RF_ClassDefaultObject | RF_ArchetypeObject, true, EInternalObjectFlags::Garbage))
-				{
-					ComponentsNeedingTextureBakingAllLandscapes += Landscape->NumComponentsNeedingTextureBaking;
-				}
-				
-				FFormatNamedArguments Args;
-				Args.Add(TEXT("OutstandingTextures"), FText::AsNumber(ComponentsNeedingTextureBakingAllLandscapes));
-				InText = FText::Format(NSLOCTEXT("TextureBaking", "TextureBakingFormat", "Baking Landscape Textures ({OutstandingTextures})"), Args);
-			});
-
-		// Create conditional GrassRenderingNotification
-		GrassRenderingNotification = MakeShared<FLandscapeNotification>(this, FLandscapeNotification::EType::GrassRendering,
-			[this]() { return NumComponentsNeedingGrassMapRender > 0 && !(FSlateApplication::Get().HasAnyMouseCaptor() || GUnrealEd->IsUserInteracting()); },
-			[this](FText& InText)
-			{
-				// Accumulating all outstanding grass maps that need to be rendered for ALL landscapes because only one such notification can be displayed at a time
-				int ComponentsNeedingGrassMapRenderAllLandscapes = 0;
-				for (const ALandscape* Landscape : TObjectRange<ALandscape>(RF_ClassDefaultObject | RF_ArchetypeObject, true, EInternalObjectFlags::Garbage))
-				{
-					ComponentsNeedingGrassMapRenderAllLandscapes += Landscape->NumComponentsNeedingGrassMapRender;
-				}
-				
-				FFormatNamedArguments Args;
-				Args.Add(TEXT("OutstandingGrassMaps"), FText::AsNumber(ComponentsNeedingGrassMapRenderAllLandscapes));
-				InText = FText::Format(NSLOCTEXT("GrassMapRender", "GrassMapRenderFormat", "Building Grass Maps ({OutstandingGrassMaps})"), Args);
-			});
-
-		LandscapeNotificationManager->RegisterNotification(TextureBakingNotification);
-		LandscapeNotificationManager->RegisterNotification(GrassRenderingNotification);
-	}
-#endif
 }
 
 ALandscapeStreamingProxy::ALandscapeStreamingProxy(const FObjectInitializer& ObjectInitializer)
@@ -3339,6 +3292,7 @@ void ALandscapeProxy::Destroyed()
 			SplineComponent->ModifySplines();
 		}
 
+		TotalComponentsNeedingGrassMapRender -= NumComponentsNeedingGrassMapRender;
 		NumComponentsNeedingGrassMapRender = 0;
 		TotalTexturesToStreamForVisibleGrassMapRender -= NumTexturesToStreamForVisibleGrassMapRender;
 		NumTexturesToStreamForVisibleGrassMapRender = 0;
@@ -4477,6 +4431,7 @@ ALandscapeProxy::~ALandscapeProxy()
 	AsyncFoliageTasks.Empty();
 
 #if WITH_EDITOR
+	TotalComponentsNeedingGrassMapRender -= NumComponentsNeedingGrassMapRender;
 	NumComponentsNeedingGrassMapRender = 0;
 	TotalTexturesToStreamForVisibleGrassMapRender -= NumTexturesToStreamForVisibleGrassMapRender;
 	NumTexturesToStreamForVisibleGrassMapRender = 0;
@@ -4799,6 +4754,7 @@ void ALandscapeProxy::UpdateGIBakedTextures(bool bBakeAllGITextures)
 		return;
 	}
 
+	TotalComponentsNeedingTextureBaking -= NumComponentsNeedingTextureBaking;
 	NumComponentsNeedingTextureBaking = 0;
 	int32 NumGenerated = 0;
 
@@ -4884,7 +4840,9 @@ void ALandscapeProxy::UpdateGIBakedTextures(bool bBakeAllGITextures)
 				}
 			}
 		}
-	
+
+	TotalComponentsNeedingTextureBaking += NumComponentsNeedingTextureBaking;
+
 	if (NumGenerated == 0)
 	{
 		// Don't check if we need to update anything for another 60 frames
diff --git a/Engine/Source/Runtime/Landscape/Private/LandscapeEditLayers.cpp b/Engine/Source/Runtime/Landscape/Private/LandscapeEditLayers.cpp
--- a/Engine/Source/Runtime/Landscape/Private/LandscapeEditLayers.cpp
+++ b/Engine/Source/Runtime/Landscape/Private/LandscapeEditLayers.cpp
@@ -57,12 +57,10 @@ LandscapeEditLayers.cpp: Landscape editing layers mode
 #include "Misc/ScopedSlowTask.h"
 #include "TextureCompiler.h"
 #include "Editor.h"
-#include "LandscapeNotification.h"
-#include "LandscapeSubsystem.h"
+#include "Widgets/Notifications/SNotificationList.h"
+#include "Framework/Notifications/NotificationManager.h"
 #include "ObjectCacheContext.h"
 #include "Components/RuntimeVirtualTextureComponent.h"
-#include "UnrealEdGlobals.h" // GUnrealEd
-#include "Editor/UnrealEdEngine.h"
 #endif
 
 #define LOCTEXT_NAMESPACE "Landscape"
@@ -1684,6 +1682,36 @@ typedef FLandscapeLayersRender_RenderThread<FLandscapeLayersWeightmapShaderParam
 
 #if WITH_EDITOR
 
+void ALandscape::ShowEditLayersResourcesNotification(const FText& InText, TWeakPtr<SNotificationItem>& NotificationItem)
+{
+	TSharedPtr<SNotificationItem> PinnedItem = NotificationItem.Pin();
+	if (!PinnedItem.IsValid())
+	{
+		FNotificationInfo Info(InText);
+		Info.bUseThrobber = true;
+		PinnedItem = FSlateNotificationManager::Get().AddNotification(Info);
+		NotificationItem = PinnedItem;
+	}
+	else
+	{
+		PinnedItem->SetText(InText);
+	}
+	PinnedItem->SetCompletionState(SNotificationItem::ECompletionState::CS_Pending);
+	PinnedItem->SetExpireDuration(1.0f);
+	PinnedItem->ExpireAndFadeout();
+}
+
+void ALandscape::HideEditLayersResourcesNotification(TWeakPtr<SNotificationItem>& InNotificationItem)
+{
+	TSharedPtr<SNotificationItem> PinnedItem = InNotificationItem.Pin();
+	if (PinnedItem.IsValid() && (PinnedItem->GetCompletionState() != SNotificationItem::ECompletionState::CS_Success))
+	{
+		PinnedItem->SetCompletionState(SNotificationItem::ECompletionState::CS_Success);
+		PinnedItem->SetExpireDuration(1.0f);
+		PinnedItem->ExpireAndFadeout();
+	}
+}
+
 bool ALandscape::IsTextureReady(UTexture2D* InTexture, bool bInWaitForStreaming)
 {
 	check(InTexture);
@@ -2498,7 +2526,7 @@ struct FLandscapeEditLayerComponentReadbackResult
 	/** Indicates which of the component's weightmaps is not needed anymore. */
 	TArray<ULandscapeLayerInfoObject*> AllZeroLayers;
 
-	FLandscapeEditLayerComponentReadbackResult(ULandscapeComponent* InLandscapeComponent, uint32 InUpdateModes)
+	FLandscapeEditLayerComponentReadbackResult(ULandscapeComponent* InLandscapeComponent, uint32 InUpdateModes = 0)
 		: LandscapeComponent(InLandscapeComponent)
 		, UpdateModes(InUpdateModes)
 	{}
@@ -3488,7 +3516,6 @@ bool ALandscape::PrepareLayersTextureResources(const TArray<FLandscapeLayer>& In
 bool ALandscape::PrepareLayersBrushResources(bool bInWaitForStreaming)
 {
 	TRACE_CPUPROFILER_EVENT_SCOPE(LandscapeLayers_PrepareLayersBrushTextureResources);
-
 	TSet<UObject*> Dependencies;
 	for (const FLandscapeLayer& Layer : LandscapeLayers)
 	{
@@ -7754,17 +7781,14 @@ void ALandscape::UpdateLayersContent(bool bInWaitForStreaming, bool bInSkipMonit
 {
 	TRACE_CPUPROFILER_EVENT_SCOPE(LandscapeLayers_UpdateLayersContent);
 
-	bool bHideNotifications = true;
+	bool bHideEditLayerResourcesNotification = true;
 	ON_SCOPE_EXIT
 	{
-		// Make sure that we don't leave any notification behind when we leave this function without explicitly displaying one :
-		if (bHideNotifications)
+		// Make sure to hide the notification if necessary when we early-out :
+		if (bHideEditLayerResourcesNotification)
 		{
-			WaitingForTexturesNotification.Reset();
-			WaitingForBrushesNotification.Reset();
-			InvalidShadingModelNotification.Reset();
-			WaitingForLandscapeResourcesStartTime = -1.0;
-			WaitingForLandscapeBrushResourcesStartTime = -1.0;
+			HideEditLayersResourcesNotification(EditLayersResourcesNotification);
+			WaitingForResourcesStartTime = -1.0;
 		}
 	};
 
@@ -7833,77 +7857,30 @@ void ALandscape::UpdateLayersContent(bool bInWaitForStreaming, bool bInSkipMonit
 	{
 		return;
 	}
-	
-	FLandscapeNotificationManager* LandscapeNotificationManager = GetWorld()->GetSubsystem<ULandscapeSubsystem>()->GetNotificationManager();
-	
-	// The Edit layers shaders only work on SM5 : cancel any update that might happen when SM5+ shading model is not active :
-	ERHIFeatureLevel::Type FeatureLevel = GetWorld() ? (ERHIFeatureLevel::Type)GetWorld()->FeatureLevel : GMaxRHIFeatureLevel;
-	if (FeatureLevel < ERHIFeatureLevel::SM5)
-	{
-		if (!InvalidShadingModelNotification.IsValid())
-		{
-			InvalidShadingModelNotification = MakeShared<FLandscapeNotification>(this, FLandscapeNotification::EType::ShadingModelInvalid);
-			static const FText NotificationText(LOCTEXT("InvalidShadingModel", "Landscape editor cannot update landscape with a feature level less than SM5!"));
-			InvalidShadingModelNotification->NotificationText = NotificationText;
-		}
-		LandscapeNotificationManager->RegisterNotification(InvalidShadingModelNotification);
-		bHideNotifications = false;
-		return;
-	}
-	else
-	{
-		InvalidShadingModelNotification.Reset();
-	}
 
 	// We need to wait until layers texture resources are ready to initialize the landscape to avoid taking the sizes and format of the default texture:
-	static constexpr double TimeBeforeDisplayingWaitingForResourcesNotification = 3.0;
-
 	bResourcesReady &= PrepareLayersTextureResources(bInWaitForStreaming);
+	// We need to wait until brushes texture/shader resources are ready to process the edit layers:
+	bResourcesReady &= PrepareLayersBrushResources(bInWaitForStreaming);
 	if (!bResourcesReady)
 	{
-		if (WaitingForLandscapeResourcesStartTime < 0.0)
+		if (FSlateApplicationBase::IsInitialized() && !bInWaitForStreaming)
 		{
-			WaitingForLandscapeResourcesStartTime = FSlateApplicationBase::Get().GetCurrentTime();
-			if (!WaitingForTexturesNotification.IsValid())
+			static constexpr double TimeBeforeDisplayingNotification = 3.0;
+			if (WaitingForResourcesStartTime < 0.0)
 			{
-				WaitingForTexturesNotification = MakeShared<FLandscapeNotification>(this, FLandscapeNotification::EType::LandscapeResourcesNotReady);
-				static const FText NotificationText(LOCTEXT("WaitForLanscapeResources", "Landscape editor waiting for landscape resources to be ready."));
-				WaitingForTexturesNotification->NotificationText = NotificationText;
-				WaitingForTexturesNotification->NotificationStartTime = WaitingForLandscapeResourcesStartTime + TimeBeforeDisplayingWaitingForResourcesNotification;
+				WaitingForResourcesStartTime = FSlateApplicationBase::Get().GetCurrentTime();
 			}
-			LandscapeNotificationManager->RegisterNotification(WaitingForTexturesNotification);
-			bHideNotifications = false;
-		}
-	}
-	else
-	{
-		WaitingForTexturesNotification.Reset();
-	}
 
-	bResourcesReady &= PrepareLayersBrushResources(bInWaitForStreaming);
-	if (!bResourcesReady)
-	{
-		if (WaitingForLandscapeBrushResourcesStartTime < 0.0)
-		{
-			WaitingForLandscapeBrushResourcesStartTime = FSlateApplicationBase::Get().GetCurrentTime();
-			if (!WaitingForBrushesNotification.IsValid())
+			if ((FSlateApplicationBase::Get().GetCurrentTime() - WaitingForResourcesStartTime) > TimeBeforeDisplayingNotification)
 			{
-				WaitingForBrushesNotification = MakeShared<FLandscapeNotification>(this, FLandscapeNotification::EType::LandscapeResourcesNotReady);
-				static const FText NotificationText(LOCTEXT("WaitForLanscapeResources", "Landscape editor waiting for landscape resources to be ready."));
-				WaitingForBrushesNotification->NotificationText = NotificationText;
-				WaitingForBrushesNotification->NotificationStartTime = WaitingForLandscapeResourcesStartTime + TimeBeforeDisplayingWaitingForResourcesNotification;
+				// let the user know we are waiting for resources :
+				static const FText NotificationText(LOCTEXT("WaitForLayersResources", "Landscape editor waiting for edit layers resources to be ready."));
+				ShowEditLayersResourcesNotification(NotificationText, EditLayersResourcesNotification);
 			}
-			LandscapeNotificationManager->RegisterNotification(WaitingForBrushesNotification);
-			bHideNotifications = false;
+			// The notification may not be visible yet (because of the initial delay) but it should not be hidden and the initial delay timer shouldn't be reset : 
+			bHideEditLayerResourcesNotification = false;
 		}
-	}
-	else
-	{
-		WaitingForBrushesNotification.Reset();
-	}
-
-	if (!bResourcesReady)
-	{
 		return;
 	}
 
diff --git a/Engine/Source/Runtime/Landscape/Private/LandscapeGrass.cpp b/Engine/Source/Runtime/Landscape/Private/LandscapeGrass.cpp
--- a/Engine/Source/Runtime/Landscape/Private/LandscapeGrass.cpp
+++ b/Engine/Source/Runtime/Landscape/Private/LandscapeGrass.cpp
@@ -2455,7 +2455,9 @@ int32 ALandscapeProxy::GetOutdatedGrassMapCount() const
 	return OutdatedGrassMaps;
 }
 
+int32 ALandscapeProxy::TotalComponentsNeedingGrassMapRender = 0;
 int32 ALandscapeProxy::TotalTexturesToStreamForVisibleGrassMapRender = 0;
+int32 ALandscapeProxy::TotalComponentsNeedingTextureBaking = 0;
 
 void ALandscapeProxy::UpdateGrassDataStatus(TSet<UTexture2D*>* OutCurrentForcedStreamedTextures, TSet<UTexture2D*>* OutDesiredForcedStreamedTextures, TSet<ULandscapeComponent*>* OutComponentsNeedingGrassMapRender, TSet<ULandscapeComponent*>* OutOutdatedComponents, bool bInEnableForceResidentFlag, int32* OutOutdatedGrassMaps) const
 {
@@ -2608,6 +2610,7 @@ void ALandscapeProxy::UpdateGrassData(bool bInShouldMarkDirty, FScopedSlowTask*
 	NumTexturesToStreamForVisibleGrassMapRender = 0;
 
 	// Remove local count from global count
+	TotalComponentsNeedingGrassMapRender -= NumComponentsNeedingGrassMapRender;
 	NumComponentsNeedingGrassMapRender = 0;
 
 	// Wait for Texture Streaming
@@ -3178,7 +3181,9 @@ void ALandscapeProxy::UpdateGrass(const TArray<FVector>& Cameras, int32& InOutNu
 					}
 				}
 
+				TotalComponentsNeedingGrassMapRender -= NumComponentsNeedingGrassMapRender;
 				NumComponentsNeedingGrassMapRender = ComponentsNeedingGrassMapRender.Num() - NumComponentsRendered - NumComponentsUnableToRender;
+				TotalComponentsNeedingGrassMapRender += NumComponentsNeedingGrassMapRender;
 
 				// Update resident flags
 				for (auto Texture : DesiredForceStreamedTextures.Difference(CurrentForcedStreamedTextures))
diff --git a/Engine/Source/Runtime/Landscape/Private/LandscapeNotification.cpp b/Engine/Source/Runtime/Landscape/Private/LandscapeNotification.cpp
deleted file mode 100644
diff --git a/Engine/Source/Runtime/Landscape/Private/LandscapeSubsystem.cpp b/Engine/Source/Runtime/Landscape/Private/LandscapeSubsystem.cpp
--- a/Engine/Source/Runtime/Landscape/Private/LandscapeSubsystem.cpp
+++ b/Engine/Source/Runtime/Landscape/Private/LandscapeSubsystem.cpp
@@ -14,7 +14,6 @@
 #include "ActorPartition/ActorPartitionSubsystem.h"
 #include "Engine/World.h"
 #include "Math/IntRect.h"
-#include "LandscapeNotification.h"
 #include "LandscapeConfigHelper.h"
 #include "Engine/Canvas.h"
 #include "EngineUtils.h"
@@ -55,17 +54,24 @@ void ULandscapeSubsystem::Initialize(FSubsystemCollectionBase& Collection)
 	GrassMapsBuilder = new FLandscapeGrassMapsBuilder(GetWorld());
 	GIBakedTextureBuilder = new FLandscapeGIBakedTextureBuilder(GetWorld());
 	PhysicalMaterialBuilder = new FLandscapePhysicalMaterialBuilder(GetWorld());
-	NotificationManager = new FLandscapeNotificationManager();
 #endif
 }
 
 void ULandscapeSubsystem::Deinitialize()
 {
 #if WITH_EDITOR
-	delete GrassMapsBuilder;
-	delete GIBakedTextureBuilder;
-	delete PhysicalMaterialBuilder;
-	delete NotificationManager;
+	if (GrassMapsBuilder)
+	{
+		delete GrassMapsBuilder;
+	}
+	if (GIBakedTextureBuilder)
+	{
+		delete GIBakedTextureBuilder;
+	}
+	if (PhysicalMaterialBuilder)
+	{
+		delete PhysicalMaterialBuilder;
+	}
 #endif
 	Proxies.Empty();
 
@@ -153,7 +159,6 @@ void ULandscapeSubsystem::Tick(float DeltaTime)
 	if (GIsEditor && !World->IsPlayInEditor())
 	{
 		LandscapePhysicalMaterial::GarbageCollectTasks();
-		NotificationManager->Tick();
 	}
 #endif
 }
diff --git a/Engine/Source/Runtime/Landscape/Public/LandscapeNotification.h b/Engine/Source/Runtime/Landscape/Public/LandscapeNotification.h
deleted file mode 100644
diff --git a/Engine/Source/Runtime/Landscape/Public/LandscapeSubsystem.h b/Engine/Source/Runtime/Landscape/Public/LandscapeSubsystem.h
--- a/Engine/Source/Runtime/Landscape/Public/LandscapeSubsystem.h
+++ b/Engine/Source/Runtime/Landscape/Public/LandscapeSubsystem.h
@@ -9,7 +9,6 @@
 
 class ALandscapeProxy;
 class ULandscapeInfo;
-class FLandscapeNotificationManager;
 
 UCLASS(MinimalAPI)
 class ULandscapeSubsystem : public UTickableWorldSubsystem
@@ -43,7 +42,6 @@ public:
 	LANDSCAPE_API void ChangeGridSize(ULandscapeInfo* LandscapeInfo, uint32 NewGridSizeInComponents);
 	LANDSCAPE_API ALandscapeProxy* FindOrAddLandscapeProxy(ULandscapeInfo* LandscapeInfo, const FIntPoint& SectionBase);
 	LANDSCAPE_API void DisplayBuildMessages(class FCanvas* Canvas, float& XPos, float& YPos);
-	FLandscapeNotificationManager* GetNotificationManager() { return NotificationManager; }
 #endif
 
 private:
@@ -58,7 +56,5 @@ private:
 	class FLandscapeGrassMapsBuilder* GrassMapsBuilder;
 	class FLandscapeGIBakedTextureBuilder* GIBakedTextureBuilder;
 	class FLandscapePhysicalMaterialBuilder* PhysicalMaterialBuilder;
-	
-	FLandscapeNotificationManager* NotificationManager;
 #endif
 };
