Landscape: Expose ALandscapeProxy::GetLandscapeActor to blueprints so that landscape streaming proxy can be used to find the owning landscape.

diff --git a/Engine/Source/Runtime/Landscape/Classes/LandscapeProxy.h b/Engine/Source/Runtime/Landscape/Classes/LandscapeProxy.h
--- a/Engine/Source/Runtime/Landscape/Classes/LandscapeProxy.h
+++ b/Engine/Source/Runtime/Landscape/Classes/LandscapeProxy.h
@@ -824,6 +824,8 @@ public:
 
 	virtual FGuid GetLandscapeGuid() const override { return LandscapeGuid; }
 	void SetLandscapeGuid(const FGuid& Guid) { LandscapeGuid = Guid; }
+
+	UFUNCTION(BlueprintCallable, Category = "Landscape|Runtime")
 	virtual ALandscape* GetLandscapeActor() PURE_VIRTUAL(GetLandscapeActor, return nullptr;)
 	virtual const ALandscape* GetLandscapeActor() const PURE_VIRTUAL(GetLandscapeActor, return nullptr;)
 
