PR #9028: Fix navmesh generation for Landscape Mesh Collision Components (Contributed by [REDACTED_PERSON])
ULandscapeMeshCollisionComponent does not SupportsGatheringGeometrySlices.

diff --git a/Engine/Source/Runtime/Landscape/Classes/LandscapeMeshCollisionComponent.h b/Engine/Source/Runtime/Landscape/Classes/LandscapeMeshCollisionComponent.h
--- a/Engine/Source/Runtime/Landscape/Classes/LandscapeMeshCollisionComponent.h
+++ b/Engine/Source/Runtime/Landscape/Classes/LandscapeMeshCollisionComponent.h
@@ -79,6 +79,10 @@ public:
 	virtual bool DoCustomNavigableGeometryExport(FNavigableGeometryExport& GeomExport) const override;
 	//End UPrimitiveComponent interface
 
+	//~ Begin INavRelevantInterface Interface
+	virtual bool SupportsGatheringGeometrySlices() const override { return false; }
+	//~ End INavRelevantInterface Interface
+
 	//~ Begin UObject Interface.
 	virtual void Serialize(FArchive& Ar) override;
 	virtual void BeginDestroy() override;
