Add flag for skipping the landscape layers update.

diff --git a/Engine/Source/Runtime/Landscape/Classes/Landscape.h b/Engine/Source/Runtime/Landscape/Classes/Landscape.h
--- a/Engine/Source/Runtime/Landscape/Classes/Landscape.h
+++ b/Engine/Source/Runtime/Landscape/Classes/Landscape.h
@@ -475,6 +475,9 @@ public:
 	/** Used to temporarily disable Grass Update in Editor */
 	bool bGrassUpdateEnabled;
 
+	UPROPERTY(Transient)
+	bool bEnableEditorLayersTick = false;
+
 	UPROPERTY()
 	TArray<FLandscapeLayer> LandscapeLayers;
 
diff --git a/Engine/Source/Runtime/Landscape/Private/LandscapeEditLayers.cpp b/Engine/Source/Runtime/Landscape/Private/LandscapeEditLayers.cpp
--- a/Engine/Source/Runtime/Landscape/Private/LandscapeEditLayers.cpp
+++ b/Engine/Source/Runtime/Landscape/Private/LandscapeEditLayers.cpp
@@ -8137,6 +8137,11 @@ void ALandscape::TickLayers(float DeltaTime)
 {
 	check(GIsEditor);
 
+	if (!bEnableEditorLayersTick)
+	{
+		return;
+	}
+
 	UWorld* World = GetWorld();
 	if (World && !World->IsPlayInEditor() && GetLandscapeInfo() && GEditor->PlayWorld == nullptr)
 	{
