Invert the default to enable landscape layers to update.

diff --git a/Engine/Source/Runtime/Landscape/Classes/Landscape.h b/Engine/Source/Runtime/Landscape/Classes/Landscape.h
--- a/Engine/Source/Runtime/Landscape/Classes/Landscape.h
+++ b/Engine/Source/Runtime/Landscape/Classes/Landscape.h
@@ -476,7 +476,7 @@ public:
 	bool bGrassUpdateEnabled;
 
 	UPROPERTY(Transient)
-	bool bEnableEditorLayersTick = false;
+	bool bEnableEditorLayersTick = true;
 
 	UPROPERTY()
 	TArray<FLandscapeLayer> LandscapeLayers;
