Fix for an assert in landscape when using local merge on D3D12 -- RDG does not yet support textures that are populated by a copy / only used as a shader resource.

diff --git a/Engine/Source/Runtime/Landscape/Private/LandscapeEditLayers.cpp b/Engine/Source/Runtime/Landscape/Private/LandscapeEditLayers.cpp
--- a/Engine/Source/Runtime/Landscape/Private/LandscapeEditLayers.cpp
+++ b/Engine/Source/Runtime/Landscape/Private/LandscapeEditLayers.cpp
@@ -6061,7 +6061,9 @@ namespace EditLayersWeightmapLocalMerge_RenderThread
 			// Allocate a texture array that can contain all edit layers textures to merge for any given component to render (this will be reused from one component to another) :
 			int32 SizeZ = InLocalMergeInfo.MaxNumEditLayersTexturesToMerge;
 			check(SizeZ > 0);
-			FRDGTextureDesc Desc = FRDGTextureDesc::Create2DArray(InLocalMergeInfo.ComponentSizeVerts, PF_B8G8R8A8, FClearValueBinding::None, TexCreate_ShaderResource, SizeZ, /*InNumMips = */1, /*InNumSamples = */1);
+
+			// TODO [chris.tchou] this texture does not have to be a render target, but RDG does not support transient shader-resource-only/copy-populated textures yet -- see issue https://jira.it.epicgames.com/browse/UE-162198
+			FRDGTextureDesc Desc = FRDGTextureDesc::Create2DArray(InLocalMergeInfo.ComponentSizeVerts, PF_B8G8R8A8, FClearValueBinding::None, TexCreate_ShaderResource | TexCreate_RenderTargetable, SizeZ, /*InNumMips = */1, /*InNumSamples = */1);
 			OutResources.EditLayersWeightmapsTextureArray = GraphBuilder.CreateTexture(Desc, TEXT("LandscapeEditLayersWeightmapsTextureArray"));
 			OutResources.EditLayersWeightmapsTextureArraySRV = GraphBuilder.CreateSRV(FRDGTextureSRVDesc::Create(OutResources.EditLayersWeightmapsTextureArray));
 		}
