[Navmesh] Fix navmesh generation not working for ULandscapeHeightfieldCollisionComponent landscape in Lazy gathering mode.
The sample cache was not reset in RecreateCollision().

diff --git a/Engine/Source/Runtime/Engine/Classes/AI/Navigation/NavigationTypes.h b/Engine/Source/Runtime/Engine/Classes/AI/Navigation/NavigationTypes.h
--- a/Engine/Source/Runtime/Engine/Classes/AI/Navigation/NavigationTypes.h
+++ b/Engine/Source/Runtime/Engine/Classes/AI/Navigation/NavigationTypes.h
@@ -765,5 +765,6 @@ struct ENGINE_API FNavHeightfieldSamples
 	FNavHeightfieldSamples();
 	void GetResourceSizeEx(FResourceSizeEx& CumulativeResourceSize);
 
-	FORCEINLINE bool IsEmpty() const { return Heights.Num() == 0; }
+	void Empty();
+	FORCEINLINE bool IsEmpty() const { return Heights.IsEmpty() && Holes.IsEmpty();  }
 };
diff --git a/Engine/Source/Runtime/Engine/Private/AI/Navigation/NavigationTypes.cpp b/Engine/Source/Runtime/Engine/Private/AI/Navigation/NavigationTypes.cpp
--- a/Engine/Source/Runtime/Engine/Private/AI/Navigation/NavigationTypes.cpp
+++ b/Engine/Source/Runtime/Engine/Private/AI/Navigation/NavigationTypes.cpp
@@ -300,6 +300,12 @@ void FNavHeightfieldSamples::GetResourceSizeEx(FResourceSizeEx& CumulativeResour
 	CumulativeResourceSize.AddDedicatedSystemMemoryBytes(sizeof(*this) + Heights.GetAllocatedSize() + Holes.GetAllocatedSize());
 }
 
+void FNavHeightfieldSamples::Empty()
+{
+	Heights.Empty();
+	Holes.Empty();
+}
+
 //----------------------------------------------------------------------//
 // FNavAgentProperties
 //----------------------------------------------------------------------//
diff --git a/Engine/Source/Runtime/Landscape/Private/LandscapeCollision.cpp b/Engine/Source/Runtime/Landscape/Private/LandscapeCollision.cpp
--- a/Engine/Source/Runtime/Landscape/Private/LandscapeCollision.cpp
+++ b/Engine/Source/Runtime/Landscape/Private/LandscapeCollision.cpp
@@ -1694,7 +1694,7 @@ bool ULandscapeHeightfieldCollisionComponent::RecreateCollision()
 		TRefCountPtr<FHeightfieldGeometryRef> HeightfieldRefLifetimeExtender = HeightfieldRef; // Ensure heightfield data is alive until removed from physics world
 		HeightfieldRef = nullptr; // Ensure data will be recreated
 		HeightfieldGuid = FGuid();
-
+		CachedHeightFieldSamples.Empty();
 		RecreatePhysicsState();
 
 		MarkRenderStateDirty();
@@ -1830,8 +1830,7 @@ bool ULandscapeMeshCollisionComponent::RecreateCollision()
 		TriMeshLifetimeExtender = MeshRef;
 		MeshRef = nullptr; // Ensure data will be recreated
 		MeshGuid = FGuid();
-		CachedHeightFieldSamples.Heights.Empty();
-		CachedHeightFieldSamples.Holes.Empty();
+		CachedHeightFieldSamples.Empty();
 	}
 
 	return Super::RecreateCollision();
