Landscape physical material render fixes:
* Added CVar for performing render captures
* Landscape physical material async tasks now trigger the compilation of their needed shaders and asynchronously wait for them in order not to render landscape physical material indices with a fallback material, which is completely incorrect. That was leading to invalid landscape physical material export during the first run because of on-demand shader compilation.
* Cleanup physical material data in landscape components when a material that uses Physical Material Output stops using them

diff --git a/Engine/Source/Runtime/Landscape/Classes/LandscapeComponent.h b/Engine/Source/Runtime/Landscape/Classes/LandscapeComponent.h
--- a/Engine/Source/Runtime/Landscape/Classes/LandscapeComponent.h
+++ b/Engine/Source/Runtime/Landscape/Classes/LandscapeComponent.h
@@ -629,7 +629,13 @@ public:
 	/** Physical material update task */
 	FLandscapePhysicalMaterialRenderTask PhysicalMaterialTask;
 	uint32 CalculatePhysicalMaterialTaskHash() const;
-#endif
+
+	/**
+	 * Get the physical materials that are configured by the landscape component graphical material.
+	 * Returns false if there are no non-null physical materials. (We probably don't want to use if no physical material connections are bound.)
+	 */
+	bool GetRenderPhysicalMaterials(TArray<UPhysicalMaterial*>& OutPhysicalMaterials) const;
+#endif // WITH_EDITOR
 
 	//~ Begin UObject Interface.	
 	virtual void PostInitProperties() override;	
diff --git a/Engine/Source/Runtime/Landscape/Classes/LandscapePhysicalMaterial.h b/Engine/Source/Runtime/Landscape/Classes/LandscapePhysicalMaterial.h
--- a/Engine/Source/Runtime/Landscape/Classes/LandscapePhysicalMaterial.h
+++ b/Engine/Source/Runtime/Landscape/Classes/LandscapePhysicalMaterial.h
@@ -26,8 +26,8 @@ public:
 		Release();
 	}
 
-	/** Initialize the task for a component. A task in progress can safely call Init() again to rekick the task. */
-	void Init(class ULandscapeComponent const* LandscapeComponent);
+	/** Initialize the task for a component. A task in progress can safely call Init() again to rekick the task. Returns false if no physical material needs to be rendered. */
+	bool Init(class ULandscapeComponent const* LandscapeComponent);
 	/** Release the task. After calling this IsValid() will return false until we call Init() again. */
 	void Release();
 
@@ -46,6 +46,9 @@ public:
 	/** Get the result data. Assumes that IsComplete() is true. This doesn't Release() the data. */
 	TArray <class UPhysicalMaterial* > const& GetResultMaterials() const;
 
+private:
+	void UpdateInternal(bool bInFlush);
+
 private:
 	int32 PoolHandle;
 	
diff --git a/Engine/Source/Runtime/Landscape/Private/LandscapeEdit.cpp b/Engine/Source/Runtime/Landscape/Private/LandscapeEdit.cpp
--- a/Engine/Source/Runtime/Landscape/Private/LandscapeEdit.cpp
+++ b/Engine/Source/Runtime/Landscape/Private/LandscapeEdit.cpp
@@ -26,6 +26,7 @@ LandscapeEdit.cpp: Landscape editing
 #include "Materials/MaterialExpressionLandscapeLayerSample.h"
 #include "Materials/MaterialExpressionLandscapeLayerBlend.h"
 #include "Materials/MaterialExpressionLandscapeLayerSwitch.h"
+#include "Materials/MaterialExpressionLandscapePhysicalMaterialOutput.h"
 #include "LandscapeDataAccess.h"
 #include "LandscapeRender.h"
 #include "LandscapePrivate.h"
@@ -56,7 +57,6 @@ LandscapeEdit.cpp: Landscape editing
 #include "NaniteBuilder.h"
 #include "Rendering/NaniteResources.h"
 #include "Misc/ScopedSlowTask.h"
-
 #include "EngineModule.h"
 #include "EngineUtils.h"
 #include "Framework/Notifications/NotificationManager.h"
@@ -95,6 +95,11 @@ static TAutoConsoleVariable<int32> CVarMobileCompressLanscapeWeightMaps(
     TEXT("Whether to compress the terrain weight maps for mobile."),
 	ECVF_ReadOnly);
 
+static TAutoConsoleVariable<int32> CVarLandscapeApplyPhysicalMaterialChangesImmediately(
+    TEXT("landscape.ApplyPhysicalMaterialChangesImmediately"),
+	1,
+    TEXT("Applies physical material task changes immediately rather than during the next cook/PIE."));
+
 #if WITH_EDITOR
 
 // Used to temporarily disable material instance updates (typically used for cases where multiple updates are called on sample component)
@@ -1674,12 +1679,50 @@ uint32 ULandscapeComponent::CalculatePhysicalMaterialTaskHash() const
 	return Hash;
 }
 
+bool ULandscapeComponent::GetRenderPhysicalMaterials(TArray<UPhysicalMaterial*>& OutPhysicalMaterials) const 
+{
+	bool bReturnValue = false;
+	OutPhysicalMaterials.Reset();
+
+	if (UMaterialInterface* Material = GetLandscapeMaterial())
+	{
+		ERHIFeatureLevel::Type FeatureLevel = GMaxRHIFeatureLevel;
+		{
+			TArray<const UMaterialExpressionLandscapePhysicalMaterialOutput*> Expressions;
+			Material->GetMaterial()->GetAllExpressionsOfType<UMaterialExpressionLandscapePhysicalMaterialOutput>(Expressions);
+			if (Expressions.Num() > 0)
+			{
+				// Assume only one valid physical material output material node
+				for (const FPhysicalMaterialInput& Input : Expressions[0]->Inputs)
+				{
+					OutPhysicalMaterials.Add(Input.PhysicalMaterial);
+					bReturnValue |= (Input.PhysicalMaterial != nullptr);
+				}
+			}
+		}
+	}
+
+	return bReturnValue;
+}
+
 void ULandscapeComponent::UpdatePhysicalMaterialTasks()
 {
 	uint32 Hash = CalculatePhysicalMaterialTaskHash();
 	if (PhysicalMaterialHash != Hash)
 	{
-		PhysicalMaterialTask.Init(this);
+		TArray<UPhysicalMaterial*> PhysicalMaterials;
+		if (GetRenderPhysicalMaterials(PhysicalMaterials))
+		{
+			bool bSuccess = PhysicalMaterialTask.Init(this);
+			check(bSuccess && PhysicalMaterialTask.IsValid());
+		}
+		else
+		{
+			// Clear the renderable physical material properties as we don't need them :
+			CollisionComponent->PhysicalMaterialRenderObjects.Reset();
+			CollisionComponent->PhysicalMaterialRenderData.RemoveBulkData();
+			PhysicalMaterialTask.Release();
+		}
 		PhysicalMaterialHash = Hash;
 	}
 
@@ -1691,8 +1734,12 @@ void ULandscapeComponent::UpdatePhysicalMaterialTasks()
 
 			PhysicalMaterialTask.Release();
 
-			// We do not force an update of the physics data here. 
-			// We don't need the information immediately in the editor and update will happen on cook or PIE.
+			// Potentially, we do not force an update of the physics data here (behind a CVar, as we don't necessarily need the 
+			//  information immediately in the editor and update will happen on cook or PIE) :
+			if (CVarLandscapeApplyPhysicalMaterialChangesImmediately.GetValueOnGameThread() != 0)
+			{
+				CollisionComponent->RecreateCollision();
+			}
 		}
 		else
 		{
@@ -1703,6 +1750,8 @@ void ULandscapeComponent::UpdatePhysicalMaterialTasks()
 
 void ULandscapeComponent::UpdateCollisionPhysicalMaterialData(TArray<UPhysicalMaterial*> const& InPhysicalMaterials, TArray<uint8> const& InMaterialIds)
 {
+	TRACE_CPUPROFILER_EVENT_SCOPE(ULandscapeComponent::UpdateCollisionPhysicalMaterialData);
+
 	// Copy the physical material array
 	CollisionComponent->PhysicalMaterialRenderObjects = InPhysicalMaterials;
 
diff --git a/Engine/Source/Runtime/Landscape/Private/LandscapePhysicalMaterial.cpp b/Engine/Source/Runtime/Landscape/Private/LandscapePhysicalMaterial.cpp
--- a/Engine/Source/Runtime/Landscape/Private/LandscapePhysicalMaterial.cpp
+++ b/Engine/Source/Runtime/Landscape/Private/LandscapePhysicalMaterial.cpp
@@ -8,8 +8,6 @@
 #include "LandscapeComponent.h"
 #include "LandscapePrivate.h"
 #include "LandscapeRender.h"
-#include "Materials/MaterialExpressionLandscapePhysicalMaterialOutput.h"
-#include "Materials/MaterialInstanceConstant.h"
 #include "MeshMaterialShader.h"
 #include "MeshPassProcessor.h"
 #include "MeshPassProcessor.inl"
@@ -20,55 +18,16 @@
 #include "SimpleMeshDrawCommandPass.h"
 #include "RHIGPUReadback.h"
 #include "RenderGraphUtils.h"
+#include "RenderCaptureInterface.h"
 
 DECLARE_GPU_STAT_NAMED(LandscapePhysicalMaterial_Draw, TEXT("LandscapePhysicalMaterial"));
 
-
-namespace
-{
-	/** Get the landscape material used by the landscape component. */
-	UMaterialInterface* GetLandscapeMaterial(ULandscapeComponent const* InLandscapeComponent)
-	{
-		UMaterialInterface* Material = InLandscapeComponent->GetLandscapeMaterial();
-		for (UMaterialInstanceConstant* MIC = Cast<UMaterialInstanceConstant>(Material); MIC; MIC = Cast<UMaterialInstanceConstant>(Material))
-		{
-			Material = MIC->Parent;
-		}
-		return Material;
-	}
-
-	/** 
-	 * Get the physical materials that are configured by the landscape component graphical material. 
-	 * Returns false if there are no non-null physical materials. (We probably don't want to use if no physical material connections are bound.)
-	 */
-	bool GetPhysicalMaterials(ULandscapeComponent const* InLandscapeComponent, TArray<UPhysicalMaterial*>& OutPhysicalMaterials)
-	{
-		bool bReturnValue = false;
-		OutPhysicalMaterials.Reset();
-
-		UMaterialInterface* Material = GetLandscapeMaterial(InLandscapeComponent);
-		if (Material != nullptr)
-		{
-			ERHIFeatureLevel::Type FeatureLevel = GMaxRHIFeatureLevel;
-			{
-				TArray<const UMaterialExpressionLandscapePhysicalMaterialOutput*> Expressions;
-				Material->GetMaterial()->GetAllExpressionsOfType<UMaterialExpressionLandscapePhysicalMaterialOutput>(Expressions);
-				if (Expressions.Num() > 0)
-				{
-					// Assume only one valid physical material output material node
-					for (const FPhysicalMaterialInput& Input : Expressions[0]->Inputs)
-					{
-						OutPhysicalMaterials.Add(Input.PhysicalMaterial);
-						bReturnValue |= (Input.PhysicalMaterial != nullptr);
-					}
-				}
-			}
-		}
-
-		return bReturnValue;
-	}
-}
-
+int32 RenderCaptureLayersNextPhysicalMaterialDraws = 0;
+static FAutoConsoleVariableRef CVarRenderCaptureLayersNextPhysicalMaterialDraws(
+	TEXT("landscape.RenderCaptureLayersNextPhysicalMaterialDraws"),
+	RenderCaptureLayersNextPhysicalMaterialDraws,
+	TEXT("Trigger N render captures during the next landscape physical material draw calls."),
+	EConsoleVariableFlags::ECVF_RenderThreadSafe);
 
 /** Material shader for rendering physical material IDs. */
 class FLandscapePhysicalMaterial : public FMeshMaterialShader
@@ -122,6 +81,19 @@ public:
 
 IMPLEMENT_MATERIAL_SHADER_TYPE(, FLandscapePhysicalMaterialPS, TEXT("/Engine/Private/LandscapePhysicalMaterial.usf"), TEXT("PSMain"), SF_Pixel);
 
+// Tries to retrieve the VS/PS for rendering landscape physical materials. If the shaders are not available yet (e.g. not compiled), it will return false and will request them to become available . 
+//  This can be called repeatedly until the shaders become available :
+bool TryGetLandscapePhysicalMaterialShaders(const FMaterial& InMaterialResource, FMaterialShaders& OutShaders)
+{
+	FMaterialShaderTypes ShaderTypes;
+	ShaderTypes.AddShaderType<FLandscapePhysicalMaterialVS>();
+	ShaderTypes.AddShaderType<FLandscapePhysicalMaterialPS>();
+	
+	FVertexFactoryType* VertexFactoryType = FindVertexFactoryType(FName(TEXT("FLandscapeFixedGridVertexFactory"), FNAME_Find));
+	check(VertexFactoryType != nullptr);
+	return InMaterialResource.TryGetShaders(ShaderTypes, VertexFactoryType, OutShaders);
+}
+
 
 /** Simple mesh processor implementation to draw using the FLandscapePhysicalMaterial mesh material shader. */
 class FLandscapePhysicalMaterialMeshProcessor : public FMeshPassProcessor
@@ -139,21 +111,6 @@ public:
 		int32 StaticMeshId = -1) override final;
 
 private:
-	bool TryAddMeshBatch(
-		const FMeshBatch& RESTRICT MeshBatch,
-		uint64 BatchElementMask,
-		const FPrimitiveSceneProxy* RESTRICT PrimitiveSceneProxy,
-		int32 StaticMeshId,
-		const FMaterialRenderProxy& MaterialRenderProxy,
-		const FMaterial& MaterialResource);
-
-	bool Process(
-		const FMeshBatch& MeshBatch,
-		uint64 BatchElementMask,
-		const FPrimitiveSceneProxy* RESTRICT PrimitiveSceneProxy,
-		const FMaterialRenderProxy& RESTRICT MaterialRenderProxy,
-		const FMaterial& RESTRICT MaterialResource);
-
 	FMeshPassProcessorRenderState PassDrawRenderState;
 };
 
@@ -174,60 +131,20 @@ void FLandscapePhysicalMaterialMeshProcessor::AddMeshBatch(
 	int32 StaticMeshId)
 {
 	const FMaterialRenderProxy* MaterialRenderProxy = MeshBatch.MaterialRenderProxy;
-	while (MaterialRenderProxy)
-	{
-		const FMaterial* Material = MaterialRenderProxy->GetMaterialNoFallback(FeatureLevel);
-		if (Material)
-		{
-			if (TryAddMeshBatch(MeshBatch, BatchElementMask, PrimitiveSceneProxy, StaticMeshId, *MaterialRenderProxy, *Material))
-			{
-				break;
-			}
-		}
-
-		MaterialRenderProxy = MaterialRenderProxy->GetFallback(FeatureLevel);
-	}
-}
+	const FMaterial* Material = MaterialRenderProxy->GetMaterialNoFallback(FeatureLevel);
+	checkf(Material != nullptr, TEXT("InitTaskRenderResources should ensure that by the time we render the physical materials, the shaders are valid and ready."));
 
-bool FLandscapePhysicalMaterialMeshProcessor::TryAddMeshBatch(
-	const FMeshBatch& RESTRICT MeshBatch,
-	uint64 BatchElementMask,
-	const FPrimitiveSceneProxy* RESTRICT PrimitiveSceneProxy,
-	int32 StaticMeshId,
-	const FMaterialRenderProxy& MaterialRenderProxy,
-	const FMaterial& MaterialResource)
-{
-	return Process(MeshBatch, BatchElementMask, PrimitiveSceneProxy, MaterialRenderProxy, MaterialResource);
-}
-
-bool FLandscapePhysicalMaterialMeshProcessor::Process(
-	const FMeshBatch& MeshBatch,
-	uint64 BatchElementMask,
-	const FPrimitiveSceneProxy* RESTRICT PrimitiveSceneProxy,
-	const FMaterialRenderProxy& RESTRICT MaterialRenderProxy,
-	const FMaterial& RESTRICT MaterialResource)
-{
-	const FVertexFactory* VertexFactory = MeshBatch.VertexFactory;
+	FMaterialShaders Shaders;
+	verify(TryGetLandscapePhysicalMaterialShaders(*Material, Shaders));
 
 	TMeshProcessorShaders<
 		FLandscapePhysicalMaterialVS,
 		FLandscapePhysicalMaterialPS> PassShaders;
-
-	FMaterialShaderTypes ShaderTypes;
-	ShaderTypes.AddShaderType<FLandscapePhysicalMaterialVS>();
-	ShaderTypes.AddShaderType<FLandscapePhysicalMaterialPS>();
-
-	FMaterialShaders Shaders;
-	if (!MaterialResource.TryGetShaders(ShaderTypes, VertexFactory->GetType(), Shaders))
-	{
-		return false;
-	}
-
 	Shaders.TryGetVertexShader(PassShaders.VertexShader);
 	Shaders.TryGetPixelShader(PassShaders.PixelShader);
 
 	const FMeshDrawingPolicyOverrideSettings OverrideSettings = ComputeMeshOverrideSettings(MeshBatch);
-	const ERasterizerFillMode MeshFillMode = ComputeMeshFillMode(MeshBatch, MaterialResource, OverrideSettings);
+	const ERasterizerFillMode MeshFillMode = ComputeMeshFillMode(MeshBatch, *Material, OverrideSettings);
 	const ERasterizerCullMode MeshCullMode = CM_None;
 
 	FMeshMaterialShaderElementData ShaderElementData;
@@ -239,8 +156,8 @@ bool FLandscapePhysicalMaterialMeshProcessor::Process(
 		MeshBatch,
 		BatchElementMask,
 		PrimitiveSceneProxy,
-		MaterialRenderProxy,
-		MaterialResource,
+		*MaterialRenderProxy,
+		*Material,
 		PassDrawRenderState,
 		PassShaders,
 		MeshFillMode,
@@ -248,11 +165,8 @@ bool FLandscapePhysicalMaterialMeshProcessor::Process(
 		SortKey,
 		EMeshPassFeatures::Default,
 		ShaderElementData);
-
-	return true;
 }
 
-
 namespace
 {
 	/** A description of a mesh to render. */
@@ -269,11 +183,11 @@ namespace
 	 * Initially we only collect the base landscape mesh, but potentially we could gather other objects like roads?
 	 * WARNING: This gets the SceneProxy pointer from the UComponent on the render thread. This doesn't feel safe but it's what the grass renderer does...
 	 */
-	void FillMeshInfos_RenderThread(const FPrimitiveSceneProxy* InSceneProxy, FMeshInfoArray& OutMeshInfos)
+	void FillMeshInfos_RenderThread(const FLandscapeComponentSceneProxy* InSceneProxy, FMeshInfoArray& OutMeshInfos)
 	{
 		const int32 MeshInfoIndex = OutMeshInfos.AddUninitialized();
 		OutMeshInfos[MeshInfoIndex].Proxy = InSceneProxy;
-		OutMeshInfos[MeshInfoIndex].MeshBatch = &((FLandscapeComponentSceneProxy*)InSceneProxy)->GetGrassMeshBatch();
+		OutMeshInfos[MeshInfoIndex].MeshBatch = &(InSceneProxy->GetGrassMeshBatch());
 		OutMeshInfos[MeshInfoIndex].MeshBatchElementMask = 1 << 0; // LOD 0 only
 	}
 	
@@ -296,51 +210,56 @@ namespace
 	{
 		FRDGBuilder GraphBuilder(RHICmdList);
 
-		// Create the view
-		FSceneViewFamily::ConstructionValues ViewFamilyInit(nullptr, SceneInterface, FEngineShowFlags(ESFIM_Game));
-		ViewFamilyInit.SetTime(FGameTime());
-		FSceneViewFamilyContext ViewFamily(ViewFamilyInit);
-		ViewFamily.LandscapeLODOverride = 0; // Force LOD 0 render
+		{
+			RenderCaptureInterface::FScopedCapture RenderCapture((RenderCaptureLayersNextPhysicalMaterialDraws != 0), GraphBuilder, TEXT("LandscapePhysicalMaterialCapture"));
+			RenderCaptureLayersNextPhysicalMaterialDraws = FMath::Max(0, RenderCaptureLayersNextPhysicalMaterialDraws - 1);
+
+			// Create the view
+			FSceneViewFamily::ConstructionValues ViewFamilyInit(nullptr, SceneInterface, FEngineShowFlags(ESFIM_Game));
+			ViewFamilyInit.SetTime(FGameTime());
+			FSceneViewFamilyContext ViewFamily(ViewFamilyInit);
+			ViewFamily.LandscapeLODOverride = 0; // Force LOD 0 render
 
-		FScenePrimitiveRenderingContextScopeHelper ScenePrimitiveRenderingContextScopeHelper(GetRendererModule().BeginScenePrimitiveRendering(GraphBuilder, &ViewFamily));
+			FScenePrimitiveRenderingContextScopeHelper ScenePrimitiveRenderingContextScopeHelper(GetRendererModule().BeginScenePrimitiveRendering(GraphBuilder, &ViewFamily));
 
-		FSceneViewInitOptions ViewInitOptions;
-		ViewInitOptions.SetViewRectangle(FIntRect(0, 0, TargetSize.X, TargetSize.Y));
-		ViewInitOptions.ViewOrigin = ViewOrigin;
-		ViewInitOptions.ViewRotationMatrix = ViewRotationMatrix;
-		ViewInitOptions.ProjectionMatrix = ProjectionMatrix;
-		ViewInitOptions.ViewFamily = &ViewFamily;
+			FSceneViewInitOptions ViewInitOptions;
+			ViewInitOptions.SetViewRectangle(FIntRect(0, 0, TargetSize.X, TargetSize.Y));
+			ViewInitOptions.ViewOrigin = ViewOrigin;
+			ViewInitOptions.ViewRotationMatrix = ViewRotationMatrix;
+			ViewInitOptions.ProjectionMatrix = ProjectionMatrix;
+			ViewInitOptions.ViewFamily = &ViewFamily;
 
-		GetRendererModule().CreateAndInitSingleView(RHICmdList, &ViewFamily, &ViewInitOptions);
-		const FSceneView* View = ViewFamily.Views[0];
+			GetRendererModule().CreateAndInitSingleView(RHICmdList, &ViewFamily, &ViewInitOptions);
+			const FSceneView* View = ViewFamily.Views[0];
 
-		FRDGTextureRef OutputTexture = GraphBuilder.CreateTexture(
-			FRDGTextureDesc::Create2D(TargetSize, PF_G8, FClearValueBinding::Black, TexCreate_RenderTargetable | TexCreate_ShaderResource),
-			TEXT("LandscapePhysicalMaterialTarget"));
+			FRDGTextureRef OutputTexture = GraphBuilder.CreateTexture(
+				FRDGTextureDesc::Create2D(TargetSize, PF_G8, FClearValueBinding::Black, TexCreate_RenderTargetable | TexCreate_ShaderResource),
+				TEXT("LandscapePhysicalMaterialTarget"));
 
-		auto* PassParameters = GraphBuilder.AllocParameters<FLandscapePhysicalMaterialPassParameters>();
-		PassParameters->View = View->ViewUniformBuffer;
-		PassParameters->RenderTargets[0] = FRenderTargetBinding(OutputTexture, ERenderTargetLoadAction::EClear);
+			auto* PassParameters = GraphBuilder.AllocParameters<FLandscapePhysicalMaterialPassParameters>();
+			PassParameters->View = View->ViewUniformBuffer;
+			PassParameters->RenderTargets[0] = FRenderTargetBinding(OutputTexture, ERenderTargetLoadAction::EClear);
 
-		AddSimpleMeshPass(GraphBuilder, PassParameters, SceneInterface->GetRenderScene(), *View, nullptr, RDG_EVENT_NAME("LandscapePhysicalMaterial"), View->UnscaledViewRect,
-			[View, &MeshInfos](FDynamicPassMeshDrawListContext* DynamicMeshPassContext)
-			{
-				FLandscapePhysicalMaterialMeshProcessor PassMeshProcessor(
-					nullptr,
-					View,
-					DynamicMeshPassContext);
-				for (auto& MeshInfo : MeshInfos)
+			AddSimpleMeshPass(GraphBuilder, PassParameters, SceneInterface->GetRenderScene(), *View, nullptr, RDG_EVENT_NAME("LandscapePhysicalMaterial"), View->UnscaledViewRect,
+				[View, &MeshInfos](FDynamicPassMeshDrawListContext* DynamicMeshPassContext)
 				{
-					const FMeshBatch& Mesh = *MeshInfo.MeshBatch;
-					if (Mesh.MaterialRenderProxy != nullptr)
+					FLandscapePhysicalMaterialMeshProcessor PassMeshProcessor(
+						nullptr,
+						View,
+						DynamicMeshPassContext);
+					for (auto& MeshInfo : MeshInfos)
 					{
-						Mesh.MaterialRenderProxy->UpdateUniformExpressionCacheIfNeeded(View->GetFeatureLevel());
-						PassMeshProcessor.AddMeshBatch(Mesh, MeshInfo.MeshBatchElementMask, MeshInfo.Proxy);
+						const FMeshBatch& Mesh = *MeshInfo.MeshBatch;
+						if (Mesh.MaterialRenderProxy != nullptr)
+						{
+							Mesh.MaterialRenderProxy->UpdateUniformExpressionCacheIfNeeded(View->GetFeatureLevel());
+							PassMeshProcessor.AddMeshBatch(Mesh, MeshInfo.MeshBatchElementMask, MeshInfo.Proxy);
+						}
 					}
-				}
-			});
+				});
 
-		AddEnqueueCopyPass(GraphBuilder, Readback, OutputTexture);
+			AddEnqueueCopyPass(GraphBuilder, Readback, OutputTexture);
+		}
 
 		GraphBuilder.Execute();
 	}
@@ -387,7 +306,7 @@ namespace
 	struct FLandscapePhysicalMaterialRenderTaskImpl
 	{
 		// Create on game thread
-		ULandscapeComponent const* LandscapeComponent = nullptr;
+		ULandscapeComponent const* LandscapeComponent = nullptr; // !Do not access on render thread!
 		uint32 InitFrameId = 0;
 		FIntPoint TargetSize = FIntPoint(ForceInitToZero);
 		FVector ViewOrigin = FVector(ForceInitToZero);
@@ -405,12 +324,12 @@ namespace
 		TArray<uint8> ResultIds;
 	};
 
-	/** Initialize the physical material render task data. */
+	/** Initialize the physical material render task data. Return false if there are no physical materials needing to be rendered by the graphical material. */
 	bool InitTask(FLandscapePhysicalMaterialRenderTaskImpl& Task, ULandscapeComponent const* InLandscapeComponent, uint32 InFrameId)
 	{
 		if (InLandscapeComponent != nullptr)
 		{
-			if (GetPhysicalMaterials(InLandscapeComponent, Task.ResultMaterials))
+			if (InLandscapeComponent->GetRenderPhysicalMaterials(Task.ResultMaterials))
 			{
 				Task.LandscapeComponent = InLandscapeComponent;
 				Task.InitFrameId = InFrameId;
@@ -443,36 +362,46 @@ namespace
 		return false;
 	}
 
-	/** Initialize the physical material render task read back resources. */
-	bool InitTaskRenderResources(FLandscapePhysicalMaterialRenderTaskImpl& Task)
+	bool InitTaskRenderResources(FLandscapePhysicalMaterialRenderTaskImpl& InTask, const FLandscapeComponentSceneProxy* InSceneProxy)
 	{
+		/** Initialize the physical material render task read back resources. */
 		//todo: Consider pooling these and throttling to the pool size?
-		if (!Task.Readback.IsValid())
+		if (!InTask.Readback.IsValid())
 		{
-			Task.Readback = MakeUnique<FRHIGPUTextureReadback>(TEXT("LandscapePhysicalMaterialReadback"));
+			InTask.Readback = MakeUnique<FRHIGPUTextureReadback>(TEXT("LandscapePhysicalMaterialReadback"));
 		}
 
-		return true;
+		// Make sure the landscape physical material shader is compiled and if not, initiate compilation : 
+		const FMeshBatch& GrassMeshBatch = InSceneProxy->GetGrassMeshBatch();
+		check(GrassMeshBatch.MaterialRenderProxy != nullptr);
+
+		ERHIFeatureLevel::Type FeatureLevel = InSceneProxy->GetScene().GetFeatureLevel();
+		const FMaterial* Material = GrassMeshBatch.MaterialRenderProxy->GetMaterialNoFallback(FeatureLevel);
+		if (Material == nullptr)
+		{
+			return false;
+		}
+
+		// Only return true when the shaders needed for rendering landscape physical materials are ready : 
+		FMaterialShaders Shaders;
+		return TryGetLandscapePhysicalMaterialShaders(*Material, Shaders);
 	}
 
 	/** Update the physical material render task on the render thread. */
-	void UpdateTask_RenderThread(FRHICommandListImmediate& RHICmdList, FLandscapePhysicalMaterialRenderTaskImpl& Task, bool bFlush)
+	void UpdateTask_RenderThread(FRHICommandListImmediate& RHICmdList, FLandscapePhysicalMaterialRenderTaskImpl& Task, const FLandscapeComponentSceneProxy* InSceneProxy, bool bFlush)
 	{
-		// WARNING: We access the UComponent to get in SceneProxy for FillMeshInfos_RenderThread(). 
-		// This isn't good style but probably works since the UComponent owns the update task and is guaranteed to be valid.
-		const FPrimitiveSceneProxy* SceneProxy = Task.LandscapeComponent->SceneProxy;
-		if (Task.CompletionState == ECompletionState::None && SceneProxy)
+		if (Task.CompletionState == ECompletionState::None && InSceneProxy)
 		{
-			// Allocate read back resources
-			if (InitTaskRenderResources(Task))
+			// Make sure everything is ready to render, if not (e.g. shader not yet compiled), wait until it is:
+			if (InitTaskRenderResources(Task, InSceneProxy))
 			{
 				// Render the pending item.
 				FMeshInfoArray MeshInfos;
-				FillMeshInfos_RenderThread(SceneProxy, MeshInfos);
+				FillMeshInfos_RenderThread(InSceneProxy, MeshInfos);
 
 				Render_RenderThread(
 					RHICmdList,
-					&SceneProxy->GetScene(),
+					&InSceneProxy->GetScene(),
 					MeshInfos,
 					Task.TargetSize,
 					Task.ViewOrigin,
@@ -514,8 +443,8 @@ public:
 	/** Frame count used to validate and garbage collect. */
 	uint32 FrameCount = 0;
 
-	/** Allocate task data from the pool. */
-	void Allocate(FLandscapePhysicalMaterialRenderTask& InTask, ULandscapeComponent const* InLandscapeComponent)
+	/** Allocate task data from the pool. Returns false if no physical material needs to be rendered. */
+	bool Allocate(FLandscapePhysicalMaterialRenderTask& InTask, ULandscapeComponent const* InLandscapeComponent)
 	{
 		check(InTask.PoolHandle == -1);
 
@@ -536,6 +465,8 @@ public:
 
 		const bool bSuccess = InitTask(Pool[Index], InLandscapeComponent, FrameCount);
 		InTask.PoolHandle = bSuccess ? Index : -1;
+
+		return bSuccess;
 	}
 
 	/** Return task data to the pool. */
@@ -600,14 +531,14 @@ public:
 static TGlobalResource< FLandscapePhysicalMaterialRenderTaskPool > GTaskPool;
 
 
-void FLandscapePhysicalMaterialRenderTask::Init(ULandscapeComponent const* LandscapeComponent)
+bool FLandscapePhysicalMaterialRenderTask::Init(ULandscapeComponent const* LandscapeComponent)
 {
 	check(IsInGameThread());
 	if (IsValid())
 	{
 		GTaskPool.Free(*this);
 	}
-	GTaskPool.Allocate(*this, LandscapeComponent);
+	return GTaskPool.Allocate(*this, LandscapeComponent);
 }
 
 void FLandscapePhysicalMaterialRenderTask::Release()
@@ -634,37 +565,36 @@ bool FLandscapePhysicalMaterialRenderTask::IsComplete() const
 
 void FLandscapePhysicalMaterialRenderTask::Tick()
 {
-	check(IsInGameThread());
-	if (IsValid() && !IsComplete())
-	{
-		FLandscapePhysicalMaterialRenderTaskImpl* Task = &GTaskPool.Pool[PoolHandle];
-
-		ENQUEUE_RENDER_COMMAND(FLandscapePhysicalMaterialUpdaterTick)(
-			[Task](FRHICommandListImmediate& RHICmdList)
-			{
-				check(Task->LandscapeComponent);
-				UpdateTask_RenderThread(RHICmdList, *Task, false);
-			});
-	}
+	UpdateInternal(/*bInFlush = */false);
 }
 
 // Note: We could add a global function that calls Flush() on multiple tasks. 
 // That could enqueue a single render thread function and use a single rendering command flush.
 // It could be useful if we see performance issue with any path that calls Flush individually for each landscape component.
 void FLandscapePhysicalMaterialRenderTask::Flush()
+{
+	UpdateInternal(/*bInFlush = */true);
+}
+
+void FLandscapePhysicalMaterialRenderTask::UpdateInternal(bool bInFlush)
 {
 	check(IsInGameThread());
 	if (IsValid() && !IsComplete())
 	{
 		FLandscapePhysicalMaterialRenderTaskImpl* Task = &GTaskPool.Pool[PoolHandle];
-
+		check(Task->LandscapeComponent != nullptr);
+		// If the material gets recompiled, for example, the component's scene proxy could be changed before the RT update runs, so we need to make sure we're getting the scene proxy on the game thread :
+		const FLandscapeComponentSceneProxy* SceneProxy = static_cast<const FLandscapeComponentSceneProxy*>(Task->LandscapeComponent->SceneProxy);
 		ENQUEUE_RENDER_COMMAND(FLandscapePhysicalMaterialFlush)(
-			[Task](FRHICommandListImmediate& RHICmdList)
-			{
-				UpdateTask_RenderThread(RHICmdList, *Task, true);
-			});
+			[Task, SceneProxy, bInFlush](FRHICommandListImmediate& RHICmdList)
+		{
+			UpdateTask_RenderThread(RHICmdList, *Task, SceneProxy, bInFlush);
+		});
 
-		FlushRenderingCommands();
+		if (bInFlush)
+		{
+			FlushRenderingCommands();
+		}
 	}
 }
 
