Separated arm/x64 DDC keys for physics geometry on static meshes and landscape heightfields

diff --git a/Engine/Source/Runtime/Engine/Private/PhysicsEngine/Experimental/ChaosDerivedData.cpp b/Engine/Source/Runtime/Engine/Private/PhysicsEngine/Experimental/ChaosDerivedData.cpp
--- a/Engine/Source/Runtime/Engine/Private/PhysicsEngine/Experimental/ChaosDerivedData.cpp
+++ b/Engine/Source/Runtime/Engine/Private/PhysicsEngine/Experimental/ChaosDerivedData.cpp
@@ -80,10 +80,19 @@ FString FChaosDerivedDataCooker::GetPluginSpecificCacheKeySuffix() const
 		Setup->GetGeometryDDCKey(SetupGeometryKey);
 	}
 
-	return FString::Printf(TEXT("%s_%s_REAL%d"),
-		*RequestedFormat.ToString(),
-		*SetupGeometryKey,
-		(int)sizeof(Chaos::FReal));
+	FString OutSuffix = FString::Printf(TEXT("%s_%s_REAL%d"),
+										*RequestedFormat.ToString(),
+										*SetupGeometryKey,
+										(int)sizeof(Chaos::FReal));
+
+#if PLATFORM_CPU_ARM_FAMILY
+	// Separate out arm keys as x64 and arm64 clang do not generate the same data for a given
+	// input. Add the arm specifically so that a) we avoid rebuilding the current DDC and
+	// b) we can remove it once we get arm64 to be consistent.
+	OutSuffix.Append(TEXT("_arm64"));
+#endif
+
+	return OutSuffix;
 }
 
 bool FChaosDerivedDataCooker::IsBuildThreadsafe() const
diff --git a/Engine/Source/Runtime/Landscape/Private/LandscapeCollision.cpp b/Engine/Source/Runtime/Landscape/Private/LandscapeCollision.cpp
--- a/Engine/Source/Runtime/Landscape/Private/LandscapeCollision.cpp
+++ b/Engine/Source/Runtime/Landscape/Private/LandscapeCollision.cpp
@@ -185,7 +185,14 @@ static FString GetHFDDCKeyString(const FName& Format, bool bDefMaterial, const F
 		CombinedStateId = FGuid(Hash[0] ^ Hash[4], Hash[1], Hash[2], Hash[3]);
 	}
 
-	const FString InterfacePrefix = FString::Printf(TEXT("%s_%s"), TEXT("CHAOS"), Chaos::ChaosVersionGUID);
+	FString InterfacePrefix = FString::Printf(TEXT("%s_%s"), TEXT("CHAOS"), Chaos::ChaosVersionGUID);
+
+#if PLATFORM_CPU_ARM_FAMILY
+	// Separate out arm keys as x64 and arm64 clang do not generate the same data for a given
+	// input. Add the arm specifically so that a) we avoid rebuilding the current DDC and
+	// b) we can remove it once we get arm64 to be consistent.
+	InterfacePrefix.Append(TEXT("_arm64"));
+#endif
 
 	const FString KeyPrefix = FString::Printf(TEXT("%s_%s_%s"), *InterfacePrefix, *Format.ToString(), (bDefMaterial ? TEXT("VIS") : TEXT("FULL")));
 	return FDerivedDataCacheInterface::BuildCacheKey(*KeyPrefix, LANDSCAPE_COLLISION_DERIVEDDATA_VER, *CombinedStateId.ToString());
