Landscape Dirtyness
- New CVar 'landscape.DirtyOnlyInMode' when enabled, modifications to the landscape that happen outside of the Landscape Edit Mode (Modifications affecting the Edit Layers like Water brush) will be tracked but packages will not be marked dirty.
- Users can use the Build > Save Modified Landscapes to save those collateral changes.

diff --git a/Engine/Source/Editor/LandscapeEditor/Private/LandscapeEditorModule.cpp b/Engine/Source/Editor/LandscapeEditor/Private/LandscapeEditorModule.cpp
--- a/Engine/Source/Editor/LandscapeEditor/Private/LandscapeEditorModule.cpp
+++ b/Engine/Source/Editor/LandscapeEditor/Private/LandscapeEditorModule.cpp
@@ -135,6 +135,14 @@ public:
 		
 			FUIAction ActionBuildNanite(FExecuteAction::CreateStatic(&BuildNanite), FCanExecuteAction());
 			Section.AddMenuEntry(NAME_None, LOCTEXT("BuildNaniteOnly", "Build Nanite Only"), LOCTEXT("BuildLandscapeNanite", "Build Nanite representation"), TAttribute<FSlateIcon>(), ActionBuildNanite, EUserInterfaceActionType::Button);
+
+			if (ULandscapeSubsystem::IsDirtyOnlyInModeEnabled())
+			{
+				FUIAction ActionSaveModifiedLandscapes(FExecuteAction::CreateStatic(&SaveModifiedLandscapes), FCanExecuteAction::CreateStatic(&HasModifiedLandscapes));
+				Section.AddMenuEntry(NAME_None,
+					LOCTEXT("SaveModifiedLandscapes", "Save Modified Landscapes"), LOCTEXT("SaveModifiedLandscapesToolTip", "Save landscapes that were modified outside of the editor mode"),
+					TAttribute<FSlateIcon>(), ActionSaveModifiedLandscapes, EUserInterfaceActionType::Button);
+			}
 		}
 		
 		ILandscapeModule& LandscapeModule = FModuleManager::GetModuleChecked<ILandscapeModule>("Landscape");
@@ -262,6 +270,30 @@ public:
 		}
 	}
 
+	static bool HasModifiedLandscapes()
+	{
+		if (UWorld* World = GEditor->GetEditorWorldContext().World())
+		{
+			if (ULandscapeSubsystem* LandscapeSubsystem = World->GetSubsystem<ULandscapeSubsystem>())
+			{
+				return LandscapeSubsystem->HasModifiedLandscapes();
+			}
+		}
+
+		return false;
+	}
+
+	static void SaveModifiedLandscapes()
+	{
+		if (UWorld* World = GEditor->GetEditorWorldContext().World())
+		{
+			if (ULandscapeSubsystem* LandscapeSubsystem = World->GetSubsystem<ULandscapeSubsystem>())
+			{
+				LandscapeSubsystem->SaveModifiedLandscapes();
+			}
+		}
+	}
+
 	static void ChangeLandscapeViewMode(ELandscapeViewMode::Type ViewMode)
 	{
 		if (ViewMode != GLandscapeViewMode)
diff --git a/Engine/Source/Runtime/Engine/Private/UnrealEngine.cpp b/Engine/Source/Runtime/Engine/Private/UnrealEngine.cpp
--- a/Engine/Source/Runtime/Engine/Private/UnrealEngine.cpp
+++ b/Engine/Source/Runtime/Engine/Private/UnrealEngine.cpp
@@ -11618,7 +11618,7 @@ float DrawMapWarnings(UWorld* World, FViewport* Viewport, FCanvas* Canvas, UCanv
 #if WITH_EDITOR && !(UE_BUILD_SHIPPING || UE_BUILD_TEST)
 	if (ULandscapeSubsystem* LandscapeSubsystem = World->GetSubsystem<ULandscapeSubsystem>())
 	{
-		LandscapeSubsystem->DisplayBuildMessages(Canvas, MessageX, MessageY);
+		LandscapeSubsystem->DisplayMessages(Canvas, MessageX, MessageY);
 	}
 #endif
 
diff --git a/Engine/Source/Runtime/Landscape/Classes/Landscape.h b/Engine/Source/Runtime/Landscape/Classes/Landscape.h
--- a/Engine/Source/Runtime/Landscape/Classes/Landscape.h
+++ b/Engine/Source/Runtime/Landscape/Classes/Landscape.h
@@ -275,6 +275,7 @@ public:
 	LANDSCAPE_API FBox GetCompleteBounds() const;
 	LANDSCAPE_API void RegisterLandscapeEdMode(ILandscapeEdModeInterface* InLandscapeEdMode) { LandscapeEdMode = InLandscapeEdMode; }
 	LANDSCAPE_API void UnregisterLandscapeEdMode() { LandscapeEdMode = nullptr; }
+	LANDSCAPE_API bool HasLandscapeEdMode() const { return LandscapeEdMode != nullptr; }
 	LANDSCAPE_API virtual bool HasLayersContent() const override;
 	LANDSCAPE_API virtual void UpdateCachedHasLayersContent(bool bInCheckComponentDataIntegrity) override;
 	LANDSCAPE_API void RequestSplineLayerUpdate();
diff --git a/Engine/Source/Runtime/Landscape/Classes/LandscapeInfo.h b/Engine/Source/Runtime/Landscape/Classes/LandscapeInfo.h
--- a/Engine/Source/Runtime/Landscape/Classes/LandscapeInfo.h
+++ b/Engine/Source/Runtime/Landscape/Classes/LandscapeInfo.h
@@ -44,6 +44,21 @@ struct FLandscapeAddCollision
 };
 #endif // WITH_EDITORONLY_DATA
 
+class ULandscapeInfo;
+
+#if WITH_EDITOR
+struct FLandscapeDirtyOnlyInModeScope
+{
+	FLandscapeDirtyOnlyInModeScope() = delete;
+	FLandscapeDirtyOnlyInModeScope(ULandscapeInfo* InLandscapeInfo);
+	~FLandscapeDirtyOnlyInModeScope();
+
+private:
+	ULandscapeInfo* LandscapeInfo;
+	bool bDirtyOnlyInModePrevious;
+};
+#endif
+
 USTRUCT()
 struct FLandscapeInfoLayerSettings
 {
@@ -149,6 +164,13 @@ private:
 #if WITH_EDITORONLY_DATA
 	UPROPERTY()
 	TArray<TScriptInterface<ILandscapeSplineInterface>> SplineActors;
+
+	// Not a property since this shouldn't be modified through transactions (no undo/redo)
+	TSet<TWeakObjectPtr<UPackage>> ModifiedPackages;
+
+	bool bDirtyOnlyInMode;
+
+	friend struct FLandscapeDirtyOnlyInModeScope;
 #endif // WITH_EDITORONLY_DATA
 
 	TSet<ULandscapeComponent*> SelectedComponents;
@@ -225,6 +247,14 @@ public:
 	 */
 	LANDSCAPE_API ALandscapeProxy* GetLandscapeProxyForLevel(ULevel* Level) const;
 
+	LANDSCAPE_API static bool IsDirtyOnlyInModeEnabled();
+
+	LANDSCAPE_API void OnModifiedPackageSaved(UPackage* InPackage);
+	LANDSCAPE_API int32 GetModifiedPackageCount() const;
+	LANDSCAPE_API TArray<UPackage*> GetModifiedPackages() const;
+
+	LANDSCAPE_API void ModifyObject(UObject* InObject, bool bAlwaysMarkDirty = true);
+	LANDSCAPE_API void MarkObjectDirty(UObject* InObject);
 #endif //WITH_EDITOR
 
 	/**
diff --git a/Engine/Source/Runtime/Landscape/Classes/LandscapeInfoMap.h b/Engine/Source/Runtime/Landscape/Classes/LandscapeInfoMap.h
--- a/Engine/Source/Runtime/Landscape/Classes/LandscapeInfoMap.h
+++ b/Engine/Source/Runtime/Landscape/Classes/LandscapeInfoMap.h
@@ -35,4 +35,5 @@ public:
 	* @returns Landscape-specific data associated with given world.
 	*/
 	LANDSCAPE_API static ULandscapeInfoMap& GetLandscapeInfoMap(const UWorld* World);
+	LANDSCAPE_API static ULandscapeInfoMap* FindLandscapeInfoMap(const UWorld* World);
 };
diff --git a/Engine/Source/Runtime/Landscape/Private/Landscape.cpp b/Engine/Source/Runtime/Landscape/Private/Landscape.cpp
--- a/Engine/Source/Runtime/Landscape/Private/Landscape.cpp
+++ b/Engine/Source/Runtime/Landscape/Private/Landscape.cpp
@@ -175,6 +175,12 @@ static FAutoConsoleVariableRef CVarLiveRebuildNaniteOnModification(
 	LiveRebuildNaniteOnModification,
 	TEXT("Trigger a rebuild of Nanite representation immediately when a modification is performed"));
 
+static int32 GLandscapeDirtyOnlyInModeEnabled = 0;
+static FAutoConsoleVariableRef CVarLandscapeDirtyOnlyInModeEnabled(
+	TEXT("landscape.DirtyOnlyInMode"),
+	GLandscapeDirtyOnlyInModeEnabled,
+	TEXT("Set to 1 to avoid Landscape being dirtied when not in Landscape mode."));
+
 #endif
 
 int32 GRenderNaniteLandscape = 1;
@@ -859,6 +865,9 @@ void ULandscapeComponent::GetLayerDebugColorKey(int32& R, int32& G, int32& B) co
 
 ULandscapeInfo::ULandscapeInfo(const FObjectInitializer& ObjectInitializer)
 	: Super(ObjectInitializer)
+#if WITH_EDITORONLY_DATA
+	, bDirtyOnlyInMode(false)
+#endif
 	, XYComponentBounds(MAX_int32, MAX_int32, MIN_int32, MIN_int32)
 {
 }
@@ -2833,6 +2842,11 @@ void ALandscapeProxy::PreSave(FObjectPreSaveContext ObjectSaveContext)
 		UpdateNaniteRepresentation();
 		UpdateRenderingMethod();
 	}
+
+	if (ULandscapeInfo* LandscapeInfo = GetLandscapeInfo())
+	{
+		LandscapeInfo->OnModifiedPackageSaved(GetPackage());
+	}
 #endif // WITH_EDITOR
 }
 
@@ -3614,6 +3628,109 @@ void ALandscape::PreSave(FObjectPreSaveContext ObjectSaveContext)
 	//}
 }
 
+bool ULandscapeInfo::IsDirtyOnlyInModeEnabled()
+{
+	return GLandscapeDirtyOnlyInModeEnabled > 0;
+}
+
+FLandscapeDirtyOnlyInModeScope::FLandscapeDirtyOnlyInModeScope(ULandscapeInfo* InLandscapeInfo)
+	: LandscapeInfo(InLandscapeInfo)
+	, bDirtyOnlyInModePrevious(InLandscapeInfo->bDirtyOnlyInMode)
+{
+	LandscapeInfo->bDirtyOnlyInMode = ULandscapeInfo::IsDirtyOnlyInModeEnabled();
+}
+
+FLandscapeDirtyOnlyInModeScope::~FLandscapeDirtyOnlyInModeScope()
+{
+	LandscapeInfo->bDirtyOnlyInMode = bDirtyOnlyInModePrevious;
+}
+
+void ULandscapeInfo::OnModifiedPackageSaved(UPackage* InPackage)
+{
+	ModifiedPackages.Remove(InPackage);
+}
+
+TArray<UPackage*> ULandscapeInfo::GetModifiedPackages() const
+{
+	TArray<UPackage*> LocalModifiedPackages;
+	LocalModifiedPackages.Reserve(ModifiedPackages.Num());
+	for (const TWeakObjectPtr<UPackage>& WeakPackagePtr : ModifiedPackages)
+	{
+		if (UPackage* Package = WeakPackagePtr.Get())
+		{
+			LocalModifiedPackages.Add(Package);
+		}
+	}
+
+	return LocalModifiedPackages;
+}
+
+int32 ULandscapeInfo::GetModifiedPackageCount() const
+{
+	int32 ModifiedPackageCount = 0;
+
+	for (const TWeakObjectPtr<UPackage>& WeakPackagePtr : ModifiedPackages)
+	{
+		if (UPackage* Package = WeakPackagePtr.Get())
+		{
+			ModifiedPackageCount++;
+		}
+	}
+
+	return ModifiedPackageCount;
+}
+
+void ULandscapeInfo::MarkObjectDirty(UObject* InObject)
+{
+	check(InObject && (InObject->IsA<ALandscapeProxy>() || InObject->GetTypedOuter<ALandscapeProxy>() != nullptr));
+		
+	if (bDirtyOnlyInMode)
+	{
+		ALandscape* LocalLandscapeActor = LandscapeActor.Get();
+		check(LocalLandscapeActor);
+		if (LocalLandscapeActor->HasLandscapeEdMode())
+		{
+			InObject->MarkPackageDirty();
+		}
+		else
+		{
+			ModifiedPackages.Add(InObject->GetPackage());
+		}
+	}
+	else
+	{
+		InObject->MarkPackageDirty();
+	}
+}
+
+void ULandscapeInfo::ModifyObject(UObject* InObject, bool bAlwaysMarkDirty)
+{
+	check(InObject && (InObject->IsA<ALandscapeProxy>() || InObject->GetTypedOuter<ALandscapeProxy>() != nullptr));
+	
+	if (!bAlwaysMarkDirty)
+	{
+		InObject->Modify(false);
+	}
+	else if(!bDirtyOnlyInMode)
+	{
+		InObject->Modify(true);
+	}
+	else
+	{
+		ALandscape* LocalLandscapeActor = LandscapeActor.Get();
+		check(LocalLandscapeActor);
+		if (LocalLandscapeActor->HasLandscapeEdMode())
+		{
+			InObject->Modify(true);
+		}
+		else
+		{
+			InObject->Modify(false);
+			ModifiedPackages.Add(InObject->GetPackage());
+		}
+	}
+}
+
 ALandscapeProxy* ULandscapeInfo::GetLandscapeProxyForLevel(ULevel* Level) const
 {
 	ALandscapeProxy* LandscapeProxy = nullptr;
diff --git a/Engine/Source/Runtime/Landscape/Private/LandscapeEdit.cpp b/Engine/Source/Runtime/Landscape/Private/LandscapeEdit.cpp
--- a/Engine/Source/Runtime/Landscape/Private/LandscapeEdit.cpp
+++ b/Engine/Source/Runtime/Landscape/Private/LandscapeEdit.cpp
@@ -6005,11 +6005,13 @@ void ULandscapeComponent::ReallocateWeightmaps(FLandscapeEditDataInterface* Data
 		}
 	}
 
+	ULandscapeInfo* LandscapeInfo = GetLandscapeInfo();
+
 	bool bMarkPackageDirty = DataInterface == nullptr ? true : DataInterface->GetShouldDirtyPackage();
 	if (InSaveToTransactionBuffer)
 	{
-		Modify(bMarkPackageDirty);
-		TargetProxy->Modify(bMarkPackageDirty);
+		LandscapeInfo->ModifyObject(this, bMarkPackageDirty);
+		LandscapeInfo->ModifyObject(TargetProxy, bMarkPackageDirty);
 	}
 
 	if (!InForceReallocate)
@@ -6143,7 +6145,7 @@ void ULandscapeComponent::ReallocateWeightmaps(FLandscapeEditDataInterface* Data
 		// No suitable weightmap texture
 		if (CurrentWeightmapTexture == nullptr)
 		{
-			MarkPackageDirty();
+			LandscapeInfo->MarkObjectDirty(this);
 
 			// Weightmap is sized the same as the component
 			int32 WeightmapSize = (SubsectionSizeQuads + 1) * NumSubsections;
diff --git a/Engine/Source/Runtime/Landscape/Private/LandscapeEditLayers.cpp b/Engine/Source/Runtime/Landscape/Private/LandscapeEditLayers.cpp
--- a/Engine/Source/Runtime/Landscape/Private/LandscapeEditLayers.cpp
+++ b/Engine/Source/Runtime/Landscape/Private/LandscapeEditLayers.cpp
@@ -5246,7 +5246,7 @@ bool ALandscape::ResolveLayersTexture(
 						bChanged = true;
 
 						// We're about to modify the texture's source data, the texture needs to know so that it can handle properly update cached platform data (additionally, the package needs to be dirtied) :
-						InOutputTexture->Modify();
+						GetLandscapeInfo()->ModifyObject(InOutputTexture);
 					}
 				}
 
@@ -7803,6 +7803,9 @@ void ALandscape::UpdateLayersContent(bool bInWaitForStreaming, bool bInSkipMonit
 		return;
 	}
 
+	// Make sure Update doesn't dirty Landscape packages when not in Landscape Ed Mode
+	FLandscapeDirtyOnlyInModeScope DirtyOnlyInMode(LandscapeInfo);
+
 	// If we went from local merge to global merge or vice versa, we need to reinitialize layers : 
 	bool bUseLocalMerge = SupportsEditLayersLocalMerge();
 	if (bLandscapeLayersAreUsingLocalMerge != bUseLocalMerge)
@@ -8994,12 +8997,12 @@ void ALandscape::UpdateLandscapeSplines(const FGuid& InTargetLayer, bool bInUpda
 					// Was never computed
 					if (Component->SplineHash == 0)
 					{
-						Component->Modify(); // mark package dirty
+						LandscapeInfo->ModifyObject(Component);
 						Component->SplineHash = DefaultSplineHash;
 					}
 
 					PreviousHashes.Add(Component, Component->SplineHash);
-					Component->Modify(false);
+					LandscapeInfo->ModifyObject(Component, false);
 					Component->SplineHash = DefaultSplineHash;
 				});
 			}
@@ -9022,13 +9025,13 @@ void ALandscape::UpdateLandscapeSplines(const FGuid& InTargetLayer, bool bInUpda
 					uint32 NewHash = Pair.Key->ComputeLayerHash();
 					if (NewHash != Pair.Value)
 					{
-						Pair.Key->MarkPackageDirty();
+						LandscapeInfo->MarkObjectDirty(Pair.Key);
 					}
 					Pair.Key->SplineHash = NewHash;
 				}
 				else if (Pair.Key->SplineHash == DefaultSplineHash && Pair.Value != DefaultSplineHash)
 				{
-					Pair.Key->MarkPackageDirty();
+					LandscapeInfo->MarkObjectDirty(Pair.Key);
 				}
 			}
 		}
diff --git a/Engine/Source/Runtime/Landscape/Private/LandscapeInfoMap.cpp b/Engine/Source/Runtime/Landscape/Private/LandscapeInfoMap.cpp
--- a/Engine/Source/Runtime/Landscape/Private/LandscapeInfoMap.cpp
+++ b/Engine/Source/Runtime/Landscape/Private/LandscapeInfoMap.cpp
@@ -52,3 +52,10 @@ ULandscapeInfoMap& ULandscapeInfoMap::GetLandscapeInfoMap(const UWorld* World)
 
 	return *FoundObject;
 }
+
+ULandscapeInfoMap* ULandscapeInfoMap::FindLandscapeInfoMap(const UWorld* World)
+{
+	ULandscapeInfoMap* FoundObject = nullptr;
+	World->PerModuleDataObjects.FindItemByClass(&FoundObject);
+	return FoundObject;
+}
diff --git a/Engine/Source/Runtime/Landscape/Private/LandscapeLight.cpp b/Engine/Source/Runtime/Landscape/Private/LandscapeLight.cpp
--- a/Engine/Source/Runtime/Landscape/Private/LandscapeLight.cpp
+++ b/Engine/Source/Runtime/Landscape/Private/LandscapeLight.cpp
@@ -752,7 +752,7 @@ void ULandscapeComponent::GetLightAndShadowMapMemoryUsage( int32& LightMapMemory
 
 void ULandscapeComponent::InvalidateLightingCacheDetailed(bool bInvalidateBuildEnqueuedLighting, bool bTranslationOnly)
 {
-	Modify();
+	GetLandscapeInfo()->ModifyObject(this);
 
 	FComponentReregisterContext ReregisterContext(this);
 
diff --git a/Engine/Source/Runtime/Landscape/Private/LandscapeSubsystem.cpp b/Engine/Source/Runtime/Landscape/Private/LandscapeSubsystem.cpp
--- a/Engine/Source/Runtime/Landscape/Private/LandscapeSubsystem.cpp
+++ b/Engine/Source/Runtime/Landscape/Private/LandscapeSubsystem.cpp
@@ -9,6 +9,7 @@
 #include "LandscapeProxy.h"
 #include "LandscapeStreamingProxy.h"
 #include "LandscapeInfo.h"
+#include "LandscapeInfoMap.h"
 #include "ProfilingDebugging/CsvProfiler.h"
 #include "WorldPartition/WorldPartitionSubsystem.h"
 #include "ActorPartition/ActorPartitionSubsystem.h"
@@ -18,6 +19,10 @@
 #include "Engine/Canvas.h"
 #include "EngineUtils.h"
 
+#if WITH_EDITOR
+#include "FileHelpers.h"
+#endif
+
 static int32 GUseStreamingManagerForCameras = 1;
 static FAutoConsoleVariableRef CVarUseStreamingManagerForCameras(
 	TEXT("grass.UseStreamingManagerForCameras"),
@@ -217,6 +222,97 @@ void ULandscapeSubsystem::BuildNanite()
 	}
 }
 
+void ULandscapeSubsystem::ForEachLandscapeInfo(TFunctionRef<bool(ULandscapeInfo*)> ForEachLandscapeInfoFunc) const
+{
+	if (ULandscapeInfoMap* LandscapeInfoMap = ULandscapeInfoMap::FindLandscapeInfoMap(GetWorld()))
+	{
+		for (const auto& Pair : LandscapeInfoMap->Map)
+		{
+			if (ULandscapeInfo* LandscapeInfo = Pair.Value)
+			{
+				if (!ForEachLandscapeInfoFunc(LandscapeInfo))
+				{
+					return;
+				}
+			}
+		}
+	}
+}
+
+bool ULandscapeSubsystem::IsDirtyOnlyInModeEnabled()
+{
+	return ULandscapeInfo::IsDirtyOnlyInModeEnabled();
+}
+
+void ULandscapeSubsystem::SaveModifiedLandscapes()
+{	
+	TSet<UPackage*> SetDirtyPackages;
+	TSet<FName> PackagesToSave;
+
+	const bool bSkipDirty = false;
+
+	// Gather list of packages to save and make them dirty so they are considered by FEditorFileUtils::SaveDirtyPackages.
+	ForEachLandscapeInfo([&](ULandscapeInfo* LandscapeInfo)
+	{
+		for(UPackage* ModifiedPackage : LandscapeInfo->GetModifiedPackages())
+		{
+			PackagesToSave.Add(ModifiedPackage->GetFName());
+			if (!ModifiedPackage->IsDirty())
+			{
+				SetDirtyPackages.Add(ModifiedPackage);
+				ModifiedPackage->SetDirtyFlag(true);
+			}
+		}
+		return true;
+	});
+
+	const bool bPromptUserToSave = true;
+	const bool bSaveMapPackages = true;
+	const bool bSaveContentPackages = true;
+	const bool bFastSave = false;
+	const bool bNotifyNoPackagesSaved = false;
+	const bool bCanBeDeclined = true;
+
+	FEditorFileUtils::SaveDirtyPackages(bPromptUserToSave, bSaveMapPackages, bSaveContentPackages, bFastSave, bNotifyNoPackagesSaved, bCanBeDeclined, nullptr,
+		[PackagesToSave](UPackage* DirtyPackage)
+		{
+			if (PackagesToSave.Contains(DirtyPackage->GetFName()))
+			{
+				return false;
+			}
+			return true;
+		});
+
+	// If Package wasn't saved it is still in the LandscapeInfo ModifiedPackage list, set its dirty flag back to false.
+	ForEachLandscapeInfo([&](ULandscapeInfo* LandscapeInfo)
+	{
+		for (UPackage* ModifiedPackage : LandscapeInfo->GetModifiedPackages())
+		{
+			if (SetDirtyPackages.Contains(ModifiedPackage))
+			{
+				ModifiedPackage->SetDirtyFlag(false);
+			}
+		}
+		return true;
+	});
+}
+
+bool ULandscapeSubsystem::HasModifiedLandscapes() const
+{
+	bool bHasModifiedLandscapes = false;
+	ForEachLandscapeInfo([&](ULandscapeInfo* LandscapeInfo)
+	{
+		if (LandscapeInfo->GetModifiedPackageCount() > 0)
+		{
+			bHasModifiedLandscapes = true;
+			return false;
+		}
+		return true;
+	});
+	
+	return bHasModifiedLandscapes;
+}
+
 bool ULandscapeSubsystem::IsGridBased() const
 {
 	return UWorld::IsPartitionedWorld(GetWorld());
@@ -245,7 +341,7 @@ ALandscapeProxy* ULandscapeSubsystem::FindOrAddLandscapeProxy(ULandscapeInfo* La
 	return FLandscapeConfigHelper::FindOrAddLandscapeStreamingProxy(LandscapeInfo, SectionBase);
 }
 
-void ULandscapeSubsystem::DisplayBuildMessages(FCanvas* Canvas, float& XPos, float& YPos)
+void ULandscapeSubsystem::DisplayMessages(FCanvas* Canvas, float& XPos, float& YPos)
 {
 	const int32 FontSizeY = 20;
 	FCanvasTextItem SmallTextItem(FVector2D(0, 0), FText::GetEmpty(), GEngine->GetSmallFont(), FLinearColor::White);
@@ -254,7 +350,7 @@ void ULandscapeSubsystem::DisplayBuildMessages(FCanvas* Canvas, float& XPos, flo
 	if (int32 OutdatedGrassMapCount = GetOutdatedGrassMapCount())
 	{
 		SmallTextItem.SetColor(FLinearColor::Red);
-		SmallTextItem.Text = FText::Format(LOCTEXT("GRASS_MAPS_NEED_TO_BE_REBUILT_FMT", "GRASS MAPS NEEDS TO BE REBUILT ({0} {0}|plural(one=object,other=objects))"), OutdatedGrassMapCount);
+		SmallTextItem.Text = FText::Format(LOCTEXT("GRASS_MAPS_NEED_TO_BE_REBUILT_FMT", "LANDSCAPE: GRASS MAPS NEEDS TO BE REBUILT ({0} {0}|plural(one=object,other=objects))"), OutdatedGrassMapCount);
 		Canvas->DrawItem(SmallTextItem, FVector2D(XPos, YPos));
 		YPos += FontSizeY;
 	}
@@ -262,7 +358,7 @@ void ULandscapeSubsystem::DisplayBuildMessages(FCanvas* Canvas, float& XPos, flo
 	if (int32 ComponentsNeedingGITextureBaking = GetOutdatedGIBakedTextureComponentsCount())
 	{
 		SmallTextItem.SetColor(FLinearColor::Red);
-		SmallTextItem.Text = FText::Format(LOCTEXT("LANDSCAPE_TEXTURES_NEED_TO_BE_REBUILT_FMT", "LANDSCAPE BAKED TEXTURES NEEDS TO BE REBUILT ({0} {0}|plural(one=object,other=objects))"), ComponentsNeedingGITextureBaking);
+		SmallTextItem.Text = FText::Format(LOCTEXT("LANDSCAPE_TEXTURES_NEED_TO_BE_REBUILT_FMT", "LANDSCAPE: BAKED TEXTURES NEEDS TO BE REBUILT ({0} {0}|plural(one=object,other=objects))"), ComponentsNeedingGITextureBaking);
 		Canvas->DrawItem(SmallTextItem, FVector2D(XPos, YPos));
 		YPos += FontSizeY;
 	}
@@ -270,10 +366,28 @@ void ULandscapeSubsystem::DisplayBuildMessages(FCanvas* Canvas, float& XPos, flo
 	if (int32 ComponentsWithOudatedPhysicalMaterial = GetOudatedPhysicalMaterialComponentsCount())
 	{
 		SmallTextItem.SetColor(FLinearColor::Red);
-		SmallTextItem.Text = FText::Format(LOCTEXT("LANDSCAPE_PHYSICALMATERIAL_NEED_TO_BE_REBUILT_FMT", "LANDSCAPE PHYSICAL MATERIAL NEEDS TO BE REBUILT ({0} {0}|plural(one=object,other=objects))"), ComponentsWithOudatedPhysicalMaterial);
+		SmallTextItem.Text = FText::Format(LOCTEXT("LANDSCAPE_PHYSICALMATERIAL_NEED_TO_BE_REBUILT_FMT", "LANDSCAPE: PHYSICAL MATERIAL NEEDS TO BE REBUILT ({0} {0}|plural(one=object,other=objects))"), ComponentsWithOudatedPhysicalMaterial);
 		Canvas->DrawItem(SmallTextItem, FVector2D(XPos, YPos));
 		YPos += FontSizeY;
 	}
+
+	if (ULandscapeInfoMap* LandscapeInfoMap = ULandscapeInfoMap::FindLandscapeInfoMap(GetWorld()))
+	{
+		int32 ModifiedNotDirtyCount = 0;
+		ForEachLandscapeInfo([&ModifiedNotDirtyCount](ULandscapeInfo* LandscapeInfo)
+		{
+			ModifiedNotDirtyCount += LandscapeInfo->GetModifiedPackageCount();
+			return true;
+		});
+				
+		if (ModifiedNotDirtyCount > 0)
+		{
+			SmallTextItem.SetColor(FLinearColor::Red);
+			SmallTextItem.Text = FText::Format(LOCTEXT("LANDSCAPE_NEED_TO_BE_SAVED", "LANDSCAPE: NEED TO BE SAVED ({0} {0}|plural(one=object,other=objects))"), ModifiedNotDirtyCount);
+			Canvas->DrawItem(SmallTextItem, FVector2D(XPos, YPos));
+			YPos += FontSizeY;
+		}
+	}
 }
 
 #endif
diff --git a/Engine/Source/Runtime/Landscape/Public/LandscapeSubsystem.h b/Engine/Source/Runtime/Landscape/Public/LandscapeSubsystem.h
--- a/Engine/Source/Runtime/Landscape/Public/LandscapeSubsystem.h
+++ b/Engine/Source/Runtime/Landscape/Public/LandscapeSubsystem.h
@@ -41,10 +41,17 @@ public:
 	LANDSCAPE_API bool IsGridBased() const;
 	LANDSCAPE_API void ChangeGridSize(ULandscapeInfo* LandscapeInfo, uint32 NewGridSizeInComponents);
 	LANDSCAPE_API ALandscapeProxy* FindOrAddLandscapeProxy(ULandscapeInfo* LandscapeInfo, const FIntPoint& SectionBase);
-	LANDSCAPE_API void DisplayBuildMessages(class FCanvas* Canvas, float& XPos, float& YPos);
+	LANDSCAPE_API void DisplayMessages(class FCanvas* Canvas, float& XPos, float& YPos);
+	LANDSCAPE_API void SaveModifiedLandscapes();
+	LANDSCAPE_API bool HasModifiedLandscapes() const;
+	LANDSCAPE_API static bool IsDirtyOnlyInModeEnabled();
 #endif
 
 private:
+#if WITH_EDITOR
+	LANDSCAPE_API void ForEachLandscapeInfo(TFunctionRef<bool(ULandscapeInfo*)> ForEachLandscapeInfoFunc) const;
+#endif
+
 	// Begin USubsystem
 	virtual void Initialize(FSubsystemCollectionBase& Collection) override;
 	virtual void Deinitialize() override;
