Removed many frames delay while the landscape BP brushes are asynchronously waiting for the required materials to be compiled : we now process them all every frame so that we can initiate their compilation all at once instead of earlying-out (which resulted in only shader compilation request per frame)

diff --git a/Engine/Source/Runtime/Landscape/Private/LandscapeEditLayers.cpp b/Engine/Source/Runtime/Landscape/Private/LandscapeEditLayers.cpp
--- a/Engine/Source/Runtime/Landscape/Private/LandscapeEditLayers.cpp
+++ b/Engine/Source/Runtime/Landscape/Private/LandscapeEditLayers.cpp
@@ -3531,16 +3531,14 @@ bool ALandscape::PrepareLayersBrushResources(bool bInWaitForStreaming)
 		}
 	}
 
+	bool bIsReady = true;
 	ERHIFeatureLevel::Type FeatureLevel = GetWorld() ? (ERHIFeatureLevel::Type)GetWorld()->FeatureLevel : GMaxRHIFeatureLevel;
 	for (UObject* Dependency : Dependencies)
 	{
 		// Streamable textures need to be fully streamed in : 
 		if (UTexture2D* Texture = Cast<UTexture2D>(Dependency))
 		{
-			if (!IsTextureReady(Texture, bInWaitForStreaming))
-			{
-				return false;
-			}
+			bIsReady &= IsTextureReady(Texture, bInWaitForStreaming);
 		}
 
 		// Material shaders need to be fully compiled : 
@@ -3548,15 +3546,13 @@ bool ALandscape::PrepareLayersBrushResources(bool bInWaitForStreaming)
 		{
 			if (FMaterialResource* MaterialResource = MaterialInterface->GetMaterialResource(FeatureLevel))
 			{
-				if (!IsMaterialResourceCompiled(MaterialResource, bInWaitForStreaming))
-				{
-					return false;
-				}
+				// Don't early-out because checking for the material resource actually requests the shaders to be loaded so we want to make sure to request them all at once instead of one by one :
+				bIsReady &= IsMaterialResourceCompiled(MaterialResource, bInWaitForStreaming);
 			}
 		}
 	}
 
-	return true;
+	return bIsReady;
 }
 
 // Must match EEditLayerHeightmapBlendMode in LandscapeLayersHeightmapsPS.usf
