Only refresh lanscape scene proxies if CVarLandscapeShowCollisionMesh has actually changed from the previous value. This prevents a 450ms stall when exiting menu in firing range map.

#test Tested that changing the cvar still triggers the refresh

diff --git a/Engine/Source/Runtime/Landscape/Private/LandscapeCollision.cpp b/Engine/Source/Runtime/Landscape/Private/LandscapeCollision.cpp
--- a/Engine/Source/Runtime/Landscape/Private/LandscapeCollision.cpp
+++ b/Engine/Source/Runtime/Landscape/Private/LandscapeCollision.cpp
@@ -61,10 +61,10 @@ using namespace PhysicsInterfaceTypes;
 // Global switch for whether to read/write to DDC for landscape cooked data
 // It's a lot faster to compute than to request from DDC, so always skip.
 bool GLandscapeCollisionSkipDDC = true;
-
+static int32 CVarLandscapeShowCollisionMeshCurrentValue = 0;
 static TAutoConsoleVariable<int32> CVarLandscapeShowCollisionMesh(
 	TEXT("landscape.ShowCollisionMesh"),
-	0,
+	CVarLandscapeShowCollisionMeshCurrentValue,
 	TEXT("Selects which heightfield to visualize when ShowFlags.Collision is used. 0 for simple, 1 for complex, 2 for editor only."),
 	ECVF_RenderThreadSafe
 );
@@ -461,9 +461,15 @@ void ULandscapeHeightfieldCollisionComponent::ApplyWorldOffset(const FVector& In
 // Callback to flag scene proxy as dirty when cvar changes
 static void OnLandscapeShowCollisionMeshChanged()
 {
-	for (ULandscapeHeightfieldCollisionComponent* LandscapeHeightfieldCollisionComponent : TObjectRange<ULandscapeHeightfieldCollisionComponent>(RF_ClassDefaultObject | RF_ArchetypeObject, true, EInternalObjectFlags::Garbage))
+	const int32 Value = CVarLandscapeShowCollisionMesh.GetValueOnAnyThread();
+	if (CVarLandscapeShowCollisionMeshCurrentValue != Value)
 	{
-		LandscapeHeightfieldCollisionComponent->MarkRenderStateDirty();
+		CVarLandscapeShowCollisionMeshCurrentValue = Value;
+
+		for (ULandscapeHeightfieldCollisionComponent* LandscapeHeightfieldCollisionComponent : TObjectRange<ULandscapeHeightfieldCollisionComponent>(RF_ClassDefaultObject | RF_ArchetypeObject, true, EInternalObjectFlags::Garbage))
+		{
+			LandscapeHeightfieldCollisionComponent->MarkRenderStateDirty();
+		}
 	}
 }
 
