Fixed Landscape Nanite intialization problem where r.Nanite is toggled mid-way through loads, and landscape can be left with some components as Nanite and some as non-Nanite (depending on toggle direction). Also fixed landscape.RenderNanite to always be respected by other places that call UpdateRenderingMethod()

diff --git a/Engine/Source/Runtime/Landscape/Classes/LandscapeProxy.h b/Engine/Source/Runtime/Landscape/Classes/LandscapeProxy.h
--- a/Engine/Source/Runtime/Landscape/Classes/LandscapeProxy.h
+++ b/Engine/Source/Runtime/Landscape/Classes/LandscapeProxy.h
@@ -885,7 +885,7 @@ public:
 	/* Invalidate the precomputed grass and baked texture data on all components */
 	LANDSCAPE_API void InvalidateGeneratedComponentData(bool bInvalidateLightingCache = false);
 
-	LANDSCAPE_API void UpdateRenderingMethod(bool bDisableNanite = false);
+	LANDSCAPE_API void UpdateRenderingMethod();
 
 #if WITH_EDITOR
 	/** Update Grass maps */
diff --git a/Engine/Source/Runtime/Landscape/Private/Landscape.cpp b/Engine/Source/Runtime/Landscape/Private/Landscape.cpp
--- a/Engine/Source/Runtime/Landscape/Private/Landscape.cpp
+++ b/Engine/Source/Runtime/Landscape/Private/Landscape.cpp
@@ -188,15 +188,6 @@ FAutoConsoleVariableRef CVarRenderNaniteLandscape(
 	TEXT("landscape.RenderNanite"),
 	GRenderNaniteLandscape,
 	TEXT("Render Landscape using Nanite."),
-	FConsoleVariableDelegate::CreateLambda([](IConsoleVariable* InVariable)
-	{
-		const bool bDisableNanite = !InVariable->GetBool();
-		ForEachObjectOfClass(ALandscapeProxy::StaticClass(), [&](UObject* Obj)
-		{
-			ALandscapeProxy* LandscapeProxy = CastChecked<ALandscapeProxy>(Obj);
-			LandscapeProxy->UpdateRenderingMethod(bDisableNanite);
-		}, /* bIncludeDerivedClasses= */ true, RF_ClassDefaultObject | RF_ArchetypeObject, EInternalObjectFlags::Garbage);
-	}),
 	ECVF_Scalability | ECVF_RenderThreadSafe
 );
 
@@ -5124,7 +5115,7 @@ void ALandscapeProxy::InvalidateGeneratedComponentData(const TSet<ULandscapeComp
 	InvalidateGeneratedComponentData(Components.Array(), bInvalidateLightingCache);
 }
 
-void ALandscapeProxy::UpdateRenderingMethod(bool bDisableNanite)
+void ALandscapeProxy::UpdateRenderingMethod()
 {
 	if (LandscapeComponents.Num() == 0)
 	{
@@ -5133,7 +5124,7 @@ void ALandscapeProxy::UpdateRenderingMethod(bool bDisableNanite)
 
 	bool bNaniteActive = false;
 
-	if (!bDisableNanite && NaniteComponent != nullptr)
+	if (NaniteComponent != nullptr && GRenderNaniteLandscape != 0)
 	{
 		bNaniteActive = UseNanite(GShaderPlatformForFeatureLevel[GMaxRHIFeatureLevel]);
 #if WITH_EDITOR
diff --git a/Engine/Source/Runtime/Landscape/Private/LandscapeSubsystem.cpp b/Engine/Source/Runtime/Landscape/Private/LandscapeSubsystem.cpp
--- a/Engine/Source/Runtime/Landscape/Private/LandscapeSubsystem.cpp
+++ b/Engine/Source/Runtime/Landscape/Private/LandscapeSubsystem.cpp
@@ -158,6 +158,8 @@ void ULandscapeSubsystem::Tick(float DeltaTime)
 		{
 			Proxy->TickGrass(*Cameras, InOutNumComponentsCreated);
 		}
+
+		Proxy->UpdateRenderingMethod();
 	}
 
 #if WITH_EDITOR
