PSO Precache: Per component precache requests

Component precache all possible used PSOs which can be needed to draw the component. PSO precaching needs the material, vertex factory type and a set of parameters (FPSOPrecacheParams).
PSO compilation can take a long time so shouldn't be done right before drawing the object because this can lead to stalls on first usage of that PSO.

PSO precaching is called from PostLoad on the component because then all required data is available. The actual PSO precaching will happen async on a background thread

diff --git a/Engine/Plugins/Experimental/Water/Source/Runtime/Private/WaterBodyComponent.cpp b/Engine/Plugins/Experimental/Water/Source/Runtime/Private/WaterBodyComponent.cpp
--- a/Engine/Plugins/Experimental/Water/Source/Runtime/Private/WaterBodyComponent.cpp
+++ b/Engine/Plugins/Experimental/Water/Source/Runtime/Private/WaterBodyComponent.cpp
@@ -1401,6 +1401,21 @@ void UWaterBodyComponent::PostLoad()
 
 	DeprecateData();
 
+	FPSOPrecacheParams PrecachePSOParams;
+	SetupPrecachePSOParams(PrecachePSOParams);
+	if (WaterMaterial)
+	{
+		WaterMaterial->PrecachePSOs(&FLocalVertexFactory::StaticType, PrecachePSOParams);
+	}
+	if (UnderwaterPostProcessMaterial)
+	{
+		UnderwaterPostProcessMaterial->PrecachePSOs(&FLocalVertexFactory::StaticType, PrecachePSOParams);
+	}
+	if (WaterInfoMaterial)
+	{
+		WaterInfoMaterial->PrecachePSOs(&FLocalVertexFactory::StaticType, PrecachePSOParams);
+	}
+
 #if WITH_EDITOR
 	RegisterOnUpdateWavesData(GetWaterWaves(), /* bRegister = */true);
 #endif // WITH_EDITOR
diff --git a/Engine/Source/Runtime/Engine/Classes/Components/InstancedStaticMeshComponent.h b/Engine/Source/Runtime/Engine/Classes/Components/InstancedStaticMeshComponent.h
--- a/Engine/Source/Runtime/Engine/Classes/Components/InstancedStaticMeshComponent.h
+++ b/Engine/Source/Runtime/Engine/Classes/Components/InstancedStaticMeshComponent.h
@@ -402,6 +402,9 @@ public:
 
 	/** Transfers ownership of instance render data to a render thread. Instance render data will be released in scene proxy destructor or on render thread task. */
 	void ReleasePerInstanceRenderData();
+
+	/** Precache all PSOs which can be used by the component */
+	virtual void PrecachePSOs() override;
 	
 	// Number of instances in the render-side instance buffer
 	virtual int32 GetNumRenderInstances() const { return PerInstanceSMData.Num(); }
diff --git a/Engine/Source/Runtime/Engine/Classes/Components/MaterialBillboardComponent.h b/Engine/Source/Runtime/Engine/Classes/Components/MaterialBillboardComponent.h
--- a/Engine/Source/Runtime/Engine/Classes/Components/MaterialBillboardComponent.h
+++ b/Engine/Source/Runtime/Engine/Classes/Components/MaterialBillboardComponent.h
@@ -81,6 +81,10 @@ class ENGINE_API UMaterialBillboardComponent : public UPrimitiveComponent
 		class UCurveFloat* DistanceToSizeCurve
 		);
 
+	//~ Begin UObject Interface.
+	virtual void PostLoad() override;
+	//~ End UObject Interface.
+	// 
 	//~ Begin UPrimitiveComponent Interface
 	virtual FPrimitiveSceneProxy* CreateSceneProxy() override;
 	virtual FBoxSphereBounds CalcBounds(const FTransform& LocalToWorld) const override;
diff --git a/Engine/Source/Runtime/Engine/Classes/Components/MeshComponent.h b/Engine/Source/Runtime/Engine/Classes/Components/MeshComponent.h
--- a/Engine/Source/Runtime/Engine/Classes/Components/MeshComponent.h
+++ b/Engine/Source/Runtime/Engine/Classes/Components/MeshComponent.h
@@ -125,6 +125,9 @@ public:
 	/** Get the material info for texture streaming. Return whether the data is valid or not. */
 	virtual bool GetMaterialStreamingData(int32 MaterialIndex, FPrimitiveMaterialInfo& MaterialData) const { return false; }
 
+	/** Precache all PSOs which can be used by the mesh component */
+	virtual void PrecachePSOs() {}
+
 	/** Generate streaming data for all materials. */
 	void GetStreamingTextureInfoInner(FStreamingTextureLevelContext& LevelContext, const TArray<FStreamingTextureBuildInfo>* PreBuiltData, float ComponentScaling, TArray<FStreamingRenderAssetPrimitiveInfo>& OutStreamingTextures) const;
 
diff --git a/Engine/Source/Runtime/Engine/Classes/Components/PrimitiveComponent.h b/Engine/Source/Runtime/Engine/Classes/Components/PrimitiveComponent.h
--- a/Engine/Source/Runtime/Engine/Classes/Components/PrimitiveComponent.h
+++ b/Engine/Source/Runtime/Engine/Classes/Components/PrimitiveComponent.h
@@ -38,6 +38,7 @@ struct FCollisionShape;
 struct FConvexVolume;
 struct FEngineShowFlags;
 struct FNavigableGeometryExport;
+struct FPSOPrecacheParams;
 
 namespace PrimitiveComponentCVars
 {
@@ -832,6 +833,12 @@ public:
 	float GetLastRenderTime() const { return LastRenderTime; }
 	float GetLastRenderTimeOnScreen() const { return LastRenderTimeOnScreen; }
 
+	/**
+	 * Setup the parameter struct used to precache the PSOs used by this component. 
+	 * Precaching uses certain component attributes to derive the shader or state used to render the component such as static lighting, cast shadows, ...
+	 */
+	virtual void SetupPrecachePSOParams(FPSOPrecacheParams& Params);
+
 	/**
 	 * Set of actors to ignore during component sweeps in MoveComponent().
 	 * All components owned by these actors will be ignored when this component moves or updates overlaps.
diff --git a/Engine/Source/Runtime/Engine/Classes/Components/SplineMeshComponent.h b/Engine/Source/Runtime/Engine/Classes/Components/SplineMeshComponent.h
--- a/Engine/Source/Runtime/Engine/Classes/Components/SplineMeshComponent.h
+++ b/Engine/Source/Runtime/Engine/Classes/Components/SplineMeshComponent.h
@@ -357,6 +357,8 @@ public:
 
 	virtual float GetTextureStreamingTransformScale() const override;
 
+	virtual void PrecachePSOs() override;
+
 private:
 	void UpdateRenderStateAndCollision_Internal(bool bConcurrent);
 };
diff --git a/Engine/Source/Runtime/Engine/Classes/Components/StaticMeshComponent.h b/Engine/Source/Runtime/Engine/Classes/Components/StaticMeshComponent.h
--- a/Engine/Source/Runtime/Engine/Classes/Components/StaticMeshComponent.h
+++ b/Engine/Source/Runtime/Engine/Classes/Components/StaticMeshComponent.h
@@ -737,6 +737,9 @@ private:
 	/** Initializes the resources used by the static mesh component. */
 	void InitResources();
 
+	/** Precache all PSOs which can be used by the static mesh component */
+	virtual void PrecachePSOs() override;
+
 #if WITH_EDITOR
 	/** Update the vertex override colors */
 	void PrivateFixupOverrideColors();
diff --git a/Engine/Source/Runtime/Engine/Private/Components/MaterialBillboardComponent.cpp b/Engine/Source/Runtime/Engine/Private/Components/MaterialBillboardComponent.cpp
--- a/Engine/Source/Runtime/Engine/Private/Components/MaterialBillboardComponent.cpp
+++ b/Engine/Source/Runtime/Engine/Private/Components/MaterialBillboardComponent.cpp
@@ -375,6 +375,28 @@ UMaterialInterface* UMaterialBillboardComponent::GetMaterial(int32 Index) const
 	return ResultMI;
 }
 
+void UMaterialBillboardComponent::PostLoad()
+{
+	Super::PostLoad();
+
+	FPSOPrecacheParams PrecachePSOParams;
+	SetupPrecachePSOParams(PrecachePSOParams);
+	PrecachePSOParams.PrimitiveType = PT_TriangleStrip;
+	PrecachePSOParams.bDisableBackFaceCulling = true;
+
+	const FVertexFactoryType* VFType = &FLocalVertexFactory::StaticType;
+
+	TArray<UMaterialInterface*> UsedMaterials;
+	GetUsedMaterials(UsedMaterials, false);
+	for (UMaterialInterface* MaterialInterface : UsedMaterials)
+	{
+		if (MaterialInterface)
+		{
+			MaterialInterface->PrecachePSOs(VFType, PrecachePSOParams);
+		}
+	}
+}
+
 void UMaterialBillboardComponent::SetMaterial(int32 ElementIndex, class UMaterialInterface* Material)
 {
 	if (Elements.IsValidIndex(ElementIndex))
diff --git a/Engine/Source/Runtime/Engine/Private/Components/MeshComponent.cpp b/Engine/Source/Runtime/Engine/Private/Components/MeshComponent.cpp
--- a/Engine/Source/Runtime/Engine/Private/Components/MeshComponent.cpp
+++ b/Engine/Source/Runtime/Engine/Private/Components/MeshComponent.cpp
@@ -88,6 +88,8 @@ void UMeshComponent::SetMaterial(int32 ElementIndex, UMaterialInterface* Materia
 				BodyInst->UpdatePhysicalMaterials();
 			}
 
+			PrecachePSOs();
+
 #if WITH_EDITOR
 			// Static Lighting is updated when compilation finishes
 			if (!IsCompiling())
diff --git a/Engine/Source/Runtime/Engine/Private/Components/PrimitiveComponent.cpp b/Engine/Source/Runtime/Engine/Private/Components/PrimitiveComponent.cpp
--- a/Engine/Source/Runtime/Engine/Private/Components/PrimitiveComponent.cpp
+++ b/Engine/Source/Runtime/Engine/Private/Components/PrimitiveComponent.cpp
@@ -4242,6 +4242,18 @@ void UPrimitiveComponent::SetLastRenderTime(float InLastRenderTime)
 	}
 }
 
+void UPrimitiveComponent::SetupPrecachePSOParams(FPSOPrecacheParams& Params)
+{
+	Params.bRenderInMainPass = bRenderInMainPass;
+	Params.bRenderInDepthPass = bRenderInDepthPass;
+	Params.bStaticLighting = HasStaticLighting();
+	Params.bAffectDynamicIndirectLighting = bAffectDynamicIndirectLighting;
+	Params.bCastShadow = CastShadow;
+	Params.bCastShadowAsTwoSided = bCastShadowAsTwoSided;
+	Params.SetMobility(Mobility);	
+	Params.SetStencilWriteMask(FRendererStencilMaskEvaluation::ToStencilMask(CustomDepthStencilWriteMask));
+}
+
 #if WITH_EDITOR
 const bool UPrimitiveComponent::ShouldGenerateAutoLOD(const int32 HierarchicalLevelIndex) const
 {	
diff --git a/Engine/Source/Runtime/Engine/Private/Components/SplineMeshComponent.cpp b/Engine/Source/Runtime/Engine/Private/Components/SplineMeshComponent.cpp
--- a/Engine/Source/Runtime/Engine/Private/Components/SplineMeshComponent.cpp
+++ b/Engine/Source/Runtime/Engine/Private/Components/SplineMeshComponent.cpp
@@ -115,6 +115,7 @@ IMPLEMENT_VERTEX_FACTORY_TYPE(FSplineMeshVertexFactory, "/Engine/Private/LocalVe
 	| EVertexFactoryFlags::SupportsPrecisePrevWorldPos
 	| EVertexFactoryFlags::SupportsPositionOnly
 	| EVertexFactoryFlags::SupportsPrimitiveIdStream
+	| EVertexFactoryFlags::SupportsPSOPrecaching
 );
 
 //////////////////////////////////////////////////////////////////////////
@@ -578,6 +579,39 @@ bool USplineMeshComponent::Modify(bool bAlwaysMarkDirty)
 }
 #endif
 
+void USplineMeshComponent::PrecachePSOs()
+{
+	if (GetStaticMesh() == nullptr || GetStaticMesh()->GetRenderData() == nullptr || !FApp::CanEverRender())
+	{
+		return;
+	}
+
+	bool bAnySectionCastsShadows = false;
+	TArray<int16, TInlineAllocator<2>> UsedMaterialIndices;
+	for (FStaticMeshLODResources& LODRenderData : GetStaticMesh()->GetRenderData()->LODResources)
+	{
+		for (FStaticMeshSection& RenderSection : LODRenderData.Sections)
+		{
+			UsedMaterialIndices.AddUnique(RenderSection.MaterialIndex);
+			bAnySectionCastsShadows |= RenderSection.bCastShadow;
+		}
+	}
+
+	FPSOPrecacheParams PrecachePSOParams;
+	SetupPrecachePSOParams(PrecachePSOParams);
+	PrecachePSOParams.bCastShadow = bAnySectionCastsShadows;
+	PrecachePSOParams.bReverseCulling = bReverseCulling;
+
+	for (uint16 MaterialIndex : UsedMaterialIndices)
+	{
+		UMaterialInterface* MaterialInterface = GetMaterial(MaterialIndex);
+		if (MaterialInterface)
+		{
+			MaterialInterface->PrecachePSOs(&FSplineMeshVertexFactory::StaticType, PrecachePSOParams);
+		}
+	}
+}
+
 FPrimitiveSceneProxy* USplineMeshComponent::CreateSceneProxy()
 {
 	// Verify that the mesh is valid before using it.
diff --git a/Engine/Source/Runtime/Engine/Private/Components/StaticMeshComponent.cpp b/Engine/Source/Runtime/Engine/Private/Components/StaticMeshComponent.cpp
--- a/Engine/Source/Runtime/Engine/Private/Components/StaticMeshComponent.cpp
+++ b/Engine/Source/Runtime/Engine/Private/Components/StaticMeshComponent.cpp
@@ -49,6 +49,7 @@
 #include "HAL/IConsoleManager.h"
 #include "Algo/AllOf.h"
 #include "Algo/Transform.h"
+#include "PipelineStateCache.h"
 
 #define LOCTEXT_NAMESPACE "StaticMeshComponent"
 
@@ -771,6 +772,9 @@ void UStaticMeshComponent::OnRegister()
 #endif //WITH_EDITORONLY_DATA
 
 	Super::OnRegister();
+
+	// World transform might have changes causing negative determinant which changes the culling mode
+	PrecachePSOs();
 }
 
 void UStaticMeshComponent::OnUnregister()
@@ -1530,6 +1534,44 @@ void UStaticMeshComponent::InitResources()
 	}
 }
 
+void UStaticMeshComponent::PrecachePSOs()
+{
+	if (StaticMesh == nullptr || StaticMesh->GetRenderData() == nullptr || !FApp::CanEverRender() || !PipelineStateCache::IsPSOPrecachingEnabled())
+	{
+		return;
+	}
+
+	bool bAnySectionCastsShadows = false;
+	TArray<int16, TInlineAllocator<2>> UsedMaterialIndices;
+	for (FStaticMeshLODResources& LODRenderData : StaticMesh->GetRenderData()->LODResources)
+	{
+		for (FStaticMeshSection& RenderSection : LODRenderData.Sections)
+		{
+			UsedMaterialIndices.AddUnique(RenderSection.MaterialIndex);
+			bAnySectionCastsShadows |= RenderSection.bCastShadow;
+		}
+	}
+
+	bool bIsLocalToWorldDeterminantNegative = GetRenderMatrix().Determinant() < 0;
+
+	FPSOPrecacheParams PrecachePSOParams;
+	SetupPrecachePSOParams(PrecachePSOParams);
+	PrecachePSOParams.bCastShadow = bAnySectionCastsShadows;
+	PrecachePSOParams.bReverseCulling = bReverseCulling || bIsLocalToWorldDeterminantNegative;
+	PrecachePSOParams.bForceLODModel = ForcedLodModel > 0;
+
+	const FVertexFactoryType* VFType = ShouldCreateNaniteProxy() ? &Nanite::FVertexFactory::StaticType :  &FLocalVertexFactory::StaticType;
+
+	for (uint16 MaterialIndex : UsedMaterialIndices)
+	{
+		UMaterialInterface* MaterialInterface = GetMaterial(MaterialIndex);
+		if (MaterialInterface)
+		{
+			MaterialInterface->PrecachePSOs(VFType, PrecachePSOParams);
+		}
+	}
+}
+
 #if WITH_EDITOR
 void UStaticMeshComponent::PrivateFixupOverrideColors()
 {
@@ -1692,7 +1734,8 @@ bool UStaticMeshComponent::ShouldCreateNaniteProxy() const
 	const bool bAllowNanite = true;
 #endif
 
-	return bAllowNanite && UseNanite(GetScene()->GetShaderPlatform()) && HasValidNaniteData();
+	EShaderPlatform ShaderPlatform = GetScene() ? GetScene()->GetShaderPlatform() : GMaxRHIShaderPlatform;
+	return bAllowNanite && UseNanite(ShaderPlatform) && HasValidNaniteData();
 }
 
 void UStaticMeshComponent::ReleaseResources()
@@ -2071,6 +2114,9 @@ void UStaticMeshComponent::PostLoad()
 
 	// Initialize the resources for the freshly loaded component.
 	InitResources();
+
+	// Precache PSOs for the used materials
+	PrecachePSOs();
 }
 
 bool UStaticMeshComponent::IsPostLoadThreadSafe() const
@@ -2181,6 +2227,9 @@ bool UStaticMeshComponent::SetStaticMesh(UStaticMesh* NewMesh)
 	// Mark cached material parameter names dirty
 	MarkCachedMaterialParameterNameIndicesDirty();
 
+	// Precache the PSOs
+	PrecachePSOs();
+
 #if WITH_EDITOR
 	// Broadcast that the static mesh has changed
 	OnStaticMeshChangedEvent.Broadcast(this);
diff --git a/Engine/Source/Runtime/Engine/Private/Components/TextRenderComponent.cpp b/Engine/Source/Runtime/Engine/Private/Components/TextRenderComponent.cpp
--- a/Engine/Source/Runtime/Engine/Private/Components/TextRenderComponent.cpp
+++ b/Engine/Source/Runtime/Engine/Private/Components/TextRenderComponent.cpp
@@ -1360,6 +1360,12 @@ void UTextRenderComponent::PostLoad()
 		}
 	}
 
+	if (TextMaterial)
+	{
+		FPSOPrecacheParams PrecachePSOParams;
+		TextMaterial->PrecachePSOs(&FLocalVertexFactory::StaticType, PrecachePSOParams);
+	}
+
 	Super::PostLoad();
 }
 
diff --git a/Engine/Source/Runtime/Engine/Private/InstancedStaticMesh.cpp b/Engine/Source/Runtime/Engine/Private/InstancedStaticMesh.cpp
--- a/Engine/Source/Runtime/Engine/Private/InstancedStaticMesh.cpp
+++ b/Engine/Source/Runtime/Engine/Private/InstancedStaticMesh.cpp
@@ -31,6 +31,7 @@
 #include "NaniteSceneProxy.h"
 #include "MaterialCachedData.h"
 #include "Collision.h"
+#include "PipelineStateCache.h"
 
 #include "Elements/Framework/EngineElementsLibrary.h"
 #include "Elements/SMInstance/SMInstanceElementData.h"
@@ -1013,7 +1014,7 @@ IMPLEMENT_VERTEX_FACTORY_TYPE(FInstancedStaticMeshVertexFactory,"/Engine/Private
 	| EVertexFactoryFlags::SupportsStaticLighting
 	| EVertexFactoryFlags::SupportsDynamicLighting
 	| EVertexFactoryFlags::SupportsPrecisePrevWorldPos
-	| EVertexFactoryFlags::SupportsPositionOnly
+	| (!ALLOW_DITHERED_LOD_FOR_INSTANCED_STATIC_MESHES ? EVertexFactoryFlags::SupportsPositionOnly : EVertexFactoryFlags::None)
 	| EVertexFactoryFlags::SupportsCachingMeshDrawCommands
 	| EVertexFactoryFlags::SupportsRayTracing
 	| EVertexFactoryFlags::SupportsRayTracingDynamicGeometry
@@ -4478,6 +4479,41 @@ void UInstancedStaticMeshComponent::PostLoad()
 	OnPostLoadPerInstanceData();
 }
 
+void UInstancedStaticMeshComponent::PrecachePSOs()
+{	
+	if (GetStaticMesh() == nullptr || GetStaticMesh()->GetRenderData() == nullptr || !FApp::CanEverRender() || !PipelineStateCache::IsPSOPrecachingEnabled())
+	{
+		return;
+	}
+
+	bool bAnySectionCastsShadows = false;
+	TArray<int16, TInlineAllocator<2>> UsedMaterialIndices;
+	for (FStaticMeshLODResources& LODRenderData : GetStaticMesh()->GetRenderData()->LODResources)
+	{
+		for (FStaticMeshSection& RenderSection : LODRenderData.Sections)
+		{
+			UsedMaterialIndices.AddUnique(RenderSection.MaterialIndex);
+			bAnySectionCastsShadows |= RenderSection.bCastShadow;
+		}
+	}
+
+	FPSOPrecacheParams PrecachePSOParams;
+	SetupPrecachePSOParams(PrecachePSOParams);
+	PrecachePSOParams.bCastShadow = bAnySectionCastsShadows;
+	PrecachePSOParams.bReverseCulling = bReverseCulling;
+	
+	const FVertexFactoryType* VFType = ShouldCreateNaniteProxy() ? &Nanite::FVertexFactory::StaticType : &FInstancedStaticMeshVertexFactory::StaticType;
+
+	for (uint16 MaterialIndex : UsedMaterialIndices)
+	{
+		UMaterialInterface* MaterialInterface = GetMaterial(MaterialIndex);
+		if (MaterialInterface)
+		{
+			MaterialInterface->PrecachePSOs(VFType, PrecachePSOParams);
+		}
+	}
+}
+
 void UInstancedStaticMeshComponent::OnPostLoadPerInstanceData()
 {
 	TRACE_CPUPROFILER_EVENT_SCOPE(UInstancedStaticMeshComponent::OnPostLoadPerInstanceData);
diff --git a/Engine/Source/Runtime/Landscape/Private/Landscape.cpp b/Engine/Source/Runtime/Landscape/Private/Landscape.cpp
--- a/Engine/Source/Runtime/Landscape/Private/Landscape.cpp
+++ b/Engine/Source/Runtime/Landscape/Private/Landscape.cpp
@@ -902,11 +902,49 @@ void ULandscapeComponent::UpdatedSharedPropertiesFromActor()
 	UpdateNavigationRelevance();
 	UpdateRejectNavmeshUnderneath();
 }
+#endif // WITH_EDITOR
 
 void ULandscapeComponent::PostLoad()
 {
 	Super::PostLoad();
 
+	{
+		TArray<UMaterialInterface*> Materials;
+		bool bGetDebugMaterials = false;
+		GetUsedMaterials(Materials, bGetDebugMaterials);
+
+		FPSOPrecacheParams PrecachePSOParams;
+		SetupPrecachePSOParams(PrecachePSOParams);
+
+		TArray<const FVertexFactoryType*, TInlineAllocator<2>> VertexFactoryTypes;
+		if (!XYOffsetmapTexture)
+		{
+			VertexFactoryTypes.Add(&FLandscapeVertexFactory::StaticType);
+		}
+		else
+		{
+			VertexFactoryTypes.Add(&FLandscapeXYOffsetVertexFactory::StaticType);
+		}
+
+		// we need the fixed grid vertex factory for both virtual texturing and grass
+		ERHIFeatureLevel::Type FeatureLevel = GMaxRHIFeatureLevel;
+		bool bNeedsFixedGridVertexFactory = UseVirtualTexturing(FeatureLevel);
+		if (bNeedsFixedGridVertexFactory)
+		{
+			VertexFactoryTypes.Add(&FLandscapeFixedGridVertexFactory::StaticType);
+		}
+
+		for (UMaterialInterface* MaterialInterface : Materials)
+		{
+			if (MaterialInterface)
+			{
+				MaterialInterface->ConditionalPostLoad();
+				MaterialInterface->PrecachePSOs(VertexFactoryTypes, PrecachePSOParams);
+			}
+		}
+	}
+
+#if WITH_EDITOR
 	ALandscapeProxy* LandscapeProxy = GetLandscapeProxy();
 	if (ensure(LandscapeProxy))
 	{
@@ -927,7 +965,6 @@ void ULandscapeComponent::PostLoad()
 		}
 	}
 
-#if WITH_EDITOR
 	if (GIsEditor && !HasAnyFlags(RF_ClassDefaultObject))
 	{
 		// This is to ensure that component relative location is exact section base offset value
@@ -993,7 +1030,6 @@ void ULandscapeComponent::PostLoad()
 		}
 		PRAGMA_ENABLE_DEPRECATION_WARNINGS;
 	}
-#endif // WITH_EDITOR
 
 #if WITH_EDITORONLY_DATA
 	// Handle old MaterialInstance
@@ -1003,17 +1039,14 @@ void ULandscapeComponent::PostLoad()
 		MaterialInstances.Add(MaterialInstance_DEPRECATED);
 		MaterialInstance_DEPRECATED = nullptr;
 
-#if WITH_EDITOR
 		if (GIsEditor && MaterialInstances.Num() > 0 && MaterialInstances[0] != nullptr)
 		{
 			MaterialInstances[0]->ConditionalPostLoad();
 			UpdateMaterialInstances();
 		}
-#endif // WITH_EDITOR
 	}
 #endif
 
-#if WITH_EDITOR
 	auto ReparentObject = [this](UObject* Object)
 	{
 		if (Object && !Object->HasAllFlags(RF_Public | RF_Standalone) && (Object->GetOuter() != GetOuter()) && (Object->GetOutermost() == GetOutermost()))
@@ -1080,7 +1113,6 @@ void ULandscapeComponent::PostLoad()
 			MobileCombinationMaterialInstance = Cast<UMaterialInstance>(MobileCombinationMaterialInstance->Parent);
 		}
 	}
-#endif
 
 #if !UE_BUILD_SHIPPING
 	// This will fix the data in case there is mismatch between save of asset/maps
@@ -1117,7 +1149,6 @@ void ULandscapeComponent::PostLoad()
 	}
 #endif // UE_BUILD_SHIPPING
 
-#if WITH_EDITOR
 	if (GIsEditor && !HasAnyFlags(RF_ClassDefaultObject))
 	{
 		// Move the MICs and Textures back to the Package if they're currently in the level
@@ -1139,7 +1170,6 @@ void ULandscapeComponent::PostLoad()
 			}
 		}
 	}
-#endif
 
 #if !UE_BUILD_SHIPPING
 	if (MobileCombinationMaterialInstances.Num() == 0)
@@ -1169,10 +1199,11 @@ void ULandscapeComponent::PostLoad()
 	}
 
 	GrassData->ConditionalDiscardDataOnLoad();
-}
 
 #endif // WITH_EDITOR
 
+}
+
 #if WITH_EDITORONLY_DATA
 TArray<ALandscapeProxy*> ALandscapeProxy::LandscapeProxies;
 #endif
diff --git a/Engine/Source/Runtime/Landscape/Private/LandscapeGrass.cpp b/Engine/Source/Runtime/Landscape/Private/LandscapeGrass.cpp
--- a/Engine/Source/Runtime/Landscape/Private/LandscapeGrass.cpp
+++ b/Engine/Source/Runtime/Landscape/Private/LandscapeGrass.cpp
@@ -3039,6 +3039,8 @@ void ALandscapeProxy::UpdateGrass(const TArray<FVector>& Cameras, int32& InOutNu
 										GrassInstancedStaticMeshComponent->NaniteOverrideMaterials = GrassVariety.NaniteOverrideMaterials;
 										GrassInstancedStaticMeshComponent->bEvaluateWorldPositionOffset = true;
 
+										GrassInstancedStaticMeshComponent->PrecachePSOs();
+
 										const FMeshMapBuildData* MeshMapBuildData = Component->GetMeshMapBuildData();
 
 										if (GrassVariety.bUseLandscapeLightmap
