Landscape: Fixed GetCompleteBounds calculation in world partition worlds when the landscape is composed of only a single streaming proxy.

GetCompleteBounds would skip the ALandscapeProxy which is equal to the main ALandscape and for a landscape with a single streaming proxy this would skip everything and return an empty bounds.

diff --git a/Engine/Source/Runtime/Landscape/Private/Landscape.cpp b/Engine/Source/Runtime/Landscape/Private/Landscape.cpp
--- a/Engine/Source/Runtime/Landscape/Private/Landscape.cpp
+++ b/Engine/Source/Runtime/Landscape/Private/Landscape.cpp
@@ -4571,34 +4571,27 @@ FBox ULandscapeInfo::GetCompleteBounds() const
 
 	FBox Bounds(EForceInit::ForceInit);
 
-	if (UWorldPartition* WorldPartition = Landscape->GetWorld()->GetWorldPartition())
+	FWorldPartitionHelpers::ForEachActorDesc<ALandscapeProxy>(Landscape->GetWorld()->GetWorldPartition(), [this, &Bounds, Landscape](const FWorldPartitionActorDesc* ActorDesc)
 	{
-		FWorldPartitionHelpers::ForEachActorDesc<ALandscapeProxy>(WorldPartition, [this, &Bounds, Landscape](const FWorldPartitionActorDesc* ActorDesc)
+		FLandscapeActorDesc* LandscapeActorDesc = (FLandscapeActorDesc*)ActorDesc;
+
+		if (LandscapeActorDesc->GridGuid == LandscapeGuid)
 		{
-			FLandscapeActorDesc* LandscapeActorDesc = (FLandscapeActorDesc*)ActorDesc;
+			ALandscapeProxy* LandscapeProxy = Cast<ALandscapeProxy>(ActorDesc->GetActor());
 
-			if (LandscapeActorDesc->GridGuid == LandscapeGuid)
+			if (LandscapeProxy)
 			{
-				ALandscapeProxy* LandscapeProxy = Cast<ALandscapeProxy>(ActorDesc->GetActor());
-
-				// Skip owning landscape actor
-				if (LandscapeProxy != Landscape)
-				{
-					if (LandscapeProxy)
-					{
-						// Prioritize loaded bounds, as the bounds in the actor desc might not be up-to-date
-						LandscapeInfoBoundsHelper::AccumulateBounds(LandscapeProxy, Bounds);
-					}
-					else
-					{
-						Bounds += ActorDesc->GetBounds();
-					}
-				}
+				// Prioritize loaded bounds, as the bounds in the actor desc might not be up-to-date
+				LandscapeInfoBoundsHelper::AccumulateBounds(LandscapeProxy, Bounds);
 			}
+			else
+			{
+				Bounds += ActorDesc->GetBounds();
+			}
+		}
 
-			return true;
-		});
-	}
+		return true;
+	});
 
 	return Bounds;
 }
