Landscape : fixed the Locked state of an edit layer that would not dirty the actor. This is a change that is persisted and therefore needs to be saved

diff --git a/Engine/Source/Runtime/Landscape/Private/LandscapeEditLayers.cpp b/Engine/Source/Runtime/Landscape/Private/LandscapeEditLayers.cpp
--- a/Engine/Source/Runtime/Landscape/Private/LandscapeEditLayers.cpp
+++ b/Engine/Source/Runtime/Landscape/Private/LandscapeEditLayers.cpp
@@ -8558,7 +8558,7 @@ void ALandscape::SetLayerLocked(int32 InLayerIndex, bool bLocked)
 		return;
 	}
 
-	Modify(false);
+	Modify();
 	Layer->bLocked = bLocked;
 }
 
