Landscape : Fix scalling of copy / paste gizmo

diff --git a/Engine/Source/Runtime/Landscape/Classes/LandscapeGizmoActiveActor.h b/Engine/Source/Runtime/Landscape/Classes/LandscapeGizmoActiveActor.h
--- a/Engine/Source/Runtime/Landscape/Classes/LandscapeGizmoActiveActor.h
+++ b/Engine/Source/Runtime/Landscape/Classes/LandscapeGizmoActiveActor.h
@@ -122,6 +122,7 @@ public:
 
 	virtual void EditorApplyTranslation(const FVector& DeltaTranslation, bool bAltDown, bool bShiftDown, bool bCtrlDown) override;
 	virtual void EditorApplyRotation(const FRotator& DeltaRotation, bool bAltDown, bool bShiftDown, bool bCtrlDown) override;
+	virtual void EditorApplyScale(const FVector& DeltaScale, const FVector* PivotLocation, bool bAltDown, bool bShiftDown, bool bCtrlDown) override;
 
 	/**
 	 * Whenever the decal actor has moved:
diff --git a/Engine/Source/Runtime/Landscape/Private/LandscapeGizmoActor.cpp b/Engine/Source/Runtime/Landscape/Private/LandscapeGizmoActor.cpp
--- a/Engine/Source/Runtime/Landscape/Private/LandscapeGizmoActor.cpp
+++ b/Engine/Source/Runtime/Landscape/Private/LandscapeGizmoActor.cpp
@@ -639,6 +639,12 @@ void ALandscapeGizmoActiveActor::EditorApplyRotation(const FRotator& DeltaRotati
 	ReregisterAllComponents();
 }
 
+void ALandscapeGizmoActiveActor::EditorApplyScale(const FVector& DeltaScale, const FVector* PivotLocation, bool bAltDown, bool bShiftDown, bool bCtrlDown)
+{
+	Super::EditorApplyScale(DeltaScale, PivotLocation, bAltDown, bShiftDown, bCtrlDown);
+	ReregisterAllComponents();
+} 
+
 ALandscapeGizmoActor* ALandscapeGizmoActiveActor::SpawnGizmoActor()
 {
 	// ALandscapeGizmoActor is history for ALandscapeGizmoActiveActor
