Change landscape info Proxy pointers to be TWeakObjectPtr, so they won't act as GC references

diff --git a/Engine/Source/Editor/DetailCustomizations/Private/LandscapeProxyUIDetails.cpp b/Engine/Source/Editor/DetailCustomizations/Private/LandscapeProxyUIDetails.cpp
--- a/Engine/Source/Editor/DetailCustomizations/Private/LandscapeProxyUIDetails.cpp
+++ b/Engine/Source/Editor/DetailCustomizations/Private/LandscapeProxyUIDetails.cpp
@@ -124,7 +124,7 @@ void FLandscapeProxyUIDetails::CustomizeDetails( IDetailLayoutBuilder& DetailBui
 					GenerateTextWidget(FText::Format(LOCTEXT("LandscapeResolutionValue", "{0} x {1}"), Size.X+1, Size.Y+1), true)
 				];
 
-				int32 LandscapeCount = LandscapeInfo->Proxies.Num() + (LandscapeInfo->LandscapeActor.Get() ? 1 : 0);
+				int32 LandscapeCount = LandscapeInfo->StreamingProxies.Num() + (LandscapeInfo->LandscapeActor.Get() ? 1 : 0);
 				RowDisplayText = LOCTEXT("LandscapeCount", "Landscape Count");
 				CategoryBuilder.AddCustomRow(RowDisplayText)
 				.RowTag(TEXT("LandscapeCount"))
diff --git a/Engine/Source/Editor/DetailCustomizations/Private/LandscapeUIDetails.cpp b/Engine/Source/Editor/DetailCustomizations/Private/LandscapeUIDetails.cpp
--- a/Engine/Source/Editor/DetailCustomizations/Private/LandscapeUIDetails.cpp
+++ b/Engine/Source/Editor/DetailCustomizations/Private/LandscapeUIDetails.cpp
@@ -124,10 +124,13 @@ void FLandscapeUIDetails::CustomizeDetails( IDetailLayoutBuilder& DetailBuilder
 						ULandscapeInfo* LandscapeInfo = Landscape->GetLandscapeInfo();
 						if (LandscapeInfo != nullptr)
 						{
-							for (ALandscapeProxy* Proxy : LandscapeInfo->Proxies)
+							for (TWeakObjectPtr<ALandscapeProxy> ProxyPtr : LandscapeInfo->StreamingProxies)
 							{
-								Proxy->UpdateNaniteRepresentation();
-								Proxy->UpdateRenderingMethod();
+								if (ALandscapeProxy* Proxy = ProxyPtr.Get())
+								{
+									Proxy->UpdateNaniteRepresentation();
+									Proxy->UpdateRenderingMethod();
+								}
 							}
 						}
 					}
diff --git a/Engine/Source/Editor/UnrealEd/Private/Commandlets/WorldPartitionConvertCommandlet.cpp b/Engine/Source/Editor/UnrealEd/Private/Commandlets/WorldPartitionConvertCommandlet.cpp
--- a/Engine/Source/Editor/UnrealEd/Private/Commandlets/WorldPartitionConvertCommandlet.cpp
+++ b/Engine/Source/Editor/UnrealEd/Private/Commandlets/WorldPartitionConvertCommandlet.cpp
@@ -974,15 +974,16 @@ int32 UWorldPartitionConvertCommandlet::Main(const FString& Params)
 		if (!LandscapeInfo->LandscapeActor.Get())
 		{
 			// Use the first proxy as the landscape template
-			ALandscapeProxy* FirstProxy = LandscapeInfo->Proxies[0];
+			if (ALandscapeProxy* FirstProxy = LandscapeInfo->StreamingProxies[0].Get())
+			{
+				FActorSpawnParameters SpawnParams;
+				FTransform LandscapeTransform = FirstProxy->LandscapeActorToWorld();
+				ALandscape* NewLandscape = MainWorld->SpawnActor<ALandscape>(ALandscape::StaticClass(), LandscapeTransform, SpawnParams);
 
-			FActorSpawnParameters SpawnParams;
-			FTransform LandscapeTransform = FirstProxy->LandscapeActorToWorld();
-			ALandscape* NewLandscape = MainWorld->SpawnActor<ALandscape>(ALandscape::StaticClass(), LandscapeTransform, SpawnParams);
-			
-			NewLandscape->GetSharedProperties(FirstProxy);
+				NewLandscape->GetSharedProperties(FirstProxy);
 
-			LandscapeInfo->RegisterActor(NewLandscape);
+				LandscapeInfo->RegisterActor(NewLandscape);
+			}
 		}
 
 		auto MoveControlPointToNewSplineActor = [&NewSplineActors, LandscapeInfo](ULandscapeSplineControlPoint* ControlPoint)
diff --git a/Engine/Source/Runtime/Landscape/Classes/LandscapeInfo.h b/Engine/Source/Runtime/Landscape/Classes/LandscapeInfo.h
--- a/Engine/Source/Runtime/Landscape/Classes/LandscapeInfo.h
+++ b/Engine/Source/Runtime/Landscape/Classes/LandscapeInfo.h
@@ -157,9 +157,12 @@ public:
 	TMap<FIntPoint, FLandscapeAddCollision> XYtoAddCollisionMap;
 #endif // WITH_EDITORONLY_DATA
 
-	UPROPERTY()
+	UE_DEPRECATED(5.1, "This property has been deprecated, please use the StreamingProxies property instead")
 	TArray<TObjectPtr<ALandscapeStreamingProxy>> Proxies;
 
+	UPROPERTY()
+	TArray<TWeakObjectPtr<ALandscapeStreamingProxy>> StreamingProxies;
+
 private:
 #if WITH_EDITORONLY_DATA
 	UPROPERTY()
diff --git a/Engine/Source/Runtime/Landscape/Private/Landscape.cpp b/Engine/Source/Runtime/Landscape/Private/Landscape.cpp
--- a/Engine/Source/Runtime/Landscape/Private/Landscape.cpp
+++ b/Engine/Source/Runtime/Landscape/Private/Landscape.cpp
@@ -3698,7 +3698,7 @@ void ALandscape::PreSave(FObjectPreSaveContext ObjectSaveContext)
 	//ULandscapeInfo* Info = GetLandscapeInfo();
 	//if (GIsEditor && Info && !ObjectSaveContext.IsProceduralSave())
 	//{
-	//	for (TSet<ALandscapeProxy*>::TIterator It(Info->Proxies); It; ++It)
+	//	for (TSet<ALandscapeProxy*>::TIterator It(Info->StreamingProxies); It; ++It)
 	//	{
 	//		ALandscapeProxy* Proxy = *It;
 	//		if (!ensure(Proxy->LandscapeActor == this))
@@ -3864,16 +3864,15 @@ ALandscapeProxy* ULandscapeInfo::GetLandscapeProxy() const
 	}
 
 	// prefer current level proxy 
-	ALandscapeProxy* Proxy = GetCurrentLevelLandscapeProxy(true);
-	if (Proxy != nullptr)
+	if (ALandscapeProxy* Proxy = GetCurrentLevelLandscapeProxy(true))
 	{
 		return Proxy;
 	}
 
 	// any proxy in the world
-	for (auto It = Proxies.CreateConstIterator(); It; ++It)
+	for (TWeakObjectPtr<ALandscapeStreamingProxy> ProxyPtr : StreamingProxies)
 	{
-		Proxy = (*It);
+		ALandscapeStreamingProxy* Proxy = ProxyPtr.Get();
 		if (Proxy != nullptr &&
 			Proxy->GetRootComponent()->IsRegistered())
 		{
@@ -3890,7 +3889,7 @@ void ULandscapeInfo::Reset()
 {
 	LandscapeActor.Reset();
 
-	Proxies.Empty();
+	StreamingProxies.Empty();
 	XYtoComponentMap.Empty();
 	XYtoAddCollisionMap.Empty();
 
@@ -3910,9 +3909,13 @@ void ULandscapeInfo::FixupProxiesTransform(bool bDirty)
 	}
 
 	// Make sure section offset of all proxies is multiple of ALandscapeProxy::ComponentSizeQuads
-	for (auto It = Proxies.CreateConstIterator(); It; ++It)
+	for (TWeakObjectPtr<ALandscapeStreamingProxy> ProxyPtr : StreamingProxies)
 	{
-		ALandscapeProxy* Proxy = *It;
+		ALandscapeProxy* Proxy = ProxyPtr.Get();
+		if (!Proxy)
+		{
+			continue;
+		}
 
 		if (bDirty)
 		{
@@ -3937,9 +3940,14 @@ void ULandscapeInfo::FixupProxiesTransform(bool bDirty)
 
 	FTransform LandscapeTM = Landscape->LandscapeActorToWorld();
 	// Update transformations of all linked landscape proxies
-	for (auto It = Proxies.CreateConstIterator(); It; ++It)
+	for (TWeakObjectPtr<ALandscapeStreamingProxy> ProxyPtr : StreamingProxies)
 	{
-		ALandscapeProxy* Proxy = *It;
+		ALandscapeProxy* Proxy = ProxyPtr.Get();
+		if (!Proxy)
+		{
+			continue;
+		}
+
 		FTransform ProxyRelativeTM(FVector(Proxy->LandscapeSectionOffset));
 		FTransform ProxyTransform = ProxyRelativeTM * LandscapeTM;
 
@@ -4071,15 +4079,17 @@ void ULandscapeInfo::Initialize(UWorld* InWorld, const FGuid& InLandscapeGuid)
 
 void ULandscapeInfo::ForAllLandscapeProxies(TFunctionRef<void(ALandscapeProxy*)> Fn) const
 {
-	ALandscape* Landscape = LandscapeActor.Get();
-	if (Landscape)
+	if (ALandscape* Landscape = LandscapeActor.Get())
 	{
 		Fn(Landscape);
 	}
 
-	for (ALandscapeProxy* LandscapeProxy : Proxies)
+	for (TWeakObjectPtr<ALandscapeStreamingProxy> StreamingProxyPtr : StreamingProxies)
 	{
-		Fn(LandscapeProxy);
+		if (ALandscapeProxy* LandscapeProxy = StreamingProxyPtr.Get())
+		{
+			Fn(LandscapeProxy);
+		}
 	}
 }
 
@@ -4128,12 +4138,15 @@ void ULandscapeInfo::RegisterActor(ALandscapeProxy* Proxy, bool bMapCheck)
 #endif // WITH_EDITOR
 
 			// update proxies reference actor
-			for (ALandscapeStreamingProxy* StreamingProxy : Proxies)
+			for (TWeakObjectPtr<ALandscapeStreamingProxy> StreamingProxyPtr : StreamingProxies)
 			{
-				StreamingProxy->LandscapeActor = LandscapeActor;
+				if (ALandscapeStreamingProxy* StreamingProxy = StreamingProxyPtr.Get())
+				{
+					StreamingProxy->LandscapeActor = LandscapeActor;
 #if WITH_EDITOR
-				StreamingProxy->FixupSharedData(Landscape);
+					StreamingProxy->FixupSharedData(Landscape);
 #endif // WITH_EDITOR
+				}
 			}
 		}
 		else if (LandscapeActor == Landscape)
@@ -4143,8 +4156,17 @@ void ULandscapeInfo::RegisterActor(ALandscapeProxy* Proxy, bool bMapCheck)
 	}
 	else
 	{
-		auto LamdbdaLowerBound = [](ALandscapeProxy* A, ALandscapeProxy* B)
+		auto LamdbdaLowerBound = [](TWeakObjectPtr<ALandscapeProxy> APtr, TWeakObjectPtr<ALandscapeProxy> BPtr)
 		{
+			ALandscapeProxy *A = APtr.Get();
+			ALandscapeProxy *B = BPtr.Get();
+
+			// sort nulls, assuming null < !null
+			if (!A || !B)
+			{
+				return !!B;
+			}
+			
 			FIntPoint SectionBaseA = A->GetSectionBaseOffset();
 			FIntPoint SectionBaseB = B->GetSectionBaseOffset();
 
@@ -4156,12 +4178,15 @@ void ULandscapeInfo::RegisterActor(ALandscapeProxy* Proxy, bool bMapCheck)
 			return SectionBaseA.Y < SectionBaseB.Y;
 		};
 
-		// Insert Proxies in a sorted fashion for generating deterministic results in the Layer system
+		// Insert Proxies in a sorted fashion into the landscape info Proxies list, for generating deterministic results in the Layer system
 		ALandscapeStreamingProxy* StreamingProxy = CastChecked<ALandscapeStreamingProxy>(Proxy);
-		if (!Proxies.Contains(Proxy))
+		TWeakObjectPtr<ALandscapeStreamingProxy> StreamingProxyPtr = StreamingProxy;
+		if (!StreamingProxies.Contains(StreamingProxyPtr))
 		{
-			uint32 InsertIndex = Algo::LowerBound(Proxies, Proxy, LamdbdaLowerBound);
-			Proxies.Insert(StreamingProxy, InsertIndex);
+			// NOTE: if a streaming proxy somehow gets garbage collected without de-registering from the Proxies list, then
+			// this search may return a non-deterministic index because the Proxies list will contain a null
+			uint32 InsertIndex = Algo::LowerBound(StreamingProxies, StreamingProxyPtr, LamdbdaLowerBound);
+			StreamingProxies.Insert(StreamingProxyPtr, InsertIndex);
 		}
 
 		StreamingProxy->LandscapeActor = LandscapeActor;
@@ -4205,15 +4230,18 @@ void ULandscapeInfo::UnregisterActor(ALandscapeProxy* Proxy)
 		}
 
 		// update proxies reference to landscape actor
-		for (ALandscapeStreamingProxy* StreamingProxy : Proxies)
+		for (TWeakObjectPtr<ALandscapeStreamingProxy> StreamingProxyPtr : StreamingProxies)
 		{
-			StreamingProxy->LandscapeActor = LandscapeActor;
+			if (ALandscapeStreamingProxy* StreamingProxy = StreamingProxyPtr.Get())
+			{
+				StreamingProxy->LandscapeActor = LandscapeActor;
+			}
 		}
 	}
 	else
 	{
 		ALandscapeStreamingProxy* StreamingProxy = CastChecked<ALandscapeStreamingProxy>(Proxy);
-		Proxies.Remove(StreamingProxy);
+		StreamingProxies.Remove(StreamingProxy);
 		StreamingProxy->LandscapeActor = nullptr;
 	}
 
diff --git a/Engine/Source/Runtime/Landscape/Private/LandscapeEdit.cpp b/Engine/Source/Runtime/Landscape/Private/LandscapeEdit.cpp
--- a/Engine/Source/Runtime/Landscape/Private/LandscapeEdit.cpp
+++ b/Engine/Source/Runtime/Landscape/Private/LandscapeEdit.cpp
@@ -5414,9 +5414,10 @@ bool ULandscapeInfo::CanDeleteLandscape(FText& OutReason) const
 	int32 UndeletedSplineCount = 0;
 
 	// Check Registered Proxies
-	for (ALandscapeProxy* RegisteredProxy : Proxies)
+	for (TWeakObjectPtr<ALandscapeStreamingProxy> ProxyPtr : StreamingProxies)
 	{
-		if (RegisteredProxy == LandscapeActor)
+		ALandscapeProxy* RegisteredProxy = ProxyPtr.Get();
+		if (!RegisteredProxy || RegisteredProxy == LandscapeActor)
 		{
 			continue;
 		}
@@ -5448,7 +5449,8 @@ bool ULandscapeInfo::CanDeleteLandscape(FText& OutReason) const
 						else
 						{
 							// If Actor is loaded it should be Registered and not pending kill (already accounted for) or pending kill (deleted)
-							check(Proxies.Contains(LandscapeProxy) == IsValidChecked(LandscapeProxy));
+							TWeakObjectPtr<ALandscapeStreamingProxy> StreamingProxyPtr = CastChecked<ALandscapeStreamingProxy>(LandscapeProxy);
+							check(StreamingProxies.Contains(StreamingProxyPtr) == IsValidChecked(LandscapeProxy));
 						}
 					}
 				}
@@ -5697,10 +5699,13 @@ void ALandscape::PostEditChangeProperty(FPropertyChangedEvent& PropertyChangedEv
 		if (bPropagateToProxies)
 		{
 			// Propagate Event to Proxies...
-			for (ALandscapeProxy* Proxy : Info->Proxies)
+			for (TWeakObjectPtr<ALandscapeStreamingProxy> ProxyPtr : Info->StreamingProxies)
 			{
-				Proxy->GetSharedProperties(this);
-				Proxy->PostEditChangeProperty(PropertyChangedEvent);
+				if (ALandscapeProxy* Proxy = ProxyPtr.Get())
+				{
+					Proxy->GetSharedProperties(this);
+					Proxy->PostEditChangeProperty(PropertyChangedEvent);
+				}
 			}
 		}
 
@@ -7048,9 +7053,13 @@ bool ALandscapeProxy::LandscapeExportHeightmapToRenderTarget(UTextureRenderTarge
 	if (InExportLandscapeProxies && (GetLandscapeActor() == this))
 	{
 		ULandscapeInfo* LandscapeInfo = GetLandscapeInfo();
-		for (ALandscapeProxy* Proxy : LandscapeInfo->Proxies)
+
+		for (TWeakObjectPtr<ALandscapeStreamingProxy> ProxyPtr : LandscapeInfo->StreamingProxies)
 		{
-			LandscapeComponentsToExport.Append(Proxy->LandscapeComponents);
+			if (ALandscapeProxy* Proxy = ProxyPtr.Get())
+			{
+				LandscapeComponentsToExport.Append(Proxy->LandscapeComponents);
+			}
 		}
 	}
 
