LandscapeSplineMeshesActor: Tag GridGuid as a UPROPERTY so it gets serialized

Fixes WorldPartitionLandscapeSplineMeshBuilder reruns failed checks

#rnx

diff --git a/Engine/Source/Runtime/Landscape/Classes/LandscapeSplineMeshesActor.h b/Engine/Source/Runtime/Landscape/Classes/LandscapeSplineMeshesActor.h
--- a/Engine/Source/Runtime/Landscape/Classes/LandscapeSplineMeshesActor.h
+++ b/Engine/Source/Runtime/Landscape/Classes/LandscapeSplineMeshesActor.h
@@ -48,6 +48,7 @@ protected:
 
 private:
 #if WITH_EDITORONLY_DATA
+	UPROPERTY()
 	FGuid GridGuid;
 #endif
 };
\ No newline at end of file
