Fixed static analysis warning

#rnx

diff --git a/Engine/Source/Runtime/Landscape/Classes/Landscape.h b/Engine/Source/Runtime/Landscape/Classes/Landscape.h
--- a/Engine/Source/Runtime/Landscape/Classes/Landscape.h
+++ b/Engine/Source/Runtime/Landscape/Classes/Landscape.h
@@ -391,7 +391,7 @@ private:
 	bool PrepareTextureResources(bool bInWaitForStreaming);
 	bool PrepareLayersTextureResources(bool bInWaitForStreaming);
 	bool PrepareLayersTextureResources(const TArray<FLandscapeLayer>& InLayers, bool bInWaitForStreaming);
-	bool PrepareLayersBrushResources(bool bInWaitForStreaming);
+	bool PrepareLayersBrushResources(ERHIFeatureLevel::Type InFeatureLevel, bool bInWaitForStreaming);
 	void InvalidateRVTForTextures(const TSet<UTexture2D*>& InTextures);
 	void PrepareLayersHeightmapsLocalMergeRenderThreadData(const FUpdateLayersContentContext& InUpdateLayersContentContext, const FEditLayersHeightmapMergeParams& InMergeParams, EditLayersHeightmapLocalMerge_RenderThread::FMergeInfo& OutRenderThreadData);
 	void PrepareLayersWeightmapsLocalMergeRenderThreadData(const FUpdateLayersContentContext& InUpdateLayersContentContext, const FEditLayersWeightmapMergeParams& InMergeParams, EditLayersWeightmapLocalMerge_RenderThread::FMergeInfo& OutRenderThreadData);
diff --git a/Engine/Source/Runtime/Landscape/Private/LandscapeEditLayers.cpp b/Engine/Source/Runtime/Landscape/Private/LandscapeEditLayers.cpp
--- a/Engine/Source/Runtime/Landscape/Private/LandscapeEditLayers.cpp
+++ b/Engine/Source/Runtime/Landscape/Private/LandscapeEditLayers.cpp
@@ -3485,7 +3485,7 @@ bool ALandscape::PrepareLayersTextureResources(const TArray<FLandscapeLayer>& In
 	return bIsReady;
 }
 
-bool ALandscape::PrepareLayersBrushResources(bool bInWaitForStreaming)
+bool ALandscape::PrepareLayersBrushResources(ERHIFeatureLevel::Type InFeatureLevel, bool bInWaitForStreaming)
 {
 	TRACE_CPUPROFILER_EVENT_SCOPE(LandscapeLayers_PrepareLayersBrushTextureResources);
 	TSet<UObject*> Dependencies;
@@ -3504,7 +3504,6 @@ bool ALandscape::PrepareLayersBrushResources(bool bInWaitForStreaming)
 	}
 
 	bool bIsReady = true;
-	ERHIFeatureLevel::Type FeatureLevel = GetWorld() ? (ERHIFeatureLevel::Type)GetWorld()->FeatureLevel : GMaxRHIFeatureLevel;
 	for (UObject* Dependency : Dependencies)
 	{
 		// Streamable textures need to be fully streamed in : 
@@ -3516,7 +3515,7 @@ bool ALandscape::PrepareLayersBrushResources(bool bInWaitForStreaming)
 		// Material shaders need to be fully compiled : 
 		if (UMaterialInterface* MaterialInterface = Cast<UMaterialInterface>(Dependency))
 		{
-			if (FMaterialResource* MaterialResource = MaterialInterface->GetMaterialResource(FeatureLevel))
+			if (FMaterialResource* MaterialResource = MaterialInterface->GetMaterialResource(InFeatureLevel))
 			{
 				// Don't early-out because checking for the material resource actually requests the shaders to be loaded so we want to make sure to request them all at once instead of one by one :
 				bIsReady &= IsMaterialResourceCompiled(MaterialResource, bInWaitForStreaming);
@@ -7769,11 +7768,18 @@ void ALandscape::UpdateLayersContent(bool bInWaitForStreaming, bool bInSkipMonit
 	bool bResourcesReady = PrepareTextureResources(bInWaitForStreaming);
 
 	ULandscapeInfo* LandscapeInfo = GetLandscapeInfo();
-	if (LandscapeInfo == nullptr || !CanHaveLayersContent() || !LandscapeInfo->AreAllComponentsRegistered())
+	if ((LandscapeInfo == nullptr) || !CanHaveLayersContent() || !LandscapeInfo->AreAllComponentsRegistered())
 	{
 		return;
 	}
 
+	UWorld* World = GetWorld();
+	check(World != nullptr);
+	ULandscapeSubsystem* LandscapeSubsystem = World->GetSubsystem<ULandscapeSubsystem>();
+	check(LandscapeSubsystem != nullptr);
+	FLandscapeNotificationManager* LandscapeNotificationManager = LandscapeSubsystem->GetNotificationManager();
+	check(LandscapeNotificationManager != nullptr);
+
 	// Make sure Update doesn't dirty Landscape packages when not in Landscape Ed Mode
 	FLandscapeDirtyOnlyInModeScope DirtyOnlyInMode(LandscapeInfo);
 
@@ -7834,11 +7840,8 @@ void ALandscape::UpdateLayersContent(bool bInWaitForStreaming, bool bInSkipMonit
 		return;
 	}
 	
-	FLandscapeNotificationManager* LandscapeNotificationManager = GetWorld()->GetSubsystem<ULandscapeSubsystem>()->GetNotificationManager();
-	
 	// The Edit layers shaders only work on SM5 : cancel any update that might happen when SM5+ shading model is not active :
-	ERHIFeatureLevel::Type FeatureLevel = GetWorld() ? (ERHIFeatureLevel::Type)GetWorld()->FeatureLevel : GMaxRHIFeatureLevel;
-	if (FeatureLevel < ERHIFeatureLevel::SM5)
+	if (World->FeatureLevel < ERHIFeatureLevel::SM5)
 	{
 		if (!InvalidShadingModelNotification.IsValid())
 		{
@@ -7880,7 +7883,7 @@ void ALandscape::UpdateLayersContent(bool bInWaitForStreaming, bool bInSkipMonit
 		WaitingForTexturesNotification.Reset();
 	}
 
-	bResourcesReady &= PrepareLayersBrushResources(bInWaitForStreaming);
+	bResourcesReady &= PrepareLayersBrushResources(World->FeatureLevel, bInWaitForStreaming);
 	if (!bResourcesReady)
 	{
 		if (WaitingForLandscapeBrushResourcesStartTime < 0.0)
