Landscape: Fix truncation warnings, part 1

diff --git a/Engine/Source/Runtime/Landscape/Classes/LandscapeGizmoActiveActor.h b/Engine/Source/Runtime/Landscape/Classes/LandscapeGizmoActiveActor.h
--- a/Engine/Source/Runtime/Landscape/Classes/LandscapeGizmoActiveActor.h
+++ b/Engine/Source/Runtime/Landscape/Classes/LandscapeGizmoActiveActor.h
@@ -170,25 +170,25 @@ public:
 	// @todo document
 	float GetWidth() const
 	{
-		return Width * GetRootComponent()->GetRelativeScale3D().X;
+		return static_cast<float>(Width * GetRootComponent()->GetRelativeScale3D().X);
 	}
 
 	// @todo document
 	float GetHeight() const
 	{
-		return Height * GetRootComponent()->GetRelativeScale3D().Y;
+		return static_cast<float>(Height * GetRootComponent()->GetRelativeScale3D().Y);
 	}
 
 	// @todo document
 	float GetLength() const
 	{
-		return LengthZ * GetRootComponent()->GetRelativeScale3D().Z;
+		return static_cast<float>(LengthZ * GetRootComponent()->GetRelativeScale3D().Z);
 	}
 
 	// @todo document
 	void SetLength(float WorldLength)
 	{
-		LengthZ = WorldLength / GetRootComponent()->GetRelativeScale3D().Z;
+		LengthZ = static_cast<float>(WorldLength / GetRootComponent()->GetRelativeScale3D().Z);
 	}
 
 	static const int32 DataTexSize = 128;
diff --git a/Engine/Source/Runtime/Landscape/Private/Landscape.cpp b/Engine/Source/Runtime/Landscape/Private/Landscape.cpp
--- a/Engine/Source/Runtime/Landscape/Private/Landscape.cpp
+++ b/Engine/Source/Runtime/Landscape/Private/Landscape.cpp
@@ -957,8 +957,8 @@ void ULandscapeComponent::PostLoad()
 
 		// check SectionBaseX/Y are correct
 		const FVector LocalRelativeLocation = GetRelativeLocation();
-		int32 CheckSectionBaseX = FMath::RoundToInt(LocalRelativeLocation.X) + LandscapeProxy->LandscapeSectionOffset.X;
-		int32 CheckSectionBaseY = FMath::RoundToInt(LocalRelativeLocation.Y) + LandscapeProxy->LandscapeSectionOffset.Y;
+		int32 CheckSectionBaseX = FMath::RoundToInt32(LocalRelativeLocation.X) + LandscapeProxy->LandscapeSectionOffset.X;
+		int32 CheckSectionBaseY = FMath::RoundToInt32(LocalRelativeLocation.Y) + LandscapeProxy->LandscapeSectionOffset.Y;
 		if (CheckSectionBaseX != SectionBaseX ||
 			CheckSectionBaseY != SectionBaseY)
 		{
@@ -1120,7 +1120,7 @@ void ULandscapeComponent::PostLoad()
 
 #if !UE_BUILD_SHIPPING
 	// This will fix the data in case there is mismatch between save of asset/maps
-	int8 MaxLOD = FMath::CeilLogTwo(SubsectionSizeQuads + 1) - 1;
+	const int8 MaxLOD = static_cast<int8>(FMath::CeilLogTwo(SubsectionSizeQuads + 1) - 1);
 
 	TArray<ULandscapeMaterialInstanceConstant*> ResolvedMaterials;
 
@@ -2025,8 +2025,8 @@ void ULandscapeComponent::AddDefaultLayerData(const FGuid& InLayerGuid, const TA
 
 			for (ULandscapeComponent* ComponentUsingHeightmap : InComponentsUsingHeightmap)
 			{
-				int32 HeightmapComponentOffsetX = FMath::RoundToInt((float)NewLayerHeightmap->Source.GetSizeX() * ComponentUsingHeightmap->HeightmapScaleBias.Z);
-				int32 HeightmapComponentOffsetY = FMath::RoundToInt((float)NewLayerHeightmap->Source.GetSizeY() * ComponentUsingHeightmap->HeightmapScaleBias.W);
+				int32 HeightmapComponentOffsetX = FMath::RoundToInt32(NewLayerHeightmap->Source.GetSizeX() * ComponentUsingHeightmap->HeightmapScaleBias.Z);
+				int32 HeightmapComponentOffsetY = FMath::RoundToInt32(NewLayerHeightmap->Source.GetSizeY() * ComponentUsingHeightmap->HeightmapScaleBias.W);
 
 				for (int32 SubsectionY = 0; SubsectionY < NumSubsections; SubsectionY++)
 				{
@@ -2053,8 +2053,8 @@ void ULandscapeComponent::AddDefaultLayerData(const FGuid& InLayerGuid, const TA
 								Mip0Data[HeightTexDataIdx].G = HeightValue & 255;
 
 								// Normal with get calculated later
-								Mip0Data[HeightTexDataIdx].B = 0.0f;
-								Mip0Data[HeightTexDataIdx].A = 0.0f;
+								Mip0Data[HeightTexDataIdx].B = 0;
+								Mip0Data[HeightTexDataIdx].A = 0;
 							}
 						}
 					}
@@ -2778,8 +2778,8 @@ void ALandscape::RenderHeightmap(const FTransform& InRenderAreaWorldTransform, c
 			// We get the overall heightmap size via the resource instead of direct GetSizeX/Y calls because apparently
 			// the latter are unreliable while the texture is being built.
 			HeightmapOffset = FIntPoint(
-				FMath::RoundToDouble(Component->HeightmapScaleBias.Z * HeightmapResource->GetSizeX()),
-				FMath::RoundToDouble(Component->HeightmapScaleBias.W * HeightmapResource->GetSizeY()));
+				FMath::RoundToInt32(Component->HeightmapScaleBias.Z * HeightmapResource->GetSizeX()),
+				FMath::RoundToInt32(Component->HeightmapScaleBias.W * HeightmapResource->GetSizeY()));
 		}
 			
 		// When mips are partially loaded, we need to take that into consideration when merging the heightmap :
@@ -4417,9 +4417,9 @@ bool ULandscapeInfo::GetOverlappedComponents(const FTransform& InAreaWorldTransf
 
 	// Indices of the landscape components needed for rendering this area : 
 	FIntRect BoundingIndices;
-	BoundingIndices.Min = FIntPoint(FMath::FloorToInt(LocalExtents.Min.X / ComponentSizeQuads), FMath::FloorToInt(LocalExtents.Min.Y / ComponentSizeQuads));
+	BoundingIndices.Min = FIntPoint(FMath::FloorToInt32(LocalExtents.Min.X / ComponentSizeQuads), FMath::FloorToInt32(LocalExtents.Min.Y / ComponentSizeQuads));
 	// The max here is meant to be an exclusive bound, hence the +1
-	BoundingIndices.Max = FIntPoint(FMath::FloorToInt(LocalExtents.Max.X / ComponentSizeQuads), FMath::FloorToInt(LocalExtents.Max.Y / ComponentSizeQuads)) + FIntPoint(1);
+	BoundingIndices.Max = FIntPoint(FMath::FloorToInt32(LocalExtents.Max.X / ComponentSizeQuads), FMath::FloorToInt32(LocalExtents.Max.Y / ComponentSizeQuads)) + FIntPoint(1);
 
 	// There could be missing components, so the effective area is actually a subset of this area :
 	FIntRect EffectiveBoundingIndices;
@@ -4715,7 +4715,7 @@ void ULandscapeMeshProxyComponent::InitializeForLandscape(ALandscapeProxy* Lands
 
 	if (InProxyLOD != INDEX_NONE)
 	{
-		ProxyLOD = FMath::Clamp<int32>(InProxyLOD, 0, FMath::CeilLogTwo(Landscape->SubsectionSizeQuads + 1) - 1);
+		ProxyLOD = static_cast<int8>(FMath::Clamp<int32>(InProxyLOD, 0, FMath::CeilLogTwo(Landscape->SubsectionSizeQuads + 1) - 1));
 	}
 }
 
@@ -5053,8 +5053,8 @@ void ALandscapeProxy::UpdateGIBakedTextures(bool bBakeAllGITextures)
 						TArray<FColor> Samples;
 						if (FMaterialUtilities::ExportBaseColor(Component, BakeSize, Samples))
 						{
-						int32 AtlasOffsetX = FMath::RoundToInt(Component->HeightmapScaleBias.Z * (float)HeightmapTexture->Source.GetSizeX()) >> 3;
-						int32 AtlasOffsetY = FMath::RoundToInt(Component->HeightmapScaleBias.W * (float)HeightmapTexture->Source.GetSizeY()) >> 3;
+							const int32 AtlasOffsetX = FMath::RoundToInt32(Component->HeightmapScaleBias.Z * HeightmapTexture->Source.GetSizeX()) >> 3;
+							const int32 AtlasOffsetY = FMath::RoundToInt32(Component->HeightmapScaleBias.W * HeightmapTexture->Source.GetSizeY()) >> 3;
 							for (int32 y = 0; y < BakeSize; y++)
 							{
 								FMemory::Memcpy(&AtlasSamples[(y + AtlasOffsetY) * AtlasSize.X + AtlasOffsetX], &Samples[y * BakeSize], sizeof(FColor) * BakeSize);
@@ -5066,7 +5066,7 @@ void ALandscapeProxy::UpdateGIBakedTextures(bool bBakeAllGITextures)
 
 					for (ULandscapeComponent* Component : Info.Components)
 					{
-					Component->BakedTextureMaterialGuid = Info.CombinedStateId;
+						Component->BakedTextureMaterialGuid = Info.CombinedStateId;
 						Component->GIBakedBaseColorTexture = AtlasTexture;
 						Component->MarkRenderStateDirty();
 					}
diff --git a/Engine/Source/Runtime/Landscape/Private/LandscapeConfigHelper.cpp b/Engine/Source/Runtime/Landscape/Private/LandscapeConfigHelper.cpp
--- a/Engine/Source/Runtime/Landscape/Private/LandscapeConfigHelper.cpp
+++ b/Engine/Source/Runtime/Landscape/Private/LandscapeConfigHelper.cpp
@@ -66,7 +66,7 @@ ALandscapeProxy* FLandscapeConfigHelper::FindOrAddLandscapeStreamingProxy(UActor
 
 	auto LandscapeProxyCreated = [InCellCoord, Landscape](APartitionActor* PartitionActor)
 	{
-		const FIntPoint CellLocation(InCellCoord.X * Landscape->GridSize, InCellCoord.Y * Landscape->GridSize);
+		const FIntPoint CellLocation(static_cast<int32>(InCellCoord.X) * Landscape->GridSize, static_cast<int32>(InCellCoord.Y) * Landscape->GridSize);
 
 		ALandscapeProxy* LandscapeProxy = CastChecked<ALandscapeProxy>(PartitionActor);
 		// copy shared properties to this new proxy
@@ -138,7 +138,7 @@ bool FLandscapeConfigHelper::ChangeGridSize(ULandscapeInfo* InLandscapeInfo, uin
 				TMap<int32, UMaterialInterface*>& LODMaterials = ComponentLODMaterials.FindOrAdd(LandscapeComponent);
 				for (int32 LODIndex = 0; LODIndex <= 8; ++LODIndex)
 				{
-					LODMaterials.Add(LODIndex, LandscapeComponent->GetLandscapeMaterial(LODIndex));
+					LODMaterials.Add(LODIndex, LandscapeComponent->GetLandscapeMaterial(static_cast<int8>(LODIndex)));
 				}
 
 				ComponentsToMove.Add(LandscapeComponent);
@@ -195,7 +195,7 @@ bool FLandscapeConfigHelper::ChangeGridSize(ULandscapeInfo* InLandscapeInfo, uin
 
 				TArray<FLandscapePerLODMaterialOverride> PerLODOverrideMaterialsForComponent;
 				TArray<FLandscapePerLODMaterialOverride> PerLODOverrideMaterialsForProxy = LandscapeProxy->GetPerLODOverrideMaterials();
-				for (int32 LODIndex = 0; LODIndex <= 8; ++LODIndex)
+				for (int8 LODIndex = 0; LODIndex <= 8; ++LODIndex)
 				{
 					UMaterialInterface* PreviousLODMaterial = PreviousLandscapeLODMaterials.FindChecked(LODIndex);
 					// If Proxy doesn't differ from Landscape override material there first
diff --git a/Engine/Source/Runtime/Landscape/Private/LandscapeDataAccess.cpp b/Engine/Source/Runtime/Landscape/Private/LandscapeDataAccess.cpp
--- a/Engine/Source/Runtime/Landscape/Private/LandscapeDataAccess.cpp
+++ b/Engine/Source/Runtime/Landscape/Private/LandscapeDataAccess.cpp
@@ -16,8 +16,8 @@ LANDSCAPE_API FLandscapeComponentDataInterface::FLandscapeComponentDataInterface
 	// Offset and stride for this component's data in heightmap texture
 	UTexture2D* HeightMapTexture = Component->GetHeightmap(bWorkOnEditingLayer);
 	HeightmapStride = HeightMapTexture->Source.GetSizeX() >> MipLevel;
-	HeightmapComponentOffsetX = FMath::RoundToInt((float)(HeightMapTexture->Source.GetSizeX() >> MipLevel) * Component->HeightmapScaleBias.Z);
-	HeightmapComponentOffsetY = FMath::RoundToInt((float)(HeightMapTexture->Source.GetSizeY() >> MipLevel) * Component->HeightmapScaleBias.W);
+	HeightmapComponentOffsetX = FMath::RoundToInt32((HeightMapTexture->Source.GetSizeX() >> MipLevel) * Component->HeightmapScaleBias.Z);
+	HeightmapComponentOffsetY = FMath::RoundToInt32((HeightMapTexture->Source.GetSizeY() >> MipLevel) * Component->HeightmapScaleBias.W);
 	HeightmapSubsectionOffset = (Component->SubsectionSizeQuads + 1) >> MipLevel;
 
 	ComponentSizeVerts = (Component->ComponentSizeQuads + 1) >> MipLevel;
diff --git a/Engine/Source/Runtime/Landscape/Private/LandscapeEditResources.cpp b/Engine/Source/Runtime/Landscape/Private/LandscapeEditResources.cpp
--- a/Engine/Source/Runtime/Landscape/Private/LandscapeEditResources.cpp
+++ b/Engine/Source/Runtime/Landscape/Private/LandscapeEditResources.cpp
@@ -28,7 +28,7 @@ void FLandscapeTexture2DResource::InitRHI()
 
 	FRHITextureCreateDesc Desc =
 		FRHITextureCreateDesc::Create2D(TEXT("FLandscapeTexture2DResource"), SizeX, SizeY, Format)
-		.SetNumMips(NumMips);
+		.SetNumMips(static_cast<uint8>(NumMips));
 
 	if (bCreateUAVs)
 	{
@@ -53,7 +53,7 @@ void FLandscapeTexture2DResource::InitRHI()
 
 	if (bCreateSRV)
 	{
-		TextureSRV = RHICreateShaderResourceView(TextureRHI, /*MipLevel = */0, NumMips, Format);
+		TextureSRV = RHICreateShaderResourceView(TextureRHI, /*MipLevel = */0, static_cast<uint8>(NumMips), Format);
 	}
 }
 
@@ -116,7 +116,7 @@ void FLandscapeTexture2DArrayResource::InitRHI()
 
 	const FRHITextureCreateDesc Desc =
 		FRHITextureCreateDesc::Create2DArray(TEXT("FLandscapeTexture2DArrayResource"), SizeX, SizeY, SizeZ, Format)
-		.SetNumMips(NumMips)
+		.SetNumMips(static_cast<uint8>(NumMips))
 		.SetFlags(Flags);
 
 	TextureRHI = RHICreateTexture(Desc);
@@ -132,7 +132,7 @@ void FLandscapeTexture2DArrayResource::InitRHI()
 
 	if (bCreateSRV)
 	{
-		TextureSRV = RHICreateShaderResourceView(TextureRHI, /*MipLevel = */0, NumMips, Format);
+		TextureSRV = RHICreateShaderResourceView(TextureRHI, /*MipLevel = */0, static_cast<uint8>(NumMips), Format);
 	}
 }
 
diff --git a/Engine/Source/Runtime/Landscape/Private/LandscapeLight.cpp b/Engine/Source/Runtime/Landscape/Private/LandscapeLight.cpp
--- a/Engine/Source/Runtime/Landscape/Private/LandscapeLight.cpp
+++ b/Engine/Source/Runtime/Landscape/Private/LandscapeLight.cpp
@@ -180,8 +180,8 @@ namespace
 		int32 LODHeightmapSizeY = LandscapeComponent->GetHeightmap()->Source.GetSizeY() >> LODValue;
 		float Ratio = (float)(LODHeightmapSizeX) / (HeightmapStride);
 
-		int32 CurrentHeightmapOffsetX = FMath::RoundToInt((float)(LODHeightmapSizeX) * LandscapeComponent->HeightmapScaleBias.Z);
-		int32 CurrentHeightmapOffsetY = FMath::RoundToInt((float)(LODHeightmapSizeY) * LandscapeComponent->HeightmapScaleBias.W);
+		int32 CurrentHeightmapOffsetX = FMath::RoundToInt32(LODHeightmapSizeX * LandscapeComponent->HeightmapScaleBias.Z);
+		int32 CurrentHeightmapOffsetY = FMath::RoundToInt32(LODHeightmapSizeY * LandscapeComponent->HeightmapScaleBias.W);
 
 		float XX = FMath::Clamp<float>((X - HeightmapOffsetX) * Ratio, 0.f, ComponentSize - 1.f) + CurrentHeightmapOffsetX;
 		int32 XI = (int32)XX;
@@ -200,12 +200,18 @@ namespace
 		FColor H3 = HeightMipData[XI + FMath::Min(YI + 1, LODHeightmapSizeY - 1) * LODHeightmapSizeX];
 		FColor H4 = HeightMipData[FMath::Min(XI + 1, LODHeightmapSizeX - 1) + FMath::Min(YI + 1, LODHeightmapSizeY - 1) * LODHeightmapSizeX];
 
-		uint16 Height = FMath::RoundToInt(FMath::Lerp(FMath::Lerp<float>(((H1.R << 8) + H1.G), ((H2.R << 8) + H2.G), XF),
-			FMath::Lerp<float>(((H3.R << 8) + H3.G), ((H4.R << 8) + H4.G), XF), YF));
-		uint8 B = FMath::RoundToInt(FMath::Lerp(FMath::Lerp<float>((H1.B), (H2.B), XF),
-			FMath::Lerp<float>((H3.B), (H4.B), XF), YF));
-		uint8 A = FMath::RoundToInt(FMath::Lerp<float>(FMath::Lerp((H1.A), (H2.A), XF),
-			FMath::Lerp<float>((H3.A), (H4.A), XF), YF));
+		const uint16 Height = static_cast<uint16>(FMath::RoundToInt(FMath::Lerp(
+			FMath::Lerp<float>(static_cast<float>((H1.R << 8) + H1.G), (H2.R << 8) + H2.G, XF),
+			FMath::Lerp<float>(static_cast<float>((H3.R << 8) + H3.G), (H4.R << 8) + H4.G, XF),
+			YF)));
+		const uint8 B = static_cast<uint8>(FMath::RoundToInt(FMath::Lerp(
+			FMath::Lerp<float>(H1.B, H2.B, XF),
+			FMath::Lerp<float>(H3.B, H4.B, XF),
+			YF)));
+		const uint8 A = static_cast<uint8>(FMath::RoundToInt(FMath::Lerp<float>(
+			FMath::Lerp<float>(H1.A, H2.A, XF),
+			FMath::Lerp<float>(H3.A, H4.A, XF),
+			YF)));
 
 		OutHeight = FColor((Height >> 8), Height & 255, B, A);
 
@@ -216,10 +222,14 @@ namespace
 			FColor X3 = XYOffsetMipData[XI + FMath::Min(YI + 1, LODHeightmapSizeY - 1) * LODHeightmapSizeX];
 			FColor X4 = XYOffsetMipData[FMath::Min(XI + 1, LODHeightmapSizeX - 1) + FMath::Min(YI + 1, LODHeightmapSizeY - 1) * LODHeightmapSizeX];
 
-			uint16 XComp = FMath::RoundToInt(FMath::Lerp(FMath::Lerp<float>(((X1.R << 8) + X1.G), ((X2.R << 8) + X2.G), XF),
-				FMath::Lerp<float>(((X3.R << 8) + X3.G), ((X4.R << 8) + X4.G), XF), YF));
-			uint16 YComp = FMath::RoundToInt(FMath::Lerp(FMath::Lerp<float>(((X1.B << 8) + X1.A), ((X2.B << 8) + X2.A), XF),
-				FMath::Lerp<float>(((X3.B << 8) + X3.A), ((X4.B << 8) + X4.A), XF), YF));
+			const uint16 XComp = static_cast<uint16>(FMath::RoundToInt(FMath::Lerp(
+				FMath::Lerp<float>(static_cast<float>((X1.R << 8) + X1.G), (X2.R << 8) + X2.G, XF),
+				FMath::Lerp<float>(static_cast<float>((X3.R << 8) + X3.G), (X4.R << 8) + X4.G, XF),
+				YF)));
+			const uint16 YComp = static_cast<uint16>(FMath::RoundToInt(FMath::Lerp(
+				FMath::Lerp<float>(static_cast<float>((X1.B << 8) + X1.A), (X2.B << 8) + X2.A, XF),
+				FMath::Lerp<float>(static_cast<float>((X3.B << 8) + X3.A), (X4.B << 8) + X4.A, XF),
+				YF)));
 
 			OutXYOffset = FColor((XComp >> 8), XComp & 255, YComp >> 8, YComp & 255);
 		}
@@ -320,22 +330,22 @@ namespace
 						FVector2D XY(float(X - DataInterface.HeightmapComponentOffsetX) / (ComponentSize - 1), float(Y - DataInterface.HeightmapComponentOffsetY) / (ComponentSize - 1));
 						XY = XY - 0.5f;
 
-						float RealLOD = GeometryLOD;
+						float RealLOD;
 
 						if (XY.X < 0.f)
 						{
 							if (XY.Y < 0.f)
 							{
 								RealLOD = FMath::Lerp(
-									FMath::Lerp<float>(NeighborLODs[0], NeighborLODs[1], XY.X + 1.f),
-									FMath::Lerp<float>(NeighborLODs[3], GeometryLOD, XY.X + 1.f),
+									FMath::Lerp<float>(static_cast<float>(NeighborLODs[0]), NeighborLODs[1], XY.X + 1.f),
+									FMath::Lerp<float>(static_cast<float>(NeighborLODs[3]), GeometryLOD, XY.X + 1.f),
 									XY.Y + 1.f); // 0
 							}
 							else
 							{
 								RealLOD = FMath::Lerp(
-									FMath::Lerp<float>(NeighborLODs[3], GeometryLOD, XY.X + 1.f),
-									FMath::Lerp<float>(NeighborLODs[5], NeighborLODs[6], XY.X + 1.f),
+									FMath::Lerp<float>(static_cast<float>(NeighborLODs[3]), GeometryLOD, XY.X + 1.f),
+									FMath::Lerp<float>(static_cast<float>(NeighborLODs[5]), NeighborLODs[6], XY.X + 1.f),
 									XY.Y); // 2
 							}
 						}
@@ -344,15 +354,15 @@ namespace
 							if (XY.Y < 0.f)
 							{
 								RealLOD = FMath::Lerp(
-									FMath::Lerp<float>(NeighborLODs[1], NeighborLODs[2], XY.X),
-									FMath::Lerp<float>(GeometryLOD, NeighborLODs[4], XY.X),
+									FMath::Lerp<float>(static_cast<float>(NeighborLODs[1]), NeighborLODs[2], XY.X),
+									FMath::Lerp<float>(static_cast<float>(GeometryLOD), NeighborLODs[4], XY.X),
 									XY.Y + 1.f); // 1
 							}
 							else
 							{
 								RealLOD = FMath::Lerp(
-									FMath::Lerp<float>(GeometryLOD, NeighborLODs[4], XY.X),
-									FMath::Lerp<float>(NeighborLODs[6], NeighborLODs[7], XY.X),
+									FMath::Lerp<float>(static_cast<float>(GeometryLOD), NeighborLODs[4], XY.X),
+									FMath::Lerp<float>(static_cast<float>(NeighborLODs[6]), NeighborLODs[7], XY.X),
 									XY.Y); // 3
 							}
 						}
@@ -374,23 +384,23 @@ namespace
 								FMath::Min(MaxLOD, LODValue + 1), HeightmapStride, Height[1], XYOffset[1]);
 
 							// Need interpolation
-							uint16 Height0 = (Height[0].R << 8) + Height[0].G;
-							uint16 Height1 = (Height[1].R << 8) + Height[1].G;
-							uint16 LerpHeight = FMath::RoundToInt(FMath::Lerp<float>(Height0, Height1, MorphAlpha));
+							const uint16 Height0 = (Height[0].R << 8) + Height[0].G;
+							const uint16 Height1 = (Height[1].R << 8) + Height[1].G;
+							const uint16 LerpHeight = static_cast<uint16>(FMath::RoundToInt(FMath::Lerp<float>(Height0, Height1, MorphAlpha)));
 
 							CompHeightData[X + Y * HeightmapStride] =
 								FColor((LerpHeight >> 8), LerpHeight & 255,
-								FMath::RoundToInt(FMath::Lerp<float>(Height[0].B, Height[1].B, MorphAlpha)),
-								FMath::RoundToInt(FMath::Lerp<float>(Height[0].A, Height[1].A, MorphAlpha)));
+								static_cast<uint8>(FMath::RoundToInt(FMath::Lerp<float>(Height[0].B, Height[1].B, MorphAlpha))),
+								static_cast<uint8>(FMath::RoundToInt(FMath::Lerp<float>(Height[0].A, Height[1].A, MorphAlpha))));
 							if (LandscapeComponent->XYOffsetmapTexture)
 							{
 								uint16 XComp0 = (XYOffset[0].R << 8) + XYOffset[0].G;
 								uint16 XComp1 = (XYOffset[1].R << 8) + XYOffset[1].G;
-								uint16 LerpXComp = FMath::RoundToInt(FMath::Lerp<float>(XComp0, XComp1, MorphAlpha));
+								uint16 LerpXComp = static_cast<uint16>(FMath::RoundToInt(FMath::Lerp<float>(XComp0, XComp1, MorphAlpha)));
 
 								uint16 YComp0 = (XYOffset[0].B << 8) + XYOffset[0].A;
 								uint16 YComp1 = (XYOffset[1].B << 8) + XYOffset[1].A;
-								uint16 LerpYComp = FMath::RoundToInt(FMath::Lerp<float>(YComp0, YComp1, MorphAlpha));
+								uint16 LerpYComp = static_cast<uint16>(FMath::RoundToInt(FMath::Lerp<float>(YComp0, YComp1, MorphAlpha)));
 
 								CompXYOffsetData[X + Y * HeightmapStride] =
 									FColor(LerpXComp >> 8, LerpXComp & 255, LerpYComp >> 8, LerpYComp & 255);
diff --git a/Engine/Source/Runtime/Landscape/Private/LandscapeNotification.cpp b/Engine/Source/Runtime/Landscape/Private/LandscapeNotification.cpp
--- a/Engine/Source/Runtime/Landscape/Private/LandscapeNotification.cpp
+++ b/Engine/Source/Runtime/Landscape/Private/LandscapeNotification.cpp
@@ -73,8 +73,12 @@ void FLandscapeNotificationManager::Tick()
 		LocalActiveNotification->SetNotificationText();
 
 		// There might be multiple notifications for this type:
-		int32 NumIdenticalNotifications = Algo::CountIf(ValidNotifications, [NotificationType = LocalActiveNotification->GetNotificationType()](const TWeakPtr<FLandscapeNotification>& Notification)
-		{ return (Notification.Pin().Get()->GetNotificationType() == NotificationType); });
+		const size_t NumIdenticalNotifications = Algo::CountIf(ValidNotifications,
+		                                                       [NotificationType = LocalActiveNotification->GetNotificationType()](
+		                                                       const TWeakPtr<FLandscapeNotification>& Notification)
+		                                                       {
+			                                                       return (Notification.Pin().Get()->GetNotificationType() == NotificationType);
+		                                                       });
 
 		FText Text = FText::Format(LOCTEXT("NotificationFooter", "{0}\n[{1}]"),
 			LocalActiveNotification->NotificationText,
diff --git a/Engine/Source/Runtime/Landscape/Private/Materials/MaterialExpressionLandscapeLayerBlend.cpp b/Engine/Source/Runtime/Landscape/Private/Materials/MaterialExpressionLandscapeLayerBlend.cpp
--- a/Engine/Source/Runtime/Landscape/Private/Materials/MaterialExpressionLandscapeLayerBlend.cpp
+++ b/Engine/Source/Runtime/Landscape/Private/Materials/MaterialExpressionLandscapeLayerBlend.cpp
@@ -212,7 +212,11 @@ int32 UMaterialExpressionLandscapeLayerBlend::Compile(class FMaterialCompiler* C
 		if (WeightCodes[LayerIdx] != INDEX_NONE)
 		{
 			// Layer input
-			int32 LayerCode = Layer.LayerInput.Expression ? Layer.LayerInput.Compile(Compiler) : Compiler->Constant3(Layer.ConstLayerInput.X, Layer.ConstLayerInput.Y, Layer.ConstLayerInput.Z);
+			const int32 LayerCode = Layer.LayerInput.Expression
+				                        ? Layer.LayerInput.Compile(Compiler)
+				                        : Compiler->Constant3(static_cast<float>(Layer.ConstLayerInput.X),
+															  static_cast<float>(Layer.ConstLayerInput.Y),
+				                                              static_cast<float>(Layer.ConstLayerInput.Z));
 
 			if (bNeedsRenormalize)
 			{
@@ -236,7 +240,11 @@ int32 UMaterialExpressionLandscapeLayerBlend::Compile(class FMaterialCompiler* C
 			const int32 WeightCode = Compiler->StaticTerrainLayerWeight(Layer.LayerName, DefaultWeightCode);
 			if (WeightCode != INDEX_NONE)
 			{
-				const int32 LayerCode = Layer.LayerInput.Expression ? Layer.LayerInput.Compile(Compiler) : Compiler->Constant3(Layer.ConstLayerInput.X, Layer.ConstLayerInput.Y, Layer.ConstLayerInput.Z);
+				const int32 LayerCode = Layer.LayerInput.Expression
+					                        ? Layer.LayerInput.Compile(Compiler)
+					                        : Compiler->Constant3(static_cast<float>(Layer.ConstLayerInput.X),
+					                                              static_cast<float>(Layer.ConstLayerInput.Y),
+					                                              static_cast<float>(Layer.ConstLayerInput.Z));
 				// Blend in the layer using the alpha value
 				OutputCode = Compiler->Lerp(OutputCode, LayerCode, WeightCode);
 			}
diff --git a/Engine/Source/Runtime/Landscape/Private/Materials/MaterialExpressionLandscapeLayerCoords.cpp b/Engine/Source/Runtime/Landscape/Private/Materials/MaterialExpressionLandscapeLayerCoords.cpp
--- a/Engine/Source/Runtime/Landscape/Private/Materials/MaterialExpressionLandscapeLayerCoords.cpp
+++ b/Engine/Source/Runtime/Landscape/Private/Materials/MaterialExpressionLandscapeLayerCoords.cpp
@@ -64,8 +64,8 @@ int32 UMaterialExpressionLandscapeLayerCoords::Compile(class FMaterialCompiler*
 	float Scale = (MappingScale == 0.0f) ? 1.0f : 1.0f / MappingScale;
 	int32 RealScale = Compiler->Constant(Scale);
 
-	float Cos = FMath::Cos(MappingRotation * PI / 180.0);
-	float Sin = FMath::Sin(MappingRotation * PI / 180.0);
+	const float Cos = FMath::Cos(MappingRotation * PI / 180.0f);
+	const float Sin = FMath::Sin(MappingRotation * PI / 180.0f);
 
 	int32 ResultUV = INDEX_NONE;
 	int32 TransformedUV = Compiler->Add(
diff --git a/Engine/Source/Runtime/Landscape/Private/Materials/MaterialExpressionLandscapeLayerWeight.cpp b/Engine/Source/Runtime/Landscape/Private/Materials/MaterialExpressionLandscapeLayerWeight.cpp
--- a/Engine/Source/Runtime/Landscape/Private/Materials/MaterialExpressionLandscapeLayerWeight.cpp
+++ b/Engine/Source/Runtime/Landscape/Private/Materials/MaterialExpressionLandscapeLayerWeight.cpp
@@ -63,7 +63,9 @@ bool UMaterialExpressionLandscapeLayerWeight::IsResultMaterialAttributes(int32 O
 
 int32 UMaterialExpressionLandscapeLayerWeight::Compile(class FMaterialCompiler* Compiler, int32 OutputIndex)
 {
-	const int32 BaseCode = Base.Expression ? Base.Compile(Compiler) : Compiler->Constant3(ConstBase.X, ConstBase.Y, ConstBase.Z);
+	const int32 BaseCode = Base.Expression
+		                       ? Base.Compile(Compiler)
+		                       : Compiler->Constant3(static_cast<float>(ConstBase.X), static_cast<float>(ConstBase.Y), static_cast<float>(ConstBase.Z));
 	const int32 WeightCode = Compiler->StaticTerrainLayerWeight(ParameterName, Compiler->Constant(PreviewWeight));
 
 	int32 ReturnCode = INDEX_NONE;
diff --git a/Engine/Source/Runtime/Landscape/Public/LandscapeDataAccess.h b/Engine/Source/Runtime/Landscape/Public/LandscapeDataAccess.h
--- a/Engine/Source/Runtime/Landscape/Public/LandscapeDataAccess.h
+++ b/Engine/Source/Runtime/Landscape/Public/LandscapeDataAccess.h
@@ -33,7 +33,7 @@ namespace LandscapeDataAccess
 
 	FORCEINLINE uint16 GetTexHeight(float Height)
 	{
-		return FMath::RoundToInt(FMath::Clamp<float>(Height * LANDSCAPE_INV_ZSCALE + MidValue, 0.f, MaxValue));		
+		return static_cast<uint16>(FMath::RoundToInt(FMath::Clamp<float>(Height * LANDSCAPE_INV_ZSCALE + MidValue, 0.f, MaxValue)));		
 	}
 
 	FORCEINLINE FColor PackHeight(uint16 Height)
diff --git a/Engine/Source/Runtime/Landscape/Public/LandscapeLight.h b/Engine/Source/Runtime/Landscape/Public/LandscapeLight.h
--- a/Engine/Source/Runtime/Landscape/Public/LandscapeLight.h
+++ b/Engine/Source/Runtime/Landscape/Public/LandscapeLight.h
@@ -114,10 +114,10 @@ namespace
 		if (LightMapRes <= 0) return 0.f;
 
 		// Assuming DXT_1 compression at the moment...
-		int32 PixelPaddingX = GPixelFormats[PF_DXT1].BlockSizeX; // "/2" ?
-		int32 PixelPaddingY = GPixelFormats[PF_DXT1].BlockSizeY;
-		int32 PatchExpandCountX = (LightMapRes >= 1.f) ? (PixelPaddingX) / LightMapRes : (PixelPaddingX);
-		int32 PatchExpandCountY = (LightMapRes >= 1.f) ? (PixelPaddingY) / LightMapRes : (PixelPaddingY);
+		const int32 PixelPaddingX = GPixelFormats[PF_DXT1].BlockSizeX; // "/2" ?
+		const int32 PixelPaddingY = GPixelFormats[PF_DXT1].BlockSizeY;
+		const int32 PatchExpandCountX = static_cast<int32>(LightMapRes >= 1.f ? PixelPaddingX / LightMapRes : PixelPaddingX);
+		const int32 PatchExpandCountY = static_cast<int32>(LightMapRes >= 1.f ? PixelPaddingY / LightMapRes : PixelPaddingY);
 
 		X = FMath::Max<int32>(1, PatchExpandCountX >> LightingLOD);
 		Y = FMath::Max<int32>(1, PatchExpandCountY >> LightingLOD);
@@ -143,8 +143,8 @@ namespace
 			}
 		}
 
-		int32 DestSize = (float)DesiredSize / CurrentSize * (ComponentSize*LightMapRes);
-		float LightMapRatio = (float)DestSize / (ComponentSize*LightMapRes) * CurrentSize / DesiredSize;
+		const float DestSize = FMath::Floor(static_cast<float>(DesiredSize) / CurrentSize * (ComponentSize*LightMapRes));
+		const float LightMapRatio = DestSize / (ComponentSize*LightMapRes) * CurrentSize / DesiredSize;
 		return LightMapRatio;
 	}
 }
