Landscape: Fix notification flickering/not showing and fix landscape notification manager leak.

The start times of landscape texture and brush notifications were not correct.

When loading a new level, if any notifications are being displayed by the landscape notification manager, the manager will leak the currently displayed notification leaving it in an incomplete state.

diff --git a/Engine/Source/Runtime/Landscape/Classes/Landscape.h b/Engine/Source/Runtime/Landscape/Classes/Landscape.h
--- a/Engine/Source/Runtime/Landscape/Classes/Landscape.h
+++ b/Engine/Source/Runtime/Landscape/Classes/Landscape.h
@@ -541,7 +541,7 @@ private:
 	bool bSplineLayerUpdateRequested;
 
 	/** Time since waiting for landscape resources to be ready (for displaying a notification to the user) */
-	double WaitingForLandscapeResourcesStartTime = -1.0;
+	double WaitingForLandscapeTextureResourcesStartTime = -1.0;
 
 	/** Time since waiting for brush resources to be ready (for displaying a notification to the user) */
 	double WaitingForLandscapeBrushResourcesStartTime = -1.0;
diff --git a/Engine/Source/Runtime/Landscape/Private/LandscapeEditLayers.cpp b/Engine/Source/Runtime/Landscape/Private/LandscapeEditLayers.cpp
--- a/Engine/Source/Runtime/Landscape/Private/LandscapeEditLayers.cpp
+++ b/Engine/Source/Runtime/Landscape/Private/LandscapeEditLayers.cpp
@@ -7759,7 +7759,7 @@ void ALandscape::UpdateLayersContent(bool bInWaitForStreaming, bool bInSkipMonit
 			WaitingForTexturesNotification.Reset();
 			WaitingForBrushesNotification.Reset();
 			InvalidShadingModelNotification.Reset();
-			WaitingForLandscapeResourcesStartTime = -1.0;
+			WaitingForLandscapeTextureResourcesStartTime = -1.0;
 			WaitingForLandscapeBrushResourcesStartTime = -1.0;
 		}
 	};
@@ -7864,45 +7864,41 @@ void ALandscape::UpdateLayersContent(bool bInWaitForStreaming, bool bInSkipMonit
 	bResourcesReady &= PrepareLayersTextureResources(bInWaitForStreaming);
 	if (!bResourcesReady)
 	{
-		if (WaitingForLandscapeResourcesStartTime < 0.0)
+		WaitingForLandscapeTextureResourcesStartTime = FSlateApplicationBase::Get().GetCurrentTime();
+		if (!WaitingForTexturesNotification.IsValid())
 		{
-			WaitingForLandscapeResourcesStartTime = FSlateApplicationBase::Get().GetCurrentTime();
-			if (!WaitingForTexturesNotification.IsValid())
-			{
-				WaitingForTexturesNotification = MakeShared<FLandscapeNotification>(this, FLandscapeNotification::EType::LandscapeResourcesNotReady);
-				static const FText NotificationText(LOCTEXT("WaitForLanscapeResources", "Landscape editor waiting for landscape resources to be ready."));
-				WaitingForTexturesNotification->NotificationText = NotificationText;
-				WaitingForTexturesNotification->NotificationStartTime = WaitingForLandscapeResourcesStartTime + TimeBeforeDisplayingWaitingForResourcesNotification;
-			}
-			LandscapeNotificationManager->RegisterNotification(WaitingForTexturesNotification);
-			bHideNotifications = false;
+			WaitingForTexturesNotification = MakeShared<FLandscapeNotification>(this, FLandscapeNotification::EType::LandscapeTextureResourcesNotReady);
+			static const FText NotificationText(LOCTEXT("WaitForLandscapeTextureResources", "Landscape editor waiting for landscape texture resources to be ready."));
+			WaitingForTexturesNotification->NotificationText = NotificationText;
+			WaitingForTexturesNotification->NotificationStartTime = WaitingForLandscapeTextureResourcesStartTime + TimeBeforeDisplayingWaitingForResourcesNotification;
 		}
+		LandscapeNotificationManager->RegisterNotification(WaitingForTexturesNotification);
+		bHideNotifications = false;
 	}
 	else
 	{
 		WaitingForTexturesNotification.Reset();
+		WaitingForLandscapeTextureResourcesStartTime = -1.0;
 	}
 
 	bResourcesReady &= PrepareLayersBrushResources(World->FeatureLevel, bInWaitForStreaming);
 	if (!bResourcesReady)
 	{
-		if (WaitingForLandscapeBrushResourcesStartTime < 0.0)
+		WaitingForLandscapeBrushResourcesStartTime = FSlateApplicationBase::Get().GetCurrentTime();
+		if (!WaitingForBrushesNotification.IsValid())
 		{
-			WaitingForLandscapeBrushResourcesStartTime = FSlateApplicationBase::Get().GetCurrentTime();
-			if (!WaitingForBrushesNotification.IsValid())
-			{
-				WaitingForBrushesNotification = MakeShared<FLandscapeNotification>(this, FLandscapeNotification::EType::LandscapeResourcesNotReady);
-				static const FText NotificationText(LOCTEXT("WaitForLanscapeResources", "Landscape editor waiting for landscape resources to be ready."));
-				WaitingForBrushesNotification->NotificationText = NotificationText;
-				WaitingForBrushesNotification->NotificationStartTime = WaitingForLandscapeResourcesStartTime + TimeBeforeDisplayingWaitingForResourcesNotification;
-			}
-			LandscapeNotificationManager->RegisterNotification(WaitingForBrushesNotification);
-			bHideNotifications = false;
+			WaitingForBrushesNotification = MakeShared<FLandscapeNotification>(this, FLandscapeNotification::EType::LandscapeBrushResourcesNotReady);
+			static const FText NotificationText(LOCTEXT("WaitForLandscapeBrushResources", "Landscape editor waiting for landscape brush resources to be ready."));
+			WaitingForBrushesNotification->NotificationText = NotificationText;
+			WaitingForBrushesNotification->NotificationStartTime = WaitingForLandscapeBrushResourcesStartTime + TimeBeforeDisplayingWaitingForResourcesNotification;
 		}
+		LandscapeNotificationManager->RegisterNotification(WaitingForBrushesNotification);
+		bHideNotifications = false;
 	}
 	else
 	{
 		WaitingForBrushesNotification.Reset();
+		WaitingForLandscapeBrushResourcesStartTime = -1.0;
 	}
 
 	if (!bResourcesReady)
diff --git a/Engine/Source/Runtime/Landscape/Private/LandscapeNotification.cpp b/Engine/Source/Runtime/Landscape/Private/LandscapeNotification.cpp
--- a/Engine/Source/Runtime/Landscape/Private/LandscapeNotification.cpp
+++ b/Engine/Source/Runtime/Landscape/Private/LandscapeNotification.cpp
@@ -46,6 +46,14 @@ void FLandscapeNotification::SetNotificationText()
 
 // ----------------------------------------------------------------------------------
 
+FLandscapeNotificationManager::~FLandscapeNotificationManager()
+{
+	if (NotificationItem.IsValid())
+	{
+		HideNotificationItem();
+	}
+}
+
 void FLandscapeNotificationManager::Tick()
 {
 	// Cleanup notifications for stale landscapes:
diff --git a/Engine/Source/Runtime/Landscape/Public/LandscapeNotification.h b/Engine/Source/Runtime/Landscape/Public/LandscapeNotification.h
--- a/Engine/Source/Runtime/Landscape/Public/LandscapeNotification.h
+++ b/Engine/Source/Runtime/Landscape/Public/LandscapeNotification.h
@@ -12,9 +12,9 @@ struct FLandscapeNotification
 	/** The various types of landscape notifications, ordered by priority (a single notification is displayed at a time). */
 	enum class EType : uint8
 	{
-		LandscapeResourcesNotReady = 0,
+		ShadingModelInvalid = 0,
+		LandscapeTextureResourcesNotReady,
 		LandscapeBrushResourcesNotReady,
-		ShadingModelInvalid,
 		TextureBaking,
 		GrassRendering,
 	};
@@ -70,6 +70,8 @@ private:
 class FLandscapeNotificationManager
 {
 public:
+	~FLandscapeNotificationManager();
+	
 	void Tick();
 
 	void RegisterNotification(const TWeakPtr<FLandscapeNotification>& InNotification);
