Landscape : Serialize CollisionHash so we don't return the collision as updated & trigger navmesh rebuilds on map load.

diff --git a/Engine/Source/Runtime/Landscape/Classes/LandscapeHeightfieldCollisionComponent.h b/Engine/Source/Runtime/Landscape/Classes/LandscapeHeightfieldCollisionComponent.h
--- a/Engine/Source/Runtime/Landscape/Classes/LandscapeHeightfieldCollisionComponent.h
+++ b/Engine/Source/Runtime/Landscape/Classes/LandscapeHeightfieldCollisionComponent.h
@@ -109,6 +109,7 @@ class ULandscapeHeightfieldCollisionComponent : public UPrimitiveComponent
 #if WITH_EDITORONLY_DATA
 	friend struct FEnableCollisionHashOptimScope;
 	
+	UPROPERTY()
 	uint32										CollisionHash = 0;
 
 	/** The collision height values. Stripped from cooked content */
