Fixed landscape edit layers readback RHI fence not being cleared upon reuse. Physical material readbacks don't need to be fixed because they use FRHIGPUBufferReadback which does the clear already

#rnx

diff --git a/Engine/Source/Runtime/Landscape/Private/LandscapeEditReadback.cpp b/Engine/Source/Runtime/Landscape/Private/LandscapeEditReadback.cpp
--- a/Engine/Source/Runtime/Landscape/Private/LandscapeEditReadback.cpp
+++ b/Engine/Source/Runtime/Landscape/Private/LandscapeEditReadback.cpp
@@ -73,6 +73,7 @@ bool InitTask_RenderThread(FLandscapeEditReadbackTaskImpl& Task)
 	{
 		Task.ReadbackFence = RHICreateGPUFence(TEXT("LandscapeEditReadbackTask"));
 	}
+	Task.ReadbackFence->Clear();
 
 	return true;
 }
