Improved formatting of landscape notifications

diff --git a/Engine/Source/Runtime/Landscape/Private/LandscapeNotification.cpp b/Engine/Source/Runtime/Landscape/Private/LandscapeNotification.cpp
--- a/Engine/Source/Runtime/Landscape/Private/LandscapeNotification.cpp
+++ b/Engine/Source/Runtime/Landscape/Private/LandscapeNotification.cpp
@@ -88,9 +88,9 @@ void FLandscapeNotificationManager::Tick()
 			                                                       return (Notification.Pin().Get()->GetNotificationType() == NotificationType);
 		                                                       });
 
-		FText Text = FText::Format(LOCTEXT("NotificationFooter", "{0}\n[{1}]"),
+		FText Text = FText::Format(LOCTEXT("NotificationFooter", "{0} for {1}"),
 			LocalActiveNotification->NotificationText,
-			(NumIdenticalNotifications > 1) ? LOCTEXT("MultipleLandscapes", "Multiple landscapes") : FText::FromString(LocalActiveNotification->GetLandscape()->GetActorLabel()));
+			(NumIdenticalNotifications > 1) ? LOCTEXT("MultipleLandscapes", "multiple landscapes") : FText::FromString(LocalActiveNotification->GetLandscape()->GetActorLabel()));
 
 		ShowNotificationItem(Text);
 	}
