Add a BP function to force landscapes to fully load in textures for more reliable enigne tests

diff --git a/Engine/Source/Runtime/Landscape/Classes/Landscape.h b/Engine/Source/Runtime/Landscape/Classes/Landscape.h
--- a/Engine/Source/Runtime/Landscape/Classes/Landscape.h
+++ b/Engine/Source/Runtime/Landscape/Classes/Landscape.h
@@ -357,6 +357,7 @@ public:
 
 	LANDSCAPE_API void SetUseGeneratedLandscapeSplineMeshesActors(bool bInEnabled);
 	LANDSCAPE_API bool GetUseGeneratedLandscapeSplineMeshesActors() const;
+	LANDSCAPE_API bool PrepareTextureResources(bool bInWaitForStreaming);
 
 protected:
 	FName GenerateUniqueLayerName(FName InName = NAME_None) const;
@@ -388,7 +389,6 @@ private:
 	uint32 UpdateCollisionAndClients(const TArrayView<FLandscapeEditLayerComponentReadbackResult>& Components);
 	uint32 UpdateAfterReadbackResolves(const TArrayView<FLandscapeEditLayerComponentReadbackResult>& Components);
 
-	bool PrepareTextureResources(bool bInWaitForStreaming);
 	bool PrepareLayersTextureResources(bool bInWaitForStreaming);
 	bool PrepareLayersTextureResources(const TArray<FLandscapeLayer>& InLayers, bool bInWaitForStreaming);
 	bool PrepareLayersBrushResources(ERHIFeatureLevel::Type InFeatureLevel, bool bInWaitForStreaming);
