Landscape build improvements:
* Build All now marks all "pending dirty" landscape actors as dirty so that the landscape can be considered fully up-to-date after the operation
* Build Nanite now shows a proper progress bar since it can take several seconds to complete

Misc:
* ULandscapeSubsystem now holds TWeakObjectPtr to landscape proxies in order not to hold references to them (to be in line with 22212286)

diff --git a/Engine/Source/Runtime/Landscape/Classes/LandscapeInfo.h b/Engine/Source/Runtime/Landscape/Classes/LandscapeInfo.h
--- a/Engine/Source/Runtime/Landscape/Classes/LandscapeInfo.h
+++ b/Engine/Source/Runtime/Landscape/Classes/LandscapeInfo.h
@@ -255,6 +255,7 @@ public:
 	LANDSCAPE_API void OnModifiedPackageSaved(UPackage* InPackage);
 	LANDSCAPE_API int32 GetModifiedPackageCount() const;
 	LANDSCAPE_API TArray<UPackage*> GetModifiedPackages() const;
+	LANDSCAPE_API void MarkModifiedPackagesAsDirty();
 
 	LANDSCAPE_API void ModifyObject(UObject* InObject, bool bAlwaysMarkDirty = true);
 	LANDSCAPE_API void MarkObjectDirty(UObject* InObject);
diff --git a/Engine/Source/Runtime/Landscape/Private/Landscape.cpp b/Engine/Source/Runtime/Landscape/Private/Landscape.cpp
--- a/Engine/Source/Runtime/Landscape/Private/Landscape.cpp
+++ b/Engine/Source/Runtime/Landscape/Private/Landscape.cpp
@@ -3749,6 +3749,18 @@ TArray<UPackage*> ULandscapeInfo::GetModifiedPackages() const
 	return LocalModifiedPackages;
 }
 
+void ULandscapeInfo::MarkModifiedPackagesAsDirty()
+{
+	for (TWeakObjectPtr<UPackage> WeakPackagePtr : ModifiedPackages)
+	{
+		if (UPackage* Package = WeakPackagePtr.Get())
+		{
+			Package->MarkPackageDirty();
+		}
+	}
+	ModifiedPackages.Empty();
+}
+
 int32 ULandscapeInfo::GetModifiedPackageCount() const
 {
 	int32 ModifiedPackageCount = 0;
diff --git a/Engine/Source/Runtime/Landscape/Private/LandscapeSubsystem.cpp b/Engine/Source/Runtime/Landscape/Private/LandscapeSubsystem.cpp
--- a/Engine/Source/Runtime/Landscape/Private/LandscapeSubsystem.cpp
+++ b/Engine/Source/Runtime/Landscape/Private/LandscapeSubsystem.cpp
@@ -19,6 +19,7 @@
 #include "LandscapeConfigHelper.h"
 #include "Engine/Canvas.h"
 #include "EngineUtils.h"
+#include "Misc/ScopedSlowTask.h"
 
 #include UE_INLINE_GENERATED_CPP_BY_NAME(LandscapeSubsystem)
 
@@ -46,12 +47,12 @@ ULandscapeSubsystem::~ULandscapeSubsystem()
 
 void ULandscapeSubsystem::RegisterActor(ALandscapeProxy* Proxy)
 {
-	Proxies.AddUnique(Proxy);
+	Proxies.AddUnique(TWeakObjectPtr<ALandscapeProxy>(Proxy));
 }
 
 void ULandscapeSubsystem::UnregisterActor(ALandscapeProxy* Proxy)
 {
-	Proxies.Remove(Proxy);
+	Proxies.Remove(TWeakObjectPtr<ALandscapeProxy>(Proxy));
 }
 
 void ULandscapeSubsystem::Initialize(FSubsystemCollectionBase& Collection)
@@ -132,30 +133,33 @@ void ULandscapeSubsystem::Tick(float DeltaTime)
 	}
 
 	int32 InOutNumComponentsCreated = 0;
-	for (ALandscapeProxy* Proxy : Proxies)
+	for (TWeakObjectPtr<ALandscapeProxy> ProxyPtr : Proxies)
 	{
-#if WITH_EDITOR
-		if (GIsEditor)
+		if (ALandscapeProxy* Proxy = ProxyPtr.Get())
 		{
-			if (ALandscape* Landscape = Cast<ALandscape>(Proxy))
+#if WITH_EDITOR
+			if (GIsEditor)
 			{
-				Landscape->TickLayers(DeltaTime);
-			}
+				if (ALandscape* Landscape = Cast<ALandscape>(Proxy))
+				{
+					Landscape->TickLayers(DeltaTime);
+				}
 
-			// editor-only
-			if (!World->IsPlayInEditor())
-			{
-				Proxy->UpdateGIBakedTextures();
-				Proxy->UpdatePhysicalMaterialTasks();
+				// editor-only
+				if (!World->IsPlayInEditor())
+				{
+					Proxy->UpdateGIBakedTextures();
+					Proxy->UpdatePhysicalMaterialTasks();
+				}
 			}
-		}
 #endif
-		if (Cameras && Proxy->ShouldTickGrass())
-		{
-			Proxy->TickGrass(*Cameras, InOutNumComponentsCreated);
-		}
+			if (Cameras && Proxy->ShouldTickGrass())
+			{
+				Proxy->TickGrass(*Cameras, InOutNumComponentsCreated);
+			}
 
-		Proxy->UpdateRenderingMethod();
+			Proxy->UpdateRenderingMethod();
+		}
 	}
 
 #if WITH_EDITOR
@@ -170,6 +174,9 @@ void ULandscapeSubsystem::Tick(float DeltaTime)
 #if WITH_EDITOR
 void ULandscapeSubsystem::BuildAll()
 {
+	// This is a deliberate action, make sure to flush all packages that are 'pending dirty' :
+	MarkModifiedLandscapesAsDirty();
+
 	BuildGrassMaps();
 	BuildGIBakedTextures();
 	BuildPhysicalMaterial();
@@ -214,10 +221,22 @@ void ULandscapeSubsystem::BuildNanite()
 		return;
 	}
 
-	for (TActorIterator<ALandscapeProxy> ProxyIt(World); ProxyIt; ++ProxyIt)
+	if (Proxies.IsEmpty())
 	{
-		ProxyIt->UpdateNaniteRepresentation();
-		ProxyIt->UpdateRenderingMethod();
+		return;
+	}
+
+	FScopedSlowTask SlowTask(Proxies.Num(), (LOCTEXT("Landscape_BuildNanite", "Building Nanite Landscape Meshes")));
+	SlowTask.MakeDialog();
+
+	for (TWeakObjectPtr<ALandscapeProxy> ProxyPtr : Proxies)
+	{
+		SlowTask.EnterProgressFrame(1, FText::Format(LOCTEXT("Landscape_BuildNaniteProgress", "Building Nanite Landscape Mesh {0} of {1})"), FText::AsNumber(SlowTask.CompletedWork), FText::AsNumber(SlowTask.TotalAmountOfWork)));
+		if (ALandscapeProxy* Proxy = ProxyPtr.Get())
+		{
+			Proxy->UpdateNaniteRepresentation();
+			Proxy->UpdateRenderingMethod();
+		}
 	}
 }
 
@@ -296,6 +315,16 @@ void ULandscapeSubsystem::SaveModifiedLandscapes()
 	});
 }
 
+void ULandscapeSubsystem::MarkModifiedLandscapesAsDirty()
+{
+	// Flush all packages that are pending mark for dirty : 
+	ForEachLandscapeInfo([&](ULandscapeInfo* LandscapeInfo)
+	{
+		LandscapeInfo->MarkModifiedPackagesAsDirty();
+		return true;
+	});
+}
+
 bool ULandscapeSubsystem::HasModifiedLandscapes() const
 {
 	bool bHasModifiedLandscapes = false;
diff --git a/Engine/Source/Runtime/Landscape/Public/LandscapeSubsystem.h b/Engine/Source/Runtime/Landscape/Public/LandscapeSubsystem.h
--- a/Engine/Source/Runtime/Landscape/Public/LandscapeSubsystem.h
+++ b/Engine/Source/Runtime/Landscape/Public/LandscapeSubsystem.h
@@ -43,6 +43,7 @@ public:
 	LANDSCAPE_API void ChangeGridSize(ULandscapeInfo* LandscapeInfo, uint32 NewGridSizeInComponents);
 	LANDSCAPE_API ALandscapeProxy* FindOrAddLandscapeProxy(ULandscapeInfo* LandscapeInfo, const FIntPoint& SectionBase);
 	LANDSCAPE_API void DisplayMessages(class FCanvas* Canvas, float& XPos, float& YPos);
+	LANDSCAPE_API void MarkModifiedLandscapesAsDirty();
 	LANDSCAPE_API void SaveModifiedLandscapes();
 	LANDSCAPE_API bool HasModifiedLandscapes() const;
 	LANDSCAPE_API static bool IsDirtyOnlyInModeEnabled();
@@ -59,7 +60,7 @@ private:
 	virtual void Deinitialize() override;
 	// End USubsystem
 
-	TArray<ALandscapeProxy*> Proxies;
+	TArray<TWeakObjectPtr<ALandscapeProxy>> Proxies;
 
 #if WITH_EDITOR
 	class FLandscapeGrassMapsBuilder* GrassMapsBuilder;
