Fixed invalid error when re-registering a landscape proxy actor, which is a valid case for ULandscapeInfo::RegisterActor : the error should only be displayed when trying to register a different actor with the same LandscapeGuid

#rnx

diff --git a/Engine/Source/Runtime/Landscape/Private/Landscape.cpp b/Engine/Source/Runtime/Landscape/Private/Landscape.cpp
--- a/Engine/Source/Runtime/Landscape/Private/Landscape.cpp
+++ b/Engine/Source/Runtime/Landscape/Private/Landscape.cpp
@@ -4163,7 +4163,7 @@ void ULandscapeInfo::RegisterActor(ALandscapeProxy* Proxy, bool bMapCheck)
 				}
 			}
 		}
-		else if (LandscapeActor == Landscape)
+		else if (LandscapeActor != Landscape)
 		{
 			UE_LOG(LogLandscape, Warning, TEXT("Multiple landscapes actors with the same GUID detected: %s vs %s"), * LandscapeActor->GetPathName(), * Landscape->GetPathName());
 		}
