Fixed typo in error message

#rnx
#skipci

diff --git a/Engine/Source/Runtime/Landscape/Private/Landscape.cpp b/Engine/Source/Runtime/Landscape/Private/Landscape.cpp
--- a/Engine/Source/Runtime/Landscape/Private/Landscape.cpp
+++ b/Engine/Source/Runtime/Landscape/Private/Landscape.cpp
@@ -4166,7 +4166,7 @@ void ULandscapeInfo::RegisterActor(ALandscapeProxy* Proxy, bool bMapCheck)
 		}
 		else if (LandscapeActor != Landscape)
 		{
-			UE_LOG(LogLandscape, Warning, TEXT("Multiple landscapes actors with the same GUID detected: %s vs %s"), * LandscapeActor->GetPathName(), * Landscape->GetPathName());
+			UE_LOG(LogLandscape, Warning, TEXT("Multiple landscape actors with the same GUID detected: %s vs %s"), * LandscapeActor->GetPathName(), * Landscape->GetPathName());
 		}
 	}
 	else
