Fixed issues with landscape collision visualization :
* Crash when using SimpleMipCollisionLevel > 0
* Collision mesh not refreshing when regenerating collisions (e.g. when changing SimpleMipCollisionLevel in the details panel)
* Use current shading model for rendering the collision mesh (rather than forcing SM6.0)
* Support for multi-view renders

#rnx

diff --git a/Engine/Source/Runtime/Landscape/Private/LandscapeCollision.cpp b/Engine/Source/Runtime/Landscape/Private/LandscapeCollision.cpp
--- a/Engine/Source/Runtime/Landscape/Private/LandscapeCollision.cpp
+++ b/Engine/Source/Runtime/Landscape/Private/LandscapeCollision.cpp
@@ -278,6 +278,9 @@ void ULandscapeHeightfieldCollisionComponent::OnCreatePhysicsState()
 	{
 		CreateCollisionObject();
 
+		// Debug display needs to update its representation, so we invalidate the collision component's render state : 
+		MarkRenderStateDirty();
+
 		if (IsValidRef(HeightfieldRef))
 		{
 			// Make transform for this landscape component PxActor
@@ -490,90 +493,75 @@ FPrimitiveSceneProxy* ULandscapeHeightfieldCollisionComponent::CreateSceneProxy(
 
 		// Constructor exists to populate Vertices and Indices arrays which are
 		// used to construct the collision mesh inside GetDynamicMeshElements
-		FLandscapeHeightfieldCollisionComponentSceneProxy(const ULandscapeHeightfieldCollisionComponent* InComponent, const Chaos::FHeightField& InHeightfield, const FLinearColor& InWireframeColor, bool bInShowHoles)
+		FLandscapeHeightfieldCollisionComponentSceneProxy(const ULandscapeHeightfieldCollisionComponent* InComponent, const Chaos::FHeightField& InHeightfield, const FLinearColor& InWireframeColor)
 			: FPrimitiveSceneProxy(InComponent)
 		{
-			const Chaos::FHeightField::FData<uint16>& GeomData = InHeightfield.GeomData;
-			const int32 CollisionSizeVerts = InComponent->CollisionSizeQuads + 1;
-			const int32 SimpleCollisionSizeVerts = InComponent->SimpleCollisionSizeQuads > 0 ? InComponent->SimpleCollisionSizeQuads + 1 : 0;
-			const int32 NumVerts = FMath::Square(CollisionSizeVerts);
-			const int32 NumSimpleVerts = FMath::Square(SimpleCollisionSizeVerts);
-
-			check(InComponent->CollisionHeightData.GetElementCount() == NumVerts + NumSimpleVerts);
-
+			const Chaos::FHeightField::FData<uint16>& GeomData = InHeightfield.GeomData; 
+			const int32 NumRows = InHeightfield.GetNumRows();
+			const int32 NumCols = InHeightfield.GetNumCols();
+			const int32 NumVerts = NumRows * NumCols;
+			const int32 NumTris = (NumRows - 1) * (NumCols - 1) * 2;
 			Vertices.SetNumUninitialized(NumVerts);
-			for (int32 i = 0; i < NumVerts; i++)
+			for (int32 I = 0; I < NumVerts; I++)
 			{
-				int32 X = i % CollisionSizeVerts;
-				int32 Y = i / CollisionSizeVerts;
-				Vertices[i] = FVector3f(X, Y, (GeomData.MinValue + (float)GeomData.Heights[i] * GeomData.HeightPerUnit) * LANDSCAPE_ZSCALE);
+				Chaos::FVec3 Point = GeomData.GetPointScaled(I);
+				Vertices[I].Position = FVector3f(Point.X, Point.Y, Point.Z);
 			}
-
-			const int32 NumTris = FMath::Square(InComponent->CollisionSizeQuads) * 2;
 			Indices.SetNumUninitialized(NumTris * 3);
 
-			const uint8* DominantLayers = nullptr;
-			if (InComponent->DominantLayerData.GetElementCount() > 0)
-			{
-				DominantLayers = (const uint8*)InComponent->DominantLayerData.LockReadOnly();
-			}
+			// Editor heightfields don't have material indices (hence, no holes), in which case InHeightfield.GeomData.MaterialIndices.Num() == 1 : 
+			const int32 NumMaterialIndices = InHeightfield.GeomData.MaterialIndices.Num();
+			const bool bHasMaterialIndices = (NumMaterialIndices > 1);
+			check(!bHasMaterialIndices || (NumMaterialIndices == ((NumRows - 1) * (NumCols - 1))));
 
 			int32 TriangleIdx = 0;
-			for (int32 Y = 0; Y < InComponent->CollisionSizeQuads; Y++)
+			for (int32 Y = 0; Y < (NumRows - 1); Y++)
 			{
-				for (int32 X = 0; X < InComponent->CollisionSizeQuads; X++)
+				for (int32 X = 0; X < (NumCols - 1); X++)
 				{
-					int32 DataIdx = X + Y * CollisionSizeVerts;
+					int32 DataIdx = X + Y * NumCols;
 					bool bHole = false;
 
-					if (bInShowHoles && DominantLayers)
+					if (bHasMaterialIndices)
 					{
-						uint8 DominantLayerIdx = DominantLayers[DataIdx];
-						if (InComponent->ComponentLayerInfos.IsValidIndex(DominantLayerIdx) && 
-							InComponent->ComponentLayerInfos[DominantLayerIdx] == ALandscapeProxy::VisibilityLayer)
-						{
-							// If it's a hole, override with the hole flag.
-							bHole = true;
-						}
+						// Material indices don't have the final row/column : 
+						int32 MaterialIndicesDataIdx = X + Y * (NumCols - 1);
+						uint8 LayerIdx = InHeightfield.GeomData.MaterialIndices[MaterialIndicesDataIdx];
+						bHole = (LayerIdx == TNumericLimits<uint8>::Max());
 					}
 
 					if (bHole)
 					{
-						Indices[TriangleIdx + 0] = (X + 0) + (Y + 0) * CollisionSizeVerts;
+						Indices[TriangleIdx + 0] = (X + 0) + (Y + 0) * NumCols;
 						Indices[TriangleIdx + 1] = Indices[TriangleIdx + 0];
 						Indices[TriangleIdx + 2] = Indices[TriangleIdx + 0];
 					}
 					else
 					{
-						Indices[TriangleIdx + 0] = (X + 0) + (Y + 0) * CollisionSizeVerts;
-						Indices[TriangleIdx + 1] = (X + 1) + (Y + 1) * CollisionSizeVerts;
-						Indices[TriangleIdx + 2] = (X + 1) + (Y + 0) * CollisionSizeVerts;
+						Indices[TriangleIdx + 0] = (X + 0) + (Y + 0) * NumCols;
+						Indices[TriangleIdx + 1] = (X + 1) + (Y + 1) * NumCols;
+						Indices[TriangleIdx + 2] = (X + 1) + (Y + 0) * NumCols;
 					}
 
 					TriangleIdx += 3;
 
 					if (bHole)
 					{
-						Indices[TriangleIdx + 0] = (X + 0) + (Y + 0) * CollisionSizeVerts;
+						Indices[TriangleIdx + 0] = (X + 0) + (Y + 0) * NumCols;
 						Indices[TriangleIdx + 1] = Indices[TriangleIdx + 0];
 						Indices[TriangleIdx + 2] = Indices[TriangleIdx + 0];
 					}
 					else
 					{
-						Indices[TriangleIdx + 0] = (X + 0) + (Y + 0) * CollisionSizeVerts;
-						Indices[TriangleIdx + 1] = (X + 0) + (Y + 1) * CollisionSizeVerts;
-						Indices[TriangleIdx + 2] = (X + 1) + (Y + 1) * CollisionSizeVerts;
+						Indices[TriangleIdx + 0] = (X + 0) + (Y + 0) * NumCols;
+						Indices[TriangleIdx + 1] = (X + 0) + (Y + 1) * NumCols;
+						Indices[TriangleIdx + 2] = (X + 1) + (Y + 1) * NumCols;
 					}
 
 					TriangleIdx += 3;
 				}
 			}
 
-			if (DominantLayers)
-			{
-				InComponent->DominantLayerData.Unlock();
-			}
-
 			WireframeMaterialInstance.Reset(new FColoredMaterialRenderProxy(
 				GEngine->WireframeMaterial ? GEngine->WireframeMaterial->GetRenderProxy() : NULL,
 				InWireframeColor));
@@ -581,21 +569,25 @@ FPrimitiveSceneProxy* ULandscapeHeightfieldCollisionComponent::CreateSceneProxy(
 
 		virtual void GetDynamicMeshElements(const TArray<const FSceneView*>& Views, const FSceneViewFamily& ViewFamily, uint32 VisibilityMap, FMeshElementCollector& Collector) const override
 		{
-			const FMatrix& LocalToWorld = GetLocalToWorld();
+			FMatrix LocalToWorldNoScale = GetLocalToWorld();
+			LocalToWorldNoScale.RemoveScaling();
 
 			const bool bDrawCollision = ViewFamily.EngineShowFlags.Collision && IsCollisionEnabled();
 
 			if (bDrawCollision && AllowDebugViewmodes() && WireframeMaterialInstance.IsValid())
 			{
-				// Set up mesh builder
-				FDynamicMeshBuilder MeshBuilder(ERHIFeatureLevel::Type::SM6);
-				MeshBuilder.AddVertices(Vertices);
-				MeshBuilder.AddTriangles(Indices);
-
-				// For each view, pass the mesh to the collector
-				for (int32 ViewIndex = 0; ViewIndex != Views.Num(); ++ViewIndex)
+				for (int32 ViewIndex = 0; ViewIndex < Views.Num(); ViewIndex++)
 				{
-					MeshBuilder.GetMesh(LocalToWorld, WireframeMaterialInstance.Get(), SDPG_World, false, false, ViewIndex, Collector);
+					if (VisibilityMap & (1 << ViewIndex))
+					{
+						const FSceneView* View = Views[ViewIndex];
+						// Set up mesh builder
+						FDynamicMeshBuilder MeshBuilder(View->GetFeatureLevel());
+						MeshBuilder.AddVertices(Vertices);
+						MeshBuilder.AddTriangles(Indices);
+
+						MeshBuilder.GetMesh(LocalToWorldNoScale, WireframeMaterialInstance.Get(), SDPG_World, false, false, ViewIndex, Collector);
+					}
 				}
 			}
 		}
@@ -628,9 +620,8 @@ FPrimitiveSceneProxy* ULandscapeHeightfieldCollisionComponent::CreateSceneProxy(
 	{
 		const Chaos::FHeightField* LocalHeightfield = nullptr;
 		FLinearColor WireframeColor;
-		bool bShowHoles = true;
 
-		switch (static_cast<EHeightfieldSource>(CVarLandscapeShowCollisionMesh.GetValueOnAnyThread()))
+		switch (static_cast<EHeightfieldSource>(CVarLandscapeShowCollisionMesh.GetValueOnGameThread()))
 		{
 		case EHeightfieldSource::Simple:
 			if (HeightfieldRef->HeightfieldSimple.IsValid())
@@ -661,7 +652,6 @@ FPrimitiveSceneProxy* ULandscapeHeightfieldCollisionComponent::CreateSceneProxy(
 			}
 
 			WireframeColor = FColor(157, 223, 149, 255);
-			bShowHoles = false;
 			break;
 
 		default:
@@ -670,7 +660,7 @@ FPrimitiveSceneProxy* ULandscapeHeightfieldCollisionComponent::CreateSceneProxy(
 
 		if (LocalHeightfield != nullptr)
 		{
-			Proxy = new FLandscapeHeightfieldCollisionComponentSceneProxy(this, *LocalHeightfield, WireframeColor, bShowHoles);
+			Proxy = new FLandscapeHeightfieldCollisionComponentSceneProxy(this, *LocalHeightfield, WireframeColor);
 		}
 	}
 
diff --git a/Engine/Source/Runtime/Landscape/Private/LandscapeEdit.cpp b/Engine/Source/Runtime/Landscape/Private/LandscapeEdit.cpp
--- a/Engine/Source/Runtime/Landscape/Private/LandscapeEdit.cpp
+++ b/Engine/Source/Runtime/Landscape/Private/LandscapeEdit.cpp
@@ -1207,6 +1207,9 @@ void ULandscapeComponent::UpdateCollisionHeightData(const FColor* const Heightma
 		CollisionComp->RegisterComponent();
 	}
 
+	// Debug display needs to update its representation, so we invalidate the collision component's render state : 
+	CollisionComp->MarkRenderStateDirty();
+
 	// Invalidate rendered physical materials
 	// These are updated in UpdatePhysicalMaterialTasks()
  	PhysicalMaterialHash = 0;
@@ -1618,6 +1621,9 @@ void ULandscapeComponent::UpdateCollisionLayerData(const FColor* const* const We
 
 		// We do not force an update of the physics data here. We don't need the layer information in the editor and it
 		// causes problems if we update it multiple times in a single frame.
+
+		// Debug display needs to update its representation, so we invalidate the collision component's render state : 
+		CollisionComp->MarkRenderStateDirty();
 	}
 }
 
