Update water info textures whenever heightmaps are streamed in.

diff --git a/Engine/Plugins/Experimental/Water/Source/Runtime/Private/WaterSubsystem.cpp b/Engine/Plugins/Experimental/Water/Source/Runtime/Private/WaterSubsystem.cpp
--- a/Engine/Plugins/Experimental/Water/Source/Runtime/Private/WaterSubsystem.cpp
+++ b/Engine/Plugins/Experimental/Water/Source/Runtime/Private/WaterSubsystem.cpp
@@ -24,6 +24,7 @@
 #include "WaterRuntimeSettings.h"
 #include "Engine/StaticMesh.h"
 #include "WaterViewExtension.h"
+#include "Landscape/Public/LandscapeSubsystem.h"
 
 #include UE_INLINE_GENERATED_CPP_BY_NAME(WaterSubsystem)
 
@@ -284,11 +285,39 @@ void UWaterSubsystem::Initialize(FSubsystemCollectionBase& Collection)
 	}
 }
 
+void UWaterSubsystem::PostInitialize()
+{
+	Super::PostInitialize();
+
+	UWorld* World = GetWorld();
+	check(World);
+
+#if WITH_EDITOR
+	// Register for heightmap streaming notifications
+	ULandscapeSubsystem* LandscapeSubsystem = World->GetSubsystem<ULandscapeSubsystem>();
+	check(LandscapeSubsystem);
+
+	check(!OnHeightmapStreamedHandle.IsValid());
+	OnHeightmapStreamedHandle = LandscapeSubsystem->GetOnHeightmapStreamedDelegate().AddUObject(this, &UWaterSubsystem::OnHeightmapStreamed);
+#endif // WITH_EDITOR
+}
+
 void UWaterSubsystem::Deinitialize()
 {
 	UWorld* World = GetWorld();
 	check(World != nullptr);
 
+#if WITH_EDITOR
+	if (OnHeightmapStreamedHandle.IsValid())
+	{
+		if (ULandscapeSubsystem* LandscapeSubsystem = World->GetSubsystem<ULandscapeSubsystem>())
+		{
+			LandscapeSubsystem->GetOnHeightmapStreamedDelegate().Remove(OnHeightmapStreamedHandle);
+		}
+		OnHeightmapStreamedHandle.Reset();
+	}
+#endif // WITH_EDITOR
+
 	FConsoleVariableDelegate NullCallback;
 	CVarShallowWaterSimulationRenderTargetSize->SetOnChangedCallback(NullCallback);
 	CVarShallowWaterSim->SetOnChangedCallback(NullCallback);
@@ -336,6 +365,14 @@ void UWaterSubsystem::ApplyRuntimeSettings(const UWaterRuntimeSettings* Settings
 #endif // WITH_EDITOR
 }
 
+#if WITH_EDITOR
+void UWaterSubsystem::OnHeightmapStreamed(const FOnHeightmapStreamedContext& InContext)
+{
+	UE_LOG(LogWater, Warning, TEXT("OnHeightmapStreamed!  Rebuilding Water Info Texture..."));
+	MarkWaterZonesInRegionForRebuild(InContext.GetUpdateRegion(), EWaterZoneRebuildFlags::UpdateWaterInfoTexture);
+}
+#endif // WITH_EDITOR
+
 FWaterBodyManager& UWaterSubsystem::GetWaterBodyManagerInternal()
 {
 	// #todo_water [roey]: Remove these pragmas when we can move WaterBodyManager to private.
@@ -481,6 +518,24 @@ void UWaterSubsystem::MarkAllWaterZonesForRebuild(EWaterZoneRebuildFlags Rebuild
 	}
 }
 
+void UWaterSubsystem::MarkWaterZonesInRegionForRebuild(const FBox2D& InUpdateRegion, EWaterZoneRebuildFlags InRebuildFlags)
+{
+	if (UWorld* World = GetWorld())
+	{
+		for (AWaterZone* WaterZone : TActorRange<AWaterZone>(World))
+		{
+			FVector2D ZoneExtent = WaterZone->GetZoneExtent();
+			FVector2D ZoneLocation = FVector2D(WaterZone->GetActorLocation());
+			const FBox2D WaterZoneBounds(ZoneLocation - ZoneExtent * 0.5, ZoneLocation + ZoneExtent * 0.5);
+
+			if (WaterZoneBounds.Intersect(InUpdateRegion))
+			{
+				WaterZone->MarkForRebuild(InRebuildFlags);
+			}
+		}
+	}
+}
+
 void UWaterSubsystem::NotifyWaterScalabilityChangedInternal(IConsoleVariable* CVar)
 {
 	OnWaterScalabilityChanged.Broadcast();
diff --git a/Engine/Plugins/Experimental/Water/Source/Runtime/Public/WaterSubsystem.h b/Engine/Plugins/Experimental/Water/Source/Runtime/Public/WaterSubsystem.h
--- a/Engine/Plugins/Experimental/Water/Source/Runtime/Public/WaterSubsystem.h
+++ b/Engine/Plugins/Experimental/Water/Source/Runtime/Public/WaterSubsystem.h
@@ -90,6 +90,7 @@ public:
 
 	// USubsystem implementation Begin
 	virtual void Initialize(FSubsystemCollectionBase& Collection) override;
+	virtual void PostInitialize() override;
 	virtual void Deinitialize() override;
 	// USubsystem implementation End
 
@@ -174,6 +175,7 @@ public:
 	UMaterialParameterCollection* GetMaterialParameterCollection() const {	return MaterialParameterCollection; }
 	
 	void MarkAllWaterZonesForRebuild(EWaterZoneRebuildFlags RebuildFlags = EWaterZoneRebuildFlags::All);
+	void MarkWaterZonesInRegionForRebuild(const FBox2D& InUpdateRegion, EWaterZoneRebuildFlags InRebuildFlags);
 
 #if WITH_EDITOR
 	/** Little scope object to temporarily change the value of bAllowWaterSubsystemOnPreviewWorld */
@@ -253,6 +255,9 @@ private:
 	TSharedPtr<FWaterViewExtension> WaterViewExtension;
 
 #if WITH_EDITOR
+	FDelegateHandle OnHeightmapStreamedHandle;
+	void OnHeightmapStreamed(const struct FOnHeightmapStreamedContext& InContext);
+
 	/** By default, there is no water subsystem allowed on preview worlds except when explicitly requested : */
 	static bool bAllowWaterSubsystemOnPreviewWorld;
 #endif // WITH_EDITOR
diff --git a/Engine/Source/Runtime/Landscape/Classes/LandscapeComponent.h b/Engine/Source/Runtime/Landscape/Classes/LandscapeComponent.h
--- a/Engine/Source/Runtime/Landscape/Classes/LandscapeComponent.h
+++ b/Engine/Source/Runtime/Landscape/Classes/LandscapeComponent.h
@@ -1030,9 +1030,10 @@ public:
 	/** @todo document */
 	LANDSCAPE_API float GetLayerWeightAtLocation( const FVector& InLocation, ULandscapeLayerInfoObject* LayerInfo, TArray<uint8>* LayerCache = NULL, bool bUseEditingWeightmap = false);
 
-	/** Extends passed region with this component section size */
+	/** Extends passed region with this component's 2D bounds (values are in landscape quads) */
 	LANDSCAPE_API void GetComponentExtent(int32& MinX, int32& MinY, int32& MaxX, int32& MaxY) const;
 
+	/** returns the 2D bounds of this component (in landscape quads) */
 	LANDSCAPE_API FIntRect GetComponentExtent() const;
 
 	/** Updates navigation properties to match landscape actor's */
diff --git a/Engine/Source/Runtime/Landscape/Private/LandscapeEditLayers.cpp b/Engine/Source/Runtime/Landscape/Private/LandscapeEditLayers.cpp
--- a/Engine/Source/Runtime/Landscape/Private/LandscapeEditLayers.cpp
+++ b/Engine/Source/Runtime/Landscape/Private/LandscapeEditLayers.cpp
@@ -16,6 +16,7 @@ LandscapeEditLayers.cpp: Landscape editing layers mode
 #include "LandscapeEditReadback.h"
 #include "LandscapeEditResources.h"
 #include "LandscapeUtils.h"
+#include "LandscapeSubsystem.h"
 
 #include "Shader.h"
 #include "GlobalShader.h"
@@ -3353,10 +3354,16 @@ bool ALandscape::PrepareTextureResources(bool bInWaitForStreaming)
 		[](const TWeakObjectPtr<UTexture2D>& Texture) { return Texture.Get(); });
 	TrackedStreamingInTextures.Empty();
 
-	// Start tracking the new textures :
+	// Textures that are still streaming in (filled out below)
 	TSet<UTexture2D*> StreamingInTexturesAfter;
 	StreamingInTexturesAfter.Reserve(TrackedStreamingInTextures.Num());
 
+	// Textures that have just completed streaming in (filled out below)
+	TSet<UTexture2D*> StreamedInTextures;
+
+	// All components containing heightmaps that have just completed streaming in (filled out below)
+	TSet<ULandscapeComponent*> StreamedInHeightmapComponents;
+
 	bool bIsReady = true;
 	Info->ForAllLandscapeProxies([&](ALandscapeProxy* Proxy)
 	{
@@ -3370,6 +3377,15 @@ bool ALandscape::PrepareTextureResources(bool bInWaitForStreaming)
 				{
 					StreamingInTexturesAfter.Add(ComponentHeightmap);
 				}
+				else
+				{
+					// If it was previously streaming in, then it has just completed.
+					if (StreamingInTexturesBefore.Contains(ComponentHeightmap))
+					{
+						StreamedInTextures.Add(ComponentHeightmap);
+						StreamedInHeightmapComponents.Add(Component);
+					}
+				}
 				bIsReady &= bIsTextureReady;
 			}
 
@@ -3381,16 +3397,44 @@ bool ALandscape::PrepareTextureResources(bool bInWaitForStreaming)
 				{
 					StreamingInTexturesAfter.Add(ComponentWeightmap);
 				}
+				else
+				{
+					// If it was previously streaming in, then it has just completed.
+					if (StreamingInTexturesBefore.Contains(ComponentWeightmap))
+					{
+						StreamedInTextures.Add(ComponentWeightmap);
+					}
+				}
 				bIsReady &= bIsTextureReady;
 			}
 		}
 	});
 
 	// The assets that were streaming in before and are not anymore can be considered streamed in: 
-	TRACE_CPUPROFILER_EVENT_SCOPE(LandscapeLayers_Difference);
-	TSet<UTexture2D*> StreamedInTextures = StreamingInTexturesBefore.Difference(StreamingInTexturesAfter);
 	InvalidateRVTForTextures(StreamedInTextures);
 
+	// If we streamed in any heightmaps, notify interested parties (i.e. water)
+	if (StreamedInHeightmapComponents.Num() > 0)
+	{
+		// Calculate update region.
+		FBox2D HeightmapUpdateRegion(ForceInit);
+		for (ULandscapeComponent* Component : StreamedInHeightmapComponents)
+		{
+			if (ALandscapeProxy* Proxy = Component->GetLandscapeProxy())
+			{
+				const FBox ProxyBox = Proxy->GetComponentsBoundingBox();
+				HeightmapUpdateRegion += FBox2D(FVector2D(ProxyBox.Min), FVector2D(ProxyBox.Max));
+			}
+		}
+
+		// Notify that heightmaps have been streamed.
+		if (ULandscapeSubsystem* LandscapeSubsystem = GetWorld()->GetSubsystem<ULandscapeSubsystem>())
+		{
+			FOnHeightmapStreamedContext Context(HeightmapUpdateRegion, StreamedInHeightmapComponents);
+			LandscapeSubsystem->GetOnHeightmapStreamedDelegate().Broadcast(Context);
+		}
+	}
+
 	// Store as a list of TWeakObjectPtr<UTexture2D> so as not to keep references on the tracked textures :
 	Algo::Transform(StreamingInTexturesAfter, TrackedStreamingInTextures, [](UTexture2D* Texture) { return TWeakObjectPtr<UTexture2D>(Texture); });
 
diff --git a/Engine/Source/Runtime/Landscape/Private/LandscapeRender.cpp b/Engine/Source/Runtime/Landscape/Private/LandscapeRender.cpp
--- a/Engine/Source/Runtime/Landscape/Private/LandscapeRender.cpp
+++ b/Engine/Source/Runtime/Landscape/Private/LandscapeRender.cpp
@@ -3643,9 +3643,10 @@ float FLandscapeComponentSceneProxy::ComputeLODBias() const
 
 double FLandscapeComponentSceneProxy::ComputeSectionResolution() const
 {
-	// LandscapeComponent Max Extend represents the half-extent of the landscape component. Multiply by two to get the actual size.
-	const double ComponentFullExtent = 2.0 * ComponentMaxExtend;
-	return ComponentFullExtent / (double)(ComponentSizeVerts);
+	// ComponentMaxExtend is the max(length,width) of the component, in world units
+	const double ComponentFullExtent = ComponentMaxExtend;
+	const double ComponentQuads = ComponentSizeVerts - 1.0;		// verts = quads + 1
+	return ComponentFullExtent / ComponentQuads;
 }
 
 //
diff --git a/Engine/Source/Runtime/Landscape/Public/LandscapeSubsystem.h b/Engine/Source/Runtime/Landscape/Public/LandscapeSubsystem.h
--- a/Engine/Source/Runtime/Landscape/Public/LandscapeSubsystem.h
+++ b/Engine/Source/Runtime/Landscape/Public/LandscapeSubsystem.h
@@ -11,6 +11,25 @@ class ALandscapeProxy;
 class ULandscapeInfo;
 class FLandscapeNotificationManager;
 
+#if WITH_EDITOR
+struct FOnHeightmapStreamedContext
+{
+private:
+	const FBox2D& UpdateRegion;
+	const TSet<class ULandscapeComponent*>& LandscapeComponentsInvolved;
+
+public:
+	const FBox2D& GetUpdateRegion() const { return UpdateRegion; }
+	const TSet<class ULandscapeComponent*>& GetLandscapeComponentsInvolved() const { return LandscapeComponentsInvolved; }
+
+	FOnHeightmapStreamedContext(const FBox2D& updateRegion, const TSet<class ULandscapeComponent*>& landscapeComponentsInvolved)
+		: UpdateRegion(updateRegion), LandscapeComponentsInvolved(landscapeComponentsInvolved)
+	{}
+};
+
+DECLARE_MULTICAST_DELEGATE_OneParam(FOnHeightmapStreamedDelegate, const FOnHeightmapStreamedContext& context);
+#endif // WITH_EDITOR
+
 UCLASS(MinimalAPI)
 class ULandscapeSubsystem : public UTickableWorldSubsystem
 {
@@ -48,7 +67,8 @@ public:
 	LANDSCAPE_API bool HasModifiedLandscapes() const;
 	LANDSCAPE_API static bool IsDirtyOnlyInModeEnabled();
 	FLandscapeNotificationManager* GetNotificationManager() { return NotificationManager; }
-#endif
+	FOnHeightmapStreamedDelegate& GetOnHeightmapStreamedDelegate() { return OnHeightmapStreamed; }
+#endif // WITH_EDITOR
 
 private:
 #if WITH_EDITOR
@@ -68,5 +88,6 @@ private:
 	class FLandscapePhysicalMaterialBuilder* PhysicalMaterialBuilder;
 	
 	FLandscapeNotificationManager* NotificationManager;
-#endif
+	FOnHeightmapStreamedDelegate OnHeightmapStreamed;
+#endif // WITH_EDITOR
 };
