Landscape: prevent landscape heightmap rendering from being disabled when nanite showflag is disabled. We can use the heightmaps as fallback in this case.
This is required for water heightmap rendering as well since we do not want to enable nanite meshes for that yet due to the issue of seams on the borders of the components' nanite meshes.

diff --git a/Engine/Source/Runtime/Landscape/Private/LandscapeRender.cpp b/Engine/Source/Runtime/Landscape/Private/LandscapeRender.cpp
--- a/Engine/Source/Runtime/Landscape/Private/LandscapeRender.cpp
+++ b/Engine/Source/Runtime/Landscape/Private/LandscapeRender.cpp
@@ -1391,7 +1391,7 @@ FPrimitiveViewRelevance FLandscapeComponentSceneProxy::GetViewRelevance(const FS
 	}
 #endif // !UE_BUILD_SHIPPING
 
-	if (bNaniteActive)
+	if (bNaniteActive && View->Family->EngineShowFlags.NaniteMeshes)
 	{
 		Result.bShadowRelevance = false;
 		Result.bVelocityRelevance = false;
