Editor Crashes in Landscape Mode after attempted to create a new Landscape

#rnx

diff --git a/Engine/Source/Runtime/Landscape/Private/Landscape.cpp b/Engine/Source/Runtime/Landscape/Private/Landscape.cpp
--- a/Engine/Source/Runtime/Landscape/Private/Landscape.cpp
+++ b/Engine/Source/Runtime/Landscape/Private/Landscape.cpp
@@ -2929,9 +2929,12 @@ void ALandscapeProxy::PreSave(FObjectPreSaveContext ObjectSaveContext)
 		UpdateRenderingMethod();
 	}
 
-	if (ULandscapeInfo* LandscapeInfo = GetLandscapeInfo())
+	if (LandscapeGuid.IsValid())
 	{
-		LandscapeInfo->OnModifiedPackageSaved(GetPackage());
+		if (ULandscapeInfo* LandscapeInfo = GetLandscapeInfo())
+		{
+			LandscapeInfo->OnModifiedPackageSaved(GetPackage());
+		}
 	}
 #endif // WITH_EDITOR
 }
