Fix crash when saving an existing already saved map containing a dirty landscape using Save As.. to rename it.

diff --git a/Engine/Source/Runtime/Landscape/Private/LandscapeSubsystem.cpp b/Engine/Source/Runtime/Landscape/Private/LandscapeSubsystem.cpp
--- a/Engine/Source/Runtime/Landscape/Private/LandscapeSubsystem.cpp
+++ b/Engine/Source/Runtime/Landscape/Private/LandscapeSubsystem.cpp
@@ -90,6 +90,13 @@ ETickableTickType ULandscapeSubsystem::GetTickableTickType() const
 	return HasAnyFlags(RF_ClassDefaultObject) || !GetWorld() || GetWorld()->IsNetMode(NM_DedicatedServer) ? ETickableTickType::Never : ETickableTickType::Always;
 }
 
+bool ULandscapeSubsystem::DoesSupportWorldType(const EWorldType::Type WorldType) const
+{
+	// we also support inactive worlds -- they are used when the world is already saved, but SaveAs renames it:
+	// then it duplicates the world (producing an inactive world), which we then need to update Landscapes in during OnPreSave()
+	return Super::DoesSupportWorldType(WorldType) || WorldType == EWorldType::Inactive;
+}
+
 void ULandscapeSubsystem::Tick(float DeltaTime)
 {
 	SCOPE_CYCLE_COUNTER(STAT_LandscapeSubsystemTick);
diff --git a/Engine/Source/Runtime/Landscape/Public/LandscapeSubsystem.h b/Engine/Source/Runtime/Landscape/Public/LandscapeSubsystem.h
--- a/Engine/Source/Runtime/Landscape/Public/LandscapeSubsystem.h
+++ b/Engine/Source/Runtime/Landscape/Public/LandscapeSubsystem.h
@@ -46,6 +46,7 @@ public:
 	virtual void Tick(float DeltaTime) override;
 	virtual bool IsTickableInEditor() const override { return true; }
 	virtual ETickableTickType GetTickableTickType() const override;
+	virtual bool DoesSupportWorldType(const EWorldType::Type WorldType) const override;
 	virtual TStatId GetStatId() const override;
 	// End FTickableGameObject overrides
 
