Fixed assert when generating grass while there's actually no grass data : in the editor, in that case FLandscapeComponentGrassData::HasData() will return true because we may have requested some mip data in HeightMipData (for mip > 0) but the actual height info, located at the beggining of HeightWeightData might actually not be there, just because we removed all grass height/weight data since no grass data is needed because all grass weights are 0.

diff --git a/Engine/Source/Runtime/Landscape/Private/LandscapeGrass.cpp b/Engine/Source/Runtime/Landscape/Private/LandscapeGrass.cpp
--- a/Engine/Source/Runtime/Landscape/Private/LandscapeGrass.cpp
+++ b/Engine/Source/Runtime/Landscape/Private/LandscapeGrass.cpp
@@ -1414,7 +1414,7 @@ bool FLandscapeComponentGrassData::HasData() const
 
 TArrayView<uint8> FLandscapeComponentGrassData::GetWeightData(const ULandscapeGrassType* GrassType)
 {
-	if (HasData())
+	if (!HeightWeightData.IsEmpty())
 	{
 		if (int32* OffsetPtr = WeightOffsets.Find(GrassType))
 		{
@@ -1435,7 +1435,7 @@ bool FLandscapeComponentGrassData::Contains(ULandscapeGrassType* GrassType) cons
 
 TArrayView<uint16> FLandscapeComponentGrassData::GetHeightData()
 {
-	if (!HasData())
+	if (HeightWeightData.IsEmpty())
 	{
 		return TArrayView<uint16>();
 	}
