Chaos - fix material resolve for landscapes with a simple collision mip. Previously we stuck to one material for simple heightfields instead of pulling the full materials from the dominant layer data

diff --git a/Engine/Source/Runtime/Landscape/Private/LandscapeCollision.cpp b/Engine/Source/Runtime/Landscape/Private/LandscapeCollision.cpp
--- a/Engine/Source/Runtime/Landscape/Private/LandscapeCollision.cpp
+++ b/Engine/Source/Runtime/Landscape/Private/LandscapeCollision.cpp
@@ -152,7 +152,7 @@ void ULandscapeMeshCollisionComponent::FTriMeshGeometryRef::GetResourceSizeEx(FR
 }
 
 // Generate a new guid to force a recache of landscape collison derived data
-#define LANDSCAPE_COLLISION_DERIVEDDATA_VER	TEXT("AB01CAE4167F48E3A3ED396A900F77AC")
+#define LANDSCAPE_COLLISION_DERIVEDDATA_VER	TEXT("75E2F3A08BE44420813DD2F2AD34021D")
 
 static FString GetHFDDCKeyString(const FName& Format, bool bDefMaterial, const FGuid& StateId, const TArray<UPhysicalMaterial*>& PhysicalMaterials)
 {
@@ -353,6 +353,7 @@ void ULandscapeHeightfieldCollisionComponent::OnCreatePhysicsState()
 
 					NewSimpleShape->SetQueryData(QueryFilterDataSimple);
 					NewSimpleShape->SetSimData(SimFilterDataSimple);
+					NewSimpleShape->SetMaterials(HeightfieldRef->UsedChaosMaterials);
 
 					Geoms.Emplace(MoveTemp(ChaosSimpleHeightFieldFromCooked));
 					ShapeArray.Emplace(MoveTemp(NewSimpleShape));
@@ -913,43 +914,50 @@ bool ULandscapeHeightfieldCollisionComponent::CookCollisionData(const FName& For
 	// Generate material indices
 	TArray<uint8> MaterialIndices;
 	MaterialIndices.Reserve(NumSamples + NumSimpleSamples);
-	for(int32 RowIndex = 0; RowIndex < CollisionSizeVerts; RowIndex++)
+
+	auto ResolveMaterials = [&MaterialIndices, &bIsMirrored, &bUseDefMaterial, &DefMaterial, &InOutMaterials, this](int32 InCollisionVertExtent, const uint8* InDominantLayers, const uint8* InRenderMaterialIds)
 	{
-		for(int32 ColIndex = 0; ColIndex < CollisionSizeVerts; ColIndex++)
+		for(int32 RowIndex = 0; RowIndex < InCollisionVertExtent; RowIndex++)
 		{
-			const int32 SrcSampleIndex = (RowIndex * CollisionSizeVerts) + (bIsMirrored ? (CollisionSizeVerts - ColIndex - 1) : ColIndex);
-
-			// Materials are not relevant on the last row/column because they are per-triangle and the last row/column don't own any
-			if(RowIndex < CollisionSizeVerts - 1 &&
-				ColIndex < CollisionSizeVerts - 1)
+			for(int32 ColIndex = 0; ColIndex < InCollisionVertExtent; ColIndex++)
 			{
-				int32 MaterialIndex = 0; // Default physical material.
-				if(!bUseDefMaterial)
-				{
-					uint8 DominantLayerIdx = DominantLayers ? DominantLayers[SrcSampleIndex] : -1;
-					ULandscapeLayerInfoObject* Layer = ComponentLayerInfos.IsValidIndex(DominantLayerIdx) ? ToRawPtr(ComponentLayerInfos[DominantLayerIdx]) : nullptr;
+				const int32 SrcSampleIndex = (RowIndex * InCollisionVertExtent) + (bIsMirrored ? (InCollisionVertExtent - ColIndex - 1) : ColIndex);
 
-					if(Layer == ALandscapeProxy::VisibilityLayer)
-					{
-						// If it's a hole, use the final index
-						MaterialIndex = TNumericLimits<uint8>::Max();
-					}
-					else if (RenderPhysicalMaterialIds)
-					{
-						uint8 RenderIdx = RenderPhysicalMaterialIds[SrcSampleIndex];
-						UPhysicalMaterial* DominantMaterial = RenderIdx > 0 ? ToRawPtr(PhysicalMaterialRenderObjects[RenderIdx - 1]) : DefMaterial;
-						MaterialIndex = InOutMaterials.AddUnique(DominantMaterial);
-					}
-					else
+				// Materials are not relevant on the last row/column because they are per-triangle and the last row/column don't own any
+				if(RowIndex < InCollisionVertExtent - 1 &&
+				   ColIndex < InCollisionVertExtent - 1)
+				{
+					int32 MaterialIndex = 0; // Default physical material.
+					if(!bUseDefMaterial)
 					{
-						UPhysicalMaterial* DominantMaterial = Layer && Layer->PhysMaterial ? ToRawPtr(Layer->PhysMaterial) : DefMaterial;
-						MaterialIndex = InOutMaterials.AddUnique(DominantMaterial);
+						uint8 DominantLayerIdx = InDominantLayers ? InDominantLayers[SrcSampleIndex] : -1;
+						ULandscapeLayerInfoObject* Layer = ComponentLayerInfos.IsValidIndex(DominantLayerIdx) ? ToRawPtr(ComponentLayerInfos[DominantLayerIdx]) : nullptr;
+
+						if(Layer == ALandscapeProxy::VisibilityLayer)
+						{
+							// If it's a hole, use the final index
+							MaterialIndex = TNumericLimits<uint8>::Max();
+						}
+						else if(InRenderMaterialIds)
+						{
+							uint8 RenderIdx = InRenderMaterialIds[SrcSampleIndex];
+							UPhysicalMaterial* DominantMaterial = RenderIdx > 0 ? ToRawPtr(PhysicalMaterialRenderObjects[RenderIdx - 1]) : DefMaterial;
+							MaterialIndex = InOutMaterials.AddUnique(DominantMaterial);
+						}
+						else
+						{
+							UPhysicalMaterial* DominantMaterial = Layer && Layer->PhysMaterial ? ToRawPtr(Layer->PhysMaterial) : DefMaterial;
+							MaterialIndex = InOutMaterials.AddUnique(DominantMaterial);
+						}
 					}
+					MaterialIndices.Add(MaterialIndex);
 				}
-				MaterialIndices.Add(MaterialIndex);
 			}
 		}
-	}
+	};
+
+	ResolveMaterials(CollisionSizeVerts, DominantLayers, RenderPhysicalMaterialIds);
+	ResolveMaterials(SimpleCollisionSizeVerts, SimpleDominantLayers, SimpleRenderPhysicalMaterialIds);
 
 	TUniquePtr<Chaos::FHeightField> Heightfield = nullptr;
 	TUniquePtr<Chaos::FHeightField> HeightfieldSimple = nullptr;
@@ -960,14 +968,18 @@ bool ULandscapeHeightfieldCollisionComponent::CookCollisionData(const FName& For
 	bool bSerializeGenerateSimpleCollision = bGenerateSimpleCollision;
 	Ar << bSerializeGenerateSimpleCollision;
 
+	const int32 NumCollisionCells = FMath::Square(CollisionSizeQuads);
+	const int32 NumSimpleCollisionCells = FMath::Square(SimpleCollisionSizeQuads);
+
 	TArrayView<const uint16> ComplexHeightView(Heights, NumSamples);
-	Heightfield = MakeUnique<Chaos::FHeightField>(ComplexHeightView, MakeArrayView(MaterialIndices), CollisionSizeVerts, CollisionSizeVerts, Chaos::FVec3(1));
+	TArrayView<uint8> ComplexMaterialIndicesView(MaterialIndices.GetData(), NumCollisionCells);
+	Heightfield = MakeUnique<Chaos::FHeightField>(ComplexHeightView, ComplexMaterialIndicesView, CollisionSizeVerts, CollisionSizeVerts, Chaos::FVec3(1));
 	Ar << Heightfield;
 	if(bGenerateSimpleCollision)
 	{
-		// #BGTODO Materials for simple geometry, currently just passing in the default
-		TArrayView<const uint16> SimpleHeightView(Heights + NumSamples, NumSimpleSamples);
-		HeightfieldSimple = MakeUnique<Chaos::FHeightField>(SimpleHeightView, MakeArrayView(MaterialIndices.GetData(), 1), SimpleCollisionSizeVerts, SimpleCollisionSizeVerts, Chaos::FVec3(1));
+		TArrayView<const uint16> SimpleHeightView(SimpleHeights, NumSimpleSamples);
+		TArrayView<uint8> SimpleMaterialIndicesView(MaterialIndices.GetData() + NumCollisionCells, NumSimpleCollisionCells);
+		HeightfieldSimple = MakeUnique<Chaos::FHeightField>(SimpleHeightView, SimpleMaterialIndicesView, SimpleCollisionSizeVerts, SimpleCollisionSizeVerts, Chaos::FVec3(1));
 		Ar << HeightfieldSimple;
 	}
 
