Landscape: removed misplaced punctuation in landscape Waiting for Resources notification toast message.

diff --git a/Engine/Source/Runtime/Landscape/Private/LandscapeEditLayers.cpp b/Engine/Source/Runtime/Landscape/Private/LandscapeEditLayers.cpp
--- a/Engine/Source/Runtime/Landscape/Private/LandscapeEditLayers.cpp
+++ b/Engine/Source/Runtime/Landscape/Private/LandscapeEditLayers.cpp
@@ -7890,7 +7890,7 @@ void ALandscape::UpdateLayersContent(bool bInWaitForStreaming, bool bInSkipMonit
 		if (!InvalidShadingModelNotification.IsValid())
 		{
 			InvalidShadingModelNotification = MakeShared<FLandscapeNotification>(this, FLandscapeNotification::EType::ShadingModelInvalid);
-			static const FText NotificationText(LOCTEXT("InvalidShadingModel", "Landscape editor cannot update landscape with a feature level less than SM5!"));
+			static const FText NotificationText(LOCTEXT("InvalidShadingModel", "Landscape editor cannot update landscape with a feature level less than SM5"));
 			InvalidShadingModelNotification->NotificationText = NotificationText;
 		}
 		LandscapeNotificationManager->RegisterNotification(InvalidShadingModelNotification);
@@ -7912,7 +7912,7 @@ void ALandscape::UpdateLayersContent(bool bInWaitForStreaming, bool bInSkipMonit
 		if (!WaitingForTexturesNotification.IsValid())
 		{
 			WaitingForTexturesNotification = MakeShared<FLandscapeNotification>(this, FLandscapeNotification::EType::LandscapeTextureResourcesNotReady);
-			static const FText NotificationText(LOCTEXT("WaitForLandscapeTextureResources", "Landscape editor waiting for landscape texture resources to be ready."));
+			static const FText NotificationText(LOCTEXT("WaitForLandscapeTextureResources", "Landscape editor waiting for landscape texture resources to be ready"));
 			WaitingForTexturesNotification->NotificationText = NotificationText;
 			WaitingForTexturesNotification->NotificationStartTime = WaitingForLandscapeTextureResourcesStartTime + TimeBeforeDisplayingWaitingForResourcesNotification;
 		}
@@ -7932,7 +7932,7 @@ void ALandscape::UpdateLayersContent(bool bInWaitForStreaming, bool bInSkipMonit
 		if (!WaitingForBrushesNotification.IsValid())
 		{
 			WaitingForBrushesNotification = MakeShared<FLandscapeNotification>(this, FLandscapeNotification::EType::LandscapeBrushResourcesNotReady);
-			static const FText NotificationText(LOCTEXT("WaitForLandscapeBrushResources", "Landscape editor waiting for landscape brush resources to be ready."));
+			static const FText NotificationText(LOCTEXT("WaitForLandscapeBrushResources", "Landscape editor waiting for landscape brush resources to be ready"));
 			WaitingForBrushesNotification->NotificationText = NotificationText;
 			WaitingForBrushesNotification->NotificationStartTime = WaitingForLandscapeBrushResourcesStartTime + TimeBeforeDisplayingWaitingForResourcesNotification;
 		}
