Ray tracing - disable adding Landscape Grass to ray tracing scene by default. Usually we don't want to trace against it due to noise and performance.

diff --git a/Engine/Source/Runtime/Landscape/Private/LandscapeGrass.cpp b/Engine/Source/Runtime/Landscape/Private/LandscapeGrass.cpp
--- a/Engine/Source/Runtime/Landscape/Private/LandscapeGrass.cpp
+++ b/Engine/Source/Runtime/Landscape/Private/LandscapeGrass.cpp
@@ -195,7 +195,7 @@ static FAutoConsoleVariableRef CVarCaptureNextGrassUpdate(
 #if RHI_RAYTRACING
 static TAutoConsoleVariable<int32> CVarRayTracingLandscapeGrass(
 	TEXT("r.RayTracing.Geometry.LandscapeGrass"),
-	1,
+	0,
 	TEXT("Include landscapes grass in ray tracing effects (default = 1)"));
 #endif
 
