1.  Allow for skiping the UpdateAllAddCollisions task when registering a landscape actor.
2.  Use an ArrayView to access the sourcemip data in ULandscapeComponent::AddDefaultLayerData.  An invalid access could happen if the components scalebias is not set correctly.
3.  Prevent crash when updating materials on a landscape component that does not have a collision component

diff --git a/Engine/Source/Runtime/Landscape/Classes/LandscapeInfo.h b/Engine/Source/Runtime/Landscape/Classes/LandscapeInfo.h
--- a/Engine/Source/Runtime/Landscape/Classes/LandscapeInfo.h
+++ b/Engine/Source/Runtime/Landscape/Classes/LandscapeInfo.h
@@ -361,7 +361,7 @@ public:
  *  @param	Proxy		Landscape actor to register
  *  @param  bMapCheck	Whether to warn about landscape errors
  */
-	LANDSCAPE_API void RegisterActor(ALandscapeProxy* Proxy, bool bMapCheck = false);
+	LANDSCAPE_API void RegisterActor(ALandscapeProxy* Proxy, bool bMapCheck = false, bool bUpdateAllAddCollisions = true);
 
 	/** Deassociates passed actor with this info object*/
 	LANDSCAPE_API void UnregisterActor(ALandscapeProxy* Proxy);
diff --git a/Engine/Source/Runtime/Landscape/Classes/LandscapeProxy.h b/Engine/Source/Runtime/Landscape/Classes/LandscapeProxy.h
--- a/Engine/Source/Runtime/Landscape/Classes/LandscapeProxy.h
+++ b/Engine/Source/Runtime/Landscape/Classes/LandscapeProxy.h
@@ -935,7 +935,7 @@ public:
 	static void AddReferencedObjects(UObject* InThis, FReferenceCollector& Collector);
 	virtual void PostLoad() override;
 
-	LANDSCAPE_API ULandscapeInfo* CreateLandscapeInfo(bool bMapCheck = false);
+	LANDSCAPE_API ULandscapeInfo* CreateLandscapeInfo(bool bMapCheck = false, bool bUpdateAllAddCollisions = true);
 	virtual LANDSCAPE_API ULandscapeInfo* GetLandscapeInfo() const override;
 
 	/** Get the LandcapeActor-to-world transform with respect to landscape section offset*/
diff --git a/Engine/Source/Runtime/Landscape/Private/Landscape.cpp b/Engine/Source/Runtime/Landscape/Private/Landscape.cpp
--- a/Engine/Source/Runtime/Landscape/Private/Landscape.cpp
+++ b/Engine/Source/Runtime/Landscape/Private/Landscape.cpp
@@ -1378,10 +1378,10 @@ ALandscape* ALandscapeStreamingProxy::GetLandscapeActor()
 	return LandscapeActor.Get();
 }
 
-ULandscapeInfo* ALandscapeProxy::CreateLandscapeInfo(bool bMapCheck)
+ULandscapeInfo* ALandscapeProxy::CreateLandscapeInfo(bool bMapCheck, bool bUpdateAllAddCollisions)
 {
 	ULandscapeInfo* LandscapeInfo = ULandscapeInfo::FindOrCreate(GetWorld(), LandscapeGuid);
-	LandscapeInfo->RegisterActor(this, bMapCheck);
+	LandscapeInfo->RegisterActor(this, bMapCheck, bUpdateAllAddCollisions);
 	return LandscapeInfo;
 }
 
@@ -2032,7 +2032,7 @@ void ULandscapeComponent::AddDefaultLayerData(const FGuid& InLayerGuid, const TA
 			ULandscapeComponent::CreateEmptyTextureMips(NewLayerHeightmap, true);
 
 			// Init Mip0 to be at 32768 which is equal to "0"
-			FColor* Mip0Data = (FColor*)NewLayerHeightmap->Source.LockMip(0);
+			TArrayView<FColor> Mip0Data((FColor*)NewLayerHeightmap->Source.LockMip(0), NewLayerHeightmap->Source.GetSizeX() * NewLayerHeightmap->Source.GetSizeY());
 
 			for (ULandscapeComponent* ComponentUsingHeightmap : InComponentsUsingHeightmap)
 			{
@@ -4119,7 +4119,7 @@ void ULandscapeInfo::ForAllLandscapeProxies(TFunctionRef<void(ALandscapeProxy*)>
 	}
 }
 
-void ULandscapeInfo::RegisterActor(ALandscapeProxy* Proxy, bool bMapCheck)
+void ULandscapeInfo::RegisterActor(ALandscapeProxy* Proxy, bool bMapCheck, bool bUpdateAllAddCollisions)
 {
 	UWorld* OwningWorld = Proxy->GetWorld();
 	// do not pass here invalid actors
@@ -4223,7 +4223,10 @@ void ULandscapeInfo::RegisterActor(ALandscapeProxy* Proxy, bool bMapCheck)
 
 #if WITH_EDITOR
 	UpdateLayerInfoMap(Proxy);
-	UpdateAllAddCollisions();
+	if(bUpdateAllAddCollisions)
+	{
+		UpdateAllAddCollisions();
+	}
 	RegisterSplineActor(Proxy);
 #endif // WITH_EDITOR
 
diff --git a/Engine/Source/Runtime/Landscape/Private/LandscapeEdit.cpp b/Engine/Source/Runtime/Landscape/Private/LandscapeEdit.cpp
--- a/Engine/Source/Runtime/Landscape/Private/LandscapeEdit.cpp
+++ b/Engine/Source/Runtime/Landscape/Private/LandscapeEdit.cpp
@@ -1720,7 +1720,7 @@ bool ULandscapeComponent::CanUpdatePhysicalMaterial()
 		return false;
 	}
 
-	return true;
+	return CollisionComponent != nullptr;
 }
 
 
