Fix for crashes involving Landscape paint undo/redo

diff --git a/Engine/Source/Runtime/Landscape/Classes/Landscape.h b/Engine/Source/Runtime/Landscape/Classes/Landscape.h
--- a/Engine/Source/Runtime/Landscape/Classes/Landscape.h
+++ b/Engine/Source/Runtime/Landscape/Classes/Landscape.h
@@ -352,7 +352,6 @@ public:
 	LANDSCAPE_API void InitializeLandscapeLayersWeightmapUsage();
 
 	LANDSCAPE_API bool ComputeLandscapeLayerBrushInfo(FTransform& OutLandscapeTransform, FIntPoint& OutLandscapeSize, FIntPoint& OutLandscapeRenderTargetSize);
-	void RequestProxyLayersWeightmapUsageUpdate();
 	void UpdateProxyLayersWeightmapUsage();
 	void ValidateProxyLayersWeightmapUsage() const;
 
diff --git a/Engine/Source/Runtime/Landscape/Classes/LandscapeComponent.h b/Engine/Source/Runtime/Landscape/Classes/LandscapeComponent.h
--- a/Engine/Source/Runtime/Landscape/Classes/LandscapeComponent.h
+++ b/Engine/Source/Runtime/Landscape/Classes/LandscapeComponent.h
@@ -463,14 +463,10 @@ private:
 	UPROPERTY()
 	FGuid LightingGuid;
 
+	/** Edit Layers that have data for this component store it here */
 	UPROPERTY()
 	TMap<FGuid, FLandscapeLayerComponentData> LayersData;
 
-	/** Compoment's Data for Editing Layer */
-	FGuid LandscapeEditingLayer;
-	mutable FGuid CachedEditingLayer;
-	mutable FLandscapeLayerComponentData* CachedEditingLayerData;
-		
 	// Final layer data
 	UPROPERTY(Transient)
 	TArray<TObjectPtr<ULandscapeWeightmapUsage>> WeightmapTexturesUsage;
@@ -650,6 +646,7 @@ public:
 
 #if WITH_EDITOR
 	virtual void BeginCacheForCookedPlatformData(const ITargetPlatform* TargetPlatform) override;
+	virtual void PreEditUndo() override;
 	virtual void PostEditUndo() override;
 	virtual void PreEditChange(FProperty* PropertyThatWillChange) override;
 	virtual void PostEditChangeProperty(FPropertyChangedEvent& PropertyChangedEvent) override;
@@ -692,10 +689,10 @@ public:
 	LANDSCAPE_API TArray<UTexture2D*>& GetWeightmapTextures(const FGuid& InLayerGuid);
 	LANDSCAPE_API const TArray<UTexture2D*>& GetWeightmapTextures(const FGuid& InLayerGuid) const;
 
-	LANDSCAPE_API TArray<FWeightmapLayerAllocationInfo>& GetWeightmapLayerAllocations(bool InReturnEditingWeightmap = false);
-	LANDSCAPE_API const TArray<FWeightmapLayerAllocationInfo>& GetWeightmapLayerAllocations(bool InReturnEditingWeightmap = false) const;
-	LANDSCAPE_API TArray<FWeightmapLayerAllocationInfo>& GetWeightmapLayerAllocations(const FGuid& InLayerGuid);
-	LANDSCAPE_API const TArray<FWeightmapLayerAllocationInfo>& GetWeightmapLayerAllocations(const FGuid& InLayerGuid) const;
+	LANDSCAPE_API TArray<FWeightmapLayerAllocationInfo>& GetWeightmapLayerAllocations(bool InReturnEditingWeightmap = false, bool bInCheckIsUpToDate = true);
+	LANDSCAPE_API const TArray<FWeightmapLayerAllocationInfo>& GetWeightmapLayerAllocations(bool InReturnEditingWeightmap = false, bool bInCheckIsUpToDate = true) const;
+	LANDSCAPE_API TArray<FWeightmapLayerAllocationInfo>& GetWeightmapLayerAllocations(const FGuid& InLayerGuid, bool bInCheckIsUpToDate = true);
+	LANDSCAPE_API const TArray<FWeightmapLayerAllocationInfo>& GetWeightmapLayerAllocations(const FGuid& InLayerGuid, bool bInCheckIsUpToDate = true) const;
 
 	const TArray<FLandscapePerLODMaterialOverride>& GetPerLODOverrideMaterials() const { return PerLODOverrideMaterials; }
 	void SetPerLODOverrideMaterials(const TArray<FLandscapePerLODMaterialOverride>& InValue) { PerLODOverrideMaterials = InValue; }
@@ -709,10 +706,10 @@ public:
 
 	LANDSCAPE_API void SetWeightmapLayerAllocations(const TArray<FWeightmapLayerAllocationInfo>& InNewWeightmapLayerAllocations);
 	LANDSCAPE_API void SetWeightmapTexturesUsage(const TArray<ULandscapeWeightmapUsage*>& InNewWeightmapTexturesUsage, bool InApplyToEditingWeightmap = false);
-	LANDSCAPE_API TArray<ULandscapeWeightmapUsage*>& GetWeightmapTexturesUsage(bool InReturnEditingWeightmap = false);
-	LANDSCAPE_API const TArray<ULandscapeWeightmapUsage*>& GetWeightmapTexturesUsage(bool InReturnEditingWeightmap = false) const;
-	LANDSCAPE_API TArray<ULandscapeWeightmapUsage*>& GetWeightmapTexturesUsage(const FGuid& InLayerGuid);
-	LANDSCAPE_API const TArray<ULandscapeWeightmapUsage*>& GetWeightmapTexturesUsage(const FGuid& InLayerGuid) const;
+	LANDSCAPE_API TArray<ULandscapeWeightmapUsage*>& GetWeightmapTexturesUsage(bool InReturnEditingWeightmap = false, bool bInCheckIsUpToDate = true);
+	LANDSCAPE_API const TArray<ULandscapeWeightmapUsage*>& GetWeightmapTexturesUsage(bool InReturnEditingWeightmap = false, bool bInCheckIsUpToDate = true) const;
+	LANDSCAPE_API TArray<ULandscapeWeightmapUsage*>& GetWeightmapTexturesUsage(const FGuid& InLayerGuid, bool bInCheckIsUpToDate = true);
+	LANDSCAPE_API const TArray<ULandscapeWeightmapUsage*>& GetWeightmapTexturesUsage(const FGuid& InLayerGuid, bool bInCheckIsUpToDate = true) const;
 	LANDSCAPE_API void InitializeLayersWeightmapUsage(const FGuid& InLayerGuid);
 
 	LANDSCAPE_API bool HasLayersData() const;
@@ -723,9 +720,14 @@ public:
 	LANDSCAPE_API void RemoveLayerData(const FGuid& InLayerGuid);
 	LANDSCAPE_API void ForEachLayer(TFunctionRef<void(const FGuid&, struct FLandscapeLayerComponentData&)> Fn);
 
-	LANDSCAPE_API void SetEditingLayer(const FGuid& InEditingLayer);
+	UE_DEPRECATED(5.1, "SetEditingLayer has been deprecated, use GetLandscapeActor()->SetEditingLayer() instead.")
+	LANDSCAPE_API void SetEditingLayer(const FGuid& InEditingLayer) {}
+
+	/** Get the Landscape Actor's editing layer data */
 	FLandscapeLayerComponentData* GetEditingLayer();
 	const FLandscapeLayerComponentData* GetEditingLayer() const;
+
+	/** Get the Landscape Actor's editing layer GUID */
 	FGuid GetEditingLayerGUID() const;
 
 	void CopyFinalLayerIntoEditingLayer(FLandscapeEditDataInterface& DataInterface, TSet<UTexture2D*>& ProcessedHeightmaps);
@@ -769,11 +771,14 @@ public:
 
 #if WITH_EDITOR
 
-	/** Deletes a layer from this component, removing all its data, adjusting other layer's weightmaps if necessary, etc. */
+	/** Deletes a material layer from the current edit layer on this component, removing all its data, adjusting other layer's weightmaps if necessary, etc. */
 	LANDSCAPE_API void DeleteLayer(ULandscapeLayerInfoObject* LayerInfo, FLandscapeEditDataInterface& LandscapeEdit);
+	
+	/** Deletes a material layer from the specified edit layer on this component, removing all its data, adjusting other layer's weightmaps if necessary, etc. */
+	void DeleteLayerInternal(ULandscapeLayerInfoObject* LayerInfo, FLandscapeEditDataInterface& LandscapeEdit, const FGuid& EditLayerGuid);
 
 	/** Deletes a layer from this component, but doesn't do anything else (assumes the user knows what he's doing, use DeleteLayer otherwise) */
-	void DeleteLayerAllocation(const FGuid& InEditLayerGuid, int32 InLayerIdx, bool bInShouldDirtyPackage);
+	void DeleteLayerAllocation(const FGuid& InEditLayerGuid, int32 InLayerAllocationIdx, bool bInShouldDirtyPackage);
 
 	/** Fills a layer to 100% on this component, adding it if needed and removing other layers that get painted away */
 	LANDSCAPE_API void FillLayer(ULandscapeLayerInfoObject* LayerInfo, FLandscapeEditDataInterface& LandscapeEdit);
@@ -1097,4 +1102,11 @@ protected:
 	{
 		return true;
 	}
+
+#if WITH_EDITOR
+public:
+	/** Records the ULandscapeComponents that are modified in any undo/redo operation that is being applied currently */
+	static uint32 UndoRedoModifiedComponentCount;
+	static TArray<ULandscapeComponent*> UndoRedoModifiedComponents;
+#endif // WITH_EDITOR
 };
diff --git a/Engine/Source/Runtime/Landscape/Classes/LandscapeInfo.h b/Engine/Source/Runtime/Landscape/Classes/LandscapeInfo.h
--- a/Engine/Source/Runtime/Landscape/Classes/LandscapeInfo.h
+++ b/Engine/Source/Runtime/Landscape/Classes/LandscapeInfo.h
@@ -366,6 +366,9 @@ public:
 	/** Deassociates passed actor with this info object*/
 	LANDSCAPE_API void UnregisterActor(ALandscapeProxy* Proxy);
 
+	/** Returns true if the Proxy is currently registered with this ULandscapeInfo */
+	bool IsRegistered(const ALandscapeProxy* Proxy) const;
+
 	/** Associates passed landscape component with this info object
 	 *  @param	Component	Landscape component to register
 	 *  @param  bMapCheck	Whether to warn about landscape errors
diff --git a/Engine/Source/Runtime/Landscape/Classes/LandscapeProxy.h b/Engine/Source/Runtime/Landscape/Classes/LandscapeProxy.h
--- a/Engine/Source/Runtime/Landscape/Classes/LandscapeProxy.h
+++ b/Engine/Source/Runtime/Landscape/Classes/LandscapeProxy.h
@@ -691,9 +691,16 @@ public:
 	UPROPERTY(Transient, NonTransactional)
 	TMap<TObjectPtr<UTexture2D>, TObjectPtr<ULandscapeWeightmapUsage>> WeightmapUsageMap;
 
+	/** True when this Proxy is registered with the LandscapeInfo */
+	bool bIsRegisteredWithLandscapeInfo = false;
+
 	/** Set to true when on undo, when it's necessary to completely regenerate weightmap usages (since some weightmap allocations are transactional and others not, e.g. splines edit layer) */
 	bool bNeedsWeightmapUsagesUpdate = false;
 
+	/** CurrentVersion is bumped whenever a landscape component has an undo/redo operation applied. This lets us detect when the weightmap fixup needs to be run. */
+	uint32 CurrentVersion = 0;
+	uint32 WeightmapFixupVersion = 0;
+
 	/** Set to true when we know that weightmap usages are being reconstructed and might be temporarily invalid as a result (ValidateProxyLayersWeightmapUsage should be called after setting this back to false) */
 	bool bTemporarilyDisableWeightmapUsagesValidation = false;
 #endif
@@ -935,6 +942,7 @@ public:
 	static void AddReferencedObjects(UObject* InThis, FReferenceCollector& Collector);
 	virtual void PostLoad() override;
 
+	/** Creates the LandscapeInfo if necessary, then registers this proxy with it */
 	LANDSCAPE_API ULandscapeInfo* CreateLandscapeInfo(bool bMapCheck = false, bool bUpdateAllAddCollisions = true);
 	virtual LANDSCAPE_API ULandscapeInfo* GetLandscapeInfo() const override;
 
diff --git a/Engine/Source/Runtime/Landscape/Private/Landscape.cpp b/Engine/Source/Runtime/Landscape/Private/Landscape.cpp
--- a/Engine/Source/Runtime/Landscape/Private/Landscape.cpp
+++ b/Engine/Source/Runtime/Landscape/Private/Landscape.cpp
@@ -200,7 +200,6 @@ ULandscapeComponent::ULandscapeComponent(const FObjectInitializer& ObjectInitial
 	: Super(ObjectInitializer)
 	, bNaniteActive(false)
 #if WITH_EDITORONLY_DATA
-	, CachedEditingLayerData(nullptr)
 	, LayerUpdateFlagPerMode(0)
 	, bPendingCollisionDataUpdate(false)
 	, bPendingLayerCollisionDataUpdate(false)
@@ -1844,9 +1843,17 @@ TArray<UTexture2D*>& ULandscapeComponent::GetWeightmapTextures(const FGuid& InLa
 	return WeightmapTextures;
 }
 
-const TArray<FWeightmapLayerAllocationInfo>& ULandscapeComponent::GetWeightmapLayerAllocations(bool InReturnEditingWeightmap) const
+const TArray<FWeightmapLayerAllocationInfo>& ULandscapeComponent::GetWeightmapLayerAllocations(bool InReturnEditingWeightmap, bool bInCheckIsUpToDate) const
 {
 #if WITH_EDITORONLY_DATA
+	if (bInCheckIsUpToDate)
+	{
+		if (ALandscapeProxy *Proxy = GetLandscapeProxy())
+		{
+			check(Proxy->WeightmapFixupVersion == Proxy->CurrentVersion);
+		}
+	}
+
 	if (InReturnEditingWeightmap)
 	{
 		if (const FLandscapeLayerComponentData* EditingLayer = GetEditingLayer())
@@ -1859,9 +1866,17 @@ const TArray<FWeightmapLayerAllocationInfo>& ULandscapeComponent::GetWeightmapLa
 	return WeightmapLayerAllocations;
 }
 
-TArray<FWeightmapLayerAllocationInfo>& ULandscapeComponent::GetWeightmapLayerAllocations(const FGuid& InLayerGuid)
+TArray<FWeightmapLayerAllocationInfo>& ULandscapeComponent::GetWeightmapLayerAllocations(const FGuid& InLayerGuid, bool bInCheckIsUpToDate)
 {
 #if WITH_EDITORONLY_DATA
+	if (bInCheckIsUpToDate)
+	{
+		if (ALandscapeProxy* Proxy = GetLandscapeProxy())
+		{
+			check(Proxy->WeightmapFixupVersion == Proxy->CurrentVersion);
+		}
+	}
+
 	if (InLayerGuid.IsValid())
 	{
 		if (FLandscapeLayerComponentData* LayerData = GetLayerData(InLayerGuid))
@@ -1874,9 +1889,17 @@ TArray<FWeightmapLayerAllocationInfo>& ULandscapeComponent::GetWeightmapLayerAll
 	return WeightmapLayerAllocations;
 }
 
-const TArray<FWeightmapLayerAllocationInfo>& ULandscapeComponent::GetWeightmapLayerAllocations(const FGuid& InLayerGuid) const
+const TArray<FWeightmapLayerAllocationInfo>& ULandscapeComponent::GetWeightmapLayerAllocations(const FGuid& InLayerGuid, bool bInCheckIsUpToDate) const
 {
 #if WITH_EDITORONLY_DATA
+	if (bInCheckIsUpToDate)
+	{
+		if (ALandscapeProxy* Proxy = GetLandscapeProxy())
+		{
+			check(Proxy->WeightmapFixupVersion == Proxy->CurrentVersion);
+		}
+	}
+
 	if (InLayerGuid.IsValid())
 	{
 		if (const FLandscapeLayerComponentData* LayerData = GetLayerData(InLayerGuid))
@@ -1889,9 +1912,17 @@ const TArray<FWeightmapLayerAllocationInfo>& ULandscapeComponent::GetWeightmapLa
 	return WeightmapLayerAllocations;
 }
 
-TArray<FWeightmapLayerAllocationInfo>& ULandscapeComponent::GetWeightmapLayerAllocations(bool InReturnEditingWeightmap)
+TArray<FWeightmapLayerAllocationInfo>& ULandscapeComponent::GetWeightmapLayerAllocations(bool InReturnEditingWeightmap, bool bInCheckIsUpToDate)
 {
 #if WITH_EDITORONLY_DATA
+	if (bInCheckIsUpToDate)
+	{
+		if (ALandscapeProxy* Proxy = GetLandscapeProxy())
+		{
+			check(Proxy->WeightmapFixupVersion == Proxy->CurrentVersion);
+		}
+	}
+
 	if (InReturnEditingWeightmap)
 	{
 		if (FLandscapeLayerComponentData* EditingLayer = GetEditingLayer())
@@ -1906,29 +1937,16 @@ TArray<FWeightmapLayerAllocationInfo>& ULandscapeComponent::GetWeightmapLayerAll
 
 #if WITH_EDITOR
 
-void ULandscapeComponent::SetEditingLayer(const FGuid& InEditingLayer)
-{
-	LandscapeEditingLayer = InEditingLayer;
-}
-
 FLandscapeLayerComponentData* ULandscapeComponent::GetEditingLayer()
 {
-	if (CachedEditingLayer != LandscapeEditingLayer)
-	{
-		CachedEditingLayer = LandscapeEditingLayer;
-		CachedEditingLayerData = CachedEditingLayer.IsValid() ? LayersData.Find(CachedEditingLayer) : nullptr;
-	}
-	return CachedEditingLayerData;
+	const FGuid& EditingLayerGuid = GetLandscapeActor()->EditingLayer;
+	return EditingLayerGuid.IsValid() ? LayersData.Find(EditingLayerGuid) : nullptr;
 }
 
 const FLandscapeLayerComponentData* ULandscapeComponent::GetEditingLayer() const
 {
-	if (CachedEditingLayer != LandscapeEditingLayer)
-	{
-		CachedEditingLayer = LandscapeEditingLayer;
-		CachedEditingLayerData = CachedEditingLayer.IsValid() ? const_cast<TMap<FGuid, FLandscapeLayerComponentData>&>(LayersData).Find(CachedEditingLayer) : nullptr;
-	}
-	return CachedEditingLayerData;
+	const FGuid& EditingLayerGuid = GetLandscapeActor()->EditingLayer;
+	return EditingLayerGuid.IsValid() ? LayersData.Find(EditingLayerGuid) : nullptr;
 }
 
 void ULandscapeComponent::CopyFinalLayerIntoEditingLayer(FLandscapeEditDataInterface& DataInterface, TSet<UTexture2D*>& ProcessedHeightmaps)
@@ -2001,11 +2019,8 @@ void ULandscapeComponent::ForEachLayer(TFunctionRef<void(const FGuid&, struct FL
 void ULandscapeComponent::AddLayerData(const FGuid& InLayerGuid, const FLandscapeLayerComponentData& InData)
 {
 	Modify();
-	check(!LandscapeEditingLayer.IsValid());
 	FLandscapeLayerComponentData& Data = LayersData.FindOrAdd(InLayerGuid);
 	Data = InData;
-	CachedEditingLayer.Invalidate();
-	CachedEditingLayerData = nullptr;
 }
 
 void ULandscapeComponent::AddDefaultLayerData(const FGuid& InLayerGuid, const TArray<ULandscapeComponent*>& InComponentsUsingHeightmap, TMap<UTexture2D*, UTexture2D*>& InOutCreatedHeightmapTextures)
@@ -2089,10 +2104,7 @@ void ULandscapeComponent::AddDefaultLayerData(const FGuid& InLayerGuid, const TA
 void ULandscapeComponent::RemoveLayerData(const FGuid& InLayerGuid)
 {
 	Modify();
-	check(!LandscapeEditingLayer.IsValid());
 	LayersData.Remove(InLayerGuid);
-	CachedEditingLayer.Invalidate();
-	CachedEditingLayerData = nullptr;
 }
 
 void ULandscapeComponent::SetHeightmap(UTexture2D* NewHeightmap)
@@ -2121,8 +2133,17 @@ void ULandscapeComponent::SetWeightmapLayerAllocations(const TArray<FWeightmapLa
 	WeightmapLayerAllocations = InNewWeightmapLayerAllocations;
 }
 
-TArray<ULandscapeWeightmapUsage*>& ULandscapeComponent::GetWeightmapTexturesUsage(bool InReturnEditingWeightmap)
+TArray<ULandscapeWeightmapUsage*>& ULandscapeComponent::GetWeightmapTexturesUsage(bool InReturnEditingWeightmap, bool bInCheckIsUpToDate)
 {
+	if (bInCheckIsUpToDate)
+	{
+		if (ALandscapeProxy* Proxy = GetLandscapeProxy())
+		{
+			check(Proxy->WeightmapFixupVersion == Proxy->CurrentVersion);
+			check(!Proxy->bNeedsWeightmapUsagesUpdate);
+		}
+	}
+
 	if (InReturnEditingWeightmap)
 	{
 		if (FLandscapeLayerComponentData* EditingLayer = GetEditingLayer())
@@ -2134,8 +2155,17 @@ TArray<ULandscapeWeightmapUsage*>& ULandscapeComponent::GetWeightmapTexturesUsag
 	return WeightmapTexturesUsage;
 }
 
-const TArray<ULandscapeWeightmapUsage*>& ULandscapeComponent::GetWeightmapTexturesUsage(bool InReturnEditingWeightmap) const
+const TArray<ULandscapeWeightmapUsage*>& ULandscapeComponent::GetWeightmapTexturesUsage(bool InReturnEditingWeightmap, bool bInCheckIsUpToDate) const
 {
+	if (bInCheckIsUpToDate)
+	{
+		if (ALandscapeProxy* Proxy = GetLandscapeProxy())
+		{
+			check(Proxy->WeightmapFixupVersion == Proxy->CurrentVersion);
+			check(!Proxy->bNeedsWeightmapUsagesUpdate);
+		}
+	}
+
 	if (InReturnEditingWeightmap)
 	{
 		if (const FLandscapeLayerComponentData* EditingLayer = GetEditingLayer())
@@ -2147,8 +2177,17 @@ const TArray<ULandscapeWeightmapUsage*>& ULandscapeComponent::GetWeightmapTextur
 	return WeightmapTexturesUsage;
 }
 
-TArray<ULandscapeWeightmapUsage*>& ULandscapeComponent::GetWeightmapTexturesUsage(const FGuid& InLayerGuid)
+TArray<ULandscapeWeightmapUsage*>& ULandscapeComponent::GetWeightmapTexturesUsage(const FGuid& InLayerGuid, bool bInCheckIsUpToDate)
 {
+	if (bInCheckIsUpToDate)
+	{
+		if (ALandscapeProxy* Proxy = GetLandscapeProxy())
+		{
+			check(Proxy->WeightmapFixupVersion == Proxy->CurrentVersion);
+			check(!Proxy->bNeedsWeightmapUsagesUpdate);
+		}
+	}
+
 	if (InLayerGuid.IsValid())
 	{
 		if (FLandscapeLayerComponentData* LayerData = GetLayerData(InLayerGuid))
@@ -2160,8 +2199,17 @@ TArray<ULandscapeWeightmapUsage*>& ULandscapeComponent::GetWeightmapTexturesUsag
 	return WeightmapTexturesUsage;
 }
 
-const TArray<ULandscapeWeightmapUsage*>& ULandscapeComponent::GetWeightmapTexturesUsage(const FGuid& InLayerGuid) const
+const TArray<ULandscapeWeightmapUsage*>& ULandscapeComponent::GetWeightmapTexturesUsage(const FGuid& InLayerGuid, bool bInCheckIsUpToDate) const
 {
+	if (bInCheckIsUpToDate)
+	{
+		if (ALandscapeProxy* Proxy = GetLandscapeProxy())
+		{
+			check(Proxy->WeightmapFixupVersion == Proxy->CurrentVersion);
+			check(!Proxy->bNeedsWeightmapUsagesUpdate);
+		}
+	}
+
 	if (InLayerGuid.IsValid())
 	{
 		if (const FLandscapeLayerComponentData* LayerData = GetLayerData(InLayerGuid))
@@ -2188,19 +2236,22 @@ void ULandscapeComponent::SetWeightmapTexturesUsage(const TArray<ULandscapeWeigh
 	}
 }
 
-void ULandscapeComponent::DeleteLayerAllocation(const FGuid& InEditLayerGuid, int32 InLayerIdx, bool bInShouldDirtyPackage)
+void ULandscapeComponent::DeleteLayerAllocation(const FGuid& InEditLayerGuid, int32 InLayerAllocationIdx, bool bInShouldDirtyPackage)
 {
-	TArray<FWeightmapLayerAllocationInfo>& ComponentWeightmapLayerAllocations = GetWeightmapLayerAllocations(InEditLayerGuid);
+	ALandscapeProxy* Proxy = GetLandscapeProxy();
+
+	// If this is called during FixupWeightmaps, then we should skip the update checks
+	TArray<FWeightmapLayerAllocationInfo>& ComponentWeightmapLayerAllocations = GetWeightmapLayerAllocations(InEditLayerGuid, !Proxy->bTemporarilyDisableWeightmapUsagesValidation);
 	TArray<UTexture2D*>& ComponentWeightmapTextures = GetWeightmapTextures(InEditLayerGuid);
-	TArray<ULandscapeWeightmapUsage*>& ComponentWeightmapTexturesUsage = GetWeightmapTexturesUsage(InEditLayerGuid);
-	const FWeightmapLayerAllocationInfo& LayerAllocation = ComponentWeightmapLayerAllocations[InLayerIdx];
+	TArray<ULandscapeWeightmapUsage*>& ComponentWeightmapTexturesUsage = GetWeightmapTexturesUsage(InEditLayerGuid, !Proxy->bTemporarilyDisableWeightmapUsagesValidation);
+
+	const FWeightmapLayerAllocationInfo& LayerAllocation = ComponentWeightmapLayerAllocations[InLayerAllocationIdx];
 	const int32 DeleteLayerWeightmapTextureIndex = LayerAllocation.WeightmapTextureIndex;
 
-	ALandscapeProxy* Proxy = GetLandscapeProxy();
 	Modify(bInShouldDirtyPackage);
 	Proxy->Modify(bInShouldDirtyPackage);
 
-	// Mark the channel as unallocated, so we can reuse it later
+	// Mark the weightmap channel as unallocated, so we can reuse it later
 	ULandscapeWeightmapUsage* Usage = ComponentWeightmapTexturesUsage.IsValidIndex(DeleteLayerWeightmapTextureIndex) ? ComponentWeightmapTexturesUsage[DeleteLayerWeightmapTextureIndex] : nullptr;
 	if (Usage) // can be null if WeightmapUsageMap hasn't been built yet
 	{
@@ -2208,27 +2259,29 @@ void ULandscapeComponent::DeleteLayerAllocation(const FGuid& InEditLayerGuid, in
 	}
 
 	// Remove the layer:
-	ComponentWeightmapLayerAllocations.RemoveAt(InLayerIdx);
+	ComponentWeightmapLayerAllocations.RemoveAt(InLayerAllocationIdx);
 
-	// Check if the weightmap texture used by the layer we just removed is used by any other layer, and if so, remove the texture too
+	// Check if the weightmap texture used by the material layer we just removed is used by any other material layer -- if not then we can remove the texture from the local list (as it's not used locally)
 	bool bCanRemoveLayerTexture = !ComponentWeightmapLayerAllocations.ContainsByPredicate([DeleteLayerWeightmapTextureIndex](const FWeightmapLayerAllocationInfo& Allocation) { return Allocation.WeightmapTextureIndex == DeleteLayerWeightmapTextureIndex; });
 	if (bCanRemoveLayerTexture)
 	{
+	    // Make sure the texture can be garbage collected, if necessary
 		ComponentWeightmapTextures[DeleteLayerWeightmapTextureIndex]->ClearFlags(RF_Standalone);
+
+		// Remove from our local list of textures and usages
 		ComponentWeightmapTextures.RemoveAt(DeleteLayerWeightmapTextureIndex);
 		if (Usage)
 		{
 			ComponentWeightmapTexturesUsage.RemoveAt(DeleteLayerWeightmapTextureIndex);
 		}
 
-		// Adjust WeightmapTextureChannel index for other layers
+		// Adjust WeightmapTextureIndex for other allocations (as we just reordered the Weightmap list with the deletions above)
 		for (FWeightmapLayerAllocationInfo& Allocation : ComponentWeightmapLayerAllocations)
 		{
 			if (Allocation.WeightmapTextureIndex > DeleteLayerWeightmapTextureIndex)
 			{
 				Allocation.WeightmapTextureIndex--;
 			}
-
 			check(Allocation.WeightmapTextureIndex < ComponentWeightmapTextures.Num());
 		}
 	}
@@ -2244,10 +2297,23 @@ void ALandscapeProxy::PostRegisterAllComponents()
 	ULandscapeInfo* LandscapeInfo = nullptr;
 	if (!IsPendingKillPending())
 	{
-		// Duplicated Landscapes don't have a valid guid until PostEditImport is called, we'll register then
+		// Duplicated or newly spawned Landscapes don't have a valid guid until PostEditImport is called, we'll register then
 		if (LandscapeGuid.IsValid())
 		{
+			LandscapeInfo = GetLandscapeInfo();
+
+			// Depending what action triggered this callback, we may have already registered.  If not register now with LandscapeInfo.
+			if ((LandscapeInfo == nullptr) || !LandscapeInfo->IsRegistered(this))
+			{
+				LandscapeInfo = CreateLandscapeInfo(true);
+			}
+
 #if WITH_EDITOR
+			if (WeightmapFixupVersion != CurrentVersion)
+			{
+				FixupWeightmaps();
+			}
+
 			if (GIsEditor)
 			{
 				// Note: This can happen when loading certain cooked assets in an editor
@@ -2263,9 +2329,7 @@ void ALandscapeProxy::PostRegisterAllComponents()
 				// Check the actual flag here not HasLayersContent() which could return true if the LandscapeActor is valid.
 				bool bHasLayersContentBefore = bHasLayersContent;
 
-				LandscapeInfo = CreateLandscapeInfo(true);
-
-				FixupWeightmaps();
+				check(WeightmapFixupVersion == CurrentVersion);
 
 				const bool bNeedOldDataMigration = !bHasLayersContentBefore && CanHaveLayersContent();
 				if (bNeedOldDataMigration && LandscapeInfo->LandscapeActor.IsValid() && LandscapeInfo->LandscapeActor.Get()->HasLayersContent())
@@ -2273,11 +2337,7 @@ void ALandscapeProxy::PostRegisterAllComponents()
 					LandscapeInfo->LandscapeActor.Get()->CopyOldDataToDefaultLayer(this);
 				}
 			}
-			else
 #endif // WITH_EDITOR
-			{
-				LandscapeInfo = CreateLandscapeInfo(true);
-			}
 		}
 
 		if (UWorld* OwningWorld = GetWorld())
@@ -3146,7 +3206,10 @@ bool ULandscapeInfo::UpdateLayerInfoMapInternal(ALandscapeProxy* Proxy, bool bIn
 						}
 					}
 
-					const TArray<FWeightmapLayerAllocationInfo>& ComponentWeightmapLayerAllocations = Component->GetWeightmapLayerAllocations();
+					// Allocations are not fixed up yet when we are registering proxies and updating the layer info map -- so don't check that they are up to date.
+					// (there is a recursive dependency here, so it's not possible to order them correctly: 
+					// WeightmapFixup requires these layers to be registered with the LandscapeInfo, but these layers are registered from the WeightmapAllocations that get fixed up)
+					const TArray<FWeightmapLayerAllocationInfo>& ComponentWeightmapLayerAllocations = Component->GetWeightmapLayerAllocations(/*InReturnEditingWeightmap*/false, /*bInCheckIsUpToDate*/false);
 
 					for (int32 AllocationIndex = 0; AllocationIndex < ComponentWeightmapLayerAllocations.Num(); AllocationIndex++)
 					{
@@ -3395,6 +3458,11 @@ void ALandscapeProxy::PostLoad()
 #endif // WITH_EDITOR
 
 	UpdateRenderingMethod();
+
+#if WITH_EDITORONLY_DATA
+	// Bump current version so that WeightmapFixup runs in PostRegisterAllComponents
+	CurrentVersion++;
+#endif // WITH_EDITORONLY_DATA
 }
 
 FIntPoint ALandscapeProxy::GetSectionBaseOffset() const
@@ -4120,6 +4188,40 @@ void ULandscapeInfo::ForAllLandscapeProxies(TFunctionRef<void(ALandscapeProxy*)>
 	}
 }
 
+bool ULandscapeInfo::IsRegistered(const ALandscapeProxy* Proxy) const
+{
+	if (Proxy == nullptr)
+		return false;
+
+	bool bResult = false;
+	if (Proxy->IsA<ALandscape>())
+	{
+		bResult = (LandscapeActor.Get() == Proxy);
+	}
+	else if (const ALandscapeStreamingProxy *StreamingProxy = Cast<ALandscapeStreamingProxy>(Proxy))
+	{
+		TWeakObjectPtr<const ALandscapeStreamingProxy> StreamingProxyPtr = StreamingProxy;
+		bResult = StreamingProxies.Contains(StreamingProxyPtr);
+	}
+
+#if WITH_EDITORONLY_DATA
+	// NOTE: during an Undo operation, the LandscapeActor/StreamingProxies are transacted, and the registration status may be restored
+	// however, in that case, the Proxy is NOT fully registered yet, because some other data in LandscapeInfo still needs to be updated (XY maps for instance are not transacted)
+	// so we trust the bIsRegisteredWithLandscapeInfo flag over the actual pointers.
+
+	// at minimum, if the proxy flag says it is registered, then the pointers should definitely be valid
+	if (Proxy->bIsRegisteredWithLandscapeInfo)
+	{
+		check(bResult == Proxy->bIsRegisteredWithLandscapeInfo);
+	}
+
+	// trust the proxy flag over the landscape info pointers
+	bResult = Proxy->bIsRegisteredWithLandscapeInfo;
+#endif // WITH_EDITORONLY_DATA
+
+	return bResult;
+}
+
 void ULandscapeInfo::RegisterActor(ALandscapeProxy* Proxy, bool bMapCheck, bool bUpdateAllAddCollisions)
 {
 	UWorld* OwningWorld = Proxy->GetWorld();
@@ -4243,6 +4345,10 @@ void ULandscapeInfo::RegisterActor(ALandscapeProxy* Proxy, bool bMapCheck, bool
 	{
 		RegisterCollisionComponent(CollComp);
 	}
+
+#if WITH_EDITORONLY_DATA
+	Proxy->bIsRegisteredWithLandscapeInfo = true;
+#endif // WITH_EDITORONLY_DATA
 }
 
 void ULandscapeInfo::UnregisterActor(ALandscapeProxy* Proxy)
@@ -4250,7 +4356,7 @@ void ULandscapeInfo::UnregisterActor(ALandscapeProxy* Proxy)
 	UWorld* OwningWorld = Proxy->GetWorld();
 	if (ALandscape* Landscape = Cast<ALandscape>(Proxy))
 	{
-		// Note: UnregisterActor sometimes gets triggered twice, e.g. it has been observed to happen during redo
+		// Note: UnregisterActor sometimes gets triggered twice, e.g. it has been observed to happen during undo/ redo
 		// Note: In some cases LandscapeActor could be updated to a new landscape actor before the old landscape is unregistered/destroyed
 		// e.g. this has been observed when merging levels in the editor
 
@@ -4271,7 +4377,8 @@ void ULandscapeInfo::UnregisterActor(ALandscapeProxy* Proxy)
 	else
 	{
 		ALandscapeStreamingProxy* StreamingProxy = CastChecked<ALandscapeStreamingProxy>(Proxy);
-		StreamingProxies.Remove(StreamingProxy);
+		TWeakObjectPtr<ALandscapeStreamingProxy> StreamingProxyPtr = StreamingProxy;
+		StreamingProxies.Remove(StreamingProxyPtr);
 		StreamingProxy->LandscapeActor = nullptr;
 	}
 
@@ -4303,6 +4410,10 @@ void ULandscapeInfo::UnregisterActor(ALandscapeProxy* Proxy)
 	UpdateLayerInfoMap();
 	UpdateAllAddCollisions();
 #endif
+
+#if WITH_EDITORONLY_DATA
+	Proxy->bIsRegisteredWithLandscapeInfo = false;
+#endif // WITH_EDITORONLY_DATA
 }
 
 #if WITH_EDITOR
diff --git a/Engine/Source/Runtime/Landscape/Private/LandscapeComponent.cpp b/Engine/Source/Runtime/Landscape/Private/LandscapeComponent.cpp
--- a/Engine/Source/Runtime/Landscape/Private/LandscapeComponent.cpp
+++ b/Engine/Source/Runtime/Landscape/Private/LandscapeComponent.cpp
@@ -10,6 +10,9 @@
 
 #include "LandscapeHLODBuilder.h"
 
+uint32 ULandscapeComponent::UndoRedoModifiedComponentCount;
+TArray<ULandscapeComponent*> ULandscapeComponent::UndoRedoModifiedComponents;
+
 #endif
 
 FName FWeightmapLayerAllocationInfo::GetLayerName() const
diff --git a/Engine/Source/Runtime/Landscape/Private/LandscapeEdit.cpp b/Engine/Source/Runtime/Landscape/Private/LandscapeEdit.cpp
--- a/Engine/Source/Runtime/Landscape/Private/LandscapeEdit.cpp
+++ b/Engine/Source/Runtime/Landscape/Private/LandscapeEdit.cpp
@@ -596,7 +596,6 @@ int32 ULandscapeComponent::GetNumMaterials() const
 	return 1;
 }
 
-#if WITH_EDITOR
 bool ULandscapeComponent::GetMaterialPropertyPath(int32 ElementIndex, UObject*& OutOwner, FString& OutPropertyPath, FProperty*& OutProperty)
 {
 	if (ElementIndex == 0)
@@ -620,7 +619,6 @@ bool ULandscapeComponent::GetMaterialPropertyPath(int32 ElementIndex, UObject*&
 
 	return false;
 }
-#endif // WITH_EDITOR
 
 class UMaterialInterface* ULandscapeComponent::GetMaterial(int32 ElementIndex) const
 {
@@ -651,46 +649,99 @@ void ULandscapeComponent::PreFeatureLevelChange(ERHIFeatureLevel::Type PendingFe
 	}
 }
 
+void ULandscapeComponent::PreEditUndo()
+{
+	// We bump the proxy's version numbers here in PRE-EDIT:
+	// This means all of the PostEditUndo() operations (across all Components on the Proxy) 
+	// see a consistent number and ensure we only do WeightmapFixup once.
+	ALandscapeProxy* Proxy = GetLandscapeProxy();
+	Proxy->CurrentVersion++;
+	
+	// Count the number of components involved in the undo/redo
+	UndoRedoModifiedComponentCount++;
+	check(UndoRedoModifiedComponents.Num() == 0);
+}
+
 void ULandscapeComponent::PostEditUndo()
 {
-	if (IsValid(this))
-	{
-		if (!GetLandscapeProxy()->HasLayersContent())
-		{
-			UpdateMaterialInstances();
-		}
-	}
+	ALandscapeProxy* Proxy = GetLandscapeProxy();
+	
+	// On undo, request a recompute weightmap usages, as they are not transacted (they combine information between multiple components)
+	Proxy->RequestProxyLayersWeightmapUsageUpdate();
 
 	Super::PostEditUndo();
 
-	// On undo, request a recompute weightmap usages which can get desynchronized since there can be duplicated between 2 components and also with the proxy : 
-	GetLandscapeProxy()->RequestProxyLayersWeightmapUsageUpdate();
-
 	if (IsValid(this))
 	{
 		EditToolRenderData.UpdateSelectionMaterial(EditToolRenderData.SelectedType, this);
-		if (!GetLandscapeProxy()->HasLayersContent())
+		if (!Proxy->HasLayersContent())
 		{
 			EditToolRenderData.UpdateDebugColorMaterial(this);
             UpdateEditToolRenderData();
 		}	
 	}
 		
-	if (GetLandscapeProxy()->HasLayersContent())
+	if (Proxy->HasLayersContent())
 	{
 		const bool bUpdateAll = true;
 		RequestHeightmapUpdate(bUpdateAll);
 		RequestWeightmapUpdate(bUpdateAll);
-
-		// Clear Cached Editing Data
-		CachedEditingLayer.Invalidate();
-		CachedEditingLayerData = nullptr;
 	}
 	else
 	{
 		TSet<ULandscapeComponent*> Components;
 		Components.Add(this);
-		GetLandscapeProxy()->FlushGrassComponents(&Components);
+		Proxy->FlushGrassComponents(&Components);
+	}
+
+	check(!UndoRedoModifiedComponents.Contains(this));
+	UndoRedoModifiedComponents.Add(this);
+
+	// Only Fixup Weightmaps and Update Material Instances when the LAST modified component has been processed by PostEditUndo()
+	if (UndoRedoModifiedComponents.Num() == UndoRedoModifiedComponentCount)
+	{
+		// process all the proxies across all of the components (except proxies that are pending kill)
+		TSet<ALandscapeProxy*> Proxies;
+		for (ULandscapeComponent* Component : UndoRedoModifiedComponents)
+		{
+			ALandscapeProxy* ComponentProxy = Component->GetLandscapeProxy();
+			if (IsValid(ComponentProxy))
+			{
+				Proxies.Add(ComponentProxy);
+			}
+		}
+
+		for (ALandscapeProxy* ComponentProxy : Proxies)
+		{
+			check(ComponentProxy->GetLandscapeGuid().IsValid());
+
+			// Here we create and register with the LandscapeInfo early (before PostRegisterAllComponents)
+			ULandscapeInfo* LandscapeInfo = ComponentProxy->GetLandscapeInfo();
+			if (!LandscapeInfo->IsRegistered(ComponentProxy))
+			{
+				LandscapeInfo->RegisterActor(ComponentProxy, true);
+			}
+
+			// So that we can ensure that weightmaps are fixed up (which depend on the Landscape Info registration)
+			ComponentProxy->FixupWeightmaps();
+		}
+
+		// Update Material Instances on each modified component (uses the weightmap allocations fixed by FixupWeightmaps)
+		for (ULandscapeComponent* Component : UndoRedoModifiedComponents)
+		{
+			if (IsValid(Component)) // Components can be pending kill, if they were removed by the undo operation
+			{
+				ALandscapeProxy* ComponentProxy = Component->GetLandscapeProxy();
+				if (IsValid(ComponentProxy) && !ComponentProxy->HasLayersContent())
+				{
+					Component->UpdateMaterialInstances();
+				}
+			}
+		}
+
+		// Clear out the temporarily recorded undo info
+		UndoRedoModifiedComponentCount = 0;
+		UndoRedoModifiedComponents.Empty();
 	}
 }
 
@@ -734,7 +785,12 @@ void ALandscapeProxy::FixupWeightmaps()
 	bTemporarilyDisableWeightmapUsagesValidation = true;
 	ON_SCOPE_EXIT
 	{
+		WeightmapFixupVersion = CurrentVersion;
 		bTemporarilyDisableWeightmapUsagesValidation = false;
+
+		// Rebuild weightmap usages, now that the allocations are fixed
+		InitializeProxyLayersWeightmapUsage();
+
 		// Now that the job is done, weightmap usages should be valid again
 		ValidateProxyLayersWeightmapUsage();
 	};
@@ -746,7 +802,6 @@ void ALandscapeProxy::FixupWeightmaps()
 			Component->FixupWeightmaps();
 		}
 	}
-
 }
 
 void ALandscapeProxy::RepairInvalidTextures()
@@ -799,10 +854,10 @@ TArray<UTexture*> ULandscapeComponent::RepairInvalidTextures()
 
 void ULandscapeComponent::FixupWeightmaps()
 {
-	// Fixup final weightmaps
+	// Fixup weightmaps in the base/render layer
 	FixupWeightmaps(FGuid());
 
-	// Also fixup all edit layers weightmaps :
+	// Also fixup all edit layers weightmaps
 	ForEachLayer([&](const FGuid& LayerGuid, FLandscapeLayerComponentData& LayerData)
 	{
 		FixupWeightmaps(LayerGuid);
@@ -815,13 +870,19 @@ void ULandscapeComponent::FixupWeightmaps(const FGuid& InEditLayerGuid)
 	{
 		ULandscapeInfo* Info = GetLandscapeInfo();
 		ALandscapeProxy* Proxy = GetLandscapeProxy();
+		ALandscape* Landscape = GetLandscapeActor();
 
 		TArray<UTexture2D*>& LocalWeightmapTextures = GetWeightmapTextures(InEditLayerGuid);
-		TArray<ULandscapeWeightmapUsage*>& LocalWeightmapTextureUsages = GetWeightmapTexturesUsage(InEditLayerGuid);
-		TArray<FWeightmapLayerAllocationInfo>& LocalWeightmapLayerAllocations = GetWeightmapLayerAllocations(InEditLayerGuid);
+		TArray<ULandscapeWeightmapUsage*>& LocalWeightmapTextureUsages = GetWeightmapTexturesUsage(InEditLayerGuid, /*bInCheckIsUpToDate*/false);
+		TArray<FWeightmapLayerAllocationInfo>& LocalWeightmapLayerAllocations = GetWeightmapLayerAllocations(InEditLayerGuid, /*bInCheckIsUpToDate*/false);
 
 		if (Info)
 		{
+			// It's very important that the Proxies be registered with LandscapeInfo before calling FixupWeightmaps
+			// Otherwise the code below may think ALL layers have been removed, and delete all of the corresponding weightmaps... which would be bad.
+			check(Info->IsRegistered(Proxy));
+
+			// We're going to re-build the texture usage array for this layer on this component
 			LocalWeightmapTextureUsages.Empty();
 			LocalWeightmapTextureUsages.AddDefaulted(LocalWeightmapTextures.Num());
 
@@ -856,28 +917,26 @@ void ULandscapeComponent::FixupWeightmaps(const FGuid& InEditLayerGuid)
 
 			if (bFixedLayerDeletion)
 			{
+				// Delete material layer in the base/render layer
 				{
+					FGuid BaseLayerGuid;
 					FLandscapeEditDataInterface LandscapeEdit(Info);
 					for (int32 Idx = 0; Idx < LayersToDelete.Num(); ++Idx)
 					{
-						DeleteLayer(LayersToDelete[Idx], LandscapeEdit);
+						DeleteLayerInternal(LayersToDelete[Idx], LandscapeEdit, BaseLayerGuid);
 					}
 				}
 
+				// For each edit layer that exists on the local component, delete material layer
 				ForEachLayer([&](const FGuid& LayerGuid, FLandscapeLayerComponentData& LayerData)
 				{
-					SetEditingLayer(LayerGuid);
 					FLandscapeEditDataInterface LandscapeEdit(Info);
 					for (int32 Idx = 0; Idx < LayersToDelete.Num(); ++Idx)
 					{
-						DeleteLayer(LayersToDelete[Idx], LandscapeEdit);
+						DeleteLayerInternal(LayersToDelete[Idx], LandscapeEdit, LayerGuid);
 					}
 				});
-								
-				// Make sure to clear editing layer and cache
-				SetEditingLayer(FGuid());
-				CachedEditingLayer.Invalidate();
-				CachedEditingLayerData = nullptr;
+							
 			}
 
 			bool bFixedWeightmapTextureIndex = false;
@@ -914,10 +973,12 @@ void ULandscapeComponent::FixupWeightmaps(const FGuid& InEditLayerGuid)
 				if (TempUsage == nullptr)
 				{
 					TempUsage = &Proxy->WeightmapUsageMap.Add(WeightmapTexture, GetLandscapeProxy()->CreateWeightmapUsage());
-					(*TempUsage)->LayerGuid.Invalidate();
+					(*TempUsage)->LayerGuid = InEditLayerGuid;
 				}
 
 				ULandscapeWeightmapUsage* Usage = *TempUsage;
+				check(Usage->LayerGuid == InEditLayerGuid); // A single weightmap must belong to exactly one layer
+
 				LocalWeightmapTextureUsages[Allocation.WeightmapTextureIndex] = Usage; // Keep a ref to it for faster access
 
 				// Detect a shared layer allocation, caused by a previous undo or layer deletion bugs
@@ -6350,9 +6411,12 @@ void ULandscapeComponent::RemoveInvalidWeightmaps()
 
 void ULandscapeComponent::RemoveInvalidWeightmaps(const FGuid& InEditLayerGuid)
 {
-	TArray<FWeightmapLayerAllocationInfo>& ComponentWeightmapLayerAllocations = GetWeightmapLayerAllocations(InEditLayerGuid);
+	ALandscapeProxy* Proxy = GetLandscapeProxy();
+	
+	// If this function is called from FixupWeightmaps, then skip the update check
+	TArray<FWeightmapLayerAllocationInfo>& ComponentWeightmapLayerAllocations = GetWeightmapLayerAllocations(InEditLayerGuid, !Proxy->bTemporarilyDisableWeightmapUsagesValidation);
 	TArray<UTexture2D*>& ComponentWeightmapTextures = GetWeightmapTextures(InEditLayerGuid);
-	TArray<ULandscapeWeightmapUsage*>& ComponentWeightmapTexturesUsage = GetWeightmapTexturesUsage(InEditLayerGuid);
+	TArray<ULandscapeWeightmapUsage*>& ComponentWeightmapTexturesUsage = GetWeightmapTexturesUsage(InEditLayerGuid, !Proxy->bTemporarilyDisableWeightmapUsagesValidation);
 
 	// Adjust WeightmapTextureIndex index for other layers
 	TSet<int32> UnUsedTextureIndices;
diff --git a/Engine/Source/Runtime/Landscape/Private/LandscapeEditInterface.cpp b/Engine/Source/Runtime/Landscape/Private/LandscapeEditInterface.cpp
--- a/Engine/Source/Runtime/Landscape/Private/LandscapeEditInterface.cpp
+++ b/Engine/Source/Runtime/Landscape/Private/LandscapeEditInterface.cpp
@@ -1582,26 +1582,34 @@ void FLandscapeEditDataInterface::GetHeightDataFast(const int32 X1, const int32
 }
 
 void ULandscapeComponent::DeleteLayer(ULandscapeLayerInfoObject* LayerInfo, FLandscapeEditDataInterface& LandscapeEdit)
+{
+	FGuid EditLayerGuid = this->GetEditingLayerGUID();
+	DeleteLayerInternal(LayerInfo, LandscapeEdit, EditLayerGuid);
+}
+
+void ULandscapeComponent::DeleteLayerInternal(ULandscapeLayerInfoObject* LayerInfo, FLandscapeEditDataInterface& LandscapeEdit, const FGuid &EditLayerGuid)
 {
 	TRACE_CPUPROFILER_EVENT_SCOPE(LandscapeComponent_DeleteLayer);
 
 	ULandscapeComponent* Component = this;
-	FGuid EditLayerGuid = Component->GetEditingLayerGUID();
-	const TArray<FWeightmapLayerAllocationInfo>& ComponentWeightmapLayerAllocations = Component->GetWeightmapLayerAllocations(EditLayerGuid);
+	
+	// This can be called during WeightmapFixup, so skip checking for updated weightmap allocations
+	const TArray<FWeightmapLayerAllocationInfo>& ComponentWeightmapLayerAllocations = Component->GetWeightmapLayerAllocations(EditLayerGuid, /*bInCheckUpToDate*/ false);
 	const TArray<UTexture2D*>& ComponentWeightmapTextures = Component->GetWeightmapTextures(EditLayerGuid);
 
-	// Find the index for this layer in this component.
-	const int32 DeleteLayerIdx = ComponentWeightmapLayerAllocations.IndexOfByPredicate(
+	// Delete the first LayerAllocation with a matching LayerInfo
+	const int32 DeleteAllocIdx = ComponentWeightmapLayerAllocations.IndexOfByPredicate(
 		[LayerInfo](const FWeightmapLayerAllocationInfo& Allocation) { return Allocation.LayerInfo == LayerInfo; });
-	if (DeleteLayerIdx == INDEX_NONE)
+	if (DeleteAllocIdx == INDEX_NONE)
 	{
 		// Layer not used for this component.
 		return;
 	}
 
-	Component->DeleteLayerAllocation(EditLayerGuid, DeleteLayerIdx, LandscapeEdit.GetShouldDirtyPackage());
+	Component->DeleteLayerAllocation(EditLayerGuid, DeleteAllocIdx, LandscapeEdit.GetShouldDirtyPackage());
 
 	// See if the deleted layer is a NoWeightBlend layer - if not, we don't have to worry about normalization
+	// If the layer doesn't exist, assume it is a blended layer (so renormalization will run whether needed or not)
 	bool bDeleteLayerIsNoWeightBlend = (LayerInfo && LayerInfo->bNoWeightBlend);
 
 	if (!bDeleteLayerIsNoWeightBlend)
@@ -1648,7 +1656,7 @@ void ULandscapeComponent::DeleteLayer(ULandscapeLayerInfoObject* LayerInfo, FLan
 						int32 OtherLayerWeightSum = 0;
 						for (int32 LayerIdx = 0; LayerIdx < ComponentWeightmapLayerAllocations.Num(); LayerIdx++)
 						{
-							if (LayerIdx != DeleteLayerIdx && LayerNoWeightBlends[LayerIdx] == false)
+							if (LayerIdx != DeleteAllocIdx && LayerNoWeightBlends[LayerIdx] == false)
 							{
 								OtherLayerWeightSum += LayerDataPtrs[LayerIdx][TexDataIndex];
 							}
@@ -1661,7 +1669,7 @@ void ULandscapeComponent::DeleteLayer(ULandscapeLayerInfoObject* LayerInfo, FLan
 							// There's nothing we can easily do if this was the only weight-blend layer on this component
 							for (int32 LayerIdx = 0; LayerIdx < ComponentWeightmapLayerAllocations.Num(); LayerIdx++)
 							{
-								if (LayerIdx != DeleteLayerIdx && LayerNoWeightBlends[LayerIdx] == false)
+								if (LayerIdx != DeleteAllocIdx && LayerNoWeightBlends[LayerIdx] == false)
 								{
 									uint8& Weight = LayerDataPtrs[LayerIdx][TexDataIndex];
 									Weight = 255;
@@ -1674,7 +1682,7 @@ void ULandscapeComponent::DeleteLayer(ULandscapeLayerInfoObject* LayerInfo, FLan
 							// Adjust other layer weights
 							for (int32 LayerIdx = 0; LayerIdx < ComponentWeightmapLayerAllocations.Num(); LayerIdx++)
 							{
-								if (LayerIdx != DeleteLayerIdx && LayerNoWeightBlends[LayerIdx] == false)
+								if (LayerIdx != DeleteAllocIdx && LayerNoWeightBlends[LayerIdx] == false)
 								{
 									uint8& Weight = LayerDataPtrs[LayerIdx][TexDataIndex];
 									Weight = FMath::Clamp<int32>(FMath::RoundToInt(255.0f * (float)Weight / (float)OtherLayerWeightSum), 0, 255);
@@ -2807,6 +2815,8 @@ void FLandscapeEditDataInterface::SetAlphaData(ULandscapeLayerInfoObject* const
 				Component->RequestWeightmapUpdate();
 			}
 
+			check(UpdateLayerIdx < ComponentWeightmapLayerAllocations.Num());
+
 			// Lock data for all the weightmaps
 			TexDataInfos.Reset();
 			const TArray<UTexture2D*>& ComponentWeightmapTextures = Component->GetWeightmapTextures(true);
diff --git a/Engine/Source/Runtime/Landscape/Private/LandscapeEditLayers.cpp b/Engine/Source/Runtime/Landscape/Private/LandscapeEditLayers.cpp
--- a/Engine/Source/Runtime/Landscape/Private/LandscapeEditLayers.cpp
+++ b/Engine/Source/Runtime/Landscape/Private/LandscapeEditLayers.cpp
@@ -2241,10 +2241,9 @@ void ALandscapeProxy::UpdateProxyLayersWeightmapUsage()
 
 void ALandscapeProxy::PostEditUndo()
 {
-	Super::PostEditUndo();
+	check(ULandscapeComponent::UndoRedoModifiedComponentCount == 0);
 
-	// On undo, request a recompute weightmap usages which can get desynchronized since there can be duplicated between 2 components and also with the proxy : 
-	RequestProxyLayersWeightmapUsageUpdate();
+	Super::PostEditUndo();
 }
 
 void ALandscape::InitializeLandscapeLayersWeightmapUsage()
@@ -2282,8 +2281,8 @@ void ALandscapeProxy::InitializeProxyLayersWeightmapUsage()
 		}
 	}
 
-	ValidateProxyLayersWeightmapUsage();
 	bNeedsWeightmapUsagesUpdate = false;
+	ValidateProxyLayersWeightmapUsage();
 }
 
 void ULandscapeComponent::InitializeLayersWeightmapUsage(const FGuid& InLayerGuid)
@@ -2305,9 +2304,10 @@ void ULandscapeComponent::InitializeLayersWeightmapUsage(const FGuid& InLayerGui
 		}
 	}
 
-	const TArray<FWeightmapLayerAllocationInfo>& ComponentWeightmapLayerAllocations = GetWeightmapLayerAllocations(InLayerGuid);
+	// Allocations should be up to date before calling InitializeLayersWeightmapUsage.
+	const TArray<FWeightmapLayerAllocationInfo>& ComponentWeightmapLayerAllocations = GetWeightmapLayerAllocations(InLayerGuid, /*bInCheckIsUpToDate*/true);
 	const TArray<UTexture2D*>& ComponentWeightmapTextures = GetWeightmapTextures(InLayerGuid);
-	TArray<ULandscapeWeightmapUsage*>& ComponentWeightmapTexturesUsage = GetWeightmapTexturesUsage(InLayerGuid);
+	TArray<ULandscapeWeightmapUsage*>& ComponentWeightmapTexturesUsage = GetWeightmapTexturesUsage(InLayerGuid, /*bInCheckIsUpToDate*/false);
 
 	ComponentWeightmapTexturesUsage.Reset();
 	ComponentWeightmapTexturesUsage.AddDefaulted(ComponentWeightmapTextures.Num());
@@ -2329,7 +2329,7 @@ void ULandscapeComponent::InitializeLayersWeightmapUsage(const FGuid& InLayerGui
 			ULandscapeWeightmapUsage* Usage = *TempUsage;
 			ComponentWeightmapTexturesUsage[Allocation.WeightmapTextureIndex] = Usage; // Keep a ref to it for faster access
 
-			// Validate that there are no duplicated allocation
+			// Validate that there are no conflicting allocations (two allocations claiming the same texture channel)
 			check(Usage->ChannelUsage[Allocation.WeightmapTextureChannel] == nullptr || Usage->ChannelUsage[Allocation.WeightmapTextureChannel] == this);
 
 			// Validate that there are no duplicated allocation (except on the splines layer, since it's updated outside of a transaction and the transactor can later restore a duplicated 
@@ -2386,50 +2386,65 @@ void ALandscapeProxy::ValidateProxyLayersWeightmapUsage() const
 		return;
 	}
 
+	// Fixup and usages should have been updated any time we run validation
+	check(WeightmapFixupVersion == CurrentVersion);
+	check(!bNeedsWeightmapUsagesUpdate);
+
 	TRACE_CPUPROFILER_EVENT_SCOPE(Landscape_ValidateProxyLayersWeightmapUsage);
-	TMap<UTexture2D*, TSet<FWeightmapLayerAllocationInfo>> PerTextureAllocations;
+	TMap<UTexture2D*, TArray<FWeightmapLayerAllocationInfo>> PerTextureAllocations;
 	if (const ALandscape* Landscape = GetLandscapeActor())
 	{
 		for (ULandscapeComponent* Component : LandscapeComponents)
 		{
-			auto ValidateWeightmapAllocationAndUsage = [&](UTexture2D* InWeightmapTexture, const FWeightmapLayerAllocationInfo& InAllocation, ULandscapeWeightmapUsage* InUsage)
+			auto ValidateWeightmapAllocationAndUsage = [&](UTexture2D* InWeightmapTexture, const FWeightmapLayerAllocationInfo& InAllocation, ULandscapeWeightmapUsage* InUsage, const FGuid &InLayerGuid)
 			{
-				check(InUsage);
-				// All usages should be in the proxy's map
-				const TObjectPtr<ULandscapeWeightmapUsage>* ProxyMapUsage = WeightmapUsageMap.Find(InWeightmapTexture);
-				check(ProxyMapUsage && (InUsage == *ProxyMapUsage));
-				check(InUsage->ChannelUsage[InAllocation.WeightmapTextureChannel] == Component);
+				if (InUsage)
+				{
+					// Each usage should also be stored in the proxy's map
+					const TObjectPtr<ULandscapeWeightmapUsage>* ProxyMapUsage = WeightmapUsageMap.Find(InWeightmapTexture);
+					check(ProxyMapUsage);
+					check(InUsage == *ProxyMapUsage);
+
+					// Our component should own the channel, and the LayerGuid should match
+					check(InUsage->ChannelUsage[InAllocation.WeightmapTextureChannel] == Component);
+					check(InUsage->LayerGuid == InLayerGuid);
+				}
 
-				// We shouldn't have any duplicate allocation : 
-				TSet<FWeightmapLayerAllocationInfo>& AllAllocations = PerTextureAllocations.FindOrAdd(InWeightmapTexture);
-				check(!AllAllocations.Contains(InAllocation));
-				AllAllocations.Add(InAllocation);
+				// There should not be any other allocations pointing to this channel on this texture
+				TArray<FWeightmapLayerAllocationInfo>& AllAllocationsForThisTexture = PerTextureAllocations.FindOrAdd(InWeightmapTexture);
+				for (FWeightmapLayerAllocationInfo &Alloc : AllAllocationsForThisTexture)
+				{
+					check(Alloc.WeightmapTextureChannel != InAllocation.WeightmapTextureChannel);
+				}
 			};
 
 			const TArray<UTexture2D*>& WeightmapTextures = Component->GetWeightmapTextures(/*InReturnEditingWeightmap = */false);
 			const TArray<ULandscapeWeightmapUsage*>& WeightmapTextureUsages = Component->GetWeightmapTexturesUsage(/*InReturnEditingWeightmap = */false);
 
-			// Validate weightmaps : 
+			// Validate weightmap allocations
+			FGuid BaseGuid;
 			for (const FWeightmapLayerAllocationInfo& Allocation : Component->GetWeightmapLayerAllocations(/*InReturnEditingWeightmap = */false))
 			{
 				if (Allocation.IsAllocated())
 				{
+					// The allocation texture index should point to a valid texture
 					UTexture2D* WeightmapTexture = WeightmapTextures[Allocation.WeightmapTextureIndex];
 					check(WeightmapTexture != nullptr);
-					if (ULandscapeWeightmapUsage* Usage = WeightmapTextureUsages.IsValidIndex(Allocation.WeightmapTextureIndex) ? WeightmapTextureUsages[Allocation.WeightmapTextureIndex] : nullptr)
-					{
-						ValidateWeightmapAllocationAndUsage(WeightmapTexture, Allocation, Usage);
-					}
+
+					// Either it's out of bounds  i.e. not initialized yet,  or it is initialized and we validate that it is correct...
+					ULandscapeWeightmapUsage* Usage = WeightmapTextureUsages.IsValidIndex(Allocation.WeightmapTextureIndex) ? WeightmapTextureUsages[Allocation.WeightmapTextureIndex] : nullptr;
+					ValidateWeightmapAllocationAndUsage(WeightmapTexture, Allocation, Usage, BaseGuid);
 				}
 			}
 
-			// Validate edit layers weightmaps : 
+			// Validate edit layers weightmap allocations : 
 			{
 				const FLandscapeLayer* SplinesLayer = Landscape->GetLandscapeSplinesReservedLayer();
 				for (const FLandscapeLayer& Layer : Landscape->LandscapeLayers)
 				{
 					FLandscapeLayerComponentData* LayerData = Component->GetLayerData(Layer.Guid);
-					// Skip validation on SplinesLayer since it can momentarily contain duplicated layer allocations after undo (since it's updated outside of a transaction) : 
+
+					// Skip validation on SplinesLayer since it can momentarily contain duplicated layer allocations after undo (since it's updated outside of a transaction) :
 					if (LayerData != nullptr && LayerData->IsInitialized() && (&Layer != SplinesLayer))
 					{
 						for (int32 LayerIdx = 0; LayerIdx < LayerData->WeightmapData.LayerAllocations.Num(); LayerIdx++)
@@ -2440,7 +2455,7 @@ void ALandscapeProxy::ValidateProxyLayersWeightmapUsage() const
 								UTexture2D* WeightmapTexture = LayerData->WeightmapData.Textures[Allocation.WeightmapTextureIndex];
 								if (ULandscapeWeightmapUsage* Usage = LayerData->WeightmapData.TextureUsages.IsValidIndex(Allocation.WeightmapTextureIndex) ? LayerData->WeightmapData.TextureUsages[Allocation.WeightmapTextureIndex].Get() : nullptr)
 								{
-									ValidateWeightmapAllocationAndUsage(WeightmapTexture, Allocation, Usage);
+									ValidateWeightmapAllocationAndUsage(WeightmapTexture, Allocation, Usage, Layer.Guid);
 								}
 							}
 						}
@@ -2451,21 +2466,6 @@ void ALandscapeProxy::ValidateProxyLayersWeightmapUsage() const
 	}
 }
 
-void ALandscape::RequestProxyLayersWeightmapUsageUpdate()
-{
-	ULandscapeInfo* Info = GetLandscapeInfo();
-
-	if (Info == nullptr)
-	{
-		return;
-	}
-
-	Info->ForAllLandscapeProxies([&](ALandscapeProxy* Proxy)
-	{
-		Proxy->RequestProxyLayersWeightmapUsageUpdate();
-	});
-}
-
 void ALandscapeProxy::RequestProxyLayersWeightmapUsageUpdate()
 {
 	bNeedsWeightmapUsagesUpdate = true;
@@ -9150,15 +9150,6 @@ void ALandscape::SetEditingLayer(const FGuid& InLayerGuid)
 	}
 
 	EditingLayer = InLayerGuid;
-
-	// Propagate Editing Layer to components (will be cached)
-	LandscapeInfo->ForAllLandscapeProxies([&](ALandscapeProxy* Proxy)
-	{
-		for (ULandscapeComponent* Component : Proxy->LandscapeComponents)
-		{
-			Component->SetEditingLayer(InLayerGuid);
-		}
-	});
 }
 
 void ALandscape::SetGrassUpdateEnabled(bool bInGrassUpdateEnabled)
diff --git a/Engine/Source/Runtime/Landscape/Private/LandscapeRender.cpp b/Engine/Source/Runtime/Landscape/Private/LandscapeRender.cpp
--- a/Engine/Source/Runtime/Landscape/Private/LandscapeRender.cpp
+++ b/Engine/Source/Runtime/Landscape/Private/LandscapeRender.cpp
@@ -741,7 +741,7 @@ FLandscapeComponentSceneProxy::FLandscapeComponentSceneProxy(ULandscapeComponent
 	: FPrimitiveSceneProxy(InComponent, NAME_LandscapeResourceNameForDebugging)
 	, FLandscapeSectionInfo(InComponent->GetWorld(), InComponent->GetLandscapeProxy()->GetLandscapeGuid(), InComponent->GetSectionBase() / InComponent->ComponentSizeQuads)
 	, MaxLOD(FMath::CeilLogTwo(InComponent->SubsectionSizeQuads + 1) - 1)
-	, NumWeightmapLayerAllocations(InComponent->GetWeightmapLayerAllocations().Num())
+	, NumWeightmapLayerAllocations(InComponent->GetWeightmapLayerAllocations(/*InReturnEditingWeightmap*/false, /*bInCheckIsUpToDate*/false).Num())
 	, StaticLightingLOD(InComponent->GetLandscapeProxy()->StaticLightingLOD)
 	, WeightmapSubsectionOffset(InComponent->WeightmapSubsectionOffset)
 	, FirstLOD(0)
@@ -1003,7 +1003,8 @@ FLandscapeComponentSceneProxy::FLandscapeComponentSceneProxy(ULandscapeComponent
 	UpdateDefaultInstanceSceneData();
 
 #if WITH_EDITOR
-	const TArray<FWeightmapLayerAllocationInfo>& ComponentWeightmapLayerAllocations = InComponent->GetWeightmapLayerAllocations();
+	// Scene proxies are sometimes constructed before it's possible to update these allocations, so don't check for updated values
+	const TArray<FWeightmapLayerAllocationInfo>& ComponentWeightmapLayerAllocations = InComponent->GetWeightmapLayerAllocations(/*InReturnEditingWeightmap*/false, /*bInCheckIsUpToDate*/false);
 	for (const FWeightmapLayerAllocationInfo& Allocation : ComponentWeightmapLayerAllocations)
 	{
 		if (Allocation.LayerInfo != nullptr)
