MGPU:  Fix for Landscape edit readback crash when loading a map in the editor with -MaxGPUCount=2.  Need to specify GPU 0 for staging buffer readback.

#rnx

diff --git a/Engine/Source/Runtime/Landscape/Private/LandscapeEditReadback.cpp b/Engine/Source/Runtime/Landscape/Private/LandscapeEditReadback.cpp
--- a/Engine/Source/Runtime/Landscape/Private/LandscapeEditReadback.cpp
+++ b/Engine/Source/Runtime/Landscape/Private/LandscapeEditReadback.cpp
@@ -134,11 +134,14 @@ bool UpdateTask_RenderThread(FRHICommandListImmediate& RHICmdList, FLandscapeEdi
 			const int32 MipWidth = FMath::Max(Task.Size.X >> MipIndex, 1);
 			const int32 MipHeight = FMath::Max(Task.Size.Y >> MipIndex, 1);
 
+			// Editor always runs on GPU zero
+			const uint32 GPUIndex = 0;
+
 			Task.Result[MipIndex].SetNum(MipWidth * MipHeight);
 
 			void* Data = nullptr;
 			int32 TargetWidth, TargetHeight;
-			RHICmdList.MapStagingSurface(Task.StagingTextures[MipIndex], Task.ReadbackFence.GetReference(), Data, TargetWidth, TargetHeight);
+			RHICmdList.MapStagingSurface(Task.StagingTextures[MipIndex], Task.ReadbackFence.GetReference(), Data, TargetWidth, TargetHeight, GPUIndex);
 			check(Data != nullptr);
 			check(MipWidth <= TargetWidth && MipHeight <= TargetHeight);
 
@@ -151,7 +154,7 @@ bool UpdateTask_RenderThread(FRHICommandListImmediate& RHICmdList, FLandscapeEdi
 				WritePtr += MipWidth;
 			}
 
-			RHICmdList.UnmapStagingSurface(Task.StagingTextures[MipIndex]);
+			RHICmdList.UnmapStagingSurface(Task.StagingTextures[MipIndex], GPUIndex);
 		}
 
 		// Write completion flag for game thread.
