Landscape: Fix truncation warnings, part 2

diff --git a/Engine/Source/Runtime/Landscape/Private/LandscapeEditInterface.cpp b/Engine/Source/Runtime/Landscape/Private/LandscapeEditInterface.cpp
--- a/Engine/Source/Runtime/Landscape/Private/LandscapeEditInterface.cpp
+++ b/Engine/Source/Runtime/Landscape/Private/LandscapeEditInterface.cpp
@@ -152,15 +152,15 @@ namespace
 		}
 	}
 
-	template<typename T>
-	FORCEINLINE T CalcValueFromValueXY( const int32* Dist, const T& ValueX, const T& ValueY, const uint8& CornerSet, const T* CornerValues )
+	template<typename T, typename F>
+	FORCEINLINE T CalcValueFromValueXY( const int32* Dist, const F& ValueX, const F& ValueY, const uint8& CornerSet, const T* CornerValues )
 	{
 		T FinalValue;
 		int32 DistX = FMath::Min(Dist[0], Dist[1]);
 		int32 DistY = FMath::Min(Dist[2], Dist[3]);
 		if (DistX+DistY > 0)
 		{
-			FinalValue = ((ValueX * DistY) + (ValueY * DistX)) / (float)(DistX + DistY);
+			FinalValue = static_cast<T>(((ValueX * DistY) + (ValueY * DistX)) / static_cast<F>((DistX + DistY)));
 		}
 		else
 		{
@@ -182,7 +182,7 @@ namespace
 			}
 			else
 			{
-				FinalValue = ValueX;
+				FinalValue = static_cast<T>(ValueX);
 			}
 		}
 		return FinalValue;
@@ -313,8 +313,8 @@ void FLandscapeEditDataInterface::SetHeightData(int32 X1, int32 Y1, int32 X2, in
 			// Find the texture data corresponding to this vertex
 			int32 SizeU = Heightmap->Source.GetSizeX();
 			int32 SizeV = Heightmap->Source.GetSizeY();
-			int32 HeightmapOffsetX = Component->HeightmapScaleBias.Z * (float)SizeU;
-			int32 HeightmapOffsetY = Component->HeightmapScaleBias.W * (float)SizeV;
+			int32 HeightmapOffsetX = static_cast<int32>(Component->HeightmapScaleBias.Z * SizeU);
+			int32 HeightmapOffsetY = static_cast<int32>(Component->HeightmapScaleBias.W * SizeV);
 
 			// Find coordinates of box that lies inside component
 			int32 ComponentX1 = FMath::Clamp<int32>(X1 - ComponentIndexX*ComponentSizeQuads, 0, ComponentSizeQuads);
@@ -379,8 +379,8 @@ void FLandscapeEditDataInterface::SetHeightData(int32 X1, int32 Y1, int32 X2, in
 							{
 								const int32 NormalDataIndex = (LandscapeX - X1) + NumVertsX * (LandscapeY - Y1);
 								FVector Normal = VertexNormals[NormalDataIndex].GetSafeNormal();
-								TexData.B = FMath::RoundToInt(127.5f * (Normal.X + 1.0f));
-								TexData.A = FMath::RoundToInt(127.5f * (Normal.Y + 1.0f));
+								TexData.B = static_cast<uint8>(FMath::RoundToInt32(127.5f * (Normal.X + 1.0f)));
+								TexData.A = static_cast<uint8>(FMath::RoundToInt32(127.5f * (Normal.Y + 1.0f)));
 							}
 							else if (InNormalData)
 							{
@@ -535,8 +535,8 @@ void FLandscapeEditDataInterface::RecalculateNormals()
 		// Find the texture data corresponding to this vertex
 		int32 SizeU = Component->GetHeightmap(true)->Source.GetSizeX();
 		int32 SizeV = Component->GetHeightmap(true)->Source.GetSizeY();
-		int32 HeightmapOffsetX = Component->HeightmapScaleBias.Z * (float)SizeU;
-		int32 HeightmapOffsetY = Component->HeightmapScaleBias.W * (float)SizeV;
+		int32 HeightmapOffsetX = static_cast<int32>(Component->HeightmapScaleBias.Z * SizeU);
+		int32 HeightmapOffsetY = static_cast<int32>(Component->HeightmapScaleBias.W * SizeV);
 
 		FLandscapeTextureDataInfo* TexDataInfo = GetTextureDataInfo(Component->GetHeightmap(true));
 		FColor* HeightmapTextureData = (FColor*)TexDataInfo->GetMipData(0);
@@ -564,8 +564,8 @@ void FLandscapeEditDataInterface::RecalculateNormals()
 						FColor& TexData = HeightmapTextureData[WriteIndex];
 
 						// Update the channels containing the normals
-						TexData.B = FMath::RoundToInt( 127.5f * (Normal.X + 1.0f) );
-						TexData.A = FMath::RoundToInt( 127.5f * (Normal.Y + 1.0f) );
+						TexData.B = static_cast<uint8>(FMath::RoundToInt32( 127.5f * (Normal.X + 1.0f) ));
+						TexData.A = static_cast<uint8>(FMath::RoundToInt32( 127.5f * (Normal.Y + 1.0f) ));
 					}
 				}
 			}
@@ -655,8 +655,8 @@ void FLandscapeEditDataInterface::GetHeightDataTemplFast(const int32 X1, const i
 							// Find the texture data corresponding to this vertex
 							int32 SizeU = Heightmap->Source.GetSizeX();
 							int32 SizeV = Heightmap->Source.GetSizeY();
-							int32 HeightmapOffsetX = Component->HeightmapScaleBias.Z * (float)SizeU;
-							int32 HeightmapOffsetY = Component->HeightmapScaleBias.W * (float)SizeV;
+							int32 HeightmapOffsetX = static_cast<int32>(Component->HeightmapScaleBias.Z * SizeU);
+							int32 HeightmapOffsetY = static_cast<int32>(Component->HeightmapScaleBias.W * SizeV);
 
 							int32 TexX = HeightmapOffsetX + (SubsectionSizeQuads+1) * SubIndexX + SubX;
 							int32 TexY = HeightmapOffsetY + (SubsectionSizeQuads+1) * SubIndexY + SubY;
@@ -861,7 +861,7 @@ void FLandscapeEditDataInterface::CalcMissingValues(const int32& X1, const int32
 									{
 										int32 Dist1 = LandscapeY - (ComponentIndexY*ComponentSizeQuads);
 										int32 Dist2 = ((ComponentIndexY+1)*ComponentSizeQuads) - LandscapeY;
-										Value[0] = (FType)(Dist2 * CornerValues[0] + Dist1 * CornerValues[2]) / (Dist1 + Dist2);
+										Value[0] = static_cast<TData>((FType)(Dist2 * CornerValues[0] + Dist1 * CornerValues[2]) / (Dist1 + Dist2));
 										Dist[0] = LandscapeX - (ComponentIndexX*ComponentSizeQuads);
 										Exist[0] = true;
 									}
@@ -880,7 +880,7 @@ void FLandscapeEditDataInterface::CalcMissingValues(const int32& X1, const int32
 									{
 										int32 Dist1 = LandscapeY - (ComponentIndexY*ComponentSizeQuads);
 										int32 Dist2 = ((ComponentIndexY+1)*ComponentSizeQuads) - LandscapeY;
-										Value[1] = (FType)(Dist2 * CornerValues[1] + Dist1 * CornerValues[3]) / (Dist1 + Dist2);
+										Value[1] = static_cast<TData>((FType)(Dist2 * CornerValues[1] + Dist1 * CornerValues[3]) / (Dist1 + Dist2));
 										Dist[1] = (ComponentIndexX+1)*ComponentSizeQuads - LandscapeX;
 										Exist[1] = true;
 									}
@@ -905,7 +905,7 @@ void FLandscapeEditDataInterface::CalcMissingValues(const int32& X1, const int32
 									{
 										int32 Dist1 = LandscapeX - (ComponentIndexX*ComponentSizeQuads);
 										int32 Dist2 = (ComponentIndexX+1)*ComponentSizeQuads - LandscapeX;
-										Value[2] = (FType)(Dist2 * CornerValues[0] + Dist1 * CornerValues[1]) / (Dist1 + Dist2);
+										Value[2] = static_cast<TData>((FType)(Dist2 * CornerValues[0] + Dist1 * CornerValues[1]) / (Dist1 + Dist2));
 										Dist[2] = LandscapeY - (ComponentIndexY*ComponentSizeQuads);
 										Exist[2] = true;
 									}
@@ -924,7 +924,7 @@ void FLandscapeEditDataInterface::CalcMissingValues(const int32& X1, const int32
 									{
 										int32 Dist1 = LandscapeX - (ComponentIndexX*ComponentSizeQuads);
 										int32 Dist2 = (ComponentIndexX+1)*ComponentSizeQuads - LandscapeX;
-										Value[3] = (FType)(Dist2 * CornerValues[2] + Dist1 * CornerValues[3]) / (Dist1 + Dist2);
+										Value[3] = static_cast<TData>((FType)(Dist2 * CornerValues[2] + Dist1 * CornerValues[3]) / (Dist1 + Dist2));
 										Dist[3] = (ComponentIndexY+1)*ComponentSizeQuads - LandscapeY;
 										Exist[3] = true;
 									}
@@ -940,11 +940,11 @@ void FLandscapeEditDataInterface::CalcMissingValues(const int32& X1, const int32
 								}
 								else if ( (Exist[0] || Exist[1]) )
 								{
-									FinalValue = ValueX;
+									FinalValue = static_cast<TData>(ValueX);
 								}
 								else if ( (Exist[2] || Exist[3]) )
 								{
-									FinalValue = ValueY;
+									FinalValue = static_cast<TData>(ValueY);
 								}
 
 								StoreData.Store(LandscapeX, LandscapeY, FinalValue);
@@ -968,13 +968,13 @@ FColor& FLandscapeEditDataInterface::GetHeightMapColor(const ULandscapeComponent
 
 	// All Heightmaps of component have the same texture size
 	const FTextureSource& HeightmapTextureSource = Component->GetHeightmap()->Source;
-	int32 SizeU = HeightmapTextureSource.GetSizeX();
-	int32 SizeV = HeightmapTextureSource.GetSizeY();
-	int32 HeightmapOffsetX = Component->HeightmapScaleBias.Z * (float)SizeU;
-	int32 HeightmapOffsetY = Component->HeightmapScaleBias.W * (float)SizeV;
+	const int32 SizeU = HeightmapTextureSource.GetSizeX();
+	const int32 SizeV = HeightmapTextureSource.GetSizeY();
+	const int32 HeightmapOffsetX = static_cast<int32>(Component->HeightmapScaleBias.Z * SizeU);
+	const int32 HeightmapOffsetY = static_cast<int32>(Component->HeightmapScaleBias.W * SizeV);
 
-	int32 TexX = HeightmapOffsetX + TexU;
-	int32 TexY = HeightmapOffsetY + TexV;
+	const int32 TexX = HeightmapOffsetX + TexU;
+	const int32 TexY = HeightmapOffsetY + TexV;
 	FColor& TexData = TextureData[ TexX + TexY * SizeU ];
 	return TexData;
 }
@@ -1320,7 +1320,7 @@ void FLandscapeEditDataInterface::GetHeightDataInternal(int32& ValidX1, int32& V
 									{
 										int32 Dist1 = LandscapeY - (ComponentIndexY*ComponentSizeQuads);
 										int32 Dist2 = ((ComponentIndexY+1)*ComponentSizeQuads) - LandscapeY;
-										Value[0] = (float)(Dist2 * CornerValues[0] + Dist1 * CornerValues[2]) / (Dist1 + Dist2);
+										Value[0] = static_cast<typename TDataAccess::DataType>((float)(Dist2 * CornerValues[0] + Dist1 * CornerValues[2]) / (Dist1 + Dist2));
 										Dist[0] = LandscapeX - (ComponentIndexX*ComponentSizeQuads);
 										Exist[0] = true;
 									}
@@ -1338,7 +1338,7 @@ void FLandscapeEditDataInterface::GetHeightDataInternal(int32& ValidX1, int32& V
 									{
 										int32 Dist1 = LandscapeY - (ComponentIndexY*ComponentSizeQuads);
 										int32 Dist2 = ((ComponentIndexY+1)*ComponentSizeQuads) - LandscapeY;
-										Value[1] = (float)(Dist2 * CornerValues[1] + Dist1 * CornerValues[3]) / (Dist1 + Dist2);
+										Value[1] = static_cast<typename TDataAccess::DataType>((float)(Dist2 * CornerValues[1] + Dist1 * CornerValues[3]) / (Dist1 + Dist2));
 										Dist[1] = (ComponentIndexX+1)*ComponentSizeQuads - LandscapeX;
 										Exist[1] = true;
 									}
@@ -1362,7 +1362,7 @@ void FLandscapeEditDataInterface::GetHeightDataInternal(int32& ValidX1, int32& V
 									{
 										int32 Dist1 = LandscapeX - (ComponentIndexX*ComponentSizeQuads);
 										int32 Dist2 = (ComponentIndexX+1)*ComponentSizeQuads - LandscapeX;
-										Value[2] = (float)(Dist2 * CornerValues[0] + Dist1 * CornerValues[1]) / (Dist1 + Dist2);
+										Value[2] = static_cast<typename TDataAccess::DataType>((float)(Dist2 * CornerValues[0] + Dist1 * CornerValues[1]) / (Dist1 + Dist2));
 										Dist[2] = LandscapeY - (ComponentIndexY*ComponentSizeQuads);
 										Exist[2] = true;
 									}
@@ -1380,7 +1380,7 @@ void FLandscapeEditDataInterface::GetHeightDataInternal(int32& ValidX1, int32& V
 									{
 										int32 Dist1 = LandscapeX - (ComponentIndexX*ComponentSizeQuads);
 										int32 Dist2 = (ComponentIndexX+1)*ComponentSizeQuads - LandscapeX;
-										Value[3] = (float)(Dist2 * CornerValues[2] + Dist1 * CornerValues[3]) / (Dist1 + Dist2);
+										Value[3] = static_cast<typename TDataAccess::DataType>((float)(Dist2 * CornerValues[2] + Dist1 * CornerValues[3]) / (Dist1 + Dist2));
 										Dist[3] = (ComponentIndexY+1)*ComponentSizeQuads - LandscapeY;
 										Exist[3] = true;
 									}
@@ -1395,19 +1395,19 @@ void FLandscapeEditDataInterface::GetHeightDataInternal(int32& ValidX1, int32& V
 								}
 								else if ( (BorderComponent[0] || BorderComponent[1]) )
 								{
-									FinalValue = ValueX;
+									FinalValue = static_cast<typename TDataAccess::DataType>(ValueX);
 								}
 								else if ( (BorderComponent[2] || BorderComponent[3]) )
 								{
-									FinalValue = ValueY;
+									FinalValue = static_cast<typename TDataAccess::DataType>(ValueY);
 								}
 								else if ( (Exist[0] || Exist[1]) )
 								{
-									FinalValue = ValueX;
+									FinalValue = static_cast<typename TDataAccess::DataType>(ValueX);
 								}
 								else if ( (Exist[2] || Exist[3]) )
 								{
-									FinalValue = ValueY;
+									FinalValue = static_cast<typename TDataAccess::DataType>(ValueY);
 								}
 
 								StoreData.Store(LandscapeX, LandscapeY, FinalValue);
@@ -1629,8 +1629,8 @@ void ULandscapeComponent::DeleteLayer(ULandscapeLayerInfoObject* LayerInfo, FLan
 		// Find the texture data corresponding to this vertex
 		const int32 SizeU = (SubsectionSizeQuads + 1) * NumSubsections;
 		const int32 SizeV = (SubsectionSizeQuads + 1) * NumSubsections;
-		const int32 WeightmapOffsetX = Component->WeightmapScaleBias.Z * (float)SizeU;
-		const int32 WeightmapOffsetY = Component->WeightmapScaleBias.W * (float)SizeV;
+		const int32 WeightmapOffsetX = static_cast<int32>(Component->WeightmapScaleBias.Z * SizeU);
+		const int32 WeightmapOffsetY = static_cast<int32>(Component->WeightmapScaleBias.W * SizeV);
 
 		for (int32 SubIndexY = 0; SubIndexY < NumSubsections; SubIndexY++)
 		{
@@ -1677,7 +1677,7 @@ void ULandscapeComponent::DeleteLayer(ULandscapeLayerInfoObject* LayerInfo, FLan
 								if (LayerIdx != DeleteLayerIdx && LayerNoWeightBlends[LayerIdx] == false)
 								{
 									uint8& Weight = LayerDataPtrs[LayerIdx][TexDataIndex];
-									Weight = FMath::Clamp<int32>(FMath::RoundToInt(255.0f * (float)Weight / (float)OtherLayerWeightSum), 0, 255);
+									Weight = static_cast<uint8>(FMath::Clamp<int32>(FMath::RoundToInt32(255.0f * (float)Weight / (float)OtherLayerWeightSum), 0, 255));
 								}
 							}
 						}
@@ -1821,8 +1821,8 @@ void ULandscapeComponent::FillLayer(ULandscapeLayerInfoObject* LayerInfo, FLands
 		// Find the texture data corresponding to this vertex
 		const int32 SizeU = (SubsectionSizeQuads + 1) * NumSubsections;
 		const int32 SizeV = (SubsectionSizeQuads + 1) * NumSubsections;
-		const int32 WeightmapOffsetX = Component->WeightmapScaleBias.Z * (float)SizeU;
-		const int32 WeightmapOffsetY = Component->WeightmapScaleBias.W * (float)SizeV;
+		const int32 WeightmapOffsetX = static_cast<int32>(Component->WeightmapScaleBias.Z * SizeU);
+		const int32 WeightmapOffsetY = static_cast<int32>(Component->WeightmapScaleBias.W * SizeV);
 
 		FWeightmapLayerAllocationInfo& FillLayerAllocation = ComponentWeightmapLayerAllocations[FillLayerIdx];
 		uint8* MaterialLayerData = (uint8*)LandscapeEdit.GetTextureDataInfo(ComponentWeightmapTextures[FillLayerAllocation.WeightmapTextureIndex])->GetMipData(0);
@@ -2105,7 +2105,7 @@ void ULandscapeComponent::ReplaceLayer(ULandscapeLayerInfoObject* FromLayerInfo,
 
 		for (int32 i = 0; i < FMath::Square(MipSize); i++)
 		{
-			DestTextureData[i*4] = FMath::Min(255, (uint16)DestTextureData[i*4] + (uint16)SrcTextureData[i*4]);
+			DestTextureData[i*4] = static_cast<uint8>(FMath::Min(255, (uint16)DestTextureData[i*4] + (uint16)SrcTextureData[i*4]));
 		}
 
 		// Update all mips
@@ -2969,7 +2969,7 @@ void FLandscapeEditDataInterface::SetAlphaData(ULandscapeLayerInfoObject* const
 										if ((255 - OtherLayerWeightSum) && MaxLayerIdx >= 0)
 										{
 											// No need to check for nullptr here because MaxLayerIdx can only be set to a valid layer
-											LayerDataPtrs[MaxLayerIdx][TexDataIndex] += 255 - OtherLayerWeightSum;
+											LayerDataPtrs[MaxLayerIdx][TexDataIndex] += static_cast<uint8>(255 - OtherLayerWeightSum);
 										}
 									}
 								}
@@ -3035,7 +3035,8 @@ void FLandscapeEditDataInterface::SetAlphaData(ULandscapeLayerInfoObject* const
 												// Exclude bNoWeightBlend layers
 												if (LayerNoWeightBlends[LayerIdx] == false)
 												{
-													Weight = FMath::Clamp<uint8>(FMath::RoundToInt((float)(255 - NewWeight) * (float)Weight / (float)OtherLayerWeightSum), 0, 255);
+													Weight = static_cast<uint8>(FMath::Clamp(
+														FMath::RoundToInt((float)(255 - NewWeight) * (float)Weight / (float)OtherLayerWeightSum), 0, 255));
 												}
 											}
 
@@ -3528,8 +3529,8 @@ void FLandscapeEditDataInterface::GetWeightDataTemplFast(ULandscapeLayerInfoObje
 									// Find the texture data corresponding to this vertex
 									int32 SizeU = WeightmapTexture->Source.GetSizeX();
 									int32 SizeV = WeightmapTexture->Source.GetSizeY();
-									int32 WeightmapOffsetX = Component->WeightmapScaleBias.Z * (float)SizeU;
-									int32 WeightmapOffsetY = Component->WeightmapScaleBias.W * (float)SizeV;
+									int32 WeightmapOffsetX = static_cast<int32>(Component->WeightmapScaleBias.Z * SizeU);
+									int32 WeightmapOffsetY = static_cast<int32>(Component->WeightmapScaleBias.W * SizeV);
 
 									int32 TexX = WeightmapOffsetX + (SubsectionSizeQuads+1) * SubIndexX + SubX;
 									int32 TexY = WeightmapOffsetY + (SubsectionSizeQuads+1) * SubIndexY + SubY;
@@ -3556,8 +3557,8 @@ void FLandscapeEditDataInterface::GetWeightDataTemplFast(ULandscapeLayerInfoObje
 									// Find the texture data corresponding to this vertex
 									int32 SizeU = ComponentWeightmapTexture->Source.GetSizeX();
 									int32 SizeV = ComponentWeightmapTexture->Source.GetSizeY();
-									int32 WeightmapOffsetX = Component->WeightmapScaleBias.Z * (float)SizeU;
-									int32 WeightmapOffsetY = Component->WeightmapScaleBias.W * (float)SizeV;
+									int32 WeightmapOffsetX = static_cast<int32>(Component->WeightmapScaleBias.Z * SizeU);
+									int32 WeightmapOffsetY = static_cast<int32>(Component->WeightmapScaleBias.W * SizeV);
 
 									int32 TexX = WeightmapOffsetX + (SubsectionSizeQuads+1) * SubIndexX + SubX;
 									int32 TexY = WeightmapOffsetY + (SubsectionSizeQuads+1) * SubIndexY + SubY;
@@ -3608,13 +3609,13 @@ uint8 FLandscapeEditDataInterface::GetWeightMapData(const ULandscapeComponent* C
 
 	if (Texture && TextureData)
 	{
-		int32 SizeU = Texture->Source.GetSizeX();
-		int32 SizeV = Texture->Source.GetSizeY();
-		int32 WeightmapOffsetX = Component->WeightmapScaleBias.Z * (float)SizeU;
-		int32 WeightmapOffsetY = Component->WeightmapScaleBias.W * (float)SizeV;
+		const int32 SizeU = Texture->Source.GetSizeX();
+		const int32 SizeV = Texture->Source.GetSizeY();
+		const int32 WeightmapOffsetX = static_cast<int32>(Component->WeightmapScaleBias.Z * SizeU);
+		const int32 WeightmapOffsetY = static_cast<int32>(Component->WeightmapScaleBias.W * SizeV);
 
-		int32 TexX = WeightmapOffsetX + TexU;
-		int32 TexY = WeightmapOffsetY + TexV;
+		const int32 TexX = WeightmapOffsetX + TexU;
+		const int32 TexY = WeightmapOffsetY + TexV;
 		return TextureData[ 4 * (TexX + TexY * SizeU) + Offset ];
 	}
 	return 0;
@@ -4071,7 +4072,7 @@ void FLandscapeEditDataInterface::GetWeightDataTempl(ULandscapeLayerInfoObject*
 										{
 											int32 Dist1 = LandscapeY - (ComponentIndexY*ComponentSizeQuads);
 											int32 Dist2 = ((ComponentIndexY+1)*ComponentSizeQuads) - LandscapeY;
-											Value[0] = (float)(Dist2 * CornerValues[0] + Dist1 * CornerValues[2]) / (Dist1 + Dist2);
+											Value[0] = static_cast<uint8>((float)(Dist2 * CornerValues[0] + Dist1 * CornerValues[2]) / (Dist1 + Dist2));
 											Dist[0] = LandscapeX - (ComponentIndexX*ComponentSizeQuads);
 											Exist[0] = true;
 										}
@@ -4089,7 +4090,7 @@ void FLandscapeEditDataInterface::GetWeightDataTempl(ULandscapeLayerInfoObject*
 										{
 											int32 Dist1 = LandscapeY - (ComponentIndexY*ComponentSizeQuads);
 											int32 Dist2 = ((ComponentIndexY+1)*ComponentSizeQuads) - LandscapeY;
-											Value[1] = (float)(Dist2 * CornerValues[1] + Dist1 * CornerValues[3]) / (Dist1 + Dist2);
+											Value[1] = static_cast<uint8>((float)(Dist2 * CornerValues[1] + Dist1 * CornerValues[3]) / (Dist1 + Dist2));
 											Dist[1] = (ComponentIndexX+1)*ComponentSizeQuads - LandscapeX;
 											Exist[1] = true;
 										}
@@ -4113,7 +4114,7 @@ void FLandscapeEditDataInterface::GetWeightDataTempl(ULandscapeLayerInfoObject*
 										{
 											int32 Dist1 = LandscapeX - (ComponentIndexX*ComponentSizeQuads);
 											int32 Dist2 = (ComponentIndexX+1)*ComponentSizeQuads - LandscapeX;
-											Value[2] = (float)(Dist2 * CornerValues[0] + Dist1 * CornerValues[1]) / (Dist1 + Dist2);
+											Value[2] = static_cast<uint8>((float)(Dist2 * CornerValues[0] + Dist1 * CornerValues[1]) / (Dist1 + Dist2));
 											Dist[2] = LandscapeY - (ComponentIndexY*ComponentSizeQuads);
 											Exist[2] = true;
 										}
@@ -4131,7 +4132,7 @@ void FLandscapeEditDataInterface::GetWeightDataTempl(ULandscapeLayerInfoObject*
 										{
 											int32 Dist1 = LandscapeX - (ComponentIndexX*ComponentSizeQuads);
 											int32 Dist2 = (ComponentIndexX+1)*ComponentSizeQuads - LandscapeX;
-											Value[3] = (float)(Dist2 * CornerValues[2] + Dist1 * CornerValues[3]) / (Dist1 + Dist2);
+											Value[3] = static_cast<uint8>((float)(Dist2 * CornerValues[2] + Dist1 * CornerValues[3]) / (Dist1 + Dist2));
 											Dist[3] = (ComponentIndexY+1)*ComponentSizeQuads - LandscapeY;
 											Exist[3] = true;
 										}
@@ -4146,11 +4147,11 @@ void FLandscapeEditDataInterface::GetWeightDataTempl(ULandscapeLayerInfoObject*
 									}
 									else if ( (Exist[0] || Exist[1]) )
 									{
-										FinalValue = ValueX;
+										FinalValue = static_cast<uint8>(ValueX);
 									}
 									else if ( (Exist[2] || Exist[3]) )
 									{
-										FinalValue = ValueY;
+										FinalValue = static_cast<uint8>(ValueY);
 									}
 
 									Weight = FinalValue;
@@ -4175,8 +4176,8 @@ void FLandscapeEditDataInterface::GetWeightDataTempl(ULandscapeLayerInfoObject*
 									// Find the texture data corresponding to this vertex
 									int32 SizeU = ComponentWeightmapTexture->Source.GetSizeX();
 									int32 SizeV = ComponentWeightmapTexture->Source.GetSizeY();
-									int32 WeightmapOffsetX = Component->WeightmapScaleBias.Z * (float)SizeU;
-									int32 WeightmapOffsetY = Component->WeightmapScaleBias.W * (float)SizeV;
+									int32 WeightmapOffsetX = static_cast<int32>(Component->WeightmapScaleBias.Z * SizeU);
+									int32 WeightmapOffsetY = static_cast<int32>(Component->WeightmapScaleBias.W * SizeV);
 
 									int32 TexX = WeightmapOffsetX + (SubsectionSizeQuads+1) * SubIndexX + SubX;
 									int32 TexY = WeightmapOffsetY + (SubsectionSizeQuads+1) * SubIndexY + SubY;
@@ -4629,8 +4630,8 @@ void FLandscapeEditDataInterface::GetEditToolTextureData(const int32 X1, const i
 								// Find the texture data corresponding to this vertex
 								int32 SizeU = EditToolTexture->Source.GetSizeX();
 								int32 SizeV = EditToolTexture->Source.GetSizeY();
-								int32 WeightmapOffsetX = Component->WeightmapScaleBias.Z * (float)SizeU;
-								int32 WeightmapOffsetY = Component->WeightmapScaleBias.W * (float)SizeV;
+								int32 WeightmapOffsetX = static_cast<int32>(Component->WeightmapScaleBias.Z * (float)SizeU);
+								int32 WeightmapOffsetY = static_cast<int32>(Component->WeightmapScaleBias.W * (float)SizeV);
 
 								int32 TexX = WeightmapOffsetX + (SubsectionSizeQuads + 1) * SubIndexX + SubX;
 								int32 TexY = WeightmapOffsetY + (SubsectionSizeQuads + 1) * SubIndexY + SubY;
@@ -4693,8 +4694,8 @@ void FLandscapeEditDataInterface::SetEditToolTextureData(int32 X1, int32 Y1, int
 			// Find the texture data corresponding to this vertex
 			int32 SizeU = EditToolTexture->Source.GetSizeX();
 			int32 SizeV = EditToolTexture->Source.GetSizeY();
-			int32 WeightmapOffsetX = Component->WeightmapScaleBias.Z * (float)SizeU;
-			int32 WeightmapOffsetY = Component->WeightmapScaleBias.W * (float)SizeV;
+			int32 WeightmapOffsetX = static_cast<int32>(Component->WeightmapScaleBias.Z * SizeU);
+			int32 WeightmapOffsetY = static_cast<int32>(Component->WeightmapScaleBias.W * SizeV);
 
 			// Find coordinates of box that lies inside component
 			int32 ComponentX1 = FMath::Clamp<int32>(X1-ComponentIndexX*ComponentSizeQuads, 0, ComponentSizeQuads);
@@ -4822,8 +4823,8 @@ void FLandscapeEditDataInterface::SetXYOffsetDataTempl(int32 X1, int32 Y1, int32
 			// Find the texture data corresponding to this vertex
 			int32 SizeU = XYOffsetTexture->Source.GetSizeX();
 			int32 SizeV = XYOffsetTexture->Source.GetSizeY();
-			int32 WeightmapOffsetX = Component->WeightmapScaleBias.Z * (float)SizeU;
-			int32 WeightmapOffsetY = Component->WeightmapScaleBias.W * (float)SizeV;
+			int32 WeightmapOffsetX = static_cast<int32>(Component->WeightmapScaleBias.Z * SizeU);
+			int32 WeightmapOffsetY = static_cast<int32>(Component->WeightmapScaleBias.W * SizeV);
 
 			// Find coordinates of box that lies inside component
 			int32 ComponentX1 = FMath::Clamp<int32>(X1-ComponentIndexX*ComponentSizeQuads, 0, ComponentSizeQuads);
@@ -4865,8 +4866,8 @@ void FLandscapeEditDataInterface::SetXYOffsetDataTempl(int32 X1, int32 Y1, int32
 							int32 TexY = WeightmapOffsetY + (SubsectionSizeQuads+1) * SubIndexY + SubY;
 							FColor& TexData = XYOffsetTextureData[ TexX + TexY * SizeU ];
 
-							uint16 XOffset = FMath::Clamp<uint16>(Value.X * LANDSCAPE_INV_XYOFFSET_SCALE + 32768.0f, 0, 65535);
-							uint16 YOffset = FMath::Clamp<uint16>(Value.Y * LANDSCAPE_INV_XYOFFSET_SCALE + 32768.0f, 0, 65535);
+							uint16 XOffset = FMath::Clamp<uint16>(static_cast<uint16>(Value.X * LANDSCAPE_INV_XYOFFSET_SCALE + 32768.0), 0, 65535);
+							uint16 YOffset = FMath::Clamp<uint16>(static_cast<uint16>(Value.Y * LANDSCAPE_INV_XYOFFSET_SCALE + 32768.0), 0, 65535);
 
 							TexData.R = XOffset >> 8;
 							TexData.G = XOffset & 255;
diff --git a/Engine/Source/Runtime/Landscape/Private/LandscapeEditResources.cpp b/Engine/Source/Runtime/Landscape/Private/LandscapeEditResources.cpp
--- a/Engine/Source/Runtime/Landscape/Private/LandscapeEditResources.cpp
+++ b/Engine/Source/Runtime/Landscape/Private/LandscapeEditResources.cpp
@@ -15,7 +15,7 @@ FLandscapeTexture2DResource::FLandscapeTexture2DResource(uint32 InSizeX, uint32
 	: SizeX(InSizeX)
 	, SizeY(InSizeY)
 	, Format(InFormat)
-	, NumMips(bInNumMips)
+	, NumMips(IntCastChecked<uint8>(bInNumMips))
 	, bCreateUAVs(bInNeedUAVs)
 	, bCreateSRV(bInNeedSRV)
 {
@@ -28,7 +28,7 @@ void FLandscapeTexture2DResource::InitRHI()
 
 	FRHITextureCreateDesc Desc =
 		FRHITextureCreateDesc::Create2D(TEXT("FLandscapeTexture2DResource"), SizeX, SizeY, Format)
-		.SetNumMips(static_cast<uint8>(NumMips));
+		.SetNumMips(NumMips);
 
 	if (bCreateUAVs)
 	{
@@ -53,7 +53,7 @@ void FLandscapeTexture2DResource::InitRHI()
 
 	if (bCreateSRV)
 	{
-		TextureSRV = RHICreateShaderResourceView(TextureRHI, /*MipLevel = */0, static_cast<uint8>(NumMips), Format);
+		TextureSRV = RHICreateShaderResourceView(TextureRHI, /*MipLevel = */0, NumMips, Format);
 	}
 }
 
diff --git a/Engine/Source/Runtime/Landscape/Private/LandscapeSplineRaster.cpp b/Engine/Source/Runtime/Landscape/Private/LandscapeSplineRaster.cpp
--- a/Engine/Source/Runtime/Landscape/Private/LandscapeSplineRaster.cpp
+++ b/Engine/Source/Runtime/Landscape/Private/LandscapeSplineRaster.cpp
@@ -72,8 +72,8 @@ protected:
 			return;
 		}
 
-		const float CosInterpX = (Interpolant.X >= 1 ? 1 : 0.5f - 0.5f * FMath::Cos(Interpolant.X * PI));
-		const float CosInterpY = (Interpolant.Y >= 1 ? 1 : 0.5f - 0.5f * FMath::Cos(Interpolant.Y * PI));
+		const float CosInterpX = static_cast<float>(Interpolant.X >= 1 ? 1 : 0.5f - 0.5f * FMath::Cos(Interpolant.X * PI));
+		const float CosInterpY = static_cast<float>(Interpolant.Y >= 1 ? 1 : 0.5f - 0.5f * FMath::Cos(Interpolant.Y * PI));
 		const float Alpha = FMath::Clamp<float>(CosInterpX * CosInterpY, 0.f, 1.f);
 
 		int32 DataIndex = (Y - MinY)*(1 + MaxX - MinX) + X - MinX;
@@ -81,7 +81,7 @@ protected:
 
 		if (HeightAlphaBlendData)
 		{
-			uint16 NewHeight = (uint16)FMath::Clamp<float>(Interpolant.Z, 0, (float)LandscapeDataAccess::MaxValue);
+			uint16 NewHeight = static_cast<uint16>(FMath::Clamp(static_cast<float>(Interpolant.Z), 0.0f, static_cast<float>(LandscapeDataAccess::MaxValue)));
 			float InterpValue = (NewHeight * Alpha) + (Dest * (1.f - Alpha));
 			Dest = (uint16)FMath::Clamp<float>(InterpValue, 0, (float)LandscapeDataAccess::MaxValue);
 
@@ -102,7 +102,7 @@ protected:
 		}
 		else
 		{
-			float Value = FMath::Lerp<float>((float)Dest, Interpolant.Z, Alpha);
+			float Value = FMath::Lerp<float>(Dest, static_cast<float>(Interpolant.Z), Alpha);
 			uint16 DValue = (uint16)FMath::Clamp<float>(Value, 0, (float)LandscapeDataAccess::MaxValue);
 			if ((bRaiseTerrain && DValue > Dest) ||
 				(bLowerTerrain && DValue < Dest))
@@ -230,20 +230,20 @@ protected:
 		
 		if (ModulateAlpha == nullptr)
 		{
-			const float CosInterpX = (Interpolant.X >= 1 ? 1 : 0.5f - 0.5f * FMath::Cos(Interpolant.X * PI));
-			const float CosInterpY = (Interpolant.Y >= 1 ? 1 : 0.5f - 0.5f * FMath::Cos(Interpolant.Y * PI));
+			const float CosInterpX = static_cast<float>(Interpolant.X >= 1 ? 1 : 0.5f - 0.5f * FMath::Cos(Interpolant.X * PI));
+			const float CosInterpY = static_cast<float>(Interpolant.Y >= 1 ? 1 : 0.5f - 0.5f * FMath::Cos(Interpolant.Y * PI));
 			Alpha = CosInterpX * CosInterpY;
 		}
 		else
 		{
-			const float InterpX = FMath::Clamp<float>(Interpolant.X, 0.0f, 1.0f);
-			const float InterpY = FMath::Clamp<float>(Interpolant.Y, 0.0f, 1.0f);
+			const float InterpX = FMath::Clamp<float>(static_cast<float>(Interpolant.X), 0.0f, 1.0f);
+			const float InterpY = FMath::Clamp<float>(static_cast<float>(Interpolant.Y), 0.0f, 1.0f);
 			Alpha = ModulateAlpha->Modulate(InterpX * InterpY, X, Y);
 		}
 
 		uint8& Dest = Data[(Y - MinY)*(1 + MaxX - MinX) + X - MinX];
-		float Value = FMath::Lerp<float>((float)Dest, Interpolant.Z, Alpha);
-		Dest = (uint32)FMath::Clamp<float>(Value, 0, LandscapeDataAccess::MaxValue);
+		float Value = FMath::Lerp<float>(Dest, static_cast<float>(Interpolant.Z), Alpha);
+		Dest = FMath::Clamp<uint8>(static_cast<uint8>(Value), 0, 255);
 	}
 
 private:
@@ -639,10 +639,10 @@ bool ULandscapeInfo::ApplySplinesInternal(bool bOnlySelected, TScriptInterface<I
 		FBox ControlPointBounds = ControlPoint->GetBounds();
 		ControlPointBounds = ControlPointBounds.TransformBy(SplineToLandscape.ToMatrixWithScale());
 
-		int32 MinX = FMath::CeilToInt(ControlPointBounds.Min.X);
-		int32 MinY = FMath::CeilToInt(ControlPointBounds.Min.Y);
-		int32 MaxX = FMath::FloorToInt(ControlPointBounds.Max.X);
-		int32 MaxY = FMath::FloorToInt(ControlPointBounds.Max.Y);
+		int32 MinX = FMath::CeilToInt32(ControlPointBounds.Min.X);
+		int32 MinY = FMath::CeilToInt32(ControlPointBounds.Min.Y);
+		int32 MaxX = FMath::FloorToInt32(ControlPointBounds.Max.X);
+		int32 MaxY = FMath::FloorToInt32(ControlPointBounds.Max.Y);
 
 		MinX = FMath::Max(MinX, LandscapeMinX);
 		MinY = FMath::Max(MinY, LandscapeMinY);
@@ -710,10 +710,10 @@ bool ULandscapeInfo::ApplySplinesInternal(bool bOnlySelected, TScriptInterface<I
 		FBox SegmentBounds = Segment->GetBounds();
 		SegmentBounds = SegmentBounds.TransformBy(SplineToLandscape.ToMatrixWithScale());
 
-		int32 MinX = FMath::CeilToInt(SegmentBounds.Min.X);
-		int32 MinY = FMath::CeilToInt(SegmentBounds.Min.Y);
-		int32 MaxX = FMath::FloorToInt(SegmentBounds.Max.X);
-		int32 MaxY = FMath::FloorToInt(SegmentBounds.Max.Y);
+		int32 MinX = FMath::CeilToInt32(SegmentBounds.Min.X);
+		int32 MinY = FMath::CeilToInt32(SegmentBounds.Min.Y);
+		int32 MaxX = FMath::FloorToInt32(SegmentBounds.Max.X);
+		int32 MaxY = FMath::FloorToInt32(SegmentBounds.Max.Y);
 
 		MinX = FMath::Max(MinX, LandscapeMinX);
 		MinY = FMath::Max(MinY, LandscapeMinY);
@@ -833,10 +833,10 @@ namespace LandscapeSplineRaster
 
 		SegmentBounds = SegmentBounds.TransformBy(SplineToLandscape.ToMatrixWithScale());
 
-		int32 MinX = FMath::CeilToInt(SegmentBounds.Min.X);
-		int32 MinY = FMath::CeilToInt(SegmentBounds.Min.Y);
-		int32 MaxX = FMath::FloorToInt(SegmentBounds.Max.X);
-		int32 MaxY = FMath::FloorToInt(SegmentBounds.Max.Y);
+		int32 MinX = FMath::CeilToInt32(SegmentBounds.Min.X);
+		int32 MinY = FMath::CeilToInt32(SegmentBounds.Min.Y);
+		int32 MaxX = FMath::FloorToInt32(SegmentBounds.Max.X);
+		int32 MaxY = FMath::FloorToInt32(SegmentBounds.Max.Y);
 
 		MinX = FMath::Max(MinX, LandscapeMinX);
 		MinY = FMath::Max(MinY, LandscapeMinY);
@@ -911,9 +911,9 @@ namespace LandscapeSplineRaster
 
 	static bool LineIntersect(const FVector2D& L1Start, const FVector2D& L1End, const FVector2D& L2Start, const FVector2D& L2End, FVector2D& Intersect, float Tolerance = KINDA_SMALL_NUMBER)
 	{
-		float tA = (L2End - L2Start) ^ (L2Start - L1Start);
-		float tB = (L1End - L1Start) ^ (L2Start - L1Start);
-		float Denom = (L2End - L2Start) ^ (L1End - L1Start);
+		float tA = static_cast<float>((L2End - L2Start) ^ (L2Start - L1Start));
+		float tB = static_cast<float>((L1End - L1Start) ^ (L2Start - L1Start));
+		float Denom = static_cast<float>((L2End - L2Start) ^ (L1End - L1Start));
 
 		if (FMath::IsNearlyZero(tA) && FMath::IsNearlyZero(tB))
 		{
@@ -974,7 +974,7 @@ namespace LandscapeSplineRaster
 					// step startSide back until before the endSide point
 					while (StartSide > 0)
 					{
-						const float Projection = (Points[StartSide].*Side - Points[StartSide - 1].*Side) | (Points[EndSide].*Side - Points[StartSide - 1].*Side);
+						const float Projection = static_cast<float>((Points[StartSide].*Side - Points[StartSide - 1].*Side) | (Points[EndSide].*Side - Points[StartSide - 1].*Side));
 						if (Projection >= 0)
 						{
 							break;
@@ -984,7 +984,7 @@ namespace LandscapeSplineRaster
 					// step endSide forwards until after the startSide point
 					while (EndSide < Points.Num() - 1)
 					{
-						const float Projection = (Points[EndSide].*Side - Points[EndSide + 1].*Side) | (Points[StartSide].*Side - Points[EndSide + 1].*Side);
+						const float Projection = static_cast<float>((Points[EndSide].*Side - Points[EndSide + 1].*Side) | (Points[StartSide].*Side - Points[EndSide + 1].*Side));
 						if (Projection >= 0)
 						{
 							break;
diff --git a/Engine/Source/Runtime/Landscape/Private/LandscapeSplines.cpp b/Engine/Source/Runtime/Landscape/Private/LandscapeSplines.cpp
--- a/Engine/Source/Runtime/Landscape/Private/LandscapeSplines.cpp
+++ b/Engine/Source/Runtime/Landscape/Private/LandscapeSplines.cpp
@@ -223,6 +223,7 @@ public:
 			{
 				const FSceneView* View = Views[ViewIndex];
 				FPrimitiveDrawInterface* PDI = Collector.GetPDI(ViewIndex);
+				const uint8 DepthPriority = static_cast<uint8>(GetDepthPriorityGroup(View));
 
 				for (const FSegmentProxy& Segment : Segments)
 				{
@@ -250,19 +251,19 @@ public:
 						PDI->SetHitProxy(Segment.HitProxy);
 
 						// center line
-						PDI->DrawLine(OldPoint.Center, NewPoint.Center, SegmentColor, GetDepthPriorityGroup(View), 0.0f, DepthBias);
+						PDI->DrawLine(OldPoint.Center, NewPoint.Center, SegmentColor, DepthPriority, 0.0f, DepthBias);
 
 						// draw sides
-						PDI->DrawLine(OldPoint.Left, NewPoint.Left, SegmentColor, GetDepthPriorityGroup(View), 0.0f, DepthBias);
-						PDI->DrawLine(OldPoint.Right, NewPoint.Right, SegmentColor, GetDepthPriorityGroup(View), 0.0f, DepthBias);
+						PDI->DrawLine(OldPoint.Left, NewPoint.Left, SegmentColor, DepthPriority, 0.0f, DepthBias);
+						PDI->DrawLine(OldPoint.Right, NewPoint.Right, SegmentColor, DepthPriority, 0.0f, DepthBias);
 
 						PDI->SetHitProxy(nullptr);
 
 						// draw falloff sides
 						if (bDrawFalloff)
 						{
-							DrawDashedLine(PDI, OldPoint.FalloffLeft, NewPoint.FalloffLeft, SegmentColor, 100, GetDepthPriorityGroup(View), DepthBias);
-							DrawDashedLine(PDI, OldPoint.FalloffRight, NewPoint.FalloffRight, SegmentColor, 100, GetDepthPriorityGroup(View), DepthBias);
+							DrawDashedLine(PDI, OldPoint.FalloffLeft, NewPoint.FalloffLeft, SegmentColor, 100, DepthPriority, DepthBias);
+							DrawDashedLine(PDI, OldPoint.FalloffRight, NewPoint.FalloffRight, SegmentColor, 100, DepthPriority, DepthBias);
 						}
 
 						OldPoint = NewPoint;
@@ -276,11 +277,12 @@ public:
 					// Draw Sprite
 					if (bDrawControlPointSprite)
 					{
-						float ControlPointSpriteScale = MyLocalToWorld.GetScaleVector().X * ControlPoint.SpriteScale;
+						double ControlPointSpriteScale = MyLocalToWorld.GetScaleVector().X * ControlPoint.SpriteScale;
 						const FLinearColor ControlPointSpriteColor = ControlPoint.bSelected ? SelectedControlPointSpriteColor : FLinearColor::White;
 						PDI->SetHitProxy(ControlPoint.HitProxy);
-						float ZoomFactor = FMath::Min<float>(View->ViewMatrices.GetProjectionMatrix().M[0][0], View->ViewMatrices.GetProjectionMatrix().M[1][1]);
-						float Scale = View->WorldToScreen(ControlPointLocation).W * (4.0f / View->UnscaledViewRect.Width() / ZoomFactor);
+						const FMatrix& ProectionMatrix = View->ViewMatrices.GetProjectionMatrix();
+						const double ZoomFactor = FMath::Min<double>(ProectionMatrix.M[0][0], ProectionMatrix.M[1][1]);
+						const double Scale = View->WorldToScreen(ControlPointLocation).W * (4.0f / View->UnscaledViewRect.Width() / ZoomFactor);
 						ControlPointSpriteScale *= Scale;
 
 #if WITH_EDITORONLY_DATA
@@ -288,17 +290,17 @@ public:
 #endif
 
 						// Clamping the scale between 10 and the ControlPoint initial scale 
-						ControlPointSpriteScale = FMath::Clamp<float>(ControlPointSpriteScale, 10, ControlPoint.SpriteScale);
+						ControlPointSpriteScale = FMath::Clamp<double>(ControlPointSpriteScale, 10, ControlPoint.SpriteScale);
 						const FVector ControlPointSpriteLocation = ControlPointLocation + FVector(0, 0, ControlPointSpriteScale * 0.75f);
 						PDI->DrawSprite(
 							ControlPointSpriteLocation,
-							ControlPointSpriteScale,
-							ControlPointSpriteScale,
+							static_cast<float>(ControlPointSpriteScale),
+							static_cast<float>(ControlPointSpriteScale),
 							ControlPointSprite->GetResource(),
 							ControlPointSpriteColor,
-							GetDepthPriorityGroup(View),
-							0, ControlPointSprite->GetResource()->GetSizeX(),
-							0, ControlPointSprite->GetResource()->GetSizeY(),
+							DepthPriority,
+							0, static_cast<float>(ControlPointSprite->GetResource()->GetSizeX()),
+							0, static_cast<float>(ControlPointSprite->GetResource()->GetSizeY()),
 							SE_BLEND_Masked);
 					}
 
@@ -315,13 +317,13 @@ public:
 						NewPoint.FalloffRight = MyLocalToWorld.TransformPosition(NewPoint.FalloffRight);
 
 						// draw end for spline connection
-						PDI->DrawPoint(NewPoint.Center, ControlPointColor, 6.0f, GetDepthPriorityGroup(View));
-						PDI->DrawLine(NewPoint.Left, NewPoint.Center, ControlPointColor, GetDepthPriorityGroup(View), 0.0f, DepthBias);
-						PDI->DrawLine(NewPoint.Right, NewPoint.Center, ControlPointColor, GetDepthPriorityGroup(View), 0.0f, DepthBias);
+						PDI->DrawPoint(NewPoint.Center, ControlPointColor, 6.0f, DepthPriority);
+						PDI->DrawLine(NewPoint.Left, NewPoint.Center, ControlPointColor, DepthPriority, 0.0f, DepthBias);
+						PDI->DrawLine(NewPoint.Right, NewPoint.Center, ControlPointColor, DepthPriority, 0.0f, DepthBias);
 						if (bDrawFalloff)
 						{
-							DrawDashedLine(PDI, NewPoint.FalloffLeft, NewPoint.Left, ControlPointColor, 100, GetDepthPriorityGroup(View), DepthBias);
-							DrawDashedLine(PDI, NewPoint.FalloffRight, NewPoint.Right, ControlPointColor, 100, GetDepthPriorityGroup(View), DepthBias);
+							DrawDashedLine(PDI, NewPoint.FalloffLeft, NewPoint.Left, ControlPointColor, 100, DepthPriority, DepthBias);
+							DrawDashedLine(PDI, NewPoint.FalloffRight, NewPoint.Right, ControlPointColor, 100, DepthPriority, DepthBias);
 						}
 					}
 					else if (ControlPoint.Points.Num() >= 2)
@@ -344,27 +346,27 @@ public:
 							PDI->SetHitProxy(ControlPoint.HitProxy);
 
 							// center line
-							PDI->DrawLine(ControlPointLocation, NewPoint.Center, ControlPointColor, GetDepthPriorityGroup(View), 0.0f, DepthBias);
+							PDI->DrawLine(ControlPointLocation, NewPoint.Center, ControlPointColor, DepthPriority, 0.0f, DepthBias);
 
 							// draw sides
-							PDI->DrawLine(OldPoint.Right, NewPoint.Left, ControlPointColor, GetDepthPriorityGroup(View), 0.0f, DepthBias);
+							PDI->DrawLine(OldPoint.Right, NewPoint.Left, ControlPointColor, DepthPriority, 0.0f, DepthBias);
 
 							PDI->SetHitProxy(nullptr);
 
 							// draw falloff sides
 							if (bDrawFalloff)
 							{
-								DrawDashedLine(PDI, OldPoint.FalloffRight, NewPoint.FalloffLeft, ControlPointColor, 100, GetDepthPriorityGroup(View), DepthBias);
+								DrawDashedLine(PDI, OldPoint.FalloffRight, NewPoint.FalloffLeft, ControlPointColor, 100, DepthPriority, DepthBias);
 							}
 
 							// draw end for spline connection
-							PDI->DrawPoint(NewPoint.Center, ControlPointColor, 6.0f, GetDepthPriorityGroup(View));
-							PDI->DrawLine(NewPoint.Left, NewPoint.Center, ControlPointColor, GetDepthPriorityGroup(View), 0.0f, DepthBias);
-							PDI->DrawLine(NewPoint.Right, NewPoint.Center, ControlPointColor, GetDepthPriorityGroup(View), 0.0f, DepthBias);
+							PDI->DrawPoint(NewPoint.Center, ControlPointColor, 6.0f, DepthPriority);
+							PDI->DrawLine(NewPoint.Left, NewPoint.Center, ControlPointColor, DepthPriority, 0.0f, DepthBias);
+							PDI->DrawLine(NewPoint.Right, NewPoint.Center, ControlPointColor, DepthPriority, 0.0f, DepthBias);
 							if (bDrawFalloff)
 							{
-								DrawDashedLine(PDI, NewPoint.FalloffLeft, NewPoint.Left, ControlPointColor, 100, GetDepthPriorityGroup(View), DepthBias);
-								DrawDashedLine(PDI, NewPoint.FalloffRight, NewPoint.Right, ControlPointColor, 100, GetDepthPriorityGroup(View), DepthBias);
+								DrawDashedLine(PDI, NewPoint.FalloffLeft, NewPoint.Left, ControlPointColor, 100, DepthPriority, DepthBias);
+								DrawDashedLine(PDI, NewPoint.FalloffRight, NewPoint.Right, ControlPointColor, 100, DepthPriority, DepthBias);
 							}
 
 							//OldPoint = NewPoint;
@@ -393,7 +395,7 @@ public:
 	}
 	uint32 GetAllocatedSize() const
 	{
-		uint32 AllocatedSize = FPrimitiveSceneProxy::GetAllocatedSize() + Segments.GetAllocatedSize() + ControlPoints.GetAllocatedSize();
+		SIZE_T AllocatedSize = FPrimitiveSceneProxy::GetAllocatedSize() + Segments.GetAllocatedSize() + ControlPoints.GetAllocatedSize();
 		for (const FSegmentProxy& Segment : Segments)
 		{
 			AllocatedSize += Segment.Points.GetAllocatedSize();
@@ -402,7 +404,7 @@ public:
 		{
 			AllocatedSize += ControlPoint.Points.GetAllocatedSize();
 		}
-		return AllocatedSize;
+		return IntCastChecked<uint32>(AllocatedSize);
 	}
 };
 #endif
@@ -1036,8 +1038,8 @@ ULandscapeSplinesComponent* ULandscapeSplinesComponent::GetStreamingSplinesCompo
 		OuterLandscape->GetLandscapeInfo())
 	{
 		FVector LandscapeLocalLocation = GetComponentTransform().GetRelativeTransform(OuterLandscape->LandscapeActorToWorld()).TransformPosition(LocalLocation);
-		const int32 ComponentIndexX = (LandscapeLocalLocation.X >= 0.0f) ? FMath::FloorToInt(LandscapeLocalLocation.X / OuterLandscape->ComponentSizeQuads) : FMath::CeilToInt(LandscapeLocalLocation.X / OuterLandscape->ComponentSizeQuads);
-		const int32 ComponentIndexY = (LandscapeLocalLocation.Y >= 0.0f) ? FMath::FloorToInt(LandscapeLocalLocation.Y / OuterLandscape->ComponentSizeQuads) : FMath::CeilToInt(LandscapeLocalLocation.Y / OuterLandscape->ComponentSizeQuads);
+		const int32 ComponentIndexX = (LandscapeLocalLocation.X >= 0.0f) ? FMath::FloorToInt32(LandscapeLocalLocation.X / OuterLandscape->ComponentSizeQuads) : FMath::CeilToInt32(LandscapeLocalLocation.X / OuterLandscape->ComponentSizeQuads);
+		const int32 ComponentIndexY = (LandscapeLocalLocation.Y >= 0.0f) ? FMath::FloorToInt32(LandscapeLocalLocation.Y / OuterLandscape->ComponentSizeQuads) : FMath::CeilToInt32(LandscapeLocalLocation.Y / OuterLandscape->ComponentSizeQuads);
 		ULandscapeComponent* LandscapeComponent = OuterLandscape->GetLandscapeInfo()->XYtoComponentMap.FindRef(FIntPoint(ComponentIndexX, ComponentIndexY));
 		if (LandscapeComponent)
 		{
@@ -1734,8 +1736,8 @@ FName ULandscapeSplineControlPoint::GetBestConnectionTo(FVector Destination) con
 			FVector SocketLocation = SocketTransform.GetTranslation();
 			FRotator SocketRotation = SocketTransform.GetRotation().Rotator();
 
-			float Score = (Destination - Location).Size() - (Destination - SocketLocation).Size(); // Score closer sockets higher
-			Score *= FMath::Abs(FVector::DotProduct((Destination - SocketLocation), SocketRotation.Vector())); // score closer rotation higher
+			const float Score = static_cast<float>((Destination - Location).Size() - (Destination - SocketLocation).Size() // Score closer sockets higher
+				* FMath::Abs(FVector::DotProduct((Destination - SocketLocation), SocketRotation.Vector()))); // score closer rotation higher
 
 			if (Score > BestScore)
 			{
@@ -2081,14 +2083,14 @@ void ULandscapeSplineControlPoint::UpdateSplinePoints(bool bUpdateCollision, boo
 		if (MeshComponent->VirtualTextureLodBias != VirtualTextureLodBias)
 		{
 			MeshComponent->Modify();
-			MeshComponent->VirtualTextureLodBias = VirtualTextureLodBias;
+			MeshComponent->VirtualTextureLodBias = static_cast<int8>(VirtualTextureLodBias);
 			MeshComponent->MarkRenderStateDirty();
 		}
 		
 		if (MeshComponent->VirtualTextureCullMips != VirtualTextureCullMips)
 		{
 			MeshComponent->Modify();
-			MeshComponent->VirtualTextureCullMips = VirtualTextureCullMips;
+			MeshComponent->VirtualTextureCullMips = static_cast<int8>(VirtualTextureCullMips);
 			MeshComponent->MarkRenderStateDirty();
 		}
 
@@ -2180,7 +2182,7 @@ void ULandscapeSplineControlPoint::UpdateSplinePoints(bool bUpdateCollision, boo
 			FVector StartLocation; FRotator StartRotation;
 			GetConnectionLocationAndRotation(Connection.GetNearConnection().SocketName, StartLocation, StartRotation);
 
-			const float Roll = FMath::DegreesToRadians(StartRotation.Roll);
+			const double Roll = FMath::DegreesToRadians(StartRotation.Roll);
 			const FVector Tangent = StartRotation.Vector();
 			const FVector BiNormal = FQuat(Tangent, -Roll).RotateVector((Tangent ^ FVector(0, 0, -1)).GetSafeNormal());
 			const FVector LeftPos = StartLocation - BiNormal * Width;
@@ -2206,7 +2208,7 @@ void ULandscapeSplineControlPoint::UpdateSplinePoints(bool bUpdateCollision, boo
 		FVector StartLocation; FRotator StartRotation;
 		GetConnectionLocationAndRotation(NAME_None, StartLocation, StartRotation);
 
-		const float Roll = FMath::DegreesToRadians(StartRotation.Roll);
+		const double Roll = FMath::DegreesToRadians(StartRotation.Roll);
 		const FVector Tangent = StartRotation.Vector();
 		const FVector BiNormal = FQuat(Tangent, -Roll).RotateVector((Tangent ^ FVector(0, 0, -1)).GetSafeNormal());
 		const FVector LeftPos = StartLocation - BiNormal * Width;
@@ -2645,7 +2647,7 @@ static float ApproxLength(const FInterpCurveVector& SplineInfo, const float Star
 	for (int32 i = 1; i <= ApproxSections; i++)
 	{
 		FVector NewPos = SplineInfo.Eval(FMath::Lerp(Start, End, (float)i / (float)ApproxSections), FVector::ZeroVector);
-		SplineLength += (NewPos - OldPos).Size();
+		SplineLength += static_cast<float>((NewPos - OldPos).Size());
 		OldPos = NewPos;
 	}
 
@@ -2744,8 +2746,8 @@ void ULandscapeSplineSegment::UpdateSplinePoints(bool bUpdateCollision, bool bUp
 	Falloffs.EndLeftSideLayer = Connections[1].ControlPoint->LeftSideLayerFalloffFactor * Connections[1].ControlPoint->SideFalloff;
 	Falloffs.StartRightSideLayer = Connections[0].ControlPoint->RightSideLayerFalloffFactor * Connections[0].ControlPoint->SideFalloff;
 	Falloffs.EndRightSideLayer = Connections[1].ControlPoint->RightSideLayerFalloffFactor * Connections[1].ControlPoint->SideFalloff;
-	const float StartRollDegrees = StartRotation.Roll * (Connections[0].TangentLen > 0 ? 1 : -1);
-	const float EndRollDegrees = EndRotation.Roll * (Connections[1].TangentLen > 0 ? -1 : 1);
+	const float StartRollDegrees = static_cast<float>(StartRotation.Roll * (Connections[0].TangentLen > 0 ? 1 : -1));
+	const float EndRollDegrees = static_cast<float>(EndRotation.Roll * (Connections[1].TangentLen > 0 ? -1 : 1));
 	const float StartRoll = FMath::DegreesToRadians(StartRollDegrees);
 	const float EndRoll = FMath::DegreesToRadians(EndRollDegrees);
 	const float StartMeshOffset = Connections[0].ControlPoint->SegmentMeshOffset;
@@ -2864,7 +2866,8 @@ void ULandscapeSplineSegment::UpdateSplinePoints(bool bUpdateCollision, bool bUp
 				Scale *= Width / USplineMeshComponent::GetAxisValue(MeshBounds.BoxExtent, CrossAxis(MeshEntry->ForwardAxis, MeshEntry->UpAxis));
 			}
 
-			const float MeshLength = FMath::Abs(USplineMeshComponent::GetAxisValue(MeshBounds.BoxExtent, MeshEntry->ForwardAxis) * 2 * USplineMeshComponent::GetAxisValue(Scale, MeshEntry->ForwardAxis));
+			const float MeshLength = static_cast<float>(FMath::Abs(USplineMeshComponent::GetAxisValue(MeshBounds.BoxExtent, MeshEntry->ForwardAxis) * 2.0 *
+				USplineMeshComponent::GetAxisValue(Scale, MeshEntry->ForwardAxis)));
 			float MeshT = (MeshLength / SplineLength);
 
 			// Improve our approximation if we're not going off the end of the spline
@@ -3130,8 +3133,8 @@ void ULandscapeSplineSegment::UpdateSplinePoints(bool bUpdateCollision, bool bUp
 			MeshComponent->TranslucencySortPriority = TranslucencySortPriority;
 
 			MeshComponent->RuntimeVirtualTextures = RuntimeVirtualTextures;
-			MeshComponent->VirtualTextureLodBias = VirtualTextureLodBias;
-			MeshComponent->VirtualTextureCullMips = VirtualTextureCullMips;
+			MeshComponent->VirtualTextureLodBias = static_cast<int8>(VirtualTextureLodBias);
+			MeshComponent->VirtualTextureCullMips = static_cast<int8>(VirtualTextureCullMips);
 			MeshComponent->VirtualTextureMainPassMaxDrawDistance = VirtualTextureMainPassMaxDrawDistance;
 			MeshComponent->VirtualTextureRenderPassType = VirtualTextureRenderPassType;
 
diff --git a/Engine/Source/Runtime/Landscape/Private/LandscapeSubsystem.cpp b/Engine/Source/Runtime/Landscape/Private/LandscapeSubsystem.cpp
--- a/Engine/Source/Runtime/Landscape/Private/LandscapeSubsystem.cpp
+++ b/Engine/Source/Runtime/Landscape/Private/LandscapeSubsystem.cpp
@@ -233,7 +233,7 @@ void ULandscapeSubsystem::BuildNanite()
 		return;
 	}
 
-	FScopedSlowTask SlowTask(Proxies.Num(), (LOCTEXT("Landscape_BuildNanite", "Building Nanite Landscape Meshes")));
+	FScopedSlowTask SlowTask(static_cast<float>(Proxies.Num()), (LOCTEXT("Landscape_BuildNanite", "Building Nanite Landscape Meshes")));
 	SlowTask.MakeDialog();
 
 	for (TWeakObjectPtr<ALandscapeProxy> ProxyPtr : Proxies)
